Missed file

#ttp 294228 - Regression: EDITOR: VIMBLUEPRINT: Animation stops when you drag a slider to its limits

#proj UE4

#summary Added code to prevent viewport throttling and changed MetaData setting so that it updates input controls without the need for a compile

Originally by [REDACTED_PERSON]

diff --git a/Engine/Source/Editor/Persona/Private/Persona.cpp b/Engine/Source/Editor/Persona/Private/Persona.cpp
--- a/Engine/Source/Editor/Persona/Private/Persona.cpp
+++ b/Engine/Source/Editor/Persona/Private/Persona.cpp
@@ -2309,6 +2309,11 @@ void FPersona::Tick(float DeltaTime)
 {
 	FBlueprintEditor::Tick(DeltaTime);
 
+	if (Viewport.IsValid())
+	{
+		Viewport.Pin()->RefreshViewport();
+	}
+
 	if (IsAnimationBeingRecorded())
 	{
 		// make sure you don't allow switch previewcomponent
