#TTP 320508 - ROCKETGDC: LICENSEE: UDN: Feature Request: PERSONA: Delete unused bones from Skeleton

#proj UE4

#summary Added ability to remove unused bones from a skeleton

diff --git a/Engine/Source/Developer/AssetTools/Private/AssetTypeActions/AssetTypeActions_Skeleton.cpp b/Engine/Source/Developer/AssetTools/Private/AssetTypeActions/AssetTypeActions_Skeleton.cpp
--- a/Engine/Source/Developer/AssetTools/Private/AssetTypeActions/AssetTypeActions_Skeleton.cpp
+++ b/Engine/Source/Developer/AssetTools/Private/AssetTypeActions/AssetTypeActions_Skeleton.cpp
@@ -287,7 +287,6 @@ void FAssetTypeActions_Skeleton::ExecuteRetargetSkeleton(TArray<TWeakObjectPtr<U
 	// warn the user to shut down any persona that is opened
 	if ( FMessageDialog::Open(EAppMsgType::YesNo, LOCTEXT("CloseReferencingEditors", "You need to close Persona or anything that references animation, mesh or animation blueprint before this step. Continue?")) == EAppReturnType::Yes )
 	{
-		TArray<UObject*> ObjectsToSync;
 		for (auto SkelIt = Skeletons.CreateConstIterator(); SkelIt; ++SkelIt)
 		{	
 			USkeleton * OldSkeleton = (*SkelIt).Get();
diff --git a/Engine/Source/Editor/Persona/Private/Persona.cpp b/Engine/Source/Editor/Persona/Private/Persona.cpp
--- a/Engine/Source/Editor/Persona/Private/Persona.cpp
+++ b/Engine/Source/Editor/Persona/Private/Persona.cpp
@@ -19,6 +19,7 @@
 #include "SAnimationBlendSpace.h"
 #include "SAnimationBlendSpace1D.h"
 #include "SKismetInspector.h"
+#include "SSkeletonWidget.h"
 
 #include "Editor/Kismet/Public/BlueprintEditorTabs.h"
 #include "Editor/Kismet/Public/BlueprintEditorModes.h"
@@ -438,6 +439,11 @@ FPersona::~FPersona()
 	FPersonaModule* PersonaModule = &FModuleManager::LoadModuleChecked<FPersonaModule>( "Persona" );
 	PersonaModule->GetMenuExtensibilityManager()->RemoveExtender(PersonaMenuExtender);
 
+	if(TargetSkeleton)
+	{
+		TargetSkeleton->UnregisterOnSkeletonHierarchyChanged(this);
+	}
+
 	if(Viewport.IsValid())
 	{
 		Viewport.Pin()->CleanupPersonaReferences();
@@ -717,6 +723,8 @@ void FPersona::InitPersona(const EToolkitMode::Type Mode, const TSharedPtr< clas
 	ObjectsBeingEdited.Add(TargetSkeleton);
 	check(TargetSkeleton);
 
+	TargetSkeleton->RegisterOnSkeletonHierarchyChanged( USkeleton::FOnSkeletonHierarchyChanged::CreateSP( this, &FPersona::OnSkeletonHierarchyChanged ) );
+
 	// We could modify the skeleton within Persona (add/remove sockets), so we need to enable undo/redo on it
 	TargetSkeleton->SetFlags( RF_Transactional );
 
@@ -772,6 +780,7 @@ void FPersona::InitPersona(const EToolkitMode::Type Mode, const TSharedPtr< clas
 			MenuBuilder.BeginSection("Persona", LOCTEXT("PersonaAssetMenuMenu", "Skeleton" ) );
 			{
 				MenuBuilder.AddMenuEntry( FPersonaCommands::Get().ChangeSkeletonPreviewMesh );
+				MenuBuilder.AddMenuEntry( FPersonaCommands::Get().RemoveUnusedBones );
 			}
 			MenuBuilder.EndSection();
 		}
@@ -1025,6 +1034,11 @@ void FPersona::CreateDefaultCommands()
 		FCanExecuteAction::CreateSP( this, &FPersona::CanChangeSkeletonPreviewMesh )
 		);
 
+	ToolkitCommands->MapAction( FPersonaCommands::Get().RemoveUnusedBones,
+		FExecuteAction::CreateSP( this, &FPersona::RemoveUnusedBones ),
+		FCanExecuteAction::CreateSP( this, &FPersona::CanRemoveBones )
+		);
+
 	// Generic deletion
 	ToolkitCommands->MapAction(FGenericCommands::Get().Delete, FExecuteAction::CreateSP(this, &FPersona::OnCommandGenericDelete));
 
@@ -1630,6 +1644,11 @@ void FPersona::ClearupPreviewMeshAnimNotifyStates()
 	}
 }
 
+void FPersona::OnSkeletonHierarchyChanged()
+{
+	OnSkeletonTreeChanged.Broadcast();
+}
+
 bool FPersona::IsInAScriptingMode() const
 {
 	return IsModeCurrent(FPersonaModes::AnimBlueprintEditMode);
@@ -2300,6 +2319,95 @@ void FPersona::ChangeSkeletonPreviewMesh()
 	}
 }
 
+void FPersona::RemoveUnusedBones()
+{
+	// Menu option cannot be called unless this is valid
+	if(TargetSkeleton)
+	{
+		TArray<FName> SkeletonBones;
+		const FReferenceSkeleton& RefSkeleton = TargetSkeleton->GetReferenceSkeleton();
+
+		for(int32 BoneIndex = 0; BoneIndex < RefSkeleton.GetNum(); ++BoneIndex)
+		{
+			SkeletonBones.Add(RefSkeleton.GetBoneName(BoneIndex));
+		}
+
+		FARFilter Filter;
+		Filter.ClassNames.Add(USkeletalMesh::StaticClass()->GetFName());
+
+		FString SkeletonString = FAssetData(TargetSkeleton).GetExportTextName();
+		Filter.TagsAndValues.Add(GET_MEMBER_NAME_CHECKED(USkeletalMesh, Skeleton), SkeletonString);
+
+		TArray<FAssetData> SkeletalMeshes;
+		FAssetRegistryModule& AssetRegistryModule = FModuleManager::LoadModuleChecked<FAssetRegistryModule>(TEXT("AssetRegistry"));
+		AssetRegistryModule.Get().GetAssets(Filter, SkeletalMeshes);
+
+		FText Message = FText::Format( LOCTEXT("TimeTakenWarning", "In order to verify bone use all Skeletal Meshes that use this skeleton will be loaded, this may take some time.\n\nProceed?\n\nNumber of Meshes: {0}"), FText::AsNumber(SkeletalMeshes.Num()) );
+		
+		if(FMessageDialog::Open( EAppMsgType::YesNo, Message ) == EAppReturnType::Yes)
+		{
+			const FText StatusUpdate = FText::Format(LOCTEXT("RemoveUnusedBones_ProcessingAssets", "Processing Skeletal Meshes for {0}"), FText::FromString(TargetSkeleton->GetName()) );
+			GWarn->BeginSlowTask(StatusUpdate, true );
+
+			// Loop through all SkeletalMeshes and remove the bones they use from our list
+			for ( int32 MeshIdx = 0; MeshIdx < SkeletalMeshes.Num(); ++MeshIdx )
+			{
+				GWarn->StatusUpdate( MeshIdx, SkeletalMeshes.Num(), StatusUpdate );
+
+				USkeletalMesh* Mesh = Cast<USkeletalMesh>(SkeletalMeshes[MeshIdx].GetAsset());
+				const FReferenceSkeleton& MeshRefSkeleton = Mesh->RefSkeleton;
+
+				for(int32 BoneIndex = 0; BoneIndex < MeshRefSkeleton.GetNum(); ++BoneIndex)
+				{
+					SkeletonBones.Remove(MeshRefSkeleton.GetBoneName(BoneIndex));
+					if(SkeletonBones.Num() == 0)
+					{
+						break;
+					}
+				}
+				if(SkeletonBones.Num() == 0)
+				{
+					break;
+				}
+			}
+
+			GWarn->EndSlowTask();
+
+			//Remove bones that are a parent to bones we aren't removing
+			for(int32 BoneIndex = RefSkeleton.GetNum() -1; BoneIndex >= 0; --BoneIndex)
+			{
+				FName CurrBoneName = RefSkeleton.GetBoneName(BoneIndex);
+				if(!SkeletonBones.Contains(CurrBoneName))
+				{
+					//We aren't removing this bone, so remove parent from list of bones to remove too
+					int32 ParentIndex = RefSkeleton.GetParentIndex(BoneIndex);
+					if(ParentIndex != INDEX_NONE)
+					{
+						SkeletonBones.Remove(RefSkeleton.GetBoneName(ParentIndex));
+					}
+				}
+			}
+
+			// If we have any bones left they are unused
+			if(SkeletonBones.Num() > 0)
+			{
+				const FText Message = FText::Format(LOCTEXT("RemoveBoneWarning", "Continuing will remove the following bones from the skeleton '{0}'. These bones are not being used by any of the SkeletalMeshes assigned to this skeleton\n\nOnce the bones have been removed all loaded animations for this skeleton will be recompressed (any that aren't loaded will be recompressed the next time they are loaded)."), FText::FromString(TargetSkeleton->GetName()) );
+
+				// Ask User whether they would like to remove the bones from the skeleton
+				if (SSkeletonBoneRemoval::ShowModal(SkeletonBones, Message))
+				{
+					//Remove these bones from the skeleton
+					TargetSkeleton->RemoveBonesFromSkeleton(SkeletonBones, true);
+				}
+			}
+			else
+			{
+				FMessageDialog::Open(EAppMsgType::Ok, LOCTEXT("NoBonesToRemove", "No unused bones were found."));
+			}
+		}
+	}
+}
+
 bool FPersona::IsAnimationBeingRecorded() const
 {
 	return (Recorder.GetAnimationObject()!=NULL);
@@ -2357,6 +2465,11 @@ bool FPersona::CanChangeSkeletonPreviewMesh() const
 	return PreviewComponent && PreviewComponent->SkeletalMesh;
 }
 
+bool FPersona::CanRemoveBones() const
+{
+	return PreviewComponent && PreviewComponent->SkeletalMesh;
+}
+
 bool FPersona::IsRecordAvailable() const
 {
 	return (GetCurrentMode() == FPersonaModes::AnimBlueprintEditMode);
@@ -2544,5 +2657,50 @@ void FAnimationRecorder::Record( USkeletalMeshComponent * Component, TArray<FTra
 	}
 }
 
+static class FMeshHierarchyCmd : private FSelfRegisteringExec
+{
+public:
+	/** Console commands, see embeded usage statement **/
+	virtual bool Exec( UWorld*, const TCHAR* Cmd, FOutputDevice& Ar ) OVERRIDE
+	{
+		bool bResult = false;
+		if(FParse::Command(&Cmd,TEXT("TMH")))
+		{
+			Ar.Log(TEXT("Starting Mesh Test"));
+			FARFilter Filter;
+			Filter.ClassNames.Add(USkeletalMesh::StaticClass()->GetFName());
+
+			TArray<FAssetData> SkeletalMeshes;
+			FAssetRegistryModule& AssetRegistryModule = FModuleManager::LoadModuleChecked<FAssetRegistryModule>(TEXT("AssetRegistry"));
+			AssetRegistryModule.Get().GetAssets(Filter, SkeletalMeshes);
+
+			const FText StatusUpdate = LOCTEXT("RemoveUnusedBones_ProcessingAssets", "Processing Skeletal Meshes");
+			GWarn->BeginSlowTask(StatusUpdate, true );
+
+			// go through all assets try load
+			for ( int32 MeshIdx = 0; MeshIdx < SkeletalMeshes.Num(); ++MeshIdx )
+			{
+				GWarn->StatusUpdate( MeshIdx, SkeletalMeshes.Num(), StatusUpdate );
+
+				USkeletalMesh* Mesh = Cast<USkeletalMesh>(SkeletalMeshes[MeshIdx].GetAsset());
+				if(Mesh->Skeleton)
+				{
+					FName MeshRoot = Mesh->RefSkeleton.GetBoneName(0);
+					FName SkelRoot = Mesh->Skeleton->GetReferenceSkeleton().GetBoneName(0);
+
+					if(MeshRoot != SkelRoot)
+					{
+						Ar.Logf( TEXT("Mesh Found '%s' %s->%s"), *Mesh->GetName(), *MeshRoot.ToString(), *SkelRoot.ToString());
+					}
+				}
+			}
+
+			GWarn->EndSlowTask();
+			Ar.Log(TEXT("Mesh Test Finished"));
+		}
+		return bResult;
+	}
+} MeshHierarchyCmdExec;
+
 #undef LOCTEXT_NAMESPACE
 
diff --git a/Engine/Source/Editor/Persona/Private/Persona.h b/Engine/Source/Editor/Persona/Private/Persona.h
--- a/Engine/Source/Editor/Persona/Private/Persona.h
+++ b/Engine/Source/Editor/Persona/Private/Persona.h
@@ -262,6 +262,9 @@ protected:
 	/** Clear up Preview Mesh's AnimNotifyStates. Called when undo or redo**/
 	void ClearupPreviewMeshAnimNotifyStates();
 
+	/** Handler connected to open skeleton so we get notified when it changes */
+	void OnSkeletonHierarchyChanged();
+
 protected:
 	USkeleton* TargetSkeleton;
 
@@ -519,6 +522,10 @@ private:
 	void ChangeSkeletonPreviewMesh();
 	bool CanChangeSkeletonPreviewMesh() const;
 
+	/** Remove unused bones from skeleton functions */
+	void RemoveUnusedBones();
+	bool CanRemoveBones() const;
+
 	/** Returns the editor objects that are applicable for our current mode (e.g mesh, animation etc) */
 	TArray<UObject*> GetEditorObjectsForMode(FName Mode) const;
 
diff --git a/Engine/Source/Editor/Persona/Private/PersonaCommands.cpp b/Engine/Source/Editor/Persona/Private/PersonaCommands.cpp
--- a/Engine/Source/Editor/Persona/Private/PersonaCommands.cpp
+++ b/Engine/Source/Editor/Persona/Private/PersonaCommands.cpp
@@ -10,4 +10,6 @@ void FPersonaCommands::RegisterCommands()
 	UI_COMMAND( RecordAnimation, "Export to Animation", "Export current graph to animation", EUserInterfaceActionType::Button, FInputGesture() );
 
 	UI_COMMAND( ChangeSkeletonPreviewMesh, "Set Preview Mesh as Default", "Changes the skeletons default preview mesh to the current open preview mesh. The skeleton will require saving after this action.", EUserInterfaceActionType::Button, FInputGesture() );
+
+	UI_COMMAND( RemoveUnusedBones, "Remove Unused Bones from Skeleton", "Removes any bones from the skeleton that are not used by any of its meshes. The skeleton and associated animations will require saving after this action.", EUserInterfaceActionType::Button, FInputGesture() );
 }
diff --git a/Engine/Source/Editor/Persona/Private/PersonaCommands.h b/Engine/Source/Editor/Persona/Private/PersonaCommands.h
--- a/Engine/Source/Editor/Persona/Private/PersonaCommands.h
+++ b/Engine/Source/Editor/Persona/Private/PersonaCommands.h
@@ -18,4 +18,7 @@ public:
 
 	// Command to allow users to set the skeletons preview mesh
 	TSharedPtr<FUICommandInfo> ChangeSkeletonPreviewMesh;
+
+	// Command to allow users to remove unused bones (not referenced by any skeletalmesh) from the skeleton
+	TSharedPtr<FUICommandInfo> RemoveUnusedBones;
 };
diff --git a/Engine/Source/Editor/UnrealEd/Private/SSkeletonWidget.cpp b/Engine/Source/Editor/UnrealEd/Private/SSkeletonWidget.cpp
--- a/Engine/Source/Editor/UnrealEd/Private/SSkeletonWidget.cpp
+++ b/Engine/Source/Editor/UnrealEd/Private/SSkeletonWidget.cpp
@@ -687,5 +687,162 @@ FReply SRemapFailures::CloseClicked()
 
 	return FReply::Handled();
 }
+
+void SSkeletonBoneRemoval::Construct( const FArguments& InArgs )
+{
+	bShouldContinue = false;
+	WidgetWindow = InArgs._WidgetWindow;
+
+	for (int32 I=0; I<InArgs._BonesToRemove.Num(); ++I)
+	{
+		BoneNames.Add( MakeShareable(new FName(InArgs._BonesToRemove[I])) ) ;
+	}
+
+	this->ChildSlot
+	[
+		SNew (SVerticalBox)
+
+		+SVerticalBox::Slot()
+		.AutoHeight() 
+		.Padding(2) 
+		.HAlign(HAlign_Fill)
+		[
+			SNew(SHorizontalBox)
+
+			+SHorizontalBox::Slot()
+			.AutoWidth()
+			.VAlign(VAlign_Center)
+			+SHorizontalBox::Slot()
+			[
+				SNew(STextBlock)
+				.Text(LOCTEXT("BoneRemovalLabel", "Bone Removal"))
+				.Font(FSlateFontInfo(FPaths::EngineContentDir() / TEXT("Slate/Fonts/Roboto-Regular.ttf"), 16))
+			]
+		]
+
+		+SVerticalBox::Slot()
+		.AutoHeight() 
+		.Padding(2, 10)
+		.HAlign(HAlign_Fill)
+		[
+			SNew(SSeparator)
+			.Orientation(Orient_Horizontal)
+		]
+
+		+SVerticalBox::Slot()
+		.AutoHeight()
+		.Padding(5)
+		[
+			SNew(STextBlock)
+			.WrapTextAt(400.f)
+			.Font( FSlateFontInfo( FPaths::EngineContentDir() / TEXT("Slate/Fonts/Roboto-Regular.ttf"), 10 ) )
+			.Text( InArgs._WarningMessage )
+		]
+
+		+SVerticalBox::Slot()
+		.AutoHeight() 
+		.Padding(5)
+		[
+			SNew(SSeparator)
+		]
+
+		+SVerticalBox::Slot()
+		.MaxHeight(300)
+		.Padding(5)
+		[
+			SNew(SListView< TSharedPtr<FName> >)
+			.OnGenerateRow(this, &SSkeletonBoneRemoval::GenerateSkeletonBoneRow)
+			.ListItemsSource(&BoneNames)
+			.HeaderRow
+			(
+			SNew(SHeaderRow)
+			+SHeaderRow::Column(TEXT("Bone Name"))
+			.DefaultLabel(NSLOCTEXT("SkeletonWidget", "BoneName", "Bone Name"))
+			)
+		]
+
+		+SVerticalBox::Slot()
+		.AutoHeight() 
+		.Padding(5)
+		[
+			SNew(SSeparator)
+		]
+
+		+SVerticalBox::Slot()
+		.AutoHeight()
+		.HAlign(HAlign_Right)
+		.VAlign(VAlign_Bottom)
+		[
+			SNew(SUniformGridPanel)
+			.SlotPadding(FEditorStyle::GetMargin("StandardDialog.SlotPadding"))
+			.MinDesiredSlotWidth(FEditorStyle::GetFloat("StandardDialog.MinDesiredSlotWidth"))
+			.MinDesiredSlotHeight(FEditorStyle::GetFloat("StandardDialog.MinDesiredSlotHeight"))
+			+SUniformGridPanel::Slot(0,0)
+			[
+				SNew(SButton) .HAlign(HAlign_Center)
+				.Text(LOCTEXT("BoneRemoval_Ok", "Ok"))
+				.OnClicked(this, &SSkeletonBoneRemoval::OnOk)
+				.HAlign(HAlign_Center)
+				.ContentPadding(FEditorStyle::GetMargin("StandardDialog.ContentPadding"))
+			]
+			+SUniformGridPanel::Slot(1,0)
+			[
+				SNew(SButton) .HAlign(HAlign_Center)
+				.Text(LOCTEXT("BoneRemoval_Cancel", "Cancel"))
+				.ContentPadding(FEditorStyle::GetMargin("StandardDialog.ContentPadding"))
+				.OnClicked(this, &SSkeletonBoneRemoval::OnCancel)
+			]
+		]
+	];
+}
+
+FReply SSkeletonBoneRemoval::OnOk()
+{
+	bShouldContinue = true;
+	CloseWindow();
+	return FReply::Handled();
+}
+
+FReply SSkeletonBoneRemoval::OnCancel()
+{
+	CloseWindow();
+	return FReply::Handled();
+}
+
+void SSkeletonBoneRemoval::CloseWindow()
+{
+	if ( WidgetWindow.IsValid() )
+	{
+		WidgetWindow.Pin()->RequestDestroyWindow();
+	}
+}
+
+bool SSkeletonBoneRemoval::ShowModal(const TArray<FName> BonesToRemove, const FText& WarningMessage)
+{
+	TSharedPtr<class SSkeletonBoneRemoval> DialogWidget;
+
+	TSharedPtr<SWindow> DialogWindow = SNew(SWindow)
+		.Title( LOCTEXT("RemapSkeleton", "Select Skeleton") )
+		.SupportsMinimize(false) .SupportsMaximize(false)
+		.SizingRule( ESizingRule::Autosized );
+
+	TSharedPtr<SBorder> DialogWrapper = 
+		SNew(SBorder)
+		.BorderImage(FEditorStyle::GetBrush("ToolPanel.GroupBorder"))
+		.Padding(4.0f)
+		[
+			SAssignNew(DialogWidget, SSkeletonBoneRemoval)
+			.WidgetWindow(DialogWindow)
+			.BonesToRemove(BonesToRemove)
+			.WarningMessage(WarningMessage)
+		];
+
+	DialogWindow->SetContent(DialogWrapper.ToSharedRef());
+
+	GEditor->EditorAddModalWindow(DialogWindow.ToSharedRef());
+
+	return DialogWidget.Get()->bShouldContinue;
+}
+
 #undef LOCTEXT_NAMESPACE 
 
diff --git a/Engine/Source/Editor/UnrealEd/Public/SSkeletonWidget.h b/Engine/Source/Editor/UnrealEd/Public/SSkeletonWidget.h
--- a/Engine/Source/Editor/UnrealEd/Public/SSkeletonWidget.h
+++ b/Engine/Source/Editor/UnrealEd/Public/SSkeletonWidget.h
@@ -374,6 +374,73 @@ private:
 	TArray< TSharedRef<FText> > FailedRemaps;
 };
 
+/////////////////////////////////////////////
+/** 
+ * Slate panel for choose displaying bones to remove
+ */
+class SSkeletonBoneRemoval : public SCompoundWidget
+{
+public:
+
+	SLATE_BEGIN_ARGS(SSkeletonBoneRemoval){}
+
+		/** The bones to remove (for list display) */
+		SLATE_ARGUMENT( TArray<FName>, BonesToRemove )
+
+		/** The window this panel has been placed in */
+		SLATE_ARGUMENT( TSharedPtr<SWindow>, WidgetWindow )
+
+		/** Message to display to the user */
+		SLATE_ARGUMENT( FText, WarningMessage )
+
+	SLATE_END_ARGS()	
+
+	/**
+	 * Constructs this widget
+	 *
+	 * @param	InArgs	The declaration data for this widget
+	 */
+	void Construct( const FArguments& InArgs );
+
+	/** Reference to our window */
+	TWeakPtr<SWindow> WidgetWindow;
+
+	/** Button Handlers */
+	FReply OnOk();
+	FReply OnCancel();
+
+	/** Handle closing to dialog window */
+	void CloseWindow();
+
+	/** Create an individual row for the bone name list */
+	TSharedRef<ITableRow> GenerateSkeletonBoneRow( TSharedPtr<FName> InBoneName, const TSharedRef<STableViewBase>& OwnerTable )
+	{
+		return
+		SNew( STableRow< TSharedPtr<FName> >, OwnerTable )
+		. Content()
+		[
+			SNew(STextBlock)
+			.Text(InBoneName->ToString())	
+		];
+	}
+
+	/**
+	 *  Show Modal window
+	 *
+	 * @param BonesToRemove		List of bones that will be removed
+	 * @param WarningMessage	Message to display to the user so they know what is going on
+	 *
+	 * @return true if successfully selected new skeleton
+	 */
+	static UNREALED_API bool ShowModal(const TArray<FName> BonesToRemove, const FText& WarningMessage);
+
+	/** Did the user choose to continue */
+	bool bShouldContinue;
+
+	/** List of bone names that will be removed */
+	TArray< TSharedPtr<FName> > BoneNames;
+};
+
 #undef LOCTEXT_NAMESPACE
 
 #endif // __SSkeletonWidget_h__
diff --git a/Engine/Source/Runtime/Engine/Classes/Animation/AnimationAsset.h b/Engine/Source/Runtime/Engine/Classes/Animation/AnimationAsset.h
--- a/Engine/Source/Runtime/Engine/Classes/Animation/AnimationAsset.h
+++ b/Engine/Source/Runtime/Engine/Classes/Animation/AnimationAsset.h
@@ -582,6 +582,10 @@ public:
 	ENGINE_API void SetSkeleton(USkeleton* NewSkeleton);
 	void ResetSkeleton(USkeleton* NewSkeleton);
 	virtual void PostLoad() OVERRIDE;
+
+	/** Validate our stored data against our skeleton and update accordingly */
+	void ValidateSkeleton();
+
 	virtual void Serialize(FArchive& Ar) OVERRIDE;
 
 #if WITH_EDITOR
diff --git a/Engine/Source/Runtime/Engine/Classes/Animation/Skeleton.h b/Engine/Source/Runtime/Engine/Classes/Animation/Skeleton.h
--- a/Engine/Source/Runtime/Engine/Classes/Animation/Skeleton.h
+++ b/Engine/Source/Runtime/Engine/Classes/Animation/Skeleton.h
@@ -511,8 +511,23 @@ protected:
 	 */
 	bool CreateReferenceSkeletonFromMesh(const USkeletalMesh * InSkeletalMesh, const TArray<int32> & RequiredRefBones);
 
-public: 
 #if WITH_EDITOR
+	DECLARE_MULTICAST_DELEGATE( FOnSkeletonHierarchyChangedMulticaster );
+	FOnSkeletonHierarchyChangedMulticaster OnSkeletonHierarchyChanged;
+
+	/** Call this when the skeleton has changed to fix dependent assets */
+	void HandleSkeletonHierarchyChange();
+
+public:
+	typedef FOnSkeletonHierarchyChangedMulticaster::FDelegate FOnSkeletonHierarchyChanged;
+
+	/** Registers a delegate to be called after notification has changed*/
+	ENGINE_API void RegisterOnSkeletonHierarchyChanged(const FOnSkeletonHierarchyChanged& Delegate);
+	ENGINE_API void UnregisterOnSkeletonHierarchyChanged(void * Unregister);
+
+	/** Removes the supplied bones from the skeleton */
+	ENGINE_API void RemoveBonesFromSkeleton(const TArray<FName>& BonesToRemove, bool bRemoveChildBones);
+
 	static const FName AnimNotifyTag;
 	static const TCHAR AnimNotifyTagDeliminator;
 #endif
diff --git a/Engine/Source/Runtime/Engine/Private/Animation/Skeleton.cpp b/Engine/Source/Runtime/Engine/Private/Animation/Skeleton.cpp
--- a/Engine/Source/Runtime/Engine/Private/Animation/Skeleton.cpp
+++ b/Engine/Source/Runtime/Engine/Private/Animation/Skeleton.cpp
@@ -816,6 +816,48 @@ void USkeleton::AddBoneToLOD(int32 LODIndex, int32 BoneIndex)
 	}
 }
 
+void USkeleton::RemoveBonesFromSkeleton( const TArray<FName>& BonesToRemove, bool bRemoveChildBones )
+{
+	TArray<int32> BonesRemoved = ReferenceSkeleton.RemoveBonesByName(BonesToRemove);
+	if(BonesRemoved.Num() > 0)
+	{
+		BonesRemoved.Sort();
+		for(int32 Index = BonesRemoved.Num()-1; Index >=0; --Index)
+		{
+			BoneTree.RemoveAt(Index);
+		}
+		HandleSkeletonHierarchyChange();
+	}
+}
+
+void USkeleton::HandleSkeletonHierarchyChange()
+{
+	MarkPackageDirty();
+
+	RegenerateGuid();
+
+	// Fix up loaded animations (any animations that aren't loaded will be fixed on load)
+	for(TObjectIterator<UAnimationAsset> It; It; ++It)
+	{
+		UAnimationAsset* CurrentAnimation = *It;
+		CurrentAnimation->ValidateSkeleton();
+	}
+
+	RefreshAllRetargetSources();
+
+	OnSkeletonHierarchyChanged.Broadcast();
+}
+
+void USkeleton::RegisterOnSkeletonHierarchyChanged(const FOnSkeletonHierarchyChanged& Delegate)
+{
+	OnSkeletonHierarchyChanged.Add(Delegate);
+}
+
+void USkeleton::UnregisterOnSkeletonHierarchyChanged(void * Unregister)
+{
+	OnSkeletonHierarchyChanged.RemoveAll(Unregister);
+}
+
 #endif
 
 void USkeleton::RegenerateGuid()
diff --git a/Engine/Source/Runtime/Engine/Private/AnimationAsset.cpp b/Engine/Source/Runtime/Engine/Private/AnimationAsset.cpp
--- a/Engine/Source/Runtime/Engine/Private/AnimationAsset.cpp
+++ b/Engine/Source/Runtime/Engine/Private/AnimationAsset.cpp
@@ -18,11 +18,8 @@ void UAnimationAsset::PostLoad()
 		Skeleton ->ConditionalPostLoad();
 	}
 
-	if (Skeleton && Skeleton->GetGuid() != SkeletonGuid)
-	{
-		// reset Skeleton
-		ResetSkeleton(Skeleton);
-	}
+	ValidateSkeleton();
+
 
 	check( Skeleton==NULL || SkeletonGuid.IsValid() );
 }
@@ -93,6 +90,16 @@ bool UAnimationAsset::GetAllAnimationSequencesReferred(TArray<UAnimSequence*>& A
 void UAnimationAsset::ReplaceReferredAnimations(const TMap<UAnimSequence*, UAnimSequence*>& ReplacementMap)
 {
 }
+
+void UAnimationAsset::ValidateSkeleton()
+{
+	if (Skeleton && Skeleton->GetGuid() != SkeletonGuid)
+	{
+		// reset Skeleton
+		ResetSkeleton(Skeleton);
+	}
+}
+
 #endif
 
 void FBoneContainer::InitializeTo(const TArray<FBoneIndexType>& InRequiredBoneIndexArray, UObject& InAsset)
diff --git a/Engine/Source/Runtime/Engine/Public/SkeletalMeshTypes.h b/Engine/Source/Runtime/Engine/Public/SkeletalMeshTypes.h
--- a/Engine/Source/Runtime/Engine/Public/SkeletalMeshTypes.h
+++ b/Engine/Source/Runtime/Engine/Public/SkeletalMeshTypes.h
@@ -145,6 +145,40 @@ private:
 	/** TMap to look up bone index from bone name. */
 	TMap<FName, int32>		NameToIndexMap;
 
+	/** Removes the specified bone, so long as it has no children. Returns whether we removed the bone or not */
+	bool RemoveIndividualBone(int32 BoneIndex, TArray<int32>& OutBonesRemoved)
+	{
+		bool bRemoveThisBone = true;
+
+		// Make sure we have no children
+		for(int32 CurrBoneIndex=BoneIndex+1; CurrBoneIndex < GetNum(); CurrBoneIndex++)
+		{
+			if( RefBoneInfo[CurrBoneIndex].ParentIndex == BoneIndex )
+			{
+				bRemoveThisBone = false;
+				break;
+			}
+		}
+
+		if(bRemoveThisBone)
+		{
+			// Update parent indices of bones further through the array
+			for(int32 CurrBoneIndex=BoneIndex+1; CurrBoneIndex < GetNum(); CurrBoneIndex++)
+			{
+				FMeshBoneInfo& Bone = RefBoneInfo[CurrBoneIndex];
+				if( Bone.ParentIndex > BoneIndex )
+				{
+					Bone.ParentIndex -= 1;
+				}
+			}
+
+			OutBonesRemoved.Add(BoneIndex);
+			RefBonePose.RemoveAt(BoneIndex, 1);
+			RefBoneInfo.RemoveAt(BoneIndex, 1);
+		}
+		return bRemoveThisBone;
+	}
+
 public:
 	void Allocate(int32 Size)
 	{
@@ -384,6 +418,25 @@ public:
 		}
 	}
 
+	/** Removes the supplied bones from the skeleton, unless they have children that aren't also going to be removed */
+	TArray<int32> RemoveBonesByName(const TArray<FName>& BonesToRemove)
+	{
+		TArray<int32> BonesRemoved;
+
+		const int32 NumBones = GetNum();
+		for(int32 BoneIndex=NumBones-1; BoneIndex>=0; BoneIndex--)
+		{
+			FMeshBoneInfo& Bone = RefBoneInfo[BoneIndex];
+
+			if(BonesToRemove.Contains(Bone.Name))
+			{
+				RemoveIndividualBone(BoneIndex, BonesRemoved);
+			}
+		}
+		RebuildNameToIndexMap();
+		return BonesRemoved;
+	}
+
 	void RebuildNameToIndexMap()
 	{
 		// Start by clearing the current map.
