Made default open the animation document when opening animation tab. At least it shows you can do more.

diff --git a/Engine/Source/Editor/Persona/Private/Persona.cpp b/Engine/Source/Editor/Persona/Private/Persona.cpp
--- a/Engine/Source/Editor/Persona/Private/Persona.cpp
+++ b/Engine/Source/Editor/Persona/Private/Persona.cpp
@@ -515,6 +515,17 @@ TSharedPtr<SWidget> FPersona::CreateEditorWidgetForAnimDocument(UObject* InAnimA
 	return Result;
 }
 
+void FPersona::ReinitMode()
+{
+	FName CurrentMode = GetCurrentMode();
+
+	if (CurrentMode == FPersonaModes::AnimationEditMode)
+	{
+		// if opening animation edit mode, open the animation editor
+		OpenNewAnimationDocumentTab(GetPreviewAnimationAsset());
+	}
+}
+
 TSharedPtr<SDockTab> FPersona::OpenNewAnimationDocumentTab(UObject* InAnimAsset)
 {
 	TSharedPtr<SDockTab> OpenedTab;
diff --git a/Engine/Source/Editor/Persona/Private/Persona.h b/Engine/Source/Editor/Persona/Private/Persona.h
--- a/Engine/Source/Editor/Persona/Private/Persona.h
+++ b/Engine/Source/Editor/Persona/Private/Persona.h
@@ -181,6 +181,9 @@ public:
 	/** Returns the editors preview scene */
 	FPreviewScene& GetPreviewScene() { return PreviewScene; }
 
+	/** Gets called when it's clicked via mode tab - Reinitialize the mode **/
+	void ReinitMode();
+
 public:
 	// IToolkit interface
 	virtual FName GetToolkitFName() const OVERRIDE;
diff --git a/Engine/Source/Editor/Persona/Private/Shared/PersonaMode.cpp b/Engine/Source/Editor/Persona/Private/Shared/PersonaMode.cpp
--- a/Engine/Source/Editor/Persona/Private/Shared/PersonaMode.cpp
+++ b/Engine/Source/Editor/Persona/Private/Shared/PersonaMode.cpp
@@ -69,6 +69,15 @@ void FPersonaAppMode::RegisterTabFactories(TSharedPtr<FTabManager> InTabManager)
 	BP->PushTabFactories(PersonaTabFactories);
 }
 
+void FPersonaAppMode::PostActivateMode()
+{
+	FApplicationMode::PostActivateMode();
+
+	if(MyPersona.IsValid())
+	{
+		MyPersona.Pin()->ReinitMode();
+	}
+}
 /////////////////////////////////////////////////////
 // FPersonaModeSharedData
 
diff --git a/Engine/Source/Editor/Persona/Private/Shared/PersonaMode.h b/Engine/Source/Editor/Persona/Private/Shared/PersonaMode.h
--- a/Engine/Source/Editor/Persona/Private/Shared/PersonaMode.h
+++ b/Engine/Source/Editor/Persona/Private/Shared/PersonaMode.h
@@ -127,6 +127,7 @@ protected:
 public:
 	// FApplicationMode interface
 	virtual void RegisterTabFactories(TSharedPtr<FTabManager> InTabManager) OVERRIDE;
+	virtual void PostActivateMode() OVERRIDE;
 	// End of FApplicationMode interface
 
 protected:
