Increased level for anim edit layout, otherwise it doesn't find where to put it.

diff --git a/Engine/Source/Editor/Persona/Private/AnimationMode/AnimationMode.cpp b/Engine/Source/Editor/Persona/Private/AnimationMode/AnimationMode.cpp
--- a/Engine/Source/Editor/Persona/Private/AnimationMode/AnimationMode.cpp
+++ b/Engine/Source/Editor/Persona/Private/AnimationMode/AnimationMode.cpp
@@ -142,7 +142,7 @@ FAnimEditAppMode::FAnimEditAppMode(TSharedPtr<FPersona> InPersona)
 	PersonaTabFactories.RegisterFactory(MakeShareable(new FSelectionDetailsSummoner(InPersona)));
 	PersonaTabFactories.RegisterFactory(MakeShareable(new FAnimAssetPropertiesSummoner(InPersona)));
 
-	TabLayout = FTabManager::NewLayout("Persona_AnimEditMode_Layout_v5")
+	TabLayout = FTabManager::NewLayout("Persona_AnimEditMode_Layout_v6")
 		->AddArea
 		(
 			FTabManager::NewPrimaryArea() ->SetOrientation(Orient_Vertical)
