#ttp 315494 - ANIM: support native anim notify class
#branch UE4
#summary Native anim notifies now supported

diff --git a/Engine/Source/Editor/Persona/Private/SAnimNotifyPanel.cpp b/Engine/Source/Editor/Persona/Private/SAnimNotifyPanel.cpp
--- a/Engine/Source/Editor/Persona/Private/SAnimNotifyPanel.cpp
+++ b/Engine/Source/Editor/Persona/Private/SAnimNotifyPanel.cpp
@@ -231,7 +231,8 @@ protected:
 	// Build up a "New Notify..." menu
 	void FillNewNotifyMenu(FMenuBuilder& MenuBuilder);
 	void FillNewNotifyStateMenu(FMenuBuilder& MenuBuilder);
-	void CreateNewNotifyAtCursor(FString NewNotifyName, FString NotifyClassPath);
+	void CreateNewBlueprintNotifyAtCursor(FString NewNotifyName, FString BlueprintPath);
+	void CreateNewNotifyAtCursor(FString NewNotifyName, UClass* NotifyClass);
 	void OnNewNotifyClicked();
 	void AddNewNotify(const FText& NewNotifyName, ETextCommit::Type CommitInfo);
 	void OnSetTriggerWeightNotifyClicked(int32 NotifyIndex);
@@ -1409,13 +1410,38 @@ void SAnimNotifyTrack::FillNewNotifyStateMenu(FMenuBuilder& MenuBuilder)
 
 			FUIAction UIAction;
 			UIAction.ExecuteAction.BindRaw(
-				this, &SAnimNotifyTrack::CreateNewNotifyAtCursor,
+				this, &SAnimNotifyTrack::CreateNewBlueprintNotifyAtCursor,
 				Label,
 				BlueprintPath);
 
 			MenuBuilder.AddMenuEntry(LabelText, Description, FSlateIcon(), UIAction);
 		}
 	}
+
+	MenuBuilder.BeginSection("NativeNotifyStates", LOCTEXT("NewStateNotifyMenu_Native", "Native Notify States"));
+	{
+		for(TObjectIterator<UClass> It ; It ; ++It)
+		{
+			UClass* Class = *It;
+
+			if(Class->IsChildOf(UAnimNotifyState::StaticClass()) && !Class->HasAllClassFlags(CLASS_Abstract) && !Class->IsInBlueprint())
+			{
+				// Found a native anim notify class (that isn't UAnimNotify)
+				const FText Description = LOCTEXT("NewNotifyStateSubMenu_NativeToolTip", "Add an existing native notify state");
+				FString Label = Class->GetName();
+				const FText LabelText = FText::FromString(Label);
+
+				FUIAction UIAction;
+				UIAction.ExecuteAction.BindRaw(
+					this, &SAnimNotifyTrack::CreateNewNotifyAtCursor,
+					Label,
+					Class);
+
+				MenuBuilder.AddMenuEntry(LabelText, Description, FSlateIcon(), UIAction);
+			}
+		}
+	}
+	MenuBuilder.EndSection();
 }
 
 void SAnimNotifyTrack::FillNewNotifyMenu(FMenuBuilder& MenuBuilder)
@@ -1444,7 +1470,7 @@ void SAnimNotifyTrack::FillNewNotifyMenu(FMenuBuilder& MenuBuilder)
 
 			FUIAction UIAction;
 			UIAction.ExecuteAction.BindRaw(
-				this, &SAnimNotifyTrack::CreateNewNotifyAtCursor,
+				this, &SAnimNotifyTrack::CreateNewBlueprintNotifyAtCursor,
 				Label,
 				BlueprintPath);
 
@@ -1452,6 +1478,31 @@ void SAnimNotifyTrack::FillNewNotifyMenu(FMenuBuilder& MenuBuilder)
 		}
 	}
 
+	MenuBuilder.BeginSection("NativeNotifies", LOCTEXT("NewNotifyMenu_Native", "Native Notifies"));
+	{
+		for(TObjectIterator<UClass> It ; It ; ++It)
+		{
+			UClass* Class = *It;
+
+			if(Class->IsChildOf(UAnimNotify::StaticClass()) && !Class->HasAllClassFlags(CLASS_Abstract) && !Class->IsInBlueprint())
+			{
+				// Found a native anim notify class (that isn't UAnimNotify)
+				const FText Description = LOCTEXT("NewNotifySubMenu_NativeToolTip", "Add an existing native notify");
+				FString Label = Class->GetName();
+				const FText LabelText = FText::FromString(Label);
+
+				FUIAction UIAction;
+				UIAction.ExecuteAction.BindRaw(
+					this, &SAnimNotifyTrack::CreateNewNotifyAtCursor,
+					Label,
+					Class);
+
+				MenuBuilder.AddMenuEntry(LabelText, Description, FSlateIcon(), UIAction);
+			}
+		}
+	}
+	MenuBuilder.EndSection();
+
 	MenuBuilder.BeginSection("AnimNotifyCustom", LOCTEXT("NewNotifySubMenu_Custom", "Custom"));
 	{
 		// now add custom anim notifiers
@@ -1469,7 +1520,7 @@ void SAnimNotifyTrack::FillNewNotifyMenu(FMenuBuilder& MenuBuilder)
 				UIAction.ExecuteAction.BindRaw(
 					this, &SAnimNotifyTrack::CreateNewNotifyAtCursor,
 					Label, 
-					FString());
+					(UClass*)nullptr);
 
 				MenuBuilder.AddMenuEntry( FText::FromString( Label ), Description, FSlateIcon(), UIAction);
 			}
@@ -1487,7 +1538,7 @@ void SAnimNotifyTrack::FillNewNotifyMenu(FMenuBuilder& MenuBuilder)
 	MenuBuilder.EndSection();
 }
 
-void SAnimNotifyTrack::CreateNewNotifyAtCursor(FString NewNotifyName, FString BlueprintPath)
+void SAnimNotifyTrack::CreateNewBlueprintNotifyAtCursor(FString NewNotifyName, FString BlueprintPath)
 {
 	TSubclassOf<UObject> BlueprintClass = NULL;
 	if( !BlueprintPath.IsEmpty() )
@@ -1495,7 +1546,12 @@ void SAnimNotifyTrack::CreateNewNotifyAtCursor(FString NewNotifyName, FString Bl
 		UBlueprint* BlueprintLibPtr = LoadObject<UBlueprint>(NULL, *BlueprintPath, NULL, 0, NULL);
 		BlueprintClass = Cast<UClass>(BlueprintLibPtr->GeneratedClass);
 	}
+	check(BlueprintClass)
+	CreateNewNotifyAtCursor(NewNotifyName, BlueprintClass);
+}
 
+void SAnimNotifyTrack::CreateNewNotifyAtCursor(FString NewNotifyName, UClass* NotifyClass)
+{
 	const float NewTime = LastClickedTime;
 	const FScopedTransaction Transaction( LOCTEXT("AddNotifyEvent", "Add Anim Notify") );
 	Sequence->Modify();
@@ -1508,9 +1564,9 @@ void SAnimNotifyTrack::CreateNewNotifyAtCursor(FString NewNotifyName, FString Bl
 	NewEvent.TriggerTimeOffset = GetTriggerTimeOffsetForType(Sequence->CalculateOffsetForNotify(NewTime));
 	NewEvent.TrackIndex = TrackIndex;
 
-	if( BlueprintClass )
+	if( NotifyClass )
 	{
-		class UObject * AnimNotifyClass = NewObject<UObject>(Sequence, *BlueprintClass);
+		class UObject * AnimNotifyClass = NewObject<UObject>(Sequence, NotifyClass);
 		NewEvent.NotifyStateClass = Cast<UAnimNotifyState>(AnimNotifyClass);
 		NewEvent.Notify = Cast<UAnimNotify>(AnimNotifyClass);
 
@@ -1528,24 +1584,6 @@ void SAnimNotifyTrack::CreateNewNotifyAtCursor(FString NewNotifyName, FString Bl
 
 	if(NewEvent.Notify)
 	{
-		struct FindClassPred
-		{
-			FindClassPred( const UClass* InClass )
-				: Class(InClass)
-			{
-			}
-			bool operator()( const FAssetData& AssetData ) const
-			{
-				if( AssetData.GetAsset()->IsA( Class ) )
-				{
-					return true;
-				}
-				return false;
-			}
-		private:
-			const UClass* Class;
-		};
-
 		TArray<FAssetData> SelectedAssets;
 		AssetSelectionUtils::GetSelectedAssets(SelectedAssets);
 
@@ -1553,7 +1591,12 @@ void SAnimNotifyTrack::CreateNewNotifyAtCursor(FString NewNotifyName, FString Bl
 		{
 			if(PropIt->GetBoolMetaData(TEXT("ExposeOnSpawn")))
 			{
-				const FAssetData* Asset = SelectedAssets.FindByPredicate( FindClassPred(PropIt->PropertyClass) );
+				UObjectProperty* Property = *PropIt;
+				const FAssetData* Asset = SelectedAssets.FindByPredicate([Property](const FAssetData& Other)
+				{
+					return Other.GetAsset()->IsA(Property->PropertyClass);
+				});
+
 				if( Asset )
 				{
 					uint8* Offset = (*PropIt)->ContainerPtrToValuePtr<uint8>(NewEvent.Notify);
@@ -2102,7 +2145,7 @@ void SAnimNotifyTrack::AddNewNotify(const FText& NewNotifyName, ETextCommit::Typ
 		FName NewName = FName( *NewNotifyName.ToString() );
 		SeqSkeleton->AddNewAnimationNotify(NewName);
 
-		CreateNewNotifyAtCursor(NewNotifyName.ToString(), "");
+		CreateNewNotifyAtCursor(NewNotifyName.ToString(), (UClass*)nullptr);
 	}
 
 	FSlateApplication::Get().DismissAllMenus();
diff --git a/Engine/Source/Runtime/Engine/Classes/Animation/AnimNotifies/AnimNotify.h b/Engine/Source/Runtime/Engine/Classes/Animation/AnimNotifies/AnimNotify.h
--- a/Engine/Source/Runtime/Engine/Classes/Animation/AnimNotifies/AnimNotify.h
+++ b/Engine/Source/Runtime/Engine/Classes/Animation/AnimNotifies/AnimNotify.h
@@ -4,8 +4,8 @@
 #pragma once
 #include "AnimNotify.generated.h"
 
-UCLASS(abstract, Blueprintable, const, hidecategories=Object, collapsecategories, MinimalAPI)
-class UAnimNotify : public UObject
+UCLASS(abstract, Blueprintable, const, hidecategories=Object, collapsecategories)
+class ENGINE_API UAnimNotify : public UObject
 {
 	GENERATED_UCLASS_BODY()
 
diff --git a/Engine/Source/Runtime/Engine/Classes/Animation/AnimNotifies/AnimNotifyState.h b/Engine/Source/Runtime/Engine/Classes/Animation/AnimNotifies/AnimNotifyState.h
--- a/Engine/Source/Runtime/Engine/Classes/Animation/AnimNotifies/AnimNotifyState.h
+++ b/Engine/Source/Runtime/Engine/Classes/Animation/AnimNotifies/AnimNotifyState.h
@@ -4,8 +4,8 @@
 #pragma once
 #include "AnimNotifyState.generated.h"
 
-UCLASS(abstract, editinlinenew, Blueprintable, const, hidecategories=Object, collapsecategories, MinimalAPI, meta=(ShowHiddenSelfPins))
-class UAnimNotifyState : public UObject
+UCLASS(abstract, editinlinenew, Blueprintable, const, hidecategories=Object, collapsecategories, meta=(ShowHiddenSelfPins))
+class ENGINE_API UAnimNotifyState : public UObject
 {
 	GENERATED_UCLASS_BODY()
 
diff --git a/Engine/Source/Runtime/Engine/Private/Animation/AnimInstance.cpp b/Engine/Source/Runtime/Engine/Private/Animation/AnimInstance.cpp
--- a/Engine/Source/Runtime/Engine/Private/Animation/AnimInstance.cpp
+++ b/Engine/Source/Runtime/Engine/Private/Animation/AnimInstance.cpp
@@ -750,21 +750,9 @@ void UAnimInstance::TriggerAnimNotifies(float DeltaSeconds)
 			continue;
 		}
 
-		// This checking-for-blueprint class is a hack to allow the old UAnimNotify_* notifies to work for a little longer until they are all removed.
-		// Once done, we can remove passing the UAnimNotify to the custom event notifies (since they cant contain custom/user data anyways).
-		// The "is blue print class" can just become a "is notify != null".
-		bool bIsBlueprintNotify = false;
-		if( AnimNotifyEvent->Notify != NULL)
+		if(AnimNotifyEvent->Notify != NULL)
 		{
-			if( !AnimNotifies[Index]->Notify->GetClass()->HasAllClassFlags(CLASS_Native) )
-			{
-				bIsBlueprintNotify = true;
-			}
-		}
-
-		if( bIsBlueprintNotify )
-		{
-			// Blueprint notify: just call Notify. UAnimNotify will forward this to the blueprintable event which will do the work.
+			// Implemented notify: just call Notify. UAnimNotify will forward this to the event which will do the work.
 			AnimNotifyEvent->Notify->Notify(SkelMeshComp, Cast<UAnimSequence>(AnimNotifyEvent->Notify->GetOuter()));
 		}
 		else if( AnimNotifyEvent->NotifyName != NAME_None )
