AnimTrail Notifies now display with their Display Name meta data rather than the class name

diff --git a/Engine/Source/Editor/Persona/Private/SAnimNotifyPanel.cpp b/Engine/Source/Editor/Persona/Private/SAnimNotifyPanel.cpp
--- a/Engine/Source/Editor/Persona/Private/SAnimNotifyPanel.cpp
+++ b/Engine/Source/Editor/Persona/Private/SAnimNotifyPanel.cpp
@@ -1428,7 +1428,7 @@ void SAnimNotifyTrack::FillNewNotifyStateMenu(FMenuBuilder& MenuBuilder)
 			{
 				// Found a native anim notify class (that isn't UAnimNotify)
 				const FText Description = LOCTEXT("NewNotifyStateSubMenu_NativeToolTip", "Add an existing native notify state");
-				FString Label = Class->GetName();
+				FString Label = Class->GetDisplayNameText().ToString();
 				const FText LabelText = FText::FromString(Label);
 
 				FUIAction UIAction;
