Fix socket renaming display bug in Persona

diff --git a/Engine/Source/Editor/Persona/Private/SSkeletonTree.cpp b/Engine/Source/Editor/Persona/Private/SSkeletonTree.cpp
--- a/Engine/Source/Editor/Persona/Private/SSkeletonTree.cpp
+++ b/Engine/Source/Editor/Persona/Private/SSkeletonTree.cpp
@@ -507,6 +507,7 @@ void FDisplayedSocketInfo::GenerateWidgetForNameColumn( TSharedPtr< SHorizontalB
 
 	FString ToolTip = GetSocketToolTip();
 
+	TAttribute<FText> SocketNameAttr = TAttribute<FText>::Create(TAttribute<FText>::FGetter::CreateSP(this, &FDisplayedSocketInfo::GetSocketNameAsText));
 	TSharedPtr< SInlineEditableTextBlock > InlineWidget;
 
 	Box->AddSlot()
@@ -514,7 +515,7 @@ void FDisplayedSocketInfo::GenerateWidgetForNameColumn( TSharedPtr< SHorizontalB
 	[
 		SAssignNew( InlineWidget, SInlineEditableTextBlock )
 			.ColorAndOpacity( TextColor )
-			.Text( FText::FromName(SocketData->SocketName) )
+			.Text( SocketNameAttr )
 			.HighlightText( FilterText )
 			.Font( TextFont )
 			.ToolTipText( ToolTip )
diff --git a/Engine/Source/Editor/Persona/Private/SSkeletonTree.h b/Engine/Source/Editor/Persona/Private/SSkeletonTree.h
--- a/Engine/Source/Editor/Persona/Private/SSkeletonTree.h
+++ b/Engine/Source/Editor/Persona/Private/SSkeletonTree.h
@@ -242,6 +242,9 @@ public:
 	/** Requests a rename on the socket item */
 	virtual void RequestRename() OVERRIDE;
 
+	/** Return socket name as FText for display in skeleton tree */
+	FText GetSocketNameAsText() const { return FText::FromName(SocketData->SocketName); }
+
 protected:
 	/** Hidden constructor, always use Make above */
 	FDisplayedSocketInfo( USkeletalMeshSocket* InSource, ESocketParentType::Type InParentType )
