#TTP 328283 - LICENSEE: UDN: Preview attached and Physics assets use strong reference

#proj UE4

#summary Changed preview asset type from raw pointer to TAssetPtr

diff --git a/Engine/Source/Editor/Persona/Private/Persona.cpp b/Engine/Source/Editor/Persona/Private/Persona.cpp
--- a/Engine/Source/Editor/Persona/Private/Persona.cpp
+++ b/Engine/Source/Editor/Persona/Private/Persona.cpp
@@ -747,24 +747,6 @@ void FPersona::InitPersona(const EToolkitMode::Type Mode, const TSharedPtr< clas
 	// We could modify the skeleton within Persona (add/remove sockets), so we need to enable undo/redo on it
 	TargetSkeleton->SetFlags( RF_Transactional );
 
-	// Validate the skeletons attached objects and display a notification to the user if any were broken
-	int32 NumBrokenAssets = TargetSkeleton->ValidatePreviewAttachedObjects();
-	if(NumBrokenAssets > 0)
-	{
-		// Tell the user that there were assets that could not be loaded
-		FFormatNamedArguments Args;
-		Args.Add( TEXT("NumBrokenAssets"), NumBrokenAssets );
-		FNotificationInfo Info( FText::Format( LOCTEXT( "MissingPreviewAttachedAssets", "{NumBrokenAssets} attached assets could not be found on loading and were removed" ), Args ) );
-
-		Info.bUseLargeFont = false;
-		Info.ExpireDuration = 5.0f;
-
-		TSharedPtr<SNotificationItem> Notification = FSlateNotificationManager::Get().AddNotification( Info );
-		if ( Notification.IsValid() )
-		{
-			Notification->SetCompletionState( SNotificationItem::CS_Fail );
-		}
-	}
 	// Initialize the asset editor and spawn tabs
 	const TSharedRef<FTabManager::FLayout> DummyLayout = FTabManager::NewLayout("NullLayout")->AddArea(FTabManager::NewPrimaryArea());
 	const bool bCreateDefaultStandaloneMenu = true;
@@ -861,6 +843,30 @@ void FPersona::InitPersona(const EToolkitMode::Type Mode, const TSharedPtr< clas
 		}
 	}
 
+	// Validate the skeleton/meshes attached objects and display a notification to the user if any were broken
+	int32 NumBrokenAssets = TargetSkeleton->ValidatePreviewAttachedObjects();
+	if (PreviewComponent->SkeletalMesh)
+	{
+		NumBrokenAssets += PreviewComponent->SkeletalMesh->ValidatePreviewAttachedObjects();
+	}
+
+	if (NumBrokenAssets > 0)
+	{
+		// Tell the user that there were assets that could not be loaded
+		FFormatNamedArguments Args;
+		Args.Add(TEXT("NumBrokenAssets"), NumBrokenAssets);
+		FNotificationInfo Info(FText::Format(LOCTEXT("MissingPreviewAttachedAssets", "{NumBrokenAssets} attached assets could not be found on loading and were removed"), Args));
+
+		Info.bUseLargeFont = false;
+		Info.ExpireDuration = 5.0f;
+
+		TSharedPtr<SNotificationItem> Notification = FSlateNotificationManager::Get().AddNotification(Info);
+		if (Notification.IsValid())
+		{
+			Notification->SetCompletionState(SNotificationItem::CS_Fail);
+		}
+	}
+
 	PreviewComponent->SetAnimClass(AnimBlueprint ? AnimBlueprint->GeneratedClass : NULL);
 
 	// We always want a preview instance unless we are using blueprints so that bone manipulation works
@@ -1911,7 +1917,7 @@ void FPersona::DuplicateAndSelectSocket( const FSelectedSocketInfo& SocketInfoTo
 			}
 			FPreviewAttachedObjectPair NewPair = Pair;
 			NewPair.AttachedTo = NewSocket->SocketName;
-			AttachObjectToPreviewComponent( NewPair.Object, NewPair.AttachedTo, &TargetSkeleton->PreviewAttachedAssetContainer );
+			AttachObjectToPreviewComponent( NewPair.GetAttachedObject(), NewPair.AttachedTo, &TargetSkeleton->PreviewAttachedAssetContainer );
 		}
 	}
 
@@ -1988,7 +1994,7 @@ void FPersona::AddPreviewAttachedObjects()
 		{
 			FPreviewAttachedObjectPair& PreviewAttachedObject = Mesh->PreviewAttachedAssetContainer[i];
 
-			AttachObjectToPreviewComponent(PreviewAttachedObject.Object, PreviewAttachedObject.AttachedTo);
+			AttachObjectToPreviewComponent(PreviewAttachedObject.GetAttachedObject(), PreviewAttachedObject.AttachedTo);
 		}
 	}
 
@@ -1997,7 +2003,7 @@ void FPersona::AddPreviewAttachedObjects()
 	{
 		FPreviewAttachedObjectPair& PreviewAttachedObject = TargetSkeleton->PreviewAttachedAssetContainer[i];
 
-		AttachObjectToPreviewComponent(PreviewAttachedObject.Object, PreviewAttachedObject.AttachedTo);
+		AttachObjectToPreviewComponent(PreviewAttachedObject.GetAttachedObject(), PreviewAttachedObject.AttachedTo);
 	}
 }
 
@@ -2101,14 +2107,14 @@ void FPersona::RemoveAttachedComponent( bool bRemovePreviewAttached /* = true */
 	{
 		for(auto Iter = TargetSkeleton->PreviewAttachedAssetContainer.CreateConstIterator(); Iter; ++Iter)
 		{
-			PreviewAttachedObjects.FindOrAdd(Iter->Object).Add(Iter->AttachedTo);
+			PreviewAttachedObjects.FindOrAdd(Iter->GetAttachedObject()).Add(Iter->AttachedTo);
 		}
 
 		if ( USkeletalMesh* PreviewMesh = GetMesh() )
 		{
 			for(auto Iter = PreviewMesh->PreviewAttachedAssetContainer.CreateConstIterator(); Iter; ++Iter)
 			{
-				PreviewAttachedObjects.FindOrAdd(Iter->Object).Add(Iter->AttachedTo);
+				PreviewAttachedObjects.FindOrAdd(Iter->GetAttachedObject()).Add(Iter->AttachedTo);
 			}
 		}
 	}
diff --git a/Engine/Source/Editor/Persona/Private/SSkeletonTree.cpp b/Engine/Source/Editor/Persona/Private/SSkeletonTree.cpp
--- a/Engine/Source/Editor/Persona/Private/SSkeletonTree.cpp
+++ b/Engine/Source/Editor/Persona/Private/SSkeletonTree.cpp
@@ -1892,7 +1892,7 @@ void SSkeletonTree::DeleteAttachedObjects( FPreviewAssetAttachContainer& Attache
 	for(auto Iter = AttachedAssets.CreateIterator(); Iter; ++Iter)
 	{
 		FPreviewAttachedObjectPair& Pair = (*Iter);
-		PersonaPtr.Pin()->RemoveAttachedObjectFromPreviewComponent(Pair.Object, Pair.AttachedTo);
+		PersonaPtr.Pin()->RemoveAttachedObjectFromPreviewComponent(Pair.GetAttachedObject(), Pair.AttachedTo);
 	}
 
 	AttachedAssets.ClearAllAttachedObjects();
@@ -2306,8 +2306,8 @@ void SSkeletonTree::RenameSocketAttachments(FName& OldSocketName, FName& NewSock
 		{
 			Pair.AttachedTo = NewSocketName;
 		}
-		PersonaPtr.Pin()->RemoveAttachedObjectFromPreviewComponent(Pair.Object, OldSocketName);
-		PersonaPtr.Pin()->AttachObjectToPreviewComponent(Pair.Object, Pair.AttachedTo);
+		PersonaPtr.Pin()->RemoveAttachedObjectFromPreviewComponent(Pair.GetAttachedObject(), OldSocketName);
+		PersonaPtr.Pin()->AttachObjectToPreviewComponent(Pair.GetAttachedObject(), Pair.AttachedTo);
 	}
 
 	if ( PersonaPtr.IsValid() )
@@ -2325,8 +2325,8 @@ void SSkeletonTree::RenameSocketAttachments(FName& OldSocketName, FName& NewSock
 				{
 					Pair.AttachedTo = NewSocketName;
 				}
-				PersonaPtr.Pin()->RemoveAttachedObjectFromPreviewComponent(Pair.Object, OldSocketName);
-				PersonaPtr.Pin()->AttachObjectToPreviewComponent(Pair.Object, Pair.AttachedTo);
+				PersonaPtr.Pin()->RemoveAttachedObjectFromPreviewComponent(Pair.GetAttachedObject(), OldSocketName);
+				PersonaPtr.Pin()->AttachObjectToPreviewComponent(Pair.GetAttachedObject(), Pair.AttachedTo);
 			}
 		}
 	}
@@ -2453,12 +2453,12 @@ void SSkeletonTree::AddAttachedAssets( const FPreviewAssetAttachContainer& Attac
 	{
 		const FPreviewAttachedObjectPair& Pair = (*Iter);
 
-		if ( !FilterText.IsEmpty() && !Pair.Object->GetName().Contains( FilterText.ToString() ) )
+		if ( !FilterText.IsEmpty() && !Pair.GetAttachedObject()->GetName().Contains( FilterText.ToString() ) )
 		{
 			continue;
 		}
 
-		TSharedRef<FDisplayedAttachedAssetInfo> DisplayInfo = FDisplayedAttachedAssetInfo::Make(Pair.AttachedTo, Pair.Object, TargetSkeleton, PersonaPtr, SharedThis( this ) );
+		TSharedRef<FDisplayedAttachedAssetInfo> DisplayInfo = FDisplayedAttachedAssetInfo::Make(Pair.AttachedTo, Pair.GetAttachedObject(), TargetSkeleton, PersonaPtr, SharedThis( this ) );
 		DisplayMirror.Add(DisplayInfo);
 
 		// for now it is a failure to not find where the asset is attached. Its possible that this might have to be changed to unloading the asset
diff --git a/Engine/Source/Runtime/Core/Public/UObject/ObjectVersion.h b/Engine/Source/Runtime/Core/Public/UObject/ObjectVersion.h
--- a/Engine/Source/Runtime/Core/Public/UObject/ObjectVersion.h
+++ b/Engine/Source/Runtime/Core/Public/UObject/ObjectVersion.h
@@ -526,7 +526,10 @@ enum EUnrealEngineObjectUE4Version
 	VER_UE4_VEHICLES_UNIT_CHANGE,
 	// removes NANs from all animations when loaded
 	// now importing should detect NaNs, so we should not have NaNs in source data
-	VER_UE4_ANIMATION_REMOVE_NANS, 
+	VER_UE4_ANIMATION_REMOVE_NANS,
+	// Change skeleton preview attached assets property type
+	VER_UE4_SKELETON_ASSET_PROPERTY_TYPE_CHANGE,
+ 
 	// -----<new versions can be added before this line>-------------------------------------------------
 	// - this needs to be the last line (see note below)
 	VER_UE4_AUTOMATIC_VERSION_PLUS_ONE,
diff --git a/Engine/Source/Runtime/Engine/Classes/Animation/PreviewAssetAttachComponent.cpp b/Engine/Source/Runtime/Engine/Classes/Animation/PreviewAssetAttachComponent.cpp
--- a/Engine/Source/Runtime/Engine/Classes/Animation/PreviewAssetAttachComponent.cpp
+++ b/Engine/Source/Runtime/Engine/Classes/Animation/PreviewAssetAttachComponent.cpp
@@ -7,7 +7,7 @@ void FPreviewAssetAttachContainer::AddAttachedObject( UObject* AttachObject, FNa
 {
 	FPreviewAttachedObjectPair Pair;
 	Pair.AttachedTo = AttachPointName;
-	Pair.Object = AttachObject;
+	Pair.SetAttachedObject(AttachObject);
 	AttachedObjects.Add( Pair );
 }
 
@@ -17,7 +17,7 @@ void FPreviewAssetAttachContainer::RemoveAttachedObject( UObject* ObjectToRemove
 	{
 		FPreviewAttachedObjectPair& Pair = AttachedObjects[i];
 
-		if( Pair.Object == ObjectToRemove && Pair.AttachedTo == AttachName )
+		if( Pair.GetAttachedObject() == ObjectToRemove && Pair.AttachedTo == AttachName )
 		{
 			AttachedObjects.RemoveAtSwap( i, 1, false );
 			break;
@@ -31,7 +31,7 @@ UObject* FPreviewAssetAttachContainer::GetAttachedObjectByAttachName( FName Atta
 	{
 		if( AttachedObjects[i].AttachedTo == AttachName )
 		{
-			return AttachedObjects[i].Object;
+			return AttachedObjects[i].GetAttachedObject();
 		}
 	}
 	return NULL;
@@ -65,4 +65,30 @@ TIterator FPreviewAssetAttachContainer::CreateIterator()
 void FPreviewAssetAttachContainer::RemoveAtSwap( int32 Index, int32 Count /* = 1 */, bool bAllowShrinking /*=true */ )
 {
 	AttachedObjects.RemoveAtSwap( Index, Count, bAllowShrinking );
+}
+
+void FPreviewAssetAttachContainer::SaveAttachedObjectsFromDeprecatedProperties()
+{
+	for (FPreviewAttachedObjectPair& Pair : AttachedObjects)
+	{
+		Pair.SaveAttachedObjectFromDeprecatedProperty();
+	}
+}
+
+int32 FPreviewAssetAttachContainer::ValidatePreviewAttachedObjects()
+{
+	int32 NumBrokenAssets = 0;
+
+	for (int32 i = AttachedObjects.Num() - 1; i >= 0; --i)
+	{
+		FPreviewAttachedObjectPair& PreviewAttachedObject = AttachedObjects[i];
+
+		if (!PreviewAttachedObject.GetAttachedObject())
+		{
+			AttachedObjects.RemoveAtSwap(i);
+			++NumBrokenAssets;
+		}
+	}
+
+	return NumBrokenAssets;
 }
\ No newline at end of file
diff --git a/Engine/Source/Runtime/Engine/Classes/Animation/PreviewAssetAttachComponent.h b/Engine/Source/Runtime/Engine/Classes/Animation/PreviewAssetAttachComponent.h
--- a/Engine/Source/Runtime/Engine/Classes/Animation/PreviewAssetAttachComponent.h
+++ b/Engine/Source/Runtime/Engine/Classes/Animation/PreviewAssetAttachComponent.h
@@ -16,13 +16,51 @@ struct FPreviewAttachedObjectPair
 {
 	GENERATED_USTRUCT_BODY()
 
+private:
+	/** the object to be attached */
+	UPROPERTY()
+	TAssetPtr<class UObject> AttachedObject;
+
+	UPROPERTY()
+	UObject* Object_DEPRECATED;
+
+public:
+
+	FPreviewAttachedObjectPair() : Object_DEPRECATED(NULL) {}
+
 	/** The name of the attach point of the Object (for example a bone or socket name) */
 	UPROPERTY()
 	FName AttachedTo;
 
-	/** the object to be attached */
-	UPROPERTY()
-	UObject* Object;
+	void SaveAttachedObjectFromDeprecatedProperty()
+	{
+		if (Object_DEPRECATED)
+		{
+			AttachedObject = Object_DEPRECATED;
+			Object_DEPRECATED = NULL;
+		}
+	}
+
+	UObject* GetAttachedObject() const
+	{
+		UObject* AttachedObjectPtr = AttachedObject.Get();
+		if (!AttachedObjectPtr)
+		{
+			// if preview mesh isn't loaded, see if we have set
+			FStringAssetReference AttachedObjectStringRef = AttachedObject.ToStringReference();
+			// load it since now is the time to load
+			if (!AttachedObjectStringRef.AssetLongPathname.IsEmpty())
+			{
+				AttachedObjectPtr = Cast<UObject>(StaticLoadObject(UObject::StaticClass(), NULL, *AttachedObjectStringRef.AssetLongPathname, NULL, LOAD_None, NULL));
+			}
+		}
+		return AttachedObjectPtr;
+	}
+
+	void SetAttachedObject(UObject* InObject)
+	{
+		AttachedObject = InObject;
+	}
 };
 
 // Iterators
@@ -89,6 +127,16 @@ public:
 	 */
 	ENGINE_API void RemoveAtSwap( int32 Index, int32 Count = 1, bool bAllowShrinking = true );
 
+	/**
+	 * Helper function to fix up attached objects after property deprecation
+	 */
+	ENGINE_API void SaveAttachedObjectsFromDeprecatedProperties();
+
+	/**
+	 * Helper function to remove invalid attached object references
+	 */
+	ENGINE_API int32 ValidatePreviewAttachedObjects();
+
 private:
 	UPROPERTY()
 	TArray<FPreviewAttachedObjectPair> AttachedObjects;
diff --git a/Engine/Source/Runtime/Engine/Classes/Engine/SkeletalMesh.h b/Engine/Source/Runtime/Engine/Classes/Engine/SkeletalMesh.h
--- a/Engine/Source/Runtime/Engine/Classes/Engine/SkeletalMesh.h
+++ b/Engine/Source/Runtime/Engine/Classes/Engine/SkeletalMesh.h
@@ -675,6 +675,13 @@ public:
 	 */
 	ENGINE_API const TArray<USkeletalMeshSocket*>& GetActiveSocketList() const;
 
+	/**
+	* Makes sure all attached objects are valid and removes any that aren't.
+	*
+	* @return		NumberOfBrokenAssets
+	*/
+	ENGINE_API int32 ValidatePreviewAttachedObjects();
+
 #endif // #if WITH_EDITOR
 
 	/**
diff --git a/Engine/Source/Runtime/Engine/Private/Animation/Skeleton.cpp b/Engine/Source/Runtime/Engine/Private/Animation/Skeleton.cpp
--- a/Engine/Source/Runtime/Engine/Private/Animation/Skeleton.cpp
+++ b/Engine/Source/Runtime/Engine/Private/Animation/Skeleton.cpp
@@ -136,6 +136,11 @@ void USkeleton::Serialize( FArchive& Ar )
 	{
 		Ar << Guid;
 	}
+
+	if (Ar.UE4Ver() < VER_UE4_SKELETON_ASSET_PROPERTY_TYPE_CHANGE)
+	{
+		PreviewAttachedAssetContainer.SaveAttachedObjectsFromDeprecatedProperties();
+	}
 }
 
 /** Remove this function when VER_UE4_REFERENCE_SKELETON_REFACTOR is removed. */
@@ -741,19 +746,7 @@ void USkeleton::SetPreviewMesh(USkeletalMesh* PreviewMesh, bool bMarkAsDirty/*=t
 
 int32 USkeleton::ValidatePreviewAttachedObjects()
 {
-	int32 NumBrokenAssets = 0;
-
-	//Load up attached objects
-	for(int32 i = PreviewAttachedAssetContainer.Num()-1; i >= 0; --i)
-	{
-		FPreviewAttachedObjectPair& PreviewAttachedObject = PreviewAttachedAssetContainer[i];
-
-		if(!PreviewAttachedObject.Object)
-		{
-			PreviewAttachedAssetContainer.RemoveAtSwap(i);
-			++NumBrokenAssets;
-		}
-	}
+	int32 NumBrokenAssets = PreviewAttachedAssetContainer.ValidatePreviewAttachedObjects();
 
 	if(NumBrokenAssets > 0)
 	{
diff --git a/Engine/Source/Runtime/Engine/Private/SkeletalMesh.cpp b/Engine/Source/Runtime/Engine/Private/SkeletalMesh.cpp
--- a/Engine/Source/Runtime/Engine/Private/SkeletalMesh.cpp
+++ b/Engine/Source/Runtime/Engine/Private/SkeletalMesh.cpp
@@ -2366,6 +2366,11 @@ void USkeletalMesh::Serialize( FArchive& Ar )
 		// now they're in the Materials array so we need to move them over
 		MoveDeprecatedShadowFlagToMaterials();
 	}
+
+	if (Ar.UE4Ver() < VER_UE4_SKELETON_ASSET_PROPERTY_TYPE_CHANGE)
+	{
+		PreviewAttachedAssetContainer.SaveAttachedObjectsFromDeprecatedProperties();
+	}
 }
 
 void USkeletalMesh::AddReferencedObjects(UObject* InThis, FReferenceCollector& Collector)
@@ -3100,6 +3105,17 @@ bool USkeletalMesh::IsSocketOnMesh( const FName& InSocketName ) const
 	return false;
 }
 
+int32 USkeletalMesh::ValidatePreviewAttachedObjects()
+{
+	int32 NumBrokenAssets = PreviewAttachedAssetContainer.ValidatePreviewAttachedObjects();
+
+	if (NumBrokenAssets > 0)
+	{
+		MarkPackageDirty();
+	}
+	return NumBrokenAssets;
+}
+
 #endif // #if WITH_EDITOR
 
 void USkeletalMesh::ReleaseCPUResources()
