Fix for build break caused by CL#2055369.

diff --git a/Engine/Source/Editor/Persona/Private/Persona.cpp b/Engine/Source/Editor/Persona/Private/Persona.cpp
--- a/Engine/Source/Editor/Persona/Private/Persona.cpp
+++ b/Engine/Source/Editor/Persona/Private/Persona.cpp
@@ -46,6 +46,7 @@
 #include "AnimGraphNode_LayeredBoneBlend.h"
 #include "AnimGraphNode_SequencePlayer.h"
 #include "AnimGraphNode_SequenceEvaluator.h"
+#include "AnimPreviewInstance.h"
 #define LOCTEXT_NAMESPACE "FPersona"
 
 /////////////////////////////////////////////////////
