Multi selection of animation notifies now possible, notify features updated to work with multiple nodes including working across multiple tracks.
#ttp 315768 - PERSONA: Add support for multi selection of anim notifies or drag/drop selection
#branch UE4

diff --git a/Engine/Source/Editor/Persona/Private/SAnimNotifyPanel.cpp b/Engine/Source/Editor/Persona/Private/SAnimNotifyPanel.cpp
--- a/Engine/Source/Editor/Persona/Private/SAnimNotifyPanel.cpp
+++ b/Engine/Source/Editor/Persona/Private/SAnimNotifyPanel.cpp
@@ -25,30 +25,14 @@ const FVector2D TextBorderSize(1.f, 1.f);
 
 DECLARE_DELEGATE_OneParam( FOnDeleteNotify, struct FAnimNotifyEvent*)
 DECLARE_DELEGATE_RetVal_FourParams( FReply, FOnNotifyNodeDragStarted, TSharedRef<SAnimNotifyNode>, const FVector2D&, const FVector2D&, const bool)
+DECLARE_DELEGATE_RetVal_FiveParams(FReply, FOnNotifyNodesDragStarted, TArray<TSharedPtr<SAnimNotifyNode>>, TSharedRef<SWidget>, const FVector2D&, const FVector2D&, const bool)
 DECLARE_DELEGATE_RetVal( float, FOnGetDraggedNodePos )
 DECLARE_DELEGATE_TwoParams( FPanTrackRequest, int32, FVector2D)
+DECLARE_DELEGATE_FourParams(FPasteNotifies, SAnimNotifyTrack*, float, ENotifyPasteMode::Type, ENotifyPasteMultipleMode::Type)
 
 class FNotifyDragDropOp;
 
-namespace ENotifyPasteMode
-{
-	enum Type
-	{
-		MousePosition,
-		SpecifiedTime,
-		OriginalTime
-	};
-}
 
-namespace ENotifyStateHandleHit
-{
-	enum Type
-	{
-		Start,
-		End,
-		None
-	};
-}
 //////////////////////////////////////////////////////////////////////////
 // SAnimNotifyNode
 
@@ -82,6 +66,8 @@ public:
 	virtual FReply OnDragDetected( const FGeometry& MyGeometry, const FPointerEvent& MouseEvent) OVERRIDE;
 	virtual FReply OnMouseMove( const FGeometry& MyGeometry, const FPointerEvent& MouseEvent) OVERRIDE;
 	virtual FReply OnMouseButtonUp( const FGeometry& MyGeometry, const FPointerEvent& MouseEvent ) OVERRIDE;
+	virtual void Tick(const FGeometry& AllottedGeometry, const double InCurrentTime, const float InDeltaTime) OVERRIDE;
+	
 	// End of SWidget interface
 
 	// SNodePanel::SNode interface
@@ -122,6 +108,26 @@ public:
 	/** The minimum possible duration that a notify state can have */
 	static const float MinimumStateDuration;
 
+	const FVector2D& GetScreenPosition() const
+	{
+		return ScreenPosition;
+	}
+
+	const float GetLastSnappedTime() const
+	{
+		return LastSnappedTime;
+	}
+
+	void ClearLastSnappedTime()
+	{
+		LastSnappedTime = -1.0f;
+	}
+
+	void SetLastSnappedTime(float NewSnapTime)
+	{
+		LastSnappedTime = NewSnapTime;
+	}
+
 private:
 	FText GetNodeTooltip() const;
 
@@ -138,6 +144,8 @@ private:
 	TAttribute<float>			ViewInputMin;
 	TAttribute<float>			ViewInputMax;
 	FVector2D					CachedAllotedGeometrySize;
+	FVector2D					ScreenPosition;
+	float						LastSnappedTime;
 
 	bool						bDrawTooltipToRight;
 	bool						bBeingDragged;	
@@ -188,6 +196,9 @@ public:
 		, _OnRequestTrackPan()
 		, _OnRequestOffsetRefresh()
 		, _OnDeleteNotify()
+		, _OnDeselectAllNotifies()
+		, _OnCopyNotifies()
+		, _OnPasteNotifies()
 		{}
 
 		SLATE_ARGUMENT( TSharedPtr<FPersona>,	Persona )
@@ -198,14 +209,17 @@ public:
 		SLATE_ATTRIBUTE( TArray<FTrackMarkerBar>, MarkerBars)
 		SLATE_ARGUMENT( int32, TrackIndex )
 		SLATE_ARGUMENT( FLinearColor, TrackColor )
-		SLATE_EVENT( FOnSelectionChanged, OnSelectionChanged )
+		SLATE_EVENT(FOnTrackSelectionChanged, OnSelectionChanged)
 		SLATE_EVENT( FOnUpdatePanel, OnUpdatePanel )
 		SLATE_EVENT( FOnGetScrubValue, OnGetScrubValue )
 		SLATE_EVENT( FOnGetDraggedNodePos, OnGetDraggedNodePos )
-		SLATE_EVENT( FOnNotifyNodeDragStarted, OnNodeDragStarted )
+		SLATE_EVENT( FOnNotifyNodesDragStarted, OnNodeDragStarted )
 		SLATE_EVENT( FPanTrackRequest, OnRequestTrackPan )
 		SLATE_EVENT( FRefreshOffsetsRequest, OnRequestOffsetRefresh )
 		SLATE_EVENT( FDeleteNotify, OnDeleteNotify )
+		SLATE_EVENT( FDeselectAllNotifies, OnDeselectAllNotifies)
+		SLATE_EVENT( FCopyNotifies, OnCopyNotifies )
+		SLATE_EVENT( FPasteNotifies, OnPasteNotifies )
 		SLATE_END_ARGS()
 public:
 
@@ -223,7 +237,36 @@ public:
 	FTrackScaleInfo GetCachedScaleInfo() const { return FTrackScaleInfo(ViewInputMin.Get(), ViewInputMax.Get(), 0.f, 0.f, CachedGeometry.Size); }
 
 	/** Updates sequences when a notify node has been successfully dragged to a new position */
-	void HandleNodeDrop(const FNotifyDragDropOp& DragDropOp);
+	void HandleNodeDrop(TSharedPtr<SAnimNotifyNode> Node);
+
+	// Number of nodes in the track currently selected
+	int32 GetNumSelectedNodes() const { return SelectedNotifyIndices.Num(); }
+
+	// Index of the track in the notify panel
+	int32 GetTrackIndex() const { return TrackIndex; }
+
+	// Time at the position of the last mouseclick
+	float GetLastClickedTime() const { return LastClickedTime; }
+
+	// Removes the node widgets from the track and adds them to the provided Array
+	void DisconnectSelectedNodesForDrag(TArray<TSharedPtr<SAnimNotifyNode>>& DragNodes);
+
+	// Adds our current selection to the provided set
+	void AppendSelectionToSet(FGraphPanelSelectionSet& SelectionSet);
+	// Adds our current selection to the provided array
+	void AppendSelectionToArray(TArray<const FAnimNotifyEvent*>& Selection) const;
+
+	/**
+	* Deselects all currently selected notify nodes
+	* @param bUpdateSelectionSet - Whether we should report a selection change to the panel
+	*/
+	void DeselectAllNotifyNodes(bool bUpdateSelectionSet = true);
+
+	// get Property Data of one element (NotifyIndex) from Notifies property of Sequence
+	static uint8* FindNotifyPropertyData(UAnimSequenceBase* Sequence, int32 NotifyIndex, UArrayProperty*& ArrayProperty);
+
+	// Paste a single Notify into this track from an exported string
+	void PasteSingleNotify(FString& NotifyString, float PasteTime);
 
 protected:
 	void CreateCommands();
@@ -237,6 +280,11 @@ protected:
 	void AddNewNotify(const FText& NewNotifyName, ETextCommit::Type CommitInfo);
 	void OnSetTriggerWeightNotifyClicked(int32 NotifyIndex);
 	void SetTriggerWeight(const FText& TriggerWeight, ETextCommit::Type CommitInfo, int32 NotifyIndex);
+	
+	// Whether we have one node selected
+	bool IsSingleNodeSelected();
+	// Checks the clipboard for an anim notify buffer, and returns whether there's only one notify
+	bool IsSingleNodeInClipboard();
 
 	void OnSetDurationNotifyClicked(int32 NotifyIndex);
 	void SetDuration(const FText& DurationText, ETextCommit::Type CommitInfo, int32 NotifyIndex);
@@ -249,16 +297,16 @@ protected:
 
 	/** Reads the contents of the clipboard for a notify to paste.
 	 *  If successful OutBuffer points to the start of the notify object's data */
-	bool ReadNotifyPasteHeader(FString& OutPropertyString, const TCHAR*& OutBuffer, float& OutOriginalTime) const;
+	bool ReadNotifyPasteHeader(FString& OutPropertyString, const TCHAR*& OutBuffer, float& OutOriginalTime, float& OutOriginalLength) const;
 
 	/** Handler for context menu paste command */
-	void OnPasteNotifyClicked(ENotifyPasteMode::Type PasteMode);
+	void OnPasteNotifyClicked(ENotifyPasteMode::Type PasteMode, ENotifyPasteMultipleMode::Type MultiplePasteType = ENotifyPasteMultipleMode::Absolute);
 
 	/** Handler for popup window asking the user for a paste time */
 	void OnPasteNotifyTimeSet(const FText& TimeText, ETextCommit::Type CommitInfo);
 
 	/** Function to paste a previously copied notify */
-	void OnPasteNotify(float TimeToPasteAt);
+	void OnPasteNotify(float TimeToPasteAt, ENotifyPasteMultipleMode::Type MultiplePasteType = ENotifyPasteMultipleMode::Absolute);
 
 	/** Provides direct access to the notify menu from the context menu */
 	void OnManageNotifies();
@@ -268,7 +316,25 @@ protected:
 	 */
 	void Update();
 
-	void SelectNotifyNode(int32 NotifyIndex);
+	/**
+	 * Selects a notify node in the graph. Supports multi selection
+	 * @param NotifyIndex - Index of the notify node to select.
+	 * @param Append - Whether to append to to current selection or start a new one.
+	 */
+	void SelectNotifyNode(int32 NotifyIndex, bool Append);
+	
+	/**
+	 * Toggles the selection status of a notify node, for example when
+	 * Control is held when clicking.
+	 * @param NotifyIndex - Index of the notify to toggle the selection status of
+	 */
+	void ToggleNotifyNodeSelectionStatus(int32 NotifyIndex);
+
+	/**
+	 * Deselects requested notify node
+	 * @param NotifyIndex - Index of the notify node to deselect
+	 */
+	void DeselectNotifyNode(int32 NotifyIndex);
 
 	int32 GetHitNotifyNode(const FGeometry& MyGeometry, const FVector2D& Position);
 
@@ -279,14 +345,12 @@ protected:
 	virtual FReply OnDrop( const FGeometry& MyGeometry, const FDragDropEvent& DragDropEvent ) OVERRIDE;
 	virtual FReply OnMouseButtonDown( const FGeometry& MyGeometry, const FPointerEvent& MouseEvent ) OVERRIDE;
 
-	float CalculateTime( const FGeometry& MyGeometry, FVector2D NodePos );
+	float CalculateTime( const FGeometry& MyGeometry, FVector2D NodePos, bool bInputIsAbsolute = true );
 
 	// Handler that is called when the user starts dragging a node
 	FReply OnNotifyNodeDragStarted( TSharedRef<SAnimNotifyNode> NotifyNode, const FVector2D& ScreenCursorPos, const FVector2D& ScreenNodePosition, const bool bDragOnMarker, int32 NotifyIndex );
 
 private:
-	// get Property Data of one element (NotifyIndex) from Notifies property of Sequence
-	uint8* FindNotifyPropertyData( int32 NotifyIndex, UArrayProperty*& ArrayProperty );
 
 	// Store the tracks geometry for later use
 	void UpdateCachedGeometry(const FGeometry& InGeometry) {CachedGeometry = InGeometry;}
@@ -299,6 +363,12 @@ private:
 		return FMargin(LeftMargin, 0, RightMargin, 0);
 	}
 
+	// Builds a UObject selection set and calls the OnSelectionChanged delegate
+	void SendSelectionChanged()
+	{
+		OnSelectionChanged.ExecuteIfBound();
+	}
+
 protected:
 	TSharedPtr<FUICommandList> AnimSequenceEditorActions;
 
@@ -311,12 +381,15 @@ protected:
 	TAttribute<float>						ViewInputMax;
 	TAttribute<FLinearColor>				TrackColor;
 	int32									TrackIndex;
-	FOnSelectionChanged						OnSelectionChanged;
+	FOnTrackSelectionChanged				OnSelectionChanged;
 	FOnUpdatePanel							OnUpdatePanel;
 	FOnGetScrubValue						OnGetScrubValue;
 	FOnGetDraggedNodePos					OnGetDraggedNodePos;
-	FOnNotifyNodeDragStarted				OnNodeDragStarted;
+	FOnNotifyNodesDragStarted				OnNodeDragStarted;
 	FPanTrackRequest						OnRequestTrackPan;
+	FDeselectAllNotifies					OnDeselectAllNotifies;
+	FCopyNotifies							OnCopyNotifies;
+	FPasteNotifies							OnPasteNotifies;
 
 	/** Delegate to call when offsets should be refreshed in a montage */
 	FRefreshOffsetsRequest					OnRequestRefreshOffsets;
@@ -335,6 +408,9 @@ protected:
 
 	/** Attribute for accessing any marker positions we have to draw */
 	TAttribute<TArray<FTrackMarkerBar>>		MarkerBars;
+
+	/** Nodes that are currently selected */
+	TArray<int32> SelectedNotifyIndices;
 };
 
 //////////////////////////////////////////////////////////////////////////
@@ -363,6 +439,8 @@ public:
 		, _WidgetWidth()
 		, _OnUpdatePanel()
 		, _OnDeleteNotify()
+		, _OnDeselectAllNotifies()
+		, _OnCopyNotifies()
 	{}
 	SLATE_ARGUMENT( int32, TrackIndex )
 	SLATE_ARGUMENT( TSharedPtr<SAnimNotifyPanel>, AnimNotifyPanel)
@@ -371,13 +449,16 @@ public:
 	SLATE_ATTRIBUTE( float, ViewInputMin )
 	SLATE_ATTRIBUTE( float, ViewInputMax )
 	SLATE_ATTRIBUTE( TArray<FTrackMarkerBar>, MarkerBars)
-	SLATE_EVENT( FOnSelectionChanged, OnSelectionChanged )
+	SLATE_EVENT( FOnTrackSelectionChanged, OnSelectionChanged)
 	SLATE_EVENT( FOnGetScrubValue, OnGetScrubValue )
 	SLATE_EVENT( FOnGetDraggedNodePos, OnGetDraggedNodePos )
 	SLATE_EVENT( FOnUpdatePanel, OnUpdatePanel )
-	SLATE_EVENT( FOnNotifyNodeDragStarted, OnNodeDragStarted )
+	SLATE_EVENT( FOnNotifyNodesDragStarted, OnNodeDragStarted )
 	SLATE_EVENT( FRefreshOffsetsRequest, OnRequestRefreshOffsets )
 	SLATE_EVENT( FDeleteNotify, OnDeleteNotify )
+	SLATE_EVENT( FDeselectAllNotifies, OnDeselectAllNotifies)
+	SLATE_EVENT( FCopyNotifies, OnCopyNotifies )
+	SLATE_EVENT( FPasteNotifies, OnPasteNotifies )
 	SLATE_END_ARGS()
 
 	void Construct(const FArguments& InArgs);
@@ -391,7 +472,13 @@ public:
 class FNotifyDragDropOp : public FDragDropOperation
 {
 public:
-	FNotifyDragDropOp(float& InCurrentDragXPosition) : CurrentDragXPosition(InCurrentDragXPosition), SnapTime(-1.f) {}
+	FNotifyDragDropOp(float& InCurrentDragXPosition) : 
+		CurrentDragXPosition(InCurrentDragXPosition), 
+		SnapTime(-1.f),
+		SelectionTimeLength(0.0f)
+	{
+	}
+
 	struct FTrackClampInfo
 	{
 		int32 TrackPos;
@@ -408,13 +495,11 @@ public:
 	{
 		if ( bDropWasHandled == false )
 		{
-			const FTrackClampInfo& ClampInfo = GetTrackClampInfo(LastNodePos+NotifyOffset);
-
-			ClampInfo.NotifyTrack->HandleNodeDrop(*this);
-
-			if (OriginalNotifyNode.IsValid())
+			for(TSharedPtr<SAnimNotifyNode> Node : SelectedNodes)
 			{
-				OriginalNotifyNode.Pin()->DropCancelled();
+				const FTrackClampInfo& ClampInfo = GetTrackClampInfo(Node->GetScreenPosition());
+				ClampInfo.NotifyTrack->HandleNodeDrop(Node);
+				Node->DropCancelled();
 			}
 		}
 		
@@ -423,85 +508,127 @@ public:
 
 	virtual void OnDragged( const class FDragDropEvent& DragDropEvent ) OVERRIDE
 	{
-		FVector2D NodePos = DragDropEvent.GetScreenSpacePosition() + DragOffset + NotifyOffset;
-		const FTrackClampInfo& ClampInfo = GetTrackClampInfo(DragDropEvent.GetScreenSpacePosition());
+		// Reset snapped node pointer
+		SnappedNode = NULL;
+
+		NodeGroupPosition = DragDropEvent.GetScreenSpacePosition() + DragOffset;
+
+		FVector2D SelectionBeginPosition = NodeGroupPosition + SelectedNodes[0]->GetNotifyPositionOffset();
 		
-		// Look for a snap on the first scrub handle
-		FVector2D TrackNodePos = ClampInfo.NotifyTrack->GetCachedGeometry().AbsoluteToLocal(NodePos);
-		float SnapX = GetSnapPosition(ClampInfo, TrackNodePos.X);
-		if(SnapX >= 0.f)
+		FTrackClampInfo* SelectionPositionClampInfo = &GetTrackClampInfo(DragDropEvent.GetScreenSpacePosition());
+		if(SelectionPositionClampInfo->NotifyTrack->GetTrackIndex() < TrackSpan)
 		{
-			EAnimEventTriggerOffsets::Type Offset = (SnapX < TrackNodePos.X) ? EAnimEventTriggerOffsets::OffsetAfter : EAnimEventTriggerOffsets::OffsetBefore;
-			NotifyEvent->TriggerTimeOffset = GetTriggerTimeOffsetForType(Offset);
-			TrackNodePos.X = SnapX;
-			NodePos = ClampInfo.NotifyTrack->GetCachedGeometry().LocalToAbsolute(TrackNodePos);
+			// Our selection has moved off the bottom of the notify panel, adjust the clamping information to keep it on the panel
+			SelectionPositionClampInfo = &ClampInfos[ClampInfos.Num() - TrackSpan - 1];
 		}
-		else
+
+		const FGeometry& TrackGeom = SelectionPositionClampInfo->NotifyTrack->GetCachedGeometry();
+		const FTrackScaleInfo& TrackScaleInfo = SelectionPositionClampInfo->NotifyTrack->GetCachedScaleInfo();
+		
+		// Tracks the movement amount to apply to the selection due to a snap.
+		float SnapMovement= 0.0f;
+		// Clamp the selection into the track
+		const float SelectionLocalLength = TrackScaleInfo.PixelsPerInput * SelectionTimeLength;
+		const float ClampedEnd = FMath::Clamp(SelectionBeginPosition.X + NodeGroupSize.X, SelectionPositionClampInfo->TrackMin, SelectionPositionClampInfo->TrackMax);
+		const float ClampedBegin = FMath::Clamp(SelectionBeginPosition.X, SelectionPositionClampInfo->TrackMin, SelectionPositionClampInfo->TrackMax);
+		if(ClampedBegin > SelectionBeginPosition.X)
 		{
-			NotifyEvent->TriggerTimeOffset = GetTriggerTimeOffsetForType(EAnimEventTriggerOffsets::NoOffset);
+			SelectionBeginPosition.X = ClampedBegin;
+		}
+		else if(ClampedEnd < SelectionBeginPosition.X + SelectionLocalLength)
+		{
+			SelectionBeginPosition.X = ClampedEnd - SelectionLocalLength;
 		}
 
-		NodePos.Y = ClampInfo.TrackPos;
-		NodePos.X = FMath::Clamp(NodePos.X, ClampInfo.TrackMin, ClampInfo.TrackMax);
-
-		if(OriginalNotifyNode.IsValid() && NotifyEvent->Duration > 0)
+		// Handle node snaps
+		for(int32 NodeIdx = 0 ; NodeIdx < SelectedNodes.Num() ; ++NodeIdx)
 		{
-			// If we didn't snap the beginning of the node, attempt to snap the end
-			if(SnapX == -1.0f)
+			TSharedPtr<SAnimNotifyNode> CurrentNode = SelectedNodes[NodeIdx];
+			FAnimNotifyEvent* CurrentEvent = CurrentNode->NotifyEvent;
+			
+			// Clear off any snap time currently stored
+			CurrentNode->ClearLastSnappedTime();
+
+			const FTrackClampInfo& NodeClamp = GetTrackClampInfo(CurrentNode->GetScreenPosition());
+
+			FVector2D EventPosition = SelectionBeginPosition + FVector2D(TrackScaleInfo.PixelsPerInput * NodeTimeOffsets[NodeIdx], 0.0f);
+
+			// Look for a snap on the first scrub handle
+			FVector2D TrackNodePos = TrackGeom.AbsoluteToLocal(EventPosition);
+			const FVector2D OriginalNodePosition = TrackNodePos;
+
+			float SnapX = GetSnapPosition(NodeClamp, TrackNodePos.X);
+			if(SnapX >= 0.f)
 			{
-				FVector2D TrackNodeEndPos = TrackNodePos + OriginalNotifyNode.Pin()->GetDurationSize();
-				SnapX = GetSnapPosition(ClampInfo, TrackNodeEndPos.X);
+				EAnimEventTriggerOffsets::Type Offset = (SnapX < TrackNodePos.X) ? EAnimEventTriggerOffsets::OffsetAfter : EAnimEventTriggerOffsets::OffsetBefore;
+				CurrentEvent->TriggerTimeOffset = GetTriggerTimeOffsetForType(Offset);
+				CurrentNode->SetLastSnappedTime(TrackScaleInfo.LocalXToInput(SnapX));
 
-				// Only attempt to snap if the node will fit on the track
-				if(SnapX >= OriginalNotifyNode.Pin()->GetDurationSize())
-				{
-					EAnimEventTriggerOffsets::Type Offset = (SnapX < TrackNodeEndPos.X) ? EAnimEventTriggerOffsets::OffsetAfter : EAnimEventTriggerOffsets::OffsetBefore;
-					NotifyEvent->EndTriggerTimeOffset = GetTriggerTimeOffsetForType(Offset);
-					TrackNodeEndPos.X = SnapX;
-					NodePos = ClampInfo.NotifyTrack->GetCachedGeometry().LocalToAbsolute(TrackNodeEndPos - OriginalNotifyNode.Pin()->GetDurationSize());
-					
-					// Fix up the snap time to correctly position the node.
-					SnapTime -= NotifyEvent->Duration;
-
-					NodePos.Y = ClampInfo.TrackPos;
-					NodePos.X = FMath::Clamp(NodePos.X, ClampInfo.TrackMin, ClampInfo.TrackMax);
-				}
-				else
+				if(SnapMovement == 0.0f)
 				{
-					// Remove any trigger time if we can't fit the node in.
-					NotifyEvent->EndTriggerTimeOffset = GetTriggerTimeOffsetForType(EAnimEventTriggerOffsets::NoOffset);
+					SnapMovement = SnapX - TrackNodePos.X;
+					TrackNodePos.X = SnapX;
+					SnapTime = TrackScaleInfo.LocalXToInput(SnapX);
+					SnappedNode = CurrentNode;
 				}
+				EventPosition = NodeClamp.NotifyTrack->GetCachedGeometry().LocalToAbsolute(TrackNodePos);
+			}
+			else
+			{
+				CurrentEvent->TriggerTimeOffset = GetTriggerTimeOffsetForType(EAnimEventTriggerOffsets::NoOffset);
 			}
 
-			//If we have a duration we need to handle clamping at the end of it
-			float NodeWidth = OriginalNotifyNode.Pin()->GetDurationSize();
-			float NotifyEnd = NodePos.X + NodeWidth;
-			float ClampedEnd = FMath::Clamp(NotifyEnd, ClampInfo.TrackMin, ClampInfo.TrackMax);
+			// Always clamp the Y to the current track
+			SelectionBeginPosition.Y = SelectionPositionClampInfo->TrackPos;
 
-			if(ClampedEnd != NotifyEnd)
+			if(CurrentNode.IsValid() && CurrentEvent->Duration > 0)
 			{
-				NodePos.X = ClampedEnd - NodeWidth;
+				// If we didn't snap the beginning of the node, attempt to snap the end
+				if(SnapX == -1.0f)
+				{
+					FVector2D TrackNodeEndPos = TrackNodePos + CurrentNode->GetDurationSize();
+					SnapX = GetSnapPosition(*SelectionPositionClampInfo, TrackNodeEndPos.X);
+
+					// Only attempt to snap if the node will fit on the track
+					if(SnapX >= CurrentNode->GetDurationSize())
+					{
+						EAnimEventTriggerOffsets::Type Offset = (SnapX < TrackNodeEndPos.X) ? EAnimEventTriggerOffsets::OffsetAfter : EAnimEventTriggerOffsets::OffsetBefore;
+						CurrentEvent->EndTriggerTimeOffset = GetTriggerTimeOffsetForType(Offset);
+						
+						if(SnapMovement == 0.0f)
+						{
+							SnapMovement = SnapX - TrackNodeEndPos.X;
+							SnapTime = TrackScaleInfo.LocalXToInput(SnapX) - CurrentEvent->Duration;
+							SnappedNode = CurrentNode;
+						}
+					}
+					else
+					{
+						// Remove any trigger time if we can't fit the node in.
+						CurrentEvent->EndTriggerTimeOffset = GetTriggerTimeOffsetForType(EAnimEventTriggerOffsets::NoOffset);
+					}
+				}
 			}
 		}
 
-		CurrentDragXPosition = ClampInfo.NotifyTrack->GetCachedGeometry().AbsoluteToLocal(FVector2D(NodePos.X,0.0f)).X;
+		SelectionBeginPosition.X += SnapMovement;
 
-		NodePos.X -= NotifyOffset.X;
-		LastNodePos = NodePos;
+		CurrentDragXPosition = SelectionPositionClampInfo->NotifyTrack->GetCachedGeometry().AbsoluteToLocal(FVector2D(SelectionBeginPosition.X,0.0f)).X;
 
-		CursorDecoratorWindow->MoveWindowTo(NodePos);
+		CursorDecoratorWindow->MoveWindowTo(SelectionBeginPosition - SelectedNodes[0]->GetNotifyPositionOffset());
+		NodeGroupPosition = SelectionBeginPosition;
 
 		//scroll view
 		float MouseXPos = DragDropEvent.GetScreenSpacePosition().X;
-		if(MouseXPos < ClampInfo.TrackMin)
+		if(MouseXPos < SelectionPositionClampInfo->TrackMin)
 		{
-			float ScreenDelta = MouseXPos - ClampInfo.TrackMin;
-			RequestTrackPan.Execute(ScreenDelta, FVector2D(ClampInfo.TrackWidth, 1.f));
+			float ScreenDelta = MouseXPos - SelectionPositionClampInfo->TrackMin;
+			RequestTrackPan.Execute(ScreenDelta, FVector2D(SelectionPositionClampInfo->TrackWidth, 1.f));
 		}
-		else if(MouseXPos > ClampInfo.TrackMax)
+		else if(MouseXPos > SelectionPositionClampInfo->TrackMax)
 		{
-			float ScreenDelta = MouseXPos - ClampInfo.TrackMax;
-			RequestTrackPan.Execute(ScreenDelta, FVector2D(ClampInfo.TrackWidth, 1.f));
+			float ScreenDelta = MouseXPos - SelectionPositionClampInfo->TrackMax;
+			RequestTrackPan.Execute(ScreenDelta, FVector2D(SelectionPositionClampInfo->TrackWidth, 1.f));
 		}
 	}
 
@@ -513,7 +640,6 @@ public:
 
 		float CurrentMinSnapDest = MaxSnapDist;
 		float SnapPosition = -1.f;
-		SnapTime = -1.f;
 
 		if(MarkerBars.IsBound())
 		{
@@ -526,7 +652,6 @@ public:
 					float ThisMinSnapDist = FMath::Abs(WidgetSpaceSnapPosition - WidgetSpaceNotifyPosition);
 					if(ThisMinSnapDist < CurrentMinSnapDest)
 					{
-						SnapTime = Bars[i].Time;
 						CurrentMinSnapDest = ThisMinSnapDist;
 						SnapPosition = WidgetSpaceSnapPosition;
 					}
@@ -536,7 +661,7 @@ public:
 		return SnapPosition;
 	}
 
-	const FTrackClampInfo& GetTrackClampInfo(const FVector2D NodePos) const
+	FTrackClampInfo& GetTrackClampInfo(const FVector2D NodePos)
 	{
 		int32 ClampInfoIndex = 0;
 		int32 SmallestNodeTrackDist = FMath::Abs(ClampInfos[0].TrackSnapTestPos - NodePos.Y);
@@ -552,29 +677,59 @@ public:
 		return ClampInfos[ClampInfoIndex];
 	}
 
-	FAnimNotifyEvent*			NotifyEvent;
-	class UAnimSequenceBase*	Sequence;
-	FVector2D					DragOffset;
-	FVector2D					NotifyOffset;
-	FVector2D					LastNodePos;
-	TWeakPtr<SAnimNotifyNode>	OriginalNotifyNode;
-	TArray<FTrackClampInfo>		ClampInfos;
-	float&						CurrentDragXPosition;
-	FPanTrackRequest			RequestTrackPan;
-	float						SnapTime;
-	TAttribute<TArray<FTrackMarkerBar>>	MarkerBars;
-
-	static TSharedRef<FNotifyDragDropOp> New(TSharedRef<SAnimNotifyNode> NotifyNode, const TArray<TSharedPtr<SAnimNotifyTrack>>& NotifyTracks, class UAnimSequenceBase* InSequence, const FVector2D &CursorPosition, const FVector2D &ScreenPositionOfNode, float& CurrentDragXPosition, FPanTrackRequest& RequestTrackPanDelegate, TAttribute<TArray<FTrackMarkerBar>>	MarkerBars)
+	class UAnimSequenceBase*			Sequence;				// The owning anim sequence
+	FVector2D							DragOffset;				// Offset from the mouse to place the decorator
+	TArray<FTrackClampInfo>				ClampInfos;				// Clamping information for all of the available tracks
+	float&								CurrentDragXPosition;	// Current X position of the drag operation
+	FPanTrackRequest					RequestTrackPan;		// Delegate to request a pan along the edges of a zoomed track
+	TArray<float>						NodeTimes;				// Times to drop each selected node at
+	float								SnapTime;				// The time that the snapped node was snapped to
+	TWeakPtr<SAnimNotifyNode>			SnappedNode;			// The node chosen for the snap
+	TAttribute<TArray<FTrackMarkerBar>>	MarkerBars;				// Branching point markers
+	TArray<TSharedPtr<SAnimNotifyNode>> SelectedNodes;			// The nodes that are in the current selection
+	TArray<float>						NodeTimeOffsets;		// Time offsets from the beginning of the selection to the nodes.
+	FVector2D							NodeGroupPosition;		// Position of the beginning of the selection
+	FVector2D							NodeGroupSize;			// Size of the entire selection
+	TSharedPtr<SWidget>					Decorator;				// The widget to display when dragging
+	float								SelectionTimeLength;	// Length of time that the selection covers
+	int32								TrackSpan;				// Number of tracks that the selection spans
+
+	static TSharedRef<FNotifyDragDropOp> New(TArray<TSharedPtr<SAnimNotifyNode>> NotifyNodes, TSharedPtr<SWidget> Decorator, const TArray<TSharedPtr<SAnimNotifyTrack>>& NotifyTracks, class UAnimSequenceBase* InSequence, const FVector2D &CursorPosition, const FVector2D &SelectionScreenPosition, const FVector2D &SelectionSize, float& CurrentDragXPosition, FPanTrackRequest& RequestTrackPanDelegate, TAttribute<TArray<FTrackMarkerBar>>	MarkerBars)
 	{
 		TSharedRef<FNotifyDragDropOp> Operation = MakeShareable(new FNotifyDragDropOp(CurrentDragXPosition));
-		Operation->NotifyEvent = NotifyNode->NotifyEvent;
 		Operation->Sequence = InSequence;
-		Operation->OriginalNotifyNode = NotifyNode;
 		Operation->RequestTrackPan = RequestTrackPanDelegate;
 
-		Operation->DragOffset = ScreenPositionOfNode - CursorPosition;
-		Operation->NotifyOffset = NotifyNode->GetNotifyPositionOffset();
+		Operation->NodeGroupPosition = SelectionScreenPosition;
+		Operation->NodeGroupSize = SelectionSize;
+		Operation->DragOffset = SelectionScreenPosition - CursorPosition;
 		Operation->MarkerBars = MarkerBars;
+		Operation->Decorator = Decorator;
+		Operation->SelectedNodes = NotifyNodes;
+		Operation->TrackSpan = NotifyNodes[0]->NotifyEvent->TrackIndex - NotifyNodes.Last()->NotifyEvent->TrackIndex;
+
+		// Caclulate offsets for the selected nodes
+		float BeginTime = MAX_flt;
+		for(TSharedPtr<SAnimNotifyNode> Node : NotifyNodes)
+		{
+			if(Node->NotifyEvent->DisplayTime < BeginTime)
+			{
+				BeginTime = Node->NotifyEvent->DisplayTime;
+			}
+		}
+
+		// Initialise node data
+		for(TSharedPtr<SAnimNotifyNode> Node : NotifyNodes)
+		{
+			Node->ClearLastSnappedTime();
+			Operation->NodeTimeOffsets.Add(Node->NotifyEvent->DisplayTime - BeginTime);
+			Operation->NodeTimes.Add(Node->NotifyEvent->DisplayTime);
+
+			// Calculate the time length of the selection. Because it is possible to have states
+			// with arbitrary durations we need to search all of the nodes and find the furthest
+			// possible point
+			Operation->SelectionTimeLength = FMath::Max(Operation->SelectionTimeLength, Node->NotifyEvent->DisplayTime + Node->NotifyEvent->Duration - BeginTime);
+		}
 
 		Operation->Construct();
 
@@ -598,16 +753,16 @@ public:
 	/** The widget decorator to use */
 	virtual TSharedPtr<SWidget> GetDefaultDecorator() const OVERRIDE
 	{
-		return OriginalNotifyNode.Pin();
+		return Decorator;
 	}
 
 	FText GetHoverText() const
 	{
 		FText HoverText = LOCTEXT("Invalid", "Invalid");
 
-		if (NotifyEvent)
+		if(SelectedNodes[0].IsValid())
 		{
-			HoverText = FText::FromName( NotifyEvent->NotifyName );
+			HoverText = FText::FromName( SelectedNodes[0]->NotifyEvent->NotifyName );
 		}
 
 		return HoverText;
@@ -1202,6 +1357,11 @@ void SAnimNotifyNode::DrawHandleOffset( const float& Offset, const float& Handle
 		);
 }
 
+void SAnimNotifyNode::Tick(const FGeometry& AllottedGeometry, const double InCurrentTime, const float InDeltaTime)
+{
+	ScreenPosition = AllottedGeometry.AbsolutePosition;
+}
+
 //////////////////////////////////////////////////////////////////////////
 // SAnimNotifyTrack
 
@@ -1226,6 +1386,9 @@ void SAnimNotifyTrack::Construct(const FArguments& InArgs)
 	OnRequestTrackPan = InArgs._OnRequestTrackPan;
 	OnRequestRefreshOffsets = InArgs._OnRequestOffsetRefresh;
 	OnDeleteNotify = InArgs._OnDeleteNotify;
+	OnDeselectAllNotifies = InArgs._OnDeselectAllNotifies;
+	OnCopyNotifies = InArgs._OnCopyNotifies;
+	OnPasteNotifies = InArgs._OnPasteNotifies;
 
 	this->ChildSlot
 	[
@@ -1615,6 +1778,8 @@ FReply SAnimNotifyTrack::OnMouseButtonUp( const FGeometry& MyGeometry, const FPo
 {
 	bool bLeftMouseButton =  MouseEvent.GetEffectingButton() == EKeys::LeftMouseButton;
 	bool bRightMouseButton =  MouseEvent.GetEffectingButton() == EKeys::RightMouseButton;
+	bool bShift = MouseEvent.IsShiftDown();
+	bool bCtrl = MouseEvent.IsControlDown();
 
 	if ( bRightMouseButton )
 	{
@@ -1630,7 +1795,23 @@ FReply SAnimNotifyTrack::OnMouseButtonUp( const FGeometry& MyGeometry, const FPo
 		FVector2D CursorPos = MouseEvent.GetScreenSpacePosition();
 		CursorPos = MyGeometry.AbsoluteToLocal(CursorPos);
 		int32 NotifyIndex = GetHitNotifyNode(MyGeometry, CursorPos);
-		SelectNotifyNode(NotifyIndex);
+
+		if(NotifyIndex == INDEX_NONE)
+		{
+			// Clicked in empty space, clear selection
+			OnDeselectAllNotifies.ExecuteIfBound();
+		}
+		else
+		{
+			if(bCtrl)
+			{
+				ToggleNotifyNodeSelectionStatus(NotifyIndex);
+			}
+			else
+			{
+				SelectNotifyNode(NotifyIndex, bShift);
+			}
+		}
 
 		return FReply::Handled();
 	}
@@ -1638,40 +1819,76 @@ FReply SAnimNotifyTrack::OnMouseButtonUp( const FGeometry& MyGeometry, const FPo
 	return FReply::Unhandled();
 }
 
-void SAnimNotifyTrack::SelectNotifyNode(int32 NotifyIndex)
+void SAnimNotifyTrack::SelectNotifyNode(int32 NotifyIndex, bool Append)
 {
 	if( NotifyIndex != INDEX_NONE )
 	{
-		// Deselect all other notifies first.
-		if( Sequence )
+		// Deselect all other notifies if necessary.
+		if (Sequence && !Append)
 		{
-			for(int32 Index=0; Index<Sequence->Notifies.Num(); Index++)
-			{
-				FAnimNotifyEvent & NotifyEvent = Sequence->Notifies[Index];
-				NotifyEvent.bSelected = false;
-			}
+			OnDeselectAllNotifies.ExecuteIfBound();
 		}
 
-		// select new one
-		FGraphPanelSelectionSet ObjectSet;
-		if( NotifyIndex < NotifyNodes.Num() )
+		// Check to see if we've already selected this node
+		if (!SelectedNotifyIndices.Contains(NotifyIndex))
 		{
-			FAnimNotifyEvent * NotifyEvent = NotifyNodes[NotifyIndex].Get()->NotifyEvent;
-			NotifyEvent->bSelected = true;
-			if( NotifyEvent->Notify )
+			// select new one
+			if (AnimNotifies.IsValidIndex(NotifyIndex))
 			{
-				ObjectSet.Add(NotifyEvent->Notify);	
-				OnSelectionChanged.ExecuteIfBound(ObjectSet);
-			}
-			else if( NotifyEvent->NotifyStateClass )
-			{
-				ObjectSet.Add(NotifyEvent->NotifyStateClass);	
-				OnSelectionChanged.ExecuteIfBound(ObjectSet);
+				FAnimNotifyEvent * NotifyEvent = AnimNotifies[NotifyIndex];
+				NotifyEvent->bSelected = true;
+				SelectedNotifyIndices.Add(NotifyIndex);
+				SendSelectionChanged();
 			}
 		}
 	}
 }
 
+void SAnimNotifyTrack::ToggleNotifyNodeSelectionStatus( int32 NotifyIndex )
+{
+	check(AnimNotifies.IsValidIndex(NotifyIndex));
+
+	bool bSelected = SelectedNotifyIndices.Contains(NotifyIndex);
+	if(bSelected)
+	{
+		SelectedNotifyIndices.Remove(NotifyIndex);
+	}
+	else
+	{
+		SelectedNotifyIndices.Add(NotifyIndex);
+	}
+
+	FAnimNotifyEvent* Event = AnimNotifies[NotifyIndex];
+	Event->bSelected = !Event->bSelected;
+	SendSelectionChanged();
+}
+
+void SAnimNotifyTrack::DeselectNotifyNode( int32 NotifyIndex )
+{
+	check(AnimNotifies.IsValidIndex(NotifyIndex));
+	FAnimNotifyEvent* Event = AnimNotifies[NotifyIndex];
+	Event->bSelected = false;
+
+	int32 ItemsRemoved = SelectedNotifyIndices.Remove(NotifyIndex);
+	check(ItemsRemoved > 0);
+
+	SendSelectionChanged();
+}
+
+void SAnimNotifyTrack::DeselectAllNotifyNodes(bool bUpdateSelectionSet)
+{
+	for (FAnimNotifyEvent& Event : Sequence->Notifies)
+	{
+		Event.bSelected = false;
+	}
+	SelectedNotifyIndices.Empty();
+
+	if(bUpdateSelectionSet)
+	{
+		SendSelectionChanged();
+	}
+}
+
 TSharedPtr<SWidget> SAnimNotifyTrack::SummonContextMenu(const FGeometry& MyGeometry, const FPointerEvent& MouseEvent)
 {
 	FVector2D CursorPos = MouseEvent.GetScreenSpacePosition();
@@ -1686,11 +1903,16 @@ TSharedPtr<SWidget> SAnimNotifyTrack::SummonContextMenu(const FGeometry& MyGeome
 	{
 		if ( NotifyIndex != INDEX_NONE )
 		{
-			SelectNotifyNode(NotifyIndex);
+			if (!AnimNotifies[NotifyIndex]->bSelected)
+			{
+				SelectNotifyNode(NotifyIndex, MouseEvent.IsControlDown());
+			}
 
 			// add menu to get threshold weight for triggering this notify
 			NewAction.ExecuteAction.BindRaw(
 				this, &SAnimNotifyTrack::OnSetTriggerWeightNotifyClicked, NotifyIndex);
+			NewAction.CanExecuteAction.BindRaw(
+				this, &SAnimNotifyTrack::IsSingleNodeSelected);
 
 			FNumberFormattingOptions Options;
 			Options.MinimumFractionalDigits = 5;
@@ -1704,6 +1926,8 @@ TSharedPtr<SWidget> SAnimNotifyTrack::SummonContextMenu(const FGeometry& MyGeome
 			{
 				NewAction.ExecuteAction.BindRaw(
 					this, &SAnimNotifyTrack::OnSetDurationNotifyClicked, NotifyIndex);
+				NewAction.CanExecuteAction.BindRaw(
+					this, &SAnimNotifyTrack::IsSingleNodeSelected);
 
 				FText SetAnimStateDurationText = FText::Format( LOCTEXT("SetAnimStateDuration", "Set AnimNotifyState duration ({0})"), FText::AsNumber( AnimNotifies[NotifyIndex]->Duration ) );
 				MenuBuilder.AddMenuEntry(SetAnimStateDurationText, LOCTEXT("SetAnimStateDuration_ToolTip", "The duration of this AnimNotifyState"), FSlateIcon(), NewAction);
@@ -1731,6 +1955,8 @@ TSharedPtr<SWidget> SAnimNotifyTrack::SummonContextMenu(const FGeometry& MyGeome
 	}
 	MenuBuilder.EndSection(); //AnimNotify
 
+	NewAction.CanExecuteAction = 0;
+
 	MenuBuilder.BeginSection("AnimEdit", LOCTEXT("NotifyEditHeading", "Edit") );
 	{
 		if ( NotifyIndex != INDEX_NONE )
@@ -1749,25 +1975,36 @@ TSharedPtr<SWidget> SAnimNotifyTrack::SummonContextMenu(const FGeometry& MyGeome
 			FString PropertyString;
 			const TCHAR* Buffer;
 			float OriginalTime;
+			float OriginalLength;
 
 			//Check whether can we show menu item to paste anim notify event
-			if( ReadNotifyPasteHeader(PropertyString, Buffer, OriginalTime) )
+			if( ReadNotifyPasteHeader(PropertyString, Buffer, OriginalTime, OriginalLength) )
 			{
 				// paste notify menu item
-				NewAction.ExecuteAction.BindRaw(
-					this, &SAnimNotifyTrack::OnPasteNotifyClicked, ENotifyPasteMode::MousePosition);
+				if (IsSingleNodeInClipboard())
+				{
+					NewAction.ExecuteAction.BindRaw(
+						this, &SAnimNotifyTrack::OnPasteNotifyClicked, ENotifyPasteMode::MousePosition, ENotifyPasteMultipleMode::Absolute);
 
-				MenuBuilder.AddMenuEntry(LOCTEXT("Paste", "Paste"), LOCTEXT("PasteToolTip", "Paste animation notify event here"), FSlateIcon(), NewAction);
+					MenuBuilder.AddMenuEntry(LOCTEXT("Paste", "Paste"), LOCTEXT("PasteToolTip", "Paste animation notify event here"), FSlateIcon(), NewAction);
+				}
+				else
+				{
+					NewAction.ExecuteAction.BindRaw(
+						this, &SAnimNotifyTrack::OnPasteNotifyClicked, ENotifyPasteMode::MousePosition, ENotifyPasteMultipleMode::Relative);
 
-				NewAction.ExecuteAction.BindRaw(
-					this, &SAnimNotifyTrack::OnPasteNotifyClicked, ENotifyPasteMode::SpecifiedTime);
+					MenuBuilder.AddMenuEntry(LOCTEXT("PasteMultRel", "Paste Multiple Relative"), LOCTEXT("PasteToolTip", "Paste multiple notifies beginning at the mouse cursor, maintaining the same relative spacing as the source."), FSlateIcon(), NewAction);
 
-				MenuBuilder.AddMenuEntry(LOCTEXT("PasteAt", "Paste at..."), LOCTEXT("PasteAtToolTip", "Paste animation notify event at a specified time"), FSlateIcon(), NewAction);
+					NewAction.ExecuteAction.BindRaw(
+						this, &SAnimNotifyTrack::OnPasteNotifyClicked, ENotifyPasteMode::MousePosition, ENotifyPasteMultipleMode::Absolute);
+
+					MenuBuilder.AddMenuEntry(LOCTEXT("PasteMultAbs", "Paste Multiple Absolute"), LOCTEXT("PasteToolTip", "Paste multiple notifies beginning at the mouse cursor, maintaining absolute spacing."), FSlateIcon(), NewAction);
+				}
 
 				if(OriginalTime < Sequence->SequenceLength)
 				{
 					NewAction.ExecuteAction.BindRaw(
-						this, &SAnimNotifyTrack::OnPasteNotifyClicked, ENotifyPasteMode::OriginalTime);
+						this, &SAnimNotifyTrack::OnPasteNotifyClicked, ENotifyPasteMode::OriginalTime, ENotifyPasteMultipleMode::Absolute);
 
 					FText DisplayText = FText::Format( LOCTEXT("PasteAtOriginalTime", "Paste at original time ({0})"), FText::AsNumber( OriginalTime) );
 					MenuBuilder.AddMenuEntry(DisplayText, LOCTEXT("PasteAtOriginalTimeToolTip", "Paste animation notify event at the time it was set to when it was copied"), FSlateIcon(), NewAction);
@@ -1802,69 +2039,9 @@ void SAnimNotifyTrack::CreateCommands()
 
 }
 
-uint8* SAnimNotifyTrack::FindNotifyPropertyData( int32 NotifyIndex, UArrayProperty*& ArrayProperty )
-{
-	// initialize to NULL
-	ArrayProperty = NULL;
-
-	if ( Sequence && Sequence->Notifies.IsValidIndex(NotifyIndex) )
-	{
-		// find Notifies property start point
-		UProperty* Property = FindField<UProperty>( Sequence->GetClass(), TEXT("Notifies"));
-
-		// found it and if it is array
-		if( Property && Property->IsA( UArrayProperty::StaticClass() ))
-		{
-			// find Property Value from UObject we got
-			uint8* PropertyValue = Property->ContainerPtrToValuePtr<uint8>(Sequence);
-
-			// it is array, so now get ArrayHelper and find the raw ptr of the data
-			ArrayProperty = CastChecked<UArrayProperty>(Property);
-			FScriptArrayHelper ArrayHelper(ArrayProperty, PropertyValue);
-
-			if (ArrayProperty->Inner && NotifyIndex < ArrayHelper.Num())
-			{
-				//Get property data based on selected index
-				return ArrayHelper.GetRawPtr(NotifyIndex);
-			}
-		}
-	}
-
-	return NULL;
-}
-
 void SAnimNotifyTrack::OnCopyNotifyClicked(int32 NotifyIndex)
 {
-	if (AnimNotifies.IsValidIndex(NotifyIndex))
-	{
-		int32 NotifyEventIndex = INDEX_NONE;
-		for( int32 i=0; i<Sequence->Notifies.Num(); i++ )
-		{
-			// find which index is what I owned
-			if( &(Sequence->Notifies[i]) == AnimNotifies[NotifyIndex] )
-			{
-				NotifyEventIndex = i;
-				break;
-			}
-		}
-
-
-		//Get array property
-		UArrayProperty* ArrayProperty = NULL;
-		uint8* PropData = FindNotifyPropertyData( NotifyEventIndex, ArrayProperty );
-
-		if (PropData)
-		{
-			//Append a sting identifier at the start in order to check whether the text represents an FAnimNotifyEvent structure.
-			const FString HeaderString( TEXT("COPY_ANIMNOTIFYEVENT") );
-			FString StrValue(HeaderString);
-			StrValue += FString::Printf( TEXT("OriginalTime=%f\n"), Sequence->Notifies[NotifyEventIndex].DisplayTime);
-
-			//Export property to text
-			ArrayProperty->Inner->ExportTextItem( StrValue, PropData, PropData, Sequence, PPF_Copy );
-			FPlatformMisc::ClipboardCopy( *StrValue );
-		}
-	}
+	OnCopyNotifies.ExecuteIfBound();
 }
 
 bool SAnimNotifyTrack::CanPasteAnimNotify() const
@@ -1872,10 +2049,11 @@ bool SAnimNotifyTrack::CanPasteAnimNotify() const
 	FString PropertyString;
 	const TCHAR* Buffer;
 	float OriginalTime;
-	return ReadNotifyPasteHeader(PropertyString, Buffer, OriginalTime);
+	float OriginalLength;
+	return ReadNotifyPasteHeader(PropertyString, Buffer, OriginalTime, OriginalLength);
 }
 
-bool SAnimNotifyTrack::ReadNotifyPasteHeader(FString& OutPropertyString, const TCHAR*& OutBuffer, float& OutOriginalTime) const
+bool SAnimNotifyTrack::ReadNotifyPasteHeader(FString& OutPropertyString, const TCHAR*& OutBuffer, float& OutOriginalTime, float& OutOriginalLength) const
 {
 	OutBuffer = NULL;
 	OutOriginalTime = -1.f;
@@ -1895,8 +2073,10 @@ bool SAnimNotifyTrack::ReadNotifyPasteHeader(FString& OutPropertyString, const T
 			OutBuffer += HeaderSize;
 
 			FString ReadLine;
-			FParse::Line( &OutBuffer, ReadLine );	// Need this to advance PastePtr, for multiple sockets
+			// Read the original time from the first notify
+			FParse::Line(&OutBuffer, ReadLine);
 			FParse::Value( *ReadLine, TEXT( "OriginalTime=" ), OutOriginalTime );
+			FParse::Value( *ReadLine, TEXT( "OriginalLength=" ), OutOriginalLength );
 			return true;
 		}
 	}
@@ -1904,41 +2084,10 @@ bool SAnimNotifyTrack::ReadNotifyPasteHeader(FString& OutPropertyString, const T
 	return false;
 }
 
-void SAnimNotifyTrack::OnPasteNotifyClicked(ENotifyPasteMode::Type PasteMode)
+void SAnimNotifyTrack::OnPasteNotifyClicked(ENotifyPasteMode::Type PasteMode, ENotifyPasteMultipleMode::Type MultiplePasteType)
 {
-	switch(PasteMode)
-	{
-		case ENotifyPasteMode::MousePosition:
-		{
-			OnPasteNotify(LastClickedTime);
-			break;
-		}
-		case ENotifyPasteMode::SpecifiedTime:
-		{
-			FString DefaultText = FString::Printf(TEXT("%f"), LastClickedTime);
-
-			TSharedRef<STextEntryPopup> TextEntry =
-				SNew(STextEntryPopup)
-				.Label(LOCTEXT("PasteAtTimeLabel", "Time"))
-				.DefaultText( FText::FromString(DefaultText) )
-				.OnTextCommitted( this, &SAnimNotifyTrack::OnPasteNotifyTimeSet );
-
-			FSlateApplication::Get().PushMenu(
-				AsShared(),
-				TextEntry,
-				FSlateApplication::Get().GetCursorPos(),
-				FPopupTransitionEffect( FPopupTransitionEffect::TypeInPopup )
-				);
-
-			TextEntry->FocusDefaultWidget();
-			break;
-		}
-		case ENotifyPasteMode::OriginalTime:
-		{
-			OnPasteNotify(-1.f);
-			break;
-		}
-	}
+	float ClickTime = PasteMode == ENotifyPasteMode::MousePosition ? LastClickedTime : -1.0f;
+	OnPasteNotifies.ExecuteIfBound(this, ClickTime, PasteMode, MultiplePasteType);
 }
 
 void SAnimNotifyTrack::OnPasteNotifyTimeSet(const FText& TimeText, ETextCommit::Type CommitInfo)
@@ -1953,76 +2102,101 @@ void SAnimNotifyTrack::OnPasteNotifyTimeSet(const FText& TimeText, ETextCommit::
 	FSlateApplication::Get().DismissAllMenus();
 }
 
-void SAnimNotifyTrack::OnPasteNotify(float TimeToPasteAt)
+void SAnimNotifyTrack::OnPasteNotify(float TimeToPasteAt, ENotifyPasteMultipleMode::Type MultiplePasteType)
 {
 	FString PropertyString;
 	const TCHAR* Buffer;
-	float OriginalTime;
+	float OriginalTime;				// Time in the original track that the node was placed
+	float OriginalLength;			// Length of the original anim sequence
+	float Shift = 0.0f;				// The amount to shift the saved display time of subsequent nodes
+	float ScaleMultiplier = 1.0f;	// Scale value to place nodes relatively in different sequences
 
-	if( ReadNotifyPasteHeader(PropertyString, Buffer, OriginalTime) )
+	if( ReadNotifyPasteHeader(PropertyString, Buffer, OriginalTime, OriginalLength) )
 	{
+		// If we're using relative mode, calculate the scale between the two sequences
+		if (MultiplePasteType == ENotifyPasteMultipleMode::Relative)
+		{
+			ScaleMultiplier = Sequence->SequenceLength / OriginalLength;
+		}
+
 		//Store object to undo pasting anim notify
 		const FScopedTransaction Transaction( LOCTEXT("PasteNotifyEvent", "Paste Anim Notify") );
-		Sequence->Modify();
-
-		// Add a new Anim Notify Event structure to the array
-		int32 NewNotifyIndex = Sequence->Notifies.AddZeroed();
 
-		//Get array property
-		UArrayProperty* ArrayProperty = NULL;
-		uint8* PropData = FindNotifyPropertyData( NewNotifyIndex, ArrayProperty );
+		Sequence->Modify();
 
-		if( PropData && ArrayProperty )
+		FString NotifyLine;
+		while (FParse::Line(&Buffer, NotifyLine))
 		{
-			ArrayProperty->Inner->ImportText( Buffer, PropData, PPF_Copy, NULL );
+			// Add a new Anim Notify Event structure to the array
+			int32 NewNotifyIndex = Sequence->Notifies.AddZeroed();
 
-			FAnimNotifyEvent& NewNotify = Sequence->Notifies[NewNotifyIndex];
+			//Get array property
+			UArrayProperty* ArrayProperty = NULL;
+			uint8* PropData = Sequence->FindNotifyPropertyData(NewNotifyIndex, ArrayProperty);
 
-			if(TimeToPasteAt >= 0.f)
+			if (PropData && ArrayProperty)
 			{
-				NewNotify.DisplayTime = TimeToPasteAt;
-				NewNotify.TriggerTimeOffset = GetTriggerTimeOffsetForType(Sequence->CalculateOffsetForNotify(TimeToPasteAt));
-			}
-			NewNotify.TrackIndex = TrackIndex;
+				ArrayProperty->Inner->ImportText(*NotifyLine, PropData, PPF_Copy, NULL);
 
-			// Get anim notify ptr of the new notify event array element (This ptr is the same as the source anim notify ptr)
-			UAnimNotify* AnimNotifyPtr = NewNotify.Notify;
-			if(AnimNotifyPtr)
-			{
-				//Create a new instance of anim notify subclass.
-				UAnimNotify* NewAnimNotifyPtr = Cast<UAnimNotify>(StaticDuplicateObject(AnimNotifyPtr, Sequence, TEXT("None")));
-				if( NewAnimNotifyPtr )
+				FAnimNotifyEvent& NewNotify = Sequence->Notifies[NewNotifyIndex];
+
+				if (TimeToPasteAt >= 0.f)
 				{
-					//Reassign new anim notify ptr to the pasted anim notify event.
-					NewNotify.Notify = NewAnimNotifyPtr;
+					Shift = TimeToPasteAt - (NewNotify.DisplayTime * ScaleMultiplier);
+					NewNotify.DisplayTime = TimeToPasteAt;
+
+					// Don't want to paste subsequent nodes at this position
+					TimeToPasteAt = -1.0f;
 				}
-			}
+				else
+				{
+					NewNotify.DisplayTime = (NewNotify.DisplayTime * ScaleMultiplier) + Shift;
+					// Clamp into the bounds of the sequence
+					NewNotify.DisplayTime = FMath::Clamp(NewNotify.DisplayTime, 0.0f, Sequence->SequenceLength);
+				}
+				NewNotify.TriggerTimeOffset = GetTriggerTimeOffsetForType(Sequence->CalculateOffsetForNotify(NewNotify.DisplayTime));
+				NewNotify.TrackIndex = TrackIndex;
 
-			UAnimNotifyState * AnimNotifyStatePtr = NewNotify.NotifyStateClass;
-			if( AnimNotifyStatePtr )
-			{
-				//Create a new instance of anim notify subclass.
-				UAnimNotifyState * NewAnimNotifyStatePtr = Cast<UAnimNotifyState>(StaticDuplicateObject(AnimNotifyStatePtr, Sequence, TEXT("None")));
-				if( NewAnimNotifyStatePtr )
+				// Get anim notify ptr of the new notify event array element (This ptr is the same as the source anim notify ptr)
+				UAnimNotify* AnimNotifyPtr = NewNotify.Notify;
+				if (AnimNotifyPtr)
 				{
-					//Reassign new anim notify ptr to the pasted anim notify event.
-					NewNotify.NotifyStateClass = NewAnimNotifyStatePtr;
+					//Create a new instance of anim notify subclass.
+					UAnimNotify* NewAnimNotifyPtr = Cast<UAnimNotify>(StaticDuplicateObject(AnimNotifyPtr, Sequence, TEXT("None")));
+					if (NewAnimNotifyPtr)
+					{
+						//Reassign new anim notify ptr to the pasted anim notify event.
+						NewNotify.Notify = NewAnimNotifyPtr;
+					}
 				}
 
-				if( (NewNotify.DisplayTime + NewNotify.Duration) > Sequence->SequenceLength)
+				UAnimNotifyState * AnimNotifyStatePtr = NewNotify.NotifyStateClass;
+				if (AnimNotifyStatePtr)
 				{
-					NewNotify.Duration = Sequence->SequenceLength - NewNotify.DisplayTime;
+					//Create a new instance of anim notify subclass.
+					UAnimNotifyState * NewAnimNotifyStatePtr = Cast<UAnimNotifyState>(StaticDuplicateObject(AnimNotifyStatePtr, Sequence, TEXT("None")));
+					if (NewAnimNotifyStatePtr)
+					{
+						//Reassign new anim notify ptr to the pasted anim notify event.
+						NewNotify.NotifyStateClass = NewAnimNotifyStatePtr;
+					}
+
+					if ((NewNotify.DisplayTime + NewNotify.Duration) > Sequence->SequenceLength)
+					{
+						NewNotify.Duration = Sequence->SequenceLength - NewNotify.DisplayTime;
+					}
 				}
 			}
-
-			Sequence->MarkPackageDirty();
-			OnUpdatePanel.ExecuteIfBound();
-		}	
-		else
-		{
-			//Remove newly created array element if pasting is not successful
-			Sequence->Notifies.RemoveAt(NewNotifyIndex);
+			else
+			{
+				//Remove newly created array element if pasting is not successful
+				Sequence->Notifies.RemoveAt(NewNotifyIndex);
+			}
 		}
+
+		OnDeselectAllNotifies.ExecuteIfBound();
+		Sequence->MarkPackageDirty();
+		OnUpdatePanel.ExecuteIfBound();
 	}
 }
 
@@ -2049,6 +2223,28 @@ void SAnimNotifyTrack::SetTriggerWeight(const FText& TriggerWeight, ETextCommit:
 	FSlateApplication::Get().DismissAllMenus();
 }
 
+bool SAnimNotifyTrack::IsSingleNodeSelected()
+{
+	return SelectedNotifyIndices.Num() == 1;
+}
+
+bool SAnimNotifyTrack::IsSingleNodeInClipboard()
+{
+	FString PropString;
+	const TCHAR* Buffer;
+	float OriginalTime;
+	float OriginalLength;
+	uint32 Count = 0;
+	if (ReadNotifyPasteHeader(PropString, Buffer, OriginalTime, OriginalLength))
+	{
+		// If reading a single line empties the buffer then we only have one notify in there.
+		FString TempLine;
+		FParse::Line(&Buffer, TempLine);
+		return *Buffer == 0;
+	}
+	return false;
+}
+
 void SAnimNotifyTrack::OnSetTriggerWeightNotifyClicked(int32 NotifyIndex)
 {
 	if (AnimNotifies.IsValidIndex(NotifyIndex))
@@ -2202,13 +2398,36 @@ int32 SAnimNotifyTrack::GetHitNotifyNode(const FGeometry& MyGeometry, const FVec
 
 FReply SAnimNotifyTrack::OnNotifyNodeDragStarted( TSharedRef<SAnimNotifyNode> NotifyNode, const FVector2D& ScreenCursorPos, const FVector2D& ScreenNodePosition, const bool bDragOnMarker, int32 NotifyIndex ) 
 {
-	SelectNotifyNode(NotifyIndex);
+	// Check to see if we've already selected the triggering node
+	if (!NotifyNode.Get().NotifyEvent->bSelected)
+	{
+		SelectNotifyNode(NotifyIndex, false);
+	}
+	
+	// Sort our nodes so we're acessing them in time order
+	SelectedNotifyIndices.Sort([this](const int32& A, const int32& B)
+	{
+		FAnimNotifyEvent* First = AnimNotifies[A];
+		FAnimNotifyEvent* Second = AnimNotifies[B];
+		return First->DisplayTime < Second->DisplayTime;
+	});
 
 	// If we're dragging one of the direction markers we don't need to call any further as we don't want the drag drop op
 	if(!bDragOnMarker)
 	{
-		NodeSlots->RemoveSlot(NotifyNode);
-		return OnNodeDragStarted.Execute(NotifyNode, ScreenCursorPos, ScreenNodePosition, bDragOnMarker);
+		TArray<TSharedPtr<SAnimNotifyNode>> NodesToDrag;
+		const float FirstNodeX = NotifyNodes[SelectedNotifyIndices[0]]->GetWidgetPosition().X;
+		
+		TSharedRef<SOverlay> DragBox = SNew(SOverlay);
+		for (auto Iter = SelectedNotifyIndices.CreateIterator(); Iter; ++Iter)
+		{
+			TSharedPtr<SAnimNotifyNode> Node = NotifyNodes[*Iter];
+			NodesToDrag.Add(Node);
+		}
+		
+		FVector2D DecoratorPosition = NodesToDrag[0]->GetWidgetPosition();
+		DecoratorPosition = CachedGeometry.LocalToAbsolute(DecoratorPosition);
+		return OnNodeDragStarted.Execute(NodesToDrag, DragBox, ScreenCursorPos, DecoratorPosition, bDragOnMarker);
 	}
 	else
 	{
@@ -2238,55 +2457,149 @@ FReply SAnimNotifyTrack::OnMouseButtonDown( const FGeometry& MyGeometry, const F
 	return FReply::Unhandled();
 }
 
-float SAnimNotifyTrack::CalculateTime( const FGeometry& MyGeometry, FVector2D NodePos )
+float SAnimNotifyTrack::CalculateTime( const FGeometry& MyGeometry, FVector2D NodePos, bool bInputIsAbsolute )
 {
-	NodePos = MyGeometry.AbsoluteToLocal(NodePos);
+	if(bInputIsAbsolute)
+	{
+		NodePos = MyGeometry.AbsoluteToLocal(NodePos);
+	}
 	FTrackScaleInfo ScaleInfo(ViewInputMin.Get(), ViewInputMax.Get(), 0, 0, MyGeometry.Size);
 	return FMath::Clamp<float>( ScaleInfo.LocalXToInput(NodePos.X), 0.f, Sequence->SequenceLength );
-	}
+}
 
 FReply SAnimNotifyTrack::OnDrop( const FGeometry& MyGeometry, const FDragDropEvent& DragDropEvent )
 {
-	TSharedPtr<FNotifyDragDropOp> DragDropOp = DragDropEvent.GetOperationAs<FNotifyDragDropOp>();
-	if (DragDropOp.IsValid())
+	return FReply::Unhandled();
+}
+
+void SAnimNotifyTrack::HandleNodeDrop(TSharedPtr<SAnimNotifyNode> Node)
+{
+	ensure(Node.IsValid());
+
+	OnDeselectAllNotifies.ExecuteIfBound();
+	const FScopedTransaction Transaction(LOCTEXT("MoveNotifyEvent", "Move Anim Notify"));
+	Sequence->Modify();
+
+	float LocalX = GetCachedGeometry().AbsoluteToLocal(Node->GetScreenPosition() + ScrubHandleSize).X;
+	float Time = GetCachedScaleInfo().LocalXToInput(LocalX);
+	FAnimNotifyEvent* DroppedEvent = Node->NotifyEvent;
+
+	if(Node->GetLastSnappedTime() != -1.0f)
 	{
-		HandleNodeDrop(*DragDropOp);
+		DroppedEvent->DisplayTime = Node->GetLastSnappedTime();
+	}
+	else
+	{
+		DroppedEvent->DisplayTime = Time;
+	}
+	DroppedEvent->RefreshTriggerOffset(Sequence->CalculateOffsetForNotify(DroppedEvent->DisplayTime));
+	DroppedEvent->RefreshEndTriggerOffset(Sequence->CalculateOffsetForNotify(DroppedEvent->DisplayTime + DroppedEvent->Duration));
+	DroppedEvent->TrackIndex = TrackIndex;
 
-		return FReply::Handled();
+	Sequence->MarkPackageDirty();
+	OnUpdatePanel.ExecuteIfBound();
+}
+
+void SAnimNotifyTrack::DisconnectSelectedNodesForDrag(TArray<TSharedPtr<SAnimNotifyNode>>& DragNodes)
+{
+	if(SelectedNotifyIndices.Num() == 0)
+	{
+		return;
 	}
 
-	return FReply::Unhandled();
+	const float FirstNodeX = NotifyNodes[SelectedNotifyIndices[0]]->GetWidgetPosition().X;
+
+	for(auto Iter = SelectedNotifyIndices.CreateIterator(); Iter; ++Iter)
+	{
+		TSharedPtr<SAnimNotifyNode> Node = NotifyNodes[*Iter];
+		NodeSlots->RemoveSlot(Node->AsShared());
+
+		DragNodes.Add(Node);
+	}
 }
 
-void SAnimNotifyTrack::HandleNodeDrop(const FNotifyDragDropOp& DragDropOp)
+void SAnimNotifyTrack::AppendSelectionToSet(FGraphPanelSelectionSet& SelectionSet)
 {
-	// if sample is highlighted, overwrite the animation for it
-	// make sure if it's same sequence we're dealing with 
-	UAnimSequenceBase* DroppedSequence = DragDropOp.Sequence;
-	FAnimNotifyEvent* DroppedEvent = DragDropOp.NotifyEvent; 
-	if (DroppedSequence == Sequence)
+	// Add our selection to the provided set
+	for(auto Iter = SelectedNotifyIndices.CreateConstIterator(); Iter; ++Iter)
 	{
-		DragDropOp.CurrentDragXPosition = -1.0f;
+		int32 Index = *Iter;
+		check(AnimNotifies.IsValidIndex(Index));
+		FAnimNotifyEvent* Event = AnimNotifies[Index];
 
-		const FScopedTransaction Transaction( LOCTEXT("MoveNotifyEvent", "Move Anim Notify") );
-		Sequence->Modify();
+		if(Event->Notify)
+		{
+			SelectionSet.Add(Event->Notify);
+		}
+		else if(Event->NotifyStateClass)
+		{
+			SelectionSet.Add(Event->NotifyStateClass);
+		}
+	}
+}
 
-		if(DragDropOp.SnapTime != -1)
+void SAnimNotifyTrack::AppendSelectionToArray(TArray<const FAnimNotifyEvent*>& Selection) const
+{
+	for(int32 Idx : SelectedNotifyIndices)
+	{
+		check(AnimNotifies.IsValidIndex(Idx));
+		FAnimNotifyEvent* Event = AnimNotifies[Idx];
+
+		if(Event)
 		{
-			DroppedEvent->DisplayTime = DragDropOp.SnapTime;
+			Selection.Add(Event);
 		}
-		else
+	}
+}
+
+void SAnimNotifyTrack::PasteSingleNotify(FString& NotifyString, float PasteTime)
+{
+	int32 NewIdx = Sequence->Notifies.AddZeroed();
+	UArrayProperty* ArrayProperty = NULL;
+	uint8* PropertyData = Sequence->FindNotifyPropertyData(NewIdx, ArrayProperty);
+
+	if(PropertyData && ArrayProperty)
+	{
+		ArrayProperty->Inner->ImportText(*NotifyString, PropertyData, PPF_Copy, NULL);
+
+		FAnimNotifyEvent& NewNotify = Sequence->Notifies[NewIdx];
+
+		if(PasteTime != -1.0f)
 		{
-			DroppedEvent->DisplayTime = CalculateTime(CachedGeometry, DragDropOp.LastNodePos+DragDropOp.NotifyOffset);
+			NewNotify.DisplayTime = PasteTime;
 		}
-		DroppedEvent->RefreshTriggerOffset(Sequence->CalculateOffsetForNotify(DroppedEvent->DisplayTime));
-		DroppedEvent->RefreshEndTriggerOffset(Sequence->CalculateOffsetForNotify(DroppedEvent->DisplayTime + DroppedEvent->Duration));
 
-		DroppedEvent->TrackIndex = TrackIndex;
+		// Make sure the notify is within the track area
+		NewNotify.DisplayTime = FMath::Clamp(NewNotify.DisplayTime, 0.0f, Sequence->SequenceLength);
+		NewNotify.TriggerTimeOffset = GetTriggerTimeOffsetForType(Sequence->CalculateOffsetForNotify(NewNotify.DisplayTime));
+		NewNotify.TrackIndex = TrackIndex;
 
-		Sequence->MarkPackageDirty();
-		OnUpdatePanel.ExecuteIfBound();
+		if(NewNotify.Notify)
+		{
+			UAnimNotify* NewNotifyObject = Cast<UAnimNotify>(StaticDuplicateObject(NewNotify.Notify, Sequence, TEXT("None")));
+			check(NewNotifyObject);
+			NewNotify.Notify = NewNotifyObject;
+		}
+		else if(NewNotify.NotifyStateClass)
+		{
+			UAnimNotifyState* NewNotifyStateObject = Cast<UAnimNotifyState>(StaticDuplicateObject(NewNotify.NotifyStateClass, Sequence, TEXT("None")));
+			check(NewNotifyStateObject);
+			NewNotify.NotifyStateClass = NewNotifyStateObject;
+
+			// Clamp duration into the sequence
+			NewNotify.Duration = FMath::Clamp(NewNotify.Duration, 1 / 30.0f, Sequence->SequenceLength - NewNotify.DisplayTime);
+			NewNotify.EndTriggerTimeOffset = GetTriggerTimeOffsetForType(Sequence->CalculateOffsetForNotify(NewNotify.DisplayTime + NewNotify.Duration));
+		}
 	}
+	else
+	{
+		// Paste failed, remove the notify
+		Sequence->Notifies.RemoveAt(NewIdx);
+	}
+
+	OnDeselectAllNotifies.ExecuteIfBound();
+	Sequence->MarkPackageDirty();
+	OnUpdatePanel.ExecuteIfBound();
 }
 
 //////////////////////////////////////////////////////////////////////////
@@ -2332,6 +2645,9 @@ void SNotifyEdTrack::Construct(const FArguments& InArgs)
 				.OnRequestTrackPan(FPanTrackRequest::CreateSP(PanelRef, &SAnimNotifyPanel::PanInputViewRange))
 				.OnRequestOffsetRefresh(InArgs._OnRequestRefreshOffsets)
 				.OnDeleteNotify(InArgs._OnDeleteNotify)
+				.OnDeselectAllNotifies(InArgs._OnDeselectAllNotifies)
+				.OnCopyNotifies(InArgs._OnCopyNotifies)
+				.OnPasteNotifies(InArgs._OnPasteNotifies)
 			]
 
 			+SHorizontalBox::Slot()
@@ -2537,6 +2853,38 @@ void SAnimNotifyPanel::Update()
 	}
 }
 
+bool SAnimNotifyPanel::ReadNotifyPasteHeader(FString& OutPropertyString, const TCHAR*& OutBuffer, float& OutOriginalTime, float& OutOriginalLength, int32& OutTrackSpan) const
+{
+	OutBuffer = NULL;
+	OutOriginalTime = -1.f;
+
+	FPlatformMisc::ClipboardPaste(OutPropertyString);
+
+	if(!OutPropertyString.IsEmpty())
+	{
+		//Remove header text
+		const FString HeaderString(TEXT("COPY_ANIMNOTIFYEVENT"));
+
+		//Check for string identifier in order to determine whether the text represents an FAnimNotifyEvent.
+		if(OutPropertyString.StartsWith(HeaderString) && OutPropertyString.Len() > HeaderString.Len())
+		{
+			int32 HeaderSize = HeaderString.Len();
+			OutBuffer = *OutPropertyString;
+			OutBuffer += HeaderSize;
+
+			FString ReadLine;
+			// Read the original time from the first notify
+			FParse::Line(&OutBuffer, ReadLine);
+			FParse::Value(*ReadLine, TEXT("OriginalTime="), OutOriginalTime);
+			FParse::Value(*ReadLine, TEXT("OriginalLength="), OutOriginalLength);
+			FParse::Value(*ReadLine, TEXT("TrackSpan="), OutTrackSpan);
+			return true;
+		}
+	}
+
+	return false;
+}
+
 void SAnimNotifyPanel::RefreshNotifyTracks()
 {
 	check (Sequence);
@@ -2567,11 +2915,14 @@ void SAnimNotifyPanel::RefreshNotifyTracks()
 			.OnGetScrubValue(OnGetScrubValue)
 			.OnGetDraggedNodePos(this, &SAnimNotifyPanel::CalculateDraggedNodePos)
 			.OnUpdatePanel(this, &SAnimNotifyPanel::Update)
-			.OnSelectionChanged(OnSelectionChanged)
+			.OnSelectionChanged(this, &SAnimNotifyPanel::OnTrackSelectionChanged)
 			.OnNodeDragStarted(this, &SAnimNotifyPanel::OnNotifyNodeDragStarted)
 			.MarkerBars(MarkerBars)
 			.OnRequestRefreshOffsets(OnRequestRefreshOffsets)
 			.OnDeleteNotify(this, &SAnimNotifyPanel::OnGenericDelete)
+			.OnDeselectAllNotifies(this, &SAnimNotifyPanel::DeselectAllNotifies)
+			.OnCopyNotifies(this, &SAnimNotifyPanel::CopySelectedNotifiesToClipboard)
+			.OnPasteNotifies(this, &SAnimNotifyPanel::OnPasteNotifies)
 		];
 
 		NotifyAnimTracks.Add(EdTrack->NotifyTrack);
@@ -2583,10 +2934,57 @@ float SAnimNotifyPanel::CalculateDraggedNodePos() const
 	return CurrentDragXPosition;
 }
 
-FReply SAnimNotifyPanel::OnNotifyNodeDragStarted(TSharedRef<SAnimNotifyNode> NotifyNode, const FVector2D& ScreenCursorPos, const FVector2D& ScreenNodePosition, const bool bDragOnMarker)
+FReply SAnimNotifyPanel::OnNotifyNodeDragStarted(TArray<TSharedPtr<SAnimNotifyNode>> NotifyNodes, TSharedRef<SWidget> Decorator, const FVector2D& ScreenCursorPos, const FVector2D& ScreenNodePosition, const bool bDragOnMarker)
 {
+	TSharedRef<SOverlay> NodeDragDecorator = SNew(SOverlay);
+	TArray<TSharedPtr<SAnimNotifyNode>> Nodes;
+
+	for(TSharedPtr<SAnimNotifyTrack> Track : NotifyAnimTracks)
+	{
+		Track->DisconnectSelectedNodesForDrag(Nodes);
+	}
+
+	FVector2D OverlayOrigin = Nodes[0]->GetScreenPosition();
+	FVector2D OverlayExtents = Nodes[0]->GetScreenPosition();
+	for(int32 Idx = 1 ; Idx < Nodes.Num() ; ++Idx)
+	{
+		TSharedPtr<SAnimNotifyNode> Node = Nodes[Idx];
+		FVector2D NodePosition = Node->GetScreenPosition();
+		float NodeDuration = Node->GetDurationSize();
+
+		if(NodePosition.X < OverlayOrigin.X)
+		{
+			OverlayOrigin.X = NodePosition.X;
+		}
+		else if(NodePosition.X + NodeDuration > OverlayExtents.X)
+		{
+			OverlayExtents.X = NodePosition.X + NodeDuration;
+		}
+
+		if(NodePosition.Y < OverlayOrigin.Y)
+		{
+			OverlayOrigin.Y = NodePosition.Y;
+		}
+		else if(NodePosition.Y + NotifyHeight > OverlayExtents.Y)
+		{
+			OverlayExtents.Y = NodePosition.Y + NotifyHeight;
+		}
+	}
+	OverlayExtents -= OverlayOrigin;
+
+	for(TSharedPtr<SAnimNotifyNode> Node : Nodes)
+	{
+		FVector2D OffsetFromFirst(Node->GetScreenPosition() - OverlayOrigin);
+
+		NodeDragDecorator->AddSlot()
+			.Padding(FMargin(OffsetFromFirst.X, OffsetFromFirst.Y, 0.0f, 0.0f))
+			[
+				Node->AsShared()
+			];
+	}
+
 	FPanTrackRequest PanRequestDelegate = FPanTrackRequest::CreateSP(this, &SAnimNotifyPanel::PanInputViewRange);
-	return FReply::Handled().BeginDragDrop(FNotifyDragDropOp::New(NotifyNode, NotifyAnimTracks, Sequence, ScreenCursorPos, ScreenNodePosition, CurrentDragXPosition, PanRequestDelegate, MarkerBars));
+	return FReply::Handled().BeginDragDrop(FNotifyDragDropOp::New(Nodes, NodeDragDecorator, NotifyAnimTracks, Sequence, ScreenCursorPos, OverlayOrigin, OverlayExtents, CurrentDragXPosition, PanRequestDelegate, MarkerBars));
 }
 
 void SAnimNotifyPanel::PostUndo()
@@ -2633,4 +3031,166 @@ void SAnimNotifyPanel::SetSequence(class UAnimSequenceBase *	InSequence)
 	}
 }
 
+void SAnimNotifyPanel::OnTrackSelectionChanged()
+{
+	// Need to collect selection info from all tracks
+	FGraphPanelSelectionSet SelectionSet;
+	for(TSharedPtr<SAnimNotifyTrack> Track : NotifyAnimTracks)
+	{
+		Track->AppendSelectionToSet(SelectionSet);
+	}
+
+	OnSelectionChanged.ExecuteIfBound(SelectionSet);
+}
+
+void SAnimNotifyPanel::DeselectAllNotifies()
+{
+	for(TSharedPtr<SAnimNotifyTrack> Track : NotifyAnimTracks)
+	{
+		Track->DeselectAllNotifyNodes(false);
+	}
+
+	OnTrackSelectionChanged();
+}
+
+void SAnimNotifyPanel::CopySelectedNotifiesToClipboard() const
+{
+	// Grab the selected events
+	TArray<const FAnimNotifyEvent*> SelectedNotifies;
+	for(TSharedPtr<SAnimNotifyTrack> Track : NotifyAnimTracks)
+	{
+		Track->AppendSelectionToArray(SelectedNotifies);
+	}
+
+	const FString HeaderString(TEXT("COPY_ANIMNOTIFYEVENT"));
+	
+	if(SelectedNotifies.Num() > 0)
+	{
+		FString StrValue(HeaderString);
+
+		// Sort by track
+		SelectedNotifies.Sort([](const FAnimNotifyEvent& A, const FAnimNotifyEvent& B)
+		{
+			return (A.TrackIndex > B.TrackIndex) || (A.TrackIndex == B.TrackIndex && A.DisplayTime < B.DisplayTime);
+		});
+
+		// Need to find how many tracks this selection spans and the minimum time to use as the beginning of the selection
+		int32 MinTrack = MAX_int32;
+		int32 MaxTrack = MIN_int32;
+		float MinTime = MAX_flt;
+		for(const FAnimNotifyEvent* Event : SelectedNotifies)
+		{
+			MinTrack = FMath::Min(MinTrack, Event->TrackIndex);
+			MaxTrack = FMath::Max(MaxTrack, Event->TrackIndex);
+			MinTime = FMath::Min(MinTime, Event->DisplayTime);
+		}
+
+		int32 TrackSpan = MaxTrack - MinTrack + 1;
+
+		StrValue += FString::Printf(TEXT("OriginalTime=%f,"), MinTime);
+		StrValue += FString::Printf(TEXT("OriginalLength=%f,"), Sequence->SequenceLength);
+		StrValue += FString::Printf(TEXT("TrackSpan=%d"), TrackSpan);
+
+		for(const FAnimNotifyEvent* Event : SelectedNotifies)
+		{
+			// Locate the notify in the sequence, we need the sequence index; but also need to
+			// keep the order we're currently in.
+			int32 Index = INDEX_NONE;
+			for(int32 NotifyIdx = 0 ; NotifyIdx < Sequence->Notifies.Num() ; ++NotifyIdx)
+			{
+				if(Event == &Sequence->Notifies[NotifyIdx])
+				{
+					Index = NotifyIdx;
+					break;
+				}
+			}
+
+			check(Index != INDEX_NONE);
+			
+			StrValue += "\n";
+			UArrayProperty* ArrayProperty = NULL;
+			uint8* PropertyData = Sequence->FindNotifyPropertyData(Index, ArrayProperty);
+			if(PropertyData && ArrayProperty)
+			{
+				ArrayProperty->Inner->ExportTextItem(StrValue, PropertyData, PropertyData, Sequence, PPF_Copy);
+			}
+
+		}
+		FPlatformMisc::ClipboardCopy(*StrValue);
+	}
+}
+
+void SAnimNotifyPanel::OnPasteNotifies(SAnimNotifyTrack* RequestTrack, float ClickTime, ENotifyPasteMode::Type PasteMode, ENotifyPasteMultipleMode::Type MultiplePasteType)
+{
+	int32 PasteIdx = RequestTrack->GetTrackIndex();
+	int32 NumTracks = NotifyAnimTracks.Num();
+	FString PropString;
+	const TCHAR* Buffer;
+	float OrigBeginTime;
+	float OrigLength;
+	int32 TrackSpan;
+	int32 FirstTrack = -1;
+	float ScaleMultiplier = 1.0f;
+
+	if(ReadNotifyPasteHeader(PropString, Buffer, OrigBeginTime, OrigLength, TrackSpan))
+	{
+		DeselectAllNotifies();
+
+		FScopedTransaction Transaction(LOCTEXT("PasteNotifyEvent", "Paste Anim Notifies"));
+		Sequence->Modify();
+
+		if(ClickTime == -1.0f)
+		{
+			// We want to place the notifies exactly where they were
+			ClickTime = OrigBeginTime;
+		}
+
+		// Expand the number of tracks if we don't have enough.
+		check(TrackSpan > 0);
+		if(PasteIdx - (TrackSpan - 1) < 0)
+		{
+			int32 TracksToAdd = PasteIdx + TrackSpan - 1;
+			while(TracksToAdd)
+			{
+				InsertTrack(PasteIdx++);
+				--TracksToAdd;
+			}
+			NumTracks = NotifyAnimTracks.Num();
+			RequestTrack = NotifyAnimTracks[NumTracks - 1 - PasteIdx].Get();
+		}
+
+		// Scaling for relative paste
+		if(MultiplePasteType == ENotifyPasteMultipleMode::Relative)
+		{
+			ScaleMultiplier = Sequence->SequenceLength / OrigLength;
+		}
+
+		// Process each line of the paste buffer and spawn notifies
+		FString CurrentLine;
+		while(FParse::Line(&Buffer, CurrentLine))
+		{
+			int32 OriginalTrack;
+			float OrigTime;
+			float PasteTime = -1.0f;
+			if(FParse::Value(*CurrentLine, TEXT("TrackIndex="), OriginalTrack) && FParse::Value(*CurrentLine, TEXT("DisplayTime="), OrigTime))
+			{
+				// Store the first track so we know where to place notifies
+				if(FirstTrack < 0)
+				{
+					FirstTrack = OriginalTrack;
+				}
+				int32 TrackOffset = OriginalTrack - FirstTrack;
+
+				float TimeOffset = OrigTime - OrigBeginTime;
+				float TimeToPaste = ClickTime + TimeOffset * ScaleMultiplier;
+
+				// Have to invert the index here as tracks are stored in reverse
+				TSharedPtr<SAnimNotifyTrack> TrackToUse = NotifyAnimTracks[NotifyAnimTracks.Num() - 1 - (PasteIdx + TrackOffset)];
+
+				TrackToUse->PasteSingleNotify(CurrentLine, TimeToPaste);
+			}
+		}
+	}
+}
+
 #undef LOCTEXT_NAMESPACE
diff --git a/Engine/Source/Editor/Persona/Private/SAnimNotifyPanel.h b/Engine/Source/Editor/Persona/Private/SAnimNotifyPanel.h
--- a/Engine/Source/Editor/Persona/Private/SAnimNotifyPanel.h
+++ b/Engine/Source/Editor/Persona/Private/SAnimNotifyPanel.h
@@ -11,15 +11,46 @@
 #include "SAnimEditorBase.h"
 
 DECLARE_DELEGATE_OneParam( FOnSelectionChanged, const FGraphPanelSelectionSet& )
+DECLARE_DELEGATE( FOnTrackSelectionChanged )
 DECLARE_DELEGATE( FOnUpdatePanel )
 DECLARE_DELEGATE_RetVal( float, FOnGetScrubValue )
 DECLARE_DELEGATE( FRefreshOffsetsRequest )
 DECLARE_DELEGATE( FDeleteNotify )
+DECLARE_DELEGATE( FDeselectAllNotifies )
+DECLARE_DELEGATE( FCopyNotifies )
 
 class SAnimNotifyNode;
 class SAnimNotifyTrack;
 class FNotifyDragDropOp;
 
+namespace ENotifyPasteMode
+{
+	enum Type
+	{
+		MousePosition,
+		OriginalTime
+	};
+}
+
+namespace ENotifyPasteMultipleMode
+{
+	enum Type
+	{
+		Relative,
+		Absolute
+	};
+}
+
+namespace ENotifyStateHandleHit
+{
+	enum Type
+	{
+		Start,
+		End,
+		None
+	};
+}
+
 //////////////////////////////////////////////////////////////////////////
 // SAnimNotifyPanel
 
@@ -73,10 +104,13 @@ public:
 	float CalculateDraggedNodePos() const;
 
 	/**Handler for when a notify node drag has been initiated */
-	FReply OnNotifyNodeDragStarted(TSharedRef<SAnimNotifyNode> NotifyNode, const FVector2D& ScreenCursorPos, const FVector2D& ScreenNodePosition, const bool bDragOnMarker);
+	FReply SAnimNotifyPanel::OnNotifyNodeDragStarted(TArray<TSharedPtr<SAnimNotifyNode>> NotifyNodes, TSharedRef<SWidget> Decorator, const FVector2D& ScreenCursorPos, const FVector2D& ScreenNodePosition, const bool bDragOnMarker);
 
 	virtual float GetSequenceLength() const OVERRIDE {return Sequence->SequenceLength;}
 
+	void CopySelectedNotifiesToClipboard() const;
+	void OnPasteNotifies(SAnimNotifyTrack* RequestTrack, float ClickTime, ENotifyPasteMode::Type PasteMode, ENotifyPasteMultipleMode::Type MultiplePasteType);
+
 private:
 	TSharedPtr<SBorder> PanelArea;
 	class UAnimSequenceBase* Sequence;
@@ -94,6 +128,9 @@ private:
 	/** Cached list of anim tracks for notify node drag drop */
 	TArray<TSharedPtr<SAnimNotifyTrack>> NotifyAnimTracks;
 
+	// Read common info from the clipboard
+	bool ReadNotifyPasteHeader(FString& OutPropertyString, const TCHAR*& OutBuffer, float& OutOriginalTime, float& OutOriginalLength, int32& OutTrackSpan) const;
+
 	// this just refresh notify tracks - UI purpose only
 	// do not call this from here. This gets called by asset. 
 	void RefreshNotifyTracks();
@@ -108,6 +145,12 @@ private:
 		return true;
 	}
 
+	// Called when a track changes it's selection; iterates all tracks collecting selected items
+	void OnTrackSelectionChanged();
+
+	// Called to deselect all notifies across all tracks
+	void DeselectAllNotifies();
+
 	/** Persona reference **/
 	TWeakPtr<FPersona> PersonaPtr;
 
diff --git a/Engine/Source/Editor/UnrealEd/Public/SCurveEditor.h b/Engine/Source/Editor/UnrealEd/Public/SCurveEditor.h
--- a/Engine/Source/Editor/UnrealEd/Public/SCurveEditor.h
+++ b/Engine/Source/Editor/UnrealEd/Public/SCurveEditor.h
@@ -38,10 +38,9 @@ struct FTrackScaleInfo
 	}
 
 	/** Local Widget Space -> Curve Input domain. */
-	float LocalXToInput(float ScreenX)
+	float LocalXToInput(float ScreenX) const
 	{
-		float LocalX = ScreenX;
-		return (LocalX/PixelsPerInput) + ViewMinInput;
+		return (ScreenX/PixelsPerInput) + ViewMinInput;
 	}
 
 	/** Curve Input domain -> local Widget Space */
@@ -51,20 +50,18 @@ struct FTrackScaleInfo
 	}
 
 	/** Local Widget Space -> Curve Output domain. */
-	float LocalYToOutput(float ScreenY)
+	float LocalYToOutput(float ScreenY) const
 	{
 		return (ViewOutputRange + ViewMinOutput) - (ScreenY/PixelsPerOutput);
-		//float LocalY = ScreenY;
-		//return (LocalY/PixelsPerOutput) + ViewMinOutput;
 	}
 
 	/** Curve Output domain -> local Widget Space */
-	float OutputToLocalY(float Output)
+	float OutputToLocalY(float Output) const
 	{
 		return (ViewOutputRange - (Output - ViewMinOutput)) * PixelsPerOutput;
 	}
 
-	float GetTrackCenterY()
+	float GetTrackCenterY() const
 	{
 		return (0.5f * WidgetSize.Y);
 	}
diff --git a/Engine/Source/Runtime/Engine/Classes/Animation/AnimSequenceBase.h b/Engine/Source/Runtime/Engine/Classes/Animation/AnimSequenceBase.h
--- a/Engine/Source/Runtime/Engine/Classes/Animation/AnimSequenceBase.h
+++ b/Engine/Source/Runtime/Engine/Classes/Animation/AnimSequenceBase.h
@@ -339,6 +339,9 @@ class UAnimSequenceBase : public UAnimationAsset
 	virtual EAnimEventTriggerOffsets::Type CalculateOffsetForNotify(float NotifyDisplayTime) const;
 
 	virtual void GetAssetRegistryTags(TArray<FAssetRegistryTag>& OutTags) const OVERRIDE;
+	
+	// Get a pointer to the data for a given Anim Notify
+	ENGINE_API uint8* FindNotifyPropertyData(int32 NotifyIndex, UArrayProperty*& ArrayProperty);
 
 #endif	//WITH_EDITORONLY_DATA
 
diff --git a/Engine/Source/Runtime/Engine/Private/Animation/AnimSequenceBase.cpp b/Engine/Source/Runtime/Engine/Private/Animation/AnimSequenceBase.cpp
--- a/Engine/Source/Runtime/Engine/Private/Animation/AnimSequenceBase.cpp
+++ b/Engine/Source/Runtime/Engine/Private/Animation/AnimSequenceBase.cpp
@@ -461,6 +461,36 @@ void UAnimSequenceBase::GetAssetRegistryTags(TArray<FAssetRegistryTag>& OutTags)
 	}
 }
 
+uint8* UAnimSequenceBase::FindNotifyPropertyData(int32 NotifyIndex, UArrayProperty*& ArrayProperty)
+{
+	// initialize to NULL
+	ArrayProperty = NULL;
+
+	if(Notifies.IsValidIndex(NotifyIndex))
+	{
+		// find Notifies property start point
+		UProperty* Property = FindField<UProperty>(GetClass(), TEXT("Notifies"));
+
+		// found it and if it is array
+		if(Property && Property->IsA(UArrayProperty::StaticClass()))
+		{
+			// find Property Value from UObject we got
+			uint8* PropertyValue = Property->ContainerPtrToValuePtr<uint8>(this);
+
+			// it is array, so now get ArrayHelper and find the raw ptr of the data
+			ArrayProperty = CastChecked<UArrayProperty>(Property);
+			FScriptArrayHelper ArrayHelper(ArrayProperty, PropertyValue);
+
+			if(ArrayProperty->Inner && NotifyIndex < ArrayHelper.Num())
+			{
+				//Get property data based on selected index
+				return ArrayHelper.GetRawPtr(NotifyIndex);
+			}
+		}
+	}
+	return NULL;
+}
+
 #endif	//WITH_EDITOR
 
 
