Fix Mac OS X compilation due to erroneous extra qualifier on SAnimNotifyPanel::OnNotifyNodeDragStarted's declarations.

diff --git a/Engine/Source/Editor/Persona/Private/SAnimNotifyPanel.h b/Engine/Source/Editor/Persona/Private/SAnimNotifyPanel.h
--- a/Engine/Source/Editor/Persona/Private/SAnimNotifyPanel.h
+++ b/Engine/Source/Editor/Persona/Private/SAnimNotifyPanel.h
@@ -104,7 +104,7 @@ public:
 	float CalculateDraggedNodePos() const;
 
 	/**Handler for when a notify node drag has been initiated */
-	FReply SAnimNotifyPanel::OnNotifyNodeDragStarted(TArray<TSharedPtr<SAnimNotifyNode>> NotifyNodes, TSharedRef<SWidget> Decorator, const FVector2D& ScreenCursorPos, const FVector2D& ScreenNodePosition, const bool bDragOnMarker);
+	FReply OnNotifyNodeDragStarted(TArray<TSharedPtr<SAnimNotifyNode>> NotifyNodes, TSharedRef<SWidget> Decorator, const FVector2D& ScreenCursorPos, const FVector2D& ScreenNodePosition, const bool bDragOnMarker);
 
 	virtual float GetSequenceLength() const OVERRIDE {return Sequence->SequenceLength;}
 
