Fixed so that in persona, it exports FBX using current preview mesh : this is for retargetted mesh

diff --git a/Engine/Source/Editor/Persona/Private/SAnimationSequenceBrowser.cpp b/Engine/Source/Editor/Persona/Private/SAnimationSequenceBrowser.cpp
--- a/Engine/Source/Editor/Persona/Private/SAnimationSequenceBrowser.cpp
+++ b/Engine/Source/Editor/Persona/Private/SAnimationSequenceBrowser.cpp
@@ -234,19 +234,22 @@ void SAnimationSequenceBrowser::OnExportToFBX(TArray<FAssetData> SelectedAssets)
 	{
 		IDesktopPlatform* DesktopPlatform = FDesktopPlatformModule::Get();
 
-		if(DesktopPlatform)
+		if(DesktopPlatform && PersonaPtr.IsValid())
 		{
+			USkeletalMesh * PreviewMesh = PersonaPtr.Pin()->GetPreviewMeshComponent()->SkeletalMesh;
+			if (PreviewMesh == NULL)
+			{
+				FMessageDialog::Open(EAppMsgType::Ok, LOCTEXT("ExportToFBXExportMissingPreviewMesh", "ERROR: Missing preview mesh") );
+				return;
+			}
+
 			TArray<TWeakObjectPtr<UAnimSequence>> AnimSequences;
 			for(auto Iter = SelectedAssets.CreateIterator(); Iter; ++Iter)
 			{
 				if(UAnimSequence* AnimSequence = Cast<UAnimSequence>(Iter->GetAsset()) )
 				{
-					USkeleton * AnimSkeleton = AnimSequence->GetSkeleton();
-					// We cannot export if we don't have a preview mesh, for the moment silently skip ones that don't
-					if (AnimSkeleton && AnimSkeleton->GetPreviewMesh())
-					{
-						AnimSequences.Add( TWeakObjectPtr<UAnimSequence>(AnimSequence) );
-					}
+					// we only shows anim sequence that belong to this skeleton
+					AnimSequences.Add( TWeakObjectPtr<UAnimSequence>(AnimSequence) );
 				}
 			}
 
@@ -353,6 +356,7 @@ void SAnimationSequenceBrowser::OnExportToFBX(TArray<FAssetData> SelectedAssets)
 				const bool bShowProgressDialog = true;
 				GWarn->BeginSlowTask( LOCTEXT("ExportToFBXProgress", "Exporting Animation(s) to FBX"), bShowProgressDialog, bShowCancel);
 
+				// make sure to use PreviewMesh, when export inside of Persona
 				const int32 NumberOfAnimations = AnimSequences.Num();
 				for(int32 i = 0; i < NumberOfAnimations; ++i)
 				{
@@ -362,7 +366,7 @@ void SAnimationSequenceBrowser::OnExportToFBX(TArray<FAssetData> SelectedAssets)
 
 					FString FileName = FString::Printf(TEXT("%s/%s"), *DestinationFolder, *AnimFileNames[i]);
 
-					FbxAnimUtils::ExportAnimFbx( *FileName, AnimSequence, AnimSequence->GetSkeleton()->GetPreviewMesh(), bSaveSkeletalMesh );
+					FbxAnimUtils::ExportAnimFbx( *FileName, AnimSequence, PreviewMesh, bSaveSkeletalMesh );
 				}
 
 				GWarn->EndSlowTask( );
