Renamed 'Animation Blueprint Editor' tab to 'Anim Bluprint Editor' so that it fits on screen

#ttp 332762 - ANIM: EDITOR: ANIMBP: Tab name cut off

#branch UE4

diff --git a/Engine/Source/Editor/Persona/Private/Shared/PersonaMode.cpp b/Engine/Source/Editor/Persona/Private/Shared/PersonaMode.cpp
--- a/Engine/Source/Editor/Persona/Private/Shared/PersonaMode.cpp
+++ b/Engine/Source/Editor/Persona/Private/Shared/PersonaMode.cpp
@@ -242,7 +242,7 @@ TSharedRef<SWidget> FRetargetSourceManagerTabSummoner::CreateTabBody(const FWork
 FAnimBlueprintDefaultsEditorSummoner::FAnimBlueprintDefaultsEditorSummoner(TSharedPtr<class FAssetEditorToolkit> InHostingApp)
 	: FWorkflowTabFactory(FPersonaTabs::AnimBlueprintDefaultsEditorID, InHostingApp)
 {
-	TabLabel = LOCTEXT("AnimBlueprintDefaultsTabTitle", "Animation Blueprint Editor");
+	TabLabel = LOCTEXT("AnimBlueprintDefaultsTabTitle", "Anim Blueprint Editor");
 
 	bIsSingleton = true;
 
