Added in new Persona tab that enables overriding of asset-referencing graph nodes in parent anim graphs.
#ttp 327253 - ANIM: DEV: Create a nice workflow for editing animation overrides in derived anim blueprints
#branch UE4

diff --git a/Engine/Source/Editor/AnimGraph/Classes/AnimGraphNode_Base.h b/Engine/Source/Editor/AnimGraph/Classes/AnimGraphNode_Base.h
--- a/Engine/Source/Editor/AnimGraph/Classes/AnimGraphNode_Base.h
+++ b/Engine/Source/Editor/AnimGraph/Classes/AnimGraphNode_Base.h
@@ -183,10 +183,10 @@ protected:
 	virtual ANIMGRAPH_API void CreatePinsForPoseLink(UProperty* PoseProperty, int32 ArrayIndex);
 
 	//
-	virtual ANIMGRAPH_API FPoseLinkMappingRecord GetLinkIDLocation(UScriptStruct* NodeType, UEdGraphPin* InputLinkPin);
+	virtual ANIMGRAPH_API FPoseLinkMappingRecord GetLinkIDLocation(const UScriptStruct* NodeType, UEdGraphPin* InputLinkPin);
 
 	//
-	virtual ANIMGRAPH_API void GetPinAssociatedProperty(UScriptStruct* NodeType, UEdGraphPin* InputPin, UProperty*& OutProperty, int32& OutIndex);
+	virtual ANIMGRAPH_API void GetPinAssociatedProperty(const UScriptStruct* NodeType, UEdGraphPin* InputPin, UProperty*& OutProperty, int32& OutIndex);
 
 	// Creates a default menu entry
 	ANIMGRAPH_API TSharedPtr<FEdGraphSchemaAction_K2NewNode> CreateDefaultMenuEntry(FGraphContextMenuBuilder& ContextMenuBuilder) const;
diff --git a/Engine/Source/Editor/AnimGraph/Private/AnimGraphNode_Base.cpp b/Engine/Source/Editor/AnimGraph/Private/AnimGraphNode_Base.cpp
--- a/Engine/Source/Editor/AnimGraph/Private/AnimGraphNode_Base.cpp
+++ b/Engine/Source/Editor/AnimGraph/Private/AnimGraphNode_Base.cpp
@@ -143,7 +143,7 @@ void UAnimGraphNode_Base::InternalPinCreation(TArray<UEdGraphPin*>* OldPins)
 	PreloadRequiredAssets();
 
 	const UAnimationGraphSchema* Schema = GetDefault<UAnimationGraphSchema>();
-	if (UStructProperty* NodeStruct = GetFNodeProperty())
+	if (const UStructProperty* NodeStruct = GetFNodeProperty())
 	{
 		// Display any currently visible optional pins
 		{
@@ -233,7 +233,7 @@ TSharedPtr<FEdGraphSchemaAction_K2NewNode> UAnimGraphNode_Base::CreateDefaultMen
 	return NodeAction;
 }
 
-void UAnimGraphNode_Base::GetPinAssociatedProperty(UScriptStruct* NodeType, UEdGraphPin* InputPin, UProperty*& OutProperty, int32& OutIndex)
+void UAnimGraphNode_Base::GetPinAssociatedProperty(const UScriptStruct* NodeType, UEdGraphPin* InputPin, UProperty*& OutProperty, int32& OutIndex)
 {
 	OutProperty = NULL;
 	OutIndex = INDEX_NONE;
@@ -261,7 +261,7 @@ void UAnimGraphNode_Base::GetPinAssociatedProperty(UScriptStruct* NodeType, UEdG
 	}
 }
 
-FPoseLinkMappingRecord UAnimGraphNode_Base::GetLinkIDLocation(UScriptStruct* NodeType, UEdGraphPin* SourcePin)
+FPoseLinkMappingRecord UAnimGraphNode_Base::GetLinkIDLocation(const UScriptStruct* NodeType, UEdGraphPin* SourcePin)
 {
 	if (SourcePin->LinkedTo.Num() > 0)
 	{
diff --git a/Engine/Source/Editor/EditorStyle/Private/SlateEditorStyle.cpp b/Engine/Source/Editor/EditorStyle/Private/SlateEditorStyle.cpp
--- a/Engine/Source/Editor/EditorStyle/Private/SlateEditorStyle.cpp
+++ b/Engine/Source/Editor/EditorStyle/Private/SlateEditorStyle.cpp
@@ -2624,6 +2624,7 @@ void FSlateEditorStyle::FStyle::SetupGraphEditorStyles()
 		Set( "Graph.ExecutionBubble", new IMAGE_BRUSH( "Old/Graph/ExecutionBubble", Icon16x16 ) );
 
 		Set( "Graph.PlayInEditor", new BOX_BRUSH( "/Graph/RegularNode_shadow_selected", FMargin(18.0f/64.0f) ) );
+		Set( "Graph.ReadOnlyBorder", new BOX_BRUSH( "/Graph/Graph_readonly_border", FMargin(18.0f / 64.0f) ) );
 
 		Set( "Graph.Panel.SolidBackground", new IMAGE_BRUSH( "/Graph/GraphPanel_SolidBackground", FVector2D(16, 16), FLinearColor::White, ESlateBrushTileType::Both) );
 		Set( "Graph.Panel.GridLineColor",   FLinearColor(0.035f, 0.035f, 0.035f) );
diff --git a/Engine/Source/Editor/GraphEditor/Private/SGraphEditorImpl.cpp b/Engine/Source/Editor/GraphEditor/Private/SGraphEditorImpl.cpp
--- a/Engine/Source/Editor/GraphEditor/Private/SGraphEditorImpl.cpp
+++ b/Engine/Source/Editor/GraphEditor/Private/SGraphEditorImpl.cpp
@@ -198,6 +198,12 @@ void SGraphEditorImpl::Construct( const FArguments& InArgs )
 			FString OverrideText = Appearance.Get().PIENotifyText;
 			return OverrideText.Len() ? OverrideText : DefaultText;
 		}
+
+		static FString GetReadOnlyText(TAttribute<FGraphAppearanceInfo>Appearance, FString DefaultText)
+		{
+			FString OverrideText = Appearance.Get().ReadOnlyText;
+			return OverrideText.Len() ? OverrideText : DefaultText;
+		}
 	};
 	
 	FString DefaultPIENotify(TEXT("SIMULATING"));
@@ -205,6 +211,11 @@ void SGraphEditorImpl::Construct( const FArguments& InArgs )
 		TAttribute<FString>::Create( TAttribute<FString>::FGetter::CreateStatic(&Local::GetPIENotifyText, Appearance, DefaultPIENotify) ) :
 		TAttribute<FString>(DefaultPIENotify);
 
+	FString DefaultReadOnlyText(LOCTEXT("GraphReadOnlyText", "READ-ONLY").ToString());
+	TAttribute<FString> ReadOnlyText = Appearance.IsBound() ?
+		TAttribute<FString>::Create(TAttribute<FString>::FGetter::CreateStatic(&Local::GetReadOnlyText, Appearance, DefaultReadOnlyText)) :
+		TAttribute<FString>(DefaultReadOnlyText);
+
 	TSharedPtr<SOverlay> OverlayWidget;
 
 	this->ChildSlot
@@ -221,7 +232,7 @@ void SGraphEditorImpl::Construct( const FArguments& InArgs )
 			.OnGetContextMenuFor( this, &SGraphEditorImpl::GraphEd_OnGetContextMenuFor )
 			.OnSelectionChanged( InArgs._GraphEvents.OnSelectionChanged )
 			.OnNodeDoubleClicked( InArgs._GraphEvents.OnNodeDoubleClicked )
-			.IsEditable( this->IsEditable )
+			.IsEditable( this, &SGraphEditorImpl::IsGraphEditable )
 			.OnDropActor( InArgs._GraphEvents.OnDropActor )
 			.OnDropStreamingLevel( InArgs._GraphEvents.OnDropStreamingLevel )
 			.IsEnabled(this, &SGraphEditorImpl::GraphEd_OnGetGraphEnabled)
@@ -276,6 +287,18 @@ void SGraphEditorImpl::Construct( const FArguments& InArgs )
 			.Text( PIENotifyText )
 		]
 
+		// Top-right corner text indicating read only when not simulating
+		+ SOverlay::Slot()
+		.Padding(20)
+		.VAlign(VAlign_Top)
+		.HAlign(HAlign_Right)
+		[
+			SNew(STextBlock)
+			.Visibility(this, &SGraphEditorImpl::ReadOnlyVisibility)
+			.TextStyle(FEditorStyle::Get(), "Graph.CornerText")
+			.Text(ReadOnlyText)
+		]
+
 		// Bottom-right corner text for notification list position
 		+SOverlay::Slot()
 		.Padding(15.f)
@@ -409,7 +432,7 @@ FActionMenuContent SGraphEditorImpl::GraphEd_OnGetContextMenuFor( const FVector2
 		}
 		else
 		{
-			return FActionMenuContent( SNew(STextBlock)  .Text( NSLOCTEXT("GraphEditor", "CannotCreateWhileDebugging", "Cannot create new nodes while debugging").ToString() ) );
+			return FActionMenuContent( SNew(STextBlock)  .Text( NSLOCTEXT("GraphEditor", "CannotCreateWhileDebugging", "Cannot create new nodes in a read only graph").ToString() ) );
 		}
 	}
 	else
@@ -483,7 +506,7 @@ FSlateColor SGraphEditorImpl::GetZoomTextColorAndOpacity() const
 
 bool SGraphEditorImpl::IsGraphEditable() const
 {
-	return (EdGraphObj != NULL) && (EdGraphObj->bEditable) && IsEditable.Get();
+	return (EdGraphObj != NULL) && IsEditable.Get();
 }
 
 TSharedPtr<SWidget> SGraphEditorImpl::GetTitleBar() const
@@ -551,6 +574,15 @@ void SGraphEditorImpl::SetPinVisibility( SGraphEditor::EPinVisibility Visibility
 	}
 }
 
+EVisibility SGraphEditorImpl::ReadOnlyVisibility() const
+{
+	if(PIENotification() == EVisibility::Hidden && !IsEditable.Get())
+	{
+		return EVisibility::Visible;
+	}
+	return EVisibility::Hidden;
+}
+
 
 
 /////////////////////////////////////////////////////
diff --git a/Engine/Source/Editor/GraphEditor/Private/SGraphEditorImpl.h b/Engine/Source/Editor/GraphEditor/Private/SGraphEditorImpl.h
--- a/Engine/Source/Editor/GraphEditor/Private/SGraphEditorImpl.h
+++ b/Engine/Source/Editor/GraphEditor/Private/SGraphEditorImpl.h
@@ -72,6 +72,9 @@ private:
 	/** Function to check whether PIE is active to display "Simulating" text in graph panel*/
 	EVisibility PIENotification( ) const;
 
+	/** Function to check whether we should show read-only text in the panel */
+	EVisibility ReadOnlyVisibility() const;
+
 	/** Notification list to pass messages to editor users  */
 	TSharedPtr<SNotificationList> NotificationListPtr;
 
diff --git a/Engine/Source/Editor/GraphEditor/Private/SGraphPanel.cpp b/Engine/Source/Editor/GraphEditor/Private/SGraphPanel.cpp
--- a/Engine/Source/Editor/GraphEditor/Private/SGraphPanel.cpp
+++ b/Engine/Source/Editor/GraphEditor/Private/SGraphPanel.cpp
@@ -459,14 +459,26 @@ int32 SGraphPanel::OnPaint( const FGeometry& AllottedGeometry, const FSlateRect&
 	++MaxLayerId;
 	PaintSurroundSunkenShadow(FEditorStyle::GetBrush(TEXT("Graph.Shadow")), AllottedGeometry, MyClippingRect, OutDrawElements, MaxLayerId);
 
-	// Draw a surrounding indicator when PIE is active, to make it clear that the graph is read-only, etc...
+	const FSlateBrush* BorderBrush = nullptr;
 	if (bShowPIENotification && (GEditor->bIsSimulatingInEditor || GEditor->PlayWorld != NULL))
 	{
+		// Draw a surrounding indicator when PIE is active, to make it clear that the graph is read-only, etc...
+		BorderBrush = FEditorStyle::GetBrush(TEXT("Graph.PlayInEditor"));
+	}
+	else if (!IsEditable.Get())
+	{
+		// Draw a different border when we're not simulating but the graph is read-only
+		BorderBrush = FEditorStyle::GetBrush(TEXT("Graph.ReadOnlyBorder"));
+	}
+
+	if (BorderBrush)
+	{
+		// Actually draw the border
 		FSlateDrawElement::MakeBox(
 			OutDrawElements,
 			MaxLayerId,
 			AllottedGeometry.ToPaintGeometry(),
-			FEditorStyle::GetBrush(TEXT("Graph.PlayInEditor")),
+			BorderBrush,
 			MyClippingRect
 			);
 	}
diff --git a/Engine/Source/Editor/Kismet/Private/BlueprintEditor.cpp b/Engine/Source/Editor/Kismet/Private/BlueprintEditor.cpp
--- a/Engine/Source/Editor/Kismet/Private/BlueprintEditor.cpp
+++ b/Engine/Source/Editor/Kismet/Private/BlueprintEditor.cpp
@@ -669,7 +669,7 @@ TSharedRef<SGraphEditor> FBlueprintEditor::CreateGraphEditorWidget(TSharedRef<FT
 
 	TSharedRef<SGraphEditor> Editor = SNew(SGraphEditor)
 		.AdditionalCommands(GraphEditorCommands)
-		.IsEditable(this, &FBlueprintEditor::InEditingMode)
+		.IsEditable(this, &FBlueprintEditor::IsEditable, InGraph)
 		.TitleBar(TitleBarWidget)
 		.TitleBarEnabledOnly(bIsInterface || bIsDelegate)
 		.Appearance(this, &FBlueprintEditor::GetGraphAppearance)
@@ -3887,7 +3887,8 @@ void FBlueprintEditor::DeleteSelectedNodes()
 bool FBlueprintEditor::CanDeleteNodes() const
 {
 	const FGraphPanelSelectionSet SelectedNodes = GetSelectedNodes();
-	return InEditingMode() && (SelectedNodes.Num() > 0);
+
+	return IsEditable(GetFocusedGraph()) && (SelectedNodes.Num() > 0);
 }
 
 void FBlueprintEditor::DeleteSelectedDuplicatableNodes()
@@ -4282,7 +4283,7 @@ bool FBlueprintEditor::CanPasteNodes() const
 	FString ClipboardContent;
 	FPlatformMisc::ClipboardPaste(ClipboardContent);
 
-	return InEditingMode() && FEdGraphUtilities::CanImportNodesFromText(FocusedGraphEd->GetCurrentGraph(), ClipboardContent);
+	return IsEditable(GetFocusedGraph()) && FEdGraphUtilities::CanImportNodesFromText(FocusedGraphEd->GetCurrentGraph(), ClipboardContent);
 }
 
 void FBlueprintEditor::DuplicateNodes()
@@ -4294,7 +4295,7 @@ void FBlueprintEditor::DuplicateNodes()
 
 bool FBlueprintEditor::CanDuplicateNodes() const
 {
-	return CanCopyNodes() && InEditingMode();
+	return CanCopyNodes() && IsEditable(GetFocusedGraph());
 }
 
 void FBlueprintEditor::OnAssignReferencedActor()
@@ -5697,7 +5698,7 @@ void FBlueprintEditor::OnRenameNode()
 
 bool FBlueprintEditor::CanRenameNodes() const
 {
-	if(InEditingMode())
+	if (IsEditable(GetFocusedGraph()))
 	{
 		const FGraphPanelSelectionSet SelectedNodes = GetSelectedNodes();
 		bool bCanRenameNodes = (SelectedNodes.Num() == 1) ? true : false;
@@ -6182,6 +6183,31 @@ UEdGraph* FBlueprintEditor::GetFocusedGraph() const
 	return nullptr;
 }
 
+bool FBlueprintEditor::IsEditable(UEdGraph* InGraph) const
+{
+	return InEditingMode() && InGraph && InGraph->bEditable;
+}
+
+FString FBlueprintEditor::GetGraphDecorationString(UEdGraph* InGraph) const
+{
+	return TEXT("");
+}
+
+bool FBlueprintEditor::IsGraphInCurrentBlueprint(UEdGraph* InGraph) const
+{
+	bool bEditable = true;
+
+	UBlueprint* EditingBP = GetBlueprintObj();
+	if(EditingBP)
+	{
+		TArray<UEdGraph*> Graphs;
+		EditingBP->GetAllGraphs(Graphs);
+		bEditable &= Graphs.Contains(InGraph);
+	}
+
+	return bEditable;
+}
+
 /////////////////////////////////////////////////////
 /////////////////////////////////////////////////////
 /////////////////////////////////////////////////////
diff --git a/Engine/Source/Editor/Kismet/Private/BlueprintEditorTabFactories.h b/Engine/Source/Editor/Kismet/Private/BlueprintEditorTabFactories.h
--- a/Engine/Source/Editor/Kismet/Private/BlueprintEditorTabFactories.h
+++ b/Engine/Source/Editor/Kismet/Private/BlueprintEditorTabFactories.h
@@ -202,7 +202,12 @@ public:
 		GraphEditor->GetViewLocation(ViewLocation, ZoomAmount);
 
 		UEdGraph* Graph = FTabPayload_UObject::CastChecked<UEdGraph>(Payload);
-		BlueprintEditorPtr.Pin()->GetBlueprintObj()->LastEditedDocuments.Add(FEditedDocumentInfo(Graph, ViewLocation, ZoomAmount));
+
+		if(BlueprintEditorPtr.Pin()->IsGraphInCurrentBlueprint(Graph))
+		{
+			// Don't save references to external graphs.
+			BlueprintEditorPtr.Pin()->GetBlueprintObj()->LastEditedDocuments.Add(FEditedDocumentInfo(Graph, ViewLocation, ZoomAmount));
+		}
 	}
 
 //	virtual void OnTabClosed(TSharedRef<SDockTab> Tab, TSharedPtr<FTabPayload> Payload) const OVERRIDE
diff --git a/Engine/Source/Editor/Kismet/Private/SGraphTitleBar.cpp b/Engine/Source/Editor/Kismet/Private/SGraphTitleBar.cpp
--- a/Engine/Source/Editor/Kismet/Private/SGraphTitleBar.cpp
+++ b/Engine/Source/Editor/Kismet/Private/SGraphTitleBar.cpp
@@ -43,10 +43,16 @@ FString SGraphTitleBar::GetTitleExtra() const
 {
 	check(EdGraphObj != NULL);
 
-	const bool bEditable = (EdGraphObj->bEditable) && (Kismet2Ptr.Pin()->InEditingMode());
-	return (bEditable)
-		? TEXT("")
-		: TEXT(" (READ-ONLY)");
+	FString ExtraText = Kismet2Ptr.Pin()->GetGraphDecorationString(EdGraphObj);
+
+	const bool bEditable = (Kismet2Ptr.Pin()->IsEditable(EdGraphObj));
+
+	if(!bEditable)
+	{
+		ExtraText += TEXT(" (READ-ONLY)");
+	}
+
+	return ExtraText;
 }
 
 
diff --git a/Engine/Source/Editor/Kismet/Public/BlueprintEditor.h b/Engine/Source/Editor/Kismet/Public/BlueprintEditor.h
--- a/Engine/Source/Editor/Kismet/Public/BlueprintEditor.h
+++ b/Engine/Source/Editor/Kismet/Public/BlueprintEditor.h
@@ -226,6 +226,9 @@ public:
 	// This is used to delay it by one frame when triggered by a tab being closed, so it can finish closing before remembering the new state
 	void RequestSaveEditedObjectState();
 
+	/** Returns whether a graph is editable or not */
+	virtual bool IsEditable(UEdGraph* InGraph) const;
+
 	/** Returns true if in editing mode */
 	bool InEditingMode() const;
 
@@ -454,6 +457,12 @@ public:
 	/** Called when a node's title is committed for a rename */
 	void OnNodeTitleCommitted(const FText& NewText, ETextCommit::Type CommitInfo, UEdGraphNode* NodeBeingChanged);
 
+	/** Called by a graph title bar to get any extra information the editor would like to display */
+	virtual FString GetGraphDecorationString(UEdGraph* InGraph) const;
+
+	/** Checks to see if the provided graph is contained within the current blueprint */
+	bool IsGraphInCurrentBlueprint(UEdGraph* InGraph) const;
+
 protected:
 	
 	// Zooming to fit the entire graph
diff --git a/Engine/Source/Editor/KismetCompiler/Private/AnimBlueprintCompiler.cpp b/Engine/Source/Editor/KismetCompiler/Private/AnimBlueprintCompiler.cpp
--- a/Engine/Source/Editor/KismetCompiler/Private/AnimBlueprintCompiler.cpp
+++ b/Engine/Source/Editor/KismetCompiler/Private/AnimBlueprintCompiler.cpp
@@ -230,7 +230,7 @@ void FAnimBlueprintCompiler::ProcessAnimationNode(UAnimGraphNode_Base* VisualAni
 	}
 
 	// Make sure the visual node has a runtime node template
-	UScriptStruct* NodeType = VisualAnimNode->GetFNodeType();
+	const UScriptStruct* NodeType = VisualAnimNode->GetFNodeType();
 	if (NodeType == NULL)
 	{
 		MessageLog.Error(TEXT("@@ has no animation node member"), VisualAnimNode);
@@ -274,6 +274,7 @@ void FAnimBlueprintCompiler::ProcessAnimationNode(UAnimGraphNode_Base* VisualAni
 
 	// Register the slightly more permanent debug information
 	NewAnimBlueprintClass->GetAnimBlueprintDebugData().NodePropertyToIndexMap.Add(TrueSourceObject, AllocatedIndex);
+	NewAnimBlueprintClass->GetAnimBlueprintDebugData().NodeGuidToIndexMap.Add(TrueSourceObject->NodeGuid, AllocatedIndex);
 	NewAnimBlueprintClass->GetDebugData().RegisterClassPropertyAssociation(TrueSourceObject, NewProperty);
 
 	// Node-specific compilation that requires compiler state info
@@ -924,7 +925,46 @@ void FAnimBlueprintCompiler::CopyTermDefaultsToDefaultObject(UObject* DefaultObj
 
 	if (bIsDerivedAnimBlueprint)
 	{
-		// Skip any work related to an anim graph, it's all done by the parent class
+		// If we are a derived animation graph; apply any stored overrides.
+		// Restore values from the root blueprint to catch values where the override has been removed.
+		UAnimBlueprint* RootAnimBP = UAnimBlueprint::FindRootAnimBlueprint(AnimBlueprint);
+		UAnimBlueprintGeneratedClass* RootAnimClass = RootAnimBP->GetAnimBlueprintGeneratedClass();
+		UObject* RootDefaultObject = RootAnimClass->GetDefaultObject();
+
+		for (TFieldIterator<UProperty> It(RootAnimClass) ; It; ++It)
+		{
+			UProperty* RootProp = *It;
+
+			if (UStructProperty* RootStructProp = Cast<UStructProperty>(RootProp))
+			{
+				if (RootStructProp->Struct->IsChildOf(FAnimNode_Base::StaticStruct()))
+				{
+					UStructProperty* ChildStructProp = FindField<UStructProperty>(NewAnimBlueprintClass, *RootStructProp->GetName());
+					check(ChildStructProp);
+					uint8* SourcePtr = RootStructProp->ContainerPtrToValuePtr<uint8>(RootDefaultObject);
+					uint8* DestPtr = ChildStructProp->ContainerPtrToValuePtr<uint8>(DefaultObject);
+					check(SourcePtr && DestPtr);
+					RootStructProp->CopyCompleteValue(DestPtr, SourcePtr);
+				}
+			}
+		}
+
+		// Patch the overridden values into the CDO
+
+		TArray<FAnimParentNodeAssetOverride*> AssetOverrides;
+		AnimBlueprint->GetAssetOverrides(AssetOverrides);
+		for (FAnimParentNodeAssetOverride* Override : AssetOverrides)
+		{
+			if (Override->NewAsset)
+			{
+				FAnimNode_Base* BaseNode = NewAnimBlueprintClass->GetPropertyInstance<FAnimNode_Base>(DefaultObject, Override->ParentNodeGuid, EPropertySearchMode::Hierarchy);
+				if (BaseNode)
+				{
+					BaseNode->OverrideAsset(Override->NewAsset);
+				}
+			}
+		}
+
 		return;
 	}
 
@@ -939,7 +979,7 @@ void FAnimBlueprintCompiler::CopyTermDefaultsToDefaultObject(UObject* DefaultObj
 
 		if (UAnimGraphNode_Base* VisualAnimNode = AllocatedNodePropertiesToNodes.FindRef(TargetProperty))
 		{
-			UStructProperty* SourceNodeProperty = VisualAnimNode->GetFNodeProperty();
+			const UStructProperty* SourceNodeProperty = VisualAnimNode->GetFNodeProperty();
 			check(SourceNodeProperty != NULL);
 			check(CastChecked<UStructProperty>(TargetProperty)->Struct == SourceNodeProperty->Struct);
 
diff --git a/Engine/Source/Editor/Persona/Private/AnimGraphNodeDetails.cpp b/Engine/Source/Editor/Persona/Private/AnimGraphNodeDetails.cpp
--- a/Engine/Source/Editor/Persona/Private/AnimGraphNodeDetails.cpp
+++ b/Engine/Source/Editor/Persona/Private/AnimGraphNodeDetails.cpp
@@ -17,6 +17,8 @@
 #include "IDetailCustomNodeBuilder.h"
 #include "PropertyCustomizationHelpers.h"
 #include "AnimGraphNode_Base.h"
+#include "GraphEditor.h"
+#include "Persona.h"
 
 #define LOCTEXT_NAMESPACE "KismetNodeWithOptionalPinsDetails"
 
@@ -65,7 +67,7 @@ void FAnimGraphNodeDetails::CustomizeDetails(class IDetailLayoutBuilder& DetailB
 	TargetSkeletonName = FString::Printf(TEXT("%s'%s'"), *TargetSkeleton->GetClass()->GetName(), *TargetSkeleton->GetPathName());
 
 	// Get the node property
-	UStructProperty* NodeProperty = AnimGraphNode->GetFNodeProperty();
+	const UStructProperty* NodeProperty = AnimGraphNode->GetFNodeProperty();
 	if (NodeProperty == NULL)
 	{
 		return;
@@ -230,7 +232,7 @@ void FAnimGraphNodeDetails::AbortDisplayOfAllNodes(TArray< TWeakObjectPtr<UObjec
 	{
 		if (UAnimGraphNode_Base* AnimNode = Cast<UAnimGraphNode_Base>(SelectionIt->Get()))
 		{
-			if (UStructProperty* NodeProperty = AnimNode->GetFNodeProperty())
+			if (const UStructProperty* NodeProperty = AnimNode->GetFNodeProperty())
 			{
 				DetailBuilder.HideProperty(NodeProperty->GetFName(), AnimNode->GetClass());
 			}
@@ -604,4 +606,307 @@ FString FBoneReferenceCustomization::GetCurrentBoneName() const
 	return OutText;
 }
 
-#undef LOCTEXT_NAMESPACE
\ No newline at end of file
+TSharedRef<IDetailCustomization> FAnimGraphParentPlayerDetails::MakeInstance(TWeakPtr<FPersona> InPersona)
+{
+	return MakeShareable(new FAnimGraphParentPlayerDetails(InPersona));
+}
+
+void FAnimGraphParentPlayerDetails::CustomizeDetails(class IDetailLayoutBuilder& DetailBuilder)
+{
+	TArray<TWeakObjectPtr<UObject>> SelectedObjects;
+	DetailBuilder.GetObjectsBeingCustomized(SelectedObjects);
+	check(SelectedObjects.Num() == 1);
+
+	EditorObject = Cast<UEditorParentPlayerListObj>(SelectedObjects[0].Get());
+	check(EditorObject);
+	
+	IDetailCategoryBuilder& Category = DetailBuilder.EditCategory("AnimGraphOverrides");
+	DetailBuilder.HideProperty("Overrides");
+
+	// Build a hierarchy of entires for a tree view in the form of Blueprint->Graph->Node
+	for(FAnimParentNodeAssetOverride& Override : EditorObject->Overrides)
+	{
+		UAnimGraphNode_Base* Node = EditorObject->GetVisualNodeFromGuid(Override.ParentNodeGuid);
+		TSharedPtr<FPlayerTreeViewEntry> NodeEntry = MakeShareable(new FPlayerTreeViewEntry(Node->GetNodeTitle(ENodeTitleType::ListView).ToString(), EPlayerTreeViewEntryType::Node, &Override));
+		
+		// Process blueprint entry
+		TSharedPtr<FPlayerTreeViewEntry>* ExistingBPEntry = ListEntries.FindByPredicate([Node](const TSharedPtr<FPlayerTreeViewEntry>& Other)
+		{
+			return Node->GetBlueprint()->GetName() == Other->EntryName;
+		});
+
+		if(!ExistingBPEntry)
+		{
+			ListEntries.Add(MakeShareable(new FPlayerTreeViewEntry(Node->GetBlueprint()->GetName(), EPlayerTreeViewEntryType::Blueprint)));
+			ExistingBPEntry = &ListEntries.Last();
+		}
+
+		// Process graph entry
+		TSharedPtr<FPlayerTreeViewEntry>* ExistingGraphEntry = (*ExistingBPEntry)->Children.FindByPredicate([Node](const TSharedPtr<FPlayerTreeViewEntry>& Other)
+		{
+			return Node->GetGraph()->GetName() == Other->EntryName;
+		});
+
+		if(!ExistingGraphEntry)
+		{
+			(*ExistingBPEntry)->Children.Add(MakeShareable(new FPlayerTreeViewEntry(Node->GetGraph()->GetName(), EPlayerTreeViewEntryType::Graph)));
+			ExistingGraphEntry = &(*ExistingBPEntry)->Children.Last();
+		}
+
+		// Process node entry
+		(*ExistingGraphEntry)->Children.Add(NodeEntry);
+	}
+
+	FDetailWidgetRow& Row = Category.AddCustomRow("");
+	TSharedRef<STreeView<TSharedPtr<FPlayerTreeViewEntry>>> TreeView = SNew(STreeView<TSharedPtr<FPlayerTreeViewEntry>>)
+		.SelectionMode(ESelectionMode::None)
+		.OnGenerateRow(this, &FAnimGraphParentPlayerDetails::OnGenerateRow)
+		.OnGetChildren(this, &FAnimGraphParentPlayerDetails::OnGetChildren)
+		.TreeItemsSource(&ListEntries)
+		.HeaderRow
+		(
+			SNew(SHeaderRow)
+			+SHeaderRow::Column(FName("Name"))
+			.FillWidth(0.5f)
+			.DefaultLabel(LOCTEXT("ParentPlayer_NameCol", "Name").ToString())
+
+			+SHeaderRow::Column(FName("Asset"))
+			.FillWidth(0.5f)
+			.DefaultLabel(LOCTEXT("ParentPlayer_AssetCol", "Asset").ToString())
+		);
+
+	// Expand top level (blueprint) entries so the panel seems less empty
+	for(TSharedPtr<FPlayerTreeViewEntry> Entry : ListEntries)
+	{
+		TreeView->SetItemExpansion(Entry, true);
+	}
+
+	Row
+	[
+		TreeView->AsShared()
+	];
+}
+
+TSharedRef<ITableRow> FAnimGraphParentPlayerDetails::OnGenerateRow(TSharedPtr<FPlayerTreeViewEntry> EntryPtr, const TSharedRef< STableViewBase >& OwnerTable)
+{
+	return SNew(SParentPlayerTreeRow, OwnerTable).Item(EntryPtr).OverrideObject(EditorObject).Persona(PersonaPtr);
+}
+
+void FAnimGraphParentPlayerDetails::OnGetChildren(TSharedPtr<FPlayerTreeViewEntry> InParent, TArray< TSharedPtr<FPlayerTreeViewEntry> >& OutChildren)
+{
+	OutChildren.Append(InParent->Children);
+}
+
+void SParentPlayerTreeRow::Construct(const FArguments& InArgs, const TSharedRef<STableViewBase>& InOwnerTableView)
+{
+	Item = InArgs._Item;
+	EditorObject = InArgs._OverrideObject;
+	Persona = InArgs._Persona;
+
+	if(Item->Override)
+	{
+		GraphNode = EditorObject->GetVisualNodeFromGuid(Item->Override->ParentNodeGuid);
+	}
+	else
+	{
+		GraphNode = NULL;
+	}
+
+	SMultiColumnTableRow<TSharedPtr<FAnimGraphParentPlayerDetails>>::Construct(FSuperRowType::FArguments(), InOwnerTableView);
+}
+
+TSharedRef<SWidget> SParentPlayerTreeRow::GenerateWidgetForColumn(const FName& ColumnName)
+{
+	TSharedPtr<SHorizontalBox> Box;
+	SAssignNew(Box, SHorizontalBox);
+
+	if(ColumnName == "Name")
+	{
+		Box->AddSlot()
+			.VAlign(VAlign_Center)
+			.AutoWidth()
+			[
+				SNew(SExpanderArrow, SharedThis(this))
+			];
+
+		Item->GenerateNameWidget(Box);
+	}
+	else if(Item->Override)
+	{
+		Box->AddSlot()
+			.Padding(2)
+			.VAlign(VAlign_Center)
+			.AutoWidth()
+			[
+				SNew(SButton)
+				.ButtonStyle(FEditorStyle::Get(), "ToggleButton")
+				.ToolTip(IDocumentation::Get()->CreateToolTip(LOCTEXT("FocusNodeButtonTip", "Open the graph that contains this node in read-only mode and focus on the node"), NULL, "Shared/Editors/Persona", "FocusNodeButton"))
+				.OnClicked(FOnClicked::CreateSP(this, &SParentPlayerTreeRow::OnFocusNodeButtonClicked))
+				.Content()
+				[
+					SNew(SImage)
+					.Image(FEditorStyle::GetBrush("GenericViewButton"))
+				]
+				
+			];
+		
+		TArray<const UClass*> AllowedClasses;
+		AllowedClasses.Add(UAnimationAsset::StaticClass());
+		Box->AddSlot()
+			.VAlign(VAlign_Center)
+			.AutoWidth()
+			[
+				SNew(SObjectPropertyEntryBox)
+				.ObjectPath(this, &SParentPlayerTreeRow::GetCurrentAssetPath)
+				.OnShouldFilterAsset(this, &SParentPlayerTreeRow::OnShouldFilterAsset)
+				.OnObjectChanged(this, &SParentPlayerTreeRow::OnAssetSelected)
+				.AllowedClass(GetCurrentAssetToUse()->GetClass())
+			];
+
+		Box->AddSlot()
+			.VAlign(VAlign_Center)
+			.AutoWidth()
+			[
+				SNew(SButton)
+				.ButtonStyle(FEditorStyle::Get(), "NoBorder")
+				.Visibility(this, &SParentPlayerTreeRow::GetResetToDefaultVisibility)
+				.OnClicked(this, &SParentPlayerTreeRow::OnResetButtonClicked)
+				.ToolTip(IDocumentation::Get()->CreateToolTip(LOCTEXT("ResetToParentButtonTip", "Undo the override, returning to the default asset for this node"), NULL, "Shared/Editors/Persona", "ResetToParentButton"))
+				.Content()
+				[
+					SNew(SImage)
+					.Image(FEditorStyle::GetBrush("PropertyWindow.DiffersFromDefault"))
+				]
+			];
+	}
+
+	return Box.ToSharedRef();
+}
+
+bool SParentPlayerTreeRow::OnShouldFilterAsset(const FAssetData& AssetData)
+{
+	const FString* SkeletonName = AssetData.TagsAndValues.Find(TEXT("Skeleton"));
+
+	if(SkeletonName)
+	{
+		USkeleton* CurrentSkeleton = GraphNode->GetAnimBlueprint()->TargetSkeleton;
+		if(*SkeletonName == FString::Printf(TEXT("%s'%s'"), *CurrentSkeleton->GetClass()->GetName(), *CurrentSkeleton->GetPathName()))
+		{
+			return false;
+		}
+	}
+
+	return true;
+}
+
+void SParentPlayerTreeRow::OnAssetSelected(const UObject* Obj)
+{
+	Item->Override->NewAsset = Cast<UAnimationAsset>(const_cast<UObject*>(Obj));
+	EditorObject->ApplyOverrideToBlueprint(*Item->Override);
+}
+
+FReply SParentPlayerTreeRow::OnFocusNodeButtonClicked()
+{
+	TSharedPtr<FPersona> SharedPersona = Persona.Pin();
+	if(SharedPersona.IsValid())
+	{
+		if(GraphNode)
+		{
+			UEdGraph* EdGraph = GraphNode->GetGraph();
+			TSharedPtr<SGraphEditor> GraphEditor = SharedPersona->OpenGraphAndBringToFront(EdGraph);
+			GraphEditor->JumpToNode(GraphNode, false);
+		}
+		return FReply::Handled();
+	}
+
+	return FReply::Unhandled();
+}
+
+const UAnimationAsset* SParentPlayerTreeRow::GetCurrentAssetToUse() const
+{
+	if(Item->Override->NewAsset)
+	{
+		return Item->Override->NewAsset;
+	}
+	
+	if(GraphNode)
+	{
+		return GraphNode->GetAnimationAsset();
+	}
+
+	return NULL;
+}
+
+EVisibility SParentPlayerTreeRow::GetResetToDefaultVisibility() const
+{
+	FAnimParentNodeAssetOverride* HierarchyOverride = EditorObject->GetBlueprint()->GetAssetOverrideForNode(Item->Override->ParentNodeGuid, true);
+
+	if(HierarchyOverride)
+	{
+		return Item->Override->NewAsset != HierarchyOverride->NewAsset ? EVisibility::Visible : EVisibility::Hidden;
+	}
+
+	return Item->Override->NewAsset != GraphNode->GetAnimationAsset() ? EVisibility::Visible : EVisibility::Hidden;
+}
+
+FReply SParentPlayerTreeRow::OnResetButtonClicked()
+{
+	FAnimParentNodeAssetOverride* HierarchyOverride = EditorObject->GetBlueprint()->GetAssetOverrideForNode(Item->Override->ParentNodeGuid, true);
+	
+	Item->Override->NewAsset = HierarchyOverride ? HierarchyOverride->NewAsset : GraphNode->GetAnimationAsset();
+
+	// Apply will remove the override from the object
+	EditorObject->ApplyOverrideToBlueprint(*Item->Override);
+	return FReply::Handled();
+}
+
+FString SParentPlayerTreeRow::GetCurrentAssetPath() const
+{
+	const UAnimationAsset* Asset = GetCurrentAssetToUse();
+	return Asset ? Asset->GetPathName() : FString("");
+}
+
+FORCENOINLINE bool FPlayerTreeViewEntry::operator==(const FPlayerTreeViewEntry& Other)
+{
+	return EntryName == Other.EntryName;
+}
+
+void FPlayerTreeViewEntry::GenerateNameWidget(TSharedPtr<SHorizontalBox> Box)
+{
+	// Get an appropriate image icon for the row
+	const FSlateBrush* EntryImageBrush = NULL;
+	switch(EntryType)
+	{
+		case EPlayerTreeViewEntryType::Blueprint:
+			EntryImageBrush = FEditorStyle::GetBrush("ClassIcon.Blueprint");
+			break;
+		case EPlayerTreeViewEntryType::Graph:
+			EntryImageBrush = FEditorStyle::GetBrush("GraphEditor.EventGraph_16x");
+			break;
+		case EPlayerTreeViewEntryType::Node:
+			EntryImageBrush = FEditorStyle::GetBrush("GraphEditor.Default_16x");
+			break;
+		default:
+			break;
+	}
+
+	Box->AddSlot()
+		.VAlign(VAlign_Center)
+		.AutoWidth()
+		[
+			SNew(SImage)
+			.Image(EntryImageBrush)
+		];
+
+	Box->AddSlot()
+		.VAlign(VAlign_Center)
+		.Padding(FMargin(5.0f, 0.0f, 0.0f, 0.0f))
+		.AutoWidth()
+		[
+			SNew(STextBlock)
+			.Font(FSlateFontInfo(FPaths::EngineContentDir() / TEXT("Slate/Fonts/Roboto-Bold.ttf"), 10))
+			.Text(EntryName)
+		];
+}
+
+#undef LOCTEXT_NAMESPACE
diff --git a/Engine/Source/Editor/Persona/Private/AnimGraphNodeDetails.h b/Engine/Source/Editor/Persona/Private/AnimGraphNodeDetails.h
--- a/Engine/Source/Editor/Persona/Private/AnimGraphNodeDetails.h
+++ b/Engine/Source/Editor/Persona/Private/AnimGraphNodeDetails.h
@@ -2,6 +2,8 @@
 
 #pragma once
 #include "PropertyEditorModule.h"
+#include "AssetData.h"
+#include "PropertyHandle.h"
 
 /////////////////////////////////////////////////////
 // FAnimGraphNodeDetails 
@@ -106,3 +108,120 @@ private:
 	// Property to change after bone has been picked
 	TSharedPtr<IPropertyHandle> BoneRefProperty;
 };
+
+//////////////////////////////////////////////////////////////////////////
+// 
+
+// Type used to identify rows in a parent player tree list
+namespace EPlayerTreeViewEntryType
+{
+	enum Type
+	{
+		Blueprint,
+		Graph,
+		Node
+	};
+}
+
+// Describes a single row entry in a player tree view
+struct FPlayerTreeViewEntry : public TSharedFromThis<FPlayerTreeViewEntry>
+{
+	FPlayerTreeViewEntry(FString Name, EPlayerTreeViewEntryType::Type InEntryType, FAnimParentNodeAssetOverride* InOverride = NULL)
+	: EntryName(Name)
+	, EntryType(InEntryType)
+	, Override(InOverride)
+	{}
+
+	FORCENOINLINE bool operator==(const FPlayerTreeViewEntry& Other);
+
+	void GenerateNameWidget(TSharedPtr<SHorizontalBox> Box);
+
+	// Name for the row
+	FString EntryName;
+
+	// What the row represents 
+	EPlayerTreeViewEntryType::Type EntryType;
+
+	// Node asset override for rows that represent nodes
+	FAnimParentNodeAssetOverride* Override;
+
+	// Children array for rows that represent blueprints and graphs.
+	TArray<TSharedPtr<FPlayerTreeViewEntry>> Children;
+};
+
+class FAnimGraphParentPlayerDetails : public IDetailCustomization
+{
+public:
+	static TSharedRef<IDetailCustomization> MakeInstance(TWeakPtr<FPersona> InPersona);
+
+	virtual void CustomizeDetails(class IDetailLayoutBuilder& DetailBuilder);
+
+private:
+
+	FAnimGraphParentPlayerDetails(TWeakPtr<FPersona> InPersona) : PersonaPtr(InPersona)
+	{}
+
+	TSharedRef<ITableRow> OnGenerateRow(TSharedPtr<FPlayerTreeViewEntry> EventPtr, const TSharedRef< STableViewBase >& OwnerTable);
+	void OnGetChildren(TSharedPtr<FPlayerTreeViewEntry> InParent, TArray< TSharedPtr<FPlayerTreeViewEntry> >& OutChildren);
+
+	// Entries in the tree view
+	TArray<TSharedPtr<FPlayerTreeViewEntry>> ListEntries;
+	
+	// Hosting Persona instance
+	TWeakPtr<FPersona> PersonaPtr;
+
+	// Editor meta-object containing override information
+	UEditorParentPlayerListObj* EditorObject;
+};
+
+class SParentPlayerTreeRow : public SMultiColumnTableRow<TSharedPtr<FAnimGraphParentPlayerDetails>>
+{
+public:
+	SLATE_BEGIN_ARGS(SParentPlayerTreeRow){}
+		SLATE_ARGUMENT(TSharedPtr<FPlayerTreeViewEntry>, Item);
+		SLATE_ARGUMENT(UEditorParentPlayerListObj*, OverrideObject);
+		SLATE_ARGUMENT(TWeakPtr<FPersona>, Persona);
+	SLATE_END_ARGS();
+
+	void Construct(const FArguments& InArgs, const TSharedRef<STableViewBase>& InOwnerTableView);
+
+	virtual TSharedRef<SWidget> GenerateWidgetForColumn(const FName& ColumnName);
+
+private:
+	
+	// Should an asset be filtered, ensures we only approve assets with matching skeletons
+	bool OnShouldFilterAsset(const FAssetData& AssetData);
+
+	// Sets the override asset when selected from the asset picker
+	void OnAssetSelected(const UObject* Obj);
+
+	void OnCloseMenu(){}
+
+	// Called when the user clicks the focus button, opens a graph panel if necessary in
+	// read only mode and focusses on the node.
+	FReply OnFocusNodeButtonClicked();
+
+	// Gets the current asset, either an override if one is selected or the original from the node.
+	const UAnimationAsset* GetCurrentAssetToUse() const;
+
+	// Whether or not we should show the reset to default button next to the asset picker
+	EVisibility GetResetToDefaultVisibility() const;
+
+	// Resets the selected asset override back to the original node's asset
+	FReply OnResetButtonClicked();
+
+	// Gets the full path to the current asset
+	FString GetCurrentAssetPath() const;
+
+	// Editor object containing all possible overrides
+	UEditorParentPlayerListObj* EditorObject;
+
+	// Tree item we are representing
+	TSharedPtr<FPlayerTreeViewEntry> Item;
+
+	// Graphnode this row represents, if any
+	UAnimGraphNode_Base* GraphNode;
+
+	// Persona editor pointer
+	TWeakPtr<FPersona> Persona;
+};
diff --git a/Engine/Source/Editor/Persona/Private/BlueprintMode/AnimBlueprintMode.cpp b/Engine/Source/Editor/Persona/Private/BlueprintMode/AnimBlueprintMode.cpp
--- a/Engine/Source/Editor/Persona/Private/BlueprintMode/AnimBlueprintMode.cpp
+++ b/Engine/Source/Editor/Persona/Private/BlueprintMode/AnimBlueprintMode.cpp
@@ -17,7 +17,7 @@
 FAnimBlueprintEditAppMode::FAnimBlueprintEditAppMode(TSharedPtr<FPersona> InPersona)
 	: FBlueprintEditorApplicationMode(StaticCastSharedPtr<FBlueprintEditor>(InPersona), FPersonaModes::AnimBlueprintEditMode, false, false)
 {
-	TabLayout = FTabManager::NewLayout( "Persona_AnimBlueprintEditMode_Layout_v5" )
+	TabLayout = FTabManager::NewLayout( "Persona_AnimBlueprintEditMode_Layout_v6" )
 		->AddArea
 		(
 			FTabManager::NewPrimaryArea()
@@ -53,8 +53,8 @@ FAnimBlueprintEditAppMode::FAnimBlueprintEditAppMode(TSharedPtr<FPersona> InPers
 						//	Left bottom - preview settings
 						FTabManager::NewStack()
 						->SetSizeCoefficient(0.5f)
-						->SetHideTabWell(true)
 						->AddTab( FPersonaTabs::AnimBlueprintDefaultsEditorID, ETabState::OpenedTab )
+						->AddTab( FPersonaTabs::AnimBlueprintParentPlayerEditorID, ETabState::ClosedTab )
 					)
 				)
 				->Split
@@ -110,6 +110,7 @@ FAnimBlueprintEditAppMode::FAnimBlueprintEditAppMode(TSharedPtr<FPersona> InPers
 	PersonaTabFactories.RegisterFactory(MakeShareable(new FPreviewViewportSummoner(InPersona)));
 	PersonaTabFactories.RegisterFactory(MakeShareable(new FSkeletonAnimNotifiesSummoner(InPersona)));
 	PersonaTabFactories.RegisterFactory(MakeShareable(new FAnimBlueprintDefaultsEditorSummoner(InPersona)));
+	PersonaTabFactories.RegisterFactory(MakeShareable(new FAnimBlueprintParentPlayerEditorSummoner(InPersona)));
 
 	// setup toolbar - clear existing toolbar extender from the BP mode
 	//@TODO: Keep this in sync with BlueprintEditorModes.cpp
diff --git a/Engine/Source/Editor/Persona/Private/EditorObjectsTracker.cpp b/Engine/Source/Editor/Persona/Private/EditorObjectsTracker.cpp
--- a/Engine/Source/Editor/Persona/Private/EditorObjectsTracker.cpp
+++ b/Engine/Source/Editor/Persona/Private/EditorObjectsTracker.cpp
@@ -20,7 +20,7 @@ UObject* FEditorObjectTracker::GetEditorObjectForClass( UClass* EdClass )
 	UObject *Obj = (EditorObjMap.Contains(EdClass) ? *EditorObjMap.Find(EdClass) : NULL);
 	if(Obj == NULL)
 	{
-		FString ObjName = EdClass->GetName();
+		FString ObjName = MakeUniqueObjectName(GetTransientPackage(), EdClass).ToString();
 		ObjName += "_EdObj";
 		Obj = StaticConstructObject(EdClass, GetTransientPackage(), FName(*ObjName), RF_Public|RF_Standalone|RF_Transient);
 		EditorObjMap.Add(EdClass,Obj);
diff --git a/Engine/Source/Editor/Persona/Private/Persona.cpp b/Engine/Source/Editor/Persona/Private/Persona.cpp
--- a/Engine/Source/Editor/Persona/Private/Persona.cpp
+++ b/Engine/Source/Editor/Persona/Private/Persona.cpp
@@ -2513,6 +2513,23 @@ bool FPersona::IsRecordAvailable() const
 	return (GetCurrentMode() == FPersonaModes::AnimBlueprintEditMode);
 }
 
+bool FPersona::IsEditable(UEdGraph* InGraph) const
+{
+	bool bEditable = FBlueprintEditor::IsEditable(InGraph);
+	bEditable &= IsGraphInCurrentBlueprint(InGraph);
+
+	return bEditable;
+}
+
+FString FPersona::GetGraphDecorationString(UEdGraph* InGraph) const
+{
+	if (!IsGraphInCurrentBlueprint(InGraph))
+	{
+		return LOCTEXT("PersonaExternalGraphDecoration", " Parent Graph Preview").ToString();
+	}
+	return TEXT("");
+}
+
 /////////////////////////////////////////////////////
 
 FAnimationRecorder::FAnimationRecorder()
diff --git a/Engine/Source/Editor/Persona/Private/Persona.h b/Engine/Source/Editor/Persona/Private/Persona.h
--- a/Engine/Source/Editor/Persona/Private/Persona.h
+++ b/Engine/Source/Editor/Persona/Private/Persona.h
@@ -235,6 +235,8 @@ protected:
 	virtual void GetCustomDebugObjects(TArray<FCustomDebugObject>& DebugList) const OVERRIDE;
 	virtual void CreateDefaultTabContents(const TArray<UBlueprint*>& InBlueprints) OVERRIDE;
 	virtual FGraphAppearanceInfo GetGraphAppearance() const OVERRIDE;
+	virtual bool IsEditable(UEdGraph* InGraph) const OVERRIDE;
+	virtual FString GetGraphDecorationString(UEdGraph* InGraph) const OVERRIDE;
 	// End of FBlueprintEditor interface
 
 	// IAssetEditorInstance interface
diff --git a/Engine/Source/Editor/Persona/Private/SAnimBlueprintParentPlayerList.cpp b/Engine/Source/Editor/Persona/Private/SAnimBlueprintParentPlayerList.cpp
new file mode 100644
--- /dev/null
+++ b/Engine/Source/Editor/Persona/Private/SAnimBlueprintParentPlayerList.cpp
@@ -0,0 +1,165 @@
+// Copyright 1998-2014 Epic Games, Inc. All Rights Reserved.
+ 
+#include "PersonaPrivatePCH.h"
+#include "SAnimBlueprintParentPlayerList.h"
+#include "PropertyEditorModule.h"
+#include "AnimGraphNodeDetails.h"
+
+SAnimBlueprintParentPlayerList::SAnimBlueprintParentPlayerList()
+	: RootBlueprint(NULL)
+{
+
+}
+
+SAnimBlueprintParentPlayerList::~SAnimBlueprintParentPlayerList()
+{
+	// Unregister delegates
+	if(RootBlueprint)
+	{
+		RootBlueprint->OnChanged().RemoveAll(this);
+	}
+
+	if(CurrentBlueprint)
+	{
+		TArray<UBlueprint*> BlueprintHierarchy;
+		CurrentBlueprint->GetBlueprintHierarchyFromClass(CurrentBlueprint->GetAnimBlueprintGeneratedClass(), BlueprintHierarchy);
+		// Start from 1 as 0 is our current BP
+		for(int32 Idx = 1; Idx < BlueprintHierarchy.Num() ; ++Idx)
+		{
+			if(UAnimBlueprint* AnimBlueprint = Cast<UAnimBlueprint>(BlueprintHierarchy[Idx]))
+			{
+				AnimBlueprint->UnregisterOnOverrideChanged(this);
+			}
+		}
+
+		CurrentBlueprint->OnChanged().RemoveAll(this);
+	}
+
+	TSharedPtr<FPersona> PersonaShared = PersonaPtr.Pin();
+	if(PersonaPtr.IsValid())
+	{
+		PersonaPtr.Pin()->UnregisterOnPostUndo(this);
+	}
+}
+
+void SAnimBlueprintParentPlayerList::Construct(const FArguments& InArgs)
+{
+	PersonaPtr = InArgs._Persona;
+
+	// Register a refresh on post undo to grab the blueprint assets again
+	TSharedPtr<FPersona> PersonaShared = PersonaPtr.Pin();
+	if(PersonaPtr.IsValid())
+	{
+		PersonaPtr.Pin()->RegisterOnPostUndo(FPersona::FOnPostUndo::CreateSP(this, &SAnimBlueprintParentPlayerList::RefreshDetailView));
+	}
+
+	UEditorParentPlayerListObj* EditorObject = Cast<UEditorParentPlayerListObj>(ObjectTracker.GetEditorObjectForClass(UEditorParentPlayerListObj::StaticClass()));
+	CurrentBlueprint = PersonaPtr.Pin()->GetAnimBlueprint();
+	check(CurrentBlueprint);
+	CurrentParentClass = CurrentBlueprint->ParentClass;
+
+	EditorObject->InitialiseFromBlueprint(CurrentBlueprint);
+	
+	if(EditorObject->Overrides.Num() > 0)
+	{
+		FPropertyEditorModule& PropertyModule = FModuleManager::LoadModuleChecked<FPropertyEditorModule>("PropertyEditor");
+		FDetailsViewArgs ViewArgs;
+		ViewArgs.bAllowSearch = false;
+		ViewArgs.bHideSelectionTip = true;
+
+		DetailView = PropertyModule.CreateDetailView(ViewArgs);
+
+		DetailView->RegisterInstancedCustomPropertyLayout(UEditorParentPlayerListObj::StaticClass(), FOnGetDetailCustomizationInstance::CreateStatic(&FAnimGraphParentPlayerDetails::MakeInstance, PersonaPtr));
+
+		this->ChildSlot
+			[
+				DetailView->AsShared()
+			];
+
+		DetailView->SetObject(EditorObject, true);
+
+		// Register delegate to refresh the view when a node changes in the root blueprint animgraph;
+		RootBlueprint = UAnimBlueprint::FindRootAnimBlueprint(CurrentBlueprint);
+		if(RootBlueprint != CurrentBlueprint)
+		{
+			RootBlueprint->OnChanged().AddSP(this, &SAnimBlueprintParentPlayerList::OnRootBlueprintChanged);
+		}
+		else
+		{
+			RootBlueprint = NULL;
+		}
+
+		// Register delegate to handle parents in the hierarchy changing their overrides.
+		TArray<UBlueprint*> BlueprintHierarchy;
+		CurrentBlueprint->GetBlueprintHierarchyFromClass(CurrentBlueprint->GetAnimBlueprintGeneratedClass(), BlueprintHierarchy);
+		// Start from 1 as 0 is our current BP
+		for(int32 Idx = 1; Idx < BlueprintHierarchy.Num() ; ++Idx)
+		{
+			if(UAnimBlueprint* AnimBlueprint = Cast<UAnimBlueprint>(BlueprintHierarchy[Idx]))
+			{
+				AnimBlueprint->RegisterOnOverrideChanged(UAnimBlueprint::FOnOverrideChanged::CreateSP(this, &SAnimBlueprintParentPlayerList::OnHierarchyOverrideChanged));
+			}
+		}
+
+		// Register a delegate for the current BP changing; so we can look for a root change.
+		CurrentBlueprint->OnChanged().AddSP(this, &SAnimBlueprintParentPlayerList::OnCurrentBlueprintChanged);
+	}
+	else
+	{
+		this->ChildSlot
+			[
+				SNew(SBox)
+				.Padding(FMargin(10.0f, 10.0f))
+				.HAlign(HAlign_Center)
+				.VAlign(VAlign_Center)
+				.Content()
+				[
+					SNew(STextBlock)
+					.WrapTextAt(300.0f)
+					.Text(NSLOCTEXT("ParentPlayerList", "NoOverrides", "No possible overrides found. Either there are no nodes referencing assets in the parent class or this is not a derived blueprint."))
+				]
+			];
+	}
+}
+
+void SAnimBlueprintParentPlayerList::OnRootBlueprintChanged(UBlueprint* InBlueprint)
+{
+	RefreshDetailView();
+}
+
+
+void SAnimBlueprintParentPlayerList::OnCurrentBlueprintChanged(UBlueprint* InBlueprint)
+{
+	UAnimBlueprint* NewRoot = UAnimBlueprint::FindRootAnimBlueprint(CurrentBlueprint);
+	if(NewRoot != RootBlueprint)
+	{
+		// The blueprint has been reparented in a way which has changed its root,
+		// The overrides we have stored are no longer valid.
+		RootBlueprint = NewRoot;
+		CurrentParentClass = CurrentBlueprint->ParentClass;
+		CurrentBlueprint->ParentAssetOverrides.Empty();
+		RefreshDetailView();
+	}
+	else if(CurrentBlueprint->ParentClass != CurrentParentClass)
+	{
+		// The blueprint has been reparented to another blueprint with the same root.
+		CurrentParentClass = CurrentBlueprint->ParentClass;
+		RefreshDetailView();
+	}
+}
+
+
+void SAnimBlueprintParentPlayerList::OnHierarchyOverrideChanged(FGuid NodeGuid, UAnimationAsset* NewAsset)
+{
+	RefreshDetailView();
+}
+
+void SAnimBlueprintParentPlayerList::RefreshDetailView()
+{
+	if(DetailView.IsValid())
+	{
+		UEditorParentPlayerListObj* EditorObject = Cast<UEditorParentPlayerListObj>(ObjectTracker.GetEditorObjectForClass(UEditorParentPlayerListObj::StaticClass()));
+		EditorObject->InitialiseFromBlueprint(CurrentBlueprint);
+		DetailView->SetObject(EditorObject, true);
+	}
+}
diff --git a/Engine/Source/Editor/Persona/Private/SAnimBlueprintParentPlayerList.h b/Engine/Source/Editor/Persona/Private/SAnimBlueprintParentPlayerList.h
new file mode 100644
--- /dev/null
+++ b/Engine/Source/Editor/Persona/Private/SAnimBlueprintParentPlayerList.h
@@ -0,0 +1,58 @@
+// Copyright 1998-2014 Epic Games, Inc. All Rights Reserved.
+
+#pragma once
+
+#include "Persona.h"
+#include "IDetailsView.h"
+#include "EditorObjectsTracker.h"
+
+class SAnimBlueprintParentPlayerList : public SCompoundWidget
+{
+public:
+	SLATE_BEGIN_ARGS(SAnimBlueprintParentPlayerList)
+		: _Persona()
+		{}
+	SLATE_ARGUMENT(TWeakPtr<FPersona>, Persona)
+		SLATE_END_ARGS()
+
+		SAnimBlueprintParentPlayerList();
+	~SAnimBlueprintParentPlayerList();
+
+	void Construct(const FArguments& InArgs);
+
+private:
+
+	// Called when the root blueprint is changed. depending on the action
+	// we need to refresh the list of available overrides incase we need
+	// to remove or add some
+	void OnRootBlueprintChanged(UBlueprint* InBlueprint);
+
+	// Called when the current blueprint changes. Used to detect reparenting
+	// So the data and UI can be updated accordingly
+	void OnCurrentBlueprintChanged(UBlueprint* InBlueprint);
+
+	// Called when an override is changed on a less derived blueprint in the
+	// current blueprint's hierarchy so we can copy them if we haven't overridden
+	// the same asset
+	void OnHierarchyOverrideChanged(FGuid NodeGuid, UAnimationAsset* NewAsset);
+
+	void RefreshDetailView();
+
+	// Persona instance we're currently in
+	TWeakPtr<FPersona> PersonaPtr;
+
+	// Object tracker to maintain single instance of editor objects
+	FEditorObjectTracker ObjectTracker;
+
+	// The blueprint currently in use
+	UAnimBlueprint* CurrentBlueprint;
+
+	// Parent class of the current blueprint, used to detect reparenting
+	UClass* CurrentParentClass;
+
+	// Root blueprint if one exists
+	UAnimBlueprint* RootBlueprint;
+
+	// The details view for the list
+	TSharedPtr<IDetailsView> DetailView;
+};
\ No newline at end of file
diff --git a/Engine/Source/Editor/Persona/Private/Shared/PersonaMode.cpp b/Engine/Source/Editor/Persona/Private/Shared/PersonaMode.cpp
--- a/Engine/Source/Editor/Persona/Private/Shared/PersonaMode.cpp
+++ b/Engine/Source/Editor/Persona/Private/Shared/PersonaMode.cpp
@@ -5,6 +5,7 @@
 #include "Persona.h"
 #include "SSkeletonAnimNotifies.h"
 #include "IDocumentation.h"
+#include "SAnimBlueprintParentPlayerList.h"
 
 #define LOCTEXT_NAMESPACE "PersonaModes"
 
@@ -20,6 +21,8 @@ const FName FPersonaTabs::RetargetSourceManagerID("RetargetSourceManager");
 // Explorer
 // Blueprint Defaults
 const FName FPersonaTabs::AnimBlueprintDefaultsEditorID("AnimBlueprintDefaultsEditor");
+
+const FName FPersonaTabs::AnimBlueprintParentPlayerEditorID("AnimBlueprintParentPlayerEditor");
 // Anim Document
 const FName FPersonaTabs::ScrubberID("ScrubberTab");
 
@@ -366,4 +369,26 @@ void FAnimBlueprintDefaultsEditorSummoner::OnCheckedChanged(ESlateCheckBoxState:
 		CurrentMode = Mode;
 	}
 }
+
+//////////////////////////////////////////////////////////////////////////
+// FAnimBlueprintParentPlayerEditorSummoner
+
+FAnimBlueprintParentPlayerEditorSummoner::FAnimBlueprintParentPlayerEditorSummoner(TSharedPtr<class FAssetEditorToolkit> InHostingApp) 
+: FWorkflowTabFactory(FPersonaTabs::AnimBlueprintParentPlayerEditorID, InHostingApp)
+{
+	TabLabel = LOCTEXT("ParentPlayerOverrideEditor", "Asset Override Editor");
+	bIsSingleton = true;
+}
+
+TSharedRef<SWidget> FAnimBlueprintParentPlayerEditorSummoner::CreateTabBody(const FWorkflowTabSpawnInfo& Info) const
+{
+	TWeakPtr<FPersona> PersonaPtr = StaticCastSharedPtr<FPersona>(HostingApp.Pin());
+	return SNew(SAnimBlueprintParentPlayerList).Persona(PersonaPtr);
+}
+
+FText FAnimBlueprintParentPlayerEditorSummoner::GetTabToolTipText(const FWorkflowTabSpawnInfo& Info) const
+{
+	return LOCTEXT("AnimSubClassTabToolTip", "Editor for overriding the animation assets referenced by the parent animation graph.");
+}
+
 #undef LOCTEXT_NAMESPACE
diff --git a/Engine/Source/Editor/Persona/Private/Shared/PersonaMode.h b/Engine/Source/Editor/Persona/Private/Shared/PersonaMode.h
--- a/Engine/Source/Editor/Persona/Private/Shared/PersonaMode.h
+++ b/Engine/Source/Editor/Persona/Private/Shared/PersonaMode.h
@@ -22,6 +22,7 @@ struct FPersonaTabs
 	// Explorer
 	// Blueprint Defaults
 	static const FName AnimBlueprintDefaultsEditorID;
+	static const FName AnimBlueprintParentPlayerEditorID;
 	// Anim Document
 	static const FName ScrubberID;
 	// Toolbar
@@ -246,4 +247,19 @@ private:
 	void OnCheckedChanged(ESlateCheckBoxState::Type NewType, EAnimBlueprintEditorMode::Type Mode);
 	EAnimBlueprintEditorMode::Type CurrentMode;
 };
-#undef LOCTEXT_NAMESPACE
+
+//////////////////////////////////////////////////////////////////////////
+// FAnimBlueprintParentPlayerEditorSummoner
+class FAnimBlueprintParentPlayerEditorSummoner : public FWorkflowTabFactory
+{
+public:
+	FAnimBlueprintParentPlayerEditorSummoner(TSharedPtr<class FAssetEditorToolkit> InHostingApp);
+
+	virtual TSharedRef<SWidget> CreateTabBody(const FWorkflowTabSpawnInfo& Info) const OVERRIDE;
+	virtual FText GetTabToolTipText(const FWorkflowTabSpawnInfo& Info) const OVERRIDE;
+
+private:
+
+};
+
+#undef LOCTEXT_NAMESPACE
\ No newline at end of file
diff --git a/Engine/Source/Editor/UnrealEd/Classes/Animation/EditorParentPlayerListObj.h b/Engine/Source/Editor/UnrealEd/Classes/Animation/EditorParentPlayerListObj.h
new file mode 100644
--- /dev/null
+++ b/Engine/Source/Editor/UnrealEd/Classes/Animation/EditorParentPlayerListObj.h
@@ -0,0 +1,33 @@
+// Copyright 1998-2014 Epic Games, Inc. All Rights Reserved.
+
+#pragma once
+#include "EditorParentPlayerListObj.generated.h"
+
+UCLASS()
+class UNREALED_API UEditorParentPlayerListObj : public UObject
+{
+	GENERATED_UCLASS_BODY()
+public:
+	// Build a possible list of overrides from a blueprint and pull in any existing overrides.
+	void InitialiseFromBlueprint(UAnimBlueprint* Blueprint);
+
+	// Adds an overridable node to the possible override list.
+	FAnimParentNodeAssetOverride& AddOverridableNode(UAnimGraphNode_Base* Node);
+
+	// Gets a pointer to a U-node from the stored GUID for that node.
+	UAnimGraphNode_Base* GetVisualNodeFromGuid(FGuid InGuid) const;
+
+	// Pushes an override into the blueprint we're editing.
+	void ApplyOverrideToBlueprint(FAnimParentNodeAssetOverride& Override);
+
+	// Get the current blueprint
+	const UAnimBlueprint* GetBlueprint() const;
+
+	// List of possible overrides to display.
+	UPROPERTY(Category = AnimationGraphOverrides, EditAnywhere)
+	TArray<FAnimParentNodeAssetOverride> Overrides;
+
+private:
+	UAnimBlueprint* AnimBlueprint;
+	TMap<FGuid, UAnimGraphNode_Base*> GuidToVisualNodeMap;
+};
diff --git a/Engine/Source/Editor/UnrealEd/Private/Animation/EditorParentPlayerListObj.cpp b/Engine/Source/Editor/UnrealEd/Private/Animation/EditorParentPlayerListObj.cpp
new file mode 100644
--- /dev/null
+++ b/Engine/Source/Editor/UnrealEd/Private/Animation/EditorParentPlayerListObj.cpp
@@ -0,0 +1,134 @@
+// Copyright 1998-2014 Epic Games, Inc. All Rights Reserved.
+
+#include "UnrealEd.h"
+#include "AnimGraphNode_Base.h"
+#include "AnimGraphNode_SequencePlayer.h"
+#include "AnimGraphNode_BlendSpaceBase.h"
+#include "AnimGraphNode_BlendSpacePlayer.h"
+#include "BlueprintEditorUtils.h"
+#include "ScopedTransaction.h"
+
+UEditorParentPlayerListObj::UEditorParentPlayerListObj(const class FPostConstructInitializeProperties& PCIP)
+	: Super(PCIP)
+	, AnimBlueprint(NULL)
+{
+
+}
+
+FAnimParentNodeAssetOverride& UEditorParentPlayerListObj::AddOverridableNode(UAnimGraphNode_Base* Node)
+{
+	FAnimParentNodeAssetOverride* ExistingEntry = Overrides.FindByPredicate([Node](const FAnimParentNodeAssetOverride& Other)
+	{
+		return Other.ParentNodeGuid == Node->NodeGuid;
+	});
+
+	if(!ExistingEntry)
+	{
+		Overrides.AddZeroed();
+		Overrides.Last().ParentNodeGuid = Node->NodeGuid;
+		Overrides.Last().NewAsset = Node->GetAnimationAsset();
+
+		GuidToVisualNodeMap.Add(Node->NodeGuid, Node);
+	}
+
+	return ExistingEntry ? *ExistingEntry : Overrides.Last();
+}
+
+void UEditorParentPlayerListObj::InitialiseFromBlueprint(UAnimBlueprint* Blueprint)
+{
+	AnimBlueprint = Blueprint;
+
+	Overrides.Empty();
+
+	TArray<UBlueprint*> BlueprintHierarchy;
+	Blueprint->GetBlueprintHierarchyFromClass(Blueprint->GetAnimBlueprintGeneratedClass(), BlueprintHierarchy);
+
+	// Search from 1 as 0 is this class and we're looking for it's parents
+	for(int32 BPIdx = 1 ; BPIdx < BlueprintHierarchy.Num() ; ++BPIdx)
+	{
+		UBlueprint* CurrBP = BlueprintHierarchy[BPIdx];
+
+		TArray<UEdGraph*> Graphs;
+		CurrBP->GetAllGraphs(Graphs);
+
+		for(UEdGraph* Graph : Graphs)
+		{
+			for(UEdGraphNode* Node : Graph->Nodes)
+			{
+				UAnimGraphNode_Base* AnimNode = Cast<UAnimGraphNode_Base>(Node);
+				if(AnimNode && AnimNode->GetAnimationAsset())
+				{
+					FAnimParentNodeAssetOverride& Override = AddOverridableNode(AnimNode);
+
+					// Check the blueprint for saved overrides
+					FAnimParentNodeAssetOverride* SavedOverride = AnimBlueprint->GetAssetOverrideForNode(Override.ParentNodeGuid);
+
+					if(SavedOverride)
+					{
+						Override.NewAsset = SavedOverride->NewAsset;
+					}
+				}
+			}
+		}
+	}
+}
+
+void UEditorParentPlayerListObj::ApplyOverrideToBlueprint(FAnimParentNodeAssetOverride& Override)
+{
+	FAnimParentNodeAssetOverride* ExistingOverride = AnimBlueprint->ParentAssetOverrides.FindByPredicate([Override](const FAnimParentNodeAssetOverride& Other)
+	{
+		return Other.ParentNodeGuid == Override.ParentNodeGuid;
+	});
+
+	UAnimBlueprintGeneratedClass* GenClass = AnimBlueprint->GetAnimBlueprintGeneratedClass();
+
+	FScopedTransaction Transaction(NSLOCTEXT("AnimOverrideEditorObj", "ApplyToBlueprintTransaction", "Apply an override to a blueprint."));
+	AnimBlueprint->Modify();
+
+	if(ExistingOverride)
+	{
+		UAnimGraphNode_Base* ExistingNode = GetVisualNodeFromGuid(ExistingOverride->ParentNodeGuid);
+		check(ExistingNode);
+
+		FAnimParentNodeAssetOverride* ParentOverride = AnimBlueprint->GetAssetOverrideForNode(Override.ParentNodeGuid, true);
+
+		if((ParentOverride && Override.NewAsset == ParentOverride->NewAsset) || Override.NewAsset == ExistingNode->GetAnimationAsset())
+		{
+			// If the asset is the same as a parent, or failing that same as the root; remove it
+			FAnimParentNodeAssetOverride OverrideToRemove = *ExistingOverride;
+			AnimBlueprint->ParentAssetOverrides.Remove(OverrideToRemove);
+		}
+		else
+		{
+			// Otherwise just update the override
+			ExistingOverride->NewAsset = Override.NewAsset;
+		}
+	}
+	else
+	{
+		UAnimGraphNode_Base* GraphNode = GetVisualNodeFromGuid(Override.ParentNodeGuid);
+		Override.ParentNodeGuid = GraphNode->NodeGuid;
+		// Can't find the override so add it to the blueprint.
+		AnimBlueprint->ParentAssetOverrides.Add(Override);
+	}
+
+	AnimBlueprint->NotifyOverrideChange(Override);
+	FBlueprintEditorUtils::MarkBlueprintAsModified(AnimBlueprint);
+
+}
+
+UAnimGraphNode_Base* UEditorParentPlayerListObj::GetVisualNodeFromGuid(FGuid InGuid) const
+{
+	UAnimGraphNode_Base* const* Result = GuidToVisualNodeMap.Find(InGuid);
+	if(Result)
+	{
+		return *Result;
+	}
+	return NULL;
+}
+
+const UAnimBlueprint* UEditorParentPlayerListObj::GetBlueprint() const
+{
+	return AnimBlueprint;
+}
+
diff --git a/Engine/Source/Editor/UnrealEd/Private/Factories/AnimBlueprintFactory.cpp b/Engine/Source/Editor/UnrealEd/Private/Factories/AnimBlueprintFactory.cpp
--- a/Engine/Source/Editor/UnrealEd/Private/Factories/AnimBlueprintFactory.cpp
+++ b/Engine/Source/Editor/UnrealEd/Private/Factories/AnimBlueprintFactory.cpp
@@ -378,6 +378,14 @@ UObject* UAnimBlueprintFactory::FactoryCreateNew(UClass* Class, UObject* InParen
 	else
 	{
 		UAnimBlueprint* NewBP = CastChecked<UAnimBlueprint>(FKismetEditorUtilities::CreateBlueprint(ParentClass, InParent, Name, BlueprintType, UAnimBlueprint::StaticClass(), UBlueprintGeneratedClass::StaticClass(), CallingContext));
+		
+		// Inherit any existing overrides in parent class
+		if (NewBP->ParentAssetOverrides.Num() > 0)
+		{
+			// We've inherited some overrides from the parent graph and need to recompile the blueprint.
+			FKismetEditorUtilities::CompileBlueprint(NewBP);
+		}
+		
 		NewBP->TargetSkeleton = TargetSkeleton;
 
 		// Because the BP itself didn't have the skeleton set when the initial compile occured, it's not set on the generated classes either
diff --git a/Engine/Source/Editor/UnrealEd/Public/GraphEditor.h b/Engine/Source/Editor/UnrealEd/Public/GraphEditor.h
--- a/Engine/Source/Editor/UnrealEd/Public/GraphEditor.h
+++ b/Engine/Source/Editor/UnrealEd/Public/GraphEditor.h
@@ -27,6 +27,8 @@ struct FGraphAppearanceInfo
 	FString CornerText;
 	/** If set, will be used as override for PIE notify text */
 	FString PIENotifyText;
+	/** If set, will be used as override for read only text */
+	FString ReadOnlyText;
 };
 
 /** Struct used to return info about action menu */
diff --git a/Engine/Source/Runtime/Engine/Classes/Animation/AnimBlueprint.h b/Engine/Source/Runtime/Engine/Classes/Animation/AnimBlueprint.h
--- a/Engine/Source/Runtime/Engine/Classes/Animation/AnimBlueprint.h
+++ b/Engine/Source/Runtime/Engine/Classes/Animation/AnimBlueprint.h
@@ -22,6 +22,31 @@ struct FAnimGroupInfo
 	}
 };
 
+USTRUCT()
+struct FAnimParentNodeAssetOverride
+{
+	GENERATED_USTRUCT_BODY()
+
+	UPROPERTY()
+	UAnimationAsset* NewAsset;
+	UPROPERTY()
+	FGuid ParentNodeGuid;
+
+	FAnimParentNodeAssetOverride(FGuid InGuid, UAnimationAsset* InNewAsset)
+		: NewAsset(InNewAsset)
+		, ParentNodeGuid(InGuid)
+	{}
+
+	FAnimParentNodeAssetOverride()
+		: NewAsset(NULL)
+	{}
+
+	bool operator ==(const FAnimParentNodeAssetOverride& Other)
+	{
+		return ParentNodeGuid == Other.ParentNodeGuid;
+	}
+};
+
 /**
  * An Anim Blueprint is essentially a specialized Blueprint whose graphs control the animation of a Skeletal Mesh.
  * It can perform blending of animations, directly control the bones of the skeleton, and output a final pose
@@ -47,6 +72,17 @@ class ENGINE_API UAnimBlueprint : public UBlueprint
 	class UAnimBlueprintGeneratedClass* GetAnimBlueprintSkeletonClass() const;
 
 #if WITH_EDITOR
+
+	// Inspects the hierarchy and looks for an override for the requested node GUID
+	// @param NodeGuid - Guid of the node to search for
+	// @param bIgnoreSelf - Ignore this blueprint and only search parents, handy for finding parent overrides
+	FAnimParentNodeAssetOverride* GetAssetOverrideForNode(FGuid NodeGuid, bool bIgnoreSelf = false) const ;
+
+	// Inspects the hierarchy and builds a list of all asset overrides for this blueprint
+	// @param OutOverrides - Array to fill with overrides
+	// @return bool - Whether any overrides were found
+	bool GetAssetOverrides(TArray<FAnimParentNodeAssetOverride*>& OutOverrides);
+
 	// UBlueprint interface
 	virtual bool SupportedByDefaultBlueprintFactory() const OVERRIDE
 	{
@@ -59,5 +95,37 @@ class ENGINE_API UAnimBlueprint : public UBlueprint
 
 	/** Returns the most base anim blueprint for a given blueprint (if it is inherited from another anim blueprint, returning null if only native / non-anim BP classes are it's parent) */
 	static UAnimBlueprint* FindRootAnimBlueprint(UAnimBlueprint* DerivedBlueprint);
+
+	DECLARE_MULTICAST_DELEGATE_TwoParams(FOnOverrideChangedMulticaster, FGuid, UAnimationAsset*);
+
+	typedef FOnOverrideChangedMulticaster::FDelegate FOnOverrideChanged;
+
+	void RegisterOnOverrideChanged(const FOnOverrideChanged& Delegate)
+	{
+		OnOverrideChanged.Add(Delegate);
+	}
+
+	void UnregisterOnOverrideChanged(SWidget* Widget)
+	{
+		OnOverrideChanged.RemoveAll(Widget);
+	}
+
+	void NotifyOverrideChange(FAnimParentNodeAssetOverride& Override)
+	{
+		OnOverrideChanged.Broadcast(Override.ParentNodeGuid, Override.NewAsset);
+	}
+
+	virtual void PostLoad() OVERRIDE;
+
+protected:
+	// Broadcast when an override is changed, allowing derived blueprints to be updated
+	FOnOverrideChangedMulticaster OnOverrideChanged;
+#endif
+
+#if WITH_EDITORONLY_DATA
+public:
+	// Array of overrides to asset containing nodes in the parent that have been overridden
+	UPROPERTY()
+	TArray<FAnimParentNodeAssetOverride> ParentAssetOverrides;
 #endif
 };
diff --git a/Engine/Source/Runtime/Engine/Classes/Animation/AnimBlueprintGeneratedClass.h b/Engine/Source/Runtime/Engine/Classes/Animation/AnimBlueprintGeneratedClass.h
--- a/Engine/Source/Runtime/Engine/Classes/Animation/AnimBlueprintGeneratedClass.h
+++ b/Engine/Source/Runtime/Engine/Classes/Animation/AnimBlueprintGeneratedClass.h
@@ -92,6 +92,9 @@ public:
 	// Map from animation node to their property index
 	TMap<TWeakObjectPtr<class UAnimGraphNode_Base>, int32> NodePropertyToIndexMap;
 
+	// Map from animation node GUID to property index
+	TMap<FGuid, int32> NodeGuidToIndexMap;
+
 	// History of snapshots of animation data
 	TSimpleRingBuffer<FAnimationFrameSnapshot>* SnapshotBuffer;
 
@@ -164,6 +167,17 @@ public:
 #endif
 };
 
+#if WITH_EDITORONLY_DATA
+namespace EPropertySearchMode
+{
+	enum Type
+	{
+		OnlyThis,
+		Hierarchy
+	};
+}
+#endif
+
 UCLASS(dependson=(UAnimBlueprint, UAnimStateMachineTypes))
 class ENGINE_API UAnimBlueprintGeneratedClass : public UBlueprintGeneratedClass
 {
@@ -172,7 +186,7 @@ class ENGINE_API UAnimBlueprintGeneratedClass : public UBlueprintGeneratedClass
 	// List of state machines present in this blueprint class
 	UPROPERTY()
 	TArray<FBakedAnimationStateMachine> BakedStateMachines;
-	
+
 	/** Target skeleton for this blueprint class */
 	UPROPERTY()
 	class USkeleton* TargetSkeleton;
@@ -198,15 +212,37 @@ class ENGINE_API UAnimBlueprintGeneratedClass : public UBlueprintGeneratedClass
 	}
 
 	template<typename StructType>
-	StructType* GetPropertyInstance(UObject* Object, class UAnimGraphNode_Base* Node)
+	const int32* GetNodePropertyIndexFromHierarchy(class UAnimGraphNode_Base* Node)
+	{
+		TArray<const UBlueprintGeneratedClass*> BlueprintHierarchy;
+		GetGeneratedClassesHierarchy(this, BlueprintHierarchy);
+
+		for (const UBlueprintGeneratedClass* Blueprint : BlueprintHierarchy)
+		{
+			if (const UAnimBlueprintGeneratedClass* AnimBlueprintClass = Cast<UAnimBlueprintGeneratedClass>(Blueprint))
+			{
+				const int32* SearchIndex = AnimBlueprintClass->AnimBlueprintDebugData.NodePropertyToIndexMap.Find(Node);
+				if (SearchIndex)
+				{
+					return SearchIndex;
+				}
+			}
+
+		}
+		return NULL;
+	}
+
+	template<typename StructType>
+	UStructProperty* GetPropertyForNode(class UAnimGraphNode_Base* Node, EPropertySearchMode::Type SearchMode = EPropertySearchMode::OnlyThis)
 	{
-		if (int32* pIndex = AnimBlueprintDebugData.NodePropertyToIndexMap.Find(Node))
+		const int32* pIndex = SearchMode == EPropertySearchMode::OnlyThis ? AnimBlueprintDebugData.NodePropertyToIndexMap.Find(Node) : GetNodePropertyIndexFromHierarchy<StructType>(Node);
+		if (pIndex)
 		{
 			if (UStructProperty* AnimationProperty = AnimNodeProperties[AnimNodeProperties.Num() - 1 - *pIndex])
 			{
 				if (AnimationProperty->Struct->IsChildOf(StructType::StaticStruct()))
 				{
-					return AnimationProperty->ContainerPtrToValuePtr<StructType>((void*)Object);
+					return AnimationProperty;
 				}
 			}
 		}
@@ -215,7 +251,37 @@ class ENGINE_API UAnimBlueprintGeneratedClass : public UBlueprintGeneratedClass
 	}
 
 	template<typename StructType>
-	StructType& GetPropertyInstanceChecked(UObject* Object, class UAnimGraphNode_Base* Node)
+	StructType* GetPropertyInstance(UObject* Object, class UAnimGraphNode_Base* Node, EPropertySearchMode::Type SearchMode = EPropertySearchMode::OnlyThis)
+	{
+		UStructProperty* AnimationProperty = GetPropertyForNode<StructType>(Node);
+		if (AnimationProperty)
+		{
+			return AnimationProperty->ContainerPtrToValuePtr<StructType>((void*)Object);
+		}
+
+		return NULL;
+	}
+
+	template<typename StructType>
+	StructType* GetPropertyInstance(UObject* Object, FGuid NodeGuid, EPropertySearchMode::Type SearchMode = EPropertySearchMode::OnlyThis)
+	{
+		const int32* pIndex = GetNodePropertyIndexFromGuid(NodeGuid, SearchMode);
+		if (pIndex)
+		{
+			if (UStructProperty* AnimProperty = AnimNodeProperties[AnimNodeProperties.Num() - 1 - *pIndex])
+			{
+				if (AnimProperty->Struct->IsChildOf(StructType::StaticStruct()))
+				{
+					return AnimProperty->ContainerPtrToValuePtr<StructType>((void*)Object);
+				}
+			}
+		}
+
+		return NULL;
+	}
+
+	template<typename StructType>
+	StructType& GetPropertyInstanceChecked(UObject* Object, class UAnimGraphNode_Base* Node, EPropertySearchMode::Type SearchMode = EPropertySearchMode::OnlyThis)
 	{
 		const int32 Index = AnimBlueprintDebugData.NodePropertyToIndexMap.FindChecked(Node);
 		UStructProperty* AnimationProperty = AnimNodeProperties[AnimNodeProperties.Num() - 1 - Index];
@@ -223,6 +289,9 @@ class ENGINE_API UAnimBlueprintGeneratedClass : public UBlueprintGeneratedClass
 		check(AnimationProperty->Struct->IsChildOf(StructType::StaticStruct()));
 		return AnimationProperty->ContainerPtrToValuePtr<StructType>((void*)Object);
 	}
+
+	const int32* GetNodePropertyIndexFromGuid(FGuid Guid, EPropertySearchMode::Type SearchMode = EPropertySearchMode::OnlyThis);
+
 #endif
 
 	// UStruct interface
diff --git a/Engine/Source/Runtime/Engine/Private/AnimBlueprint.cpp b/Engine/Source/Runtime/Engine/Private/AnimBlueprint.cpp
--- a/Engine/Source/Runtime/Engine/Private/AnimBlueprint.cpp
+++ b/Engine/Source/Runtime/Engine/Private/AnimBlueprint.cpp
@@ -70,4 +70,78 @@ UAnimBlueprint* UAnimBlueprint::FindRootAnimBlueprint(UAnimBlueprint* DerivedBlu
 	return ParentBP;
 }
 
+FAnimParentNodeAssetOverride* UAnimBlueprint::GetAssetOverrideForNode(FGuid NodeGuid, bool bIgnoreSelf) const
+{
+	TArray<UBlueprint*> Hierarchy;
+	GetBlueprintHierarchyFromClass(GetAnimBlueprintGeneratedClass(), Hierarchy);
+
+	for (int32 Idx = bIgnoreSelf ? 1 : 0 ; Idx < Hierarchy.Num() ; ++Idx)
+	{
+		if (UAnimBlueprint* AnimBlueprint = Cast<UAnimBlueprint>(Hierarchy[Idx]))
+		{
+			FAnimParentNodeAssetOverride* Override = AnimBlueprint->ParentAssetOverrides.FindByPredicate([NodeGuid](const FAnimParentNodeAssetOverride& Other)
+			{
+				return Other.ParentNodeGuid == NodeGuid;
+			});
+
+			if (Override)
+			{
+				return Override;
+			}
+		}
+	}
+
+	return nullptr;
+}
+
+bool UAnimBlueprint::GetAssetOverrides(TArray<FAnimParentNodeAssetOverride*>& OutOverrides)
+{
+	TArray<UBlueprint*> Hierarchy;
+	GetBlueprintHierarchyFromClass(GetAnimBlueprintGeneratedClass(), Hierarchy);
+
+	for (UBlueprint* Blueprint : Hierarchy)
+	{
+		if (UAnimBlueprint* AnimBlueprint = Cast<UAnimBlueprint>(Blueprint))
+		{
+			for (FAnimParentNodeAssetOverride& Override : AnimBlueprint->ParentAssetOverrides)
+			{
+				bool OverrideExists = OutOverrides.ContainsByPredicate([Override](const FAnimParentNodeAssetOverride* Other)
+				{
+					return Override.ParentNodeGuid == Other->ParentNodeGuid;
+				});
+
+				if (!OverrideExists)
+				{
+					OutOverrides.Add(&Override);
+				}
+			}
+		}
+	}
+
+	return OutOverrides.Num() > 0;
+}
+
+void UAnimBlueprint::PostLoad()
+{
+	Super::PostLoad();
+#if WITH_EDITOR
+	// Validate animation overrides
+	UAnimBlueprintGeneratedClass* GeneratedClass = GetAnimBlueprintGeneratedClass();
+	
+	if (GeneratedClass)
+	{
+		// If there is no index for the guid, remove the entry.
+		ParentAssetOverrides.RemoveAll([GeneratedClass](const FAnimParentNodeAssetOverride& Element)
+		{
+			if (!GeneratedClass->GetNodePropertyIndexFromGuid(Element.ParentNodeGuid, EPropertySearchMode::Hierarchy))
+			{
+				return true;
+			}
+
+			return false;
+		});
+	}
+#endif
+}
+
 #endif
\ No newline at end of file
diff --git a/Engine/Source/Runtime/Engine/Private/Animation/AnimBlueprintGeneratedClass.cpp b/Engine/Source/Runtime/Engine/Private/Animation/AnimBlueprintGeneratedClass.cpp
--- a/Engine/Source/Runtime/Engine/Private/Animation/AnimBlueprintGeneratedClass.cpp
+++ b/Engine/Source/Runtime/Engine/Private/Animation/AnimBlueprintGeneratedClass.cpp
@@ -232,3 +232,36 @@ void UAnimBlueprintGeneratedClass::PurgeClass(bool bRecompilingOnLoad)
 
 	BakedStateMachines.Empty();
 }
+
+#if WITH_EDITORONLY_DATA
+
+const int32* UAnimBlueprintGeneratedClass::GetNodePropertyIndexFromGuid(FGuid Guid, EPropertySearchMode::Type SearchMode /*= EPropertySearchMode::OnlyThis*/)
+{
+	if (SearchMode == EPropertySearchMode::OnlyThis)
+	{
+		return AnimBlueprintDebugData.NodeGuidToIndexMap.Find(Guid);
+	}
+	else
+	{
+		TArray<const UBlueprintGeneratedClass*> BlueprintHierarchy;
+		GetGeneratedClassesHierarchy(this, BlueprintHierarchy);
+
+		for (const UBlueprintGeneratedClass* Blueprint : BlueprintHierarchy)
+		{
+			if (const UAnimBlueprintGeneratedClass* AnimBlueprintClass = Cast<UAnimBlueprintGeneratedClass>(Blueprint))
+			{
+				const int32* NodeIndex = AnimBlueprintClass->AnimBlueprintDebugData.NodeGuidToIndexMap.Find(Guid);
+
+				if (NodeIndex)
+				{
+					return NodeIndex;
+				}
+			}
+
+		}
+	}
+
+	return NULL;
+}
+
+#endif
