#TTP 334038 - CRASH: FPersona::AttachObjectToPreviewComponent() FAssetTypeActions_AnimBlueprint::OpenAssetEditor()

#proj UE4

#summary Changed validation code so that it runs on SetPreviewMesh not just on initializing Persona

diff --git a/Engine/Source/Editor/Persona/Private/Persona.cpp b/Engine/Source/Editor/Persona/Private/Persona.cpp
--- a/Engine/Source/Editor/Persona/Private/Persona.cpp
+++ b/Engine/Source/Editor/Persona/Private/Persona.cpp
@@ -844,30 +844,6 @@ void FPersona::InitPersona(const EToolkitMode::Type Mode, const TSharedPtr< clas
 		}
 	}
 
-	// Validate the skeleton/meshes attached objects and display a notification to the user if any were broken
-	int32 NumBrokenAssets = TargetSkeleton->ValidatePreviewAttachedObjects();
-	if (PreviewComponent->SkeletalMesh)
-	{
-		NumBrokenAssets += PreviewComponent->SkeletalMesh->ValidatePreviewAttachedObjects();
-	}
-
-	if (NumBrokenAssets > 0)
-	{
-		// Tell the user that there were assets that could not be loaded
-		FFormatNamedArguments Args;
-		Args.Add(TEXT("NumBrokenAssets"), NumBrokenAssets);
-		FNotificationInfo Info(FText::Format(LOCTEXT("MissingPreviewAttachedAssets", "{NumBrokenAssets} attached assets could not be found on loading and were removed"), Args));
-
-		Info.bUseLargeFont = false;
-		Info.ExpireDuration = 5.0f;
-
-		TSharedPtr<SNotificationItem> Notification = FSlateNotificationManager::Get().AddNotification(Info);
-		if (Notification.IsValid())
-		{
-			Notification->SetCompletionState(SNotificationItem::CS_Fail);
-		}
-	}
-
 	PreviewComponent->SetAnimClass(AnimBlueprint ? AnimBlueprint->GeneratedClass : NULL);
 
 	// We always want a preview instance unless we are using blueprints so that bone manipulation works
@@ -1583,6 +1559,7 @@ void FPersona::SetPreviewMesh(USkeletalMesh* NewPreviewMesh)
 	}
 	else
 	{
+		ValidatePreviewAttachedAssets(NewPreviewMesh);
 		if (NewPreviewMesh != PreviewComponent->SkeletalMesh)
 		{
 			if ( PreviewComponent->SkeletalMesh != NULL )
@@ -2530,6 +2507,33 @@ FString FPersona::GetGraphDecorationString(UEdGraph* InGraph) const
 	return TEXT("");
 }
 
+void FPersona::ValidatePreviewAttachedAssets(USkeletalMesh* PreviewSkeletalMesh)
+{
+	// Validate the skeleton/meshes attached objects and display a notification to the user if any were broken
+	int32 NumBrokenAssets = TargetSkeleton->ValidatePreviewAttachedObjects();
+	if (PreviewSkeletalMesh)
+	{
+		NumBrokenAssets += PreviewSkeletalMesh->ValidatePreviewAttachedObjects();
+	}
+
+	if (NumBrokenAssets > 0)
+	{
+		// Tell the user that there were assets that could not be loaded
+		FFormatNamedArguments Args;
+		Args.Add(TEXT("NumBrokenAssets"), NumBrokenAssets);
+		FNotificationInfo Info(FText::Format(LOCTEXT("MissingPreviewAttachedAssets", "{NumBrokenAssets} attached assets could not be found on loading and were removed"), Args));
+
+		Info.bUseLargeFont = false;
+		Info.ExpireDuration = 5.0f;
+
+		TSharedPtr<SNotificationItem> Notification = FSlateNotificationManager::Get().AddNotification(Info);
+		if (Notification.IsValid())
+		{
+			Notification->SetCompletionState(SNotificationItem::CS_Fail);
+		}
+	}
+}
+
 /////////////////////////////////////////////////////
 
 FAnimationRecorder::FAnimationRecorder()
diff --git a/Engine/Source/Editor/Persona/Private/Persona.h b/Engine/Source/Editor/Persona/Private/Persona.h
--- a/Engine/Source/Editor/Persona/Private/Persona.h
+++ b/Engine/Source/Editor/Persona/Private/Persona.h
@@ -105,6 +105,9 @@ public:
 	// Sets the current preview mesh
 	void SetPreviewMesh(USkeletalMesh* NewPreviewMesh);
 
+	// Validate preview attached assets on skeleton and supplied skeletal mesh, notifying user if any are removed
+	void ValidatePreviewAttachedAssets(USkeletalMesh* SkeletalMeshToUse);
+
 	// Thunks for SContentReference
 	UObject* GetMeshAsObject() const;
 	UObject* GetPhysicsAssetAsObject() const;
