Leading and trailing space on socket names is now ignored.
#ttp 333461 - EDITOR: Static Mesh Editor: Can give sockets blank names
#branch UE4
#proj Editor.Persona, Editor.UnrealEd, Runtime.Core
#add Added calls to FText::TrimPrecedingAndTrailing in the verify and commit delegates for socket name textboxes (in SSocketManager, SSkeletonTree and FSkeletalMeshSocketDetails).
#change Fixed a bug in FText::TrimPrecedingAndTrailing causing trailing whitespace to not be removed. This maybe should be corrected in FString::Trim and FString::TrimTrailing, which return the result by value instead of by ref-to-*this (TTP'd).

diff --git a/Engine/Source/Editor/Persona/Private/SSkeletonTree.cpp b/Engine/Source/Editor/Persona/Private/SSkeletonTree.cpp
--- a/Engine/Source/Editor/Persona/Private/SSkeletonTree.cpp
+++ b/Engine/Source/Editor/Persona/Private/SSkeletonTree.cpp
@@ -565,7 +565,9 @@ bool FDisplayedSocketInfo::OnVerifySocketNameChanged( const FText& InText, FText
 	// but you can have a socket with the same name on the mesh *and* the skeleton.
 	bool bVerifyName = true;
 
-	if( InText.IsEmpty() )
+	FText NewText = FText::TrimPrecedingAndTrailing(InText);
+
+	if (NewText.IsEmpty())
 	{
 		OutErrorMessage = LOCTEXT( "EmptySocketName_Error", "Sockets must have a name!");
 		bVerifyName = false;
@@ -579,13 +581,13 @@ bool FDisplayedSocketInfo::OnVerifySocketNameChanged( const FText& InText, FText
 
 			if ( Mesh )
 			{
-				bVerifyName = !PersonaPtr.Pin()->DoesSocketAlreadyExist( SocketData, InText, Mesh->GetMeshOnlySocketList() );
+				bVerifyName = !PersonaPtr.Pin()->DoesSocketAlreadyExist( SocketData, NewText, Mesh->GetMeshOnlySocketList() );
 			}
 		}
 		else
 		{
 			// ...and if we're on the skeleton, check the skeleton for dupes
-			bVerifyName = !PersonaPtr.Pin()->DoesSocketAlreadyExist( SocketData, InText, TargetSkeleton->Sockets );
+			bVerifyName = !PersonaPtr.Pin()->DoesSocketAlreadyExist( SocketData, NewText, TargetSkeleton->Sockets );
 		}
 
 		// Needs to be checked on verify.
@@ -602,12 +604,14 @@ bool FDisplayedSocketInfo::OnVerifySocketNameChanged( const FText& InText, FText
 
 void FDisplayedSocketInfo::OnCommitSocketName( const FText& InText, ETextCommit::Type CommitInfo )
 {
-	const FScopedTransaction Transaction( LOCTEXT( "RenameSocket", "Rename Socket" ) );
+	FText NewText = FText::TrimPrecedingAndTrailing(InText);
+
+	const FScopedTransaction Transaction(LOCTEXT("RenameSocket", "Rename Socket"));
 	SocketData->SetFlags( RF_Transactional );	// Undo doesn't work without this!
 	SocketData->Modify();
 
 	FName OldSocketName = SocketData->SocketName;
-	SocketData->SocketName = FName( *InText.ToString() );
+	SocketData->SocketName = FName( *NewText.ToString() );
 
 	if ( SkeletonTree.IsValid() )
 	{
diff --git a/Engine/Source/Editor/Persona/Private/SkeletalMeshSocketDetails.cpp b/Engine/Source/Editor/Persona/Private/SkeletalMeshSocketDetails.cpp
--- a/Engine/Source/Editor/Persona/Private/SkeletalMeshSocketDetails.cpp
+++ b/Engine/Source/Editor/Persona/Private/SkeletalMeshSocketDetails.cpp
@@ -161,16 +161,18 @@ void FSkeletalMeshSocketDetails::OnSocketNameCommitted(const FText& InSearchText
 		{
 			FPersona* PersonaInstance = static_cast<FPersona*> (EditorInstance);
 
+			FText NewText = FText::TrimPrecedingAndTrailing(InSearchText);
+
 			FText OutErrorMessage;
-			if( VerifySocketName(PersonaInstance, TargetSocket, InSearchText, OutErrorMessage) )
+			if (VerifySocketName(PersonaInstance, TargetSocket, NewText, OutErrorMessage))
 			{
 				// tell persona to rename the socket.
-				PersonaInstance->RenameSocket( TargetSocket, *InSearchText.ToString() );
+				PersonaInstance->RenameSocket(TargetSocket, *NewText.ToString());
 				PersonaInstance->SetSelectedSocket( FSelectedSocketInfo( TargetSocket, true ) );
 				
 				// update the pre-edit socket name and the property
-				PreEditSocketName = InSearchText;
-				SocketNameTextBox->SetText( InSearchText );
+				PreEditSocketName = NewText;
+				SocketNameTextBox->SetText(NewText);
 			}
 			else
 			{
@@ -192,19 +194,21 @@ bool FSkeletalMeshSocketDetails::VerifySocketName( FPersona* PersonaPtr, const U
 	// but you can have a socket with the same name on the mesh *and* the skeleton.
 	bool bVerifyName = true;
 
-	if( InText.IsEmpty() )
+	FText NewText = FText::TrimPrecedingAndTrailing(InText);
+
+	if (NewText.IsEmpty())
 	{
 		OutErrorMessage = NSLOCTEXT( "SkeletalMeshSocketDetails", "EmptySocketName_Error", "Sockets must have a name!");
 		bVerifyName = false;
 	}
 	else
 	{
-		if( PersonaPtr->DoesSocketAlreadyExist(Socket, InText, PersonaPtr->GetSkeleton()->Sockets) )
+		if (PersonaPtr->DoesSocketAlreadyExist(Socket, NewText, PersonaPtr->GetSkeleton()->Sockets))
 		{
 			bVerifyName = false;
 		}
 
-		if( bVerifyName && PersonaPtr->DoesSocketAlreadyExist(Socket, InText, PersonaPtr->GetMesh()->GetMeshOnlySocketList()) )
+		if (bVerifyName && PersonaPtr->DoesSocketAlreadyExist(Socket, NewText, PersonaPtr->GetMesh()->GetMeshOnlySocketList()))
 		{
 			bVerifyName = false;
 		}
diff --git a/Engine/Source/Editor/UnrealEd/Private/SSocketManager.cpp b/Engine/Source/Editor/UnrealEd/Private/SSocketManager.cpp
--- a/Engine/Source/Editor/UnrealEd/Private/SSocketManager.cpp
+++ b/Engine/Source/Editor/UnrealEd/Private/SSocketManager.cpp
@@ -89,7 +89,8 @@ private:
 	{
 		bool bVerifyName = true;
 
-		if(InNewText.IsEmpty())
+		FText NewText = FText::TrimPrecedingAndTrailing(InNewText);
+		if(NewText.IsEmpty())
 		{
 			OutErrorMessage = LOCTEXT( "EmptySocketName_Error", "Sockets must have a name!");
 			bVerifyName = false;
@@ -99,8 +100,8 @@ private:
 			TSharedPtr<SocketListItem> SocketItemPinned = SocketItem.Pin();
 			TSharedPtr<SSocketManager> SocketManagerPinned = SocketManagerPtr.Pin();
 
-			if (SocketItemPinned.IsValid() && SocketItemPinned->Socket != nullptr && SocketItemPinned->Socket->SocketName.ToString() != InNewText.ToString() &&
-				SocketManagerPinned.IsValid() && SocketManagerPinned->CheckForDuplicateSocket(InNewText.ToString()))
+			if (SocketItemPinned.IsValid() && SocketItemPinned->Socket != nullptr && SocketItemPinned->Socket->SocketName.ToString() != NewText.ToString() &&
+				SocketManagerPinned.IsValid() && SocketManagerPinned->CheckForDuplicateSocket(NewText.ToString()))
 			{
 				OutErrorMessage = LOCTEXT("DuplicateSocket_Error", "Socket name in use!");
 				bVerifyName = false;
@@ -112,6 +113,8 @@ private:
 
 	void OnCommitSocketName( const FText& InText, ETextCommit::Type CommitInfo )
 	{
+		FText NewText = FText::TrimPrecedingAndTrailing(InText);
+
 		TSharedPtr<SocketListItem> PinnedSocketItem = SocketItem.Pin();
 		if (PinnedSocketItem.IsValid())
 		{
@@ -126,7 +129,7 @@ private:
 				SelectedSocket->PreEditChange(ChangedProperty);
 
 				// Edit the property itself
-				SelectedSocket->SocketName = FName(*InText.ToString());
+				SelectedSocket->SocketName = FName(*NewText.ToString());
 
 				// Post edit
 				FPropertyChangedEvent PropertyChangedEvent( ChangedProperty );
diff --git a/Engine/Source/Runtime/Core/Private/Internationalization/Text.cpp b/Engine/Source/Runtime/Core/Private/Internationalization/Text.cpp
--- a/Engine/Source/Runtime/Core/Private/Internationalization/Text.cpp
+++ b/Engine/Source/Runtime/Core/Private/Internationalization/Text.cpp
@@ -222,7 +222,8 @@ FText FText::TrimTrailing( const FText& InText )
 FText FText::TrimPrecedingAndTrailing( const FText& InText )
 {
 	FString TrimmedString = InText.ToString();
-	TrimmedString.Trim().TrimTrailing();
+	TrimmedString.Trim();
+	TrimmedString.TrimTrailing();
 
 	FText NewText = FText( MoveTemp( TrimmedString ) );
 
