- slot group node support for multi montage
 - added customizeProperty option for anim node properties, so that you can customize
 - slot node is customized so that you can add slot/group name as well as you can go to manage section
 - skeleton now stores list of slot/group names available

#code reivew: Laurent.Delayen, Michael.Noland

diff --git a/Engine/Source/Editor/AnimGraph/Classes/AnimGraphNode_Slot.h b/Engine/Source/Editor/AnimGraph/Classes/AnimGraphNode_Slot.h
--- a/Engine/Source/Editor/AnimGraph/Classes/AnimGraphNode_Slot.h
+++ b/Engine/Source/Editor/AnimGraph/Classes/AnimGraphNode_Slot.h
@@ -21,5 +21,6 @@ class UAnimGraphNode_Slot : public UAnimGraphNode_Base
 
 	// UAnimGraphNode_Base interface
 	virtual FString GetNodeCategory() const OVERRIDE;
+	virtual void BakeDataDuringCompilation(class FCompilerResultsLog& MessageLog) OVERRIDE;
 	// End of UAnimGraphNode_Base interface
 };
diff --git a/Engine/Source/Editor/AnimGraph/Classes/AnimationGraphSchema.h b/Engine/Source/Editor/AnimGraph/Classes/AnimationGraphSchema.h
--- a/Engine/Source/Editor/AnimGraph/Classes/AnimationGraphSchema.h
+++ b/Engine/Source/Editor/AnimGraph/Classes/AnimationGraphSchema.h
@@ -26,6 +26,9 @@ class UAnimationGraphSchema : public UEdGraphSchema_K2
 	UPROPERTY()
 	FName NAME_AlwaysAsPin;
 
+	UPROPERTY()
+	FName NAME_CustomizeProperty;
+
 	UPROPERTY()
 	FName NAME_OnEvaluate;
 	
diff --git a/Engine/Source/Editor/AnimGraph/Private/AnimGraphNode_Base.cpp b/Engine/Source/Editor/AnimGraph/Private/AnimGraphNode_Base.cpp
--- a/Engine/Source/Editor/AnimGraph/Private/AnimGraphNode_Base.cpp
+++ b/Engine/Source/Editor/AnimGraph/Private/AnimGraphNode_Base.cpp
@@ -46,9 +46,11 @@ public:
 		const bool bOptional_ShowByDefault = TestProperty->HasMetaData(Schema->NAME_PinShownByDefault);
 		const bool bOptional_HideByDefault = TestProperty->HasMetaData(Schema->NAME_PinHiddenByDefault);
 		const bool bNeverShow = TestProperty->HasMetaData(Schema->NAME_NeverAsPin);
+		const bool bPropertyIsCustomized = TestProperty->HasMetaData(Schema->NAME_CustomizeProperty);
 
 		Record.bCanToggleVisibility = bOptional_ShowByDefault || bOptional_HideByDefault;
 		Record.bShowPin = bAlwaysShow || bOptional_ShowByDefault;
+		Record.bPropertyIsCustomized = bPropertyIsCustomized;
  	}
 
 	virtual void CustomizePinData(UEdGraphPin* Pin, FName SourcePropertyName, int32 ArrayIndex, UProperty* Property) const OVERRIDE
diff --git a/Engine/Source/Editor/AnimGraph/Private/AnimGraphNode_Slot.cpp b/Engine/Source/Editor/AnimGraph/Private/AnimGraphNode_Slot.cpp
--- a/Engine/Source/Editor/AnimGraph/Private/AnimGraphNode_Slot.cpp
+++ b/Engine/Source/Editor/AnimGraph/Private/AnimGraphNode_Slot.cpp
@@ -63,4 +63,14 @@ FString UAnimGraphNode_Slot::GetNodeCategory() const
 	return TEXT("Blends");
 }
 
+void UAnimGraphNode_Slot::BakeDataDuringCompilation(class FCompilerResultsLog& MessageLog)
+{
+	UAnimBlueprint* AnimBlueprint = GetAnimBlueprint();
+	if (AnimBlueprint->TargetSkeleton)
+	{
+		AnimBlueprint->TargetSkeleton->AddSlotNodeName(Node.SlotName);
+		AnimBlueprint->TargetSkeleton->AddSlotGroupName(Node.GroupName);
+	}
+}
+
 #undef LOCTEXT_NAMESPACE
diff --git a/Engine/Source/Editor/AnimGraph/Private/AnimationGraphSchema.cpp b/Engine/Source/Editor/AnimGraph/Private/AnimationGraphSchema.cpp
--- a/Engine/Source/Editor/AnimGraph/Private/AnimationGraphSchema.cpp
+++ b/Engine/Source/Editor/AnimGraph/Private/AnimationGraphSchema.cpp
@@ -35,6 +35,7 @@ UAnimationGraphSchema::UAnimationGraphSchema(const class FPostConstructInitializ
 	NAME_PinShownByDefault = TEXT("PinShownByDefault");
 	NAME_AlwaysAsPin = TEXT("AlwaysAsPin");
 	NAME_OnEvaluate = TEXT("OnEvaluate");
+	NAME_CustomizeProperty = TEXT("CustomizeProperty");
 	DefaultEvaluationHandlerName = TEXT("EvaluateGraphExposedInputs");
 }
 
diff --git a/Engine/Source/Editor/BlueprintGraph/Classes/K2Node.h b/Engine/Source/Editor/BlueprintGraph/Classes/K2Node.h
--- a/Engine/Source/Editor/BlueprintGraph/Classes/K2Node.h
+++ b/Engine/Source/Editor/BlueprintGraph/Classes/K2Node.h
@@ -68,16 +68,20 @@ struct FOptionalPinFromProperty
 	UPROPERTY(EditAnywhere, BlueprintReadWrite, Category=Hi)
 	bool bCanToggleVisibility;
 
+	UPROPERTY(EditAnywhere, BlueprintReadWrite, Category=Hi)
+	bool bPropertyIsCustomized;
+
 	FOptionalPinFromProperty()
 	{
 	}
 	
-	FOptionalPinFromProperty(FName InPropertyName, bool bInShowPin, bool bInCanToggleVisibility, const FString& InFriendlyName, const FText& InTooltip)
+	FOptionalPinFromProperty(FName InPropertyName, bool bInShowPin, bool bInCanToggleVisibility, const FString& InFriendlyName, const FText& InTooltip, bool bInPropertyIsCustomized)
 		: PropertyName(InPropertyName)
 		, PropertyFriendlyName(InFriendlyName)
 		, PropertyTooltip(InTooltip)
 		, bShowPin(bInShowPin)
 		, bCanToggleVisibility(bInCanToggleVisibility)
+		, bPropertyIsCustomized(bInPropertyIsCustomized)
 	{
 	}
 };
diff --git a/Engine/Source/Editor/Persona/Private/AnimGraphNodeDetails.cpp b/Engine/Source/Editor/Persona/Private/AnimGraphNodeDetails.cpp
--- a/Engine/Source/Editor/Persona/Private/AnimGraphNodeDetails.cpp
+++ b/Engine/Source/Editor/Persona/Private/AnimGraphNodeDetails.cpp
@@ -93,12 +93,29 @@ void FAnimGraphNodeDetails::CustomizeDetails(class IDetailLayoutBuilder& DetailB
 			continue;
 		}
 
+		// if customized, do not do anything
+		if(TargetPropertyHandle->IsCustomized())
+		{
+			continue;
+		}
+		
+		// sometimes because of order of customization
+		// this gets called first for the node you'd like to customize
+		// then the above statement won't work
+		// so you can mark certain property to have meta data "CustomizeProperty"
+		// which will trigger below statement
+		if (OptionalPin.bPropertyIsCustomized)
+		{
+			continue;
+		}
+
 		IDetailPropertyRow& PropertyRow = CurrentCategory.AddProperty( TargetPropertyHandle );
 
 		if(!TargetPropertyHandle->GetProperty())
 		{
 			continue;
 		}
+
 		TSharedPtr<SWidget> NameWidget; 
 		TSharedPtr<SWidget> ValueWidget;
 		FDetailWidgetRow Row;
@@ -317,6 +334,7 @@ void FInputScaleBiasCustomization::CustomizeStructHeader(TSharedRef<class IPrope
 }
 END_SLATE_FUNCTION_BUILD_OPTIMIZATION
 
+/////////////////////////////////////////////////////
 void FInputScaleBiasCustomization::CustomizeStructChildren( TSharedRef<class IPropertyHandle> StructPropertyHandle, class IDetailChildrenBuilder& StructBuilder, IStructCustomizationUtils& StructCustomizationUtils )
 {
 	// nothing here
diff --git a/Engine/Source/Editor/Persona/Private/BlueprintMode/AnimBlueprintMode.cpp b/Engine/Source/Editor/Persona/Private/BlueprintMode/AnimBlueprintMode.cpp
--- a/Engine/Source/Editor/Persona/Private/BlueprintMode/AnimBlueprintMode.cpp
+++ b/Engine/Source/Editor/Persona/Private/BlueprintMode/AnimBlueprintMode.cpp
@@ -7,6 +7,8 @@
 #include "Editor/Kismet/Public/BlueprintEditorTabs.h"
 #include "Editor/Kismet/Public/SBlueprintEditorToolbar.h"
 #include "Editor/Kismet/Public/BlueprintEditorModes.h"
+#include "SSkeletonSlotNames.h"
+#include "SSkeletonSlotGroupNames.h"
 #include "SSkeletonAnimNotifies.h"
 
 #include "AnimBlueprintMode.h"
@@ -111,6 +113,8 @@ FAnimBlueprintEditAppMode::FAnimBlueprintEditAppMode(TSharedPtr<FPersona> InPers
 	PersonaTabFactories.RegisterFactory(MakeShareable(new FSkeletonAnimNotifiesSummoner(InPersona)));
 	PersonaTabFactories.RegisterFactory(MakeShareable(new FAnimBlueprintDefaultsEditorSummoner(InPersona)));
 	PersonaTabFactories.RegisterFactory(MakeShareable(new FAnimBlueprintParentPlayerEditorSummoner(InPersona)));
+	PersonaTabFactories.RegisterFactory(MakeShareable(new FSkeletonSlotNamesSummoner(InPersona)));
+	PersonaTabFactories.RegisterFactory(MakeShareable(new FSkeletonSlotGroupNamesSummoner(InPersona)));
 
 	// setup toolbar - clear existing toolbar extender from the BP mode
 	//@TODO: Keep this in sync with BlueprintEditorModes.cpp
diff --git a/Engine/Source/Editor/Persona/Private/Persona.cpp b/Engine/Source/Editor/Persona/Private/Persona.cpp
--- a/Engine/Source/Editor/Persona/Private/Persona.cpp
+++ b/Engine/Source/Editor/Persona/Private/Persona.cpp
@@ -46,6 +46,9 @@
 #include "AnimGraphNode_LayeredBoneBlend.h"
 #include "AnimGraphNode_SequencePlayer.h"
 #include "AnimGraphNode_SequenceEvaluator.h"
+#include "AnimGraphNode_Slot.h"
+#include "Customization/AnimGraphNodeSlotDetails.h"
+
 #include "AnimPreviewInstance.h"
 #define LOCTEXT_NAMESPACE "FPersona"
 
@@ -870,6 +873,11 @@ void FPersona::InitPersona(const EToolkitMode::Type Mode, const TSharedPtr< clas
 	{
 		OpenNewDocumentTab( InitAnimationAsset );
 	}
+
+	// register customization of Slot node for this Persona
+	// this is so that you can open the manage window per Persona
+	Inspector->GetPropertyView()->RegisterInstancedCustomPropertyLayout(UAnimGraphNode_Slot::StaticClass(),
+		FOnGetDetailCustomizationInstance::CreateStatic(&FAnimGraphNodeSlotDetails::MakeInstance, SharedThis(this) ));
 }
 
 
diff --git a/Engine/Source/Editor/Persona/Private/PersonaModule.cpp b/Engine/Source/Editor/Persona/Private/PersonaModule.cpp
--- a/Engine/Source/Editor/Persona/Private/PersonaModule.cpp
+++ b/Engine/Source/Editor/Persona/Private/PersonaModule.cpp
@@ -68,6 +68,18 @@ void FPersonaModule::StartupModule()
 void FPersonaModule::ShutdownModule()
 {
 	MenuExtensibilityManager.Reset();
+
+	// unregsiter when shut down
+	{
+		FPropertyEditorModule& PropertyModule = FModuleManager::LoadModuleChecked<FPropertyEditorModule>("PropertyEditor");
+		PropertyModule.UnregisterCustomPropertyLayout("SkeletalMeshSocket");
+		PropertyModule.UnregisterCustomPropertyLayout("AnimNotify");
+		PropertyModule.UnregisterCustomPropertyLayout("AnimNotifyState");
+		PropertyModule.UnregisterCustomPropertyLayout("AnimGraphNode_Base");
+
+		PropertyModule.UnregisterStructPropertyLayout("InputScaleBias");
+		PropertyModule.UnregisterStructPropertyLayout("BoneReference");
+	}
 }
 
 
diff --git a/Engine/Source/Editor/Persona/Private/SSkeletonSlotGroupNames.cpp b/Engine/Source/Editor/Persona/Private/SSkeletonSlotGroupNames.cpp
new file mode 100644
--- /dev/null
+++ b/Engine/Source/Editor/Persona/Private/SSkeletonSlotGroupNames.cpp
@@ -0,0 +1,417 @@
+// Copyright 1998-2014 Epic Games, Inc. All Rights Reserved.
+
+
+#include "PersonaPrivatePCH.h"
+
+#include "Persona.h"
+#include "SSkeletonSlotGroupNames.h"
+#include "AssetRegistryModule.h"
+#include "ScopedTransaction.h"
+
+#define LOCTEXT_NAMESPACE "SkeletonSlotGroupNames"
+
+static const FName ColumnId_SlotGroupNameLabel( "SlotGroupName" );
+
+//////////////////////////////////////////////////////////////////////////
+// SMorphTargetListRow
+
+typedef TSharedPtr< FDisplayedSlotGroupNameInfo > FDisplayedSlotGroupNameInfoPtr;
+
+class SSlotGroupNameListRow
+	: public SMultiColumnTableRow< FDisplayedSlotGroupNameInfoPtr >
+{
+public:
+
+	SLATE_BEGIN_ARGS( SSlotGroupNameListRow ) {}
+
+	/** The item for this row **/
+	SLATE_ARGUMENT( FDisplayedSlotGroupNameInfoPtr, Item )
+
+	/* Widget used to display the list of morph targets */
+	SLATE_ARGUMENT( TSharedPtr<SSkeletonSlotGroupNames>, SlotGroupNameListView )
+
+	/* Persona used to update the viewport when a weight slider is dragged */
+	SLATE_ARGUMENT( TWeakPtr<FPersona>, Persona )
+
+	SLATE_END_ARGS()
+
+	void Construct( const FArguments& InArgs, const TSharedRef<STableViewBase>& OwnerTableView );
+
+	/** Overridden from SMultiColumnTableRow.  Generates a widget for this column of the tree row. */
+	virtual TSharedRef<SWidget> GenerateWidgetForColumn( const FName& ColumnName ) OVERRIDE;
+
+private:
+
+	/** Widget used to display the list of slot name */
+	TSharedPtr<SSkeletonSlotGroupNames> SlotGroupNameListView;
+
+	/** The notify being displayed by this row */
+	FDisplayedSlotGroupNameInfoPtr	Item;
+
+	/** Pointer back to the Persona that owns us */
+	TWeakPtr<FPersona> PersonaPtr;
+};
+
+void SSlotGroupNameListRow::Construct( const FArguments& InArgs, const TSharedRef<STableViewBase>& InOwnerTableView )
+{
+	Item = InArgs._Item;
+	SlotGroupNameListView = InArgs._SlotGroupNameListView;
+	PersonaPtr = InArgs._Persona;
+
+	check( Item.IsValid() );
+
+	SMultiColumnTableRow< FDisplayedSlotGroupNameInfoPtr >::Construct( FSuperRowType::FArguments(), InOwnerTableView );
+}
+
+TSharedRef< SWidget > SSlotGroupNameListRow::GenerateWidgetForColumn( const FName& ColumnName )
+{
+	check( ColumnName == ColumnId_SlotGroupNameLabel );
+
+	return SNew( SVerticalBox )
+	+ SVerticalBox::Slot()
+	.AutoHeight()
+	.Padding( 0.0f, 4.0f )
+	.VAlign( VAlign_Center )
+	[
+		SAssignNew(Item->InlineEditableText, SInlineEditableTextBlock)
+		.Text( FText::FromName(Item->Name) )
+		.OnVerifyTextChanged(SlotGroupNameListView.Get(), &SSkeletonSlotGroupNames::OnVerifyNotifyNameCommit, Item)
+		.OnTextCommitted(SlotGroupNameListView.Get(), &SSkeletonSlotGroupNames::OnNotifyNameCommitted, Item)
+		.IsSelected(SlotGroupNameListView.Get(), &SSkeletonSlotGroupNames::IsSelected)
+	];
+}
+
+/////////////////////////////////////////////////////
+// FSkeletonSlotGroupNamesSummoner
+
+FSkeletonSlotGroupNamesSummoner::FSkeletonSlotGroupNamesSummoner(TSharedPtr<class FAssetEditorToolkit> InHostingApp)
+	: FWorkflowTabFactory(FPersonaTabs::SkeletonSlotGroupNamesID, InHostingApp)
+{
+	TabLabel = LOCTEXT("SkeletonSlotGroupNamesTabTitle", "Animation Slot Group Name");
+
+	EnableTabPadding();
+	bIsSingleton = true;
+
+	ViewMenuDescription = LOCTEXT("SkeletonSlotGroupNamesMenu", "Animation Slot Group Name");
+	ViewMenuTooltip = LOCTEXT("SkeletonSlotGroupNames_ToolTip", "Shows the skeletons slot name list");
+}
+
+TSharedRef<SWidget> FSkeletonSlotGroupNamesSummoner::CreateTabBody(const FWorkflowTabSpawnInfo& Info) const
+{
+	return SNew(SSkeletonSlotGroupNames)
+		.Persona(StaticCastSharedPtr<FPersona>(HostingApp.Pin()));
+}
+
+/////////////////////////////////////////////////////
+// SSkeletonSlotGroupNames
+
+void SSkeletonSlotGroupNames::Construct(const FArguments& InArgs)
+{
+	PersonaPtr = InArgs._Persona;
+	TargetSkeleton = PersonaPtr.Pin()->GetSkeleton();
+
+	PersonaPtr.Pin()->RegisterOnPostUndo(FPersona::FOnPostUndo::CreateSP( this, &SSkeletonSlotGroupNames::PostUndo ) );
+
+	this->ChildSlot
+	[
+		SNew( SVerticalBox )
+
+		+ SVerticalBox::Slot()
+		.AutoHeight()
+		.Padding( FMargin( 0.0f, 0.0f, 0.0f, 4.0f ) )
+		[
+			SAssignNew( NameFilterBox, SSearchBox )
+			.SelectAllTextWhenFocused( true )
+			.OnTextChanged( this, &SSkeletonSlotGroupNames::OnFilterTextChanged )
+			.OnTextCommitted( this, &SSkeletonSlotGroupNames::OnFilterTextCommitted )
+			.HintText( LOCTEXT( "SlotGroupNameSearchBoxHint", "Search Animation Slot Group Name...") )
+		]
+
+		+ SVerticalBox::Slot()
+		.FillHeight( 1.0f )		// This is required to make the scrollbar work, as content overflows Slate containers by default
+		[
+			SAssignNew( SlotGroupNameListView, SSlotGroupNameListType )
+			.ListItemsSource( &NotifyList )
+			.OnGenerateRow( this, &SSkeletonSlotGroupNames::GenerateNotifyRow )
+			.OnContextMenuOpening( this, &SSkeletonSlotGroupNames::OnGetContextMenuContent )
+			.OnSelectionChanged( this, &SSkeletonSlotGroupNames::OnNotifySelectionChanged )
+			.ItemHeight( 22.0f )
+			.HeaderRow
+			(
+				SNew( SHeaderRow )
+				+ SHeaderRow::Column( ColumnId_SlotGroupNameLabel )
+				.DefaultLabel( LOCTEXT( "SlotGroupNameNameLabel", "Notify Name" ).ToString() )
+			)
+		]
+	];
+
+	CreateSlotGroupNameList();
+}
+
+SSkeletonSlotGroupNames::~SSkeletonSlotGroupNames()
+{
+	if ( PersonaPtr.IsValid() )
+	{
+		PersonaPtr.Pin()->UnregisterOnPostUndo( this );
+	}
+}
+
+void SSkeletonSlotGroupNames::OnFilterTextChanged( const FText& SearchText )
+{
+	FilterText = SearchText;
+
+	RefreshSlotGroupNameListWithFilter();
+}
+
+void SSkeletonSlotGroupNames::OnFilterTextCommitted( const FText& SearchText, ETextCommit::Type CommitInfo )
+{
+	// Just do the same as if the user typed in the box
+	OnFilterTextChanged( SearchText );
+}
+
+TSharedRef<ITableRow> SSkeletonSlotGroupNames::GenerateNotifyRow(TSharedPtr<FDisplayedSlotGroupNameInfo> InInfo, const TSharedRef<STableViewBase>& OwnerTable)
+{
+	check( InInfo.IsValid() );
+
+	return
+		SNew( SSlotGroupNameListRow, OwnerTable )
+		.Persona( PersonaPtr )
+		.Item( InInfo )
+		.SlotGroupNameListView( SharedThis( this ) );
+}
+
+TSharedPtr<SWidget> SSkeletonSlotGroupNames::OnGetContextMenuContent() const
+{
+	const bool bShouldCloseWindowAfterMenuSelection = true;
+	FMenuBuilder MenuBuilder( bShouldCloseWindowAfterMenuSelection, NULL);
+
+	MenuBuilder.BeginSection("SlotGroupNameAction", LOCTEXT( "SlotGroupNameActions", "Selected Notify Actions" ) );
+	{
+		{
+			FUIAction Action = FUIAction( FExecuteAction::CreateSP( this, &SSkeletonSlotGroupNames::OnDeleteSlotGroupName ), 
+				FCanExecuteAction::CreateSP( this, &SSkeletonSlotGroupNames::CanPerformDelete ) );
+			const FText Label = LOCTEXT("DeleteSlotGroupNameButtonLabel", "Delete");
+			const FText ToolTip = LOCTEXT("DeleteSlotGroupNameButtonTooltip", "Deletes the selected anim SlotGroupName.");
+			MenuBuilder.AddMenuEntry( Label, ToolTip, FSlateIcon(), Action);
+		}
+
+		// for now disable rename, it doesn't work correctly
+// 		{
+// 			FUIAction Action = FUIAction( FExecuteAction::CreateSP( this, &SSkeletonSlotGroupNames::OnRenameSlotGroupName ), 
+// 				FCanExecuteAction::CreateSP( this, &SSkeletonSlotGroupNames::CanPerformRename ) );
+// 			const FText Label = LOCTEXT("RenameSlotGroupNameButtonLabel", "Rename");
+// 			const FText ToolTip = LOCTEXT("RenameSlotGroupNameButtonTooltip", "Renames the selected anim SlotGroupName.");
+// 			MenuBuilder.AddMenuEntry( Label, ToolTip, FSlateIcon(), Action);
+// 		}
+	}
+	MenuBuilder.EndSection();
+
+	return MenuBuilder.MakeWidget();
+}
+
+void SSkeletonSlotGroupNames::OnNotifySelectionChanged(TSharedPtr<FDisplayedSlotGroupNameInfo> Selection, ESelectInfo::Type SelectInfo)
+{
+	if(Selection.IsValid())
+	{
+		ShowNotifyInDetailsView(Selection->Name);
+	}
+}
+
+bool SSkeletonSlotGroupNames::CanPerformDelete() const
+{
+	TArray< TSharedPtr< FDisplayedSlotGroupNameInfo > > SelectedRows = SlotGroupNameListView->GetSelectedItems();
+	
+	// if it's trying to delete default slot group name, they can't
+	if (SelectedRows.Num() == 1 && SelectedRows[0]->Name == USkeleton::DefaultSlotGroupName)
+	{
+		return false;
+	}
+
+	return SelectedRows.Num() > 0;
+}
+
+bool SSkeletonSlotGroupNames::CanPerformRename() const
+{
+	TArray< TSharedPtr< FDisplayedSlotGroupNameInfo > > SelectedRows = SlotGroupNameListView->GetSelectedItems();
+	return SelectedRows.Num() == 1;
+}
+
+void SSkeletonSlotGroupNames::OnDeleteSlotGroupName()
+{
+	TArray< TSharedPtr< FDisplayedSlotGroupNameInfo > > SelectedRows = SlotGroupNameListView->GetSelectedItems();
+
+	const FScopedTransaction Transaction( LOCTEXT("DeleteSlotGroupName", "Delete Slot Name") );
+
+	// this one deletes all SlotGroupName with same name. 
+	TargetSkeleton->Modify();
+
+	for(int Selection = 0; Selection < SelectedRows.Num(); ++Selection)
+	{
+		FName& SelectedName = SelectedRows[Selection]->Name;
+		TargetSkeleton->RemoveSlotGroupName(SelectedName);
+	}
+
+	CreateSlotGroupNameList( NameFilterBox->GetText().ToString() );
+}
+
+// void SSkeletonSlotGroupNames::GetCompatibleAnimBlueprints( TMultiMap<class UAnimBlueprint *, class UAnimGraphNode_Slot *>& OutAssets )
+// {
+// 	// second, go through all animgraphnode_slot
+// 	for(FObjectIterator Iter(UAnimGraphNode_Slot::StaticClass()); Iter; ++Iter)
+// 	{
+// 		// if I find one, add the BP and slot nodes to the list
+// 		UAnimGraphNode_Slot * SlotNode = CastChecked<UAnimGraphNode_Slot>(Iter);
+// 		UAnimBlueprint * AnimBlueprint = Cast<UAnimBlueprint>(SlotNode->GetBlueprint());
+// 		if (AnimBlueprint)
+// 		{
+// 			OutAssets.AddUnique(AnimBlueprint, SlotNode);
+// 		}
+// 	}
+// }
+
+void SSkeletonSlotGroupNames::OnRenameSlotGroupName()
+{
+	TArray< TSharedPtr< FDisplayedSlotGroupNameInfo > > SelectedRows = SlotGroupNameListView->GetSelectedItems();
+
+	check(SelectedRows.Num() == 1); // Should be guaranteed by CanPerformRename
+
+	SelectedRows[0]->InlineEditableText->EnterEditingMode();
+}
+
+bool SSkeletonSlotGroupNames::OnVerifyNotifyNameCommit( const FText& NewName, FText& OutErrorMessage, TSharedPtr<FDisplayedSlotGroupNameInfo> Item )
+{
+	bool bValid(true);
+
+	if(NewName.IsEmpty())
+	{
+		OutErrorMessage = LOCTEXT( "NameMissing_Error", "You must provide a name." );
+		bValid = false;
+	}
+
+	FName NotifyName( *NewName.ToString() );
+	if(NotifyName != Item->Name)
+	{
+		if(TargetSkeleton->DoesHaveSlotGroupName(NotifyName))
+		{
+			OutErrorMessage = FText::Format( LOCTEXT("AlreadyInUseMessage", "'{0}' is already in use."), NewName );
+			bValid = false;
+		}
+	}
+
+	return bValid;
+}
+
+void SSkeletonSlotGroupNames::OnNotifyNameCommitted( const FText& NewName, ETextCommit::Type, TSharedPtr<FDisplayedSlotGroupNameInfo> Item )
+{
+	const FScopedTransaction Transaction( LOCTEXT("RenameSlotGroupName", "Rename Anim Notify") );
+
+	FName NewSlotGroupName = FName( *NewName.ToString() );
+	FName SlotToRename = Item->Name;
+
+	// rename notify in skeleton
+	TargetSkeleton->Modify();
+
+	TArray<FAssetData> CompatibleAnimMontages;
+	GetCompatibleAnimMontage(CompatibleAnimMontages);
+
+	// @todo need to fix up anim blueprint
+
+	RefreshSlotGroupNameListWithFilter();
+}
+
+void SSkeletonSlotGroupNames::RefreshSlotGroupNameListWithFilter()
+{
+	CreateSlotGroupNameList( NameFilterBox->GetText().ToString() );
+}
+
+void SSkeletonSlotGroupNames::CreateSlotGroupNameList( const FString& SearchText )
+{
+	NotifyList.Empty();
+
+	if ( TargetSkeleton )
+	{
+		const TArray<FName> & SlotGroupNames = TargetSkeleton->GetSlotGroupNames();
+		for(int i = 0; i < SlotGroupNames.Num(); ++i)
+		{
+			const FName& NotifyName = SlotGroupNames[i];
+			if ( !SearchText.IsEmpty() )
+			{
+				if ( NotifyName.ToString().Contains( SearchText ) )
+				{
+					NotifyList.Add( FDisplayedSlotGroupNameInfo::Make( NotifyName ) );
+				}
+			}
+			else
+			{
+				NotifyList.Add( FDisplayedSlotGroupNameInfo::Make( NotifyName ) );
+			}
+		}
+	}
+
+	SlotGroupNameListView->RequestListRefresh();
+}
+
+void SSkeletonSlotGroupNames::ShowNotifyInDetailsView(FName NotifyName)
+{
+	// @todo nothing to show now, but in the future
+	// we can show the list of montage that are used by this slot node?
+}
+
+void SSkeletonSlotGroupNames::GetCompatibleAnimMontage(TArray<class FAssetData>& OutAssets)
+{
+	//Get the skeleton tag to search for
+	FString SkeletonExportName = FAssetData(TargetSkeleton).GetExportTextName();
+
+	// Load the asset registry module
+	FAssetRegistryModule& AssetRegistryModule = FModuleManager::LoadModuleChecked<FAssetRegistryModule>(TEXT("AssetRegistry"));
+
+	TArray<FAssetData> AssetDataList;
+	AssetRegistryModule.Get().GetAssetsByClass(UAnimMontage::StaticClass()->GetFName(), AssetDataList, true);
+
+	OutAssets.Empty(AssetDataList.Num());
+
+	for( int32 AssetIndex = 0; AssetIndex < AssetDataList.Num(); ++AssetIndex )
+	{
+		const FAssetData& PossibleAnimMontage = AssetDataList[AssetIndex];
+		if( SkeletonExportName == PossibleAnimMontage.TagsAndValues.FindRef("Skeleton") )
+		{
+			OutAssets.Add( PossibleAnimMontage );
+		}
+	}
+}
+
+UObject* SSkeletonSlotGroupNames::ShowInDetailsView( UClass* EdClass )
+{
+	UObject *Obj = EditorObjectTracker.GetEditorObjectForClass(EdClass);
+
+	if(Obj != NULL)
+	{
+		PersonaPtr.Pin()->SetDetailObject(Obj);
+	}
+	return Obj;
+}
+
+void SSkeletonSlotGroupNames::ClearDetailsView()
+{
+	PersonaPtr.Pin()->SetDetailObject(NULL);
+}
+
+void SSkeletonSlotGroupNames::PostUndo()
+{
+	RefreshSlotGroupNameListWithFilter();
+}
+
+void SSkeletonSlotGroupNames::AddReferencedObjects( FReferenceCollector& Collector )
+{
+	EditorObjectTracker.AddReferencedObjects(Collector);
+}
+
+void SSkeletonSlotGroupNames::NotifyUser( FNotificationInfo& NotificationInfo )
+{
+	TSharedPtr<SNotificationItem> Notification = FSlateNotificationManager::Get().AddNotification( NotificationInfo );
+	if ( Notification.IsValid() )
+	{
+		Notification->SetCompletionState( SNotificationItem::CS_Fail );
+	}
+}
+#undef LOCTEXT_NAMESPACE
diff --git a/Engine/Source/Editor/Persona/Private/SSkeletonSlotGroupNames.h b/Engine/Source/Editor/Persona/Private/SSkeletonSlotGroupNames.h
new file mode 100644
--- /dev/null
+++ b/Engine/Source/Editor/Persona/Private/SSkeletonSlotGroupNames.h
@@ -0,0 +1,165 @@
+// Copyright 1998-2014 Epic Games, Inc. All Rights Reserved.
+
+#pragma once
+
+#include "EditorObjectsTracker.h"
+#include "IDocumentation.h"
+
+#define LOCTEXT_NAMESPACE "SkeletonSlotGroupNames"
+/////////////////////////////////////////////////////
+// FSkeletonTreeSummoner
+struct FSkeletonSlotGroupNamesSummoner : public FWorkflowTabFactory
+{
+public:
+	FSkeletonSlotGroupNamesSummoner(TSharedPtr<class FAssetEditorToolkit> InHostingApp);
+
+	virtual TSharedRef<SWidget> CreateTabBody(const FWorkflowTabSpawnInfo& Info) const OVERRIDE;
+
+	// Create a tooltip widget for the tab
+	virtual TSharedPtr<SToolTip> CreateTabToolTipWidget(const FWorkflowTabSpawnInfo& Info) const OVERRIDE
+	{
+		return  IDocumentation::Get()->CreateToolTip(LOCTEXT("WindowTooltip", "This tab lets you modify custom animation Slot Group Name"), NULL, TEXT("Shared/Editors/Persona"), TEXT("AnimationSlotGroupName_Window"));
+	}
+};
+
+//////////////////////////////////////////////////////////////////////////
+// FDisplayedSlotGroupNameInfo
+
+class FDisplayedSlotGroupNameInfo
+{
+public:
+	FName Name;
+
+	/** Handle to editable text block for rename */
+	TSharedPtr<SInlineEditableTextBlock> InlineEditableText;
+
+	/** Static function for creating a new item, but ensures that you can only have a TSharedRef to one */
+	static TSharedRef<FDisplayedSlotGroupNameInfo> Make(const FName& NotifyName)
+	{
+		return MakeShareable(new FDisplayedSlotGroupNameInfo(NotifyName));
+	}
+
+protected:
+	/** Hidden constructor, always use Make above */
+	FDisplayedSlotGroupNameInfo(const FName& InNotifyName)
+		: Name( InNotifyName )
+	{}
+
+	/** Hidden constructor, always use Make above */
+	FDisplayedSlotGroupNameInfo() {}
+};
+
+/** Widgets list type */
+typedef SListView< TSharedPtr<FDisplayedSlotGroupNameInfo> > SSlotGroupNameListType;
+
+class SSkeletonSlotGroupNames : public SCompoundWidget, public FGCObject
+{
+public:
+	SLATE_BEGIN_ARGS( SSkeletonSlotGroupNames )
+		: _Persona()
+	{}
+
+	SLATE_ARGUMENT( TWeakPtr<FPersona>, Persona )
+		
+		SLATE_END_ARGS()
+public:
+	void Construct(const FArguments& InArgs);
+
+	~SSkeletonSlotGroupNames();
+
+	/**
+	* Accessor so our rows can grab the filter text for highlighting
+	*
+	*/
+	FText& GetFilterText() { return FilterText; }
+
+	/** Creates an editor object from the given type to be used in a details panel */
+	UObject* ShowInDetailsView( UClass* EdClass );
+
+	/** Clears the detail view of whatever we displayed last */
+	void ClearDetailsView();
+
+	/** This triggers a UI repopulation after undo has been called */
+	void PostUndo();
+
+	// FGCObject interface start
+	virtual void AddReferencedObjects( FReferenceCollector& Collector ) OVERRIDE;
+	// FGCObject interface end
+
+	/** When user attempts to commit the name of a track*/
+	bool OnVerifyNotifyNameCommit( const FText& NewName, FText& OutErrorMessage, TSharedPtr<FDisplayedSlotGroupNameInfo> Item );
+
+	/** When user commits the name of a track*/
+	void OnNotifyNameCommitted( const FText& NewName, ETextCommit::Type, TSharedPtr<FDisplayedSlotGroupNameInfo> Item );
+
+	/** Dummy handler to stop editable text boxes swallowing our list selected events */
+	bool IsSelected(){return false;}
+private:
+
+	/** Called when the user changes the contents of the search box */
+	void OnFilterTextChanged( const FText& SearchText );
+
+	/** Called when the user changes the contents of the search box */
+	void OnFilterTextCommitted( const FText& SearchText, ETextCommit::Type CommitInfo );
+
+	/** Delegate handler for generating rows in SlotGroupNameListView */ 
+	TSharedRef<ITableRow> GenerateNotifyRow( TSharedPtr<FDisplayedSlotGroupNameInfo> InInfo, const TSharedRef<STableViewBase>& OwnerTable );
+
+	/** Delegate handler called when the user right clicks in SlotGroupNameListView */
+	TSharedPtr<SWidget> OnGetContextMenuContent() const;
+
+	/** Delegate handler for when the user selects something in SlotGroupNameListView */
+	void OnNotifySelectionChanged( TSharedPtr<FDisplayedSlotGroupNameInfo> Selection, ESelectInfo::Type SelectInfo );
+
+	/** Delegate handler for determining whether we can show the delete menu options */
+	bool CanPerformDelete() const;
+
+	/** Delegate handler for deleting anim SlotGroupName */
+	void OnDeleteSlotGroupName();
+
+	/** Delegate handler for determining whether we can show the rename menu options */
+	bool CanPerformRename() const;
+
+	/** Delegate handler for renaming anim SlotGroupName */
+	void OnRenameSlotGroupName();
+
+	/** Wrapper that populates SlotGroupNameListView using current filter test */
+	void RefreshSlotGroupNameListWithFilter();
+
+	/** Populates SlotGroupNameListView based on the skeletons SlotGroupName and the supplied filter text */
+	void CreateSlotGroupNameList( const FString& SearchText = FString("") );
+
+	/** handler for user selecting a Notify in SlotGroupNameListView - populates the details panel */
+	void ShowNotifyInDetailsView( FName NotifyName );
+
+	/** Populates OutAssets with the AnimSequences that match Personas current skeleton */
+	void GetCompatibleAnimMontage( TArray<class FAssetData>& OutAssets );
+
+//	void GetCompatibleAnimBlueprints( TMultiMap<class UAnimBlueprint * , class UAnimGraphNode_Slot *>& OutAssets );
+
+	/** Utility function to display notifications to the user */
+	void NotifyUser( FNotificationInfo& NotificationInfo );
+
+	/** The owning Persona instance */
+	TWeakPtr<class FPersona> PersonaPtr;
+
+	/** The skeleton we are currently editing */
+	USkeleton* TargetSkeleton;
+
+	/** SSearchBox to filter the notify list */
+	TSharedPtr<SSearchBox>	NameFilterBox;
+
+	/** Widget used to display the list of SlotGroupName */
+	TSharedPtr<SSlotGroupNameListType> SlotGroupNameListView;
+
+	/** A list of SlotGroupName. Used by the SlotGroupNameListView. */
+	TArray< TSharedPtr<FDisplayedSlotGroupNameInfo> > NotifyList;
+
+	/** Current text typed into NameFilterBox */
+	FText FilterText;
+
+	/** Tracks objects created for displaying in the details panel*/
+	FEditorObjectTracker EditorObjectTracker;
+};
+
+#undef LOCTEXT_NAMESPACE
\ No newline at end of file
diff --git a/Engine/Source/Editor/Persona/Private/SSkeletonSlotNames.cpp b/Engine/Source/Editor/Persona/Private/SSkeletonSlotNames.cpp
new file mode 100644
--- /dev/null
+++ b/Engine/Source/Editor/Persona/Private/SSkeletonSlotNames.cpp
@@ -0,0 +1,494 @@
+// Copyright 1998-2014 Epic Games, Inc. All Rights Reserved.
+
+
+#include "PersonaPrivatePCH.h"
+
+#include "Persona.h"
+#include "SSkeletonSlotNames.h"
+#include "AssetRegistryModule.h"
+#include "ScopedTransaction.h"
+
+#define LOCTEXT_NAMESPACE "SkeletonSlotNames"
+
+static const FName ColumnId_SlotNameLabel( "SlotName" );
+
+//////////////////////////////////////////////////////////////////////////
+// SMorphTargetListRow
+
+typedef TSharedPtr< FDisplayedSlotNameInfo > FDisplayedSlotNameInfoPtr;
+
+class SSlotNameListRow
+	: public SMultiColumnTableRow< FDisplayedSlotNameInfoPtr >
+{
+public:
+
+	SLATE_BEGIN_ARGS( SSlotNameListRow ) {}
+
+	/** The item for this row **/
+	SLATE_ARGUMENT( FDisplayedSlotNameInfoPtr, Item )
+
+	/* Widget used to display the list of morph targets */
+	SLATE_ARGUMENT( TSharedPtr<SSkeletonSlotNames>, SlotNameListView )
+
+	/* Persona used to update the viewport when a weight slider is dragged */
+	SLATE_ARGUMENT( TWeakPtr<FPersona>, Persona )
+
+	SLATE_END_ARGS()
+
+	void Construct( const FArguments& InArgs, const TSharedRef<STableViewBase>& OwnerTableView );
+
+	/** Overridden from SMultiColumnTableRow.  Generates a widget for this column of the tree row. */
+	virtual TSharedRef<SWidget> GenerateWidgetForColumn( const FName& ColumnName ) OVERRIDE;
+
+private:
+
+	/** Widget used to display the list of slot name */
+	TSharedPtr<SSkeletonSlotNames> SlotNameListView;
+
+	/** The notify being displayed by this row */
+	FDisplayedSlotNameInfoPtr	Item;
+
+	/** Pointer back to the Persona that owns us */
+	TWeakPtr<FPersona> PersonaPtr;
+};
+
+void SSlotNameListRow::Construct( const FArguments& InArgs, const TSharedRef<STableViewBase>& InOwnerTableView )
+{
+	Item = InArgs._Item;
+	SlotNameListView = InArgs._SlotNameListView;
+	PersonaPtr = InArgs._Persona;
+
+	check( Item.IsValid() );
+
+	SMultiColumnTableRow< FDisplayedSlotNameInfoPtr >::Construct( FSuperRowType::FArguments(), InOwnerTableView );
+}
+
+TSharedRef< SWidget > SSlotNameListRow::GenerateWidgetForColumn( const FName& ColumnName )
+{
+	check( ColumnName == ColumnId_SlotNameLabel );
+
+	return SNew( SVerticalBox )
+	+ SVerticalBox::Slot()
+	.AutoHeight()
+	.Padding( 0.0f, 4.0f )
+	.VAlign( VAlign_Center )
+	[
+		SAssignNew(Item->InlineEditableText, SInlineEditableTextBlock)
+		.Text( FText::FromName(Item->Name) )
+		.OnVerifyTextChanged(SlotNameListView.Get(), &SSkeletonSlotNames::OnVerifyNotifyNameCommit, Item)
+		.OnTextCommitted(SlotNameListView.Get(), &SSkeletonSlotNames::OnNotifyNameCommitted, Item)
+		.IsSelected(SlotNameListView.Get(), &SSkeletonSlotNames::IsSelected)
+	];
+}
+
+/////////////////////////////////////////////////////
+// FSkeletonSlotNamesSummoner
+
+FSkeletonSlotNamesSummoner::FSkeletonSlotNamesSummoner(TSharedPtr<class FAssetEditorToolkit> InHostingApp)
+	: FWorkflowTabFactory(FPersonaTabs::SkeletonSlotNamesID, InHostingApp)
+{
+	TabLabel = LOCTEXT("SkeletonSlotNamesTabTitle", "Animation SlotName");
+
+	EnableTabPadding();
+	bIsSingleton = true;
+
+	ViewMenuDescription = LOCTEXT("SkeletonSlotNamesMenu", "Animation SlotName");
+	ViewMenuTooltip = LOCTEXT("SkeletonSlotNames_ToolTip", "Shows the skeletons slot name list");
+}
+
+TSharedRef<SWidget> FSkeletonSlotNamesSummoner::CreateTabBody(const FWorkflowTabSpawnInfo& Info) const
+{
+	return SNew(SSkeletonSlotNames)
+		.Persona(StaticCastSharedPtr<FPersona>(HostingApp.Pin()));
+}
+
+/////////////////////////////////////////////////////
+// SSkeletonSlotNames
+
+void SSkeletonSlotNames::Construct(const FArguments& InArgs)
+{
+	PersonaPtr = InArgs._Persona;
+	TargetSkeleton = PersonaPtr.Pin()->GetSkeleton();
+
+	PersonaPtr.Pin()->RegisterOnPostUndo(FPersona::FOnPostUndo::CreateSP( this, &SSkeletonSlotNames::PostUndo ) );
+
+	this->ChildSlot
+	[
+		SNew( SVerticalBox )
+
+		+ SVerticalBox::Slot()
+		.AutoHeight()
+		.Padding( FMargin( 0.0f, 0.0f, 0.0f, 4.0f ) )
+		[
+			SAssignNew( NameFilterBox, SSearchBox )
+			.SelectAllTextWhenFocused( true )
+			.OnTextChanged( this, &SSkeletonSlotNames::OnFilterTextChanged )
+			.OnTextCommitted( this, &SSkeletonSlotNames::OnFilterTextCommitted )
+			.HintText( LOCTEXT( "SlotNameSearchBoxHint", "Search Animation SlotName...") )
+		]
+
+		+ SVerticalBox::Slot()
+		.FillHeight( 1.0f )		// This is required to make the scrollbar work, as content overflows Slate containers by default
+		[
+			SAssignNew( SlotNameListView, SSlotNameListType )
+			.ListItemsSource( &NotifyList )
+			.OnGenerateRow( this, &SSkeletonSlotNames::GenerateNotifyRow )
+			.OnContextMenuOpening( this, &SSkeletonSlotNames::OnGetContextMenuContent )
+			.OnSelectionChanged( this, &SSkeletonSlotNames::OnNotifySelectionChanged )
+			.ItemHeight( 22.0f )
+			.HeaderRow
+			(
+				SNew( SHeaderRow )
+				+ SHeaderRow::Column( ColumnId_SlotNameLabel )
+				.DefaultLabel( LOCTEXT( "SlotNameNameLabel", "Slot Name" ).ToString() )
+			)
+		]
+	];
+
+	CreateSlotNameList();
+}
+
+SSkeletonSlotNames::~SSkeletonSlotNames()
+{
+	if ( PersonaPtr.IsValid() )
+	{
+		PersonaPtr.Pin()->UnregisterOnPostUndo( this );
+	}
+}
+
+void SSkeletonSlotNames::OnFilterTextChanged( const FText& SearchText )
+{
+	FilterText = SearchText;
+
+	RefreshSlotNameListWithFilter();
+}
+
+void SSkeletonSlotNames::OnFilterTextCommitted( const FText& SearchText, ETextCommit::Type CommitInfo )
+{
+	// Just do the same as if the user typed in the box
+	OnFilterTextChanged( SearchText );
+}
+
+TSharedRef<ITableRow> SSkeletonSlotNames::GenerateNotifyRow(TSharedPtr<FDisplayedSlotNameInfo> InInfo, const TSharedRef<STableViewBase>& OwnerTable)
+{
+	check( InInfo.IsValid() );
+
+	return
+		SNew( SSlotNameListRow, OwnerTable )
+		.Persona( PersonaPtr )
+		.Item( InInfo )
+		.SlotNameListView( SharedThis( this ) );
+}
+
+TSharedPtr<SWidget> SSkeletonSlotNames::OnGetContextMenuContent() const
+{
+	const bool bShouldCloseWindowAfterMenuSelection = true;
+	FMenuBuilder MenuBuilder( bShouldCloseWindowAfterMenuSelection, NULL);
+
+	MenuBuilder.BeginSection("SlotNameAction", LOCTEXT( "SlotNameActions", "Selected Notify Actions" ) );
+	{
+		{
+			FUIAction Action = FUIAction( FExecuteAction::CreateSP( this, &SSkeletonSlotNames::OnDeleteSlotName ), 
+				FCanExecuteAction::CreateSP( this, &SSkeletonSlotNames::CanPerformDelete ) );
+			const FText Label = LOCTEXT("DeleteSlotNameButtonLabel", "Delete");
+			const FText ToolTip = LOCTEXT("DeleteSlotNameButtonTooltip", "Deletes the selected anim SlotName.");
+			MenuBuilder.AddMenuEntry( Label, ToolTip, FSlateIcon(), Action);
+		}
+
+		// for now disable rename, it doesn't work correctly
+// 		{
+// 			FUIAction Action = FUIAction( FExecuteAction::CreateSP( this, &SSkeletonSlotNames::OnRenameSlotName ), 
+// 				FCanExecuteAction::CreateSP( this, &SSkeletonSlotNames::CanPerformRename ) );
+// 			const FText Label = LOCTEXT("RenameSlotNameButtonLabel", "Rename");
+// 			const FText ToolTip = LOCTEXT("RenameSlotNameButtonTooltip", "Renames the selected anim SlotName.");
+// 			MenuBuilder.AddMenuEntry( Label, ToolTip, FSlateIcon(), Action);
+// 		}
+	}
+	MenuBuilder.EndSection();
+
+	return MenuBuilder.MakeWidget();
+}
+
+void SSkeletonSlotNames::OnNotifySelectionChanged(TSharedPtr<FDisplayedSlotNameInfo> Selection, ESelectInfo::Type SelectInfo)
+{
+	if(Selection.IsValid())
+	{
+		ShowNotifyInDetailsView(Selection->Name);
+	}
+}
+
+bool SSkeletonSlotNames::CanPerformDelete() const
+{
+	TArray< TSharedPtr< FDisplayedSlotNameInfo > > SelectedRows = SlotNameListView->GetSelectedItems();
+	return SelectedRows.Num() > 0;
+}
+
+bool SSkeletonSlotNames::CanPerformRename() const
+{
+	TArray< TSharedPtr< FDisplayedSlotNameInfo > > SelectedRows = SlotNameListView->GetSelectedItems();
+	return SelectedRows.Num() == 1;
+}
+
+void SSkeletonSlotNames::OnDeleteSlotName()
+{
+	TArray< TSharedPtr< FDisplayedSlotNameInfo > > SelectedRows = SlotNameListView->GetSelectedItems();
+
+	const FScopedTransaction Transaction( LOCTEXT("DeleteSlotName", "Delete Slot Name") );
+
+	// this one deletes all SlotName with same name. 
+	TArray<FName> SelectedSlotNames;
+	TargetSkeleton->Modify();
+
+	for(int Selection = 0; Selection < SelectedRows.Num(); ++Selection)
+	{
+		FName& SelectedName = SelectedRows[Selection]->Name;
+		TargetSkeleton->RemoveSlotNodeName(SelectedName);
+		SelectedSlotNames.Add(SelectedName);
+	}
+
+	TArray<FAssetData> CompatibleAnimMontages;
+	GetCompatibleAnimMontage(CompatibleAnimMontages);
+
+	int32 NumAnimationsModified = 0;
+
+	for( int32 AssetIndex = 0; AssetIndex < CompatibleAnimMontages.Num(); ++AssetIndex )
+	{
+		const FAssetData& PossibleAnimAsset = CompatibleAnimMontages[AssetIndex];
+		UAnimMontage* Montage = CastChecked<UAnimMontage>(PossibleAnimAsset.GetAsset());
+	
+		for (auto & SlotAnim : Montage->SlotAnimTracks)
+		{
+			if ( SelectedSlotNames.Contains(SlotAnim.SlotName) )
+			{
+				// found one, set transaction mark
+				Montage->Modify();
+
+				// if this name is deleted, just clear it
+				// I don't think changing to something else will help
+				SlotAnim.SlotName = NAME_None;
+				// although people shouldn't have same slot name in the same montage, 
+				// just in case anybody does, it will loop
+				// this might cause issue because Modify will get called again.
+
+				++NumAnimationsModified;
+			}
+		}
+	}
+
+// 	TMultiMap<class UAnimBlueprint *, class UAnimGraphNode_Slot *> AssetsToFix;
+// 	GetCompatibleAnimBlueprints( AssetsToFix );
+
+	if(NumAnimationsModified > 0)
+	{
+		// Tell the user that the socket is a duplicate
+		FFormatNamedArguments Args;
+		Args.Add( TEXT("NumAnimationsModified"), NumAnimationsModified );
+		FNotificationInfo Info( FText::Format( LOCTEXT( "SlotNamesDeleted", "{NumAnimationsModified} animation(s) modified to delete notifications" ), Args ) );
+
+		Info.bUseLargeFont = false;
+		Info.ExpireDuration = 5.0f;
+
+		NotifyUser( Info );
+	}
+
+	CreateSlotNameList( NameFilterBox->GetText().ToString() );
+}
+
+// void SSkeletonSlotNames::GetCompatibleAnimBlueprints( TMultiMap<class UAnimBlueprint *, class UAnimGraphNode_Slot *>& OutAssets )
+// {
+// 	// second, go through all animgraphnode_slot
+// 	for(FObjectIterator Iter(UAnimGraphNode_Slot::StaticClass()); Iter; ++Iter)
+// 	{
+// 		// if I find one, add the BP and slot nodes to the list
+// 		UAnimGraphNode_Slot * SlotNode = CastChecked<UAnimGraphNode_Slot>(Iter);
+// 		UAnimBlueprint * AnimBlueprint = Cast<UAnimBlueprint>(SlotNode->GetBlueprint());
+// 		if (AnimBlueprint)
+// 		{
+// 			OutAssets.AddUnique(AnimBlueprint, SlotNode);
+// 		}
+// 	}
+// }
+
+void SSkeletonSlotNames::OnRenameSlotName()
+{
+	TArray< TSharedPtr< FDisplayedSlotNameInfo > > SelectedRows = SlotNameListView->GetSelectedItems();
+
+	check(SelectedRows.Num() == 1); // Should be guaranteed by CanPerformRename
+
+	SelectedRows[0]->InlineEditableText->EnterEditingMode();
+}
+
+bool SSkeletonSlotNames::OnVerifyNotifyNameCommit( const FText& NewName, FText& OutErrorMessage, TSharedPtr<FDisplayedSlotNameInfo> Item )
+{
+	bool bValid(true);
+
+	if(NewName.IsEmpty())
+	{
+		OutErrorMessage = LOCTEXT( "NameMissing_Error", "You must provide a name." );
+		bValid = false;
+	}
+
+	FName NotifyName( *NewName.ToString() );
+	if(NotifyName != Item->Name)
+	{
+		if(TargetSkeleton->DoesHaveSlotNodeName(NotifyName))
+		{
+			OutErrorMessage = FText::Format( LOCTEXT("AlreadyInUseMessage", "'{0}' is already in use."), NewName );
+			bValid = false;
+		}
+	}
+
+	return bValid;
+}
+
+void SSkeletonSlotNames::OnNotifyNameCommitted( const FText& NewName, ETextCommit::Type, TSharedPtr<FDisplayedSlotNameInfo> Item )
+{
+	const FScopedTransaction Transaction( LOCTEXT("RenameSlotName", "Rename Anim Notify") );
+
+	FName NewSlotName = FName( *NewName.ToString() );
+	FName SlotToRename = Item->Name;
+
+	// rename notify in skeleton
+	TargetSkeleton->Modify();
+
+	TArray<FAssetData> CompatibleAnimMontages;
+	GetCompatibleAnimMontage(CompatibleAnimMontages);
+
+	int32 NumAnimationsModified = 0;
+
+	for(int32 AssetIndex = 0; AssetIndex < CompatibleAnimMontages.Num(); ++AssetIndex)
+	{
+		const FAssetData& PossibleAnimAsset = CompatibleAnimMontages[AssetIndex];
+		UAnimMontage* Montage = Cast<UAnimMontage>(PossibleAnimAsset.GetAsset());
+
+		for (auto & SlotAnim : Montage->SlotAnimTracks)
+		{
+			if(SlotAnim.SlotName == SlotToRename)
+			{
+				// found one, set transaction mark
+				Montage->Modify();
+
+				// if this name is deleted, just clear it
+				// I don't think changing to something else will help
+				SlotAnim.SlotName = NewSlotName;
+				// although people shouldn't have same slot name in the same montage, 
+				// just in case anybody does, it will loop
+				// this might cause issue because Modify will get called again.
+
+				++NumAnimationsModified;
+			}
+		}
+	}
+
+	if(NumAnimationsModified > 0)
+	{
+		// Tell the user that the socket is a duplicate
+		FFormatNamedArguments Args;
+		Args.Add( TEXT("NumAnimationsModified"), NumAnimationsModified );
+		FNotificationInfo Info( FText::Format( LOCTEXT( "SlotNamesRenamed", "{NumAnimationsModified} animation(s) modified to rename notification" ), Args ) );
+
+		Info.bUseLargeFont = false;
+		Info.ExpireDuration = 5.0f;
+
+		NotifyUser( Info );
+	}
+
+	RefreshSlotNameListWithFilter();
+}
+
+void SSkeletonSlotNames::RefreshSlotNameListWithFilter()
+{
+	CreateSlotNameList( NameFilterBox->GetText().ToString() );
+}
+
+void SSkeletonSlotNames::CreateSlotNameList( const FString& SearchText )
+{
+	NotifyList.Empty();
+
+	if ( TargetSkeleton )
+	{
+		const TArray<FName> & SlotNodeNames = TargetSkeleton->GetSlotNodeNames();
+		for(int i = 0; i < SlotNodeNames.Num(); ++i)
+		{
+			const FName& NotifyName = SlotNodeNames[i];
+			if ( !SearchText.IsEmpty() )
+			{
+				if ( NotifyName.ToString().Contains( SearchText ) )
+				{
+					NotifyList.Add( FDisplayedSlotNameInfo::Make( NotifyName ) );
+				}
+			}
+			else
+			{
+				NotifyList.Add( FDisplayedSlotNameInfo::Make( NotifyName ) );
+			}
+		}
+	}
+
+	SlotNameListView->RequestListRefresh();
+}
+
+void SSkeletonSlotNames::ShowNotifyInDetailsView(FName NotifyName)
+{
+	// @todo nothing to show now, but in the future
+	// we can show the list of montage that are used by this slot node?
+}
+
+void SSkeletonSlotNames::GetCompatibleAnimMontage(TArray<class FAssetData>& OutAssets)
+{
+	//Get the skeleton tag to search for
+	FString SkeletonExportName = FAssetData(TargetSkeleton).GetExportTextName();
+
+	// Load the asset registry module
+	FAssetRegistryModule& AssetRegistryModule = FModuleManager::LoadModuleChecked<FAssetRegistryModule>(TEXT("AssetRegistry"));
+
+	TArray<FAssetData> AssetDataList;
+	AssetRegistryModule.Get().GetAssetsByClass(UAnimMontage::StaticClass()->GetFName(), AssetDataList, true);
+
+	OutAssets.Empty(AssetDataList.Num());
+
+	for( int32 AssetIndex = 0; AssetIndex < AssetDataList.Num(); ++AssetIndex )
+	{
+		const FAssetData& PossibleAnimMontage = AssetDataList[AssetIndex];
+		if( SkeletonExportName == PossibleAnimMontage.TagsAndValues.FindRef("Skeleton") )
+		{
+			OutAssets.Add( PossibleAnimMontage );
+		}
+	}
+}
+
+UObject* SSkeletonSlotNames::ShowInDetailsView( UClass* EdClass )
+{
+	UObject *Obj = EditorObjectTracker.GetEditorObjectForClass(EdClass);
+
+	if(Obj != NULL)
+	{
+		PersonaPtr.Pin()->SetDetailObject(Obj);
+	}
+	return Obj;
+}
+
+void SSkeletonSlotNames::ClearDetailsView()
+{
+	PersonaPtr.Pin()->SetDetailObject(NULL);
+}
+
+void SSkeletonSlotNames::PostUndo()
+{
+	RefreshSlotNameListWithFilter();
+}
+
+void SSkeletonSlotNames::AddReferencedObjects( FReferenceCollector& Collector )
+{
+	EditorObjectTracker.AddReferencedObjects(Collector);
+}
+
+void SSkeletonSlotNames::NotifyUser( FNotificationInfo& NotificationInfo )
+{
+	TSharedPtr<SNotificationItem> Notification = FSlateNotificationManager::Get().AddNotification( NotificationInfo );
+	if ( Notification.IsValid() )
+	{
+		Notification->SetCompletionState( SNotificationItem::CS_Fail );
+	}
+}
+#undef LOCTEXT_NAMESPACE
diff --git a/Engine/Source/Editor/Persona/Private/SSkeletonSlotNames.h b/Engine/Source/Editor/Persona/Private/SSkeletonSlotNames.h
new file mode 100644
--- /dev/null
+++ b/Engine/Source/Editor/Persona/Private/SSkeletonSlotNames.h
@@ -0,0 +1,165 @@
+// Copyright 1998-2014 Epic Games, Inc. All Rights Reserved.
+
+#pragma once
+
+#include "EditorObjectsTracker.h"
+#include "IDocumentation.h"
+
+#define LOCTEXT_NAMESPACE "SkeletonSlotNames"
+/////////////////////////////////////////////////////
+// FSkeletonTreeSummoner
+struct FSkeletonSlotNamesSummoner : public FWorkflowTabFactory
+{
+public:
+	FSkeletonSlotNamesSummoner(TSharedPtr<class FAssetEditorToolkit> InHostingApp);
+
+	virtual TSharedRef<SWidget> CreateTabBody(const FWorkflowTabSpawnInfo& Info) const OVERRIDE;
+
+	// Create a tooltip widget for the tab
+	virtual TSharedPtr<SToolTip> CreateTabToolTipWidget(const FWorkflowTabSpawnInfo& Info) const OVERRIDE
+	{
+		return  IDocumentation::Get()->CreateToolTip(LOCTEXT("WindowTooltip", "This tab lets you modify custom animation SlotName"), NULL, TEXT("Shared/Editors/Persona"), TEXT("AnimationSlotName_Window"));
+	}
+};
+
+//////////////////////////////////////////////////////////////////////////
+// FDisplayedSlotNameInfo
+
+class FDisplayedSlotNameInfo
+{
+public:
+	FName Name;
+
+	/** Handle to editable text block for rename */
+	TSharedPtr<SInlineEditableTextBlock> InlineEditableText;
+
+	/** Static function for creating a new item, but ensures that you can only have a TSharedRef to one */
+	static TSharedRef<FDisplayedSlotNameInfo> Make(const FName& NotifyName)
+	{
+		return MakeShareable(new FDisplayedSlotNameInfo(NotifyName));
+	}
+
+protected:
+	/** Hidden constructor, always use Make above */
+	FDisplayedSlotNameInfo(const FName& InNotifyName)
+		: Name( InNotifyName )
+	{}
+
+	/** Hidden constructor, always use Make above */
+	FDisplayedSlotNameInfo() {}
+};
+
+/** Widgets list type */
+typedef SListView< TSharedPtr<FDisplayedSlotNameInfo> > SSlotNameListType;
+
+class SSkeletonSlotNames : public SCompoundWidget, public FGCObject
+{
+public:
+	SLATE_BEGIN_ARGS( SSkeletonSlotNames )
+		: _Persona()
+	{}
+
+	SLATE_ARGUMENT( TWeakPtr<FPersona>, Persona )
+		
+		SLATE_END_ARGS()
+public:
+	void Construct(const FArguments& InArgs);
+
+	~SSkeletonSlotNames();
+
+	/**
+	* Accessor so our rows can grab the filter text for highlighting
+	*
+	*/
+	FText& GetFilterText() { return FilterText; }
+
+	/** Creates an editor object from the given type to be used in a details panel */
+	UObject* ShowInDetailsView( UClass* EdClass );
+
+	/** Clears the detail view of whatever we displayed last */
+	void ClearDetailsView();
+
+	/** This triggers a UI repopulation after undo has been called */
+	void PostUndo();
+
+	// FGCObject interface start
+	virtual void AddReferencedObjects( FReferenceCollector& Collector ) OVERRIDE;
+	// FGCObject interface end
+
+	/** When user attempts to commit the name of a track*/
+	bool OnVerifyNotifyNameCommit( const FText& NewName, FText& OutErrorMessage, TSharedPtr<FDisplayedSlotNameInfo> Item );
+
+	/** When user commits the name of a track*/
+	void OnNotifyNameCommitted( const FText& NewName, ETextCommit::Type, TSharedPtr<FDisplayedSlotNameInfo> Item );
+
+	/** Dummy handler to stop editable text boxes swallowing our list selected events */
+	bool IsSelected(){return false;}
+private:
+
+	/** Called when the user changes the contents of the search box */
+	void OnFilterTextChanged( const FText& SearchText );
+
+	/** Called when the user changes the contents of the search box */
+	void OnFilterTextCommitted( const FText& SearchText, ETextCommit::Type CommitInfo );
+
+	/** Delegate handler for generating rows in SlotNameListView */ 
+	TSharedRef<ITableRow> GenerateNotifyRow( TSharedPtr<FDisplayedSlotNameInfo> InInfo, const TSharedRef<STableViewBase>& OwnerTable );
+
+	/** Delegate handler called when the user right clicks in SlotNameListView */
+	TSharedPtr<SWidget> OnGetContextMenuContent() const;
+
+	/** Delegate handler for when the user selects something in SlotNameListView */
+	void OnNotifySelectionChanged( TSharedPtr<FDisplayedSlotNameInfo> Selection, ESelectInfo::Type SelectInfo );
+
+	/** Delegate handler for determining whether we can show the delete menu options */
+	bool CanPerformDelete() const;
+
+	/** Delegate handler for deleting anim SlotName */
+	void OnDeleteSlotName();
+
+	/** Delegate handler for determining whether we can show the rename menu options */
+	bool CanPerformRename() const;
+
+	/** Delegate handler for renaming anim SlotName */
+	void OnRenameSlotName();
+
+	/** Wrapper that populates SlotNameListView using current filter test */
+	void RefreshSlotNameListWithFilter();
+
+	/** Populates SlotNameListView based on the skeletons SlotName and the supplied filter text */
+	void CreateSlotNameList( const FString& SearchText = FString("") );
+
+	/** handler for user selecting a Notify in SlotNameListView - populates the details panel */
+	void ShowNotifyInDetailsView( FName NotifyName );
+
+	/** Populates OutAssets with the AnimSequences that match Personas current skeleton */
+	void GetCompatibleAnimMontage( TArray<class FAssetData>& OutAssets );
+
+//	void GetCompatibleAnimBlueprints( TMultiMap<class UAnimBlueprint * , class UAnimGraphNode_Slot *>& OutAssets );
+
+	/** Utility function to display notifications to the user */
+	void NotifyUser( FNotificationInfo& NotificationInfo );
+
+	/** The owning Persona instance */
+	TWeakPtr<class FPersona> PersonaPtr;
+
+	/** The skeleton we are currently editing */
+	USkeleton* TargetSkeleton;
+
+	/** SSearchBox to filter the notify list */
+	TSharedPtr<SSearchBox>	NameFilterBox;
+
+	/** Widget used to display the list of SlotName */
+	TSharedPtr<SSlotNameListType> SlotNameListView;
+
+	/** A list of SlotName. Used by the SlotNameListView. */
+	TArray< TSharedPtr<FDisplayedSlotNameInfo> > NotifyList;
+
+	/** Current text typed into NameFilterBox */
+	FText FilterText;
+
+	/** Tracks objects created for displaying in the details panel*/
+	FEditorObjectTracker EditorObjectTracker;
+};
+
+#undef LOCTEXT_NAMESPACE
\ No newline at end of file
diff --git a/Engine/Source/Editor/Persona/Private/Shared/PersonaMode.cpp b/Engine/Source/Editor/Persona/Private/Shared/PersonaMode.cpp
--- a/Engine/Source/Editor/Persona/Private/Shared/PersonaMode.cpp
+++ b/Engine/Source/Editor/Persona/Private/Shared/PersonaMode.cpp
@@ -6,6 +6,7 @@
 #include "SSkeletonAnimNotifies.h"
 #include "IDocumentation.h"
 #include "SAnimBlueprintParentPlayerList.h"
+#include "SSkeletonSlotGroupNames.h"
 
 #define LOCTEXT_NAMESPACE "PersonaModes"
 
@@ -35,6 +36,8 @@ const FName FPersonaTabs::AnimAssetPropertiesID("AnimAssetPropertiesTab");
 const FName FPersonaTabs::MeshAssetPropertiesID("MeshAssetPropertiesTab");
 const FName FPersonaTabs::PreviewManagerID("AnimPreviewSetup");		//@TODO: Name
 const FName FPersonaTabs::SkeletonAnimNotifiesID("SkeletonAnimNotifies");
+const FName FPersonaTabs::SkeletonSlotNamesID("SkeletonSlotNames");
+const FName FPersonaTabs::SkeletonSlotGroupNamesID("SkeletonSlotGroupNames");
 
 /////////////////////////////////////////////////////
 // FPersonaMode
@@ -60,6 +63,8 @@ FPersonaAppMode::FPersonaAppMode(TSharedPtr<class FPersona> InPersona, FName InM
 	PersonaTabFactories.RegisterFactory(MakeShareable(new FMorphTargetTabSummoner(InPersona)));
 	PersonaTabFactories.RegisterFactory(MakeShareable(new FSkeletonAnimNotifiesSummoner(InPersona)));
 	PersonaTabFactories.RegisterFactory(MakeShareable(new FRetargetSourceManagerTabSummoner(InPersona)));
+	PersonaTabFactories.RegisterFactory(MakeShareable(new FSkeletonSlotNamesSummoner(InPersona)));
+	PersonaTabFactories.RegisterFactory(MakeShareable(new FSkeletonSlotGroupNamesSummoner(InPersona)));
 }
 
 void FPersonaAppMode::RegisterTabFactories(TSharedPtr<FTabManager> InTabManager)
diff --git a/Engine/Source/Editor/Persona/Private/Shared/PersonaMode.h b/Engine/Source/Editor/Persona/Private/Shared/PersonaMode.h
--- a/Engine/Source/Editor/Persona/Private/Shared/PersonaMode.h
+++ b/Engine/Source/Editor/Persona/Private/Shared/PersonaMode.h
@@ -34,6 +34,8 @@ struct FPersonaTabs
 	static const FName MeshAssetPropertiesID;
 	static const FName PreviewManagerID;
 	static const FName SkeletonAnimNotifiesID;
+	static const FName SkeletonSlotNamesID;
+	static const FName SkeletonSlotGroupNamesID;
 
 	// Blueprint Document
 
diff --git a/Engine/Source/Editor/UMGEditor/Private/WidgetGraphNode_Base.cpp b/Engine/Source/Editor/UMGEditor/Private/WidgetGraphNode_Base.cpp
--- a/Engine/Source/Editor/UMGEditor/Private/WidgetGraphNode_Base.cpp
+++ b/Engine/Source/Editor/UMGEditor/Private/WidgetGraphNode_Base.cpp
@@ -51,6 +51,8 @@ public:
 
 		Record.bCanToggleVisibility = bOptional_ShowByDefault || bOptional_HideByDefault;
 		Record.bShowPin = bAlwaysShow || bOptional_ShowByDefault;
+		// this was for anim node graph, not sure if widget graph node would need it
+		Record.bPropertyIsCustomized = false;
 	}
 
 	virtual void CustomizePinData(UEdGraphPin* Pin, FName SourcePropertyName, int32 ArrayIndex, UProperty* Property) const OVERRIDE
diff --git a/Engine/Source/Runtime/Engine/Classes/Animation/AnimNode_Slot.h b/Engine/Source/Runtime/Engine/Classes/Animation/AnimNode_Slot.h
--- a/Engine/Source/Runtime/Engine/Classes/Animation/AnimNode_Slot.h
+++ b/Engine/Source/Runtime/Engine/Classes/Animation/AnimNode_Slot.h
@@ -17,9 +17,14 @@ struct ENGINE_API FAnimNode_Slot : public FAnimNode_Base
 	FPoseLink Source;
 
 	// The name of this slot, exposed to gameplay code, etc...
-	UPROPERTY(EditAnywhere, BlueprintReadWrite, Category=Settings)
+	UPROPERTY(EditAnywhere, BlueprintReadWrite, Category=Settings, meta=(CustomizeProperty))
 	FName SlotName;
 
+	// The group name of this slot, exposed to gameplay code, etc...
+	// you can only play one montage per group
+	UPROPERTY(EditAnywhere, BlueprintReadWrite, Category=Settings, meta=(CustomizeProperty))
+	FName GroupName;
+
 protected:
 	/** Weight of Source Branch. This is the weight of the input pose coming from children.
 	This is different than (1.f - SourceWeight) since the Slot can play additive animations, which are overlayed on top of the source pose. */
diff --git a/Engine/Source/Runtime/Engine/Classes/Animation/Skeleton.h b/Engine/Source/Runtime/Engine/Classes/Animation/Skeleton.h
--- a/Engine/Source/Runtime/Engine/Classes/Animation/Skeleton.h
+++ b/Engine/Source/Runtime/Engine/Classes/Animation/Skeleton.h
@@ -197,6 +197,14 @@ private:
 	UPROPERTY(duplicatetransient, AssetRegistrySearchable)
 	TAssetPtr<class USkeletalMesh> PreviewSkeletalMesh;
 
+	/** @todo collection of slotnode names available for this skeleton */
+	UPROPERTY()
+	TArray<FName> SlotNodeNames;
+
+	/** @todo collection of graoup names for slot node available for this skeleton */
+	UPROPERTY()
+	TArray<FName> SlotGroupNames;
+
 public:
 
 	/** AnimNotifiers that has been created. Right now there is no delete step for this, but in the future we'll supply delete**/
@@ -253,6 +261,17 @@ public:
 	// @todo document
 	ENGINE_API void AddNewAnimationNotify(FName NewAnimNotifyName);
 
+	// @todo document
+	ENGINE_API void AddSlotNodeName(FName SlotNodeName);
+	ENGINE_API void RemoveSlotNodeName(FName SlotNodeName);
+	ENGINE_API bool DoesHaveSlotNodeName(FName SlotNodeName) const;
+	ENGINE_API const TArray<FName> & GetSlotNodeNames() const;
+
+	ENGINE_API void AddSlotGroupName(FName GroupName);
+	ENGINE_API void RemoveSlotGroupName(FName GroupName);
+	ENGINE_API bool DoesHaveSlotGroupName(FName GroupName) const;
+	ENGINE_API const TArray<FName> & GetSlotGroupNames() const;
+
 	/** Returns the skeletons preview mesh, loading it if necessary */
 	ENGINE_API USkeletalMesh* GetPreviewMesh(bool bFindIfNotSet=false);
 
@@ -530,6 +549,7 @@ public:
 
 	static const FName AnimNotifyTag;
 	static const TCHAR AnimNotifyTagDeliminator;
+	ENGINE_API static const FName DefaultSlotGroupName;
 #endif
 private:
 	void RegenerateGuid();
diff --git a/Engine/Source/Runtime/Engine/Private/Animation/AnimNode_Slot.cpp b/Engine/Source/Runtime/Engine/Private/Animation/AnimNode_Slot.cpp
--- a/Engine/Source/Runtime/Engine/Private/Animation/AnimNode_Slot.cpp
+++ b/Engine/Source/Runtime/Engine/Private/Animation/AnimNode_Slot.cpp
@@ -88,5 +88,6 @@ void FAnimNode_Slot::GatherDebugData(FNodeDebugData& DebugData)
 FAnimNode_Slot::FAnimNode_Slot()
 : SourceWeight(0.f)
 , SlotNodeWeight(0.f)
+, GroupName(USkeleton::DefaultSlotGroupName)
 {
 }
diff --git a/Engine/Source/Runtime/Engine/Private/Animation/Skeleton.cpp b/Engine/Source/Runtime/Engine/Private/Animation/Skeleton.cpp
--- a/Engine/Source/Runtime/Engine/Private/Animation/Skeleton.cpp
+++ b/Engine/Source/Runtime/Engine/Private/Animation/Skeleton.cpp
@@ -14,6 +14,7 @@
 #if WITH_EDITOR
 const FName USkeleton::AnimNotifyTag = FName(TEXT("AnimNotifyList"));
 const TCHAR USkeleton::AnimNotifyTagDeliminator = TEXT(';');
+const FName USkeleton::DefaultSlotGroupName = FName(TEXT("Default"));
 #endif 
 
 FArchive& operator<<(FArchive& Ar, FReferencePose & P)
@@ -33,6 +34,8 @@ FArchive& operator<<(FArchive& Ar, FReferencePose & P)
 USkeleton::USkeleton(const class FPostConstructInitializeProperties& PCIP)
 	: Super(PCIP)
 {
+	SlotGroupNames.Empty();
+	SlotGroupNames.Add(DefaultSlotGroupName);
 }
 
 void USkeleton::PostInitProperties()
@@ -856,7 +859,50 @@ void USkeleton::UnregisterOnSkeletonHierarchyChanged(void * Unregister)
 	OnSkeletonHierarchyChanged.RemoveAll(Unregister);
 }
 
-#endif
+void USkeleton::AddSlotNodeName(FName SlotNodeName)
+{
+	SlotNodeNames.AddUnique(SlotNodeName);
+}
+
+void USkeleton::RemoveSlotNodeName(FName SlotNodeName)
+{
+	SlotNodeNames.Remove(SlotNodeName);
+}
+
+bool USkeleton::DoesHaveSlotNodeName(FName SlotNodeName) const
+{
+	return SlotNodeNames.Contains(SlotNodeName);
+}
+
+const TArray<FName> & USkeleton::GetSlotNodeNames() const
+{
+	return SlotNodeNames;
+}
+
+void USkeleton::AddSlotGroupName(FName GroupName)
+{
+	SlotGroupNames.AddUnique(GroupName);
+}
+
+void USkeleton::RemoveSlotGroupName(FName GroupName)
+{
+	// we can't delete the deafult one
+	if (GroupName != DefaultSlotGroupName)
+	{
+		SlotGroupNames.Remove(GroupName);
+	}
+}
+
+const TArray<FName> & USkeleton::GetSlotGroupNames() const
+{
+	return SlotGroupNames;
+}
+
+bool USkeleton::DoesHaveSlotGroupName(FName GroupName) const
+{
+	return SlotGroupNames.Contains(GroupName);
+}
+#endif // WITH_EDITORONLY_DATA
 
 void USkeleton::RegenerateGuid()
 {
