Missing files from CL 2075622

diff --git a/Engine/Source/Editor/Persona/Private/Customization/AnimGraphNodeSlotDetails.cpp b/Engine/Source/Editor/Persona/Private/Customization/AnimGraphNodeSlotDetails.cpp
new file mode 100644
--- /dev/null
+++ b/Engine/Source/Editor/Persona/Private/Customization/AnimGraphNodeSlotDetails.cpp
@@ -0,0 +1,371 @@
+// Copyright 1998-2014 Epic Games, Inc. All Rights Reserved.
+
+#include "PersonaPrivatePCH.h"
+
+#include "Editor/PropertyEditor/Public/IDetailsView.h"
+#include "Editor/PropertyEditor/Public/PropertyEditing.h"
+#include "Editor/PropertyEditor/Public/PropertyEditorModule.h"
+
+#include "AnimGraphNode_Base.h"
+
+#include "ScopedTransaction.h"
+#include "AnimGraphNodeSlotDetails.h"
+
+#define LOCTEXT_NAMESPACE "AnimNodeSlotDetails"
+
+////////////////////////////////////////////////////////////////
+FAnimGraphNodeSlotDetails::FAnimGraphNodeSlotDetails(TSharedRef<class FPersona> InPersonalEditor)
+:PersonaEditor(InPersonalEditor)
+{
+}
+
+void FAnimGraphNodeSlotDetails::CustomizeDetails(class IDetailLayoutBuilder& DetailBuilder)
+{
+	SlotNodeNamePropertyHandle = DetailBuilder.GetProperty(TEXT("Node.SlotName"));
+	GroupNamePropertyHandle = DetailBuilder.GetProperty(TEXT("Node.GroupName"));
+
+	DetailBuilder.HideProperty(SlotNodeNamePropertyHandle);
+	DetailBuilder.HideProperty(GroupNamePropertyHandle);
+
+	TArray<UObject*> OuterObjects;
+	SlotNodeNamePropertyHandle->GetOuterObjects(OuterObjects);
+
+	check(SlotNodeNamePropertyHandle.IsValid());
+	check(GroupNamePropertyHandle.IsValid());
+
+	Skeleton = NULL;
+
+	// only support one object, it doesn't make sense to have same slot node name in multiple nodes
+	for (auto Object : OuterObjects)
+	{
+		// if not this object, we have problem 
+		const UAnimGraphNode_Base * Node = CastChecked<UAnimGraphNode_Base>(Object);
+		const UAnimBlueprint * AnimBlueprint = Node->GetAnimBlueprint();
+		if (AnimBlueprint)
+		{
+			// if skeleton is already set and it's not same as target skeleton, we have some problem ,return
+			if (Skeleton != NULL && Skeleton == AnimBlueprint->TargetSkeleton)
+			{
+				// abort
+				return;
+			}
+
+			Skeleton = AnimBlueprint->TargetSkeleton;
+		}
+	}
+
+	check (Skeleton);
+
+	// this is used for 2 things, to create name, but also for another pop-up box to show off, it's a bit hacky to use this, but I need widget to parent on
+	TSharedRef<SWidget> SlotNodePropertyNameWidget = SlotNodeNamePropertyHandle->CreatePropertyNameWidget();
+	TSharedRef<SWidget> GroupNodePropertyNameWidget = GroupNamePropertyHandle->CreatePropertyNameWidget();
+
+	// get the current slot node name
+	int32 CurrentlySelectedIndex = RefreshSlotNames();
+	TSharedPtr<FString> CurrentlySelectedNode;
+
+	if (CurrentlySelectedIndex != INDEX_NONE)
+	{
+		CurrentlySelectedNode = SlotNameComboList[CurrentlySelectedIndex];
+	}
+
+	// create combo box
+	IDetailCategoryBuilder& SlotNameGroup = DetailBuilder.EditCategory(TEXT("Settings"));
+	SlotNameGroup.AddCustomRow(LOCTEXT("SlotNameTitleLabel", "Slot Name").ToString())
+	.NameContent()
+	[
+		SlotNodePropertyNameWidget
+	]
+	.ValueContent()
+	[
+		SNew(SHorizontalBox)
+
+		+SHorizontalBox::Slot()
+		[
+			SAssignNew(SlotNameComboBox, STextComboBox)
+			.OptionsSource(&SlotNameComboList)
+			.OnSelectionChanged(this, &FAnimGraphNodeSlotDetails::OnSlotNameChanged)
+			.InitiallySelectedItem(CurrentlySelectedNode)
+			.ContentPadding(2)
+		]
+
+		+SHorizontalBox::Slot()
+		.AutoWidth()
+		[
+			SNew(SButton)
+			.Text(LOCTEXT("SlotName_AddButtonLabel", "Add"))
+			.ToolTipText(LOCTEXT("SlotName_AddButtonToolTipText", "New Slot Node Names") )
+			.OnClicked(this, &FAnimGraphNodeSlotDetails::OnAddSlotName, SlotNodePropertyNameWidget)
+			.Content()
+			[
+				SNew(SImage)
+				.Image(FEditorStyle::GetBrush("PropertyWindow.Button_AddToArray"))
+			]
+		]
+
+		+SHorizontalBox::Slot()
+		.AutoWidth()
+		[
+			SNew(SButton)
+			.Text(LOCTEXT("SlotName_ManageButtonLabel", "Manage"))
+			.ToolTipText(LOCTEXT("SlotName_ManageButtonToolTipText", "Manage Slot Node Names"))
+			.OnClicked(this, &FAnimGraphNodeSlotDetails::OnManageSlotName)
+			.Content()
+			[
+				SNew(SImage)
+				.Image(FEditorStyle::GetBrush("MeshPaint.FindInCB"))
+			]
+		]
+	];
+
+	// group names
+	CurrentlySelectedIndex = RefreshGroupNames();
+	CurrentlySelectedNode = nullptr;
+
+	if(CurrentlySelectedIndex != INDEX_NONE)
+	{
+		CurrentlySelectedNode = GroupNameComboList[CurrentlySelectedIndex];
+	}
+
+	//IDetailGroup& GroupNameGroup = DetailBuilder.EditCategory(TEXT("GroupName_Group"), LOCTEXT("GroupNameTitleLabel", "Group Name").ToString());
+	SlotNameGroup.AddCustomRow(LOCTEXT("GroupNameTitleLabel", "Group Name").ToString())
+	.NameContent()
+	[
+		GroupNodePropertyNameWidget
+	]
+	.ValueContent()
+	[
+		SNew(SHorizontalBox)
+
+		+SHorizontalBox::Slot()
+		[
+			SAssignNew(GroupNameComboBox, STextComboBox)
+			.OptionsSource(&GroupNameComboList)
+			.OnSelectionChanged(this, &FAnimGraphNodeSlotDetails::OnGroupNameChanged)
+			.InitiallySelectedItem(CurrentlySelectedNode)
+			.ContentPadding(2)
+		]
+
+		+SHorizontalBox::Slot()
+		.AutoWidth()
+		[
+			SNew(SButton)
+			.Text(LOCTEXT("GroupName_AddButtonLabel", "Add"))
+			.ToolTipText(LOCTEXT("GroupName_AddButtonToolTipText", "New Slot Group Names"))
+			.OnClicked(this, &FAnimGraphNodeSlotDetails::OnAddGroupName, GroupNodePropertyNameWidget)
+			.Content()
+			[
+				SNew(SImage)
+				.Image(FEditorStyle::GetBrush("PropertyWindow.Button_AddToArray"))
+			]
+		]
+
+		+SHorizontalBox::Slot()
+		.AutoWidth()
+		[
+			SNew(SButton)
+			.Text(LOCTEXT("GroupName_ManageButtonLabel", "Manage"))
+			.ToolTipText(LOCTEXT("GroupName_ManageButtonToolTipText", "Manage Slot Group Names"))
+			.OnClicked(this, &FAnimGraphNodeSlotDetails::OnManageGroupName)
+			.Content()
+			[
+				SNew(SImage)
+				.Image(FEditorStyle::GetBrush("MeshPaint.FindInCB"))
+			]
+		]
+	];
+
+}
+
+int32 FAnimGraphNodeSlotDetails::RefreshSlotNames()
+{
+	SlotNameComboList.Empty();
+
+	FName CurrentSlotNodeName = NAME_None;
+	SlotNodeNamePropertyHandle->GetValue(CurrentSlotNodeName);
+
+	// add slot node names
+	int32 CurrentSelectedIndex = INDEX_NONE;
+	const TArray<FName> & SlotNodeNames = Skeleton->GetSlotNodeNames();
+	for(auto SlotName : SlotNodeNames)
+	{
+		SlotNameComboList.Add(MakeShareable(new FString(SlotName.ToString())));
+
+		if(CurrentSlotNodeName == SlotName)
+		{
+			CurrentSelectedIndex = SlotNameComboList.Num()-1;
+		}
+	}
+
+	return CurrentSelectedIndex;
+}
+////////////////////////////////////////////////////////////////
+
+void FAnimGraphNodeSlotDetails::OnSlotNameChanged(TSharedPtr<FString> NewSelection, ESelectInfo::Type SelectInfo)
+{
+	// if it's set from code, we did that on purpose
+	if(SelectInfo != ESelectInfo::Direct)
+	{
+		FString NewValue = *NewSelection.Get();
+		const TArray<FName> & SlotNodeNames = Skeleton->GetSlotNodeNames();
+
+		for(auto SlotNodeName : SlotNodeNames )
+		{
+			if(NewValue == SlotNodeName.ToString())
+			{
+				// trigget transaction 
+				//const FScopedTransaction Transaction(LOCTEXT("ChangeSlotNodeName", "Change Collision Profile"));
+				// set profile set up
+				ensure(SlotNodeNamePropertyHandle->SetValue(NewValue) ==  FPropertyAccess::Result::Success);
+				return;
+			}
+		}
+	}
+}
+
+FReply FAnimGraphNodeSlotDetails::OnAddSlotName(TSharedRef<SWidget> Widget)
+{
+	TSharedRef<STextEntryPopup> TextEntry =
+		SNew(STextEntryPopup)
+		.Label(LOCTEXT("NewSlotName_AskSlotName", "New Slot Name"))
+		.OnTextCommitted(this, &FAnimGraphNodeSlotDetails::AddSlotName);
+
+	// Show dialog to enter new track name
+	FSlateApplication::Get().PushMenu(
+		Widget,
+		TextEntry,
+		FSlateApplication::Get().GetCursorPos(),
+		FPopupTransitionEffect(FPopupTransitionEffect::TypeInPopup)
+		);
+
+	TextEntry->FocusDefaultWidget();
+
+	return FReply::Handled();
+}
+
+void FAnimGraphNodeSlotDetails::AddSlotName(const FText & SlotNameToAdd, ETextCommit::Type CommitInfo)
+{
+	if(!SlotNameToAdd.IsEmpty())
+	{
+		const FScopedTransaction Transaction(LOCTEXT("NewSlotName_AddSlotName", "Add New Slot Node Name"));
+		// before insert, make sure everything behind is fixed
+		Skeleton->AddSlotNodeName(FName(*SlotNameToAdd.ToString()));
+		Skeleton->Modify(true);
+
+		// refresh combo box
+		int32 SelectedIndex = RefreshSlotNames();
+		if (SelectedIndex != INDEX_NONE)
+		{
+			SlotNameComboBox->SetSelectedItem(SlotNameComboList[SelectedIndex]);
+		}
+
+		// @todo add bunch of notifies
+		FSlateApplication::Get().DismissAllMenus();
+	}
+}
+
+FReply FAnimGraphNodeSlotDetails::OnManageSlotName()
+{
+	if(PersonaEditor.IsValid())
+	{
+		PersonaEditor.Pin()->GetTabManager()->InvokeTab(FPersonaTabs::SkeletonSlotNamesID);
+	}
+	return FReply::Handled();
+}
+
+void FAnimGraphNodeSlotDetails::OnGroupNameChanged(TSharedPtr<FString> NewSelection, ESelectInfo::Type SelectInfo)
+{
+	// if it's set from code, we did that on purpose
+	if(SelectInfo != ESelectInfo::Direct)
+	{
+		FString NewValue = *NewSelection.Get();
+		const TArray<FName> & SlotGroupNames = Skeleton->GetSlotGroupNames();
+
+		for(auto SlotGroupName : SlotGroupNames)
+		{
+			if(NewValue == SlotGroupName.ToString())
+			{
+				// trigget transaction 
+				//const FScopedTransaction Transaction(LOCTEXT("ChangeSlotGroupName", "Change Collision Profile"));
+				// set profile set up
+				ensure(GroupNamePropertyHandle->SetValue(NewValue) ==  FPropertyAccess::Result::Success);
+				return;
+			}
+		}
+	}
+}
+
+FReply FAnimGraphNodeSlotDetails::OnAddGroupName(TSharedRef<SWidget> Widget)
+{
+	TSharedRef<STextEntryPopup> TextEntry =
+		SNew(STextEntryPopup)
+		.Label(LOCTEXT("NewGroupName_AskGroupName", "New Group Name"))
+		.OnTextCommitted(this, &FAnimGraphNodeSlotDetails::AddGroupName);
+
+	// Show dialog to enter new track name
+	FSlateApplication::Get().PushMenu(
+		Widget,
+		TextEntry,
+		FSlateApplication::Get().GetCursorPos(),
+		FPopupTransitionEffect(FPopupTransitionEffect::TypeInPopup)
+		);
+
+	TextEntry->FocusDefaultWidget();
+
+	return FReply::Handled();
+}
+
+void FAnimGraphNodeSlotDetails::AddGroupName(const FText & GroupNameToAdd, ETextCommit::Type CommitInfo)
+{
+	if(!GroupNameToAdd.IsEmpty())
+	{
+		const FScopedTransaction Transaction(LOCTEXT("NewGroupName_AddGroupName", "Add New Slot Node Name"));
+		// before insert, make sure everything behind is fixed
+		Skeleton->AddSlotGroupName(FName(*GroupNameToAdd.ToString()));
+		Skeleton->Modify(true);
+
+		// refresh combo box
+		int32 SelectedIndex = RefreshGroupNames();
+		if(SelectedIndex != INDEX_NONE)
+		{
+			GroupNameComboBox->SetSelectedItem(GroupNameComboList[SelectedIndex]);
+		}
+
+		// @todo add bunch of notifies
+		FSlateApplication::Get().DismissAllMenus();
+	}
+}
+
+FReply FAnimGraphNodeSlotDetails::OnManageGroupName()
+{
+	if(PersonaEditor.IsValid())
+	{
+		PersonaEditor.Pin()->GetTabManager()->InvokeTab(FPersonaTabs::SkeletonSlotGroupNamesID);
+	}
+	return FReply::Handled();
+}
+
+int32 FAnimGraphNodeSlotDetails::RefreshGroupNames()
+{
+	GroupNameComboList.Empty();
+
+	FName CurrentSlotGroupName = NAME_None;
+	GroupNamePropertyHandle->GetValue(CurrentSlotGroupName);
+
+	// add slot node names
+	int32 CurrentSelectedIndex = INDEX_NONE;
+	const TArray<FName> & SlotGroupNames = Skeleton->GetSlotGroupNames();
+	for(auto GroupName : SlotGroupNames)
+	{
+		GroupNameComboList.Add(MakeShareable(new FString(GroupName.ToString())));
+
+		if(CurrentSlotGroupName == GroupName)
+		{
+			CurrentSelectedIndex = GroupNameComboList.Num()-1;
+		}
+	}
+
+	return CurrentSelectedIndex;
+}
+
+#undef LOCTEXT_NAMESPACE
diff --git a/Engine/Source/Editor/Persona/Private/Customization/AnimGraphNodeSlotDetails.h b/Engine/Source/Editor/Persona/Private/Customization/AnimGraphNodeSlotDetails.h
new file mode 100644
--- /dev/null
+++ b/Engine/Source/Editor/Persona/Private/Customization/AnimGraphNodeSlotDetails.h
@@ -0,0 +1,57 @@
+// Copyright 1998-2014 Epic Games, Inc. All Rights Reserved.
+
+#pragma once
+
+
+/**
+ * Customizes a DataTable asset to use a dropdown
+ */
+class FAnimGraphNodeSlotDetails : public IDetailCustomization
+{
+public:
+	FAnimGraphNodeSlotDetails(TSharedRef<class FPersona> PersonalEditor);
+
+	static TSharedRef<IDetailCustomization> MakeInstance(TSharedRef<class FPersona> PersonalEditor) 
+	{
+		return MakeShareable( new FAnimGraphNodeSlotDetails(PersonalEditor) );
+	}
+
+	// IDetailCustomization interface
+	virtual void CustomizeDetails(class IDetailLayoutBuilder& DetailBuilder) OVERRIDE;
+	
+private:
+	// persona editor 
+	TWeakPtr<class FPersona> PersonaEditor;
+
+	// property of the two 
+	TSharedPtr<IPropertyHandle> SlotNodeNamePropertyHandle;
+	TSharedPtr<IPropertyHandle> GroupNamePropertyHandle;
+
+	// slot node names
+	TSharedPtr< class STextComboBox>	SlotNameComboBox;
+	TArray< TSharedPtr< FString > >		SlotNameComboList;
+
+	void OnSlotNameChanged(TSharedPtr<FString> NewSelection, ESelectInfo::Type SelectInfo);
+
+	// slot node names buttons
+	FReply OnAddSlotName(TSharedRef<SWidget> Widget);
+	FReply OnManageSlotName();
+
+	void AddSlotName(const FText & SlotNameToAdd, ETextCommit::Type CommitInfo);
+	int32 RefreshSlotNames();
+
+	// group names
+	TSharedPtr< class STextComboBox>	GroupNameComboBox;
+	TArray< TSharedPtr< FString > >		GroupNameComboList;
+
+	void OnGroupNameChanged(TSharedPtr<FString> NewSelection, ESelectInfo::Type SelectInfo);
+
+	// slot node names buttons
+	FReply OnAddGroupName(TSharedRef<SWidget> Widget);
+	FReply OnManageGroupName();
+
+	void AddGroupName(const FText & GroupNameToAdd, ETextCommit::Type CommitInfo);
+	int32 RefreshGroupNames();
+
+	USkeleton * Skeleton;
+};
