UE4 non-unity fixes

diff --git a/Engine/Source/Editor/BlueprintGraph/Private/K2Node_Knot.cpp b/Engine/Source/Editor/BlueprintGraph/Private/K2Node_Knot.cpp
--- a/Engine/Source/Editor/BlueprintGraph/Private/K2Node_Knot.cpp
+++ b/Engine/Source/Editor/BlueprintGraph/Private/K2Node_Knot.cpp
@@ -2,6 +2,7 @@
 
 #include "BlueprintGraphPrivatePCH.h"
 #include "K2ActionMenuBuilder.h"
+#include "Kismet2NameValidators.h"
 
 #define LOCTEXT_NAMESPACE "K2Node_Knot"
 
diff --git a/Engine/Source/Editor/GraphEditor/Private/SGraphNodeKnot.cpp b/Engine/Source/Editor/GraphEditor/Private/SGraphNodeKnot.cpp
--- a/Engine/Source/Editor/GraphEditor/Private/SGraphNodeKnot.cpp
+++ b/Engine/Source/Editor/GraphEditor/Private/SGraphNodeKnot.cpp
@@ -3,6 +3,9 @@
 
 #include "GraphEditorCommon.h"
 #include "SGraphNodeKnot.h"
+#include "ScopedTransaction.h"
+#include "BlueprintEditorUtils.h"
+
 
 class FAmbivalentDirectionDragConnection : public FDragConnection
 {
diff --git a/Engine/Source/Editor/Persona/Private/Customization/AnimGraphNodeSlotDetails.cpp b/Engine/Source/Editor/Persona/Private/Customization/AnimGraphNodeSlotDetails.cpp
--- a/Engine/Source/Editor/Persona/Private/Customization/AnimGraphNodeSlotDetails.cpp
+++ b/Engine/Source/Editor/Persona/Private/Customization/AnimGraphNodeSlotDetails.cpp
@@ -10,6 +10,7 @@
 
 #include "ScopedTransaction.h"
 #include "AnimGraphNodeSlotDetails.h"
+#include "Persona.h"
 
 #define LOCTEXT_NAMESPACE "AnimNodeSlotDetails"
 
diff --git a/Engine/Source/Editor/Persona/Private/Persona.cpp b/Engine/Source/Editor/Persona/Private/Persona.cpp
--- a/Engine/Source/Editor/Persona/Private/Persona.cpp
+++ b/Engine/Source/Editor/Persona/Private/Persona.cpp
@@ -6,6 +6,7 @@
 #include "SPersonaToolbar.h"
 #include "PersonaModule.h"
 #include "AnimGraphDefinitions.h"
+#include "IDetailsView.h"
 
 #include "Toolkits/IToolkitHost.h"
 
diff --git a/Engine/Source/Editor/Persona/Private/Shared/PersonaMode.cpp b/Engine/Source/Editor/Persona/Private/Shared/PersonaMode.cpp
--- a/Engine/Source/Editor/Persona/Private/Shared/PersonaMode.cpp
+++ b/Engine/Source/Editor/Persona/Private/Shared/PersonaMode.cpp
@@ -7,6 +7,7 @@
 #include "IDocumentation.h"
 #include "SAnimBlueprintParentPlayerList.h"
 #include "SSkeletonSlotGroupNames.h"
+#include "SSkeletonSlotNames.h"
 
 #define LOCTEXT_NAMESPACE "PersonaModes"
 
