Fix asset on exit if property editor hasn't been loaded before shutdown. (Found when loading up to the project browser, then quitting).

diff --git a/Engine/Source/Editor/Persona/Private/PersonaModule.cpp b/Engine/Source/Editor/Persona/Private/PersonaModule.cpp
--- a/Engine/Source/Editor/Persona/Private/PersonaModule.cpp
+++ b/Engine/Source/Editor/Persona/Private/PersonaModule.cpp
@@ -70,6 +70,7 @@ void FPersonaModule::ShutdownModule()
 	MenuExtensibilityManager.Reset();
 
 	// unregsiter when shut down
+	if(FModuleManager::Get().IsModuleLoaded("PropertyEditor"))
 	{
 		FPropertyEditorModule& PropertyModule = FModuleManager::LoadModuleChecked<FPropertyEditorModule>("PropertyEditor");
 		PropertyModule.UnregisterCustomPropertyLayout("SkeletalMeshSocket");
