Fix crash on opening Persona twice (with Benn.Gallagher)

diff --git a/Engine/Source/Editor/Persona/Private/Customization/AnimGraphNodeSlotDetails.cpp b/Engine/Source/Editor/Persona/Private/Customization/AnimGraphNodeSlotDetails.cpp
--- a/Engine/Source/Editor/Persona/Private/Customization/AnimGraphNodeSlotDetails.cpp
+++ b/Engine/Source/Editor/Persona/Private/Customization/AnimGraphNodeSlotDetails.cpp
@@ -15,7 +15,7 @@
 #define LOCTEXT_NAMESPACE "AnimNodeSlotDetails"
 
 ////////////////////////////////////////////////////////////////
-FAnimGraphNodeSlotDetails::FAnimGraphNodeSlotDetails(TSharedRef<class FPersona> InPersonalEditor)
+FAnimGraphNodeSlotDetails::FAnimGraphNodeSlotDetails(TWeakPtr<class FPersona> InPersonalEditor)
 :PersonaEditor(InPersonalEditor)
 {
 }
diff --git a/Engine/Source/Editor/Persona/Private/Customization/AnimGraphNodeSlotDetails.h b/Engine/Source/Editor/Persona/Private/Customization/AnimGraphNodeSlotDetails.h
--- a/Engine/Source/Editor/Persona/Private/Customization/AnimGraphNodeSlotDetails.h
+++ b/Engine/Source/Editor/Persona/Private/Customization/AnimGraphNodeSlotDetails.h
@@ -9,9 +9,9 @@
 class FAnimGraphNodeSlotDetails : public IDetailCustomization
 {
 public:
-	FAnimGraphNodeSlotDetails(TSharedRef<class FPersona> PersonalEditor);
+	FAnimGraphNodeSlotDetails(TWeakPtr<class FPersona> PersonalEditor);
 
-	static TSharedRef<IDetailCustomization> MakeInstance(TSharedRef<class FPersona> PersonalEditor) 
+	static TSharedRef<IDetailCustomization> MakeInstance(TWeakPtr<class FPersona> PersonalEditor) 
 	{
 		return MakeShareable( new FAnimGraphNodeSlotDetails(PersonalEditor) );
 	}
diff --git a/Engine/Source/Editor/Persona/Private/Persona.cpp b/Engine/Source/Editor/Persona/Private/Persona.cpp
--- a/Engine/Source/Editor/Persona/Private/Persona.cpp
+++ b/Engine/Source/Editor/Persona/Private/Persona.cpp
@@ -880,8 +880,9 @@ void FPersona::InitPersona(const EToolkitMode::Type Mode, const TSharedPtr< clas
 
 	// register customization of Slot node for this Persona
 	// this is so that you can open the manage window per Persona
+	TWeakPtr<FPersona> PersonaPtr = SharedThis(this);
 	Inspector->GetPropertyView()->RegisterInstancedCustomPropertyLayout(UAnimGraphNode_Slot::StaticClass(),
-		FOnGetDetailCustomizationInstance::CreateStatic(&FAnimGraphNodeSlotDetails::MakeInstance, SharedThis(this) ));
+		FOnGetDetailCustomizationInstance::CreateStatic(&FAnimGraphNodeSlotDetails::MakeInstance, PersonaPtr));
 }
 
 
