Fix issue with show floor/sky not working after mode switch in Persona

#TTP 320395 - EDITOR: PERSONA: Show floor/sky viewport options don't update automatically on the skeleton and animation tabs

diff --git a/Engine/Source/Editor/Persona/Private/AnimationEditorViewportClient.cpp b/Engine/Source/Editor/Persona/Private/AnimationEditorViewportClient.cpp
--- a/Engine/Source/Editor/Persona/Private/AnimationEditorViewportClient.cpp
+++ b/Engine/Source/Editor/Persona/Private/AnimationEditorViewportClient.cpp
@@ -187,6 +187,24 @@ FAnimationViewportClient::FAnimationViewportClient( FPreviewScene& InPreviewScen
 	}
 }
 
+FAnimationViewportClient::~FAnimationViewportClient()
+{
+	if (PersonaPtr.IsValid())
+	{
+		if (EditorFloorComp)
+		{
+			PreviewScene->RemoveComponent(EditorFloorComp);
+			EditorFloorComp = NULL;
+		}
+
+		if (EditorSkyComp)
+		{
+			PreviewScene->RemoveComponent(EditorSkyComp);
+			EditorSkyComp = NULL;
+		}
+	}
+}
+
 FLinearColor FAnimationViewportClient::GetBackgroundColor() const
 {
 	return SelectedHSVColor.HSVToLinearRGB();
diff --git a/Engine/Source/Editor/Persona/Private/AnimationEditorViewportClient.h b/Engine/Source/Editor/Persona/Private/AnimationEditorViewportClient.h
--- a/Engine/Source/Editor/Persona/Private/AnimationEditorViewportClient.h
+++ b/Engine/Source/Editor/Persona/Private/AnimationEditorViewportClient.h
@@ -36,6 +36,7 @@ protected:
 
 public:
 	FAnimationViewportClient( FPreviewScene& InPreviewScene, TWeakPtr<FPersona> InPersonaPtr );
+	virtual ~FAnimationViewportClient();
 
 	// FEditorViewportClient interface
 	virtual FLinearColor GetBackgroundColor() const OVERRIDE;
