Changed edit box population code in blend space editor so that numbers are no longer formatted

#TTP 335871 - LICENSEE: UDN: PERSONA: BLENDSPACE: Setting Axis range over 1000 a second time will cause it to display as a single integer

diff --git a/Engine/Source/Editor/Persona/Private/SAnimationBlendSpaceBase.cpp b/Engine/Source/Editor/Persona/Private/SAnimationBlendSpaceBase.cpp
--- a/Engine/Source/Editor/Persona/Private/SAnimationBlendSpaceBase.cpp
+++ b/Engine/Source/Editor/Persona/Private/SAnimationBlendSpaceBase.cpp
@@ -176,7 +176,7 @@ public:
 				[
 					SNew(SEditableTextBox)
 					.MinDesiredWidth(10.0f)
-					.Text( FText::AsNumber( GetVectorValueForParam(ParamIndex, SamplePosition) ) )
+					.Text(FText::AsNumber(GetVectorValueForParam(ParamIndex, SamplePosition), NULL, FInternationalization::Get().GetInvariantCulture()))
 					.OnTextCommitted( this, &SSampleEditPopUp::OnValueModified, SampleIndex, ParamIndex  )
 				]
 			];
@@ -416,9 +416,9 @@ void SBlendSpaceParameterWidget::UpdateParametersFromBlendSpace()
 		const FBlendParameter& BlendParam = BlendSpace->GetBlendParameter(I);
 
 		Param_DisplayNames[I]->SetText(FText::FromString(BlendParam.DisplayName));
-		Param_MinValues[I]->SetText( FText::AsNumber( BlendParam.Min ) );
-		Param_MaxValues[I]->SetText( FText::AsNumber( BlendParam.Max) );
-		Param_GridNumbers[I]->SetText(FText::AsNumber(BlendParam.GridNum));
+		Param_MinValues[I]->SetText(FText::AsNumber(BlendParam.Min, NULL, FInternationalization::Get().GetInvariantCulture()));
+		Param_MaxValues[I]->SetText(FText::AsNumber(BlendParam.Max, NULL, FInternationalization::Get().GetInvariantCulture()));
+		Param_GridNumbers[I]->SetText(FText::AsNumber(BlendParam.GridNum, NULL, FInternationalization::Get().GetInvariantCulture()));
 	}
 }
 
