Fixed crash when dragging notifies after selecting a node, closing Persona and reopening the original asset.
#ttp 337774 - CRASH: Assertion failed: (Index >= 0) & (Index < ArrayNum) Function:

diff --git a/Engine/Source/Editor/Persona/Private/SAnimNotifyPanel.cpp b/Engine/Source/Editor/Persona/Private/SAnimNotifyPanel.cpp
--- a/Engine/Source/Editor/Persona/Private/SAnimNotifyPanel.cpp
+++ b/Engine/Source/Editor/Persona/Private/SAnimNotifyPanel.cpp
@@ -148,7 +148,8 @@ private:
 	float						LastSnappedTime;
 
 	bool						bDrawTooltipToRight;
-	bool						bBeingDragged;	
+	bool						bBeingDragged;
+	bool						bSelected;
 
 	/** The scrub handle currently being dragged, if any */
 	ENotifyStateHandleHit::Type CurrentDragHandle;
@@ -255,6 +256,8 @@ public:
 	void AppendSelectionToSet(FGraphPanelSelectionSet& SelectionSet);
 	// Adds our current selection to the provided array
 	void AppendSelectionToArray(TArray<const FAnimNotifyEvent*>& Selection) const;
+	// Gets the currently selected SAnimNotifyNode instances
+	void AppendSelectedNodeWidgetsToArray(TArray<TSharedPtr<SAnimNotifyNode>>& NodeArray) const;
 
 	/**
 	* Deselects all currently selected notify nodes
@@ -805,6 +808,7 @@ void SAnimNotifyNode::Construct(const FArguments& InArgs)
 	bBeingDragged = false;
 	CurrentDragHandle = ENotifyStateHandleHit::None;
 	bDrawTooltipToRight = true;
+	bSelected = false;
 
 	OnNodeDragStarted = InArgs._OnNodeDragStarted;
 	PanTrackRequest = InArgs._PanTrackRequest;
@@ -1026,7 +1030,7 @@ int32 SAnimNotifyNode::OnPaint(const FGeometry& AllottedGeometry, const FSlateRe
 		FLinearColor(1.0f, 1.0f, 1.0f,0.1f));
 
 	FText Text = GetNotifyText();
-	FLinearColor NodeColour = NotifyEvent->bSelected ? FLinearColor(1.0f, 0.5f, 0.0f) : FLinearColor::Red;
+	FLinearColor NodeColour = bSelected ? FLinearColor(1.0f, 0.5f, 0.0f) : FLinearColor::Red;
 	
 	float HalfScrubHandleWidth = ScrubHandleSize.X / 2.0f;
 
@@ -1082,7 +1086,7 @@ int32 SAnimNotifyNode::OnPaint(const FGeometry& AllottedGeometry, const FSlateRe
 
 	// Frame
 	// Drawing lines is slow, reserved for single selected node
-	if( NotifyEvent->bSelected )
+	if( bSelected )
 	{
 		TArray<FVector2D> LinePoints;
 
@@ -1833,10 +1837,10 @@ void SAnimNotifyTrack::SelectNotifyNode(int32 NotifyIndex, bool Append)
 		if (!SelectedNotifyIndices.Contains(NotifyIndex))
 		{
 			// select new one
-			if (AnimNotifies.IsValidIndex(NotifyIndex))
+			if (NotifyNodes.IsValidIndex(NotifyIndex))
 			{
-				FAnimNotifyEvent * NotifyEvent = AnimNotifies[NotifyIndex];
-				NotifyEvent->bSelected = true;
+				TSharedPtr<SAnimNotifyNode> Node = NotifyNodes[NotifyIndex];
+				Node->bSelected = true;
 				SelectedNotifyIndices.Add(NotifyIndex);
 				SendSelectionChanged();
 			}
@@ -1846,7 +1850,7 @@ void SAnimNotifyTrack::SelectNotifyNode(int32 NotifyIndex, bool Append)
 
 void SAnimNotifyTrack::ToggleNotifyNodeSelectionStatus( int32 NotifyIndex )
 {
-	check(AnimNotifies.IsValidIndex(NotifyIndex));
+	check(NotifyNodes.IsValidIndex(NotifyIndex));
 
 	bool bSelected = SelectedNotifyIndices.Contains(NotifyIndex);
 	if(bSelected)
@@ -1858,16 +1862,16 @@ void SAnimNotifyTrack::ToggleNotifyNodeSelectionStatus( int32 NotifyIndex )
 		SelectedNotifyIndices.Add(NotifyIndex);
 	}
 
-	FAnimNotifyEvent* Event = AnimNotifies[NotifyIndex];
-	Event->bSelected = !Event->bSelected;
+	TSharedPtr<SAnimNotifyNode> Node = NotifyNodes[NotifyIndex];
+	Node->bSelected = !Node->bSelected;
 	SendSelectionChanged();
 }
 
 void SAnimNotifyTrack::DeselectNotifyNode( int32 NotifyIndex )
 {
-	check(AnimNotifies.IsValidIndex(NotifyIndex));
-	FAnimNotifyEvent* Event = AnimNotifies[NotifyIndex];
-	Event->bSelected = false;
+	check(NotifyNodes.IsValidIndex(NotifyIndex));
+	TSharedPtr<SAnimNotifyNode> Node = NotifyNodes[NotifyIndex];
+	Node->bSelected = false;
 
 	int32 ItemsRemoved = SelectedNotifyIndices.Remove(NotifyIndex);
 	check(ItemsRemoved > 0);
@@ -1877,9 +1881,9 @@ void SAnimNotifyTrack::DeselectNotifyNode( int32 NotifyIndex )
 
 void SAnimNotifyTrack::DeselectAllNotifyNodes(bool bUpdateSelectionSet)
 {
-	for (FAnimNotifyEvent& Event : Sequence->Notifies)
+	for(TSharedPtr<SAnimNotifyNode> Node : NotifyNodes)
 	{
-		Event.bSelected = false;
+		Node->bSelected = false;
 	}
 	SelectedNotifyIndices.Empty();
 
@@ -1903,7 +1907,7 @@ TSharedPtr<SWidget> SAnimNotifyTrack::SummonContextMenu(const FGeometry& MyGeome
 	{
 		if ( NotifyIndex != INDEX_NONE )
 		{
-			if (!AnimNotifies[NotifyIndex]->bSelected)
+			if (!NotifyNodes[NotifyIndex]->bSelected)
 			{
 				SelectNotifyNode(NotifyIndex, MouseEvent.IsControlDown());
 			}
@@ -2399,7 +2403,7 @@ int32 SAnimNotifyTrack::GetHitNotifyNode(const FGeometry& MyGeometry, const FVec
 FReply SAnimNotifyTrack::OnNotifyNodeDragStarted( TSharedRef<SAnimNotifyNode> NotifyNode, const FVector2D& ScreenCursorPos, const FVector2D& ScreenNodePosition, const bool bDragOnMarker, int32 NotifyIndex ) 
 {
 	// Check to see if we've already selected the triggering node
-	if (!NotifyNode.Get().NotifyEvent->bSelected)
+	if (!NotifyNode->bSelected)
 	{
 		SelectNotifyNode(NotifyIndex, false);
 	}
@@ -2602,6 +2606,17 @@ void SAnimNotifyTrack::PasteSingleNotify(FString& NotifyString, float PasteTime)
 	OnUpdatePanel.ExecuteIfBound();
 }
 
+void SAnimNotifyTrack::AppendSelectedNodeWidgetsToArray(TArray<TSharedPtr<SAnimNotifyNode>>& NodeArray) const
+{
+	for(TSharedPtr<SAnimNotifyNode> Node : NotifyNodes)
+	{
+		if(Node->bSelected)
+		{
+			NodeArray.Add(Node);
+		}
+	}
+}
+
 //////////////////////////////////////////////////////////////////////////
 // SSequenceEdTrack
 
@@ -2829,19 +2844,13 @@ void SAnimNotifyPanel::DeleteNotify(FAnimNotifyEvent* Notify)
 {
 	for (int32 I=0; I<Sequence->Notifies.Num(); ++I)
 	{
-		if (Notify == &Sequence->Notifies[I])
+		if (*Notify == Sequence->Notifies[I])
 		{
 			Sequence->Notifies.RemoveAt(I);
 			Sequence->MarkPackageDirty();
 			break;
 		}
 	}
-
-	// clear selection before updating
-	FGraphPanelSelectionSet ObjectSet;
-	OnSelectionChanged.ExecuteIfBound(ObjectSet);
-
-	Update();
 }
 
 void SAnimNotifyPanel::Update()
@@ -3001,20 +3010,27 @@ void SAnimNotifyPanel::OnGenericDelete()
 	// so don't delete anything when the key is pressed.
 	if(HasKeyboardFocus() || HasFocusedDescendants()) 
 	{
-		// Find selected notifies and delete them
-		for(int32 NotifyIndex = 0 ; NotifyIndex < Sequence->Notifies.Num() ;)
+		TArray<TSharedPtr<SAnimNotifyNode>> SelectedNodes;
+		for(TSharedPtr<SAnimNotifyTrack> Track : NotifyAnimTracks)
 		{
-			FAnimNotifyEvent& Event = Sequence->Notifies[NotifyIndex];
-			if(Event.bSelected)
-			{
-				DeleteNotify(&Event);
-			}
-			else
+			Track->AppendSelectedNodeWidgetsToArray(SelectedNodes);
+		}
+
+		if(SelectedNodes.Num() > 0)
+		{
+			FScopedTransaction Transaction(LOCTEXT("DeleteNotifies", "Delete Animation Notifies"));
+			Sequence->Modify(true);
+			for(TSharedPtr<SAnimNotifyNode> Node : SelectedNodes)
 			{
-				// Only increment if we don't remove something
-				NotifyIndex++;
+				DeleteNotify(Node->NotifyEvent);
 			}
 		}
+
+		// clear selection and update the panel
+		FGraphPanelSelectionSet ObjectSet;
+		OnSelectionChanged.ExecuteIfBound(ObjectSet);
+
+		Update();
 	}
 }
 
diff --git a/Engine/Source/Runtime/Engine/Classes/Animation/AnimSequenceBase.h b/Engine/Source/Runtime/Engine/Classes/Animation/AnimSequenceBase.h
--- a/Engine/Source/Runtime/Engine/Classes/Animation/AnimSequenceBase.h
+++ b/Engine/Source/Runtime/Engine/Classes/Animation/AnimSequenceBase.h
@@ -71,10 +71,6 @@ struct FAnimNotifyEvent
 	UPROPERTY()
 	int32 TrackIndex;
 
-	/** If notify is selected in editor */
-	UPROPERTY(Transient)
-	bool bSelected;
-
 #endif // WITH_EDITORONLY_DATA
 
 	FAnimNotifyEvent()
@@ -89,7 +85,6 @@ struct FAnimNotifyEvent
 		, Duration(0)
 #if WITH_EDITORONLY_DATA
 		, TrackIndex(0)
-		, bSelected(false)
 #endif // WITH_EDITORONLY_DATA
 	{
 	}
