Altered animation import to detect and scale existing animation curves on import and reimport to avoid keys becoming inaccessible.
#ttp 320313 - LICENSEE: UDN: PERSONA: Float Track Rescale on Animation Reimport

diff --git a/Engine/Source/Editor/Persona/Private/Persona.cpp b/Engine/Source/Editor/Persona/Private/Persona.cpp
--- a/Engine/Source/Editor/Persona/Private/Persona.cpp
+++ b/Engine/Source/Editor/Persona/Private/Persona.cpp
@@ -447,6 +447,7 @@ FPersona::FPersona()
 
 FPersona::~FPersona()
 {
+	FEditorDelegates::OnAssetPostImport.RemoveAll(this);
 	FReimportManager::Instance()->OnPostReimport().RemoveAll(this);
 
 	FPersonaModule* PersonaModule = &FModuleManager::LoadModuleChecked<FPersonaModule>( "Persona" );
@@ -883,6 +884,9 @@ void FPersona::InitPersona(const EToolkitMode::Type Mode, const TSharedPtr< clas
 	TWeakPtr<FPersona> PersonaPtr = SharedThis(this);
 	Inspector->GetPropertyView()->RegisterInstancedCustomPropertyLayout(UAnimGraphNode_Slot::StaticClass(),
 		FOnGetDetailCustomizationInstance::CreateStatic(&FAnimGraphNodeSlotDetails::MakeInstance, PersonaPtr));
+
+	// Register post import callback to catch animation imports when we have the asset open (we need to reinit)
+	FEditorDelegates::OnAssetPostImport.AddRaw(this, &FPersona::OnPostImport);
 }
 
 
@@ -998,14 +1002,25 @@ UDebugSkelMeshComponent* FPersona::GetPreviewMeshComponent()
 }
 
 void FPersona::OnPostReimport(UObject* InObject, bool bSuccess)
+{
+	ConditionalRefreshEditor(InObject);
+}
+
+void FPersona::OnPostImport(UFactory* InFactory, UObject* InObject)
+{
+	ConditionalRefreshEditor(InObject);
+}
+
+void FPersona::ConditionalRefreshEditor(UObject* InObject)
 {
 	// Ignore if this is regarding a different object
-	if ( InObject != TargetSkeleton && InObject != TargetSkeleton->GetPreviewMesh() && InObject != GetAnimationAssetBeingEdited() )
+	if(InObject != TargetSkeleton && InObject != TargetSkeleton->GetPreviewMesh() && InObject != GetAnimationAssetBeingEdited())
 	{
 		return;
 	}
 
 	RefreshViewport();
+	ReinitMode();
 }
 
 /** Called when graph editor focus is changed */
diff --git a/Engine/Source/Editor/Persona/Private/Persona.h b/Engine/Source/Editor/Persona/Private/Persona.h
--- a/Engine/Source/Editor/Persona/Private/Persona.h
+++ b/Engine/Source/Editor/Persona/Private/Persona.h
@@ -187,6 +187,9 @@ public:
 	/** Gets called when it's clicked via mode tab - Reinitialize the mode **/
 	void ReinitMode();
 
+	/** Called when an asset is imported into the editor */
+	void OnPostImport(UFactory* InFactory, UObject* InObject);
+
 public:
 	// IToolkit interface
 	virtual FName GetToolkitFName() const OVERRIDE;
@@ -268,6 +271,9 @@ protected:
 	/** Callback when an object has been reimported, and whether it worked */
 	void OnPostReimport(UObject* InObject, bool bSuccess);
 
+	/** Refreshes the viewport and current editor mode if the passed in object is open in the editor */
+	void ConditionalRefreshEditor(UObject* InObject);
+
 	/** Clear up Preview Mesh's AnimNotifyStates. Called when undo or redo**/
 	void ClearupPreviewMeshAnimNotifyStates();
 
diff --git a/Engine/Source/Editor/UnrealEd/Private/SkeletalMeshEdit.cpp b/Engine/Source/Editor/UnrealEd/Private/SkeletalMeshEdit.cpp
--- a/Engine/Source/Editor/UnrealEd/Private/SkeletalMeshEdit.cpp
+++ b/Engine/Source/Editor/UnrealEd/Private/SkeletalMeshEdit.cpp
@@ -911,10 +911,22 @@ namespace AnimationTransformDebug
 bool UnFbx::FFbxImporter::ImportAnimation(USkeleton * Skeleton, UAnimSequence * DestSeq, const FString & FileName, TArray<FbxNode*>& SortedLinks, TArray<FbxNode*>& NodeArray, FbxAnimStack* CurAnimStack, const int32 ResampleRate, const FbxTimeSpan AnimTimeSpan)
 {
 	FbxTime SequenceLength = AnimTimeSpan.GetDuration();
+	float PreviousSequenceLength = DestSeq->SequenceLength;
 
 	// if you have one pose(thus 0.f duration), it still contains animation, so we'll need to consider that as MINIMUM_ANIMATION_LENGTH time length
 	DestSeq->SequenceLength = FGenericPlatformMath::Max<float>(SequenceLength.GetSecondDouble(), MINIMUM_ANIMATION_LENGTH);
 
+	if(PreviousSequenceLength > MINIMUM_ANIMATION_LENGTH && DestSeq->RawCurveData.FloatCurves.Num() > 0)
+	{
+		// The sequence already existed when we began the import. We need to scale the key times for all curves to match the new 
+		// duration before importing over them. This is to catch any user-added curves
+		float ScaleFactor = DestSeq->SequenceLength / PreviousSequenceLength;
+		for(FFloatCurve& Curve : DestSeq->RawCurveData.FloatCurves)
+		{
+			Curve.FloatCurve.ScaleCurve(0.0f, ScaleFactor);
+		}
+	}
+
 	//
 	// shape animation START
 	//
