Fix for branching points not moving in a montage

diff --git a/Engine/Source/Editor/Persona/Private/SAnimMontagePanel.cpp b/Engine/Source/Editor/Persona/Private/SAnimMontagePanel.cpp
--- a/Engine/Source/Editor/Persona/Private/SAnimMontagePanel.cpp
+++ b/Engine/Source/Editor/Persona/Private/SAnimMontagePanel.cpp
@@ -46,6 +46,7 @@ public:
 	SLATE_BEGIN_ARGS( SMontageBranchingPointNode )
 		: _Montage()
 		, _BranchingPoint()
+		, _AllowDrag(true)
 		{}
 		SLATE_ARGUMENT( class UAnimMontage*, Montage )
 		SLATE_ARGUMENT( FBranchingPoint*, BranchingPoint )
