Revamped  LOD panel for skeletal mesh to look like static mesh LOD

diff --git a/Engine/Source/Developer/SimplygonMeshReduction/Private/SimplygonMeshReduction.cpp b/Engine/Source/Developer/SimplygonMeshReduction/Private/SimplygonMeshReduction.cpp
--- a/Engine/Source/Developer/SimplygonMeshReduction/Private/SimplygonMeshReduction.cpp
+++ b/Engine/Source/Developer/SimplygonMeshReduction/Private/SimplygonMeshReduction.cpp
@@ -225,6 +225,19 @@ public:
 			FSkeletalMeshLODInfo* NewLODInfo = new( SkeletalMesh->LODInfo ) FSkeletalMeshLODInfo;
 			FSkeletalMeshLODInfo& OldLODInfo = SkeletalMesh->LODInfo[0];
 			*NewLODInfo = OldLODInfo;
+
+			// creates LOD Material map for a newly generated LOD model
+			int32 NumSections = LODModels[0]->NumNonClothingSections();
+			if (NewLODInfo->LODMaterialMap.Num() != NumSections)
+			{
+				NewLODInfo->LODMaterialMap.Empty(NumSections);
+				NewLODInfo->LODMaterialMap.AddUninitialized(NumSections);
+			}
+
+			for (int32 Index = 0; Index < NumSections; Index++)
+			{
+				NewLODInfo->LODMaterialMap[Index] = Index;
+			}
 		}
 
 		// now try bone reduction process if it's setup
diff --git a/Engine/Source/Editor/Persona/Private/Persona.cpp b/Engine/Source/Editor/Persona/Private/Persona.cpp
--- a/Engine/Source/Editor/Persona/Private/Persona.cpp
+++ b/Engine/Source/Editor/Persona/Private/Persona.cpp
@@ -438,6 +438,7 @@ TArray<UObject*> GetEditorObjectsOfClass( const TArray< UObject* >& Objects, con
 FPersona::FPersona()
 	: TargetSkeleton(NULL)
 	, PreviewComponent(NULL)
+	, PersonaMeshDetailLayout(NULL)
 	, PreviewScene(FPreviewScene::ConstructionValues().AllowAudioPlayback(true).ShouldSimulatePhysics(true))
 {
 	// Register to be notified when properties are edited
diff --git a/Engine/Source/Editor/Persona/Private/Persona.h b/Engine/Source/Editor/Persona/Private/Persona.h
--- a/Engine/Source/Editor/Persona/Private/Persona.h
+++ b/Engine/Source/Editor/Persona/Private/Persona.h
@@ -304,6 +304,8 @@ public:
 	/** Shared data between modes - for now only used for viewport **/
 	FPersonaModeSharedData ModeSharedData;
 
+	/** holding this pointer to refresh persona mesh detials tab when LOD is changed **/
+	class IDetailLayoutBuilder* PersonaMeshDetailLayout;
 private:
 	// called when animation asset has been changed
 	DECLARE_MULTICAST_DELEGATE_OneParam( FOnAnimChangedMulticaster, UAnimationAsset* );
diff --git a/Engine/Source/Editor/Persona/Private/PersonaMeshDetails.cpp b/Engine/Source/Editor/Persona/Private/PersonaMeshDetails.cpp
--- a/Engine/Source/Editor/Persona/Private/PersonaMeshDetails.cpp
+++ b/Engine/Source/Editor/Persona/Private/PersonaMeshDetails.cpp
@@ -18,56 +18,80 @@
 
 #define LOCTEXT_NAMESPACE "PersonaMeshDetails"
 
+/**
+* FPersonaMeshDetails
+*/
 TSharedRef<IDetailCustomization> FPersonaMeshDetails::MakeInstance(TSharedPtr<FPersona> InPersona)
 {
 	return MakeShareable( new FPersonaMeshDetails(InPersona) );
 }
 
+void FPersonaMeshDetails::AddLODLevelCategories(IDetailLayoutBuilder& DetailLayout)
+{
+	USkeletalMesh* SkelMesh = PersonaPtr->GetMesh();
+
+	if (SkelMesh)
+	{
+		const int32 SkelMeshLODCount = SkelMesh->LODInfo.Num();
+
+		// Create information panel for each LOD level.
+		for (int32 LODIndex = 0; LODIndex < SkelMeshLODCount; ++LODIndex)
+		{
+			FString CategoryName = FString(TEXT("LOD"));
+			CategoryName.AppendInt(LODIndex);
+
+			FString LODLevelString = FText::Format(LOCTEXT("LODLevel", "LOD{0}"), FText::AsNumber(LODIndex)).ToString();
+
+			IDetailCategoryBuilder& LODCategory = DetailLayout.EditCategory(*CategoryName, LODLevelString, ECategoryPriority::Important);
+
+			{
+				FMaterialListDelegates MaterialListDelegates;
+
+				MaterialListDelegates.OnGetMaterials.BindSP(this, &FPersonaMeshDetails::OnGetMaterialsForView, LODIndex);
+				MaterialListDelegates.OnMaterialChanged.BindSP(this, &FPersonaMeshDetails::OnMaterialChanged, LODIndex);
+				MaterialListDelegates.OnGenerateCustomNameWidgets.BindSP(this, &FPersonaMeshDetails::OnGenerateCustomNameWidgetsForMaterial, LODIndex);
+				MaterialListDelegates.OnGenerateCustomMaterialWidgets.BindSP(this, &FPersonaMeshDetails::OnGenerateCustomMaterialWidgetsForMaterial, LODIndex);
+
+				LODCategory.AddCustomBuilder(MakeShareable(new FMaterialList(LODCategory.GetParentLayout(), MaterialListDelegates)));
+			}
+			// add each LOD Info to each LOD category
+			TSharedRef<IPropertyHandle> LODInfoProperty = DetailLayout.GetProperty(FName("LODInfo"), USkeletalMesh::StaticClass());
+
+			if (LODInfoProperty->IsValidHandle())
+			{
+				uint32 NumChildren;
+				LODInfoProperty->GetNumChildren(NumChildren);
+				check((uint32)LODIndex < NumChildren);
+
+				// changing property name to "LOD Info" because it shows only array index
+				TSharedPtr<IPropertyHandle> LODInfoChild = LODInfoProperty->GetChildHandle(LODIndex);
+				LODInfoChild->CreatePropertyNameWidget(FString(TEXT("LOD Info")));
+				LODCategory.AddProperty(LODInfoChild);
+			}
+		}
+	}
+}
+
 void FPersonaMeshDetails::CustomizeDetails( IDetailLayoutBuilder& DetailLayout )
 {
 	SelectedObjects = DetailLayout.GetDetailsView().GetSelectedObjects();
 	check(SelectedObjects.Num()<=1); // The OnGenerateCustomWidgets delegate will not be useful if we try to process more than one object.
 
-	FMaterialListDelegates MaterialListDelegates; 
-	MaterialListDelegates.OnGetMaterials.BindSP( this, &FPersonaMeshDetails::OnGetMaterialsForView );
-	MaterialListDelegates.OnMaterialChanged.BindSP( this, &FPersonaMeshDetails::OnMaterialChanged );
-	MaterialListDelegates.OnGenerateCustomNameWidgets.BindSP(this, &FPersonaMeshDetails::OnGenerateCustomNameWidgetsForMaterial);
-	MaterialListDelegates.OnGenerateCustomMaterialWidgets.BindSP(this, &FPersonaMeshDetails::OnGenerateCustomMaterialWidgetsForMaterial);
+	// copy temporarily to refresh Mesh details tab from the LOD settings window
+	PersonaPtr->PersonaMeshDetailLayout = &DetailLayout;
 
-	TSharedRef<FMaterialList> MaterialList( new FMaterialList( DetailLayout, MaterialListDelegates ) );
+	// add multiple LOD levels to LOD category
+	AddLODLevelCategories(DetailLayout);
 
-	// only show the category if there are materials to display
-	if( MaterialList->IsDisplayingMaterials() )
-	{
-		// Make a category for the materials.
-		IDetailCategoryBuilder& MaterialCategory = DetailLayout.EditCategory("Materials", TEXT(""), ECategoryPriority::TypeSpecific);
-
-		MaterialCategory.AddCustomBuilder( MaterialList );
-	}
+	// LODInfo doesn't need to be showed anymore because it was moved to each LOD category
+	TSharedRef<IPropertyHandle> LODInfoProperty = DetailLayout.GetProperty(FName("LODInfo"), USkeletalMesh::StaticClass());
+	DetailLayout.HideProperty(LODInfoProperty);
 
 #if WITH_APEX_CLOTHING
 	IDetailCategoryBuilder& ClothingCategory = DetailLayout.EditCategory("Clothing", TEXT(""), ECategoryPriority::TypeSpecific);
 	CustomizeClothingProperties(DetailLayout,ClothingCategory);
 #endif// #if WITH_APEX_CLOTHING
 
-	// hide other materials property
-	IDetailCategoryBuilder& SkeletalMeshCategory = DetailLayout.EditCategory("Skeletal Mesh", TEXT(""), ECategoryPriority::TypeSpecific);
-	TSharedRef<IPropertyHandle> MaterialProperty = DetailLayout.GetProperty("Materials", USkeletalMesh::StaticClass());
-
-	DetailLayout.HideProperty(MaterialProperty);
-	IDetailCategoryBuilder& MaterialCategory = DetailLayout.EditCategory("Materials");
-	MaterialCategory.AddCustomRow("AddMaterial")
-	[
-		SNew(SVerticalBox)
-		+SVerticalBox::Slot()
-		.Padding(50,10,50,0)
-		[
-			SNew(SButton)
-			.OnClicked(FOnClicked::CreateSP(this, &FPersonaMeshDetails::OnAddMaterialButtonClicked))
-			.Text(LOCTEXT("SKMeshAddMaterial", "Add Material Element"))
-		]
-	];
-
 	IDetailCategoryBuilder& AdditionalMeshCategory = DetailLayout.EditCategory("Additional Meshes", LOCTEXT("AdditionalMeshesCollapsable", "Additional Meshes").ToString(), ECategoryPriority::TypeSpecific);
 	AdditionalMeshCategory.AddCustomRow("")
 	[
@@ -75,26 +99,52 @@ void FPersonaMeshDetails::CustomizeDetails( IDetailLayoutBuilder& DetailLayout )
 	];
 }
 
-void FPersonaMeshDetails::OnGetMaterialsForView( IMaterialListBuilder& MaterialList )
+void FPersonaMeshDetails::OnGetMaterialsForView(IMaterialListBuilder& MaterialList, int32 LODIndex)
 {
-	for(auto Iter = SelectedObjects.CreateIterator(); Iter; ++Iter)
+	USkeletalMesh* SkelMesh = PersonaPtr->GetMesh();
+
+	FSkeletalMeshResource* ImportedResource = SkelMesh->GetImportedResource();
+
+	if (ImportedResource && ImportedResource->LODModels.IsValidIndex(LODIndex))
 	{
-		auto ObjectPtr = (*Iter);
-		if(USkeletalMesh* Mesh = Cast<USkeletalMesh>(ObjectPtr.Get()))
+		FStaticLODModel& Model = ImportedResource->LODModels[LODIndex];
+
+		bool bHasMaterialMap = SkelMesh->LODInfo.IsValidIndex(LODIndex) && SkelMesh->LODInfo[LODIndex].LODMaterialMap.Num() > 0;
+
+		if (LODIndex == 0 || !bHasMaterialMap)
+		{
+			int32 NumSections = Model.NumNonClothingSections();
+			for (int32 SectionIdx = 0; SectionIdx < NumSections; SectionIdx++)
+			{
+				int32 MaterialIndex = Model.Sections[SectionIdx].MaterialIndex;
+
+				if (SkelMesh->Materials.IsValidIndex(MaterialIndex))
+				{
+					MaterialList.AddMaterial(SectionIdx, SkelMesh->Materials[MaterialIndex].MaterialInterface, true);
+				}
+			}
+		}
+		else // refers to LODMaterialMap
 		{
-			for(int MaterialIndex = 0; MaterialIndex < Mesh->Materials.Num(); ++MaterialIndex)
+			TArray<int32>& MaterialMap = SkelMesh->LODInfo[LODIndex].LODMaterialMap;
+
+			for(int32 MapIdx = 0; MapIdx < MaterialMap.Num(); MapIdx++)
 			{
-				MaterialList.AddMaterial( MaterialIndex, Mesh->Materials[MaterialIndex].MaterialInterface, true );
+				int32 MaterialIndex = MaterialMap[MapIdx];
+
+				if (!SkelMesh->Materials.IsValidIndex(MaterialIndex))
+				{
+					MaterialMap[MapIdx] = MaterialIndex = SkelMesh->Materials.Add(FSkeletalMaterial());
+				}
+
+				MaterialList.AddMaterial(MapIdx, SkelMesh->Materials[MaterialIndex].MaterialInterface, true);
 			}
 		}
-	}
 
-#if WITH_APEX_CLOTHING
-	CheckLODMaterialMapChanges();
-#endif// #if WITH_APEX_CLOTHING
+	}
 }
 
-TSharedRef<SWidget> FPersonaMeshDetails::OnGenerateCustomNameWidgetsForMaterial(UMaterialInterface* Material, int32 SlotIndex)
+TSharedRef<SWidget> FPersonaMeshDetails::OnGenerateCustomNameWidgetsForMaterial(UMaterialInterface* Material, int32 SlotIndex, int32 LODIndex)
 {
 	return
 		SNew(SCheckBox)
@@ -109,22 +159,33 @@ TSharedRef<SWidget> FPersonaMeshDetails::OnGenerateCustomNameWidgetsForMaterial(
 		];
 }
 
-
-TSharedRef<SWidget> FPersonaMeshDetails::OnGenerateCustomMaterialWidgetsForMaterial(UMaterialInterface* Material, int32 SlotIndex)
+TSharedRef<SWidget> FPersonaMeshDetails::OnGenerateCustomMaterialWidgetsForMaterial(UMaterialInterface* Material, int32 SlotIndex, int32 LODIndex)
 {
+	// currently a slot index means a section index
+	int32 SectionIndex = SlotIndex;
+
 #if WITH_APEX_CLOTHING
-	if(SlotIndex == 0)
+
+	if (LODIndex == 0 && SectionIndex == 0)
 	{
-		ClothingComboBoxes.Empty();
+		ClothingComboLODInfos.Empty();
 		// Generate strings for the combo boxes assets
 		UpdateComboBoxStrings();
 	}
-	while( ClothingComboBoxes.Num() <= SlotIndex )
+	
+	TArray<TSharedPtr< STextComboBox >>& ComboBoxes = ClothingComboLODInfos[LODIndex].ClothingComboBoxes;
+	TArray<TSharedPtr<FString> >& ComboStrings = ClothingComboLODInfos[LODIndex].ClothingComboStrings;
+	TMap<FString, FClothAssetSubmeshIndex>&	ComboStringReverseLookup = ClothingComboLODInfos[LODIndex].ClothingComboStringReverseLookup;
+	TArray<int32>& ComboSelectedIndices = ClothingComboLODInfos[LODIndex].ClothingComboSelectedIndices;
+
+	while (ComboBoxes.Num() <= SectionIndex)
 	{
-		ClothingComboBoxes.AddZeroed();
+		ComboBoxes.AddZeroed();
 	}
+
 #endif // WITH_APEX_CLOTHING
 
+	int32 MaterialIndex = GetMaterialIndex(LODIndex, SectionIndex);
 	TSharedRef<SWidget> MaterialWidget
 	= SNew(SVerticalBox)
 
@@ -132,8 +193,8 @@ TSharedRef<SWidget> FPersonaMeshDetails::OnGenerateCustomMaterialWidgetsForMater
 		.Padding(0,2,0,0)
 		[
 			SNew(SCheckBox)
-			.IsChecked(this, &FPersonaMeshDetails::IsShadowCastingEnabled, SlotIndex)
-			.OnCheckStateChanged(this, &FPersonaMeshDetails::OnShadowCastingChanged, SlotIndex)
+			.IsChecked(this, &FPersonaMeshDetails::IsShadowCastingEnabled, MaterialIndex)
+			.OnCheckStateChanged(this, &FPersonaMeshDetails::OnShadowCastingChanged, MaterialIndex)
 			[
 				SNew(STextBlock)
 				.Font(FEditorStyle::GetFontStyle("StaticMeshEditor.NormalFont"))
@@ -150,8 +211,8 @@ TSharedRef<SWidget> FPersonaMeshDetails::OnGenerateCustomMaterialWidgetsForMater
 			.AutoWidth()
 			[
 				SNew(SButton)
-				.IsEnabled(this, &FPersonaMeshDetails::CanDeleteMaterialElement, SlotIndex)
-				.OnClicked(this, &FPersonaMeshDetails::OnDeleteButtonClicked, SlotIndex)
+				.IsEnabled(this, &FPersonaMeshDetails::CanDeleteMaterialElement, LODIndex, SectionIndex)
+				.OnClicked(this, &FPersonaMeshDetails::OnDeleteButtonClicked, LODIndex, SectionIndex)
 				.ButtonStyle(FEditorStyle::Get(), "HoverHintOnly")
 				.VAlign(VAlign_Center)
 				.HAlign(HAlign_Center)
@@ -168,7 +229,7 @@ TSharedRef<SWidget> FPersonaMeshDetails::OnGenerateCustomMaterialWidgetsForMater
 				.Font(FEditorStyle::GetFontStyle("StaticMeshEditor.NormalFont"))
 				.Text(LOCTEXT("PersonaDeleteMaterialLabel", "Delete"))
 			]
-	]
+		]
 
 #if WITH_APEX_CLOTHING
 
@@ -178,7 +239,7 @@ TSharedRef<SWidget> FPersonaMeshDetails::OnGenerateCustomMaterialWidgetsForMater
 		.Padding(0,2,0,0)
 		[
 			SNew( SHorizontalBox )
-			.Visibility(this, &FPersonaMeshDetails::IsClothingComboBoxVisible, SlotIndex)
+			.Visibility(this, &FPersonaMeshDetails::IsClothingComboBoxVisible, LODIndex, SectionIndex)
 
 			+SHorizontalBox::Slot()
 			.AutoWidth()
@@ -192,34 +253,16 @@ TSharedRef<SWidget> FPersonaMeshDetails::OnGenerateCustomMaterialWidgetsForMater
 			.Padding(2,2,0,0)
 			.AutoWidth()
 			[
-				SAssignNew(ClothingComboBoxes[SlotIndex], STextComboBox)
+				SAssignNew(ComboBoxes[SlotIndex], STextComboBox)
 				.ContentPadding(FMargin(6.0f, 2.0f))
-				.OptionsSource(&ClothingComboStrings)
+				.OptionsSource(&ComboStrings)
 				.ToolTipText(LOCTEXT("SectionsComboBoxToolTip", "Select the clothing asset and submesh to use as clothing for this section"))
-				.OnSelectionChanged(this, &FPersonaMeshDetails::HandleSectionsComboBoxSelectionChanged, SlotIndex)
-			]
-		]
-
-		// check-box for importing cloth-LODs automatically
-		+SVerticalBox::Slot()
-		.AutoHeight()
-		.Padding(0,2,0,0)
-		[
-			SNew( SCheckBox )
-			.IsChecked( this, &FPersonaMeshDetails::IsEnablingClothingLOD, SlotIndex )
-			.OnCheckStateChanged( this, &FPersonaMeshDetails::OnEnableClothingLOD, SlotIndex )
-			.Visibility(this, &FPersonaMeshDetails::IsClothingLODCheckBoxVisible, SlotIndex)
-			[
-				SNew( STextBlock )
-				.Font(FEditorStyle::GetFontStyle("StaticMeshEditor.NormalFont"))
-				.Text( LOCTEXT("EnableClothLODCheckBox", "Enable Clothing LOD") )
-				.ToolTipText( LOCTEXT("UseLODTooltip", "If enabled, clothing data will be used for each LOD if there is a corresponding APEX clothing LOD available. If disabled, clothing data will only be used when LOD 0 is rendered.") )
+				.OnSelectionChanged(this, &FPersonaMeshDetails::HandleSectionsComboBoxSelectionChanged, LODIndex, SectionIndex)
 			]
 		];
 
-		int32 SelectedIndex = ClothingComboSelectedIndices.IsValidIndex(SlotIndex) ? ClothingComboSelectedIndices[SlotIndex] : -1;
-		ClothingComboBoxes[SlotIndex]->SetSelectedItem(ClothingComboStrings.IsValidIndex(SelectedIndex) ? ClothingComboStrings[SelectedIndex] : ClothingComboStrings[0]);
-
+		int32 SelectedIndex = ComboSelectedIndices.IsValidIndex(SlotIndex) ? ComboSelectedIndices[SlotIndex] : -1;
+		ComboBoxes[SlotIndex]->SetSelectedItem(ComboStrings.IsValidIndex(SelectedIndex) ? ComboStrings[SelectedIndex] : ComboStrings[0]);
 #else
 	;
 #endif// #if WITH_APEX_CLOTHING
@@ -227,27 +270,27 @@ TSharedRef<SWidget> FPersonaMeshDetails::OnGenerateCustomMaterialWidgetsForMater
 	return MaterialWidget;
 }
 
-ESlateCheckBoxState::Type FPersonaMeshDetails::IsSectionSelected(int32 SlotIndex) const
+ESlateCheckBoxState::Type FPersonaMeshDetails::IsSectionSelected(int32 SectionIndex) const
 {
 	ESlateCheckBoxState::Type State = ESlateCheckBoxState::Unchecked;
 	USkeletalMesh* Mesh = Cast<USkeletalMesh>(SelectedObjects[0].Get());
 
 	if (Mesh)
 	{
-		State = Mesh->SelectedEditorSection == SlotIndex ? ESlateCheckBoxState::Checked : ESlateCheckBoxState::Unchecked;
+		State = Mesh->SelectedEditorSection == SectionIndex ? ESlateCheckBoxState::Checked : ESlateCheckBoxState::Unchecked;
 	}
 
 	return State;
 }
 
-void FPersonaMeshDetails::OnSectionSelectedChanged(ESlateCheckBoxState::Type NewState, int32 SlotIndex)
+void FPersonaMeshDetails::OnSectionSelectedChanged(ESlateCheckBoxState::Type NewState, int32 SectionIndex)
 {
 	USkeletalMesh* Mesh = Cast<USkeletalMesh>(SelectedObjects[0].Get());
 	if (Mesh)
 	{
 		if (NewState == ESlateCheckBoxState::Checked)
 		{
-			Mesh->SelectedEditorSection = SlotIndex;
+			Mesh->SelectedEditorSection = SectionIndex;
 		}
 		else if (NewState == ESlateCheckBoxState::Unchecked)
 		{
@@ -257,20 +300,20 @@ void FPersonaMeshDetails::OnSectionSelectedChanged(ESlateCheckBoxState::Type New
 	}
 }
 
-ESlateCheckBoxState::Type FPersonaMeshDetails::IsShadowCastingEnabled( int32 SlotIndex ) const
+ESlateCheckBoxState::Type FPersonaMeshDetails::IsShadowCastingEnabled(int32 MaterialIndex) const
 {
 	ESlateCheckBoxState::Type State = ESlateCheckBoxState::Unchecked;
 	USkeletalMesh* Mesh = Cast<USkeletalMesh>( SelectedObjects[0].Get() );
 
-	if (Mesh && SlotIndex < Mesh->Materials.Num())
+	if (Mesh && MaterialIndex < Mesh->Materials.Num())
 	{
-		State = Mesh->Materials[SlotIndex].bEnableShadowCasting ? ESlateCheckBoxState::Checked : ESlateCheckBoxState::Unchecked;
+		State = Mesh->Materials[MaterialIndex].bEnableShadowCasting ? ESlateCheckBoxState::Checked : ESlateCheckBoxState::Unchecked;
 	}
 
 	return State;
 }
 
-void FPersonaMeshDetails::OnShadowCastingChanged( ESlateCheckBoxState::Type NewState, int32 SlotIndex )
+void FPersonaMeshDetails::OnShadowCastingChanged(ESlateCheckBoxState::Type NewState, int32 MaterialIndex)
 {
 	USkeletalMesh* Mesh = Cast<USkeletalMesh>( SelectedObjects[0].Get() );
 
@@ -280,13 +323,13 @@ void FPersonaMeshDetails::OnShadowCastingChanged( ESlateCheckBoxState::Type NewS
 		{
 			const FScopedTransaction Transaction( LOCTEXT( "SetShadowCastingFlag", "Set Shadow Casting For Material" ) );
 			Mesh->Modify();
-			Mesh->Materials[SlotIndex].bEnableShadowCasting = true;
+			Mesh->Materials[MaterialIndex].bEnableShadowCasting = true;
 		}
 		else if (NewState == ESlateCheckBoxState::Unchecked)
 		{
 			const FScopedTransaction Transaction( LOCTEXT( "ClearShadowCastingFlag", "Clear Shadow Casting For Material" ) );
 			Mesh->Modify();
-			Mesh->Materials[SlotIndex].bEnableShadowCasting = false;
+			Mesh->Materials[MaterialIndex].bEnableShadowCasting = false;
 		}
 		for (TObjectIterator<USkinnedMeshComponent> It; It; ++It)
 		{
@@ -302,38 +345,142 @@ void FPersonaMeshDetails::OnShadowCastingChanged( ESlateCheckBoxState::Type NewS
 	}
 }
 
-void FPersonaMeshDetails::OnMaterialChanged( UMaterialInterface* NewMaterial, UMaterialInterface* PrevMaterial, int32 SlotIndex, bool bReplaceAll )
+int32 FPersonaMeshDetails::GetMaterialIndex(int32 LODIndex, int32 SectionIndex)
+{
+	USkeletalMesh* SkelMesh = PersonaPtr->GetMesh();
+
+	check(LODIndex < SkelMesh->LODInfo.Num());
+
+	FSkeletalMeshLODInfo& Info = SkelMesh->LODInfo[LODIndex];
+	if (LODIndex == 0 || Info.LODMaterialMap.Num() == 0)
+	{
+		FSkeletalMeshResource* ImportedResource = SkelMesh->GetImportedResource();
+		check(ImportedResource && ImportedResource->LODModels.IsValidIndex(LODIndex));
+
+		return ImportedResource->LODModels[LODIndex].Sections[SectionIndex].MaterialIndex;
+	}
+	else
+	{
+		check(SectionIndex < Info.LODMaterialMap.Num());
+		return Info.LODMaterialMap[SectionIndex];
+	}
+}
+
+bool FPersonaMeshDetails::IsDuplicatedMaterialIndex(int32 LODIndex, int32 MaterialIndex)
+{
+	USkeletalMesh* SkelMesh = PersonaPtr->GetMesh();
+
+	// finding whether this material index is being used in parent LODs
+	for (int32 LODInfoIdx = 0; LODInfoIdx < LODIndex; LODInfoIdx++)
+	{
+		FSkeletalMeshLODInfo& Info = SkelMesh->LODInfo[LODInfoIdx];
+		if (LODIndex == 0 || Info.LODMaterialMap.Num() == 0)
+		{
+			FSkeletalMeshResource* ImportedResource = SkelMesh->GetImportedResource();
+
+			if (ImportedResource && ImportedResource->LODModels.IsValidIndex(LODInfoIdx))
+			{
+				FStaticLODModel& Model = ImportedResource->LODModels[LODInfoIdx];
+
+				for (int32 SectionIdx = 0; SectionIdx < Model.Sections.Num(); SectionIdx++)
+				{
+					if (MaterialIndex == Model.Sections[SectionIdx].MaterialIndex)
+					{
+						return true;
+					}
+				}
+			}
+		}
+		else // if LODMaterialMap exists
+		{
+			for (int32 MapIdx = 0; MapIdx < Info.LODMaterialMap.Num(); MapIdx++)
+			{
+				if (MaterialIndex == Info.LODMaterialMap[MapIdx])
+				{
+					return true;
+				}
+			}
+		}
+	}
+
+	return false;
+}
+
+void FPersonaMeshDetails::OnMaterialChanged(UMaterialInterface* NewMaterial, UMaterialInterface* PrevMaterial, int32 SlotIndex, bool bReplaceAll, int32 LODIndex)
 {
 	// Whether or not we made a transaction and need to end it
 	bool bMadeTransaction = false;
 
 	USkeletalMesh* Mesh = Cast<USkeletalMesh>(SelectedObjects[0].Get());
+	UProperty* MaterialProperty = FindField<UProperty>(USkeletalMesh::StaticClass(), "Materials");
+	check(MaterialProperty);
+	Mesh->PreEditChange(MaterialProperty);
 
-	for(int MaterialIndex = 0; MaterialIndex < Mesh->Materials.Num(); ++MaterialIndex)
+	if (LODIndex == 0)
 	{
+		check(Mesh->GetImportedResource());
+
+		int MaterialIndex = Mesh->GetImportedResource()->LODModels[LODIndex].Sections[SlotIndex].MaterialIndex;
 		UMaterialInterface* Material = Mesh->Materials[MaterialIndex].MaterialInterface;
-		if(( Material == PrevMaterial || bReplaceAll ) && MaterialIndex == SlotIndex)
+		if (Material == PrevMaterial || bReplaceAll)
 		{
-
 			// Begin a transaction for undo/redo the first time we encounter a material to replace.  
 			// There is only one transaction for all replacement
-			if( !bMadeTransaction )
+			if (!bMadeTransaction)
 			{
-				GEditor->BeginTransaction( LOCTEXT("PersonaReplaceMaterial", "Replace material on mesh") );
+				GEditor->BeginTransaction(LOCTEXT("PersonaReplaceMaterial", "Replace material on mesh"));
 				bMadeTransaction = true;
 			}
 
-			UProperty* MaterialProperty = FindField<UProperty>( USkeletalMesh::StaticClass(), "Materials" );
+			Mesh->Materials[MaterialIndex].MaterialInterface = NewMaterial;
+		}
+	}
+	else
+	{
+		check(Mesh->LODInfo.IsValidIndex(LODIndex));
 
-			Mesh->PreEditChange( MaterialProperty );
+		int32 MaterialIndex;
+		if (Mesh->LODInfo[LODIndex].LODMaterialMap.Num() > 0)
+		{
+			MaterialIndex = Mesh->LODInfo[LODIndex].LODMaterialMap[SlotIndex];
+		}
+		else
+		{
+			MaterialIndex = SlotIndex;
+		}
 
-			Mesh->Materials[MaterialIndex].MaterialInterface = NewMaterial;
+		bool bIsUsedInParentLODs = IsDuplicatedMaterialIndex(LODIndex, MaterialIndex);
 
-			FPropertyChangedEvent PropertyChangedEvent( MaterialProperty );
-			Mesh->PostEditChangeProperty( PropertyChangedEvent );
+		if (bIsUsedInParentLODs)
+		{
+			int32 NewMaterialIndex = Mesh->Materials.Add(FSkeletalMaterial(NewMaterial));
+
+			if (Mesh->LODInfo[LODIndex].LODMaterialMap.Num() > 0)
+			{
+				Mesh->LODInfo[LODIndex].LODMaterialMap[SlotIndex] = NewMaterialIndex;
+			}
+			else
+			{
+				// @TODO : need to overwrite section's material index
+				MaterialIndex = SlotIndex;
+			}
+		}
+		else
+		{
+			// Begin a transaction for undo/redo the first time we encounter a material to replace.  
+			// There is only one transaction for all replacement
+			if (!bMadeTransaction)
+			{
+				GEditor->BeginTransaction(LOCTEXT("PersonaReplaceMaterial", "Replace material on mesh"));
+				bMadeTransaction = true;
+			}
+			Mesh->Materials[MaterialIndex].MaterialInterface = NewMaterial;
 		}
 	}
 
+	FPropertyChangedEvent PropertyChangedEvent(MaterialProperty);
+	Mesh->PostEditChangeProperty(PropertyChangedEvent);
+
 	if( bMadeTransaction )
 	{
 		// End the transation if we created one
@@ -637,84 +784,6 @@ TSharedRef<SUniformGridPanel> FPersonaMeshDetails::MakeApexDetailsWidget(int32 A
 	return Grid;
 }
 
-ESlateCheckBoxState::Type FPersonaMeshDetails::IsEnablingClothingLOD(int32 MaterialIndex) const
-{
-	USkeletalMesh* SkelMesh = PersonaPtr->GetMesh();
-
-	check(SkelMesh);
-
-	FStaticLODModel& Model = SkelMesh->GetImportedResource()->LODModels[0];
-	int32 SectionIndex = ApexClothingUtils::FindSectionByMaterialIndex(SkelMesh, 0, MaterialIndex);
-
-	return Model.Sections.IsValidIndex(SectionIndex) && Model.Sections[SectionIndex].bEnableClothLOD ? ESlateCheckBoxState::Checked : ESlateCheckBoxState::Unchecked;
-}
-
-void FPersonaMeshDetails::OnEnableClothingLOD(ESlateCheckBoxState::Type CheckState, int32 MaterialIndex)
-{
-	USkeletalMesh* SkelMesh = PersonaPtr->GetMesh();
-
-	check(SkelMesh);
-
-	FStaticLODModel& Model = SkelMesh->GetImportedResource()->LODModels[0];
-	
-	int32 SectionIndex = ApexClothingUtils::FindSectionByMaterialIndex(SkelMesh, 0, MaterialIndex);
-
-	if(SectionIndex >= 0)
-	{
-		Model.Sections[SectionIndex].bEnableClothLOD = (CheckState == ESlateCheckBoxState::Checked);
-
-		// re-import when check box is changed
-		ApexClothingUtils::ReImportClothingSectionFromClothingAsset(SkelMesh, SectionIndex);
-		PersonaPtr->GetPreviewMeshComponent()->SetClothingLOD(PersonaPtr->GetPreviewMeshComponent()->PredictedLODLevel);
-	}
-}
-
-EVisibility FPersonaMeshDetails::IsClothingLODCheckBoxVisible(int32 MaterialIndex) const
-{
-	USkeletalMesh* SkelMesh = PersonaPtr->GetMesh();
-	check(SkelMesh);
-
-	int32 SectionIndex = ApexClothingUtils::FindSectionByMaterialIndex(SkelMesh, 0, MaterialIndex);
-
-	// AssetIndex could be invalid while reimporting skeletal mesh
-	int32 AssetIndex = SkelMesh->GetClothAssetIndex(SectionIndex);
-
-	if(!SkelMesh->ClothingAssets.IsValidIndex(AssetIndex))
-	{
-		return EVisibility::Collapsed;
-	}
-
-	FClothingAssetData& Asset = SkelMesh->ClothingAssets[AssetIndex];
-	int32 NumLODs = ApexClothingUtils::GetNumLODs(Asset.ApexClothingAsset->GetAsset());
-	
-	const FClothAssetSubmeshIndex* SubmeshInfo = ClothingComboStringReverseLookup.Find(*ClothingComboBoxes[MaterialIndex]->GetSelectedItem());
-
-	if(!SubmeshInfo)
-	{
-		return EVisibility::Collapsed;
-	}
-
-	// temporarily return Hidden while different asset is being selected from clothing combo box
-	// sometimes it causes asset index inconsistency for a split second
-	if(SubmeshInfo->AssetIndex != AssetIndex)
-	{
-		return EVisibility::Hidden;
-	}
-
-	int32 SubmeshIndex = SubmeshInfo->SubmeshIndex;
-
-	int32 NumLODsHasCorrespondSubmesh = 0;
-	for(int32 LODIndex=0; LODIndex < NumLODs; LODIndex++)
-	{
-		if(SubmeshIndex < ApexClothingUtils::GetNumRenderSubmeshes(Asset.ApexClothingAsset->GetAsset(), LODIndex))
-		{
-			NumLODsHasCorrespondSubmesh++;
-		}
-	}
-
-	return NumLODsHasCorrespondSubmesh > 1 ? EVisibility::Visible : EVisibility::Hidden;
-}
-
 FReply FPersonaMeshDetails::OnReimportApexFileClicked(int32 AssetIndex, IDetailLayoutBuilder* DetailLayout)
 {
 	USkeletalMesh* SkelMesh = PersonaPtr->GetMesh();
@@ -913,13 +982,111 @@ FReply FPersonaMeshDetails::OnOpenClothingFileClicked(IDetailLayoutBuilder* Deta
 
 void FPersonaMeshDetails::UpdateComboBoxStrings()
 {
-	// Clear out old strings
-	ClothingComboStrings.Empty();
-	ClothingComboStringReverseLookup.Empty();
-	ClothingComboStrings.Add(MakeShareable(new FString(TEXT("None"))));
+	ClothingComboLODInfos.Empty();
 
 	USkeletalMesh* SkelMesh = PersonaPtr->GetMesh();
 
+	if (!SkelMesh)
+	{
+		return;
+	}
+
+	int32 NumLODs = SkelMesh->LODInfo.Num();
+
+	for (int32 LODIdx = 0; LODIdx < NumLODs; LODIdx++)
+	{
+		FClothingComboInfo* NewLODInfo = new (ClothingComboLODInfos)FClothingComboInfo;
+
+		TArray<TSharedPtr< STextComboBox >>& ClothingComboBoxes = NewLODInfo->ClothingComboBoxes;
+		TArray<TSharedPtr<FString> >& ClothingComboStrings = NewLODInfo->ClothingComboStrings;
+		TMap<FString, FClothAssetSubmeshIndex>&	ClothingComboStringReverseLookup = NewLODInfo->ClothingComboStringReverseLookup;
+		TArray<int32>& ClothingComboSelectedIndices = NewLODInfo->ClothingComboSelectedIndices;
+
+		// Clear out old strings
+		ClothingComboStrings.Empty();
+		ClothingComboStringReverseLookup.Empty();
+		ClothingComboStrings.Add(MakeShareable(new FString(TEXT("None"))));
+
+		FStaticLODModel& LODModel = SkelMesh->GetImportedResource()->LODModels[LODIdx];
+
+		// # of non clothing sections means # of slots
+		int32 NumNonClothingSections = LODModel.NumNonClothingSections();
+
+		TArray<FClothAssetSubmeshIndex> SectionAssetSubmeshIndices;
+
+		ClothingComboSelectedIndices.Empty(NumNonClothingSections);
+		SectionAssetSubmeshIndices.Empty(NumNonClothingSections);
+
+		for (int32 SecIdx = 0; SecIdx < NumNonClothingSections; SecIdx++)
+		{
+			ClothingComboSelectedIndices.Add(0); // None
+			new(SectionAssetSubmeshIndices)FClothAssetSubmeshIndex(INDEX_NONE, INDEX_NONE);
+		}
+
+		// Setup Asset's sub-mesh indices
+		for (int32 SecIdx = 0; SecIdx < NumNonClothingSections; SecIdx++)
+		{
+			int32 ClothSection = LODModel.Sections[SecIdx].CorrespondClothSectionIndex;
+			if (ClothSection >= 0)
+			{
+				FSkelMeshChunk& Chunk = LODModel.Chunks[LODModel.Sections[ClothSection].ChunkIndex];
+				SectionAssetSubmeshIndices[SecIdx] = FClothAssetSubmeshIndex(Chunk.CorrespondClothAssetIndex, Chunk.ClothAssetSubmeshIndex);
+			}
+		}
+
+		for (int32 AssetIndex = 0; AssetIndex < SkelMesh->ClothingAssets.Num(); AssetIndex++)
+		{
+			FClothingAssetData& ClothingAssetData = SkelMesh->ClothingAssets[AssetIndex];
+
+			TArray<FSubmeshInfo> SubmeshInfos;
+			// if failed to get sub-mesh info, then skip
+			if (!ApexClothingUtils::GetSubmeshInfoFromApexAsset(ClothingAssetData.ApexClothingAsset->GetAsset(), LODIdx, SubmeshInfos))
+			{
+				continue;
+			}
+
+			if (SubmeshInfos.Num() == 1)
+			{
+				FString ComboString = ClothingAssetData.AssetName.ToString();
+				ClothingComboStrings.Add(MakeShareable(new FString(ComboString)));
+				FClothAssetSubmeshIndex& AssetSubmeshIndex = ClothingComboStringReverseLookup.Add(ComboString, FClothAssetSubmeshIndex(AssetIndex, SubmeshInfos[0].SubmeshIndex)); //submesh index is not always same as ordered index
+
+				for (int32 SlotIdx = 0; SlotIdx < NumNonClothingSections; SlotIdx++)
+				{
+					if (SectionAssetSubmeshIndices[SlotIdx] == AssetSubmeshIndex)
+					{
+						ClothingComboSelectedIndices[SlotIdx] = ClothingComboStrings.Num() - 1;
+					}
+				}
+			}
+			else
+			{
+				for (int32 SubIdx = 0; SubIdx < SubmeshInfos.Num(); SubIdx++)
+				{
+					FString ComboString = FText::Format(LOCTEXT("Persona_ApexClothingAsset_Submesh", "{0} submesh {1}"), FText::FromName(ClothingAssetData.AssetName), FText::AsNumber(SubmeshInfos[SubIdx].SubmeshIndex)).ToString(); //submesh index is not always same as SubIdx
+					ClothingComboStrings.Add(MakeShareable(new FString(ComboString)));
+					FClothAssetSubmeshIndex& AssetSubmeshIndex = ClothingComboStringReverseLookup.Add(ComboString, FClothAssetSubmeshIndex(AssetIndex, SubIdx));
+
+					for (int32 SlotIdx = 0; SlotIdx < NumNonClothingSections; SlotIdx++)
+					{
+						if (SectionAssetSubmeshIndices[SlotIdx] == AssetSubmeshIndex)
+						{
+							ClothingComboSelectedIndices[SlotIdx] = NewLODInfo->ClothingComboStrings.Num() - 1;
+						}
+					}
+				}
+			}
+		}
+		// Cause combo boxes to refresh their strings lists and selected value
+		for (int32 i = 0; i < NewLODInfo->ClothingComboBoxes.Num(); i++)
+		{
+			ClothingComboBoxes[i]->RefreshOptions();
+			ClothingComboBoxes[i]->SetSelectedItem(ClothingComboSelectedIndices.IsValidIndex(i) && ClothingComboStrings.IsValidIndex(ClothingComboSelectedIndices[i]) ? ClothingComboStrings[ClothingComboSelectedIndices[i]] : ClothingComboStrings[0]);
+		}
+	}
+
+
+#if 0
 	if(SkelMesh)
 	{
 		int32 NumMaterials = SkelMesh->Materials.Num();
@@ -1005,24 +1172,26 @@ void FPersonaMeshDetails::UpdateComboBoxStrings()
 		ClothingComboBoxes[i]->RefreshOptions();
 		ClothingComboBoxes[i]->SetSelectedItem(ClothingComboSelectedIndices.IsValidIndex(i) && ClothingComboStrings.IsValidIndex(ClothingComboSelectedIndices[i]) ? ClothingComboStrings[ClothingComboSelectedIndices[i]] : ClothingComboStrings[0]);
 	}
+#endif
 }
 
-EVisibility FPersonaMeshDetails::IsClothingComboBoxVisible(int32 MaterialIndex) const
+EVisibility FPersonaMeshDetails::IsClothingComboBoxVisible(int32 LODIndex, int32 SectionIndex) const
 {
 	// if this material doesn't have correspondence with any mesh section, hide this combo box 
 	USkeletalMesh* SkelMesh = PersonaPtr->GetMesh();
 
-	int32 SectionIndex = ApexClothingUtils::FindSectionByMaterialIndex(SkelMesh, 0, MaterialIndex);
+	// ClothingComboStrings.Num() should be greater than 1 because it is including "None"
+	bool bExistClothingLOD = ClothingComboLODInfos.IsValidIndex(LODIndex) && ClothingComboLODInfos[LODIndex].ClothingComboStrings.Num() > 1;
 
-	return SkelMesh && SkelMesh->ClothingAssets.Num() > 0 && SectionIndex >= 0 ? EVisibility::Visible : EVisibility::Collapsed;
+	return SkelMesh && SkelMesh->ClothingAssets.Num() > 0 && bExistClothingLOD ? EVisibility::Visible : EVisibility::Collapsed;
 }
 
-FString FPersonaMeshDetails::HandleSectionsComboBoxGetRowText( TSharedPtr<FString> Section, int32 SlotIndex)
+FString FPersonaMeshDetails::HandleSectionsComboBoxGetRowText( TSharedPtr<FString> Section, int32 LODIndex, int32 SlotIndex)
 {
 	return *Section;
 }
 
-void FPersonaMeshDetails::HandleSectionsComboBoxSelectionChanged( TSharedPtr<FString> SelectedItem, ESelectInfo::Type SelectInfo, int32 SlotIndex )
+void FPersonaMeshDetails::HandleSectionsComboBoxSelectionChanged( TSharedPtr<FString> SelectedItem, ESelectInfo::Type SelectInfo, int32 LODIndex, int32 SectionIndex )
 {		
 	// Selection set by code shouldn't make any content changes.
 	if( SelectInfo == ESelectInfo::Direct )
@@ -1034,20 +1203,17 @@ void FPersonaMeshDetails::HandleSectionsComboBoxSelectionChanged( TSharedPtr<FSt
 
 	FClothAssetSubmeshIndex* ComboItem;
 
-	int32 SectionIndex = ApexClothingUtils::FindSectionByMaterialIndex(SkelMesh, 0, SlotIndex);
-
-	if(SectionIndex < 0)
-	{
-		// if this slot doesn't correspond to any section, set to None
-		ClothingComboBoxes[SlotIndex]->SetSelectedItem(ClothingComboStrings[0]);
-		return;
-	}
+	TArray<TSharedPtr< STextComboBox >>& ComboBoxes = ClothingComboLODInfos[LODIndex].ClothingComboBoxes;
+	TArray<TSharedPtr<FString> >& ComboStrings = ClothingComboLODInfos[LODIndex].ClothingComboStrings;
+	TMap<FString, FClothAssetSubmeshIndex>&	ComboStringReverseLookup = ClothingComboLODInfos[LODIndex].ClothingComboStringReverseLookup;
+	TArray<int32>& ComboSelectedIndices = ClothingComboLODInfos[LODIndex].ClothingComboSelectedIndices;
 
 	// Look up the item
-	if (SelectedItem.IsValid() && (ComboItem = ClothingComboStringReverseLookup.Find(*SelectedItem)) != NULL)
+	if (SelectedItem.IsValid() && (ComboItem = ComboStringReverseLookup.Find(*SelectedItem)) != NULL)
 	{
 		if( !ApexClothingUtils::ImportClothingSectionFromClothingAsset(
 			SkelMesh,
+			LODIndex,
 			SectionIndex,
 			ComboItem->AssetIndex,
 			ComboItem->SubmeshIndex) )
@@ -1055,80 +1221,13 @@ void FPersonaMeshDetails::HandleSectionsComboBoxSelectionChanged( TSharedPtr<FSt
 			// Failed.
 			FMessageDialog::Open( EAppMsgType::Ok, LOCTEXT("ImportClothingSectionFailed", "An error occurred using this clothing asset for this skeletal mesh material section.") );
 			// Change back to previously selected one.
-			int32 SelectedIndex = ClothingComboSelectedIndices.IsValidIndex(SlotIndex) ? ClothingComboSelectedIndices[SlotIndex] : 0;
-			ClothingComboBoxes[SlotIndex]->SetSelectedItem(ClothingComboStrings[SelectedIndex]);
+			int32 SelectedIndex = ComboSelectedIndices.IsValidIndex(SectionIndex) ? ComboSelectedIndices[SectionIndex] : 0;
+			ComboBoxes[SectionIndex]->SetSelectedItem(ComboStrings[SelectedIndex]);
 		}
 	}
 	else
 	{
-		ApexClothingUtils::RestoreOriginalClothingSectionAllLODs(SkelMesh, SectionIndex);
-	}
-}
-
-void FPersonaMeshDetails::CheckLODMaterialMapChanges()
-{
-	// checking whether MaterialMap is changed or not
-	USkeletalMesh* SkelMesh = PersonaPtr->GetMesh();
-	// initialization
-	if(OldLODInfo.Num() == 0)
-	{
-		OldLODInfo = SkelMesh->LODInfo;
-	}
-
-	int32 NumLODInfos = SkelMesh->LODInfo.Num();
-	int32 NumOldLODInfos = OldLODInfo.Num();
-
-	bool bNeedReimport = false;
-	// LOD is added or removed
-	if(NumLODInfos != NumOldLODInfos)
-	{
-		bNeedReimport = true;
-
-		TArray<uint32> SectionIndices;
-		SkelMesh->GetOriginSectionIndicesWithCloth(0, SectionIndices);
-
-		// shows a warning if MaterialMap has duplicated cloth mapping
-		int32 LastIndex = NumLODInfos-1;
-
-		TArray<int32>& MaterialMap = SkelMesh->LODInfo[LastIndex].LODMaterialMap;
-
-		for(int32 Index=0; Index < SectionIndices.Num(); Index++)
-		{
-			int32 SectionIndex = SectionIndices[Index];
-			int32 DuplicateCount = 0;
-			for(int32 MapIdx=0; MapIdx < MaterialMap.Num(); MapIdx++)
-			{
-				if(MaterialMap[MapIdx] == SectionIndex)
-				{
-					DuplicateCount++;
-				}
-			}
-
-			if(DuplicateCount > 1)
-			{
-				FMessageDialog::Open( EAppMsgType::Ok, FText::Format( LOCTEXT("DuplicateClothSection", "{0} cloth section indices are duplicated. Please check LOD {1} LODMaterialMap."), FText::AsNumber( DuplicateCount ), FText::AsNumber( LastIndex ) ) );
-				break;
-			}
-		}
-	}
-	else
-	{
-		for(int32 LODInfoIdx=0; LODInfoIdx < NumLODInfos; LODInfoIdx++)
-		{
-			// If LODMaterialMap is changed, then update clothing sections
-			if(SkelMesh->LODInfo[LODInfoIdx].LODMaterialMap != OldLODInfo[LODInfoIdx].LODMaterialMap)
-			{
-				bNeedReimport = true;
-				break;
-			}
-		}
-	}
-
-	if(bNeedReimport)
-	{
-		ApexClothingUtils::ReImportClothingSectionsFromClothingAsset(SkelMesh);
-		// refresh OldLODInfo
-		OldLODInfo = SkelMesh->LODInfo;
+		ApexClothingUtils::RestoreOriginalClothingSection(SkelMesh, LODIndex, SectionIndex);
 	}
 }
 
@@ -1143,29 +1242,18 @@ void FPersonaMeshDetails::UpdateClothPhysicsProperties(int32 AssetIndex)
 
 #endif// #if WITH_APEX_CLOTHING
 
-
-FReply FPersonaMeshDetails::OnAddMaterialButtonClicked()
-{
-	USkeletalMesh* SkelMesh = PersonaPtr->GetMesh();
-
-	if(SkelMesh)
-	{
-		SkelMesh->Materials.Add(FSkeletalMaterial());
-	}
-
-	return FReply::Handled();
-}
-
-bool FPersonaMeshDetails::CanDeleteMaterialElement( int32 SectionIndex ) const
+bool FPersonaMeshDetails::CanDeleteMaterialElement(int32 LODIndex, int32 SectionIndex) const
 {
 	// Only allow deletion of extra elements
 	return SectionIndex != 0;
 }
 
-FReply FPersonaMeshDetails::OnDeleteButtonClicked( int32 SectionIndex )
+FReply FPersonaMeshDetails::OnDeleteButtonClicked(int32 LODIndex, int32 SectionIndex)
 {
 	ensure(SectionIndex != 0);
 
+	int32 MaterialIndex = GetMaterialIndex(LODIndex, SectionIndex);
+
 	USkeletalMesh* SkelMesh = PersonaPtr->GetMesh();
 
 	// Move any mappings pointing to the requested material to point to the first
@@ -1185,12 +1273,12 @@ FReply FPersonaMeshDetails::OnDeleteButtonClicked( int32 SectionIndex )
 			{
 				int32 CurrentMapping = *LodMaterialIter;
 
-				if(CurrentMapping == SectionIndex)
+				if (CurrentMapping == MaterialIndex)
 				{
 					// Set to first material
 					*LodMaterialIter = 0;
 				}
-				else if(CurrentMapping > SectionIndex)
+				else if (CurrentMapping > MaterialIndex)
 				{
 					// Decrement to keep correct reference after removal
 					*LodMaterialIter = CurrentMapping - 1;
@@ -1206,29 +1294,24 @@ FReply FPersonaMeshDetails::OnDeleteButtonClicked( int32 SectionIndex )
 			{
 				FSkelMeshSection& Section = *SectionIter;
 
-				if(Section.MaterialIndex == SectionIndex)
+				if (Section.MaterialIndex == MaterialIndex)
 				{
 					Section.MaterialIndex = 0;
 				}
-				else if(Section.MaterialIndex > SectionIndex)
+				else if (Section.MaterialIndex > MaterialIndex)
 				{
 					Section.MaterialIndex--;
 				}
 			}
 		}
 
-		SkelMesh->Materials.RemoveAt(SectionIndex);
+		SkelMesh->Materials.RemoveAt(MaterialIndex);
 
 		// Notify the change in material
 		FPropertyChangedEvent PropertyChangedEvent( MaterialProperty );
 		SkelMesh->PostEditChangeProperty(PropertyChangedEvent);
 	}
 
-#if WITH_APEX_CLOTHING
-	// Check LOD map changes and update if necessary
-	CheckLODMaterialMapChanges();
-#endif// #if WITH_APEX_CLOTHING
-
 	return FReply::Handled();
 }
 
diff --git a/Engine/Source/Editor/Persona/Private/PersonaMeshDetails.h b/Engine/Source/Editor/Persona/Private/PersonaMeshDetails.h
--- a/Engine/Source/Editor/Persona/Private/PersonaMeshDetails.h
+++ b/Engine/Source/Editor/Persona/Private/PersonaMeshDetails.h
@@ -24,6 +24,18 @@ struct FClothAssetSubmeshIndex
 	}
 };
 
+struct FClothingComboInfo
+{
+	/* Per-material clothing combo boxes, array size must be same to # of sections */
+	TArray<TSharedPtr< STextComboBox >>		ClothingComboBoxes;
+	/* Clothing combo box strings */
+	TArray<TSharedPtr<FString> >			ClothingComboStrings;
+	/* Mapping from a combo box string to the asset and submesh it was generated from */
+	TMap<FString, FClothAssetSubmeshIndex>	ClothingComboStringReverseLookup;
+	/* The currently-selected index from each clothing combo box */
+	TArray<int32>							ClothingComboSelectedIndices;
+};
+
 class FPersonaMeshDetails : public IDetailCustomization
 
 {
@@ -42,7 +54,7 @@ private:
 	 *
 	 * @param OutMaterials	Handle to a material list builder that materials should be added to
 	 */
-	void OnGetMaterialsForView( class IMaterialListBuilder& OutMaterials );
+	void OnGetMaterialsForView( class IMaterialListBuilder& OutMaterials, int32 LODIndex );
 
 	/**
 	 * Called when a user drags a new material over a list item to replace it
@@ -52,7 +64,7 @@ private:
 	 * @param SlotIndex		The index of the slot on the component where materials should be replaces
 	 * @param bReplaceAll	If true all materials in the slot should be replaced not just ones using PrevMaterial
 	 */
-	void OnMaterialChanged( UMaterialInterface* NewMaterial, UMaterialInterface* PrevMaterial, int32 SlotIndex, bool bReplaceAll );
+	void OnMaterialChanged(UMaterialInterface* NewMaterial, UMaterialInterface* PrevMaterial, int32 SlotIndex, bool bReplaceAll, int32 LODIndex);
 
 	
 	/**
@@ -61,7 +73,7 @@ private:
 	 * @param Material		The material that is being displayed
 	 * @param SlotIndex		The index of the material slot
 	 */
-	TSharedRef<SWidget> OnGenerateCustomNameWidgetsForMaterial(UMaterialInterface* Material, int32 SlotIndex);
+	TSharedRef<SWidget> OnGenerateCustomNameWidgetsForMaterial(UMaterialInterface* Material, int32 SlotIndex, int32 LODIndex);
 
 	/**
 	 * Called by the material list widget on generating each thumbnail widget
@@ -69,7 +81,7 @@ private:
 	 * @param Material		The material that is being displayed
 	 * @param SlotIndex		The index of the material slot
 	 */
-	TSharedRef<SWidget> OnGenerateCustomMaterialWidgetsForMaterial(UMaterialInterface* Material, int32 SlotIndex);
+	TSharedRef<SWidget> OnGenerateCustomMaterialWidgetsForMaterial(UMaterialInterface* Material, int32 SlotIndex, int32 LODIndex);
 
 	/**
 	 * Handler for check box display based on whether the material is highlighted
@@ -88,35 +100,38 @@ private:
 		/**
 	 * Handler for check box display based on whether the material has shadow casting enabled
 	 *
-	 * @param SectionIndex	The material section that is being tested
+	 * @param MaterialIndex	The material index which a section in a specific LOD model has
 	 */
-	ESlateCheckBoxState::Type IsShadowCastingEnabled(int32 SectionIndex) const;
+	ESlateCheckBoxState::Type IsShadowCastingEnabled(int32 MaterialIndex) const;
 
 	/**
 	 * Handler for changing shadow casting status on a material
 	 *
-	 * @param SectionIndex	The material section that is being tested
-	 */
-	void OnShadowCastingChanged(ESlateCheckBoxState::Type NewState, int32 SectionIndex);
-
-	/**
-	 * Handler for adding a material element when the user clicks a button
+	 * @param MaterialIndex	The material index which a section in a specific LOD model has
 	 */
-	FReply OnAddMaterialButtonClicked();
+	void OnShadowCastingChanged(ESlateCheckBoxState::Type NewState, int32 MaterialIndex);
 
 	/**
 	 * Handler for enabling delete button on materials
 	 *
 	 * @param SectionIndex - index of the section to check
 	 */
-	bool CanDeleteMaterialElement(int32 SectionIndex) const;
+	bool CanDeleteMaterialElement(int32 LODIndex, int32 SectionIndex) const;
 
 	/**
 	 * Handler for deleting material elements
 	 * 
 	 * @Param SectionIndex - material section to remove
 	 */
-	FReply OnDeleteButtonClicked(int32 SectionIndex);
+	FReply OnDeleteButtonClicked(int32 LODIndex, int32 SectionIndex);
+
+	/** Creates the UI for Current LOD panel */
+	void AddLODLevelCategories(IDetailLayoutBuilder& DetailLayout);
+
+	bool IsDuplicatedMaterialIndex(int32 LODIndex, int32 MaterialIndex);
+
+	/** Get a material index from LOD index and section index */
+	int32 GetMaterialIndex(int32 LODIndex, int32 SectionIndex);
 
 private:
 	// Container for the objects to display
@@ -125,24 +140,19 @@ private:
 	// Pointer back to Persona
 	TSharedPtr<FPersona> PersonaPtr;
 
+	IDetailLayoutBuilder* MeshDetailLayout;
 
 #if WITH_APEX_CLOTHING
 private:
-	/* Per-material clothing combo boxes, array size must be same to # of sections */
-	TArray<TSharedPtr< STextComboBox >>		ClothingComboBoxes;
-	/* Clothing combo box strings */
-	TArray<TSharedPtr<FString> >			ClothingComboStrings;
-	/* Mapping from a combo box string to the asset and submesh it was generated from */
-	TMap<FString, FClothAssetSubmeshIndex>	ClothingComboStringReverseLookup;
-	/* The currently-selected index from each clothing combo box */
-	TArray<int32>							ClothingComboSelectedIndices;
-	/* To check the changes of LOD info */
-	TArray<struct FSkeletalMeshLODInfo>		OldLODInfo;
+
+	// info about clothing combo boxes for multiple LOD
+	TArray<FClothingComboInfo>				ClothingComboLODInfos;
+
 
 	/* Clothing combo box functions */
-	EVisibility IsClothingComboBoxVisible(int32 MaterialIndex) const;
-	FString HandleSectionsComboBoxGetRowText( TSharedPtr<FString> Section, int32 SectionIndex );
-	void HandleSectionsComboBoxSelectionChanged( TSharedPtr<FString> SelectedItem, ESelectInfo::Type SelectInfo, int32 SlotIndex );
+	EVisibility IsClothingComboBoxVisible(int32 LODIndex, int32 SectionIndex) const;
+	FString HandleSectionsComboBoxGetRowText(TSharedPtr<FString> Section, int32 LODIndex, int32 SectionIndex);
+	void HandleSectionsComboBoxSelectionChanged(TSharedPtr<FString> SelectedItem, ESelectInfo::Type SelectInfo, int32 LODIndex, int32 SectionIndex );
 	void UpdateComboBoxStrings();
 
 	/* Generate slate UI for Clothing category */
@@ -163,13 +173,6 @@ private:
 	/* Removes a clothing asset */ 
 	FReply OnRemoveApexFileClicked(int32 AssetIndex, IDetailLayoutBuilder* DetailLayout);
 
-	/* Enabling clothing LOD functions */
-	ESlateCheckBoxState::Type IsEnablingClothingLOD(int32 MaterialIndex) const;
-	void OnEnableClothingLOD(ESlateCheckBoxState::Type CheckState, int32 MaterialIndex);
-	EVisibility IsClothingLODCheckBoxVisible(int32 MaterialIndex) const;
-
-	/* if LODMaterialMap is changed, then re-map clothing sections by changed info */
-	void CheckLODMaterialMapChanges();
 
 	/* if physics properties are changed, then save to the clothing asset */
 	void UpdateClothPhysicsProperties(int32 AssetIndex);
diff --git a/Engine/Source/Editor/Persona/Private/SAnimationEditorViewport.cpp b/Engine/Source/Editor/Persona/Private/SAnimationEditorViewport.cpp
--- a/Engine/Source/Editor/Persona/Private/SAnimationEditorViewport.cpp
+++ b/Engine/Source/Editor/Persona/Private/SAnimationEditorViewport.cpp
@@ -17,6 +17,7 @@
 #include "AnimationEditorViewportClient.h"
 #include "ScopedTransaction.h"
 #include "Editor/UnrealEd/Public/LODUtilities.h"
+#include "DetailLayoutBuilder.h"
 
 #define LOCTEXT_NAMESPACE "PersonaViewportToolbar"
 
@@ -1355,6 +1356,11 @@ void SAnimationEditorViewportTabBody::OnLODChanged()
 	{
 		OnSetLODModel( LOD_Auto );
 	}
+
+	if (PersonaPtr.Pin()->PersonaMeshDetailLayout)
+	{
+		PersonaPtr.Pin()->PersonaMeshDetailLayout->ForceRefreshDetails();
+	}
 }
 
 void SAnimationEditorViewportTabBody::OnMuteAudio()
diff --git a/Engine/Source/Editor/UnrealEd/Private/ApexClothingUtils.cpp b/Engine/Source/Editor/UnrealEd/Private/ApexClothingUtils.cpp
--- a/Engine/Source/Editor/UnrealEd/Private/ApexClothingUtils.cpp
+++ b/Engine/Source/Editor/UnrealEd/Private/ApexClothingUtils.cpp
@@ -23,8 +23,6 @@ DEFINE_LOG_CATEGORY_STATIC(LogApexClothingUtils, Log, All);
 
 #define LOCTEXT_NAMESPACE "ApexClothingUtils"
 
-#define USE_MATERIAL_MAP 1
-
 namespace ApexClothingUtils
 {
 #if WITH_APEX_CLOTHING
@@ -1425,113 +1423,45 @@ bool ImportClothingSectionFromClothingAsset( USkeletalMesh* SkelMesh, uint32 LOD
 											BoneNames, IndexBuffer, SubmeshTriangleCount);
 }
 
-
-bool ImportClothingSectionFromClothingAsset( USkeletalMesh* SkelMesh, uint32 SectionIndex, int32 AssetIndex, int32 AssetSubmeshIndex)
+void ReImportClothingSectionsFromClothingAsset(USkeletalMesh* SkelMesh)
 {
-	check( SkelMesh->ClothingAssets.IsValidIndex(AssetIndex) );
-
-	int32 NumMeshLOD = SkelMesh->LODInfo.Num();
+	check(SkelMesh);
 
-	int32 NumImportLOD;
+	FSkeletalMeshResource* ImportedResource= SkelMesh->GetImportedResource();
 
-	int32 NumAssetLOD = GetNumLODs(SkelMesh->ClothingAssets[AssetIndex].ApexClothingAsset->GetAsset());
+	int32 NumLODs = ImportedResource->LODModels.Num();
 
-	FSkeletalMeshResource* ImportedResource = SkelMesh->GetImportedResource();
-	if(ImportedResource->LODModels[0].Sections[SectionIndex].bEnableClothLOD)
+	for (int32 LODIndex = 0; LODIndex < NumLODs; LODIndex++)
 	{
-		NumImportLOD = FMath::Min(NumMeshLOD, NumAssetLOD);
-	}
-	else
-	{
-		NumImportLOD = 1;
-	}
+		int32 NumSections = ImportedResource->LODModels[LODIndex].NumNonClothingSections();
 
-	bool bResult = false;
-	for(int32 LODIndex=0; LODIndex < NumImportLOD; LODIndex++)
-	{
-		uint32 NumNonClothSection = ImportedResource->LODModels[LODIndex].NumNonClothingSections();
-		if(SectionIndex < NumNonClothSection)
+		for (int32 SectionIndex = 0; SectionIndex < NumSections; SectionIndex++)
 		{
-			int32 SecIdx = SectionIndex;
-	#if USE_MATERIAL_MAP
-			int32 NumMaterialMap = SkelMesh->LODInfo[LODIndex].LODMaterialMap.Num();
-			if(NumMaterialMap > 0)
-			{
-				int32 MaterialIndex = ImportedResource->LODModels[LODIndex].Sections[SectionIndex].MaterialIndex;
-				int32 NewMatIndex = SkelMesh->LODInfo[LODIndex].LODMaterialMap[MaterialIndex];
-				SecIdx = FindSectionByMaterialIndex(SkelMesh, LODIndex, NewMatIndex);
-				if(SecIdx < 0)
-				{ 
-					continue;
-				}
-			}
-	#endif// #if USE_MATERIAL_MAP
-			bResult |= ImportClothingSectionFromClothingAsset(SkelMesh, LODIndex, SecIdx, AssetIndex, AssetSubmeshIndex);
-		}
-	}
-
-	// Restores other LODs if mesh LOD is greater than asset LOD
-	for(int32 LODIndex=NumImportLOD; LODIndex < NumMeshLOD; LODIndex++)
-	{
-		FStaticLODModel& LODModel = ImportedResource->LODModels[LODIndex];
-		uint32 NumNonClothSection = LODModel.NumNonClothingSections();
-		if(SectionIndex < NumNonClothSection)
-		{
-			int32 SecIdx = SectionIndex;
-	#if USE_MATERIAL_MAP
-			if(SkelMesh->LODInfo[LODIndex].LODMaterialMap.Num() > 0)
-			{
-				int32 MaterialIndex = LODModel.Sections[SectionIndex].MaterialIndex;
-				int32 NewMatIndex = SkelMesh->LODInfo[LODIndex].LODMaterialMap[MaterialIndex];
-				SecIdx = FindSectionByMaterialIndex(SkelMesh, LODIndex, NewMatIndex);
-				if(SecIdx < 0)
-				{ 
-					continue;
-				}
-			}
-	#endif// #if USE_MATERIAL_MAP
-			RestoreOriginalClothingSection(SkelMesh, LODIndex, SecIdx);
+			ReImportClothingSectionFromClothingAsset(SkelMesh, LODIndex, SectionIndex);
 		}
 	}
-
-	return bResult;
 }
 
-void ReImportClothingSectionsFromClothingAsset(USkeletalMesh* SkelMesh)
+void ReImportClothingSectionFromClothingAsset(USkeletalMesh* SkelMesh, int32 LODIndex, uint32 SectionIndex)
 {
 	check(SkelMesh);
 
 	FSkeletalMeshResource* ImportedResource= SkelMesh->GetImportedResource();
-	int32 NumSections = ImportedResource->LODModels[0].NumNonClothingSections();
+	FStaticLODModel& Model = ImportedResource->LODModels[LODIndex];
 
-	for(int32 SectionIndex=0; SectionIndex < NumSections; SectionIndex++)
-	{
-		ReImportClothingSectionFromClothingAsset(SkelMesh, SectionIndex);
-	}
-}
+	check(Model.Sections.IsValidIndex(SectionIndex));
 
-void ReImportClothingSectionFromClothingAsset(USkeletalMesh* SkelMesh, uint32 SectionIndex)
-{
-	check(SkelMesh);
-
-	FSkeletalMeshResource* ImportedResource= SkelMesh->GetImportedResource();
-	FStaticLODModel& LOD_0_Model = ImportedResource->LODModels[0];
-
-	check(LOD_0_Model.Sections.IsValidIndex(SectionIndex));
-
-	int16 ClothSecIdx = LOD_0_Model.Sections[SectionIndex].CorrespondClothSectionIndex;
+	int16 ClothSecIdx = Model.Sections[SectionIndex].CorrespondClothSectionIndex;
 	if(ClothSecIdx < 0)
 	{
-		// if there are sections which has clothing in other LODs, restore all LODs
-		RestoreOriginalClothingSectionAllLODs(SkelMesh, SectionIndex);
 		return;
 	}
 
-	int16 ClothChunkIdx = LOD_0_Model.Sections[ClothSecIdx].ChunkIndex;	
-	int16 AssetIndex = LOD_0_Model.Chunks[ClothChunkIdx].CorrespondClothAssetIndex;
-	int16 SubmeshIdx = LOD_0_Model.Chunks[ClothChunkIdx].ClothAssetSubmeshIndex;
+	int16 ClothChunkIdx = Model.Sections[ClothSecIdx].ChunkIndex;
+	int16 AssetIndex = Model.Chunks[ClothChunkIdx].CorrespondClothAssetIndex;
+	int16 SubmeshIdx = Model.Chunks[ClothChunkIdx].ClothAssetSubmeshIndex;
 
-	ImportClothingSectionFromClothingAsset(SkelMesh, SectionIndex, AssetIndex, SubmeshIdx);
+	ImportClothingSectionFromClothingAsset(SkelMesh, LODIndex, SectionIndex, AssetIndex, SubmeshIdx);
 }
 
 static bool IsProperApexFile(NxClothingAsset& ApexClothingAsset, USkeletalMesh* SkelMesh)
@@ -1608,66 +1538,27 @@ EClothUtilRetType ImportApexAssetFromApexFile(FString& ApexFile, USkeletalMesh*
 
 		FSkeletalMeshResource* ImportedResource= SkelMesh->GetImportedResource();
 		int32 NumLODs = ImportedResource->LODModels.Num();
-		FStaticLODModel& LOD_0_Model = ImportedResource->LODModels[0];
-
-		// clear mapping data of LODs other than LOD_0
-		for(int32 LODIndex=1; LODIndex < NumLODs; LODIndex++)
-		{
-			TArray<uint32> SectionIndicesLOD;
-			SkelMesh->GetOriginSectionIndicesWithCloth(LODIndex, AssetIndex, SectionIndicesLOD);
-
-			for(int32 i=0; i < SectionIndicesLOD.Num(); i++)
-			{
-				RestoreOriginalClothingSection(SkelMesh, LODIndex, SectionIndicesLOD[i]);
-			}
-		}
-
-		TArray<uint32> SectionIndices;
-		SkelMesh->GetOriginSectionIndicesWithCloth(0, AssetIndex, SectionIndices);
-
-		
-		if(!SkelMesh->IsEnabledClothLOD(AssetIndex))
-		{
-			NumLODs = 1;
-		}
 
 		for(int32 LODIndex=0; LODIndex < NumLODs; LODIndex++)
 		{
+			TArray<uint32> SectionIndices;
+			SkelMesh->GetOriginSectionIndicesWithCloth(LODIndex, AssetIndex, SectionIndices);
+
 			FStaticLODModel& LODModel = ImportedResource->LODModels[LODIndex];
 
 			for(int32 i=0; i < SectionIndices.Num(); i++)
 			{
-				//use LOD_0 info
-				FSkelMeshSection& Section = LOD_0_Model.Sections[SectionIndices[i]];
-
-				if(!Section.bEnableClothLOD && LODIndex > 0)
-				{
-					continue;
-				}
+				FSkelMeshSection& Section = LODModel.Sections[SectionIndices[i]];
 
 				int16 ClothSecIdx = Section.CorrespondClothSectionIndex;
-				int16 ClothChunkIdx = LOD_0_Model.Sections[ClothSecIdx].ChunkIndex;	
-				int16 SubmeshIdx = LOD_0_Model.Chunks[ClothChunkIdx].ClothAssetSubmeshIndex;
+				int16 ClothChunkIdx = LODModel.Sections[ClothSecIdx].ChunkIndex;
+				int16 SubmeshIdx = LODModel.Chunks[ClothChunkIdx].ClothAssetSubmeshIndex;
 
 				int32 SecIdx = SectionIndices[i];
 
 				int32 NumNonClothSection = LODModel.NumNonClothingSections();
 				if(SecIdx < NumNonClothSection)
 				{
-	#if USE_MATERIAL_MAP
-					int32 NumMaterialMap = SkelMesh->LODInfo[LODIndex].LODMaterialMap.Num();
-					if(NumMaterialMap > 0)
-					{
-						int32 MaterialIndex = LODModel.Sections[SecIdx].MaterialIndex;
-						check(MaterialIndex < NumMaterialMap);
-						int32 NewMatIndex = SkelMesh->LODInfo[LODIndex].LODMaterialMap[MaterialIndex];
-						SecIdx = FindSectionByMaterialIndex(SkelMesh, LODIndex, NewMatIndex);
-						if(SecIdx < 0)
-						{ 
-							continue;
-						}
-					}
-	#endif// #if USE_MATERIAL_MAP
 					if(!ImportClothingSectionFromClothingAsset(SkelMesh, LODIndex, SecIdx, AssetIndex, SubmeshIdx))
 					{
 						// if import is failed, restores original section because sub-mesh might be removed
@@ -1926,8 +1817,6 @@ void ReapplyClothingDataToSkeletalMesh( USkeletalMesh* SkelMesh, FClothingBackup
 			continue;
 		}
 
-		LODModel.Sections[SecIdx].bEnableClothLOD = Section.bEnableClothLOD;
-
 		int16 ChunkIndex = Section.ChunkIndex;
 
 		FSkelMeshChunk& Chunk = ClothingBackup.Chunks[ChunkIndex];
@@ -1972,6 +1861,7 @@ void GetPhysicsPropertiesFromApexAsset(NxClothingAsset *InAsset, FClothPhysicsPr
 	const NxParameterized::Interface* AssetParams = InAsset->getAssetNxParameterized();
 
 	uint32 MaterialIndex;
+	// uses default material index
 	verify(NxParameterized::getParamU32(*AssetParams, "materialIndex", MaterialIndex));
 	
 	// ClothingMaterialLibraryParameters 
diff --git a/Engine/Source/Editor/UnrealEd/Public/ApexClothingUtils.h b/Engine/Source/Editor/UnrealEd/Public/ApexClothingUtils.h
--- a/Engine/Source/Editor/UnrealEd/Public/ApexClothingUtils.h
+++ b/Engine/Source/Editor/UnrealEd/Public/ApexClothingUtils.h
@@ -81,15 +81,15 @@ namespace ApexClothingUtils
 
 #if WITH_APEX_CLOTHING
 
-	//if AssetIndex is -1, means newly added asset, otherwise re-import
+	// if AssetIndex is -1, means newly added asset, otherwise re-import
 	UNREALED_API EClothUtilRetType ImportApexAssetFromApexFile(FString& ApexFile,USkeletalMesh* SkelMesh, int32 AssetIndex=-1);
 	UNREALED_API bool GetSubmeshInfoFromApexAsset(physx::apex::NxClothingAsset *InAsset, uint32 LODIndex, TArray<FSubmeshInfo>& OutSubmeshInfos);
 
-	//Imported LOD is decided according to bUseLOD in clothing asset
-	UNREALED_API bool ImportClothingSectionFromClothingAsset(USkeletalMesh* SkelMesh, uint32 SectionIndex, int32 AssetIndex, int32 AssetSubmeshIndex);
+	// import cloth by a specifed section and LOD index
+	UNREALED_API bool ImportClothingSectionFromClothingAsset(USkeletalMesh* SkelMesh, uint32 LODIndex, uint32 SectionIndex, int32 AssetIndex, int32 AssetSubmeshIndex);
 
 	UNREALED_API void ReImportClothingSectionsFromClothingAsset(USkeletalMesh* SkelMesh);
-	UNREALED_API void ReImportClothingSectionFromClothingAsset(USkeletalMesh* SkelMesh, uint32 SectionIndex);
+	UNREALED_API void ReImportClothingSectionFromClothingAsset(USkeletalMesh* SkelMesh, int32 LODIndex, uint32 SectionIndex);
 
 	UNREALED_API void BackupClothingDataFromSkeletalMesh(USkeletalMesh* SkelMesh, FClothingBackup& ClothingBackup);
 	UNREALED_API void ReapplyClothingDataToSkeletalMesh(USkeletalMesh* SkelMesh, FClothingBackup& ClothingBackup);
@@ -98,7 +98,7 @@ namespace ApexClothingUtils
 
 	UNREALED_API physx::apex::NxClothingAsset* CreateApexClothingAssetFromBuffer(const uint8* Buffer, int32 BufferSize);
 
-	// if don't specify MaterialIndex, default material index will be used 
+	// if MaterialIndex is not specified, default material index will be used
 	UNREALED_API void GetPhysicsPropertiesFromApexAsset(physx::apex::NxClothingAsset *InAsset, FClothPhysicsProperties& OutPropertyInfo);
 	UNREALED_API void SetPhysicsPropertiesToApexAsset(physx::apex::NxClothingAsset *InAsset, FClothPhysicsProperties& InPropertyInfo);
 #endif // #if WITH_APEX_CLOTHING
diff --git a/Engine/Source/Runtime/Engine/Classes/Engine/SkeletalMesh.h b/Engine/Source/Runtime/Engine/Classes/Engine/SkeletalMesh.h
--- a/Engine/Source/Runtime/Engine/Classes/Engine/SkeletalMesh.h
+++ b/Engine/Source/Runtime/Engine/Classes/Engine/SkeletalMesh.h
@@ -759,14 +759,15 @@ public:
 	ENGINE_API void InitMorphTargets();
 
 #if WITH_APEX_CLOTHING
+	ENGINE_API bool  HasClothSectionsInAllLODs(int AssetIndex);
 	ENGINE_API bool	 HasClothSections(int32 LODIndex,int AssetIndex);
 	ENGINE_API void	 GetOriginSectionIndicesWithCloth(int32 LODIndex, TArray<uint32>& OutSectionIndices);
 	ENGINE_API void	 GetOriginSectionIndicesWithCloth(int32 LODIndex, int32 AssetIndex, TArray<uint32>& OutSectionIndices);
 	ENGINE_API void	 GetClothSectionIndices(int32 LODIndex, int32 AssetIndex, TArray<uint32>& OutSectionIndices);
 	//moved from ApexClothingUtils because of compile issues
 	ENGINE_API void  LoadClothCollisionVolumes(int32 AssetIndex, physx::apex::NxClothingAsset* ClothingAsset);
-	ENGINE_API bool  IsEnabledClothLOD(int32 AssetIndex);
-	ENGINE_API int32 GetClothAssetIndex(int32 SectionIndex);
+	ENGINE_API bool IsMappedClothingLOD(int32 LODIndex, int32 AssetIndex);
+	ENGINE_API int32 GetClothAssetIndex(int32 LODIndex, int32 SectionIndex);
 #endif// #if WITH_APEX_CLOTHING
 
 private:
diff --git a/Engine/Source/Runtime/Engine/Private/SkeletalMesh.cpp b/Engine/Source/Runtime/Engine/Private/SkeletalMesh.cpp
--- a/Engine/Source/Runtime/Engine/Private/SkeletalMesh.cpp
+++ b/Engine/Source/Runtime/Engine/Private/SkeletalMesh.cpp
@@ -1146,7 +1146,7 @@ FArchive& operator<<(FArchive& Ar,FSkelMeshSection& S)
 		
 		Ar << S.TriangleSorting;
 
-	//for clothing info
+	// for clothing info
 	if( Ar.UE4Ver() >= VER_UE4_APEX_CLOTH )
 	{
 		Ar << S.bDisabled;
@@ -1155,7 +1155,7 @@ FArchive& operator<<(FArchive& Ar,FSkelMeshSection& S)
 
 	if( Ar.UE4Ver() >= VER_UE4_APEX_CLOTH_LOD )
 	{
-		Ar << S.bEnableClothLOD;
+		Ar << S.bEnableClothLOD_DEPRECATED;
 	}
 
 	return Ar;
diff --git a/Engine/Source/Runtime/Engine/Private/SkeletalMeshComponentPhysics.cpp b/Engine/Source/Runtime/Engine/Private/SkeletalMeshComponentPhysics.cpp
--- a/Engine/Source/Runtime/Engine/Private/SkeletalMeshComponentPhysics.cpp
+++ b/Engine/Source/Runtime/Engine/Private/SkeletalMeshComponentPhysics.cpp
@@ -205,6 +205,18 @@ void USkeletalMesh::LoadClothCollisionVolumes(int32 AssetIndex, NxClothingAsset*
 	}
 }
 
+bool USkeletalMesh::HasClothSectionsInAllLODs(int AssetIndex)
+{
+	bool bResult = false;
+
+	for (int32 LODIndex = 0; LODIndex < LODInfo.Num(); LODIndex++)
+	{
+		bResult |= HasClothSections(LODIndex, AssetIndex);
+	}
+
+	return bResult;
+}
+
 bool USkeletalMesh::HasClothSections(int32 LODIndex, int AssetIndex)
 {
 	FSkeletalMeshResource* Resource = GetImportedResource();
@@ -289,44 +301,56 @@ void USkeletalMesh::GetOriginSectionIndicesWithCloth(int32 LODIndex, int32 Asset
 	}
 }
 
-bool USkeletalMesh::IsEnabledClothLOD(int32 AssetIndex)
+bool USkeletalMesh::IsMappedClothingLOD(int32 InLODIndex, int32 InAssetIndex)
 {
 	FSkeletalMeshResource* Resource = GetImportedResource();
 
-	FStaticLODModel& LODModel = Resource->LODModels[0];
+	if (Resource && Resource->LODModels.IsValidIndex(InLODIndex))
+	{
+		FStaticLODModel& LODModel = Resource->LODModels[InLODIndex];
 
-	TArray<uint32> SectionIndices;
-	GetOriginSectionIndicesWithCloth(0, AssetIndex, SectionIndices);
+		int32 NumSections = LODModel.Sections.Num();
 
-	//if found more than 1 section enabled, return true
-	for(int32 i=0; i < SectionIndices.Num(); i++)
-	{
-		if(LODModel.Sections[SectionIndices[i]].bEnableClothLOD)
+		// loop reversely for optimized search
+		for (int32 SecIdx = NumSections-1; SecIdx >= 0; SecIdx--)
 		{
-			return true;
+			int32 ClothAssetIndex = LODModel.Chunks[LODModel.Sections[SecIdx].ChunkIndex].CorrespondClothAssetIndex;
+			if (ClothAssetIndex == InAssetIndex)
+			{
+				return true;
+			}
+			else if (ClothAssetIndex == INDEX_NONE) // no more cloth sections
+			{
+				return false;
+			}
 		}
 	}
 
 	return false;
 }
 
-int32 USkeletalMesh::GetClothAssetIndex(int32 SectionIndex)
+int32 USkeletalMesh::GetClothAssetIndex(int32 LODIndex, int32 SectionIndex)
 {
 	FSkeletalMeshResource* Resource = GetImportedResource();
 
-	FStaticLODModel& LODModel = Resource->LODModels[0];
+	// no LODs
+	if (!Resource || !Resource->LODModels.IsValidIndex(LODIndex))
+	{
+		return INDEX_NONE;
+	}
+	FStaticLODModel& LODModel = Resource->LODModels[LODIndex];
 
-	//no sections
+	// no sections
 	if(!LODModel.Sections.IsValidIndex(SectionIndex))
 	{
-		return -1;
+		return INDEX_NONE;
 	}
 	int16 ClothSecIdx = LODModel.Sections[SectionIndex].CorrespondClothSectionIndex;
 
-	//no mapping
+	// no mapping
 	if(ClothSecIdx < 0)
 	{
-		return -1;
+		return INDEX_NONE;
 	}
 
 	int16 ChunkIdx = LODModel.Sections[ClothSecIdx].ChunkIndex;
@@ -1914,8 +1938,8 @@ void USkeletalMeshComponent::ValidateClothingActors()
 	{
 		FClothingAssetData& ClothAsset = SkeletalMesh->ClothingAssets[AssetIdx];
 
-		// if there exist mapped sections, create a clothing actor
-		if(SkeletalMesh->HasClothSections(0, AssetIdx))
+		// if there exist mapped sections for all LODs, create a clothing actor
+		if (SkeletalMesh->HasClothSectionsInAllLODs(AssetIdx))
 		{
 			if(CreateClothingActor(AssetIdx, ClothAsset.ApexClothingAsset))
 			{
@@ -1923,10 +1947,10 @@ void USkeletalMeshComponent::ValidateClothingActors()
 			}
 		}
 		else 
-		{	//don't have cloth sections but clothing actor is alive
+		{	// don't have cloth sections but a clothing actor is alive
 			if(IsValidClothingActor(AssetIdx))
 			{
-				//clear because mapped sections are removed
+				// clear this clothing actor because mapped sections are removed
 				ClothingActors[AssetIdx].Clear(true);
 			}
 		}
@@ -2074,42 +2098,44 @@ void USkeletalMeshComponent::SetClothingLOD(int32 LODIndex)
 
 	bool bFrozen = false;
 
-	for(int32 i=0; i<NumActors; i++)
+	for (int32 AssetIndex = 0; AssetIndex<NumActors; AssetIndex++)
 	{
-		if(IsValidClothingActor(i))
+		if (IsValidClothingActor(AssetIndex))
 		{			
-			int32 CurLODIndex = (int32)ClothingActors[i].ApexClothingActor->getGraphicalLod();
+			FClothingActor& Actor = ClothingActors[AssetIndex];
 
-			int32 NumClothLODs = ClothingActors[i].ParentClothingAsset->GetAsset()->getNumGraphicalLodLevels();
+			int32 CurLODIndex = (int32)Actor.ApexClothingActor->getGraphicalLod();
 
-			bool bEnabledLOD = SkeletalMesh->IsEnabledClothLOD(i);
+			// check whether clothing LOD is mapped for this LOD index
+			bool IsMappedClothLOD = SkeletalMesh->IsMappedClothingLOD(LODIndex, AssetIndex);
 
-			//Change Clothing LOD if new LOD index is different from current index
-			if(CurLODIndex != LODIndex)
+			// Change Clothing LOD if a new LOD index is different from the current index
+			if (IsMappedClothLOD && CurLODIndex != LODIndex)
 			{
 				//physical LOD is changed by graphical LOD
-				ClothingActors[i].ApexClothingActor->setGraphicalLOD(LODIndex);
+				Actor.ApexClothingActor->setGraphicalLOD(LODIndex);
 
-				if(ClothingActors[i].ApexClothingActor->isFrozen())
+				if (Actor.ApexClothingActor->isFrozen())
 				{
 					bFrozen = true;
 				}
 			}
 
-			//decide whether enable or disable
-			if( (LODIndex > 0  && !bEnabledLOD) ||
-			    (LODIndex >= NumClothLODs) )
+			int32 NumClothLODs = Actor.ParentClothingAsset->GetAsset()->getNumGraphicalLodLevels();
+
+			// decide whether should enable or disable
+			if (!IsMappedClothLOD || (LODIndex >= NumClothLODs))
 			{
-				//disable clothing simulation
-				ClothingActors[i].ApexClothingActor->forcePhysicalLod(0);
+				// disable clothing simulation
+				Actor.ApexClothingActor->forcePhysicalLod(0);
 			}
 			else
 			{	
-				int32 CurPhysLOD = ClothingActors[i].ApexClothingActor->getActivePhysicalLod();
-				//enable clothing simulation
+				int32 CurPhysLOD = Actor.ApexClothingActor->getActivePhysicalLod();
+				// if disabled, enables clothing simulation
 				if(CurPhysLOD == 0)
 				{
-					ClothingActors[i].ApexClothingActor->forcePhysicalLod(1);
+					Actor.ApexClothingActor->forcePhysicalLod(1);
 				}
 			}
 		}
@@ -2120,7 +2146,7 @@ void USkeletalMeshComponent::SetClothingLOD(int32 LODIndex)
 #if WITH_EDITOR
 		FMessageDialog::Open(EAppMsgType::Ok, NSLOCTEXT("SkelmeshComponent", "Warning_FrozenCloth", "Clothing will be melted from frozen state"));
 #endif
-		//melt it because rendering mesh is broken if frozen when changing LODs
+		// melt it because rendering mesh will be broken if frozen when changing LODs
 		FreezeClothSection(false);
 	}
 
@@ -4400,6 +4426,13 @@ void USkeletalMeshComponent::DrawClothingPhysicalMeshWire(FPrimitiveDrawInterfac
 
 				if(NumSimulVertices > 0)
 				{
+					// if # of simulated vertices is bigger than loaded info, it will be LOD transition period 
+					// just skip this tick because PredictedLODLevel is different from internal clothing LOD. it will be matched at next tick
+					if (NumSimulVertices > NumPhysicalMeshVerts)
+					{
+						return;
+					}
+
 					bUseSimulatedResult = true;
 					SimulatedPhysicalMeshVertices.AddUninitialized(NumPhysicalMeshVerts);
 
diff --git a/Engine/Source/Runtime/Engine/Public/SkeletalMeshTypes.h b/Engine/Source/Runtime/Engine/Public/SkeletalMeshTypes.h
--- a/Engine/Source/Runtime/Engine/Public/SkeletalMeshTypes.h
+++ b/Engine/Source/Runtime/Engine/Public/SkeletalMeshTypes.h
@@ -726,8 +726,8 @@ struct FSkelMeshChunk
 		, NumRigidVertices(0)
 		, NumSoftVertices(0)
 		, MaxBoneInfluences(4)
-		, CorrespondClothAssetIndex(-1)
-		, ClothAssetSubmeshIndex(-1)
+		, CorrespondClothAssetIndex(INDEX_NONE)
+		, ClothAssetSubmeshIndex(INDEX_NONE)
 	{}
 
 	FSkelMeshChunk(const FSkelMeshChunk& Other)
@@ -871,7 +871,8 @@ struct FSkelMeshSection
 	int16 CorrespondClothSectionIndex;
 
 	/** Decide whether enabling clothing LOD for this section or not, just using skelmesh LOD_0's one to decide */
-	uint8 bEnableClothLOD;
+	/** no need anymore because each clothing LOD will be assigned to each mesh LOD  */
+	uint8 bEnableClothLOD_DEPRECATED;
 
 	FSkelMeshSection()
 		: MaterialIndex(0)
@@ -882,7 +883,6 @@ struct FSkelMeshSection
 		, bSelected(false)
 		, bDisabled(false)
 		, CorrespondClothSectionIndex(-1)
-		, bEnableClothLOD(true)
 	{}
 
 	// Serialization.
