Fix the crash hapepns when bone is selected, and mesh is swapped with the one without that selected bone.

diff --git a/Engine/Source/Editor/Persona/Private/AnimationEditorViewportClient.cpp b/Engine/Source/Editor/Persona/Private/AnimationEditorViewportClient.cpp
--- a/Engine/Source/Editor/Persona/Private/AnimationEditorViewportClient.cpp
+++ b/Engine/Source/Editor/Persona/Private/AnimationEditorViewportClient.cpp
@@ -375,6 +375,8 @@ void FAnimationViewportClient::SetPreviewMeshComponent(UDebugSkelMeshComponent*
 {
 	PreviewSkelMeshComp = InPreviewSkelMeshComp;
 
+	PreviewSkelMeshComp->BonesOfInterest.Empty();
+
 	UpdateCameraSetup();
 }
 
