Editor: Replace some literal red values with FLinearColor::Red

diff --git a/Engine/Source/Editor/Persona/Private/SAnimCurvePanel.cpp b/Engine/Source/Editor/Persona/Private/SAnimCurvePanel.cpp
--- a/Engine/Source/Editor/Persona/Private/SAnimCurvePanel.cpp
+++ b/Engine/Source/Editor/Persona/Private/SAnimCurvePanel.cpp
@@ -182,7 +182,7 @@ int32 SAnimCurveEd::OnPaint(const FGeometry& AllottedGeometry, const FSlateRect&
 		LinePoints,
 		MyClippingRect,
 		ESlateDrawEffect::None,
-		FLinearColor(1, 0, 0)
+		FLinearColor::Red
 		);
 
 	// now draw scrub with new layer ID + 1;
diff --git a/Engine/Source/Editor/Persona/Private/SAnimNotifyPanel.cpp b/Engine/Source/Editor/Persona/Private/SAnimNotifyPanel.cpp
--- a/Engine/Source/Editor/Persona/Private/SAnimNotifyPanel.cpp
+++ b/Engine/Source/Editor/Persona/Private/SAnimNotifyPanel.cpp
@@ -1493,7 +1493,7 @@ int32 SAnimNotifyTrack::OnPaint(const FGeometry& AllottedGeometry, const FSlateR
 			LinePoints,
 			MyClippingRect,
 			ESlateDrawEffect::None,
-			FLinearColor(1, 0, 0)
+			FLinearColor::Red
 			);
 	}
 
diff --git a/Engine/Source/Editor/Persona/Private/STrack.cpp b/Engine/Source/Editor/Persona/Private/STrack.cpp
--- a/Engine/Source/Editor/Persona/Private/STrack.cpp
+++ b/Engine/Source/Editor/Persona/Private/STrack.cpp
@@ -412,7 +412,7 @@ int32 STrack::OnPaint( const FGeometry& AllottedGeometry, const FSlateRect& MyCl
 			LinePoints,
 			MyClippingRect,
 			ESlateDrawEffect::None,
-			FLinearColor(1, 0, 0)
+			FLinearColor::Red
 			);
 	}
 
