TTP #306297 : Feature Request: EDITOR: Changeable LOD reduction settings for SkeletalMeshes

diff --git a/Engine/Source/Developer/SimplygonMeshReduction/Private/SimplygonMeshReduction.cpp b/Engine/Source/Developer/SimplygonMeshReduction/Private/SimplygonMeshReduction.cpp
--- a/Engine/Source/Developer/SimplygonMeshReduction/Private/SimplygonMeshReduction.cpp
+++ b/Engine/Source/Developer/SimplygonMeshReduction/Private/SimplygonMeshReduction.cpp
@@ -1336,24 +1336,12 @@ private:
 	 */
 	void SetReductionSettings( const FSkeletalMeshOptimizationSettings& Settings, float BoundsRadius, int32 SourceTriCount, SimplygonSDK::spReductionSettings ReductionSettings )
 	{
-		float MaxDeviation = 1.0f;
-		float ReductionRatio = 1.0f;
-		if(Settings.ReductionMethod == SMOT_MaxDeviation)
-		{
-			// Compute max deviation from quality.
-			MaxDeviation =  FMath::Max<float>( 0.0f, Settings.MaxDeviationPercentage * BoundsRadius );
-
-			// Set the reduction ratio such that at least 1 triangle remains.
-			ReductionRatio = FMath::Max<float>( 1.0f / SourceTriCount, 0.05f );
-		}
-		else
-		{
-			MaxDeviation = SimplygonSDK::REAL_MAX;
-
-			// Set the reduction ratio such that at least 1 triangle or 5% of the original triangles remain, whichever is larger.
-			float MinTriCount = FMath::Max<float>( 1.0f / SourceTriCount, 0.05f );
-			ReductionRatio = FMath::Max<float>( Settings.NumOfTrianglesPercentage, MinTriCount );
-		}
+		// Compute max deviation from quality.
+		float MaxDeviation = Settings.MaxDeviationPercentage > 0.0f ? Settings.MaxDeviationPercentage * BoundsRadius : SimplygonSDK::REAL_MAX;
+		// Set the reduction ratio such that at least 1 triangle or 5% of the original triangles remain, whichever is larger.
+		float MinReductionRatio = FMath::Max<float>(1.0f / SourceTriCount, 0.05f);
+		float MaxReductionRatio = (Settings.MaxDeviationPercentage > 0.0f && Settings.NumOfTrianglesPercentage == 1.0f) ? MinReductionRatio : 1.0f;
+		float ReductionRatio = FMath::Clamp(Settings.NumOfTrianglesPercentage, MinReductionRatio, MaxReductionRatio);
 
 		//Enable feature flags for those features where we have set an importance
 		unsigned int FeatureFlagsMask = SimplygonSDK::SG_FEATUREFLAGS_GROUP;
diff --git a/Engine/Source/Editor/Persona/Private/AnimGraphNodeDetails.cpp b/Engine/Source/Editor/Persona/Private/AnimGraphNodeDetails.cpp
--- a/Engine/Source/Editor/Persona/Private/AnimGraphNodeDetails.cpp
+++ b/Engine/Source/Editor/Persona/Private/AnimGraphNodeDetails.cpp
@@ -426,20 +426,27 @@ void FBoneReferenceCustomization::CustomizeHeader( TSharedRef<IPropertyHandle> S
 	TArray<UObject*> Objects;
 	StructPropertyHandle->GetOuterObjects(Objects);
 	UAnimGraphNode_Base* AnimGraphNode = NULL;
+	USkeletalMesh* SkeletalMesh = NULL;
+	TargetSkeleton = NULL;
 
 	for (auto OuterIter = Objects.CreateIterator() ; OuterIter ; ++OuterIter)
 	{
 		AnimGraphNode = Cast<UAnimGraphNode_Base>(*OuterIter);
 		if (AnimGraphNode)
 		{
+			TargetSkeleton = AnimGraphNode->GetAnimBlueprint()->TargetSkeleton;
+			break;
+		}
+		SkeletalMesh = Cast<USkeletalMesh>(*OuterIter);
+		if (SkeletalMesh)
+		{
+			TargetSkeleton = SkeletalMesh->Skeleton;
 			break;
 		}
 	}
 
-	if (AnimGraphNode)
+	if (TargetSkeleton)
 	{
-		TargetSkeleton = AnimGraphNode->GetAnimBlueprint()->TargetSkeleton;
-
 		HeaderRow.NameContent()
 		[
 			StructPropertyHandle->CreatePropertyNameWidget()
diff --git a/Engine/Source/Editor/Persona/Private/AnimViewportLODCommands.cpp b/Engine/Source/Editor/Persona/Private/AnimViewportLODCommands.cpp
--- a/Engine/Source/Editor/Persona/Private/AnimViewportLODCommands.cpp
+++ b/Engine/Source/Editor/Persona/Private/AnimViewportLODCommands.cpp
@@ -10,7 +10,5 @@ void FAnimViewportLODCommands::RegisterCommands()
 	UI_COMMAND( LOD1, "LOD 1", "Force select LOD 1", EUserInterfaceActionType::RadioButton, FInputGesture() );
 	UI_COMMAND( LOD2, "LOD 2", "Force select LOD 2", EUserInterfaceActionType::RadioButton, FInputGesture() );
 	UI_COMMAND( LOD3, "LOD 3", "Force select LOD 3", EUserInterfaceActionType::RadioButton, FInputGesture() );
-
-	UI_COMMAND( ShowLevelOfDetailSettings, "Level of Detail Settings", "Show the Level of Detail Settings Dialog", EUserInterfaceActionType::Button, FInputGesture() );
 }
 
diff --git a/Engine/Source/Editor/Persona/Private/AnimViewportLODCommands.h b/Engine/Source/Editor/Persona/Private/AnimViewportLODCommands.h
--- a/Engine/Source/Editor/Persona/Private/AnimViewportLODCommands.h
+++ b/Engine/Source/Editor/Persona/Private/AnimViewportLODCommands.h
@@ -36,9 +36,6 @@ public:
 	/** LOD 3 */
 	TSharedPtr< FUICommandInfo > LOD3;
 
-	/** Show Level of Detail Settings Dialog */
-	TSharedPtr< FUICommandInfo > ShowLevelOfDetailSettings;
-
 public:
 	/** Registers our commands with the binding system */
 	virtual void RegisterCommands() override;
diff --git a/Engine/Source/Editor/Persona/Private/PersonaMeshDetails.cpp b/Engine/Source/Editor/Persona/Private/PersonaMeshDetails.cpp
--- a/Engine/Source/Editor/Persona/Private/PersonaMeshDetails.cpp
+++ b/Engine/Source/Editor/Persona/Private/PersonaMeshDetails.cpp
@@ -15,9 +15,412 @@
 #endif // #if WITH_APEX_CLOTHING
 
 #include "MeshMode/SAdditionalMeshesEditor.h"
+#include "LODUtilities.h"
+#include "Developer/MeshUtilities/Public/MeshUtilities.h"
+#include "FbxMeshUtils.h"
+
+#include "AnimGraphNodeDetails.h"
 
 #define LOCTEXT_NAMESPACE "PersonaMeshDetails"
 
+/** Returns true if automatic mesh reduction is available. */
+static bool IsAutoMeshReductionAvailable()
+{
+	static bool bAutoMeshReductionAvailable = FModuleManager::Get().LoadModuleChecked<IMeshUtilities>("MeshUtilities").GetMeshReductionInterface() != NULL;
+	return bAutoMeshReductionAvailable;
+}
+
+/**
+* FMeshReductionSettings
+*/
+
+static UEnum& GetFeatureImportanceEnum()
+{
+	static FName FeatureImportanceName(TEXT("EMeshFeatureImportance::Off"));
+	static UEnum* FeatureImportanceEnum = NULL;
+	if (FeatureImportanceEnum == NULL)
+	{
+		UEnum::LookupEnumName(FeatureImportanceName, &FeatureImportanceEnum);
+		check(FeatureImportanceEnum);
+	}
+	return *FeatureImportanceEnum;
+}
+
+static UEnum& GetFeatureSimpificationEnum()
+{
+	static FName FeatureSimpificationName(TEXT("SMOT_NumOfTriangles"));
+	static UEnum* FeatureSimpificationEnum = NULL;
+	if (FeatureSimpificationEnum == NULL)
+	{
+		UEnum::LookupEnumName(FeatureSimpificationName, &FeatureSimpificationEnum);
+		check(FeatureSimpificationEnum);
+	}
+	return *FeatureSimpificationEnum;
+}
+
+static void FillEnumOptions(TArray<TSharedPtr<FString> >& OutStrings, UEnum& InEnum)
+{
+	for (int32 EnumIndex = 0; EnumIndex < InEnum.NumEnums() - 1; ++EnumIndex)
+	{
+		OutStrings.Add(MakeShareable(new FString(InEnum.GetEnumName(EnumIndex))));
+	}
+}
+
+
+FSkelMeshReductionSettingsLayout::FSkelMeshReductionSettingsLayout(int32 InLODIndex, TSharedRef<FPersonaMeshDetails> InParentLODSettings, TSharedPtr<IPropertyHandle> InBoneToRemoveProperty)
+: LODIndex(InLODIndex)
+, ParentLODSettings(InParentLODSettings)
+, BoneToRemoveProperty(InBoneToRemoveProperty)
+{
+	FillEnumOptions(SimplificationOptions, GetFeatureSimpificationEnum());
+	FillEnumOptions(ImportanceOptions, GetFeatureImportanceEnum());
+}
+
+FSkelMeshReductionSettingsLayout::~FSkelMeshReductionSettingsLayout()
+{
+}
+
+void FSkelMeshReductionSettingsLayout::GenerateHeaderRowContent(FDetailWidgetRow& NodeRow)
+{
+	NodeRow.NameContent()
+		[
+			SNew(STextBlock)
+			.Text(LOCTEXT("MeshReductionSettings", "Reduction Settings"))
+			.Font(IDetailLayoutBuilder::GetDetailFont())
+		];
+}
+
+
+BEGIN_SLATE_FUNCTION_BUILD_OPTIMIZATION
+void FSkelMeshReductionSettingsLayout::GenerateChildContent(IDetailChildrenBuilder& ChildrenBuilder)
+{
+	{
+		ChildrenBuilder.AddChildContent(LOCTEXT("PercentTriangles", "Percent Triangles").ToString())
+			.NameContent()
+			[
+				SNew(STextBlock)
+				.Font(IDetailLayoutBuilder::GetDetailFont())
+				.Text(LOCTEXT("PercentTriangles", "Percent Triangles"))
+			]
+		.ValueContent()
+			[
+				SNew(SSpinBox<float>)
+				.Font(IDetailLayoutBuilder::GetDetailFont())
+				.MinValue(0.0f)
+				.MaxValue(100.0f)
+				.Value(this, &FSkelMeshReductionSettingsLayout::GetPercentTriangles)
+				.OnValueChanged(this, &FSkelMeshReductionSettingsLayout::OnPercentTrianglesChanged)
+			];
+
+	}
+
+	{
+		ChildrenBuilder.AddChildContent(LOCTEXT("MaxDeviation", "Max Deviation").ToString())
+			.NameContent()
+			[
+				SNew(STextBlock)
+				.Font(IDetailLayoutBuilder::GetDetailFont())
+				.Text(LOCTEXT("MaxDeviation", "Max Deviation"))
+			]
+		.ValueContent()
+			[
+				SNew(SSpinBox<float>)
+				.Font(IDetailLayoutBuilder::GetDetailFont())
+				.MinValue(0.0f)
+				.MaxValue(100.0f)
+				.Value(this, &FSkelMeshReductionSettingsLayout::GetMaxDeviation)
+				.OnValueChanged(this, &FSkelMeshReductionSettingsLayout::OnMaxDeviationChanged)
+			];
+
+	}
+
+	{
+		ChildrenBuilder.AddChildContent(LOCTEXT("Silhouette_MeshSimplification", "Silhouette").ToString())
+			.NameContent()
+			[
+				SNew(STextBlock)
+				.Font(IDetailLayoutBuilder::GetDetailFont())
+				.Text(LOCTEXT("Silhouette_MeshSimplification", "Silhouette"))
+			]
+		.ValueContent()
+			[
+				SAssignNew(SilhouetteCombo, STextComboBox)
+				//.Font( IDetailLayoutBuilder::GetDetailFont() )
+				.ContentPadding(0)
+				.OptionsSource(&ImportanceOptions)
+				.InitiallySelectedItem(ImportanceOptions[ReductionSettings.SilhouetteImportance])
+				.OnSelectionChanged(this, &FSkelMeshReductionSettingsLayout::OnSilhouetteImportanceChanged)
+			];
+
+	}
+
+	{
+		ChildrenBuilder.AddChildContent(LOCTEXT("Texture_MeshSimplification", "Texture").ToString())
+			.NameContent()
+			[
+				SNew(STextBlock)
+				.Font(IDetailLayoutBuilder::GetDetailFont())
+				.Text(LOCTEXT("Texture_MeshSimplification", "Texture"))
+			]
+		.ValueContent()
+			[
+				SAssignNew(TextureCombo, STextComboBox)
+				//.Font( IDetailLayoutBuilder::GetDetailFont() )
+				.ContentPadding(0)
+				.OptionsSource(&ImportanceOptions)
+				.InitiallySelectedItem(ImportanceOptions[ReductionSettings.TextureImportance])
+				.OnSelectionChanged(this, &FSkelMeshReductionSettingsLayout::OnTextureImportanceChanged)
+			];
+
+	}
+
+	{
+		ChildrenBuilder.AddChildContent(LOCTEXT("Shading_MeshSimplification", "Shading").ToString())
+			.NameContent()
+			[
+				SNew(STextBlock)
+				.Font(IDetailLayoutBuilder::GetDetailFont())
+				.Text(LOCTEXT("Shading_MeshSimplification", "Shading"))
+			]
+		.ValueContent()
+			[
+				SAssignNew(ShadingCombo, STextComboBox)
+				//.Font( IDetailLayoutBuilder::GetDetailFont() )
+				.ContentPadding(0)
+				.OptionsSource(&ImportanceOptions)
+				.InitiallySelectedItem(ImportanceOptions[ReductionSettings.ShadingImportance])
+				.OnSelectionChanged(this, &FSkelMeshReductionSettingsLayout::OnShadingImportanceChanged)
+			];
+
+	}
+
+	{
+		ChildrenBuilder.AddChildContent(LOCTEXT("Skinning_MeshSimplification", "Skinning").ToString())
+			.NameContent()
+			[
+				SNew(STextBlock)
+				.Font(IDetailLayoutBuilder::GetDetailFont())
+				.Text(LOCTEXT("Skinning_MeshSimplification", "Skinning"))
+			]
+		.ValueContent()
+			[
+				SAssignNew(SkinningCombo, STextComboBox)
+				//.Font( IDetailLayoutBuilder::GetDetailFont() )
+				.ContentPadding(0)
+				.OptionsSource(&ImportanceOptions)
+				.InitiallySelectedItem(ImportanceOptions[ReductionSettings.SkinningImportance])
+				.OnSelectionChanged(this, &FSkelMeshReductionSettingsLayout::OnSkinningImportanceChanged)
+			];
+
+	}
+
+	{
+		ChildrenBuilder.AddChildContent(LOCTEXT("WeldingThreshold", "Welding Threshold").ToString())
+			.NameContent()
+			[
+				SNew(STextBlock)
+				.Font(IDetailLayoutBuilder::GetDetailFont())
+				.Text(LOCTEXT("WeldingThreshold", "Welding Threshold"))
+			]
+		.ValueContent()
+			[
+				SNew(SSpinBox<float>)
+				.Font(IDetailLayoutBuilder::GetDetailFont())
+				.MinValue(0.0f)
+				.MaxValue(10.0f)
+				.Value(this, &FSkelMeshReductionSettingsLayout::GetWeldingThreshold)
+				.OnValueChanged(this, &FSkelMeshReductionSettingsLayout::OnWeldingThresholdChanged)
+			];
+
+	}
+
+	{
+		ChildrenBuilder.AddChildContent(LOCTEXT("RecomputeNormals", "Recompute Normals").ToString())
+			.NameContent()
+			[
+				SNew(STextBlock)
+				.Font(IDetailLayoutBuilder::GetDetailFont())
+				.Text(LOCTEXT("RecomputeNormals", "Recompute Normals"))
+
+			]
+		.ValueContent()
+			[
+				SNew(SCheckBox)
+				.IsChecked(this, &FSkelMeshReductionSettingsLayout::ShouldRecalculateNormals)
+				.OnCheckStateChanged(this, &FSkelMeshReductionSettingsLayout::OnRecalculateNormalsChanged)
+			];
+	}
+
+	{
+		ChildrenBuilder.AddChildContent(LOCTEXT("HardEdgeAngle", "Hard Edge Angle").ToString())
+			.NameContent()
+			[
+				SNew(STextBlock)
+				.Font(IDetailLayoutBuilder::GetDetailFont())
+				.Text(LOCTEXT("HardEdgeAngle", "Hard Edge Angle"))
+			]
+		.ValueContent()
+			[
+				SNew(SSpinBox<float>)
+				.Font(IDetailLayoutBuilder::GetDetailFont())
+				.MinValue(0.0f)
+				.MaxValue(180.0f)
+				.Value(this, &FSkelMeshReductionSettingsLayout::GetHardAngleThreshold)
+				.OnValueChanged(this, &FSkelMeshReductionSettingsLayout::OnHardAngleThresholdChanged)
+			];
+
+	}
+
+	{
+		if (BoneToRemoveProperty.IsValid() && BoneToRemoveProperty->IsValidHandle())
+		{
+			ChildrenBuilder.AddChildProperty(BoneToRemoveProperty.ToSharedRef());
+		}
+	}
+
+	{
+		ChildrenBuilder.AddChildContent(LOCTEXT("MaxBonesPerVertex", "Max Bones Per Vertex").ToString())
+			.NameContent()
+			[
+				SNew(STextBlock)
+				.Font(IDetailLayoutBuilder::GetDetailFont())
+				.Text(LOCTEXT("MaxBonesPerVertex", "Max Bones Per Vertex"))
+			]
+		.ValueContent()
+			[
+				SNew(SSpinBox<int32>)
+				.Font(IDetailLayoutBuilder::GetDetailFont())
+				.MinValue(1)
+				.MaxValue(MAX_TOTAL_INFLUENCES)
+				.Value(this, &FSkelMeshReductionSettingsLayout::GetMaxBonesPerVertex)
+				.OnValueChanged(this, &FSkelMeshReductionSettingsLayout::OnMaxBonesPerVertexChanged)
+			];
+
+	}
+
+	{
+		ChildrenBuilder.AddChildContent(LOCTEXT("ApplyChanges", "Apply Changes").ToString())
+			.ValueContent()
+			.HAlign(HAlign_Left)
+			[
+				SNew(SButton)
+				.OnClicked(this, &FSkelMeshReductionSettingsLayout::OnApplyChanges)
+				.IsEnabled(ParentLODSettings.Pin().ToSharedRef(), &FPersonaMeshDetails::IsApplyNeeded)
+				[
+					SNew(STextBlock)
+					.Text(LOCTEXT("ApplyChanges", "Apply Changes"))
+					.Font(IDetailLayoutBuilder::GetDetailFont())
+				]
+			];
+	}
+
+	SilhouetteCombo->SetSelectedItem(ImportanceOptions[ReductionSettings.SilhouetteImportance]);
+	TextureCombo->SetSelectedItem(ImportanceOptions[ReductionSettings.TextureImportance]);
+	ShadingCombo->SetSelectedItem(ImportanceOptions[ReductionSettings.ShadingImportance]);
+	SkinningCombo->SetSelectedItem(ImportanceOptions[ReductionSettings.SkinningImportance]);
+}
+END_SLATE_FUNCTION_BUILD_OPTIMIZATION
+
+const FSkeletalMeshOptimizationSettings& FSkelMeshReductionSettingsLayout::GetSettings() const
+{
+	return ReductionSettings;
+}
+
+void FSkelMeshReductionSettingsLayout::UpdateSettings(const FSkeletalMeshOptimizationSettings& InSettings)
+{
+	ReductionSettings = InSettings;
+}
+
+FReply FSkelMeshReductionSettingsLayout::OnApplyChanges()
+{
+	if (ParentLODSettings.IsValid())
+	{
+		ParentLODSettings.Pin()->ApplyChanges();
+	}
+	return FReply::Handled();
+}
+
+float FSkelMeshReductionSettingsLayout::GetPercentTriangles() const
+{
+	return ReductionSettings.NumOfTrianglesPercentage * 100.0f; // Display fraction as percentage.
+}
+
+float FSkelMeshReductionSettingsLayout::GetMaxDeviation() const
+{
+	return ReductionSettings.MaxDeviationPercentage * 2000.0f;
+}
+
+float FSkelMeshReductionSettingsLayout::GetWeldingThreshold() const
+{
+	return ReductionSettings.WeldingThreshold;
+}
+
+ESlateCheckBoxState::Type FSkelMeshReductionSettingsLayout::ShouldRecalculateNormals() const
+{
+	return ReductionSettings.bRecalcNormals ? ESlateCheckBoxState::Checked : ESlateCheckBoxState::Unchecked;
+}
+
+float FSkelMeshReductionSettingsLayout::GetHardAngleThreshold() const
+{
+	return ReductionSettings.NormalsThreshold;
+}
+
+int32 FSkelMeshReductionSettingsLayout::GetMaxBonesPerVertex() const
+{
+	return ReductionSettings.MaxBonesPerVertex;
+}
+
+void FSkelMeshReductionSettingsLayout::OnPercentTrianglesChanged(float NewValue)
+{
+	// Percentage -> fraction.
+	ReductionSettings.NumOfTrianglesPercentage = NewValue * 0.01f;
+}
+
+void FSkelMeshReductionSettingsLayout::OnMaxDeviationChanged(float NewValue)
+{
+	ReductionSettings.MaxDeviationPercentage = NewValue / 2000.0f;
+}
+
+void FSkelMeshReductionSettingsLayout::OnWeldingThresholdChanged(float NewValue)
+{
+	ReductionSettings.WeldingThreshold = NewValue;
+}
+
+void FSkelMeshReductionSettingsLayout::OnRecalculateNormalsChanged(ESlateCheckBoxState::Type NewValue)
+{
+	ReductionSettings.bRecalcNormals = NewValue == ESlateCheckBoxState::Checked;
+}
+
+void FSkelMeshReductionSettingsLayout::OnHardAngleThresholdChanged(float NewValue)
+{
+	ReductionSettings.NormalsThreshold = NewValue;
+}
+
+void FSkelMeshReductionSettingsLayout::OnMaxBonesPerVertexChanged(int32 NewValue)
+{
+	ReductionSettings.MaxBonesPerVertex = NewValue;
+}
+
+void FSkelMeshReductionSettingsLayout::OnSilhouetteImportanceChanged(TSharedPtr<FString> NewValue, ESelectInfo::Type SelectInfo)
+{
+	ReductionSettings.SilhouetteImportance = (SkeletalMeshOptimizationImportance)ImportanceOptions.Find(NewValue);
+}
+
+void FSkelMeshReductionSettingsLayout::OnTextureImportanceChanged(TSharedPtr<FString> NewValue, ESelectInfo::Type SelectInfo)
+{
+	ReductionSettings.TextureImportance = (SkeletalMeshOptimizationImportance)ImportanceOptions.Find(NewValue);
+}
+
+void FSkelMeshReductionSettingsLayout::OnShadingImportanceChanged(TSharedPtr<FString> NewValue, ESelectInfo::Type SelectInfo)
+{
+	ReductionSettings.ShadingImportance = (SkeletalMeshOptimizationImportance)ImportanceOptions.Find(NewValue);
+}
+
+void FSkelMeshReductionSettingsLayout::OnSkinningImportanceChanged(TSharedPtr<FString> NewValue, ESelectInfo::Type SelectInfo)
+{
+	ReductionSettings.SkinningImportance = (SkeletalMeshOptimizationImportance)ImportanceOptions.Find(NewValue);
+}
+
 /**
 * FPersonaMeshDetails
 */
@@ -34,9 +437,59 @@ void FPersonaMeshDetails::AddLODLevelCategories(IDetailLayoutBuilder& DetailLayo
 	{
 		const int32 SkelMeshLODCount = SkelMesh->LODInfo.Num();
 
+		if (ReductionSettingsWidgets.Num() < SkelMeshLODCount)
+		{
+			ReductionSettingsWidgets.AddZeroed(SkelMeshLODCount - ReductionSettingsWidgets.Num());
+		}
+
 		// Create information panel for each LOD level.
 		for (int32 LODIndex = 0; LODIndex < SkelMeshLODCount; ++LODIndex)
 		{
+			if (LODIndex > 0 && IsAutoMeshReductionAvailable())
+			{
+				TSharedRef<IPropertyHandle> LODInfoProperty = DetailLayout.GetProperty(FName("LODInfo"), USkeletalMesh::StaticClass());
+				uint32 NumChildren = 0;
+				LODInfoProperty->GetNumChildren(NumChildren);
+
+				check(NumChildren > (uint32)LODIndex);
+	
+				TSharedPtr<IPropertyHandle> ChildHandle = LODInfoProperty->GetChildHandle(LODIndex);
+				check(ChildHandle.IsValid());
+				TSharedPtr<IPropertyHandle> ReductionHandle = ChildHandle->GetChildHandle(FName("ReductionSettings"));
+				check(ReductionHandle.IsValid());
+				TSharedPtr<IPropertyHandle> BonesToRemoveHandle = ReductionHandle->GetChildHandle(FName("BonesToRemove"));
+				ReductionSettingsWidgets[LODIndex] = MakeShareable(new FSkelMeshReductionSettingsLayout(LODIndex, SharedThis(this), BonesToRemoveHandle));
+			}
+
+//			if (LODIndex < SkelMesh->LODInfo.Num())
+			{
+				FSkeletalMeshLODInfo& LODInfo = SkelMesh->LODInfo[LODIndex];
+				if (ReductionSettingsWidgets[LODIndex].IsValid())
+				{
+					ReductionSettingsWidgets[LODIndex]->UpdateSettings(LODInfo.ReductionSettings);
+				}
+
+				//LODScreenSizes[LODIndex] = SrcModel.ScreenSize;
+			}
+			/*
+			else if (LODIndex > 0)
+			{
+				if (ReductionSettingsWidgets[LODIndex].IsValid() && ReductionSettingsWidgets[LODIndex - 1].IsValid())
+				{
+					FMeshReductionSettings ReductionSettings = ReductionSettingsWidgets[LODIndex - 1]->GetSettings();
+					// By default create LODs with half the triangles of the previous LOD.
+					ReductionSettings.PercentTriangles *= 0.5f;
+					ReductionSettingsWidgets[LODIndex]->UpdateSettings(ReductionSettings);
+				}
+
+				if (LODScreenSizes[LODIndex] >= LODScreenSizes[LODIndex - 1])
+				{
+					const float DefaultScreenSizeDifference = 0.01f;
+					LODScreenSizes[LODIndex] = LODScreenSizes[LODIndex - 1] - DefaultScreenSizeDifference;
+				}
+			}
+			*/
+
 			FString CategoryName = FString(TEXT("LOD"));
 			CategoryName.AppendInt(LODIndex);
 
@@ -68,24 +521,301 @@ void FPersonaMeshDetails::AddLODLevelCategories(IDetailLayoutBuilder& DetailLayo
 				LODInfoChild->CreatePropertyNameWidget(FString(TEXT("LOD Info")));
 				LODCategory.AddProperty(LODInfoChild);
 			}
+
+			if (ReductionSettingsWidgets[LODIndex].IsValid())
+			{
+				LODCategory.AddCustomBuilder(ReductionSettingsWidgets[LODIndex].ToSharedRef());
+			}
 		}
 	}
 }
 
+void FPersonaMeshDetails::CustomizeLODSettingsCategories(IDetailLayoutBuilder& DetailLayout)
+{
+	USkeletalMesh* SkelMesh = PersonaPtr->GetMesh();
+	LODCount = SkelMesh->LODInfo.Num();
+
+	UpdateLODNames();
+
+	IDetailCategoryBuilder& LODSettingsCategory = DetailLayout.EditCategory("LodSettings", LOCTEXT("LodSettingsCategory", "LOD Settings").ToString(), ECategoryPriority::TypeSpecific);
+
+	LODSettingsCategory.AddCustomRow(LOCTEXT("LODImport", "LOD Import").ToString())
+		.NameContent()
+		[
+			SNew(STextBlock)
+			.Font(IDetailLayoutBuilder::GetDetailFont())
+			.Text(LOCTEXT("LODImport", "LOD Import"))
+		]
+	.ValueContent()
+		[
+			SNew(STextComboBox)
+			.ContentPadding(0)
+			.OptionsSource(&LODNames)
+			.InitiallySelectedItem(LODNames[0])
+			.OnSelectionChanged(this, &FPersonaMeshDetails::OnImportLOD, &DetailLayout)
+		];
+
+	// Add Number of LODs slider.
+	const int32 MinAllowedLOD = 1;
+	LODSettingsCategory.AddCustomRow(LOCTEXT("NumberOfLODs", "Number of LODs").ToString())
+		.NameContent()
+		[
+			SNew(STextBlock)
+			.Font(IDetailLayoutBuilder::GetDetailFont())
+			.Text(LOCTEXT("NumberOfLODs", "Number of LODs"))
+		]
+	.ValueContent()
+		[
+			SNew(SSpinBox<int32>)
+			.Font(IDetailLayoutBuilder::GetDetailFont())
+			.Value(this, &FPersonaMeshDetails::GetLODCount)
+			.OnValueChanged(this, &FPersonaMeshDetails::OnLODCountChanged)
+			.OnValueCommitted(this, &FPersonaMeshDetails::OnLODCountCommitted)
+			.MinValue(MinAllowedLOD)
+			.MaxValue(MAX_SKELETAL_MESH_LODS)
+			.ToolTipText(this, &FPersonaMeshDetails::GetLODCountTooltip)
+			.IsEnabled(IsAutoMeshReductionAvailable())
+		];
+
+	LODSettingsCategory.AddCustomRow(LOCTEXT("ApplyChanges", "Apply Changes").ToString())
+		.ValueContent()
+		.HAlign(HAlign_Left)
+		[
+			SNew(SButton)
+			.OnClicked(this, &FPersonaMeshDetails::OnApplyChanges)
+			.IsEnabled(this, &FPersonaMeshDetails::IsApplyNeeded)
+			[
+				SNew(STextBlock)
+				.Text(LOCTEXT("ApplyChanges", "Apply Changes"))
+				.Font(DetailLayout.GetDetailFont())
+			]
+		];
+}
+
+void FPersonaMeshDetails::OnImportLOD(TSharedPtr<FString> NewValue, ESelectInfo::Type SelectInfo, IDetailLayoutBuilder* DetailLayout)
+{
+	int32 LODIndex = 0;
+	if (LODNames.Find(NewValue, LODIndex) && LODIndex > 0)
+	{
+		USkeletalMesh* SkelMesh = PersonaPtr->GetMesh();
+		check(SkelMesh);
+
+		FbxMeshUtils::ImportMeshLODDialog(SkelMesh, LODIndex);
+
+		DetailLayout->ForceRefreshDetails();
+	}
+}
+
+int32 FPersonaMeshDetails::GetLODCount() const
+{
+	return LODCount;
+}
+
+void FPersonaMeshDetails::OnLODCountChanged(int32 NewValue)
+{
+	LODCount = FMath::Clamp<int32>(NewValue, 1, MAX_SKELETAL_MESH_LODS);
+
+	UpdateLODNames();
+}
+
+void FPersonaMeshDetails::OnLODCountCommitted(int32 InValue, ETextCommit::Type CommitInfo)
+{
+	OnLODCountChanged(InValue);
+}
+
+FReply FPersonaMeshDetails::OnApplyChanges()
+{
+	ApplyChanges();
+	return FReply::Handled();
+}
+
+void FPersonaMeshDetails::RefreshBonesToRemove(TArray<FBoneReference>& InOutBonesToRemove, int32 LODIndex)
+{
+	if (LODIndex == 0)
+	{
+		// we don't remove bones from LOD 0 
+		InOutBonesToRemove.Empty();
+		return;
+	}
+
+	USkeletalMesh* SkelMesh = PersonaPtr->GetMesh();
+	USkeleton* Skeleton = SkelMesh->Skeleton;
+
+	for (int32 Index = 0; Index < InOutBonesToRemove.Num(); Index++)
+	{
+		int32 BoneIndex = Skeleton->GetReferenceSkeleton().FindBoneIndex(InOutBonesToRemove[Index].BoneName);
+		if (BoneIndex != INDEX_NONE)
+		{
+			Skeleton->RemoveBoneFromLOD(LODIndex, BoneIndex);
+		}
+	}
+
+	InOutBonesToRemove.Empty();
+
+	// index for accessing to BoneReductionSettingsForLODs because it doesn't have for LOD 0
+	int32 LODSettingIndex = LODIndex - 1;
+
+	if (Skeleton->BoneReductionSettingsForLODs.Num() > LODSettingIndex)
+	{
+		const TArray<FName> & BonesToRemove = Skeleton->BoneReductionSettingsForLODs[LODSettingIndex].BonesToRemove;
+
+		// filter out only the parent, not all children
+		int32 TotalNumOfBonesToRemove = BonesToRemove.Num();
+		const FReferenceSkeleton & RefSkeleton = Skeleton->GetReferenceSkeleton();
+		for (int32 Index = 0; Index < TotalNumOfBonesToRemove; ++Index)
+		{
+			bool bNoParent = true;
+			FName BoneName = BonesToRemove[Index];
+			int32 BoneIndex = RefSkeleton.FindBoneIndex(BoneName);
+			int32 ParentBoneIndex = Skeleton->GetReferenceSkeleton().GetParentIndex(BoneIndex);
+			for (int32 J = 0; J < TotalNumOfBonesToRemove; ++J)
+			{
+				int32 CurrentBoneIndex = RefSkeleton.FindBoneIndex(BonesToRemove[J]);
+				if (CurrentBoneIndex == ParentBoneIndex)
+				{
+					bNoParent = false;
+					break;
+				}
+			}
+
+			// if no parent, we mark that as to display
+			if (bNoParent)
+			{
+				FBoneReference BoneRef;
+				BoneRef.BoneName = BoneName;
+				BoneRef.BoneIndex = BoneIndex;
+				InOutBonesToRemove.AddUnique(BoneRef);
+			}
+		}
+	}
+}
+
+void FPersonaMeshDetails::ApplyChanges()
+{
+	USkeletalMesh* SkelMesh = PersonaPtr->GetMesh();
+	check(SkelMesh);
+
+	FSkeletalMeshUpdateContext UpdateContext;
+	UpdateContext.SkeletalMesh = SkelMesh;
+	UpdateContext.AssociatedComponents.Push(PersonaPtr->PreviewComponent);
+
+	// remove LODs
+	int32 CurrentNumLODs = SkelMesh->LODInfo.Num();
+	if (LODCount < CurrentNumLODs)
+	{
+		for (int32 LODIdx = CurrentNumLODs - 1; LODIdx >= LODCount; LODIdx--)
+		{
+			FLODUtilities::RemoveLOD(UpdateContext, LODIdx);
+		}
+	}
+	else // generate LODs
+	{
+		TArray<FSkeletalMeshOptimizationSettings> Settings;
+
+		int32 SettingIndex = INDEX_NONE;
+
+		for (int32 LODIdx = 0; LODIdx < CurrentNumLODs; LODIdx++)
+		{
+			if (ReductionSettingsWidgets[LODIdx].IsValid())
+			{
+				FSkeletalMeshLODInfo& Info = SkelMesh->LODInfo[LODIdx];
+				SettingIndex = Settings.Add(ReductionSettingsWidgets[LODIdx]->GetSettings());
+				// BonesToRemove in LODInfo is the latest info which the user assigned because it is set by UPROPERTY directly from the editor 
+				// and swap to make settings always be the latest
+				TArray<FBoneReference> TempBonesForSwap = Settings[SettingIndex].BonesToRemove;
+				Settings[SettingIndex].BonesToRemove = Info.ReductionSettings.BonesToRemove;
+				Info.ReductionSettings.BonesToRemove = TempBonesForSwap;
+				RefreshBonesToRemove(Settings[SettingIndex].BonesToRemove, LODIdx);
+			}
+		}
+
+		for (int32 LODIdx = CurrentNumLODs; LODIdx < LODCount; LODIdx++)
+		{
+			FSkeletalMeshOptimizationSettings Setting;
+			if (SettingIndex != INDEX_NONE)
+			{
+				Setting = Settings[SettingIndex];
+			}
+
+			Setting.NumOfTrianglesPercentage *= 0.5f;
+			SettingIndex = Settings.Add(Setting);
+		}
+
+		FLODUtilities::SimplifySkeletalMesh(UpdateContext, Settings);
+
+		// save reduction settings
+		for (int32 LODIdx = 1; LODIdx < LODCount; ++LODIdx)
+		{
+			if (SkelMesh->LODInfo.IsValidIndex(LODIdx))
+			{
+				FSkeletalMeshLODInfo& Info = SkelMesh->LODInfo[LODIdx];
+				Info.ReductionSettings = Settings[LODIdx-1];
+			}
+		}
+	}
+
+	PersonaPtr->PersonaMeshDetailLayout->ForceRefreshDetails();
+}
+
+void FPersonaMeshDetails::UpdateLODNames()
+{
+	LODNames.Empty();
+	LODNames.Add(MakeShareable(new FString(LOCTEXT("BaseLOD", "Base LOD").ToString())));
+	for (int32 LODLevelID = 1; LODLevelID < LODCount; ++LODLevelID)
+	{
+		LODNames.Add(MakeShareable(new FString(FText::Format(NSLOCTEXT("LODSettingsLayout", "LODLevel_Reimport", "Reimport LOD Level {0}"), FText::AsNumber(LODLevelID)).ToString())));
+	}
+	LODNames.Add(MakeShareable(new FString(FText::Format(NSLOCTEXT("LODSettingsLayout", "LODLevel_Import", "Import LOD Level {0}"), FText::AsNumber(LODCount)).ToString())));
+}
+
+bool FPersonaMeshDetails::IsApplyNeeded() const
+{
+	USkeletalMesh* SkelMesh = PersonaPtr->GetMesh();
+	check(SkelMesh);
+
+	if (SkelMesh->LODInfo.Num() != LODCount)
+	{
+		return true;
+	}
+
+
+	for (int32 LODIndex = 0; LODIndex < LODCount; ++LODIndex)
+	{
+		FSkeletalMeshLODInfo& Info = SkelMesh->LODInfo[LODIndex];
+		if (ReductionSettingsWidgets.IsValidIndex(LODIndex)
+			&& ReductionSettingsWidgets[LODIndex].IsValid()
+			&& Info.ReductionSettings != ReductionSettingsWidgets[LODIndex]->GetSettings())
+		{
+			return true;
+		}
+	}
+
+	return false;
+}
+
+FText FPersonaMeshDetails::GetLODCountTooltip() const
+{
+	if (IsAutoMeshReductionAvailable())
+	{
+		return LOCTEXT("LODCountTooltip", "The number of LODs for this static mesh. If auto mesh reduction is available, setting this number will determine the number of LOD levels to auto generate.");
+	}
+
+	return LOCTEXT("LODCountTooltip_Disabled", "Auto mesh reduction is unavailable! Please provide a mesh reduction interface such as Simplygon to use this feature or manually import LOD levels.");
+}
+
+
 void FPersonaMeshDetails::CustomizeDetails( IDetailLayoutBuilder& DetailLayout )
 {
 	SelectedObjects = DetailLayout.GetDetailsView().GetSelectedObjects();
 	check(SelectedObjects.Num()<=1); // The OnGenerateCustomWidgets delegate will not be useful if we try to process more than one object.
 
 	// copy temporarily to refresh Mesh details tab from the LOD settings window
-	PersonaPtr->PersonaMeshDetailLayout = &DetailLayout;
+	PersonaPtr->PersonaMeshDetailLayout = &DetailLayout;	
 
 	// add multiple LOD levels to LOD category
 	AddLODLevelCategories(DetailLayout);
 
-	// LODInfo doesn't need to be showed anymore because it was moved to each LOD category
-	TSharedRef<IPropertyHandle> LODInfoProperty = DetailLayout.GetProperty(FName("LODInfo"), USkeletalMesh::StaticClass());
-	DetailLayout.HideProperty(LODInfoProperty);
+	CustomizeLODSettingsCategories(DetailLayout);
 
 #if WITH_APEX_CLOTHING
 	IDetailCategoryBuilder& ClothingCategory = DetailLayout.EditCategory("Clothing", TEXT(""), ECategoryPriority::TypeSpecific);
@@ -97,6 +827,46 @@ void FPersonaMeshDetails::CustomizeDetails( IDetailLayoutBuilder& DetailLayout )
 	[
 		SNew(SAdditionalMeshesEditor, PersonaPtr)
 	];
+
+	HideUnnecessaryProperties(DetailLayout);
+}
+
+void FPersonaMeshDetails::HideUnnecessaryProperties(IDetailLayoutBuilder& DetailLayout)
+{
+	// LODInfo doesn't need to be showed anymore because it was moved to each LOD category
+	TSharedRef<IPropertyHandle> LODInfoProperty = DetailLayout.GetProperty(FName("LODInfo"), USkeletalMesh::StaticClass());
+	DetailLayout.HideProperty(LODInfoProperty);
+	uint32 NumChildren = 0;
+	LODInfoProperty->GetNumChildren(NumChildren);
+	// Hide reduction settings property because it is duplicated with Reduction settings layout created by ReductionSettingsWidgets
+	for (uint32 ChildIdx = 0; ChildIdx < NumChildren; ChildIdx++)
+	{
+		TSharedPtr<IPropertyHandle> ChildHandle = LODInfoProperty->GetChildHandle(ChildIdx);
+		if (ChildHandle.IsValid())
+		{
+			TSharedPtr<IPropertyHandle> MaterialMapHandle = ChildHandle->GetChildHandle(FName("LODMaterialMap"));
+			DetailLayout.HideProperty(MaterialMapHandle);
+
+			TSharedPtr<IPropertyHandle> TriangleSortHandle = ChildHandle->GetChildHandle(FName("TriangleSortSettings"));
+			DetailLayout.HideProperty(TriangleSortHandle);
+
+			TSharedPtr<IPropertyHandle> ReductionHandle = ChildHandle->GetChildHandle(FName("ReductionSettings"));
+			DetailLayout.HideProperty(ReductionHandle);
+		}
+	}
+
+	TSharedRef<IPropertyHandle> MaterialsProperty = DetailLayout.GetProperty(FName("Materials"), USkeletalMesh::StaticClass());
+	DetailLayout.HideProperty(MaterialsProperty);
+
+	// hide all properties in Mirroring category to hide Mirroring category itself
+	IDetailCategoryBuilder& MirroringCategory = DetailLayout.EditCategory("Mirroring", TEXT(""), ECategoryPriority::Default);
+	TArray<TSharedRef<IPropertyHandle>> MirroringProperties;
+	MirroringCategory.GetDefaultProperties(MirroringProperties);
+	for (int32 MirrorPropertyIdx = 0; MirrorPropertyIdx < MirroringProperties.Num(); MirrorPropertyIdx++)
+	{
+		DetailLayout.HideProperty(MirroringProperties[MirrorPropertyIdx]);
+	}
+
 }
 
 void FPersonaMeshDetails::OnGetMaterialsForView(IMaterialListBuilder& MaterialList, int32 LODIndex)
@@ -418,9 +1188,7 @@ void FPersonaMeshDetails::OnMaterialChanged(UMaterialInterface* NewMaterial, UMa
 
 	if (LODIndex == 0)
 	{
-		check(Mesh->GetImportedResource());
-
-		int MaterialIndex = Mesh->GetImportedResource()->LODModels[LODIndex].Sections[SlotIndex].MaterialIndex;
+		int MaterialIndex = GetMaterialIndex(LODIndex, SlotIndex);
 		UMaterialInterface* Material = Mesh->Materials[MaterialIndex].MaterialInterface;
 		if (Material == PrevMaterial || bReplaceAll)
 		{
diff --git a/Engine/Source/Editor/Persona/Private/PersonaMeshDetails.h b/Engine/Source/Editor/Persona/Private/PersonaMeshDetails.h
--- a/Engine/Source/Editor/Persona/Private/PersonaMeshDetails.h
+++ b/Engine/Source/Editor/Persona/Private/PersonaMeshDetails.h
@@ -36,8 +36,64 @@ struct FClothingComboInfo
 	TArray<int32>							ClothingComboSelectedIndices;
 };
 
-class FPersonaMeshDetails : public IDetailCustomization
 
+class FSkelMeshReductionSettingsLayout : public IDetailCustomNodeBuilder, public TSharedFromThis<FSkelMeshReductionSettingsLayout>
+{
+public:
+	FSkelMeshReductionSettingsLayout(int32 InLODIndex, TSharedRef<class FPersonaMeshDetails> InParentLODSettings, TSharedPtr<IPropertyHandle> InBoneToRemoveProperty);
+	virtual ~FSkelMeshReductionSettingsLayout();
+
+	const FSkeletalMeshOptimizationSettings& GetSettings() const;
+	void UpdateSettings(const FSkeletalMeshOptimizationSettings& InSettings);
+private:
+	/** IDetailCustomNodeBuilder Interface*/
+	virtual void SetOnRebuildChildren(FSimpleDelegate InOnRegenerateChildren) override {}
+	virtual void GenerateHeaderRowContent(FDetailWidgetRow& NodeRow) override;
+	virtual void GenerateChildContent(IDetailChildrenBuilder& ChildrenBuilder) override;
+	virtual void Tick(float DeltaTime) override{}
+	virtual bool RequiresTick() const override { return false; }
+	virtual FName GetName() const override { static FName MeshReductionSettings("MeshReductionSettings"); return MeshReductionSettings; }
+	virtual bool InitiallyCollapsed() const override { return true; }
+
+	FReply OnApplyChanges();
+	float GetPercentTriangles() const;
+	float GetMaxDeviation() const;
+	float GetWeldingThreshold() const;
+	ESlateCheckBoxState::Type ShouldRecalculateNormals() const;
+	float GetHardAngleThreshold() const;
+	int32 GetMaxBonesPerVertex() const;
+
+	void OnPercentTrianglesChanged(float NewValue);
+	void OnMaxDeviationChanged(float NewValue);
+	void OnReductionAmountChanged(float NewValue);
+	void OnRecalculateNormalsChanged(ESlateCheckBoxState::Type NewValue);
+	void OnWeldingThresholdChanged(float NewValue);
+	void OnHardAngleThresholdChanged(float NewValue);
+	void OnMaxBonesPerVertexChanged(int32 NewValue);
+
+	void OnSilhouetteImportanceChanged(TSharedPtr<FString> NewValue, ESelectInfo::Type SelectInfo);
+	void OnTextureImportanceChanged(TSharedPtr<FString> NewValue, ESelectInfo::Type SelectInfo);
+	void OnShadingImportanceChanged(TSharedPtr<FString> NewValue, ESelectInfo::Type SelectInfo);
+	void OnSkinningImportanceChanged(TSharedPtr<FString> NewValue, ESelectInfo::Type SelectInfo);
+
+	void UpdateBonesToRemoveProperties(int32 LODIndex);
+	void RefreshBonesToRemove();
+
+private:
+	int32 LODIndex;
+	TWeakPtr<class FPersonaMeshDetails> ParentLODSettings;
+	TSharedPtr<IPropertyHandle>	BoneToRemoveProperty;
+	FSkeletalMeshOptimizationSettings ReductionSettings;
+
+	TArray<TSharedPtr<FString> > ImportanceOptions;
+	TArray<TSharedPtr<FString> > SimplificationOptions;
+	TSharedPtr<STextComboBox> SilhouetteCombo;
+	TSharedPtr<STextComboBox> TextureCombo;
+	TSharedPtr<STextComboBox> ShadingCombo;
+	TSharedPtr<STextComboBox> SkinningCombo;
+};
+
+class FPersonaMeshDetails : public IDetailCustomization
 {
 public:
 	FPersonaMeshDetails(TSharedPtr<FPersona> InPersona) : PersonaPtr(InPersona) {}
@@ -133,6 +189,38 @@ private:
 	/** Get a material index from LOD index and section index */
 	int32 GetMaterialIndex(int32 LODIndex, int32 SectionIndex);
 
+	/** for LOD settings category */
+	void CustomizeLODSettingsCategories(IDetailLayoutBuilder& DetailLayout);
+
+	void OnImportLOD(TSharedPtr<FString> NewValue, ESelectInfo::Type SelectInfo, IDetailLayoutBuilder* DetailLayout);
+	void UpdateLODNames();
+	int32 GetLODCount() const;
+	void OnLODCountChanged(int32 NewValue);
+	void OnLODCountCommitted(int32 InValue, ETextCommit::Type CommitInfo);
+	FText GetLODCountTooltip() const;
+	/** apply LOD changes if the user modified LOD reduction settings */
+	FReply OnApplyChanges();
+	/** hide properties which don't need to be showed to end users */
+	void HideUnnecessaryProperties(IDetailLayoutBuilder& DetailLayout);
+	/** clear "None" bones and remove already included bones */
+	void RefreshBonesToRemove(TArray<FBoneReference>& BonesToRemove, int32 LODIndex);
+
+public:
+
+	bool IsApplyNeeded() const;
+	void ApplyChanges();
+
+	USkeletalMesh* GetMesh()
+	{ 
+		if (PersonaPtr.IsValid())
+		{
+			return PersonaPtr->GetMesh();
+		}
+		else
+		{
+			return NULL;
+		}
+	}
 private:
 	// Container for the objects to display
 	TArray< TWeakObjectPtr<UObject>> SelectedObjects;
@@ -142,6 +230,14 @@ private:
 
 	IDetailLayoutBuilder* MeshDetailLayout;
 
+	/** LOD import options */
+	TArray<TSharedPtr<FString> > LODNames;
+	/** Helper value that corresponds to the 'Number of LODs' spinbox.*/
+	int32 LODCount;
+
+	/** Simplification options for each LOD level */
+	TArray<TSharedPtr<FSkelMeshReductionSettingsLayout>> ReductionSettingsWidgets;
+
 #if WITH_APEX_CLOTHING
 private:
 
diff --git a/Engine/Source/Editor/Persona/Private/SAnimViewportToolBar.cpp b/Engine/Source/Editor/Persona/Private/SAnimViewportToolBar.cpp
--- a/Engine/Source/Editor/Persona/Private/SAnimViewportToolBar.cpp
+++ b/Engine/Source/Editor/Persona/Private/SAnimViewportToolBar.cpp
@@ -684,13 +684,6 @@ TSharedRef<SWidget> SAnimViewportToolBar::GenerateLODMenu() const
 			}
 		}
 		ShowMenuBuilder.EndSection();
-
-		// Commands
-		ShowMenuBuilder.BeginSection("AnimViewportLODSettings");
-		{
-			ShowMenuBuilder.AddMenuEntry( Actions.ShowLevelOfDetailSettings);
-		}
-		ShowMenuBuilder.EndSection();
 	}
 
 	return ShowMenuBuilder.MakeWidget();
diff --git a/Engine/Source/Editor/Persona/Private/SAnimationEditorViewport.cpp b/Engine/Source/Editor/Persona/Private/SAnimationEditorViewport.cpp
--- a/Engine/Source/Editor/Persona/Private/SAnimationEditorViewport.cpp
+++ b/Engine/Source/Editor/Persona/Private/SAnimationEditorViewport.cpp
@@ -642,10 +642,6 @@ void SAnimationEditorViewportTabBody::BindCommands()
 		FCanExecuteAction(),
 		FIsActionChecked::CreateSP(this, &SAnimationEditorViewportTabBody::IsLODModelSelected, LOD_3));
 
-	CommandList.MapAction(
-		ViewportLODMenuCommands.ShowLevelOfDetailSettings,
-		FExecuteAction::CreateSP(this, &SAnimationEditorViewportTabBody::OnShowLevelOfDetailSettings));
-
 	CommandList.MapAction( 
 		ViewportShowMenuCommands.ToggleGrid,
 		FExecuteAction::CreateSP(this, &SAnimationEditorViewportTabBody::OnShowGrid),
@@ -1220,20 +1216,6 @@ void SAnimationEditorViewportTabBody::OnSetLODModel(ELODViewSelection LODSelecti
 	}
 }
 
-void SAnimationEditorViewportTabBody::OnShowLevelOfDetailSettings()
-{
-	UDebugSkelMeshComponent* PreviewComponent = PersonaPtr.Pin()->PreviewComponent;
-	if(TargetSkeleton)
-	{
-		FSkeletalMeshUpdateContext UpdateContext;
-		UpdateContext.SkeletalMesh = PersonaPtr.Pin()->GetMesh();
-		UpdateContext.AssociatedComponents.Push(PreviewComponent);
-		UpdateContext.OnLODChanged.BindSP(this, &SAnimationEditorViewportTabBody::OnLODChanged);
-
-		FLODUtilities::CreateLODSettingsDialog(UpdateContext);
-	}
-}
-
 FLinearColor SAnimationEditorViewportTabBody::GetViewportBackgroundColor() const
 {
 	TSharedRef<FAnimationViewportClient> AnimViewportClient = StaticCastSharedRef<FAnimationViewportClient>(LevelViewportClient.ToSharedRef());
diff --git a/Engine/Source/Editor/Persona/Private/SAnimationEditorViewport.h b/Engine/Source/Editor/Persona/Private/SAnimationEditorViewport.h
--- a/Engine/Source/Editor/Persona/Private/SAnimationEditorViewport.h
+++ b/Engine/Source/Editor/Persona/Private/SAnimationEditorViewport.h
@@ -247,9 +247,6 @@ private:
 	/** Function to set LOD model selection*/
 	void OnSetLODModel( ELODViewSelection LODSelectionType );
 
-	/** Function to show Level of Detail Settings dialog */
-	void OnShowLevelOfDetailSettings( );
-
 	/** Function to show/hide mesh info*/
 	void OnShowDisplayInfo();
 
diff --git a/Engine/Source/Editor/UnrealEd/Private/Fbx/FbxSkeletalMeshImport.cpp b/Engine/Source/Editor/UnrealEd/Private/Fbx/FbxSkeletalMeshImport.cpp
--- a/Engine/Source/Editor/UnrealEd/Private/Fbx/FbxSkeletalMeshImport.cpp
+++ b/Engine/Source/Editor/UnrealEd/Private/Fbx/FbxSkeletalMeshImport.cpp
@@ -1217,6 +1217,9 @@ USkeletalMesh* UnFbx::FFbxImporter::ImportSkeletalMesh(UObject* InParent, TArray
 	SkeletalMesh->LODInfo.Empty();
 	SkeletalMesh->LODInfo.AddZeroed();
 	SkeletalMesh->LODInfo[0].LODHysteresis = 0.02f;
+	FSkeletalMeshOptimizationSettings Settings;
+	// set default reduction settings values
+	SkeletalMesh->LODInfo[0].ReductionSettings = Settings;
 
 	// Create initial bounding box based on expanded version of reference pose for meshes without physics assets. Can be overridden by artist.
 	FBox BoundingBox( SkelMeshImportDataPtr->Points.GetTypedData(), SkelMeshImportDataPtr->Points.Num() );
diff --git a/Engine/Source/Editor/UnrealEd/Private/LODUtilities.cpp b/Engine/Source/Editor/UnrealEd/Private/LODUtilities.cpp
--- a/Engine/Source/Editor/UnrealEd/Private/LODUtilities.cpp
+++ b/Engine/Source/Editor/UnrealEd/Private/LODUtilities.cpp
@@ -292,9 +292,16 @@ void FLODUtilities::SimplifySkeletalMesh( FSkeletalMeshUpdateContext& UpdateCont
 	if ( MeshReduction && MeshReduction->IsSupported() && SkeletalMesh )
 	{
 		// Simplify each LOD
-		for( int32 LODIndex = 0; LODIndex < InSettings.Num(); ++LODIndex)
+		for (int32 SettingIndex = 0; SettingIndex < InSettings.Num(); ++SettingIndex)
 		{
-			uint32 DesiredLOD = LODIndex+1;
+			uint32 DesiredLOD = SettingIndex + 1;
+
+			// check whether reduction settings are same or not
+			if (SkeletalMesh->LODInfo.IsValidIndex(DesiredLOD) 
+			 && SkeletalMesh->LODInfo[DesiredLOD].ReductionSettings == InSettings[SettingIndex])
+			{
+				continue;
+			}
 
 			{
 				FFormatNamedArguments Args;
@@ -304,7 +311,7 @@ void FLODUtilities::SimplifySkeletalMesh( FSkeletalMeshUpdateContext& UpdateCont
 				GWarn->BeginSlowTask( StatusUpdate, true );
 			}
 
-			if ( MeshReduction->ReduceSkeletalMesh( SkeletalMesh, DesiredLOD, InSettings[LODIndex] , true ) )
+			if (MeshReduction->ReduceSkeletalMesh(SkeletalMesh, DesiredLOD, InSettings[SettingIndex], true))
 			{
 				check( SkeletalMesh->LODInfo.Num() >= 2 );
 				SkeletalMesh->MarkPackageDirty();
diff --git a/Engine/Source/Runtime/Engine/Classes/Animation/AnimationAsset.h b/Engine/Source/Runtime/Engine/Classes/Animation/AnimationAsset.h
--- a/Engine/Source/Runtime/Engine/Classes/Animation/AnimationAsset.h
+++ b/Engine/Source/Runtime/Engine/Classes/Animation/AnimationAsset.h
@@ -138,253 +138,7 @@ struct FAnimExtractContext
 	}
 };
 
-/**
- * This is a native transient structure.
- * Contains:
- * - BoneIndicesArray: Array of RequiredBoneIndices for Current Asset. In increasing order. Mapping to current Array of Transforms (Pose).
- * - BoneSwitchArray: Size of current Skeleton. true if Bone is contained in RequiredBones array, false otherwise.
- **/
-struct FBoneContainer
-{
-private:
-	/** Array of RequiredBonesIndices. In increasing order. */
-	TArray<FBoneIndexType>	BoneIndicesArray;
-	/** Array sized by Current RefPose. true if Bone is contained in RequiredBones array, false otherwise. */
-	TBitArray<>				BoneSwitchArray;
-
-	/** Asset BoneIndicesArray was made for. Typically a SkeletalMesh. */
-	TWeakObjectPtr<UObject>	Asset;
-	/** If Asset is a SkeletalMesh, this will be a pointer to it. Can be NULL if Asset is a USkeleton. */
-	TWeakObjectPtr<USkeletalMesh> AssetSkeletalMesh;
-	/** If Asset is a Skeleton that will be it. If Asset is a SkeletalMesh, that will be its Skeleton. */
-	TWeakObjectPtr<USkeleton> AssetSkeleton;
-
-	/** Pointer to RefSkeleton of Asset. */
-	const FReferenceSkeleton* RefSkeleton;
-
-	/** Mapping table between Skeleton Bone Indices and Pose Bone Indices. */
-	TArray<int32> SkeletonToPoseBoneIndexArray;
-
-	/** Mapping table between Pose Bone Indices and Skeleton Bone Indices. */
-	TArray<int32> PoseToSkeletonBoneIndexArray;
-
-	/** For debugging. */
-	/** Disable Retargeting. Extract animation, but do not retarget it. */
-	bool bDisableRetargeting;
-	/** Disable animation compression, use RAW data instead. */
-	bool bUseRAWData;
-
-public:
-
-	FBoneContainer()
-		: Asset(NULL)
-		, AssetSkeletalMesh(NULL)
-		, AssetSkeleton(NULL)
-		, RefSkeleton(NULL)
-		, bDisableRetargeting(false)
-		, bUseRAWData(false)
-	{
-		BoneIndicesArray.Empty();
-		BoneSwitchArray.Empty();
-		SkeletonToPoseBoneIndexArray.Empty();
-		PoseToSkeletonBoneIndexArray.Empty();
-	}
-
-	FBoneContainer(const TArray<FBoneIndexType>& InRequiredBoneIndexArray, UObject& InAsset)
-		: BoneIndicesArray(InRequiredBoneIndexArray)
-		, Asset(&InAsset)
-		, AssetSkeletalMesh(NULL)
-		, AssetSkeleton(NULL)
-		, RefSkeleton(NULL)
-		, bDisableRetargeting(false)
-		, bUseRAWData(false)
-	{
-		Initialize();
-	}
-
-	/** Initialize BoneContainer to a new Asset, RequiredBonesArray and RefPoseArray. */
-	void InitializeTo(const TArray<FBoneIndexType>& InRequiredBoneIndexArray, UObject& InAsset);
-
-	/** Returns true if FBoneContainer is Valid. Needs an Asset, a RefPoseArray, and a RequiredBonesArray. */
-	const bool IsValid() const
-	{
-		return (Asset.IsValid() && (RefSkeleton != NULL) && (BoneIndicesArray.Num() > 0));
-	}
-
-	/** Get Asset this BoneContainer was made for. Typically a SkeletalMesh, but could also be a USkeleton. */
-	UObject* GetAsset() const
-	{
-		return Asset.Get();
-	}
-
-	/** Get SkeletalMesh Asset this BoneContainer was made for. Could be NULL if Asset is a Skeleton. */
-	USkeletalMesh* GetSkeletalMeshAsset() const
-	{
-		return AssetSkeletalMesh.Get();
-	}
-
-	/** Get Skeleton Asset. Could either be the SkeletalMesh's Skeleton, or the Skeleton this BoneContainer was made for. Is non NULL is BoneContainer is valid. */
-	USkeleton * GetSkeletonAsset() const
-	{
-		return AssetSkeleton.Get();
-	}
-
-	/** Disable Retargeting for debugging. */
-	void SetDisableRetargeting(bool InbDisableRetargeting)
-	{
-		bDisableRetargeting = InbDisableRetargeting;
-	}
-
-	/** True if retargeting is disabled for debugging. */
-	bool GetDisableRetargeting() const
-	{
-		return bDisableRetargeting;
-	}
-
-	/** Ignore compressed data and use RAW data instead, for debugging. */
-	void SetUseRAWData(bool InbUseRAWData)
-	{
-		bUseRAWData = InbUseRAWData;
-	}
-
-	/** True if we're requesting RAW data instead of compressed data. For debugging. */
-	bool ShouldUseRawData() const
-	{
-		return bUseRAWData;
-	}
-
- 	/**
-	 * returns Required Bone Indices Array
-	 */
-	const TArray<FBoneIndexType>& GetBoneIndicesArray() const
-	{
-		return BoneIndicesArray;
-	}
-
-	/**
-	 * returns Bone Switch Array. BitMask for RequiredBoneIndex array.
-	 */
-	const TBitArray<>& GetBoneSwitchArray() const
-	{
-		return BoneSwitchArray;
-	}
-
-	/** Pointer to RefPoseArray for current Asset. */
-	const TArray<FTransform>& GetRefPoseArray() const
-	{
-		return RefSkeleton->GetRefBonePose();
-	}
-
-	/** Access to Asset's RefSkeleton. */
-	const FReferenceSkeleton& GetReferenceSkeleton() const
-	{
-		return *RefSkeleton;
-	}
-
-	/** Number of Bones in RefPose for current asset. This is NOT the number of bones in RequiredBonesArray, but the TOTAL number of bones in the RefPose of the current Asset! */ 
-	const int32 GetNumBones() const
-	{
-		return RefSkeleton->GetNum();
-	}
-
-	/** Get BoneIndex for BoneName for current Asset. */
-	ENGINE_API int32 GetPoseBoneIndexForBoneName(const FName& BoneName) const;
 
-	/** Get ParentBoneIndex for current Asset. */
-	ENGINE_API int32 GetParentBoneIndex(const int32& BoneIndex) const;
-
-	/** Get Depth between bones for current asset. */
-	int32 GetDepthBetweenBones(const int32& BoneIndex, const int32& ParentBoneIndex) const;
-
-	/** Returns true if bone is child of for current asset. */
-	bool BoneIsChildOf(const int32& BoneIndex, const int32& ParentBoneIndex) const;
-
-	/**
-	 * Serializes the bones
-	 *
-	 * @param Ar - The archive to serialize into.
-	 * @param Rect - The bone container to serialize.
-	 *
-	 * @return Reference to the Archive after serialization.
-	 */
-	friend FArchive& operator<<( FArchive& Ar, FBoneContainer& B )
-	{ 
-		Ar 
-			<< B.BoneIndicesArray 
-			<< B.BoneSwitchArray 
-			<< B.Asset 
-			<< B.AssetSkeletalMesh 
-			<< B.AssetSkeleton 
-			<< B.SkeletonToPoseBoneIndexArray
-			<< B.PoseToSkeletonBoneIndexArray
-			<< B.bDisableRetargeting
-			<< B.bUseRAWData
-			;	
-		return Ar;
-	}
-
-	/**
-	 * Returns true of RequiredBonesArray contains this bone index.
-	 */
-	bool Contains(FBoneIndexType NewIndex) const
-	{
-		return BoneSwitchArray[NewIndex];
-	}
-
-	/** Const accessor to GetSkeletonToPoseBoneIndexArray(). */
-	TArray<int32> const & GetSkeletonToPoseBoneIndexArray() const
-	{
-		return SkeletonToPoseBoneIndexArray;
-	}
-
-	/** Const accessor to GetSkeletonToPoseBoneIndexArray(). */
-	TArray<int32> const & GetPoseToSkeletonBoneIndexArray() const
-	{
-		return PoseToSkeletonBoneIndexArray;
-	}
-
-private:
-	/** Initialize FBoneContainer. */
-	ENGINE_API void Initialize();
-
-	/** Cache remapping data if current Asset is a SkeletalMesh, with all compatible Skeletons. */
-	void RemapFromSkelMesh(USkeletalMesh const & SourceSkeletalMesh, USkeleton & TargetSkeleton);
-
-	/** Cache remapping data if current Asset is a Skeleton, with all compatible Skeletons. */
-	void RemapFromSkeleton(USkeleton const & SourceSkeleton);
-};
-
-USTRUCT()
-struct FBoneReference
-{
-	GENERATED_USTRUCT_BODY()
-
-	/** Name of bone to control. This is the main bone chain to modify from. **/
-	UPROPERTY(EditAnywhere, Category=BoneReference)
-	FName BoneName;
-
-	/** Cached bone index for run time - right now bone index of skeleton **/
-	int32 BoneIndex;
-
-	FBoneReference()
-		: BoneIndex(INDEX_NONE)
-	{
-	}
-
-	bool operator==( const FBoneReference& Other ) const
-	{
-		// faster to compare, and BoneName won't matter
-		return BoneIndex==Other.BoneIndex;
-	}
-	/** Initialize Bone Reference, return TRUE if success, otherwise, return false **/
-	ENGINE_API bool Initialize(const FBoneContainer& RequiredBones);
-
-	// @fixme laurent - only used by blendspace 'PerBoneBlend'. Fix this to support SkeletalMesh pose.
-	ENGINE_API bool Initialize(const USkeleton* Skeleton);
-
-	/** return true if valid. Otherwise return false **/
-	ENGINE_API bool IsValid(const FBoneContainer& RequiredBones) const;
-};
 /**
  * Information about an animation asset that needs to be ticked
  */
diff --git a/Engine/Source/Runtime/Engine/Classes/Engine/SkeletalMesh.h b/Engine/Source/Runtime/Engine/Classes/Engine/SkeletalMesh.h
--- a/Engine/Source/Runtime/Engine/Classes/Engine/SkeletalMesh.h
+++ b/Engine/Source/Runtime/Engine/Classes/Engine/SkeletalMesh.h
@@ -10,6 +10,9 @@
 #include "Animation/PreviewAssetAttachComponent.h"
 #include "SkeletalMesh.generated.h"
 
+/** The maximum number of skeletal mesh LODs allowed. */
+#define MAX_SKELETAL_MESH_LODS 5
+
 class UMorphTarget;
 
 #if WITH_APEX_CLOTHING
@@ -114,105 +117,314 @@
 
 };
 
-/** Struct containing information for a particular LOD level, such as materials and info for when to use it. */
+/**
+* This is a native transient structure.
+* Contains:
+* - BoneIndicesArray: Array of RequiredBoneIndices for Current Asset. In increasing order. Mapping to current Array of Transforms (Pose).
+* - BoneSwitchArray: Size of current Skeleton. true if Bone is contained in RequiredBones array, false otherwise.
+**/
+struct FBoneContainer
+{
+private:
+	/** Array of RequiredBonesIndices. In increasing order. */
+	TArray<FBoneIndexType>	BoneIndicesArray;
+	/** Array sized by Current RefPose. true if Bone is contained in RequiredBones array, false otherwise. */
+	TBitArray<>				BoneSwitchArray;
+
+	/** Asset BoneIndicesArray was made for. Typically a SkeletalMesh. */
+	TWeakObjectPtr<UObject>	Asset;
+	/** If Asset is a SkeletalMesh, this will be a pointer to it. Can be NULL if Asset is a USkeleton. */
+	TWeakObjectPtr<USkeletalMesh> AssetSkeletalMesh;
+	/** If Asset is a Skeleton that will be it. If Asset is a SkeletalMesh, that will be its Skeleton. */
+	TWeakObjectPtr<USkeleton> AssetSkeleton;
+
+	/** Pointer to RefSkeleton of Asset. */
+	const FReferenceSkeleton* RefSkeleton;
+
+	/** Mapping table between Skeleton Bone Indices and Pose Bone Indices. */
+	TArray<int32> SkeletonToPoseBoneIndexArray;
+
+	/** Mapping table between Pose Bone Indices and Skeleton Bone Indices. */
+	TArray<int32> PoseToSkeletonBoneIndexArray;
+
+	/** For debugging. */
+	/** Disable Retargeting. Extract animation, but do not retarget it. */
+	bool bDisableRetargeting;
+	/** Disable animation compression, use RAW data instead. */
+	bool bUseRAWData;
+
+public:
+
+	FBoneContainer()
+		: Asset(NULL)
+		, AssetSkeletalMesh(NULL)
+		, AssetSkeleton(NULL)
+		, RefSkeleton(NULL)
+		, bDisableRetargeting(false)
+		, bUseRAWData(false)
+	{
+		BoneIndicesArray.Empty();
+		BoneSwitchArray.Empty();
+		SkeletonToPoseBoneIndexArray.Empty();
+		PoseToSkeletonBoneIndexArray.Empty();
+	}
+
+	FBoneContainer(const TArray<FBoneIndexType>& InRequiredBoneIndexArray, UObject& InAsset)
+		: BoneIndicesArray(InRequiredBoneIndexArray)
+		, Asset(&InAsset)
+		, AssetSkeletalMesh(NULL)
+		, AssetSkeleton(NULL)
+		, RefSkeleton(NULL)
+		, bDisableRetargeting(false)
+		, bUseRAWData(false)
+	{
+		Initialize();
+	}
+
+	/** Initialize BoneContainer to a new Asset, RequiredBonesArray and RefPoseArray. */
+	void InitializeTo(const TArray<FBoneIndexType>& InRequiredBoneIndexArray, UObject& InAsset);
+
+	/** Returns true if FBoneContainer is Valid. Needs an Asset, a RefPoseArray, and a RequiredBonesArray. */
+	const bool IsValid() const
+	{
+		return (Asset.IsValid() && (RefSkeleton != NULL) && (BoneIndicesArray.Num() > 0));
+	}
+
+	/** Get Asset this BoneContainer was made for. Typically a SkeletalMesh, but could also be a USkeleton. */
+	UObject* GetAsset() const
+	{
+		return Asset.Get();
+	}
+
+	/** Get SkeletalMesh Asset this BoneContainer was made for. Could be NULL if Asset is a Skeleton. */
+	USkeletalMesh* GetSkeletalMeshAsset() const
+	{
+		return AssetSkeletalMesh.Get();
+	}
+
+	/** Get Skeleton Asset. Could either be the SkeletalMesh's Skeleton, or the Skeleton this BoneContainer was made for. Is non NULL is BoneContainer is valid. */
+	USkeleton * GetSkeletonAsset() const
+	{
+		return AssetSkeleton.Get();
+	}
+
+	/** Disable Retargeting for debugging. */
+	void SetDisableRetargeting(bool InbDisableRetargeting)
+	{
+		bDisableRetargeting = InbDisableRetargeting;
+	}
+
+	/** True if retargeting is disabled for debugging. */
+	bool GetDisableRetargeting() const
+	{
+		return bDisableRetargeting;
+	}
+
+	/** Ignore compressed data and use RAW data instead, for debugging. */
+	void SetUseRAWData(bool InbUseRAWData)
+	{
+		bUseRAWData = InbUseRAWData;
+	}
+
+	/** True if we're requesting RAW data instead of compressed data. For debugging. */
+	bool ShouldUseRawData() const
+	{
+		return bUseRAWData;
+	}
+
+	/**
+	* returns Required Bone Indices Array
+	*/
+	const TArray<FBoneIndexType>& GetBoneIndicesArray() const
+	{
+		return BoneIndicesArray;
+	}
+
+	/**
+	* returns Bone Switch Array. BitMask for RequiredBoneIndex array.
+	*/
+	const TBitArray<>& GetBoneSwitchArray() const
+	{
+		return BoneSwitchArray;
+	}
+
+	/** Pointer to RefPoseArray for current Asset. */
+	const TArray<FTransform>& GetRefPoseArray() const
+	{
+		return RefSkeleton->GetRefBonePose();
+	}
+
+	/** Access to Asset's RefSkeleton. */
+	const FReferenceSkeleton& GetReferenceSkeleton() const
+	{
+		return *RefSkeleton;
+	}
+
+	/** Number of Bones in RefPose for current asset. This is NOT the number of bones in RequiredBonesArray, but the TOTAL number of bones in the RefPose of the current Asset! */
+	const int32 GetNumBones() const
+	{
+		return RefSkeleton->GetNum();
+	}
+
+	/** Get BoneIndex for BoneName for current Asset. */
+	ENGINE_API int32 GetPoseBoneIndexForBoneName(const FName& BoneName) const;
+
+	/** Get ParentBoneIndex for current Asset. */
+	ENGINE_API int32 GetParentBoneIndex(const int32& BoneIndex) const;
+
+	/** Get Depth between bones for current asset. */
+	int32 GetDepthBetweenBones(const int32& BoneIndex, const int32& ParentBoneIndex) const;
+
+	/** Returns true if bone is child of for current asset. */
+	bool BoneIsChildOf(const int32& BoneIndex, const int32& ParentBoneIndex) const;
+
+	/**
+	* Serializes the bones
+	*
+	* @param Ar - The archive to serialize into.
+	* @param Rect - The bone container to serialize.
+	*
+	* @return Reference to the Archive after serialization.
+	*/
+	friend FArchive& operator<<(FArchive& Ar, FBoneContainer& B)
+	{
+		Ar
+			<< B.BoneIndicesArray
+			<< B.BoneSwitchArray
+			<< B.Asset
+			<< B.AssetSkeletalMesh
+			<< B.AssetSkeleton
+			<< B.SkeletonToPoseBoneIndexArray
+			<< B.PoseToSkeletonBoneIndexArray
+			<< B.bDisableRetargeting
+			<< B.bUseRAWData
+			;
+		return Ar;
+	}
+
+	/**
+	* Returns true of RequiredBonesArray contains this bone index.
+	*/
+	bool Contains(FBoneIndexType NewIndex) const
+	{
+		return BoneSwitchArray[NewIndex];
+	}
+
+	/** Const accessor to GetSkeletonToPoseBoneIndexArray(). */
+	TArray<int32> const & GetSkeletonToPoseBoneIndexArray() const
+	{
+		return SkeletonToPoseBoneIndexArray;
+	}
+
+	/** Const accessor to GetSkeletonToPoseBoneIndexArray(). */
+	TArray<int32> const & GetPoseToSkeletonBoneIndexArray() const
+	{
+		return PoseToSkeletonBoneIndexArray;
+	}
+
+private:
+	/** Initialize FBoneContainer. */
+	ENGINE_API void Initialize();
+
+	/** Cache remapping data if current Asset is a SkeletalMesh, with all compatible Skeletons. */
+	void RemapFromSkelMesh(USkeletalMesh const & SourceSkeletalMesh, USkeleton & TargetSkeleton);
+
+	/** Cache remapping data if current Asset is a Skeleton, with all compatible Skeletons. */
+	void RemapFromSkeleton(USkeleton const & SourceSkeleton);
+};
+
 USTRUCT()
-struct FSkeletalMeshLODInfo
+struct FBoneReference
 {
 	GENERATED_USTRUCT_BODY()
 
-	/**	Indicates when to use this LOD. A smaller number means use this LOD when further away. */
-	UPROPERTY(EditAnywhere, Category=SkeletalMeshLODInfo)
-	float DisplayFactor;
+		/** Name of bone to control. This is the main bone chain to modify from. **/
+		UPROPERTY(EditAnywhere, Category = BoneReference)
+		FName BoneName;
 
-	/**	Used to avoid 'flickering' when on LOD boundary. Only taken into account when moving from complex->simple. */
-	UPROPERTY(EditAnywhere, Category=SkeletalMeshLODInfo)
-	float LODHysteresis;
+	/** Cached bone index for run time - right now bone index of skeleton **/
+	int32 BoneIndex;
 
-	/** Mapping table from this LOD's materials to the USkeletalMesh materials array. */
-	UPROPERTY(EditAnywhere, editfixedsize, Category=SkeletalMeshLODInfo)
-	TArray<int32> LODMaterialMap;
-
-	/** Per-section control over whether to enable shadow casting. */
-	UPROPERTY()
-	TArray<bool> bEnableShadowCasting_DEPRECATED;
-
-	/** Per-section sorting options */
-	UPROPERTY()
-	TArray<TEnumAsByte<enum ETriangleSortOption> > TriangleSorting_DEPRECATED;
-
-	UPROPERTY(EditAnywhere, editfixedsize, Category=SkeletalMeshLODInfo)
-	TArray<struct FTriangleSortSettings> TriangleSortSettings;
-
-	/** Whether to disable morph targets for this LOD. */
-	UPROPERTY()
-	uint32 bHasBeenSimplified:1;
-
-
-	FSkeletalMeshLODInfo()
-		: DisplayFactor(0)
-		, LODHysteresis(0)
-		, bHasBeenSimplified(false)
+	FBoneReference()
+		: BoneIndex(INDEX_NONE)
 	{
 	}
 
+	bool operator==(const FBoneReference& Other) const
+	{
+		// faster to compare, and BoneName won't matter
+		return BoneIndex == Other.BoneIndex;
+	}
+	/** Initialize Bone Reference, return TRUE if success, otherwise, return false **/
+	ENGINE_API bool Initialize(const FBoneContainer& RequiredBones);
+
+	// @fixme laurent - only used by blendspace 'PerBoneBlend'. Fix this to support SkeletalMesh pose.
+	ENGINE_API bool Initialize(const USkeleton* Skeleton);
+
+	/** return true if valid. Otherwise return false **/
+	ENGINE_API bool IsValid(const FBoneContainer& RequiredBones) const;
 };
 
 /**
- * FSkeletalMeshOptimizationSettings - The settings used to optimize a skeletal mesh LOD.
- */
+* FSkeletalMeshOptimizationSettings - The settings used to optimize a skeletal mesh LOD.
+*/
 USTRUCT()
 struct FSkeletalMeshOptimizationSettings
 {
 	GENERATED_USTRUCT_BODY()
 
-	/** The method to use when optimizing the skeletal mesh LOD */
-	UPROPERTY()
-	TEnumAsByte<enum SkeletalMeshOptimizationType> ReductionMethod;
+		/** The method to use when optimizing the skeletal mesh LOD */
+		UPROPERTY()
+		TEnumAsByte<enum SkeletalMeshOptimizationType> ReductionMethod;
 
 	/** If ReductionMethod equals SMOT_NumOfTriangles this value is the ratio of triangles [0-1] to remove from the mesh */
 	UPROPERTY()
-	float NumOfTrianglesPercentage;
+		float NumOfTrianglesPercentage;
 
 	/**If ReductionMethod equals SMOT_MaxDeviation this value is the maximum deviation from the base mesh as a percentage of the bounding sphere. */
 	UPROPERTY()
-	float MaxDeviationPercentage;
+		float MaxDeviationPercentage;
 
 	/** The welding threshold distance. Vertices under this distance will be welded. */
 	UPROPERTY()
-	float WeldingThreshold; 
+		float WeldingThreshold;
 
 	/** Whether Normal smoothing groups should be preserved. If false then NormalsThreshold is used **/
 	UPROPERTY()
-	bool bRecalcNormals;
+		bool bRecalcNormals;
 
 	/** If the angle between two triangles are above this value, the normals will not be
 	smooth over the edge between those two triangles. Set in degrees. This is only used when PreserveNormals is set to false*/
 	UPROPERTY()
-	float NormalsThreshold;
+		float NormalsThreshold;
 
 	/** How important the shape of the geometry is. */
 	UPROPERTY()
-	TEnumAsByte<enum SkeletalMeshOptimizationImportance> SilhouetteImportance;
+		TEnumAsByte<enum SkeletalMeshOptimizationImportance> SilhouetteImportance;
 
 	/** How important texture density is. */
 	UPROPERTY()
-	TEnumAsByte<enum SkeletalMeshOptimizationImportance> TextureImportance;
+		TEnumAsByte<enum SkeletalMeshOptimizationImportance> TextureImportance;
 
 	/** How important shading quality is. */
 	UPROPERTY()
-	TEnumAsByte<enum SkeletalMeshOptimizationImportance> ShadingImportance;
+		TEnumAsByte<enum SkeletalMeshOptimizationImportance> ShadingImportance;
 
 	/** How important skinning quality is. */
 	UPROPERTY()
-	TEnumAsByte<enum SkeletalMeshOptimizationImportance> SkinningImportance;
+		TEnumAsByte<enum SkeletalMeshOptimizationImportance> SkinningImportance;
 
 	/** The ratio of bones that will be removed from the mesh */
 	UPROPERTY()
-	float BoneReductionRatio;
+		float BoneReductionRatio;
 
 	/** Maximum number of bones that can be assigned to each vertex. */
 	UPROPERTY()
-	int32 MaxBonesPerVertex;
+		int32 MaxBonesPerVertex;
 
+	UPROPERTY(EditAnywhere, Category = ReductionSettings)
+		TArray<FBoneReference> BonesToRemove;
+
+
 	FSkeletalMeshOptimizationSettings()
 		: ReductionMethod(SMOT_MaxDeviation)
@@ -230,6 +442,75 @@
 	{
 	}
 
+	/** Equality operator. */
+	bool operator==(const FSkeletalMeshOptimizationSettings& Other) const
+	{
+		return ReductionMethod == Other.ReductionMethod
+			&& NumOfTrianglesPercentage == Other.NumOfTrianglesPercentage
+			&& MaxDeviationPercentage == Other.MaxDeviationPercentage
+			&& WeldingThreshold == Other.WeldingThreshold
+			&& NormalsThreshold == Other.NormalsThreshold
+			&& SilhouetteImportance == Other.SilhouetteImportance
+			&& TextureImportance == Other.TextureImportance
+			&& ShadingImportance == Other.ShadingImportance
+			&& SkinningImportance == Other.SkinningImportance
+			&& bRecalcNormals == Other.bRecalcNormals
+			&& BoneReductionRatio == Other.BoneReductionRatio
+			&& MaxBonesPerVertex == Other.MaxBonesPerVertex
+			&& BonesToRemove == Other.BonesToRemove;
+	}
+
+	/** Inequality. */
+	bool operator!=(const FSkeletalMeshOptimizationSettings& Other) const
+	{
+		return !(*this == Other);
+	}
+};
+
+/** Struct containing information for a particular LOD level, such as materials and info for when to use it. */
+USTRUCT()
+struct FSkeletalMeshLODInfo
+{
+	GENERATED_USTRUCT_BODY()
+
+	/**	Indicates when to use this LOD. A smaller number means use this LOD when further away. */
+	UPROPERTY(EditAnywhere, Category=SkeletalMeshLODInfo)
+	float DisplayFactor;
+
+	/**	Used to avoid 'flickering' when on LOD boundary. Only taken into account when moving from complex->simple. */
+	UPROPERTY(EditAnywhere, Category=SkeletalMeshLODInfo)
+	float LODHysteresis;
+
+	/** Mapping table from this LOD's materials to the USkeletalMesh materials array. */
+	UPROPERTY(EditAnywhere, editfixedsize, Category=SkeletalMeshLODInfo)
+	TArray<int32> LODMaterialMap;
+
+	/** Per-section control over whether to enable shadow casting. */
+	UPROPERTY()
+	TArray<bool> bEnableShadowCasting_DEPRECATED;
+
+	/** Per-section sorting options */
+	UPROPERTY()
+	TArray<TEnumAsByte<enum ETriangleSortOption> > TriangleSorting_DEPRECATED;
+
+	UPROPERTY(EditAnywhere, editfixedsize, Category=SkeletalMeshLODInfo)
+	TArray<struct FTriangleSortSettings> TriangleSortSettings;
+
+	/** Whether to disable morph targets for this LOD. */
+	UPROPERTY()
+	uint32 bHasBeenSimplified:1;
+
+	/** Reduction settings to apply when building render data. */
+	UPROPERTY(EditAnywhere, Category = ReductionSettings)
+	FSkeletalMeshOptimizationSettings ReductionSettings;
+
+	FSkeletalMeshLODInfo()
+		: DisplayFactor(0)
+		, LODHysteresis(0)
+		, bHasBeenSimplified(false)
+	{
+	}
+
 };
 
 USTRUCT()
@@ -507,7 +788,7 @@
 	TArray<struct FSkeletalMeshLODInfo> LODInfo;
 
 	/** If true, use 32 bit UVs. If false, use 16 bit UVs to save memory */
-	UPROPERTY(EditAnywhere, Category=SkeletalMesh)
+	UPROPERTY(EditAnywhere, Category=Mesh)
 	uint32 bUseFullPrecisionUVs:1;
 
 	/** true if this mesh has ever been simplified with Simplygon. */
