Fixed bone reduction checking problem

diff --git a/Engine/Source/Editor/Persona/Private/PersonaMeshDetails.cpp b/Engine/Source/Editor/Persona/Private/PersonaMeshDetails.cpp
--- a/Engine/Source/Editor/Persona/Private/PersonaMeshDetails.cpp
+++ b/Engine/Source/Editor/Persona/Private/PersonaMeshDetails.cpp
@@ -683,7 +683,6 @@ void FPersonaMeshDetails::RefreshBonesToRemove(TArray<FBoneReference>& InOutBone
 			{
 				FBoneReference BoneRef;
 				BoneRef.BoneName = BoneName;
-				BoneRef.BoneIndex = BoneIndex;
 				InOutBonesToRemove.AddUnique(BoneRef);
 			}
 		}
@@ -797,7 +796,7 @@ FText FPersonaMeshDetails::GetLODCountTooltip() const
 {
 	if (IsAutoMeshReductionAvailable())
 	{
-		return LOCTEXT("LODCountTooltip", "The number of LODs for this static mesh. If auto mesh reduction is available, setting this number will determine the number of LOD levels to auto generate.");
+		return LOCTEXT("LODCountTooltip", "The number of LODs for this skeletal mesh. If auto mesh reduction is available, setting this number will determine the number of LOD levels to auto generate.");
 	}
 
 	return LOCTEXT("LODCountTooltip_Disabled", "Auto mesh reduction is unavailable! Please provide a mesh reduction interface such as Simplygon to use this feature or manually import LOD levels.");
