Preventing sounds and other notifies from playing when animation playback buttons set the position directly, so the notifies between the two positions are not important.

Came from question on UDN: https://udn.unrealengine.com/questions/212774/sound-notify-getting-played-twice-in-anim-sequence.html

diff --git a/Engine/Source/Editor/Persona/Private/SAnimationScrubPanel.cpp b/Engine/Source/Editor/Persona/Private/SAnimationScrubPanel.cpp
--- a/Engine/Source/Editor/Persona/Private/SAnimationScrubPanel.cpp
+++ b/Engine/Source/Editor/Persona/Private/SAnimationScrubPanel.cpp
@@ -96,7 +96,7 @@ FReply SAnimationScrubPanel::OnClick_Forward_End()
 	if (PreviewInstance)
 	{
 		PreviewInstance->SetPlaying(false);
-		PreviewInstance->SetPosition(PreviewInstance->GetLength());
+		PreviewInstance->SetPosition(PreviewInstance->GetLength(), false);
 	}
 
 	return FReply::Handled();
@@ -119,7 +119,7 @@ FReply SAnimationScrubPanel::OnClick_Backward_End()
 	if (PreviewInstance)
 	{
 		PreviewInstance->SetPlaying(false);
-		PreviewInstance->SetPosition(0.f);
+		PreviewInstance->SetPosition(0.f, false);
 	}
 	return FReply::Handled();
 }
@@ -147,7 +147,7 @@ FReply SAnimationScrubPanel::OnClick_Forward()
 			//if we're at the end of the animation, jump back to the beginning before playing
 			if ( GetScrubValue() >= GetSequenceLength() )
 			{
-				PreviewInstance->SetPosition(0.0f);
+				PreviewInstance->SetPosition(0.0f, false);
 			}
 
 			PreviewInstance->SetReverse(false);
@@ -184,7 +184,7 @@ FReply SAnimationScrubPanel::OnClick_Backward()
 			//if we're at the beginning of the animation, jump back to the end before playing
 			if ( GetScrubValue() <= 0.0f )
 			{
-				PreviewInstance->SetPosition(GetSequenceLength());
+				PreviewInstance->SetPosition(GetSequenceLength(), false);
 			}
 
 			PreviewInstance->SetPlaying(true);
@@ -477,7 +477,7 @@ void SAnimationScrubPanel::OnCropAnimSequence( bool bFromStart, float CurrentTim
 				FAnimationUtils::CompressAnimSequence(AnimSequence, false, false);
 
 				//Resetting slider position to the first frame
-				PreviewInstance->SetPosition( 0.0f );
+				PreviewInstance->SetPosition( 0.0f, false );
 			}
 		}
 	}
