Removed LOD bone reduction menus on Skeleton tree because the users can set in reduction settings of each LOD panel

diff --git a/Engine/Source/Editor/Persona/Private/PersonaMeshDetails.cpp b/Engine/Source/Editor/Persona/Private/PersonaMeshDetails.cpp
--- a/Engine/Source/Editor/Persona/Private/PersonaMeshDetails.cpp
+++ b/Engine/Source/Editor/Persona/Private/PersonaMeshDetails.cpp
@@ -641,6 +641,15 @@ void FPersonaMeshDetails::RefreshBonesToRemove(TArray<FBoneReference>& InOutBone
 	USkeletalMesh* SkelMesh = PersonaPtr->GetMesh();
 	USkeleton* Skeleton = SkelMesh->Skeleton;
 
+	// index for accessing to BoneReductionSettingsForLODs because it doesn't have for LOD 0
+	int32 LODSettingIndex = LODIndex - 1;
+
+	// remove all previous data before setting new bones to remove
+	if (Skeleton->BoneReductionSettingsForLODs.IsValidIndex(LODSettingIndex))
+	{
+		Skeleton->BoneReductionSettingsForLODs[LODSettingIndex].BonesToRemove.Empty();
+	}
+
 	for (int32 Index = 0; Index < InOutBonesToRemove.Num(); Index++)
 	{
 		int32 BoneIndex = Skeleton->GetReferenceSkeleton().FindBoneIndex(InOutBonesToRemove[Index].BoneName);
@@ -652,9 +661,6 @@ void FPersonaMeshDetails::RefreshBonesToRemove(TArray<FBoneReference>& InOutBone
 
 	InOutBonesToRemove.Empty();
 
-	// index for accessing to BoneReductionSettingsForLODs because it doesn't have for LOD 0
-	int32 LODSettingIndex = LODIndex - 1;
-
 	if (Skeleton->BoneReductionSettingsForLODs.Num() > LODSettingIndex)
 	{
 		const TArray<FName> & BonesToRemove = Skeleton->BoneReductionSettingsForLODs[LODSettingIndex].BonesToRemove;
@@ -683,7 +689,7 @@ void FPersonaMeshDetails::RefreshBonesToRemove(TArray<FBoneReference>& InOutBone
 			{
 				FBoneReference BoneRef;
 				BoneRef.BoneName = BoneName;
-				InOutBonesToRemove.AddUnique(BoneRef);
+				InOutBonesToRemove.Add(BoneRef);
 			}
 		}
 	}
diff --git a/Engine/Source/Editor/Persona/Private/SSkeletonTree.cpp b/Engine/Source/Editor/Persona/Private/SSkeletonTree.cpp
--- a/Engine/Source/Editor/Persona/Private/SSkeletonTree.cpp
+++ b/Engine/Source/Editor/Persona/Private/SSkeletonTree.cpp
@@ -1386,25 +1386,6 @@ TSharedPtr< SWidget > SSkeletonTree::CreateContextMenu()
 					);
 			}
 			MenuBuilder.EndSection();
-
-			if (bMeshReductionSupported)
-			{
-				MenuBuilder.BeginSection("SkeletonTreeBoneReductionForLOD", LOCTEXT( "BoneReductionHeader", "LOD Bone Reduction" ) );
-				{
-					MenuBuilder.AddSubMenu(
-						LOCTEXT( "SkeletonTreeBoneReductionForLOD_RemoveBonesFromLOD", "Remove Selected Bone from..." ),
-						FText::GetEmpty(),
-						FNewMenuDelegate::CreateStatic(&SSkeletonTree::CreateMenuForBoneReduction, this, TargetSkeleton, false)
-						);
-
-					MenuBuilder.AddSubMenu(
-						LOCTEXT( "SkeletonTreeBoneReductionForLOD_AddBonesToLOD", "Add Remove Selected Bone to..." ),
-						FText::GetEmpty(),
-						FNewMenuDelegate::CreateStatic(&SSkeletonTree::CreateMenuForBoneReduction, this, TargetSkeleton, true)
-						);
-				}
-				MenuBuilder.EndSection();
-			}
 		}
 
 		if(BoneTreeSelection.HasSelectedOfType(ESkeletonTreeRowType::Socket))
@@ -1468,6 +1449,7 @@ TSharedPtr< SWidget > SSkeletonTree::CreateContextMenu()
 	return MenuBuilder.MakeWidget();
 }
 
+/*
 void SSkeletonTree::CreateMenuForBoneReduction(FMenuBuilder& MenuBuilder, SSkeletonTree * Widget, USkeleton * Skeleton, bool bAdd)
 {
 	if (bMeshReductionSupported)
@@ -1513,6 +1495,7 @@ void SSkeletonTree::CreateMenuForBoneReduction(FMenuBuilder& MenuBuilder, SSkele
 		}
 	}
 }
+*/
 
 void SSkeletonTree::SetBoneTranslationRetargetingModeRecursive(EBoneTranslationRetargetingMode::Type NewRetargetingMode)
 {
diff --git a/Engine/Source/Editor/Persona/Private/SSkeletonTree.h b/Engine/Source/Editor/Persona/Private/SSkeletonTree.h
--- a/Engine/Source/Editor/Persona/Private/SSkeletonTree.h
+++ b/Engine/Source/Editor/Persona/Private/SSkeletonTree.h
@@ -542,9 +542,6 @@ private:
 	/** Deletes a set of attached objects from a FPreviewAssetAttachContainer and notifies Persona*/
 	void DeleteAttachedObjects( FPreviewAssetAttachContainer& AttachedAssets );
 
-	/** Submenu creator handler for the given skeleton */
-	static void CreateMenuForBoneReduction(FMenuBuilder& MenuBuilder, SSkeletonTree * Widget, USkeleton * Skeleton, bool bAdd);
-
 private:
 	/** Pointer back to the kismet 2 tool that owns us */
 	TWeakPtr<FPersona> PersonaPtr;
diff --git a/Engine/Source/Runtime/Engine/Classes/Engine/SkeletalMesh.h b/Engine/Source/Runtime/Engine/Classes/Engine/SkeletalMesh.h
--- a/Engine/Source/Runtime/Engine/Classes/Engine/SkeletalMesh.h
+++ b/Engine/Source/Runtime/Engine/Classes/Engine/SkeletalMesh.h
@@ -445,6 +445,28 @@ struct FSkeletalMeshOptimizationSettings
 	/** Equality operator. */
 	bool operator==(const FSkeletalMeshOptimizationSettings& Other) const
 	{
+		// first, check whether bones to remove are same or not
+		const TArray<FBoneReference>& TempBones1 = BonesToRemove.Num() > Other.BonesToRemove.Num() ? BonesToRemove : Other.BonesToRemove;
+		const TArray<FBoneReference>& TempBones2 = BonesToRemove.Num() > Other.BonesToRemove.Num() ? Other.BonesToRemove : BonesToRemove;
+
+		for (int32 Index = 0; Index < TempBones2.Num(); Index++)
+		{
+			if (TempBones1[Index].BoneName != TempBones2[Index].BoneName)
+			{
+				return false;
+			}
+		}
+
+		// check remained bones 
+		for (int32 Index = TempBones2.Num(); Index < TempBones1.Num(); Index++)
+		{
+			// if it has an actual bone name, these are not the same
+			if (TempBones1[Index].BoneName != FName("None"))
+			{
+				return false;
+			}
+		}
+
 		return ReductionMethod == Other.ReductionMethod
 			&& NumOfTrianglesPercentage == Other.NumOfTrianglesPercentage
 			&& MaxDeviationPercentage == Other.MaxDeviationPercentage
@@ -456,8 +478,7 @@ struct FSkeletalMeshOptimizationSettings
 			&& SkinningImportance == Other.SkinningImportance
 			&& bRecalcNormals == Other.bRecalcNormals
 			&& BoneReductionRatio == Other.BoneReductionRatio
-			&& MaxBonesPerVertex == Other.MaxBonesPerVertex
-			&& BonesToRemove == Other.BonesToRemove;
+			&& MaxBonesPerVertex == Other.MaxBonesPerVertex;
 	}
 
 	/** Inequality. */
