Only show retargetting in Persona's skeleton mode

diff --git a/Engine/Source/Editor/Persona/Private/SSkeletonTree.cpp b/Engine/Source/Editor/Persona/Private/SSkeletonTree.cpp
--- a/Engine/Source/Editor/Persona/Private/SSkeletonTree.cpp
+++ b/Engine/Source/Editor/Persona/Private/SSkeletonTree.cpp
@@ -799,6 +799,20 @@ void SSkeletonTree::Construct(const FArguments& InArgs)
 	FSkeletonTreeCommands::Register();
 	BindCommands();
 
+	TSharedRef<SHeaderRow> TreeHeaderRow = SNew(SHeaderRow)
+		+ SHeaderRow::Column(ColumnID_BoneLabel)
+		.DefaultLabel(LOCTEXT("SkeletonBoneNameLabel", "Name"))
+		.FillWidth(0.75f);
+
+	if (PersonaPtr.Pin()->IsModeCurrent(FPersonaModes::SkeletonDisplayMode))
+	{
+		TreeHeaderRow->AddColumn(
+			SHeaderRow::Column(ColumnID_RetargetingLabel)
+			.DefaultLabel(LOCTEXT("SkeletonBoneTranslationRetargetingLabel", "Translation Retargeting"))
+			.FillWidth(0.25f)
+			);
+	}
+
 	this->ChildSlot
 	[
 		SNew( SVerticalBox )
@@ -865,15 +879,7 @@ void SSkeletonTree::Construct(const FArguments& InArgs)
 			.ItemHeight( 24 )
 			.HeaderRow
 			(
-				SNew(SHeaderRow)
-
-				+ SHeaderRow::Column( ColumnID_BoneLabel )
-				.DefaultLabel( LOCTEXT( "SkeletonBoneNameLabel", "Name" ) )
-				.FillWidth( 0.75f )
-
-				+ SHeaderRow::Column( ColumnID_RetargetingLabel )
-				.DefaultLabel( LOCTEXT( "SkeletonBoneTranslationRetargetingLabel", "Translation Retargeting" ) )
-				.FillWidth( 0.25f )
+				TreeHeaderRow
 			)
 		]
 	];
