Added doc anchors for each animation editor type

diff --git a/Engine/Source/Editor/Persona/Private/Persona.cpp b/Engine/Source/Editor/Persona/Private/Persona.cpp
--- a/Engine/Source/Editor/Persona/Private/Persona.cpp
+++ b/Engine/Source/Editor/Persona/Private/Persona.cpp
@@ -479,7 +479,7 @@ FPersona::~FPersona()
 	// NOTE: Any tabs that we still have hanging out when destroyed will be cleaned up by FBaseToolkit's destructor
 }
 
-TSharedPtr<SWidget> FPersona::CreateEditorWidgetForAnimDocument(UObject* InAnimAsset)
+TSharedPtr<SWidget> FPersona::CreateEditorWidgetForAnimDocument(UObject* InAnimAsset, FString& DocumentLink)
 {
 	TSharedPtr<SWidget> Result;
 	if (InAnimAsset)
@@ -489,30 +489,55 @@ TSharedPtr<SWidget> FPersona::CreateEditorWidgetForAnimDocument(UObject* InAnimA
 			Result = SNew(SSequenceEditor)
 				.Persona(SharedThis(this))
 				.Sequence(Sequence);
+
+			DocumentLink = TEXT("Engine/Animation/Sequences");
 		}
 		else if (UAnimComposite* Composite = Cast<UAnimComposite>(InAnimAsset))
 		{
 			Result = SNew(SAnimCompositeEditor)
 				.Persona(SharedThis(this))
 				.Composite(Composite);
+
+			DocumentLink = TEXT("Engine/Animation/AnimationComposite");
 		}
 		else if (UAnimMontage* Montage = Cast<UAnimMontage>(InAnimAsset))
 		{
 			Result = SNew(SMontageEditor)
 				.Persona(SharedThis(this))
 				.Montage(Montage);
+
+			DocumentLink = TEXT("Engine/Animation/AnimMontage");
 		}
 		else if (UBlendSpace* BlendSpace = Cast<UBlendSpace>(InAnimAsset))
 		{
 			Result = SNew(SBlendSpaceEditor)
 				.Persona(SharedThis(this))
 				.BlendSpace(BlendSpace);
+
+			if (Cast<UAimOffsetBlendSpace>(InAnimAsset))
+			{
+				DocumentLink = TEXT("Engine/Animation/AimOffset");
+			}
+			else
+			{
+				DocumentLink = TEXT("Engine/Animation/Blendspaces");
+			}
 		}
 		else if (UBlendSpace1D* BlendSpace1D = Cast<UBlendSpace1D>(InAnimAsset))
 		{
 			Result = SNew(SBlendSpaceEditor1D)
 				.Persona(SharedThis(this))
 				.BlendSpace1D(BlendSpace1D);
+
+
+			if (Cast<UAimOffsetBlendSpace1D>(InAnimAsset))
+			{
+				DocumentLink = TEXT("Engine/Animation/AimOffset");
+			}
+			else
+			{
+				DocumentLink = TEXT("Engine/Animation/Blendspaces");
+			}
 		}
 	}
 
@@ -552,7 +577,8 @@ TSharedPtr<SDockTab> FPersona::OpenNewAnimationDocumentTab(UObject* InAnimAsset)
 
 	if (InAnimAsset != NULL)
 	{
-		TSharedPtr<SWidget> TabContents = CreateEditorWidgetForAnimDocument(InAnimAsset);
+		FString	DocumentLink;
+		TSharedPtr<SWidget> TabContents = CreateEditorWidgetForAnimDocument(InAnimAsset, DocumentLink);
 		if (TabContents.IsValid())
 		{
 			if ( SharedAnimAssetBeingEdited.IsValid() )
@@ -574,6 +600,7 @@ TSharedPtr<SDockTab> FPersona::OpenNewAnimationDocumentTab(UObject* InAnimAsset)
 				OpenedTab->SetContent(TabContents.ToSharedRef());
 				OpenedTab->ActivateInParent(ETabActivationCause::SetDirectly);
 				OpenedTab->SetLabel(NameAttribute);
+				OpenedTab->SetLeftContent(IDocumentation::Get()->CreateAnchor(DocumentLink));
 			}
 			else
 			{
@@ -585,6 +612,8 @@ TSharedPtr<SDockTab> FPersona::OpenNewAnimationDocumentTab(UObject* InAnimAsset)
 					[
 						TabContents.ToSharedRef()
 					];
+
+				OpenedTab->SetLeftContent(IDocumentation::Get()->CreateAnchor(DocumentLink));
 			
 				TabManager->InsertNewDocumentTab("Document", FTabManager::ESearchPreference::RequireClosedTab, OpenedTab.ToSharedRef());
 
diff --git a/Engine/Source/Editor/Persona/Private/Persona.h b/Engine/Source/Editor/Persona/Private/Persona.h
--- a/Engine/Source/Editor/Persona/Private/Persona.h
+++ b/Engine/Source/Editor/Persona/Private/Persona.h
@@ -265,8 +265,8 @@ protected:
 	// 
 	TSharedPtr<SDockTab> OpenNewAnimationDocumentTab(UObject* InAnimAsset);
 
-	// Creates an editor widget for a specified animation document
-	TSharedPtr<SWidget> CreateEditorWidgetForAnimDocument(UObject* InAnimAsset);
+	// Creates an editor widget for a specified animation document and returns the document link 
+	TSharedPtr<SWidget> CreateEditorWidgetForAnimDocument(UObject* InAnimAsset, FString& OutDocumentLink);
 
 	/** Callback when an object has been reimported, and whether it worked */
 	void OnPostReimport(UObject* InObject, bool bSuccess);
diff --git a/Engine/Source/Editor/Persona/Private/SAnimCompositeEditor.cpp b/Engine/Source/Editor/Persona/Private/SAnimCompositeEditor.cpp
--- a/Engine/Source/Editor/Persona/Private/SAnimCompositeEditor.cpp
+++ b/Engine/Source/Editor/Persona/Private/SAnimCompositeEditor.cpp
@@ -15,6 +15,11 @@
 //////////////////////////////////////////////////////////////////////////
 // SAnimCompositeEditor
 
+TSharedRef<SWidget> SAnimCompositeEditor::CreateDocumentAnchor()
+{
+	return IDocumentation::Get()->CreateAnchor(TEXT("Engine/Animation/AnimationComposite"));
+}
+
 void SAnimCompositeEditor::Construct(const FArguments& InArgs)
 {
 	bRebuildPanel = false;
diff --git a/Engine/Source/Editor/Persona/Private/SAnimCompositeEditor.h b/Engine/Source/Editor/Persona/Private/SAnimCompositeEditor.h
--- a/Engine/Source/Editor/Persona/Private/SAnimCompositeEditor.h
+++ b/Engine/Source/Editor/Persona/Private/SAnimCompositeEditor.h
@@ -78,4 +78,8 @@ public:
 	// SWidget interface start
 	virtual void Tick( const FGeometry& AllottedGeometry, const double InCurrentTime, const float InDeltaTime ) override;
 	// SWidget interface end
+
+	// Begin SAnimEditorBase interface
+	virtual TSharedRef<SWidget> CreateDocumentAnchor() override;
+	// End SAnimEditorBase interface
 };
diff --git a/Engine/Source/Editor/Persona/Private/SAnimEditorBase.cpp b/Engine/Source/Editor/Persona/Private/SAnimEditorBase.cpp
--- a/Engine/Source/Editor/Persona/Private/SAnimEditorBase.cpp
+++ b/Engine/Source/Editor/Persona/Private/SAnimEditorBase.cpp
@@ -17,6 +17,11 @@
 //////////////////////////////////////////////////////////////////////////
 // SAnimEditorBase
 
+TSharedRef<SWidget> SAnimEditorBase::CreateDocumentAnchor()
+{
+	return IDocumentation::Get()->CreateAnchor(TEXT("Engine/Animation/Sequences"));
+}
+
 void SAnimEditorBase::Construct(const FArguments& InArgs)
 {
 	PersonaPtr = InArgs._Persona;
@@ -44,6 +49,12 @@ void SAnimEditorBase::Construct(const FArguments& InArgs)
 					.ColorAndOpacity( FLinearColor(1,1,1,0.5) )
 					.Text( this, &SAnimEditorBase::GetEditorObjectName )
 				]
+				+ SHorizontalBox::Slot()
+					.HAlign(HAlign_Left)
+					.VAlign(VAlign_Center)
+					[
+						CreateDocumentAnchor()
+					]
 			]
 		]
 
diff --git a/Engine/Source/Editor/Persona/Private/SAnimEditorBase.h b/Engine/Source/Editor/Persona/Private/SAnimEditorBase.h
--- a/Engine/Source/Editor/Persona/Private/SAnimEditorBase.h
+++ b/Engine/Source/Editor/Persona/Private/SAnimEditorBase.h
@@ -65,6 +65,9 @@ public:
 	virtual void AddReferencedObjects( FReferenceCollector& Collector ) override;
 	// FGCObject interface end
 
+	/** Creates a anchor widget for each animation editor type */
+	virtual TSharedRef<SWidget> CreateDocumentAnchor();
+
 protected:
 	/** Persona reference **/
 	TWeakPtr<FPersona> PersonaPtr;
diff --git a/Engine/Source/Editor/Persona/Private/SAnimationBlendSpaceBase.cpp b/Engine/Source/Editor/Persona/Private/SAnimationBlendSpaceBase.cpp
--- a/Engine/Source/Editor/Persona/Private/SAnimationBlendSpaceBase.cpp
+++ b/Engine/Source/Editor/Persona/Private/SAnimationBlendSpaceBase.cpp
@@ -1201,6 +1201,17 @@ void SBlendSpaceEditorBase::PostUndo()
 
 TSharedRef<SWidget> SBlendSpaceEditorBase::MakeEditorHeader() const
 {
+	FString DocumentLink;
+
+	if (Cast<UAimOffsetBlendSpace>(BlendSpace) || Cast<UAimOffsetBlendSpace1D>(BlendSpace))
+	{
+		DocumentLink = TEXT("Engine/Animation/AimOffset");
+	}
+	else
+	{
+		DocumentLink = TEXT("Engine/Animation/Blendspaces");
+	}
+	
 	return SNew(SBorder)
 	. BorderImage( FEditorStyle::GetBrush( TEXT("Graph.TitleBackground") ) )
 	. HAlign(HAlign_Center)
@@ -1215,6 +1226,12 @@ TSharedRef<SWidget> SBlendSpaceEditorBase::MakeEditorHeader() const
 			.ColorAndOpacity( FLinearColor(1,1,1,0.5) )
 			.Text( this, &SBlendSpaceEditorBase::GetBlendSpaceDisplayName )
 		]
+		+ SHorizontalBox::Slot()
+		.HAlign(HAlign_Left)
+		.VAlign(VAlign_Center)
+		[
+			IDocumentation::Get()->CreateAnchor(DocumentLink)
+		]
 	];
 }
 
diff --git a/Engine/Source/Editor/Persona/Private/SMontageEditor.cpp b/Engine/Source/Editor/Persona/Private/SMontageEditor.cpp
--- a/Engine/Source/Editor/Persona/Private/SMontageEditor.cpp
+++ b/Engine/Source/Editor/Persona/Private/SMontageEditor.cpp
@@ -20,6 +20,11 @@
 //////////////////////////////////////////////////////////////////////////
 // SMontageEditor
 
+TSharedRef<SWidget> SMontageEditor::CreateDocumentAnchor()
+{
+	return IDocumentation::Get()->CreateAnchor(TEXT("Engine/Animation/AnimMontage"));
+}
+
 void SMontageEditor::Construct(const FArguments& InArgs)
 {
 	MontageObj = InArgs._Montage;
diff --git a/Engine/Source/Editor/Persona/Private/SMontageEditor.h b/Engine/Source/Editor/Persona/Private/SMontageEditor.h
--- a/Engine/Source/Editor/Persona/Private/SMontageEditor.h
+++ b/Engine/Source/Editor/Persona/Private/SMontageEditor.h
@@ -142,4 +142,8 @@ public:
 	// SWidget interface start
 	virtual void Tick( const FGeometry& AllottedGeometry, const double InCurrentTime, const float InDeltaTime ) override;
 	// SWidget interface end
+
+	// Begin SAnimEditorBase interface
+	virtual TSharedRef<SWidget> CreateDocumentAnchor() override;
+	// End SAnimEditorBase interface
 };
