Shift click in skeleton tree performs recursive expansion/collapse + Removed retargetting context menu options from non skeleton modes

diff --git a/Engine/Source/Editor/Persona/Private/SSkeletonTree.cpp b/Engine/Source/Editor/Persona/Private/SSkeletonTree.cpp
--- a/Engine/Source/Editor/Persona/Private/SSkeletonTree.cpp
+++ b/Engine/Source/Editor/Persona/Private/SSkeletonTree.cpp
@@ -804,7 +804,7 @@ void SSkeletonTree::Construct(const FArguments& InArgs)
 		.DefaultLabel(LOCTEXT("SkeletonBoneNameLabel", "Name"))
 		.FillWidth(0.75f);
 
-	if (PersonaPtr.Pin()->IsModeCurrent(FPersonaModes::SkeletonDisplayMode))
+	if (IsInSkeletonMode())
 	{
 		TreeHeaderRow->AddColumn(
 			SHeaderRow::Column(ColumnID_RetargetingLabel)
@@ -876,6 +876,7 @@ void SSkeletonTree::Construct(const FArguments& InArgs)
 			.OnSelectionChanged( this, &SSkeletonTree::OnSelectionChanged )
 			.OnItemScrolledIntoView( this, &SSkeletonTree::OnItemScrolledIntoView )
 			.OnMouseButtonDoubleClick( this, &SSkeletonTree::OnTreeDoubleClick)
+			.OnSetExpansionRecursive(this, &SSkeletonTree::SetTreeItemExpansionRecursive)
 			.ItemHeight( 24 )
 			.HeaderRow
 			(
@@ -1364,34 +1365,37 @@ TSharedPtr< SWidget > SSkeletonTree::CreateContextMenu()
 
 			MenuBuilder.EndSection();
 
-			MenuBuilder.BeginSection("SkeletonTreeBoneTranslationRetargeting", LOCTEXT( "BoneTranslationRetargetingHeader", "Bone Translation Retargeting" ) );
+			if(IsInSkeletonMode())
 			{
-				FUIAction RecursiveRetargetingSkeletonAction = FUIAction( FExecuteAction::CreateSP( this, &SSkeletonTree::SetBoneTranslationRetargetingModeRecursive, EBoneTranslationRetargetingMode::Skeleton));
-				FUIAction RecursiveRetargetingAnimationAction = FUIAction( FExecuteAction::CreateSP( this, &SSkeletonTree::SetBoneTranslationRetargetingModeRecursive, EBoneTranslationRetargetingMode::Animation));
-				FUIAction RecursiveRetargetingAnimationScaledAction = FUIAction( FExecuteAction::CreateSP( this, &SSkeletonTree::SetBoneTranslationRetargetingModeRecursive, EBoneTranslationRetargetingMode::AnimationScaled));
-
-				MenuBuilder.AddMenuEntry
-					( LOCTEXT( "SetTranslationRetargetingSkeletonChildrenAction", "Recursively Set Translation Retargeting Skeleton" )
-					, LOCTEXT( "BoneTranslationRetargetingSkeletonToolTip", "Use translation from Skeleton." )
-					, FSlateIcon()
-					, RecursiveRetargetingSkeletonAction 
-					);
-
-				MenuBuilder.AddMenuEntry
-					( LOCTEXT( "SetTranslationRetargetingAnimationChildrenAction", "Recursively Set Translation Retargeting Animation" )
-					, LOCTEXT( "BoneTranslationRetargetingAnimationToolTip", "Use translation from animation." )
-					, FSlateIcon()
-					, RecursiveRetargetingAnimationAction 
-					);
-
-				MenuBuilder.AddMenuEntry
-					( LOCTEXT( "SetTranslationRetargetingAnimationScaledChildrenAction", "Recursively Set Translation Retargeting AnimationScaled" )
-					, LOCTEXT( "BoneTranslationRetargetingAnimationScaledToolTip", "Use translation from animation, scale length by Skeleton's proportions." )
-					, FSlateIcon()
-					, RecursiveRetargetingAnimationScaledAction
-					);
+				MenuBuilder.BeginSection("SkeletonTreeBoneTranslationRetargeting", LOCTEXT("BoneTranslationRetargetingHeader", "Bone Translation Retargeting"));
+				{
+					FUIAction RecursiveRetargetingSkeletonAction = FUIAction(FExecuteAction::CreateSP(this, &SSkeletonTree::SetBoneTranslationRetargetingModeRecursive, EBoneTranslationRetargetingMode::Skeleton));
+					FUIAction RecursiveRetargetingAnimationAction = FUIAction(FExecuteAction::CreateSP(this, &SSkeletonTree::SetBoneTranslationRetargetingModeRecursive, EBoneTranslationRetargetingMode::Animation));
+					FUIAction RecursiveRetargetingAnimationScaledAction = FUIAction(FExecuteAction::CreateSP(this, &SSkeletonTree::SetBoneTranslationRetargetingModeRecursive, EBoneTranslationRetargetingMode::AnimationScaled));
+
+					MenuBuilder.AddMenuEntry
+						(LOCTEXT("SetTranslationRetargetingSkeletonChildrenAction", "Recursively Set Translation Retargeting Skeleton")
+						, LOCTEXT("BoneTranslationRetargetingSkeletonToolTip", "Use translation from Skeleton.")
+						, FSlateIcon()
+						, RecursiveRetargetingSkeletonAction
+						);
+
+					MenuBuilder.AddMenuEntry
+						(LOCTEXT("SetTranslationRetargetingAnimationChildrenAction", "Recursively Set Translation Retargeting Animation")
+						, LOCTEXT("BoneTranslationRetargetingAnimationToolTip", "Use translation from animation.")
+						, FSlateIcon()
+						, RecursiveRetargetingAnimationAction
+						);
+
+					MenuBuilder.AddMenuEntry
+						(LOCTEXT("SetTranslationRetargetingAnimationScaledChildrenAction", "Recursively Set Translation Retargeting AnimationScaled")
+						, LOCTEXT("BoneTranslationRetargetingAnimationScaledToolTip", "Use translation from animation, scale length by Skeleton's proportions.")
+						, FSlateIcon()
+						, RecursiveRetargetingAnimationScaledAction
+						);
+				}
+				MenuBuilder.EndSection();
 			}
-			MenuBuilder.EndSection();
 		}
 
 		if(BoneTreeSelection.HasSelectedOfType(ESkeletonTreeRowType::Socket))
@@ -2043,6 +2047,17 @@ void SSkeletonTree::OnTreeDoubleClick( FDisplayedTreeRowInfoPtr InItem )
 	InItem->OnItemDoubleClicked();
 }
 
+void SSkeletonTree::SetTreeItemExpansionRecursive(TSharedPtr< FDisplayedTreeRowInfo > TreeItem, bool bInExpansionState) const
+{
+	SkeletonTreeView->SetItemExpansion(TreeItem, bInExpansionState);
+
+	// Recursively go through the children.
+	for (auto It = TreeItem->Children.CreateIterator(); It; ++It)
+	{
+		SetTreeItemExpansionRecursive(*It, bInExpansionState);
+	}
+}
+
 void SSkeletonTree::PostUndo()
 {
 	// Rebuild the tree view whenever we undo a change to the skeleton
@@ -2465,5 +2480,10 @@ void SSkeletonTree::AddAttachedAssets( const FPreviewAssetAttachContainer& Attac
 	}
 }
 
+bool SSkeletonTree::IsInSkeletonMode() const
+{
+	return PersonaPtr.Pin()->IsModeCurrent(FPersonaModes::SkeletonDisplayMode);
+}
+
 
 #undef LOCTEXT_NAMESPACE
diff --git a/Engine/Source/Editor/Persona/Private/SSkeletonTree.h b/Engine/Source/Editor/Persona/Private/SSkeletonTree.h
--- a/Engine/Source/Editor/Persona/Private/SSkeletonTree.h
+++ b/Engine/Source/Editor/Persona/Private/SSkeletonTree.h
@@ -401,6 +401,9 @@ public:
 	/** Callback for when the user double-clicks on an item in the tree */
 	void OnTreeDoubleClick( FDisplayedTreeRowInfoPtr InItem );
 
+	/** Handle recursive expansion/contraction of the tree */
+	void SetTreeItemExpansionRecursive(TSharedPtr< FDisplayedTreeRowInfo > TreeItem, bool bInExpansionState) const;
+
 	/** Called when a socket has been renamed */
 	void RenameSocketAttachments(FName& OldSocketName, FName& NewSocketName);
 
@@ -542,6 +545,9 @@ private:
 	/** Deletes a set of attached objects from a FPreviewAssetAttachContainer and notifies Persona*/
 	void DeleteAttachedObjects( FPreviewAssetAttachContainer& AttachedAssets );
 
+	// Is Persona currently in Skeleton Mode
+	bool IsInSkeletonMode() const;
+
 private:
 	/** Pointer back to the kismet 2 tool that owns us */
 	TWeakPtr<FPersona> PersonaPtr;
