Stopped adding a skeletal socket from dirtying skeletal meshes

diff --git a/Engine/Source/Editor/Persona/Private/SSkeletonTree.cpp b/Engine/Source/Editor/Persona/Private/SSkeletonTree.cpp
--- a/Engine/Source/Editor/Persona/Private/SSkeletonTree.cpp
+++ b/Engine/Source/Editor/Persona/Private/SSkeletonTree.cpp
@@ -2306,13 +2306,18 @@ void SSkeletonTree::RenameSocketAttachments(FName& OldSocketName, FName& NewSock
 {
 	const FScopedTransaction Transaction( LOCTEXT( "RenameSocketAttachments", "Rename Socket Attachments" ) );
 
-	TargetSkeleton->Modify();
-
+	bool bSkeletonModified = false;
 	for(int AttachedObjectIndex = 0; AttachedObjectIndex < TargetSkeleton->PreviewAttachedAssetContainer.Num(); ++AttachedObjectIndex)
 	{
 		FPreviewAttachedObjectPair& Pair = TargetSkeleton->PreviewAttachedAssetContainer[AttachedObjectIndex];
 		if(Pair.AttachedTo == OldSocketName)
 		{
+			// Only modify the skeleton if we actually intend to change something.
+			if(!bSkeletonModified)
+			{
+				TargetSkeleton->Modify();
+				bSkeletonModified = true;
+			}
 			Pair.AttachedTo = NewSocketName;
 		}
 		PersonaPtr.Pin()->RemoveAttachedObjectFromPreviewComponent(Pair.GetAttachedObject(), OldSocketName);
@@ -2325,13 +2330,19 @@ void SSkeletonTree::RenameSocketAttachments(FName& OldSocketName, FName& NewSock
 
 		if ( Mesh )
 		{
-			Mesh->Modify();
-
+			bool bMeshModified = false;
 			for(int AttachedObjectIndex = 0; AttachedObjectIndex < Mesh->PreviewAttachedAssetContainer.Num(); ++AttachedObjectIndex)
 			{
 				FPreviewAttachedObjectPair& Pair = Mesh->PreviewAttachedAssetContainer[AttachedObjectIndex];
 				if(Pair.AttachedTo == OldSocketName)
 				{
+					// Only modify the mesh if we actually intend to change something. Avoids dirtying
+					// meshes when we don't actually update any data on them. (such as adding a new socket)
+					if(!bMeshModified)
+					{
+						Mesh->Modify();
+						bMeshModified = true;
+					}
 					Pair.AttachedTo = NewSocketName;
 				}
 				PersonaPtr.Pin()->RemoveAttachedObjectFromPreviewComponent(Pair.GetAttachedObject(), OldSocketName);
