Removed animation viewport scrub panel from animation mode in Persona

diff --git a/Engine/Source/Editor/Persona/Private/SAnimationEditorViewport.cpp b/Engine/Source/Editor/Persona/Private/SAnimationEditorViewport.cpp
--- a/Engine/Source/Editor/Persona/Private/SAnimationEditorViewport.cpp
+++ b/Engine/Source/Editor/Persona/Private/SAnimationEditorViewport.cpp
@@ -292,8 +292,11 @@ void SAnimationEditorViewportTabBody::Construct(const FArguments& InArgs)
 	bPreviewLockModeOn = 0;
 	LODSelection = LOD_Auto;
 
+	TSharedPtr<FPersona> SharedPersona = PersonaPtr.Pin();
+	check(SharedPersona.IsValid());
+
 	// register delegate for anim change notification
-	PersonaPtr.Pin()->RegisterOnAnimChanged(FPersona::FOnAnimChanged::CreateSP(this, &SAnimationEditorViewportTabBody::AnimChanged));
+	SharedPersona->RegisterOnAnimChanged(FPersona::FOnAnimChanged::CreateSP(this, &SAnimationEditorViewportTabBody::AnimChanged));
 
 	const FSlateFontInfo SmallLayoutFont( FPaths::EngineContentDir() / TEXT("Slate/Fonts/Roboto-Regular.ttf"), 10 );
 
@@ -309,9 +312,10 @@ void SAnimationEditorViewportTabBody::Construct(const FArguments& InArgs)
 
 	ViewportWidget = SNew(SAnimationEditorViewport, InArgs._Persona, SharedThis(this));
 
+	TSharedPtr<SVerticalBox> ViewportContainer = nullptr;
 	this->ChildSlot
 	[
-		SNew(SVerticalBox)
+		SAssignNew(ViewportContainer, SVerticalBox)
 
 		// Build our toolbar level toolbar
 		+SVerticalBox::Slot()
@@ -340,7 +344,11 @@ void SAnimationEditorViewportTabBody::Construct(const FArguments& InArgs)
 				.OnClicked(this, &SAnimationEditorViewportTabBody::ClickedOnViewportCornerText)
 			]
 		]
-		+SVerticalBox::Slot()
+	];
+
+	if(!SharedPersona->IsModeCurrent(FPersonaModes::AnimationEditMode) && ViewportContainer.IsValid())
+	{
+		ViewportContainer->AddSlot()
 		.AutoHeight()
 		[
 			SAssignNew(ScrubPanelContainer, SVerticalBox)
@@ -353,10 +361,10 @@ void SAnimationEditorViewportTabBody::Construct(const FArguments& InArgs)
 				.ViewInputMax(this, &SAnimationEditorViewportTabBody::GetViewMaxInput)
 				.bAllowZoom(true)
 			]
-		]
-	];
+		];
 
-	UpdateScrubPanel(Cast<UAnimationAsset>(PersonaPtr.Pin()->GetPreviewAnimationAsset()));
+		UpdateScrubPanel(Cast<UAnimationAsset>(SharedPersona->GetPreviewAnimationAsset()));
+	}
 
 	LevelViewportClient = ViewportWidget->GetViewportClient();
 	LevelViewportClient->VisibilityDelegate.BindSP(this, &SAnimationEditorViewportTabBody::IsVisible);
@@ -1121,32 +1129,36 @@ void SAnimationEditorViewportTabBody::PopulateUVChoices()
 
 void SAnimationEditorViewportTabBody::UpdateScrubPanel(UAnimationAsset* AnimAsset)
 {
-	ScrubPanelContainer->ClearChildren();
-	bool bUseDefaultScrubPanel = true;
-	if (UAnimMontage* Montage = Cast<UAnimMontage>(AnimAsset))
+	// We might not have a scrub panel if we're in animation mode.
+	if(ScrubPanelContainer.IsValid())
 	{
-		ScrubPanelContainer->AddSlot()
-		.AutoHeight()
-		[
-			SNew(SAnimMontageScrubPanel)
-			.Persona(PersonaPtr)
-			.ViewInputMin(this, &SAnimationEditorViewportTabBody::GetViewMinInput)
-			.ViewInputMax(this, &SAnimationEditorViewportTabBody::GetViewMaxInput)
-			.bAllowZoom(true)
-		];
-		bUseDefaultScrubPanel = false;
-	}
-	if (bUseDefaultScrubPanel)
-	{
-		ScrubPanelContainer->AddSlot()
-		.AutoHeight()
-		[
-			SNew(SAnimationScrubPanel)
-			.Persona(PersonaPtr)
-			.ViewInputMin(this, &SAnimationEditorViewportTabBody::GetViewMinInput)
-			.ViewInputMax(this, &SAnimationEditorViewportTabBody::GetViewMaxInput)
-			.bAllowZoom(true)
-		];
+		ScrubPanelContainer->ClearChildren();
+		bool bUseDefaultScrubPanel = true;
+		if(UAnimMontage* Montage = Cast<UAnimMontage>(AnimAsset))
+		{
+			ScrubPanelContainer->AddSlot()
+				.AutoHeight()
+				[
+					SNew(SAnimMontageScrubPanel)
+					.Persona(PersonaPtr)
+					.ViewInputMin(this, &SAnimationEditorViewportTabBody::GetViewMinInput)
+					.ViewInputMax(this, &SAnimationEditorViewportTabBody::GetViewMaxInput)
+					.bAllowZoom(true)
+				];
+			bUseDefaultScrubPanel = false;
+		}
+		if(bUseDefaultScrubPanel)
+		{
+			ScrubPanelContainer->AddSlot()
+				.AutoHeight()
+				[
+					SNew(SAnimationScrubPanel)
+					.Persona(PersonaPtr)
+					.ViewInputMin(this, &SAnimationEditorViewportTabBody::GetViewMinInput)
+					.ViewInputMax(this, &SAnimationEditorViewportTabBody::GetViewMaxInput)
+					.bAllowZoom(true)
+				];
+		}
 	}
 }
 
