Changed animation layout to be simpler and have a distinction for each mode.
- Made sure animation editor is big enough in the default mode.

diff --git a/Engine/Source/Editor/Persona/Private/AnimationMode/AnimationMode.cpp b/Engine/Source/Editor/Persona/Private/AnimationMode/AnimationMode.cpp
--- a/Engine/Source/Editor/Persona/Private/AnimationMode/AnimationMode.cpp
+++ b/Engine/Source/Editor/Persona/Private/AnimationMode/AnimationMode.cpp
@@ -142,7 +142,7 @@ FAnimEditAppMode::FAnimEditAppMode(TSharedPtr<FPersona> InPersona)
 	PersonaTabFactories.RegisterFactory(MakeShareable(new FSelectionDetailsSummoner(InPersona)));
 	PersonaTabFactories.RegisterFactory(MakeShareable(new FAnimAssetPropertiesSummoner(InPersona)));
 
-	TabLayout = FTabManager::NewLayout("Persona_AnimEditMode_Layout_v6")
+	TabLayout = FTabManager::NewLayout("Persona_AnimEditMode_Layout_v7")
 		->AddArea
 		(
 			FTabManager::NewPrimaryArea() ->SetOrientation(Orient_Vertical)
@@ -164,6 +164,7 @@ FAnimEditAppMode::FAnimEditAppMode(TSharedPtr<FPersona> InPersona)
 					// Left 1/3rd - Skeleton and Anim properties
 					FTabManager::NewSplitter()
 					->SetOrientation(Orient_Vertical)
+					->SetSizeCoefficient(0.3f)
 					->Split
 					(
 						FTabManager::NewStack()
@@ -180,17 +181,16 @@ FAnimEditAppMode::FAnimEditAppMode(TSharedPtr<FPersona> InPersona)
 					// Middle 1/3rd - Viewport and anim document area
 					FTabManager::NewSplitter()
 					->SetOrientation(Orient_Vertical)
+					->SetSizeCoefficient(0.4f)
 					->Split
 					(
 						FTabManager::NewStack()
-						->SetSizeCoefficient(0.75f)
 						->SetHideTabWell(true)
 						->AddTab( FPersonaTabs::PreviewViewportID, ETabState::OpenedTab )
 					)
 					->Split
 					(
 						FTabManager::NewStack()
-						->SetSizeCoefficient(0.25f)
 						->AddTab( "Document", ETabState::ClosedTab )
 					)
 				)
@@ -199,6 +199,7 @@ FAnimEditAppMode::FAnimEditAppMode(TSharedPtr<FPersona> InPersona)
 					// Right 1/3rd - Details panel and quick browser
 					FTabManager::NewSplitter()
 					->SetOrientation(Orient_Vertical)
+					->SetSizeCoefficient(0.3f)
 					->Split
 					(
 						FTabManager::NewStack()
diff --git a/Engine/Source/Editor/Persona/Private/MeshMode/MeshMode.cpp b/Engine/Source/Editor/Persona/Private/MeshMode/MeshMode.cpp
--- a/Engine/Source/Editor/Persona/Private/MeshMode/MeshMode.cpp
+++ b/Engine/Source/Editor/Persona/Private/MeshMode/MeshMode.cpp
@@ -74,7 +74,7 @@ FMeshEditAppMode::FMeshEditAppMode(TSharedPtr<FPersona> InPersona)
 	PersonaTabFactories.RegisterFactory(MakeShareable(new FSelectionDetailsSummoner(InPersona)));
 	PersonaTabFactories.RegisterFactory(MakeShareable(new FMeshPropertiesSummoner(InPersona)));
 
-	TabLayout = FTabManager::NewLayout( "Persona_MeshEditMode_Layout_v6" )
+	TabLayout = FTabManager::NewLayout( "Persona_MeshEditMode_Layout_v7" )
 		->AddArea
 		(
 			FTabManager::NewPrimaryArea()
@@ -97,12 +97,7 @@ FMeshEditAppMode::FMeshEditAppMode(TSharedPtr<FPersona> InPersona)
 					// Left 1/3rd - Skeleton tree and mesh panel
 					FTabManager::NewSplitter()
 					->SetOrientation(Orient_Vertical)
-					->SetSizeCoefficient(0.125f)
- 					->Split
- 					(
-						FTabManager::NewStack()
-						->AddTab( FPersonaTabs::SkeletonTreeViewID, ETabState::OpenedTab )
- 					)
+					->SetSizeCoefficient(0.3f)
 					->Split
 					(
 						FTabManager::NewStack()
@@ -114,7 +109,7 @@ FMeshEditAppMode::FMeshEditAppMode(TSharedPtr<FPersona> InPersona)
 					// Middle 1/3rd - Viewport
 					FTabManager::NewSplitter()
 					->SetOrientation(Orient_Vertical)
-					->SetSizeCoefficient(0.750f)
+					->SetSizeCoefficient(0.5f)
 					->Split
 					(
 						FTabManager::NewStack()
@@ -127,13 +122,7 @@ FMeshEditAppMode::FMeshEditAppMode(TSharedPtr<FPersona> InPersona)
 					// Right 1/3rd - Details panel and quick browser
 					FTabManager::NewSplitter()
 					->SetOrientation(Orient_Vertical)
-					->SetSizeCoefficient(0.125f)
-					->Split
-					(
-						FTabManager::NewStack()
-						->AddTab( FBlueprintEditorTabs::DetailsID, ETabState::OpenedTab )
-						->AddTab( FPersonaTabs::AssetBrowserID, ETabState::OpenedTab )
-					)
+					->SetSizeCoefficient(0.2f)
 					->Split
 					(
 						FTabManager::NewStack()
diff --git a/Engine/Source/Editor/Persona/Private/Persona.cpp b/Engine/Source/Editor/Persona/Private/Persona.cpp
--- a/Engine/Source/Editor/Persona/Private/Persona.cpp
+++ b/Engine/Source/Editor/Persona/Private/Persona.cpp
@@ -351,7 +351,7 @@ public:
 	{
 		PersonaTabFactories.RegisterFactory(MakeShareable(new FSelectionDetailsSummoner(InPersona)));
 
-		TabLayout = FTabManager::NewLayout( "Persona_SkeletonEditMode_Layout_v2" )
+		TabLayout = FTabManager::NewLayout( "Persona_SkeletonEditMode_Layout_v3" )
 			->AddArea
 			(
 				FTabManager::NewPrimaryArea() ->SetOrientation(Orient_Vertical)
@@ -372,6 +372,7 @@ public:
 						// Left 1/3rd - Skeleton tree and mesh panel
 						FTabManager::NewSplitter()
 						->SetOrientation(Orient_Vertical)
+						->SetSizeCoefficient(0.3f)
 						->Split
 						(
 							FTabManager::NewStack()
@@ -383,7 +384,7 @@ public:
 						// Middle 1/3rd - Viewport
 						FTabManager::NewSplitter()
 						->SetOrientation(Orient_Vertical)
-						->SetSizeCoefficient(0.75f)
+						->SetSizeCoefficient(0.5f)
 						->Split
 						(
 							FTabManager::NewStack()
@@ -393,19 +394,15 @@ public:
 					)
 					->Split
 					(
-						// Right 1/3rd - Details panel and quick browser
+						// Right 1/3rd - Details panel 
 						FTabManager::NewSplitter()
 						->SetOrientation(Orient_Vertical)
+						->SetSizeCoefficient(0.2f)
 						->Split
 						(
 							FTabManager::NewStack()
 							->AddTab( FBlueprintEditorTabs::DetailsID, ETabState::OpenedTab )	//@TODO: FPersonaTabs::AnimPropertiesID
 						)
-						->Split
-						(
-							FTabManager::NewStack()
-							->AddTab( FPersonaTabs::AssetBrowserID, ETabState::OpenedTab )
-						)
 					)
 				)
 			);
