Changed Customize socket -> Mesh Socket

diff --git a/Engine/Source/Editor/Persona/Private/SSkeletonTree.cpp b/Engine/Source/Editor/Persona/Private/SSkeletonTree.cpp
--- a/Engine/Source/Editor/Persona/Private/SSkeletonTree.cpp
+++ b/Engine/Source/Editor/Persona/Private/SSkeletonTree.cpp
@@ -530,8 +530,8 @@ void FDisplayedSocketInfo::GenerateWidgetForNameColumn( TSharedPtr< SHorizontalB
 	if ( ParentType == ESocketParentType::Mesh )
 	{
 		FText SocketSuffix = IsSocketCustomized() ?
-			LOCTEXT( "CustomizedSuffix", " [Customized]" ) :
-			LOCTEXT( "MeshSuffix", " [Mesh]" );
+			LOCTEXT( "CustomizedSuffix", " [Mesh]" ) :
+			LOCTEXT( "MeshSuffix", " [Mesh Only]" );
 
 		Box->AddSlot()
 		.AutoWidth()
@@ -647,7 +647,7 @@ FString FDisplayedSocketInfo::GetSocketToolTip()
 	}
 	else if ( ParentType == ESocketParentType::Skeleton )
 	{
-		ToolTip = LOCTEXT( "SocketToolTipSkeleton", "This socket is on the skeleton (shared with all meshes that use the skeleton), and the current mesh has customized it" );
+		ToolTip = LOCTEXT( "SocketToolTipSkeleton", "This socket is on the skeleton (shared with all meshes that use the skeleton), and the current mesh has duplciated version of it" );
 	}
 	else
 	{
@@ -987,11 +987,11 @@ void SSkeletonTree::BindCommands()
 		FCanExecuteAction::CreateSP( this, &SSkeletonTree::CanRenameSelected ) );
 
 	CommandList.MapAction(
-		MenuActions.CustomizeSocket,
+		MenuActions.CreateMeshSocket,
 		FExecuteAction::CreateSP( this, &SSkeletonTree::OnCustomizeSocket ) );
 
 	CommandList.MapAction(
-		MenuActions.RemoveSocketCustomization,
+		MenuActions.RemoveMeshSocket,
 		FExecuteAction::CreateSP( this, &SSkeletonTree::OnDeleteSelectedRows ) ); // Removing customization just deletes the mesh socket
 
 	CommandList.MapAction(
@@ -1412,7 +1412,7 @@ TSharedPtr< SWidget > SSkeletonTree::CreateContextMenu()
 
 				if ( DisplayedSocketInfo->IsSocketCustomized() && DisplayedSocketInfo->GetParentType() == ESocketParentType::Mesh )
 				{
-					MenuBuilder.AddMenuEntry( Actions.RemoveSocketCustomization );
+					MenuBuilder.AddMenuEntry( Actions.RemoveMeshSocket );
 				}
 
 				USkeletalMeshSocket* SelectedSocket = static_cast< USkeletalMeshSocket* >( DisplayedSocketInfo->GetData() );
@@ -1424,7 +1424,7 @@ TSharedPtr< SWidget > SSkeletonTree::CreateContextMenu()
 				{
 					if ( DisplayedSocketInfo->GetParentType() == ESocketParentType::Skeleton )
 					{
-						MenuBuilder.AddMenuEntry( Actions.CustomizeSocket );
+						MenuBuilder.AddMenuEntry( Actions.CreateMeshSocket );
 					}
 					else if ( DisplayedSocketInfo->GetParentType() == ESocketParentType::Mesh )
 					{
@@ -1769,7 +1769,7 @@ void SSkeletonTree::OnCustomizeSocket()
 
 			if ( Mesh )
 			{
-				const FScopedTransaction Transaction( LOCTEXT( "CustomizeSocket", "Customize Socket" ) );
+				const FScopedTransaction Transaction( LOCTEXT( "CreateMeshSocket", "Create Mesh Socket" ) );
 				Mesh->Modify();
 
 				USkeletalMeshSocket* NewSocket = ConstructObject<USkeletalMeshSocket>( USkeletalMeshSocket::StaticClass(), Mesh );
diff --git a/Engine/Source/Editor/Persona/Private/SkeletonTreeCommands.cpp b/Engine/Source/Editor/Persona/Private/SkeletonTreeCommands.cpp
--- a/Engine/Source/Editor/Persona/Private/SkeletonTreeCommands.cpp
+++ b/Engine/Source/Editor/Persona/Private/SkeletonTreeCommands.cpp
@@ -16,8 +16,8 @@ void FSkeletonTreeCommands::RegisterCommands()
 	UI_COMMAND( PasteSockets, "Paste Selected Sockets", "Paste selected sockets from clipboard", EUserInterfaceActionType::Button, FInputGesture( EModifierKey::Control, EKeys::V ) );
 
 	UI_COMMAND( AddSocket, "Add Socket", "Add a socket to this bone in the skeleton (disabled when socket filter is set to \"Mesh Sockets\" or \"Sockets Hidden\" mode)", EUserInterfaceActionType::Button, FInputGesture() );
-	UI_COMMAND( CustomizeSocket, "Customize Socket", "Customize this socket for the current mesh (copy it from the skeleton to the mesh)", EUserInterfaceActionType::Button, FInputGesture() );
-	UI_COMMAND( RemoveSocketCustomization, "Remove Socket Customization", "Remove customization of this socket for the current mesh (reverts to the socket in the skeleton)", EUserInterfaceActionType::Button, FInputGesture() );
+	UI_COMMAND( CreateMeshSocket, "Create Mesh Socket", "Duplicate this socket from skeleton to the current mesh and modify the socket data for it", EUserInterfaceActionType::Button, FInputGesture() );
+	UI_COMMAND( RemoveMeshSocket, "Remove Mesh Socket", "Remove duplicated version of this socket for the current mesh (reverts to the socket in the skeleton)", EUserInterfaceActionType::Button, FInputGesture() );
 	UI_COMMAND( PromoteSocketToSkeleton, "Promote Socket To Skeleton", "Makes this socket available for all meshes that use the same skeleton (copies the socket from this mesh to the skeleton)", EUserInterfaceActionType::Button, FInputGesture() );
 
 	UI_COMMAND(DeleteSelectedRows, "Delete", "Delete all selected sockets and attached meshes in the tree", EUserInterfaceActionType::Button, FInputGesture(EKeys::Platform_Delete));
diff --git a/Engine/Source/Editor/Persona/Private/SkeletonTreeCommands.h b/Engine/Source/Editor/Persona/Private/SkeletonTreeCommands.h
--- a/Engine/Source/Editor/Persona/Private/SkeletonTreeCommands.h
+++ b/Engine/Source/Editor/Persona/Private/SkeletonTreeCommands.h
@@ -40,10 +40,10 @@ public:
 	TSharedPtr< FUICommandInfo > AddSocket;
 
 	/** Customize a socket for a mesh */
-	TSharedPtr< FUICommandInfo > CustomizeSocket;
+	TSharedPtr< FUICommandInfo > CreateMeshSocket;
 
 	/** Remove customization for a socket (actually just deletes the mesh socket) */
-	TSharedPtr< FUICommandInfo > RemoveSocketCustomization;
+	TSharedPtr< FUICommandInfo > RemoveMeshSocket;
 
 	/** Promotes a mesh-only socket to the skeleton */
 	TSharedPtr< FUICommandInfo > PromoteSocketToSkeleton;
