Fixed persona issue with showing menu multiple if multiple persona is opened

diff --git a/Engine/Source/Editor/Persona/Private/Persona.cpp b/Engine/Source/Editor/Persona/Private/Persona.cpp
--- a/Engine/Source/Editor/Persona/Private/Persona.cpp
+++ b/Engine/Source/Editor/Persona/Private/Persona.cpp
@@ -860,7 +860,7 @@ void FPersona::InitPersona(const EToolkitMode::Type Mode, const TSharedPtr< clas
 					
 	PersonaModule->GetMenuExtensibilityManager()->AddExtender(PersonaMenuExtender);					
 
-	AddMenuExtender(PersonaModule->GetMenuExtensibilityManager()->GetAllExtenders(GetToolkitCommands(), GetEditingObjects()));
+	AddMenuExtender(PersonaMenuExtender);
 
 	check(PreviewComponent == NULL);
 
