- Changed Display Info in Persona's viewport to be simple as static mesh editor
- Changed the name of "DisplayFactor" property in SkeletalMeshLODInfo into "ScreenSize" and redirected property

diff --git a/Engine/Config/BaseEditorUserSettings.ini b/Engine/Config/BaseEditorUserSettings.ini
--- a/Engine/Config/BaseEditorUserSettings.ini
+++ b/Engine/Config/BaseEditorUserSettings.ini
@@ -465,7 +465,7 @@ GridSize=25
 ViewModeType=2
 ViewportBackgroundColor=(R=0.04,G=0.04,B=0.04)
 ViewFOV=53.43
-bShowMeshStats=true
+bShowMeshStats=false
 
 
 [UnrealEd.UIEditorOptions]
diff --git a/Engine/Config/BaseEngine.ini b/Engine/Config/BaseEngine.ini
--- a/Engine/Config/BaseEngine.ini
+++ b/Engine/Config/BaseEngine.ini
@@ -433,6 +433,9 @@ bUseStreamingPause=false
 +EnumRedirects=(EnumName="ENavLinkDirection",OldEnumEntry="ESmartNavLinkDir::OneWay",NewEnumEntry="ENavLinkDirection::LeftToRight")
 +EnumRedirects=(EnumName="ENavLinkDirection",OldEnumEntry="ESmartNavLinkDir::BothWays",NewEnumEntry="ENavLinkDirection::BothWays")
 
+;SkeletalMesh redirects
++TaggedPropertyRedirects=(ClassName="SkeletalMeshLODInfo",OldPropertyName="DisplayFactor",NewPropertyName="ScreenSize")
+
 ;Behavior Tree editor redirects
 +ActiveClassRedirects=(OldClassName="BehaviorTreeNode",NewClassName="BehaviorTreeGraphNode")
 +ActiveClassRedirects=(OldClassName="BehaviorTreeNode_Composite",NewClassName="BehaviorTreeGraphNode_Composite")
diff --git a/Engine/Source/Developer/SimplygonMeshReduction/Private/SimplygonMeshReduction.cpp b/Engine/Source/Developer/SimplygonMeshReduction/Private/SimplygonMeshReduction.cpp
--- a/Engine/Source/Developer/SimplygonMeshReduction/Private/SimplygonMeshReduction.cpp
+++ b/Engine/Source/Developer/SimplygonMeshReduction/Private/SimplygonMeshReduction.cpp
@@ -283,7 +283,7 @@ public:
 			if( bCalcLODDistance )
 			{
 				float ViewDistance = CalculateViewDistance( MaxDeviation );
-				SkeletalMesh->LODInfo[LODIndex].DisplayFactor = 2.0f * SkeletalMesh->Bounds.SphereRadius / ViewDistance;
+				SkeletalMesh->LODInfo[LODIndex].ScreenSize = 2.0f * SkeletalMesh->Bounds.SphereRadius / ViewDistance;
 			}
 
 			// Flag this LOD as having been simplified.
diff --git a/Engine/Source/Editor/Persona/Private/AnimationEditorViewportClient.cpp b/Engine/Source/Editor/Persona/Private/AnimationEditorViewportClient.cpp
--- a/Engine/Source/Editor/Persona/Private/AnimationEditorViewportClient.cpp
+++ b/Engine/Source/Editor/Persona/Private/AnimationEditorViewportClient.cpp
@@ -765,128 +765,202 @@ void FAnimationViewportClient::DisplayInfo(FCanvas* Canvas, FSceneView* View, bo
 		}
 	}
 
-	if (bDisplayAllInfo && PreviewSkelMeshComp != NULL && PreviewSkelMeshComp->MeshObject != NULL)
+	if (PreviewSkelMeshComp != NULL && PreviewSkelMeshComp->MeshObject != NULL)
 	{
-		FSkeletalMeshResource* SkelMeshResource = PreviewSkelMeshComp->GetSkeletalMeshResource();
-		check(SkelMeshResource);
-
-		// Draw stats about the mesh
-		const FBoxSphereBounds& SkelBounds = PreviewSkelMeshComp->Bounds;
-		const FPlane ScreenPosition = View->Project(SkelBounds.Origin);
+		if (bDisplayAllInfo)
+		{
+			FSkeletalMeshResource* SkelMeshResource = PreviewSkelMeshComp->GetSkeletalMeshResource();
+			check(SkelMeshResource);
 
-		const int32 HalfX = Viewport->GetSizeXY().X/2;
-		const int32 HalfY = Viewport->GetSizeXY().Y/2;
+			// Draw stats about the mesh
+			const FBoxSphereBounds& SkelBounds = PreviewSkelMeshComp->Bounds;
+			const FPlane ScreenPosition = View->Project(SkelBounds.Origin);
 
-		const float ScreenRadius = FMath::Max((float)HalfX * View->ViewMatrices.ProjMatrix.M[0][0], (float)HalfY * View->ViewMatrices.ProjMatrix.M[1][1]) * SkelBounds.SphereRadius / FMath::Max(ScreenPosition.W,1.0f);
-		const float LODFactor = ScreenRadius / 320.0f;
+			const int32 HalfX = Viewport->GetSizeXY().X / 2;
+			const int32 HalfY = Viewport->GetSizeXY().Y / 2;
 
-		int32 NumBonesInUse;
-		int32 NumBonesMappedToVerts; 
-		int32 NumChunksInUse;
-		int32 NumSectionsInUse;
-		FString WeightUsage;
+			const float ScreenRadius = FMath::Max((float)HalfX * View->ViewMatrices.ProjMatrix.M[0][0], (float)HalfY * View->ViewMatrices.ProjMatrix.M[1][1]) * SkelBounds.SphereRadius / FMath::Max(ScreenPosition.W, 1.0f);
+			const float LODFactor = ScreenRadius / 320.0f;
 
-		const int32 LODIndex = FMath::Clamp(PreviewSkelMeshComp->PredictedLODLevel, 0, SkelMeshResource->LODModels.Num()-1);
-		FStaticLODModel& LODModel = SkelMeshResource->LODModels[ LODIndex ];
+			int32 NumBonesInUse;
+			int32 NumBonesMappedToVerts;
+			int32 NumChunksInUse;
+			int32 NumSectionsInUse;
+			FString WeightUsage;
 
-		NumBonesInUse = LODModel.RequiredBones.Num();
-		NumBonesMappedToVerts = LODModel.ActiveBoneIndices.Num();
-		NumChunksInUse = LODModel.Chunks.Num();
-		NumSectionsInUse = LODModel.Sections.Num();
+			const int32 LODIndex = FMath::Clamp(PreviewSkelMeshComp->PredictedLODLevel, 0, SkelMeshResource->LODModels.Num() - 1);
+			FStaticLODModel& LODModel = SkelMeshResource->LODModels[LODIndex];
 
-		InfoString = FString::Printf( TEXT("LOD [%d] Bones:%d (Mapped to Vertices:%d) Polys:%d"), 
-			LODIndex, 
-			NumBonesInUse,	   
-			NumBonesMappedToVerts,
-			LODModel.GetTotalFaces());
+			NumBonesInUse = LODModel.RequiredBones.Num();
+			NumBonesMappedToVerts = LODModel.ActiveBoneIndices.Num();
+			NumChunksInUse = LODModel.Chunks.Num();
+			NumSectionsInUse = LODModel.Sections.Num();
 
-		Canvas->DrawShadowedString( CurXOffset, CurYOffset, *InfoString, GEngine->GetSmallFont(), TextColor );
+			InfoString = FString::Printf(TEXT("LOD [%d] Bones:%d (Mapped to Vertices:%d) Polys:%d"),
+				LODIndex,
+				NumBonesInUse,
+				NumBonesMappedToVerts,
+				LODModel.GetTotalFaces());
 
-		InfoString = FString::Printf( TEXT("(Display Factor:%3.2f, FOV:%3.0f)"), LODFactor, ViewFOV);
-		CurYOffset += YL + 2;
-		Canvas->DrawShadowedString( CurXOffset, CurYOffset, *InfoString, GEngine->GetSmallFont(), TextColor );
+			Canvas->DrawShadowedString(CurXOffset, CurYOffset, *InfoString, GEngine->GetSmallFont(), TextColor);
 
-		CurYOffset += 1; // --
+			InfoString = FString::Printf(TEXT("(Screen Size:%3.2f, FOV:%3.0f)"), LODFactor, ViewFOV);
+			CurYOffset += YL + 2;
+			Canvas->DrawShadowedString(CurXOffset, CurYOffset, *InfoString, GEngine->GetSmallFont(), TextColor);
 
-		uint32 TotalRigidVertices = 0;
-		uint32 TotalSoftVertices = 0;
-		for(int32 ChunkIndex = 0;ChunkIndex < LODModel.Chunks.Num();ChunkIndex++)
-		{
-			int32 ChunkRigidVerts = LODModel.Chunks[ChunkIndex].GetNumRigidVertices();
-			int32 ChunkSoftVerts = LODModel.Chunks[ChunkIndex].GetNumSoftVertices();
-
-			InfoString = FString::Printf( TEXT(" [Chunk %d] Verts:%d (Rigid:%d Soft:%d) Bones:%d"), 
-				ChunkIndex,
-				ChunkRigidVerts + ChunkSoftVerts,
-				ChunkRigidVerts,
-				ChunkSoftVerts,
-				LODModel.Chunks[ChunkIndex].BoneMap.Num()
-				);
+			CurYOffset += 1; // --
 
-			CurYOffset += YL + 2;
-			Canvas->DrawShadowedString( CurXOffset, CurYOffset, *InfoString, GEngine->GetSmallFont(), TextColor*0.8f );
+			uint32 TotalRigidVertices = 0;
+			uint32 TotalSoftVertices = 0;
+			for (int32 ChunkIndex = 0; ChunkIndex < LODModel.Chunks.Num(); ChunkIndex++)
+			{
+				int32 ChunkRigidVerts = LODModel.Chunks[ChunkIndex].GetNumRigidVertices();
+				int32 ChunkSoftVerts = LODModel.Chunks[ChunkIndex].GetNumSoftVertices();
 
-			TotalRigidVertices += ChunkRigidVerts;
-			TotalSoftVertices += ChunkSoftVerts;
-		}
+				InfoString = FString::Printf(TEXT(" [Chunk %d] Verts:%d (Rigid:%d Soft:%d) Bones:%d"),
+					ChunkIndex,
+					ChunkRigidVerts + ChunkSoftVerts,
+					ChunkRigidVerts,
+					ChunkSoftVerts,
+					LODModel.Chunks[ChunkIndex].BoneMap.Num()
+					);
 
-		InfoString = FString::Printf( TEXT("TOTAL Verts:%d (Rigid:%d Soft:%d)"), 
-			LODModel.NumVertices,
-			TotalRigidVertices,
-			TotalSoftVertices );
+				CurYOffset += YL + 2;
+				Canvas->DrawShadowedString(CurXOffset, CurYOffset, *InfoString, GEngine->GetSmallFont(), TextColor*0.8f);
 
-		CurYOffset += 1; // --
+				TotalRigidVertices += ChunkRigidVerts;
+				TotalSoftVertices += ChunkSoftVerts;
+			}
 
+			InfoString = FString::Printf(TEXT("TOTAL Verts:%d (Rigid:%d Soft:%d)"),
+				LODModel.NumVertices,
+				TotalRigidVertices,
+				TotalSoftVertices);
 
-		CurYOffset += YL + 2;
-		Canvas->DrawShadowedString( CurXOffset, CurYOffset, *InfoString, GEngine->GetSmallFont(), TextColor );
+			CurYOffset += 1; // --
 
-		InfoString = FString::Printf( TEXT("Chunks:%d Sections:%d %s"), 
-			NumChunksInUse, 
-			NumSectionsInUse,
-			*WeightUsage
-			);
 
-		CurYOffset += YL + 2;
-		Canvas->DrawShadowedString( CurXOffset, CurYOffset, *InfoString, GEngine->GetSmallFont(), TextColor );
-		
-		if (PreviewSkelMeshComp->BonesOfInterest.Num() > 0)
-		{
-			int32 BoneIndex = PreviewSkelMeshComp->BonesOfInterest[0];
-			const FName BoneName = PreviewSkelMeshComp->SkeletalMesh->RefSkeleton.GetBoneName(BoneIndex);
-			FTransform ReferenceTransform = PreviewSkelMeshComp->SkeletalMesh->RefSkeleton.GetRefBonePose()[BoneIndex];
-			FTransform LocalTransform = PreviewSkelMeshComp->LocalAtoms[BoneIndex];
-			FTransform ComponentTransform = PreviewSkelMeshComp->SpaceBases[BoneIndex];
+			CurYOffset += YL + 2;
+			Canvas->DrawShadowedString(CurXOffset, CurYOffset, *InfoString, GEngine->GetSmallFont(), TextColor);
 
-			CurYOffset += YL + 2;						
-			InfoString = FString::Printf(TEXT("Local :%s"), *LocalTransform.ToString());
-			Canvas->DrawShadowedString( CurXOffset, CurYOffset, *InfoString, GEngine->GetSmallFont(), TextColor );
+			InfoString = FString::Printf(TEXT("Chunks:%d Sections:%d %s"),
+				NumChunksInUse,
+				NumSectionsInUse,
+				*WeightUsage
+				);
 
 			CurYOffset += YL + 2;
-			InfoString = FString::Printf(TEXT("Component :%s"), *ComponentTransform.ToString());
-			Canvas->DrawShadowedString( CurXOffset, CurYOffset, *InfoString, GEngine->GetSmallFont(), TextColor );
+			Canvas->DrawShadowedString(CurXOffset, CurYOffset, *InfoString, GEngine->GetSmallFont(), TextColor);
+
+			if (PreviewSkelMeshComp->BonesOfInterest.Num() > 0)
+			{
+				int32 BoneIndex = PreviewSkelMeshComp->BonesOfInterest[0];
+				const FName BoneName = PreviewSkelMeshComp->SkeletalMesh->RefSkeleton.GetBoneName(BoneIndex);
+				FTransform ReferenceTransform = PreviewSkelMeshComp->SkeletalMesh->RefSkeleton.GetRefBonePose()[BoneIndex];
+				FTransform LocalTransform = PreviewSkelMeshComp->LocalAtoms[BoneIndex];
+				FTransform ComponentTransform = PreviewSkelMeshComp->SpaceBases[BoneIndex];
+
+				CurYOffset += YL + 2;
+				InfoString = FString::Printf(TEXT("Local :%s"), *LocalTransform.ToString());
+				Canvas->DrawShadowedString(CurXOffset, CurYOffset, *InfoString, GEngine->GetSmallFont(), TextColor);
+
+				CurYOffset += YL + 2;
+				InfoString = FString::Printf(TEXT("Component :%s"), *ComponentTransform.ToString());
+				Canvas->DrawShadowedString(CurXOffset, CurYOffset, *InfoString, GEngine->GetSmallFont(), TextColor);
+
+				CurYOffset += YL + 2;
+				InfoString = FString::Printf(TEXT("Reference :%s"), *ReferenceTransform.ToString());
+				Canvas->DrawShadowedString(CurXOffset, CurYOffset, *InfoString, GEngine->GetSmallFont(), TextColor);
+			}
 
 			CurYOffset += YL + 2;
-			InfoString = FString::Printf(TEXT("Reference :%s"), *ReferenceTransform.ToString());
+			InfoString = FString::Printf(TEXT("Approximate Size: %ix%ix%i"),
+				FMath::RoundToInt(PreviewSkelMeshComp->Bounds.BoxExtent.X * 2.0f),
+				FMath::RoundToInt(PreviewSkelMeshComp->Bounds.BoxExtent.Y * 2.0f),
+				FMath::RoundToInt(PreviewSkelMeshComp->Bounds.BoxExtent.Z * 2.0f));
 			Canvas->DrawShadowedString(CurXOffset, CurYOffset, *InfoString, GEngine->GetSmallFont(), TextColor);
+
+			uint32 NumNotiesWithErrors = PreviewSkelMeshComp->AnimNotifyErrors.Num();
+			for (uint32 i = 0; i < NumNotiesWithErrors; ++i)
+			{
+				uint32 NumErrors = PreviewSkelMeshComp->AnimNotifyErrors[i].Errors.Num();
+				for (uint32 ErrorIdx = 0; ErrorIdx < NumErrors; ++ErrorIdx)
+				{
+					CurYOffset += YL + 2;
+					Canvas->DrawShadowedString(CurXOffset, CurYOffset, *PreviewSkelMeshComp->AnimNotifyErrors[i].Errors[ErrorIdx], GEngine->GetSmallFont(), FLinearColor(1.0f, 0.0f, 0.0f, 1.0f));
+				}
+			}
 		}
+		else // simplified default display info to be same as static mesh editor
+		{
+			FSkeletalMeshResource* SkelMeshResource = PreviewSkelMeshComp->GetSkeletalMeshResource();
+			check(SkelMeshResource);
 
-		CurYOffset += YL + 2;
-		InfoString = FString::Printf(TEXT("Approximate Size: %ix%ix%i"), 
-			FMath::RoundToInt(PreviewSkelMeshComp->Bounds.BoxExtent.X * 2.0f),
-			FMath::RoundToInt(PreviewSkelMeshComp->Bounds.BoxExtent.Y * 2.0f),
-			FMath::RoundToInt(PreviewSkelMeshComp->Bounds.BoxExtent.Z * 2.0f));
-		Canvas->DrawShadowedString( CurXOffset, CurYOffset, *InfoString, GEngine->GetSmallFont(), TextColor );
+			const int32 LODIndex = FMath::Clamp(PreviewSkelMeshComp->PredictedLODLevel, 0, SkelMeshResource->LODModels.Num() - 1);
+			FStaticLODModel& LODModel = SkelMeshResource->LODModels[LODIndex];
 
-		uint32 NumNotiesWithErrors = PreviewSkelMeshComp->AnimNotifyErrors.Num();
-		for (uint32 i = 0; i < NumNotiesWithErrors; ++i)
-		{
-			uint32 NumErrors = PreviewSkelMeshComp->AnimNotifyErrors[i].Errors.Num();
-			for (uint32 ErrorIdx = 0; ErrorIdx < NumErrors; ++ErrorIdx)
+			InfoString = FString::Printf(TEXT("LOD: %d"), LODIndex);
+			Canvas->DrawShadowedString(CurXOffset, CurYOffset, *InfoString, GEngine->GetSmallFont(), TextColor);
+
+			/*
+				float ComputeBoundsScreenSize( const FVector4& Origin, const float SphereRadius, const FSceneView& View )
+				{
+					// Only need one component from a view transformation; just calculate the one we're interested in.
+					const float Divisor =  Dot3(Origin - View.ViewMatrices.ViewOrigin, View.ViewMatrices.ViewMatrix.GetColumn(2));
+
+					// Get projection multiple accounting for view scaling.
+					const float ScreenMultiple = FMath::Max(View.ViewRect.Width() / 2.0f * View.ViewMatrices.ProjMatrix.M[0][0],
+						View.ViewRect.Height() / 2.0f * View.ViewMatrices.ProjMatrix.M[1][1]);
+
+					const float ScreenRadius = ScreenMultiple * SphereRadius / FMath::Max(Divisor, 1.0f);
+					const float ScreenArea = PI * ScreenRadius * ScreenRadius;
+					return FMath::Clamp(ScreenArea / View.ViewRect.Area(), 0.0f, 1.0f);
+				}			
+			*/
+			// current screen size
+			// Draw stats about the mesh
+			const FBoxSphereBounds& SkelBounds = PreviewSkelMeshComp->Bounds;
+			const FPlane ScreenPosition = View->Project(SkelBounds.Origin);
+
+			const int32 HalfX = Viewport->GetSizeXY().X / 2;
+			const int32 HalfY = Viewport->GetSizeXY().Y / 2;
+			const float ScreenRadius = FMath::Max((float)HalfX * View->ViewMatrices.ProjMatrix.M[0][0], (float)HalfY * View->ViewMatrices.ProjMatrix.M[1][1]) * SkelBounds.SphereRadius / FMath::Max(ScreenPosition.W, 1.0f);
+			const float LODFactor = ScreenRadius / 320.0f;
+
+			float ScreenSize = ComputeBoundsScreenSize(ScreenPosition, SkelBounds.SphereRadius, *View);
+
+			InfoString = FString::Printf(TEXT("Current Screen Size:%3.2f"), LODFactor);
+			CurYOffset += YL + 2;
+			Canvas->DrawShadowedString(CurXOffset, CurYOffset, *InfoString, GEngine->GetSmallFont(), TextColor);
+
+			// Triangles
+			uint32 NumTotalTriangles = 0;
+			int32 NumSections = LODModel.NumNonClothingSections();
+			for (int32 SectionIndex = 0; SectionIndex < NumSections; SectionIndex++)
 			{
-				CurYOffset += YL + 2;
-				Canvas->DrawShadowedString(CurXOffset, CurYOffset, *PreviewSkelMeshComp->AnimNotifyErrors[i].Errors[ErrorIdx], GEngine->GetSmallFont(), FLinearColor(1.0f, 0.0f, 0.0f, 1.0f));
+				NumTotalTriangles += LODModel.Sections[SectionIndex].NumTriangles;
 			}
+			InfoString = FString::Printf(TEXT("Triangles: %d"), NumTotalTriangles);
+			CurYOffset += YL + 2;
+			Canvas->DrawShadowedString(CurXOffset, CurYOffset, *InfoString, GEngine->GetSmallFont(), TextColor);
+
+			// Vertices
+			InfoString = FString::Printf(TEXT("Vertices: %d"), LODModel.NumVertices);
+			CurYOffset += YL + 2;
+			Canvas->DrawShadowedString(CurXOffset, CurYOffset, *InfoString, GEngine->GetSmallFont(), TextColor);
+
+			// UV Channels
+			InfoString = FString::Printf(TEXT("UV Channels: %d"), LODModel.NumTexCoords);
+			CurYOffset += YL + 2;
+			Canvas->DrawShadowedString(CurXOffset, CurYOffset, *InfoString, GEngine->GetSmallFont(), TextColor);
+			
+			// Approx Size 
+			CurYOffset += YL + 2;
+			InfoString = FString::Printf(TEXT("Approx Size: %ix%ix%i"),
+				FMath::RoundToInt(PreviewSkelMeshComp->Bounds.BoxExtent.X * 2.0f),
+				FMath::RoundToInt(PreviewSkelMeshComp->Bounds.BoxExtent.Y * 2.0f),
+				FMath::RoundToInt(PreviewSkelMeshComp->Bounds.BoxExtent.Z * 2.0f));
+			Canvas->DrawShadowedString(CurXOffset, CurYOffset, *InfoString, GEngine->GetSmallFont(), TextColor);
 		}
 	}
 }
@@ -962,7 +1036,7 @@ void FAnimationViewportClient::ProcessClick(class FSceneView& View, class HHitPr
 		{
 			PersonaPtr.Pin()->DeselectAll();
 		}
-	}
+}
 }
 
 bool FAnimationViewportClient::InputWidgetDelta( FViewport* Viewport, EAxisList::Type CurrentAxis, FVector& Drag, FRotator& Rot, FVector& Scale )
diff --git a/Engine/Source/Editor/UnrealEd/Private/Fbx/FbxFactory.cpp b/Engine/Source/Editor/UnrealEd/Private/Fbx/FbxFactory.cpp
--- a/Engine/Source/Editor/UnrealEd/Private/Fbx/FbxFactory.cpp
+++ b/Engine/Source/Editor/UnrealEd/Private/Fbx/FbxFactory.cpp
@@ -398,7 +398,7 @@ UObject* UFbxFactory::FactoryCreateBinary
 
 								if (bImportSucceeded)
 								{
-									BaseSkeletalMesh->LODInfo[LODIndex].DisplayFactor = 1.0f / (MaxLODLevel * LODIndex);
+									BaseSkeletalMesh->LODInfo[LODIndex].ScreenSize = 1.0f / (MaxLODLevel * LODIndex);
 								}
 								else
 								{
diff --git a/Engine/Source/Editor/UnrealEd/Private/Fbx/FbxSkeletalMeshImport.cpp b/Engine/Source/Editor/UnrealEd/Private/Fbx/FbxSkeletalMeshImport.cpp
--- a/Engine/Source/Editor/UnrealEd/Private/Fbx/FbxSkeletalMeshImport.cpp
+++ b/Engine/Source/Editor/UnrealEd/Private/Fbx/FbxSkeletalMeshImport.cpp
@@ -1549,7 +1549,7 @@ USkeletalMesh* UnFbx::FFbxImporter::ReimportSkeletalMesh(USkeletalMesh* Mesh, UF
 				// Don't override DispalyFactor of existing LODs
 				if(LODIndex >= NumPrevLODs)
 				{
-					BaseSkeletalMesh->LODInfo[LODIndex].DisplayFactor = 1.0f / (MaxLODLevel * LODIndex);
+					BaseSkeletalMesh->LODInfo[LODIndex].ScreenSize = 1.0f / (MaxLODLevel * LODIndex);
 				}
 			}
 
diff --git a/Engine/Source/Runtime/Engine/Classes/Engine/SkeletalMesh.h b/Engine/Source/Runtime/Engine/Classes/Engine/SkeletalMesh.h
--- a/Engine/Source/Runtime/Engine/Classes/Engine/SkeletalMesh.h
+++ b/Engine/Source/Runtime/Engine/Classes/Engine/SkeletalMesh.h
@@ -496,7 +496,7 @@ struct FSkeletalMeshLODInfo
 
 	/**	Indicates when to use this LOD. A smaller number means use this LOD when further away. */
 	UPROPERTY(EditAnywhere, Category=SkeletalMeshLODInfo)
-	float DisplayFactor;
+	float ScreenSize;
 
 	/**	Used to avoid 'flickering' when on LOD boundary. Only taken into account when moving from complex->simple. */
 	UPROPERTY(EditAnywhere, Category=SkeletalMeshLODInfo)
@@ -526,7 +526,7 @@ struct FSkeletalMeshLODInfo
 	FSkeletalMeshOptimizationSettings ReductionSettings;
 
 	FSkeletalMeshLODInfo()
-		: DisplayFactor(0)
+		: ScreenSize(0)
 		, LODHysteresis(0)
 		, bHasBeenSimplified(false)
 	{
diff --git a/Engine/Source/Runtime/Engine/Private/SkeletalMesh.cpp b/Engine/Source/Runtime/Engine/Private/SkeletalMesh.cpp
--- a/Engine/Source/Runtime/Engine/Private/SkeletalMesh.cpp
+++ b/Engine/Source/Runtime/Engine/Private/SkeletalMesh.cpp
@@ -2709,10 +2709,10 @@ void USkeletalMesh::DebugVerifySkeletalMeshLOD()
 	{
 		for(int32 i=1; i<LODInfo.Num(); i++)
 		{
-			if (LODInfo[i].DisplayFactor <= 0.1f)
+			if (LODInfo[i].ScreenSize <= 0.1f)
 			{
 				// too small
-				UE_LOG(LogSkeletalMesh, Warning, TEXT("SkelMeshLOD (%s) : DisplayFactor for LOD %d may be too small (%0.5f)"), *GetPathName(), i, LODInfo[i].DisplayFactor);
+				UE_LOG(LogSkeletalMesh, Warning, TEXT("SkelMeshLOD (%s) : ScreenSize for LOD %d may be too small (%0.5f)"), *GetPathName(), i, LODInfo[i].ScreenSize);
 			}
 		}
 	}
diff --git a/Engine/Source/Runtime/Engine/Private/SkeletalMeshMerge.cpp b/Engine/Source/Runtime/Engine/Private/SkeletalMeshMerge.cpp
--- a/Engine/Source/Runtime/Engine/Private/SkeletalMeshMerge.cpp
+++ b/Engine/Source/Runtime/Engine/Private/SkeletalMeshMerge.cpp
@@ -392,7 +392,7 @@ void FSkeletalMeshMerge::GenerateLODModel( int32 LODIdx )
 	MergeLODModel.Size = 0;
 	// add the new LOD info entry
 	FSkeletalMeshLODInfo& MergeLODInfo = *new(MergeMesh->LODInfo) FSkeletalMeshLODInfo;
-	MergeLODInfo.DisplayFactor = MergeLODInfo.LODHysteresis = MAX_FLT;
+	MergeLODInfo.ScreenSize = MergeLODInfo.LODHysteresis = MAX_FLT;
 
 	// generate an array with info about new sections that need to be created
 	TArray<FNewSectionInfo> NewSectionArray;
@@ -479,7 +479,7 @@ void FSkeletalMeshMerge::GenerateLODModel( int32 LODIdx )
 			const FSkeletalMeshLODInfo& SrcLODInfo = MergeSectionInfo.SkelMesh->LODInfo[SourceLODIdx];
 
 			// keep track of the lowest LOD displayfactor and hysterisis
-			MergeLODInfo.DisplayFactor = FMath::Min<float>(MergeLODInfo.DisplayFactor,SrcLODInfo.DisplayFactor);
+			MergeLODInfo.ScreenSize = FMath::Min<float>(MergeLODInfo.ScreenSize, SrcLODInfo.ScreenSize);
 			MergeLODInfo.LODHysteresis = FMath::Min<float>(MergeLODInfo.LODHysteresis,SrcLODInfo.LODHysteresis);
 
 			// get the source skel LOD model from this merge entry
diff --git a/Engine/Source/Runtime/Engine/Private/SkeletalRender.cpp b/Engine/Source/Runtime/Engine/Private/SkeletalRender.cpp
--- a/Engine/Source/Runtime/Engine/Private/SkeletalRender.cpp
+++ b/Engine/Source/Runtime/Engine/Private/SkeletalRender.cpp
@@ -93,7 +93,7 @@ void FSkeletalMeshObject::UpdateMinDesiredLODLevel(const FSceneView* View, const
 		for(int32 LODLevel = SkeletalMeshResource->LODModels.Num()-1; LODLevel > 0; LODLevel--) 
 		{
 			// Get DistanceFactor for this LOD
-			float LODDistanceFactor = SkeletalMeshLODInfo[LODLevel].DisplayFactor;
+			float LODDistanceFactor = SkeletalMeshLODInfo[LODLevel].ScreenSize;
 
 			// If we are considering shifting to a better (lower) LOD, bias with hysteresis.
 			if(LODLevel  <= CurrentLODLevel)
