Added a space between "Current Screen Size:" and a value

diff --git a/Engine/Source/Editor/Persona/Private/AnimationEditorViewportClient.cpp b/Engine/Source/Editor/Persona/Private/AnimationEditorViewportClient.cpp
--- a/Engine/Source/Editor/Persona/Private/AnimationEditorViewportClient.cpp
+++ b/Engine/Source/Editor/Persona/Private/AnimationEditorViewportClient.cpp
@@ -899,26 +899,11 @@ void FAnimationViewportClient::DisplayInfo(FCanvas* Canvas, FSceneView* View, bo
 			const int32 LODIndex = FMath::Clamp(PreviewSkelMeshComp->PredictedLODLevel, 0, SkelMeshResource->LODModels.Num() - 1);
 			FStaticLODModel& LODModel = SkelMeshResource->LODModels[LODIndex];
 
+			// Current LOD 
 			InfoString = FString::Printf(TEXT("LOD: %d"), LODIndex);
 			Canvas->DrawShadowedString(CurXOffset, CurYOffset, *InfoString, GEngine->GetSmallFont(), TextColor);
 
-			/*
-				float ComputeBoundsScreenSize( const FVector4& Origin, const float SphereRadius, const FSceneView& View )
-				{
-					// Only need one component from a view transformation; just calculate the one we're interested in.
-					const float Divisor =  Dot3(Origin - View.ViewMatrices.ViewOrigin, View.ViewMatrices.ViewMatrix.GetColumn(2));
-
-					// Get projection multiple accounting for view scaling.
-					const float ScreenMultiple = FMath::Max(View.ViewRect.Width() / 2.0f * View.ViewMatrices.ProjMatrix.M[0][0],
-						View.ViewRect.Height() / 2.0f * View.ViewMatrices.ProjMatrix.M[1][1]);
-
-					const float ScreenRadius = ScreenMultiple * SphereRadius / FMath::Max(Divisor, 1.0f);
-					const float ScreenArea = PI * ScreenRadius * ScreenRadius;
-					return FMath::Clamp(ScreenArea / View.ViewRect.Area(), 0.0f, 1.0f);
-				}			
-			*/
 			// current screen size
-			// Draw stats about the mesh
 			const FBoxSphereBounds& SkelBounds = PreviewSkelMeshComp->Bounds;
 			const FPlane ScreenPosition = View->Project(SkelBounds.Origin);
 
@@ -929,7 +914,7 @@ void FAnimationViewportClient::DisplayInfo(FCanvas* Canvas, FSceneView* View, bo
 
 			float ScreenSize = ComputeBoundsScreenSize(ScreenPosition, SkelBounds.SphereRadius, *View);
 
-			InfoString = FString::Printf(TEXT("Current Screen Size:%3.2f"), LODFactor);
+			InfoString = FString::Printf(TEXT("Current Screen Size: %3.2f"), LODFactor);
 			CurYOffset += YL + 2;
 			Canvas->DrawShadowedString(CurXOffset, CurYOffset, *InfoString, GEngine->GetSmallFont(), TextColor);
 
