Persona usability -

Added tool bar - mainly import/reimport options

diff --git a/Engine/Source/Editor/EditorStyle/Private/SlateEditorStyle.cpp b/Engine/Source/Editor/EditorStyle/Private/SlateEditorStyle.cpp
--- a/Engine/Source/Editor/EditorStyle/Private/SlateEditorStyle.cpp
+++ b/Engine/Source/Editor/EditorStyle/Private/SlateEditorStyle.cpp
@@ -3899,16 +3899,13 @@ void FSlateEditorStyle::FStyle::SetupPersonaStyle()
 	// Persona
 #if WITH_EDITOR
 	{
-		Set( "Persona.ConvertAnimationGraph", new IMAGE_BRUSH( "Old/Graph/ConvertIcon", Icon40x40 ) );
+		// Persona viewport
 		Set( "AnimViewportMenu.TranslateMode", new IMAGE_BRUSH( "Icons/icon_translate_40x", Icon32x32) );
 		Set( "AnimViewportMenu.TranslateMode.Small", new IMAGE_BRUSH( "Icons/icon_translate_40x", Icon16x16 ) );
 		Set( "AnimViewportMenu.RotateMode", new IMAGE_BRUSH( "Icons/icon_rotate_40x", Icon32x32) );
 		Set( "AnimViewportMenu.RotateMode.Small", new IMAGE_BRUSH( "Icons/icon_rotate_40x", Icon16x16 ) );
-		Set( "Persona.ReimportAsset", new IMAGE_BRUSH( "Icons/Reimport_12x", Icon12x12 ) );
 		Set( "AnimViewportMenu.CameraFollow", new IMAGE_BRUSH( "Persona/Viewport/Camera_FollowBounds_40px", Icon32x32) );
 		Set( "AnimViewportMenu.CameraFollow.Small", new IMAGE_BRUSH( "Persona/Viewport/Camera_FollowBounds_40px", Icon16x16 ) );
-		Set( "SkeletonTree.SkeletonSocket", new IMAGE_BRUSH( "Persona/SkeletonTree/icon_SocketG_16px", Icon16x16) );
-		Set( "SkeletonTree.MeshSocket", new IMAGE_BRUSH( "Persona/SkeletonTree/icon_SocketC_16px", Icon16x16) );
 		Set( "AnimViewport.LocalSpaceEditing", new IMAGE_BRUSH( "Icons/icon_axis_local_16px", FVector2D( 16, 16 ) ) );
 		Set( "AnimViewport.WorldSpaceEditing", new IMAGE_BRUSH( "Icons/icon_axis_world_16px", FVector2D( 16, 16 ) ) );
 		Set( "AnimViewportMenu.SetShowNormals", new IMAGE_BRUSH( TEXT("Icons/icon_StaticMeshEd_Normals_40x"), Icon40x40 ) );
@@ -3928,6 +3925,36 @@ void FSlateEditorStyle::FStyle::SetupPersonaStyle()
 			.SetShadowOffset( FVector2D( 1,1 ) )
 			.SetShadowColorAndOpacity( FLinearColor(0,0,0,0.9f) )
 			);
+
+		// persona commands
+		Set("Persona.ChangeSkeletonPreviewMesh", new IMAGE_BRUSH(TEXT("Icons/icon_ShowSkeletalMeshes_40x"), Icon40x40));
+		Set("Persona.ChangeSkeletonPreviewMesh.Small", new IMAGE_BRUSH(TEXT("Icons/icon_ShowSkeletalMeshes_40x"), Icon20x20));
+		Set("Persona.AnimNotifyWindow", new IMAGE_BRUSH(TEXT("Icons/icon_ShowNavigation_40x"), Icon40x40));
+		Set("Persona.AnimNotifyWindow.Small", new IMAGE_BRUSH(TEXT("Icons/icon_ShowNavigation_40x"), Icon20x20));
+		Set("Persona.RetargetSourceMgr", new IMAGE_BRUSH(TEXT("Icons/icon_PlayDefaultPlayerStart_40x"), Icon40x40));
+		Set("Persona.RetargetSourceMgr.Small", new IMAGE_BRUSH(TEXT("Icons/icon_PlayDefaultPlayerStart_40x"), Icon20x20));
+		Set("Persona.ImportMesh", new IMAGE_BRUSH(TEXT("Icons/icon_PhAT_Prismatic_40x"), Icon40x40));
+		Set("Persona.ImportMesh.Small", new IMAGE_BRUSH(TEXT("Icons/icon_PhAT_Prismatic_40x"), Icon20x20));
+		Set("Persona.ReimportMesh", new IMAGE_BRUSH(TEXT("Icons/icon_Placement_RecentlyPlaced_40x"), Icon40x40));
+		Set("Persona.ReimportMesh.Small", new IMAGE_BRUSH(TEXT("Icons/icon_Placement_RecentlyPlaced_40x"), Icon20x20));
+		Set("Persona.ImportLODs", new IMAGE_BRUSH(TEXT("Icons/icon_StaticMeshEd_AdditionalData_40x"), Icon40x40));
+		Set("Persona.ImportLODs.Small", new IMAGE_BRUSH(TEXT("Icons/icon_StaticMeshEd_AdditionalData_40x"), Icon20x20));
+		Set("Persona.AddBodyPart", new IMAGE_BRUSH(TEXT("Icons/icon_Placement_AllContent_40x"), Icon40x40));
+		Set("Persona.AddBodyPart.Small", new IMAGE_BRUSH(TEXT("Icons/icon_Placement_AllContent_40x"), Icon20x20));
+		Set("Persona.ImportAnimation", new IMAGE_BRUSH(TEXT("Icons/icon_PhAT_Prismatic_40x"), Icon40x40));
+		Set("Persona.ImportAnimation.Small", new IMAGE_BRUSH(TEXT("Icons/icon_PhAT_Prismatic_40x"), Icon20x20));
+		Set("Persona.ReimportAnimation", new IMAGE_BRUSH(TEXT("Icons/icon_PhAT_PhysMat_40x"), Icon40x40));
+		Set("Persona.ReimportAnimation.Small", new IMAGE_BRUSH(TEXT("Icons/icon_PhAT_PhysMat_40x"), Icon20x20));
+		Set("Persona.ApplyCompression", new IMAGE_BRUSH(TEXT("Icons/icon_PhAT_NewBody_40x"), Icon40x40));
+		Set("Persona.ApplyCompression.Small", new IMAGE_BRUSH(TEXT("Icons/icon_PhAT_NewBody_40x"), Icon20x20));
+		Set("Persona.ExportToFBX", new IMAGE_BRUSH(TEXT("Icons/icon_DestructibleMeshEd_ImportFBX_40x"), Icon40x40));
+		Set("Persona.ExportToFBX.Small", new IMAGE_BRUSH(TEXT("Icons/icon_DestructibleMeshEd_ImportFBX_40x"), Icon20x20));
+
+		// persona extras
+		Set("Persona.ConvertAnimationGraph", new IMAGE_BRUSH("Old/Graph/ConvertIcon", Icon40x40));
+		Set("Persona.ReimportAsset", new IMAGE_BRUSH("Icons/Reimport_12x", Icon12x12));
+		Set("SkeletonTree.SkeletonSocket", new IMAGE_BRUSH("Persona/SkeletonTree/icon_SocketG_16px", Icon16x16));
+		Set("SkeletonTree.MeshSocket", new IMAGE_BRUSH("Persona/SkeletonTree/icon_SocketC_16px", Icon16x16));
 	}
 
 	// Kismet 2
diff --git a/Engine/Source/Editor/Persona/Private/Persona.cpp b/Engine/Source/Editor/Persona/Private/Persona.cpp
--- a/Engine/Source/Editor/Persona/Private/Persona.cpp
+++ b/Engine/Source/Editor/Persona/Private/Persona.cpp
@@ -35,7 +35,6 @@
 #include "PhysicsMode/PhysicsMode.h"
 #include "AnimationMode/AnimationMode.h"
 #include "BlueprintMode/AnimBlueprintMode.h"
-#include "Editor/ContentBrowser/Public/ContentBrowserModule.h"
 #include "Runtime/AssetRegistry/Public/AssetRegistryModule.h"
 
 #include "Editor/Persona/Private/AnimationEditorViewportClient.h"
@@ -59,6 +58,10 @@
 #include "AnimationCompressionPanel.h"
 #include "FbxAnimUtils.h"
 
+#include "SAnimationDlgs.h"
+#include "Developer/AssetTools/Public/AssetToolsModule.h"
+#include "FbxMeshUtils.h"
+
 #define LOCTEXT_NAMESPACE "FPersona"
 
 /////////////////////////////////////////////////////
@@ -72,273 +75,6 @@ struct FLocalCharEditorCallbacks
 	}
 };
 
-/////////////////////////////////////////////////////
-/////////////////////////////////////////////////////
-/////////////////////////////////////////////////////
-/////////////////////////////////////////////////////
-
-/////////////////////////////////////////////////////
-// Create Animation dialog for recording animation
-/////////////////////////////////////////////////////
-
-class SCreateAnimationDlg : public SWindow
-{
-public:
-	SLATE_BEGIN_ARGS(SCreateAnimationDlg)
-		: _Title(LOCTEXT("Title", "Create New Animation Object"))
-	{
-	}
-	SLATE_ARGUMENT(FText, Title)
-		SLATE_ARGUMENT(FText, DefaultAssetPath)
-		SLATE_END_ARGS()
-
-		SCreateAnimationDlg()
-		: UserResponse(EAppReturnType::Cancel)
-	{
-	}
-
-	void Construct(const FArguments& InArgs);
-
-public:
-	/** Displays the dialog in a blocking fashion */
-	EAppReturnType::Type ShowModal();
-
-	/** Gets the resulting asset path */
-	FString GetAssetPath();
-
-	/** Gets the resulting asset name */
-	FString GetAssetName();
-
-	/** Gets the resulting full asset path (path+'/'+name) */
-	FString GetFullAssetPath();
-
-	/** Get Duration **/
-	float GetDuration();
-
-protected:
-	void OnPathChange(const FString& NewPath);
-	void OnNameChange(const FText& NewName, ETextCommit::Type CommitInfo);
-	void OnDurationChange(const FText& NewName, ETextCommit::Type CommitInfo);
-	FReply OnButtonClick(EAppReturnType::Type ButtonID);
-
-	bool ValidatePackage();
-
-	EAppReturnType::Type UserResponse;
-	FText AssetPath;
-	FText AssetName;
-	FText Duration;
-};
-
-void SCreateAnimationDlg::Construct(const FArguments& InArgs)
-{
-	AssetPath = FText::FromString(FPackageName::GetLongPackagePath(InArgs._DefaultAssetPath.ToString()));
-	AssetName = FText::FromString(FPackageName::GetLongPackageAssetName(InArgs._DefaultAssetPath.ToString()));
-
-	FPathPickerConfig PathPickerConfig;
-	PathPickerConfig.DefaultPath = AssetPath.ToString();
-	PathPickerConfig.OnPathSelected = FOnPathSelected::CreateSP(this, &SCreateAnimationDlg::OnPathChange);
-	PathPickerConfig.bAddDefaultPath = true;
-
-	FContentBrowserModule& ContentBrowserModule = FModuleManager::LoadModuleChecked<FContentBrowserModule>("ContentBrowser");
-
-	SWindow::Construct(SWindow::FArguments()
-		.Title(InArgs._Title)
-		.SupportsMinimize(false)
-		.SupportsMaximize(false)
-		//.SizingRule( ESizingRule::Autosized )
-		.ClientSize(FVector2D(450,450))
-		[
-			SNew(SVerticalBox)
-
-			+ SVerticalBox::Slot() // Add user input block
-			.Padding(2)
-			[
-				SNew(SBorder)
-				.BorderImage(FEditorStyle::GetBrush("ToolPanel.GroupBorder"))
-				[
-					SNew(SVerticalBox)
-
-					+SVerticalBox::Slot()
-					.AutoHeight()
-					[
-						SNew(STextBlock)
-						.Text(LOCTEXT("SelectPath", "Select Path to create animation"))
-						.Font( FSlateFontInfo( FPaths::EngineContentDir() / TEXT("Slate/Fonts/Roboto-Regular.ttf"), 14 ) )
-					]
-
-					+SVerticalBox::Slot()
-					.FillHeight(1)
-					.Padding(3)
-					[
-						ContentBrowserModule.Get().CreatePathPicker(PathPickerConfig)
-					]
-
-					+SVerticalBox::Slot()
-					.AutoHeight()
-					[
-						SNew(SSeparator)
-					]
-
-					+SVerticalBox::Slot()
-					.AutoHeight()
-					.Padding(3)
-					[
-						SNew(SHorizontalBox)
-
-						+SHorizontalBox::Slot()
-						.AutoWidth()
-						.Padding(0, 0, 10, 0)
-						.VAlign(VAlign_Center)
-						[
-							SNew(STextBlock)
-							.Text(LOCTEXT("AnimationName", "Animation Name"))
-						]
-
-						+SHorizontalBox::Slot()
-						[
-							SNew(SEditableTextBox)
-							.Text(AssetName)
-							.OnTextCommitted(this, &SCreateAnimationDlg::OnNameChange)
-							.MinDesiredWidth(250)
-						]
-					]
-
-					+SVerticalBox::Slot()
-					.AutoHeight()
-					.Padding(3)
-					[
-						SNew(SHorizontalBox)
-
-						+SHorizontalBox::Slot()
-						.AutoWidth()
-						.Padding(0, 0, 10, 0)
-						.VAlign(VAlign_Center)
-						[
-							SNew(STextBlock)
-							.Text(LOCTEXT("DurationInSec", "Duration (in sec) "))
-						]
-
-						+SHorizontalBox::Slot()
-						[
-							SNew(SEditableTextBox)
-							.Text(Duration)
-							.OnTextCommitted(this, &SCreateAnimationDlg::OnDurationChange)
-							.MinDesiredWidth(100)
-						]
-					]
-				]
-			]
-
-			+SVerticalBox::Slot()
-			.AutoHeight()
-			.HAlign(HAlign_Right)
-			.Padding(5)
-			[
-				SNew(SUniformGridPanel)
-				.SlotPadding(FEditorStyle::GetMargin("StandardDialog.SlotPadding"))
-				.MinDesiredSlotWidth(FEditorStyle::GetFloat("StandardDialog.MinDesiredSlotWidth"))
-				.MinDesiredSlotHeight(FEditorStyle::GetFloat("StandardDialog.MinDesiredSlotHeight"))
-				+SUniformGridPanel::Slot(0,0)
-				[
-					SNew(SButton) 
-					.HAlign(HAlign_Center)
-					.ContentPadding(FEditorStyle::GetMargin("StandardDialog.ContentPadding"))
-					.Text(LOCTEXT("OK", "OK"))
-					.OnClicked(this, &SCreateAnimationDlg::OnButtonClick, EAppReturnType::Ok)
-				]
-				+SUniformGridPanel::Slot(1,0)
-					[
-						SNew(SButton) 
-						.HAlign(HAlign_Center)
-						.ContentPadding(FEditorStyle::GetMargin("StandardDialog.ContentPadding"))
-						.Text(LOCTEXT("Cancel", "Cancel"))
-						.OnClicked(this, &SCreateAnimationDlg::OnButtonClick, EAppReturnType::Cancel)
-					]
-			]
-		]);
-}
-
-void SCreateAnimationDlg::OnDurationChange(const FText& InDuration, ETextCommit::Type CommitInfo)
-{
-	Duration = InDuration;
-}
-
-void SCreateAnimationDlg::OnNameChange(const FText& NewName, ETextCommit::Type CommitInfo)
-{
-	AssetName = NewName;
-}
-
-void SCreateAnimationDlg::OnPathChange(const FString& NewPath)
-{
-	AssetPath = FText::FromString(NewPath);
-}
-
-FReply SCreateAnimationDlg::OnButtonClick(EAppReturnType::Type ButtonID)
-{
-	UserResponse = ButtonID;
-
-	if (ButtonID != EAppReturnType::Cancel)
-	{
-		if (!ValidatePackage())
-		{
-			// reject the request
-			return FReply::Handled();
-		}
-	}
-
-	RequestDestroyWindow();
-
-	return FReply::Handled();
-}
-
-/** Ensures supplied package name information is valid */
-bool SCreateAnimationDlg::ValidatePackage()
-{
-	FText Reason;
-	FString FullPath = GetFullAssetPath();
-
-	if(!FPackageName::IsValidLongPackageName(FullPath, false, &Reason)
-		|| !FName(*AssetName.ToString()).IsValidObjectName(Reason))
-	{
-		FMessageDialog::Open(EAppMsgType::Ok, Reason);
-		return false;
-	}
-
-	if (GetDuration() <= 0.0f)
-	{
-		FMessageDialog::Open(EAppMsgType::Ok, LOCTEXT("InvalidDuration", "Invalid duration"));
-		return false;
-	}
-
-	return true;
-}
-
-EAppReturnType::Type SCreateAnimationDlg::ShowModal()
-{
-	GEditor->EditorAddModalWindow(SharedThis(this));
-	return UserResponse;
-}
-
-FString SCreateAnimationDlg::GetAssetPath()
-{
-	return AssetPath.ToString();
-}
-
-FString SCreateAnimationDlg::GetAssetName()
-{
-	return AssetName.ToString();
-}
-
-FString SCreateAnimationDlg::GetFullAssetPath()
-{
-	return AssetPath.ToString() + "/" + AssetName.ToString();
-}
-
-float SCreateAnimationDlg::GetDuration()
-{
-	return FCString::Atof(*Duration.ToString());
-}
-
 //////////////////////////////////////////////////////
 /////////////////////////////////////////////////////
 // FSkeletonEditAppMode
@@ -697,6 +433,79 @@ UObject* FPersona::GetAnimBlueprintAsObject() const
 	return GetAnimBlueprint();
 }
 
+void FPersona::ExtendMenu()
+{
+	// Add additional persona editor menus
+	struct Local
+	{
+		static void AddPersonaFileMenu(FMenuBuilder& MenuBuilder)
+		{
+			// View
+			MenuBuilder.BeginSection("Persona", LOCTEXT("PersonaEditorMenu", "Animation"));
+			{
+
+			}
+			MenuBuilder.EndSection();
+		}
+
+		static void AddPersonaAssetMenu(FMenuBuilder& MenuBuilder)
+		{
+			// View
+			MenuBuilder.BeginSection("Persona", LOCTEXT("PersonaAssetMenuMenu_Skeleton", "Skeleton"));
+			{
+				MenuBuilder.AddMenuEntry(FPersonaCommands::Get().ChangeSkeletonPreviewMesh);
+				MenuBuilder.AddMenuEntry(FPersonaCommands::Get().RemoveUnusedBones);
+			}
+			MenuBuilder.EndSection();
+
+			//  Animation menu
+			MenuBuilder.BeginSection("Persona", LOCTEXT("PersonaAssetMenuMenu_Animation", "Animation"));
+			{
+				MenuBuilder.AddMenuEntry(FPersonaCommands::Get().ApplyCompression);
+				MenuBuilder.AddMenuEntry(FPersonaCommands::Get().ExportToFBX);
+				MenuBuilder.AddMenuEntry(FPersonaCommands::Get().RecordAnimation);
+				MenuBuilder.AddMenuEntry(FPersonaCommands::Get().AddLoopingInterpolation);
+			}
+			MenuBuilder.EndSection();
+		}
+	};
+
+	if(MenuExtender.IsValid())
+	{
+		RemoveMenuExtender(MenuExtender);
+		MenuExtender.Reset();
+	}
+
+	MenuExtender = MakeShareable(new FExtender);
+
+	UAnimBlueprint* AnimBlueprint = GetAnimBlueprint();
+	if(AnimBlueprint)
+	{
+		TSharedPtr<FExtender> MenuExtender = MakeShareable(new FExtender);
+		FKismet2Menu::SetupBlueprintEditorMenu(MenuExtender, *this);
+		AddMenuExtender(MenuExtender);
+
+		MenuExtender->AddMenuExtension(
+			"FileLoadAndSave",
+			EExtensionHook::After,
+			GetToolkitCommands(),
+			FMenuExtensionDelegate::CreateStatic(&Local::AddPersonaFileMenu));
+	}
+
+	MenuExtender->AddMenuExtension(
+			"AssetEditorActions",
+			EExtensionHook::After,
+			GetToolkitCommands(),
+			FMenuExtensionDelegate::CreateStatic(&Local::AddPersonaAssetMenu)
+			);
+
+	AddMenuExtender(MenuExtender);
+
+	// add extensible menu if exists
+	FPersonaModule* PersonaModule = &FModuleManager::LoadModuleChecked<FPersonaModule>("Persona");
+	AddMenuExtender(PersonaModule->GetMenuExtensibilityManager()->GetAllExtenders(GetToolkitCommands(), GetEditingObjects()));
+}
+
 void FPersona::InitPersona(const EToolkitMode::Type Mode, const TSharedPtr< class IToolkitHost >& InitToolkitHost, USkeleton* InitSkeleton, UAnimBlueprint* InitAnimBlueprint, UAnimationAsset* InitAnimationAsset, class USkeletalMesh * InitMesh)
 {
 	FReimportManager::Instance()->OnPostReimport().AddRaw(this, &FPersona::OnPostReimport);
@@ -796,69 +605,6 @@ void FPersona::InitPersona(const EToolkitMode::Type Mode, const TSharedPtr< clas
 	}
 	CommonInitialization(AnimBlueprints);
 
-	FPersonaModule* PersonaModule = &FModuleManager::LoadModuleChecked<FPersonaModule>( "Persona" );
-
-	// Add additional persona editor menus
-	struct Local
-	{
-		static void AddPersonaFileMenu( FMenuBuilder& MenuBuilder )
-		{
-			// View
-			MenuBuilder.BeginSection("Persona", LOCTEXT("PersonaEditorMenu", "Animation" ) );
-			{
-			
-			}
-			MenuBuilder.EndSection();
-		}
-
-		static void AddPersonaAssetMenu( FMenuBuilder& MenuBuilder )
-		{
-			// View
-			MenuBuilder.BeginSection("Persona", LOCTEXT("PersonaAssetMenuMenu", "Skeleton" ) );
-			{
-				MenuBuilder.AddMenuEntry( FPersonaCommands::Get().ChangeSkeletonPreviewMesh );
-				MenuBuilder.AddMenuEntry( FPersonaCommands::Get().RemoveUnusedBones );
-			}
-			MenuBuilder.EndSection();
-
-			//  Animation menu
-			MenuBuilder.BeginSection("Persona", LOCTEXT("PersonaAssetMenuMenu", "Animation"));
-			{
-				MenuBuilder.AddMenuEntry(FPersonaCommands::Get().ApplyCompression);
-				MenuBuilder.AddMenuEntry(FPersonaCommands::Get().ExportToFBX);
-				MenuBuilder.AddMenuEntry(FPersonaCommands::Get().RecordAnimation );
-				MenuBuilder.AddMenuEntry(FPersonaCommands::Get().AddLoopingInterpolation);
-			}
-			MenuBuilder.EndSection();
-		}
-	};
-
-	PersonaMenuExtender = MakeShareable(new FExtender);
-
-	UAnimBlueprint* AnimBlueprint = GetAnimBlueprint();
-	if (AnimBlueprint)
-	{
-		TSharedPtr<FExtender> MenuExtender = MakeShareable(new FExtender);
-		FKismet2Menu::SetupBlueprintEditorMenu( MenuExtender, *this );
-		AddMenuExtender(MenuExtender);
-
-		PersonaMenuExtender->AddMenuExtension(
-			"FileLoadAndSave", 
-			EExtensionHook::After, 
-			GetToolkitCommands(), 
-			FMenuExtensionDelegate::CreateStatic(&Local::AddPersonaFileMenu));
-	}
-
-	PersonaMenuExtender->AddMenuExtension(
-			"AssetEditorActions",
-			EExtensionHook::After,
-			GetToolkitCommands(),
-			FMenuExtensionDelegate::CreateStatic( &Local::AddPersonaAssetMenu )
-			);
-					
-	AddMenuExtender(PersonaMenuExtender);
-	AddMenuExtender(PersonaModule->GetMenuExtensibilityManager()->GetAllExtenders(GetToolkitCommands(), GetEditingObjects()));
-
 	check(PreviewComponent == NULL);
 
 	// Create the preview component
@@ -888,6 +634,7 @@ void FPersona::InitPersona(const EToolkitMode::Type Mode, const TSharedPtr< clas
 		}
 	}
 
+	UAnimBlueprint* AnimBlueprint = GetAnimBlueprint();
 	PreviewComponent->SetAnimInstanceClass(AnimBlueprint ? AnimBlueprint->GeneratedClass : NULL);
 
 	// We always want a preview instance unless we are using blueprints so that bone manipulation works
@@ -901,8 +648,8 @@ void FPersona::InitPersona(const EToolkitMode::Type Mode, const TSharedPtr< clas
 		AnimBlueprint->SetObjectBeingDebugged(PreviewComponent->AnimScriptInstance);
 	}
 
+	ExtendMenu();
 	ExtendDefaultPersonaToolbar();
-
 	RegenerateMenusAndToolbars();
 
 	// Activate the initial mode (which will populate with a real layout)
@@ -931,11 +678,54 @@ void FPersona::InitPersona(const EToolkitMode::Type Mode, const TSharedPtr< clas
 
 void FPersona::ExtendDefaultPersonaToolbar()
 {
-	TSharedPtr<FExtender> Extender = MakeShareable(new FExtender);
+	// If the ToolbarExtender is valid, remove it before rebuilding it
+	if(ToolbarExtender.IsValid())
+	{
+		RemoveToolbarExtender(ToolbarExtender);
+		ToolbarExtender.Reset();
+	}
+
+	ToolbarExtender = MakeShareable(new FExtender);
 
-	PersonaToolbar->SetupPersonaToolbar(Extender, SharedThis(this));
+	PersonaToolbar->SetupPersonaToolbar(ToolbarExtender, SharedThis(this));
 
-	AddToolbarExtender(Extender);
+	// extend extra menu/toolbars
+	struct Local
+	{
+		static void FillToolbar(FToolBarBuilder& ToolbarBuilder)
+		{
+			ToolbarBuilder.BeginSection("Person");
+			{
+				ToolbarBuilder.AddToolBarButton(FPersonaCommands::Get().ChangeSkeletonPreviewMesh, NAME_None, LOCTEXT("Toolbar_ChangePreviewMesh", "Set Preview"));
+				ToolbarBuilder.AddToolBarButton(FPersonaCommands::Get().AnimNotifyWindow);
+				ToolbarBuilder.AddToolBarButton(FPersonaCommands::Get().RetargetSourceMgr, NAME_None, LOCTEXT("Toolbar_RetargetSourceMgr", "Retarget Source"));
+				ToolbarBuilder.AddToolBarButton(FPersonaCommands::Get().ImportMesh);
+
+				ToolbarBuilder.AddToolBarButton(FPersonaCommands::Get().ReimportMesh);
+				//ToolbarBuilder.AddToolBarButton(FPersonaCommands::Get().ImportLODs);
+				//ToolbarBuilder.AddToolBarButton(FPersonaCommands::Get().AddBodyPart);
+
+				ToolbarBuilder.AddToolBarButton(FPersonaCommands::Get().ImportAnimation, NAME_None, LOCTEXT("Toolbar_ImportAnimation", "Import"));
+				ToolbarBuilder.AddToolBarButton(FPersonaCommands::Get().ReimportAnimation, NAME_None, LOCTEXT("Toolbar_ReimportAnimation", "Reimport"));
+				ToolbarBuilder.AddToolBarButton(FPersonaCommands::Get().ApplyCompression, NAME_None, LOCTEXT("Toolbar_ApplyCompression", "Compression"));
+				ToolbarBuilder.AddToolBarButton(FPersonaCommands::Get().ExportToFBX, NAME_None, LOCTEXT("Toolbar_ExportToFBX", "Export"));
+			}
+			ToolbarBuilder.EndSection();
+		}
+	};
+
+
+	ToolbarExtender->AddToolBarExtension(
+		"Asset",
+		EExtensionHook::After,
+		GetToolkitCommands(),
+		FToolBarExtensionDelegate::CreateStatic(&Local::FillToolbar)
+		);
+
+	AddToolbarExtender(ToolbarExtender);
+
+	FPersonaModule* PersonaModule = &FModuleManager::LoadModuleChecked<FPersonaModule>("Persona");
+	AddToolbarExtender(PersonaModule->GetToolBarExtensibilityManager()->GetAllExtenders(GetToolkitCommands(), GetEditingObjects()));
 }
 
 UBlueprint* FPersona::GetBlueprintObj() const
@@ -1099,21 +889,21 @@ void FPersona::CreateDefaultCommands()
 		FExecuteAction::CreateSP(this, &FPersona::OnApplyCompression),
 		FCanExecuteAction::CreateSP(this, &FPersona::HasValidAnimationSequencePlaying),
 		FIsActionChecked(),
-		FIsActionButtonVisible::CreateSP(this, &FPersona::IsInAnimationMode)
+		FIsActionButtonVisible::CreateSP(this, &FPersona::IsInPersonaMode, FPersonaModes::AnimationEditMode)
 		);
 
 	ToolkitCommands->MapAction(FPersonaCommands::Get().ExportToFBX,
 		FExecuteAction::CreateSP(this, &FPersona::OnExportToFBX),
 		FCanExecuteAction::CreateSP(this, &FPersona::HasValidAnimationSequencePlaying),
 		FIsActionChecked(),
-		FIsActionButtonVisible::CreateSP(this, &FPersona::IsInAnimationMode)
+		FIsActionButtonVisible::CreateSP(this, &FPersona::IsInPersonaMode, FPersonaModes::AnimationEditMode)
 		);
 
 	ToolkitCommands->MapAction(FPersonaCommands::Get().AddLoopingInterpolation,
 		FExecuteAction::CreateSP(this, &FPersona::OnAddLoopingInterpolation),
 		FCanExecuteAction::CreateSP(this, &FPersona::HasValidAnimationSequencePlaying),
 		FIsActionChecked(),
-		FIsActionButtonVisible::CreateSP(this, &FPersona::IsInAnimationMode)
+		FIsActionButtonVisible::CreateSP(this, &FPersona::IsInPersonaMode, FPersonaModes::AnimationEditMode)
 		);
 
 	ToolkitCommands->MapAction( FPersonaCommands::Get().ChangeSkeletonPreviewMesh,
@@ -1126,6 +916,65 @@ void FPersona::CreateDefaultCommands()
 		FCanExecuteAction::CreateSP( this, &FPersona::CanRemoveBones )
 		);
 
+	ToolkitCommands->MapAction(FPersonaCommands::Get().AnimNotifyWindow,
+		FExecuteAction::CreateSP(this, &FPersona::OnAnimNotifyWindow),
+		FCanExecuteAction(),
+		FIsActionChecked(),
+		FIsActionButtonVisible::CreateSP(this, &FPersona::IsInPersonaMode, FPersonaModes::SkeletonDisplayMode)
+		);
+
+	ToolkitCommands->MapAction(FPersonaCommands::Get().RetargetSourceMgr,
+		FExecuteAction::CreateSP(this, &FPersona::OnRetargetSourceMgr),
+		FCanExecuteAction(),
+		FIsActionChecked(),
+		FIsActionButtonVisible::CreateSP(this, &FPersona::IsInPersonaMode, FPersonaModes::SkeletonDisplayMode)
+		);
+
+	ToolkitCommands->MapAction(FPersonaCommands::Get().ImportMesh,
+		FExecuteAction::CreateSP(this, &FPersona::OnImportAsset, FBXIT_SkeletalMesh),
+		FCanExecuteAction(),
+		FIsActionChecked(),
+		FIsActionButtonVisible::CreateSP(this, &FPersona::IsInPersonaMode, FPersonaModes::SkeletonDisplayMode)
+		);
+
+	ToolkitCommands->MapAction(FPersonaCommands::Get().ReimportMesh,
+		FExecuteAction::CreateSP(this, &FPersona::OnReimportMesh),
+		FCanExecuteAction(),
+		FIsActionChecked(),
+		FIsActionButtonVisible::CreateSP(this, &FPersona::IsInPersonaMode, FPersonaModes::MeshEditMode)
+		);
+
+// 	ToolkitCommands->MapAction(FPersonaCommands::Get().ImportLODs,
+// 		FExecuteAction::CreateSP(this, &FPersona::OnImportLODs),
+// 		FCanExecuteAction(),
+// 		FIsActionChecked(),
+// 		FIsActionButtonVisible::CreateSP(this, &FPersona::IsInPersonaMode, FPersonaModes::MeshEditMode)
+// 		);
+
+// 	ToolkitCommands->MapAction(FPersonaCommands::Get().AddBodyPart,
+// 		FExecuteAction::CreateSP(this, &FPersona::OnAddBodyPart),
+// 		FCanExecuteAction(),
+// 		FIsActionChecked(),
+// 		FIsActionButtonVisible::CreateSP(this, &FPersona::IsInPersonaMode, FPersonaModes::MeshEditMode)
+// 		);
+
+
+	// animation menu options
+	// import animation
+	ToolkitCommands->MapAction(FPersonaCommands::Get().ImportAnimation,
+		FExecuteAction::CreateSP(this, &FPersona::OnImportAsset, FBXIT_Animation),
+		FCanExecuteAction(),
+		FIsActionChecked(),
+		FIsActionButtonVisible::CreateSP(this, &FPersona::IsInPersonaMode, FPersonaModes::AnimationEditMode)
+		);
+
+	ToolkitCommands->MapAction(FPersonaCommands::Get().ReimportAnimation,
+		FExecuteAction::CreateSP(this, &FPersona::OnReimportAnimation),
+		FCanExecuteAction::CreateSP(this, &FPersona::HasValidAnimationSequencePlaying),
+		FIsActionChecked(),
+		FIsActionButtonVisible::CreateSP(this, &FPersona::IsInPersonaMode, FPersonaModes::AnimationEditMode)
+		);
+
 	// Generic deletion
 	ToolkitCommands->MapAction(FGenericCommands::Get().Delete, FExecuteAction::CreateSP(this, &FPersona::OnCommandGenericDelete));
 
@@ -2595,9 +2444,9 @@ bool FPersona::HasValidAnimationSequencePlaying() const
 	return (AnimSequence != NULL);
 }
 
-bool FPersona::IsInAnimationMode() const
+bool FPersona::IsInPersonaMode(const FName InPersonaMode) const
 {
-	return (GetCurrentMode() == FPersonaModes::AnimationEditMode);
+	return (GetCurrentMode() == InPersonaMode);
 }
 
 void FPersona::ChangeSkeletonPreviewMesh()
@@ -3004,6 +2853,70 @@ void FAnimationRecorder::Record( USkeletalMeshComponent * Component, TArray<FTra
 	}
 }
 
+void FPersona::OnAnimNotifyWindow()
+{
+	TabManager->InvokeTab(FPersonaTabs::SkeletonAnimNotifiesID);
+}
+
+void FPersona::OnRetargetSourceMgr()
+{
+	TabManager->InvokeTab(FPersonaTabs::RetargetSourceManagerID);
+}
+
+void FPersona::OnImportAsset(enum EFBXImportType DefaultImportType)
+{
+	// open dialog
+	// get path
+	FString AssetPath;
+
+	TSharedRef<SImportPathDialog> NewAnimDlg =
+	SNew(SImportPathDialog);
+
+	if(NewAnimDlg->ShowModal() != EAppReturnType::Cancel)
+	{
+		AssetPath = NewAnimDlg->GetAssetPath();
+	
+		UFbxImportUI* ImportUI = ConstructObject<UFbxImportUI>(UFbxImportUI::StaticClass());
+		ImportUI->Skeleton = TargetSkeleton;
+		ImportUI->MeshTypeToImport = DefaultImportType;
+
+		FbxMeshUtils::SetImportOption(ImportUI);
+
+		// now I have to set skeleton on it. 
+		FAssetToolsModule& AssetToolsModule = FModuleManager::GetModuleChecked<FAssetToolsModule>("AssetTools");
+		TArray<UObject*> NewAssets = AssetToolsModule.Get().ImportAssets(AssetPath);
+
+		if ( NewAssets.Num() > 0 )
+		{
+
+		}
+		else
+		{
+			// error
+		}
+	}
+}
+
+void FPersona::OnReimportMesh()
+{
+	// Reimport the asset
+	UDebugSkelMeshComponent* PreviewComponent = GetPreviewMeshComponent();
+	if(PreviewComponent && PreviewComponent->SkeletalMesh)
+	{
+		FReimportManager::Instance()->Reimport(PreviewComponent->SkeletalMesh, true);
+	}
+}
+
+void FPersona::OnReimportAnimation()
+{
+	// Reimport the asset
+	UAnimSequence * AnimSequence = Cast<UAnimSequence> (GetAnimationAssetBeingEdited());
+	if (AnimSequence)
+	{
+		FReimportManager::Instance()->Reimport(AnimSequence, true);
+	}
+}
+
 static class FMeshHierarchyCmd : private FSelfRegisteringExec
 {
 public:
diff --git a/Engine/Source/Editor/Persona/Private/Persona.h b/Engine/Source/Editor/Persona/Private/Persona.h
--- a/Engine/Source/Editor/Persona/Private/Persona.h
+++ b/Engine/Source/Editor/Persona/Private/Persona.h
@@ -547,7 +547,7 @@ private:
 	void OnExportToFBX();
 	void OnAddLoopingInterpolation();
 	bool HasValidAnimationSequencePlaying() const;
-	bool IsInAnimationMode() const;
+	bool IsInPersonaMode(const FName InPersonaMode) const;
 
 	/** Change skeleton preview mesh functions */
 	void ChangeSkeletonPreviewMesh();
@@ -557,11 +557,24 @@ private:
 	void RemoveUnusedBones();
 	bool CanRemoveBones() const;
 
+	// tool bar actions
+	void OnAnimNotifyWindow();
+	void OnRetargetSourceMgr();
+	void OnReimportMesh();
+	void OnImportAsset(enum EFBXImportType DefaultImportType);
+	void OnReimportAnimation();
+
+	/** Extend menu and toolbar */
+	void ExtendMenu();
+
 	/** Returns the editor objects that are applicable for our current mode (e.g mesh, animation etc) */
 	TArray<UObject*> GetEditorObjectsForMode(FName Mode) const;
 
 	/** The extender to pass to the level editor to extend it's window menu */
-	TSharedPtr<FExtender> PersonaMenuExtender;
+	TSharedPtr<FExtender> MenuExtender;
+
+	/** Toolbar extender */
+	TSharedPtr<FExtender> ToolbarExtender;
 
 	/** Preview scene for the editor */
 	FPreviewScene PreviewScene;
diff --git a/Engine/Source/Editor/Persona/Private/PersonaCommands.cpp b/Engine/Source/Editor/Persona/Private/PersonaCommands.cpp
--- a/Engine/Source/Editor/Persona/Private/PersonaCommands.cpp
+++ b/Engine/Source/Editor/Persona/Private/PersonaCommands.cpp
@@ -9,11 +9,21 @@ void FPersonaCommands::RegisterCommands()
 	// skeleton menu
 	UI_COMMAND( ChangeSkeletonPreviewMesh, "Set Preview Mesh as Default", "Changes the skeletons default preview mesh to the current open preview mesh. The skeleton will require saving after this action.", EUserInterfaceActionType::Button, FInputGesture() );
 	UI_COMMAND( RemoveUnusedBones, "Remove Unused Bones from Skeleton", "Removes any bones from the skeleton that are not used by any of its meshes. The skeleton and associated animations will require saving after this action.", EUserInterfaceActionType::Button, FInputGesture() );
+	UI_COMMAND( AnimNotifyWindow, "Anim Notifies", "You can manage animation notifies that belong to the skeleton.", EUserInterfaceActionType::Button, FInputGesture() );
+	UI_COMMAND( RetargetSourceMgr, "Retarget Source Manager", "Create, delete, or edit retarget sources for this skeleton. ", EUserInterfaceActionType::Button, FInputGesture() );
+	UI_COMMAND( ImportMesh, "Import Mesh", "Import new mesh for this skeleton. ", EUserInterfaceActionType::Button, FInputGesture() );
+
+	// mesh menu
+	UI_COMMAND( ReimportMesh, "Reimport Mesh", "Reimport the current mesh.", EUserInterfaceActionType::Button, FInputGesture() );
+	UI_COMMAND( ImportLODs, "Import LODs", "Import LODs of the current mesh.", EUserInterfaceActionType::Button, FInputGesture() );
+	UI_COMMAND( AddBodyPart, "Add Body Part", "Add extra body parts for the current mesh", EUserInterfaceActionType::Button, FInputGesture());
 
 	// animation menu
+	// import animation
+	UI_COMMAND( ImportAnimation, "Import Animation", "Import new animation for the skeleton.", EUserInterfaceActionType::Button, FInputGesture());
+	UI_COMMAND( ReimportAnimation, "Reimport Animation", "Reimport current animation.", EUserInterfaceActionType::Button, FInputGesture());
 	UI_COMMAND( ApplyCompression, "Apply Compression", "Apply cvompression to current animation", EUserInterfaceActionType::Button, FInputGesture() );
 	UI_COMMAND( RecordAnimation, "Record to new Animation", "Create new animation from currently playing", EUserInterfaceActionType::Button, FInputGesture() );
 	UI_COMMAND( ExportToFBX, "Export to FBX", "Export current animation to FBX", EUserInterfaceActionType::Button, FInputGesture() );
 	UI_COMMAND( AddLoopingInterpolation, "Add Looping Interpolation", "Add an extra first frame at the end of the frame to create interpolation when looping", EUserInterfaceActionType::Button, FInputGesture() );
-
 }
diff --git a/Engine/Source/Editor/Persona/Private/PersonaCommands.h b/Engine/Source/Editor/Persona/Private/PersonaCommands.h
--- a/Engine/Source/Editor/Persona/Private/PersonaCommands.h
+++ b/Engine/Source/Editor/Persona/Private/PersonaCommands.h
@@ -17,20 +17,34 @@ public:
 	// skeleton menu options
 	// Command to allow users to set the skeletons preview mesh
 	TSharedPtr<FUICommandInfo> ChangeSkeletonPreviewMesh;
-
 	// Command to allow users to remove unused bones (not referenced by any skeletalmesh) from the skeleton
 	TSharedPtr<FUICommandInfo> RemoveUnusedBones;
+	// Command to show Anim Notify window
+	TSharedPtr<FUICommandInfo> AnimNotifyWindow;
+	// Command to show Retarget Source Manager
+	TSharedPtr<FUICommandInfo> RetargetSourceMgr;
+	// Import Mesh for this Skeleton
+	TSharedPtr<FUICommandInfo> ImportMesh;
+			
+	// mesh menu options
+	// reimport current mesh
+	TSharedPtr<FUICommandInfo> ReimportMesh;
+	// import LODs
+	TSharedPtr<FUICommandInfo> ImportLODs;
+	// import Body part
+	TSharedPtr<FUICommandInfo> AddBodyPart;
 
 	// animation menu options
+	// import animation
+	TSharedPtr<FUICommandInfo> ImportAnimation;
+	// reimport animation
+	TSharedPtr<FUICommandInfo> ReimportAnimation;
 	// record animation 
  	TSharedPtr<FUICommandInfo> RecordAnimation;
-
 	// apply compression
 	TSharedPtr<FUICommandInfo> ApplyCompression;
-
 	// export to FBX
 	TSharedPtr<FUICommandInfo> ExportToFBX;
-
 	// Add looping interpolation
 	TSharedPtr<FUICommandInfo> AddLoopingInterpolation;
 };
diff --git a/Engine/Source/Editor/Persona/Private/PersonaMeshDetails.cpp b/Engine/Source/Editor/Persona/Private/PersonaMeshDetails.cpp
--- a/Engine/Source/Editor/Persona/Private/PersonaMeshDetails.cpp
+++ b/Engine/Source/Editor/Persona/Private/PersonaMeshDetails.cpp
@@ -827,7 +827,7 @@ void FPersonaMeshDetails::CustomizeDetails( IDetailLayoutBuilder& DetailLayout )
 	CustomizeClothingProperties(DetailLayout,ClothingCategory);
 #endif// #if WITH_APEX_CLOTHING
 
-	IDetailCategoryBuilder& AdditionalMeshCategory = DetailLayout.EditCategory("Additional Meshes", LOCTEXT("AdditionalMeshesCollapsable", "Additional Meshes").ToString(), ECategoryPriority::TypeSpecific);
+	IDetailCategoryBuilder& AdditionalMeshCategory = DetailLayout.EditCategory("AdditionalBodyPart", LOCTEXT("AdditionalMeshesCollapsable", "Additional Body Part").ToString(), ECategoryPriority::TypeSpecific);
 	AdditionalMeshCategory.AddCustomRow("")
 	[
 		SNew(SAdditionalMeshesEditor, PersonaPtr)
diff --git a/Engine/Source/Editor/Persona/Private/PersonaModule.cpp b/Engine/Source/Editor/Persona/Private/PersonaModule.cpp
--- a/Engine/Source/Editor/Persona/Private/PersonaModule.cpp
+++ b/Engine/Source/Editor/Persona/Private/PersonaModule.cpp
@@ -21,6 +21,7 @@ const FName PersonaAppName = FName(TEXT("PersonaApp"));
 void FPersonaModule::StartupModule()
 {
 	MenuExtensibilityManager = MakeShareable(new FExtensibilityManager);
+	ToolBarExtensibilityManager = MakeShareable(new FExtensibilityManager);
 
 	// Load all blueprint animnotifies from asset registry so they are available from drop downs in anim segment detail views
 	{
@@ -68,6 +69,7 @@ void FPersonaModule::StartupModule()
 void FPersonaModule::ShutdownModule()
 {
 	MenuExtensibilityManager.Reset();
+	ToolBarExtensibilityManager.Reset();
 
 	// unregsiter when shut down
 	if(FModuleManager::Get().IsModuleLoaded("PropertyEditor"))
diff --git a/Engine/Source/Editor/Persona/Public/PersonaModule.h b/Engine/Source/Editor/Persona/Public/PersonaModule.h
--- a/Engine/Source/Editor/Persona/Public/PersonaModule.h
+++ b/Engine/Source/Editor/Persona/Public/PersonaModule.h
@@ -47,9 +47,10 @@ public:
 
 	/** Gets the extensibility managers for outside entities to extend persona editor's menus and toolbars */
 	virtual TSharedPtr<FExtensibilityManager> GetMenuExtensibilityManager() {return MenuExtensibilityManager;}
-
+	virtual TSharedPtr<FExtensibilityManager> GetToolBarExtensibilityManager() {return ToolBarExtensibilityManager;}
 private:
 	TSharedPtr<FExtensibilityManager> MenuExtensibilityManager;
+	TSharedPtr<FExtensibilityManager> ToolBarExtensibilityManager;
 };
 
 
diff --git a/Engine/Source/Editor/UnrealEd/Classes/Factories/FbxImportUI.h b/Engine/Source/Editor/UnrealEd/Classes/Factories/FbxImportUI.h
--- a/Engine/Source/Editor/UnrealEd/Classes/Factories/FbxImportUI.h
+++ b/Engine/Source/Editor/UnrealEd/Classes/Factories/FbxImportUI.h
@@ -22,7 +22,7 @@ enum EFBXImportType
 	FBXIT_MAX,
 };
 
-UCLASS(config=EditorUserSettings, AutoExpandCategories=(General, USkeletalMesh, UStaticMesh, Materials), HideCategories=Object)
+UCLASS(config=EditorUserSettings, AutoExpandCategories=(General, USkeletalMesh, UStaticMesh, Materials), HideCategories=Object, MinimalAPI)
 class UFbxImportUI : public UObject
 {
 	GENERATED_UCLASS_BODY()
diff --git a/Engine/Source/Editor/UnrealEd/Private/FbxMeshUtils.cpp b/Engine/Source/Editor/UnrealEd/Private/FbxMeshUtils.cpp
--- a/Engine/Source/Editor/UnrealEd/Private/FbxMeshUtils.cpp
+++ b/Engine/Source/Editor/UnrealEd/Private/FbxMeshUtils.cpp
@@ -626,6 +626,13 @@ namespace FbxMeshUtils
 			}
 		}
 	}
+
+	void SetImportOption(UFbxImportUI* ImportUI)
+	{
+		UnFbx::FFbxImporter* FFbxImporter = UnFbx::FFbxImporter::GetInstance();
+		UnFbx::FBXImportOptions* ImportOptions = FFbxImporter->GetImportOptions();
+		ApplyImportUIToImportOptions(ImportUI, *ImportOptions);
+	}
 }  //end namespace MeshUtils
 
 #undef LOCTEXT_NAMESPACE
\ No newline at end of file
diff --git a/Engine/Source/Editor/UnrealEd/Public/FbxMeshUtils.h b/Engine/Source/Editor/UnrealEd/Public/FbxMeshUtils.h
--- a/Engine/Source/Editor/UnrealEd/Public/FbxMeshUtils.h
+++ b/Engine/Source/Editor/UnrealEd/Public/FbxMeshUtils.h
@@ -75,4 +75,11 @@ namespace FbxMeshUtils
 	 */
 	UNREALED_API void ImportMeshLODDialog( class UObject* Mesh, int32 LODLevel );
 
+	/**
+	 * Sets import option before importing
+	 *
+	 * @param ImportUI			The importUI you'd like to apply to
+	 */
+	UNREALED_API void SetImportOption(UFbxImportUI* ImportUI);
+
 }  //end namespace ExportMeshUtils
