Missing files

diff --git a/Engine/Source/Editor/Persona/Private/SAnimationDlgs.cpp b/Engine/Source/Editor/Persona/Private/SAnimationDlgs.cpp
new file mode 100644
--- /dev/null
+++ b/Engine/Source/Editor/Persona/Private/SAnimationDlgs.cpp
@@ -0,0 +1,362 @@
+// Copyright 1998-2014 Epic Games, Inc. All Rights Reserved.
+
+
+#include "PersonaPrivatePCH.h"
+#include "Persona.h"
+#include "SAnimationDlgs.h"
+#include "Editor/ContentBrowser/Public/ContentBrowserModule.h"
+
+#define LOCTEXT_NAMESPACE "SAnimationDlgs"
+
+/////////////////////////////////////////////////////
+/////////////////////////////////////////////////////
+/////////////////////////////////////////////////////
+/////////////////////////////////////////////////////
+
+/////////////////////////////////////////////////////
+// Create Animation dialog for recording animation
+/////////////////////////////////////////////////////
+
+FText SCreateAnimationDlg::LastUsedAssetPath;
+FText SImportPathDialog::LastUsedAssetPath;
+
+void SCreateAnimationDlg::Construct(const FArguments& InArgs)
+{
+	AssetPath = FText::FromString(FPackageName::GetLongPackagePath(InArgs._DefaultAssetPath.ToString()));
+	AssetName = FText::FromString(FPackageName::GetLongPackageAssetName(InArgs._DefaultAssetPath.ToString()));
+
+	if (AssetPath.IsEmpty())
+	{
+		AssetPath = LastUsedAssetPath;
+	}
+	else
+	{
+		LastUsedAssetPath = AssetPath;
+	}
+
+	FPathPickerConfig PathPickerConfig;
+	PathPickerConfig.DefaultPath = AssetPath.ToString();
+	PathPickerConfig.OnPathSelected = FOnPathSelected::CreateSP(this, &SCreateAnimationDlg::OnPathChange);
+	PathPickerConfig.bAddDefaultPath = true;
+
+	FContentBrowserModule& ContentBrowserModule = FModuleManager::LoadModuleChecked<FContentBrowserModule>("ContentBrowser");
+
+	SWindow::Construct(SWindow::FArguments()
+		.Title(LOCTEXT("SCreateAnimationDlg_Title", "Create New Animation Object"))
+		.SupportsMinimize(false)
+		.SupportsMaximize(false)
+		//.SizingRule( ESizingRule::Autosized )
+		.ClientSize(FVector2D(450,450))
+		[
+			SNew(SVerticalBox)
+
+			+ SVerticalBox::Slot() // Add user input block
+			.Padding(2)
+			[
+				SNew(SBorder)
+				.BorderImage(FEditorStyle::GetBrush("ToolPanel.GroupBorder"))
+				[
+					SNew(SVerticalBox)
+
+					+SVerticalBox::Slot()
+					.AutoHeight()
+					[
+						SNew(STextBlock)
+						.Text(LOCTEXT("SelectPath", "Select Path to create animation"))
+						.Font( FSlateFontInfo( FPaths::EngineContentDir() / TEXT("Slate/Fonts/Roboto-Regular.ttf"), 14 ) )
+					]
+
+					+SVerticalBox::Slot()
+					.FillHeight(1)
+					.Padding(3)
+					[
+						ContentBrowserModule.Get().CreatePathPicker(PathPickerConfig)
+					]
+
+					+SVerticalBox::Slot()
+					.AutoHeight()
+					[
+						SNew(SSeparator)
+					]
+
+					+SVerticalBox::Slot()
+					.AutoHeight()
+					.Padding(3)
+					[
+						SNew(SHorizontalBox)
+
+						+SHorizontalBox::Slot()
+						.AutoWidth()
+						.Padding(0, 0, 10, 0)
+						.VAlign(VAlign_Center)
+						[
+							SNew(STextBlock)
+							.Text(LOCTEXT("AnimationName", "Animation Name"))
+						]
+
+						+SHorizontalBox::Slot()
+						[
+							SNew(SEditableTextBox)
+							.Text(AssetName)
+							.OnTextCommitted(this, &SCreateAnimationDlg::OnNameChange)
+							.MinDesiredWidth(250)
+						]
+					]
+
+					+SVerticalBox::Slot()
+					.AutoHeight()
+					.Padding(3)
+					[
+						SNew(SHorizontalBox)
+
+						+SHorizontalBox::Slot()
+						.AutoWidth()
+						.Padding(0, 0, 10, 0)
+						.VAlign(VAlign_Center)
+						[
+							SNew(STextBlock)
+							.Text(LOCTEXT("DurationInSec", "Duration (in sec) "))
+						]
+
+						+SHorizontalBox::Slot()
+						[
+							SNew(SEditableTextBox)
+							.Text(Duration)
+							.OnTextCommitted(this, &SCreateAnimationDlg::OnDurationChange)
+							.MinDesiredWidth(100)
+						]
+					]
+				]
+			]
+
+			+SVerticalBox::Slot()
+			.AutoHeight()
+			.HAlign(HAlign_Right)
+			.Padding(5)
+			[
+				SNew(SUniformGridPanel)
+				.SlotPadding(FEditorStyle::GetMargin("StandardDialog.SlotPadding"))
+				.MinDesiredSlotWidth(FEditorStyle::GetFloat("StandardDialog.MinDesiredSlotWidth"))
+				.MinDesiredSlotHeight(FEditorStyle::GetFloat("StandardDialog.MinDesiredSlotHeight"))
+				+SUniformGridPanel::Slot(0,0)
+				[
+					SNew(SButton) 
+					.HAlign(HAlign_Center)
+					.ContentPadding(FEditorStyle::GetMargin("StandardDialog.ContentPadding"))
+					.Text(LOCTEXT("OK", "OK"))
+					.OnClicked(this, &SCreateAnimationDlg::OnButtonClick, EAppReturnType::Ok)
+				]
+				+SUniformGridPanel::Slot(1,0)
+					[
+						SNew(SButton) 
+						.HAlign(HAlign_Center)
+						.ContentPadding(FEditorStyle::GetMargin("StandardDialog.ContentPadding"))
+						.Text(LOCTEXT("Cancel", "Cancel"))
+						.OnClicked(this, &SCreateAnimationDlg::OnButtonClick, EAppReturnType::Cancel)
+					]
+			]
+		]);
+}
+
+void SCreateAnimationDlg::OnDurationChange(const FText& InDuration, ETextCommit::Type CommitInfo)
+{
+	Duration = InDuration;
+}
+
+void SCreateAnimationDlg::OnNameChange(const FText& NewName, ETextCommit::Type CommitInfo)
+{
+	AssetName = NewName;
+}
+
+void SCreateAnimationDlg::OnPathChange(const FString& NewPath)
+{
+	AssetPath = FText::FromString(NewPath);
+	LastUsedAssetPath = AssetPath;
+}
+
+FReply SCreateAnimationDlg::OnButtonClick(EAppReturnType::Type ButtonID)
+{
+	UserResponse = ButtonID;
+
+	if (ButtonID != EAppReturnType::Cancel)
+	{
+		if (!ValidatePackage())
+		{
+			// reject the request
+			return FReply::Handled();
+		}
+	}
+
+	RequestDestroyWindow();
+
+	return FReply::Handled();
+}
+
+/** Ensures supplied package name information is valid */
+bool SCreateAnimationDlg::ValidatePackage()
+{
+	FText Reason;
+	FString FullPath = GetFullAssetPath();
+
+	if(!FPackageName::IsValidLongPackageName(FullPath, false, &Reason)
+		|| !FName(*AssetName.ToString()).IsValidObjectName(Reason))
+	{
+		FMessageDialog::Open(EAppMsgType::Ok, Reason);
+		return false;
+	}
+
+	if (GetDuration() <= 0.0f)
+	{
+		FMessageDialog::Open(EAppMsgType::Ok, LOCTEXT("InvalidDuration", "Invalid duration"));
+		return false;
+	}
+
+	return true;
+}
+
+EAppReturnType::Type SCreateAnimationDlg::ShowModal()
+{
+	GEditor->EditorAddModalWindow(SharedThis(this));
+	return UserResponse;
+}
+
+FString SCreateAnimationDlg::GetAssetPath()
+{
+	return AssetPath.ToString();
+}
+
+FString SCreateAnimationDlg::GetAssetName()
+{
+	return AssetName.ToString();
+}
+
+FString SCreateAnimationDlg::GetFullAssetPath()
+{
+	return AssetPath.ToString() + "/" + AssetName.ToString();
+}
+
+float SCreateAnimationDlg::GetDuration()
+{
+	return FCString::Atof(*Duration.ToString());
+}
+
+/////////////////////////////////////////////////////
+// create path picker for importing assets
+/////////////////////////////////////////////////////
+
+void SImportPathDialog::Construct(const FArguments& InArgs)
+{
+	AssetPath = FText::FromString(FPackageName::GetLongPackagePath(InArgs._DefaultAssetPath.ToString()));
+
+	if(AssetPath.IsEmpty())
+	{
+		AssetPath = LastUsedAssetPath;
+	}
+	else
+	{
+		LastUsedAssetPath = AssetPath;
+	}
+
+	FPathPickerConfig PathPickerConfig;
+	PathPickerConfig.DefaultPath = AssetPath.ToString();
+	PathPickerConfig.OnPathSelected = FOnPathSelected::CreateSP(this, &SImportPathDialog::OnPathChange);
+	PathPickerConfig.bAddDefaultPath = true;
+
+	FContentBrowserModule& ContentBrowserModule = FModuleManager::LoadModuleChecked<FContentBrowserModule>("ContentBrowser");
+
+	SWindow::Construct(SWindow::FArguments()
+		.Title(LOCTEXT("SImportPathDialog_Title", "Select folder to import to"))
+		.SupportsMinimize(false)
+		.SupportsMaximize(false)
+		//.SizingRule( ESizingRule::Autosized )
+		.ClientSize(FVector2D(450, 450))
+		[
+			SNew(SVerticalBox)
+
+			+ SVerticalBox::Slot() // Add user input block
+			.Padding(2)
+			[
+				SNew(SBorder)
+				.BorderImage(FEditorStyle::GetBrush("ToolPanel.GroupBorder"))
+				[
+					SNew(SVerticalBox)
+
+					+SVerticalBox::Slot()
+					.AutoHeight()
+					[
+						SNew(STextBlock)
+						.Text(LOCTEXT("SelectPath", "Select Path to create animation"))
+						.Font(FSlateFontInfo(FPaths::EngineContentDir() / TEXT("Slate/Fonts/Roboto-Regular.ttf"), 14))
+					]
+
+					+SVerticalBox::Slot()
+					.FillHeight(1)
+					.Padding(3)
+					[
+						ContentBrowserModule.Get().CreatePathPicker(PathPickerConfig)
+					]
+
+					+SVerticalBox::Slot()
+					.AutoHeight()
+					[
+						SNew(SSeparator)
+					]
+				]
+			]
+
+			+SVerticalBox::Slot()
+			.AutoHeight()
+			.HAlign(HAlign_Right)
+			.Padding(5)
+			[
+				SNew(SUniformGridPanel)
+				.SlotPadding(FEditorStyle::GetMargin("StandardDialog.SlotPadding"))
+				.MinDesiredSlotWidth(FEditorStyle::GetFloat("StandardDialog.MinDesiredSlotWidth"))
+				.MinDesiredSlotHeight(FEditorStyle::GetFloat("StandardDialog.MinDesiredSlotHeight"))
+				+SUniformGridPanel::Slot(0, 0)
+				[
+					SNew(SButton)
+					.HAlign(HAlign_Center)
+					.ContentPadding(FEditorStyle::GetMargin("StandardDialog.ContentPadding"))
+					.Text(LOCTEXT("OK", "OK"))
+					.OnClicked(this, &SImportPathDialog::OnButtonClick, EAppReturnType::Ok)
+				]
+				+SUniformGridPanel::Slot(1, 0)
+				[
+					SNew(SButton)
+					.HAlign(HAlign_Center)
+					.ContentPadding(FEditorStyle::GetMargin("StandardDialog.ContentPadding"))
+					.Text(LOCTEXT("Cancel", "Cancel"))
+					.OnClicked(this, &SImportPathDialog::OnButtonClick, EAppReturnType::Cancel)
+				]
+			]
+		]);
+}
+
+void SImportPathDialog::OnPathChange(const FString& NewPath)
+{
+	AssetPath = FText::FromString(NewPath);
+	LastUsedAssetPath = AssetPath;
+}
+
+FReply SImportPathDialog::OnButtonClick(EAppReturnType::Type ButtonID)
+{
+	UserResponse = ButtonID;
+	RequestDestroyWindow();
+
+	return FReply::Handled();
+}
+
+EAppReturnType::Type SImportPathDialog::ShowModal()
+{
+	GEditor->EditorAddModalWindow(SharedThis(this));
+	return UserResponse;
+}
+
+FString SImportPathDialog::GetAssetPath()
+{
+	return AssetPath.ToString();
+}
+
+#undef LOCTEXT_NAMESPACE
+
diff --git a/Engine/Source/Editor/Persona/Private/SAnimationDlgs.h b/Engine/Source/Editor/Persona/Private/SAnimationDlgs.h
new file mode 100644
--- /dev/null
+++ b/Engine/Source/Editor/Persona/Private/SAnimationDlgs.h
@@ -0,0 +1,90 @@
+// Copyright 1998-2014 Epic Games, Inc. All Rights Reserved.
+
+#pragma once
+
+class SCreateAnimationDlg: public SWindow
+{
+public:
+	SLATE_BEGIN_ARGS(SCreateAnimationDlg)
+	{
+	}
+	
+	SLATE_ARGUMENT(FText, DefaultAssetPath)
+	SLATE_END_ARGS()
+
+		SCreateAnimationDlg()
+		: UserResponse(EAppReturnType::Cancel)
+	{
+		}
+
+	void Construct(const FArguments& InArgs);
+
+public:
+	/** Displays the dialog in a blocking fashion */
+	EAppReturnType::Type ShowModal();
+
+	/** Gets the resulting asset path */
+	FString GetAssetPath();
+
+	/** Gets the resulting asset name */
+	FString GetAssetName();
+
+	/** Gets the resulting full asset path (path+'/'+name) */
+	FString GetFullAssetPath();
+
+	/** Get Duration **/
+	float GetDuration();
+
+protected:
+	void OnPathChange(const FString& NewPath);
+	void OnNameChange(const FText& NewName, ETextCommit::Type CommitInfo);
+	void OnDurationChange(const FText& NewName, ETextCommit::Type CommitInfo);
+	FReply OnButtonClick(EAppReturnType::Type ButtonID);
+
+	bool ValidatePackage();
+
+	EAppReturnType::Type UserResponse;
+	FText AssetPath;
+	FText AssetName;
+	FText Duration;
+
+	static FText LastUsedAssetPath;
+};
+
+class SImportPathDialog: public SWindow
+{
+public:
+	SLATE_BEGIN_ARGS(SImportPathDialog)
+	{
+	}
+
+	SLATE_ARGUMENT(FText, DefaultAssetPath)
+	SLATE_END_ARGS()
+
+	SImportPathDialog()
+		: UserResponse(EAppReturnType::Cancel)
+	{
+		}
+
+	void Construct(const FArguments& InArgs);
+
+public:
+	/** Displays the dialog in a blocking fashion */
+	EAppReturnType::Type ShowModal();
+
+	/** Gets the resulting asset path */
+	FString GetAssetPath();
+
+protected:
+	void OnPathChange(const FString& NewPath);
+	FReply OnButtonClick(EAppReturnType::Type ButtonID);
+
+	EAppReturnType::Type UserResponse;
+	FText AssetPath;
+	FText AssetName;
+	FText Duration;
+
+	static FText LastUsedAssetPath;
+};
+
+
