Added ability for native and blueprint notifies to dynamically set custom names in the notify panel.

diff --git a/Engine/Source/Editor/Persona/Private/SAnimNotifyPanel.cpp b/Engine/Source/Editor/Persona/Private/SAnimNotifyPanel.cpp
--- a/Engine/Source/Editor/Persona/Private/SAnimNotifyPanel.cpp
+++ b/Engine/Source/Editor/Persona/Private/SAnimNotifyPanel.cpp
@@ -810,6 +810,16 @@ void SAnimNotifyNode::Construct(const FArguments& InArgs)
 	bDrawTooltipToRight = true;
 	bSelected = false;
 
+	// Cache notify name for blueprint / Native notifies.
+	if(NotifyEvent->Notify)
+	{
+		NotifyEvent->NotifyName = FName(*NotifyEvent->Notify->GetNotifyName());
+	}
+	else if(NotifyEvent->NotifyStateClass)
+	{
+		NotifyEvent->NotifyName = FName(*NotifyEvent->NotifyStateClass->GetNotifyName());
+	}
+
 	OnNodeDragStarted = InArgs._OnNodeDragStarted;
 	PanTrackRequest = InArgs._PanTrackRequest;
 
@@ -1741,6 +1751,11 @@ void SAnimNotifyTrack::CreateNewNotifyAtCursor(FString NewNotifyName, UClass* No
 		if( NewEvent.NotifyStateClass )
 		{
 			NewEvent.Duration = 1 / 30.f;
+			NewEvent.NotifyName = FName(*NewEvent.NotifyStateClass->GetNotifyName());
+		}
+		else
+		{
+			NewEvent.NotifyName = FName(*NewEvent.Notify->GetNotifyName());
 		}
 	}
 	else
@@ -2762,6 +2777,9 @@ void SAnimNotifyPanel::Construct(const FArguments& InArgs)
 		]
 	];
 
+	OnPropertyChangedHandle = FCoreDelegates::FOnObjectPropertyChanged::FDelegate::CreateSP(this, &SAnimNotifyPanel::OnPropertyChanged);
+	FCoreDelegates::OnObjectPropertyChanged.Add(OnPropertyChangedHandle);
+
 	Update();
 }
 
@@ -2773,6 +2791,8 @@ SAnimNotifyPanel::~SAnimNotifyPanel()
 		PersonaPtr.Pin()->UnregisterOnPostUndo(this);
 		PersonaPtr.Pin()->UnregisterOnGenericDelete(this);
 	}
+
+	FCoreDelegates::OnObjectPropertyChanged.Remove(OnPropertyChangedHandle);
 }
 
 FReply SAnimNotifyPanel::InsertTrack(int32 TrackIndexToInsert)
@@ -3209,4 +3229,20 @@ void SAnimNotifyPanel::OnPasteNotifies(SAnimNotifyTrack* RequestTrack, float Cli
 	}
 }
 
+void SAnimNotifyPanel::OnPropertyChanged(UObject* ChangedObject, FPropertyChangedEvent& PropertyEvent)
+{
+	// Don't process if it's an interactive change; wait till we receive the final event.
+	if(PropertyEvent.ChangeType != EPropertyChangeType::Interactive)
+	{
+		for(FAnimNotifyEvent& Event : Sequence->Notifies)
+		{
+			if(Event.Notify == ChangedObject || Event.NotifyStateClass == ChangedObject)
+			{
+				// If we've changed a notify present in the sequence, refresh our tracks.
+				RefreshNotifyTracks();
+			}
+		}
+	}
+}
+
 #undef LOCTEXT_NAMESPACE
diff --git a/Engine/Source/Editor/Persona/Private/SAnimNotifyPanel.h b/Engine/Source/Editor/Persona/Private/SAnimNotifyPanel.h
--- a/Engine/Source/Editor/Persona/Private/SAnimNotifyPanel.h
+++ b/Engine/Source/Editor/Persona/Private/SAnimNotifyPanel.h
@@ -111,6 +111,10 @@ public:
 	void CopySelectedNotifiesToClipboard() const;
 	void OnPasteNotifies(SAnimNotifyTrack* RequestTrack, float ClickTime, ENotifyPasteMode::Type PasteMode, ENotifyPasteMultipleMode::Type MultiplePasteType);
 
+	/** Handler for properties changing on objects */
+	FCoreDelegates::FOnObjectPropertyChanged::FDelegate OnPropertyChangedHandle;
+	void OnPropertyChanged(UObject* ChangedObject, FPropertyChangedEvent& PropertyEvent);
+
 private:
 	TSharedPtr<SBorder> PanelArea;
 	class UAnimSequenceBase* Sequence;
diff --git a/Engine/Source/Runtime/Engine/Classes/Animation/AnimNotifies/AnimNotify.h b/Engine/Source/Runtime/Engine/Classes/Animation/AnimNotifies/AnimNotify.h
--- a/Engine/Source/Runtime/Engine/Classes/Animation/AnimNotifies/AnimNotify.h
+++ b/Engine/Source/Runtime/Engine/Classes/Animation/AnimNotifies/AnimNotify.h
@@ -14,6 +14,12 @@ class ENGINE_API UAnimNotify : public UObject
 {
 	GENERATED_UCLASS_BODY()
 
+	/** 
+	 * Implementable event to get a custom name for the notify
+	 */
+	UFUNCTION(BlueprintNativeEvent)
+	FString GetNotifyName() const;
+
 	UFUNCTION(BlueprintImplementableEvent)
 	virtual bool Received_Notify(USkeletalMeshComponent* MeshComp, UAnimSequenceBase* Animation) const;
 
diff --git a/Engine/Source/Runtime/Engine/Classes/Animation/AnimNotifies/AnimNotifyState.h b/Engine/Source/Runtime/Engine/Classes/Animation/AnimNotifies/AnimNotifyState.h
--- a/Engine/Source/Runtime/Engine/Classes/Animation/AnimNotifies/AnimNotifyState.h
+++ b/Engine/Source/Runtime/Engine/Classes/Animation/AnimNotifies/AnimNotifyState.h
@@ -14,6 +14,12 @@ class ENGINE_API UAnimNotifyState : public UObject
 {
 	GENERATED_UCLASS_BODY()
 
+	/** 
+	 * Implementable event to get a custom name for the notify
+	 */
+	UFUNCTION(BlueprintNativeEvent)
+	FString GetNotifyName() const;
+
 	UFUNCTION(BlueprintImplementableEvent)
 	virtual bool Received_NotifyBegin(USkeletalMeshComponent * MeshComp, UAnimSequenceBase * Animation) const;
 	
diff --git a/Engine/Source/Runtime/Engine/Classes/Animation/AnimNotifies/AnimNotifyState_TimedParticleEffect.h b/Engine/Source/Runtime/Engine/Classes/Animation/AnimNotifies/AnimNotifyState_TimedParticleEffect.h
--- a/Engine/Source/Runtime/Engine/Classes/Animation/AnimNotifies/AnimNotifyState_TimedParticleEffect.h
+++ b/Engine/Source/Runtime/Engine/Classes/Animation/AnimNotifies/AnimNotifyState_TimedParticleEffect.h
@@ -52,6 +52,9 @@ class UAnimNotifyState_TimedParticleEffect : public UAnimNotifyState
 	virtual void NotifyTick(class USkeletalMeshComponent * MeshComp, class UAnimSequenceBase * Animation, float FrameDeltaTime) override;
 	virtual void NotifyEnd(class USkeletalMeshComponent * MeshComp, class UAnimSequenceBase * Animation) override;
 
+	// Overriden from UAnimNotifyState to provide custom notify name.
+	virtual FString GetNotifyName_Implementation() const override;
+
 private:
 	bool ValidateParameters(USkeletalMeshComponent* MeshComp);
 };
\ No newline at end of file
diff --git a/Engine/Source/Runtime/Engine/Private/Animation/AnimNotify.cpp b/Engine/Source/Runtime/Engine/Private/Animation/AnimNotify.cpp
--- a/Engine/Source/Runtime/Engine/Private/Animation/AnimNotify.cpp
+++ b/Engine/Source/Runtime/Engine/Private/Animation/AnimNotify.cpp
@@ -29,3 +29,24 @@ class UWorld* UAnimNotify::GetWorld() const
 {
 	return (MeshContext ? MeshContext->GetWorld() : NULL);
 }
+
+FString UAnimNotify::GetNotifyName_Implementation() const
+{
+	UObject* ClassGeneratedBy = GetClass()->ClassGeneratedBy;
+	FString NotifyName;
+
+	if(ClassGeneratedBy)
+	{
+		// GeneratedBy will be valid for blueprint types and gives a clean name without a suffix
+		NotifyName = ClassGeneratedBy->GetName();
+	}
+	else
+	{
+		// Native notify classes are clean without a suffix otherwise
+		NotifyName = GetClass()->GetName();
+	}
+
+	NotifyName = NotifyName.Replace(TEXT("AnimNotify_"), TEXT(""), ESearchCase::CaseSensitive);
+	
+	return NotifyName;
+}
diff --git a/Engine/Source/Runtime/Engine/Private/Animation/AnimNotifyState.cpp b/Engine/Source/Runtime/Engine/Private/Animation/AnimNotifyState.cpp
--- a/Engine/Source/Runtime/Engine/Private/Animation/AnimNotifyState.cpp
+++ b/Engine/Source/Runtime/Engine/Private/Animation/AnimNotifyState.cpp
@@ -31,3 +31,24 @@ void UAnimNotifyState::NotifyEnd(class USkeletalMeshComponent * MeshComp, class
 	Received_NotifyEnd(MeshComp, Animation);
 }
 
+FString UAnimNotifyState::GetNotifyName_Implementation() const
+{
+	UObject* ClassGeneratedBy = GetClass()->ClassGeneratedBy;
+	FString NotifyName;
+
+	if(ClassGeneratedBy)
+	{
+		// GeneratedBy will be valid for blueprint types and gives a clean name without a suffix
+		NotifyName = ClassGeneratedBy->GetName();
+	}
+	else
+	{
+		// Native notify classes are clean without a suffix otherwise
+		NotifyName = GetClass()->GetName();
+	}
+
+	NotifyName = NotifyName.Replace(TEXT("AnimNotifyState_"), TEXT(""), ESearchCase::CaseSensitive);
+	
+	return NotifyName;
+}
+
diff --git a/Engine/Source/Runtime/Engine/Private/Animation/AnimNotifyState_TimedParticleEffect.cpp b/Engine/Source/Runtime/Engine/Private/Animation/AnimNotifyState_TimedParticleEffect.cpp
--- a/Engine/Source/Runtime/Engine/Private/Animation/AnimNotifyState_TimedParticleEffect.cpp
+++ b/Engine/Source/Runtime/Engine/Private/Animation/AnimNotifyState_TimedParticleEffect.cpp
@@ -92,6 +92,16 @@ bool UAnimNotifyState_TimedParticleEffect::ValidateParameters(USkeletalMeshCompo
 	return bValid;
 }
 
+FString UAnimNotifyState_TimedParticleEffect::GetNotifyName_Implementation() const
+{
+	if(PSTemplate)
+	{
+		return PSTemplate->GetName();
+	}
+
+	return UAnimNotifyState::GetNotifyName_Implementation();
+}
+
 #if WITH_EDITORONLY_DATA
 void UAnimNotifyState_TimedParticleEffect::PreEditChange(UProperty* PropertyAboutToChange)
 {
