Fix Clang compile error.

diff --git a/Engine/Source/Editor/Persona/Private/Persona.cpp b/Engine/Source/Editor/Persona/Private/Persona.cpp
--- a/Engine/Source/Editor/Persona/Private/Persona.cpp
+++ b/Engine/Source/Editor/Persona/Private/Persona.cpp
@@ -766,7 +766,7 @@ void FPersona::ExtendDefaultPersonaToolbar()
 							FIsActionChecked(),
 							FIsActionButtonVisible::CreateSP(PersonaPtr, &FPersona::IsInPersonaMode, FPersonaModes::AnimationEditMode)
 							),
-						FOnGetContent::CreateSP(PersonaPtr, &GenerateCreateAssetMenu, Skeleton),
+						FOnGetContent::CreateSP(PersonaPtr, &FPersona::GenerateCreateAssetMenu, Skeleton),
 						LOCTEXT("OpenBlueprint_Label", "Create Asset"),
 						LOCTEXT("OpenBlueprint_ToolTip", "Create Assets for this skeleton."),
 						FSlateIcon(FEditorStyle::GetStyleSetName(), "Persona.CreateAsset")
