Enabled physics based selection of bones if a persona viewport click fails to find a hit proxy

diff --git a/Engine/Source/Editor/Persona/Private/AnimationEditorViewportClient.cpp b/Engine/Source/Editor/Persona/Private/AnimationEditorViewportClient.cpp
--- a/Engine/Source/Editor/Persona/Private/AnimationEditorViewportClient.cpp
+++ b/Engine/Source/Editor/Persona/Private/AnimationEditorViewportClient.cpp
@@ -75,6 +75,7 @@ FAnimationViewportClient::FAnimationViewportClient( FPreviewScene& InPreviewScen
 	, SelectedWindActor(NULL)
 	, PrevWindStrength(0.0f)
 	, GravityScaleSliderValue(0.25f)
+	, BodyTraceDistance(100000.0f)
 {
 	// load config
 	ConfigOption = UPersonaOptions::StaticClass()->GetDefaultObject<UPersonaOptions>();
@@ -378,6 +379,19 @@ void FAnimationViewportClient::SetPreviewMeshComponent(UDebugSkelMeshComponent*
 	PreviewSkelMeshComp->BonesOfInterest.Empty();
 
 	UpdateCameraSetup();
+
+	// Setup physics data from physics assets if available, clearing any physics setup on the component
+	UPhysicsAsset* PhysAsset = PreviewSkelMeshComp->GetPhysicsAsset();
+	if(PhysAsset)
+	{
+		PhysAsset->InvalidateAllPhysicsMeshes();
+		PreviewSkelMeshComp->TermArticulated();
+		PreviewSkelMeshComp->InitArticulated(GetWorld()->GetPhysicsScene());
+
+		// Set to block all to enable tracing.
+		PreviewSkelMeshComp->SetCollisionProfileName(UCollisionProfile::BlockAll_ProfileName);
+		PreviewSkelMeshComp->RefreshBoneTransforms();
+	}
 }
 
 void FAnimationViewportClient::Draw(const FSceneView* View, FPrimitiveDrawInterface* PDI)
@@ -979,6 +993,14 @@ TWeakObjectPtr<AWindDirectionalSource> FAnimationViewportClient::FindSelectedWin
 
 void FAnimationViewportClient::ProcessClick(class FSceneView& View, class HHitProxy* HitProxy, FKey Key, EInputEvent Event, uint32 HitX, uint32 HitY)
 {
+	TSharedPtr<FPersona> SharedPersona = PersonaPtr.Pin();
+
+	if(!SharedPersona.IsValid())
+	{
+		// No reason to continue, cannot select elements
+		return;
+	}
+
 	if ( HitProxy )
 	{
 		if ( HitProxy->IsA( HPersonaSocketProxy::StaticGetType() ) )
@@ -986,20 +1008,13 @@ void FAnimationViewportClient::ProcessClick(class FSceneView& View, class HHitPr
 			HPersonaSocketProxy* SocketProxy = ( HPersonaSocketProxy* )HitProxy;
 
 			// Tell Persona that the socket has been selected - this will sort out the skeleton tree, etc.
-			if ( PersonaPtr.IsValid() )
-			{
-				PersonaPtr.Pin()->SetSelectedSocket( SocketProxy->SocketInfo );
-			}
+			SharedPersona->SetSelectedSocket( SocketProxy->SocketInfo );
 		}
 		else if ( HitProxy->IsA( HPersonaBoneProxy::StaticGetType() ) )
 		{
 			HPersonaBoneProxy* BoneProxy = ( HPersonaBoneProxy* )HitProxy;
-
 			// Tell Persona that the bone has been selected - this will sort out the skeleton tree, etc.
-			if ( PersonaPtr.IsValid() )
-			{
-				PersonaPtr.Pin()->SetSelectedBone( PreviewSkelMeshComp.Get()->SkeletalMesh->Skeleton, BoneProxy->BoneName );
-			}
+			SharedPersona->SetSelectedBone( PreviewSkelMeshComp.Get()->SkeletalMesh->Skeleton, BoneProxy->BoneName );
 		}
 		else if ( HitProxy->IsA( HActor::StaticGetType() ) )
 		{
@@ -1009,17 +1024,26 @@ void FAnimationViewportClient::ProcessClick(class FSceneView& View, class HHitPr
 			if( WindActor )
 			{
 				//clear previously selected things
-				PersonaPtr.Pin()->DeselectAll();
+				SharedPersona->DeselectAll();
 				SelectedWindActor = WindActor;
 			}
 		}
 	}
 	else
 	{
-		// Clicking outside of a hit proxy should de-select things
-		if ( PersonaPtr.IsValid() )
+		// Cast for phys bodies if we didn't get any hit proxies
+		FHitResult Result(1.0f);
+		const FViewportClick Click(&View, this, EKeys::Invalid, IE_Released, Viewport->GetMouseX(), Viewport->GetMouseY());
+		bool bHit = PreviewSkelMeshComp.Get()->LineTraceComponent(Result, Click.GetOrigin(), Click.GetOrigin() + Click.GetDirection() * BodyTraceDistance, FCollisionQueryParams(true));
+		
+		if(bHit)
 		{
-			PersonaPtr.Pin()->DeselectAll();
+			SharedPersona->SetSelectedBone(PreviewSkelMeshComp->SkeletalMesh->Skeleton, Result.BoneName);
+		}
+		else
+		{
+			// We didn't hit a proxy or a physics object, so deselect all objects
+			SharedPersona->DeselectAll();
 		}
 }
 }
diff --git a/Engine/Source/Editor/Persona/Private/AnimationEditorViewportClient.h b/Engine/Source/Editor/Persona/Private/AnimationEditorViewportClient.h
--- a/Engine/Source/Editor/Persona/Private/AnimationEditorViewportClient.h
+++ b/Engine/Source/Editor/Persona/Private/AnimationEditorViewportClient.h
@@ -262,6 +262,9 @@ private:
 
 	/** Wind actor used for preview */
 	TWeakObjectPtr<AWindDirectionalSource> SelectedWindActor;
+
+	/** Distance to trace for physics bodies */
+	const float BodyTraceDistance;
 private:
 	int32 FindSelectedBone() const;
 	class USkeletalMeshSocket* FindSelectedSocket() const;
