Added animated preview viewports to Persona asset browser tooltips

diff --git a/Engine/Source/Editor/ContentBrowser/Private/AssetViewWidgets.cpp b/Engine/Source/Editor/ContentBrowser/Private/AssetViewWidgets.cpp
--- a/Engine/Source/Editor/ContentBrowser/Private/AssetViewWidgets.cpp
+++ b/Engine/Source/Editor/ContentBrowser/Private/AssetViewWidgets.cpp
@@ -131,6 +131,8 @@ void SAssetViewItem::Construct( const FArguments& InArgs )
 	OnAssetsDragDropped = InArgs._OnAssetsDragDropped;
 	OnPathsDragDropped = InArgs._OnPathsDragDropped;
 	OnFilesDragDropped = InArgs._OnFilesDragDropped;
+	OnGetCustomAssetToolTip = InArgs._OnGetCustomAssetToolTip;
+	OnVisualizeAssetToolTip = InArgs._OnVisualizeAssetToolTip;
 
 	bDraggedOver = false;
 
@@ -428,7 +430,12 @@ TSharedRef<SToolTip> SAssetViewItem::CreateToolTipWidget() const
 {
 	if ( AssetItem.IsValid() )
 	{
-		if(AssetItem->GetType() != EAssetItemType::Folder)
+		if(OnGetCustomAssetToolTip.IsBound() && AssetItem->GetType() != EAssetItemType::Folder)
+		{
+			FAssetData& AssetData = StaticCastSharedPtr<FAssetViewAsset>(AssetItem)->Data;
+			return OnGetCustomAssetToolTip.Execute(AssetData);
+		}
+		else if(AssetItem->GetType() != EAssetItemType::Folder)
 		{
 			const FAssetData& AssetData = StaticCastSharedPtr<FAssetViewAsset>(AssetItem)->Data;
 			UClass* AssetClass = FindObject<UClass>(ANY_PACKAGE, *AssetData.AssetClass.ToString());
@@ -892,6 +899,18 @@ void SAssetViewItem::HandleAssetLoaded(UObject* InAsset) const
 	}
 }
 
+bool SAssetViewItem::OnVisualizeTooltip(const TSharedPtr<SWidget>& TooltipContent)
+{
+	if(OnVisualizeAssetToolTip.IsBound() && TooltipContent.IsValid() && AssetItem->GetType() != EAssetItemType::Folder)
+	{
+		FAssetData& AssetData = StaticCastSharedPtr<FAssetViewAsset>(AssetItem)->Data;
+		return OnVisualizeAssetToolTip.Execute(TooltipContent, AssetData);
+	}
+
+	// No custom behaviour, return false to allow slate to visualize the widget
+	return false;
+}
+
 ///////////////////////////////
 // SAssetListItem
 ///////////////////////////////
@@ -916,6 +935,8 @@ void SAssetListItem::Construct( const FArguments& InArgs )
 		.OnAssetsDragDropped(InArgs._OnAssetsDragDropped)
 		.OnPathsDragDropped(InArgs._OnPathsDragDropped)
 		.OnFilesDragDropped(InArgs._OnFilesDragDropped)
+		.OnGetCustomAssetToolTip(InArgs._OnGetCustomAssetToolTip)
+		.OnVisualizeAssetToolTip(InArgs._OnVisualizeAssetToolTip)
 		);
 
 	AssetThumbnail = InArgs._AssetThumbnail;
@@ -1117,6 +1138,8 @@ void SAssetTileItem::Construct( const FArguments& InArgs )
 		.OnAssetsDragDropped(InArgs._OnAssetsDragDropped)
 		.OnPathsDragDropped(InArgs._OnPathsDragDropped)
 		.OnFilesDragDropped(InArgs._OnFilesDragDropped)
+		.OnGetCustomAssetToolTip(InArgs._OnGetCustomAssetToolTip)
+		.OnVisualizeAssetToolTip(InArgs._OnVisualizeAssetToolTip)
 		);
 
 	AssetThumbnail = InArgs._AssetThumbnail;
@@ -1306,6 +1329,8 @@ void SAssetColumnItem::Construct( const FArguments& InArgs )
 		.OnAssetsDragDropped(InArgs._OnAssetsDragDropped)
 		.OnPathsDragDropped(InArgs._OnPathsDragDropped)
 		.OnFilesDragDropped(InArgs._OnFilesDragDropped)
+		.OnGetCustomAssetToolTip(InArgs._OnGetCustomAssetToolTip)
+		.OnVisualizeAssetToolTip(InArgs._OnVisualizeAssetToolTip)
 		);
 	
 	HighlightText = InArgs._HighlightText;
@@ -1339,8 +1364,6 @@ TSharedRef<SWidget> SAssetColumnItem::GenerateWidgetForColumn( const FName& Colu
 		const float IconOverlaySize = IconBrush->ImageSize.X * 0.6f;
 
 		Content = SNew(SHorizontalBox)
-			.ToolTipText(this, &SAssetColumnItem::GetAssetNameToolTipText)
-
 			// Icon
 			+SHorizontalBox::Slot()
 			.AutoWidth()
@@ -1427,6 +1450,7 @@ TSharedRef<SWidget> SAssetColumnItem::GenerateWidgetForColumn( const FName& Colu
 
 	return SNew(SBox)
 		.Padding(FMargin(0, 0, 6, 0)) // Add a little right padding so text from this column does not run directly into text from the next.
+		.ToolTip(CreateToolTipWidget())
 		[
 			Content.ToSharedRef()
 		];
@@ -1446,6 +1470,8 @@ void SAssetColumnItem::OnAssetDataChanged()
 	{
 		PathText->SetText( GetAssetPathText() );
 	}
+
+	SetToolTip( CreateToolTipWidget() );
 }
 
 FString SAssetColumnItem::GetAssetNameToolTipText() const
diff --git a/Engine/Source/Editor/ContentBrowser/Private/AssetViewWidgets.h b/Engine/Source/Editor/ContentBrowser/Private/AssetViewWidgets.h
--- a/Engine/Source/Editor/ContentBrowser/Private/AssetViewWidgets.h
+++ b/Engine/Source/Editor/ContentBrowser/Private/AssetViewWidgets.h
@@ -82,6 +82,12 @@ public:
 		/** Delegate for when a list of files is dropped on this folder (if it is a folder) from an external source */
 		SLATE_EVENT( FOnFilesDragDropped, OnFilesDragDropped )
 
+		/** Delegate to call (if bound) to get a custom tooltip for this view item */
+		SLATE_EVENT( FOnGetCustomAssetToolTip, OnGetCustomAssetToolTip )
+
+		/** Delegate for when an item is about to show a tool tip */
+		SLATE_EVENT( FOnVisualizeAssetToolTip, OnVisualizeAssetToolTip)
+
 	SLATE_END_ARGS()
 
 	/** Virtual destructor */
@@ -109,6 +115,9 @@ public:
 	/** Delegate handling when an asset is loaded */
 	void HandleAssetLoaded(UObject* InAsset) const;
 
+	/* About to show a tool tip */
+	virtual bool OnVisualizeTooltip( const TSharedPtr<SWidget>& TooltipContent ) override;
+
 protected:
 	/** Handles starting a name change */
 	virtual void HandleBeginNameChange( const FText& OriginalText );
@@ -204,6 +213,12 @@ protected:
 	/** Called when any asset item is destroyed. Used in thumbnail management */
 	FOnItemDestroyed OnItemDestroyed;
 
+	/** Called if bound to get a custom asset item tooltip */
+	FOnGetCustomAssetToolTip OnGetCustomAssetToolTip;
+
+	/** Called if bound when about to show a tooltip */
+	FOnVisualizeAssetToolTip OnVisualizeAssetToolTip;
+
 	/** The geometry last frame. Used when telling popup messages where to appear. */
 	FGeometry LastGeometry;
 
@@ -320,6 +335,12 @@ public:
 		/** Delegate for when a list of files is dropped on this folder (if it is a folder) from an external source */
 		SLATE_EVENT( FOnFilesDragDropped, OnFilesDragDropped )
 
+		/** Delegate to request a custom tool tip if necessary */
+		SLATE_EVENT(FOnGetCustomAssetToolTip, OnGetCustomAssetToolTip)
+
+		/* Delegate to signal when the item is about to show a tooltip */
+		SLATE_EVENT(FOnVisualizeAssetToolTip, OnVisualizeAssetToolTip)
+
 	SLATE_END_ARGS()
 
 	/** Destructor */
@@ -417,6 +438,12 @@ public:
 		/** Delegate for when a list of files is dropped on this folder (if it is a folder) from an external source */
 		SLATE_EVENT( FOnFilesDragDropped, OnFilesDragDropped )
 
+		/** Delegate to request a custom tool tip if necessary */
+		SLATE_EVENT(FOnGetCustomAssetToolTip, OnGetCustomAssetToolTip)
+
+		/* Delegate to signal when the item is about to show a tooltip */
+		SLATE_EVENT(FOnVisualizeAssetToolTip, OnVisualizeAssetToolTip)
+
 	SLATE_END_ARGS()
 
 	/** Destructor */
@@ -482,6 +509,12 @@ public:
 		/** Delegate for when a list of files is dropped on this folder (if it is a folder) from an external source */
 		SLATE_EVENT( FOnFilesDragDropped, OnFilesDragDropped )
 
+		/** Delegate to request a custom tool tip if necessary */
+		SLATE_EVENT( FOnGetCustomAssetToolTip, OnGetCustomAssetToolTip)
+
+		/* Delegate to signal when the item is about to show a tooltip */
+		SLATE_EVENT( FOnVisualizeAssetToolTip, OnVisualizeAssetToolTip)
+
 	SLATE_END_ARGS()
 
 	/** Constructs this widget with InArgs */
@@ -555,5 +588,12 @@ public:
 		this->AssetColumnItem->Tick(AllottedGeometry, InCurrentTime, InDeltaTime);
 	}
 
+	virtual bool OnVisualizeTooltip(const TSharedPtr<SWidget>& TooltipContent)
+	{
+		// We take the content from the asset column item during construction,
+		// so let the item handle the tooltip callback
+		return AssetColumnItem->OnVisualizeTooltip(TooltipContent);
+	}
+
 	TSharedPtr<SAssetColumnItem> AssetColumnItem;
 };
diff --git a/Engine/Source/Editor/ContentBrowser/Private/SAssetPicker.cpp b/Engine/Source/Editor/ContentBrowser/Private/SAssetPicker.cpp
--- a/Engine/Source/Editor/ContentBrowser/Private/SAssetPicker.cpp
+++ b/Engine/Source/Editor/ContentBrowser/Private/SAssetPicker.cpp
@@ -193,6 +193,8 @@ void SAssetPicker::Construct( const FArguments& InArgs )
 		.OnAssetSelected(InArgs._AssetPickerConfig.OnAssetSelected)
 		.OnAssetsActivated(this, &SAssetPicker::HandleAssetsActivated)
 		.OnGetAssetContextMenu(InArgs._AssetPickerConfig.OnGetAssetContextMenu)
+		.OnGetCustomAssetToolTip(InArgs._AssetPickerConfig.OnGetCustomAssetToolTip)
+		.OnVisualizeAssetToolTip(InArgs._AssetPickerConfig.OnVisualizeAssetToolTip)
 		.AreRealTimeThumbnailsAllowed(this, &SAssetPicker::IsHovered)
 		.FrontendFilters(FrontendFilters)
 		.InitialSourcesData(CurrentSourcesData)
diff --git a/Engine/Source/Editor/ContentBrowser/Private/SAssetView.cpp b/Engine/Source/Editor/ContentBrowser/Private/SAssetView.cpp
--- a/Engine/Source/Editor/ContentBrowser/Private/SAssetView.cpp
+++ b/Engine/Source/Editor/ContentBrowser/Private/SAssetView.cpp
@@ -159,6 +159,8 @@ void SAssetView::Construct( const FArguments& InArgs )
 	OnFindInAssetTreeRequested = InArgs._OnFindInAssetTreeRequested;
 	OnAssetRenameCommitted = InArgs._OnAssetRenameCommitted;
 	OnAssetTagWantsToBeDisplayed = InArgs._OnAssetTagWantsToBeDisplayed;
+	OnGetCustomAssetToolTip = InArgs._OnGetCustomAssetToolTip;
+	OnVisualizeAssetToolTip = InArgs._OnVisualizeAssetToolTip;
 	HighlightedText = InArgs._HighlightedText;
 	ThumbnailLabel = InArgs._ThumbnailLabel;
 	AllowThumbnailHintLabel = InArgs._AllowThumbnailHintLabel;
@@ -2663,7 +2665,9 @@ TSharedRef<ITableRow> SAssetView::MakeListViewWidget(TSharedPtr<FAssetViewItem>
 			.ThumbnailLabel( ThumbnailLabel )
 			.ThumbnailHintColorAndOpacity( this, &SAssetView::GetThumbnailHintColorAndOpacity )
 			.AllowThumbnailHintLabel( AllowThumbnailHintLabel )
-			.IsSelected( FIsSelected::CreateSP(TableRowWidget.Get(), &STableRow<TSharedPtr<FAssetViewItem>>::IsSelectedExclusively) );
+			.IsSelected( FIsSelected::CreateSP(TableRowWidget.Get(), &STableRow<TSharedPtr<FAssetViewItem>>::IsSelectedExclusively) )
+			.OnGetCustomAssetToolTip(OnGetCustomAssetToolTip)
+			.OnVisualizeAssetToolTip(OnVisualizeAssetToolTip);
 
 		TableRowWidget->SetContent(Item);
 
@@ -2748,7 +2752,9 @@ TSharedRef<ITableRow> SAssetView::MakeTileViewWidget(TSharedPtr<FAssetViewItem>
 			.ThumbnailLabel( ThumbnailLabel )
 			.ThumbnailHintColorAndOpacity( this, &SAssetView::GetThumbnailHintColorAndOpacity )
 			.AllowThumbnailHintLabel( AllowThumbnailHintLabel )
-			.IsSelected( FIsSelected::CreateSP(TableRowWidget.Get(), &STableRow<TSharedPtr<FAssetViewItem>>::IsSelectedExclusively) );
+			.IsSelected( FIsSelected::CreateSP(TableRowWidget.Get(), &STableRow<TSharedPtr<FAssetViewItem>>::IsSelectedExclusively) )
+			.OnGetCustomAssetToolTip(OnGetCustomAssetToolTip)
+			.OnVisualizeAssetToolTip(OnVisualizeAssetToolTip);
 
 		TableRowWidget->SetContent(Item);
 
@@ -2779,6 +2785,8 @@ TSharedRef<ITableRow> SAssetView::MakeColumnViewWidget(TSharedPtr<FAssetViewItem
 				.OnAssetsDragDropped(this, &SAssetView::OnAssetsDragDropped)
 				.OnPathsDragDropped(this, &SAssetView::OnPathsDragDropped)
 				.OnFilesDragDropped(this, &SAssetView::OnFilesDragDropped)
+				.OnGetCustomAssetToolTip(OnGetCustomAssetToolTip)
+				.OnVisualizeAssetToolTip(OnVisualizeAssetToolTip)
 		);
 }
 
diff --git a/Engine/Source/Editor/ContentBrowser/Private/SAssetView.h b/Engine/Source/Editor/ContentBrowser/Private/SAssetView.h
--- a/Engine/Source/Editor/ContentBrowser/Private/SAssetView.h
+++ b/Engine/Source/Editor/ContentBrowser/Private/SAssetView.h
@@ -57,6 +57,12 @@ public:
 		/** Called when the user has committed a rename of one or more assets */
 		SLATE_EVENT( FOnAssetRenameCommitted, OnAssetRenameCommitted )
 
+		/** Called to get a custom asset item tool tip (if necessary) */
+		SLATE_EVENT( FOnGetCustomAssetToolTip, OnGetCustomAssetToolTip )
+
+		/** Called when an asset item is about to show a tooltip */
+		SLATE_EVENT( FOnVisualizeAssetToolTip, OnVisualizeAssetToolTip )
+
 		/** The warning text to display when there are no assets to show */
 		SLATE_ATTRIBUTE( FText, AssetShowWarningText )
 
@@ -672,6 +678,12 @@ private:
 
 	/** Called to check if an asset tag should be display in details view. */
 	FOnShouldDisplayAssetTag OnAssetTagWantsToBeDisplayed;
+
+	/** Called to get a custom asset item tooltip (If necessary) */
+	FOnGetCustomAssetToolTip OnGetCustomAssetToolTip;
+
+	/** Called when a custom asset item is about to show a tooltip */
+	FOnVisualizeAssetToolTip OnVisualizeAssetToolTip;
 	
 	/** When true, filtered list items will be sorted next tick. Provided another sort hasn't happened recently or we are renaming an asset */
 	bool bPendingSortFilteredItems;
diff --git a/Engine/Source/Editor/ContentBrowser/Public/ContentBrowserDelegates.h b/Engine/Source/Editor/ContentBrowser/Public/ContentBrowserDelegates.h
--- a/Engine/Source/Editor/ContentBrowser/Public/ContentBrowserDelegates.h
+++ b/Engine/Source/Editor/ContentBrowser/Public/ContentBrowserDelegates.h
@@ -70,4 +70,10 @@ DECLARE_DELEGATE_RetVal_OneParam(TSharedRef<FExtender>, FContentBrowserMenuExten
 DECLARE_DELEGATE_RetVal_OneParam(TSharedRef<FExtender>, FContentBrowserMenuExtender_SelectedPaths, const TArray<FString>& /*SelectedPaths*/);
 
 /** Called to request the menu when right clicking on an asset */
-DECLARE_DELEGATE_RetVal_ThreeParams(TSharedPtr<SWidget>, FOnGetFolderContextMenu, const TArray<FString>& /*SelectedPaths*/, FContentBrowserMenuExtender_SelectedPaths /*MenuExtender*/, FOnCreateNewFolder /*CreationDelegate*/);
\ No newline at end of file
+DECLARE_DELEGATE_RetVal_ThreeParams(TSharedPtr<SWidget>, FOnGetFolderContextMenu, const TArray<FString>& /*SelectedPaths*/, FContentBrowserMenuExtender_SelectedPaths /*MenuExtender*/, FOnCreateNewFolder /*CreationDelegate*/);
+
+/** Called to request a custom asset item tooltip */
+DECLARE_DELEGATE_RetVal_OneParam( TSharedRef<SToolTip>, FOnGetCustomAssetToolTip, class FAssetData& /*AssetData*/);
+
+/** Called when an asset item visualizes its tooltip */
+DECLARE_DELEGATE_RetVal_TwoParams(bool, FOnVisualizeAssetToolTip, const TSharedPtr<SWidget>& /*ToolTipContent*/, class FAssetData& /*AssetData*/);
\ No newline at end of file
diff --git a/Engine/Source/Editor/ContentBrowser/Public/IContentBrowserSingleton.h b/Engine/Source/Editor/ContentBrowser/Public/IContentBrowserSingleton.h
--- a/Engine/Source/Editor/ContentBrowser/Public/IContentBrowserSingleton.h
+++ b/Engine/Source/Editor/ContentBrowser/Public/IContentBrowserSingleton.h
@@ -74,6 +74,12 @@ struct FAssetPickerConfig
 	/** The delegate that fires when an asset is right clicked and a context menu is requested */
 	FOnGetAssetContextMenu OnGetAssetContextMenu;
 
+	/** Fired when an asset item is constructed and a tooltip is requested. If unbound the item will use the default widget */
+	FOnGetCustomAssetToolTip OnGetCustomAssetToolTip;
+
+	/** Fired when an asset item is about to show it's tool tip */
+	FOnVisualizeAssetToolTip OnVisualizeAssetToolTip;
+
 	/** If more detailed filtering is required than simply Filter, this delegate will get fired for every asset to determine if it should be culled. */
 	FOnShouldFilterAsset OnShouldFilterAsset;
 
diff --git a/Engine/Source/Editor/Persona/Private/AnimationMode/AnimationMode.cpp b/Engine/Source/Editor/Persona/Private/AnimationMode/AnimationMode.cpp
--- a/Engine/Source/Editor/Persona/Private/AnimationMode/AnimationMode.cpp
+++ b/Engine/Source/Editor/Persona/Private/AnimationMode/AnimationMode.cpp
@@ -122,15 +122,7 @@ TSharedRef<SWidget> FAnimAssetPropertiesSummoner::CreateTabBody(const FWorkflowT
 {
 	TSharedPtr<FPersona> PersonaApp = StaticCastSharedPtr<FPersona>(HostingApp.Pin());
 
-	return SNew(SOverlay)
-		+SOverlay::Slot()
-		[
-			SNew(SAnimAssetPropertiesTabBody, PersonaApp)
-		]
-		+SOverlay::Slot()
-		[
-			SNew(SAnimDifferentAssetBeingPreviewedWarning, PersonaApp)
-		];
+	return SNew(SAnimAssetPropertiesTabBody, PersonaApp);
 }
 
 /////////////////////////////////////////////////////
@@ -215,50 +207,4 @@ FAnimEditAppMode::FAnimEditAppMode(TSharedPtr<FPersona> InPersona)
 		);
 }
 
-/////////////////////////////////////////////////////
-// SAnimDifferentAssetBeingPreviewedWarning
-
-void SAnimDifferentAssetBeingPreviewedWarning::Construct(const FArguments& InArgs, TSharedPtr<FPersona> InPersona)
-{
-	PersonaPtr = InPersona;
-
-	FSlateFontInfo BoldFont = FEditorStyle::GetFontStyle(TEXT("NotificationList.FontBold"));
-
-	Visibility = EVisibility::HitTestInvisible;
-
-	ChildSlot
-	[
-		SNew(SBorder)
-		.Padding(20.0f)
-		.Visibility(this, &SAnimDifferentAssetBeingPreviewedWarning::GetVisibility)
-		.BorderImage(FEditorStyle::GetBrush("NotificationList.ItemBackground"))
-		.BorderBackgroundColor(FLinearColor(1.0f, 1.0f, 1.0f, 0.62f))
-		.HAlign(HAlign_Center)
-		.VAlign(VAlign_Center)
-		[
-			SNew(STextBlock)
-			.Font(BoldFont)
-			.Text(LOCTEXT("AssetPreviewMismatch", "Currently previewing\nanother asset"))
-		]
-	];
-}
-
-EVisibility SAnimDifferentAssetBeingPreviewedWarning::GetVisibility() const
-{
-	UObject* AssetBeingEdited = PersonaPtr.Pin()->GetAnimationAssetBeingEdited();
-	UObject* AssetBeingPreviewed = PersonaPtr.Pin()->GetPreviewAnimationAsset();
-
-	if (AssetBeingEdited != NULL)
-	{
-		if ((AssetBeingPreviewed != NULL) && (AssetBeingPreviewed != AssetBeingEdited))
-		{
-			return EVisibility::HitTestInvisible;
-		}
-	}
-	
-	return EVisibility::Collapsed;
-}
-
-/////////////////////////////////////////////////////
-
 #undef LOCTEXT_NAMESPACE
\ No newline at end of file
diff --git a/Engine/Source/Editor/Persona/Private/AnimationMode/AnimationMode.h b/Engine/Source/Editor/Persona/Private/AnimationMode/AnimationMode.h
--- a/Engine/Source/Editor/Persona/Private/AnimationMode/AnimationMode.h
+++ b/Engine/Source/Editor/Persona/Private/AnimationMode/AnimationMode.h
@@ -10,21 +10,3 @@ class FAnimEditAppMode : public FPersonaAppMode
 public:
 	FAnimEditAppMode(TSharedPtr<class FPersona> InPersona);
 };
-
-/////////////////////////////////////////////////////
-// SAnimDifferentAssetBeingPreviewedWarning
-
-class SAnimDifferentAssetBeingPreviewedWarning : public SCompoundWidget
-{
-public:
-	SLATE_BEGIN_ARGS(SAnimDifferentAssetBeingPreviewedWarning) {}
-	SLATE_END_ARGS()
-
-	void Construct(const FArguments& InArgs, TSharedPtr<FPersona> InPersona);
-
-private:
-	// Pointer back to owning Persona instance (the keeper of state)
-	TWeakPtr<class FPersona> PersonaPtr;
-private:
-	EVisibility GetVisibility() const;
-};
\ No newline at end of file
diff --git a/Engine/Source/Editor/Persona/Private/Persona.cpp b/Engine/Source/Editor/Persona/Private/Persona.cpp
--- a/Engine/Source/Editor/Persona/Private/Persona.cpp
+++ b/Engine/Source/Editor/Persona/Private/Persona.cpp
@@ -276,17 +276,6 @@ TSharedPtr<SWidget> FPersona::CreateEditorWidgetForAnimDocument(UObject* InAnimA
 
 	if (Result.IsValid())
 	{
-		// Wrap the document editor with a potential preview mismatch warning
-		Result = SNew(SOverlay)
-			+SOverlay::Slot()
-			[
-				Result.ToSharedRef()
-			]
-			+SOverlay::Slot()
-			[
-				SNew(SAnimDifferentAssetBeingPreviewedWarning, SharedThis(this))
-			];
-
 		InAnimAsset->SetFlags(RF_Transactional);
 	}
 
diff --git a/Engine/Source/Editor/Persona/Private/SAnimationSequenceBrowser.cpp b/Engine/Source/Editor/Persona/Private/SAnimationSequenceBrowser.cpp
--- a/Engine/Source/Editor/Persona/Private/SAnimationSequenceBrowser.cpp
+++ b/Engine/Source/Editor/Persona/Private/SAnimationSequenceBrowser.cpp
@@ -8,10 +8,13 @@
 #include "Persona.h"
 #include "AssetRegistryModule.h"
 #include "SSkeletonWidget.h"
+#include "Editor/ContentBrowser/Public/ContentBrowserModule.h"
 #include "FeedbackContextEditor.h"
 #include "EditorAnimUtils.h"
 #include "Editor/ContentBrowser/Public/FrontendFilterBase.h"
 #include "Runtime/AssetRegistry/Public/AssetRegistryModule.h"
+#include "SceneViewport.h"
+#include "AnimPreviewInstance.h"
 
 #define LOCTEXT_NAMESPACE "SequenceBrowser"
 
@@ -40,29 +43,27 @@ const int32 SAnimationSequenceBrowser::MaxAssetsHistory = 10;
 
 SAnimationSequenceBrowser::~SAnimationSequenceBrowser()
 {
-	if (PersonaPtr.IsValid())
-	{
-		PersonaPtr.Pin()->SetSequenceBrowser(NULL);
-	}
-}
-
-void SAnimationSequenceBrowser::OnAnimSelected(const FAssetData& AssetData)
-{
-	CacheOriginalAnimAssetHistory();
-	TSharedPtr<FPersona> Persona = PersonaPtr.Pin();
-	if (Persona.IsValid())
+	if(PreviewComponent)
 	{
-		if (UObject* RawAsset = AssetData.GetAsset())
+		for(int32 ComponentIdx = PreviewComponent->AttachChildren.Num() - 1 ; ComponentIdx >= 0 ; --ComponentIdx)
 		{
-			if (UAnimationAsset* Asset = Cast<UAnimationAsset>(RawAsset))
-			{
-				Persona->SetPreviewAnimationAsset(Asset);
-			}
-			else if(UVertexAnimation* VertexAnim = Cast<UVertexAnimation>(RawAsset))
+			USceneComponent* Component = PreviewComponent->AttachChildren[ComponentIdx];
+			if(Component)
 			{
-				Persona->SetPreviewVertexAnim(VertexAnim);
+				CleanupPreviewSceneComponent(Component);
 			}
 		}
+		PreviewComponent->AttachChildren.Empty();
+	}
+
+	if(ViewportClient.IsValid())
+	{
+		ViewportClient->Viewport = NULL;
+	}
+
+	if (PersonaPtr.IsValid())
+	{
+		PersonaPtr.Pin()->SetSequenceBrowser(NULL);
 	}
 }
 
@@ -330,6 +331,8 @@ void SAnimationSequenceBrowser::Construct(const FArguments& InArgs)
 
 	FContentBrowserModule& ContentBrowserModule = FModuleManager::Get().LoadModuleChecked<FContentBrowserModule>(TEXT("ContentBrowser"));
 
+	CreateAssetTooltipResources();
+
 	// Configure filter for asset picker
 	FAssetPickerConfig Config;
 	Config.Filter.bRecursiveClasses = true;
@@ -350,7 +353,6 @@ void SAnimationSequenceBrowser::Construct(const FArguments& InArgs)
 	}
 
 	// Configure response to click and double-click
-	Config.OnAssetSelected = FOnAssetSelected::CreateSP(this, &SAnimationSequenceBrowser::OnAnimSelected);
 	Config.OnAssetDoubleClicked = FOnAssetDoubleClicked::CreateSP(this, &SAnimationSequenceBrowser::OnRequestOpenAsset, false);
 	Config.OnGetAssetContextMenu = FOnGetAssetContextMenu::CreateSP(this, &SAnimationSequenceBrowser::OnGetAssetContextMenu);
 	Config.OnAssetTagWantsToBeDisplayed = FOnShouldDisplayAssetTag::CreateSP(this, &SAnimationSequenceBrowser::CanShowColumnForAssetRegistryTag);
@@ -362,6 +364,9 @@ void SAnimationSequenceBrowser::Construct(const FArguments& InArgs)
 	TSharedPtr<FFrontendFilterCategory> AnimCategory = MakeShareable( new FFrontendFilterCategory(LOCTEXT("ExtraAnimationFilters", "Anim Filters"), LOCTEXT("ExtraAnimationFiltersTooltip", "Filter assets by all filters in this category.")) );
 	Config.ExtraFrontendFilters.Add( MakeShareable(new FFrontendFilter_AdditiveAnimAssets(AnimCategory)) );
 	
+	Config.OnGetCustomAssetToolTip = FOnGetCustomAssetToolTip::CreateSP(this, &SAnimationSequenceBrowser::CreateCustomAssetToolTip);
+	Config.OnVisualizeAssetToolTip = FOnVisualizeAssetToolTip::CreateSP(this, &SAnimationSequenceBrowser::OnVisualizeAssetToolTip);
+
 	TWeakPtr< SMenuAnchor > MenuAnchorPtr;
 	
 	this->ChildSlot
@@ -634,4 +639,250 @@ void SAnimationSequenceBrowser::SelectAsset(UAnimationAsset * AnimAsset)
 		}
 	}
 }
+
+TSharedRef<SToolTip> SAnimationSequenceBrowser::CreateCustomAssetToolTip(FAssetData& AssetData)
+{
+	// Make a list of tags to show
+	TArray<UObject::FAssetRegistryTag> Tags;
+	UClass* AssetClass = FindObject<UClass>(ANY_PACKAGE, *AssetData.AssetClass.ToString());
+	check(AssetClass);
+	AssetClass->GetDefaultObject()->GetAssetRegistryTags(Tags);
+
+	TArray<FName> TagsToShow;
+	for(UObject::FAssetRegistryTag& TagEntry : Tags)
+	{
+		if(TagEntry.Name != FName(TEXT("Skeleton")) && TagEntry.Type != UObject::FAssetRegistryTag::TT_Hidden)
+		{
+			TagsToShow.Add(TagEntry.Name);
+		}
+	}
+
+	// Add asset registry tags to a text list; except skeleton as that is implied in Persona
+	TSharedRef<SVerticalBox> DescriptionBox = SNew(SVerticalBox);
+	bool bDescriptionCreated = false;
+	for(TPair<FName, FString> TagPair : AssetData.TagsAndValues)
+	{
+		if(TagsToShow.Contains(TagPair.Key))
+		{
+			if(!bDescriptionCreated)
+			{
+				bDescriptionCreated = true;
+			}
+
+			DescriptionBox->AddSlot()
+			.AutoHeight()
+			.Padding(0,0,5,0)
+			[
+				SNew(SHorizontalBox)
+				+ SHorizontalBox::Slot()
+				.AutoWidth()
+				[
+					SNew(STextBlock)
+					.Text(FText::Format(LOCTEXT("AssetTagKey", "{0} :"), FText::FromName(TagPair.Key)))
+					.ColorAndOpacity(FSlateColor::UseSubduedForeground())
+				]
+
+				+ SHorizontalBox::Slot()
+				.AutoWidth()
+				[
+					SNew(STextBlock)
+					.Text(TagPair.Value)
+					.ColorAndOpacity(FSlateColor::UseForeground())
+				]
+			];
+		}
+	}
+
+	TSharedPtr<SHorizontalBox> ContentBox = nullptr;
+	TSharedRef<SToolTip> ToolTip = SNew(SToolTip)
+	.TextMargin(1)
+	.BorderImage(FEditorStyle::GetBrush("ContentBrowser.TileViewTooltip.ToolTipBorder"))
+	[
+		SNew(SBorder)
+		.Padding(6)
+		.BorderImage(FEditorStyle::GetBrush("ContentBrowser.TileViewTooltip.NonContentBorder"))
+		[
+			SNew(SVerticalBox)
+			+SVerticalBox::Slot()
+			.AutoHeight()
+			.Padding(0,0,0,4)
+			[
+				SNew(SBorder)
+				.Padding(6)
+				.BorderImage(FEditorStyle::GetBrush("ContentBrowser.TileViewTooltip.ContentBorder"))
+				[
+					SNew(SBox)
+					.HAlign(HAlign_Center)
+					[
+						SNew(STextBlock)
+						.Text(FText::FromName(AssetData.AssetName))
+						.Font(FEditorStyle::GetFontStyle("ContentBrowser.TileViewTooltip.NameFont"))
+					]
+				]
+			]
+		
+			+ SVerticalBox::Slot()
+			[
+				SAssignNew(ContentBox, SHorizontalBox)
+				+SHorizontalBox::Slot()
+				.AutoWidth()
+				[
+					SNew(SBorder)
+					.Padding(6)
+					.BorderImage(FEditorStyle::GetBrush("ContentBrowser.TileViewTooltip.ContentBorder"))
+					[
+						ViewportWidget.ToSharedRef()
+					]
+				]
+			]
+		]
+	];
+
+	// If we have a description, add an extra section to the tooltip for it.
+	if(bDescriptionCreated)
+	{
+		ContentBox->AddSlot()
+		.Padding(4, 0, 0, 0)
+		[
+			SNew(SBorder)
+			.Padding(6)
+			.BorderImage(FEditorStyle::GetBrush("ContentBrowser.TileViewTooltip.ContentBorder"))
+			[
+				DescriptionBox
+			]
+		];
+	}
+
+	return ToolTip;
+}
+
+void SAnimationSequenceBrowser::CreateAssetTooltipResources()
+{
+	SAssignNew(ViewportWidget, SViewport)
+		.EnableGammaCorrection(false)
+		.ViewportSize(FVector2D(128, 128));
+
+	ViewportClient = MakeShareable(new FAnimationAssetViewportClient(PreviewScene));
+	SceneViewport = MakeShareable(new FSceneViewport(ViewportClient.Get(), ViewportWidget));
+	PreviewComponent = ConstructObject<UDebugSkelMeshComponent>(UDebugSkelMeshComponent::StaticClass());
+
+	// Client options
+	ViewportClient->ViewportType = LVT_Perspective;
+	ViewportClient->bSetListenerPosition = false;
+	// Default view until we need to show the viewport
+	ViewportClient->SetViewLocation(EditorViewportDefs::DefaultPerspectiveViewLocation);
+	ViewportClient->SetViewRotation(EditorViewportDefs::DefaultPerspectiveViewRotation);
+
+	ViewportClient->Viewport = SceneViewport.Get();
+	ViewportClient->SetRealtime(true);
+	ViewportClient->SetViewMode(VMI_Lit);
+	ViewportClient->ToggleOrbitCamera(true);
+
+	// Add the scene viewport
+	ViewportWidget->SetViewportInterface(SceneViewport.ToSharedRef());
+
+	// Setup the preview component to ensure an animation will update when requested
+	PreviewComponent->MeshComponentUpdateFlag = EMeshComponentUpdateFlag::AlwaysTickPoseAndRefreshBones;
+	PreviewScene.AddComponent(PreviewComponent, FTransform::Identity);
+
+	const UDestructableMeshEditorSettings* Options = GetDefault<UDestructableMeshEditorSettings>();
+
+	PreviewScene.SetLightDirection(Options->AnimPreviewLightingDirection);
+	PreviewScene.GetScene()->UpdateDynamicSkyLight(Options->AnimPreviewSkyBrightness * FLinearColor(Options->AnimPreviewSkyColor), Options->AnimPreviewSkyBrightness * FLinearColor(Options->AnimPreviewFloorColor));
+	PreviewScene.SetLightColor(Options->AnimPreviewDirectionalColor);
+	PreviewScene.SetLightBrightness(Options->AnimPreviewLightBrightness);
+}
+
+bool SAnimationSequenceBrowser::OnVisualizeAssetToolTip(const TSharedPtr<SWidget>& TooltipContent, FAssetData& AssetData)
+{
+	// Resolve the asset
+	USkeletalMesh* MeshToUse = nullptr;
+	UClass* AssetClass = FindObject<UClass>(ANY_PACKAGE, *AssetData.AssetClass.ToString());
+	if(AssetClass->IsChildOf(UAnimationAsset::StaticClass()))
+	{
+		// Set up the viewport to show the asset. Catching the visualize allows us to use
+		// one viewport between all of the assets in the sequence browser.
+		UAnimationAsset* Asset = StaticCast<UAnimationAsset*>(AssetData.GetAsset());
+		check(Asset);
+
+		USkeleton* Skeleton = Asset->GetSkeleton();
+		
+		MeshToUse = Skeleton->GetPreviewMesh(true);
+		check(MeshToUse);
+		if(PreviewComponent->SkeletalMesh != MeshToUse)
+		{
+			PreviewComponent->SetSkeletalMesh(MeshToUse);
+		}
+
+		PreviewComponent->EnablePreview(true, Asset, NULL);
+		PreviewComponent->PreviewInstance->SetLooping(true);
+
+		float HalfFov = FMath::DegreesToRadians(ViewportClient->ViewFOV) / 2.0f;
+		float TargetDist = MeshToUse->Bounds.SphereRadius / FMath::Tan(HalfFov);
+
+		ViewportClient->SetViewRotation(FRotator(0.0f, 135.0f, 0.0f));
+		ViewportClient->SetViewLocationForOrbiting(FVector(0.0f, 0.0f, MeshToUse->Bounds.BoxExtent.Z / 2.0f), TargetDist);
+	}
+
+	// We return false here as we aren't visualizing the tooltip - just detecting when it is about to be shown.
+	// We still want slate to draw it.
+	return false;
+}
+
+void SAnimationSequenceBrowser::CleanupPreviewSceneComponent(USceneComponent* Component)
+{
+	if(Component)
+	{
+		for(int32 ComponentIdx = Component->AttachChildren.Num() - 1 ; ComponentIdx >= 0 ; --ComponentIdx)
+		{
+			USceneComponent* ChildComponent = Component->AttachChildren[ComponentIdx];
+			CleanupPreviewSceneComponent(ChildComponent);
+		}
+		Component->AttachChildren.Empty();
+		Component->DestroyComponent();
+	}
+}
+
+void SAnimationSequenceBrowser::Tick(const FGeometry& AllottedGeometry, const double InCurrentTime, const float InDeltaTime)
+{
+	SCompoundWidget::Tick(AllottedGeometry, InCurrentTime, InDeltaTime);
+
+	if(PreviewComponent && ViewportWidget->GetVisibility() == EVisibility::Visible)
+	{
+		// Tick the world to update preview viewport for tooltips
+		PreviewComponent->GetScene()->GetWorld()->Tick(LEVELTICK_All, InDeltaTime);
+	}
+}
+
+FAnimationAssetViewportClient::FAnimationAssetViewportClient(FPreviewScene& InPreviewScene) : FEditorViewportClient(GLevelEditorModeTools(), &InPreviewScene)
+{
+	SetViewMode(VMI_Lit);
+
+	// Always composite editor objects after post processing in the editor
+	EngineShowFlags.CompositeEditorPrimitives = true;
+	EngineShowFlags.DisableAdvancedFeatures();
+
+	// Setup defaults for the common draw helper.
+	DrawHelper.bDrawPivot = false;
+	DrawHelper.bDrawWorldBox = false;
+	DrawHelper.bDrawKillZ = false;
+	DrawHelper.bDrawGrid = true;
+	DrawHelper.GridColorAxis = FColor(70, 70, 70);
+	DrawHelper.GridColorMajor = FColor(40, 40, 40);
+	DrawHelper.GridColorMinor = FColor(20, 20, 20);
+	DrawHelper.PerspectiveGridSize = HALF_WORLD_MAX1;
+	bDrawAxes = false;
+}
+
+FSceneInterface* FAnimationAssetViewportClient::GetScene() const
+{
+	return PreviewScene->GetScene();
+}
+
+FLinearColor FAnimationAssetViewportClient::GetBackgroundColor() const
+{
+	return FLinearColor(0.8f, 0.85f, 0.85f);
+}
+
+
 #undef LOCTEXT_NAMESPACE
diff --git a/Engine/Source/Editor/Persona/Private/SAnimationSequenceBrowser.h b/Engine/Source/Editor/Persona/Private/SAnimationSequenceBrowser.h
--- a/Engine/Source/Editor/Persona/Private/SAnimationSequenceBrowser.h
+++ b/Engine/Source/Editor/Persona/Private/SAnimationSequenceBrowser.h
@@ -5,6 +5,20 @@
 
 #include "AssetData.h"
 #include "Editor/ContentBrowser/Public/ContentBrowserModule.h"
+#include "PreviewScene.h"
+
+//////////////////////////////////////////////////////////////////////////
+// FAnimationAssetViewportClient
+class FAnimationAssetViewportClient : public FEditorViewportClient
+{
+public:
+	FAnimationAssetViewportClient(FPreviewScene& InPreviewScene);
+
+	// FEditorViewportClient interface
+	virtual FSceneInterface* GetScene() const override;
+	virtual FLinearColor GetBackgroundColor() const override;
+	// End of FEditorViewportClient interface
+};
 
 //////////////////////////////////////////////////////////////////////////
 // SAnimationSequenceBrowser
@@ -21,7 +35,6 @@ public:
 public:
 	void Construct(const FArguments& InArgs);
 
-	void OnAnimSelected(const FAssetData& AssetData);
 	void OnRequestOpenAsset(const FAssetData& AssetData, bool bFromHistory);
 
 	virtual ~SAnimationSequenceBrowser();
@@ -56,6 +69,11 @@ public:
 
 	/** Refresh list */
 	void SelectAsset(UAnimationAsset * AnimAsset);
+	
+	/**
+	 * Update the widget
+	 */
+	virtual void Tick(const FGeometry& AllottedGeometry, const double InCurrentTime, const float InDeltaTime) override;
 
 protected:
 	bool CanShowColumnForAssetRegistryTag(FName AssetType, FName TagName) const;
@@ -99,6 +117,51 @@ protected:
 	 */
 	TSharedRef<SWidget> CreateHistoryMenu(bool bInBackHistory) const;
 
+	/**
+	 * Create the viewport and required objects to control the viewport used in tooltips
+	 */
+	void CreateAssetTooltipResources();
+
+	/**
+	 * Create a custom asset view tooltip for the provided asset
+	 */
+	TSharedRef<SToolTip> CreateCustomAssetToolTip(FAssetData& AssetData);
+
+	/**
+	 * Called as a tooltip is about to show;
+	 */
+	bool OnVisualizeAssetToolTip(const TSharedPtr<SWidget>& TooltipContent, FAssetData& AssetData);
+
+	/**
+	 * Cleaup a component from the preview scene along with any attached children
+	 */
+	void CleanupPreviewSceneComponent(USceneComponent* Component);
+
+	/**
+	 * The actual viewport widget
+	 */	
+	TSharedPtr<SViewport> ViewportWidget;
+
+	/**
+	 * The scene viewport data
+	 */
+	TSharedPtr<FSceneViewport> SceneViewport;
+
+	/**
+	 * Custom viewport client used for tooltip previews
+	 */
+	TSharedPtr<FAnimationAssetViewportClient> ViewportClient;
+
+	/**
+	 * Skeletal component to preview the animation asset on
+	 */
+	UDebugSkelMeshComponent* PreviewComponent;
+
+	/**
+	 * The scene to show in the asset previews
+	 */
+	FPreviewScene PreviewScene;
+
 protected:
 
 	// Pointer back to persona tool that owns us
diff --git a/Engine/Source/Runtime/Slate/Private/Widgets/SViewport.cpp b/Engine/Source/Runtime/Slate/Private/Widgets/SViewport.cpp
--- a/Engine/Source/Runtime/Slate/Private/Widgets/SViewport.cpp
+++ b/Engine/Source/Runtime/Slate/Private/Widgets/SViewport.cpp
@@ -27,6 +27,7 @@ void SViewport::Construct( const FArguments& InArgs )
 	bEnableGammaCorrection = InArgs._EnableGammaCorrection;
 	bEnableBlending = InArgs._EnableBlending;
 	bIgnoreTextureAlpha = InArgs._IgnoreTextureAlpha;
+	ViewportSize = InArgs._ViewportSize;
 
 	this->ChildSlot
 	[
diff --git a/Engine/Source/Runtime/Slate/Public/Widgets/SViewport.h b/Engine/Source/Runtime/Slate/Public/Widgets/SViewport.h
--- a/Engine/Source/Runtime/Slate/Public/Widgets/SViewport.h
+++ b/Engine/Source/Runtime/Slate/Public/Widgets/SViewport.h
@@ -18,6 +18,7 @@ public:
 		, _EnableGammaCorrection(true)
 		, _EnableBlending(false)
 		, _IgnoreTextureAlpha(true)
+		, _ViewportSize(FVector2D(320, 240))
 		{}
 
 		SLATE_DEFAULT_SLOT( FArguments, Content )
@@ -45,6 +46,12 @@ public:
 		 * If false, the texture alpha is used during blending
 		 */
 		SLATE_ARGUMENT( bool, IgnoreTextureAlpha )
+
+		/**
+		 * Size of the viewport widget
+		 */
+		SLATE_ATTRIBUTE(FVector2D, ViewportSize);
+
 	SLATE_END_ARGS()
 
 	/**
@@ -69,7 +76,7 @@ public:
 	 */
 	virtual FVector2D ComputeDesiredSize() const override
 	{
-		return FVector2D( 320, 240 );
+		return ViewportSize.Get();
 	}
 
 	/**
@@ -172,6 +179,9 @@ private:
 	// Whether or not to show the disabled effect when this viewport is disabled.
 	TAttribute<bool> ShowDisabledEffect;
 
+	// Size of the viewport
+	TAttribute<FVector2D> ViewportSize;
+
 	// Whether or not this viewport renders directly to the window back-buffer.
 	bool bRenderDirectlyToWindow;
 
