anim montage slot name shows a warning message to provide a valid slot name if it is empty

diff --git a/Engine/Source/Editor/Persona/Private/SAnimMontagePanel.cpp b/Engine/Source/Editor/Persona/Private/SAnimMontagePanel.cpp
--- a/Engine/Source/Editor/Persona/Private/SAnimMontagePanel.cpp
+++ b/Engine/Source/Editor/Persona/Private/SAnimMontagePanel.cpp
@@ -402,15 +402,16 @@ void SAnimMontagePanel::Update()
 					.AutoWidth()
 					.Padding( FMargin(5.0f, 0.5f) )
 					[
-						SNew(SEditableTextBox)
+						SAssignNew(SlotNameTextBox, SEditableTextBox)
 						.Style( FEditorStyle::Get(), "SpecialEditableTextBox" )
 						.HintText( LOCTEXT("SlotName", "Slot Name") )
-						.Text( Editor, &SMontageEditor::GetMontageSlotName, SlotAnimIdx )
+						.Text(this, &SAnimMontagePanel::GetMontageSlotName, SlotAnimIdx)
 						.Font( FEditorStyle::GetFontStyle("Editor.SearchBoxFont") )
 						.SelectAllTextWhenFocused( true )
 						.RevertTextOnEscape( true )
 						.ClearKeyboardFocusOnCommit( false )
 						.OnTextCommitted( this, &SAnimMontagePanel::OnSlotNodeNameChangeCommit, SlotAnimIdx )
+						.OnTextChanged(this, &SAnimMontagePanel::OnSlotNameChanged)
 					]
 				];
 
@@ -723,4 +724,31 @@ void SAnimMontagePanel::OnSlotNodeNameChangeCommit( const FText& NewText, ETextC
 	MontageEditor.Pin()->RenameSlotNode(SlodeNodeIndex, NewText.ToString());
 }
 
+void SAnimMontagePanel::CheckSlotName(const FText& SlotName) const
+{
+	if (SlotName.IsEmpty())
+	{
+		FText Error = LOCTEXT("Error_SlotNameIsEmpty", "Please provide a slot name for this asset.");
+		SlotNameTextBox->SetError(Error);
+	}
+	else
+	{
+		SlotNameTextBox->SetError(FText::GetEmpty());
+	}
+}
+
+FText SAnimMontagePanel::GetMontageSlotName(int32 SlotIndex) const
+{
+	FText SlotName = MontageEditor.Pin()->GetMontageSlotName(SlotIndex);
+
+	CheckSlotName(SlotName);
+
+	return SlotName;
+}
+
+void SAnimMontagePanel::OnSlotNameChanged(const FText& NewText)
+{
+	CheckSlotName(NewText);
+}
+
 #undef LOCTEXT_NAMESPACE
\ No newline at end of file
diff --git a/Engine/Source/Editor/Persona/Private/SAnimMontagePanel.h b/Engine/Source/Editor/Persona/Private/SAnimMontagePanel.h
--- a/Engine/Source/Editor/Persona/Private/SAnimMontagePanel.h
+++ b/Engine/Source/Editor/Persona/Private/SAnimMontagePanel.h
@@ -82,6 +82,16 @@ public:
 
 	void ClearSelected();
 
+	// helper method to check whether the slot name is empty or not. If empty, shows an error message to provide a valid name
+	void CheckSlotName(const FText& SlotName) const;
+
+	// check the slot name whether valid or not while the user is typing
+	void OnSlotNameChanged(const FText& NewText);
+
+	// get slot name from a montage editor and check the slot name whether valid or not
+	FText GetMontageSlotName(int32 SlotIndex) const;
+
+
 private:
 	TWeakPtr<FPersona> Persona;
 	TWeakPtr<SMontageEditor>	MontageEditor;
@@ -94,6 +104,8 @@ private:
 	FString LastContextHeading;
 
 	STrackNodeSelectionSet SelectionSet;
+
+	TSharedPtr<SEditableTextBox> SlotNameTextBox;
 	
 	void CreateNewSlot(const FText& NewSlotName, ETextCommit::Type CommitInfo);
 	void CreateNewSection(const FText& NewSectionName, ETextCommit::Type CommitInfo, float StartTime);
