TTP #334268 : Skeletal Mesh Editor: UI: Missing Tooltips and Rich Tooltip cabilities

diff --git a/Engine/Source/Editor/Persona/Private/MeshMode/SAdditionalMeshesEditor.cpp b/Engine/Source/Editor/Persona/Private/MeshMode/SAdditionalMeshesEditor.cpp
--- a/Engine/Source/Editor/Persona/Private/MeshMode/SAdditionalMeshesEditor.cpp
+++ b/Engine/Source/Editor/Persona/Private/MeshMode/SAdditionalMeshesEditor.cpp
@@ -177,6 +177,11 @@ void SAdditionalMeshesEditor::Construct(const FArguments& InArgs, TSharedPtr<FPe
 				.ButtonContent()
 				[
 					SNew(STextBlock) .Text(LOCTEXT("AdditionalMeshesAddButtonLabel", "Add Mesh"))
+					.ToolTip(IDocumentation::Get()->CreateToolTip(
+					LOCTEXT("AdditionalMeshesAddButtonTooltip", "Add a body part which is a skeletal mesh sharing the same skeleton"),
+					NULL,
+					TEXT("Shared/Editors/Persona"),
+					TEXT("AdditionalMeshes")))
 				]
 			]
 			+SHorizontalBox::Slot()
diff --git a/Engine/Source/Editor/Persona/Private/PersonaMeshDetails.cpp b/Engine/Source/Editor/Persona/Private/PersonaMeshDetails.cpp
--- a/Engine/Source/Editor/Persona/Private/PersonaMeshDetails.cpp
+++ b/Engine/Source/Editor/Persona/Private/PersonaMeshDetails.cpp
@@ -1291,7 +1291,11 @@ void FPersonaMeshDetails::CustomizeClothingProperties(IDetailLayoutBuilder& Deta
 			SNew(SButton)
 			.OnClicked(this, &FPersonaMeshDetails::OnOpenClothingFileClicked, &DetailLayout)
 			.Text(LOCTEXT("AddAPEXClothingFile", "Add APEX clothing file..."))
-			.ToolTipText(LOCTEXT("AddAPEXClothingFileTip", "Select a new APEX clothing file and add it to the skeletal mesh."))
+			.ToolTip(IDocumentation::Get()->CreateToolTip(
+				LOCTEXT("AddClothingButtonTooltip", "Select a new APEX clothing file and add it to the skeletal mesh."),
+				NULL,
+				TEXT("Shared/Editors/Persona"),
+				TEXT("AddClothing")))
 		]
 	];
 }
diff --git a/Engine/Source/Editor/Persona/Private/SSkeletonTree.cpp b/Engine/Source/Editor/Persona/Private/SSkeletonTree.cpp
--- a/Engine/Source/Editor/Persona/Private/SSkeletonTree.cpp
+++ b/Engine/Source/Editor/Persona/Private/SSkeletonTree.cpp
@@ -314,7 +314,11 @@ TSharedRef< SWidget > FDisplayedMeshBoneInfo::GenerateWidgetForDataColumn()
 	return SNew( SComboButton )
 		.ContentPadding(3)
 		.OnGetMenuContent( this, &FDisplayedMeshBoneInfo::CreateBoneTranslationRetargetingModeMenu )
-		.ToolTipText( LOCTEXT( "RetargetingToolTip", "Set bone translation retargeting mode" ) )
+		.ToolTip(IDocumentation::Get()->CreateToolTip(
+		LOCTEXT("RetargetingToolTip", "Set bone translation retargeting mode"),
+		NULL,
+		TEXT("Shared/Editors/Persona"),
+		TEXT("TranslationRetargeting")))
 		.ButtonContent()
 		[
 			SNew( STextBlock )
