Fixed to show "Save Animation Assets" menu in All Modes and made it prompt to show users what they're saving

diff --git a/Engine/Source/Editor/Persona/Private/Persona.cpp b/Engine/Source/Editor/Persona/Private/Persona.cpp
--- a/Engine/Source/Editor/Persona/Private/Persona.cpp
+++ b/Engine/Source/Editor/Persona/Private/Persona.cpp
@@ -946,9 +946,7 @@ void FPersona::CreateDefaultCommands()
 	// save all animation assets
 	ToolkitCommands->MapAction(FPersonaCommands::Get().SaveAnimationAssets,
 		FExecuteAction::CreateSP(this, &FPersona::SaveAnimationAssets_Execute),
-		FCanExecuteAction::CreateSP(this, &FPersona::CanSaveAnimationAssets),
-		FIsActionChecked(),
-		FIsActionButtonVisible::CreateSP(this, &FPersona::IsRecordAvailable)
+		FCanExecuteAction::CreateSP(this, &FPersona::CanSaveAnimationAssets)
 		);
 
 	// record animation
@@ -1523,14 +1521,14 @@ void FPersona::SaveAsset_Execute()
 
 void FPersona::SaveAnimationAssets_Execute()
 {
-	// only save animation assets or animation blueprints
+	// only save animation assets related to skeletons
 	TArray<UClass*> SaveClasses;
 	SaveClasses.Add(UAnimationAsset::StaticClass());
 	SaveClasses.Add(UAnimBlueprint::StaticClass());
 	SaveClasses.Add(USkeletalMesh::StaticClass());
 	SaveClasses.Add(USkeleton::StaticClass());
 
-	const bool bPromptUserToSave = false;
+	const bool bPromptUserToSave = true;
 	const bool bFastSave = false;
 	FEditorFileUtils::SaveDirtyContentPackages(SaveClasses, bPromptUserToSave, bFastSave);
 }
