Fixed TTP  #342121 : FEATURE REQUEST: PERSONA: Ability to remove all display text without removing mesh

diff --git a/Engine/Config/BaseEditorUserSettings.ini b/Engine/Config/BaseEditorUserSettings.ini
--- a/Engine/Config/BaseEditorUserSettings.ini
+++ b/Engine/Config/BaseEditorUserSettings.ini
@@ -466,7 +466,7 @@ GridSize=25
 ViewModeType=2
 ViewportBackgroundColor=(R=0.04,G=0.04,B=0.04)
 ViewFOV=53.43
-bShowMeshStats=false
+ShowMeshStats=1
 
 
 [UnrealEd.UIEditorOptions]
diff --git a/Engine/Source/Editor/Persona/Private/AnimViewportShowCommands.cpp b/Engine/Source/Editor/Persona/Private/AnimViewportShowCommands.cpp
--- a/Engine/Source/Editor/Persona/Private/AnimViewportShowCommands.cpp
+++ b/Engine/Source/Editor/Persona/Private/AnimViewportShowCommands.cpp
@@ -21,7 +21,12 @@ void FAnimViewportShowCommands::RegisterCommands()
 
 	UI_COMMAND( ShowBones, "Bones", "Display Skeleton in viewport", EUserInterfaceActionType::ToggleButton, FInputGesture() );
 	UI_COMMAND( ShowBoneNames, "Bone Names", "Display Bone Names in Viewport", EUserInterfaceActionType::ToggleButton, FInputGesture() );
-	UI_COMMAND( ShowDisplayInfo, "Display Info", "Display Mesh Info in Viewport", EUserInterfaceActionType::ToggleButton, FInputGesture() );
+
+	// below 3 menus are radio button styles
+	UI_COMMAND(ShowDisplayInfoBasic, "Basic", "Display Basic Mesh Info in Viewport", EUserInterfaceActionType::RadioButton, FInputGesture());
+	UI_COMMAND(ShowDisplayInfoDetailed, "Detailed", "Display Detailed Mesh Info in Viewport", EUserInterfaceActionType::RadioButton, FInputGesture());
+	UI_COMMAND(HideDisplayInfo, "None", "Hide All Display Info in Viewport", EUserInterfaceActionType::RadioButton, FInputGesture());
+
 	UI_COMMAND( ShowBoneWeight, "View Selected Bone Weight", "Display Selected Bone Weight in Viewport", EUserInterfaceActionType::ToggleButton, FInputGesture() );
 	UI_COMMAND( ShowRawAnimation, "Uncompressed Animation", "Display Skeleton With Uncompressed Animation Data", EUserInterfaceActionType::ToggleButton, FInputGesture() );
 	UI_COMMAND( ShowNonRetargetedAnimation, "NonRetargeted Animation", "Display Skeleton With non retargeted Animation Data", EUserInterfaceActionType::ToggleButton, FInputGesture() );
diff --git a/Engine/Source/Editor/Persona/Private/AnimViewportShowCommands.h b/Engine/Source/Editor/Persona/Private/AnimViewportShowCommands.h
--- a/Engine/Source/Editor/Persona/Private/AnimViewportShowCommands.h
+++ b/Engine/Source/Editor/Persona/Private/AnimViewportShowCommands.h
@@ -64,7 +64,9 @@ public:
 	TSharedPtr< FUICommandInfo > ShowBoneNames;
 
 	/** Show skeletal mesh info */
-	TSharedPtr< FUICommandInfo > ShowDisplayInfo;
+	TSharedPtr< FUICommandInfo > ShowDisplayInfoBasic;
+	TSharedPtr< FUICommandInfo > ShowDisplayInfoDetailed;
+	TSharedPtr< FUICommandInfo > HideDisplayInfo;
 
 	/** Show selected bone weight */
 	TSharedPtr< FUICommandInfo > ShowBoneWeight;
diff --git a/Engine/Source/Editor/Persona/Private/AnimationEditorViewportClient.cpp b/Engine/Source/Editor/Persona/Private/AnimationEditorViewportClient.cpp
--- a/Engine/Source/Editor/Persona/Private/AnimationEditorViewportClient.cpp
+++ b/Engine/Source/Editor/Persona/Private/AnimationEditorViewportClient.cpp
@@ -522,7 +522,10 @@ void FAnimationViewportClient::DrawCanvas( FViewport& InViewport, FSceneView& Vi
 		}
 
 		// Display info
-		DisplayInfo(&Canvas, &View, IsShowingMeshStats());
+		if (IsShowingMeshStats())
+		{
+			DisplayInfo(&Canvas, &View, IsDetailedMeshStats());
+		}
 
 		// Draw name of selected bone
 		if (PreviewSkelMeshComp->BonesOfInterest.Num() == 1)
@@ -1987,18 +1990,28 @@ bool FAnimationViewportClient::IsSetDrawUVOverlayChecked() const
 	return bDrawUVs;
 }
 
-void FAnimationViewportClient::OnToggleShowMeshStats()
+void FAnimationViewportClient::OnSetShowMeshStats(int32 ShowMode)
 {
-	ConfigOption->SetShowMeshStats(!ConfigOption->bShowMeshStats);
+	ConfigOption->SetShowMeshStats(ShowMode);
 }
 
 bool FAnimationViewportClient::IsShowingMeshStats() const
 {
-	const bool bShouldBeEnabled = ConfigOption->bShowMeshStats;
+	const bool bShouldBeEnabled = ConfigOption->ShowMeshStats != EDisplayInfoMode::None;
 	const bool bCanBeEnabled = !PersonaPtr.Pin()->IsModeCurrent(FPersonaModes::AnimBlueprintEditMode);
 
 	return bShouldBeEnabled && bCanBeEnabled;
 }
 
+bool FAnimationViewportClient::IsDetailedMeshStats() const
+{
+	return ConfigOption->ShowMeshStats == EDisplayInfoMode::Detailed;
+}
+
+int32 FAnimationViewportClient::GetShowMeshStats() const
+{
+	return ConfigOption->ShowMeshStats;
+}
+
 #undef LOCTEXT_NAMESPACE
 
diff --git a/Engine/Source/Editor/Persona/Private/AnimationEditorViewportClient.h b/Engine/Source/Editor/Persona/Private/AnimationEditorViewportClient.h
--- a/Engine/Source/Editor/Persona/Private/AnimationEditorViewportClient.h
+++ b/Engine/Source/Editor/Persona/Private/AnimationEditorViewportClient.h
@@ -18,6 +18,20 @@ namespace ELocalAxesMode
 	};
 };
 
+//////////////////////////////////////////////////////////////////////////
+// ELocalAxesMode
+
+namespace EDisplayInfoMode
+{
+	enum Type
+	{
+		None,
+		Basic,	
+		Detailed,
+		NumInfoModes
+	};
+};
+
 /////////////////////////////////////////////////////////////////////////
 // FAnimationViewportClient
 
@@ -199,9 +213,13 @@ public:
 	/** Function to get gravity scale label */
 	FString GetGravityScaleLabel() const;
 	/** Function to set mesh stat drawing state */
-	void OnToggleShowMeshStats();
+	void OnSetShowMeshStats(int32 ShowMode);
 	/** Whether or not mesh stats are being displayed */
 	bool IsShowingMeshStats() const;
+	/** Whether detailed mesh stats are being displayed or basic mesh stats */
+	bool IsDetailedMeshStats() const;
+
+	int32 GetShowMeshStats() const;
 
 public:
 
diff --git a/Engine/Source/Editor/Persona/Private/SAnimViewportToolBar.cpp b/Engine/Source/Editor/Persona/Private/SAnimViewportToolBar.cpp
--- a/Engine/Source/Editor/Persona/Private/SAnimViewportToolBar.cpp
+++ b/Engine/Source/Editor/Persona/Private/SAnimViewportToolBar.cpp
@@ -476,11 +476,11 @@ TSharedRef<SWidget> SAnimViewportToolBar::GenerateShowMenu() const
 			}
 			ShowMenuBuilder.EndSection();
 
-			ShowMenuBuilder.BeginSection("AnimviewportInfo", LOCTEXT("ShowInfo_Actions_Info", "Info") );
-			{
-				ShowMenuBuilder.AddMenuEntry( Actions.ShowDisplayInfo );
-			}
-			ShowMenuBuilder.EndSection();
+			ShowMenuBuilder.AddMenuSeparator();
+			ShowMenuBuilder.AddSubMenu(
+				LOCTEXT("AnimviewportInfo", "Display Info"),
+				LOCTEXT("AnimviewportInfoSubMenuToolTip", "Display Mesh Info in Viewport"),
+				FNewMenuDelegate::CreateRaw(this, &SAnimViewportToolBar::FillShowDisplayInfoMenu));
 
 #if WITH_APEX_CLOTHING
 			UDebugSkelMeshComponent* PreviewComp = Viewport.Pin()->GetPersona().Pin()->PreviewComponent;
@@ -629,9 +629,22 @@ void SAnimViewportToolBar::FillShowClothingMenu(FMenuBuilder& MenuBuilder) const
 		MenuBuilder.AddMenuEntry(Actions.ShowOnlyClothSections);
 		MenuBuilder.AddMenuEntry(Actions.HideOnlyClothSections);
 	}
+	MenuBuilder.EndSection();
 #endif // #if WITH_APEX_CLOTHING
 }
 
+void SAnimViewportToolBar::FillShowDisplayInfoMenu(FMenuBuilder& MenuBuilder) const
+{
+	const FAnimViewportShowCommands& Actions = FAnimViewportShowCommands::Get();
+
+	// display info levels
+	{
+		MenuBuilder.AddMenuEntry(Actions.ShowDisplayInfoBasic);
+		MenuBuilder.AddMenuEntry(Actions.ShowDisplayInfoDetailed);
+		MenuBuilder.AddMenuEntry(Actions.HideDisplayInfo);
+	}
+}
+
 FText SAnimViewportToolBar::GetLODMenuLabel() const
 {
 	FText Label = LOCTEXT("LODMenu_AutoLabel", "LOD Auto");
diff --git a/Engine/Source/Editor/Persona/Private/SAnimViewportToolBar.h b/Engine/Source/Editor/Persona/Private/SAnimViewportToolBar.h
--- a/Engine/Source/Editor/Persona/Private/SAnimViewportToolBar.h
+++ b/Engine/Source/Editor/Persona/Private/SAnimViewportToolBar.h
@@ -50,6 +50,12 @@ private:
 	* Generates the Show -> Clothing sub menu content
 	*/
 	void FillShowClothingMenu(FMenuBuilder& MenuBuilder) const;
+
+	/**
+	* Generates the Show -> Display Info sub menu content
+	*/
+	void FillShowDisplayInfoMenu(FMenuBuilder& MenuBuilder) const;
+
 	/**
 	 * Generates the toolbar LOD menu content 
 	 *
diff --git a/Engine/Source/Editor/Persona/Private/SAnimationEditorViewport.cpp b/Engine/Source/Editor/Persona/Private/SAnimationEditorViewport.cpp
--- a/Engine/Source/Editor/Persona/Private/SAnimationEditorViewport.cpp
+++ b/Engine/Source/Editor/Persona/Private/SAnimationEditorViewport.cpp
@@ -490,10 +490,22 @@ void SAnimationEditorViewportTabBody::BindCommands()
 
 	//Display info
 	CommandList.MapAction( 
-		ViewportShowMenuCommands.ShowDisplayInfo,
-		FExecuteAction::CreateSP(this, &SAnimationEditorViewportTabBody::OnShowDisplayInfo),
+		ViewportShowMenuCommands.ShowDisplayInfoBasic,
+		FExecuteAction::CreateSP(this, &SAnimationEditorViewportTabBody::OnShowDisplayInfo, (int32)EDisplayInfoMode::Basic),
 		FCanExecuteAction(),
-		FIsActionChecked::CreateSP(this, &SAnimationEditorViewportTabBody::IsShowingMeshInfo));
+		FIsActionChecked::CreateSP(this, &SAnimationEditorViewportTabBody::IsShowingMeshInfo, (int32)EDisplayInfoMode::Basic));
+
+	CommandList.MapAction(
+		ViewportShowMenuCommands.ShowDisplayInfoDetailed,
+		FExecuteAction::CreateSP(this, &SAnimationEditorViewportTabBody::OnShowDisplayInfo, (int32)EDisplayInfoMode::Detailed),
+		FCanExecuteAction(),
+		FIsActionChecked::CreateSP(this, &SAnimationEditorViewportTabBody::IsShowingMeshInfo, (int32)EDisplayInfoMode::Detailed));
+
+	CommandList.MapAction(
+		ViewportShowMenuCommands.HideDisplayInfo,
+		FExecuteAction::CreateSP(this, &SAnimationEditorViewportTabBody::OnShowDisplayInfo, (int32)EDisplayInfoMode::None),
+		FCanExecuteAction(),
+		FIsActionChecked::CreateSP(this, &SAnimationEditorViewportTabBody::IsShowingMeshInfo, (int32)EDisplayInfoMode::None));
 
 	//Bone weight
 	CommandList.MapAction( 
@@ -825,14 +837,14 @@ bool SAnimationEditorViewportTabBody::IsShowingAdditiveBase() const
 	return PreviewComponent != NULL && PreviewComponent->bDisplayAdditiveBasePose;
 }
 
-void SAnimationEditorViewportTabBody::OnShowDisplayInfo()
+void SAnimationEditorViewportTabBody::OnShowDisplayInfo(int32 DisplayInfoMode)
 {
-	GetAnimationViewportClient()->OnToggleShowMeshStats();
+	GetAnimationViewportClient()->OnSetShowMeshStats(DisplayInfoMode);
 }
 
-bool SAnimationEditorViewportTabBody::IsShowingMeshInfo() const
+bool SAnimationEditorViewportTabBody::IsShowingMeshInfo(int32 DisplayInfoMode) const
 {
-	return GetAnimationViewportClient()->IsShowingMeshStats();
+	return GetAnimationViewportClient()->GetShowMeshStats() == DisplayInfoMode;
 }
 
 void SAnimationEditorViewportTabBody::OnShowBoneWeight()
diff --git a/Engine/Source/Editor/Persona/Private/SAnimationEditorViewport.h b/Engine/Source/Editor/Persona/Private/SAnimationEditorViewport.h
--- a/Engine/Source/Editor/Persona/Private/SAnimationEditorViewport.h
+++ b/Engine/Source/Editor/Persona/Private/SAnimationEditorViewport.h
@@ -248,10 +248,10 @@ private:
 	void OnSetLODModel( ELODViewSelection LODSelectionType );
 
 	/** Function to show/hide mesh info*/
-	void OnShowDisplayInfo();
+	void OnShowDisplayInfo(int32 DisplayInfoMode);
 
 	/** Function to check whether mesh info is displayed or not */
-	bool IsShowingMeshInfo() const;
+	bool IsShowingMeshInfo(int32 DisplayInfoMode) const;
 
 	/** Function to show/hide grid in the viewport */
 	void OnShowGrid();
diff --git a/Engine/Source/Editor/UnrealEd/Classes/Preferences/PersonaOptions.h b/Engine/Source/Editor/UnrealEd/Classes/Preferences/PersonaOptions.h
--- a/Engine/Source/Editor/UnrealEd/Classes/Preferences/PersonaOptions.h
+++ b/Engine/Source/Editor/UnrealEd/Classes/Preferences/PersonaOptions.h
@@ -30,8 +30,9 @@ class UNREALED_API UPersonaOptions : public UObject
 	UPROPERTY(EditAnywhere, config, Category=Options)
 	uint32 bMuteAudio:1;
 
-	UPROPERTY(EditAnywhere, config, Category=Options)
-	uint32 bShowMeshStats:1;
+	// currently Stats can have None, Basic and Detailed. Please refer to EDisplayInfoMode.
+	UPROPERTY(EditAnywhere, config, Category = Options)
+	int32 ShowMeshStats;
 
 	UPROPERTY(EditAnywhere, config, Category=Options)
 	int32 GridSize;
@@ -59,6 +60,6 @@ public:
 	void SetViewModeIndex( int32 InViewModeIndex );
 	void SetViewFOV( float InViewFOV );
 	void SetDefaultLocalAxesSelection( uint32 InDefaultLocalAxesSelection );
-	void SetShowMeshStats( bool bInShowMeshStats );
+	void SetShowMeshStats( int32 InShowMeshStats );
 };
 
diff --git a/Engine/Source/Editor/UnrealEd/Private/PreferenceStubs.cpp b/Engine/Source/Editor/UnrealEd/Private/PreferenceStubs.cpp
--- a/Engine/Source/Editor/UnrealEd/Private/PreferenceStubs.cpp
+++ b/Engine/Source/Editor/UnrealEd/Private/PreferenceStubs.cpp
@@ -101,9 +101,9 @@ void UPersonaOptions::SetDefaultLocalAxesSelection( uint32 InDefaultLocalAxesSel
 }
 
 
-void UPersonaOptions::SetShowMeshStats( bool bInShowMeshStats )
+void UPersonaOptions::SetShowMeshStats( int32 InShowMeshStats )
 {
-	bShowMeshStats = bInShowMeshStats;
+	ShowMeshStats = InShowMeshStats;
 	SaveConfig();
 }
 
