Fixed a regression : clothing combo boxes didní»t work properly when the user changed the clothing asset name

diff --git a/Engine/Source/Editor/Persona/Private/PersonaMeshDetails.cpp b/Engine/Source/Editor/Persona/Private/PersonaMeshDetails.cpp
--- a/Engine/Source/Editor/Persona/Private/PersonaMeshDetails.cpp
+++ b/Engine/Source/Editor/Persona/Private/PersonaMeshDetails.cpp
@@ -1759,8 +1759,6 @@ FReply FPersonaMeshDetails::OnOpenClothingFileClicked(IDetailLayoutBuilder* Deta
 
 void FPersonaMeshDetails::UpdateComboBoxStrings()
 {
-	ClothingComboLODInfos.Empty();
-
 	USkeletalMesh* SkelMesh = PersonaPtr->GetMesh();
 
 	if (!SkelMesh)
@@ -1770,14 +1768,19 @@ void FPersonaMeshDetails::UpdateComboBoxStrings()
 
 	int32 NumLODs = SkelMesh->LODInfo.Num();
 
+	while (ClothingComboLODInfos.Num() < NumLODs)
+	{
+		new (ClothingComboLODInfos)FClothingComboInfo;
+	}
+
 	for (int32 LODIdx = 0; LODIdx < NumLODs; LODIdx++)
 	{
-		FClothingComboInfo* NewLODInfo = new (ClothingComboLODInfos)FClothingComboInfo;
+		FClothingComboInfo& NewLODInfo = ClothingComboLODInfos[LODIdx];
 
-		TArray<TSharedPtr< STextComboBox >>& ClothingComboBoxes = NewLODInfo->ClothingComboBoxes;
-		TArray<TSharedPtr<FString> >& ClothingComboStrings = NewLODInfo->ClothingComboStrings;
-		TMap<FString, FClothAssetSubmeshIndex>&	ClothingComboStringReverseLookup = NewLODInfo->ClothingComboStringReverseLookup;
-		TArray<int32>& ClothingComboSelectedIndices = NewLODInfo->ClothingComboSelectedIndices;
+		TArray<TSharedPtr< STextComboBox >>& ClothingComboBoxes = NewLODInfo.ClothingComboBoxes;
+		TArray<TSharedPtr<FString> >& ClothingComboStrings = NewLODInfo.ClothingComboStrings;
+		TMap<FString, FClothAssetSubmeshIndex>&	ClothingComboStringReverseLookup = NewLODInfo.ClothingComboStringReverseLookup;
+		TArray<int32>& ClothingComboSelectedIndices = NewLODInfo.ClothingComboSelectedIndices;
 
 		// Clear out old strings
 		ClothingComboStrings.Empty();
@@ -1848,14 +1851,14 @@ void FPersonaMeshDetails::UpdateComboBoxStrings()
 					{
 						if (SectionAssetSubmeshIndices[SlotIdx] == AssetSubmeshIndex)
 						{
-							ClothingComboSelectedIndices[SlotIdx] = NewLODInfo->ClothingComboStrings.Num() - 1;
+							ClothingComboSelectedIndices[SlotIdx] = NewLODInfo.ClothingComboStrings.Num() - 1;
 						}
 					}
 				}
 			}
 		}
 		// Cause combo boxes to refresh their strings lists and selected value
-		for (int32 i = 0; i < NewLODInfo->ClothingComboBoxes.Num(); i++)
+		for (int32 i = 0; i < NewLODInfo.ClothingComboBoxes.Num(); i++)
 		{
 			ClothingComboBoxes[i]->RefreshOptions();
 			ClothingComboBoxes[i]->SetSelectedItem(ClothingComboSelectedIndices.IsValidIndex(i) && ClothingComboStrings.IsValidIndex(ClothingComboSelectedIndices[i]) ? ClothingComboStrings[ClothingComboSelectedIndices[i]] : ClothingComboStrings[0]);
