Altered animation asset setup for floating preview windows to force play the animation; otherwise certain assets would cause the preview to stop updating for other assets.

diff --git a/Engine/Source/Editor/Persona/Private/SAnimationSequenceBrowser.cpp b/Engine/Source/Editor/Persona/Private/SAnimationSequenceBrowser.cpp
--- a/Engine/Source/Editor/Persona/Private/SAnimationSequenceBrowser.cpp
+++ b/Engine/Source/Editor/Persona/Private/SAnimationSequenceBrowser.cpp
@@ -815,7 +815,7 @@ bool SAnimationSequenceBrowser::OnVisualizeAssetToolTip(const TSharedPtr<SWidget
 		}
 
 		PreviewComponent->EnablePreview(true, Asset, NULL);
-		PreviewComponent->PreviewInstance->SetLooping(true);
+		PreviewComponent->PreviewInstance->PlayAnim(true);
 
 		float HalfFov = FMath::DegreesToRadians(ViewportClient->ViewFOV) / 2.0f;
 		float TargetDist = MeshToUse->Bounds.SphereRadius / FMath::Tan(HalfFov);
