Removed Persona generic delete and created specific UI commands for the notify panel to stop Persona consuming all delete input and making the hotkeys not work in the My Blueprint view.

#ttp 341212 - PERSONA: Unable to delete My Blueprint items in an Anim BP with hotkey

diff --git a/Engine/Source/Editor/Persona/Private/Persona.cpp b/Engine/Source/Editor/Persona/Private/Persona.cpp
--- a/Engine/Source/Editor/Persona/Private/Persona.cpp
+++ b/Engine/Source/Editor/Persona/Private/Persona.cpp
@@ -1056,10 +1056,6 @@ void FPersona::CreateDefaultCommands()
 		FIsActionChecked(),
 		FIsActionButtonVisible::CreateSP(this, &FPersona::IsInPersonaMode, FPersonaModes::AnimationEditMode)
 		);
-
-	// Generic deletion
-	ToolkitCommands->MapAction(FGenericCommands::Get().Delete, FExecuteAction::CreateSP(this, &FPersona::OnCommandGenericDelete));
-
 }
 
 void FPersona::StartEditingDefaults(bool bAutoFocus, bool bForceRefresh)
diff --git a/Engine/Source/Editor/Persona/Private/SAnimNotifyPanel.cpp b/Engine/Source/Editor/Persona/Private/SAnimNotifyPanel.cpp
--- a/Engine/Source/Editor/Persona/Private/SAnimNotifyPanel.cpp
+++ b/Engine/Source/Editor/Persona/Private/SAnimNotifyPanel.cpp
@@ -2728,6 +2728,14 @@ bool SNotifyEdTrack::CanDeleteTrack()
 	return AnimPanelPtr.Pin()->CanDeleteTrack(TrackIndex);
 }
 
+//////////////////////////////////////////////////////////////////////////
+// FAnimNotifyPanelCommands
+
+void FAnimNotifyPanelCommands::RegisterCommands()
+{
+	UI_COMMAND(DeleteNotify, "Delete", "Deletes the selected notifies.", EUserInterfaceActionType::Button, FInputGesture(EKeys::Platform_Delete));
+}
+
 //////////////////////////////////////////////////////////////////////////
 // SAnimNotifyPanel
 
@@ -2745,13 +2753,16 @@ void SAnimNotifyPanel::Construct(const FArguments& InArgs)
 	Sequence = InArgs._Sequence;
 	MarkerBars = InArgs._MarkerBars;
 
+	FAnimNotifyPanelCommands::Register();
+	BindCommands();
+
 	// @todo anim : this is kinda hack to make sure it has 1 track is alive
 	// we can do this whenever import or asset is created, but it's more places to handle than here
 	// the function name in that case will need to change
 	Sequence->InitializeNotifyTrack();
 	Sequence->RegisterOnNotifyChanged(UAnimSequenceBase::FOnNotifyChanged::CreateSP( this, &SAnimNotifyPanel::RefreshNotifyTracks )  );
 	PersonaPtr.Pin()->RegisterOnPostUndo(FPersona::FOnPostUndo::CreateSP( this, &SAnimNotifyPanel::PostUndo ) );
-	PersonaPtr.Pin()->RegisterOnGenericDelete(FPersona::FOnDeleteGeneric::CreateSP(this, &SAnimNotifyPanel::OnGenericDelete));
+	PersonaPtr.Pin()->RegisterOnGenericDelete(FPersona::FOnDeleteGeneric::CreateSP(this, &SAnimNotifyPanel::OnDeletePressed));
 
 	CurrentPosition = InArgs._CurrentPosition;
 	OnSelectionChanged = InArgs._OnSelectionChanged;
@@ -2948,7 +2959,7 @@ void SAnimNotifyPanel::RefreshNotifyTracks()
 			.OnNodeDragStarted(this, &SAnimNotifyPanel::OnNotifyNodeDragStarted)
 			.MarkerBars(MarkerBars)
 			.OnRequestRefreshOffsets(OnRequestRefreshOffsets)
-			.OnDeleteNotify(this, &SAnimNotifyPanel::OnGenericDelete)
+			.OnDeleteNotify(this, &SAnimNotifyPanel::OnDeletePressed)
 			.OnDeselectAllNotifies(this, &SAnimNotifyPanel::DeselectAllNotifies)
 			.OnCopyNotifies(this, &SAnimNotifyPanel::CopySelectedNotifiesToClipboard)
 			.OnPasteNotifies(this, &SAnimNotifyPanel::OnPasteNotifies)
@@ -3024,7 +3035,7 @@ void SAnimNotifyPanel::PostUndo()
 	}
 }
 
-void SAnimNotifyPanel::OnGenericDelete()
+void SAnimNotifyPanel::OnDeletePressed()
 {
 	// If there's no focus on the panel it's likely the user is not editing notifies
 	// so don't delete anything when the key is pressed.
@@ -3245,4 +3256,25 @@ void SAnimNotifyPanel::OnPropertyChanged(UObject* ChangedObject, FPropertyChange
 	}
 }
 
+void SAnimNotifyPanel::BindCommands()
+{
+	check(!UICommandList.IsValid());
+
+	UICommandList = MakeShareable(new FUICommandList);
+	const FAnimNotifyPanelCommands& Commands = FAnimNotifyPanelCommands::Get();
+
+	UICommandList->MapAction(
+		Commands.DeleteNotify,
+		FExecuteAction::CreateSP(this, &SAnimNotifyPanel::OnDeletePressed));
+}
+
+FReply SAnimNotifyPanel::OnKeyDown(const FGeometry& MyGeometry, const FKeyboardEvent& InKeyboardEvent)
+{
+	if(UICommandList->ProcessCommandBindings(InKeyboardEvent))
+	{
+		return FReply::Handled();
+	}
+	return FReply::Unhandled();
+}
+
 #undef LOCTEXT_NAMESPACE
diff --git a/Engine/Source/Editor/Persona/Private/SAnimNotifyPanel.h b/Engine/Source/Editor/Persona/Private/SAnimNotifyPanel.h
--- a/Engine/Source/Editor/Persona/Private/SAnimNotifyPanel.h
+++ b/Engine/Source/Editor/Persona/Private/SAnimNotifyPanel.h
@@ -54,6 +54,20 @@ namespace ENotifyStateHandleHit
 //////////////////////////////////////////////////////////////////////////
 // SAnimNotifyPanel
 
+class FAnimNotifyPanelCommands : public TCommands<FAnimNotifyPanelCommands>
+{
+public:
+	FAnimNotifyPanelCommands()
+		: TCommands<FAnimNotifyPanelCommands>("AnimNotifyPanel", NSLOCTEXT("Contexts", "AnimNotifyPanel", "Anim Notify Panel"), NAME_None, FEditorStyle::GetStyleSetName())
+	{
+
+	}
+
+	TSharedPtr<FUICommandInfo> DeleteNotify;
+
+	virtual void RegisterCommands() override;
+};
+
 // @todo anim : register when it's opened for the animsequence
 // broadcast when animsequence changed, so that we refresh for multiple window
 class SAnimNotifyPanel: public SAnimTrackPanel
@@ -115,6 +129,9 @@ public:
 	FCoreDelegates::FOnObjectPropertyChanged::FDelegate OnPropertyChangedHandle;
 	void OnPropertyChanged(UObject* ChangedObject, FPropertyChangedEvent& PropertyEvent);
 
+	/** Handler for key press events */
+	virtual FReply OnKeyDown(const FGeometry& MyGeometry, const FKeyboardEvent& InKeyboardEvent);
+
 private:
 	TSharedPtr<SBorder> PanelArea;
 	class UAnimSequenceBase* Sequence;
@@ -140,8 +157,8 @@ private:
 	void RefreshNotifyTracks();
 	void PostUndo();
 
-	/** Handler for generic editor delete command */
-	void OnGenericDelete();
+	/** Handler for delete command */
+	void OnDeletePressed();
 
 	/** We support keyboard focus to detect when we should process key commands like delete */
 	virtual bool SupportsKeyboardFocus() const override
@@ -155,9 +172,15 @@ private:
 	// Called to deselect all notifies across all tracks
 	void DeselectAllNotifies();
 
+	// Binds the UI commands for this widget to delegates
+	void BindCommands();
+
 	/** Persona reference **/
 	TWeakPtr<FPersona> PersonaPtr;
 
 	/** Attribute for accessing any section/branching point positions we have to draw */
 	TAttribute<TArray<FTrackMarkerBar>>	MarkerBars;
+
+	/** UI commands for this widget */
+	TSharedPtr<FUICommandList> UICommandList;
 };
\ No newline at end of file
