Tooltip fixes

342151	PERSONA: Looping Interpolation tooltip worded incorrectly
342154	POLISH: PERSONA: Looping Interpolation warning message contains grammatical error
342168	PERSONA: Add mesh button rich tool tip contains grammatical errors

diff --git a/Engine/Source/Editor/Persona/Private/Persona.cpp b/Engine/Source/Editor/Persona/Private/Persona.cpp
--- a/Engine/Source/Editor/Persona/Private/Persona.cpp
+++ b/Engine/Source/Editor/Persona/Private/Persona.cpp
@@ -2520,7 +2520,7 @@ void FPersona::ExportToFBX(TArray<TWeakObjectPtr<UAnimSequence>> & AnimSequences
 
 void FPersona::AddLoopingInterpolation(TArray<TWeakObjectPtr<UAnimSequence>> & AnimSequences)
 {
-	FText WarningMessage = LOCTEXT("AddLoopiingInterpolation", "This will add extra frame at the end of the animation with the first frame to create a better looping interpolation. This action isn't undoable. Would you like to proceed?");
+	FText WarningMessage = LOCTEXT("AddLoopiingInterpolation", "This will add an extra first frame at the end of the animation to create a better looping interpolation. This action cannot be undone. Would you like to proceed?");
 
 	if(FMessageDialog::Open(EAppMsgType::YesNo, WarningMessage) == EAppReturnType::Yes)
 	{
diff --git a/Engine/Source/Editor/Persona/Private/PersonaCommands.cpp b/Engine/Source/Editor/Persona/Private/PersonaCommands.cpp
--- a/Engine/Source/Editor/Persona/Private/PersonaCommands.cpp
+++ b/Engine/Source/Editor/Persona/Private/PersonaCommands.cpp
@@ -28,5 +28,5 @@ void FPersonaCommands::RegisterCommands()
 	UI_COMMAND( ApplyCompression, "Apply Compression", "Apply cvompression to current animation", EUserInterfaceActionType::Button, FInputGesture() );
 	UI_COMMAND( RecordAnimation, "Record to new Animation", "Create new animation from currently playing", EUserInterfaceActionType::Button, FInputGesture() );
 	UI_COMMAND( ExportToFBX, "Export to FBX", "Export current animation to FBX", EUserInterfaceActionType::Button, FInputGesture() );
-	UI_COMMAND( AddLoopingInterpolation, "Add Looping Interpolation", "Add an extra first frame at the end of the frame to create interpolation when looping", EUserInterfaceActionType::Button, FInputGesture() );
+	UI_COMMAND( AddLoopingInterpolation, "Add Looping Interpolation", "Add an extra first frame at the end of the animation to create interpolation when looping", EUserInterfaceActionType::Button, FInputGesture() );
 }
