Fixed TTP# 342144 POLISH: PERSONA: Screen Size name inconsistency

diff --git a/Engine/Source/Editor/Persona/Private/AnimationEditorViewportClient.cpp b/Engine/Source/Editor/Persona/Private/AnimationEditorViewportClient.cpp
--- a/Engine/Source/Editor/Persona/Private/AnimationEditorViewportClient.cpp
+++ b/Engine/Source/Editor/Persona/Private/AnimationEditorViewportClient.cpp
@@ -813,7 +813,7 @@ void FAnimationViewportClient::DisplayInfo(FCanvas* Canvas, FSceneView* View, bo
 			NumChunksInUse = LODModel.Chunks.Num();
 			NumSectionsInUse = LODModel.Sections.Num();
 
-			InfoString = FString::Printf(TEXT("LOD [%d] Bones:%d (Mapped to Vertices:%d) Polys:%d"),
+			InfoString = FString::Printf(TEXT("LOD: %d, Bones: %d (Mapped to Vertices: %d), Polys: %d"),
 				LODIndex,
 				NumBonesInUse,
 				NumBonesMappedToVerts,
@@ -821,7 +821,7 @@ void FAnimationViewportClient::DisplayInfo(FCanvas* Canvas, FSceneView* View, bo
 
 			Canvas->DrawShadowedString(CurXOffset, CurYOffset, *InfoString, GEngine->GetSmallFont(), TextColor);
 
-			InfoString = FString::Printf(TEXT("(Screen Size:%3.2f, FOV:%3.0f)"), LODFactor, ViewFOV);
+			InfoString = FString::Printf(TEXT("Current Screen Size: %3.2f, FOV:%3.0f"), LODFactor, ViewFOV);
 			CurYOffset += YL + 2;
 			Canvas->DrawShadowedString(CurXOffset, CurYOffset, *InfoString, GEngine->GetSmallFont(), TextColor);
 
@@ -834,7 +834,7 @@ void FAnimationViewportClient::DisplayInfo(FCanvas* Canvas, FSceneView* View, bo
 				int32 ChunkRigidVerts = LODModel.Chunks[ChunkIndex].GetNumRigidVertices();
 				int32 ChunkSoftVerts = LODModel.Chunks[ChunkIndex].GetNumSoftVertices();
 
-				InfoString = FString::Printf(TEXT(" [Chunk %d] Verts:%d (Rigid:%d Soft:%d) Bones:%d"),
+				InfoString = FString::Printf(TEXT(" [Chunk %d] Verts:%d (Rigid:%d Soft:%d), Bones:%d"),
 					ChunkIndex,
 					ChunkRigidVerts + ChunkSoftVerts,
 					ChunkRigidVerts,
@@ -860,7 +860,7 @@ void FAnimationViewportClient::DisplayInfo(FCanvas* Canvas, FSceneView* View, bo
 			CurYOffset += YL + 2;
 			Canvas->DrawShadowedString(CurXOffset, CurYOffset, *InfoString, GEngine->GetSmallFont(), TextColor);
 
-			InfoString = FString::Printf(TEXT("Chunks:%d Sections:%d %s"),
+			InfoString = FString::Printf(TEXT("Chunks:%d, Sections:%d %s"),
 				NumChunksInUse,
 				NumSectionsInUse,
 				*WeightUsage
