Fixed a "Slot Name" combo box problem in Anim Graph

diff --git a/Engine/Source/Editor/Persona/Private/Customization/AnimGraphNodeSlotDetails.cpp b/Engine/Source/Editor/Persona/Private/Customization/AnimGraphNodeSlotDetails.cpp
--- a/Engine/Source/Editor/Persona/Private/Customization/AnimGraphNodeSlotDetails.cpp
+++ b/Engine/Source/Editor/Persona/Private/Customization/AnimGraphNodeSlotDetails.cpp
@@ -198,6 +198,12 @@ int32 FAnimGraphNodeSlotDetails::RefreshSlotNames()
 		}
 	}
 
+	// if SlotNameComboList is empty, just add a temporary message to avoid a combo box problem which doesn't show the list properly if empty
+	if (SlotNameComboList.Num() == 0)
+	{
+		SlotNameComboList.Add(MakeShareable(new FString("Please Add a slot name using the neighboring + button")));
+	}
+
 	return CurrentSelectedIndex;
 }
 ////////////////////////////////////////////////////////////////
