Change Zoom in Persona tracks to be based on mouse position instead of track centre

diff --git a/Engine/Source/Editor/Persona/Private/SAnimCurvePanel.cpp b/Engine/Source/Editor/Persona/Private/SAnimCurvePanel.cpp
--- a/Engine/Source/Editor/Persona/Private/SAnimCurvePanel.cpp
+++ b/Engine/Source/Editor/Persona/Private/SAnimCurvePanel.cpp
@@ -113,6 +113,7 @@ public:
 protected:
 	// SWidget interface
 	virtual int32 OnPaint(const FPaintArgs& Args, const FGeometry& AllottedGeometry, const FSlateRect& MyClippingRect, FSlateWindowElementList& OutDrawElements, int32 LayerId, const FWidgetStyle& InWidgetStyle, bool bParentEnabled) const override;
+	virtual FReply OnMouseWheel(const FGeometry& MyGeometry, const FPointerEvent& MouseEvent);
 	// SWidget interface
 
 	// SCurveEditor interface
@@ -189,6 +190,26 @@ int32 SAnimCurveEd::OnPaint(const FPaintArgs& Args, const FGeometry& AllottedGeo
 	return NewLayerId;
 }
 
+FReply SAnimCurveEd::OnMouseWheel(const FGeometry& MyGeometry, const FPointerEvent& MouseEvent)
+{
+	const float ZoomDelta = -0.1f * MouseEvent.GetWheelDelta();
+
+	const FVector2D WidgetSpace = MyGeometry.AbsoluteToLocal(MouseEvent.GetScreenSpacePosition());
+	const float ZoomRatio = FMath::Clamp((WidgetSpace.X / MyGeometry.Size.X), 0.f, 1.f);
+
+	{
+		const float InputViewSize = ViewMaxInput.Get() - ViewMinInput.Get();
+		const float InputChange = InputViewSize * ZoomDelta;
+
+		float NewViewMinInput = ViewMinInput.Get() - (InputChange * ZoomRatio);
+		float NewViewMaxInput = ViewMaxInput.Get() + (InputChange * (1.f - ZoomRatio));
+
+		SetInputMinMax(NewViewMinInput, NewViewMaxInput);
+	}
+
+	return FReply::Handled();
+}
+
 void SAnimCurveEd::Construct(const FArguments& InArgs)
 {
 	OnGetScrubValue = InArgs._OnGetScrubValue;
diff --git a/Engine/Source/Editor/Persona/Private/SAnimNotifyPanel.cpp b/Engine/Source/Editor/Persona/Private/SAnimNotifyPanel.cpp
--- a/Engine/Source/Editor/Persona/Private/SAnimNotifyPanel.cpp
+++ b/Engine/Source/Editor/Persona/Private/SAnimNotifyPanel.cpp
@@ -200,6 +200,7 @@ public:
 		, _OnDeselectAllNotifies()
 		, _OnCopyNotifies()
 		, _OnPasteNotifies()
+		, _OnSetInputViewRange()
 		{}
 
 		SLATE_ARGUMENT( TSharedPtr<FPersona>,	Persona )
@@ -221,6 +222,7 @@ public:
 		SLATE_EVENT( FDeselectAllNotifies, OnDeselectAllNotifies)
 		SLATE_EVENT( FCopyNotifies, OnCopyNotifies )
 		SLATE_EVENT( FPasteNotifies, OnPasteNotifies )
+		SLATE_EVENT(FOnSetInputViewRange, OnSetInputViewRange)
 		SLATE_END_ARGS()
 public:
 
@@ -230,6 +232,7 @@ public:
 	// SWidget interface
 	virtual void Tick( const FGeometry& AllottedGeometry, const double InCurrentTime, const float InDeltaTime ) {UpdateCachedGeometry(AllottedGeometry);}
 	virtual int32 OnPaint(const FPaintArgs& Args, const FGeometry& AllottedGeometry, const FSlateRect& MyClippingRect, FSlateWindowElementList& OutDrawElements, int32 LayerId, const FWidgetStyle& InWidgetStyle, bool bParentEnabled) const override;
+	virtual FReply OnMouseWheel(const FGeometry& MyGeometry, const FPointerEvent& MouseEvent) override;
 	// End of SWidget interface
 
 	/** Returns the cached rendering geometry of this track */
@@ -393,6 +396,7 @@ protected:
 	FDeselectAllNotifies					OnDeselectAllNotifies;
 	FCopyNotifies							OnCopyNotifies;
 	FPasteNotifies							OnPasteNotifies;
+	FOnSetInputViewRange					OnSetInputViewRange;
 
 	/** Delegate to call when offsets should be refreshed in a montage */
 	FRefreshOffsetsRequest					OnRequestRefreshOffsets;
@@ -444,6 +448,7 @@ public:
 		, _OnDeleteNotify()
 		, _OnDeselectAllNotifies()
 		, _OnCopyNotifies()
+		, _OnSetInputViewRange()
 	{}
 	SLATE_ARGUMENT( int32, TrackIndex )
 	SLATE_ARGUMENT( TSharedPtr<SAnimNotifyPanel>, AnimNotifyPanel)
@@ -462,6 +467,7 @@ public:
 	SLATE_EVENT( FDeselectAllNotifies, OnDeselectAllNotifies)
 	SLATE_EVENT( FCopyNotifies, OnCopyNotifies )
 	SLATE_EVENT( FPasteNotifies, OnPasteNotifies )
+	SLATE_EVENT(FOnSetInputViewRange, OnSetInputViewRange)
 	SLATE_END_ARGS()
 
 	void Construct(const FArguments& InArgs);
@@ -1403,6 +1409,7 @@ void SAnimNotifyTrack::Construct(const FArguments& InArgs)
 	OnDeselectAllNotifies = InArgs._OnDeselectAllNotifies;
 	OnCopyNotifies = InArgs._OnCopyNotifies;
 	OnPasteNotifies = InArgs._OnPasteNotifies;
+	OnSetInputViewRange = InArgs._OnSetInputViewRange;
 
 	this->ChildSlot
 	[
@@ -1561,6 +1568,26 @@ int32 SAnimNotifyTrack::OnPaint(const FPaintArgs& Args, const FGeometry& Allotte
 	return SCompoundWidget::OnPaint(Args, AllottedGeometry, MyClippingRect, OutDrawElements, CustomLayerId, InWidgetStyle, bParentEnabled);
 }
 
+FReply SAnimNotifyTrack::OnMouseWheel(const FGeometry& MyGeometry, const FPointerEvent& MouseEvent)
+{
+	const float ZoomDelta = -0.1f * MouseEvent.GetWheelDelta();
+
+	const FVector2D WidgetSpace = MyGeometry.AbsoluteToLocal(MouseEvent.GetScreenSpacePosition());
+	const float ZoomRatio = FMath::Clamp( (WidgetSpace.X / MyGeometry.Size.X), 0.f, 1.f);
+
+	{
+		const float InputViewSize = ViewInputMax.Get() - ViewInputMin.Get();
+		const float InputChange = InputViewSize * ZoomDelta;
+
+		float ViewMinInput = ViewInputMin.Get() - (InputChange * ZoomRatio);
+		float ViewMaxInput = ViewInputMax.Get() + (InputChange * (1.f - ZoomRatio));
+
+		OnSetInputViewRange.ExecuteIfBound(ViewMinInput, ViewMaxInput);
+	}
+
+	return FReply::Handled();
+}
+
 void SAnimNotifyTrack::FillNewNotifyStateMenu(FMenuBuilder& MenuBuilder)
 {
 	// Load the asset registry module
@@ -2678,6 +2705,7 @@ void SNotifyEdTrack::Construct(const FArguments& InArgs)
 				.OnDeselectAllNotifies(InArgs._OnDeselectAllNotifies)
 				.OnCopyNotifies(InArgs._OnCopyNotifies)
 				.OnPasteNotifies(InArgs._OnPasteNotifies)
+				.OnSetInputViewRange(InArgs._OnSetInputViewRange)
 			]
 
 			+SHorizontalBox::Slot()
@@ -2963,6 +2991,7 @@ void SAnimNotifyPanel::RefreshNotifyTracks()
 			.OnDeselectAllNotifies(this, &SAnimNotifyPanel::DeselectAllNotifies)
 			.OnCopyNotifies(this, &SAnimNotifyPanel::CopySelectedNotifiesToClipboard)
 			.OnPasteNotifies(this, &SAnimNotifyPanel::OnPasteNotifies)
+			.OnSetInputViewRange(OnSetInputViewRange)
 		];
 
 		NotifyAnimTracks.Add(EdTrack->NotifyTrack);
diff --git a/Engine/Source/Editor/Persona/Private/SAnimTrackPanel.cpp b/Engine/Source/Editor/Persona/Private/SAnimTrackPanel.cpp
--- a/Engine/Source/Editor/Persona/Private/SAnimTrackPanel.cpp
+++ b/Engine/Source/Editor/Persona/Private/SAnimTrackPanel.cpp
@@ -77,12 +77,15 @@ FReply SAnimTrackPanel::OnMouseWheel( const FGeometry& MyGeometry, const FPointe
 {
 	const float ZoomDelta = -0.1f * MouseEvent.GetWheelDelta();
 
+	const FVector2D MouseWidgetPos = MyGeometry.AbsoluteToLocal(MouseEvent.GetScreenSpacePosition());
+	const float ZoomRatio = FMath::Clamp((MouseWidgetPos.X / (MyGeometry.Size.X - WidgetWidth)), 0.f, 1.f);
+
 	{
 		const float InputViewSize = ViewInputMax.Get() - ViewInputMin.Get();
 		const float InputChange = InputViewSize * ZoomDelta;
 
-		float ViewMinInput = ViewInputMin.Get() - (InputChange * 0.5f);
-		float ViewMaxInput = ViewInputMax.Get() + (InputChange * 0.5f);
+		float ViewMinInput = ViewInputMin.Get() - (InputChange * ZoomRatio);
+		float ViewMaxInput = ViewInputMax.Get() + (InputChange * (1.f - ZoomRatio));
 		
 		OnSetInputViewRange.Execute(ViewMinInput, ViewMaxInput);
 	}
diff --git a/Engine/Source/Editor/UnrealEd/Public/SCurveEditor.h b/Engine/Source/Editor/UnrealEd/Public/SCurveEditor.h
--- a/Engine/Source/Editor/UnrealEd/Public/SCurveEditor.h
+++ b/Engine/Source/Editor/UnrealEd/Public/SCurveEditor.h
@@ -393,6 +393,9 @@ protected:
 	UNREALED_API virtual int32 OnPaint( const FPaintArgs& Args, const FGeometry& AllottedGeometry, const FSlateRect& MyClippingRect, 
 		FSlateWindowElementList& OutDrawElements, int32 LayerId, const FWidgetStyle& InWidgetStyle, bool bParentEnabled ) const ;
 
+	/** Update view range */
+	UNREALED_API void SetInputMinMax(float NewMin, float NewMax);
+
 private:
 
 	/** Curve selection*/
@@ -468,7 +471,6 @@ private:
 	/** Gradient editor */
 	TSharedPtr<class SColorGradientEditor> GradientViewer;
 
-	void SetInputMinMax(float NewMin, float NewMax);
 protected:
 	/** Minimum input of view range  */
 	TAttribute<float>	ViewMinInput;
