Fixed TTP 342235 : PERSONA: Animation Montage: No slot name warning displays on wrong slot

Fixed TTP 342236 : Regression: PERSONA: Animation Montage: Slot name has last letter cut off

diff --git a/Engine/Source/Editor/Persona/Private/SAnimMontagePanel.cpp b/Engine/Source/Editor/Persona/Private/SAnimMontagePanel.cpp
--- a/Engine/Source/Editor/Persona/Private/SAnimMontagePanel.cpp
+++ b/Engine/Source/Editor/Persona/Private/SAnimMontagePanel.cpp
@@ -387,6 +387,10 @@ void SAnimMontagePanel::Update()
 		// Anim Segment Tracks
 		// ===================================
 		{
+			int32 NumAnimTracks = Montage->SlotAnimTracks.Num();
+			SlotNameTextBoxes.Empty(NumAnimTracks);
+			SlotNameTextBoxes.AddZeroed(NumAnimTracks);
+
 			for (int32 SlotAnimIdx=0; SlotAnimIdx < Montage->SlotAnimTracks.Num(); SlotAnimIdx++)
 			{
 				TSharedRef<S2ColumnWidget> SectionTrack = Create2ColumnWidget(MontageSlots.ToSharedRef());
@@ -402,7 +406,7 @@ void SAnimMontagePanel::Update()
 					.AutoWidth()
 					.Padding( FMargin(5.0f, 0.5f) )
 					[
-						SAssignNew(SlotNameTextBox, SEditableTextBox)
+						SAssignNew(SlotNameTextBoxes[SlotAnimIdx], SEditableTextBox)
 						.Style( FEditorStyle::Get(), "SpecialEditableTextBox" )
 						.HintText( LOCTEXT("SlotName", "Slot Name") )
 						.Text(this, &SAnimMontagePanel::GetMontageSlotName, SlotAnimIdx)
@@ -411,7 +415,7 @@ void SAnimMontagePanel::Update()
 						.RevertTextOnEscape( true )
 						.ClearKeyboardFocusOnCommit( false )
 						.OnTextCommitted( this, &SAnimMontagePanel::OnSlotNodeNameChangeCommit, SlotAnimIdx )
-						.OnTextChanged(this, &SAnimMontagePanel::OnSlotNameChanged)
+						.OnTextChanged(this, &SAnimMontagePanel::OnSlotNameChanged, SlotAnimIdx)
 					]
 				];
 
@@ -719,21 +723,26 @@ void SAnimMontagePanel::ClearSelected()
 	MontageEditor.Pin()->ClearDetailsView();
 }
 
-void SAnimMontagePanel::OnSlotNodeNameChangeCommit( const FText& NewText, ETextCommit::Type CommitInfo, int32 SlodeNodeIndex )
+void SAnimMontagePanel::OnSlotNodeNameChangeCommit(const FText& NewText, ETextCommit::Type CommitInfo, int32 SlotNodeIndex)
 {
-	MontageEditor.Pin()->RenameSlotNode(SlodeNodeIndex, NewText.ToString());
+	MontageEditor.Pin()->RenameSlotNode(SlotNodeIndex, NewText.ToString());
 }
 
-void SAnimMontagePanel::CheckSlotName(const FText& SlotName) const
+void SAnimMontagePanel::CheckSlotName(const FText& SlotName, int32 SlotNodeIndex) const
 {
+	if (!SlotNameTextBoxes.IsValidIndex(SlotNodeIndex))
+	{
+		return;
+	}
+
 	if (SlotName.IsEmpty())
 	{
 		FText Error = LOCTEXT("Error_SlotNameIsEmpty", "Please provide a slot name for this asset.");
-		SlotNameTextBox->SetError(Error);
+		SlotNameTextBoxes[SlotNodeIndex]->SetError(Error);
 	}
 	else
 	{
-		SlotNameTextBox->SetError(FText::GetEmpty());
+		SlotNameTextBoxes[SlotNodeIndex]->SetError(FText::GetEmpty());
 	}
 }
 
@@ -741,14 +750,14 @@ FText SAnimMontagePanel::GetMontageSlotName(int32 SlotIndex) const
 {
 	FText SlotName = MontageEditor.Pin()->GetMontageSlotName(SlotIndex);
 
-	CheckSlotName(SlotName);
+	CheckSlotName(SlotName, SlotIndex);
 
 	return SlotName;
 }
 
-void SAnimMontagePanel::OnSlotNameChanged(const FText& NewText)
+void SAnimMontagePanel::OnSlotNameChanged(const FText& NewText, int32 SlotNodeIndex)
 {
-	CheckSlotName(NewText);
+	CheckSlotName(NewText, SlotNodeIndex);
 }
 
 #undef LOCTEXT_NAMESPACE
\ No newline at end of file
diff --git a/Engine/Source/Editor/Persona/Private/SAnimMontagePanel.h b/Engine/Source/Editor/Persona/Private/SAnimMontagePanel.h
--- a/Engine/Source/Editor/Persona/Private/SAnimMontagePanel.h
+++ b/Engine/Source/Editor/Persona/Private/SAnimMontagePanel.h
@@ -78,15 +78,15 @@ public:
 	void ShowBranchPointInDetailsView(int32 BranchPointIndex);
 	void ShowSectionInDetailsView(int32 SectionIndex);
 
-	void OnSlotNodeNameChangeCommit( const FText& NewText, ETextCommit::Type CommitInfo, int32 SlodeNodeIndex );
+	void OnSlotNodeNameChangeCommit(const FText& NewText, ETextCommit::Type CommitInfo, int32 SlotNodeIndex);
 
 	void ClearSelected();
 
 	// helper method to check whether the slot name is empty or not. If empty, shows an error message to provide a valid name
-	void CheckSlotName(const FText& SlotName) const;
+	void CheckSlotName(const FText& SlotName, int32 SlotNodeIndex) const;
 
 	// check the slot name whether valid or not while the user is typing
-	void OnSlotNameChanged(const FText& NewText);
+	void OnSlotNameChanged(const FText& NewText, int32 SlotNodeIndex);
 
 	// get slot name from a montage editor and check the slot name whether valid or not
 	FText GetMontageSlotName(int32 SlotIndex) const;
@@ -105,7 +105,8 @@ private:
 
 	STrackNodeSelectionSet SelectionSet;
 
-	TSharedPtr<SEditableTextBox> SlotNameTextBox;
+	// text box array for multiple slot names
+	TArray<TSharedPtr<SEditableTextBox>> SlotNameTextBoxes;
 	
 	void CreateNewSlot(const FText& NewSlotName, ETextCommit::Type CommitInfo);
 	void CreateNewSection(const FText& NewSectionName, ETextCommit::Type CommitInfo, float StartTime);
