Fixed TTP #342237 : Crash: Expanding and collapsing montage sections

diff --git a/Engine/Source/Editor/Persona/Private/SAnimMontagePanel.cpp b/Engine/Source/Editor/Persona/Private/SAnimMontagePanel.cpp
--- a/Engine/Source/Editor/Persona/Private/SAnimMontagePanel.cpp
+++ b/Engine/Source/Editor/Persona/Private/SAnimMontagePanel.cpp
@@ -391,7 +391,7 @@ void SAnimMontagePanel::Update()
 			SlotNameTextBoxes.Empty(NumAnimTracks);
 			SlotNameTextBoxes.AddZeroed(NumAnimTracks);
 
-			for (int32 SlotAnimIdx=0; SlotAnimIdx < Montage->SlotAnimTracks.Num(); SlotAnimIdx++)
+			for (int32 SlotAnimIdx = 0; SlotAnimIdx < NumAnimTracks; SlotAnimIdx++)
 			{
 				TSharedRef<S2ColumnWidget> SectionTrack = Create2ColumnWidget(MontageSlots.ToSharedRef());
 
@@ -728,14 +728,24 @@ void SAnimMontagePanel::OnSlotNodeNameChangeCommit(const FText& NewText, ETextCo
 	MontageEditor.Pin()->RenameSlotNode(SlotNodeIndex, NewText.ToString());
 }
 
-void SAnimMontagePanel::CheckSlotName(const FText& SlotName, int32 SlotNodeIndex) const
+void SAnimMontagePanel::CheckSlotName(const FText& SlotName, int32 SlotNodeIndex, bool bShouldCheckCollapsed) const
 {
 	if (!SlotNameTextBoxes.IsValidIndex(SlotNodeIndex))
 	{
 		return;
 	}
 
-	if (SlotName.IsEmpty())
+	bool bCanSetError = true;
+	if (bShouldCheckCollapsed)
+	{
+		// Check whether the text box is collapsed or not. If collapsed, hide the error message
+		FWidgetPath WidgetPath;
+		FSlateApplication::Get().FindPathToWidget(SlotNameTextBoxes[SlotNodeIndex].ToSharedRef(), WidgetPath);
+
+		bCanSetError = WidgetPath.IsValid();
+	}
+
+	if (bCanSetError && SlotName.IsEmpty())
 	{
 		FText Error = LOCTEXT("Error_SlotNameIsEmpty", "Please provide a slot name for this asset.");
 		SlotNameTextBoxes[SlotNodeIndex]->SetError(Error);
@@ -750,7 +760,7 @@ FText SAnimMontagePanel::GetMontageSlotName(int32 SlotIndex) const
 {
 	FText SlotName = MontageEditor.Pin()->GetMontageSlotName(SlotIndex);
 
-	CheckSlotName(SlotName, SlotIndex);
+	CheckSlotName(SlotName, SlotIndex, true);
 
 	return SlotName;
 }
diff --git a/Engine/Source/Editor/Persona/Private/SAnimMontagePanel.h b/Engine/Source/Editor/Persona/Private/SAnimMontagePanel.h
--- a/Engine/Source/Editor/Persona/Private/SAnimMontagePanel.h
+++ b/Engine/Source/Editor/Persona/Private/SAnimMontagePanel.h
@@ -83,7 +83,7 @@ public:
 	void ClearSelected();
 
 	// helper method to check whether the slot name is empty or not. If empty, shows an error message to provide a valid name
-	void CheckSlotName(const FText& SlotName, int32 SlotNodeIndex) const;
+	void CheckSlotName(const FText& SlotName, int32 SlotNodeIndex, bool bShouldCheckCollapsed = false) const;
 
 	// check the slot name whether valid or not while the user is typing
 	void OnSlotNameChanged(const FText& NewText, int32 SlotNodeIndex);
