Added duplicate curve option to curve dropdown in applicable animation assets.

diff --git a/Engine/Source/Editor/Persona/Private/SAnimCurvePanel.cpp b/Engine/Source/Editor/Persona/Private/SAnimCurvePanel.cpp
--- a/Engine/Source/Editor/Persona/Private/SAnimCurvePanel.cpp
+++ b/Engine/Source/Editor/Persona/Private/SAnimCurvePanel.cpp
@@ -284,6 +284,9 @@ public:
 	// input handling for curve name
 	void NewCurveNameEntered( const FText& NewText, ETextCommit::Type CommitInfo );
 
+	// Duplicate the current track
+	void DuplicateTrack();
+
 	// Delete current track
 	void DeleteTrack();
 
@@ -436,6 +439,15 @@ SCurveEdTrack::~SCurveEdTrack()
 	delete CurveInterface;
 }
 
+void SCurveEdTrack::DuplicateTrack()
+{
+	TSharedPtr<SAnimCurvePanel> SharedPanel = PanelPtr.Pin();
+	if(SharedPanel.IsValid())
+	{
+		SharedPanel->DuplicateTrack(CurveInterface->CurveData->CurveName.ToString());
+	}
+}
+
 void SCurveEdTrack::DeleteTrack()
 {
 	if ( PanelPtr.IsValid() )
@@ -525,6 +537,13 @@ FReply SCurveEdTrack::OnContextMenu()
 
 	MenuBuilder.BeginSection("AnimCurvePanelTrackOptions", LOCTEXT("TrackOptionsHeading", "Track Options") );
 	{
+		MenuBuilder.AddMenuEntry(
+			LOCTEXT("DuplicateTrack", "Duplicate Track"),
+			LOCTEXT("DuplicateTrackTooltip", "Duplicate this track"),
+			FSlateIcon(),
+			FUIAction(FExecuteAction::CreateSP(this, &SCurveEdTrack::DuplicateTrack))
+		);
+
 		MenuBuilder.AddMenuEntry(
 			LOCTEXT("RemoveTrack", "Remove Track"),
 			LOCTEXT("RemoveTrackTooltip", "Remove this track"),
@@ -684,6 +703,20 @@ void SAnimCurvePanel::AddTrack(const FText & CurveNameToAdd, ETextCommit::Type C
 	}
 }
 
+FReply SAnimCurvePanel::DuplicateTrack(const FString& CurveNameToDuplicate)
+{
+	const FScopedTransaction Transaction( LOCTEXT("AnimCurve_DuplicateTrack", "Duplicate Curve") );
+	
+	if(Sequence->RawCurveData.DuplicateCurveData(*CurveNameToDuplicate))
+	{
+		Sequence->Modify();
+		UpdatePanel(true);
+
+		return FReply::Handled();
+	}
+	return FReply::Unhandled();
+}
+
 FReply SAnimCurvePanel::DeleteTrack(const FString & CurveNameToDelete)
 {
 	const FScopedTransaction Transaction( LOCTEXT("AnimCurve_AddTrack", "Add New Curve") );
diff --git a/Engine/Source/Editor/Persona/Private/SAnimCurvePanel.h b/Engine/Source/Editor/Persona/Private/SAnimCurvePanel.h
--- a/Engine/Source/Editor/Persona/Private/SAnimCurvePanel.h
+++ b/Engine/Source/Editor/Persona/Private/SAnimCurvePanel.h
@@ -67,6 +67,11 @@ public:
 	 */
 	void AddTrack(const FText & CurveNameToAdd, ETextCommit::Type CommitInfo);
 	
+	/**
+	 * Duplicate Track
+	 */
+	FReply DuplicateTrack(const FString& CurveNameToDuplicate);
+
 	/**
 	 * Delete Track
 	 *
diff --git a/Engine/Source/Runtime/Engine/Classes/Animation/AnimSequenceBase.h b/Engine/Source/Runtime/Engine/Classes/Animation/AnimSequenceBase.h
--- a/Engine/Source/Runtime/Engine/Classes/Animation/AnimSequenceBase.h
+++ b/Engine/Source/Runtime/Engine/Classes/Animation/AnimSequenceBase.h
@@ -9,6 +9,7 @@
 
 #include "AnimationAsset.h"
 #include "Curves/CurveBase.h"
+#include "Kismet2NameValidators.h"
 #include "AnimSequenceBase.generated.h"
 
 #define DEFAULT_SAMPLERATE			30.f
@@ -256,6 +257,27 @@ struct FRawCurveTracks
 	 * Delete curve data 
 	 */
 	ENGINE_API bool DeleteCurveData(FName InCurveName);
+	/**
+	 * Duplicate curve data
+	 * Duplicates a curve, creating a new name using a validator (FCurveNameValidator)
+	 */
+	ENGINE_API bool DuplicateCurveData(FName InCurveName);
+};
+
+/**
+ * Name validator for anim curves
+ */
+class FCurveNameValidator : public FStringSetNameValidator
+{
+public:
+	FCurveNameValidator(FRawCurveTracks& Tracks, const FString& ExistingName)
+		: FStringSetNameValidator(ExistingName)
+	{
+		for(FFloatCurve& Curve : Tracks.FloatCurves)
+		{
+			Names.Add(Curve.CurveName.ToString());
+		}
+	}
 };
 
 UENUM()
diff --git a/Engine/Source/Runtime/Engine/Private/Animation/AnimSequenceBase.cpp b/Engine/Source/Runtime/Engine/Private/Animation/AnimSequenceBase.cpp
--- a/Engine/Source/Runtime/Engine/Private/Animation/AnimSequenceBase.cpp
+++ b/Engine/Source/Runtime/Engine/Private/Animation/AnimSequenceBase.cpp
@@ -148,6 +148,25 @@ bool FRawCurveTracks::AddCurveData(FName InCurveName, int32 CurveFlags /*= ACF_D
 	return false;
 }
 
+ENGINE_API bool FRawCurveTracks::DuplicateCurveData(FName InCurveName)
+{
+	if(FFloatCurve* ExistingCurve = GetCurveData(InCurveName))
+	{
+		TSharedPtr<INameValidatorInterface> Validator = MakeShareable(new FCurveNameValidator(*this, FString(TEXT(""))));
+
+		// Use the validator to pick a reasonable name for the duplicated curve.
+		FString NewCurveName = InCurveName.ToString();
+		Validator->FindValidString(NewCurveName);
+
+		// Add the curve to the track and set its data to the existing curve
+		FloatCurves.Add(FFloatCurve(*NewCurveName, ExistingCurve->GetCurveTypeFlags()));
+		FloatCurves.Last().FloatCurve = ExistingCurve->FloatCurve;
+		
+		return true;
+	}
+	return false;
+}
+
 /////////////////////////////////////////////////////
 
 UAnimSequenceBase::UAnimSequenceBase(const class FPostConstructInitializeProperties& PCIP)
