CIS fix from CL#2234570 - Moved editor only derived code out of non edior only files.

diff --git a/Engine/Source/Editor/Persona/Private/SAnimCurvePanel.cpp b/Engine/Source/Editor/Persona/Private/SAnimCurvePanel.cpp
--- a/Engine/Source/Editor/Persona/Private/SAnimCurvePanel.cpp
+++ b/Engine/Source/Editor/Persona/Private/SAnimCurvePanel.cpp
@@ -8,6 +8,7 @@
 #include "SCurveEditor.h"
 #include "Editor/KismetWidgets/Public/SScrubWidget.h"
 #include "AssetRegistryModule.h"
+#include "Kismet2NameValidators.h"
 
 #define LOCTEXT_NAMESPACE "AnimCurvePanel"
 
@@ -588,6 +589,22 @@ FVector2D SCurveEdTrack::GetDesiredSize() const
 //////////////////////////////////////////////////////////////////////////
 // SAnimCurvePanel
 
+/**
+ * Name validator for anim curves
+ */
+class FCurveNameValidator : public FStringSetNameValidator
+{
+public:
+	FCurveNameValidator(FRawCurveTracks& Tracks, const FString& ExistingName)
+		: FStringSetNameValidator(ExistingName)
+	{
+		for(FFloatCurve& Curve : Tracks.FloatCurves)
+		{
+			Names.Add(Curve.CurveName.ToString());
+		}
+	}
+};
+
 void SAnimCurvePanel::Construct(const FArguments& InArgs)
 {
 	SAnimTrackPanel::Construct( SAnimTrackPanel::FArguments()
@@ -707,7 +724,13 @@ FReply SAnimCurvePanel::DuplicateTrack(const FString& CurveNameToDuplicate)
 {
 	const FScopedTransaction Transaction( LOCTEXT("AnimCurve_DuplicateTrack", "Duplicate Curve") );
 	
-	if(Sequence->RawCurveData.DuplicateCurveData(*CurveNameToDuplicate))
+	TSharedPtr<INameValidatorInterface> Validator = MakeShareable(new FCurveNameValidator(Sequence->RawCurveData, FString(TEXT(""))));
+
+	// Use the validator to pick a reasonable name for the duplicated curve.
+	FString NewCurveName = CurveNameToDuplicate;
+	Validator->FindValidString(NewCurveName);
+
+	if(Sequence->RawCurveData.DuplicateCurveData(*CurveNameToDuplicate, *NewCurveName))
 	{
 		Sequence->Modify();
 		UpdatePanel(true);
diff --git a/Engine/Source/Runtime/Engine/Classes/Animation/AnimSequenceBase.h b/Engine/Source/Runtime/Engine/Classes/Animation/AnimSequenceBase.h
--- a/Engine/Source/Runtime/Engine/Classes/Animation/AnimSequenceBase.h
+++ b/Engine/Source/Runtime/Engine/Classes/Animation/AnimSequenceBase.h
@@ -9,7 +9,6 @@
 
 #include "AnimationAsset.h"
 #include "Curves/CurveBase.h"
-#include "Kismet2NameValidators.h"
 #include "AnimSequenceBase.generated.h"
 
 #define DEFAULT_SAMPLERATE			30.f
@@ -259,25 +258,9 @@ struct FRawCurveTracks
 	ENGINE_API bool DeleteCurveData(FName InCurveName);
 	/**
 	 * Duplicate curve data
-	 * Duplicates a curve, creating a new name using a validator (FCurveNameValidator)
+	 * 
 	 */
-	ENGINE_API bool DuplicateCurveData(FName InCurveName);
-};
-
-/**
- * Name validator for anim curves
- */
-class FCurveNameValidator : public FStringSetNameValidator
-{
-public:
-	FCurveNameValidator(FRawCurveTracks& Tracks, const FString& ExistingName)
-		: FStringSetNameValidator(ExistingName)
-	{
-		for(FFloatCurve& Curve : Tracks.FloatCurves)
-		{
-			Names.Add(Curve.CurveName.ToString());
-		}
-	}
+	ENGINE_API bool DuplicateCurveData(FName CurveToCopyName, FName NewName);
 };
 
 UENUM()
diff --git a/Engine/Source/Runtime/Engine/Private/Animation/AnimSequenceBase.cpp b/Engine/Source/Runtime/Engine/Private/Animation/AnimSequenceBase.cpp
--- a/Engine/Source/Runtime/Engine/Private/Animation/AnimSequenceBase.cpp
+++ b/Engine/Source/Runtime/Engine/Private/Animation/AnimSequenceBase.cpp
@@ -148,18 +148,13 @@ bool FRawCurveTracks::AddCurveData(FName InCurveName, int32 CurveFlags /*= ACF_D
 	return false;
 }
 
-ENGINE_API bool FRawCurveTracks::DuplicateCurveData(FName InCurveName)
+ENGINE_API bool FRawCurveTracks::DuplicateCurveData(FName CurveToCopyName, FName NewName)
 {
-	if(FFloatCurve* ExistingCurve = GetCurveData(InCurveName))
+	FFloatCurve* ExistingCurve = GetCurveData(CurveToCopyName);
+	if(ExistingCurve != NULL && GetCurveData(NewName) == NULL)
 	{
-		TSharedPtr<INameValidatorInterface> Validator = MakeShareable(new FCurveNameValidator(*this, FString(TEXT(""))));
-
-		// Use the validator to pick a reasonable name for the duplicated curve.
-		FString NewCurveName = InCurveName.ToString();
-		Validator->FindValidString(NewCurveName);
-
 		// Add the curve to the track and set its data to the existing curve
-		FloatCurves.Add(FFloatCurve(*NewCurveName, ExistingCurve->GetCurveTypeFlags()));
+		FloatCurves.Add(FFloatCurve(NewName, ExistingCurve->GetCurveTypeFlags()));
 		FloatCurves.Last().FloatCurve = ExistingCurve->FloatCurve;
 		
 		return true;
