New sockets are created with the bone name instead of "NewSocket"

diff --git a/Engine/Source/Editor/Persona/Private/SSkeletonTree.cpp b/Engine/Source/Editor/Persona/Private/SSkeletonTree.cpp
--- a/Engine/Source/Editor/Persona/Private/SSkeletonTree.cpp
+++ b/Engine/Source/Editor/Persona/Private/SSkeletonTree.cpp
@@ -1749,7 +1749,8 @@ void SSkeletonTree::OnAddSocket()
 		check(NewSocket);
 
 		NewSocket->BoneName = *static_cast<FName*>( TreeSelection.GetSingleSelectedItem()->GetData() );
-		NewSocket->SocketName = PersonaPtr.Pin()->GenerateUniqueSocketName( *LOCTEXT( "NewSocketDefaultName", "NewSocket" ).ToString() );
+		FString SocketName = NewSocket->BoneName.ToString() + LOCTEXT("SocketPostfix", "Socket").ToString();
+		NewSocket->SocketName = PersonaPtr.Pin()->GenerateUniqueSocketName( *SocketName );
 
 		TargetSkeleton->Sockets.Add( NewSocket );
 
