Fix for crash on create animation for other skeleton

diff --git a/Engine/Source/Editor/Persona/Private/SAnimationSequenceBrowser.cpp b/Engine/Source/Editor/Persona/Private/SAnimationSequenceBrowser.cpp
--- a/Engine/Source/Editor/Persona/Private/SAnimationSequenceBrowser.cpp
+++ b/Engine/Source/Editor/Persona/Private/SAnimationSequenceBrowser.cpp
@@ -304,7 +304,7 @@ void SAnimationSequenceBrowser::OnCreateCopy(TArray<FAssetData> Selected)
 			{
 				FAssetRegistryModule::AssetCreated(AssetToOpen);
 				// once all success, attempt to open new persona module with new skeleton
-				EToolkitMode::Type Mode = EToolkitMode::WorldCentric;
+				EToolkitMode::Type Mode = EToolkitMode::Standalone;
 				FPersonaModule& PersonaModule = FModuleManager::LoadModuleChecked<FPersonaModule>( "Persona" );
 				PersonaModule.CreatePersona( Mode, TSharedPtr<IToolkitHost>(), NewSkeleton, NULL, AnimAsset, NULL );
 			}
