Removed code path in curve panel that wouldn't refresh the entire panel. This stops an intermittent crash when adding curves where a realloc can break widgets further down the chain from the panel.

diff --git a/Engine/Source/Editor/Persona/Private/SAnimCurvePanel.cpp b/Engine/Source/Editor/Persona/Private/SAnimCurvePanel.cpp
--- a/Engine/Source/Editor/Persona/Private/SAnimCurvePanel.cpp
+++ b/Engine/Source/Editor/Persona/Private/SAnimCurvePanel.cpp
@@ -733,7 +733,7 @@ FReply SAnimCurvePanel::DuplicateTrack(const FString& CurveNameToDuplicate)
 	if(Sequence->RawCurveData.DuplicateCurveData(*CurveNameToDuplicate, *NewCurveName))
 	{
 		Sequence->Modify();
-		UpdatePanel(true);
+		UpdatePanel();
 
 		return FReply::Handled();
 	}
@@ -746,7 +746,7 @@ FReply SAnimCurvePanel::DeleteTrack(const FString & CurveNameToDelete)
 	if ( Sequence->RawCurveData.DeleteCurveData(*CurveNameToDelete) )
 	{
 		Sequence->Modify(true);
-		UpdatePanel(true);
+		UpdatePanel();
 		return FReply::Handled();
 	}
 
@@ -841,40 +841,23 @@ ESlateCheckBoxState::Type SAnimCurvePanel::AreAllCurvesOfMode(EAnimCurveFlags Mo
 	return ESlateCheckBoxState::Undetermined;
 }
 
-void SAnimCurvePanel::UpdatePanel(bool bClearAll/*=false*/)
+void SAnimCurvePanel::UpdatePanel()
 {
 	if(Sequence != NULL)
 	{
 		// see if we need to clear or not
 		FChildren * Children = PanelSlot->GetChildren();
 
-		int32 TotalCurve = Sequence->RawCurveData.FloatCurves.Num();
-		int32 CurveEditor = Children->Num();
-		int32 CurrentIt = 0;
-
-		// a curve is removed, refresh all window
-		if (bClearAll || (CurveEditor > TotalCurve))
+		for(int32 Id = Children->Num() - 1; Id >= 0; --Id)
 		{
-			for (int32 Id=Children->Num()-1; Id>=0; --Id)
-			{
-				PanelSlot->RemoveAt(Id);
-			}
-
-			// Clear all tracks as we're re-adding them all anyway.
-			Tracks.Empty();
-
-			// refresh from 0 index
-			CurrentIt = 0;
-		}
-		else
-		{
-			// otherwise start from new curve, 
-			// this way, we don't destroy old curves
-			CurrentIt = CurveEditor;
+			PanelSlot->RemoveAt(Id);
 		}
+		// Clear all tracks as we're re-adding them all anyway.
+		Tracks.Empty();
 
 		// Updating new tracks
-		for (; CurrentIt<TotalCurve; ++CurrentIt)
+		int32 TotalCurve = Sequence->RawCurveData.FloatCurves.Num();
+		for(int32 CurrentIt = 0 ; CurrentIt < TotalCurve ; ++CurrentIt)
 		{
 			FFloatCurve&  Curve = Sequence->RawCurveData.FloatCurves[CurrentIt];
 
@@ -912,7 +895,7 @@ void SAnimCurvePanel::SetSequence(class UAnimSequenceBase *	InSequence)
 	if (InSequence != Sequence)
 	{
 		Sequence = InSequence;
-		UpdatePanel(true);
+		UpdatePanel();
 	}
 }
 
@@ -1045,7 +1028,7 @@ void SAnimCurvePanel::ToggleEditability(ESlateCheckBoxState::Type NewType, FName
 
 FReply		SAnimCurvePanel::RefreshPanel()
 {
-	UpdatePanel(true);
+	UpdatePanel();
 	return FReply::Handled();
 }
 
@@ -1059,7 +1042,7 @@ FReply		SAnimCurvePanel::ShowAll(bool bShow)
 			Curve.SetCurveTypeFlag(ACF_Editable, bShow);
 		}
 
-		UpdatePanel(true);
+		UpdatePanel();
 	}
 
 	return FReply::Handled();
diff --git a/Engine/Source/Editor/Persona/Private/SAnimCurvePanel.h b/Engine/Source/Editor/Persona/Private/SAnimCurvePanel.h
--- a/Engine/Source/Editor/Persona/Private/SAnimCurvePanel.h
+++ b/Engine/Source/Editor/Persona/Private/SAnimCurvePanel.h
@@ -131,7 +131,7 @@ private:
 private:
 	/**
 	 * Update Panel 
-	 * @todo this has to be more efficient. Right now it does refresh whole things if it has to
+	 * @todo this has to be more efficient. Right now it refreshes the entire panel
 	 */
-	void UpdatePanel(bool bClearAll=false);
+	void UpdatePanel();
 };
\ No newline at end of file
