Fixes to anim notify placement issues.
- Non-state notifies shouldn't end up with end offsets
- State notifies no longer jump around when dragged
- Begin and end snaps no longer sometimes offset on the wrong side

diff --git a/Engine/Source/Editor/Persona/Private/SAnimNotifyPanel.cpp b/Engine/Source/Editor/Persona/Private/SAnimNotifyPanel.cpp
--- a/Engine/Source/Editor/Persona/Private/SAnimNotifyPanel.cpp
+++ b/Engine/Source/Editor/Persona/Private/SAnimNotifyPanel.cpp
@@ -240,8 +240,10 @@ public:
 
 	FTrackScaleInfo GetCachedScaleInfo() const { return FTrackScaleInfo(ViewInputMin.Get(), ViewInputMax.Get(), 0.f, 0.f, CachedGeometry.Size); }
 
-	/** Updates sequences when a notify node has been successfully dragged to a new position */
-	void HandleNodeDrop(TSharedPtr<SAnimNotifyNode> Node);
+	/** Updates sequences when a notify node has been successfully dragged to a new position
+	 *	@param Offset - Offset from the widget to the time handle 
+	 */
+	void HandleNodeDrop(TSharedPtr<SAnimNotifyNode> Node, float Offset = 0.0f);
 
 	// Number of nodes in the track currently selected
 	int32 GetNumSelectedNodes() const { return SelectedNotifyIndices.Num(); }
@@ -504,10 +506,13 @@ public:
 	{
 		if ( bDropWasHandled == false )
 		{
-			for(TSharedPtr<SAnimNotifyNode> Node : SelectedNodes)
+			int32 NumNodes = SelectedNodes.Num();
+			for(int32 CurrentNode = 0 ; CurrentNode < NumNodes ; ++CurrentNode)
 			{
+				TSharedPtr<SAnimNotifyNode> Node = SelectedNodes[CurrentNode];
+				float NodePositionOffset = NodeXOffsets[CurrentNode];
 				const FTrackClampInfo& ClampInfo = GetTrackClampInfo(Node->GetScreenPosition());
-				ClampInfo.NotifyTrack->HandleNodeDrop(Node);
+				ClampInfo.NotifyTrack->HandleNodeDrop(Node, NodePositionOffset);
 				Node->DropCancelled();
 			}
 		}
@@ -565,11 +570,21 @@ public:
 			// Look for a snap on the first scrub handle
 			FVector2D TrackNodePos = TrackGeom.AbsoluteToLocal(EventPosition);
 			const FVector2D OriginalNodePosition = TrackNodePos;
+			float SequenceEnd = TrackScaleInfo.InputToLocalX(Sequence->SequenceLength);
 
 			float SnapX = GetSnapPosition(NodeClamp, TrackNodePos.X);
 			if(SnapX >= 0.f)
 			{
-				EAnimEventTriggerOffsets::Type Offset = (SnapX < TrackNodePos.X) ? EAnimEventTriggerOffsets::OffsetAfter : EAnimEventTriggerOffsets::OffsetBefore;
+				EAnimEventTriggerOffsets::Type Offset = EAnimEventTriggerOffsets::NoOffset;
+				if(SnapX == 0.0f || SnapX == SequenceEnd)
+				{
+					Offset = SnapX > 0.0f ? EAnimEventTriggerOffsets::OffsetBefore : EAnimEventTriggerOffsets::OffsetAfter;
+				}
+				else
+				{
+					Offset = (SnapX < TrackNodePos.X) ? EAnimEventTriggerOffsets::OffsetAfter : EAnimEventTriggerOffsets::OffsetBefore;
+				}
+
 				CurrentEvent->TriggerTimeOffset = GetTriggerTimeOffsetForType(Offset);
 				CurrentNode->SetLastSnappedTime(TrackScaleInfo.LocalXToInput(SnapX));
 
@@ -601,7 +616,16 @@ public:
 					// Only attempt to snap if the node will fit on the track
 					if(SnapX >= CurrentNode->GetDurationSize())
 					{
-						EAnimEventTriggerOffsets::Type Offset = (SnapX < TrackNodeEndPos.X) ? EAnimEventTriggerOffsets::OffsetAfter : EAnimEventTriggerOffsets::OffsetBefore;
+						EAnimEventTriggerOffsets::Type Offset = EAnimEventTriggerOffsets::NoOffset;
+						if(SnapX == SequenceEnd)
+						{
+							// Only need to check the end of the sequence here; end handle can't hit the beginning
+							Offset = EAnimEventTriggerOffsets::OffsetBefore;
+						}
+						else
+						{
+							Offset = (SnapX < TrackNodeEndPos.X) ? EAnimEventTriggerOffsets::OffsetAfter : EAnimEventTriggerOffsets::OffsetBefore;
+						}
 						CurrentEvent->EndTriggerTimeOffset = GetTriggerTimeOffsetForType(Offset);
 						
 						if(SnapMovement == 0.0f)
@@ -667,6 +691,23 @@ public:
 				}
 			}
 		}
+
+		if(SnapPosition == -1.0f)
+		{
+			// Didn't snap to a bar, snap to the track bounds
+			float WidgetSpaceEndPosition = ScaleInfo.InputToLocalX(Sequence->SequenceLength);
+			float SnapDistBegin = FMath::Abs(-WidgetSpaceNotifyPosition);
+			float SnapDistEnd = FMath::Abs(WidgetSpaceEndPosition - WidgetSpaceNotifyPosition);
+			if(SnapDistBegin < CurrentMinSnapDest)
+			{
+				SnapPosition = 0.0f;
+			}
+			else if(SnapDistEnd < CurrentMinSnapDest)
+			{
+				SnapPosition = WidgetSpaceEndPosition;
+			}
+		}
+
 		return SnapPosition;
 	}
 
@@ -697,6 +738,7 @@ public:
 	TAttribute<TArray<FTrackMarkerBar>>	MarkerBars;				// Branching point markers
 	TArray<TSharedPtr<SAnimNotifyNode>> SelectedNodes;			// The nodes that are in the current selection
 	TArray<float>						NodeTimeOffsets;		// Time offsets from the beginning of the selection to the nodes.
+	TArray<float>						NodeXOffsets;			// Offsets in X from the widget position to the scrub handle for each node.
 	FVector2D							NodeGroupPosition;		// Position of the beginning of the selection
 	FVector2D							NodeGroupSize;			// Size of the entire selection
 	TSharedPtr<SWidget>					Decorator;				// The widget to display when dragging
@@ -716,7 +758,7 @@ public:
 		Operation->Decorator = Decorator;
 		Operation->SelectedNodes = NotifyNodes;
 		Operation->TrackSpan = NotifyNodes[0]->NotifyEvent->TrackIndex - NotifyNodes.Last()->NotifyEvent->TrackIndex;
-
+		
 		// Caclulate offsets for the selected nodes
 		float BeginTime = MAX_flt;
 		for(TSharedPtr<SAnimNotifyNode> Node : NotifyNodes)
@@ -733,6 +775,7 @@ public:
 			Node->ClearLastSnappedTime();
 			Operation->NodeTimeOffsets.Add(Node->NotifyEvent->DisplayTime - BeginTime);
 			Operation->NodeTimes.Add(Node->NotifyEvent->DisplayTime);
+			Operation->NodeXOffsets.Add(Node->GetNotifyPositionOffset().X);
 
 			// Calculate the time length of the selection. Because it is possible to have states
 			// with arbitrary durations we need to search all of the nodes and find the furthest
@@ -2518,7 +2561,7 @@ FReply SAnimNotifyTrack::OnDrop( const FGeometry& MyGeometry, const FDragDropEve
 	return FReply::Unhandled();
 }
 
-void SAnimNotifyTrack::HandleNodeDrop(TSharedPtr<SAnimNotifyNode> Node)
+void SAnimNotifyTrack::HandleNodeDrop(TSharedPtr<SAnimNotifyNode> Node, float Offset)
 {
 	ensure(Node.IsValid());
 
@@ -2526,7 +2569,7 @@ void SAnimNotifyTrack::HandleNodeDrop(TSharedPtr<SAnimNotifyNode> Node)
 	const FScopedTransaction Transaction(LOCTEXT("MoveNotifyEvent", "Move Anim Notify"));
 	Sequence->Modify();
 
-	float LocalX = GetCachedGeometry().AbsoluteToLocal(Node->GetScreenPosition() + ScrubHandleSize).X;
+	float LocalX = GetCachedGeometry().AbsoluteToLocal(Node->GetScreenPosition() + Offset).X;
 	float Time = GetCachedScaleInfo().LocalXToInput(LocalX);
 	FAnimNotifyEvent* DroppedEvent = Node->NotifyEvent;
 
@@ -2539,7 +2582,16 @@ void SAnimNotifyTrack::HandleNodeDrop(TSharedPtr<SAnimNotifyNode> Node)
 		DroppedEvent->DisplayTime = Time;
 	}
 	DroppedEvent->RefreshTriggerOffset(Sequence->CalculateOffsetForNotify(DroppedEvent->DisplayTime));
-	DroppedEvent->RefreshEndTriggerOffset(Sequence->CalculateOffsetForNotify(DroppedEvent->DisplayTime + DroppedEvent->Duration));
+
+	if(DroppedEvent->Duration > 0.0f)
+	{
+		DroppedEvent->RefreshEndTriggerOffset(Sequence->CalculateOffsetForNotify(DroppedEvent->DisplayTime + DroppedEvent->Duration));
+	}
+	else
+	{
+		DroppedEvent->EndTriggerTimeOffset = 0.0f;
+	}
+
 	DroppedEvent->TrackIndex = TrackIndex;
 
 	Sequence->MarkPackageDirty();
@@ -3014,7 +3066,8 @@ FReply SAnimNotifyPanel::OnNotifyNodeDragStarted(TArray<TSharedPtr<SAnimNotifyNo
 	}
 
 	FVector2D OverlayOrigin = Nodes[0]->GetScreenPosition();
-	FVector2D OverlayExtents = Nodes[0]->GetScreenPosition();
+	FVector2D OverlayExtents = OverlayOrigin;
+	OverlayExtents.X += Nodes[0]->GetDurationSize();
 	for(int32 Idx = 1 ; Idx < Nodes.Num() ; ++Idx)
 	{
 		TSharedPtr<SAnimNotifyNode> Node = Nodes[Idx];
diff --git a/Engine/Source/Editor/Persona/Private/SMontageEditor.cpp b/Engine/Source/Editor/Persona/Private/SMontageEditor.cpp
--- a/Engine/Source/Editor/Persona/Private/SMontageEditor.cpp
+++ b/Engine/Source/Editor/Persona/Private/SMontageEditor.cpp
@@ -174,6 +174,10 @@ void SMontageEditor::RefreshNotifyTriggerOffsets()
 			PredictedOffset = MontageObj->CalculateOffsetForNotify(Notify.DisplayTime + Notify.Duration);
 			Notify.RefreshEndTriggerOffset(PredictedOffset);
 		}
+		else
+		{
+			Notify.EndTriggerTimeOffset = 0.0f;
+		}
 	}
 
 	for(auto Iter = MontageObj->BranchingPoints.CreateIterator(); Iter; ++Iter)
diff --git a/Engine/Source/Runtime/Core/Public/UObject/ObjectVersion.h b/Engine/Source/Runtime/Core/Public/UObject/ObjectVersion.h
--- a/Engine/Source/Runtime/Core/Public/UObject/ObjectVersion.h
+++ b/Engine/Source/Runtime/Core/Public/UObject/ObjectVersion.h
@@ -586,6 +586,8 @@ enum EUnrealEngineObjectUE4Version
 	VER_UE4_BLUEPRINT_USE_SCS_ROOTCOMPONENT_SCALE,
 	// Changed level streaming to have a linear color since the visualization doesn't gamma correct.
 	VER_UE4_LEVEL_STREAMING_DRAW_COLOR_TYPE_CHANGE,
+	// Cleared end triggers from non-state anim notifies
+	VER_UE4_CLEAR_NOTIFY_TRIGGERS,
 
 	// -----<new versions can be added before this line>-------------------------------------------------
 	// - this needs to be the last line (see note below)
diff --git a/Engine/Source/Runtime/Engine/Private/Animation/AnimSequenceBase.cpp b/Engine/Source/Runtime/Engine/Private/Animation/AnimSequenceBase.cpp
--- a/Engine/Source/Runtime/Engine/Private/Animation/AnimSequenceBase.cpp
+++ b/Engine/Source/Runtime/Engine/Private/Animation/AnimSequenceBase.cpp
@@ -175,19 +175,34 @@ void UAnimSequenceBase::PostLoad()
 	Super::PostLoad();
 
 	// Convert Notifies to new data
-	if( GIsEditor && GetLinkerUE4Version() < VER_UE4_ANIMNOTIFY_NAMECHANGE && Notifies.Num() > 0 )
+	if( GIsEditor && Notifies.Num() > 0 )
 	{
-		LOG_SCOPE_VERBOSITY_OVERRIDE(LogAnimation, ELogVerbosity::Warning);
- 		// convert animnotifies
- 		for (int32 I=0; I<Notifies.Num(); ++I)
- 		{
- 			if (Notifies[I].Notify!=NULL)
- 			{
-				FString Label = Notifies[I].Notify->GetClass()->GetName();
-				Label = Label.Replace(TEXT("AnimNotify_"), TEXT(""), ESearchCase::CaseSensitive);
-				Notifies[I].NotifyName = FName(*Label);
- 			}
- 		}
+		if(GetLinkerUE4Version() < VER_UE4_ANIMNOTIFY_NAMECHANGE)
+		{
+			LOG_SCOPE_VERBOSITY_OVERRIDE(LogAnimation, ELogVerbosity::Warning);
+			// convert animnotifies
+			for(int32 I = 0; I < Notifies.Num(); ++I)
+			{
+				if(Notifies[I].Notify != NULL)
+				{
+					FString Label = Notifies[I].Notify->GetClass()->GetName();
+					Label = Label.Replace(TEXT("AnimNotify_"), TEXT(""), ESearchCase::CaseSensitive);
+					Notifies[I].NotifyName = FName(*Label);
+				}
+			}
+		}
+
+		if(GetLinkerUE4Version() < VER_UE4_CLEAR_NOTIFY_TRIGGERS)
+		{
+			for(FAnimNotifyEvent Notify : Notifies)
+			{
+				if(Notify.Notify)
+				{
+					// Clear end triggers for notifies that are not notify states
+					Notify.EndTriggerTimeOffset = 0.0f;
+				}
+			}
+		}
 	}
 
 	if ( GetLinkerUE4Version() < VER_UE4_MORPHTARGET_CURVE_INTEGRATION )
