Filtering of anim assets by skeleton in anim blueprints

diff --git a/Engine/Source/Editor/Persona/Private/AnimInstanceDetails.cpp b/Engine/Source/Editor/Persona/Private/AnimInstanceDetails.cpp
new file mode 100644
--- /dev/null
+++ b/Engine/Source/Editor/Persona/Private/AnimInstanceDetails.cpp
@@ -0,0 +1,122 @@
+// Copyright 1998-2014 Epic Games, Inc. All Rights Reserved.
+
+#include "PersonaPrivatePCH.h"
+#include "DetailLayoutBuilder.h"
+#include "PropertyHandle.h"
+#include "PropertyEditorModule.h"
+#include "AnimInstanceDetails.h"
+#include "ObjectEditorUtils.h"
+#include "DetailCategoryBuilder.h"
+#include "DetailWidgetRow.h"
+#include "IDetailPropertyRow.h"
+#include "IDetailCustomNodeBuilder.h"
+#include "PropertyCustomizationHelpers.h"
+#include "AssetData.h"
+
+/////////////////////////////////////////////////////
+// FAnimInstanceDetails 
+
+TSharedRef<IDetailCustomization> FAnimInstanceDetails::MakeInstance()
+{
+	return MakeShareable(new FAnimInstanceDetails());
+}
+
+void FAnimInstanceDetails::CustomizeDetails(class IDetailLayoutBuilder& DetailBuilder)
+{
+	TArray<TWeakObjectPtr<UObject>> Objects;
+	DetailBuilder.GetObjectsBeingCustomized(Objects);
+
+	TArray<UClass*> ObjectClasses;
+	USkeleton* TargetSkeleton = NULL;
+
+	// Grab the skeleton we are displaying for filtering
+	for (TWeakObjectPtr<UObject> Object : Objects)
+	{
+		if (UObject* ObjectPtr = Object.Get())
+		{
+			ObjectClasses.AddUnique(Object->GetClass());
+			if (UAnimInstance* AnimInst = Cast<UAnimInstance>(ObjectPtr))
+			{
+				if (TargetSkeleton && TargetSkeleton != AnimInst->CurrentSkeleton)
+				{
+					TargetSkeleton = NULL;
+					break;
+				}
+				else
+				{
+					TargetSkeleton = AnimInst->CurrentSkeleton;
+				}
+			}
+		}
+	}
+
+	// Grab common base class for property population
+	UClass* CommonBaseClass = UClass::FindCommonBase(ObjectClasses);
+
+	// If everything uses the same skeleton we can filter based on that skeleton
+	if (TargetSkeleton)
+	{
+		TargetSkeletonName = FString::Printf(TEXT("%s'%s'"), *TargetSkeleton->GetClass()->GetName(), *TargetSkeleton->GetPathName());
+
+		for (TFieldIterator<UProperty> It(CommonBaseClass, EFieldIteratorFlags::IncludeSuper); It; ++It)
+		{
+			UProperty* TargetProperty = *It;
+
+			TSharedRef<IPropertyHandle> TargetPropertyHandle = DetailBuilder.GetProperty(*TargetProperty->GetName(), CommonBaseClass);
+			if ((*TargetPropertyHandle).GetProperty() )
+			{
+				IDetailCategoryBuilder& CurrentCategory = DetailBuilder.EditCategory(FObjectEditorUtils::GetCategoryFName(TargetProperty));
+				
+				IDetailPropertyRow& PropertyRow = CurrentCategory.AddProperty(TargetPropertyHandle);
+
+				TSharedPtr<SWidget> NameWidget;
+				TSharedPtr<SWidget> ValueWidget;
+				FDetailWidgetRow Row;
+				PropertyRow.GetDefaultWidgets(NameWidget, ValueWidget, Row);
+
+				TSharedRef<SWidget> TempWidget = CreateFilteredObjectPropertyWidget(TargetProperty, TargetPropertyHandle);
+				ValueWidget = (TempWidget == SNullWidget::NullWidget) ? ValueWidget : TempWidget;
+
+				const bool bShowChildren = true;
+				PropertyRow.CustomWidget(bShowChildren)
+					.NameContent()
+					.MinDesiredWidth(Row.NameWidget.MinWidth)
+					.MaxDesiredWidth(Row.NameWidget.MaxWidth)
+					[
+						NameWidget.ToSharedRef()
+					]
+				.ValueContent()
+					.MinDesiredWidth(Row.ValueWidget.MinWidth)
+					.MaxDesiredWidth(Row.ValueWidget.MaxWidth)
+					[
+						ValueWidget.ToSharedRef()
+					];
+			}
+		}
+	}
+}
+
+TSharedRef<SWidget> FAnimInstanceDetails::CreateFilteredObjectPropertyWidget(UProperty* TargetProperty, TSharedRef<IPropertyHandle> TargetPropertyHandle)
+{
+	if(const UObjectPropertyBase* ObjectProperty = Cast<const UObjectPropertyBase>( TargetProperty ))
+	{
+		if(ObjectProperty->PropertyClass->IsChildOf(UAnimationAsset::StaticClass()))
+		{
+			bool bAllowClear = !(ObjectProperty->PropertyFlags & CPF_NoClear);
+
+			return SNew(SObjectPropertyEntryBox)
+				.PropertyHandle(TargetPropertyHandle)
+				.AllowedClass(ObjectProperty->PropertyClass)
+				.AllowClear(bAllowClear)
+				.OnShouldFilterAsset(FOnShouldFilterAsset::CreateSP(this, &FAnimInstanceDetails::OnShouldFilterAnimAsset));
+		}
+	}
+
+	return SNullWidget::NullWidget;
+}
+
+bool FAnimInstanceDetails::OnShouldFilterAnimAsset(const FAssetData& AssetData) const
+{
+	const FString* SkeletonName = AssetData.TagsAndValues.Find(TEXT("Skeleton"));
+	return *SkeletonName != TargetSkeletonName;
+}
diff --git a/Engine/Source/Editor/Persona/Private/AnimInstanceDetails.h b/Engine/Source/Editor/Persona/Private/AnimInstanceDetails.h
new file mode 100644
--- /dev/null
+++ b/Engine/Source/Editor/Persona/Private/AnimInstanceDetails.h
@@ -0,0 +1,28 @@
+// Copyright 1998-2014 Epic Games, Inc. All Rights Reserved.
+
+#pragma once
+#include "PropertyEditorModule.h"
+
+/////////////////////////////////////////////////////
+// FAnimInstanceDetails 
+
+class FAnimInstanceDetails : public IDetailCustomization
+{
+public:
+	static TSharedRef<IDetailCustomization> MakeInstance();
+
+	// IDetailCustomization interface
+	virtual void CustomizeDetails(class IDetailLayoutBuilder& DetailBuilder) override;
+	// End of IDetailCustomization interface
+
+protected:
+
+	// Creates a filtered object widget if the supplied property is an object property
+	TSharedRef<SWidget> CreateFilteredObjectPropertyWidget(UProperty* TargetProperty, TSharedRef<IPropertyHandle> TargetPropertyHandle);
+
+	/** Delegate to handle filtering of asset pickers */
+	bool OnShouldFilterAnimAsset( const FAssetData& AssetData ) const;
+
+	/** Path to the current blueprints skeleton to allow us to filter asset pickers */
+	FString TargetSkeletonName;
+};
\ No newline at end of file
diff --git a/Engine/Source/Editor/Persona/Private/PersonaModule.cpp b/Engine/Source/Editor/Persona/Private/PersonaModule.cpp
--- a/Engine/Source/Editor/Persona/Private/PersonaModule.cpp
+++ b/Engine/Source/Editor/Persona/Private/PersonaModule.cpp
@@ -12,6 +12,7 @@
 #include "SkeletalMeshSocketDetails.h"
 #include "AnimNotifyDetails.h"
 #include "AnimGraphNodeDetails.h"
+#include "AnimInstanceDetails.h"
 
 IMPLEMENT_MODULE( FPersonaModule, Persona );
 
@@ -60,6 +61,7 @@ void FPersonaModule::StartupModule()
 		PropertyModule.RegisterCustomClassLayout( "AnimNotify", FOnGetDetailCustomizationInstance::CreateStatic( &FAnimNotifyDetails::MakeInstance ) );
 		PropertyModule.RegisterCustomClassLayout( "AnimNotifyState", FOnGetDetailCustomizationInstance::CreateStatic(&FAnimNotifyDetails::MakeInstance));
 		PropertyModule.RegisterCustomClassLayout( "AnimGraphNode_Base", FOnGetDetailCustomizationInstance::CreateStatic( &FAnimGraphNodeDetails::MakeInstance ) );
+		PropertyModule.RegisterCustomClassLayout( "AnimInstance", FOnGetDetailCustomizationInstance::CreateStatic(&FAnimInstanceDetails::MakeInstance));
 
 		PropertyModule.RegisterCustomPropertyTypeLayout( "InputScaleBias", FOnGetPropertyTypeCustomizationInstance::CreateStatic( &FInputScaleBiasCustomization::MakeInstance ) );
 		PropertyModule.RegisterCustomPropertyTypeLayout( "BoneReference", FOnGetPropertyTypeCustomizationInstance::CreateStatic( &FBoneReferenceCustomization::MakeInstance ) );
