Unshelved from pending changelist '2065351':

Added Smartnames to skeletons to manage elements of sequences centrally. Added new animation metadata as the first feature to use smartnames
#ttp - 332035

diff --git a/Engine/Config/BaseEngine.ini b/Engine/Config/BaseEngine.ini
--- a/Engine/Config/BaseEngine.ini
+++ b/Engine/Config/BaseEngine.ini
@@ -648,6 +648,7 @@ bUseStreamingPause=false
 
 ; AnimSequence
 +TaggedPropertyRedirects=(ClassName="AnimSequence",OldPropertyName="BasePose",NewPropertyName="RetargetSource")
++TaggedPropertyRedirects=(ClassName="AnimCurveBase", OldPropertyName="CurveName", NewPropertyName="LastObservedName")
 
 ;EMeshComponentUpdateFlags rename
 +EnumRedirects=(EnumName="EMeshComponentUpdateFlag",OldEnumEntry="SMU_AlwaysTickPoseAndRefreshBones",NewEnumEntry="EMeshComponentUpdateFlag::AlwaysTickPoseAndRefreshBones")
diff --git a/Engine/Source/Editor/Persona/Private/Persona.h b/Engine/Source/Editor/Persona/Private/Persona.h
--- a/Engine/Source/Editor/Persona/Private/Persona.h
+++ b/Engine/Source/Editor/Persona/Private/Persona.h
@@ -330,6 +330,8 @@ private:
 	DECLARE_MULTICAST_DELEGATE( FOnChangeSkeletonTree )
 	// Called when the notifies of the current animation are changed
 	DECLARE_MULTICAST_DELEGATE( FOnChangeAnimNotifies )
+	// Called when the curve panel is changed / updated
+	DECLARE_MULTICAST_DELEGATE( FOnChangeCurves )
 	// Called when the preview viewport is created
 	DECLARE_MULTICAST_DELEGATE_OneParam(FOnCreateViewport, TWeakPtr<class SAnimationEditorViewportTabBody>)
 	// Called when generic delete happens
@@ -459,6 +461,24 @@ public:
 	/** Delegate for when the skeletons animation notifies have been changed */
 	FOnChangeAnimNotifies OnAnimNotifiesChanged;
 
+	// Curve changed
+	typedef FOnChangeCurves::FDelegate FOnCurvesChanged;
+
+	/** Registers delegate for changing / updating of curves panel */
+	void RegisterOnChangeCurves(const FOnCurvesChanged& Delegate)
+	{
+		OnCurvesChanged.Add(Delegate);
+	}
+
+	/** Unregisters delegate for changing / updating of curves panel */
+	void UnregisterOnChangeCurves(SWidget* Widget)
+	{
+		OnCurvesChanged.RemoveAll(Widget);
+	}
+
+	/** Delegate for changing / updating of curves panel */
+	FOnChangeCurves OnCurvesChanged;
+
 	// Viewport Created
 	typedef FOnCreateViewport::FDelegate FOnViewportCreated;
 
diff --git a/Engine/Source/Editor/Persona/Private/SAnimCompositeEditor.cpp b/Engine/Source/Editor/Persona/Private/SAnimCompositeEditor.cpp
--- a/Engine/Source/Editor/Persona/Private/SAnimCompositeEditor.cpp
+++ b/Engine/Source/Editor/Persona/Private/SAnimCompositeEditor.cpp
@@ -69,6 +69,7 @@ void SAnimCompositeEditor::Construct(const FArguments& InArgs)
 		.Padding(0, 10)
 		[
 			SAssignNew( AnimCurvePanel, SAnimCurvePanel )
+			.Persona(PersonaPtr)
 			.Sequence(CompositeObj)
 			.WidgetWidth(S2ColumnWidget::DEFAULT_RIGHT_COLUMN_WIDTH)
 			.ViewInputMin(this, &SAnimEditorBase::GetViewMinInput)
diff --git a/Engine/Source/Editor/Persona/Private/SAnimCurvePanel.cpp b/Engine/Source/Editor/Persona/Private/SAnimCurvePanel.cpp
--- a/Engine/Source/Editor/Persona/Private/SAnimCurvePanel.cpp
+++ b/Engine/Source/Editor/Persona/Private/SAnimCurvePanel.cpp
@@ -79,6 +79,24 @@ public:
 			BaseSequence.Get()->SetFlags(BaseSequence.Get()->GetFlags() | RF_Transactional);
 		}
 	}
+
+	/** Called to get the name of a curve back from the animation skeleton */
+	virtual FText GetCurveName(USkeleton::AnimCurveUID Uid) const
+	{
+		if(BaseSequence.IsValid())
+		{
+			FSmartNameMapping* NameMapping = BaseSequence.Get()->GetSkeleton()->SmartNames.GetContainer(USkeleton::AnimCurveMappingName);
+			if(NameMapping)
+			{
+				FName CurveName;
+				if(NameMapping->GetName(Uid, CurveName))
+				{
+					return FText::FromName(CurveName);
+				}
+			}
+		}
+		return FText::GetEmpty();
+	}
 };
 
 //////////////////////////////////////////////////////////////////////////
@@ -261,7 +279,7 @@ public:
 	SLATE_BEGIN_ARGS( SCurveEdTrack )
 		: _AnimCurvePanel()
 		, _Sequence()
-		, _CurveName()
+		, _CurveUid()
 		, _WidgetWidth()
 		, _ViewInputMin()
 		, _ViewInputMax()
@@ -271,7 +289,7 @@ public:
 	SLATE_ARGUMENT( TSharedPtr<SAnimCurvePanel>, AnimCurvePanel)
 	// editing related variables
 	SLATE_ARGUMENT( class UAnimSequenceBase*, Sequence )
-	SLATE_ARGUMENT( FName, CurveName ) 
+	SLATE_ARGUMENT( USkeleton::AnimCurveUID, CurveUid )
 	// widget viewing related variables
 	SLATE_ARGUMENT( float, WidgetWidth ) // @todo do I need this?
 	SLATE_ATTRIBUTE( float, ViewInputMin )
@@ -309,6 +327,9 @@ public:
 	void ToggleExpandEditor(ESlateCheckBoxState::Type NewType);
 	const FSlateBrush* GetExpandContent() const;
 	FVector2D GetDesiredSize() const;
+
+	// Bound to attribute for curve name, uses curve interface to request from skeleton
+	FText GetCurveName(USkeleton::AnimCurveUID Uid) const;
 };
 
 //////////////////////////////////////////////////////////////////////////
@@ -325,93 +346,113 @@ void SCurveEdTrack::Construct(const FArguments& InArgs)
 	check (Sequence);
 
 	// get the curve data
-	FAnimCurveBase * Curve = Sequence->RawCurveData.GetCurveData(InArgs._CurveName);
+	FAnimCurveBase * Curve = Sequence->RawCurveData.GetCurveData(InArgs._CurveUid);
 	check (Curve);
 
 	CurveInterface = new FAnimCurveBaseInterface(Sequence, Curve);
 	int32 NumberOfKeys = Sequence->GetNumberOfFrames();
 	//////////////////////////////
-	this->ChildSlot
+	
+	TSharedPtr<SBorder> CurveBorder = nullptr;
+	TSharedPtr<SHorizontalBox> InnerBox = nullptr;
+
+	SAssignNew(CurveBorder, SBorder)
+	.Padding(FMargin(2.0f, 2.0f))
 	[
-		SNew( SBorder )
-		.Padding( FMargin(2.0f, 2.0f) )
+		SAssignNew(InnerBox, SHorizontalBox)
+	];
+	
+	FFloatCurve* CurveData = (FFloatCurve*)Curve;
+	bool bIsMetadata = CurveData->GetCurveTypeFlag(ACF_Metadata);
+	if(!bIsMetadata)
+	{
+		InnerBox->AddSlot()
+		.FillWidth(1)
 		[
-			SNew(SHorizontalBox)
-			+SHorizontalBox::Slot()
-			.FillWidth(1)
-			[
-				// Notification editor panel
-				SAssignNew(CurveEditor, SAnimCurveEd)
-				.ViewMinInput(InArgs._ViewInputMin)
-				.ViewMaxInput(InArgs._ViewInputMax)
-				.DataMinInput(0.f)
-				.DataMaxInput(Sequence->SequenceLength)
-				// @fixme fix this to delegate
-				.TimelineLength(Sequence->SequenceLength)
-				.NumberOfKeys(NumberOfKeys)
-				.DesiredSize(this, &SCurveEdTrack::GetDesiredSize)
-				.OnSetInputViewRange(InArgs._OnSetInputViewRange)
-				.OnGetScrubValue(InArgs._OnGetScrubValue)
-			]
+			// Notification editor panel
+			SAssignNew(CurveEditor, SAnimCurveEd)
+			.ViewMinInput(InArgs._ViewInputMin)
+			.ViewMaxInput(InArgs._ViewInputMax)
+			.DataMinInput(0.f)
+			.DataMaxInput(Sequence->SequenceLength)
+			// @fixme fix this to delegate
+			.TimelineLength(Sequence->SequenceLength)
+			.NumberOfKeys(NumberOfKeys)
+			.DesiredSize(this, &SCurveEdTrack::GetDesiredSize)
+			.OnSetInputViewRange(InArgs._OnSetInputViewRange)
+			.OnGetScrubValue(InArgs._OnGetScrubValue)
+		];
+
+		//Inform track widget about the curve and whether it is editable or not.
+		CurveEditor->SetCurveOwner(CurveInterface, true);
+	}
 
-			+SHorizontalBox::Slot()
-			.AutoWidth()
+	TSharedPtr<SHorizontalBox> NameBox = nullptr;
+	SHorizontalBox::FSlot& CurveSlot = InnerBox->AddSlot()
+	[
+		SNew(SBox)
+		.WidthOverride(InArgs._WidgetWidth)
+		.VAlign(VAlign_Center)
+		[
+			SAssignNew(NameBox, SHorizontalBox)
+			+ SHorizontalBox::Slot()
+			.HAlign(HAlign_Center)
+			.Padding(FMargin(0.0f, 5.0f, 0.0f, 5.0f))
 			[
-				SNew( SBox )
-				.WidthOverride(InArgs._WidgetWidth)
-				.VAlign(VAlign_Center)
-				[
-					SNew(SHorizontalBox)
-					+SHorizontalBox::Slot()
-					.HAlign(HAlign_Center)
-					.FillWidth(1)
-					[
-						// Name of track
-						SNew(SEditableText)
-						.MinDesiredWidth(64.0f)
-						.IsEnabled(true)
-						.Font(FEditorStyle::GetFontStyle("CurveEd.InfoFont"))
-						.SelectAllTextWhenFocused(true)
-						.Text(FText::FromName(Curve->CurveName))
-						.OnTextCommitted(this, &SCurveEdTrack::NewCurveNameEntered)
-					]
+				// Name of track
+				SNew(SEditableText)
+				.MinDesiredWidth(64.0f)
+				.IsEnabled(true)
+				.Font(FEditorStyle::GetFontStyle("CurveEd.InfoFont"))
+				.SelectAllTextWhenFocused(true)
+				.Text(this, &SCurveEdTrack::GetCurveName, Curve->CurveUid)
+				.OnTextCommitted(this, &SCurveEdTrack::NewCurveNameEntered)
+			]
+		]
+	];
 
-					+SHorizontalBox::Slot()
-					.AutoWidth()
-					[
-						// Name of track
-						SNew( SCheckBox )
-						.IsChecked(this, &SCurveEdTrack::IsEditorExpanded)
-						.OnCheckStateChanged(this, &SCurveEdTrack::ToggleExpandEditor)
-						.ToolTipText( LOCTEXT("Expand window", "Expand window") )
-						.IsEnabled( true )
-						[
-							SNew( SImage )
-							.Image(this, &SCurveEdTrack::GetExpandContent)
-						]
-					]
+	// Need to autowidth non-metadata names to maximise curve editor area and
+	// add the expansion checkbox (unnecessary for metadata)
+	if(!bIsMetadata)
+	{
+		CurveSlot.AutoWidth();
 
-					+SHorizontalBox::Slot()
-					.AutoWidth()
-					[
-						// Name of track
-						SNew(SButton)
-						.ToolTipText( LOCTEXT("DisplayTrackOptionsMenuTooltip", "Display track options menu") )
-						.OnClicked( this, &SCurveEdTrack::OnContextMenu )
-						.Content()
-						[
-							SNew(SImage)
-							.Image(FEditorStyle::GetBrush("ComboButton.Arrow"))
-							.ColorAndOpacity( FSlateColor::UseForeground() )
-						]
-					]
-				]
+		NameBox->AddSlot()
+		.AutoWidth()
+		[
+			// Name of track
+			SNew(SCheckBox)
+			.IsChecked(this, &SCurveEdTrack::IsEditorExpanded)
+			.OnCheckStateChanged(this, &SCurveEdTrack::ToggleExpandEditor)
+			.ToolTipText(LOCTEXT("Expand window", "Expand window"))
+			.IsEnabled(true)
+			[
+				SNew(SImage)
+				.Image(this, &SCurveEdTrack::GetExpandContent)
 			]
+		];
+	}
+
+	// Add track options combo button
+	NameBox->AddSlot()
+	.Padding(FMargin(0.0f, 5.0f, 0.0f, 5.0f))
+	.AutoWidth()
+	[
+		SNew(SButton)
+		.ToolTipText(LOCTEXT("DisplayTrackOptionsMenuTooltip", "Display track options menu"))
+		.OnClicked(this, &SCurveEdTrack::OnContextMenu)
+		.Content()
+		[
+			SNew(SImage)
+			.Image(FEditorStyle::GetBrush("ComboButton.Arrow"))
+			.ColorAndOpacity(FSlateColor::UseForeground())
 		]
 	];
 
-	//Inform track widget about the curve and whether it is editable or not.
-	CurveEditor->SetCurveOwner(CurveInterface, true );
+	this->ChildSlot
+	[
+		CurveBorder->AsShared()
+	];
 }
 
 /** return a widget */
@@ -429,9 +470,19 @@ const FSlateBrush* SCurveEdTrack::GetExpandContent() const
 }
 void SCurveEdTrack::NewCurveNameEntered( const FText& NewText, ETextCommit::Type CommitInfo )
 {
-	if ( CommitInfo == ETextCommit::OnEnter || CommitInfo == ETextCommit::OnUserMovedFocus )
+	if(CommitInfo == ETextCommit::OnEnter || CommitInfo == ETextCommit::OnUserMovedFocus)
 	{
-		CurveInterface->CurveData->CurveName = FName( *NewText.ToString() );
+		if(USkeleton* Skeleton = CurveInterface->BaseSequence->GetSkeleton())
+		{
+			// Check that the name doesn't already exist
+			FName RequestedName = FName(*NewText.ToString());
+			FSmartNameMapping* NameMapping = Skeleton->SmartNames.GetContainer(USkeleton::AnimCurveMappingName);
+			if(!NameMapping->Exists(RequestedName))
+			{
+				FScopedTransaction Transaction(LOCTEXT("CurveEditor_RenameCurve", "Rename Curve"));
+				Skeleton->RenameSmartnameAndModify(USkeleton::AnimCurveMappingName, CurveInterface->CurveData->CurveUid, FName(*NewText.ToString()));
+			}
+		}
 	}
 }
 
@@ -446,15 +497,16 @@ void SCurveEdTrack::DuplicateTrack()
 	TSharedPtr<SAnimCurvePanel> SharedPanel = PanelPtr.Pin();
 	if(SharedPanel.IsValid())
 	{
-		SharedPanel->DuplicateTrack(CurveInterface->CurveData->CurveName.ToString());
+		SharedPanel->DuplicateTrack(CurveInterface->CurveData->CurveUid);
 	}
 }
 
 void SCurveEdTrack::DeleteTrack()
 {
-	if ( PanelPtr.IsValid() )
+	TSharedPtr<SAnimCurvePanel> SharedPanel = PanelPtr.Pin();
+	if(SharedPanel.IsValid())
 	{
-		PanelPtr.Pin()->DeleteTrack(CurveInterface->CurveData->CurveName.ToString());
+		SharedPanel->DeleteTrack(CurveInterface->CurveData->CurveUid);
 	}
 }
 
@@ -507,61 +559,15 @@ ESlateCheckBoxState::Type SCurveEdTrack::IsCurveOfMode(EAnimCurveFlags ModeToTes
 
 FReply SCurveEdTrack::OnContextMenu()
 {
-	FMenuBuilder MenuBuilder(true, NULL);
-
-	MenuBuilder.BeginSection("AnimCurvePanelCurveTypes", LOCTEXT("CurveTypesHeading", "Curve Types"));
+	TSharedPtr<SAnimCurvePanel> PanelShared = PanelPtr.Pin();
+	if(PanelShared.IsValid())
 	{
-		MenuBuilder.AddWidget(
-			SNew(SCheckBox)
-			.IsChecked( this, &SCurveEdTrack::IsCurveOfMode, ACF_DrivesMorphTarget )
-			.OnCheckStateChanged( this, &SCurveEdTrack::ToggleCurveMode, ACF_DrivesMorphTarget )
-			.ToolTipText(LOCTEXT("MorphCurveModeTooltip", "This curve drives a morph target"))
-			[
-				SNew(STextBlock)
-				.Text(LOCTEXT("MorphCurveMode", "Morph Curve"))
-			],
-			FText()
-		);
-
-		MenuBuilder.AddWidget(
-			SNew(SCheckBox)
-			.IsChecked( this, &SCurveEdTrack::IsCurveOfMode, ACF_DrivesMaterial )
-			.OnCheckStateChanged( this, &SCurveEdTrack::ToggleCurveMode, ACF_DrivesMaterial )
-			.ToolTipText(LOCTEXT("MaterialCurveModeTooltip", "This curve drives a material"))
-			[
-				SNew(STextBlock)
-				.Text(LOCTEXT("MaterialCurveMode", "Material Curve"))
-			],
-			FText()
-		);
-	}
-	MenuBuilder.EndSection();
-
-	MenuBuilder.BeginSection("AnimCurvePanelTrackOptions", LOCTEXT("TrackOptionsHeading", "Track Options") );
-	{
-		MenuBuilder.AddMenuEntry(
-			LOCTEXT("DuplicateTrack", "Duplicate Track"),
-			LOCTEXT("DuplicateTrackTooltip", "Duplicate this track"),
-			FSlateIcon(),
-			FUIAction(FExecuteAction::CreateSP(this, &SCurveEdTrack::DuplicateTrack))
-		);
-
-		MenuBuilder.AddMenuEntry(
-			LOCTEXT("RemoveTrack", "Remove Track"),
-			LOCTEXT("RemoveTrackTooltip", "Remove this track"),
-			FSlateIcon(),
-			FUIAction(
-			FExecuteAction::CreateSP( this, &SCurveEdTrack::DeleteTrack ),
-			FCanExecuteAction()
-			)
-		);
-
-		FSlateApplication::Get().PushMenu(	SharedThis(this),
-			MenuBuilder.MakeWidget(),
-			FSlateApplication::Get().GetCursorPos(),
-			FPopupTransitionEffect(FPopupTransitionEffect::TypeInPopup) );
+		FFloatCurve* Curve = (FFloatCurve*)(CurveInterface->CurveData);
+		FSlateApplication::Get().PushMenu(SharedThis(this),
+										  PanelShared->CreateCurveContextMenu(Curve),
+										  FSlateApplication::Get().GetCursorPos(),
+										  FPopupTransitionEffect(FPopupTransitionEffect::TypeInPopup));
 	}
-	MenuBuilder.EndSection();
 
 	return FReply::Handled();
 }
@@ -587,6 +593,12 @@ FVector2D SCurveEdTrack::GetDesiredSize() const
 		return FVector2D(128.f, 32.f);
 	}
 }
+
+FText SCurveEdTrack::GetCurveName(USkeleton::AnimCurveUID Uid) const
+{
+	return CurveInterface->GetCurveName(Uid);
+}
+
 //////////////////////////////////////////////////////////////////////////
 // SAnimCurvePanel
 
@@ -596,12 +608,16 @@ FVector2D SCurveEdTrack::GetDesiredSize() const
 class FCurveNameValidator : public FStringSetNameValidator
 {
 public:
-	FCurveNameValidator(FRawCurveTracks& Tracks, const FString& ExistingName)
+	FCurveNameValidator(FRawCurveTracks& Tracks, FSmartNameMapping* NameMapping, const FString& ExistingName)
 		: FStringSetNameValidator(ExistingName)
 	{
+		FName CurveName;
 		for(FFloatCurve& Curve : Tracks.FloatCurves)
 		{
-			Names.Add(Curve.CurveName.ToString());
+			if(NameMapping->GetName(Curve.CurveUid, CurveName))
+			{
+				Names.Add(CurveName.ToString());
+			}
 		}
 	}
 };
@@ -616,6 +632,7 @@ void SAnimCurvePanel::Construct(const FArguments& InArgs)
 		.InputMax(InArgs._InputMax)
 		.OnSetInputViewRange(InArgs._OnSetInputViewRange));
 
+	WeakPersona = InArgs._Persona;
 	Sequence = InArgs._Sequence;
 	WidgetWidth = InArgs._WidgetWidth;
 	OnGetScrubValue = InArgs._OnGetScrubValue;
@@ -640,9 +657,9 @@ void SAnimCurvePanel::Construct(const FArguments& InArgs)
 					[
 						// Name of track
 						SNew(SButton)
-						.Text( LOCTEXT("AddFloatTrack", "Add Float Track") )
+						.Text( LOCTEXT("AddFloatTrack", "Add...") )
 						.ToolTipText( LOCTEXT("AddTrackTooltip", "Add float track above here") )
-						.OnClicked( this, &SAnimCurvePanel::CreateNewTrack )
+						.OnClicked( this, &SAnimCurvePanel::AddButtonClicked )
 					]
 
 					+SHorizontalBox::Slot()
@@ -673,9 +690,10 @@ void SAnimCurvePanel::Construct(const FArguments& InArgs)
 				]
 
 				+SVerticalBox::Slot()
+				.Padding(FMargin(0.0f, 5.0f, 0.0f, 0.0f))
 				.AutoHeight()
 				[
-					SAssignNew( PanelSlot, SSplitter )
+					SAssignNew(PanelSlot, SSplitter)
 					.Orientation(Orient_Vertical)
 				]
 			]
@@ -685,73 +703,103 @@ void SAnimCurvePanel::Construct(const FArguments& InArgs)
 	UpdatePanel();
 }
 
-FReply SAnimCurvePanel::CreateNewTrack()
+FReply SAnimCurvePanel::AddButtonClicked()
 {
-	TSharedRef<STextEntryPopup> TextEntry =
-		SNew(STextEntryPopup)
-		.Label( LOCTEXT("TrackName", "Track Name") )
-		.OnTextCommitted( this, &SAnimCurvePanel::AddTrack );
+	USkeleton* CurrentSkeleton = Sequence->GetSkeleton();
+	check(CurrentSkeleton);
+
+	FMenuBuilder MenuBuilder(true, NULL);
+	
+	MenuBuilder.BeginSection("ConstantCurves", LOCTEXT("ConstantCurveHeading", "Constant Curve"));
+	{
+		MenuBuilder.AddSubMenu(
+			LOCTEXT("NewMetadataSubMenu", "Add Metadata..."),
+			LOCTEXT("NewMetadataSubMenuToolTip", "Add a new metadata entry to the sequence"),
+			FNewMenuDelegate::CreateRaw(this, &SAnimCurvePanel::FillMetadataEntryMenu));
+	}
+	MenuBuilder.EndSection();
+
+	MenuBuilder.BeginSection("Curves", LOCTEXT("CurveHeading", "Curve"));
+	{
+		MenuBuilder.AddSubMenu(
+			LOCTEXT("NewVariableCurveSubMenu", "Add Variable Curve..."),
+			LOCTEXT("NewVariableCurveSubMenuToolTip", "Add a new variable curve to the sequence"),
+			FNewMenuDelegate::CreateRaw(this, &SAnimCurvePanel::FillVariableCurveMenu));
+	}
+	MenuBuilder.EndSection();
 
 	// Show dialog to enter new track name
 	FSlateApplication::Get().PushMenu(
 		SharedThis( this ),
-		TextEntry,
+		MenuBuilder.MakeWidget(),
 		FSlateApplication::Get().GetCursorPos(),
 		FPopupTransitionEffect( FPopupTransitionEffect::TypeInPopup)
 		);
 
-	TextEntry->FocusDefaultWidget();
-
 	return FReply::Handled();	
 }
 
-void SAnimCurvePanel::AddTrack(const FText & CurveNameToAdd, ETextCommit::Type CommitInfo)
+void SAnimCurvePanel::CreateTrack(const FText& ComittedText, ETextCommit::Type CommitInfo)
 {
 	if ( CommitInfo == ETextCommit::OnEnter )
 	{
-		const FScopedTransaction Transaction( LOCTEXT("AnimCurve_AddTrack", "Add New Curve") );
-		// before insert, make sure everything behind is fixed
-		if ( Sequence->RawCurveData.AddCurveData(FName( *CurveNameToAdd.ToString() )) )
+		USkeleton* Skeleton = Sequence->GetSkeleton();
+		if(Skeleton && !ComittedText.IsEmpty())
 		{
-			Sequence->Modify(true);
-			UpdatePanel();
+			const FScopedTransaction Transaction(LOCTEXT("AnimCurve_AddTrack", "Add New Curve"));
+			USkeleton::AnimCurveUID CurveUid;
+
+			if(Skeleton->AddSmartnameAndModify(USkeleton::AnimCurveMappingName, FName(*ComittedText.ToString()), CurveUid))
+			{
+				AddVariableCurve(CurveUid);
+			}
 		}
 
 		FSlateApplication::Get().DismissAllMenus();
 	}
 }
 
-FReply SAnimCurvePanel::DuplicateTrack(const FString& CurveNameToDuplicate)
+FReply SAnimCurvePanel::DuplicateTrack(USkeleton::AnimCurveUID Uid)
 {
 	const FScopedTransaction Transaction( LOCTEXT("AnimCurve_DuplicateTrack", "Duplicate Curve") );
 	
-	TSharedPtr<INameValidatorInterface> Validator = MakeShareable(new FCurveNameValidator(Sequence->RawCurveData, FString(TEXT(""))));
+	FSmartNameMapping* NameMapping = Sequence->GetSkeleton()->SmartNames.GetContainer(USkeleton::AnimCurveMappingName);
+	FName CurveNameToCopy;
+	USkeleton::AnimCurveUID NewUid;
 
-	// Use the validator to pick a reasonable name for the duplicated curve.
-	FString NewCurveName = CurveNameToDuplicate;
-	Validator->FindValidString(NewCurveName);
-
-	if(Sequence->RawCurveData.DuplicateCurveData(*CurveNameToDuplicate, *NewCurveName))
+	// Must have a curve that exists to duplicate
+	if(NameMapping->Exists(Uid))
 	{
-		Sequence->Modify();
-		UpdatePanel();
+		NameMapping->GetName(Uid, CurveNameToCopy);
+		TSharedPtr<INameValidatorInterface> Validator = MakeShareable(new FCurveNameValidator(Sequence->RawCurveData, NameMapping, FString(TEXT(""))));
 
-		return FReply::Handled();
+		// Use the validator to pick a reasonable name for the duplicated curve.
+		FString NewCurveName = CurveNameToCopy.ToString();
+		Validator->FindValidString(NewCurveName);
+		if(NameMapping->AddName(*NewCurveName, NewUid))
+		{
+			if(Sequence->RawCurveData.DuplicateCurveData(Uid, NewUid))
+			{
+				Sequence->Modify();
+				UpdatePanel();
+
+				return FReply::Handled();
+			}
+		}
 	}
 	return FReply::Unhandled();
 }
 
-FReply SAnimCurvePanel::DeleteTrack(const FString & CurveNameToDelete)
+void SAnimCurvePanel::DeleteTrack(USkeleton::AnimCurveUID Uid)
 {
-	const FScopedTransaction Transaction( LOCTEXT("AnimCurve_AddTrack", "Add New Curve") );
-	if ( Sequence->RawCurveData.DeleteCurveData(*CurveNameToDelete) )
+	const FScopedTransaction Transaction( LOCTEXT("AnimCurve_DeleteTrack", "Delete Curve") );
+	
+	if(Sequence->RawCurveData.GetCurveData(Uid))
 	{
 		Sequence->Modify(true);
+		Sequence->RawCurveData.DeleteCurveData(Uid);
 		UpdatePanel();
-		return FReply::Handled();
 	}
-
-	return FReply::Unhandled();	
 }
 
 FReply SAnimCurvePanel::OnContextMenu()
@@ -803,9 +851,9 @@ EVisibility SAnimCurvePanel::IsSetAllTracksButtonVisible() const
 void SAnimCurvePanel::ToggleAllCurveModes(ESlateCheckBoxState::Type NewState, EAnimCurveFlags ModeToSet)
 {
 	const ESlateCheckBoxState::Type CurrentAllState = AreAllCurvesOfMode(ModeToSet);
-	for( auto It=Tracks.CreateIterator(); It; ++It )
+	for(TWeakPtr<SCurveEdTrack> TrackWeak : Tracks)
 	{
-		TSharedPtr<SCurveEdTrack> TrackWidget = It->Value.Pin();
+		TSharedPtr<SCurveEdTrack> TrackWidget = TrackWeak.Pin();
 		if( TrackWidget.IsValid() )
 		{
 			const ESlateCheckBoxState::Type CurrentTrackState = TrackWidget->IsCurveOfMode(ModeToSet);
@@ -820,9 +868,9 @@ void SAnimCurvePanel::ToggleAllCurveModes(ESlateCheckBoxState::Type NewState, EA
 ESlateCheckBoxState::Type SAnimCurvePanel::AreAllCurvesOfMode(EAnimCurveFlags ModeToSet) const
 {
 	int32 NumChecked = 0;
-	for( auto It=Tracks.CreateConstIterator(); It; ++It )
+	for(const TWeakPtr<SCurveEdTrack> TrackWeak : Tracks)
 	{
-		TSharedPtr<SCurveEdTrack> TrackWidget = It->Value.Pin();
+		const TSharedPtr<SCurveEdTrack> TrackWidget = TrackWeak.Pin();
 		if( TrackWidget.IsValid() )
 		{
 			if( TrackWidget->IsCurveOfMode(ModeToSet) )
@@ -846,37 +894,64 @@ void SAnimCurvePanel::UpdatePanel()
 {
 	if(Sequence != NULL)
 	{
-		// see if we need to clear or not
-		FChildren * Children = PanelSlot->GetChildren();
+		USkeleton* CurrentSkeleton = Sequence->GetSkeleton();
+		FSmartNameMapping* MetadataNameMap = CurrentSkeleton->SmartNames.GetContainer(USkeleton::AnimCurveMappingName);
+		// Sort the raw curves before setting up display
+		Sequence->RawCurveData.FloatCurves.Sort([MetadataNameMap](const FFloatCurve& A, const FFloatCurve& B)
+		{
+			bool bAMeta = A.GetCurveTypeFlag(ACF_Metadata);
+			bool bBMeta = B.GetCurveTypeFlag(ACF_Metadata);
+			
+			if(bAMeta != bBMeta)
+			{
+				return !bAMeta;
+			}
+
+			FName AName;
+			FName BName;
+			MetadataNameMap->GetName(A.CurveUid, AName);
+			MetadataNameMap->GetName(B.CurveUid, BName);
+
+			return AName < BName;
+		});
 
-		for(int32 Id = Children->Num() - 1; Id >= 0; --Id)
+		// see if we need to clear or not
+		FChildren * VariableChildren = PanelSlot->GetChildren();
+		for (int32 Id=VariableChildren->Num()-1; Id>=0; --Id)
 		{
 			PanelSlot->RemoveAt(Id);
 		}
+
 		// Clear all tracks as we're re-adding them all anyway.
 		Tracks.Empty();
 
 		// Updating new tracks
+		FSmartNameMapping* NameMapping = CurrentSkeleton->SmartNames.GetContainer(USkeleton::AnimCurveMappingName);
+
 		int32 TotalCurve = Sequence->RawCurveData.FloatCurves.Num();
 		for(int32 CurrentIt = 0 ; CurrentIt < TotalCurve ; ++CurrentIt)
 		{
 			FFloatCurve&  Curve = Sequence->RawCurveData.FloatCurves[CurrentIt];
 
+			const bool bEditable = Curve.GetCurveTypeFlag(ACF_Editable);
+			const bool bConstant = Curve.GetCurveTypeFlag(ACF_Metadata);
+			FName CurveName;
+
 			// if editable, add to the list
-			if (Curve.GetCurveTypeFlag(ACF_Editable))
+			if(bEditable && NameMapping->GetName(Curve.CurveUid, CurveName))
 			{
 				TSharedPtr<SCurveEdTrack> CurrentTrack;
 				PanelSlot->AddSlot()
 				.SizeRule(SSplitter::SizeToContent)
 				[
 					SNew(SVerticalBox)
-					+SVerticalBox::Slot()
+					+ SVerticalBox::Slot()
 					.AutoHeight()
 					.VAlign(VAlign_Center)
 					[
 						SAssignNew(CurrentTrack, SCurveEdTrack)
 						.Sequence(Sequence)
-						.CurveName(Curve.CurveName)
+						.CurveUid(Curve.CurveUid)
 						.AnimCurvePanel(SharedThis(this))
 						.WidgetWidth(WidgetWidth)
 						.ViewInputMin(ViewInputMin)
@@ -885,9 +960,15 @@ void SAnimCurvePanel::UpdatePanel()
 						.OnSetInputViewRange(OnSetInputViewRange)
 					]
 				];
-				Tracks.Add( Curve.CurveName.ToString(), CurrentTrack );
+				Tracks.Add(CurrentTrack);
 			}
 		}
+
+		TSharedPtr<FPersona> SharedPersona = WeakPersona.Pin();
+		if(SharedPersona.IsValid())
+		{
+			SharedPersona->OnCurvesChanged.Broadcast();
+		}
 	}
 }
 
@@ -918,23 +999,30 @@ TSharedRef<SWidget> SAnimCurvePanel::GenerateCurveList()
 				]
 			];
 
+		// Mapping to retrieve curve names
+		FSmartNameMapping* NameMapping = Sequence->GetSkeleton()->SmartNames.GetContainer(USkeleton::AnimCurveMappingName);
+		check(NameMapping);
+
 		for (auto Iter=Sequence->RawCurveData.FloatCurves.CreateConstIterator(); Iter; ++Iter)
 		{
 			const FFloatCurve& Curve= *Iter;
 
+			FName CurveName;
+			NameMapping->GetName(Curve.CurveUid, CurveName);
+
 			ListBox->AddSlot()
 				.AutoHeight()
 				.VAlign(VAlign_Center)
 				.Padding( 2.0f, 2.0f )
 				[
 					SNew( SCheckBox )
-					.IsChecked(this, &SAnimCurvePanel::IsCurveEditable, Curve.CurveName)
-					.OnCheckStateChanged(this, &SAnimCurvePanel::ToggleEditability, Curve.CurveName)
+					.IsChecked(this, &SAnimCurvePanel::IsCurveEditable, Curve.CurveUid)
+					.OnCheckStateChanged(this, &SAnimCurvePanel::ToggleEditability, Curve.CurveUid)
 					.ToolTipText( LOCTEXT("Show Curves", "Show or Hide Curves") )
 					.IsEnabled( true )
 					[
 						SNew( STextBlock )
-						.Text(FText::FromName(Curve.CurveName))
+						.Text(FText::FromName(CurveName))
 					]
 				];
 		}
@@ -999,11 +1087,11 @@ TSharedRef<SWidget> SAnimCurvePanel::GenerateCurveList()
 	return NewWidget;
 }
 
-ESlateCheckBoxState::Type SAnimCurvePanel::IsCurveEditable(FName CurveName) const
+ESlateCheckBoxState::Type SAnimCurvePanel::IsCurveEditable(USkeleton::AnimCurveUID Uid) const
 {
 	if ( Sequence )
 	{
-		const FFloatCurve * Curve = Sequence->RawCurveData.GetCurveData(CurveName);
+		const FFloatCurve * Curve = Sequence->RawCurveData.GetCurveData(Uid);
 		if ( Curve )
 		{
 			return Curve->GetCurveTypeFlag(ACF_Editable)? ESlateCheckBoxState::Checked : ESlateCheckBoxState::Unchecked;
@@ -1013,13 +1101,13 @@ ESlateCheckBoxState::Type SAnimCurvePanel::IsCurveEditable(FName CurveName) cons
 	return ESlateCheckBoxState::Undetermined;
 }
 
-void SAnimCurvePanel::ToggleEditability(ESlateCheckBoxState::Type NewType, FName CurveName)
+void SAnimCurvePanel::ToggleEditability(ESlateCheckBoxState::Type NewType, USkeleton::AnimCurveUID Uid)
 {
 	bool bEdit = (NewType == ESlateCheckBoxState::Checked);
 
 	if ( Sequence )
 	{
-		FFloatCurve * Curve = Sequence->RawCurveData.GetCurveData(CurveName);
+		FFloatCurve * Curve = Sequence->RawCurveData.GetCurveData(Uid);
 		if ( Curve )
 		{
 			Curve->SetCurveTypeFlag(ACF_Editable, bEdit);
@@ -1048,4 +1136,256 @@ FReply		SAnimCurvePanel::ShowAll(bool bShow)
 
 	return FReply::Handled();
 }
+
+void SAnimCurvePanel::FillMetadataEntryMenu(FMenuBuilder& Builder)
+{
+	USkeleton* CurrentSkeleton = Sequence->GetSkeleton();
+	check(CurrentSkeleton);
+
+	FSmartNameMapping* Mapping = CurrentSkeleton->SmartNames.GetContainer(USkeleton::AnimCurveMappingName);
+	TArray<USkeleton::AnimCurveUID> CurveUids;
+	Mapping->FillUidArray(CurveUids);
+
+	Builder.BeginSection(NAME_None, LOCTEXT("MetadataMenu_ListHeading", "Available Names"));
+	for(USkeleton::AnimCurveUID Id : CurveUids)
+	{
+		if(!Sequence->RawCurveData.GetCurveData(Id))
+		{
+			FName CurveName;
+			if(Mapping->GetName(Id, CurveName))
+			{
+				const FText Description = LOCTEXT("NewMetadataSubMenu_ToolTip", "Add an existing metadata curve");
+				const FText Label = FText::FromName(CurveName);
+
+				FUIAction UIAction;
+				UIAction.ExecuteAction.BindRaw(
+					this, &SAnimCurvePanel::AddMetadataEntry,
+					Id);
+
+				Builder.AddMenuEntry(Label, Description, FSlateIcon(), UIAction);
+			}
+		}
+	}
+	Builder.EndSection();
+
+	Builder.AddMenuSeparator();
+
+	const FText Description = LOCTEXT("NewMetadataCreateNew_ToolTip", "Create a new metadata entry");
+	const FText Label = LOCTEXT("NewMetadataCreateNew_Label","Create New");
+	FUIAction UIAction;
+	UIAction.ExecuteAction.BindRaw(this, &SAnimCurvePanel::CreateNewMetadataEntryClicked);
+
+	Builder.AddMenuEntry(Label, Description, FSlateIcon(), UIAction);
+}
+
+void SAnimCurvePanel::FillVariableCurveMenu(FMenuBuilder& Builder)
+{
+	USkeleton* CurrentSkeleton = Sequence->GetSkeleton();
+	check(CurrentSkeleton);
+
+	FSmartNameMapping* Mapping = CurrentSkeleton->SmartNames.GetContainer(USkeleton::AnimCurveMappingName);
+	TArray<USkeleton::AnimCurveUID> CurveUids;
+	Mapping->FillUidArray(CurveUids);
+
+	Builder.BeginSection(NAME_None, LOCTEXT("VariableMenu_ListHeading", "Available Names"));
+	for(USkeleton::AnimCurveUID Id : CurveUids)
+	{
+		if(!Sequence->RawCurveData.GetCurveData(Id))
+		{
+			FName CurveName;
+			if(Mapping->GetName(Id, CurveName))
+			{
+				const FText Description = LOCTEXT("NewVariableSubMenu_ToolTip", "Add an existing variable curve");
+				const FText Label = FText::FromName(CurveName);
+
+				FUIAction UIAction;
+				UIAction.ExecuteAction.BindRaw(
+					this, &SAnimCurvePanel::AddVariableCurve,
+					Id);
+
+				Builder.AddMenuEntry(Label, Description, FSlateIcon(), UIAction);
+			}
+		}
+	}
+	Builder.EndSection();
+
+	Builder.AddMenuSeparator();
+
+	const FText Description = LOCTEXT("NewVariableCurveCreateNew_ToolTip", "Create a new variable curve");
+	const FText Label = LOCTEXT("NewVariableCurveCreateNew_Label", "Create Curve");
+	FUIAction UIAction;
+	UIAction.ExecuteAction.BindRaw(this, &SAnimCurvePanel::CreateNewCurveClicked);
+
+	Builder.AddMenuEntry(Label, Description, FSlateIcon(), UIAction);
+}
+
+void SAnimCurvePanel::AddMetadataEntry(USkeleton::AnimCurveUID Uid)
+{
+	if(Sequence->RawCurveData.AddCurveData(Uid))
+	{
+		Sequence->Modify(true);
+		FFloatCurve* Curve = Sequence->RawCurveData.GetCurveData(Uid);
+		Curve->FloatCurve.AddKey(0.0f, 1.0f);
+		Curve->SetCurveTypeFlag(ACF_Metadata, true);
+		RefreshPanel();
+	}
+}
+
+void SAnimCurvePanel::CreateNewMetadataEntryClicked()
+{
+	TSharedRef<STextEntryPopup> TextEntry =
+		SNew(STextEntryPopup)
+		.Label(LOCTEXT("NewMetadataCurveEntryLabal", "Metadata Name"))
+		.OnTextCommitted(this, &SAnimCurvePanel::CreateNewMetadataEntry);
+
+	FSlateApplication& SlateApp = FSlateApplication::Get();
+	SlateApp.PushMenu(
+		AsShared(),
+		TextEntry,
+		SlateApp.GetCursorPos(),
+		FPopupTransitionEffect::TypeInPopup
+		);
+
+	TextEntry->FocusDefaultWidget();
+}
+
+void SAnimCurvePanel::CreateNewMetadataEntry(const FText& CommittedText, ETextCommit::Type CommitType)
+{
+	FSlateApplication::Get().DismissAllMenus();
+	if(CommitType == ETextCommit::OnEnter)
+	{
+		// Add the name to the skeleton and then add the new curve to the sequence
+		USkeleton* Skeleton = Sequence->GetSkeleton();
+		if(Skeleton && !CommittedText.IsEmpty())
+		{
+			USkeleton::AnimCurveUID CurveUid;
+
+			if(Skeleton->AddSmartnameAndModify(USkeleton::AnimCurveMappingName, FName(*CommittedText.ToString()), CurveUid))
+			{
+				AddMetadataEntry(CurveUid);
+			}
+		}
+	}
+}
+
+void SAnimCurvePanel::CreateNewCurveClicked()
+{
+	TSharedRef<STextEntryPopup> TextEntry =
+		SNew(STextEntryPopup)
+		.Label(LOCTEXT("NewCurveEntryLabal", "Curve Name"))
+		.OnTextCommitted(this, &SAnimCurvePanel::CreateTrack);
+
+	FSlateApplication& SlateApp = FSlateApplication::Get();
+	SlateApp.PushMenu(
+		AsShared(),
+		TextEntry,
+		SlateApp.GetCursorPos(),
+		FPopupTransitionEffect::TypeInPopup
+		);
+
+	TextEntry->FocusDefaultWidget();
+}
+
+TSharedRef<SWidget> SAnimCurvePanel::CreateCurveContextMenu(FFloatCurve* Curve) const
+{
+	FMenuBuilder MenuBuilder(true, NULL);
+
+	MenuBuilder.BeginSection("AnimCurvePanelCurveTypes", LOCTEXT("CurveTypesHeading", "Curve Types"));
+	{
+		MenuBuilder.AddWidget(
+			SNew(SCheckBox)
+			.IsChecked(this, &SAnimCurvePanel::GetCurveFlagAsCheckboxState, Curve, ACF_DrivesMorphTarget)
+			.OnCheckStateChanged(this, &SAnimCurvePanel::SetCurveFlagFromCheckboxState, Curve, ACF_DrivesMorphTarget)
+			.ToolTipText(LOCTEXT("MorphCurveModeTooltip", "This curve drives a morph target"))
+			[
+				SNew(STextBlock)
+				.Text(LOCTEXT("MorphCurveMode", "Morph Curve"))
+			],
+			FText()
+			);
+
+		MenuBuilder.AddWidget(
+			SNew(SCheckBox)
+			.IsChecked(this, &SAnimCurvePanel::GetCurveFlagAsCheckboxState, Curve, ACF_DrivesMaterial)
+			.OnCheckStateChanged(this, &SAnimCurvePanel::SetCurveFlagFromCheckboxState, Curve, ACF_DrivesMaterial)
+			.ToolTipText(LOCTEXT("MaterialCurveModeTooltip", "This curve drives a material"))
+			[
+				SNew(STextBlock)
+				.Text(LOCTEXT("MaterialCurveMode", "Material Curve"))
+			],
+			FText()
+			);
+	}
+	MenuBuilder.EndSection();
+
+	MenuBuilder.BeginSection("AnimCurvePanelTrackOptions", LOCTEXT("TrackOptionsHeading", "Track Options"));
+	{
+		FText TypeToggleLabel = LOCTEXT("TypeToggleToMetadataLabel", "Convert to Metadata");
+		FText TypeToggleToolTip = LOCTEXT("TypeToggleToMetadataToolTip", "Turns this curve into a Metadata entry. This is a destructive operation and will remove the keys in this curve");
+		bool bIsConstantCurve = Curve->GetCurveTypeFlag(ACF_Metadata);
+
+		FUIAction NewAction;
+
+		if(bIsConstantCurve)
+		{
+			TypeToggleLabel = LOCTEXT("TypeToggleToVariableLabel", "Convert to Variable Curve");
+			TypeToggleToolTip = LOCTEXT("TypeToggleToVariableToolTip", "Turns this curve into a variable curve.");
+		}
+
+		NewAction.ExecuteAction.BindSP(this, &SAnimCurvePanel::ToggleCurveTypeMenuCallback, Curve);
+		MenuBuilder.AddMenuEntry(
+			TypeToggleLabel,
+			TypeToggleToolTip,
+			FSlateIcon(),
+			NewAction);
+
+		NewAction.ExecuteAction.BindSP(this, &SAnimCurvePanel::DeleteTrack, Curve->CurveUid);
+		MenuBuilder.AddMenuEntry(
+			LOCTEXT("RemoveTrack", "Remove Track"),
+			LOCTEXT("RemoveTrackTooltip", "Remove this track"),
+			FSlateIcon(),
+			NewAction);
+	}
+	MenuBuilder.EndSection();
+
+	return MenuBuilder.MakeWidget();
+}
+
+ESlateCheckBoxState::Type SAnimCurvePanel::GetCurveFlagAsCheckboxState(FFloatCurve* Curve, EAnimCurveFlags InFlag) const
+{
+	return Curve->GetCurveTypeFlag(InFlag) ? ESlateCheckBoxState::Checked : ESlateCheckBoxState::Unchecked;
+}
+
+void SAnimCurvePanel::SetCurveFlagFromCheckboxState(ESlateCheckBoxState::Type CheckState, FFloatCurve* Curve, EAnimCurveFlags InFlag)
+{
+	bool Enabled = CheckState == ESlateCheckBoxState::Checked;
+	Curve->SetCurveTypeFlag(InFlag, Enabled);
+}
+
+void SAnimCurvePanel::ToggleCurveTypeMenuCallback(FFloatCurve* Curve)
+{
+	check(Curve);
+
+	FScopedTransaction Transaction(LOCTEXT("CurvePanel_ToggleCurveType", "Toggle curve type"));
+	Sequence->Modify(true);
+	bool bIsSet = Curve->GetCurveTypeFlag(ACF_Metadata);
+	Curve->SetCurveTypeFlag(ACF_Metadata, !bIsSet);
+
+	if(!bIsSet)
+	{
+		// We're moving to a metadata curve, we need to clear out the keys.
+		Curve->FloatCurve.Reset();
+		Curve->FloatCurve.AddKey(0.0f, 1.0f);
+	}
+
+	UpdatePanel();
+}
+
+void SAnimCurvePanel::AddVariableCurve(USkeleton::AnimCurveUID CurveUid)
+{
+	Sequence->Modify(true);
+	Sequence->RawCurveData.AddCurveData(CurveUid);
+	UpdatePanel();
+}
+
 #undef LOCTEXT_NAMESPACE
diff --git a/Engine/Source/Editor/Persona/Private/SAnimCurvePanel.h b/Engine/Source/Editor/Persona/Private/SAnimCurvePanel.h
--- a/Engine/Source/Editor/Persona/Private/SAnimCurvePanel.h
+++ b/Engine/Source/Editor/Persona/Private/SAnimCurvePanel.h
@@ -8,6 +8,7 @@
 #include "SNodePanel.h"
 #include "SAnimTrackPanel.h"
 #include "SCurveEditor.h"
+#include "Animation/AnimSequenceBase.h"
 
 DECLARE_DELEGATE_OneParam( FOnSelectionChanged, const FGraphPanelSelectionSet& )
 DECLARE_DELEGATE( FOnUpdatePanel )
@@ -29,6 +30,10 @@ public:
 		, _OnGetScrubValue()
 	{}
 
+	/**
+	 * Weak ptr to our Persona instance
+	 */
+	SLATE_ARGUMENT(TWeakPtr<FPersona>, Persona)
 	/**
 	 * AnimSequenceBase to be used for this panel
 	 */
@@ -60,23 +65,29 @@ public:
 	/**
 	 * Create New Curve Track
 	 */
-	FReply CreateNewTrack();
+	FReply AddButtonClicked();
 
 	/**
 	 * Add Track
 	 */
-	void AddTrack(const FText & CurveNameToAdd, ETextCommit::Type CommitInfo);
-	
+	void CreateTrack(const FText& CommittedText, ETextCommit::Type CommitInfo);
+
+	/**
+	 * Add new variable curve
+	 */
+	void AddVariableCurve(USkeleton::AnimCurveUID CurveUid);
+
+
 	/**
 	 * Duplicate Track
 	 */
-	FReply DuplicateTrack(const FString& CurveNameToDuplicate);
+	FReply DuplicateTrack(USkeleton::AnimCurveUID Uid);
 
 	/**
 	 * Delete Track
 	 *
 	 */
-	FReply DeleteTrack(const FString & CurveNameToDelete);
+	void DeleteTrack(USkeleton::AnimCurveUID Uid);
 
 	/**
 	 * Build and display curve track context menu for all tracks.
@@ -99,12 +110,30 @@ public:
 	 */
 	EVisibility IsSetAllTracksButtonVisible() const;
 
+	/**
+	 * Creates a context menu for the provided curve
+	 */
+	TSharedRef<SWidget> CreateCurveContextMenu(FFloatCurve* Curve) const;
+
+	/**
+	 * Sets the specified flag value to State for the provided curve
+	 */
+	void SetCurveFlag(FFloatCurve* CurveInterface, bool State, EAnimCurveFlags FlagToSet);
+
+	/**
+	* Update Panel
+	* Used internally and by sequence editor to refresh the panel contents
+	* @todo this has to be more efficient. Right now it refreshes the entire panel
+	*/
+	void UpdatePanel();
+
 private:
+	TWeakPtr<FPersona> WeakPersona;
 	TSharedPtr<SSplitter> PanelSlot;
 	class UAnimSequenceBase* Sequence;
 	TAttribute<float> CurrentPosition;
 	FOnGetScrubValue OnGetScrubValue;
-	TMap<FString, TWeakPtr<class SCurveEdTrack> > Tracks;
+	TArray<TWeakPtr<class SCurveEdTrack>> Tracks;
 
 	/**
 	 * This is to control visibility of the curves, so you can edit or not
@@ -114,11 +143,11 @@ private:
 	/**
 	 * Returns true if this curve is editable
 	 */
-	ESlateCheckBoxState::Type	IsCurveEditable(FName CurveName) const;
+	ESlateCheckBoxState::Type	IsCurveEditable(USkeleton::AnimCurveUID Uid) const;
 	/**
 	 * Toggle curve visibility
 	 */
-	void		ToggleEditability(ESlateCheckBoxState::Type NewType, FName CurveName);
+	void		ToggleEditability(ESlateCheckBoxState::Type NewType, USkeleton::AnimCurveUID Uid);
 	/**
 	 * Refresh Panel
 	 */
@@ -128,10 +157,50 @@ private:
 	 */
 	FReply		ShowAll(bool bShow);
 
-private:
+	void FillEntryMenu(FMenuBuilder& Builder, FText AddToolTip, FExecuteAction& AddDelegate, FText CreateToolTip, FExecuteAction& CreateDelegate);
+
 	/**
-	 * Update Panel 
-	 * @todo this has to be more efficient. Right now it refreshes the entire panel
+	 * Fills a menu for creating and adding new metadata entries
 	 */
-	void UpdatePanel();
+	void FillMetadataEntryMenu(FMenuBuilder& Builder);
+
+	/**
+	 * Fills a menu for creating and adding new variable curves
+	 */
+	void FillVariableCurveMenu(FMenuBuilder& Builder);
+
+	/**
+	 * Adds a metadata entry from a Uid for a name that exists in the skeleton
+	 */
+	void AddMetadataEntry(USkeleton::AnimCurveUID Uid);
+
+	/**
+	 * Handler for the new metadata item in the context menu, spawns a text entry for the name
+	 */
+	void CreateNewMetadataEntryClicked();
+
+	/**
+	 * Create a new curve name in the skeleton using the committed text and add that curve to the sequence
+	 */
+	void CreateNewMetadataEntry(const FText& CommittedText, ETextCommit::Type CommitType);
+
+	/**
+	* Handler for the new curve item in the context menu, spawns a text entry for the name
+	*/
+	void CreateNewCurveClicked();
+
+	/**
+	 * Convert the requested flag bool value into a checkbox state
+	 */
+	ESlateCheckBoxState::Type GetCurveFlagAsCheckboxState(FFloatCurve* Curve, EAnimCurveFlags InFlag) const;
+
+	/**
+	 * Convert a given checkbox state into a flag value in the provided curve
+	 */
+	void SetCurveFlagFromCheckboxState(ESlateCheckBoxState::Type CheckState, FFloatCurve* Curve, EAnimCurveFlags InFlag);
+
+	/**
+	 * Handler for converting a curve from metadata to variable curve or vice versa
+	 */
+	void ToggleCurveTypeMenuCallback(FFloatCurve* Curve);
 };
\ No newline at end of file
diff --git a/Engine/Source/Editor/Persona/Private/SMontageEditor.cpp b/Engine/Source/Editor/Persona/Private/SMontageEditor.cpp
--- a/Engine/Source/Editor/Persona/Private/SMontageEditor.cpp
+++ b/Engine/Source/Editor/Persona/Private/SMontageEditor.cpp
@@ -88,6 +88,7 @@ void SMontageEditor::Construct(const FArguments& InArgs)
 	.Padding(0, 10)
 	[
 		SAssignNew( AnimCurvePanel, SAnimCurvePanel )
+		.Persona(InArgs._Persona)
 		.Sequence(MontageObj)
 		.WidgetWidth(S2ColumnWidget::DEFAULT_RIGHT_COLUMN_WIDTH)
 		.ViewInputMin(this, &SAnimEditorBase::GetViewMinInput)
diff --git a/Engine/Source/Editor/Persona/Private/SSequenceEditor.cpp b/Engine/Source/Editor/Persona/Private/SSequenceEditor.cpp
--- a/Engine/Source/Editor/Persona/Private/SSequenceEditor.cpp
+++ b/Engine/Source/Editor/Persona/Private/SSequenceEditor.cpp
@@ -49,6 +49,7 @@ void SSequenceEditor::Construct(const FArguments& InArgs)
 	.Padding(0, 10)
 	[
 		SAssignNew( AnimCurvePanel, SAnimCurvePanel )
+		.Persona(InArgs._Persona)
 		.Sequence(SequenceObj)
 		.WidgetWidth(S2ColumnWidget::DEFAULT_RIGHT_COLUMN_WIDTH)
 		.ViewInputMin(this, &SAnimEditorBase::GetViewMinInput)
@@ -85,14 +86,18 @@ void SSequenceEditor::SetSequenceObj(UAnimSequenceBase * NewSequence)
 
 void SSequenceEditor::PostUndo()
 {
-	PersonaPtr.Pin()->SetPreviewAnimationAsset(SequenceObj);
-	// need to refresh panel
+	TSharedPtr<FPersona> SharedPersona = PersonaPtr.Pin();
+	if(SharedPersona.IsValid())
+	{
+		SharedPersona->SetPreviewAnimationAsset(SequenceObj);
+	}
+
 	if( SequenceObj )
 	{
 		SetInputViewRange(0, SequenceObj->SequenceLength);
 
-		// update notify panel 
 		AnimNotifyPanel->Update();
+		AnimCurvePanel->UpdatePanel();
 	}
 }
 
diff --git a/Engine/Source/Editor/Persona/Private/SSkeletonSmartNameManager.cpp b/Engine/Source/Editor/Persona/Private/SSkeletonSmartNameManager.cpp
new file mode 100644
--- /dev/null
+++ b/Engine/Source/Editor/Persona/Private/SSkeletonSmartNameManager.cpp
@@ -0,0 +1,482 @@
+// Copyright 1998-2014 Epic Games, Inc. All Rights Reserved.
+
+
+#include "PersonaPrivatePCH.h"
+
+#include "SSkeletonSmartNameManager.h"
+#include "ScopedTransaction.h"
+#include "AssetRegistryModule.h"
+#include "AssetData.h"
+
+#define LOCTEXT_NAMESPACE "SkeletonSmartNameManager"
+
+static const FName ColumnId_SmartNameLabel("SmartNameLabel");
+
+//////////////////////////////////////////////////////////////////////////
+
+typedef TSharedPtr<FDisplayedSmartNameInfo> FDisplayedSmartNameInfoPtr;
+
+class SSmartNameListRow : public SMultiColumnTableRow < FDisplayedSmartNameInfoPtr >
+{
+public:
+	SLATE_BEGIN_ARGS(SSmartNameListRow) {}
+
+	SLATE_ARGUMENT(FDisplayedSmartNameInfoPtr, Item)							// The display object for this row
+	SLATE_ARGUMENT(TSharedPtr<SSkeletonSmartNameManager>, NameListWidget)		// The list that we are in
+
+	SLATE_END_ARGS()
+
+	void Construct(const FArguments& InArgs, const TSharedRef<STableViewBase>& InOwnerTableView);
+
+	virtual TSharedRef<SWidget> GenerateWidgetForColumn(const FName& ColumnName) override;
+
+	// The editable text boxes in the rows consume clicks on the rows, this stops that behaviour
+	// similar to SSkeletonAnimNotifies
+	bool IsSelected() {return false;}
+
+private:
+
+	// Gets the name of the smartname we are displaying
+	FText GetItemName() const;
+
+	FDisplayedSmartNameInfoPtr Item;						// The display object for this row
+	TSharedPtr<SSkeletonSmartNameManager> NameListWidget;	// The list we are in
+};
+
+void SSmartNameListRow::Construct(const FArguments& InArgs, const TSharedRef<STableViewBase>& InOwnerTableView)
+{
+	Item = InArgs._Item;
+	NameListWidget = InArgs._NameListWidget;
+
+	check(Item.IsValid());
+
+	SMultiColumnTableRow< FDisplayedSmartNameInfoPtr >::Construct(FSuperRowType::FArguments(), InOwnerTableView);
+}
+
+TSharedRef<SWidget> SSmartNameListRow::GenerateWidgetForColumn(const FName& ColumnName)
+{
+	return SNew(SVerticalBox)
+	+ SVerticalBox::Slot()
+	.AutoHeight()
+	.Padding(0.0f, 4.0f)
+	.VAlign(VAlign_Center)
+	[
+		SAssignNew(Item->EditableText, SInlineEditableTextBlock)
+		.OnTextCommitted(NameListWidget.Get(), &SSkeletonSmartNameManager::OnNameCommitted, Item)
+		.IsSelected(this, &SSmartNameListRow::IsSelected)
+		.Text(this, &SSmartNameListRow::GetItemName)
+	];
+}
+
+FText SSmartNameListRow::GetItemName() const
+{
+	const FSmartNameMapping* Mapping = Item->Skeleton->SmartNames.GetContainer(Item->ContainerName);
+	check(Mapping);
+	FName ItemName;
+	Mapping->GetName(Item->NameUid, ItemName);
+	
+	return FText::FromName(ItemName);
+}
+
+//////////////////////////////////////////////////////////////////////////
+
+FSkeletonCurveNameManagerSummoner::FSkeletonCurveNameManagerSummoner(TSharedPtr<class FAssetEditorToolkit> InHostingApp)
+	:FWorkflowTabFactory(FPersonaTabs::CurveNameManagerID, InHostingApp)
+{
+	TabLabel = LOCTEXT("CurveNameTabTitle", "Skeleton Curves");
+
+	EnableTabPadding();
+	bIsSingleton = true;
+
+	ViewMenuDescription = LOCTEXT("CurveNameTabMenuDesc", "Skeleton Curves");
+	ViewMenuTooltip = LOCTEXT("CurveNameTabMenuToolTip", "Shows the list of curves for the current skeleton");
+}
+
+TSharedRef<SWidget> FSkeletonCurveNameManagerSummoner::CreateTabBody(const FWorkflowTabSpawnInfo& Info) const
+{
+	TSharedPtr<FAssetEditorToolkit> Toolkit = HostingApp.Pin();
+	check(Toolkit.IsValid());
+
+	return SNew(SCurveNameManager)
+		.Persona(StaticCastSharedPtr<FPersona>(Toolkit))
+		.ContainerName(FName(USkeleton::AnimCurveMappingName))
+		.NameColumnDescription(LOCTEXT("CurveNameTabColumnDesc", "Name").ToString());
+}
+
+//////////////////////////////////////////////////////////////////////////
+
+SSkeletonSmartNameManager::~SSkeletonSmartNameManager()
+{
+	TSharedPtr<FPersona> SharedPersona = WeakPersona.Pin();
+	if(SharedPersona.IsValid())
+	{
+		SharedPersona->UnregisterOnPostUndo(this);
+	}
+}
+
+void SSkeletonSmartNameManager::Construct(const FArguments& InArgs)
+{
+	WeakPersona = InArgs._Persona;
+	ContainerName = InArgs._ContainerName;
+
+	TSharedPtr<FPersona> SharedPersona = WeakPersona.Pin();
+	check(SharedPersona.IsValid());
+
+	CurrentSkeleton = SharedPersona->GetSkeleton();
+	CurrentFilterText = FText::GetEmpty();
+
+	SharedPersona->RegisterOnPostUndo(FPersona::FOnPostUndo::CreateSP(this, &SSkeletonSmartNameManager::OnPostUndo));
+
+	this->ChildSlot
+	[
+		SNew(SVerticalBox)
+
+		// Search box
+		+ SVerticalBox::Slot()
+		.AutoHeight()
+		[
+			SNew(SSearchBox)
+			.SelectAllTextWhenFocused(true)
+			.OnTextChanged(this, &SSkeletonSmartNameManager::OnFilterTextChanged)
+			.OnTextCommitted(this, &SSkeletonSmartNameManager::OnFilterTextCommitted)
+			.HintText(LOCTEXT("SmartNameSearchHint", "Search Names . . . "))
+		]
+
+		// Name list
+		+ SVerticalBox::Slot()
+		.FillHeight(1.0f)
+		[
+			SAssignNew(ListWidget, SSmartNameListType)
+			.ListItemsSource(&DisplayedInfoList)
+			.OnGenerateRow(this, &SSkeletonSmartNameManager::OnGenerateRow)
+			.OnContextMenuOpening(this, &SSkeletonSmartNameManager::OnGetContextMenu)
+			.ItemHeight(22.0f)
+			.SelectionMode(ESelectionMode::Multi)
+			.HeaderRow
+			(
+				SNew(SHeaderRow)
+				+SHeaderRow::Column(ColumnId_SmartNameLabel)
+				.DefaultLabel(FString(InArgs._NameColumnDescription))
+			)
+		]
+	];
+
+	GenerateDisplayedList(CurrentFilterText);
+}
+
+FText& SSkeletonSmartNameManager::GetNameFilterText()
+{
+	return CurrentFilterText;
+}
+
+void SSkeletonSmartNameManager::OnFilterTextChanged(const FText& FilterText)
+{
+	CurrentFilterText = FilterText;
+
+	GenerateDisplayedList(CurrentFilterText);
+}
+
+void SSkeletonSmartNameManager::OnFilterTextCommitted(const FText& FilterText, ETextCommit::Type CommitType)
+{
+	OnFilterTextChanged(FilterText);
+}
+
+TSharedRef<ITableRow> SSkeletonSmartNameManager::OnGenerateRow(TSharedPtr<FDisplayedSmartNameInfo> InInfo, const TSharedRef<STableViewBase>& OwningTable)
+{
+	check(InInfo.IsValid());
+
+	return SNew(SSmartNameListRow, OwningTable)
+		.Item(InInfo)
+		.NameListWidget(SharedThis(this));
+}
+
+TSharedPtr<SWidget> SSkeletonSmartNameManager::OnGetContextMenu() const
+{
+	FMenuBuilder MenuBuilder(true, nullptr);
+
+	MenuBuilder.BeginSection("SmartNameAction", LOCTEXT("SmartNameActions", "Actions"));
+	{
+		FUIAction Action;
+		// Rename selection
+		{
+			Action.ExecuteAction = FExecuteAction::CreateSP(this, &SSkeletonSmartNameManager::OnRenameClicked);
+			Action.CanExecuteAction = FCanExecuteAction::CreateSP(this, &SSkeletonSmartNameManager::CanRename);
+			const FText Label = LOCTEXT("RenameSmartNameLabel", "Rename");
+			const FText ToolTip = LOCTEXT("RenameSmartNameToolTip", "Rename the selected item");
+			MenuBuilder.AddMenuEntry(Label, ToolTip, FSlateIcon(), Action);
+		}
+		// Delete a name
+		{
+			Action.ExecuteAction = FExecuteAction::CreateSP(this, &SSkeletonSmartNameManager::OnDeleteNameClicked);
+			Action.CanExecuteAction = FCanExecuteAction::CreateSP(this, &SSkeletonSmartNameManager::CanDelete);
+			const FText Label = LOCTEXT("DeleteSmartNameLabel", "Delete");
+			const FText ToolTip = LOCTEXT("DeleteSmartNameToolTip", "Delete the selected item");
+			MenuBuilder.AddMenuEntry(Label, ToolTip, FSlateIcon(), Action);
+		}
+		// Add a name
+		MenuBuilder.AddMenuSeparator();
+		{
+			Action.ExecuteAction = FExecuteAction::CreateSP(this, &SSkeletonSmartNameManager::OnAddClicked);
+			Action.CanExecuteAction = nullptr;
+			const FText Label = LOCTEXT("AddSmartNameLabel", "Add...");
+			const FText ToolTip = LOCTEXT("AddSmartNameToolTip", "Add an entry to the skeleton");
+			MenuBuilder.AddMenuEntry(Label, ToolTip, FSlateIcon(), Action);
+		}
+	}
+	MenuBuilder.EndSection();
+
+	return MenuBuilder.MakeWidget();
+}
+
+void SSkeletonSmartNameManager::GenerateDisplayedList(const FText& FilterText)
+{
+	FSmartNameMapping* Mapping = CurrentSkeleton->SmartNames.GetContainer(ContainerName);
+	if(Mapping)
+	{
+		DisplayedInfoList.Empty();
+		TArray<FSmartNameMapping::UID> UidList;
+		Mapping->FillUidArray(UidList);
+
+		for(FSmartNameMapping::UID Uid : UidList)
+		{
+			bool bAddToList = true;
+			if(!FilterText.IsEmpty())
+			{
+				FName CurrentName;
+				Mapping->GetName(Uid, CurrentName);
+
+				if(!CurrentName.ToString().Contains(*FilterText.ToString()))
+				{
+					bAddToList = false;
+				}
+			}
+			
+			if(bAddToList)
+			{
+				DisplayedInfoList.Add(FDisplayedSmartNameInfo::Make(CurrentSkeleton, ContainerName, Uid));
+			}
+		}
+
+		ListWidget->RequestListRefresh();
+	}
+}
+
+void SSkeletonSmartNameManager::OnNameCommitted(const FText& NewName, ETextCommit::Type, TSharedPtr<FDisplayedSmartNameInfo> Item)
+{
+	FSmartNameMapping* Mapping = CurrentSkeleton->SmartNames.GetContainer(ContainerName);
+	if(Mapping)
+	{
+		FName NewName(*NewName.ToString());
+		if(!Mapping->Exists(NewName))
+		{
+			FScopedTransaction Transaction(LOCTEXT("TransactionRename", "Rename Element"));
+			CurrentSkeleton->RenameSmartnameAndModify(ContainerName, Item->NameUid, NewName);
+		}
+	}
+}
+
+void SSkeletonSmartNameManager::OnDeleteNameClicked()
+{
+	TArray<FDisplayedSmartNameInfoPtr> SelectedItems = ListWidget->GetSelectedItems();
+
+	bool bModifiedList = false;
+
+	for(FDisplayedSmartNameInfoPtr Item : SelectedItems)
+	{
+		CurrentSkeleton->RemoveSmartnameAndModify(Item->ContainerName, Item->NameUid);
+	}
+
+	GenerateDisplayedList(CurrentFilterText);
+	ListWidget->RequestListRefresh();
+}
+
+void SSkeletonSmartNameManager::OnRenameClicked()
+{
+	TArray<FDisplayedSmartNameInfoPtr> SelectedItems = ListWidget->GetSelectedItems();
+
+	SelectedItems[0]->EditableText->EnterEditingMode();
+}
+
+void SSkeletonSmartNameManager::OnAddClicked()
+{
+	TSharedRef<STextEntryPopup> TextEntry =
+		SNew(STextEntryPopup)
+		.Label(LOCTEXT("NewSmartnameLabel", "New Name"))
+		.OnTextCommitted(this, &SSkeletonSmartNameManager::CreateNewNameEntry);
+
+	FSlateApplication& SlateApp = FSlateApplication::Get();
+	SlateApp.PushMenu(
+		AsShared(),
+		TextEntry,
+		SlateApp.GetCursorPos(),
+		FPopupTransitionEffect::TypeInPopup
+		);
+
+	TextEntry->FocusDefaultWidget();
+}
+
+void SSkeletonSmartNameManager::OnPostUndo()
+{
+	GenerateDisplayedList(CurrentFilterText);
+}
+
+bool SSkeletonSmartNameManager::CanDelete()
+{
+	return ListWidget->GetNumItemsSelected() > 0;
+}
+
+bool SSkeletonSmartNameManager::CanRename()
+{
+	return ListWidget->GetNumItemsSelected() == 1;
+}
+
+void SSkeletonSmartNameManager::CreateNewNameEntry(const FText& CommittedText, ETextCommit::Type CommitType)
+{
+	FSlateApplication::Get().DismissAllMenus();
+	if(!CommittedText.IsEmpty() && CommitType == ETextCommit::OnEnter)
+	{
+		if(FSmartNameMapping* NameMapping = CurrentSkeleton->SmartNames.GetContainer(ContainerName))
+		{
+			FName NewName = FName(*CommittedText.ToString());
+			FSmartNameMapping::UID NewUid;
+			if(CurrentSkeleton->AddSmartnameAndModify(ContainerName, NewName, NewUid))
+			{
+				// Successfully added
+				GenerateDisplayedList(CurrentFilterText);
+				ListWidget->RequestListRefresh();
+			}
+		}
+	}
+}
+
+//////////////////////////////////////////////////////////////////////////
+
+void SCurveNameManager::OnDeleteNameClicked()
+{
+	TArray<FDisplayedSmartNameInfoPtr> SelectedItems = ListWidget->GetSelectedItems();
+	TArray<FSmartNameMapping::UID> SelectedUids;
+
+	for(FDisplayedSmartNameInfoPtr Item : SelectedItems)
+	{
+		SelectedUids.Add(Item->NameUid);
+	}
+
+	FAssetRegistryModule& AssetModule = FModuleManager::LoadModuleChecked<FAssetRegistryModule>(TEXT("AssetRegistry"));
+	TArray<FAssetData> AnimationAssets;
+	AssetModule.Get().GetAssetsByClass(UAnimSequenceBase::StaticClass()->GetFName(), AnimationAssets, true);
+	FAssetData SkeletonData(CurrentSkeleton);
+	const FString CurrentSkeletonName = SkeletonData.GetExportTextName();
+
+	GWarn->BeginSlowTask(LOCTEXT("CollectAnimationsTaskDesc", "Collecting assets..."), true);
+
+	for(int32 Idx = AnimationAssets.Num() - 1 ; Idx >= 0 ; --Idx)
+	{
+		FAssetData& Data = AnimationAssets[Idx];
+		bool bRemove = true;
+
+		const FString* SkeletonData = Data.TagsAndValues.Find("Skeleton");
+		if(SkeletonData && *SkeletonData == CurrentSkeletonName)
+		{
+			const FString* CurveData = Data.TagsAndValues.Find(USkeleton::CurveTag);
+			FString CurveDataCopy;
+
+			if(!CurveData)
+			{
+				// This asset is old; it hasn't been loaded since before smartnames were added for curves.
+				// unfortunately the only way to delete curves safely is to load old assets to see if they're
+				// using the selected name. We only load what we have to here.
+				UObject* Asset = Data.GetAsset();
+				check(Asset);
+				TArray<UObject::FAssetRegistryTag> Tags;
+				Asset->GetAssetRegistryTags(Tags);
+				
+				UObject::FAssetRegistryTag* CurveTag = Tags.FindByPredicate([](const UObject::FAssetRegistryTag& Tag)
+				{
+					return Tag.Name == USkeleton::CurveTag;
+				});
+				CurveDataCopy = CurveTag->Value;
+			}
+			else
+			{
+				CurveDataCopy = *CurveData;
+			}
+
+			TArray<FString> ParsedStringUids;
+			CurveDataCopy.ParseIntoArray(&ParsedStringUids, &USkeleton::CurveTagDelimiter, true);
+
+			for(const FString& UidString : ParsedStringUids)
+			{
+				FSmartNameMapping::UID Uid = (FSmartNameMapping::UID)TCString<TCHAR>::Strtoui64(*UidString, NULL, 10);
+				if(SelectedUids.Contains(Uid))
+				{
+					bRemove = false;
+					break;
+				}
+			}
+		}
+
+		if(bRemove)
+		{
+			AnimationAssets.RemoveAt(Idx);
+		}
+	}
+
+	GWarn->EndSlowTask();
+
+	// AnimationAssets now only contains assets that are using the selected curve(s)
+	if(AnimationAssets.Num() > 0)
+	{
+		FString ConfirmMessage = LOCTEXT("DeleteCurveMessage", "The following assets use the selected curve(s), deleting will remove the curves from these assets. Continue?\n\n").ToString();
+		
+		for(FAssetData& Data : AnimationAssets)
+		{
+			ConfirmMessage += Data.AssetName.ToString() + " (" + Data.AssetClass.ToString() + ")\n";
+		}
+
+		FText Title = LOCTEXT("DeleteCurveDialogTitle", "Confirm Deletion");
+		FText Message = FText::FromString(ConfirmMessage);
+		if(FMessageDialog::Open(EAppMsgType::YesNo, Message, &Title) == EAppReturnType::Yes)
+		{
+			// Proceed to delete the curves
+			GWarn->BeginSlowTask(FText::Format(LOCTEXT("DeleteCurvesTaskDesc", "Deleting curve from skeleton {1}"), FText::FromString(CurrentSkeleton->GetName())), true);
+			FScopedTransaction Transaction(LOCTEXT("DeleteCurvesTransactionName", "Delete skeleton curve"));
+
+			// Remove curves from animation assets
+			for(FAssetData& Data : AnimationAssets)
+			{
+				UAnimSequenceBase* Sequence = Cast<UAnimSequenceBase>(Data.GetAsset());	
+				Sequence->Modify(true);
+				for(FSmartNameMapping::UID Uid : SelectedUids)
+				{
+					Sequence->RawCurveData.DeleteCurveData(Uid);
+				}
+			}
+
+			// Remove names from skeleton
+			CurrentSkeleton->Modify(true);
+			FSmartNameMapping* Mapping = CurrentSkeleton->SmartNames.GetContainer(ContainerName);
+			for(FSmartNameMapping::UID Uid : SelectedUids)
+			{
+				Mapping->Remove(Uid);
+			}
+
+			GWarn->EndSlowTask();
+		}
+	}
+
+	GenerateDisplayedList(CurrentFilterText);
+
+	TSharedPtr<FPersona> SharedPersona = WeakPersona.Pin();
+	if(SharedPersona.IsValid())
+	{
+		if(SharedPersona->GetCurrentMode() == FPersonaModes::AnimationEditMode)
+		{
+			// If we're in animation mode then we need to reinitialise the mode to
+			// refresh the sequencer which will now have a widget it shouldn't have
+			SharedPersona->ReinitMode();
+		}
+	}
+}
+
+
+#undef LOCTEXT_NAMESPACE
diff --git a/Engine/Source/Editor/Persona/Private/SSkeletonSmartNameManager.h b/Engine/Source/Editor/Persona/Private/SSkeletonSmartNameManager.h
new file mode 100644
--- /dev/null
+++ b/Engine/Source/Editor/Persona/Private/SSkeletonSmartNameManager.h
@@ -0,0 +1,189 @@
+// Copyright 1998-2014 Epic Games, Inc. All Rights Reserved.
+
+#pragma once
+
+#include "IDocumentation.h"
+#include "Persona.h"
+
+#define LOCTEXT_NAMESPACE "SkeletonSmartNameManager"
+
+struct FSkeletonCurveNameManagerSummoner : public FWorkflowTabFactory
+{
+	FSkeletonCurveNameManagerSummoner(TSharedPtr<class FAssetEditorToolkit> InHostingApp);
+
+	virtual TSharedRef<SWidget> CreateTabBody(const FWorkflowTabSpawnInfo& Info) const override;
+
+	virtual TSharedPtr<SToolTip> CreateTabToolTipWidget(const FWorkflowTabSpawnInfo& Info) const override
+	{
+		return IDocumentation::Get()->CreateToolTip(LOCTEXT("CurveNamesTooltip", "Allows management of curve names and metadata names stored on the current skeleton (Add / Remove / Rename)"), nullptr, TEXT("Shared/Editors/Persona"), TEXT("CurveNameManager_Window"));
+	}
+};
+
+//////////////////////////////////////////////////////////////////////////
+// Used to store display information for each name in the list.
+
+class FDisplayedSmartNameInfo
+{
+public:
+	FSmartNameMapping::UID					NameUid;		// The UID of the displayed name
+	FName									ContainerName;	// The container in the skeleton this name resides in
+	const class USkeleton*					Skeleton;		// The skeleton we're associated with
+	TSharedPtr<SInlineEditableTextBlock>	EditableText;	// The editable text box in the list, used to focus from the context menu
+
+	static TSharedRef<FDisplayedSmartNameInfo> Make(const class USkeleton* InSkeleton, const FName& InContainerName, const FSmartNameMapping::UID InNameUid)
+	{
+		return MakeShareable(new FDisplayedSmartNameInfo(InSkeleton, InContainerName, InNameUid));
+	}
+
+private:
+
+	// Hidden constructors, use Make() to create an instance of this class
+	FDisplayedSmartNameInfo();
+
+	FDisplayedSmartNameInfo(const class USkeleton* InSkeleton, const FName& InContainerName, const FSmartNameMapping::UID InNameUid)
+		: NameUid(InNameUid)
+		, ContainerName(InContainerName)
+		, Skeleton(InSkeleton)
+	{}
+
+};
+
+//////////////////////////////////////////////////////////////////////////
+// Generalized smart name list and search widget 
+
+typedef SListView<TSharedPtr<FDisplayedSmartNameInfo>> SSmartNameListType;
+
+class SSkeletonSmartNameManager : public SCompoundWidget
+{
+public:
+	SLATE_BEGIN_ARGS(SSkeletonSmartNameManager)
+		: _Persona()
+		, _ContainerName()
+		, _NameColumnDescription()
+	{}
+
+	SLATE_ARGUMENT(TWeakPtr<FPersona>, Persona)			// Our Persona instance
+	SLATE_ARGUMENT(FName, ContainerName)				// The name of the smartname container to manage
+	SLATE_ARGUMENT(FString, NameColumnDescription)		// Description text for the display column header
+
+	SLATE_END_ARGS();
+
+	virtual ~SSkeletonSmartNameManager();
+
+	void Construct(const FArguments& InArgs);
+
+	// Get the current filter text from the search box
+	FText& GetNameFilterText();
+
+	// When a name is committed after being edited in the list
+	virtual void OnNameCommitted(const FText& NewName, ETextCommit::Type CommitType, TSharedPtr<FDisplayedSmartNameInfo> Item);
+
+protected:
+
+	// Updates the search filter
+	void OnFilterTextChanged(const FText& FilterText);
+
+	// Updates the search filter
+	void OnFilterTextCommitted(const FText& FilterText, ETextCommit::Type CommitType);
+
+	// Generates a table row for the provided info object
+	TSharedRef<ITableRow> OnGenerateRow(TSharedPtr<FDisplayedSmartNameInfo> InInfo, const TSharedRef<STableViewBase>& OwningTable);
+
+	// Creates a context menu
+	TSharedPtr<SWidget> OnGetContextMenu() const;
+
+	// Clears and regenerates the displayed info list for the window
+	void GenerateDisplayedList(const FText& FilterText);
+
+	// Attempt to delete a smartname. This can be overriden in derived classes to handle the
+	// delete as this version will simply find the correct name container in the skeleton and
+	// remove the name.
+	virtual void OnDeleteNameClicked();
+
+	// Rename clicked on the context menu, simply activates the editable text box in the
+	// row. Can be overridden in derived classes to provide extra functionality if
+	// necessary
+	virtual void OnRenameClicked();
+
+	// Add clicked in the context menu. Shows a text entry popup to add a new name to the
+	// smartname container in the skeleton
+	virtual void OnAddClicked();
+
+	// Called by Persona after an undo operation, base version simply refreshes the
+	// list to repopulate it if something had been deleted
+	virtual void OnPostUndo();
+
+	// Used to control delete availability. Checks to see if we have anything selected
+	bool CanDelete();
+	
+	// Used to control rename availability. Checks to see if we have one item selected
+	bool CanRename();
+	
+	// Adds a new smartname entry to the skeleton in the container we are managing
+	void CreateNewNameEntry(const FText& CommittedText, ETextCommit::Type CommitType);
+
+	TWeakPtr<FPersona> WeakPersona;										// Persona instance this manager is in
+	class USkeleton* CurrentSkeleton;									// The current skeleton being edited by out Persona instance
+	FName ContainerName;												// Name of the skeleton smart name container to display in the list
+	TSharedPtr<SSearchBox> ListSearchBox;								// The search box at the top of the window
+	TSharedPtr<SSmartNameListType> ListWidget;							// The list inside the window
+	TArray<TSharedPtr<FDisplayedSmartNameInfo>> DisplayedInfoList;		// List of information created for the list
+	FText CurrentFilterText;											// The text currently entered in the search box at the top of the window
+
+};
+
+//////////////////////////////////////////////////////////////////////////
+// Specialised smart name manager to handle animation curves
+class SCurveNameManager : public SSkeletonSmartNameManager
+{
+public:
+
+	SLATE_BEGIN_ARGS(SCurveNameManager)
+		: _Persona()
+		, _ContainerName()
+		, _NameColumnDescription()
+		{}
+
+		SLATE_ARGUMENT(TWeakPtr<FPersona>, Persona)			// Our Persona instance
+		SLATE_ARGUMENT(FName, ContainerName)				// Name of the container to manage in the skeleton
+		SLATE_ARGUMENT(FString, NameColumnDescription)		// Description to show in the name column
+	SLATE_END_ARGS();
+
+	void Construct(const FArguments& InArgs)
+	{
+		TSharedPtr<FPersona> SharedPersona = InArgs._Persona.Pin();
+		if(SharedPersona.IsValid())
+		{
+			SharedPersona->RegisterOnChangeCurves(FPersona::FOnCurvesChanged::CreateSP(this, &SCurveNameManager::OnCurvesChanged));
+		}
+
+		SSkeletonSmartNameManager::Construct(SSkeletonSmartNameManager::FArguments()
+											 .ContainerName(InArgs._ContainerName)
+											 .Persona(InArgs._Persona)
+											 .NameColumnDescription(InArgs._NameColumnDescription));
+	}
+
+	virtual ~SCurveNameManager()
+	{
+		TSharedPtr<FPersona> SharedPersona = WeakPersona.Pin();
+		if(SharedPersona.IsValid())
+		{
+			SharedPersona->UnregisterOnChangeCurves(this);
+		}
+	}
+
+protected:
+
+	// Checks dependant animations that could be using the requested curve name and
+	// checks that the user wants to remove that curve.
+	virtual void OnDeleteNameClicked() override;
+
+	// Regenerates the list of curve names.
+	void OnCurvesChanged()
+	{
+		GenerateDisplayedList(CurrentFilterText);
+	}
+
+};
+
+#undef LOCTEXT_NAMESPACE
diff --git a/Engine/Source/Editor/Persona/Private/Shared/PersonaMode.cpp b/Engine/Source/Editor/Persona/Private/Shared/PersonaMode.cpp
--- a/Engine/Source/Editor/Persona/Private/Shared/PersonaMode.cpp
+++ b/Engine/Source/Editor/Persona/Private/Shared/PersonaMode.cpp
@@ -8,6 +8,7 @@
 #include "SAnimBlueprintParentPlayerList.h"
 #include "SSkeletonSlotGroupNames.h"
 #include "SSkeletonSlotNames.h"
+#include "SSkeletonSmartNameManager.h"
 
 #define LOCTEXT_NAMESPACE "PersonaModes"
 
@@ -39,6 +40,7 @@ const FName FPersonaTabs::PreviewManagerID("AnimPreviewSetup");		//@TODO: Name
 const FName FPersonaTabs::SkeletonAnimNotifiesID("SkeletonAnimNotifies");
 const FName FPersonaTabs::SkeletonSlotNamesID("SkeletonSlotNames");
 const FName FPersonaTabs::SkeletonSlotGroupNamesID("SkeletonSlotGroupNames");
+const FName FPersonaTabs::CurveNameManagerID("CurveNameManager");
 
 /////////////////////////////////////////////////////
 // FPersonaMode
@@ -66,6 +68,7 @@ FPersonaAppMode::FPersonaAppMode(TSharedPtr<class FPersona> InPersona, FName InM
 	PersonaTabFactories.RegisterFactory(MakeShareable(new FRetargetSourceManagerTabSummoner(InPersona)));
 	PersonaTabFactories.RegisterFactory(MakeShareable(new FSkeletonSlotNamesSummoner(InPersona)));
 	PersonaTabFactories.RegisterFactory(MakeShareable(new FSkeletonSlotGroupNamesSummoner(InPersona)));
+	PersonaTabFactories.RegisterFactory(MakeShareable(new FSkeletonCurveNameManagerSummoner(InPersona)));
 }
 
 void FPersonaAppMode::RegisterTabFactories(TSharedPtr<FTabManager> InTabManager)
diff --git a/Engine/Source/Editor/Persona/Private/Shared/PersonaMode.h b/Engine/Source/Editor/Persona/Private/Shared/PersonaMode.h
--- a/Engine/Source/Editor/Persona/Private/Shared/PersonaMode.h
+++ b/Engine/Source/Editor/Persona/Private/Shared/PersonaMode.h
@@ -36,6 +36,7 @@ struct FPersonaTabs
 	static const FName SkeletonAnimNotifiesID;
 	static const FName SkeletonSlotNamesID;
 	static const FName SkeletonSlotGroupNamesID;
+	static const FName CurveNameManagerID;
 
 	// Blueprint Document
 
diff --git a/Engine/Source/Editor/UnrealEd/Private/SkeletalMeshEdit.cpp b/Engine/Source/Editor/UnrealEd/Private/SkeletalMeshEdit.cpp
--- a/Engine/Source/Editor/UnrealEd/Private/SkeletalMeshEdit.cpp
+++ b/Engine/Source/Editor/UnrealEd/Private/SkeletalMeshEdit.cpp
@@ -943,13 +943,18 @@ bool UnFbx::FFbxImporter::ImportCurveToAnimSequence(class UAnimSequence * Target
 	if (TargetSequence && FbxCurve)
 	{
 		FName Name = *CurveName;
+		FSmartNameMapping* NameMapping = TargetSequence->GetSkeleton()->SmartNames.GetContainer(USkeleton::AnimCurveMappingName);
 
-		FFloatCurve * CurveToImport = TargetSequence->RawCurveData.GetCurveData(Name);
+		// Add or retrieve curve
+		USkeleton::AnimCurveUID Uid;
+		NameMapping->AddName(Name, Uid);
+
+		FFloatCurve * CurveToImport = TargetSequence->RawCurveData.GetCurveData(Uid);
 		if(CurveToImport==NULL)
 		{
-			if(TargetSequence->RawCurveData.AddCurveData(Name, CurveFlags))
+			if(TargetSequence->RawCurveData.AddCurveData(Uid, CurveFlags))
 			{
-				CurveToImport = TargetSequence->RawCurveData.GetCurveData(Name);
+				CurveToImport = TargetSequence->RawCurveData.GetCurveData(Uid);
 			}
 			else
 			{
@@ -1025,6 +1030,32 @@ bool UnFbx::FFbxImporter::ImportAnimation(USkeleton * Skeleton, UAnimSequence *
 							{
 								// now see if we have one already exists. If so, just overwrite that. if not, add new one. 
 								ImportCurveToAnimSequence(DestSeq, *ChannelName, Curve,  ACF_DrivesMorphTarget | ACF_TriggerEvent, AnimTimeSpan, 0.01f /** for some reason blend shape values are coming as 100 scaled **/);
+
+
+//								FName Name = *ChannelName;
+//								USkeleton::AnimCurveUID CurveId;
+//								Skeleton->SmartNames.GetContainer(USkeleton::AnimCurveMappingName)->AddName(Name, CurveId);
+//								FFloatCurve * CurveToImport = DestSeq->RawCurveData.GetCurveData(CurveId);
+//								if (CurveToImport==NULL)
+//								{
+//									if ( DestSeq->RawCurveData.AddCurveData(CurveId, ACF_DrivesMorphTarget | ACF_TriggerEvent) )
+//									{
+//										CurveToImport = DestSeq->RawCurveData.GetCurveData(CurveId);
+//									}
+//									else
+//									{
+//										// this should not happen, we already checked before adding
+//										ensureMsg (0, TEXT("FBX Import: Critical error: no memory?"));
+//									}
+//								}
+//								else
+//								{
+//									CurveToImport->FloatCurve.Reset();
+//								}
+//	
+//								ImportCurve(Curve, CurveToImport, AnimTimeSpan);
+
+
 							}
 						}
 					}
diff --git a/Engine/Source/Runtime/Core/Public/UObject/ObjectVersion.h b/Engine/Source/Runtime/Core/Public/UObject/ObjectVersion.h
--- a/Engine/Source/Runtime/Core/Public/UObject/ObjectVersion.h
+++ b/Engine/Source/Runtime/Core/Public/UObject/ObjectVersion.h
@@ -588,6 +588,8 @@ enum EUnrealEngineObjectUE4Version
 	VER_UE4_LEVEL_STREAMING_DRAW_COLOR_TYPE_CHANGE,
 	// Cleared end triggers from non-state anim notifies
 	VER_UE4_CLEAR_NOTIFY_TRIGGERS,
+	// Convert old curve names stored in anim assets into skeleton smartnames
+	VER_UE4_SKELETON_ADD_SMARTNAMES,
 
 	// -----<new versions can be added before this line>-------------------------------------------------
 	// - this needs to be the last line (see note below)
diff --git a/Engine/Source/Runtime/Engine/Classes/Animation/AnimInstance.h b/Engine/Source/Runtime/Engine/Classes/Animation/AnimInstance.h
--- a/Engine/Source/Runtime/Engine/Classes/Animation/AnimInstance.h
+++ b/Engine/Source/Runtime/Engine/Classes/Animation/AnimInstance.h
@@ -545,8 +545,8 @@ public:
 	/** Trigger AnimNotifies **/
 	void TriggerAnimNotifies(float DeltaSeconds);
 
-	/** Add curve float data **/
-	void AddCurveValue(const FName & CurveName, float Value, int32 CurveTypeFlags);
+	/** Add curve float data using a curve Uid, the name of the curve will be resolved from the skeleton **/
+	void AddCurveValue(const USkeleton::AnimCurveUID Uid, float Value, int32 CurveTypeFlags);
 
 	/* native cpp montage interface functions **/
 	void Montage_SetEndDelegate(FOnMontageEnded & OnMontageEnded);
@@ -593,6 +593,12 @@ protected:
 	/** Called by blueprint functions that modify the montages current position. */
 	void OnMontagePositionChanged(FAnimMontageInstance* MontageInstance, FName ToSectionName);
 
+	/** 
+	 * Add curve float data, using a curve name. External values should all be added using
+	 * The curve UID to the public version of this method
+	 */
+	void AddCurveValue(const FName & CurveName, float Value, int32 CurveTypeFlags);
+
 #if WITH_EDITORONLY_DATA
 	// Returns true if a snapshot is being played back and the remainder of Update should be skipped.
 	bool UpdateSnapshotAndSkipRemainingUpdate();
diff --git a/Engine/Source/Runtime/Engine/Classes/Animation/AnimSequenceBase.h b/Engine/Source/Runtime/Engine/Classes/Animation/AnimSequenceBase.h
--- a/Engine/Source/Runtime/Engine/Classes/Animation/AnimSequenceBase.h
+++ b/Engine/Source/Runtime/Engine/Classes/Animation/AnimSequenceBase.h
@@ -8,6 +8,8 @@
  */
 
 #include "AnimationAsset.h"
+#include "SmartName.h"
+#include "Skeleton.h"
 #include "Curves/CurveBase.h"
 #include "AnimSequenceBase.generated.h"
 
@@ -155,16 +157,33 @@ struct FAnimCurveBase
 {
 	GENERATED_USTRUCT_BODY()
 
-	/** Name of curve. - in the future, we might need TMap, right now TMap doesn't support property serialization, but in the future we can change this to TMap*/
+	// Last observed name of the curve. We store this so we can recover from situations that
+	// mean the skeleton doesn't have a mapped name for our UID (such as a user saving the an
+	// animation but not the skeleton).
 	UPROPERTY()
-	FName		CurveName;
+	FName		LastObservedName;
+
+	FSmartNameMapping::UID CurveUid;
 
 	FAnimCurveBase(){}
 
 	FAnimCurveBase(FName InName)
-		: CurveName(InName)
+		: LastObservedName(InName)
 	{	
 	}
+
+	FAnimCurveBase(USkeleton::AnimCurveUID Uid)
+		: CurveUid(Uid)
+	{}
+
+	// To be able to use typedef'd types we need to serialize manually
+	virtual void Serialize(FArchive& Ar)
+	{
+		if(Ar.UE4Ver() >= VER_UE4_SKELETON_ADD_SMARTNAMES)
+		{
+			Ar << CurveUid;
+		}
+	}
 };
 
 /** native enum for curve types **/
@@ -178,6 +197,8 @@ enum EAnimCurveFlags
 	ACF_Editable			= 0x00000004,
 	// Used as a material curve
 	ACF_DrivesMaterial		= 0x00000008,
+	// Is a metadata 'curve'
+	ACF_Metadata		= 0x00000010,
 
 	// default flag when created
 	ACF_DefaultCurve		= ACF_TriggerEvent | ACF_Editable,
@@ -202,10 +223,21 @@ struct FFloatCurve : public FAnimCurveBase
 	{
 	}
 
+	FFloatCurve(USkeleton::AnimCurveUID Uid, int32 InCurveTypeFlags)
+		: FAnimCurveBase(Uid)
+		, CurveTypeFlags(InCurveTypeFlags)
+	{}
+
 	/**
 	 * Set InFlag to bValue
 	 */
 	ENGINE_API void SetCurveTypeFlag(EAnimCurveFlags InFlag, bool bValue);
+
+	/**
+	 * Toggle the value of the specified flag
+	 */
+	ENGINE_API void ToggleCurveTypeFlag(EAnimCurveFlags InFlag);
+
 	/**
 	 * Return true if InFlag is set, false otherwise 
 	 */
@@ -221,6 +253,11 @@ struct FFloatCurve : public FAnimCurveBase
 	 */
 	int32 GetCurveTypeFlags() const;
 
+	virtual void Serialize(FArchive& Ar) override
+	{
+		FAnimCurveBase::Serialize(Ar);
+	}
+
 private:
 
 	/** Curve Type Flags */
@@ -244,23 +281,30 @@ struct FRawCurveTracks
 	 */
 	void EvaluateCurveData(class UAnimInstance* Instance, float CurrentTime, float BlendWeight ) const;
 	/**
-	 * Find curve data based on name, you give pointer
+	 * Find curve data based on the curve UID
 	 */
-	ENGINE_API FFloatCurve * GetCurveData(FName InCurveName);
+	ENGINE_API FFloatCurve * GetCurveData(USkeleton::AnimCurveUID Uid);
 	/**
-	 * Add new curve named InCurveName and return true if success
+	 * Add new curve from the provided UID and return true if success
 	 * bVectorInterpCurve == true, then it will create FVectorCuve, otherwise, FFloatCurve
 	 */
-	ENGINE_API bool AddCurveData(FName InCurveName, int32 CurveFlags = ACF_DefaultCurve);
+	ENGINE_API bool AddCurveData(USkeleton::AnimCurveUID Uid, int32 CurveFlags = ACF_DefaultCurve);
 	/*
 	 * Delete curve data 
 	 */
-	ENGINE_API bool DeleteCurveData(FName InCurveName);
+	ENGINE_API bool DeleteCurveData(USkeleton::AnimCurveUID Uid);
 	/**
 	 * Duplicate curve data
 	 * 
 	 */
-	ENGINE_API bool DuplicateCurveData(FName CurveToCopyName, FName NewName);
+	ENGINE_API bool DuplicateCurveData(USkeleton::AnimCurveUID ToCopyUid, USkeleton::AnimCurveUID NewUid);
+
+	/**
+	 * Updates the LastObservedName field of the curves from the provided name container
+	 */
+	void UpdateLastObservedNames(FSmartNameMapping* NameMapping);
+
+	void Serialize(FArchive& Ar);
 };
 
 UENUM()
@@ -371,6 +415,9 @@ class UAnimSequenceBase : public UAnimationAsset
 protected:
 	virtual void ExtractRootTrack(float Pos, FTransform & RootTransform, const FBoneContainer * RequiredBones) const;
 
+public:
+	virtual void Serialize(FArchive& Ar) override;
+
 #if WITH_EDITOR
 private:
 	DECLARE_MULTICAST_DELEGATE( FOnNotifyChangedMulticaster );
diff --git a/Engine/Source/Runtime/Engine/Classes/Animation/Skeleton.h b/Engine/Source/Runtime/Engine/Classes/Animation/Skeleton.h
--- a/Engine/Source/Runtime/Engine/Classes/Animation/Skeleton.h
+++ b/Engine/Source/Runtime/Engine/Classes/Animation/Skeleton.h
@@ -7,6 +7,7 @@
 
 #pragma once
 #include "PreviewAssetAttachComponent.h"
+#include "SmartName.h"
 #include "Skeleton.generated.h"
 
 class UAnimSequence;
@@ -191,6 +192,16 @@ public:
 	/** Serializable retarget sources for this skeleton **/
 	TMap< FName, FReferencePose > AnimRetargetSources;
 
+	// Typedefs for greater smartname UID readability, add one for each smartname category 
+	typedef FSmartNameMapping::UID AnimCurveUID;
+
+	// Names for smartname mappings, if you're adding a new category of smartnames add a new name here
+	static ENGINE_API const FName AnimCurveMappingName;
+
+	// Container for smart name mappings
+	UPROPERTY()
+	FSmartNameContainer SmartNames;
+
 #if WITH_EDITORONLY_DATA
 private:
 	/** The default skeletal mesh to use when previewing this skeleton */
@@ -207,6 +218,17 @@ private:
 
 public:
 
+	// Adds a new name to the smart name container and modifies the skeleton so it can be saved
+	// return bool - Whether a name was added (false if already present)
+	ENGINE_API bool AddSmartnameAndModify(FName ContainerName, FName NewName, FSmartNameMapping::UID& NewUid);
+
+	// Renames a smartname in the specified container and modifies the skeleton
+	// return bool - Whether the rename was sucessful
+	ENGINE_API bool RenameSmartnameAndModify(FName ContainerName, FSmartNameMapping::UID Uid, FName NewName);
+
+	// Removes a smartname from the specified container and modifies the skeleton
+	ENGINE_API void RemoveSmartnameAndModify(FName ContainerName, FSmartNameMapping::UID Uid);
+
 	/** AnimNotifiers that has been created. Right now there is no delete step for this, but in the future we'll supply delete**/
 	UPROPERTY()
 	TArray<FName> AnimationNotifies;
@@ -547,8 +569,13 @@ public:
 	/** Removes the supplied bones from the skeleton */
 	ENGINE_API void RemoveBonesFromSkeleton(const TArray<FName>& BonesToRemove, bool bRemoveChildBones);
 
+	// Asset registry information for animation notifies
 	static const FName AnimNotifyTag;
-	static const TCHAR AnimNotifyTagDeliminator;
+	static const TCHAR AnimNotifyTagDelimiter;
+
+	// Asset registry information for animation curves
+	ENGINE_API static const FName CurveTag;
+	ENGINE_API static const TCHAR CurveTagDelimiter;
 #endif
 
 public:
diff --git a/Engine/Source/Runtime/Engine/Classes/Animation/SmartName.h b/Engine/Source/Runtime/Engine/Classes/Animation/SmartName.h
new file mode 100644
--- /dev/null
+++ b/Engine/Source/Runtime/Engine/Classes/Animation/SmartName.h
@@ -0,0 +1,97 @@
+// Copyright 1998-2014 Epic Games, Inc. All Rights Reserved.
+
+#pragma once
+
+#include "SmartName.generated.h"
+
+USTRUCT()
+struct ENGINE_API FSmartNameMapping
+{
+	GENERATED_USTRUCT_BODY();
+
+	FSmartNameMapping();
+
+	// ID type, should be used to access names as fundamental type may change.
+	// NOTE: If changing this, also change the type of NextUid below. UHT ignores,
+	// this typedef and needs to know the fundamental type to parse properties.
+	// If this changes and UHT handles typedefs, set NextUid type to UID.
+	typedef uint16 UID;
+	// Max UID used for overflow checking
+	static const UID MaxUID = MAX_uint16;
+
+	// Add a name to the mapping, if it exists, get it.
+	// @param Name - The name to add/get
+	// @param OUT OutUid - The UID of the newly created or retrieved name
+	// @return bool - true if the name was added, false if it existed (OutUid will be correctly set still in this case)
+	bool AddName(FName Name, UID& OutUid);
+
+	// Get a name from the mapping
+	// @param Uid - UID of the name to retrieve
+	// @param OUT OutName - Retrieved name
+	// @return bool - true if name existed and OutName is valid
+	bool GetName(const UID& Uid, FName& OutName) const;
+
+	// Fill an array with all used UIDs
+	// @param Array - Array to fill
+	void FillUidArray(TArray<UID>& Array) const;
+
+	// Fill an array with all used names
+	// @param Array - Array to fill
+	void FillNameArray(TArray<FName>& Array) const;
+
+	// Change a name
+	// @param Uid - UID of the name to change
+	// @param NewName - New name to set 
+	// @return bool - true if the name was found and changed, false if the name wasn't present in the mapping
+	bool Rename(const UID& Uid, FName NewName);
+
+	// Remove a name from the mapping
+	// @param Uid - UID of the name to remove
+	// @return bool - true if the name was found and removed, false if the name wasn't present in the mapping
+	bool Remove(const UID& Uid);
+
+	// Check whether a name already exists in the mapping
+	// @param Uid - the UID to check
+	// @return bool - whether the name was found
+	bool Exists(const UID& Uid);
+
+	// Check whether a name already exists in the mapping
+	// @param Name - the name to check
+	// @return bool - whether the name was found
+	bool Exists(const FName& Name);
+
+	// Get the number of names currently stored in this container
+	int32 GetNumNames();
+
+	// Serialize this to the provided archive; required for TMap serialization
+	void Serialize(FArchive& Ar);
+
+	friend FArchive& operator<<(FArchive& Ar, FSmartNameMapping& Elem);
+
+private:
+
+	UID NextUid;				// The next UID to use 
+	TMap<UID, FName> UidMap;	// Mapping of UIDs and names
+	TArray<UID> FreeList;		// Freelist of UIDs from removing/moving names around. Populated on load
+};
+
+USTRUCT()
+struct ENGINE_API FSmartNameContainer
+{
+	GENERATED_USTRUCT_BODY();
+
+	// Add a new smartname container with the provided name
+	void AddContainer(FName NewContainerName);
+
+	// Get a container by name
+	FSmartNameMapping* GetContainer(FName ContainerName);
+	const FSmartNameMapping* GetContainer(FName ContainerName) const;
+
+	// Serialize this to the provided archive; required for TMap serialization
+	void Serialize(FArchive& Ar);
+
+	friend FArchive& operator<<(FArchive& Ar, FSmartNameContainer& Elem);
+
+private:
+	TMap<FName, FSmartNameMapping> NameMappings;	// List of smartname mappings
+};
\ No newline at end of file
diff --git a/Engine/Source/Runtime/Engine/Private/Animation/AnimInstance.cpp b/Engine/Source/Runtime/Engine/Private/Animation/AnimInstance.cpp
--- a/Engine/Source/Runtime/Engine/Private/Animation/AnimInstance.cpp
+++ b/Engine/Source/Runtime/Engine/Private/Animation/AnimInstance.cpp
@@ -1060,6 +1060,18 @@ void UAnimInstance::AddCurveValue(const FName & CurveName, float Value, int32 Cu
 	}
 }
 
+void UAnimInstance::AddCurveValue(const USkeleton::AnimCurveUID Uid, float Value, int32 CurveTypeFlags)
+{
+	FName CurrentCurveName;
+	// Grab the smartname mapping from our current skeleton and resolve the curve name. We cannot cache
+	// the smart name mapping as the skeleton can change at any time.
+	if(FSmartNameMapping* NameMapping = CurrentSkeleton->SmartNames.GetContainer(USkeleton::AnimCurveMappingName))
+	{
+		NameMapping->GetName(Uid, CurrentCurveName);
+	}
+	AddCurveValue(CurrentCurveName, Value, CurveTypeFlags);
+}
+
 void UAnimInstance::TriggerAnimNotifies(float DeltaSeconds)
 {
 	USkeletalMeshComponent * SkelMeshComp = GetSkelMeshComponent();
diff --git a/Engine/Source/Runtime/Engine/Private/Animation/AnimSequence.cpp b/Engine/Source/Runtime/Engine/Private/Animation/AnimSequence.cpp
--- a/Engine/Source/Runtime/Engine/Private/Animation/AnimSequence.cpp
+++ b/Engine/Source/Runtime/Engine/Private/Animation/AnimSequence.cpp
@@ -1504,12 +1504,17 @@ void UAnimSequence::UpgradeMorphTargetCurves()
 		const FCurveTrack & CurveTrack = (*CurveIter);
 
 		{
+			// Add curves into naming system 
+			FSmartNameMapping* NameMapping = GetSkeleton()->SmartNames.GetContainer(USkeleton::AnimCurveMappingName);
+			check(NameMapping); // This name mapping should always be present
+			USkeleton::AnimCurveUID NewId;
+			NameMapping->AddName(CurveTrack.CurveName, NewId);
 
 			// add it first, it won't add if duplicate
-			RawCurveData.AddCurveData(CurveTrack.CurveName);
-
+			RawCurveData.AddCurveData(NewId);
+			
 			// get the last one I just added
-			FFloatCurve * NewCurve = RawCurveData.GetCurveData(CurveTrack.CurveName);
+			FFloatCurve * NewCurve = RawCurveData.GetCurveData(NewId);
 			check (NewCurve);
 
 			// since the import code change, this might not match all the time
diff --git a/Engine/Source/Runtime/Engine/Private/Animation/AnimSequenceBase.cpp b/Engine/Source/Runtime/Engine/Private/Animation/AnimSequenceBase.cpp
--- a/Engine/Source/Runtime/Engine/Private/Animation/AnimSequenceBase.cpp
+++ b/Engine/Source/Runtime/Engine/Private/Animation/AnimSequenceBase.cpp
@@ -81,6 +81,12 @@ void FFloatCurve::SetCurveTypeFlag(EAnimCurveFlags InFlag, bool bValue)
 	}
 }
 
+void FFloatCurve::ToggleCurveTypeFlag(EAnimCurveFlags InFlag)
+{
+	bool Current = GetCurveTypeFlag(InFlag);
+	SetCurveTypeFlag(InFlag, !Current);
+}
+
 bool FFloatCurve::GetCurveTypeFlag(EAnimCurveFlags InFlag) const
 {
 	return (CurveTypeFlags & InFlag) != 0;
@@ -106,16 +112,16 @@ void FRawCurveTracks::EvaluateCurveData(class UAnimInstance* Instance, float Cur
 	for (auto CurveIter = FloatCurves.CreateConstIterator(); CurveIter; ++CurveIter)
 	{
 		const FFloatCurve & Curve = *CurveIter;
-		Instance->AddCurveValue( Curve.CurveName, Curve.FloatCurve.Eval(CurrentTime)*BlendWeight, Curve.GetCurveTypeFlags() );
+
+		Instance->AddCurveValue( Curve.CurveUid, Curve.FloatCurve.Eval(CurrentTime)*BlendWeight, Curve.GetCurveTypeFlags() );
 	}
 }
 
-FFloatCurve * FRawCurveTracks::GetCurveData(FName InCurveName) 
+FFloatCurve * FRawCurveTracks::GetCurveData(USkeleton::AnimCurveUID Uid)
 {
-	for (auto CurveIter = FloatCurves.CreateIterator(); CurveIter; ++CurveIter)
+	for(FFloatCurve& Curve : FloatCurves)
 	{
-		FFloatCurve & Curve = *CurveIter;
-		if (Curve.CurveName == InCurveName)
+		if(Curve.CurveUid == Uid)
 		{
 			return &Curve;
 		}
@@ -124,13 +130,13 @@ FFloatCurve * FRawCurveTracks::GetCurveData(FName InCurveName)
 	return NULL;
 }
 
-bool FRawCurveTracks::DeleteCurveData(FName InCurveName)
+bool FRawCurveTracks::DeleteCurveData(USkeleton::AnimCurveUID Uid)
 {
-	for ( int32 Id=0; Id<FloatCurves.Num(); ++Id )
+	for( int32 Idx = 0; Idx < FloatCurves.Num(); ++Idx )
 	{
-		if (FloatCurves[Id].CurveName == InCurveName)
+		if( FloatCurves[Idx].CurveUid == Uid )
 		{
-			FloatCurves.RemoveAt(Id);
+			FloatCurves.RemoveAt(Idx);
 			return true;
 		}
 	}
@@ -138,25 +144,45 @@ bool FRawCurveTracks::DeleteCurveData(FName InCurveName)
 	return false;
 }
 
-bool FRawCurveTracks::AddCurveData(FName InCurveName, int32 CurveFlags /*= ACF_DefaultCurve*/)
+bool FRawCurveTracks::AddCurveData(USkeleton::AnimCurveUID Uid, int32 CurveFlags /*= ACF_DefaultCurve*/)
 {
-	// if we don't have same name used
-	if ( GetCurveData(InCurveName) == NULL )
+	if(GetCurveData(Uid) == NULL)
 	{
-		FloatCurves.Add( FFloatCurve(InCurveName, CurveFlags) );
+		FloatCurves.Add(FFloatCurve(Uid, CurveFlags));
 		return true;
 	}
-
 	return false;
 }
 
-ENGINE_API bool FRawCurveTracks::DuplicateCurveData(FName CurveToCopyName, FName NewName)
+void FRawCurveTracks::Serialize(FArchive& Ar)
+{
+	if(Ar.UE4Ver() >= VER_UE4_SKELETON_ADD_SMARTNAMES)
+	{
+		for(FFloatCurve& Curve : FloatCurves)
+		{
+			Curve.Serialize(Ar);
+		}
+	}
+}
+
+void FRawCurveTracks::UpdateLastObservedNames(FSmartNameMapping* NameMapping)
 {
-	FFloatCurve* ExistingCurve = GetCurveData(CurveToCopyName);
-	if(ExistingCurve != NULL && GetCurveData(NewName) == NULL)
+	if(NameMapping)
+	{
+		for(FFloatCurve& Curve : FloatCurves)
+		{
+			NameMapping->GetName(Curve.CurveUid, Curve.LastObservedName);
+		}
+	}
+}
+
+ENGINE_API bool FRawCurveTracks::DuplicateCurveData(USkeleton::AnimCurveUID ToCopyUid, USkeleton::AnimCurveUID NewUid)
+{
+	FFloatCurve* ExistingCurve = GetCurveData(ToCopyUid);
+	if(ExistingCurve && GetCurveData(NewUid) == NULL)
 	{
 		// Add the curve to the track and set its data to the existing curve
-		FloatCurves.Add(FFloatCurve(NewName, ExistingCurve->GetCurveTypeFlags()));
+		FloatCurves.Add(FFloatCurve(NewUid, ExistingCurve->GetCurveTypeFlags()));
 		FloatCurves.Last().FloatCurve = ExistingCurve->FloatCurve;
 		
 		return true;
@@ -217,6 +243,41 @@ void UAnimSequenceBase::PostLoad()
 #if WITH_EDITORONLY_DATA
 	InitializeNotifyTrack();
 	UpdateAnimNotifyTrackCache();
+
+	if(USkeleton* Skeleton = GetSkeleton())
+	{
+		// Get the name mapping object for curves
+		FSmartNameMapping* NameMapping = Skeleton->SmartNames.GetContainer(USkeleton::AnimCurveMappingName);
+		
+		// Fix up the existing curves to work with smartnames
+		if(GetLinkerUE4Version() < VER_UE4_SKELETON_ADD_SMARTNAMES)
+		{
+			for(FFloatCurve& Curve : RawCurveData.FloatCurves)
+			{
+				// Add the names of the curves into the smartname mapping and store off the curve uid which will be saved next time the sequence saves.
+				NameMapping->AddName(Curve.LastObservedName, Curve.CurveUid);
+			}
+		}
+		else
+		{
+			TArray<FFloatCurve*> UnlinkedCurves;
+			for(FFloatCurve& Curve : RawCurveData.FloatCurves)
+			{
+				if(!NameMapping->Exists(Curve.CurveUid))
+				{
+					// The skeleton doesn't know our name. Use the last observed name that was saved with the
+					// curve to create a new name. This can happen if a user saves an animation but not a skeleton
+					// either when updating the assets or editing the curves within.
+					UnlinkedCurves.Add(&Curve);
+				}
+			}
+
+			for(FFloatCurve* Curve : UnlinkedCurves)
+			{
+				NameMapping->AddName(Curve->LastObservedName, Curve->CurveUid);
+			}
+		}
+	}
 #endif
 }
 
@@ -503,7 +564,7 @@ void UAnimSequenceBase::GetAssetRegistryTags(TArray<FAssetRegistryTag>& OutTags)
 			// only add if not BP anim notify since they're handled separate
 			if(Iter->IsBlueprintNotify() == false)
 			{
-				NotifyList += FString::Printf(TEXT("%s%c"), *Iter->NotifyName.ToString(), USkeleton::AnimNotifyTagDeliminator);
+				NotifyList += FString::Printf(TEXT("%s%c"), *Iter->NotifyName.ToString(), USkeleton::AnimNotifyTagDelimiter);
 			}
 		}
 		
@@ -512,6 +573,18 @@ void UAnimSequenceBase::GetAssetRegistryTags(TArray<FAssetRegistryTag>& OutTags)
 			OutTags.Add(FAssetRegistryTag(USkeleton::AnimNotifyTag, NotifyList, FAssetRegistryTag::TT_Hidden));
 		}
 	}
+
+	// Add curve IDs to a tag list, or a blank tag if we have no curves.
+	// The blank list is necessary when we attempt to delete a curve so
+	// an old asset can be detected from its asset data so we load as few
+	// as possible.
+	FString CurveIdList;
+
+	for(const FFloatCurve& Curve : RawCurveData.FloatCurves)
+	{
+		CurveIdList += FString::Printf(TEXT("%u%c"), Curve.CurveUid, USkeleton::CurveTagDelimiter);
+	}
+	OutTags.Add(FAssetRegistryTag(USkeleton::CurveTag, CurveIdList, FAssetRegistryTag::TT_Hidden));
 }
 
 uint8* UAnimSequenceBase::FindNotifyPropertyData(int32 NotifyIndex, UArrayProperty*& ArrayProperty)
@@ -584,3 +657,18 @@ void UAnimSequenceBase::ExtractRootTrack(float Pos, FTransform & RootTransform,
 	RootTransform = FTransform::Identity;
 };
 
+void UAnimSequenceBase::Serialize(FArchive& Ar)
+{
+	Super::Serialize(Ar);
+
+	if(Ar.ArIsSaving && Ar.UE4Ver() >= VER_UE4_SKELETON_ADD_SMARTNAMES)
+	{
+		if(USkeleton* Skeleton = GetSkeleton())
+		{
+			FSmartNameMapping* Mapping = GetSkeleton()->SmartNames.GetContainer(USkeleton::AnimCurveMappingName);
+			check(Mapping); // Should always exist
+			RawCurveData.UpdateLastObservedNames(Mapping);
+		}
+	}
+	RawCurveData.Serialize(Ar);
+}
diff --git a/Engine/Source/Runtime/Engine/Private/Animation/Skeleton.cpp b/Engine/Source/Runtime/Engine/Private/Animation/Skeleton.cpp
--- a/Engine/Source/Runtime/Engine/Private/Animation/Skeleton.cpp
+++ b/Engine/Source/Runtime/Engine/Private/Animation/Skeleton.cpp
@@ -15,9 +15,15 @@
 
 #if WITH_EDITOR
 const FName USkeleton::AnimNotifyTag = FName(TEXT("AnimNotifyList"));
-const TCHAR USkeleton::AnimNotifyTagDeliminator = TEXT(';');
+const TCHAR USkeleton::AnimNotifyTagDelimiter = TEXT(';');
+
+const FName USkeleton::CurveTag = FName(TEXT("CurveUIDList"));
+const TCHAR USkeleton::CurveTagDelimiter = TEXT(';');
 #endif 
 
+// Names of smartname containers for skeleton properties
+const FName USkeleton::AnimCurveMappingName = FName(TEXT("AnimationCurves"));
+
 const FName USkeleton::DefaultSlotGroupName = FName(TEXT("Default"));
 
 FArchive& operator<<(FArchive& Ar, FReferencePose & P)
@@ -71,6 +77,9 @@ void USkeleton::PostLoad()
 #if WITH_EDITOR
 	if ( GIsEditor )
 	{
+		// Set up smart name mappings, if these have been loaded they will already be present
+		SmartNames.AddContainer(AnimCurveMappingName);
+
 		CollectAnimationNotifies();
 	}
 #endif
@@ -149,6 +158,12 @@ void USkeleton::Serialize( FArchive& Ar )
 	{
 		PreviewAttachedAssetContainer.SaveAttachedObjectsFromDeprecatedProperties();
 	}
+
+	// If we should be using smartnames, serialize the mappings
+	if(Ar.UE4Ver() >= VER_UE4_SKELETON_ADD_SMARTNAMES)
+	{
+		SmartNames.Serialize(Ar);
+	}
 #endif
 }
 
@@ -687,7 +702,7 @@ void USkeleton::CollectAnimationNotifies()
 			if (const FString * Value = Asset.TagsAndValues.Find(USkeleton::AnimNotifyTag))
 			{
 				TArray<FString> NotifyList;
-				Value->ParseIntoArray(&NotifyList, &AnimNotifyTagDeliminator, true);
+				Value->ParseIntoArray(&NotifyList, &AnimNotifyTagDelimiter, true);
 				for (auto NotifyIter = NotifyList.CreateConstIterator(); NotifyIter; ++NotifyIter)
 				{
 					FString NotifyName = *NotifyIter;
@@ -907,6 +922,52 @@ bool USkeleton::DoesHaveSlotGroupName(FName GroupName) const
 {
 	return SlotGroupNames.Contains(GroupName);
 }
+
+bool USkeleton::AddSmartnameAndModify(FName ContainerName, FName NewName, FSmartNameMapping::UID& NewUid)
+{
+	bool Successful = false;
+	FSmartNameMapping* RequestedMapping = SmartNames.GetContainer(ContainerName);
+	if(RequestedMapping)
+	{
+		if(RequestedMapping->AddName(NewName, NewUid))
+		{
+			Modify(true);
+			Successful = true;
+		}
+	}
+	return Successful;
+}
+
+bool USkeleton::RenameSmartnameAndModify(FName ContainerName, FSmartNameMapping::UID Uid, FName NewName)
+{
+	bool Successful = false;
+	FSmartNameMapping* RequestedMapping = SmartNames.GetContainer(ContainerName);
+	if(RequestedMapping)
+	{
+		FName CurrentName;
+		RequestedMapping->GetName(Uid, CurrentName);
+		if(CurrentName != NewName)
+		{
+			Modify();
+			Successful = RequestedMapping->Rename(Uid, NewName);
+		}
+	}
+	return Successful;
+}
+
+void USkeleton::RemoveSmartnameAndModify(FName ContainerName, FSmartNameMapping::UID Uid)
+{
+	FSmartNameMapping* RequestedMapping = SmartNames.GetContainer(ContainerName);
+	if(RequestedMapping)
+	{
+		if(RequestedMapping->Exists(Uid))
+		{
+			Modify();
+			RequestedMapping->Remove(Uid);
+		}
+	}
+}
+
 #endif // WITH_EDITORONLY_DATA
 
 void USkeleton::RegenerateGuid()
diff --git a/Engine/Source/Runtime/Engine/Private/Animation/SmartName.cpp b/Engine/Source/Runtime/Engine/Private/Animation/SmartName.cpp
new file mode 100644
--- /dev/null
+++ b/Engine/Source/Runtime/Engine/Private/Animation/SmartName.cpp
@@ -0,0 +1,182 @@
+// Copyright 1998-2014 Epic Games, Inc. All Rights Reserved.
+
+#include "EnginePrivate.h"
+
+FSmartNameMapping::FSmartNameMapping()
+: NextUid(0)
+{}
+
+bool FSmartNameMapping::AddName(FName Name, UID& OutUid)
+{
+	check(Name.IsValid());
+
+	// Check for UID overflow
+	if(NextUid == MaxUID && FreeList.Num() == 0)
+	{
+		// No UIDs left
+		UE_LOG(LogAnimation, Error, TEXT("No more UIDs left in skeleton smartname container, consider changing UID type to longer type."));
+		return false;
+	}
+	else
+	{
+		const UID* ExistingUid = UidMap.FindKey(Name);
+
+		if(ExistingUid)
+		{
+			// Already present in the list
+			OutUid = *ExistingUid;
+			return false;
+		}
+
+		if(FreeList.Num() > 0)
+		{
+			// We've got some UIDs that are unused < NextUid. Use them first
+			OutUid = FreeList.Pop();
+		}
+		else
+		{
+			OutUid = NextUid++;
+		}
+		UidMap.Add(OutUid, Name);
+
+		return true;
+	}
+}
+
+bool FSmartNameMapping::GetName(const UID& Uid, FName& OutName) const
+{
+	check(Uid < NextUid);
+	const FName* FoundName = UidMap.Find(Uid);
+	if(FoundName)
+	{
+		OutName = *FoundName;
+		return true;
+	}
+	return false;
+}
+
+bool FSmartNameMapping::Rename(const UID& Uid, FName NewName)
+{
+	check(Uid < NextUid);
+	FName* ExistingName = UidMap.Find(Uid);
+	if(ExistingName)
+	{
+		*ExistingName = NewName;
+		return true;
+	}
+	return false;
+}
+
+bool FSmartNameMapping::Remove(const UID& Uid)
+{
+	check(Uid < NextUid);
+	bool bRemoved = UidMap.Remove(Uid) != 0;
+
+	if(bRemoved)
+	{
+		// The provided UID has been cleared. If it is at the end of the mapping decrement next id,
+		// otherwise add to the freelist
+		if(Uid == NextUid - 1)
+		{
+			--NextUid;
+		}
+		else
+		{
+			FreeList.Add(Uid);
+		}
+	}
+
+	return bRemoved;
+}
+
+void FSmartNameMapping::Serialize(FArchive& Ar)
+{
+	if(Ar.UE4Ver() >= VER_UE4_SKELETON_ADD_SMARTNAMES)
+	{
+		Ar << NextUid;
+		Ar << UidMap;
+
+		if(Ar.ArIsLoading)
+		{
+			// Sort the Uid Map
+			UidMap.KeySort([](const UID& A, const UID& B)
+			{
+				return A < B;
+			});
+
+			// Build freelist as necessary
+			UID ExpectedID = 0;
+			UID Last = 0;
+			for(TPair<UID, FName>& Pair : UidMap)
+			{
+				UID CurrentUid = Pair.Key;
+				while(CurrentUid != ExpectedID)
+				{
+					FreeList.Add(ExpectedID++);
+				}
+
+				Last = CurrentUid;
+				++ExpectedID;
+			}
+
+			// In the case of multipled queued deletions previously,
+			// fix up the next UID to match the last used key we found.
+			NextUid = Last + 1;
+		}
+	}
+}
+
+int32 FSmartNameMapping::GetNumNames()
+{
+	return UidMap.Num();
+}
+
+void FSmartNameMapping::FillUidArray(TArray<UID>& Array) const
+{
+	UidMap.GenerateKeyArray(Array);
+}
+
+void FSmartNameMapping::FillNameArray(TArray<FName>& Array) const
+{
+	UidMap.GenerateValueArray(Array);
+}
+
+bool FSmartNameMapping::Exists(const UID& Uid)
+{
+	return UidMap.Find(Uid) != nullptr;
+}
+
+bool FSmartNameMapping::Exists(const FName& Name)
+{
+	return UidMap.FindKey(Name) != nullptr;
+}
+
+FArchive& operator<<(FArchive& Ar, FSmartNameMapping& Elem)
+{
+	Elem.Serialize(Ar);
+
+	return Ar;
+}
+
+void FSmartNameContainer::AddContainer(FName NewContainerName)
+{
+	if(!NameMappings.Find(NewContainerName))
+	{
+		NameMappings.Add(NewContainerName);
+	}
+}
+
+FSmartNameMapping* FSmartNameContainer::GetContainer(FName ContainerName)
+{
+	return NameMappings.Find(ContainerName);
+}
+
+const FSmartNameMapping* FSmartNameContainer::GetContainer(FName ContainerName) const
+{
+	return NameMappings.Find(ContainerName);
+}
+
+void FSmartNameContainer::Serialize(FArchive& Ar)
+{
+	Ar << NameMappings;
+}
