Fixed a build error.

diff --git a/Engine/Source/Editor/Persona/Private/SSkeletonSmartNameManager.cpp b/Engine/Source/Editor/Persona/Private/SSkeletonSmartNameManager.cpp
--- a/Engine/Source/Editor/Persona/Private/SSkeletonSmartNameManager.cpp
+++ b/Engine/Source/Editor/Persona/Private/SSkeletonSmartNameManager.cpp
@@ -261,12 +261,12 @@ void SSkeletonSmartNameManager::GenerateDisplayedList(const FText& FilterText)
 	}
 }
 
-void SSkeletonSmartNameManager::OnNameCommitted(const FText& NewName, ETextCommit::Type, TSharedPtr<FDisplayedSmartNameInfo> Item)
+void SSkeletonSmartNameManager::OnNameCommitted(const FText& InNewName, ETextCommit::Type, TSharedPtr<FDisplayedSmartNameInfo> Item)
 {
 	FSmartNameMapping* Mapping = CurrentSkeleton->SmartNames.GetContainer(ContainerName);
 	if(Mapping)
 	{
-		FName NewName(*NewName.ToString());
+		FName NewName(*InNewName.ToString());
 		if(!Mapping->Exists(NewName))
 		{
 			FScopedTransaction Transaction(LOCTEXT("TransactionRename", "Rename Element"));
