Altered notify customisations to allow outers of animation assets not just anim sequences.
#ttp 321617 - UE4: Persona: Discrepencies between assigning sockets in Montage/AnimSequence

diff --git a/Engine/Source/Editor/Persona/Private/AnimNotifyDetails.cpp b/Engine/Source/Editor/Persona/Private/AnimNotifyDetails.cpp
--- a/Engine/Source/Editor/Persona/Private/AnimNotifyDetails.cpp
+++ b/Engine/Source/Editor/Persona/Private/AnimNotifyDetails.cpp
@@ -56,9 +56,9 @@ void FAnimNotifyDetails::AddBoneNameProperty(IDetailLayoutBuilder& DetailBuilder
 		{
 			TWeakObjectPtr<UObject> SelectedObject = SelectedObjects[0];
 			const UObject* OuterObject = SelectedObject->GetOuter();
-			if (const UAnimSequence* AnimSequence = Cast<const UAnimSequence>(OuterObject))
+			if (const UAnimationAsset* AnimAsset = Cast<const UAnimationAsset>(OuterObject))
 			{
-				if (const USkeleton* Skeleton = AnimSequence->GetSkeleton())
+				if (const USkeleton* Skeleton = AnimAsset->GetSkeleton())
 				{
 					AnimNotifyCategory.AddProperty(BoneNameProperty.ToSharedRef())
 						.CustomWidget()
