Added marquee selection to notify panel and made fixes to marquee selection / dropping introduced wih recent sorting changes.

diff --git a/Engine/Source/Editor/Persona/Private/SAnimNotifyPanel.cpp b/Engine/Source/Editor/Persona/Private/SAnimNotifyPanel.cpp
--- a/Engine/Source/Editor/Persona/Private/SAnimNotifyPanel.cpp
+++ b/Engine/Source/Editor/Persona/Private/SAnimNotifyPanel.cpp
@@ -24,7 +24,7 @@ const FVector2D TextBorderSize(1.f, 1.f);
 #define LOCTEXT_NAMESPACE "AnimNotifyPanel"
 
 DECLARE_DELEGATE_OneParam( FOnDeleteNotify, struct FAnimNotifyEvent*)
-DECLARE_DELEGATE_RetVal_FourParams( FReply, FOnNotifyNodeDragStarted, TSharedRef<SAnimNotifyNode>, const FVector2D&, const FVector2D&, const bool)
+DECLARE_DELEGATE_RetVal_FourParams( FReply, FOnNotifyNodeDragStarted, TSharedRef<SAnimNotifyNode>, const FPointerEvent&, const FVector2D&, const bool)
 DECLARE_DELEGATE_RetVal_FiveParams(FReply, FOnNotifyNodesDragStarted, TArray<TSharedPtr<SAnimNotifyNode>>, TSharedRef<SWidget>, const FVector2D&, const FVector2D&, const bool)
 DECLARE_DELEGATE_RetVal( float, FOnGetDraggedNodePos )
 DECLARE_DELEGATE_TwoParams( FPanTrackRequest, int32, FVector2D)
@@ -276,6 +276,9 @@ public:
 	// Paste a single Notify into this track from an exported string
 	void PasteSingleNotify(FString& NotifyString, float PasteTime);
 
+	// Uses the given track space rect and marquee information to refresh selection information
+	void RefreshMarqueeSelectedNodes(FSlateRect& Rect, FNotifyMarqueeOperation& Marquee);
+
 protected:
 	void CreateCommands();
 
@@ -328,21 +331,24 @@ protected:
 	 * Selects a notify node in the graph. Supports multi selection
 	 * @param NotifyIndex - Index of the notify node to select.
 	 * @param Append - Whether to append to to current selection or start a new one.
+	 * @param bUpdateSelection - Whether to immediately inform Persona of a selection change
 	 */
-	void SelectNotifyNode(int32 NotifyIndex, bool Append);
+	void SelectNotifyNode(int32 NotifyIndex, bool Append, bool bUpdateSelection = true);
 	
 	/**
 	 * Toggles the selection status of a notify node, for example when
 	 * Control is held when clicking.
 	 * @param NotifyIndex - Index of the notify to toggle the selection status of
+	 * @param bUpdateSelection - Whether to immediately inform Persona of a selection change
 	 */
-	void ToggleNotifyNodeSelectionStatus(int32 NotifyIndex);
+	void ToggleNotifyNodeSelectionStatus(int32 NotifyIndex, bool bUpdateSelection = true);
 
 	/**
 	 * Deselects requested notify node
 	 * @param NotifyIndex - Index of the notify node to deselect
+	 * @param bUpdateSelection - Whether to immediately inform Persona of a selection change
 	 */
-	void DeselectNotifyNode(int32 NotifyIndex);
+	void DeselectNotifyNode(int32 NotifyIndex, bool bUpdateSelection = true);
 
 	int32 GetHitNotifyNode(const FGeometry& MyGeometry, const FVector2D& Position);
 
@@ -356,7 +362,7 @@ protected:
 	float CalculateTime( const FGeometry& MyGeometry, FVector2D NodePos, bool bInputIsAbsolute = true );
 
 	// Handler that is called when the user starts dragging a node
-	FReply OnNotifyNodeDragStarted( TSharedRef<SAnimNotifyNode> NotifyNode, const FVector2D& ScreenCursorPos, const FVector2D& ScreenNodePosition, const bool bDragOnMarker, int32 NotifyIndex );
+	FReply OnNotifyNodeDragStarted( TSharedRef<SAnimNotifyNode> NotifyNode, const FPointerEvent& MouseEvent, const FVector2D& ScreenNodePosition, const bool bDragOnMarker, int32 NotifyIndex );
 
 private:
 
@@ -515,6 +521,9 @@ public:
 				ClampInfo.NotifyTrack->HandleNodeDrop(Node, NodePositionOffset);
 				Node->DropCancelled();
 			}
+
+			Sequence->MarkPackageDirty();
+			OnUpdatePanel.ExecuteIfBound();
 		}
 		
 		FDragDropOperation::OnDrop(bDropWasHandled, MouseEvent);
@@ -744,12 +753,26 @@ public:
 	TSharedPtr<SWidget>					Decorator;				// The widget to display when dragging
 	float								SelectionTimeLength;	// Length of time that the selection covers
 	int32								TrackSpan;				// Number of tracks that the selection spans
-
-	static TSharedRef<FNotifyDragDropOp> New(TArray<TSharedPtr<SAnimNotifyNode>> NotifyNodes, TSharedPtr<SWidget> Decorator, const TArray<TSharedPtr<SAnimNotifyTrack>>& NotifyTracks, class UAnimSequenceBase* InSequence, const FVector2D &CursorPosition, const FVector2D &SelectionScreenPosition, const FVector2D &SelectionSize, float& CurrentDragXPosition, FPanTrackRequest& RequestTrackPanDelegate, TAttribute<TArray<FTrackMarkerBar>>	MarkerBars)
+	FOnUpdatePanel						OnUpdatePanel;			// Delegate to redraw the notify panel
+
+	static TSharedRef<FNotifyDragDropOp> New(
+		TArray<TSharedPtr<SAnimNotifyNode>>			NotifyNodes, 
+		TSharedPtr<SWidget>							Decorator, 
+		const TArray<TSharedPtr<SAnimNotifyTrack>>& NotifyTracks, 
+		class UAnimSequenceBase*					InSequence, 
+		const FVector2D&							CursorPosition, 
+		const FVector2D&							SelectionScreenPosition, 
+		const FVector2D&							SelectionSize, 
+		float&										CurrentDragXPosition, 
+		FPanTrackRequest&							RequestTrackPanDelegate, 
+		TAttribute<TArray<FTrackMarkerBar>>			MarkerBars,
+		FOnUpdatePanel&								UpdatePanel
+		)
 	{
 		TSharedRef<FNotifyDragDropOp> Operation = MakeShareable(new FNotifyDragDropOp(CurrentDragXPosition));
 		Operation->Sequence = InSequence;
 		Operation->RequestTrackPan = RequestTrackPanDelegate;
+		Operation->OnUpdatePanel = UpdatePanel;
 
 		Operation->NodeGroupPosition = SelectionScreenPosition;
 		Operation->NodeGroupSize = SelectionSize;
@@ -899,7 +922,7 @@ FReply SAnimNotifyNode::OnDragDetected( const FGeometry& MyGeometry, const FPoin
 		}
 	}
 
-	return OnNodeDragStarted.Execute(SharedThis(this), MouseEvent.GetScreenSpacePosition(), ScreenNodePosition, bDragOnMarker);
+	return OnNodeDragStarted.Execute(SharedThis(this), MouseEvent, ScreenNodePosition, bDragOnMarker);
 
 }
 
@@ -1908,7 +1931,7 @@ FReply SAnimNotifyTrack::OnMouseButtonUp( const FGeometry& MyGeometry, const FPo
 	return FReply::Unhandled();
 }
 
-void SAnimNotifyTrack::SelectNotifyNode(int32 NotifyIndex, bool Append)
+void SAnimNotifyTrack::SelectNotifyNode(int32 NotifyIndex, bool Append, bool bUpdateSelection)
 {
 	if( NotifyIndex != INDEX_NONE )
 	{
@@ -1927,13 +1950,17 @@ void SAnimNotifyTrack::SelectNotifyNode(int32 NotifyIndex, bool Append)
 				TSharedPtr<SAnimNotifyNode> Node = NotifyNodes[NotifyIndex];
 				Node->bSelected = true;
 				SelectedNotifyIndices.Add(NotifyIndex);
-				SendSelectionChanged();
+
+				if(bUpdateSelection)
+				{
+					SendSelectionChanged();
+				}
 			}
 		}
 	}
 }
 
-void SAnimNotifyTrack::ToggleNotifyNodeSelectionStatus( int32 NotifyIndex )
+void SAnimNotifyTrack::ToggleNotifyNodeSelectionStatus( int32 NotifyIndex, bool bUpdateSelection )
 {
 	check(NotifyNodes.IsValidIndex(NotifyIndex));
 
@@ -1949,10 +1976,14 @@ void SAnimNotifyTrack::ToggleNotifyNodeSelectionStatus( int32 NotifyIndex )
 
 	TSharedPtr<SAnimNotifyNode> Node = NotifyNodes[NotifyIndex];
 	Node->bSelected = !Node->bSelected;
-	SendSelectionChanged();
+
+	if(bUpdateSelection)
+	{
+		SendSelectionChanged();
+	}
 }
 
-void SAnimNotifyTrack::DeselectNotifyNode( int32 NotifyIndex )
+void SAnimNotifyTrack::DeselectNotifyNode( int32 NotifyIndex, bool bUpdateSelection )
 {
 	check(NotifyNodes.IsValidIndex(NotifyIndex));
 	TSharedPtr<SAnimNotifyNode> Node = NotifyNodes[NotifyIndex];
@@ -1961,7 +1992,10 @@ void SAnimNotifyTrack::DeselectNotifyNode( int32 NotifyIndex )
 	int32 ItemsRemoved = SelectedNotifyIndices.Remove(NotifyIndex);
 	check(ItemsRemoved > 0);
 
-	SendSelectionChanged();
+	if(bUpdateSelection)
+	{
+		SendSelectionChanged();
+	}
 }
 
 void SAnimNotifyTrack::DeselectAllNotifyNodes(bool bUpdateSelectionSet)
@@ -2485,12 +2519,12 @@ int32 SAnimNotifyTrack::GetHitNotifyNode(const FGeometry& MyGeometry, const FVec
 	return INDEX_NONE;
 }
 
-FReply SAnimNotifyTrack::OnNotifyNodeDragStarted( TSharedRef<SAnimNotifyNode> NotifyNode, const FVector2D& ScreenCursorPos, const FVector2D& ScreenNodePosition, const bool bDragOnMarker, int32 NotifyIndex ) 
+FReply SAnimNotifyTrack::OnNotifyNodeDragStarted( TSharedRef<SAnimNotifyNode> NotifyNode, const FPointerEvent& MouseEvent, const FVector2D& ScreenNodePosition, const bool bDragOnMarker, int32 NotifyIndex ) 
 {
 	// Check to see if we've already selected the triggering node
 	if (!NotifyNode->bSelected)
 	{
-		SelectNotifyNode(NotifyIndex, false);
+		SelectNotifyNode(NotifyIndex, MouseEvent.IsShiftDown());
 	}
 	
 	// Sort our nodes so we're acessing them in time order
@@ -2516,7 +2550,7 @@ FReply SAnimNotifyTrack::OnNotifyNodeDragStarted( TSharedRef<SAnimNotifyNode> No
 		
 		FVector2D DecoratorPosition = NodesToDrag[0]->GetWidgetPosition();
 		DecoratorPosition = CachedGeometry.LocalToAbsolute(DecoratorPosition);
-		return OnNodeDragStarted.Execute(NodesToDrag, DragBox, ScreenCursorPos, DecoratorPosition, bDragOnMarker);
+		return OnNodeDragStarted.Execute(NodesToDrag, DragBox, MouseEvent.GetScreenSpacePosition(), DecoratorPosition, bDragOnMarker);
 	}
 	else
 	{
@@ -2593,9 +2627,6 @@ void SAnimNotifyTrack::HandleNodeDrop(TSharedPtr<SAnimNotifyNode> Node, float Of
 	}
 
 	DroppedEvent->TrackIndex = TrackIndex;
-
-	Sequence->MarkPackageDirty();
-	OnUpdatePanel.ExecuteIfBound();
 }
 
 void SAnimNotifyTrack::DisconnectSelectedNodesForDrag(TArray<TSharedPtr<SAnimNotifyNode>>& DragNodes)
@@ -2711,6 +2742,49 @@ void SAnimNotifyTrack::AppendSelectedNodeWidgetsToArray(TArray<TSharedPtr<SAnimN
 	}
 }
 
+void SAnimNotifyTrack::RefreshMarqueeSelectedNodes(FSlateRect& Rect, FNotifyMarqueeOperation& Marquee)
+{
+	if(Marquee.Operation != FNotifyMarqueeOperation::Replace)
+	{
+		// Maintain the original selection from before the operation
+		for(int32 Idx = 0 ; Idx < NotifyNodes.Num() ; ++Idx)
+		{
+			TSharedPtr<SAnimNotifyNode> Notify = NotifyNodes[Idx];
+			bool bWasSelected = Marquee.OriginalSelection.Contains(Notify);
+			if(bWasSelected)
+			{
+				SelectNotifyNode(Idx, true, false);
+			}
+			else if(SelectedNotifyIndices.Contains(Idx))
+			{
+				DeselectNotifyNode(Idx, false);
+			}
+		}
+	}
+
+	for(int32 Index = 0 ; Index < NotifyNodes.Num() ; ++Index)
+	{
+		TSharedPtr<SAnimNotifyNode> Node = NotifyNodes[Index];
+		FSlateRect NodeRect = FSlateRect(Node->GetWidgetPosition(), Node->GetWidgetPosition() + Node->GetSize());
+
+		if(FSlateRect::DoRectanglesIntersect(Rect, NodeRect))
+		{
+			// Either select or deselect the intersecting node, depending on the type of selection operation
+			if(Marquee.Operation == FNotifyMarqueeOperation::Remove)
+			{
+				if(SelectedNotifyIndices.Contains(Index))
+				{
+					DeselectNotifyNode(Index, false);
+				}
+			}
+			else
+			{
+				SelectNotifyNode(Index, true, false);
+			}
+		}
+	}
+}
+
 //////////////////////////////////////////////////////////////////////////
 // SSequenceEdTrack
 
@@ -3106,7 +3180,8 @@ FReply SAnimNotifyPanel::OnNotifyNodeDragStarted(TArray<TSharedPtr<SAnimNotifyNo
 	}
 
 	FPanTrackRequest PanRequestDelegate = FPanTrackRequest::CreateSP(this, &SAnimNotifyPanel::PanInputViewRange);
-	return FReply::Handled().BeginDragDrop(FNotifyDragDropOp::New(Nodes, NodeDragDecorator, NotifyAnimTracks, Sequence, ScreenCursorPos, OverlayOrigin, OverlayExtents, CurrentDragXPosition, PanRequestDelegate, MarkerBars));
+	FOnUpdatePanel UpdateDelegate = FOnUpdatePanel::CreateSP(this, &SAnimNotifyPanel::Update);
+	return FReply::Handled().BeginDragDrop(FNotifyDragDropOp::New(Nodes, NodeDragDecorator, NotifyAnimTracks, Sequence, ScreenCursorPos, OverlayOrigin, OverlayExtents, CurrentDragXPosition, PanRequestDelegate, MarkerBars, UpdateDelegate));
 }
 
 void SAnimNotifyPanel::PostUndo()
@@ -3359,4 +3434,121 @@ FReply SAnimNotifyPanel::OnKeyDown(const FGeometry& MyGeometry, const FKeyboardE
 	return FReply::Unhandled();
 }
 
+FReply SAnimNotifyPanel::OnMouseButtonDown(const FGeometry& MyGeometry, const FPointerEvent& MouseEvent)
+{
+	SAnimTrackPanel::OnMouseButtonDown(MyGeometry, MouseEvent);
+
+	bool bLeftButton = MouseEvent.IsMouseButtonDown(EKeys::LeftMouseButton);
+
+	if(bLeftButton)
+	{
+		TArray<TSharedPtr<SAnimNotifyNode>> SelectedNodes;
+		for(TSharedPtr<SAnimNotifyTrack> Track : NotifyAnimTracks)
+		{
+			Track->AppendSelectedNodeWidgetsToArray(SelectedNodes);
+		}
+
+		Marquee.Start(MyGeometry.AbsoluteToLocal(MouseEvent.GetScreenSpacePosition()), Marquee.OperationTypeFromMouseEvent(MouseEvent), SelectedNodes);
+		if(Marquee.Operation == FNotifyMarqueeOperation::Replace)
+		{
+			// Remove and Add operations preserve selections, replace starts afresh
+			DeselectAllNotifies();
+		}
+
+		return FReply::Handled().DetectDrag(SharedThis(this), EKeys::LeftMouseButton);
+	}
+
+	return FReply::Unhandled();
+}
+
+FReply SAnimNotifyPanel::OnMouseButtonUp(const FGeometry& MyGeometry, const FPointerEvent& MouseEvent)
+{
+	if(Marquee.bActive)
+	{
+		OnTrackSelectionChanged();
+	}
+	Marquee = FNotifyMarqueeOperation();
+
+	return FReply::Handled().ReleaseMouseCapture();
+}
+
+FReply SAnimNotifyPanel::OnMouseMove(const FGeometry& MyGeometry, const FPointerEvent& MouseEvent)
+{
+	FReply BaseReply = SAnimTrackPanel::OnMouseMove(MyGeometry, MouseEvent);
+	if(!BaseReply.IsEventHandled())
+	{
+		bool bLeftButton = MouseEvent.IsMouseButtonDown(EKeys::LeftMouseButton);
+
+		if(bLeftButton && Marquee.bActive)
+		{
+			Marquee.Rect.UpdateEndPoint(MyGeometry.AbsoluteToLocal(MouseEvent.GetScreenSpacePosition()));
+			RefreshMarqueeSelectedNodes(MyGeometry);
+		}
+		return FReply::Handled();
+	}
+
+	return FReply::Unhandled();
+}
+
+int32 SAnimNotifyPanel::OnPaint(const FPaintArgs& Args, const FGeometry& AllottedGeometry, const FSlateRect& MyClippingRect, FSlateWindowElementList& OutDrawElements, int32 LayerId, const FWidgetStyle& InWidgetStyle, bool bParentEnabled) const
+{
+	LayerId = SAnimTrackPanel::OnPaint(Args, AllottedGeometry, MyClippingRect, OutDrawElements, LayerId, InWidgetStyle, bParentEnabled);
+
+	FVector2D Origin = AllottedGeometry.AbsoluteToLocal(Marquee.Rect.GetUpperLeft());
+	FVector2D Extents = AllottedGeometry.AbsoluteToLocal(Marquee.Rect.GetSize());
+
+	if(Marquee.IsValid())
+	{
+		FSlateDrawElement::MakeBox(
+			OutDrawElements,
+			LayerId++,
+			AllottedGeometry.ToPaintGeometry(Marquee.Rect.GetUpperLeft(), Marquee.Rect.GetSize()),
+			FEditorStyle::GetBrush(TEXT("MarqueeSelection")),
+			MyClippingRect
+			);
+	}
+
+	return LayerId;
+}
+
+void SAnimNotifyPanel::RefreshMarqueeSelectedNodes(const FGeometry& PanelGeo)
+{
+	if(Marquee.IsValid())
+	{
+		FSlateRect MarqueeRect = Marquee.Rect.ToSlateRect();
+		for(TSharedPtr<SAnimNotifyTrack> Track : NotifyAnimTracks)
+		{
+			if(Marquee.Operation == FNotifyMarqueeOperation::Replace || Marquee.OriginalSelection.Num() == 0)
+			{
+				Track->DeselectAllNotifyNodes(false);
+			}
+
+			const FGeometry& TrackGeo = Track->GetCachedGeometry();
+
+			FSlateRect TrackClip = TrackGeo.GetClippingRect();
+			FSlateRect PanelClip = PanelGeo.GetClippingRect();
+			FVector2D PanelSpaceOrigin = TrackClip.GetTopLeft() - PanelClip.GetTopLeft();
+			FVector2D TrackSpaceOrigin = MarqueeRect.GetTopLeft() - PanelSpaceOrigin;
+			FSlateRect MarqueeTrackSpace(TrackSpaceOrigin, TrackSpaceOrigin + MarqueeRect.GetSize());
+
+			Track->RefreshMarqueeSelectedNodes(MarqueeTrackSpace, Marquee);
+		}
+	}
+}
+
+FReply SAnimNotifyPanel::OnDragDetected(const FGeometry& MyGeometry, const FPointerEvent& MouseEvent)
+{
+	Marquee.bActive = true;
+	return FReply::Handled().CaptureMouse(SharedThis(this));
+}
+
+void SAnimNotifyPanel::OnKeyboardFocusLost(const FKeyboardFocusEvent& InKeyboardFocusEvent)
+{
+	if(Marquee.bActive)
+	{
+		OnTrackSelectionChanged();
+	}
+	Marquee = FNotifyMarqueeOperation();
+}
+
 #undef LOCTEXT_NAMESPACE
diff --git a/Engine/Source/Editor/Persona/Private/SAnimNotifyPanel.h b/Engine/Source/Editor/Persona/Private/SAnimNotifyPanel.h
--- a/Engine/Source/Editor/Persona/Private/SAnimNotifyPanel.h
+++ b/Engine/Source/Editor/Persona/Private/SAnimNotifyPanel.h
@@ -51,6 +51,70 @@ namespace ENotifyStateHandleHit
 	};
 }
 
+struct FNotifyMarqueeOperation
+{
+	FNotifyMarqueeOperation()
+		: Operation(Add)
+		, bActive(false)
+	{
+	}
+
+	enum Type
+	{
+		/** Holding down Ctrl removes nodes */
+		Remove,
+		/** Holding down Shift adds to the selection */
+		Add,
+		/** When nothing is pressed, marquee replaces selection */
+		Replace
+	} Operation;
+
+	bool IsValid() const
+	{
+		return Rect.IsValid() && bActive;
+	}
+
+	void Start(const FVector2D& InStartLocation, FNotifyMarqueeOperation::Type InOperationType, TArray<TSharedPtr<SAnimNotifyNode>>& InOriginalSelection)
+	{
+		Rect = FMarqueeRect(InStartLocation);
+		Operation = InOperationType;
+		OriginalSelection = InOriginalSelection;
+	}
+
+	void End()
+	{
+		Rect = FMarqueeRect();
+	}
+
+
+	/** Given a mouse event, figure out what the marquee selection should do based on the state of Shift and Ctrl keys */
+	static FNotifyMarqueeOperation::Type OperationTypeFromMouseEvent(const FPointerEvent& MouseEvent)
+	{
+		if(MouseEvent.IsControlDown())
+		{
+			return FNotifyMarqueeOperation::Remove;
+		}
+		else if(MouseEvent.IsShiftDown())
+		{
+			return FNotifyMarqueeOperation::Add;
+		}
+		else
+		{
+			return FNotifyMarqueeOperation::Replace;
+		}
+	}
+
+public:
+	/** The marquee rectangle being dragged by the user */
+	FMarqueeRect Rect;
+
+	/** Whether the marquee has been activated, usually by a drag */
+	bool bActive;
+
+	/** The original selection state before the marquee selection */
+	TArray<TSharedPtr<SAnimNotifyNode>> OriginalSelection;
+};
+
 //////////////////////////////////////////////////////////////////////////
 // SAnimNotifyPanel
 
@@ -129,8 +193,17 @@ public:
 	FCoreDelegates::FOnObjectPropertyChanged::FDelegate OnPropertyChangedHandle;
 	void OnPropertyChanged(UObject* ChangedObject, FPropertyChangedEvent& PropertyEvent);
 
-	/** Handler for key press events */
-	virtual FReply OnKeyDown(const FGeometry& MyGeometry, const FKeyboardEvent& InKeyboardEvent);
+	/** SWidget Interface */
+	virtual FReply OnKeyDown(const FGeometry& MyGeometry, const FKeyboardEvent& InKeyboardEvent);	
+	virtual FReply OnMouseButtonDown(const FGeometry& MyGeometry, const FPointerEvent& MouseEvent) override;
+	virtual FReply OnMouseButtonUp(const FGeometry& MyGeometry, const FPointerEvent& MouseEvent) override;
+	virtual FReply OnMouseMove(const FGeometry& MyGeometry, const FPointerEvent& MouseEvent) override;
+	virtual FReply OnDragDetected(const FGeometry& MyGeometry, const FPointerEvent& MouseEvent) override;
+	virtual void OnKeyboardFocusLost(const FKeyboardFocusEvent& InKeyboardFocusEvent);
+	virtual int32 OnPaint(const FPaintArgs& Args, const FGeometry& AllottedGeometry, const FSlateRect& MyClippingRect, FSlateWindowElementList& OutDrawElements, int32 LayerId, const FWidgetStyle& InWidgetStyle, bool bParentEnabled) const override;
+	/** End SWidget Interface */
+
+	void RefreshMarqueeSelectedNodes(const FGeometry& PanelGeo);
 
 private:
 	TSharedPtr<SBorder> PanelArea;
@@ -139,6 +212,9 @@ private:
 	TAttribute<float> CurrentPosition;
 	FOnSelectionChanged OnSelectionChanged;
 	FOnGetScrubValue OnGetScrubValue;
+	
+	/** Manager for mouse controlled marquee selection */
+	FNotifyMarqueeOperation Marquee;
 
 	/** Delegate to request a refresh of the offsets calculated for notifies */
 	FRefreshOffsetsRequest OnRequestRefreshOffsets;
