Fix additional meshes display
#TTP 343406 - PERSONA: "Additional Body Parts" are not listed when added to the mesh preview

diff --git a/Engine/Source/Editor/Persona/Private/MeshMode/SAdditionalMeshesEditor.cpp b/Engine/Source/Editor/Persona/Private/MeshMode/SAdditionalMeshesEditor.cpp
--- a/Engine/Source/Editor/Persona/Private/MeshMode/SAdditionalMeshesEditor.cpp
+++ b/Engine/Source/Editor/Persona/Private/MeshMode/SAdditionalMeshesEditor.cpp
@@ -131,6 +131,8 @@ private:
 
 void SAdditionalMeshesDisplayPanel::Construct(const FArguments& InArgs)
 {
+	SWrapBox::Construct(SWrapBox::FArguments().UseAllottedWidth(true));
+
 	OnRemoveAdditionalMesh = InArgs._OnRemoveAdditionalMesh;
 }
 
