Moved out AnimatinoRecorder feature, so eventually can be used for everywhere. We have to move the header to public if we really want other module to use it.

diff --git a/Engine/Source/Editor/Persona/Private/Persona.cpp b/Engine/Source/Editor/Persona/Private/Persona.cpp
--- a/Engine/Source/Editor/Persona/Private/Persona.cpp
+++ b/Engine/Source/Editor/Persona/Private/Persona.cpp
@@ -2780,188 +2780,6 @@ void FPersona::ValidatePreviewAttachedAssets(USkeletalMesh* PreviewSkeletalMesh)
 	}
 }
 
-/////////////////////////////////////////////////////
-
-FAnimationRecorder::FAnimationRecorder()
-	: SampleRate(30)
-	, AnimationObject(NULL)
-{
-
-}
-
-FAnimationRecorder::~FAnimationRecorder()
-{
-	StopRecord();
-}
-
-void FAnimationRecorder::StartRecord(USkeletalMeshComponent * Component, UAnimSequence * InAnimationObject, float InDuration)
-{
-	Duration = InDuration;
-	check (Duration > 0.f);
-	TimePassed = 0.f;
-	AnimationObject = InAnimationObject;
-
-	AnimationObject->RawAnimationData.Empty();
-	AnimationObject->TrackToSkeletonMapTable.Empty();
-	AnimationObject->AnimationTrackNames.Empty();
-
-	PreviousSpacesBases = Component->SpaceBases;
-
-	check (SampleRate!=0.f);
-	const float IntervalTime = 1.f/SampleRate;
-	// prepare to record
-	// how many frames we need?
-	int32 MaxFrame = Duration/IntervalTime+1;
-	int32 NumBones = PreviousSpacesBases.Num();
-	check (MaxFrame > 0);
-	check (Component->SkeletalMesh);
-
-	LastFrame = 0;
-	AnimationObject->SequenceLength = Duration;
-	AnimationObject->NumFrames = MaxFrame;
-
-	USkeleton * AnimSkeleton = AnimationObject->GetSkeleton();
-	// add all frames
-	for (int32 BoneIndex=0; BoneIndex <PreviousSpacesBases.Num(); ++BoneIndex)
-	{
-		// verify if this bone exists in skeleton
-		int32 BoneTreeIndex = AnimSkeleton->GetSkeletonBoneIndexFromMeshBoneIndex(Component->SkeletalMesh, BoneIndex);
-		if (BoneTreeIndex != INDEX_NONE)
-		{
-			FName BoneTreeName = AnimSkeleton->GetReferenceSkeleton().GetBoneName(BoneTreeIndex);
-			int32 NewTrack = AnimationObject->RawAnimationData.AddZeroed(1);
-			AnimationObject->AnimationTrackNames.Add(BoneTreeName);
-
-			FRawAnimSequenceTrack& RawTrack = AnimationObject->RawAnimationData[NewTrack];
-			// add mapping to skeleton bone track
-			AnimationObject->TrackToSkeletonMapTable.Add(FTrackToSkeletonMap(BoneTreeIndex));
-			RawTrack.PosKeys.AddZeroed(MaxFrame);
-			RawTrack.RotKeys.AddZeroed(MaxFrame);
-			RawTrack.ScaleKeys.AddZeroed(MaxFrame);
-		}
-	}
-
-	// record the first frame;
-	Record(Component, PreviousSpacesBases, 0);
-}
-
-void FAnimationRecorder::StopRecord()
-{
-	if (AnimationObject)
-	{
-		// if LastFrame doesn't match with NumFrames, we need to adjust data
-		if ( LastFrame + 1 < AnimationObject->NumFrames )
-		{
-			int32 NewMaxFrame = LastFrame + 1;
-			int32 OldMaxFrame = AnimationObject->NumFrames;
-			// Set New Frames/Time
-			AnimationObject->NumFrames = NewMaxFrame;
-			AnimationObject->SequenceLength = (float)(NewMaxFrame)/SampleRate;
-
-			for (int32 TrackId=0; TrackId <AnimationObject->RawAnimationData.Num(); ++TrackId)
-			{
-				FRawAnimSequenceTrack& RawTrack = AnimationObject->RawAnimationData[TrackId];
-				// add mapping to skeleton bone track
-				RawTrack.PosKeys.RemoveAt(NewMaxFrame, OldMaxFrame-NewMaxFrame);
-				RawTrack.RotKeys.RemoveAt(NewMaxFrame, OldMaxFrame-NewMaxFrame);
-				RawTrack.ScaleKeys.RemoveAt(NewMaxFrame, OldMaxFrame-NewMaxFrame);
-			}
-		}
-
-		AnimationObject->PostProcessSequence();
-		AnimationObject->MarkPackageDirty();
-		AnimationObject = NULL;
-
-		PreviousSpacesBases.Empty();
-	}
-}
-
-// return false if nothing to update
-// return true if it has properly updated
-bool FAnimationRecorder::UpdateRecord(USkeletalMeshComponent * Component, float DeltaTime)
-{
-	const float MaxDelta = 1.f/10.f;
-	float UseDelta = FMath::Min<float>(DeltaTime, MaxDelta);
-
-	float PreviousTimePassed = TimePassed;
-	TimePassed += UseDelta;
-
-	check (SampleRate!=0.f);
-	const float IntervalTime = 1.f/SampleRate;
-
-	// time passed has been updated
-	// now find what frames we need to update
-	int32 MaxFrame = Duration/IntervalTime;
-	int32 FramesRecorded = (PreviousTimePassed / IntervalTime);
-	int32 FramesToRecord = FMath::Min<int32>(TimePassed / IntervalTime, MaxFrame);
-
-	TArray<FTransform> & SpaceBases = Component->SpaceBases;
-
-	check (SpaceBases.Num() == PreviousSpacesBases.Num());
-
-	TArray<FTransform> BlendedSpaceBases;
-	BlendedSpaceBases.AddZeroed(SpaceBases.Num());
-
-	// if we need to record frame
-	while (FramesToRecord > FramesRecorded)
-	{
-		// find what frames we need to record
-		// convert to time
-		float NewTime = (FramesRecorded + 1) * IntervalTime;
-
-		float PercentageOfNew = (NewTime - PreviousTimePassed)/UseDelta;
-		//float PercentageOfOld = (TimePassed-NewTime)/UseDelta;
-
-		// make sure this is almost 1 = PercentageOfNew
-		//ensure ( FMath::IsNearlyEqual(PercentageOfNew + PercentageOfOld, 1.f, KINDA_SMALL_NUMBER));
-
-		// for now we just concern component space, not skeleton space
-		for (int32 BoneIndex=0; BoneIndex<SpaceBases.Num(); ++BoneIndex)
-		{
-			BlendedSpaceBases[BoneIndex].Blend(PreviousSpacesBases[BoneIndex], SpaceBases[BoneIndex], PercentageOfNew);
-		}
-
-		Record(Component, BlendedSpaceBases, FramesRecorded + 1);
-		++FramesRecorded;
-	}
-
-	//save to current transform
-	PreviousSpacesBases = Component->SpaceBases;
-
-	// if we reached MaxFrame, return false;
-	return (MaxFrame > FramesRecorded);
-}
-
-void FAnimationRecorder::Record( USkeletalMeshComponent * Component, TArray<FTransform> SpacesBases, int32 FrameToAdd )
-{
-	if (ensure (AnimationObject))
-	{
-		USkeleton * AnimSkeleton = AnimationObject->GetSkeleton();
-		for (int32 TrackIndex=0; TrackIndex <AnimationObject->RawAnimationData.Num(); ++TrackIndex)
-		{
-			// verify if this bone exists in skeleton
-			int32 BoneTreeIndex = AnimationObject->GetSkeletonIndexFromTrackIndex(TrackIndex);
-			if (BoneTreeIndex != INDEX_NONE)
-			{
-				int32 BoneIndex = AnimSkeleton->GetMeshBoneIndexFromSkeletonBoneIndex(Component->SkeletalMesh, BoneTreeIndex);
-				int32 ParentIndex = Component->SkeletalMesh->RefSkeleton.GetParentIndex(BoneIndex);
-				FTransform LocalTransform = SpacesBases[BoneIndex];
-				if ( ParentIndex != INDEX_NONE )
-				{
-					LocalTransform.SetToRelativeTransform(SpacesBases[ParentIndex]);
-				}
-
-				FRawAnimSequenceTrack& RawTrack = AnimationObject->RawAnimationData[TrackIndex];
-				RawTrack.PosKeys[FrameToAdd] = LocalTransform.GetTranslation();
-				RawTrack.RotKeys[FrameToAdd] = LocalTransform.GetRotation();
-				RawTrack.ScaleKeys[FrameToAdd] = LocalTransform.GetScale3D();
-			}
-		}
-
-		LastFrame = FrameToAdd;
-	}
-}
-
 void FPersona::OnAnimNotifyWindow()
 {
 	TabManager->InvokeTab(FPersonaTabs::SkeletonAnimNotifiesID);
diff --git a/Engine/Source/Editor/Persona/Private/Persona.h b/Engine/Source/Editor/Persona/Private/Persona.h
--- a/Engine/Source/Editor/Persona/Private/Persona.h
+++ b/Engine/Source/Editor/Persona/Private/Persona.h
@@ -9,35 +9,10 @@
 #include "PersonaModule.h"
 #include "PersonaCommands.h"
 
+#include "Shared/AnimationRecorder.h"
 //////////////////////////////////////////////////////////////////////////
 // FPersona
 
-// records the mesh pose to animation input
-struct FAnimationRecorder
-{
-private:
-	int32 SampleRate;
-	float Duration;
-	int32 LastFrame;
-	float TimePassed;
-	UAnimSequence * AnimationObject;
-	TArray<FTransform> PreviousSpacesBases;
-	
-public:
-	FAnimationRecorder();
-	~FAnimationRecorder();
-
-	void StartRecord(USkeletalMeshComponent * Component, UAnimSequence * InAnimationObject, float InDuration);
-	void StopRecord();
-	// return false if nothing to update
-	// return true if it has properly updated
-	bool UpdateRecord(USkeletalMeshComponent * Component, float DeltaTime);
-	const UAnimSequence * GetAnimationObject() const { return AnimationObject; }
-
-private:
-	void Record( USkeletalMeshComponent * Component, TArray<FTransform> SpacesBases, int32 FrameToAdd );
-};
-
 /**
  * Persona asset editor (extends Blueprint editor)
  */
diff --git a/Engine/Source/Editor/Persona/Private/Shared/AnimationRecorder.cpp b/Engine/Source/Editor/Persona/Private/Shared/AnimationRecorder.cpp
new file mode 100644
--- /dev/null
+++ b/Engine/Source/Editor/Persona/Private/Shared/AnimationRecorder.cpp
@@ -0,0 +1,191 @@
+// Copyright 1998-2014 Epic Games, Inc. All Rights Reserved.
+
+
+#include "PersonaPrivatePCH.h"
+
+#define LOCTEXT_NAMESPACE "FAnimationRecorder"
+
+/////////////////////////////////////////////////////
+
+FAnimationRecorder::FAnimationRecorder()
+	: SampleRate(30)
+	, AnimationObject(NULL)
+{
+
+}
+
+FAnimationRecorder::~FAnimationRecorder()
+{
+	StopRecord();
+}
+
+void FAnimationRecorder::StartRecord(USkeletalMeshComponent * Component, UAnimSequence * InAnimationObject, float InDuration)
+{
+	Duration = InDuration;
+	check (Duration > 0.f);
+	TimePassed = 0.f;
+	AnimationObject = InAnimationObject;
+
+	AnimationObject->RawAnimationData.Empty();
+	AnimationObject->TrackToSkeletonMapTable.Empty();
+	AnimationObject->AnimationTrackNames.Empty();
+
+	PreviousSpacesBases = Component->SpaceBases;
+
+	check (SampleRate!=0.f);
+	const float IntervalTime = 1.f/SampleRate;
+	// prepare to record
+	// how many frames we need?
+	int32 MaxFrame = Duration/IntervalTime+1;
+	int32 NumBones = PreviousSpacesBases.Num();
+	check (MaxFrame > 0);
+	check (Component->SkeletalMesh);
+
+	LastFrame = 0;
+	AnimationObject->SequenceLength = Duration;
+	AnimationObject->NumFrames = MaxFrame;
+
+	USkeleton * AnimSkeleton = AnimationObject->GetSkeleton();
+	// add all frames
+	for (int32 BoneIndex=0; BoneIndex <PreviousSpacesBases.Num(); ++BoneIndex)
+	{
+		// verify if this bone exists in skeleton
+		int32 BoneTreeIndex = AnimSkeleton->GetSkeletonBoneIndexFromMeshBoneIndex(Component->SkeletalMesh, BoneIndex);
+		if (BoneTreeIndex != INDEX_NONE)
+		{
+			FName BoneTreeName = AnimSkeleton->GetReferenceSkeleton().GetBoneName(BoneTreeIndex);
+			int32 NewTrack = AnimationObject->RawAnimationData.AddZeroed(1);
+			AnimationObject->AnimationTrackNames.Add(BoneTreeName);
+
+			FRawAnimSequenceTrack& RawTrack = AnimationObject->RawAnimationData[NewTrack];
+			// add mapping to skeleton bone track
+			AnimationObject->TrackToSkeletonMapTable.Add(FTrackToSkeletonMap(BoneTreeIndex));
+			RawTrack.PosKeys.AddZeroed(MaxFrame);
+			RawTrack.RotKeys.AddZeroed(MaxFrame);
+			RawTrack.ScaleKeys.AddZeroed(MaxFrame);
+		}
+	}
+
+	// record the first frame;
+	Record(Component, PreviousSpacesBases, 0);
+}
+
+void FAnimationRecorder::StopRecord()
+{
+	if (AnimationObject)
+	{
+		// if LastFrame doesn't match with NumFrames, we need to adjust data
+		if ( LastFrame + 1 < AnimationObject->NumFrames )
+		{
+			int32 NewMaxFrame = LastFrame + 1;
+			int32 OldMaxFrame = AnimationObject->NumFrames;
+			// Set New Frames/Time
+			AnimationObject->NumFrames = NewMaxFrame;
+			AnimationObject->SequenceLength = (float)(NewMaxFrame)/SampleRate;
+
+			for (int32 TrackId=0; TrackId <AnimationObject->RawAnimationData.Num(); ++TrackId)
+			{
+				FRawAnimSequenceTrack& RawTrack = AnimationObject->RawAnimationData[TrackId];
+				// add mapping to skeleton bone track
+				RawTrack.PosKeys.RemoveAt(NewMaxFrame, OldMaxFrame-NewMaxFrame);
+				RawTrack.RotKeys.RemoveAt(NewMaxFrame, OldMaxFrame-NewMaxFrame);
+				RawTrack.ScaleKeys.RemoveAt(NewMaxFrame, OldMaxFrame-NewMaxFrame);
+			}
+		}
+
+		AnimationObject->PostProcessSequence();
+		AnimationObject->MarkPackageDirty();
+		AnimationObject = NULL;
+
+		PreviousSpacesBases.Empty();
+	}
+}
+
+// return false if nothing to update
+// return true if it has properly updated
+bool FAnimationRecorder::UpdateRecord(USkeletalMeshComponent * Component, float DeltaTime)
+{
+	const float MaxDelta = 1.f/10.f;
+	float UseDelta = FMath::Min<float>(DeltaTime, MaxDelta);
+
+	float PreviousTimePassed = TimePassed;
+	TimePassed += UseDelta;
+
+	check (SampleRate!=0.f);
+	const float IntervalTime = 1.f/SampleRate;
+
+	// time passed has been updated
+	// now find what frames we need to update
+	int32 MaxFrame = Duration/IntervalTime;
+	int32 FramesRecorded = (PreviousTimePassed / IntervalTime);
+	int32 FramesToRecord = FMath::Min<int32>(TimePassed / IntervalTime, MaxFrame);
+
+	TArray<FTransform> & SpaceBases = Component->SpaceBases;
+
+	check (SpaceBases.Num() == PreviousSpacesBases.Num());
+
+	TArray<FTransform> BlendedSpaceBases;
+	BlendedSpaceBases.AddZeroed(SpaceBases.Num());
+
+	// if we need to record frame
+	while (FramesToRecord > FramesRecorded)
+	{
+		// find what frames we need to record
+		// convert to time
+		float NewTime = (FramesRecorded + 1) * IntervalTime;
+
+		float PercentageOfNew = (NewTime - PreviousTimePassed)/UseDelta;
+		//float PercentageOfOld = (TimePassed-NewTime)/UseDelta;
+
+		// make sure this is almost 1 = PercentageOfNew
+		//ensure ( FMath::IsNearlyEqual(PercentageOfNew + PercentageOfOld, 1.f, KINDA_SMALL_NUMBER));
+
+		// for now we just concern component space, not skeleton space
+		for (int32 BoneIndex=0; BoneIndex<SpaceBases.Num(); ++BoneIndex)
+		{
+			BlendedSpaceBases[BoneIndex].Blend(PreviousSpacesBases[BoneIndex], SpaceBases[BoneIndex], PercentageOfNew);
+		}
+
+		Record(Component, BlendedSpaceBases, FramesRecorded + 1);
+		++FramesRecorded;
+	}
+
+	//save to current transform
+	PreviousSpacesBases = Component->SpaceBases;
+
+	// if we reached MaxFrame, return false;
+	return (MaxFrame > FramesRecorded);
+}
+
+void FAnimationRecorder::Record( USkeletalMeshComponent * Component, TArray<FTransform> SpacesBases, int32 FrameToAdd )
+{
+	if (ensure (AnimationObject))
+	{
+		USkeleton * AnimSkeleton = AnimationObject->GetSkeleton();
+		for (int32 TrackIndex=0; TrackIndex <AnimationObject->RawAnimationData.Num(); ++TrackIndex)
+		{
+			// verify if this bone exists in skeleton
+			int32 BoneTreeIndex = AnimationObject->GetSkeletonIndexFromTrackIndex(TrackIndex);
+			if (BoneTreeIndex != INDEX_NONE)
+			{
+				int32 BoneIndex = AnimSkeleton->GetMeshBoneIndexFromSkeletonBoneIndex(Component->SkeletalMesh, BoneTreeIndex);
+				int32 ParentIndex = Component->SkeletalMesh->RefSkeleton.GetParentIndex(BoneIndex);
+				FTransform LocalTransform = SpacesBases[BoneIndex];
+				if ( ParentIndex != INDEX_NONE )
+				{
+					LocalTransform.SetToRelativeTransform(SpacesBases[ParentIndex]);
+				}
+
+				FRawAnimSequenceTrack& RawTrack = AnimationObject->RawAnimationData[TrackIndex];
+				RawTrack.PosKeys[FrameToAdd] = LocalTransform.GetTranslation();
+				RawTrack.RotKeys[FrameToAdd] = LocalTransform.GetRotation();
+				RawTrack.ScaleKeys[FrameToAdd] = LocalTransform.GetScale3D();
+			}
+		}
+
+		LastFrame = FrameToAdd;
+	}
+}
+
+#undef LOCTEXT_NAMESPACE
+
diff --git a/Engine/Source/Editor/Persona/Private/Shared/AnimationRecorder.h b/Engine/Source/Editor/Persona/Private/Shared/AnimationRecorder.h
new file mode 100644
--- /dev/null
+++ b/Engine/Source/Editor/Persona/Private/Shared/AnimationRecorder.h
@@ -0,0 +1,31 @@
+// Copyright 1998-2014 Epic Games, Inc. All Rights Reserved.
+
+//////////////////////////////////////////////////////////////////////////
+// FAnimationRecorder
+
+// records the mesh pose to animation input
+struct FAnimationRecorder
+{
+private:
+	int32 SampleRate;
+	float Duration;
+	int32 LastFrame;
+	float TimePassed;
+	UAnimSequence * AnimationObject;
+	TArray<FTransform> PreviousSpacesBases;
+	
+public:
+	FAnimationRecorder();
+	~FAnimationRecorder();
+
+	void StartRecord(USkeletalMeshComponent * Component, UAnimSequence * InAnimationObject, float InDuration);
+	void StopRecord();
+	// return false if nothing to update
+	// return true if it has properly updated
+	bool UpdateRecord(USkeletalMeshComponent * Component, float DeltaTime);
+	const UAnimSequence * GetAnimationObject() const { return AnimationObject; }
+
+private:
+	void Record( USkeletalMeshComponent * Component, TArray<FTransform> SpacesBases, int32 FrameToAdd );
+};
+
