ANIMGRAPH: SLOTNODE: Add refresh option to refresh correctly.

diff --git a/Engine/Source/Editor/Persona/Private/Customization/AnimGraphNodeSlotDetails.cpp b/Engine/Source/Editor/Persona/Private/Customization/AnimGraphNodeSlotDetails.cpp
--- a/Engine/Source/Editor/Persona/Private/Customization/AnimGraphNodeSlotDetails.cpp
+++ b/Engine/Source/Editor/Persona/Private/Customization/AnimGraphNodeSlotDetails.cpp
@@ -198,6 +198,11 @@ int32 FAnimGraphNodeSlotDetails::RefreshSlotNames()
 		}
 	}
 
+	if (SlotNameComboBox.IsValid())
+	{
+		SlotNameComboBox->RefreshOptions();
+	}
+
 	// if SlotNameComboList is empty, just add a temporary message to avoid a combo box problem which doesn't show the list properly if empty
 	if (SlotNameComboList.Num() == 0)
 	{
