CIS Fix

diff --git a/Engine/Source/Editor/Persona/Private/Shared/AnimationRecorder.h b/Engine/Source/Editor/Persona/Private/Shared/AnimationRecorder.h
--- a/Engine/Source/Editor/Persona/Private/Shared/AnimationRecorder.h
+++ b/Engine/Source/Editor/Persona/Private/Shared/AnimationRecorder.h
@@ -1,5 +1,7 @@
 // Copyright 1998-2014 Epic Games, Inc. All Rights Reserved.
 
+#pragma once
+
 //////////////////////////////////////////////////////////////////////////
 // FAnimationRecorder
 
