NUXF-236 - Persona zooms out too far on actor focus

diff --git a/Engine/Source/Editor/Persona/Private/AnimationEditorViewportClient.cpp b/Engine/Source/Editor/Persona/Private/AnimationEditorViewportClient.cpp
--- a/Engine/Source/Editor/Persona/Private/AnimationEditorViewportClient.cpp
+++ b/Engine/Source/Editor/Persona/Private/AnimationEditorViewportClient.cpp
@@ -75,6 +75,7 @@ FAnimationViewportClient::FAnimationViewportClient( FPreviewScene& InPreviewScen
 	, GravityScaleSliderValue(0.25f)
 	, PrevWindStrength(0.0f)
 	, SelectedWindActor(NULL)
+	, bFocusOnDraw(false)
 	, BodyTraceDistance(100000.0f)
 {
 	// load config
@@ -507,6 +508,12 @@ void FAnimationViewportClient::Draw(const FSceneView* View, FPrimitiveDrawInterf
 				Node->Draw(PDI, PreviewSkelMeshComp.Get());
 			}
 		}
+
+		if(bFocusOnDraw)
+		{
+			bFocusOnDraw = false;
+			FocusViewportOnPreviewMesh();
+		}
 	}
 }
 
@@ -616,7 +623,7 @@ void FAnimationViewportClient::Tick(float DeltaSeconds)
 		// Handle updating the preview component to represent the effects of root motion	
 		FBoxSphereBounds Bounds = EditorFloorComp->CalcBounds(EditorFloorComp->GetRelativeTransform());
 		PreviewComp->ConsumeRootMotion(Bounds.GetBox().Min, Bounds.GetBox().Max);
-	}
+	}	
 }
 
 void FAnimationViewportClient::SetCameraTargetLocation(const FSphere &BoundSphere, float DeltaSeconds)
@@ -1781,7 +1788,7 @@ void FAnimationViewportClient::UpdateCameraSetup()
 
 void FAnimationViewportClient::FocusViewportOnSphere( FSphere& Sphere )
 {
-	FBox Box( Sphere.Center - FVector(Sphere.W), Sphere.Center + FVector(Sphere.W) );
+	FBox Box( Sphere.Center - FVector(Sphere.W, 0.0f, 0.0f), Sphere.Center + FVector(Sphere.W, 0.0f, 0.0f) );
 
 	bool bInstant = false;
 	FocusViewportOnBox( Box, bInstant );
@@ -1791,6 +1798,14 @@ void FAnimationViewportClient::FocusViewportOnSphere( FSphere& Sphere )
 
 void FAnimationViewportClient::FocusViewportOnPreviewMesh()
 {
+	FIntPoint ViewportSize = Viewport->GetSizeXY();
+
+	if(!(ViewportSize.SizeSquared() > 0))
+	{
+		// We cannot focus fully right now as the viewport does not know its size
+		// and we must have the aspect to correctly focus on the component,
+		bFocusOnDraw = true;
+	}
 	FSphere Sphere = GetCameraTarget();
 	FocusViewportOnSphere(Sphere);
 }
diff --git a/Engine/Source/Editor/Persona/Private/AnimationEditorViewportClient.h b/Engine/Source/Editor/Persona/Private/AnimationEditorViewportClient.h
--- a/Engine/Source/Editor/Persona/Private/AnimationEditorViewportClient.h
+++ b/Engine/Source/Editor/Persona/Private/AnimationEditorViewportClient.h
@@ -281,6 +281,9 @@ private:
 	/** Wind actor used for preview */
 	TWeakObjectPtr<AWindDirectionalSource> SelectedWindActor;
 
+	/** Focus on the preview component the next time we draw the viewport */
+	bool bFocusOnDraw;
+
 	/** Distance to trace for physics bodies */
 	const float BodyTraceDistance;
 private:
diff --git a/Engine/Source/Editor/UnrealEd/Private/EditorViewportClient.cpp b/Engine/Source/Editor/UnrealEd/Private/EditorViewportClient.cpp
--- a/Engine/Source/Editor/UnrealEd/Private/EditorViewportClient.cpp
+++ b/Engine/Source/Editor/UnrealEd/Private/EditorViewportClient.cpp
@@ -406,6 +406,13 @@ void FEditorViewportClient::FocusViewportOnBox( const FBox& BoundingBox, bool bI
 	const FVector Position = BoundingBox.GetCenter();
 	float Radius = BoundingBox.GetExtent().Size();
 
+	float AspectToUse = AspectRatio;
+	FIntPoint ViewportSize = Viewport->GetSizeXY();
+	if(!bUseControllingActorViewInfo && ViewportSize.X > 0 && ViewportSize.Y > 0)
+	{
+		AspectToUse = Viewport->GetDesiredAspectRatio();
+	}
+
 	const bool bEnable=false;
 	ToggleOrbitCamera(bEnable);
 
@@ -417,9 +424,9 @@ void FEditorViewportClient::FocusViewportOnBox( const FBox& BoundingBox, bool bI
 			* than the width of the viewport, we scale the radius by the aspect ratio in order to compensate for the fact that we have
 			* less visible vertically than horizontally.
 			*/
-			if( AspectRatio > 1.0f )
+			if( AspectToUse > 1.0f )
 			{
-				Radius *= AspectRatio;
+				Radius *= AspectToUse;
 			}
 
 			/** 
@@ -451,7 +458,7 @@ void FEditorViewportClient::FocusViewportOnBox( const FBox& BoundingBox, bool bI
 				* to find the new OrthoZoom value for the viewport. The 15.0f is a fudge factor.
 				*/
 				float NewOrthoZoom;
-				uint32 MinAxisSize = (AspectRatio > 1.0f) ? Viewport->GetSizeXY().Y : Viewport->GetSizeXY().X;
+				uint32 MinAxisSize = (AspectToUse > 1.0f) ? Viewport->GetSizeXY().Y : Viewport->GetSizeXY().X;
 				float Zoom = Radius / (MinAxisSize / 2);
 
 				NewOrthoZoom = Zoom * (Viewport->GetSizeXY().X*15.0f);
