TTP #331607 : Persona: Move CollectAnimationNotifies into Persona initialization instead of USkeleton::PostLoad

diff --git a/Engine/Source/Editor/Persona/Private/Persona.cpp b/Engine/Source/Editor/Persona/Private/Persona.cpp
--- a/Engine/Source/Editor/Persona/Private/Persona.cpp
+++ b/Engine/Source/Editor/Persona/Private/Persona.cpp
@@ -612,6 +612,8 @@ void FPersona::InitPersona(const EToolkitMode::Type Mode, const TSharedPtr< clas
 	ObjectsBeingEdited.Add(TargetSkeleton);
 	check(TargetSkeleton);
 
+	TargetSkeleton->CollectAnimationNotifies();
+
 	TargetSkeleton->RegisterOnSkeletonHierarchyChanged( USkeleton::FOnSkeletonHierarchyChanged::CreateSP( this, &FPersona::OnSkeletonHierarchyChanged ) );
 
 	// We could modify the skeleton within Persona (add/remove sockets), so we need to enable undo/redo on it
diff --git a/Engine/Source/Runtime/Engine/Classes/Animation/Skeleton.h b/Engine/Source/Runtime/Engine/Classes/Animation/Skeleton.h
--- a/Engine/Source/Runtime/Engine/Classes/Animation/Skeleton.h
+++ b/Engine/Source/Runtime/Engine/Classes/Animation/Skeleton.h
@@ -278,7 +278,7 @@ public:
 #if WITH_EDITORONLY_DATA
 
 	// @todo document
-	void CollectAnimationNotifies();
+	ENGINE_API void CollectAnimationNotifies();
 
 	// @todo document
 	ENGINE_API void AddNewAnimationNotify(FName NewAnimNotifyName);
diff --git a/Engine/Source/Runtime/Engine/Private/Animation/Skeleton.cpp b/Engine/Source/Runtime/Engine/Private/Animation/Skeleton.cpp
--- a/Engine/Source/Runtime/Engine/Private/Animation/Skeleton.cpp
+++ b/Engine/Source/Runtime/Engine/Private/Animation/Skeleton.cpp
@@ -77,13 +77,6 @@ void USkeleton::PostLoad()
 
 	// catch any case if guid isn't valid
 	check(Guid.IsValid());
-
-#if WITH_EDITOR
-	if ( GIsEditor )
-	{
-		CollectAnimationNotifies();
-	}
-#endif
 }
 
 void USkeleton::PostDuplicate(bool bDuplicateForPIE)
