- Fixed an issue with blendspace not showing bone name for per bone interpolation
- Fixed comment

diff --git a/Engine/Source/Editor/Persona/Private/AnimGraphNodeDetails.cpp b/Engine/Source/Editor/Persona/Private/AnimGraphNodeDetails.cpp
--- a/Engine/Source/Editor/Persona/Private/AnimGraphNodeDetails.cpp
+++ b/Engine/Source/Editor/Persona/Private/AnimGraphNodeDetails.cpp
@@ -427,6 +427,7 @@ void FBoneReferenceCustomization::CustomizeHeader( TSharedRef<IPropertyHandle> S
 	StructPropertyHandle->GetOuterObjects(Objects);
 	UAnimGraphNode_Base* AnimGraphNode = NULL;
 	USkeletalMesh* SkeletalMesh = NULL;
+	UAnimationAsset * AnimationAsset = NULL;
 	TargetSkeleton = NULL;
 
 	for (auto OuterIter = Objects.CreateIterator() ; OuterIter ; ++OuterIter)
@@ -443,6 +444,12 @@ void FBoneReferenceCustomization::CustomizeHeader( TSharedRef<IPropertyHandle> S
 			TargetSkeleton = SkeletalMesh->Skeleton;
 			break;
 		}
+		AnimationAsset = Cast<UAnimationAsset>(*OuterIter);
+		if(AnimationAsset)
+		{
+			TargetSkeleton = AnimationAsset->GetSkeleton();
+			break;
+		}
 	}
 
 	if (TargetSkeleton)
diff --git a/Engine/Source/Runtime/Engine/Private/Animation/Skeleton.cpp b/Engine/Source/Runtime/Engine/Private/Animation/Skeleton.cpp
--- a/Engine/Source/Runtime/Engine/Private/Animation/Skeleton.cpp
+++ b/Engine/Source/Runtime/Engine/Private/Animation/Skeleton.cpp
@@ -1,7 +1,7 @@
 // Copyright 1998-2014 Epic Games, Inc. All Rights Reserved.
 
 /*=============================================================================
-	UnSkeletalAnim.cpp: Skeletal mesh animation functions.
+	Skeleton.cpp: Skeleton features
 =============================================================================*/ 
 
 #include "EnginePrivate.h"
