Separated Bone Selection Widget out of the customization, so that we can use it for other stuff

diff --git a/Engine/Source/Editor/Persona/Private/AnimGraphNodeDetails.cpp b/Engine/Source/Editor/Persona/Private/AnimGraphNodeDetails.cpp
--- a/Engine/Source/Editor/Persona/Private/AnimGraphNodeDetails.cpp
+++ b/Engine/Source/Editor/Persona/Private/AnimGraphNodeDetails.cpp
@@ -19,6 +19,7 @@
 #include "AnimGraphNode_Base.h"
 #include "GraphEditor.h"
 #include "Persona.h"
+#include "BoneSelectionWidget.h"
 
 #define LOCTEXT_NAMESPACE "KismetNodeWithOptionalPinsDetails"
 
@@ -428,7 +429,7 @@ void FBoneReferenceCustomization::CustomizeHeader( TSharedRef<IPropertyHandle> S
 	UAnimGraphNode_Base* AnimGraphNode = NULL;
 	USkeletalMesh* SkeletalMesh = NULL;
 	UAnimationAsset * AnimationAsset = NULL;
-	TargetSkeleton = NULL;
+	USkeleton * TargetSkeleton = NULL;
 
 	for (auto OuterIter = Objects.CreateIterator() ; OuterIter ; ++OuterIter)
 	{
@@ -459,20 +460,13 @@ void FBoneReferenceCustomization::CustomizeHeader( TSharedRef<IPropertyHandle> S
 			StructPropertyHandle->CreatePropertyNameWidget()
 		];
 
-		FString DefaultTooltip = StructPropertyHandle->GetToolTipText();
-		FText FinalTooltip = FText::Format(LOCTEXT("BoneClickToolTip", "{0}\nClick to choose a different bone"), FText::FromString(DefaultTooltip));
-
 		HeaderRow.ValueContent()
 		[
-			SAssignNew(BonePickerButton, SComboButton)
-				.OnGetMenuContent(FOnGetContent::CreateSP(this, &FBoneReferenceCustomization::CreateSkeletonWidgetMenu, StructPropertyHandle))
-				.ContentPadding(0)
-				.ButtonContent()
-				[
-					SNew(STextBlock)
-					.Text(this, &FBoneReferenceCustomization::GetCurrentBoneName)
-					.ToolTipText(FinalTooltip)
-				]
+			SNew(SBoneSelectionWidget)
+			.Skeleton(TargetSkeleton)
+			.Tooltip(FText::FromString(StructPropertyHandle->GetToolTipText()))
+			.OnBoneSelectionChanged(this, &FBoneReferenceCustomization::OnBoneSelectionChanged)
+			.OnGetSelectedBone(this, &FBoneReferenceCustomization::GetSelectedBone)
 		];
 	}
 }
@@ -483,159 +477,17 @@ void FBoneReferenceCustomization::CustomizeChildren( TSharedRef<IPropertyHandle>
 	// No child customisations as the properties are shown in the header
 }
 
-TSharedRef<SWidget> FBoneReferenceCustomization::CreateSkeletonWidgetMenu( TSharedRef<IPropertyHandle> TargetPropertyHandle)
-{
-	SAssignNew(TreeView, STreeView<TSharedPtr<FBoneNameInfo>>)
-		.TreeItemsSource(&SkeletonTreeInfo)
-		.OnGenerateRow(this, &FBoneReferenceCustomization::MakeTreeRowWidget)
-		.OnGetChildren(this, &FBoneReferenceCustomization::GetChildrenForInfo)
-		.OnSelectionChanged(this, &FBoneReferenceCustomization::OnSelectionChanged)
-		.SelectionMode(ESelectionMode::Single);
-
-	RebuildBoneList();
-
-	TSharedPtr<SSearchBox> SearchWidgetToFocus = NULL;
-	TSharedRef<SBorder> MenuWidget = SNew(SBorder)
-		.Padding(6)
-		.BorderImage(FEditorStyle::GetBrush("NoBorder"))
-		.Content()
-		[
-			SNew(SBox)
-			.WidthOverride(300)
-			.HeightOverride(512)
-			.Content()
-			[
-				SNew(SVerticalBox)
-				+SVerticalBox::Slot()
-				.AutoHeight()
-				[
-					SNew(STextBlock)
-						.Font(FEditorStyle::GetFontStyle("BoldFont"))
-						.Text(LOCTEXT("BonePickerTitle", "Pick Bone..."))
-				]
-				+SVerticalBox::Slot()
-				.AutoHeight()
-				[
-					SNew(SSeparator)
-						.SeparatorImage(FEditorStyle::GetBrush("Menu.Separator"))
-						.Orientation(Orient_Horizontal)
-				]
-				+SVerticalBox::Slot()
-				.AutoHeight()
-				[
-					SAssignNew(SearchWidgetToFocus, SSearchBox)
-					.SelectAllTextWhenFocused(true)
-					.OnTextChanged(this, &FBoneReferenceCustomization::OnFilterTextChanged)
-					.OnTextCommitted(this, &FBoneReferenceCustomization::OnFilterTextCommitted)
-					.HintText(NSLOCTEXT("BonePicker", "Search", "Search..."))
-				]
-				+SVerticalBox::Slot()
-					[
-						TreeView->AsShared()
-					]
-			]
-		];
-
-	BonePickerButton->SetMenuContentWidgetToFocus(SearchWidgetToFocus);
-
-	return MenuWidget;
-}
-
-TSharedRef<ITableRow> FBoneReferenceCustomization::MakeTreeRowWidget( TSharedPtr<FBoneNameInfo> InInfo, const TSharedRef<STableViewBase>& OwnerTable )
-{
-	return SNew(STableRow<TSharedPtr<FBoneNameInfo>>, OwnerTable)
-		.Content()
-		[
-			SNew(STextBlock)
-				.HighlightText(FilterText)
-				.Text(InInfo->BoneName.ToString())
-		];
-}
-
-void FBoneReferenceCustomization::GetChildrenForInfo( TSharedPtr<FBoneNameInfo> InInfo, TArray< TSharedPtr<FBoneNameInfo> >& OutChildren )
-{
-	OutChildren = InInfo->Children;
-}
-
-void FBoneReferenceCustomization::OnFilterTextChanged( const FText& InFilterText )
-{
-	FilterText = InFilterText;
-
-	RebuildBoneList();
-}
-
-void FBoneReferenceCustomization::OnFilterTextCommitted( const FText& SearchText, ETextCommit::Type CommitInfo )
+void FBoneReferenceCustomization::OnBoneSelectionChanged(FName Name)
 {
-	// Already committed as the text was typed
+	BoneRefProperty->SetValue(Name);
 }
-
-void FBoneReferenceCustomization::RebuildBoneList()
-{
-	SkeletonTreeInfo.Empty();
-	SkeletonTreeInfoFlat.Empty();
-	const FReferenceSkeleton& RefSkeleton = TargetSkeleton->GetReferenceSkeleton();
-	for(int32 BoneIdx = 0 ; BoneIdx < RefSkeleton.GetNum() ; ++BoneIdx)
-	{
-		TSharedRef<FBoneNameInfo> BoneInfo = MakeShareable(new FBoneNameInfo(RefSkeleton.GetBoneName(BoneIdx)));
-
-		// Filter if Necessary
-		if(!FilterText.IsEmpty() && !BoneInfo->BoneName.ToString().Contains(FilterText.ToString()))
-		{
-			continue;
-		}
-
-		int32 ParentIdx = RefSkeleton.GetParentIndex(BoneIdx);
-		bool bAddToParent = false;
-
-		if(ParentIdx != INDEX_NONE && FilterText.IsEmpty())
-		{
-			// We have a parent, search for it in the flat list
-			FName ParentName = RefSkeleton.GetBoneName(ParentIdx);
-
-			for(int32 FlatListIdx = 0 ; FlatListIdx < SkeletonTreeInfoFlat.Num() ; ++FlatListIdx)
-			{
-				TSharedPtr<FBoneNameInfo> InfoEntry = SkeletonTreeInfoFlat[FlatListIdx];
-				if(InfoEntry->BoneName == ParentName)
-				{
-					bAddToParent = true;
-					ParentIdx = FlatListIdx;
-					break;
-				}
-			}
-
-			if(bAddToParent)
-			{
-				SkeletonTreeInfoFlat[ParentIdx]->Children.Add(BoneInfo);
-			}
-			else
-			{
-				SkeletonTreeInfo.Add(BoneInfo);
-			}
-		}
-		else
-		{
-			SkeletonTreeInfo.Add(BoneInfo);
-		}
-
-		SkeletonTreeInfoFlat.Add(BoneInfo);
-		TreeView->SetItemExpansion(BoneInfo, true);
-	}
-
-	TreeView->RequestTreeRefresh();
-}
-
-void FBoneReferenceCustomization::OnSelectionChanged( TSharedPtr<FBoneNameInfo> BoneInfo, ESelectInfo::Type SelectInfo )
-{
-	FilterText = FText::FromString("");
-	BoneRefProperty->SetValue(BoneInfo->BoneName);
-	BonePickerButton->SetIsOpen(false);
-}
-
-FString FBoneReferenceCustomization::GetCurrentBoneName() const
+FName FBoneReferenceCustomization::GetSelectedBone() const
 {
 	FString OutText;
+	
 	BoneRefProperty->GetValueAsFormattedString(OutText);
-	return OutText;
+
+	return FName(*OutText);
 }
 
 TSharedRef<IDetailCustomization> FAnimGraphParentPlayerDetails::MakeInstance(TWeakPtr<FPersona> InPersona)
diff --git a/Engine/Source/Editor/Persona/Private/AnimGraphNodeDetails.h b/Engine/Source/Editor/Persona/Private/AnimGraphNodeDetails.h
--- a/Engine/Source/Editor/Persona/Private/AnimGraphNodeDetails.h
+++ b/Engine/Source/Editor/Persona/Private/AnimGraphNodeDetails.h
@@ -65,46 +65,11 @@ public:
 	virtual void CustomizeChildren(TSharedRef<IPropertyHandle> StructPropertyHandle, class IDetailChildrenBuilder& ChildBuilder, IPropertyTypeCustomizationUtils& StructCustomizationUtils) override;
 
 private:
-	// Storage object for bone hierarchy
-	struct FBoneNameInfo
-	{
-		FBoneNameInfo(FName Name) : BoneName(Name) {}
 
-		FName BoneName;
-		TArray<TSharedPtr<FBoneNameInfo>> Children;
-	};
+	// Bone tree widget delegates
+	void OnBoneSelectionChanged(FName Name);
+	FName GetSelectedBone() const;
 
-	// Creates the combo button menu when clicked
-	TSharedRef<SWidget> CreateSkeletonWidgetMenu( TSharedRef<IPropertyHandle> TargetPropertyHandle);
-	// Using the current filter, repopulate the tree view
-	void RebuildBoneList();
-	// Make a single tree row widget
-	TSharedRef<ITableRow> MakeTreeRowWidget(TSharedPtr<FBoneNameInfo> InInfo, const TSharedRef<STableViewBase>& OwnerTable);
-	// Get the children for the provided bone info
-	void GetChildrenForInfo(TSharedPtr<FBoneNameInfo> InInfo, TArray< TSharedPtr<FBoneNameInfo> >& OutChildren);
-
-	// Called when the user changes the search filter
-	void OnFilterTextChanged( const FText& InFilterText );
-	void OnFilterTextCommitted( const FText& SearchText, ETextCommit::Type CommitInfo );
-	// Called when the user selects a bone name
-	void OnSelectionChanged(TSharedPtr<FBoneNameInfo>, ESelectInfo::Type SelectInfo);
-	// Gets the current bone name, used to get the right name for the combo button
-	FString GetCurrentBoneName() const;
-
-	// Skeleton to search
-	USkeleton* TargetSkeleton;
-
-	// Base combo button 
-	TSharedPtr<SComboButton> BonePickerButton;
-	// Tree view used in the button menu
-	TSharedPtr<STreeView<TSharedPtr<FBoneNameInfo>>> TreeView;
-
-	// Tree info entries for bone picker
-	TArray<TSharedPtr<FBoneNameInfo>> SkeletonTreeInfo;
-	// Mirror of SkeletonTreeInfo but flattened for searching
-	TArray<TSharedPtr<FBoneNameInfo>> SkeletonTreeInfoFlat;
-	// Text to filter bone tree with
-	FText FilterText;
 	// Property to change after bone has been picked
 	TSharedPtr<IPropertyHandle> BoneRefProperty;
 };
diff --git a/Engine/Source/Editor/Persona/Private/Shared/BoneSelectionWidget.cpp b/Engine/Source/Editor/Persona/Private/Shared/BoneSelectionWidget.cpp
new file mode 100644
--- /dev/null
+++ b/Engine/Source/Editor/Persona/Private/Shared/BoneSelectionWidget.cpp
@@ -0,0 +1,205 @@
+// Copyright 1998-2014 Epic Games, Inc. All Rights Reserved.
+
+
+#include "PersonaPrivatePCH.h"
+#include "BoneSelectionWidget.h"
+
+#define LOCTEXT_NAMESPACE "SBoneSelectionWidget"
+
+/////////////////////////////////////////////////////
+
+void SBoneSelectionWidget::Construct(const FArguments& InArgs)
+{
+	if(InArgs._Skeleton.IsValid())
+	{
+		TargetSkeleton = InArgs._Skeleton.Get();
+	}
+
+	check(TargetSkeleton);
+
+	OnBoneSelectionChanged = InArgs._OnBoneSelectionChanged;
+	OnGetSelectedBone = InArgs._OnGetSelectedBone;
+
+	FText FinalTooltip = FText::Format(LOCTEXT("BoneClickToolTip", "{0}\nClick to choose a different bone"), InArgs._Tooltip);
+
+	this->ChildSlot
+	[
+		SAssignNew(BonePickerButton, SComboButton)
+		.OnGetMenuContent(FOnGetContent::CreateSP(this, &SBoneSelectionWidget::CreateSkeletonWidgetMenu))
+		.ContentPadding(0)
+		.ButtonContent()
+		[
+			SNew(STextBlock)
+			.Text(this, &SBoneSelectionWidget::GetCurrentBoneName)
+			.ToolTipText(FinalTooltip)
+		]
+	];
+}
+
+TSharedRef<SWidget> SBoneSelectionWidget::CreateSkeletonWidgetMenu()
+{
+	SAssignNew(TreeView, STreeView<TSharedPtr<FBoneNameInfo>>)
+		.TreeItemsSource(&SkeletonTreeInfo)
+		.OnGenerateRow(this, &SBoneSelectionWidget::MakeTreeRowWidget)
+		.OnGetChildren(this, &SBoneSelectionWidget::GetChildrenForInfo)
+		.OnSelectionChanged(this, &SBoneSelectionWidget::OnSelectionChanged)
+		.SelectionMode(ESelectionMode::Single);
+
+	RebuildBoneList();
+
+	TSharedPtr<SSearchBox> SearchWidgetToFocus = NULL;
+	TSharedRef<SBorder> MenuWidget = SNew(SBorder)
+		.Padding(6)
+		.BorderImage(FEditorStyle::GetBrush("NoBorder"))
+		.Content()
+		[
+			SNew(SBox)
+			.WidthOverride(300)
+			.HeightOverride(512)
+			.Content()
+			[
+				SNew(SVerticalBox)
+				+SVerticalBox::Slot()
+				.AutoHeight()
+				[
+					SNew(STextBlock)
+					.Font(FEditorStyle::GetFontStyle("BoldFont"))
+						.Text(LOCTEXT("BonePickerTitle", "Pick Bone..."))
+				]
+				+SVerticalBox::Slot()
+				.AutoHeight()
+				[
+					SNew(SSeparator)
+					.SeparatorImage(FEditorStyle::GetBrush("Menu.Separator"))
+						.Orientation(Orient_Horizontal)
+				]
+				+SVerticalBox::Slot()
+				.AutoHeight()
+				[
+					SAssignNew(SearchWidgetToFocus, SSearchBox)
+					.SelectAllTextWhenFocused(true)
+					.OnTextChanged(this, &SBoneSelectionWidget::OnFilterTextChanged)
+					.OnTextCommitted(this, &SBoneSelectionWidget::OnFilterTextCommitted)
+					.HintText(NSLOCTEXT("BonePicker", "Search", "Search..."))
+				]
+				+SVerticalBox::Slot()
+					[
+						TreeView->AsShared()
+					]
+			]
+		];
+
+	BonePickerButton->SetMenuContentWidgetToFocus(SearchWidgetToFocus);
+
+	return MenuWidget;
+}
+
+TSharedRef<ITableRow> SBoneSelectionWidget::MakeTreeRowWidget(TSharedPtr<FBoneNameInfo> InInfo, const TSharedRef<STableViewBase>& OwnerTable)
+{
+	return SNew(STableRow<TSharedPtr<FBoneNameInfo>>, OwnerTable)
+		.Content()
+		[
+			SNew(STextBlock)
+			.HighlightText(FilterText)
+				.Text(InInfo->BoneName.ToString())
+		];
+}
+
+void SBoneSelectionWidget::GetChildrenForInfo(TSharedPtr<FBoneNameInfo> InInfo, TArray< TSharedPtr<FBoneNameInfo> >& OutChildren)
+{
+	OutChildren = InInfo->Children;
+}
+
+void SBoneSelectionWidget::OnFilterTextChanged(const FText& InFilterText)
+{
+	FilterText = InFilterText;
+
+	RebuildBoneList();
+}
+
+void SBoneSelectionWidget::OnFilterTextCommitted(const FText& SearchText, ETextCommit::Type CommitInfo)
+{
+	// Already committed as the text was typed
+}
+
+void SBoneSelectionWidget::RebuildBoneList()
+{
+	SkeletonTreeInfo.Empty();
+	SkeletonTreeInfoFlat.Empty();
+	const FReferenceSkeleton& RefSkeleton = TargetSkeleton->GetReferenceSkeleton();
+	for(int32 BoneIdx = 0; BoneIdx < RefSkeleton.GetNum(); ++BoneIdx)
+	{
+		TSharedRef<FBoneNameInfo> BoneInfo = MakeShareable(new FBoneNameInfo(RefSkeleton.GetBoneName(BoneIdx)));
+
+		// Filter if Necessary
+		if(!FilterText.IsEmpty() && !BoneInfo->BoneName.ToString().Contains(FilterText.ToString()))
+		{
+			continue;
+		}
+
+		int32 ParentIdx = RefSkeleton.GetParentIndex(BoneIdx);
+		bool bAddToParent = false;
+
+		if(ParentIdx != INDEX_NONE && FilterText.IsEmpty())
+		{
+			// We have a parent, search for it in the flat list
+			FName ParentName = RefSkeleton.GetBoneName(ParentIdx);
+
+			for(int32 FlatListIdx = 0; FlatListIdx < SkeletonTreeInfoFlat.Num(); ++FlatListIdx)
+			{
+				TSharedPtr<FBoneNameInfo> InfoEntry = SkeletonTreeInfoFlat[FlatListIdx];
+				if(InfoEntry->BoneName == ParentName)
+				{
+					bAddToParent = true;
+					ParentIdx = FlatListIdx;
+					break;
+				}
+			}
+
+			if(bAddToParent)
+			{
+				SkeletonTreeInfoFlat[ParentIdx]->Children.Add(BoneInfo);
+			}
+			else
+			{
+				SkeletonTreeInfo.Add(BoneInfo);
+			}
+		}
+		else
+		{
+			SkeletonTreeInfo.Add(BoneInfo);
+		}
+
+		SkeletonTreeInfoFlat.Add(BoneInfo);
+		TreeView->SetItemExpansion(BoneInfo, true);
+	}
+
+	TreeView->RequestTreeRefresh();
+}
+
+void SBoneSelectionWidget::OnSelectionChanged(TSharedPtr<FBoneNameInfo> BoneInfo, ESelectInfo::Type SelectInfo)
+{
+	FilterText = FText::FromString("");
+
+	if(OnBoneSelectionChanged.IsBound())
+	{
+		OnBoneSelectionChanged.Execute(BoneInfo->BoneName);
+	}
+
+	BonePickerButton->SetIsOpen(false);
+}
+
+FString SBoneSelectionWidget::GetCurrentBoneName() const
+{
+	if(OnGetSelectedBone.IsBound())
+	{
+		FName Name = OnGetSelectedBone.Execute();
+
+		return Name.ToString();
+	}
+
+	// @todo implement default solution?
+	return TEXT("");
+}
+#undef LOCTEXT_NAMESPACE
+
diff --git a/Engine/Source/Editor/Persona/Public/BoneSelectionWidget.h b/Engine/Source/Editor/Persona/Public/BoneSelectionWidget.h
new file mode 100644
--- /dev/null
+++ b/Engine/Source/Editor/Persona/Public/BoneSelectionWidget.h
@@ -0,0 +1,89 @@
+
+#ifndef __BoneSelectionWidget_h__
+#define __BoneSelectionWidget_h__
+
+#pragma once
+DECLARE_DELEGATE_OneParam(FOnBoneSelectionChanged, FName);
+DECLARE_DELEGATE_RetVal(FName, FGetSelectedBone);
+
+class SBoneSelectionWidget : public SCompoundWidget
+{
+public: 
+
+	SLATE_BEGIN_ARGS( SBoneSelectionWidget )
+		:_Skeleton()
+		,_Tooltip()
+		,_OnBoneSelectionChanged()
+		,_OnGetSelectedBone()
+	{}
+
+		/** Vertical alignment */
+		SLATE_ARGUMENT(TWeakObjectPtr<USkeleton>, Skeleton);
+		/** Set tooltip attribute */
+		SLATE_ARGUMENT(FText, Tooltip);
+
+		/** set selected bone name */
+		SLATE_EVENT(FOnBoneSelectionChanged, OnBoneSelectionChanged);
+
+		/** get selected bone name **/
+		SLATE_EVENT(FGetSelectedBone, OnGetSelectedBone);
+
+	SLATE_END_ARGS();
+
+	/**
+	 * Construct this widget
+	 *
+	 * @param	InArgs	The declaration data for this widget
+	 */
+	void Construct( const FArguments& InArgs );
+
+private: 
+
+	// Storage object for bone hierarchy
+	struct FBoneNameInfo
+	{
+		FBoneNameInfo(FName Name): BoneName(Name) {}
+
+		FName BoneName;
+		TArray<TSharedPtr<FBoneNameInfo>> Children;
+	};
+
+	// Creates the combo button menu when clicked
+	TSharedRef<SWidget> CreateSkeletonWidgetMenu();
+	// Using the current filter, repopulate the tree view
+	void RebuildBoneList();
+	// Make a single tree row widget
+	TSharedRef<ITableRow> MakeTreeRowWidget(TSharedPtr<FBoneNameInfo> InInfo, const TSharedRef<STableViewBase>& OwnerTable);
+	// Get the children for the provided bone info
+	void GetChildrenForInfo(TSharedPtr<FBoneNameInfo> InInfo, TArray< TSharedPtr<FBoneNameInfo> >& OutChildren);
+	// Called when the user changes the search filter
+	void OnFilterTextChanged(const FText& InFilterText);
+	void OnFilterTextCommitted(const FText& SearchText, ETextCommit::Type CommitInfo);
+	// Called when the user selects a bone name
+	void OnSelectionChanged(TSharedPtr<FBoneNameInfo>, ESelectInfo::Type SelectInfo);
+	// Gets the current bone name, used to get the right name for the combo button
+	FString GetCurrentBoneName() const;
+
+	// Tree view used in the button menu
+	TSharedPtr<STreeView<TSharedPtr<FBoneNameInfo>>> TreeView;
+
+	// Tree info entries for bone picker
+	TArray<TSharedPtr<FBoneNameInfo>> SkeletonTreeInfo;
+	// Mirror of SkeletonTreeInfo but flattened for searching
+	TArray<TSharedPtr<FBoneNameInfo>> SkeletonTreeInfoFlat;
+
+	// Base combo button 
+	TSharedPtr<SComboButton> BonePickerButton;
+
+	// Skeleton to search
+	USkeleton* TargetSkeleton;
+
+	// Text to filter bone tree with
+	FText FilterText;
+
+	// delegates
+	FOnBoneSelectionChanged OnBoneSelectionChanged;
+	FGetSelectedBone OnGetSelectedBone;
+};
+
+#endif		//__BoneSelectionWidget_h__
\ No newline at end of file
