343756	POLISH: PERSONA: Typo in the Compression button's tooltip

diff --git a/Engine/Source/Editor/Persona/Private/PersonaCommands.cpp b/Engine/Source/Editor/Persona/Private/PersonaCommands.cpp
--- a/Engine/Source/Editor/Persona/Private/PersonaCommands.cpp
+++ b/Engine/Source/Editor/Persona/Private/PersonaCommands.cpp
@@ -25,7 +25,7 @@ void FPersonaCommands::RegisterCommands()
 	// import animation
 	UI_COMMAND( ImportAnimation, "Import Animation", "Import new animation for the skeleton.", EUserInterfaceActionType::Button, FInputGesture());
 	UI_COMMAND( ReimportAnimation, "Reimport Animation", "Reimport current animation.", EUserInterfaceActionType::Button, FInputGesture());
-	UI_COMMAND( ApplyCompression, "Apply Compression", "Apply cvompression to current animation", EUserInterfaceActionType::Button, FInputGesture() );
+	UI_COMMAND( ApplyCompression, "Apply Compression", "Apply compression to current animation", EUserInterfaceActionType::Button, FInputGesture() );
 	UI_COMMAND( RecordAnimation, "Record to new Animation", "Create new animation from currently playing", EUserInterfaceActionType::Button, FInputGesture() );
 	UI_COMMAND( ExportToFBX, "Export to FBX", "Export current animation to FBX", EUserInterfaceActionType::Button, FInputGesture() );
 	UI_COMMAND( AddLoopingInterpolation, "Add Looping Interpolation", "Add an extra first frame at the end of the animation to create interpolation when looping", EUserInterfaceActionType::Button, FInputGesture() );
