WidthScaleCurve Improvements

Changed the text entry box for the width scale curve property to have it search the list of curves and provide suggestions like the bone names properties do.

TTP 333827 - UE4: EDITOR: Improve Widget for WidthScaleCurve property of UAnimNotifyState_Trails

diff --git a/Engine/Source/Editor/Persona/Private/AnimNotifyDetails.cpp b/Engine/Source/Editor/Persona/Private/AnimNotifyDetails.cpp
--- a/Engine/Source/Editor/Persona/Private/AnimNotifyDetails.cpp
+++ b/Engine/Source/Editor/Persona/Private/AnimNotifyDetails.cpp
@@ -31,8 +31,11 @@ void FAnimNotifyDetails::CustomizeDetails( IDetailLayoutBuilder& DetailBuilder )
 		}
 		else if (DetailObjectClass->GetName().Find(TEXT("AnimNotifyState_Trail")) != INDEX_NONE)
 		{
+			AddPropertyDefault(DetailBuilder, DetailObjectClass, TEXT("PSTemplate"), TEXT("Trail"));	
 			AddBoneNameProperty(DetailBuilder, DetailObjectClass, TEXT("FirstSocketName"), TEXT("Trail"));
 			AddBoneNameProperty(DetailBuilder, DetailObjectClass, TEXT("SecondSocketName"), TEXT("Trail"));
+			AddPropertyDefault(DetailBuilder, DetailObjectClass, TEXT("WidthScaleMode"), TEXT("Trail"));
+			AddCurveNameProperty(DetailBuilder, DetailObjectClass, TEXT("WidthScaleCurve"), TEXT("Trail"));
 		}
 		else if (DetailObjectClass->GetName().Find(TEXT("AnimNotifyState_TimedParticleEffect")) != INDEX_NONE)
 		{
@@ -41,15 +44,25 @@ void FAnimNotifyDetails::CustomizeDetails( IDetailLayoutBuilder& DetailBuilder )
 	}
 }
 
+void FAnimNotifyDetails::AddPropertyDefault(IDetailLayoutBuilder& DetailBuilder, const UClass* PropertyClass, const TCHAR* PropertyName, const TCHAR* CategoryName)
+{
+	IDetailCategoryBuilder& AnimNotifyCategory = DetailBuilder.EditCategory(CategoryName, TEXT(""), ECategoryPriority::TypeSpecific);
+	TSharedPtr<IPropertyHandle> Property = DetailBuilder.GetProperty(FName(PropertyName), PropertyClass);
+	if (Property.IsValid())
+	{
+		AnimNotifyCategory.AddProperty(Property);
+	}
+}
+
 void FAnimNotifyDetails::AddBoneNameProperty(IDetailLayoutBuilder& DetailBuilder, const UClass* PropertyClass, const TCHAR* PropertyName, const TCHAR* CategoryName)
 {
 	IDetailCategoryBuilder& AnimNotifyCategory = DetailBuilder.EditCategory(CategoryName, TEXT(""), ECategoryPriority::TypeSpecific);
-	int32 PropIndex = BoneNameProperties.Num();
+	int32 PropIndex = NameProperties.Num();
 	TSharedPtr<IPropertyHandle> BoneNameProperty = DetailBuilder.GetProperty(FName(PropertyName), PropertyClass);
 
 	if (BoneNameProperty.IsValid() && BoneNameProperty->GetProperty())
 	{
-		BoneNameProperties.Add(BoneNameProperty);
+		NameProperties.Add(BoneNameProperty);
 		// get all the possible suggestions for the bones and sockets.
 		const TArray< TWeakObjectPtr<UObject> >& SelectedObjects = DetailBuilder.GetDetailsView().GetSelectedObjects();
 		if (SelectedObjects.Num() == 1)
@@ -87,7 +100,53 @@ void FAnimNotifyDetails::AddBoneNameProperty(IDetailLayoutBuilder& DetailBuilder
 	}
 }
 
+void FAnimNotifyDetails::AddCurveNameProperty(IDetailLayoutBuilder& DetailBuilder, const UClass* PropertyClass, const TCHAR* PropertyName, const TCHAR* CategoryName)
+{
+	IDetailCategoryBuilder& AnimNotifyCategory = DetailBuilder.EditCategory(CategoryName, TEXT(""), ECategoryPriority::TypeSpecific);
+	int32 PropIndex = NameProperties.Num();
+	TSharedPtr<IPropertyHandle> CurveNameProperty = DetailBuilder.GetProperty(FName(PropertyName), PropertyClass);
+
+	if (CurveNameProperty.IsValid() && CurveNameProperty->GetProperty())
+	{
+		NameProperties.Add(CurveNameProperty);
+		// get all the possible suggestions for the curves.
+		const TArray< TWeakObjectPtr<UObject> >& SelectedObjects = DetailBuilder.GetDetailsView().GetSelectedObjects();
+		if (SelectedObjects.Num() == 1)
+		{
+			TWeakObjectPtr<UObject> SelectedObject = SelectedObjects[0];
+			const UObject* OuterObject = SelectedObject->GetOuter();
+			if (const UAnimationAsset* AnimAsset = Cast<const UAnimationAsset>(OuterObject))
+			{
+				if (const USkeleton* Skeleton = AnimAsset->GetSkeleton())
+				{
+					AnimNotifyCategory.AddProperty(CurveNameProperty.ToSharedRef())
+						.CustomWidget()
+						.NameContent()
+						[
+							SNew(SHorizontalBox)
+							+ SHorizontalBox::Slot()
+							.Padding(FMargin(2, 1, 0, 1))
+							[
+								SNew(STextBlock)
+								.Text(CurveNameProperty->GetPropertyDisplayName())
+								.Font(FEditorStyle::GetFontStyle(TEXT("PropertyWindow.NormalFont")))
+							]
+						]
+					.ValueContent()
+						[
+							SNew(SAssetSearchBoxForCurves, Skeleton, CurveNameProperty)
+							.IncludeSocketsForSuggestions(true)
+							.MustMatchPossibleSuggestions(true)
+							.HintText(NSLOCTEXT("AnimNotifyDetails", "Curve Name Hint Text", "Curve Name..."))
+							.OnTextCommitted(this, &FAnimNotifyDetails::OnSearchBoxCommitted, PropIndex)
+						];
+				}
+			}
+		}
+	}
+}
+
 void FAnimNotifyDetails::OnSearchBoxCommitted(const FText& InSearchText, ETextCommit::Type CommitInfo, int32 PropertyIndex )
 {
-	BoneNameProperties[PropertyIndex]->SetValue( InSearchText.ToString() );
+	NameProperties[PropertyIndex]->SetValue( InSearchText.ToString() );
 }
diff --git a/Engine/Source/Editor/Persona/Private/AnimNotifyDetails.h b/Engine/Source/Editor/Persona/Private/AnimNotifyDetails.h
--- a/Engine/Source/Editor/Persona/Private/AnimNotifyDetails.h
+++ b/Engine/Source/Editor/Persona/Private/AnimNotifyDetails.h
@@ -15,12 +15,18 @@ public:
 
 private:
 
-	/** Array of pointers to the bone name properties */
-	TArray<TSharedPtr<IPropertyHandle>> BoneNameProperties;
+	/** Array of pointers to the name properties */
+	TArray<TSharedPtr<IPropertyHandle>> NameProperties;
 
 	/** Adds a Bone Name property to the details layout. */
 	void AddBoneNameProperty(IDetailLayoutBuilder& DetailBuilder, const UClass* PropertyClass, const TCHAR* PropertyName, const TCHAR* CategoryName);
 
+	/** Adds a Curve Name property to the details layout. */
+	void AddCurveNameProperty(IDetailLayoutBuilder& DetailBuilder, const UClass* PropertyClass, const TCHAR* PropertyName, const TCHAR* CategoryName);
+
+	/** Adds the default layout for a property. */
+	void AddPropertyDefault(IDetailLayoutBuilder& DetailBuilder, const UClass* PropertyClass, const TCHAR* PropertyName, const TCHAR* CategoryName);
+
 	/** Handle the committed text */
 	void OnSearchBoxCommitted(const FText& InSearchText, ETextCommit::Type CommitInfo, int32 PropertyIndex);
 
diff --git a/Engine/Source/Editor/Persona/Private/AssetSearchBoxUtilPersona.cpp b/Engine/Source/Editor/Persona/Private/AssetSearchBoxUtilPersona.cpp
--- a/Engine/Source/Editor/Persona/Private/AssetSearchBoxUtilPersona.cpp
+++ b/Engine/Source/Editor/Persona/Private/AssetSearchBoxUtilPersona.cpp
@@ -68,3 +68,50 @@ void SAssetSearchBoxForBones::Construct( const FArguments& InArgs, const class U
 		.MustMatchPossibleSuggestions(InArgs._MustMatchPossibleSuggestions)
 	];
 }
+
+void SAssetSearchBoxForCurves::Construct(const FArguments& InArgs, const class USkeleton* InSkeleton, TSharedPtr<class IPropertyHandle> CurveNameProperty)
+{
+	check(InSkeleton);
+
+	// get the currently chosen curve, if any
+	const FString PropertyName = CurveNameProperty->GetPropertyDisplayName();
+	FString CurValue;
+	CurveNameProperty->GetValue(CurValue);
+	if (CurValue == FString("None"))
+	{
+		CurValue.Empty();
+	}
+
+	Skeleton = InSkeleton;
+	
+	// create the asset search box
+	ChildSlot
+		[
+			SNew(SAssetSearchBox)
+			.InitialText(FText::FromString(CurValue))
+			.HintText(InArgs._HintText)
+			.OnTextCommitted(InArgs._OnTextCommitted)
+			.PossibleSuggestions(this, &SAssetSearchBoxForCurves::GetCurveSearchSuggestions)
+			.DelayChangeNotificationsWhileTyping(true)
+			.MustMatchPossibleSuggestions(InArgs._MustMatchPossibleSuggestions)
+		];
+}
+
+TArray<FString> SAssetSearchBoxForCurves::GetCurveSearchSuggestions() const
+{
+	TArray<FString> PossibleSuggestions;
+	if (USkeleton* Skel = Skeleton.Get())
+	{
+		if (FSmartNameMapping* Mapping = Skel->SmartNames.GetContainer(USkeleton::AnimCurveMappingName))
+		{
+			TArray<FName> Names; 
+			Mapping->FillNameArray(Names);
+			for (const FName& Name : Names)
+			{
+				PossibleSuggestions.Add(Name.ToString());
+			}
+		}
+	}
+
+	return PossibleSuggestions;
+}
\ No newline at end of file
diff --git a/Engine/Source/Editor/Persona/Private/AssetSearchBoxUtilPersona.h b/Engine/Source/Editor/Persona/Private/AssetSearchBoxUtilPersona.h
--- a/Engine/Source/Editor/Persona/Private/AssetSearchBoxUtilPersona.h
+++ b/Engine/Source/Editor/Persona/Private/AssetSearchBoxUtilPersona.h
@@ -32,3 +32,42 @@ public:
 		/** Constructs this widget with InArgs */
 		void Construct( const FArguments& InArgs, const class UObject* Outer, TSharedPtr<class IPropertyHandle> ParentBoneProperty );
 };
+
+
+class SAssetSearchBoxForCurves : public SCompoundWidget
+{
+public:
+	SLATE_BEGIN_ARGS(SAssetSearchBoxForCurves)
+		: _OnTextCommitted()
+		, _HintText()
+		, _MustMatchPossibleSuggestions(false)
+		, _IncludeSocketsForSuggestions(false)
+	{}
+
+	/** Where to place the suggestion list */
+	SLATE_ARGUMENT(EMenuPlacement, SuggestionListPlacement)
+
+	/** Invoked whenever the text is committed (e.g. user presses enter) */
+	SLATE_EVENT(FOnTextCommitted, OnTextCommitted)
+
+	/** Hint text to display for the search text when there is no value */
+	SLATE_ATTRIBUTE(FText, HintText)
+
+	/** Whether the SearchBox allow entries that don't match the possible suggestions */
+	SLATE_ATTRIBUTE(bool, MustMatchPossibleSuggestions)
+
+	/** Whether the possible suggestions should include sockets */
+	SLATE_ATTRIBUTE(bool, IncludeSocketsForSuggestions)
+
+	SLATE_END_ARGS()
+
+	/** Constructs this widget with InArgs */
+	void Construct(const FArguments& InArgs, const class USkeleton* InSkeleton, TSharedPtr<class IPropertyHandle> ParentCurveProperty);
+
+private:
+
+	TWeakObjectPtr<USkeleton> Skeleton;
+
+	/** Get the search suggestions */
+	TArray<FString> GetCurveSearchSuggestions() const;
+};
