Fix crash in blend space editor

diff --git a/Engine/Source/Editor/Persona/Private/SAnimationBlendSpace.cpp b/Engine/Source/Editor/Persona/Private/SAnimationBlendSpace.cpp
--- a/Engine/Source/Editor/Persona/Private/SAnimationBlendSpace.cpp
+++ b/Engine/Source/Editor/Persona/Private/SAnimationBlendSpace.cpp
@@ -1006,7 +1006,10 @@ int32 SBlendSpaceGridWidget::OnPaint( const FPaintArgs& Args, const FGeometry& A
 			if (bPreviewOn)
 			{
 				TOptional<FVector2D> PreviewBSInput = GetWidgetPosFromEditorPos(PreviewInput.Get(), WindowRect);
-				DrawSamplePoint(PreviewBSInput.GetValue(), EBlendSpaceSampleState::Highlighted, AllottedGeometry, MyClippingRect, OutDrawElements, HighlightLayer);
+				if (PreviewBSInput.IsSet())
+				{
+					DrawSamplePoint(PreviewBSInput.GetValue(), EBlendSpaceSampleState::Highlighted, AllottedGeometry, MyClippingRect, OutDrawElements, HighlightLayer);
+				}
 			}
 		}
 		else if (bIsOnValidSample)
