Fix non-unity

diff --git a/Engine/Source/Editor/DetailCustomizations/Private/RigDetails.cpp b/Engine/Source/Editor/DetailCustomizations/Private/RigDetails.cpp
--- a/Engine/Source/Editor/DetailCustomizations/Private/RigDetails.cpp
+++ b/Engine/Source/Editor/DetailCustomizations/Private/RigDetails.cpp
@@ -3,6 +3,7 @@
 #include "DetailCustomizationsPrivatePCH.h"
 #include "RigDetails.h"
 #include "Animation/Rig.h"
+#include "ScopedTransaction.h"
 
 #define LOCTEXT_NAMESPACE	"RigDetails"
 
diff --git a/Engine/Source/Editor/Persona/Private/SRetargetManager.cpp b/Engine/Source/Editor/Persona/Private/SRetargetManager.cpp
--- a/Engine/Source/Editor/Persona/Private/SRetargetManager.cpp
+++ b/Engine/Source/Editor/Persona/Private/SRetargetManager.cpp
@@ -13,6 +13,7 @@
 #include "BoneSelectionWidget.h"
 #include "SRetargetSourceWindow.h"
 #include "SRigWindow.h"
+#include "Editor/AnimGraph/Classes/AnimPreviewInstance.h"
 
 #define LOCTEXT_NAMESPACE "SRetargetManager"
 
diff --git a/Engine/Source/Editor/Persona/Private/SRigWindow.cpp b/Engine/Source/Editor/Persona/Private/SRigWindow.cpp
--- a/Engine/Source/Editor/Persona/Private/SRigWindow.cpp
+++ b/Engine/Source/Editor/Persona/Private/SRigWindow.cpp
@@ -8,6 +8,7 @@
 #include "AssetRegistryModule.h"
 #include "Editor/PropertyEditor/Public/PropertyEditorModule.h"
 #include "Editor/ContentBrowser/Public/ContentBrowserModule.h"
+#include "WorkflowOrientedApp/SContentReference.h"
 #include "AssetNotifications.h"
 #include "Animation/Rig.h"
 #include "BoneSelectionWidget.h"
