TTP #342758 : Feature Request: PERSONA: Ability to reimport animations in asset browser with right click

diff --git a/Engine/Source/Editor/Persona/Private/SAnimationSequenceBrowser.cpp b/Engine/Source/Editor/Persona/Private/SAnimationSequenceBrowser.cpp
--- a/Engine/Source/Editor/Persona/Private/SAnimationSequenceBrowser.cpp
+++ b/Engine/Source/Editor/Persona/Private/SAnimationSequenceBrowser.cpp
@@ -146,6 +146,16 @@ TSharedPtr<SWidget> SAnimationSequenceBrowser::OnGetAssetContextMenu(const TArra
 				FCanExecuteAction()
 				)
 				);
+
+			MenuBuilder.AddMenuEntry(
+				LOCTEXT("ReimportAnimation", "Reimport Animation"),
+				LOCTEXT("ReimportAnimation_ToolTip", "Reimport current animaion."),
+				FSlateIcon(),
+				FUIAction(
+				FExecuteAction::CreateSP(this, &SAnimationSequenceBrowser::OnReimportAnimation, SelectedAssets),
+				FCanExecuteAction()
+				)
+				);
 		}
 		MenuBuilder.EndSection();
 	}
@@ -283,6 +293,21 @@ void SAnimationSequenceBrowser::OnAddLoopingInterpolation(TArray<FAssetData> Sel
 	}
 }
 
+void SAnimationSequenceBrowser::OnReimportAnimation(TArray<FAssetData> SelectedAssets)
+{
+	if (SelectedAssets.Num() > 0)
+	{
+		TArray<TWeakObjectPtr<UAnimSequence>> AnimSequences;
+		for (auto Iter = SelectedAssets.CreateIterator(); Iter; ++Iter)
+		{
+			if (UAnimSequence* AnimSequence = Cast<UAnimSequence>(Iter->GetAsset()))
+			{
+				FReimportManager::Instance()->Reimport(AnimSequence, true);
+			}
+		}
+	}
+}
+
 void SAnimationSequenceBrowser::OnCreateCopy(TArray<FAssetData> Selected)
 {
 	if ( Selected.Num() > 0 )
diff --git a/Engine/Source/Editor/Persona/Private/SAnimationSequenceBrowser.h b/Engine/Source/Editor/Persona/Private/SAnimationSequenceBrowser.h
--- a/Engine/Source/Editor/Persona/Private/SAnimationSequenceBrowser.h
+++ b/Engine/Source/Editor/Persona/Private/SAnimationSequenceBrowser.h
@@ -62,6 +62,9 @@ public:
 	/** Delegate to handle Add looping interpolation context menu option */
 	void OnAddLoopingInterpolation(TArray<FAssetData> SelectedAssets);
 
+	/** Delegate to handle reimport selected animations */
+	void OnReimportAnimation(TArray<FAssetData> SelectedAssets);
+
 	/** This will allow duplicate the current object, and remap to new skeleton 
 	 *	Only allowed for AnimSequence 
 	 */
