Fixed TTP #338432 : In the retarget source manager, if you rename an asset and then undo a duplicate asset is added

diff --git a/Engine/Source/Editor/Persona/Private/SRetargetSourceWindow.cpp b/Engine/Source/Editor/Persona/Private/SRetargetSourceWindow.cpp
--- a/Engine/Source/Editor/Persona/Private/SRetargetSourceWindow.cpp
+++ b/Engine/Source/Editor/Persona/Private/SRetargetSourceWindow.cpp
@@ -289,7 +289,9 @@ void SRetargetSourceWindow::OnRenameCommit( const FName & OldName,  const FStrin
 			Skeleton->Modify();
 
 			Skeleton->AnimRetargetSources.Remove(OldName);
-			Skeleton->AnimRetargetSources.Add(NewPoseName, NewPose);
+			Skeleton->AnimRetargetSources.Add(NewPoseName, NewPose);		
+			// push renamed info to the stack
+			RenamedPoseNameStack.Push(FPoseNameForUndo(OldName, NewPoseName));
 
 			// need to verify if these pose is used by anybody else
 			FAssetRegistryModule& AssetRegistryModule = FModuleManager::LoadModuleChecked<FAssetRegistryModule>(TEXT("AssetRegistry"));
@@ -745,6 +747,30 @@ SRetargetSourceWindow::~SRetargetSourceWindow()
 
 void SRetargetSourceWindow::PostUndo()
 {
+	// handle renamed pose because currently adding old one without removing new one when the user did undo action
+	// As a result, both exist in the list so should remove new one 
+	if (RenamedPoseNameStack.Num() > 0 && Skeleton->AnimRetargetSources.Num() > 0)
+	{
+		FPoseNameForUndo& RenamedPoseName = RenamedPoseNameStack.Top();
+
+		int32 NumSources = Skeleton->AnimRetargetSources.Num();
+
+		FReferencePose * Pose = Skeleton->AnimRetargetSources.Find(RenamedPoseName.OldPoseName);
+		if (Pose)
+		{
+			FReferencePose * NewPose = Skeleton->AnimRetargetSources.Find(RenamedPoseName.NewPoseName);
+
+			// if both the old pose name and the new pose name exist 
+			if (NewPose)
+			{
+				Skeleton->Modify();
+				Skeleton->AnimRetargetSources.Remove(RenamedPoseName.NewPoseName);
+				// remove the last one because we handled it
+				RenamedPoseNameStack.Pop();
+			}
+		}
+	}
+
 	CreateRetargetSourceList();
 }
 
diff --git a/Engine/Source/Editor/Persona/Private/SRetargetSourceWindow.h b/Engine/Source/Editor/Persona/Private/SRetargetSourceWindow.h
--- a/Engine/Source/Editor/Persona/Private/SRetargetSourceWindow.h
+++ b/Engine/Source/Editor/Persona/Private/SRetargetSourceWindow.h
@@ -202,6 +202,20 @@ private:
 	/** Item to rename. Only valid for adding **/
 	mutable TSharedPtr<FDisplayedRetargetSourceInfo> ItemToRename;
 
+	/** data to handle renamed pose when undo action is executed */
+	struct FPoseNameForUndo
+	{
+		FName OldPoseName;
+		FName NewPoseName;
+
+		FPoseNameForUndo(FName OldName, FName NewName)
+			: OldPoseName(OldName), NewPoseName(NewName)
+		{}
+	};
+
+	/** stack for renamed poses for multiple undo actions **/
+	TArray<FPoseNameForUndo> RenamedPoseNameStack;
+
 	/** Rename committed. Called by Inline Widget **/
 	void OnRenameCommit( const FName & OldName,  const FString& NewName );
 	/** Verify Rename is legit or not **/
