Fixed "New Notify" context menus not showing classes derived more than one step away from the base notify or notify state classes.
#ttp 343214 - ANIM: Extended AnimNotify doesn't show up

diff --git a/Engine/Source/Editor/Persona/Private/SAnimNotifyPanel.cpp b/Engine/Source/Editor/Persona/Private/SAnimNotifyPanel.cpp
--- a/Engine/Source/Editor/Persona/Private/SAnimNotifyPanel.cpp
+++ b/Engine/Source/Editor/Persona/Private/SAnimNotifyPanel.cpp
@@ -213,6 +213,8 @@ public:
 		SLATE_ARGUMENT( FLinearColor, TrackColor )
 		SLATE_EVENT(FOnTrackSelectionChanged, OnSelectionChanged)
 		SLATE_EVENT( FOnUpdatePanel, OnUpdatePanel )
+		SLATE_EVENT( FOnGetBlueprintNotifyData, OnGetNotifyBlueprintData )
+		SLATE_EVENT( FOnGetBlueprintNotifyData, OnGetNotifyStateBlueprintData )
 		SLATE_EVENT( FOnGetScrubValue, OnGetScrubValue )
 		SLATE_EVENT( FOnGetDraggedNodePos, OnGetDraggedNodePos )
 		SLATE_EVENT( FOnNotifyNodesDragStarted, OnNodeDragStarted )
@@ -280,6 +282,7 @@ public:
 	void RefreshMarqueeSelectedNodes(FSlateRect& Rect, FNotifyMarqueeOperation& Marquee);
 
 protected:
+
 	void CreateCommands();
 
 	// Build up a "New Notify..." menu
@@ -366,6 +369,16 @@ protected:
 
 private:
 
+	// Data structure for bluprint notify context menu entries
+	struct BlueprintNotifyMenuInfo
+	{
+		FString NotifyName;
+		FString BlueprintPath;
+	};
+
+	// Format notify asset data into the information needed for menu display
+	void GetNotifyMenuData(TArray<FAssetData>& NotifyAssetData, TArray<BlueprintNotifyMenuInfo>& OutNotifyMenuData);
+
 	// Store the tracks geometry for later use
 	void UpdateCachedGeometry(const FGeometry& InGeometry) {CachedGeometry = InGeometry;}
 
@@ -397,6 +410,8 @@ protected:
 	int32									TrackIndex;
 	FOnTrackSelectionChanged				OnSelectionChanged;
 	FOnUpdatePanel							OnUpdatePanel;
+	FOnGetBlueprintNotifyData				OnGetNotifyBlueprintData;
+	FOnGetBlueprintNotifyData				OnGetNotifyStateBlueprintData;
 	FOnGetScrubValue						OnGetScrubValue;
 	FOnGetDraggedNodePos					OnGetDraggedNodePos;
 	FOnNotifyNodesDragStarted				OnNodeDragStarted;
@@ -469,6 +484,8 @@ public:
 	SLATE_EVENT( FOnGetScrubValue, OnGetScrubValue )
 	SLATE_EVENT( FOnGetDraggedNodePos, OnGetDraggedNodePos )
 	SLATE_EVENT( FOnUpdatePanel, OnUpdatePanel )
+	SLATE_EVENT( FOnGetBlueprintNotifyData, OnGetNotifyBlueprintData )
+	SLATE_EVENT( FOnGetBlueprintNotifyData, OnGetNotifyStateBlueprintData )
 	SLATE_EVENT( FOnNotifyNodesDragStarted, OnNodeDragStarted )
 	SLATE_EVENT( FRefreshOffsetsRequest, OnRequestRefreshOffsets )
 	SLATE_EVENT( FDeleteNotify, OnDeleteNotify )
@@ -1463,6 +1480,8 @@ void SAnimNotifyTrack::Construct(const FArguments& InArgs)
 	OnSelectionChanged = InArgs._OnSelectionChanged;
 	AnimNotifies = InArgs._AnimNotifies;
 	OnUpdatePanel = InArgs._OnUpdatePanel;
+	OnGetNotifyBlueprintData = InArgs._OnGetNotifyBlueprintData;
+	OnGetNotifyStateBlueprintData = InArgs._OnGetNotifyStateBlueprintData;
 	TrackIndex = InArgs._TrackIndex;
 	OnGetScrubValue = InArgs._OnGetScrubValue;
 	OnGetDraggedNodePos = InArgs._OnGetDraggedNodePos;
@@ -1656,36 +1675,24 @@ FReply SAnimNotifyTrack::OnMouseWheel(const FGeometry& MyGeometry, const FPointe
 
 void SAnimNotifyTrack::FillNewNotifyStateMenu(FMenuBuilder& MenuBuilder)
 {
-	// Load the asset registry module
-	FAssetRegistryModule& AssetRegistryModule = FModuleManager::LoadModuleChecked<FAssetRegistryModule>(TEXT("AssetRegistry"));
+	const static FText Description = LOCTEXT("AddsAnExistingAnimNotify", "Add an existing notify");
 
-	// Collect a full list of assets with the specified class
-	TArray<FAssetData> AssetData;
-	AssetRegistryModule.Get().GetAssetsByClass(UBlueprint::StaticClass()->GetFName(), AssetData);
+	TArray<FAssetData> NotifyAssetData;
+	TArray<BlueprintNotifyMenuInfo> NotifyMenuData;
+	OnGetNotifyStateBlueprintData.ExecuteIfBound(NotifyAssetData);
+	GetNotifyMenuData(NotifyAssetData, NotifyMenuData);
 
-	const FName BPParentClassName( TEXT( "ParentClass" ) );
-	const FString BPAnimNotify( TEXT("Class'/Script/Engine.AnimNotifyState'" ));
-
-	for (int32 AssetIndex = 0; AssetIndex < AssetData.Num(); ++AssetIndex)
+	for(BlueprintNotifyMenuInfo& NotifyData : NotifyMenuData)
 	{
-		FString TagValue = AssetData[ AssetIndex ].TagsAndValues.FindRef(BPParentClassName);
-		if(TagValue == BPAnimNotify)
-		{
-			FString BlueprintPath = AssetData[AssetIndex].ObjectPath.ToString();
-			FString Label = AssetData[AssetIndex].AssetName.ToString();
-			Label = Label.Replace(TEXT("AnimNotifyState_"), TEXT(""), ESearchCase::CaseSensitive);
+		const FText LabelText = FText::FromString(NotifyData.NotifyName);
 
-			const FText Description = LOCTEXT("AddsAnExistingAnimNotify", "Add an existing notify");
-			const FText LabelText = FText::FromString( Label );
-
-			FUIAction UIAction;
-			UIAction.ExecuteAction.BindRaw(
-				this, &SAnimNotifyTrack::CreateNewBlueprintNotifyAtCursor,
-				Label,
-				BlueprintPath);
+		FUIAction UIAction;
+		UIAction.ExecuteAction.BindRaw(
+			this, &SAnimNotifyTrack::CreateNewBlueprintNotifyAtCursor,
+			NotifyData.NotifyName,
+			NotifyData.BlueprintPath);
 
-			MenuBuilder.AddMenuEntry(LabelText, Description, FSlateIcon(), UIAction);
-		}
+		MenuBuilder.AddMenuEntry(LabelText, Description, FSlateIcon(), UIAction);
 	}
 
 	MenuBuilder.BeginSection("NativeNotifyStates", LOCTEXT("NewStateNotifyMenu_Native", "Native Notify States"));
@@ -1716,36 +1723,24 @@ void SAnimNotifyTrack::FillNewNotifyStateMenu(FMenuBuilder& MenuBuilder)
 
 void SAnimNotifyTrack::FillNewNotifyMenu(FMenuBuilder& MenuBuilder)
 {
-	// Load the asset registry module
-	FAssetRegistryModule& AssetRegistryModule = FModuleManager::LoadModuleChecked<FAssetRegistryModule>(TEXT("AssetRegistry"));
+	const static FText Description = LOCTEXT("NewNotifySubMenu_ToolTip", "Add an existing notify");
 
-	// Collect a full list of assets with the specified class
-	TArray<FAssetData> AssetData;
-	AssetRegistryModule.Get().GetAssetsByClass(UBlueprint::StaticClass()->GetFName(), AssetData);
+	TArray<FAssetData> NotifyAssetData;
+	TArray<BlueprintNotifyMenuInfo> NotifyMenuData;
+	OnGetNotifyBlueprintData.ExecuteIfBound(NotifyAssetData);
+	GetNotifyMenuData(NotifyAssetData, NotifyMenuData);
 
-	const FName BPParentClassName( TEXT( "ParentClass" ) );
-	const FString BPAnimNotify( TEXT("Class'/Script/Engine.AnimNotify'" ));
-		
-	for (int32 AssetIndex = 0; AssetIndex < AssetData.Num(); ++AssetIndex)
+	for(BlueprintNotifyMenuInfo& NotifyData : NotifyMenuData)
 	{
-		FString TagValue = AssetData[ AssetIndex ].TagsAndValues.FindRef(BPParentClassName);
-		if(TagValue == BPAnimNotify)
-		{
-			FString BlueprintPath = AssetData[AssetIndex].ObjectPath.ToString();
-			FString Label = AssetData[AssetIndex].AssetName.ToString();
-			Label = Label.Replace(TEXT("AnimNotify_"), TEXT(""), ESearchCase::CaseSensitive);
+		const FText LabelText = FText::FromString( NotifyData.NotifyName );
 
-			const FText Description = LOCTEXT("NewNotifySubMenu_ToolTip", "Add an existing notify");
-			const FText LabelText = FText::FromString( Label );
-
-			FUIAction UIAction;
-			UIAction.ExecuteAction.BindRaw(
-				this, &SAnimNotifyTrack::CreateNewBlueprintNotifyAtCursor,
-				Label,
-				BlueprintPath);
-
-			MenuBuilder.AddMenuEntry(LabelText, Description, FSlateIcon(), UIAction);
-		}
+		FUIAction UIAction;
+		UIAction.ExecuteAction.BindRaw(
+			this, &SAnimNotifyTrack::CreateNewBlueprintNotifyAtCursor,
+			NotifyData.NotifyName,
+			NotifyData.BlueprintPath);
+		
+		MenuBuilder.AddMenuEntry(LabelText, Description, FSlateIcon(), UIAction);
 	}
 
 	MenuBuilder.BeginSection("NativeNotifies", LOCTEXT("NewNotifyMenu_Native", "Native Notifies"));
@@ -2785,6 +2780,25 @@ void SAnimNotifyTrack::RefreshMarqueeSelectedNodes(FSlateRect& Rect, FNotifyMarq
 	}
 }
 
+void SAnimNotifyTrack::GetNotifyMenuData(TArray<FAssetData>& NotifyAssetData, TArray<BlueprintNotifyMenuInfo>& OutNotifyMenuData)
+{
+	for(FAssetData& NotifyData : NotifyAssetData)
+	{
+		OutNotifyMenuData.AddZeroed();
+		BlueprintNotifyMenuInfo& MenuInfo = OutNotifyMenuData.Last();
+
+		MenuInfo.BlueprintPath = NotifyData.ObjectPath.ToString();
+		MenuInfo.NotifyName = NotifyData.AssetName.ToString();
+		MenuInfo.NotifyName = MenuInfo.NotifyName.Replace(TEXT("AnimNotify_"), TEXT(""), ESearchCase::CaseSensitive);
+		MenuInfo.NotifyName = MenuInfo.NotifyName.Replace(TEXT("AnimNotifyState_"), TEXT(""), ESearchCase::CaseSensitive);
+	}
+
+	OutNotifyMenuData.Sort([](const BlueprintNotifyMenuInfo& A, const BlueprintNotifyMenuInfo& B)
+	{
+		return A.NotifyName < B.NotifyName;
+	});
+}
+
 //////////////////////////////////////////////////////////////////////////
 // SSequenceEdTrack
 
@@ -2820,6 +2834,8 @@ void SNotifyEdTrack::Construct(const FArguments& InArgs)
 				.ViewInputMax(InArgs._ViewInputMax)
 				.OnSelectionChanged(InArgs._OnSelectionChanged)
 				.OnUpdatePanel(InArgs._OnUpdatePanel)
+				.OnGetNotifyBlueprintData(InArgs._OnGetNotifyBlueprintData)
+				.OnGetNotifyStateBlueprintData(InArgs._OnGetNotifyStateBlueprintData)
 				.OnGetScrubValue(InArgs._OnGetScrubValue)
 				.OnGetDraggedNodePos(InArgs._OnGetDraggedNodePos)
 				.OnNodeDragStarted(InArgs._OnNodeDragStarted)
@@ -2946,6 +2962,13 @@ void SAnimNotifyPanel::Construct(const FArguments& InArgs)
 	OnPropertyChangedHandle = FCoreDelegates::FOnObjectPropertyChanged::FDelegate::CreateSP(this, &SAnimNotifyPanel::OnPropertyChanged);
 	FCoreDelegates::OnObjectPropertyChanged.Add(OnPropertyChangedHandle);
 
+	// Base notify classes used to search asset data for children.
+	NotifyClassNames.Add(TEXT("Class'/Script/Engine.AnimNotify'"));
+	NotifyStateClassNames.Add(TEXT("Class'/Script/Engine.AnimNotifyState'"));
+
+	PopulateNotifyBlueprintClasses(NotifyClassNames);
+	PopulateNotifyBlueprintClasses(NotifyStateClassNames);
+
 	Update();
 }
 
@@ -3110,6 +3133,8 @@ void SAnimNotifyPanel::RefreshNotifyTracks()
 			.OnGetScrubValue(OnGetScrubValue)
 			.OnGetDraggedNodePos(this, &SAnimNotifyPanel::CalculateDraggedNodePos)
 			.OnUpdatePanel(this, &SAnimNotifyPanel::Update)
+			.OnGetNotifyBlueprintData(this, &SAnimNotifyPanel::OnGetNotifyBlueprintData, &NotifyClassNames)
+			.OnGetNotifyStateBlueprintData(this, &SAnimNotifyPanel::OnGetNotifyBlueprintData, &NotifyStateClassNames)
 			.OnSelectionChanged(this, &SAnimNotifyPanel::OnTrackSelectionChanged)
 			.OnNodeDragStarted(this, &SAnimNotifyPanel::OnNotifyNodeDragStarted)
 			.MarkerBars(MarkerBars)
@@ -3552,4 +3577,78 @@ void SAnimNotifyPanel::OnKeyboardFocusLost(const FKeyboardFocusEvent& InKeyboard
 	Marquee = FNotifyMarqueeOperation();
 }
 
+void SAnimNotifyPanel::PopulateNotifyBlueprintClasses(TArray<FString>& InOutAllowedClasses)
+{
+	TArray<FAssetData> TempArray;
+	OnGetNotifyBlueprintData(TempArray, &InOutAllowedClasses);
+}
+
+void SAnimNotifyPanel::OnGetNotifyBlueprintData(TArray<FAssetData>& OutNotifyData, TArray<FString>* InOutAllowedClassNames)
+{
+	// If we have nothing to seach with, early out
+	if(InOutAllowedClassNames == NULL || InOutAllowedClassNames->Num() == 0)
+	{
+		return;
+	}
+
+	TArray<FAssetData> AssetDataList;
+	TArray<FString> FoundClasses;
+
+	// Load the asset registry module
+	FAssetRegistryModule& AssetRegistryModule = FModuleManager::LoadModuleChecked<FAssetRegistryModule>(TEXT("AssetRegistry"));
+
+	// Collect a full list of assets with the specified class
+	AssetRegistryModule.Get().GetAssetsByClass(UBlueprint::StaticClass()->GetFName(), AssetDataList);
+
+	static const FName BPParentClassName(TEXT("ParentClass"));
+	static const FName BPGenClassName(TEXT("GeneratedClass"));
+
+	int32 BeginClassCount = InOutAllowedClassNames->Num();
+	int32 CurrentClassCount = -1;
+
+	while(BeginClassCount != CurrentClassCount)
+	{
+		BeginClassCount = InOutAllowedClassNames->Num();
+
+		for(int32 AssetIndex = 0; AssetIndex < AssetDataList.Num(); ++AssetIndex)
+		{
+			FAssetData& AssetData = AssetDataList[AssetIndex];
+			FString TagValue = AssetData.TagsAndValues.FindRef(BPParentClassName);
+
+			if(InOutAllowedClassNames->Contains(TagValue))
+			{
+				FString GenClass = AssetData.TagsAndValues.FindRef(BPGenClassName);
+
+				if(!OutNotifyData.Contains(AssetData))
+				{
+					// Output the assetdata and record it as found in this request
+					OutNotifyData.Add(AssetData);
+					FoundClasses.Add(GenClass);
+				}
+
+				if(!InOutAllowedClassNames->Contains(GenClass))
+				{
+					// Expand the class list to account for a new possible parent class found
+					InOutAllowedClassNames->Add(GenClass);
+				}
+			}
+		}
+
+		CurrentClassCount = InOutAllowedClassNames->Num();
+	}
+
+	// One less here to include the base class originally added
+	if(FoundClasses.Num() < InOutAllowedClassNames->Num() - 1)
+	{
+		// Less classes found, some may have been deleted or reparented
+		for(int32 ClassIndex = InOutAllowedClassNames->Num() - 1 ; ClassIndex >= 0 ; --ClassIndex)
+		{
+			if(!FoundClasses.Contains((*InOutAllowedClassNames)[ClassIndex]))
+			{
+				InOutAllowedClassNames->RemoveAt(ClassIndex);
+			}
+		}
+	}
+}
+
 #undef LOCTEXT_NAMESPACE
diff --git a/Engine/Source/Editor/Persona/Private/SAnimNotifyPanel.h b/Engine/Source/Editor/Persona/Private/SAnimNotifyPanel.h
--- a/Engine/Source/Editor/Persona/Private/SAnimNotifyPanel.h
+++ b/Engine/Source/Editor/Persona/Private/SAnimNotifyPanel.h
@@ -18,6 +18,7 @@ DECLARE_DELEGATE( FRefreshOffsetsRequest )
 DECLARE_DELEGATE( FDeleteNotify )
 DECLARE_DELEGATE( FDeselectAllNotifies )
 DECLARE_DELEGATE( FCopyNotifies )
+DECLARE_DELEGATE_OneParam( FOnGetBlueprintNotifyData, TArray<FAssetData>& )
 
 class SAnimNotifyNode;
 class SAnimNotifyTrack;
@@ -251,6 +252,18 @@ private:
 	// Binds the UI commands for this widget to delegates
 	void BindCommands();
 
+	/** Populates the given class array with all classes deriving from those originally present
+	 * @param InAssetsToSearch Assets to search to detect child classes
+	 * @param InOutAllowedClassNames Classes to allow, this will be expanded to cover all derived classes of those originally present
+	 */
+	void PopulateNotifyBlueprintClasses(TArray<FString>& InOutAllowedClasses);
+
+	/** Find blueprints matching allowed classes and all derived blueprints 
+	 * @param OutNotifyData Asset data matching allowed classes and their children
+	 * @param InOutAllowedClassNames Classes to allow, this will be expanded to cover all derived classes of those originally present
+	 */
+	void OnGetNotifyBlueprintData(TArray<FAssetData>& OutNotifyData, TArray<FString>* InOutAllowedClassNames);
+
 	/** Persona reference **/
 	TWeakPtr<FPersona> PersonaPtr;
 
@@ -259,4 +272,10 @@ private:
 
 	/** UI commands for this widget */
 	TSharedPtr<FUICommandList> UICommandList;
+
+	/** Classes that are known to be derived from blueprint notifies */
+	TArray<FString> NotifyClassNames;
+
+	/** Classes that are known to be derived from blueprint state notifies */
+	TArray<FString> NotifyStateClassNames;
 };
\ No newline at end of file
