- Retargeting
- Added "advanced" option

diff --git a/Engine/Source/Editor/Persona/Private/SRigWindow.cpp b/Engine/Source/Editor/Persona/Private/SRigWindow.cpp
--- a/Engine/Source/Editor/Persona/Private/SRigWindow.cpp
+++ b/Engine/Source/Editor/Persona/Private/SRigWindow.cpp
@@ -191,6 +191,7 @@ void SRigWindow::Construct(const FArguments& InArgs)
 {
 	PersonaPtr = InArgs._Persona;
 	Skeleton = NULL;
+	bDisplayAdvanced = false;
 
 	if ( PersonaPtr.IsValid() )
 	{
@@ -239,6 +240,17 @@ void SRigWindow::Construct(const FArguments& InArgs)
 				.AssetPickerSizeOverride(FVector2D(250, 700))
 				.InitialAssetViewType(EAssetViewType::List)
 			]
+
+			+SHorizontalBox::Slot()
+			.HAlign(HAlign_Right)
+			[
+				SNew(SButton)
+				.OnClicked(FOnClicked::CreateSP(this, &SRigWindow::OnToggleAdvanced))
+				.HAlign(HAlign_Center)
+				.VAlign(VAlign_Center)
+				.Text(this, &SRigWindow::GetAdvancedButtonText)
+				.ToolTipText(LOCTEXT("ToggleAdvanced_Tooltip", "Toggle Base/Advanced configuration"))
+			]
 		]
 
 		// now show bone mapping
@@ -325,25 +337,28 @@ void SRigWindow::CreateBoneMappingList( const FString& SearchText)
 			const FString & DisplayName = Node.DisplayName;
 			const FName & BoneName = Skeleton->GetRigBoneMapping(Name);
 
-			if ( bDoFiltering )
+			if (Node.bAdvanced == bDisplayAdvanced)
 			{
-				if (!Name.ToString().Contains( SearchText ))
-				{
-					continue; // Skip items that don't match our filter
-				}
-				if (!DisplayName.Contains( SearchText ))
+				if(bDoFiltering)
 				{
-					continue;
+					if(!Name.ToString().Contains(SearchText))
+					{
+						continue; // Skip items that don't match our filter
+					}
+					if(!DisplayName.Contains(SearchText))
+					{
+						continue;
+					}
+					if(!BoneName.ToString().Contains(SearchText))
+					{
+						continue;
+					}
 				}
-				if (!BoneName.ToString().Contains( SearchText ))
-				{
-					continue;
-				}
-			}
 
-			TSharedRef<FDisplayedBoneMappingInfo> Info = FDisplayedBoneMappingInfo::Make( Name, DisplayName, Skeleton );
+				TSharedRef<FDisplayedBoneMappingInfo> Info = FDisplayedBoneMappingInfo::Make(Name, DisplayName, Skeleton);
 
-			BoneMappingList.Add( Info );
+				BoneMappingList.Add(Info);
+			}
 		}
 	}
 
@@ -401,5 +416,23 @@ FName SRigWindow::GetBoneMapping(FName NodeName)
 	return Skeleton->GetRigBoneMapping(NodeName);
 }
 
+FReply SRigWindow::OnToggleAdvanced()
+{
+	bDisplayAdvanced = !bDisplayAdvanced;
+
+	CreateBoneMappingList();
+
+	return FReply::Handled();
+}
+
+FString SRigWindow::GetAdvancedButtonText() const
+{
+	if (bDisplayAdvanced)
+	{
+		return TEXT("Show Base");
+	}
+
+	return TEXT("Show Advanced");
+}
 #undef LOCTEXT_NAMESPACE
 
diff --git a/Engine/Source/Editor/Persona/Private/SRigWindow.h b/Engine/Source/Editor/Persona/Private/SRigWindow.h
--- a/Engine/Source/Editor/Persona/Private/SRigWindow.h
+++ b/Engine/Source/Editor/Persona/Private/SRigWindow.h
@@ -145,6 +145,9 @@ private:
 	/** Current text typed into NameFilterBox */
 	FText FilterText;
 
+	/** show advanced? */
+	bool bDisplayAdvanced;
+
 	/** Delegate for undo/redo transaction **/
 	void PostUndo();
 
@@ -161,4 +164,7 @@ private:
 
 	void OnBoneMappingChanged( FName NodeName, FName BoneName );
 	FName GetBoneMapping( FName NodeName );
+
+	FReply OnToggleAdvanced();
+	FString GetAdvancedButtonText() const;
 };
\ No newline at end of file
