Fixed filtering text issue with rig window

diff --git a/Engine/Source/Editor/Persona/Private/SRigWindow.cpp b/Engine/Source/Editor/Persona/Private/SRigWindow.cpp
--- a/Engine/Source/Editor/Persona/Private/SRigWindow.cpp
+++ b/Engine/Source/Editor/Persona/Private/SRigWindow.cpp
@@ -373,7 +373,7 @@ void SRigWindow::OnAssetSelected(UObject* Object)
 		const FScopedTransaction Transaction(LOCTEXT("RigAssetChanged", "Select Rig"));
 		Skeleton->Modify();
 		Skeleton->SetRigConfig(Cast<URig>(Object));
-		CreateBoneMappingList();
+		CreateBoneMappingList(FilterText.ToString());
 
 		FAssetNotifications::SkeletonNeedsToBeSaved(Skeleton);
 	}
@@ -400,7 +400,7 @@ SRigWindow::~SRigWindow()
 
 void SRigWindow::PostUndo()
 {
-	CreateBoneMappingList();
+	CreateBoneMappingList(FilterText.ToString());
 }
 
 void SRigWindow::OnBoneMappingChanged(FName NodeName, FName BoneName)
@@ -420,7 +420,7 @@ FReply SRigWindow::OnToggleAdvanced()
 {
 	bDisplayAdvanced = !bDisplayAdvanced;
 
-	CreateBoneMappingList();
+	CreateBoneMappingList(FilterText.ToString());
 
 	return FReply::Handled();
 }
