Grab hand for pannable tracks in Persona

diff --git a/Engine/Source/Editor/Persona/Private/SAnimCurvePanel.cpp b/Engine/Source/Editor/Persona/Private/SAnimCurvePanel.cpp
--- a/Engine/Source/Editor/Persona/Private/SAnimCurvePanel.cpp
+++ b/Engine/Source/Editor/Persona/Private/SAnimCurvePanel.cpp
@@ -137,6 +137,7 @@ protected:
 	// SWidget interface
 	virtual int32 OnPaint(const FPaintArgs& Args, const FGeometry& AllottedGeometry, const FSlateRect& MyClippingRect, FSlateWindowElementList& OutDrawElements, int32 LayerId, const FWidgetStyle& InWidgetStyle, bool bParentEnabled) const override;
 	virtual FReply OnMouseWheel(const FGeometry& MyGeometry, const FPointerEvent& MouseEvent);
+	virtual FCursorReply OnCursorQuery(const FGeometry& MyGeometry, const FPointerEvent& CursorEvent) const override;
 	// SWidget interface
 
 	// SCurveEditor interface
@@ -233,6 +234,16 @@ FReply SAnimCurveEd::OnMouseWheel(const FGeometry& MyGeometry, const FPointerEve
 	return FReply::Handled();
 }
 
+FCursorReply SAnimCurveEd::OnCursorQuery(const FGeometry& MyGeometry, const FPointerEvent& CursorEvent) const
+{
+	if (ViewMinInput.Get() > 0.f || ViewMaxInput.Get() < TimelineLength.Get())
+	{
+		return FCursorReply::Cursor(EMouseCursor::GrabHand);
+	}
+
+	return FCursorReply::Unhandled();
+}
+
 void SAnimCurveEd::Construct(const FArguments& InArgs)
 {
 	OnGetScrubValue = InArgs._OnGetScrubValue;
diff --git a/Engine/Source/Editor/Persona/Private/SAnimNotifyPanel.cpp b/Engine/Source/Editor/Persona/Private/SAnimNotifyPanel.cpp
--- a/Engine/Source/Editor/Persona/Private/SAnimNotifyPanel.cpp
+++ b/Engine/Source/Editor/Persona/Private/SAnimNotifyPanel.cpp
@@ -235,6 +235,7 @@ public:
 	virtual void Tick( const FGeometry& AllottedGeometry, const double InCurrentTime, const float InDeltaTime ) {UpdateCachedGeometry(AllottedGeometry);}
 	virtual int32 OnPaint(const FPaintArgs& Args, const FGeometry& AllottedGeometry, const FSlateRect& MyClippingRect, FSlateWindowElementList& OutDrawElements, int32 LayerId, const FWidgetStyle& InWidgetStyle, bool bParentEnabled) const override;
 	virtual FReply OnMouseWheel(const FGeometry& MyGeometry, const FPointerEvent& MouseEvent) override;
+	virtual FCursorReply OnCursorQuery(const FGeometry& MyGeometry, const FPointerEvent& CursorEvent) const override;
 	// End of SWidget interface
 
 	/** Returns the cached rendering geometry of this track */
@@ -1673,6 +1674,16 @@ FReply SAnimNotifyTrack::OnMouseWheel(const FGeometry& MyGeometry, const FPointe
 	return FReply::Handled();
 }
 
+FCursorReply SAnimNotifyTrack::OnCursorQuery(const FGeometry& MyGeometry, const FPointerEvent& CursorEvent) const
+{
+	if (ViewInputMin.Get() > 0.f || ViewInputMax.Get() < Sequence->SequenceLength)
+	{
+		return FCursorReply::Cursor(EMouseCursor::GrabHand);
+	}
+
+	return FCursorReply::Unhandled();
+}
+
 void SAnimNotifyTrack::FillNewNotifyStateMenu(FMenuBuilder& MenuBuilder)
 {
 	const static FText Description = LOCTEXT("AddsAnExistingAnimNotify", "Add an existing notify");
diff --git a/Engine/Source/Editor/Persona/Private/STrack.cpp b/Engine/Source/Editor/Persona/Private/STrack.cpp
--- a/Engine/Source/Editor/Persona/Private/STrack.cpp
+++ b/Engine/Source/Editor/Persona/Private/STrack.cpp
@@ -777,6 +777,11 @@ FCursorReply STrack::OnCursorQuery( const FGeometry& MyGeometry, const FPointerE
 		return FCursorReply::Cursor( EMouseCursor::ResizeLeftRight );
 	}
 
+	if (ViewInputMin.Get() > TrackMinValue.Get() || ViewInputMax.Get() < TrackMaxValue.Get())
+	{
+		return FCursorReply::Cursor(EMouseCursor::GrabHand);
+	}
+
 	return FCursorReply::Unhandled();
 }
 
