When adding sockets in Persona automatically select them in the viewport

#TTP 344032 - skeletal mesh interface: new sockets are not automatically selected after their creation (which can make the socket creation workflow a little less than ideal when making large numbers of them at once)

diff --git a/Engine/Source/Editor/Persona/Private/SSkeletonTree.cpp b/Engine/Source/Editor/Persona/Private/SSkeletonTree.cpp
--- a/Engine/Source/Editor/Persona/Private/SSkeletonTree.cpp
+++ b/Engine/Source/Editor/Persona/Private/SSkeletonTree.cpp
@@ -1757,6 +1757,9 @@ void SSkeletonTree::OnAddSocket()
 
 		TargetSkeleton->Sockets.Add( NewSocket );
 
+		FSelectedSocketInfo SocketInfo(NewSocket, true);
+		PersonaPtr.Pin()->SetSelectedSocket(SocketInfo, false);
+
 		CreateFromSkeleton( TargetSkeleton->GetBoneTree(), NewSocket );
 	}
 }
