Fixed Persona mouse hover previews being frozen, viewport visibility delegate provided.

diff --git a/Engine/Source/Editor/Persona/Private/SAnimationSequenceBrowser.cpp b/Engine/Source/Editor/Persona/Private/SAnimationSequenceBrowser.cpp
--- a/Engine/Source/Editor/Persona/Private/SAnimationSequenceBrowser.cpp
+++ b/Engine/Source/Editor/Persona/Private/SAnimationSequenceBrowser.cpp
@@ -818,6 +818,7 @@ void SAnimationSequenceBrowser::CreateAssetTooltipResources()
 	ViewportClient->SetRealtime(true);
 	ViewportClient->SetViewMode(VMI_Lit);
 	ViewportClient->ToggleOrbitCamera(true);
+	ViewportClient->VisibilityDelegate.BindSP(this, &SAnimationSequenceBrowser::IsToolTipPreviewVisible);
 
 	// Add the scene viewport
 	ViewportWidget->SetViewportInterface(SceneViewport.ToSharedRef());
@@ -861,7 +862,7 @@ bool SAnimationSequenceBrowser::OnVisualizeAssetToolTip(const TSharedPtr<SWidget
 		float HalfFov = FMath::DegreesToRadians(ViewportClient->ViewFOV) / 2.0f;
 		float TargetDist = MeshToUse->Bounds.SphereRadius / FMath::Tan(HalfFov);
 
-		ViewportClient->SetViewRotation(FRotator(0.0f, 135.0f, 0.0f));
+		ViewportClient->SetViewRotation(FRotator(0.0f, -45.0f, 0.0f));
 		ViewportClient->SetViewLocationForOrbiting(FVector(0.0f, 0.0f, MeshToUse->Bounds.BoxExtent.Z / 2.0f), TargetDist);
 	}
 
@@ -895,6 +896,16 @@ void SAnimationSequenceBrowser::Tick(const FGeometry& AllottedGeometry, const do
 	}
 }
 
+bool SAnimationSequenceBrowser::IsToolTipPreviewVisible()
+{
+	bool bVisible = false;
+	if(ViewportWidget.IsValid())
+	{
+		bVisible = ViewportWidget->GetVisibility() == EVisibility::Visible;
+	}
+	return bVisible;
+}
+
 FAnimationAssetViewportClient::FAnimationAssetViewportClient(FPreviewScene& InPreviewScene) : FEditorViewportClient(GLevelEditorModeTools(), &InPreviewScene)
 {
 	SetViewMode(VMI_Lit);
diff --git a/Engine/Source/Editor/Persona/Private/SAnimationSequenceBrowser.h b/Engine/Source/Editor/Persona/Private/SAnimationSequenceBrowser.h
--- a/Engine/Source/Editor/Persona/Private/SAnimationSequenceBrowser.h
+++ b/Engine/Source/Editor/Persona/Private/SAnimationSequenceBrowser.h
@@ -140,6 +140,11 @@ protected:
 	 */
 	void CleanupPreviewSceneComponent(USceneComponent* Component);
 
+	/**
+	 * Has the animation asset preview tool tip been visualized
+	 */
+	bool IsToolTipPreviewVisible();
+
 	/**
 	 * The actual viewport widget
 	 */	
