Dont mark the skeleton as dirty when we are setting the preview mesh on loading. Stops the skeleton becoming dirty for no obvious reason.

diff --git a/Engine/Source/Editor/Persona/Private/Persona.cpp b/Engine/Source/Editor/Persona/Private/Persona.cpp
--- a/Engine/Source/Editor/Persona/Private/Persona.cpp
+++ b/Engine/Source/Editor/Persona/Private/Persona.cpp
@@ -645,7 +645,7 @@ void FPersona::InitPersona(const EToolkitMode::Type Mode, const TSharedPtr< clas
 
 		if(!TargetSkeleton->GetPreviewMesh())
 		{
-			TargetSkeleton->SetPreviewMesh(InitMesh);
+			TargetSkeleton->SetPreviewMesh(InitMesh, false);
 		}
 	}
 	else if (TargetSkeleton)
