Remove "Highlight Origin" from Persona as it doesnt work and wont be fixed any time soon

diff --git a/Engine/Source/Editor/Persona/Private/AnimViewportShowCommands.cpp b/Engine/Source/Editor/Persona/Private/AnimViewportShowCommands.cpp
--- a/Engine/Source/Editor/Persona/Private/AnimViewportShowCommands.cpp
+++ b/Engine/Source/Editor/Persona/Private/AnimViewportShowCommands.cpp
@@ -6,7 +6,6 @@
 void FAnimViewportShowCommands::RegisterCommands()
 {
 	UI_COMMAND( ToggleGrid, "Grid", "Display Grid", EUserInterfaceActionType::ToggleButton, FInputGesture() );
-	UI_COMMAND( HighlightOrigin, "Highlight Origin", "Highlight the origin lines on the grid", EUserInterfaceActionType::ToggleButton, FInputGesture() );
 
 	UI_COMMAND( ToggleFloor, "Floor", "Display Floor", EUserInterfaceActionType::ToggleButton, FInputGesture() );
 	UI_COMMAND( ToggleSky, "Sky", "Display Sky", EUserInterfaceActionType::ToggleButton, FInputGesture() );
diff --git a/Engine/Source/Editor/Persona/Private/AnimViewportShowCommands.h b/Engine/Source/Editor/Persona/Private/AnimViewportShowCommands.h
--- a/Engine/Source/Editor/Persona/Private/AnimViewportShowCommands.h
+++ b/Engine/Source/Editor/Persona/Private/AnimViewportShowCommands.h
@@ -24,9 +24,6 @@ public:
 	/** Option to show grid */
 	TSharedPtr< FUICommandInfo > ToggleGrid;
 
-	/** Option to highlight the origin */
-	TSharedPtr< FUICommandInfo > HighlightOrigin;
-
 	/** Option to show floor */
 	TSharedPtr< FUICommandInfo > ToggleFloor;
 
diff --git a/Engine/Source/Editor/Persona/Private/AnimationEditorViewportClient.cpp b/Engine/Source/Editor/Persona/Private/AnimationEditorViewportClient.cpp
--- a/Engine/Source/Editor/Persona/Private/AnimationEditorViewportClient.cpp
+++ b/Engine/Source/Editor/Persona/Private/AnimationEditorViewportClient.cpp
@@ -279,18 +279,6 @@ bool FAnimationViewportClient::IsShowingGrid() const
 	return FEditorViewportClient::IsSetShowGridChecked();
 }
 
-void FAnimationViewportClient::OnToggleHighlightOrigin()
-{
-	DrawHelper.AxesLineThickness = 1.0f - DrawHelper.AxesLineThickness;
-
-	ConfigOption->SetHighlightOrigin( DrawHelper.AxesLineThickness != 0.0f );
-}
-
-bool FAnimationViewportClient::IsHighlightingOrigin() const
-{
-	return ( DrawHelper.AxesLineThickness != 0.0f );
-}
-
 void FAnimationViewportClient::OnToggleShowFloor()
 {
 	if (EditorFloorComp)
diff --git a/Engine/Source/Editor/Persona/Private/AnimationEditorViewportClient.h b/Engine/Source/Editor/Persona/Private/AnimationEditorViewportClient.h
--- a/Engine/Source/Editor/Persona/Private/AnimationEditorViewportClient.h
+++ b/Engine/Source/Editor/Persona/Private/AnimationEditorViewportClient.h
@@ -95,12 +95,6 @@ public:
 	/** Function to check whether grid is displayed or not */
 	bool IsShowingGrid() const;
 
-	/** Function to show/hide grid in the viewport */
-	void OnToggleHighlightOrigin();
-
-	/** Function to check whether grid is displayed or not */
-	bool IsHighlightingOrigin() const;
-
 	/** Function to show/hide floor in the viewport */
 	void OnToggleShowFloor();
 
diff --git a/Engine/Source/Editor/Persona/Private/SAnimViewportToolBar.cpp b/Engine/Source/Editor/Persona/Private/SAnimViewportToolBar.cpp
--- a/Engine/Source/Editor/Persona/Private/SAnimViewportToolBar.cpp
+++ b/Engine/Source/Editor/Persona/Private/SAnimViewportToolBar.cpp
@@ -551,7 +551,6 @@ void SAnimViewportToolBar::FillShowSceneMenu(FMenuBuilder& MenuBuilder) const
 	MenuBuilder.BeginSection("AnimViewportGrid", LOCTEXT("Viewport_GridLabel", "Grid"));
 	{
 		MenuBuilder.AddMenuEntry(Actions.ToggleGrid);
-		MenuBuilder.AddMenuEntry(Actions.HighlightOrigin);
 	}
 	MenuBuilder.EndSection();
 
diff --git a/Engine/Source/Editor/Persona/Private/SAnimationEditorViewport.cpp b/Engine/Source/Editor/Persona/Private/SAnimationEditorViewport.cpp
--- a/Engine/Source/Editor/Persona/Private/SAnimationEditorViewport.cpp
+++ b/Engine/Source/Editor/Persona/Private/SAnimationEditorViewport.cpp
@@ -675,13 +675,6 @@ void SAnimationEditorViewportTabBody::BindCommands()
 		FCanExecuteAction(),
 		FIsActionChecked::CreateSP(this, &SAnimationEditorViewportTabBody::IsShowingGrid));
 
-	// Highlight origin should be disabled if the grid isn't showing
-	CommandList.MapAction( 
-		ViewportShowMenuCommands.HighlightOrigin,
-		FExecuteAction::CreateSP(this, &SAnimationEditorViewportTabBody::OnHighlightOrigin),
-		FCanExecuteAction::CreateSP(this, &SAnimationEditorViewportTabBody::IsShowingGrid),
-		FIsActionChecked::CreateSP(this, &SAnimationEditorViewportTabBody::IsHighlightingOrigin));
-
 	CommandList.MapAction( 
 		ViewportShowMenuCommands.ToggleFloor,
 		FExecuteAction::CreateSP(this, &SAnimationEditorViewportTabBody::OnShowFloor),
@@ -914,18 +907,6 @@ bool SAnimationEditorViewportTabBody::IsShowingGrid() const
 	return AnimViewportClient->IsShowingGrid();
 }
 
-void SAnimationEditorViewportTabBody::OnHighlightOrigin()
-{
-	TSharedRef<FAnimationViewportClient> AnimViewportClient = StaticCastSharedRef<FAnimationViewportClient>(LevelViewportClient.ToSharedRef());
-	AnimViewportClient->OnToggleHighlightOrigin();
-}
-
-bool SAnimationEditorViewportTabBody::IsHighlightingOrigin() const
-{
-	TSharedRef<FAnimationViewportClient> AnimViewportClient = StaticCastSharedRef<FAnimationViewportClient>(LevelViewportClient.ToSharedRef());
-	return AnimViewportClient->IsHighlightingOrigin();
-}
-
 void SAnimationEditorViewportTabBody::OnShowFloor()
 {
 	TSharedRef<FAnimationViewportClient> AnimViewportClient = StaticCastSharedRef<FAnimationViewportClient>(LevelViewportClient.ToSharedRef());
diff --git a/Engine/Source/Editor/Persona/Private/SAnimationEditorViewport.h b/Engine/Source/Editor/Persona/Private/SAnimationEditorViewport.h
--- a/Engine/Source/Editor/Persona/Private/SAnimationEditorViewport.h
+++ b/Engine/Source/Editor/Persona/Private/SAnimationEditorViewport.h
@@ -256,12 +256,6 @@ private:
 	/** Function to show/hide grid in the viewport */
 	void OnShowGrid();
 
-	/** Function to show/hide the origin highlight */
-	void OnHighlightOrigin();
-
-	/** Function to query whether the origin highlighting is being displayed */
-	bool IsHighlightingOrigin() const;
-
 	/** Function to show/hide floor in the viewport */
 	void OnShowFloor();
 
