Fixed saving retarget pose when cloth exists

diff --git a/Engine/Source/Editor/Persona/Private/SRetargetManager.cpp b/Engine/Source/Editor/Persona/Private/SRetargetManager.cpp
--- a/Engine/Source/Editor/Persona/Private/SRetargetManager.cpp
+++ b/Engine/Source/Editor/Persona/Private/SRetargetManager.cpp
@@ -247,7 +247,7 @@ FReply SRetargetManager::OnSaveRetargetBasePose()
 			// get space bases and calculate local
 			const TArray<FTransform> & SpaceBases = PreviewMeshComp->SpaceBases;
 			// @todo check to see if skeleton vs preview mesh makes it different for missing bones
-			const FReferenceSkeleton& RefSkeleton = Skeleton->GetReferenceSkeleton();
+			const FReferenceSkeleton& RefSkeleton = PreviewMesh->RefSkeleton;
 			TArray<FTransform> & NewRetargetBasePose = PreviewMesh->RetargetBasePose;
 			// if you're using master pose component in preview, this won't work
 			check (PreviewMesh->RefSkeleton.GetNum() == SpaceBases.Num());
