Support preview mesh per asset.
This is only used when opening persona first time. It doesn't switch mesh whenever you open.

diff --git a/Engine/Source/Editor/DetailCustomizations/Private/AnimMontageSegmentDetails.cpp b/Engine/Source/Editor/DetailCustomizations/Private/AnimMontageSegmentDetails.cpp
--- a/Engine/Source/Editor/DetailCustomizations/Private/AnimMontageSegmentDetails.cpp
+++ b/Engine/Source/Editor/DetailCustomizations/Private/AnimMontageSegmentDetails.cpp
@@ -221,7 +221,7 @@ void SAnimationSegmentViewport::InitSkeleton()
 
 	if( PreviewComponent != NULL && Skeleton != NULL )
 	{
-		USkeletalMesh* PreviewMesh = Skeleton->GetPreviewMesh();
+		USkeletalMesh* PreviewMesh = Skeleton->GetAssetPreviewMesh(AnimSequence);
 		if (PreviewMesh)
 		{
 			UAnimSingleNodeInstance * Preview = PreviewComponent->PreviewInstance;
diff --git a/Engine/Source/Editor/Persona/Private/Persona.cpp b/Engine/Source/Editor/Persona/Private/Persona.cpp
--- a/Engine/Source/Editor/Persona/Private/Persona.cpp
+++ b/Engine/Source/Editor/Persona/Private/Persona.cpp
@@ -638,17 +638,30 @@ void FPersona::InitPersona(const EToolkitMode::Type Mode, const TSharedPtr< clas
 	// preview worlds (in UWorld::CleanupWorld), but we want this to stick around while Persona exists
 	PreviewComponent->AddToRoot();
 
+	bool bSetMesh = false;
+
 	// Set the mesh
 	if (InitMesh != NULL)
 	{
 		SetPreviewMesh(InitMesh);
+		bSetMesh = true;
 
 		if(!TargetSkeleton->GetPreviewMesh())
 		{
 			TargetSkeleton->SetPreviewMesh(InitMesh, false);
 		}
 	}
-	else if (TargetSkeleton)
+	else if (InitAnimationAsset != NULL)
+	{
+		USkeletalMesh * AssetMesh = InitAnimationAsset->GetPreviewMesh();
+		if (AssetMesh)
+		{
+			SetPreviewMesh(AssetMesh);
+			bSetMesh = true;
+		}
+	}
+
+	if (!bSetMesh && TargetSkeleton)
 	{
 		//If no preview mesh set, just find the first mesh that uses this skeleton
 		USkeletalMesh * PreviewMesh = TargetSkeleton->GetPreviewMesh(true);
diff --git a/Engine/Source/Editor/Persona/Private/SAnimationSequenceBrowser.cpp b/Engine/Source/Editor/Persona/Private/SAnimationSequenceBrowser.cpp
--- a/Engine/Source/Editor/Persona/Private/SAnimationSequenceBrowser.cpp
+++ b/Engine/Source/Editor/Persona/Private/SAnimationSequenceBrowser.cpp
@@ -157,6 +157,16 @@ TSharedPtr<SWidget> SAnimationSequenceBrowser::OnGetAssetContextMenu(const TArra
 				FCanExecuteAction()
 				)
 				);
+
+			MenuBuilder.AddMenuEntry(
+				LOCTEXT("SetCurrentPreviewMesh", "Set Current Preview Mesh"),
+				LOCTEXT("SetCurrentPreviewMesh_ToolTip", "Set current preview mesh to be used when previewed by this asset. This only applies when you open Persona using this asset."),
+				FSlateIcon(),
+				FUIAction(
+				FExecuteAction::CreateSP(this, &SAnimationSequenceBrowser::OnSetCurrentPreviewMesh, SelectedAssets),
+				FCanExecuteAction()
+				)
+				);
 		}
 		MenuBuilder.EndSection();
 	}
@@ -276,6 +286,27 @@ void SAnimationSequenceBrowser::OnExportToFBX(TArray<FAssetData> SelectedAssets)
 	}
 }
 
+void SAnimationSequenceBrowser::OnSetCurrentPreviewMesh(TArray<FAssetData> SelectedAssets)
+{
+	if(SelectedAssets.Num() > 0)
+	{
+		USkeletalMesh * PreviewMesh = PersonaPtr.Pin()->GetMesh();
+
+		if (PreviewMesh)
+		{
+			TArray<TWeakObjectPtr<UAnimSequence>> AnimSequences;
+			for(auto Iter = SelectedAssets.CreateIterator(); Iter; ++Iter)
+			{
+				UAnimationAsset * AnimAsset = Cast<UAnimationAsset>(Iter->GetAsset());
+				if (AnimAsset)
+				{
+					AnimAsset->SetPreviewMesh(PreviewMesh);
+				}
+			}
+		}
+	}
+}
+
 void SAnimationSequenceBrowser::OnAddLoopingInterpolation(TArray<FAssetData> SelectedAssets)
 {
 	if(SelectedAssets.Num() > 0)
@@ -849,7 +880,7 @@ bool SAnimationSequenceBrowser::OnVisualizeAssetToolTip(const TSharedPtr<SWidget
 
 		USkeleton* Skeleton = Asset->GetSkeleton();
 		
-		MeshToUse = Skeleton->GetPreviewMesh(true);
+		MeshToUse = Skeleton->GetAssetPreviewMesh(Asset);
 		check(MeshToUse);
 		if(PreviewComponent->SkeletalMesh != MeshToUse)
 		{
diff --git a/Engine/Source/Editor/Persona/Private/SAnimationSequenceBrowser.h b/Engine/Source/Editor/Persona/Private/SAnimationSequenceBrowser.h
--- a/Engine/Source/Editor/Persona/Private/SAnimationSequenceBrowser.h
+++ b/Engine/Source/Editor/Persona/Private/SAnimationSequenceBrowser.h
@@ -62,6 +62,9 @@ public:
 	/** Delegate to handle Add looping interpolation context menu option */
 	void OnAddLoopingInterpolation(TArray<FAssetData> SelectedAssets);
 
+	/** Delegate to handle set current preview mesh context menu option */
+	void OnSetCurrentPreviewMesh(TArray<FAssetData> SelectedAssets);
+
 	/** Delegate to handle reimport selected animations */
 	void OnReimportAnimation(TArray<FAssetData> SelectedAssets);
 
diff --git a/Engine/Source/Editor/UnrealEd/Private/EditorExporters.cpp b/Engine/Source/Editor/UnrealEd/Private/EditorExporters.cpp
--- a/Engine/Source/Editor/UnrealEd/Private/EditorExporters.cpp
+++ b/Engine/Source/Editor/UnrealEd/Private/EditorExporters.cpp
@@ -2501,7 +2501,7 @@ bool UAnimSequenceExporterFBX::ExportBinary( UObject* Object, const TCHAR* Type,
 {
 	UAnimSequence* AnimSequence = CastChecked<UAnimSequence>( Object );
 	USkeleton* AnimSkeleton = AnimSequence->GetSkeleton();
-	USkeletalMesh* PreviewMesh = AnimSkeleton->GetPreviewMesh(true);
+	USkeletalMesh* PreviewMesh = AnimSkeleton->GetAssetPreviewMesh(AnimSequence);
 
 	if (AnimSkeleton && PreviewMesh)
 	{
diff --git a/Engine/Source/Editor/UnrealEd/Private/Factories/ActorFactory.cpp b/Engine/Source/Editor/UnrealEd/Private/Factories/ActorFactory.cpp
--- a/Engine/Source/Editor/UnrealEd/Private/Factories/ActorFactory.cpp
+++ b/Engine/Source/Editor/UnrealEd/Private/Factories/ActorFactory.cpp
@@ -745,7 +745,7 @@ USkeletalMesh* UActorFactoryAnimationAsset::GetSkeletalMeshFromAsset( UObject* A
 	if( AnimationAsset != NULL )
 	{
 		// base it on preview skeletal mesh, just to have something
-		SkeletalMesh = AnimationAsset->GetSkeleton()? AnimationAsset->GetSkeleton()->GetPreviewMesh(true) : NULL;
+		SkeletalMesh = AnimationAsset->GetSkeleton()? AnimationAsset->GetSkeleton()->GetAssetPreviewMesh(AnimationAsset) : NULL;
 	}
 	else if( VertexAnimation != NULL )
 	{
diff --git a/Engine/Source/Editor/UnrealEd/Private/LevelEditorViewport.cpp b/Engine/Source/Editor/UnrealEd/Private/LevelEditorViewport.cpp
--- a/Engine/Source/Editor/UnrealEd/Private/LevelEditorViewport.cpp
+++ b/Engine/Source/Editor/UnrealEd/Private/LevelEditorViewport.cpp
@@ -642,7 +642,7 @@ static bool AttemptApplyObjToActor( UObject* ObjToUse, AActor* ActorToApplyTo, i
 
 						if(bShouldChangeMesh)
 						{
-							SkelComponent->SetSkeletalMesh(NeedsSkeleton->GetPreviewMesh(true));
+							SkelComponent->SetSkeletalMesh(NeedsSkeleton->GetAssetPreviewMesh(DroppedObjAsAnimationAsset));
 						}
 
 						if(DroppedObjAsAnimationAsset)
diff --git a/Engine/Source/Editor/UnrealEd/Private/ThumbnailHelpers.cpp b/Engine/Source/Editor/UnrealEd/Private/ThumbnailHelpers.cpp
--- a/Engine/Source/Editor/UnrealEd/Private/ThumbnailHelpers.cpp
+++ b/Engine/Source/Editor/UnrealEd/Private/ThumbnailHelpers.cpp
@@ -526,7 +526,7 @@ bool FAnimationSequenceThumbnailScene::SetAnimation(UAnimSequenceBase* InAnimati
 	{
 		if (USkeleton* Skeleton = InAnimation->GetSkeleton())
 		{
-			USkeletalMesh* PreviewSkeletalMesh = Skeleton->GetPreviewMesh(true);
+			USkeletalMesh* PreviewSkeletalMesh = Skeleton->GetAssetPreviewMesh(InAnimation);
 
 			PreviewActor->SkeletalMeshComponent->SetSkeletalMesh(PreviewSkeletalMesh);
 
@@ -650,7 +650,7 @@ bool FBlendSpaceThumbnailScene::SetBlendSpace(class UBlendSpaceBase* InBlendSpac
 	{
 		if (USkeleton* Skeleton = InBlendSpace->GetSkeleton())
 		{
-			USkeletalMesh* PreviewSkeletalMesh = Skeleton->GetPreviewMesh(true);
+			USkeletalMesh* PreviewSkeletalMesh = Skeleton->GetAssetPreviewMesh(InBlendSpace);
 
 			PreviewActor->SkeletalMeshComponent->SetSkeletalMesh(PreviewSkeletalMesh);
 
diff --git a/Engine/Source/Runtime/Engine/Classes/Animation/AnimationAsset.h b/Engine/Source/Runtime/Engine/Classes/Animation/AnimationAsset.h
--- a/Engine/Source/Runtime/Engine/Classes/Animation/AnimationAsset.h
+++ b/Engine/Source/Runtime/Engine/Classes/Animation/AnimationAsset.h
@@ -415,10 +415,20 @@ public:
 	ENGINE_API virtual void ReplaceReferredAnimations(const TMap<UAnimSequence*, UAnimSequence*>& ReplacementMap);	
 #endif
 
-#if WITH_EDITORONLY_DATA
+#if WITH_EDITOR
 	/** Information for thumbnail rendering */
 	UPROPERTY(VisibleAnywhere, Instanced, Category = Thumbnail)
 	class UThumbnailInfo* ThumbnailInfo;
+
+private:
+	/** The default skeletal mesh to use when previewing this asset - this only applies when you open Persona using this asset*/
+	UPROPERTY(duplicatetransient, AssetRegistrySearchable)
+	TAssetPtr<class USkeletalMesh> PreviewSkeletalMesh;
+
+public:
+	ENGINE_API void SetPreviewMesh(USkeletalMesh* PreviewMesh);
+	ENGINE_API USkeletalMesh* UAnimationAsset::GetPreviewMesh();
+
 #endif
 
 public:
diff --git a/Engine/Source/Runtime/Engine/Classes/Animation/Skeleton.h b/Engine/Source/Runtime/Engine/Classes/Animation/Skeleton.h
--- a/Engine/Source/Runtime/Engine/Classes/Animation/Skeleton.h
+++ b/Engine/Source/Runtime/Engine/Classes/Animation/Skeleton.h
@@ -346,6 +346,7 @@ public:
 	/** Returns the skeletons preview mesh, loading it if necessary */
 	ENGINE_API USkeletalMesh* GetPreviewMesh(bool bFindIfNotSet=false);
 	ENGINE_API USkeletalMesh* GetPreviewMesh() const;
+	ENGINE_API USkeletalMesh* GetAssetPreviewMesh(class UAnimationAsset* AnimAsset);
 
 	/** Returns the skeletons preview mesh, loading it if necessary */
 	ENGINE_API void SetPreviewMesh(USkeletalMesh* PreviewMesh, bool bMarkAsDirty=true);
diff --git a/Engine/Source/Runtime/Engine/Private/Animation/Skeleton.cpp b/Engine/Source/Runtime/Engine/Private/Animation/Skeleton.cpp
--- a/Engine/Source/Runtime/Engine/Private/Animation/Skeleton.cpp
+++ b/Engine/Source/Runtime/Engine/Private/Animation/Skeleton.cpp
@@ -763,6 +763,24 @@ USkeletalMesh* USkeleton::GetPreviewMesh() const
 	return PreviewSkeletalMesh.Get();
 }
 
+USkeletalMesh* USkeleton::GetAssetPreviewMesh(UAnimationAsset* AnimAsset) 
+{
+	USkeletalMesh * PreviewMesh = NULL;
+	// return asset preview asset
+	// if nothing is assigned, return skeleton asset
+	if (AnimAsset)
+	{
+		PreviewMesh = AnimAsset->GetPreviewMesh();
+	}
+
+	if (!PreviewMesh)
+	{
+		PreviewMesh = GetPreviewMesh(false);
+	}
+
+	return PreviewMesh;
+}
+
 void USkeleton::SetPreviewMesh(USkeletalMesh* PreviewMesh, bool bMarkAsDirty/*=true*/)
 {
 	if (bMarkAsDirty)
diff --git a/Engine/Source/Runtime/Engine/Private/AnimationAsset.cpp b/Engine/Source/Runtime/Engine/Private/AnimationAsset.cpp
--- a/Engine/Source/Runtime/Engine/Private/AnimationAsset.cpp
+++ b/Engine/Source/Runtime/Engine/Private/AnimationAsset.cpp
@@ -96,6 +96,34 @@ void UAnimationAsset::ReplaceReferredAnimations(const TMap<UAnimSequence*, UAnim
 {
 }
 
+USkeletalMesh* UAnimationAsset::GetPreviewMesh() 
+{
+	USkeletalMesh* PreviewMesh = PreviewSkeletalMesh.Get();
+	if(!PreviewMesh)
+	{
+		// if preview mesh isn't loaded, see if we have set
+		FStringAssetReference PreviewMeshStringRef = PreviewSkeletalMesh.ToStringReference();
+		// load it since now is the time to load
+		if(!PreviewMeshStringRef.ToString().IsEmpty())
+		{
+			PreviewMesh = Cast<USkeletalMesh>(StaticLoadObject(USkeletalMesh::StaticClass(), NULL, *PreviewMeshStringRef.ToString(), NULL, LOAD_None, NULL));
+			// if somehow skeleton changes, just nullify it. 
+			if (PreviewMesh && PreviewMesh->Skeleton != Skeleton)
+			{
+				PreviewMesh = NULL;
+				SetPreviewMesh(NULL);
+			}
+		}
+	}
+
+	return PreviewMesh;
+}
+
+void UAnimationAsset::SetPreviewMesh(USkeletalMesh* PreviewMesh)
+{
+	Modify();
+	PreviewSkeletalMesh = PreviewMesh;
+}
 #endif
 
 void UAnimationAsset::ValidateSkeleton()
