Fixing some dubious uses of slate that will be deprecated shortly.

diff --git a/Engine/Source/Editor/Kismet/Private/WorkflowOrientedApp/WorkflowTabManager.cpp b/Engine/Source/Editor/Kismet/Private/WorkflowOrientedApp/WorkflowTabManager.cpp
--- a/Engine/Source/Editor/Kismet/Private/WorkflowOrientedApp/WorkflowTabManager.cpp
+++ b/Engine/Source/Editor/Kismet/Private/WorkflowOrientedApp/WorkflowTabManager.cpp
@@ -256,17 +256,18 @@ TSharedRef< SWidget > FTabInfo::CreateHistoryNavigationWidget()
 {
 	if(!HistoryNavigationWidget.IsValid())
 	{
-		TWeakPtr< SMenuAnchor > MenuAnchorPtr;
+		TWeakPtr< SMenuAnchor > BackMenuAnchorPtr;
+		TWeakPtr< SMenuAnchor > FwdMenuAnchorPtr;
 		HistoryNavigationWidget = 
 			SNew(SHorizontalBox)
 			+SHorizontalBox::Slot()
 			.AutoWidth()
 			[
 				SNew(SBorder)
-				.OnMouseButtonDown(this, &FTabInfo::OnMouseDownHisory, MenuAnchorPtr)
+				.OnMouseButtonDown(this, &FTabInfo::OnMouseDownHisory, BackMenuAnchorPtr)
 				.BorderImage( FEditorStyle::GetBrush("NoBorder") )
 				[
-					SAssignNew(MenuAnchorPtr, SMenuAnchor)
+					SAssignNew(BackMenuAnchorPtr, SMenuAnchor)
 					.Placement( MenuPlacement_BelowAnchor )
 					.OnGetMenuContent( this, &FTabInfo::CreateHistoryMenu, true )
 					[
@@ -287,10 +288,10 @@ TSharedRef< SWidget > FTabInfo::CreateHistoryNavigationWidget()
 			.AutoWidth()
 			[
 				SNew(SBorder)
-				.OnMouseButtonDown(this, &FTabInfo::OnMouseDownHisory, MenuAnchorPtr)
+				.OnMouseButtonDown(this, &FTabInfo::OnMouseDownHisory, FwdMenuAnchorPtr)
 				.BorderImage( FEditorStyle::GetBrush("NoBorder") )
 				[
-					SAssignNew(MenuAnchorPtr, SMenuAnchor)
+					SAssignNew(FwdMenuAnchorPtr, SMenuAnchor)
 					.Placement( MenuPlacement_BelowAnchor )
 					.OnGetMenuContent( this, &FTabInfo::CreateHistoryMenu, false )
 					[
diff --git a/Engine/Source/Editor/Persona/Private/SAnimationSequenceBrowser.cpp b/Engine/Source/Editor/Persona/Private/SAnimationSequenceBrowser.cpp
--- a/Engine/Source/Editor/Persona/Private/SAnimationSequenceBrowser.cpp
+++ b/Engine/Source/Editor/Persona/Private/SAnimationSequenceBrowser.cpp
@@ -439,8 +439,8 @@ void SAnimationSequenceBrowser::Construct(const FArguments& InArgs)
 	Config.OnGetCustomAssetToolTip = FOnGetCustomAssetToolTip::CreateSP(this, &SAnimationSequenceBrowser::CreateCustomAssetToolTip);
 	Config.OnVisualizeAssetToolTip = FOnVisualizeAssetToolTip::CreateSP(this, &SAnimationSequenceBrowser::OnVisualizeAssetToolTip);
 
-	TWeakPtr< SMenuAnchor > MenuAnchorPtr;
-	
+	TWeakPtr< SMenuAnchor > BackMenuAnchorPtr;
+	TWeakPtr< SMenuAnchor > FwdMenuAnchorPtr;	
 	this->ChildSlot
 	[
 		SNew(SVerticalBox)
@@ -463,10 +463,10 @@ void SAnimationSequenceBrowser::Construct(const FArguments& InArgs)
 			.AutoWidth()
 			[
 				SNew(SBorder)
-				.OnMouseButtonDown(this, &SAnimationSequenceBrowser::OnMouseDownHisory, MenuAnchorPtr)
+				.OnMouseButtonDown(this, &SAnimationSequenceBrowser::OnMouseDownHisory, BackMenuAnchorPtr)
 				.BorderImage( FEditorStyle::GetBrush("NoBorder") )
 				[
-					SAssignNew(MenuAnchorPtr, SMenuAnchor)
+					SAssignNew(BackMenuAnchorPtr, SMenuAnchor)
 					.Placement( MenuPlacement_BelowAnchor )
 					.OnGetMenuContent( this, &SAnimationSequenceBrowser::CreateHistoryMenu, true )
 					[
@@ -487,10 +487,10 @@ void SAnimationSequenceBrowser::Construct(const FArguments& InArgs)
 			.AutoWidth()
 			[
 				SNew(SBorder)
-				.OnMouseButtonDown(this, &SAnimationSequenceBrowser::OnMouseDownHisory, MenuAnchorPtr)
+				.OnMouseButtonDown(this, &SAnimationSequenceBrowser::OnMouseDownHisory, FwdMenuAnchorPtr)
 				.BorderImage( FEditorStyle::GetBrush("NoBorder") )
 				[
-					SAssignNew(MenuAnchorPtr, SMenuAnchor)
+					SAssignNew(FwdMenuAnchorPtr, SMenuAnchor)
 					.Placement( MenuPlacement_BelowAnchor )
 					.OnGetMenuContent( this, &SAnimationSequenceBrowser::CreateHistoryMenu, false )
 					[
diff --git a/Engine/Source/Editor/UnrealEd/Private/EditorClassUtils.cpp b/Engine/Source/Editor/UnrealEd/Private/EditorClassUtils.cpp
--- a/Engine/Source/Editor/UnrealEd/Private/EditorClassUtils.cpp
+++ b/Engine/Source/Editor/UnrealEd/Private/EditorClassUtils.cpp
@@ -98,7 +98,7 @@ TSharedRef<SWidget> FEditorClassUtils::GetSourceLink(const UClass* Class, const
 
 		TWeakObjectPtr<UBlueprint> BlueprintPtr = Blueprint;
 
-		SAssignNew(SourceHyperlink, SHyperlink)
+		SourceHyperlink = SNew(SHyperlink)
 			.Style(FEditorStyle::Get(), "EditBPHyperlink")
 			.TextStyle(FEditorStyle::Get(), "DetailsView.EditBlueprintHyperlinkStyle")
 			.OnNavigate_Static(&Local::OnEditBlueprintClicked, BlueprintPtr, ObjectWeakPtr)
@@ -119,7 +119,7 @@ TSharedRef<SWidget> FEditorClassUtils::GetSourceLink(const UClass* Class, const
 				}
 			};
 
-			SAssignNew(SourceHyperlink, SHyperlink)
+			SourceHyperlink = SNew(SHyperlink)
 				.Style(FCoreStyle::Get(), "Hyperlink")
 				.TextStyle(FEditorStyle::Get(), "DetailsView.GoToCodeHyperlinkStyle")
 				.OnNavigate_Static(&Local::OnEditCodeClicked, ClassHeaderPath)
diff --git a/Engine/Source/Editor/UnrealEd/Private/Toolkits/SStandaloneAssetEditorToolkitHost.cpp b/Engine/Source/Editor/UnrealEd/Private/Toolkits/SStandaloneAssetEditorToolkitHost.cpp
--- a/Engine/Source/Editor/UnrealEd/Private/Toolkits/SStandaloneAssetEditorToolkitHost.cpp
+++ b/Engine/Source/Editor/UnrealEd/Private/Toolkits/SStandaloneAssetEditorToolkitHost.cpp
@@ -131,6 +131,8 @@ void SStandaloneAssetEditorToolkitHost::RestoreFromLayout( const TSharedRef<FTab
 	TSharedRef<SWidget> SuperSearchWidget = SuperSearchModule.MakeSearchBox(ExposedEditableTextBox);
 #endif
 	
+	MenuOverlayWidgetContent.Reset();
+	MenuWidgetContent.Reset();
 	this->ChildSlot
 	[
 		SNew( SVerticalBox )
