Improve Record Animation feature.

diff --git a/Engine/Source/Editor/EditorStyle/Private/SlateEditorStyle.cpp b/Engine/Source/Editor/EditorStyle/Private/SlateEditorStyle.cpp
--- a/Engine/Source/Editor/EditorStyle/Private/SlateEditorStyle.cpp
+++ b/Engine/Source/Editor/EditorStyle/Private/SlateEditorStyle.cpp
@@ -4370,6 +4370,12 @@ void FSlateEditorStyle::FStyle::SetupPersonaStyle()
 		Set("Persona.ExportToFBX.Small", new IMAGE_BRUSH(TEXT("Icons/icon_Persona_ExportFBX_40x"), Icon20x20));
 		Set("Persona.CreateAsset", new IMAGE_BRUSH(TEXT("Icons/icon_Persona_CreateAsset_40x"), Icon40x40));
 		Set("Persona.CreateAsset.Small", new IMAGE_BRUSH(TEXT("Icons/icon_Persona_CreateAsset_40x"), Icon20x20));
+		Set("Persona.StartRecordAnimation", new IMAGE_BRUSH(TEXT("Icons/icon_Persona_StartRecord_40x"), Icon40x40));
+		Set("Persona.StartRecordAnimation.Small", new IMAGE_BRUSH(TEXT("Icons/icon_Persona_StartRecord_40x"), Icon20x20));
+		Set("Persona.StopRecordAnimation", new IMAGE_BRUSH(TEXT("Icons/icon_Persona_StopRecord_40x"), Icon40x40));
+		Set("Persona.StopRecordAnimation.Small", new IMAGE_BRUSH(TEXT("Icons/icon_Persona_StopRecord_40x"), Icon20x20));
+		Set("Persona.StopRecordAnimation_Alt", new IMAGE_BRUSH(TEXT("Icons/icon_Persona_StopRecord_Alt_40x"), Icon40x40));
+		Set("Persona.StopRecordAnimation_Alt.Small", new IMAGE_BRUSH(TEXT("Icons/icon_Persona_StopRecord_Alt_40x"), Icon20x20));
 
 		// persona extras
 		Set("Persona.ConvertAnimationGraph", new IMAGE_BRUSH("Old/Graph/ConvertIcon", Icon40x40));
diff --git a/Engine/Source/Editor/Persona/Private/Persona.cpp b/Engine/Source/Editor/Persona/Private/Persona.cpp
--- a/Engine/Source/Editor/Persona/Private/Persona.cpp
+++ b/Engine/Source/Editor/Persona/Private/Persona.cpp
@@ -456,14 +456,13 @@ void FPersona::ExtendMenu()
 		static void AddPersonaFileMenu(FMenuBuilder& MenuBuilder)
 		{
 			// View
-			MenuBuilder.BeginSection("Persona", LOCTEXT("PersonaEditorMenu", "Animation"));
+			MenuBuilder.BeginSection("Persona", LOCTEXT("PersonaEditorMenu_File", "Blueprint"));
 			{
-
 			}
 			MenuBuilder.EndSection();
 		}
 
-		static void AddPersonaAssetMenu(FMenuBuilder& MenuBuilder)
+		static void AddPersonaAssetMenu(FMenuBuilder& MenuBuilder, FPersona* PersonaPtr)
 		{
 			// View
 			MenuBuilder.BeginSection("Persona", LOCTEXT("PersonaAssetMenuMenu_Skeleton", "Skeleton"));
@@ -479,10 +478,17 @@ void FPersona::ExtendMenu()
 			{
 				MenuBuilder.AddMenuEntry(FPersonaCommands::Get().ApplyCompression);
 				MenuBuilder.AddMenuEntry(FPersonaCommands::Get().ExportToFBX);
-				MenuBuilder.AddMenuEntry(FPersonaCommands::Get().RecordAnimation);
 				MenuBuilder.AddMenuEntry(FPersonaCommands::Get().AddLoopingInterpolation);
 			}
 			MenuBuilder.EndSection();
+
+			MenuBuilder.BeginSection("Persona", LOCTEXT("PersonaAssetMenuMenu_Record", "Record"));
+			{
+				MenuBuilder.AddMenuEntry(FPersonaCommands::Get().RecordAnimation, 
+					NAME_None, 
+					TAttribute<FText>(PersonaPtr, &FPersona::GetRecordMenuLabel));
+			}
+			MenuBuilder.EndSection();
 		}
 	};
 
@@ -518,7 +524,7 @@ void FPersona::ExtendMenu()
 			"AssetEditorActions",
 			EExtensionHook::After,
 			GetToolkitCommands(),
-			FMenuExtensionDelegate::CreateStatic(&Local::AddPersonaAssetMenu)
+			FMenuExtensionDelegate::CreateStatic(&Local::AddPersonaAssetMenu, this)
 			);
 
 	AddMenuExtender(MenuExtender);
@@ -768,6 +774,13 @@ void FPersona::ExtendDefaultPersonaToolbar()
 				ToolbarBuilder.AddToolBarButton(FPersonaCommands::Get().ImportAnimation, NAME_None, LOCTEXT("Toolbar_ImportAnimation", "Import"));
 				ToolbarBuilder.AddToolBarButton(FPersonaCommands::Get().ReimportAnimation, NAME_None, LOCTEXT("Toolbar_ReimportAnimation", "Reimport"));
 				ToolbarBuilder.AddToolBarButton(FPersonaCommands::Get().ExportToFBX, NAME_None, LOCTEXT("Toolbar_ExportToFBX", "Export"));
+
+				ToolbarBuilder.AddToolBarButton(FPersonaCommands::Get().RecordAnimation, 
+											NAME_None,
+											 TAttribute<FText>(PersonaPtr, &FPersona::GetRecordStatusLabel),
+											 TAttribute<FText>(PersonaPtr, &FPersona::GetRecordStatusTooltip),
+											 TAttribute<FSlateIcon>(PersonaPtr, &FPersona::GetRecordStatusImage),
+											 NAME_None);
 			}
 			ToolbarBuilder.EndSection();
 
@@ -970,7 +983,7 @@ void FPersona::CreateDefaultCommands()
 	// record animation
 	ToolkitCommands->MapAction( FPersonaCommands::Get().RecordAnimation,
 		FExecuteAction::CreateSP( this, &FPersona::RecordAnimation ),
-		FCanExecuteAction::CreateSP( this, &FPersona::CanRecordAnimation ),
+		FCanExecuteAction(),
 		FIsActionChecked(),
 		FIsActionButtonVisible::CreateSP( this, &FPersona::IsRecordAvailable )
 		);
@@ -2294,66 +2307,61 @@ void FPersona::RedoAction()
 	GEditor->RedoTransaction();
 }
 
+FSlateIcon FPersona::GetRecordStatusImage() const
+{
+	if (Recorder.InRecording())
+	{
+		return FSlateIcon(FEditorStyle::GetStyleSetName(), "Persona.StopRecordAnimation");
+	}
 
-// get record configuration
-bool GetRecordConfig( FString& AssetPath, FString& AssetName, float& Duration )
+	return FSlateIcon(FEditorStyle::GetStyleSetName(), "Persona.StartRecordAnimation");
+}
+
+FText FPersona::GetRecordMenuLabel() const
 {
-	TSharedRef<SCreateAnimationDlg> NewAnimDlg = 
-		SNew(SCreateAnimationDlg);
+	if(Recorder.InRecording())
+	{
+		return LOCTEXT("Persona_StopRecordAnimationMenuLabel", "Stop Record Animation");
+	}
 
-	if (NewAnimDlg->ShowModal() != EAppReturnType::Cancel)
+	return LOCTEXT("Persona_StartRecordAnimationLabel", "Start Record Animation");
+}
+
+FText FPersona::GetRecordStatusLabel() const
+{
+	if (Recorder.InRecording())
 	{
-		AssetPath = NewAnimDlg->GetFullAssetPath();
-		AssetName = NewAnimDlg->GetAssetName();
-		Duration = NewAnimDlg->GetDuration();
-		return true;
+		return LOCTEXT("Persona_StopRecordAnimationLabel", "Stop");
 	}
 
-	return false;
+	return LOCTEXT("Persona_StartRecordAnimationLabel", "Record");
 }
 
-void FPersona::RecordAnimation()
+FText FPersona::GetRecordStatusTooltip() const
 {
-	// options
-	// - whether restart anim tree?
-	// - sample rate for now is 30
+	if (Recorder.InRecording())
+	{
+		return LOCTEXT("Persona_StopRecordAnimation", "Stop Record Animation");
+	}
+
+	return LOCTEXT("Persona_StartRecordAnimation", "Start Record Animation");
+}
 
+void FPersona::RecordAnimation()
+{
 	if (!PreviewComponent || !PreviewComponent->SkeletalMesh)
 	{
 		// error
 		return;
 	}
 
-	// ask for path
-	FString AssetPath;
-	FString AssetName;
-	float Duration = 0.f;
-
-	if ( GetRecordConfig(AssetPath, AssetName, Duration) )
+	if (Recorder.InRecording())
 	{
-		// create the asset
-		UObject* 	Parent = CreatePackage(NULL, *AssetPath);
-		UObject* Object = LoadObject<UObject>(Parent, *AssetName, NULL, LOAD_None, NULL);
-		// if object with same name exists, warn user
-		if (Object)
-		{
-			FMessageDialog::Open( EAppMsgType::Ok, NSLOCTEXT("UnrealEd", "Error_AssetExist", "Asset with same name exists. Can't overwrite another asset") ) ;
-			return; // Move on to next sequence...
-		}
-
-		// If not, create new one now.
-		UAnimSequence * NewSeq = ConstructObject<UAnimSequence>( UAnimSequence::StaticClass(), Parent, *AssetName, RF_Public|RF_Standalone );
-		if (NewSeq)
-		{
-			// restart anim tree
-			//PreviewComponent->InitAnim(true);
-
-			// set skeleton
-			NewSeq->SetSkeleton ((PreviewComponent->SkeletalMesh)? PreviewComponent->SkeletalMesh->Skeleton : TargetSkeleton);
-			// Notify the asset registry
-			FAssetRegistryModule::AssetCreated(NewSeq);
-			Recorder.StartRecord(PreviewComponent, NewSeq, Duration);
-		}
+		Recorder.StopRecord(true);
+	}
+	else
+	{
+		Recorder.TriggerRecordAnimation(PreviewComponent);
 	}
 }
 
@@ -2680,11 +2688,6 @@ void FPersona::RemoveUnusedBones()
 	}
 }
 
-bool FPersona::IsAnimationBeingRecorded() const
-{
-	return (Recorder.GetAnimationObject()!=NULL);
-}
-
 void FPersona::Tick(float DeltaTime)
 {
 	FBlueprintEditor::Tick(DeltaTime);
@@ -2694,44 +2697,13 @@ void FPersona::Tick(float DeltaTime)
 		Viewport.Pin()->RefreshViewport();
 	}
 
-	if (IsAnimationBeingRecorded())
+	if (Recorder.InRecording())
 	{
 		// make sure you don't allow switch previewcomponent
-		if (Recorder.UpdateRecord(PreviewComponent, DeltaTime) == false)
-		{
-			// before stop record, save the animation, it will get cleared after stop recording
-			const UAnimSequence* NewAnim = Recorder.GetAnimationObject();
-			Recorder.StopRecord();
-
-			// notify to user
-			if (NewAnim)
-			{
-				const FText NotificationText = FText::Format( LOCTEXT("RecordAnimation", "'{0}' has been successfully recorded" ), FText::FromString( NewAnim->GetName() ) );
-
-				//This is not showing well in the Persona, so opening dialog first. 
-				//right now it will crash if you don't wait until end of the record, so it is important for users to know
-				//this is done
-/*				
-				FNotificationInfo Info(NotificationText);
-				Info.ExpireDuration = 3.0f;
-				Info.bUseLargeFont = false;
-				TSharedPtr<SNotificationItem> Notification = FSlateNotificationManager::Get().AddNotification(Info);
-				if ( Notification.IsValid() )
-				{
-					Notification->SetCompletionState( SNotificationItem::CS_Success );
-				}*/
-
-				FMessageDialog::Open(EAppMsgType::Ok, NotificationText);
-			}
-		}
+		Recorder.UpdateRecord(PreviewComponent, DeltaTime);
 	}
 }
 
-bool FPersona::CanRecordAnimation() const
-{
-	return (!IsAnimationBeingRecorded());
-}
-
 bool FPersona::CanChangeSkeletonPreviewMesh() const
 {
 	return PreviewComponent && PreviewComponent->SkeletalMesh;
@@ -2744,7 +2716,8 @@ bool FPersona::CanRemoveBones() const
 
 bool FPersona::IsRecordAvailable() const
 {
-	return (GetCurrentMode() == FPersonaModes::AnimBlueprintEditMode || GetCurrentMode() == FPersonaModes::AnimationEditMode );
+	// make sure mesh exists
+	return (PreviewComponent && PreviewComponent->SkeletalMesh);
 }
 
 bool FPersona::IsEditable(UEdGraph* InGraph) const
@@ -2855,6 +2828,11 @@ void FPersona::OnReimportAnimation()
 	}
 }
 
+TStatId FPersona::GetStatId() const
+{
+	RETURN_QUICK_DECLARE_CYCLE_STAT(FPersona, STATGROUP_Tickables);
+}
+
 static class FMeshHierarchyCmd : private FSelfRegisteringExec
 {
 public:
diff --git a/Engine/Source/Editor/Persona/Private/Persona.h b/Engine/Source/Editor/Persona/Private/Persona.h
--- a/Engine/Source/Editor/Persona/Private/Persona.h
+++ b/Engine/Source/Editor/Persona/Private/Persona.h
@@ -9,7 +9,7 @@
 #include "PersonaModule.h"
 #include "PersonaCommands.h"
 
-#include "Shared/AnimationRecorder.h"
+#include "AnimationRecorder.h"
 //////////////////////////////////////////////////////////////////////////
 // FPersona
 
@@ -193,6 +193,7 @@ public:
 
 	// FTickableEditorObject interface
 	virtual void Tick(float DeltaTime) override;
+	virtual TStatId GetStatId() const override;
 	// End of FTickableEditorObject interface
 
 	/** Returns the image brush to use for each modes dirty marker */
@@ -275,6 +276,9 @@ public:
 	mutable TWeakObjectPtr<UObject> SharedAnimAssetBeingEdited;
 	TWeakPtr<SDockTab> SharedAnimDocumentTab;
 
+	/** Animation recorder **/
+	FAnimationRecorder Recorder;
+
 public:
 	/** Viewport widget */
 	TWeakPtr<class SAnimationEditorViewportTabBody> Viewport;
@@ -535,14 +539,14 @@ protected:
 	TSharedPtr<class FPersonaToolbar> PersonaToolbar;
 
 private:
-	/** Animation recorder **/
-	FAnimationRecorder Recorder;
 
 	/** Recording animation functions **/
 	void RecordAnimation();
-	bool CanRecordAnimation() const;
 	bool IsRecordAvailable() const;
-	bool IsAnimationBeingRecorded() const;
+	FSlateIcon GetRecordStatusImage() const;
+	FText GetRecordStatusTooltip() const;
+	FText GetRecordStatusLabel() const;
+	FText GetRecordMenuLabel() const;
 
 	/** Animation menu functions **/
 	void OnApplyCompression();
diff --git a/Engine/Source/Editor/Persona/Private/SAnimationDlgs.cpp b/Engine/Source/Editor/Persona/Private/SAnimationDlgs.cpp
--- a/Engine/Source/Editor/Persona/Private/SAnimationDlgs.cpp
+++ b/Engine/Source/Editor/Persona/Private/SAnimationDlgs.cpp
@@ -102,30 +102,6 @@ void SCreateAnimationDlg::Construct(const FArguments& InArgs)
 							.MinDesiredWidth(250)
 						]
 					]
-
-					+SVerticalBox::Slot()
-					.AutoHeight()
-					.Padding(3)
-					[
-						SNew(SHorizontalBox)
-
-						+SHorizontalBox::Slot()
-						.AutoWidth()
-						.Padding(0, 0, 10, 0)
-						.VAlign(VAlign_Center)
-						[
-							SNew(STextBlock)
-							.Text(LOCTEXT("DurationInSec", "Duration (in sec) "))
-						]
-
-						+SHorizontalBox::Slot()
-						[
-							SNew(SEditableTextBox)
-							.Text(Duration)
-							.OnTextCommitted(this, &SCreateAnimationDlg::OnDurationChange)
-							.MinDesiredWidth(100)
-						]
-					]
 				]
 			]
 
@@ -158,11 +134,6 @@ void SCreateAnimationDlg::Construct(const FArguments& InArgs)
 		]);
 }
 
-void SCreateAnimationDlg::OnDurationChange(const FText& InDuration, ETextCommit::Type CommitInfo)
-{
-	Duration = InDuration;
-}
-
 void SCreateAnimationDlg::OnNameChange(const FText& NewName, ETextCommit::Type CommitInfo)
 {
 	AssetName = NewName;
@@ -205,12 +176,6 @@ bool SCreateAnimationDlg::ValidatePackage()
 		return false;
 	}
 
-	if (GetDuration() <= 0.0f)
-	{
-		FMessageDialog::Open(EAppMsgType::Ok, LOCTEXT("InvalidDuration", "Invalid duration"));
-		return false;
-	}
-
 	return true;
 }
 
@@ -235,11 +200,6 @@ FString SCreateAnimationDlg::GetFullAssetPath()
 	return AssetPath.ToString() + "/" + AssetName.ToString();
 }
 
-float SCreateAnimationDlg::GetDuration()
-{
-	return FCString::Atof(*Duration.ToString());
-}
-
 /////////////////////////////////////////////////////
 // create path picker for importing assets
 /////////////////////////////////////////////////////
diff --git a/Engine/Source/Editor/Persona/Private/SAnimationDlgs.h b/Engine/Source/Editor/Persona/Private/SAnimationDlgs.h
--- a/Engine/Source/Editor/Persona/Private/SAnimationDlgs.h
+++ b/Engine/Source/Editor/Persona/Private/SAnimationDlgs.h
@@ -32,13 +32,9 @@ public:
 	/** Gets the resulting full asset path (path+'/'+name) */
 	FString GetFullAssetPath();
 
-	/** Get Duration **/
-	float GetDuration();
-
 protected:
 	void OnPathChange(const FString& NewPath);
 	void OnNameChange(const FText& NewName, ETextCommit::Type CommitInfo);
-	void OnDurationChange(const FText& NewName, ETextCommit::Type CommitInfo);
 	FReply OnButtonClick(EAppReturnType::Type ButtonID);
 
 	bool ValidatePackage();
@@ -46,7 +42,6 @@ protected:
 	EAppReturnType::Type UserResponse;
 	FText AssetPath;
 	FText AssetName;
-	FText Duration;
 
 	static FText LastUsedAssetPath;
 };
diff --git a/Engine/Source/Editor/Persona/Private/SAnimationEditorViewport.cpp b/Engine/Source/Editor/Persona/Private/SAnimationEditorViewport.cpp
--- a/Engine/Source/Editor/Persona/Private/SAnimationEditorViewport.cpp
+++ b/Engine/Source/Editor/Persona/Private/SAnimationEditorViewport.cpp
@@ -339,10 +339,27 @@ void SAnimationEditorViewportTabBody::Construct(const FArguments& InArgs)
 				SNew(SButton)
 				.ButtonStyle(FEditorStyle::Get(), "NoBorder")
 				.Visibility(this, &SAnimationEditorViewportTabBody::GetViewportCornerTextVisibility)
-				.TextStyle(FEditorStyle::Get(), "Persona.Viewport.BlueprintDirtyText")
-				.Text(this, &SAnimationEditorViewportTabBody::GetViewportCornerText)
-				.ToolTipText(LOCTEXT("BlueprintStatusTooltip", "Shows the status of the animation blueprint.\nClick to recompile a dirty blueprint"))
 				.OnClicked(this, &SAnimationEditorViewportTabBody::ClickedOnViewportCornerText)
+				.Content()
+				[
+					SNew(SHorizontalBox)
+					+SHorizontalBox::Slot()
+					.AutoWidth()
+					[
+						SNew(SImage)
+						.Visibility(this, &SAnimationEditorViewportTabBody::GetViewportCornerImageVisibility)
+						.Image(this, &SAnimationEditorViewportTabBody::GetViewportCornerImage)
+					]
+						
+					+SHorizontalBox::Slot()
+					.FillWidth(1)
+					[
+						SNew(STextBlock)
+						.TextStyle(FEditorStyle::Get(), "Persona.Viewport.BlueprintDirtyText")
+						.Text(this, &SAnimationEditorViewportTabBody::GetViewportCornerText)
+						.ToolTipText(this, &SAnimationEditorViewportTabBody::GetViewportCornerTooltip)
+					]
+				]
 			]
 		]
 	];
@@ -1765,10 +1782,45 @@ bool SAnimationEditorViewportTabBody::IsSectionsDisplayMode(int32 DisplayMode) c
 }
 #endif // #if WITH_APEX_CLOTHING
 
+EVisibility SAnimationEditorViewportTabBody::GetViewportCornerImageVisibility() const
+{
+	TSharedPtr<FPersona> Persona = PersonaPtr.Pin();	
+	return Persona->Recorder.InRecording()? EVisibility::Visible : EVisibility::Collapsed;
+}
+
+const FSlateBrush * SAnimationEditorViewportTabBody::GetViewportCornerImage() const
+{
+	static int32 Count=0;
+
+	TSharedPtr<FPersona> Persona = PersonaPtr.Pin();
+//	if(Persona->Recorder.InRecording())
+	{
+		if(Count++ < 5)
+		{
+			return FEditorStyle::GetBrush("Persona.StopRecordAnimation");
+		}
+		else
+		{
+			if(Count == 10)
+			{
+				Count = 0;
+			}
+
+			return FEditorStyle::GetBrush("Persona.StopRecordAnimation_Alt");
+		}
+	}
+
+//	return NULL;
+}
+
 EVisibility SAnimationEditorViewportTabBody::GetViewportCornerTextVisibility() const
 {
 	TSharedPtr<FPersona> Persona = PersonaPtr.Pin();
-	if (Persona->IsModeCurrent(FPersonaModes::AnimBlueprintEditMode))
+	if (Persona->Recorder.InRecording())
+	{
+		return EVisibility::Visible;
+	}
+	else if (Persona->IsModeCurrent(FPersonaModes::AnimBlueprintEditMode))
 	{
 		if (UBlueprint* Blueprint = Persona->GetBlueprintObj())
 		{
@@ -1783,6 +1835,16 @@ EVisibility SAnimationEditorViewportTabBody::GetViewportCornerTextVisibility() c
 FText SAnimationEditorViewportTabBody::GetViewportCornerText() const
 {
 	TSharedPtr<FPersona> Persona = PersonaPtr.Pin();
+	if(Persona->Recorder.InRecording())
+	{
+		const FString & Name = Persona->Recorder.GetAnimationObject()->GetName();
+		float TimeRecorded = Persona->Recorder.GetTimeRecorded();
+		FNumberFormattingOptions NumberOption;
+		NumberOption.MaximumFractionalDigits = 2;
+		NumberOption.MinimumFractionalDigits = 2;
+		return FText::Format(LOCTEXT("AnimRecorder", "Recording '{0}' - Time {1} sec(s)]\nTo stop, click here. "),
+			FText::FromString(Name), FText::AsNumber(TimeRecorded, &NumberOption));
+	}
 	if (Persona->IsModeCurrent(FPersonaModes::AnimBlueprintEditMode))
 	{
 		if (UBlueprint* Blueprint = Persona->GetBlueprintObj())
@@ -1806,12 +1868,33 @@ FText SAnimationEditorViewportTabBody::GetViewportCornerText() const
 	return FText::GetEmpty();
 }
 
+FText SAnimationEditorViewportTabBody::GetViewportCornerTooltip() const
+{
+	TSharedPtr<FPersona> Persona = PersonaPtr.Pin();
+	if(Persona->Recorder.InRecording())
+	{
+		return LOCTEXT("RecordingStatusTooltip", "Shows the status of animation recording.\nClick to stop the recording.");
+	}
+	if(Persona->IsModeCurrent(FPersonaModes::AnimBlueprintEditMode))
+	{
+		return LOCTEXT("BlueprintStatusTooltip", "Shows the status of the animation blueprint.\nClick to recompile a dirty blueprint");
+	}
+
+return FText::GetEmpty();
+}
+
 FReply SAnimationEditorViewportTabBody::ClickedOnViewportCornerText()
 {
 	TSharedPtr<FPersona> Persona = PersonaPtr.Pin();
-	if (UBlueprint* Blueprint = Persona->GetBlueprintObj())
+	// if it's recording, it won't be able to see the message
+	// so disable it
+	if( Persona->Recorder.InRecording() )
+	{
+		Persona->Recorder.StopRecord(true);
+	}
+	else if(UBlueprint* Blueprint = Persona->GetBlueprintObj())
 	{
-		if (!Blueprint->IsUpToDate())
+		if(!Blueprint->IsUpToDate())
 		{
 			FKismetEditorUtilities::CompileBlueprint(Blueprint);
 		}
diff --git a/Engine/Source/Editor/Persona/Private/SAnimationEditorViewport.h b/Engine/Source/Editor/Persona/Private/SAnimationEditorViewport.h
--- a/Engine/Source/Editor/Persona/Private/SAnimationEditorViewport.h
+++ b/Engine/Source/Editor/Persona/Private/SAnimationEditorViewport.h
@@ -436,7 +436,11 @@ private:
 	/** Function to replace root translation */
 	void UpdateMeshRootTranslation(const FText& NewText, ETextCommit::Type CommitInfo, int32 Dimension);
 
+	EVisibility GetViewportCornerImageVisibility() const;
+	const FSlateBrush * GetViewportCornerImage() const;
+
 	EVisibility GetViewportCornerTextVisibility() const;
 	FText GetViewportCornerText() const;
+	FText GetViewportCornerTooltip() const;
 	FReply ClickedOnViewportCornerText();
 };
diff --git a/Engine/Source/Editor/Persona/Private/SAnimationSequenceBrowser.cpp b/Engine/Source/Editor/Persona/Private/SAnimationSequenceBrowser.cpp
--- a/Engine/Source/Editor/Persona/Private/SAnimationSequenceBrowser.cpp
+++ b/Engine/Source/Editor/Persona/Private/SAnimationSequenceBrowser.cpp
@@ -930,7 +930,9 @@ void SAnimationSequenceBrowser::Tick(const FGeometry& AllottedGeometry, const do
 bool SAnimationSequenceBrowser::IsToolTipPreviewVisible()
 {
 	bool bVisible = false;
-	if(ViewportWidget.IsValid())
+	// during persona recording, disable this
+	if( PersonaPtr.IsValid() && PersonaPtr.Pin()->Recorder.InRecording() == false 
+		&& ViewportWidget.IsValid())
 	{
 		bVisible = ViewportWidget->GetVisibility() == EVisibility::Visible;
 	}
diff --git a/Engine/Source/Editor/Persona/Private/Shared/AnimationRecorder.cpp b/Engine/Source/Editor/Persona/Private/Shared/AnimationRecorder.cpp
--- a/Engine/Source/Editor/Persona/Private/Shared/AnimationRecorder.cpp
+++ b/Engine/Source/Editor/Persona/Private/Shared/AnimationRecorder.cpp
@@ -3,13 +3,16 @@
 
 #include "PersonaPrivatePCH.h"
 #include "AnimationRecorder.h"
+#include "SAnimationDlgs.h"
+#include "Runtime/AssetRegistry/Public/AssetRegistryModule.h"
 
 #define LOCTEXT_NAMESPACE "FAnimationRecorder"
 
 /////////////////////////////////////////////////////
 
 FAnimationRecorder::FAnimationRecorder()
-	: SampleRate(30)
+	: IntervalTime(1.0f/DEFAULT_SAMPLERATE)
+	, MaxFrame(DEFAULT_SAMPLERATE * 60 * 10) // for now, set the max limit to 10 mins. 
 	, AnimationObject(NULL)
 {
 
@@ -17,13 +20,67 @@ FAnimationRecorder::FAnimationRecorder()
 
 FAnimationRecorder::~FAnimationRecorder()
 {
-	StopRecord();
+	StopRecord(false);
 }
 
-void FAnimationRecorder::StartRecord(USkeletalMeshComponent * Component, UAnimSequence * InAnimationObject, float InDuration)
+// get record configuration
+bool GetRecordConfig(FString& AssetPath, FString& AssetName)
+{
+	TSharedRef<SCreateAnimationDlg> NewAnimDlg =
+		SNew(SCreateAnimationDlg);
+
+	if(NewAnimDlg->ShowModal() != EAppReturnType::Cancel)
+	{
+		AssetPath = NewAnimDlg->GetFullAssetPath();
+		AssetName = NewAnimDlg->GetAssetName();
+		return true;
+	}
+
+	return false;
+}
+
+bool FAnimationRecorder::TriggerRecordAnimation(USkeletalMeshComponent * Component)
+{
+	// ask for path
+	FString AssetPath;
+	FString AssetName;
+
+	if (!Component || !Component->SkeletalMesh || !Component->SkeletalMesh->Skeleton)
+	{
+		return false;
+	}
+
+	if(GetRecordConfig(AssetPath, AssetName))
+	{
+		// create the asset
+		UObject* 	Parent = CreatePackage(NULL, *AssetPath);
+		UObject* Object = LoadObject<UObject>(Parent, *AssetName, NULL, LOAD_None, NULL);
+		// if object with same name exists, warn user
+		if(Object)
+		{
+			FMessageDialog::Open(EAppMsgType::Ok, NSLOCTEXT("UnrealEd", "Error_AssetExist", "Asset with same name exists. Can't overwrite another asset"));
+			return false; // Move on to next sequence...
+		}
+
+		// If not, create new one now.
+		UAnimSequence * NewSeq = ConstructObject<UAnimSequence>(UAnimSequence::StaticClass(), Parent, *AssetName, RF_Public|RF_Standalone);
+		if(NewSeq)
+		{
+			// set skeleton
+			NewSeq->SetSkeleton(Component->SkeletalMesh->Skeleton);
+			// Notify the asset registry
+			FAssetRegistryModule::AssetCreated(NewSeq);
+			StartRecord(Component, NewSeq);
+
+			return true;
+		}
+	}
+
+	return false;
+}
+
+void FAnimationRecorder::StartRecord(USkeletalMeshComponent * Component, UAnimSequence * InAnimationObject)
 {
-	Duration = InDuration;
-	check (Duration > 0.f);
 	TimePassed = 0.f;
 	AnimationObject = InAnimationObject;
 
@@ -33,18 +90,9 @@ void FAnimationRecorder::StartRecord(USkeletalMeshComponent * Component, UAnimSe
 
 	PreviousSpacesBases = Component->SpaceBases;
 
-	check (SampleRate!=0.f);
-	const float IntervalTime = 1.f/SampleRate;
-	// prepare to record
-	// how many frames we need?
-	int32 MaxFrame = Duration/IntervalTime+1;
-	int32 NumBones = PreviousSpacesBases.Num();
-	check (MaxFrame > 0);
-	check (Component->SkeletalMesh);
-
 	LastFrame = 0;
-	AnimationObject->SequenceLength = Duration;
-	AnimationObject->NumFrames = MaxFrame;
+	AnimationObject->SequenceLength = 0.f;
+	AnimationObject->NumFrames = 0;
 
 	USkeleton* AnimSkeleton = AnimationObject->GetSkeleton();
 	// add all frames
@@ -54,72 +102,94 @@ void FAnimationRecorder::StartRecord(USkeletalMeshComponent * Component, UAnimSe
 		int32 BoneTreeIndex = AnimSkeleton->GetSkeletonBoneIndexFromMeshBoneIndex(Component->SkeletalMesh, BoneIndex);
 		if (BoneTreeIndex != INDEX_NONE)
 		{
+			// add tracks for the bone existing
 			FName BoneTreeName = AnimSkeleton->GetReferenceSkeleton().GetBoneName(BoneTreeIndex);
-			int32 NewTrack = AnimationObject->RawAnimationData.AddZeroed(1);
 			AnimationObject->AnimationTrackNames.Add(BoneTreeName);
-
-			FRawAnimSequenceTrack& RawTrack = AnimationObject->RawAnimationData[NewTrack];
 			// add mapping to skeleton bone track
 			AnimationObject->TrackToSkeletonMapTable.Add(FTrackToSkeletonMap(BoneTreeIndex));
-			RawTrack.PosKeys.AddZeroed(MaxFrame);
-			RawTrack.RotKeys.AddZeroed(MaxFrame);
-			RawTrack.ScaleKeys.AddZeroed(MaxFrame);
 		}
 	}
 
+	int32 NumTracks = AnimationObject->AnimationTrackNames.Num();
+	// add tracks
+	AnimationObject->RawAnimationData.AddZeroed(NumTracks);
+
 	// record the first frame;
 	Record(Component, PreviousSpacesBases, 0);
 }
 
-void FAnimationRecorder::StopRecord()
+void FAnimationRecorder::StopRecord(bool bShowMessage)
 {
 	if (AnimationObject)
 	{
-		// if LastFrame doesn't match with NumFrames, we need to adjust data
-		if ( LastFrame + 1 < AnimationObject->NumFrames )
-		{
-			int32 NewMaxFrame = LastFrame + 1;
-			int32 OldMaxFrame = AnimationObject->NumFrames;
-			// Set New Frames/Time
-			AnimationObject->NumFrames = NewMaxFrame;
-			AnimationObject->SequenceLength = (float)(NewMaxFrame)/SampleRate;
+		int32 NumFrames = LastFrame  + 1;
+		AnimationObject->NumFrames = NumFrames;
 
-			for (int32 TrackId=0; TrackId <AnimationObject->RawAnimationData.Num(); ++TrackId)
+		// can't use TimePassed. That is just total time that has been passed, not necessarily match with frame count
+		AnimationObject->SequenceLength = NumFrames * IntervalTime;
+		AnimationObject->PostProcessSequence();
+		AnimationObject->MarkPackageDirty();
+
+		// notify to user
+		if (bShowMessage)
+		{
+			const FText NotificationText = FText::Format( LOCTEXT("RecordAnimation", "'{0}' has been successfully recorded [{1} frames : {2} sec(s)]" ), 
+				FText::FromString( AnimationObject->GetName() ),
+				FText::AsNumber( AnimationObject->NumFrames ),
+				FText::AsNumber( AnimationObject->SequenceLength ));
+
+			//This is not showing well in the Persona, so opening dialog first. 
+			//right now it will crash if you don't wait until end of the record, so it is important for users to know
+			//this is done
+					
+			FNotificationInfo Info(NotificationText);
+			Info.ExpireDuration = 5.0f;
+			Info.bUseLargeFont = false;
+			TSharedPtr<SNotificationItem> Notification = FSlateNotificationManager::Get().AddNotification(Info);
+			if ( Notification.IsValid() )
 			{
-				FRawAnimSequenceTrack& RawTrack = AnimationObject->RawAnimationData[TrackId];
-				// add mapping to skeleton bone track
-				RawTrack.PosKeys.RemoveAt(NewMaxFrame, OldMaxFrame-NewMaxFrame);
-				RawTrack.RotKeys.RemoveAt(NewMaxFrame, OldMaxFrame-NewMaxFrame);
-				RawTrack.ScaleKeys.RemoveAt(NewMaxFrame, OldMaxFrame-NewMaxFrame);
+				Notification->SetCompletionState( SNotificationItem::CS_Success );
 			}
+
+			FMessageDialog::Open(EAppMsgType::Ok, NotificationText);
 		}
 
-		AnimationObject->PostProcessSequence();
-		AnimationObject->MarkPackageDirty();
 		AnimationObject = NULL;
-
 		PreviousSpacesBases.Empty();
 	}
 }
 
 // return false if nothing to update
 // return true if it has properly updated
-bool FAnimationRecorder::UpdateRecord(USkeletalMeshComponent * Component, float DeltaTime)
+void FAnimationRecorder::UpdateRecord(USkeletalMeshComponent * Component, float DeltaTime)
 {
+	// if no animation object, return
+	if (!AnimationObject)
+	{
+		return;
+	}
+
 	const float MaxDelta = 1.f/10.f;
 	float UseDelta = FMath::Min<float>(DeltaTime, MaxDelta);
 
+	// if longer than 1 second, you'll see the halt on the animation
+	if (UseDelta <= 0.f)
+	{
+		return;
+	}
+
 	float PreviousTimePassed = TimePassed;
 	TimePassed += UseDelta;
 
-	check (SampleRate!=0.f);
-	const float IntervalTime = 1.f/SampleRate;
-
 	// time passed has been updated
 	// now find what frames we need to update
-	int32 MaxFrame = Duration/IntervalTime;
-	int32 FramesRecorded = (PreviousTimePassed / IntervalTime);
-	int32 FramesToRecord = FMath::Min<int32>(TimePassed / IntervalTime, MaxFrame);
+	int32 FramesRecorded = LastFrame;
+	int32 FramesToRecord = FPlatformMath::TruncToInt(TimePassed / IntervalTime);
+
+	if (FramesRecorded >= FramesToRecord)
+	{
+		return;
+	}
 
 	TArray<FTransform> & SpaceBases = Component->SpaceBases;
 
@@ -128,23 +198,23 @@ bool FAnimationRecorder::UpdateRecord(USkeletalMeshComponent * Component, float
 	TArray<FTransform> BlendedSpaceBases;
 	BlendedSpaceBases.AddZeroed(SpaceBases.Num());
 
+	UE_LOG(LogAnimation, Warning, TEXT("DeltaTime : %0.2f, Current Frame Count : %d, Frames To Record : %d, TimePassed : %0.2f"), DeltaTime
+		, FramesRecorded, FramesToRecord, TimePassed);
+
 	// if we need to record frame
 	while (FramesToRecord > FramesRecorded)
 	{
 		// find what frames we need to record
 		// convert to time
-		float NewTime = (FramesRecorded + 1) * IntervalTime;
-
-		float PercentageOfNew = (NewTime - PreviousTimePassed)/UseDelta;
-		//float PercentageOfOld = (TimePassed-NewTime)/UseDelta;
+		float CurrentTime = (FramesRecorded + 1) * IntervalTime;
+		float BlendAlpha = (CurrentTime - PreviousTimePassed)/UseDelta;
 
-		// make sure this is almost 1 = PercentageOfNew
-		//ensure ( FMath::IsNearlyEqual(PercentageOfNew + PercentageOfOld, 1.f, KINDA_SMALL_NUMBER));
+		UE_LOG(LogAnimation, Warning, TEXT("Current Frame Count : %d, BlendAlpha : %0.2f"), FramesRecorded + 1, BlendAlpha);
 
 		// for now we just concern component space, not skeleton space
 		for (int32 BoneIndex=0; BoneIndex<SpaceBases.Num(); ++BoneIndex)
 		{
-			BlendedSpaceBases[BoneIndex].Blend(PreviousSpacesBases[BoneIndex], SpaceBases[BoneIndex], PercentageOfNew);
+			BlendedSpaceBases[BoneIndex].Blend(PreviousSpacesBases[BoneIndex], SpaceBases[BoneIndex], BlendAlpha);
 		}
 
 		Record(Component, BlendedSpaceBases, FramesRecorded + 1);
@@ -154,8 +224,12 @@ bool FAnimationRecorder::UpdateRecord(USkeletalMeshComponent * Component, float
 	//save to current transform
 	PreviousSpacesBases = Component->SpaceBases;
 
-	// if we reached MaxFrame, return false;
-	return (MaxFrame > FramesRecorded);
+	// if we passed MaxFrame, just stop it
+	if (FramesRecorded > MaxFrame)
+	{
+		UE_LOG(LogAnimation, Warning, TEXT("Animation Recording exceeds the time limited (10 mins). Stopping recording animation... "))
+		StopRecord(true);
+	}
 }
 
 void FAnimationRecorder::Record( USkeletalMeshComponent * Component, TArray<FTransform> SpacesBases, int32 FrameToAdd )
@@ -178,9 +252,12 @@ void FAnimationRecorder::Record( USkeletalMeshComponent * Component, TArray<FTra
 				}
 
 				FRawAnimSequenceTrack& RawTrack = AnimationObject->RawAnimationData[TrackIndex];
-				RawTrack.PosKeys[FrameToAdd] = LocalTransform.GetTranslation();
-				RawTrack.RotKeys[FrameToAdd] = LocalTransform.GetRotation();
-				RawTrack.ScaleKeys[FrameToAdd] = LocalTransform.GetScale3D();
+				RawTrack.PosKeys.Add(LocalTransform.GetTranslation());
+				RawTrack.RotKeys.Add(LocalTransform.GetRotation());
+				RawTrack.ScaleKeys.Add(LocalTransform.GetScale3D());
+
+				// verification
+				check (FrameToAdd == RawTrack.PosKeys.Num()-1);
 			}
 		}
 
diff --git a/Engine/Source/Editor/Persona/Private/Shared/AnimationRecorder.h b/Engine/Source/Editor/Persona/Private/Shared/AnimationRecorder.h
deleted file mode 100644
diff --git a/Engine/Source/Editor/Persona/Public/AnimationRecorder.h b/Engine/Source/Editor/Persona/Public/AnimationRecorder.h
new file mode 100644
--- /dev/null
+++ b/Engine/Source/Editor/Persona/Public/AnimationRecorder.h
@@ -0,0 +1,35 @@
+// Copyright 1998-2014 Epic Games, Inc. All Rights Reserved.
+
+#pragma once
+
+//////////////////////////////////////////////////////////////////////////
+// FAnimationRecorder
+
+// records the mesh pose to animation input
+struct PERSONA_API FAnimationRecorder
+{
+private:
+	float IntervalTime;
+	int32 MaxFrame;
+	int32 LastFrame;
+	float TimePassed;
+	UAnimSequence * AnimationObject;
+	TArray<FTransform> PreviousSpacesBases;
+
+public:
+	FAnimationRecorder();
+	~FAnimationRecorder();
+
+	// also provides interface
+	bool TriggerRecordAnimation(USkeletalMeshComponent * Component);
+	void StartRecord(USkeletalMeshComponent * Component, UAnimSequence * InAnimationObject);
+	void StopRecord(bool bShowMessage);
+	void UpdateRecord(USkeletalMeshComponent * Component, float DeltaTime);
+	const UAnimSequence* GetAnimationObject() const { return AnimationObject; }
+	bool InRecording() const { return AnimationObject != NULL; }
+	float GetTimeRecorded() const { return TimePassed; }
+
+private:
+	void Record( USkeletalMeshComponent * Component, TArray<FTransform> SpacesBases, int32 FrameToAdd );
+};
+
