Check box for turning on/off skeleton retargeting

#TTP 349472 - UE4: Animaiton: Skeleton Tree should always have option to display retargeting information, in all tabs.

diff --git a/Engine/Source/Editor/Persona/Private/SSkeletonTree.cpp b/Engine/Source/Editor/Persona/Private/SSkeletonTree.cpp
--- a/Engine/Source/Editor/Persona/Private/SSkeletonTree.cpp
+++ b/Engine/Source/Editor/Persona/Private/SSkeletonTree.cpp
@@ -792,6 +792,7 @@ void SSkeletonTree::Construct(const FArguments& InArgs)
 
 	BoneFilter = EBoneFilter::All;
 	SocketFilter = ESocketFilter::Active;
+	bShowingRetargetingOptions = false;
 
 	PersonaPtr = InArgs._Persona;
 	IsEditable = InArgs._IsEditable;
@@ -811,20 +812,6 @@ void SSkeletonTree::Construct(const FArguments& InArgs)
 	FSkeletonTreeCommands::Register();
 	BindCommands();
 
-	TSharedRef<SHeaderRow> TreeHeaderRow = SNew(SHeaderRow)
-		+ SHeaderRow::Column(ColumnID_BoneLabel)
-		.DefaultLabel(LOCTEXT("SkeletonBoneNameLabel", "Name"))
-		.FillWidth(0.75f);
-
-	if (IsInSkeletonMode())
-	{
-		TreeHeaderRow->AddColumn(
-			SHeaderRow::Column(ColumnID_RetargetingLabel)
-			.DefaultLabel(LOCTEXT("SkeletonBoneTranslationRetargetingLabel", "Translation Retargeting"))
-			.FillWidth(0.25f)
-			);
-	}
-
 	this->ChildSlot
 	[
 		SNew( SVerticalBox )
@@ -878,31 +865,33 @@ void SSkeletonTree::Construct(const FArguments& InArgs)
 					.Text( this, &SSkeletonTree::GetSocketFilterMenuTitle )
 				]
 			]
+
+			+ SHorizontalBox::Slot()
+			.Padding(0.0f, 0.0f, 2.0f, 0.0f)
+			.VAlign(VAlign_Center)
+			.AutoWidth()
+			[
+				SNew(SCheckBox)
+				.IsChecked(this, &SSkeletonTree::IsShowingRetargetingOptions)
+				.ToolTipText(LOCTEXT("SocketFilterToolTip", "Change which types of sockets are shown"))
+				.OnCheckStateChanged(this, &SSkeletonTree::OnChangeShowingRetargetingOptions)
+				[
+					SNew(STextBlock)
+						.Text(LOCTEXT("ShowRetargetingOptions", "Show Retargeting Options").ToString())
+				]
+			]
 		]
 
 		+ SVerticalBox::Slot()
 		.Padding( FMargin( 0.0f, 4.0f, 0.0f, 0.0f ) )
 		[
-			SAssignNew( SkeletonTreeView, SMeshSkeletonTreeRowType )
-			.TreeItemsSource( &SkeletonRowList )
-			.OnGenerateRow( this, &SSkeletonTree::MakeTreeRowWidget )
-			.OnGetChildren( this, &SSkeletonTree::GetChildrenForInfo )
-			.OnContextMenuOpening( this, &SSkeletonTree::CreateContextMenu )
-			.OnSelectionChanged( this, &SSkeletonTree::OnSelectionChanged )
-			.OnItemScrolledIntoView( this, &SSkeletonTree::OnItemScrolledIntoView )
-			.OnMouseButtonDoubleClick( this, &SSkeletonTree::OnTreeDoubleClick)
-			.OnSetExpansionRecursive(this, &SSkeletonTree::SetTreeItemExpansionRecursive)
-			.ItemHeight( 24 )
-			.HeaderRow
-			(
-				TreeHeaderRow
-			)
+			SAssignNew(TreeHolder, SOverlay)
 		]
 	];
 
 	if (PersonaPtr.IsValid())
 	{
-		CreateFromSkeleton(TargetSkeleton->GetBoneTree());
+		CreateTreeColumns();
 	}
 }
 
@@ -1067,6 +1056,44 @@ bool SSkeletonTree::AttachToParent( TSharedRef<FDisplayedTreeRowInfo> ItemToAtta
 	return false;
 }
 
+void SSkeletonTree::CreateTreeColumns()
+{
+	TSharedRef<SHeaderRow> TreeHeaderRow = SNew(SHeaderRow)
+		+ SHeaderRow::Column(ColumnID_BoneLabel)
+		.DefaultLabel(LOCTEXT("SkeletonBoneNameLabel", "Name"))
+		.FillWidth(0.75f);
+
+	if (bShowingRetargetingOptions)
+	{
+		TreeHeaderRow->AddColumn(
+			SHeaderRow::Column(ColumnID_RetargetingLabel)
+			.DefaultLabel(LOCTEXT("SkeletonBoneTranslationRetargetingLabel", "Translation Retargeting"))
+			.FillWidth(0.25f)
+			);
+	}
+
+	TreeHolder->ClearChildren();
+	TreeHolder->AddSlot()
+		[
+			SAssignNew(SkeletonTreeView, SMeshSkeletonTreeRowType)
+			.TreeItemsSource(&SkeletonRowList)
+			.OnGenerateRow(this, &SSkeletonTree::MakeTreeRowWidget)
+			.OnGetChildren(this, &SSkeletonTree::GetChildrenForInfo)
+			.OnContextMenuOpening(this, &SSkeletonTree::CreateContextMenu)
+			.OnSelectionChanged(this, &SSkeletonTree::OnSelectionChanged)
+			.OnItemScrolledIntoView(this, &SSkeletonTree::OnItemScrolledIntoView)
+			.OnMouseButtonDoubleClick(this, &SSkeletonTree::OnTreeDoubleClick)
+			.OnSetExpansionRecursive(this, &SSkeletonTree::SetTreeItemExpansionRecursive)
+			.ItemHeight(24)
+			.HeaderRow
+			(
+			TreeHeaderRow
+			)
+		];
+
+	CreateFromSkeleton(TargetSkeleton->GetBoneTree());
+}
+
 void SSkeletonTree::CreateFromSkeleton( const TArray<FBoneNode>& SourceSkeleton, USkeletalMeshSocket* SocketToRename )
 {
 	SkeletonRowList.Empty();
@@ -1380,7 +1407,7 @@ TSharedPtr< SWidget > SSkeletonTree::CreateContextMenu()
 
 			MenuBuilder.EndSection();
 
-			if(IsInSkeletonMode())
+			if(bShowingRetargetingOptions)
 			{
 				MenuBuilder.BeginSection("SkeletonTreeBoneTranslationRetargeting", LOCTEXT("BoneTranslationRetargetingHeader", "Bone Translation Retargeting"));
 				{
@@ -2518,9 +2545,15 @@ void SSkeletonTree::AddAttachedAssets( const FPreviewAssetAttachContainer& Attac
 	}
 }
 
-bool SSkeletonTree::IsInSkeletonMode() const
+void SSkeletonTree::OnChangeShowingRetargetingOptions(ESlateCheckBoxState::Type NewState)
+{
+	bShowingRetargetingOptions = NewState == ESlateCheckBoxState::Checked;
+	CreateTreeColumns();
+}
+
+ESlateCheckBoxState::Type SSkeletonTree::IsShowingRetargetingOptions() const
 {
-	return PersonaPtr.Pin()->IsModeCurrent(FPersonaModes::SkeletonDisplayMode);
+	return bShowingRetargetingOptions ? ESlateCheckBoxState::Checked : ESlateCheckBoxState::Unchecked;
 }
 
 
diff --git a/Engine/Source/Editor/Persona/Private/SSkeletonTree.h b/Engine/Source/Editor/Persona/Private/SSkeletonTree.h
--- a/Engine/Source/Editor/Persona/Private/SSkeletonTree.h
+++ b/Engine/Source/Editor/Persona/Private/SSkeletonTree.h
@@ -380,6 +380,9 @@ public:
 
 	virtual ~SSkeletonTree();
 
+	/** Creates the tree control and then populates using CreateFromSkeleton */
+	void CreateTreeColumns();
+
 	/** Function to build the skeleton tree widgets from the source skeleton tree */
 	void CreateFromSkeleton( const TArray<FBoneNode>& SourceSkeleton, USkeletalMeshSocket* SocketToRename = NULL );
 
@@ -530,6 +533,12 @@ private:
 	/** We can only add sockets in Active, Skeleton or All mode (otherwise they just disappear) */
 	bool IsAddingSocketsAllowed() const;
 
+	/** Handler for "Show Retargeting Options" check box IsChecked functionality */
+	ESlateCheckBoxState::Type IsShowingRetargetingOptions() const;
+
+	/**  Handler for when we change the "Show Retargeting Options" check box */
+	void OnChangeShowingRetargetingOptions(ESlateCheckBoxState::Type NewState);
+
 	/** This replicates the socket filter to the previewcomponent so that the viewport can use the same settings */
 	void SetPreviewComponentSocketFilter() const;
 
@@ -551,9 +560,6 @@ private:
 	/** Deletes a set of attached objects from a FPreviewAssetAttachContainer and notifies Persona*/
 	void DeleteAttachedObjects( FPreviewAssetAttachContainer& AttachedAssets );
 
-	// Is Persona currently in Skeleton Mode
-	bool IsInSkeletonMode() const;
-
 private:
 	/** Pointer back to the kismet 2 tool that owns us */
 	TWeakPtr<FPersona> PersonaPtr;
@@ -563,6 +569,9 @@ private:
 
 	USkeleton* TargetSkeleton;
 
+	/** Widget user to hold the skeleton tree */
+	TSharedPtr<SOverlay> TreeHolder;
+
 	/** Widget used to display the skeleton hierarchy */
 	TSharedPtr<SMeshSkeletonTreeRowType> SkeletonTreeView;
 
@@ -587,6 +596,8 @@ private:
 	/** Current type of sockets to show */
 	ESocketFilter::Type SocketFilter;
 
+	bool bShowingRetargetingOptions;
+
 	/** Points to an item that is being requested to be renamed */
 	TSharedPtr<FDisplayedTreeRowInfo> DeferredRenameRequest;
 
