Retargeting : select rig: Fixed so that you don't have to have "show engine content" to see Humanoid Rig

diff --git a/Engine/Source/Editor/Persona/Private/SRigPicker.cpp b/Engine/Source/Editor/Persona/Private/SRigPicker.cpp
new file mode 100644
--- /dev/null
+++ b/Engine/Source/Editor/Persona/Private/SRigPicker.cpp
@@ -0,0 +1,128 @@
+// Copyright 1998-2014 Epic Games, Inc. All Rights Reserved.
+
+#include "PersonaPrivatePCH.h"
+#include "SRigPicker.h"
+#include "Editor/ContentBrowser/Public/ContentBrowserModule.h"
+
+#define LOCTEXT_NAMESPACE "RigPicker"
+
+static const FVector2D ContentBrowserWindowSize(300.0f, 300.0f);
+
+URig* SRigPicker::EngineHumanoidRig = NULL;
+
+void SRigPicker::Construct(const FArguments& InArgs)
+{
+	if (EngineHumanoidRig == NULL)
+	{
+		EngineHumanoidRig = LoadObject<URig>(nullptr, TEXT("/Engine/EngineMeshes/Humanoid.Humanoid"), nullptr, LOAD_None, nullptr);
+	}
+
+	CurrentObject = InArgs._InitialObject;
+	ShouldFilterAsset = InArgs._OnShouldFilterAsset;
+	OnSetReference = InArgs._OnSetReference;
+	OnClose = InArgs._OnClose;
+
+	FMenuBuilder MenuBuilder(true, NULL);
+
+	MenuBuilder.BeginSection(NAME_None, LOCTEXT("EngineAssets", "Default"));
+	{
+		if (EngineHumanoidRig)
+		{
+			MenuBuilder.AddMenuEntry(
+				LOCTEXT("PickDefaultRig", "Select Humanoid Rig"),
+				LOCTEXT("PickDefaultRig_Tooltip", "Select Engine Defined Humanoid Rig"),
+				FSlateIcon(),
+				FUIAction(FExecuteAction::CreateSP(this, &SRigPicker::OnSelectDefault)));
+		}
+
+		MenuBuilder.AddMenuEntry(
+			LOCTEXT("ClearAsset", "Clear"),
+			LOCTEXT("ClearAsset_ToolTip", "Clears the asset set on this field"),
+			FSlateIcon(),
+			FUIAction(FExecuteAction::CreateSP(this, &SRigPicker::OnClear))
+			);
+	}
+	MenuBuilder.EndSection();
+
+	MenuBuilder.BeginSection(NAME_None, LOCTEXT("BrowseHeader", "Browse"));
+	{
+		TSharedPtr<SWidget> MenuContent;
+
+		FContentBrowserModule& ContentBrowserModule = FModuleManager::Get().LoadModuleChecked<FContentBrowserModule>(TEXT("ContentBrowser"));
+
+		FAssetPickerConfig AssetPickerConfig;
+		AssetPickerConfig.Filter.ClassNames.Add(URig::StaticClass()->GetFName());
+		// Allow child classes
+		AssetPickerConfig.Filter.bRecursiveClasses = true;
+		// Set a delegate for setting the asset from the picker
+		AssetPickerConfig.OnAssetSelected = FOnAssetSelected::CreateSP(this, &SRigPicker::OnAssetSelected);
+		// Use the smallest size thumbnails
+		AssetPickerConfig.ThumbnailScale = 0;
+		// Use the list view by default
+		AssetPickerConfig.InitialAssetViewType = EAssetViewType::List;
+		// The initial selection should be the current value
+		AssetPickerConfig.InitialAssetSelection = CurrentObject;
+		// We'll do clearing ourselves
+		AssetPickerConfig.bAllowNullSelection = false;
+		// Focus search box
+		AssetPickerConfig.bFocusSearchBoxWhenOpened = true;
+		// Apply custom filter
+		AssetPickerConfig.OnShouldFilterAsset = FOnShouldFilterAsset::CreateSP(this, &SRigPicker::OnShouldFilterAsset);
+		// Don't allow dragging
+		AssetPickerConfig.bAllowDragging = false;
+
+		MenuContent =
+			SNew(SBox)
+			.WidthOverride(ContentBrowserWindowSize.X)
+			.HeightOverride(ContentBrowserWindowSize.Y)
+			[
+				ContentBrowserModule.Get().CreateAssetPicker(AssetPickerConfig)
+			];
+
+		MenuBuilder.AddWidget(MenuContent.ToSharedRef(), FText::GetEmpty(), true);
+	}
+	MenuBuilder.EndSection();
+
+	ChildSlot
+	[
+		MenuBuilder.MakeWidget()
+	];
+}
+
+void SRigPicker::OnClear()
+{
+	SetValue(NULL);
+	OnClose.ExecuteIfBound();
+}
+
+void SRigPicker::OnAssetSelected(const FAssetData& AssetData)
+{
+	SetValue(AssetData);
+	OnClose.ExecuteIfBound();
+}
+
+bool SRigPicker::OnShouldFilterAsset(const class FAssetData& AssetData)
+{
+	if (EngineHumanoidRig && AssetData.ObjectPath == FName(*EngineHumanoidRig->GetPathName()))
+	{
+		return true;
+	}
+
+	if (ShouldFilterAsset.IsBound())
+	{
+		return ShouldFilterAsset.Execute(AssetData);
+	}
+
+	return false;
+}
+void SRigPicker::SetValue(const FAssetData& AssetData)
+{
+	OnSetReference.ExecuteIfBound(AssetData.GetAsset());
+}
+
+void SRigPicker::OnSelectDefault()
+{
+	OnAssetSelected(EngineHumanoidRig);
+}
+
+#undef LOCTEXT_NAMESPACE
\ No newline at end of file
diff --git a/Engine/Source/Editor/Persona/Private/SRigPicker.h b/Engine/Source/Editor/Persona/Private/SRigPicker.h
new file mode 100644
--- /dev/null
+++ b/Engine/Source/Editor/Persona/Private/SRigPicker.h
@@ -0,0 +1,61 @@
+// Copyright 1998-2014 Epic Games, Inc. All Rights Reserved.
+#pragma once
+
+
+/** Called when an asset is selected in the asset view */
+DECLARE_DELEGATE_OneParam(FOnAssetSelected, const class FAssetData& /*AssetData*/);
+
+/**
+ * A widget used to pick rig 
+ */
+class SRigPicker : public SCompoundWidget
+{
+public:
+	DECLARE_DELEGATE_OneParam(FOnSetReference, UObject*)
+
+	SLATE_BEGIN_ARGS( SRigPicker )
+		: _InitialObject(NULL)
+	{}
+		SLATE_ARGUMENT( FAssetData, InitialObject )
+		SLATE_EVENT( FOnShouldFilterAsset, OnShouldFilterAsset )
+		SLATE_EVENT( FOnSetReference, OnSetReference )
+		SLATE_EVENT( FSimpleDelegate, OnClose )
+	SLATE_END_ARGS()
+
+	/**
+	 * Construct the widget.
+	 * @param	InArgs				Arguments for widget construction
+	 * @param	InPropertyHandle	The property handle that this widget will operate on.
+	 */
+	void Construct( const FArguments& InArgs );
+
+private:
+	/** 
+	 * Delegate for handling selection in the asset browser.
+	 * @param	AssetData	The chosen asset data
+	 */
+	void OnAssetSelected( const FAssetData& AssetData );
+	void OnSelectDefault();
+	void OnClear();
+	bool OnShouldFilterAsset(const class FAssetData& AssetData);
+
+	/** 
+	 * Set the value of the asset referenced by this property editor.
+	 * Will set the underlying property handle if there is one.
+	 * @param	AssetData	The asset data for the object to set the reference to
+	 */
+	void SetValue( const FAssetData& AssetData );
+private:
+	FAssetData CurrentObject;
+
+	/** Delegate for filtering valid assets */
+	FOnShouldFilterAsset ShouldFilterAsset;
+
+	/** Delegate to call when our object value should be set */
+	FOnSetReference OnSetReference;
+
+	/** Delegate for closing the containing menu */
+	FSimpleDelegate OnClose;
+
+	static URig* EngineHumanoidRig;
+};
\ No newline at end of file
diff --git a/Engine/Source/Editor/Persona/Private/SRigWindow.cpp b/Engine/Source/Editor/Persona/Private/SRigWindow.cpp
--- a/Engine/Source/Editor/Persona/Private/SRigWindow.cpp
+++ b/Engine/Source/Editor/Persona/Private/SRigWindow.cpp
@@ -14,6 +14,7 @@
 #include "BoneSelectionWidget.h"
 #include "SSearchBox.h"
 #include "SInlineEditableTextBlock.h"
+#include "SRigPicker.h"
 
 #define LOCTEXT_NAMESPACE "SRigWindow"
 
@@ -229,18 +230,20 @@ void SRigWindow::Construct(const FArguments& InArgs)
 
 			+SHorizontalBox::Slot()
 			[
-				// rig asset picker
-				//ContentBrowserModule.Get().CreateAssetPicker(AssetPickerConfig)
-				SNew(SContentReference)
-				//.WidthOverride(80.f)
-				.AllowClearingReference(true)
-				.AllowSelectingNewAsset(true)
-				.AssetReference(this, &SRigWindow::GetRigObject)
-				.AllowedClass(URig::StaticClass())
-				.OnShouldFilterAsset(this, &SRigWindow::ShouldFilterAssetBased)
-				.OnSetReference(this, &SRigWindow::OnAssetSelected)
-				.AssetPickerSizeOverride(FVector2D(250, 700))
-				.InitialAssetViewType(EAssetViewType::List)
+				SAssignNew( AssetComboButton, SComboButton )
+				//.ToolTipText( this, &SPropertyEditorAsset::OnGetToolTip )
+				.ButtonStyle( FEditorStyle::Get(), "PropertyEditor.AssetComboStyle" )
+				.ForegroundColor(FEditorStyle::GetColor("PropertyEditor.AssetName.ColorAndOpacity"))
+				.OnGetMenuContent( this, &SRigWindow::MakeRigPickerWithMenu )
+				.ContentPadding(2.0f)
+				.ButtonContent()
+				[
+					// Show the name of the asset or actor
+					SNew(STextBlock)
+					.TextStyle( FEditorStyle::Get(), "PropertyEditor.AssetClass" )
+					.Font( FEditorStyle::GetFontStyle( "PropertyWindow.NormalFont" ) )
+					.Text(this,&SRigWindow::GetAssetName)
+				]
 			]
 
 			+SHorizontalBox::Slot()
@@ -372,6 +375,8 @@ void SRigWindow::OnAssetSelected(UObject* Object)
 {
 	if (Skeleton)
 	{
+		AssetComboButton->SetIsOpen(false);
+
 		const FScopedTransaction Transaction(LOCTEXT("RigAssetChanged", "Select Rig"));
 		Skeleton->Modify();
 		Skeleton->SetRigConfig(Cast<URig>(Object));
@@ -382,7 +387,7 @@ void SRigWindow::OnAssetSelected(UObject* Object)
 }
 
 /** Returns true if the asset shouldn't show  */
-bool SRigWindow::ShouldFilterAssetBased(const class FAssetData& AssetData)
+bool SRigWindow::ShouldFilterAsset(const class FAssetData& AssetData)
 {
 	return (AssetData.GetAsset() == GetRigObject());
 }
@@ -436,5 +441,33 @@ FString SRigWindow::GetAdvancedButtonText() const
 
 	return TEXT("Show Advanced");
 }
+
+TSharedRef<SWidget> SRigWindow::MakeRigPickerWithMenu()
+{
+	// rig asset picker
+	return	
+		SNew(SRigPicker)
+		.InitialObject(Skeleton->GetRig())
+		.OnShouldFilterAsset(this, &SRigWindow::ShouldFilterAsset)
+		.OnSetReference(this, &SRigWindow::OnAssetSelected)
+		.OnClose(this, &SRigWindow::CloseComboButton );
+}
+
+void SRigWindow::CloseComboButton()
+{
+	AssetComboButton->SetIsOpen(false);
+}
+
+FString SRigWindow::GetAssetName() const
+{
+	UObject* Rig = GetRigObject();
+	if (Rig)
+	{
+		return Rig->GetName();
+	}
+
+	return TEXT("None");
+}
+
 #undef LOCTEXT_NAMESPACE
 
diff --git a/Engine/Source/Editor/Persona/Private/SRigWindow.h b/Engine/Source/Editor/Persona/Private/SRigWindow.h
--- a/Engine/Source/Editor/Persona/Private/SRigWindow.h
+++ b/Engine/Source/Editor/Persona/Private/SRigWindow.h
@@ -148,20 +148,26 @@ private:
 	/** show advanced? */
 	bool bDisplayAdvanced;
 
+	/** rig combo button */
+	TSharedPtr< class SComboButton > AssetComboButton;
+
 	/** Delegate for undo/redo transaction **/
 	void PostUndo();
 
 	/**
 	 * Callback for asset picker
 	 */
-	/* Set new reference for skeletal mesh */
+	/* Set rig set combo box*/
 	void OnAssetSelected(UObject* Object);
+	FString GetAssetName() const;
+	void CloseComboButton();
+	TSharedRef<SWidget> SRigWindow::MakeRigPickerWithMenu();
 
 	/** Returns true if the asset shouldn't show  */
-	bool ShouldFilterAssetBased(const class FAssetData& AssetData);
+	bool ShouldFilterAsset(const class FAssetData& AssetData);
 
 	UObject* GetRigObject() const;
-
+	
 	void OnBoneMappingChanged( FName NodeName, FName BoneName );
 	FName GetBoneMapping( FName NodeName );
 
