Fix compile error

diff --git a/Engine/Source/Editor/Persona/Private/SRigWindow.h b/Engine/Source/Editor/Persona/Private/SRigWindow.h
--- a/Engine/Source/Editor/Persona/Private/SRigWindow.h
+++ b/Engine/Source/Editor/Persona/Private/SRigWindow.h
@@ -161,7 +161,7 @@ private:
 	void OnAssetSelected(UObject* Object);
 	FString GetAssetName() const;
 	void CloseComboButton();
-	TSharedRef<SWidget> SRigWindow::MakeRigPickerWithMenu();
+	TSharedRef<SWidget> MakeRigPickerWithMenu();
 
 	/** Returns true if the asset shouldn't show  */
 	bool ShouldFilterAsset(const class FAssetData& AssetData);
