Fixed non-unity build

diff --git a/Engine/Source/Editor/Persona/Private/SRigPicker.cpp b/Engine/Source/Editor/Persona/Private/SRigPicker.cpp
--- a/Engine/Source/Editor/Persona/Private/SRigPicker.cpp
+++ b/Engine/Source/Editor/Persona/Private/SRigPicker.cpp
@@ -3,6 +3,7 @@
 #include "PersonaPrivatePCH.h"
 #include "SRigPicker.h"
 #include "Editor/ContentBrowser/Public/ContentBrowserModule.h"
+#include "Animation/Rig.h"
 
 #define LOCTEXT_NAMESPACE "RigPicker"
 
diff --git a/Engine/Source/Editor/Persona/Private/SRigPicker.h b/Engine/Source/Editor/Persona/Private/SRigPicker.h
--- a/Engine/Source/Editor/Persona/Private/SRigPicker.h
+++ b/Engine/Source/Editor/Persona/Private/SRigPicker.h
@@ -1,6 +1,7 @@
 // Copyright 1998-2014 Epic Games, Inc. All Rights Reserved.
 #pragma once
 
+#include "Editor/ContentBrowser/Public/ContentBrowserDelegates.h"
 
 /** Called when an asset is selected in the asset view */
 DECLARE_DELEGATE_OneParam(FOnAssetSelected, const class FAssetData& /*AssetData*/);
