If ref pose, just disable editing. It doesn't work.

diff --git a/Engine/Source/Editor/Persona/Private/SAnimationEditorViewport.cpp b/Engine/Source/Editor/Persona/Private/SAnimationEditorViewport.cpp
--- a/Engine/Source/Editor/Persona/Private/SAnimationEditorViewport.cpp
+++ b/Engine/Source/Editor/Persona/Private/SAnimationEditorViewport.cpp
@@ -215,7 +215,7 @@ bool SAnimationEditorViewportTabBody::CanUseGizmos() const
 	{
 		if (Component->bForceRefpose)
 		{
-			return true;
+			return false;
 		}
 		else if (Component->IsPreviewOn())
 		{
@@ -1332,7 +1332,7 @@ void SAnimationEditorViewportTabBody::SaveData(class SAnimationEditorViewportTab
 {
 	if ( PersonaPtr.IsValid() && OldViewport )
 	{
-		FPersonaModeSharedData & SharedData = PersonaPtr.Pin()->ModeSharedData;
+		FPersonaModeSharedData& SharedData = PersonaPtr.Pin()->ModeSharedData;
 
 		TSharedRef<FAnimationViewportClient> OldAnimViewportClient = StaticCastSharedRef<FAnimationViewportClient>(OldViewport->LevelViewportClient.ToSharedRef());
 		// set camera set up
@@ -1353,7 +1353,7 @@ void SAnimationEditorViewportTabBody::RestoreData()
 {
 	if ( PersonaPtr.IsValid() )
 	{
-		FPersonaModeSharedData & SharedData = PersonaPtr.Pin()->ModeSharedData;
+		FPersonaModeSharedData& SharedData = PersonaPtr.Pin()->ModeSharedData;
 		TSharedRef<FAnimationViewportClient> AnimViewportClient = StaticCastSharedRef<FAnimationViewportClient>(LevelViewportClient.ToSharedRef());
 
 
@@ -1838,7 +1838,7 @@ FText SAnimationEditorViewportTabBody::GetViewportCornerText() const
 	TSharedPtr<FPersona> Persona = PersonaPtr.Pin();
 	if(Persona->Recorder.InRecording())
 	{
-		const FString & Name = Persona->Recorder.GetAnimationObject()->GetName();
+		const FString& Name = Persona->Recorder.GetAnimationObject()->GetName();
 		float TimeRecorded = Persona->Recorder.GetTimeRecorded();
 		FNumberFormattingOptions NumberOption;
 		NumberOption.MaximumFractionalDigits = 2;
