Multi-Montages support
. Slots now belong to a Group. One montage can play at a time per group.
. Removed old Slot and Group name lists and created new Anim Slot Manager. Add, Create Slots and Groups, assign Slots to groups and save changes.
. New Montage assets and Slot nodes default to a valid 'DefaultSlot' so they are functional when created. (Eliminates common mistake of empty slot names).
. Montage tracks and slot nodes picks slots from a combo box of valid entries, no more typing name by hand.
. Montage tracks will warn if not setup properly (different groups used, same slot used multiple times).

diff --git a/Engine/Source/Editor/AnimGraph/Private/AnimGraphNode_Slot.cpp b/Engine/Source/Editor/AnimGraph/Private/AnimGraphNode_Slot.cpp
--- a/Engine/Source/Editor/AnimGraph/Private/AnimGraphNode_Slot.cpp
+++ b/Engine/Source/Editor/AnimGraph/Private/AnimGraphNode_Slot.cpp
@@ -43,8 +43,12 @@ FText UAnimGraphNode_Slot::GetNodeTitle(ENodeTitleType::Type TitleType) const
 	//        choose to mark this dirty when that happens for this to properly work
 	else //if (!CachedNodeTitles.IsTitleCached(TitleType))
 	{
+		UAnimBlueprint* AnimBlueprint = GetAnimBlueprint();
+		FName GroupName = (AnimBlueprint->TargetSkeleton) ? AnimBlueprint->TargetSkeleton->GetSlotGroupName(Node.SlotName) : FAnimSlotGroup::DefaultGroupName;
+
 		FFormatNamedArguments Args;
 		Args.Add(TEXT("SlotName"), FText::FromName(Node.SlotName));
+		Args.Add(TEXT("GroupName"), FText::FromName(GroupName));
 
 		// FText::Format() is slow, so we cache this to save on performance
 		if (TitleType == ENodeTitleType::ListView || TitleType == ENodeTitleType::MenuTitle)
@@ -53,7 +57,7 @@ FText UAnimGraphNode_Slot::GetNodeTitle(ENodeTitleType::Type TitleType) const
 		}
 		else
 		{
-			CachedNodeTitles.SetCachedTitle(TitleType, FText::Format(LOCTEXT("SlotNodeTitle", "{SlotName}\nSlot"), Args));
+			CachedNodeTitles.SetCachedTitle(TitleType, FText::Format(LOCTEXT("SlotNodeTitle", "Slot '{SlotName}'\nGroup '{GroupName}'"), Args));
 		}
 	}
 	return CachedNodeTitles[TitleType];
@@ -69,8 +73,7 @@ void UAnimGraphNode_Slot::BakeDataDuringCompilation(class FCompilerResultsLog& M
 	UAnimBlueprint* AnimBlueprint = GetAnimBlueprint();
 	if (AnimBlueprint->TargetSkeleton)
 	{
-		AnimBlueprint->TargetSkeleton->AddSlotNodeName(Node.SlotName);
-		AnimBlueprint->TargetSkeleton->AddSlotGroupName(Node.GroupName);
+		AnimBlueprint->TargetSkeleton->RegisterSlotNode(Node.SlotName);
 	}
 }
 
diff --git a/Engine/Source/Editor/EditorStyle/Private/SlateEditorStyle.cpp b/Engine/Source/Editor/EditorStyle/Private/SlateEditorStyle.cpp
--- a/Engine/Source/Editor/EditorStyle/Private/SlateEditorStyle.cpp
+++ b/Engine/Source/Editor/EditorStyle/Private/SlateEditorStyle.cpp
@@ -4429,6 +4429,12 @@ void FSlateEditorStyle::FStyle::SetupPersonaStyle()
 		Set("Persona.ReimportAsset", new IMAGE_BRUSH("Icons/Reimport_12x", Icon12x12));
 		Set("SkeletonTree.SkeletonSocket", new IMAGE_BRUSH("Persona/SkeletonTree/icon_SocketG_16px", Icon16x16));
 		Set("SkeletonTree.MeshSocket", new IMAGE_BRUSH("Persona/SkeletonTree/icon_SocketC_16px", Icon16x16));
+
+		// Anim Slot Manager
+		Set("AnimSlotManager.SaveSkeleton", new IMAGE_BRUSH("Persona/AnimSlotManager/icon_SaveSkeleton_40x", Icon40x40));
+		Set("AnimSlotManager.AddGroup", new IMAGE_BRUSH("Persona/AnimSlotManager/icon_AddGroup_40x", Icon40x40));
+		Set("AnimSlotManager.AddSlot", new IMAGE_BRUSH("Persona/AnimSlotManager/icon_AddSlot_40x", Icon40x40));
+		Set("AnimSlotManager.Warning", new IMAGE_BRUSH("Persona/AnimSlotManager/icon_Warning_14x", Icon16x16));
 	}
 
 	// Kismet 2
diff --git a/Engine/Source/Editor/Persona/Private/BlueprintMode/AnimBlueprintMode.cpp b/Engine/Source/Editor/Persona/Private/BlueprintMode/AnimBlueprintMode.cpp
--- a/Engine/Source/Editor/Persona/Private/BlueprintMode/AnimBlueprintMode.cpp
+++ b/Engine/Source/Editor/Persona/Private/BlueprintMode/AnimBlueprintMode.cpp
@@ -8,7 +8,6 @@
 #include "Editor/Kismet/Public/SBlueprintEditorToolbar.h"
 #include "Editor/Kismet/Public/BlueprintEditorModes.h"
 #include "SSkeletonSlotNames.h"
-#include "SSkeletonSlotGroupNames.h"
 #include "SSkeletonAnimNotifies.h"
 
 #include "AnimBlueprintMode.h"
@@ -114,7 +113,6 @@ FAnimBlueprintEditAppMode::FAnimBlueprintEditAppMode(TSharedPtr<FPersona> InPers
 	PersonaTabFactories.RegisterFactory(MakeShareable(new FAnimBlueprintDefaultsEditorSummoner(InPersona)));
 	PersonaTabFactories.RegisterFactory(MakeShareable(new FAnimBlueprintParentPlayerEditorSummoner(InPersona)));
 	PersonaTabFactories.RegisterFactory(MakeShareable(new FSkeletonSlotNamesSummoner(InPersona)));
-	PersonaTabFactories.RegisterFactory(MakeShareable(new FSkeletonSlotGroupNamesSummoner(InPersona)));
 
 	// setup toolbar - clear existing toolbar extender from the BP mode
 	//@TODO: Keep this in sync with BlueprintEditorModes.cpp
diff --git a/Engine/Source/Editor/Persona/Private/Customization/AnimGraphNodeSlotDetails.cpp b/Engine/Source/Editor/Persona/Private/Customization/AnimGraphNodeSlotDetails.cpp
--- a/Engine/Source/Editor/Persona/Private/Customization/AnimGraphNodeSlotDetails.cpp
+++ b/Engine/Source/Editor/Persona/Private/Customization/AnimGraphNodeSlotDetails.cpp
@@ -25,16 +25,13 @@ FAnimGraphNodeSlotDetails::FAnimGraphNodeSlotDetails(TWeakPtr<class FPersona> In
 void FAnimGraphNodeSlotDetails::CustomizeDetails(class IDetailLayoutBuilder& DetailBuilder)
 {
 	SlotNodeNamePropertyHandle = DetailBuilder.GetProperty(TEXT("Node.SlotName"));
-	GroupNamePropertyHandle = DetailBuilder.GetProperty(TEXT("Node.GroupName"));
 
 	DetailBuilder.HideProperty(SlotNodeNamePropertyHandle);
-	DetailBuilder.HideProperty(GroupNamePropertyHandle);
 
 	TArray<UObject*> OuterObjects;
 	SlotNodeNamePropertyHandle->GetOuterObjects(OuterObjects);
 
 	check(SlotNodeNamePropertyHandle.IsValid());
-	check(GroupNamePropertyHandle.IsValid());
 
 	Skeleton = NULL;
 
@@ -61,16 +58,12 @@ void FAnimGraphNodeSlotDetails::CustomizeDetails(class IDetailLayoutBuilder& Det
 
 	// this is used for 2 things, to create name, but also for another pop-up box to show off, it's a bit hacky to use this, but I need widget to parent on
 	TSharedRef<SWidget> SlotNodePropertyNameWidget = SlotNodeNamePropertyHandle->CreatePropertyNameWidget();
-	TSharedRef<SWidget> GroupNodePropertyNameWidget = GroupNamePropertyHandle->CreatePropertyNameWidget();
 
-	// get the current slot node name
-	int32 CurrentlySelectedIndex = RefreshSlotNames();
-	TSharedPtr<FString> CurrentlySelectedNode;
+	// fill combo with groups and slots
+	RefreshComboLists();
 
-	if (CurrentlySelectedIndex != INDEX_NONE)
-	{
-		CurrentlySelectedNode = SlotNameComboList[CurrentlySelectedIndex];
-	}
+	int32 FoundIndex = SlotNameList.Find(SlotNameComboSelectedName);
+	TSharedPtr<FString> ComboBoxSelectedItem = SlotNameComboListItems[FoundIndex];
 
 	// create combo box
 	IDetailCategoryBuilder& SlotNameGroup = DetailBuilder.EditCategory(TEXT("Settings"));
@@ -80,96 +73,30 @@ void FAnimGraphNodeSlotDetails::CustomizeDetails(class IDetailLayoutBuilder& Det
 		SlotNodePropertyNameWidget
 	]
 	.ValueContent()
+	.MinDesiredWidth(125.f * 3.f)
+	.MaxDesiredWidth(125.f * 3.f)
 	[
 		SNew(SHorizontalBox)
 
 		+SHorizontalBox::Slot()
+		.AutoWidth()
 		[
 			SAssignNew(SlotNameComboBox, STextComboBox)
-			.OptionsSource(&SlotNameComboList)
+			.OptionsSource(&SlotNameComboListItems)
 			.OnSelectionChanged(this, &FAnimGraphNodeSlotDetails::OnSlotNameChanged)
-			.InitiallySelectedItem(CurrentlySelectedNode)
-			.ContentPadding(2)
-		]
-
-		+SHorizontalBox::Slot()
-		.AutoWidth()
-		[
-			SNew(SButton)
-			.Text(LOCTEXT("SlotName_AddButtonLabel", "Add"))
-			.ToolTipText(LOCTEXT("SlotName_AddButtonToolTipText", "New Slot Node Names") )
-			.OnClicked(this, &FAnimGraphNodeSlotDetails::OnAddSlotName, SlotNodePropertyNameWidget)
-			.Content()
-			[
-				SNew(SImage)
-				.Image(FEditorStyle::GetBrush("PropertyWindow.Button_AddToArray"))
-			]
-		]
-
-		+SHorizontalBox::Slot()
-		.AutoWidth()
-		[
-			SNew(SButton)
-			.Text(LOCTEXT("SlotName_ManageButtonLabel", "Manage"))
-			.ToolTipText(LOCTEXT("SlotName_ManageButtonToolTipText", "Manage Slot Node Names"))
-			.OnClicked(this, &FAnimGraphNodeSlotDetails::OnManageSlotName)
-			.Content()
-			[
-				SNew(SImage)
-				.Image(FEditorStyle::GetBrush("MeshPaint.FindInCB"))
-			]
-		]
-	];
-
-	// group names
-	CurrentlySelectedIndex = RefreshGroupNames();
-	CurrentlySelectedNode = nullptr;
-
-	if(CurrentlySelectedIndex != INDEX_NONE)
-	{
-		CurrentlySelectedNode = GroupNameComboList[CurrentlySelectedIndex];
-	}
-
-	//IDetailGroup& GroupNameGroup = DetailBuilder.EditCategory(TEXT("GroupName_Group"), LOCTEXT("GroupNameTitleLabel", "Group Name").ToString());
-	SlotNameGroup.AddCustomRow(LOCTEXT("GroupNameTitleLabel", "Group Name").ToString())
-	.NameContent()
-	[
-		GroupNodePropertyNameWidget
-	]
-	.ValueContent()
-	[
-		SNew(SHorizontalBox)
-
-		+SHorizontalBox::Slot()
-		[
-			SAssignNew(GroupNameComboBox, STextComboBox)
-			.OptionsSource(&GroupNameComboList)
-			.OnSelectionChanged(this, &FAnimGraphNodeSlotDetails::OnGroupNameChanged)
-			.InitiallySelectedItem(CurrentlySelectedNode)
+			.OnComboBoxOpening(this, &FAnimGraphNodeSlotDetails::OnSlotListOpening)
+			.InitiallySelectedItem(ComboBoxSelectedItem)
 			.ContentPadding(2)
+			.ToolTipText(FText::FromString(*ComboBoxSelectedItem))
 		]
 
 		+SHorizontalBox::Slot()
 		.AutoWidth()
 		[
 			SNew(SButton)
-			.Text(LOCTEXT("GroupName_AddButtonLabel", "Add"))
-			.ToolTipText(LOCTEXT("GroupName_AddButtonToolTipText", "New Slot Group Names"))
-			.OnClicked(this, &FAnimGraphNodeSlotDetails::OnAddGroupName, GroupNodePropertyNameWidget)
-			.Content()
-			[
-				SNew(SImage)
-				.Image(FEditorStyle::GetBrush("PropertyWindow.Button_AddToArray"))
-			]
-		]
-
-		+SHorizontalBox::Slot()
-		.AutoWidth()
-		[
-			SNew(SButton)
-			.Text(LOCTEXT("GroupName_ManageButtonLabel", "Manage"))
-			.ToolTipText(LOCTEXT("GroupName_ManageButtonToolTipText", "Manage Slot Group Names"))
-			.OnClicked(this, &FAnimGraphNodeSlotDetails::OnManageGroupName)
+			.Text(LOCTEXT("AnimSlotNode_DetailPanelManageButtonLabel", "Anim Slot Manager"))
+			.ToolTipText(LOCTEXT("AnimSlotNode_DetailPanelManageButtonToolTipText", "Open Anim Slot Manager to edit Slots and Groups."))
+			.OnClicked(this, &FAnimGraphNodeSlotDetails::OnOpenAnimSlotManager)
 			.Content()
 			[
 				SNew(SImage)
@@ -177,209 +104,99 @@ void FAnimGraphNodeSlotDetails::CustomizeDetails(class IDetailLayoutBuilder& Det
 			]
 		]
 	];
-
 }
 
-int32 FAnimGraphNodeSlotDetails::RefreshSlotNames()
+void FAnimGraphNodeSlotDetails::RefreshComboLists(bool bOnlyRefreshIfDifferent /*= false*/)
 {
-	SlotNameComboList.Empty();
+	SlotNodeNamePropertyHandle->GetValue(SlotNameComboSelectedName);
 
-	FName CurrentSlotNodeName = NAME_None;
-	SlotNodeNamePropertyHandle->GetValue(CurrentSlotNodeName);
+	// Make sure slot node exists in our skeleton.
+	Skeleton->RegisterSlotNode(SlotNameComboSelectedName);
 
-	// add slot node names
-	int32 CurrentSelectedIndex = INDEX_NONE;
-	const TArray<FName> & SlotNodeNames = Skeleton->GetSlotNodeNames();
-	for(auto SlotName : SlotNodeNames)
+	// Refresh Slot Names
 	{
-		SlotNameComboList.Add(MakeShareable(new FString(SlotName.ToString())));
+		TArray< TSharedPtr< FString > > NewSlotNameComboListItems;
+		TArray< FName > NewSlotNameList;
+		bool bIsSlotNameListDifferent = false;
 
-		if(CurrentSlotNodeName == SlotName)
+		const TArray<FAnimSlotGroup>& SlotGroups = Skeleton->GetSlotGroups();
+		for (auto SlotGroup : SlotGroups)
 		{
-			CurrentSelectedIndex = SlotNameComboList.Num()-1;
-		}
-	}
-
-	if (SlotNameComboBox.IsValid())
-	{
-		SlotNameComboBox->RefreshOptions();
-	}
-
-	// if SlotNameComboList is empty, just add a temporary message to avoid a combo box problem which doesn't show the list properly if empty
-	if (SlotNameComboList.Num() == 0)
-	{
-		SlotNameComboList.Add(MakeShareable(new FString("Please Add a slot name using the neighboring + button")));
-	}
-
-	return CurrentSelectedIndex;
-}
-////////////////////////////////////////////////////////////////
+			int32 Index = 0;
+			for (auto SlotName : SlotGroup.SlotNames)
+			{
+				NewSlotNameList.Add(SlotName);
 
-void FAnimGraphNodeSlotDetails::OnSlotNameChanged(TSharedPtr<FString> NewSelection, ESelectInfo::Type SelectInfo)
-{
-	// if it's set from code, we did that on purpose
-	if(SelectInfo != ESelectInfo::Direct)
-	{
-		FString NewValue = *NewSelection.Get();
-		const TArray<FName> & SlotNodeNames = Skeleton->GetSlotNodeNames();
+				FString ComboItemString = FString::Printf(TEXT("%s.%s"), *SlotGroup.GroupName.ToString(), *SlotName.ToString());
+				NewSlotNameComboListItems.Add(MakeShareable(new FString(ComboItemString)));
 
-		for(auto SlotNodeName : SlotNodeNames )
-		{
-			if(NewValue == SlotNodeName.ToString())
-			{
-				// trigget transaction 
-				//const FScopedTransaction Transaction(LOCTEXT("ChangeSlotNodeName", "Change Collision Profile"));
-				// set profile set up
-				ensure(SlotNodeNamePropertyHandle->SetValue(NewValue) ==  FPropertyAccess::Result::Success);
-				return;
+				bIsSlotNameListDifferent = bIsSlotNameListDifferent || (!SlotNameComboListItems.IsValidIndex(Index) || (SlotNameComboListItems[Index] != NewSlotNameComboListItems[Index]));
+				Index++;
 			}
 		}
-	}
-}
-
-FReply FAnimGraphNodeSlotDetails::OnAddSlotName(TSharedRef<SWidget> Widget)
-{
-	TSharedRef<STextEntryPopup> TextEntry =
-		SNew(STextEntryPopup)
-		.Label(LOCTEXT("NewSlotName_AskSlotName", "New Slot Name"))
-		.OnTextCommitted(this, &FAnimGraphNodeSlotDetails::AddSlotName);
-
-	// Show dialog to enter new track name
-	FSlateApplication::Get().PushMenu(
-		Widget,
-		TextEntry,
-		FSlateApplication::Get().GetCursorPos(),
-		FPopupTransitionEffect(FPopupTransitionEffect::TypeInPopup)
-		);
-
-	TextEntry->FocusDefaultWidget();
 
-	return FReply::Handled();
-}
-
-void FAnimGraphNodeSlotDetails::AddSlotName(const FText& SlotNameToAdd, ETextCommit::Type CommitInfo)
-{
-	if(!SlotNameToAdd.IsEmpty())
-	{
-		const FScopedTransaction Transaction(LOCTEXT("NewSlotName_AddSlotName", "Add New Slot Node Name"));
-		// before insert, make sure everything behind is fixed
-		Skeleton->AddSlotNodeName(FName(*SlotNameToAdd.ToString()));
-		Skeleton->Modify(true);
-
-		// refresh combo box
-		int32 SelectedIndex = RefreshSlotNames();
-		if (SelectedIndex != INDEX_NONE)
+		// Refresh if needed
+		if (bIsSlotNameListDifferent || !bOnlyRefreshIfDifferent || (NewSlotNameComboListItems.Num() == 0))
 		{
-			SlotNameComboBox->SetSelectedItem(SlotNameComboList[SelectedIndex]);
-		}
+			SlotNameComboListItems = NewSlotNameComboListItems;
+			SlotNameList = NewSlotNameList;
 
-		// @todo add bunch of notifies
-		FSlateApplication::Get().DismissAllMenus();
+			if (SlotNameComboBox.IsValid())
+			{
+				if (Skeleton->ContainsSlotName(SlotNameComboSelectedName))
+				{
+					int32 FoundIndex = SlotNameList.Find(SlotNameComboSelectedName);
+					TSharedPtr<FString> ComboItem = SlotNameComboListItems[FoundIndex];
+
+					SlotNameComboBox->SetSelectedItem(ComboItem);
+					SlotNameComboBox->SetToolTipText(FText::FromString(*ComboItem));
+				}
+				SlotNameComboBox->RefreshOptions();
+			}
+		}
 	}
 }
 
-FReply FAnimGraphNodeSlotDetails::OnManageSlotName()
-{
-	if(PersonaEditor.IsValid())
-	{
-		PersonaEditor.Pin()->GetTabManager()->InvokeTab(FPersonaTabs::SkeletonSlotNamesID);
-	}
-	return FReply::Handled();
-}
+////////////////////////////////////////////////////////////////
 
-void FAnimGraphNodeSlotDetails::OnGroupNameChanged(TSharedPtr<FString> NewSelection, ESelectInfo::Type SelectInfo)
+void FAnimGraphNodeSlotDetails::OnSlotNameChanged(TSharedPtr<FString> NewSelection, ESelectInfo::Type SelectInfo)
 {
 	// if it's set from code, we did that on purpose
 	if(SelectInfo != ESelectInfo::Direct)
 	{
-		FString NewValue = *NewSelection.Get();
-		const TArray<FName> & SlotGroupNames = Skeleton->GetSlotGroupNames();
-
-		for(auto SlotGroupName : SlotGroupNames)
+		int32 ItemIndex = SlotNameComboListItems.Find(NewSelection);
+		if (ItemIndex != INDEX_NONE)
 		{
-			if(NewValue == SlotGroupName.ToString())
+			SlotNameComboSelectedName = SlotNameList[ItemIndex];
+			if (SlotNameComboBox.IsValid())
+			{
+				SlotNameComboBox->SetToolTipText(FText::FromString(*NewSelection));
+			}
+
+			if (Skeleton->ContainsSlotName(SlotNameComboSelectedName))
 			{
-				// trigget transaction 
-				//const FScopedTransaction Transaction(LOCTEXT("ChangeSlotGroupName", "Change Collision Profile"));
+				// trigger transaction 
+				//const FScopedTransaction Transaction(LOCTEXT("ChangeSlotNodeName", "Change Collision Profile"));
 				// set profile set up
-				ensure(GroupNamePropertyHandle->SetValue(NewValue) ==  FPropertyAccess::Result::Success);
-				return;
+				ensure(SlotNodeNamePropertyHandle->SetValue(SlotNameComboSelectedName.ToString()) == FPropertyAccess::Result::Success);
 			}
 		}
 	}
 }
 
-FReply FAnimGraphNodeSlotDetails::OnAddGroupName(TSharedRef<SWidget> Widget)
+void FAnimGraphNodeSlotDetails::OnSlotListOpening()
 {
-	TSharedRef<STextEntryPopup> TextEntry =
-		SNew(STextEntryPopup)
-		.Label(LOCTEXT("NewGroupName_AskGroupName", "New Group Name"))
-		.OnTextCommitted(this, &FAnimGraphNodeSlotDetails::AddGroupName);
-
-	// Show dialog to enter new track name
-	FSlateApplication::Get().PushMenu(
-		Widget,
-		TextEntry,
-		FSlateApplication::Get().GetCursorPos(),
-		FPopupTransitionEffect(FPopupTransitionEffect::TypeInPopup)
-		);
-
-	TextEntry->FocusDefaultWidget();
-
-	return FReply::Handled();
-}
-
-void FAnimGraphNodeSlotDetails::AddGroupName(const FText& GroupNameToAdd, ETextCommit::Type CommitInfo)
-{
-	if(!GroupNameToAdd.IsEmpty())
-	{
-		const FScopedTransaction Transaction(LOCTEXT("NewGroupName_AddGroupName", "Add New Slot Node Name"));
-		// before insert, make sure everything behind is fixed
-		Skeleton->AddSlotGroupName(FName(*GroupNameToAdd.ToString()));
-		Skeleton->Modify(true);
-
-		// refresh combo box
-		int32 SelectedIndex = RefreshGroupNames();
-		if(SelectedIndex != INDEX_NONE)
-		{
-			GroupNameComboBox->SetSelectedItem(GroupNameComboList[SelectedIndex]);
-		}
-
-		// @todo add bunch of notifies
-		FSlateApplication::Get().DismissAllMenus();
-	}
+	// Refresh Slot Names, in case we used the Anim Slot Manager to make changes.
+	RefreshComboLists(true);
 }
 
-FReply FAnimGraphNodeSlotDetails::OnManageGroupName()
+FReply FAnimGraphNodeSlotDetails::OnOpenAnimSlotManager()
 {
 	if(PersonaEditor.IsValid())
 	{
-		PersonaEditor.Pin()->GetTabManager()->InvokeTab(FPersonaTabs::SkeletonSlotGroupNamesID);
+		PersonaEditor.Pin()->GetTabManager()->InvokeTab(FPersonaTabs::SkeletonSlotNamesID);
 	}
 	return FReply::Handled();
 }
 
-int32 FAnimGraphNodeSlotDetails::RefreshGroupNames()
-{
-	GroupNameComboList.Empty();
-
-	FName CurrentSlotGroupName = NAME_None;
-	GroupNamePropertyHandle->GetValue(CurrentSlotGroupName);
-
-	// add slot node names
-	int32 CurrentSelectedIndex = INDEX_NONE;
-	const TArray<FName> & SlotGroupNames = Skeleton->GetSlotGroupNames();
-	for(auto GroupName : SlotGroupNames)
-	{
-		GroupNameComboList.Add(MakeShareable(new FString(GroupName.ToString())));
-
-		if(CurrentSlotGroupName == GroupName)
-		{
-			CurrentSelectedIndex = GroupNameComboList.Num()-1;
-		}
-	}
-
-	return CurrentSelectedIndex;
-}
-
 #undef LOCTEXT_NAMESPACE
diff --git a/Engine/Source/Editor/Persona/Private/Customization/AnimGraphNodeSlotDetails.h b/Engine/Source/Editor/Persona/Private/Customization/AnimGraphNodeSlotDetails.h
--- a/Engine/Source/Editor/Persona/Private/Customization/AnimGraphNodeSlotDetails.h
+++ b/Engine/Source/Editor/Persona/Private/Customization/AnimGraphNodeSlotDetails.h
@@ -25,33 +25,20 @@ private:
 
 	// property of the two 
 	TSharedPtr<IPropertyHandle> SlotNodeNamePropertyHandle;
-	TSharedPtr<IPropertyHandle> GroupNamePropertyHandle;
 
 	// slot node names
-	TSharedPtr< class STextComboBox>	SlotNameComboBox;
-	TArray< TSharedPtr< FString > >		SlotNameComboList;
+	TSharedPtr<class STextComboBox>	SlotNameComboBox;
+	TArray< TSharedPtr< FString > > SlotNameComboListItems;
+	TArray< FName > SlotNameList;
+	FName SlotNameComboSelectedName;
 
 	void OnSlotNameChanged(TSharedPtr<FString> NewSelection, ESelectInfo::Type SelectInfo);
+	void OnSlotListOpening();
 
 	// slot node names buttons
-	FReply OnAddSlotName(TSharedRef<SWidget> Widget);
-	FReply OnManageSlotName();
+	FReply OnOpenAnimSlotManager();
 
-	void AddSlotName(const FText& SlotNameToAdd, ETextCommit::Type CommitInfo);
-	int32 RefreshSlotNames();
-
-	// group names
-	TSharedPtr< class STextComboBox>	GroupNameComboBox;
-	TArray< TSharedPtr< FString > >		GroupNameComboList;
-
-	void OnGroupNameChanged(TSharedPtr<FString> NewSelection, ESelectInfo::Type SelectInfo);
-
-	// slot node names buttons
-	FReply OnAddGroupName(TSharedRef<SWidget> Widget);
-	FReply OnManageGroupName();
-
-	void AddGroupName(const FText& GroupNameToAdd, ETextCommit::Type CommitInfo);
-	int32 RefreshGroupNames();
+	void RefreshComboLists(bool bOnlyRefreshIfDifferent = false);
 
 	USkeleton* Skeleton;
 };
diff --git a/Engine/Source/Editor/Persona/Private/SAnimMontagePanel.cpp b/Engine/Source/Editor/Persona/Private/SAnimMontagePanel.cpp
--- a/Engine/Source/Editor/Persona/Private/SAnimMontagePanel.cpp
+++ b/Engine/Source/Editor/Persona/Private/SAnimMontagePanel.cpp
@@ -356,6 +356,43 @@ void SAnimMontagePanel::Update()
 			SAssignNew( MontageSlots, SVerticalBox )
 			);
 
+		/************************************************************************/
+		/* Status Bar                                                                     */
+		/************************************************************************/
+		{
+			MontageSlots->AddSlot()
+				.AutoHeight()
+				[
+					// Header, shows name of timeline we are editing
+					SNew(SBorder)
+					.BorderImage(FEditorStyle::GetBrush(TEXT("Graph.TitleBackground")))
+					.HAlign(HAlign_Center)
+					[
+						SNew(SHorizontalBox)
+						+ SHorizontalBox::Slot()
+						.AutoWidth()
+						.FillWidth(3.f)
+						.HAlign(HAlign_Right)
+						.VAlign(VAlign_Center)
+						[
+							SAssignNew(StatusBarWarningImage, SImage)
+							.Image(FEditorStyle::GetBrush("AnimSlotManager.Warning"))
+							.Visibility(EVisibility::Hidden)
+						]
+
+						+ SHorizontalBox::Slot()
+						.AutoWidth()
+						.VAlign(VAlign_Center)
+						.HAlign(HAlign_Center)
+						[
+							SAssignNew(StatusBarTextBlock, STextBlock)
+							.Font(FSlateFontInfo(FPaths::EngineContentDir() / TEXT("Slate/Fonts/Roboto-Regular.ttf"), 12))
+							.ColorAndOpacity(FLinearColor(1, 1, 1, 0.5))
+						]
+					]
+				];
+		}
+
 		// ===================================
 		// Section Name Track
 		// ===================================
@@ -390,35 +427,75 @@ void SAnimMontagePanel::Update()
 		// ===================================
 		{
 			int32 NumAnimTracks = Montage->SlotAnimTracks.Num();
-			SlotNameTextBoxes.Empty(NumAnimTracks);
-			SlotNameTextBoxes.AddZeroed(NumAnimTracks);
+
+			SlotNameComboBoxes.Empty(NumAnimTracks);
+			SlotNameComboSelectedNames.Empty(NumAnimTracks);
+			SlotWarningImages.Empty(NumAnimTracks);
+			SlotNameComboBoxes.AddZeroed(NumAnimTracks);
+			SlotNameComboSelectedNames.AddZeroed(NumAnimTracks);
+			SlotWarningImages.AddZeroed(NumAnimTracks);
+
+			RefreshComboLists();
+			check(SlotNameComboBoxes.Num() == NumAnimTracks);
+			check(SlotNameComboSelectedNames.Num() == NumAnimTracks);
 
 			for (int32 SlotAnimIdx = 0; SlotAnimIdx < NumAnimTracks; SlotAnimIdx++)
 			{
 				TSharedRef<S2ColumnWidget> SectionTrack = Create2ColumnWidget(MontageSlots.ToSharedRef());
+				
+				int32 FoundIndex = SlotNameList.Find(SlotNameComboSelectedNames[SlotAnimIdx]);
+				TSharedPtr<FString> ComboItem = SlotNameComboListItems[FoundIndex];
 
 				// Right column
 				SectionTrack->RightColumn->AddSlot()
-					.VAlign(VAlign_Center)
+				.VAlign(VAlign_Center)
+				.HAlign(HAlign_Fill)
 				.FillHeight(1)
 				[
-					SNew(SHorizontalBox)
+					SNew(SVerticalBox)
 
-					+SHorizontalBox::Slot()
-					.AutoWidth()
-					.Padding( FMargin(5.0f, 0.5f) )
+					+ SVerticalBox::Slot()
+					.HAlign(HAlign_Fill)
 					[
-						SAssignNew(SlotNameTextBoxes[SlotAnimIdx], SEditableTextBox)
-						.Style( FEditorStyle::Get(), "SpecialEditableTextBox" )
-						.HintText( LOCTEXT("SlotName", "Slot Name") )
-						.Text(this, &SAnimMontagePanel::GetMontageSlotName, SlotAnimIdx)
-						.Font( FEditorStyle::GetFontStyle("Editor.SearchBoxFont") )
-						.SelectAllTextWhenFocused( true )
-						.RevertTextOnEscape( true )
-						.ClearKeyboardFocusOnCommit( false )
-						.OnTextCommitted( this, &SAnimMontagePanel::OnSlotNodeNameChangeCommit, SlotAnimIdx )
-						.OnTextChanged(this, &SAnimMontagePanel::OnSlotNameChanged, SlotAnimIdx)
+						SAssignNew(SlotNameComboBoxes[SlotAnimIdx], STextComboBox)
+						.OptionsSource(&SlotNameComboListItems)
+						.OnSelectionChanged(this, &SAnimMontagePanel::OnSlotNameChanged, SlotAnimIdx)
+						.OnComboBoxOpening(this, &SAnimMontagePanel::OnSlotListOpening, SlotAnimIdx)
+						.InitiallySelectedItem(ComboItem)
+						.ContentPadding(2)
+						.ToolTipText(*ComboItem)
 					]
+
+					+ SVerticalBox::Slot()
+						.HAlign(HAlign_Left)
+						[
+							SNew(SHorizontalBox)
+							+ SHorizontalBox::Slot()
+							.AutoWidth()
+							.HAlign(HAlign_Left)
+							[
+								SNew(SButton)
+								.Text(LOCTEXT("AnimSlotNode_DetailPanelManageButtonLabel", "Anim Slot Manager"))
+								.ToolTipText(LOCTEXT("AnimSlotNode_DetailPanelManageButtonToolTipText", "Open Anim Slot Manager to edit Slots and Groups."))
+								.OnClicked(this, &SAnimMontagePanel::OnOpenAnimSlotManager)
+								.Content()
+								[
+									SNew(SImage)
+									.Image(FEditorStyle::GetBrush("MeshPaint.FindInCB"))
+								]
+							]
+
+							+ SHorizontalBox::Slot()
+								.AutoWidth()
+								.FillWidth(2.f)
+								.HAlign(HAlign_Left)
+								.VAlign(VAlign_Center)
+								[
+									SAssignNew(SlotWarningImages[SlotAnimIdx], SImage)
+									.Image(FEditorStyle::GetBrush("AnimSlotManager.Warning"))
+									.Visibility(EVisibility::Hidden)
+								]
+						]
 				];
 
 				SectionTrack->LeftColumn->AddSlot()
@@ -499,6 +576,8 @@ void SAnimMontagePanel::Update()
 			}
 		}
 	}
+
+	UpdateSlotGroupWarningVisibility();
 }
 
 void SAnimMontagePanel::SetMontage(class UAnimMontage * InMontage)
@@ -531,7 +610,7 @@ void SAnimMontagePanel::SummonTrackContextMenu( FMenuBuilder& MenuBuilder, float
 	MenuBuilder.EndSection();
 
 	// Sections
-	MenuBuilder.BeginSection("AnimMontageSections", LOCTEXT("Sections", "Sections") );
+	MenuBuilder.BeginSection("AnimMontageSections", LOCTEXT("Sections", "Sections"));
 	{
 		UIAction.ExecuteAction.BindRaw(this, &SAnimMontagePanel::OnNewSectionClicked, static_cast<float>(DataPosX));
 		MenuBuilder.AddMenuEntry(LOCTEXT("NewMontageSection", "New Montage Section"), LOCTEXT("NewMontageSectionToolTip", "Adds a new Montage Section"), FSlateIcon(), UIAction);
@@ -578,22 +657,7 @@ void SAnimMontagePanel::SummonBranchNodeContextMenu( FMenuBuilder& MenuBuilder,
 /** Slots */
 void SAnimMontagePanel::OnNewSlotClicked()
 {
-	// Show dialog to enter new track name
-	TSharedRef<STextEntryPopup> TextEntry =
-		SNew(STextEntryPopup)
-		.Label( LOCTEXT("NewSlotNameLabel", "Slot Name") )
-		.OnTextCommitted( this, &SAnimMontagePanel::CreateNewSlot );
-
-
-	// Show dialog to enter new event name
-	FSlateApplication::Get().PushMenu(
-		AsShared(), // Menu being summoned from a menu that is closing: Parent widget should be k2 not the menu thats open or it will be closed when the menu is dismissed
-		TextEntry,
-		FSlateApplication::Get().GetCursorPos(),
-		FPopupTransitionEffect( FPopupTransitionEffect::TypeInPopup )
-		);
-
-	TextEntry->FocusDefaultWidget();
+	MontageEditor.Pin()->AddNewMontageSlot(FAnimSlotGroup::DefaultSlotName.ToString());
 }
 
 void SAnimMontagePanel::CreateNewSlot(const FText& NewSlotName, ETextCommit::Type CommitInfo)
@@ -725,51 +789,167 @@ void SAnimMontagePanel::ClearSelected()
 	MontageEditor.Pin()->ClearDetailsView();
 }
 
-void SAnimMontagePanel::OnSlotNodeNameChangeCommit(const FText& NewText, ETextCommit::Type CommitInfo, int32 SlotNodeIndex)
+void SAnimMontagePanel::OnSlotNameChanged(TSharedPtr<FString> NewSelection, ESelectInfo::Type SelectInfo, int32 AnimSlotIndex)
 {
-	MontageEditor.Pin()->RenameSlotNode(SlotNodeIndex, NewText.ToString());
+	// if it's set from code, we did that on purpose
+	if (SelectInfo != ESelectInfo::Direct)
+	{
+		int32 ItemIndex = SlotNameComboListItems.Find(NewSelection);
+		if (ItemIndex != INDEX_NONE)
+		{
+			FName NewSlotName = SlotNameList[ItemIndex];
+
+			SlotNameComboSelectedNames[AnimSlotIndex] = NewSlotName;
+			if (SlotNameComboBoxes[AnimSlotIndex].IsValid())
+			{
+				SlotNameComboBoxes[AnimSlotIndex]->SetToolTipText(FText::FromString(*NewSelection));
+			}
+
+			if (Montage->GetSkeleton()->ContainsSlotName(NewSlotName))
+			{
+				MontageEditor.Pin()->RenameSlotNode(AnimSlotIndex, NewSlotName.ToString());
+			}
+			
+			UpdateSlotGroupWarningVisibility();
+		}
+	}
 }
 
-void SAnimMontagePanel::CheckSlotName(const FText& SlotName, int32 SlotNodeIndex, bool bShouldCheckCollapsed) const
+void SAnimMontagePanel::OnSlotListOpening(int32 AnimSlotIndex)
 {
-	if (!SlotNameTextBoxes.IsValidIndex(SlotNodeIndex))
-	{
-		return;
-	}
+	// Refresh Slot Names, in case we used the Anim Slot Manager to make changes.
+	RefreshComboLists(true);
+}
 
-	bool bCanSetError = true;
-	if (bShouldCheckCollapsed)
+FReply SAnimMontagePanel::OnOpenAnimSlotManager()
+{
+	if (Persona.IsValid())
 	{
-		// Check whether the text box is collapsed or not. If collapsed, hide the error message
-		FWidgetPath WidgetPath;
-		FSlateApplication::Get().FindPathToWidget(SlotNameTextBoxes[SlotNodeIndex].ToSharedRef(), WidgetPath);
-
-		bCanSetError = WidgetPath.IsValid();
+		Persona.Pin()->GetTabManager()->InvokeTab(FPersonaTabs::SkeletonSlotNamesID);
 	}
+	return FReply::Handled();
+}
 
-	if (bCanSetError && SlotName.IsEmpty())
+void SAnimMontagePanel::RefreshComboLists(bool bOnlyRefreshIfDifferent /*= false*/)
+{
+	// Make sure all slots defined in the montage are registered in our skeleton.
+	int32 NumAnimTracks = Montage->SlotAnimTracks.Num();
+	for (int32 TrackIndex = 0; TrackIndex < NumAnimTracks; TrackIndex++)
 	{
-		FText Error = LOCTEXT("Error_SlotNameIsEmpty", "Please provide a slot name for this asset.");
-		SlotNameTextBoxes[SlotNodeIndex]->SetError(Error);
+		FName TrackSlotName = Montage->SlotAnimTracks[TrackIndex].SlotName;
+		Montage->GetSkeleton()->RegisterSlotNode(TrackSlotName);
+		SlotNameComboSelectedNames[TrackIndex] = TrackSlotName;
 	}
-	else
+
+	// Refresh Slot Names
 	{
-		SlotNameTextBoxes[SlotNodeIndex]->SetError(FText::GetEmpty());
+		TArray<TSharedPtr<FString>>	NewSlotNameComboListItems;
+		TArray<FName> NewSlotNameList;
+
+		bool bIsSlotNameListDifferent = false;
+
+		const TArray<FAnimSlotGroup>& SlotGroups = Montage->GetSkeleton()->GetSlotGroups();
+		for (auto SlotGroup : SlotGroups)
+		{
+			int32 Index = 0;
+			for (auto SlotName : SlotGroup.SlotNames)
+			{
+				NewSlotNameList.Add(SlotName);
+
+				FString ComboItemString = FString::Printf(TEXT("%s.%s"), *SlotGroup.GroupName.ToString(), *SlotName.ToString());
+				NewSlotNameComboListItems.Add(MakeShareable(new FString(ComboItemString)));
+
+				bIsSlotNameListDifferent = bIsSlotNameListDifferent || (!SlotNameComboListItems.IsValidIndex(Index) || (SlotNameComboListItems[Index] != NewSlotNameComboListItems[Index]));
+				Index++;
+			}
+		}
+
+		// Refresh if needed
+		if (bIsSlotNameListDifferent || !bOnlyRefreshIfDifferent || (NewSlotNameComboListItems.Num() == 0))
+		{
+			SlotNameComboListItems = NewSlotNameComboListItems;
+			SlotNameList = NewSlotNameList;
+
+			// Update Combo Boxes
+			for (int32 TrackIndex = 0; TrackIndex < NumAnimTracks; TrackIndex++)
+			{
+				if (SlotNameComboBoxes[TrackIndex].IsValid())
+				{
+					FName SelectedSlotName = SlotNameComboSelectedNames[TrackIndex];
+					if (Montage->GetSkeleton()->ContainsSlotName(SelectedSlotName))
+					{
+						int32 FoundIndex = SlotNameList.Find(SelectedSlotName);
+						TSharedPtr<FString> ComboItem = SlotNameComboListItems[FoundIndex];
+
+						SlotNameComboBoxes[TrackIndex]->SetSelectedItem(ComboItem);
+						SlotNameComboBoxes[TrackIndex]->SetToolTipText(FText::FromString(*ComboItem));
+					}
+					SlotNameComboBoxes[TrackIndex]->RefreshOptions();
+				}
+			}
+		}
 	}
 }
 
-FText SAnimMontagePanel::GetMontageSlotName(int32 SlotIndex) const
+void SAnimMontagePanel::UpdateSlotGroupWarningVisibility()
 {
-	FText SlotName = MontageEditor.Pin()->GetMontageSlotName(SlotIndex);
+	bool bShowStatusBarWarning = false;
+	FName MontageGroupName = Montage->GetGroupName();
 
-	CheckSlotName(SlotName, SlotIndex, true);
+	int32 NumAnimTracks = Montage->SlotAnimTracks.Num();
+	if (NumAnimTracks > 0)
+	{
+		TArray<FName> UniqueSlotNameList;
+		for (int32 TrackIndex = 0; TrackIndex < NumAnimTracks; TrackIndex++)
+		{
+			FName CurrentSlotName = SlotNameComboSelectedNames[TrackIndex];
+			FName CurrentSlotGroupName = Montage->GetSkeleton()->GetSlotGroupName(CurrentSlotName);
 
-	return SlotName;
-}
+			int32 SlotNameCount = 0;
+			for (int32 Index = 0; Index < NumAnimTracks; Index++)
+			{
+				if (CurrentSlotName == SlotNameComboSelectedNames[Index])
+				{
+					SlotNameCount++;
+				}
+			}
+			
+			// Verify that slot names are unique.
+			bool bSlotNameAlreadyInUse = UniqueSlotNameList.Contains(CurrentSlotName);
+			if (!bSlotNameAlreadyInUse)
+			{
+				UniqueSlotNameList.Add(CurrentSlotName);
+			}
 
-void SAnimMontagePanel::OnSlotNameChanged(const FText& NewText, int32 SlotNodeIndex)
-{
-	CheckSlotName(NewText, SlotNodeIndex);
+			bool bDifferentGroupName = (CurrentSlotGroupName != MontageGroupName);
+			bool bShowWarning = bDifferentGroupName || bSlotNameAlreadyInUse;
+			bShowStatusBarWarning = bShowStatusBarWarning || bShowWarning;
+
+			SlotWarningImages[TrackIndex]->SetVisibility(bShowWarning ? EVisibility::Visible : EVisibility::Hidden);
+			if (bDifferentGroupName)
+			{
+				FText WarningText = FText::Format(LOCTEXT("AnimMontagePanel_SlotGroupMismatchToolTipText", "Slot's group '{0}' is different than the Montage's group '{1}'. All slots must belong to the same group."), FText::FromName(CurrentSlotGroupName), FText::FromName(MontageGroupName));
+				SlotWarningImages[TrackIndex]->SetToolTipText(WarningText);
+				StatusBarTextBlock->SetText(WarningText);
+				StatusBarTextBlock->SetToolTipText(WarningText);
+			}
+			if (bSlotNameAlreadyInUse)
+			{
+				FText WarningText = FText::Format(LOCTEXT("AnimMontagePanel_SlotNameAlreadyInUseToolTipText", "Slot named '{0}' is already used in this Montage. All slots must be unique"), FText::FromName(CurrentSlotName));
+				SlotWarningImages[TrackIndex]->SetToolTipText(WarningText);
+				StatusBarTextBlock->SetText(WarningText);
+				StatusBarTextBlock->SetToolTipText(WarningText);
+			}
+		}
+	}
+
+	// Update Status bar
+	StatusBarWarningImage->SetVisibility(bShowStatusBarWarning ? EVisibility::Visible : EVisibility::Hidden);
+	if (!bShowStatusBarWarning)
+	{
+		StatusBarTextBlock->SetText(FText::Format(LOCTEXT("AnimMontagePanel_StatusBarText", "Montage Group: '{0}'"), FText::FromName(MontageGroupName)));
+		StatusBarTextBlock->SetToolTipText(LOCTEXT("AnimMontagePanel_StatusBarToolTipText", "The Montage Group is set by the first slot's group."));
+	}
 }
 
 #undef LOCTEXT_NAMESPACE
\ No newline at end of file
diff --git a/Engine/Source/Editor/Persona/Private/SAnimMontagePanel.h b/Engine/Source/Editor/Persona/Private/SAnimMontagePanel.h
--- a/Engine/Source/Editor/Persona/Private/SAnimMontagePanel.h
+++ b/Engine/Source/Editor/Persona/Private/SAnimMontagePanel.h
@@ -78,20 +78,8 @@ public:
 	void ShowBranchPointInDetailsView(int32 BranchPointIndex);
 	void ShowSectionInDetailsView(int32 SectionIndex);
 
-	void OnSlotNodeNameChangeCommit(const FText& NewText, ETextCommit::Type CommitInfo, int32 SlotNodeIndex);
-
 	void ClearSelected();
 
-	// helper method to check whether the slot name is empty or not. If empty, shows an error message to provide a valid name
-	void CheckSlotName(const FText& SlotName, int32 SlotNodeIndex, bool bShouldCheckCollapsed = false) const;
-
-	// check the slot name whether valid or not while the user is typing
-	void OnSlotNameChanged(const FText& NewText, int32 SlotNodeIndex);
-
-	// get slot name from a montage editor and check the slot name whether valid or not
-	FText GetMontageSlotName(int32 SlotIndex) const;
-
-
 private:
 	TWeakPtr<FPersona> Persona;
 	TWeakPtr<SMontageEditor>	MontageEditor;
@@ -105,9 +93,26 @@ private:
 
 	STrackNodeSelectionSet SelectionSet;
 
-	// text box array for multiple slot names
-	TArray<TSharedPtr<SEditableTextBox>> SlotNameTextBoxes;
-	
+	TArray< TSharedPtr< class STextComboBox> >	SlotNameComboBoxes;
+	TArray< FName > SlotNameComboSelectedNames;
+
+	TArray< TSharedPtr< FString > > SlotNameComboListItems;
+	TArray< FName > SlotNameList;
+
+	TArray< TSharedPtr<SImage> > SlotWarningImages;
+
+	/************************************************************************/
+	/* Status Bar                                                                     */
+	/************************************************************************/
+	TSharedPtr<STextBlock> StatusBarTextBlock;
+	TSharedPtr<SImage> StatusBarWarningImage;
+
+	void OnSlotNameChanged(TSharedPtr<FString> NewSelection, ESelectInfo::Type SelectInfo, int32 AnimSlotIndex);
+	void OnSlotListOpening(int32 AnimSlotIndex);
+	FReply OnOpenAnimSlotManager();
+	void RefreshComboLists(bool bOnlyRefreshIfDifferent = false);
+	void UpdateSlotGroupWarningVisibility();
+
 	void CreateNewSlot(const FText& NewSlotName, ETextCommit::Type CommitInfo);
 	void CreateNewSection(const FText& NewSectionName, ETextCommit::Type CommitInfo, float StartTime);
 	void CreateNewBranch(const FText& NewBranchName, ETextCommit::Type CommitInfo, float StartTime);
diff --git a/Engine/Source/Editor/Persona/Private/SMontageEditor.cpp b/Engine/Source/Editor/Persona/Private/SMontageEditor.cpp
--- a/Engine/Source/Editor/Persona/Private/SMontageEditor.cpp
+++ b/Engine/Source/Editor/Persona/Private/SMontageEditor.cpp
@@ -460,7 +460,7 @@ void SMontageEditor::EnsureSlotNode()
 {
 	if (MontageObj && MontageObj->SlotAnimTracks.Num()==0)
 	{
-		AddNewMontageSlot(TEXT(""));
+		AddNewMontageSlot(FAnimSlotGroup::DefaultSlotName.ToString());
 		MontageObj->MarkPackageDirty();
 	}
 }
diff --git a/Engine/Source/Editor/Persona/Private/SSkeletonSlotGroupNames.cpp b/Engine/Source/Editor/Persona/Private/SSkeletonSlotGroupNames.cpp
deleted file mode 100644
diff --git a/Engine/Source/Editor/Persona/Private/SSkeletonSlotGroupNames.h b/Engine/Source/Editor/Persona/Private/SSkeletonSlotGroupNames.h
deleted file mode 100644
diff --git a/Engine/Source/Editor/Persona/Private/SSkeletonSlotNames.cpp b/Engine/Source/Editor/Persona/Private/SSkeletonSlotNames.cpp
--- a/Engine/Source/Editor/Persona/Private/SSkeletonSlotNames.cpp
+++ b/Engine/Source/Editor/Persona/Private/SSkeletonSlotNames.cpp
@@ -11,6 +11,7 @@
 #include "SInlineEditableTextBlock.h"
 #include "SNotificationList.h"
 #include "NotificationManager.h"
+#include "STextEntryPopup.h"
 
 #define LOCTEXT_NAMESPACE "SkeletonSlotNames"
 
@@ -71,18 +72,23 @@ TSharedRef< SWidget > SSlotNameListRow::GenerateWidgetForColumn( const FName& Co
 {
 	check( ColumnName == ColumnId_SlotNameLabel );
 
-	return SNew( SVerticalBox )
-	+ SVerticalBox::Slot()
-	.AutoHeight()
-	.Padding( 0.0f, 4.0f )
-	.VAlign( VAlign_Center )
-	[
-		SAssignNew(Item->InlineEditableText, SInlineEditableTextBlock)
-		.Text( FText::FromName(Item->Name) )
-		.OnVerifyTextChanged(SlotNameListView.Get(), &SSkeletonSlotNames::OnVerifyNotifyNameCommit, Item)
-		.OnTextCommitted(SlotNameListView.Get(), &SSkeletonSlotNames::OnNotifyNameCommitted, Item)
-		.IsSelected(SlotNameListView.Get(), &SSkeletonSlotNames::IsSelected)
-	];
+	// Items can be either Slots or Groups.
+	FText ItemText = Item->bIsGroupItem ? FText::Format(LOCTEXT("AnimSlotManagerGroupItem", "(Group) {0}"), FText::FromName(Item->Name))
+		: FText::Format(LOCTEXT("AnimSlotManagerSlotItem", "(Slot) {0}"), FText::FromName(Item->Name));
+
+	return
+		SNew(SHorizontalBox)
+		+ SHorizontalBox::Slot()
+		.AutoWidth()
+		[
+			SNew(SExpanderArrow, SharedThis(this))
+		]
+		+ SHorizontalBox::Slot()
+		.AutoWidth()
+		[
+			SNew(STextBlock)
+			.Text(ItemText)
+		];
 }
 
 /////////////////////////////////////////////////////
@@ -91,13 +97,13 @@ TSharedRef< SWidget > SSlotNameListRow::GenerateWidgetForColumn( const FName& Co
 FSkeletonSlotNamesSummoner::FSkeletonSlotNamesSummoner(TSharedPtr<class FAssetEditorToolkit> InHostingApp)
 	: FWorkflowTabFactory(FPersonaTabs::SkeletonSlotNamesID, InHostingApp)
 {
-	TabLabel = LOCTEXT("SkeletonSlotNamesTabTitle", "Animation SlotName");
+	TabLabel = LOCTEXT("AnimSlotManagerTabTitle", "Anim Slot Manager");
 
 	EnableTabPadding();
 	bIsSingleton = true;
 
-	ViewMenuDescription = LOCTEXT("SkeletonSlotNamesMenu", "Animation SlotName");
-	ViewMenuTooltip = LOCTEXT("SkeletonSlotNames_ToolTip", "Shows the skeletons slot name list");
+	ViewMenuDescription = LOCTEXT("SkeletonSlotNamesMenu", "Anim Slots");
+	ViewMenuTooltip = LOCTEXT("SkeletonSlotNames_ToolTip", "Manage Skeleton's Slots and Groups.");
 }
 
 TSharedRef<SWidget> FSkeletonSlotNamesSummoner::CreateTabBody(const FWorkflowTabSpawnInfo& Info) const
@@ -116,10 +122,48 @@ void SSkeletonSlotNames::Construct(const FArguments& InArgs)
 
 	PersonaPtr.Pin()->RegisterOnPostUndo(FPersona::FOnPostUndo::CreateSP( this, &SSkeletonSlotNames::PostUndo ) );
 
+	// Toolbar
+	FToolBarBuilder ToolbarBuilder(TSharedPtr< FUICommandList >(), FMultiBoxCustomization::None);
+
+	// Save USkeleton
+	ToolbarBuilder.AddToolBarButton(
+		FUIAction(FExecuteAction::CreateSP(this, &SSkeletonSlotNames::OnSaveSkeleton))
+		, NAME_None
+		, LOCTEXT("AnimSlotManagerToolbarSaveLabel", "Save")
+		, LOCTEXT("AnimSlotManagerToolbarSaveTooltip", "Saves changes into Skeleton asset")
+		, FSlateIcon(FEditorStyle::GetStyleSetName(), "AnimSlotManager.SaveSkeleton")
+		);
+
+	ToolbarBuilder.AddSeparator();
+
+	// Add Slot
+	ToolbarBuilder.AddToolBarButton(
+		FUIAction(FExecuteAction::CreateSP(this, &SSkeletonSlotNames::OnAddSlot))
+		, NAME_None
+		, LOCTEXT("AnimSlotManagerToolbarAddSlotLabel", "Add Slot")
+		, LOCTEXT("AnimSlotManagerToolbarAddSlotTooltip", "Create a new unique Slot name")
+		, FSlateIcon(FEditorStyle::GetStyleSetName(), "AnimSlotManager.AddSlot")
+		);
+
+	// Add Group
+	ToolbarBuilder.AddToolBarButton(
+		FUIAction(FExecuteAction::CreateSP(this, &SSkeletonSlotNames::OnAddGroup))
+		, NAME_None
+		, LOCTEXT("AnimSlotManagerToolbarAddGroupLabel", "Add Group")
+		, LOCTEXT("AnimSlotManagerToolbarAddGroupTooltip", "Create a new unique Group name")
+		, FSlateIcon(FEditorStyle::GetStyleSetName(), "AnimSlotManager.AddGroup")
+		);
+
 	this->ChildSlot
 	[
 		SNew( SVerticalBox )
 
+		+ SVerticalBox::Slot()
+		.AutoHeight()
+		[
+			ToolbarBuilder.MakeWidget()
+		]
+
 		+ SVerticalBox::Slot()
 		.AutoHeight()
 		.Padding( FMargin( 0.0f, 0.0f, 0.0f, 4.0f ) )
@@ -128,16 +172,18 @@ void SSkeletonSlotNames::Construct(const FArguments& InArgs)
 			.SelectAllTextWhenFocused( true )
 			.OnTextChanged( this, &SSkeletonSlotNames::OnFilterTextChanged )
 			.OnTextCommitted( this, &SSkeletonSlotNames::OnFilterTextCommitted )
-			.HintText( LOCTEXT( "SlotNameSearchBoxHint", "Search Animation SlotName...") )
+			.HintText( LOCTEXT( "AnimSlotManagerSlotNameSearchBoxHint", "Slot name filter...") )
 		]
 
 		+ SVerticalBox::Slot()
 		.FillHeight( 1.0f )		// This is required to make the scrollbar work, as content overflows Slate containers by default
 		[
 			SAssignNew( SlotNameListView, SSlotNameListType )
-			.ListItemsSource( &NotifyList )
+			.TreeItemsSource(&NotifyList)
 			.OnGenerateRow( this, &SSkeletonSlotNames::GenerateNotifyRow )
+			.OnGetChildren(this, &SSkeletonSlotNames::GetChildrenForInfo)
 			.OnContextMenuOpening( this, &SSkeletonSlotNames::OnGetContextMenuContent )
+			.SelectionMode(ESelectionMode::Single)
 			.OnSelectionChanged( this, &SSkeletonSlotNames::OnNotifySelectionChanged )
 			.ItemHeight( 22.0f )
 			.HeaderRow
@@ -184,118 +230,216 @@ TSharedRef<ITableRow> SSkeletonSlotNames::GenerateNotifyRow(TSharedPtr<FDisplaye
 		.SlotNameListView( SharedThis( this ) );
 }
 
+void SSkeletonSlotNames::GetChildrenForInfo(TSharedPtr<FDisplayedSlotNameInfo> InInfo, TArray< TSharedPtr<FDisplayedSlotNameInfo> >& OutChildren)
+{
+	check(InInfo.IsValid());
+	OutChildren = InInfo->Children;
+}
+
 TSharedPtr<SWidget> SSkeletonSlotNames::OnGetContextMenuContent() const
 {
 	const bool bShouldCloseWindowAfterMenuSelection = true;
 	FMenuBuilder MenuBuilder( bShouldCloseWindowAfterMenuSelection, NULL);
 
-	MenuBuilder.BeginSection("SlotNameAction", LOCTEXT( "SlotNameActions", "Selected Notify Actions" ) );
+	TArray< TSharedPtr< FDisplayedSlotNameInfo > > SelectedItems = SlotNameListView.Get()->GetSelectedItems();
+
+	bool bHasSelectedItem = (SelectedItems.Num() > 0);
+	bool bShowGroupItem = bHasSelectedItem && SelectedItems[0].Get()->bIsGroupItem;
+	bool bShowSlotItem = bHasSelectedItem && !SelectedItems[0].Get()->bIsGroupItem;
+
+	if (bShowGroupItem)
+	{
+
+	}
+	else if (bShowSlotItem)
 	{
+		MenuBuilder.BeginSection("SlotManagerSlotActions", LOCTEXT("SlotManagerSlotActions", "Slot Actions"));
+		// Set Slot Group
 		{
-			FUIAction Action = FUIAction( FExecuteAction::CreateSP( this, &SSkeletonSlotNames::OnDeleteSlotName ), 
-				FCanExecuteAction::CreateSP( this, &SSkeletonSlotNames::CanPerformDelete ) );
-			const FText Label = LOCTEXT("DeleteSlotNameButtonLabel", "Delete");
-			const FText ToolTip = LOCTEXT("DeleteSlotNameButtonTooltip", "Deletes the selected anim SlotName.");
-			MenuBuilder.AddMenuEntry( Label, ToolTip, FSlateIcon(), Action);
+			MenuBuilder.AddSubMenu(
+				FText::Format(LOCTEXT("ContextMenuSetSlotGroupLabel", "Set Slot {0} Group to"), FText::FromName(SelectedItems[0].Get()->Name)),
+				FText::Format(LOCTEXT("ContextMenuSetSlotGroupToolTip", "Set Slot {0} Group"), FText::FromName(SelectedItems[0].Get()->Name)),
+				FNewMenuDelegate::CreateRaw(this, &SSkeletonSlotNames::FillSetSlotGroupSubMenu));
+			MenuBuilder.EndSection();
 		}
+	}
 
-		// for now disable rename, it doesn't work correctly
-// 		{
-// 			FUIAction Action = FUIAction( FExecuteAction::CreateSP( this, &SSkeletonSlotNames::OnRenameSlotName ), 
-// 				FCanExecuteAction::CreateSP( this, &SSkeletonSlotNames::CanPerformRename ) );
-// 			const FText Label = LOCTEXT("RenameSlotNameButtonLabel", "Rename");
-// 			const FText ToolTip = LOCTEXT("RenameSlotNameButtonTooltip", "Renames the selected anim SlotName.");
-// 			MenuBuilder.AddMenuEntry( Label, ToolTip, FSlateIcon(), Action);
-// 		}
+	MenuBuilder.BeginSection("SlotManagerGeneralActions", LOCTEXT("SlotManagerGeneralActions", "Slot Manager Actions"));
+	// Add Slot
+	{
+		FUIAction Action = FUIAction(FExecuteAction::CreateSP(this, &SSkeletonSlotNames::OnAddSlot));
+		const FText Label = LOCTEXT("AnimSlotManagerContextMenuAddSlotLabel", "Add Slot");
+		const FText ToolTip = LOCTEXT("AnimSlotManagerContextMenuAddSlotTooltip", "Adds a new Slot");
+		MenuBuilder.AddMenuEntry(Label, ToolTip, FSlateIcon(), Action);
+	}
+	// Add Group
+	{
+		FUIAction Action = FUIAction(FExecuteAction::CreateSP(this, &SSkeletonSlotNames::OnAddGroup));
+		const FText Label = LOCTEXT("AnimSlotManagerContextMenuAddGroupLabel", "Add Group");
+		const FText ToolTip = LOCTEXT("AnimSlotManagerContextMenuAddGroupTooltip", "Adds a new Group");
+		MenuBuilder.AddMenuEntry(Label, ToolTip, FSlateIcon(), Action);
 	}
 	MenuBuilder.EndSection();
 
 	return MenuBuilder.MakeWidget();
 }
 
-void SSkeletonSlotNames::OnNotifySelectionChanged(TSharedPtr<FDisplayedSlotNameInfo> Selection, ESelectInfo::Type SelectInfo)
+void SSkeletonSlotNames::FillSetSlotGroupSubMenu(FMenuBuilder& MenuBuilder)
 {
-	if(Selection.IsValid())
+	const TArray<FAnimSlotGroup>& SlotGroups = TargetSkeleton->GetSlotGroups();
+	for (auto SlotGroup : SlotGroups)
 	{
-		ShowNotifyInDetailsView(Selection->Name);
+		const FName& GroupName = SlotGroup.GroupName;
+
+		const FText ToolTip = FText::Format(LOCTEXT("ContextMenuSetSlotSubMenuToolTip", "Changes slot's group to {0}"), FText::FromName(GroupName));
+/*		FString Label = Class->GetDisplayNameText().ToString();*/
+		const FText Label = FText::FromName(GroupName);
+
+		FUIAction UIAction;
+		UIAction.ExecuteAction.BindRaw(this, &SSkeletonSlotNames::ContextMenuOnSetSlot,	GroupName);
+		MenuBuilder.AddMenuEntry(Label, ToolTip, FSlateIcon(), UIAction);
 	}
 }
 
-bool SSkeletonSlotNames::CanPerformDelete() const
+void SSkeletonSlotNames::ContextMenuOnSetSlot(FName InNewGroupName)
 {
-	TArray< TSharedPtr< FDisplayedSlotNameInfo > > SelectedRows = SlotNameListView->GetSelectedItems();
-	return SelectedRows.Num() > 0;
+	TArray< TSharedPtr< FDisplayedSlotNameInfo > > SelectedItems = SlotNameListView.Get()->GetSelectedItems();
+
+	bool bHasSelectedItem = (SelectedItems.Num() > 0);
+	bool bShowSlotItem = bHasSelectedItem && !SelectedItems[0].Get()->bIsGroupItem;
+
+	if (bShowSlotItem)
+	{
+		const FName SlotName = SelectedItems[0].Get()->Name;
+		if (TargetSkeleton->ContainsSlotName(SlotName))
+		{
+			const FScopedTransaction Transaction(LOCTEXT("AnimSlotManager_ContextMenuOnSetSlot", "Set Slot Group Name"));
+			TargetSkeleton->SetSlotGroupName(SlotName, InNewGroupName);
+			TargetSkeleton->Modify(true);
+
+			RefreshSlotNameListWithFilter();
+		}
+
+		// Highlight newly created item.
+		TSharedPtr< FDisplayedSlotNameInfo > Item = FindItemNamed(SlotName);
+		if (Item.IsValid())
+		{
+			SlotNameListView->SetSelection(Item);
+		}
+
+		FSlateApplication::Get().DismissAllMenus();
+	}
 }
 
-bool SSkeletonSlotNames::CanPerformRename() const
+void SSkeletonSlotNames::OnNotifySelectionChanged(TSharedPtr<FDisplayedSlotNameInfo> Selection, ESelectInfo::Type SelectInfo)
 {
-	TArray< TSharedPtr< FDisplayedSlotNameInfo > > SelectedRows = SlotNameListView->GetSelectedItems();
-	return SelectedRows.Num() == 1;
+	if(Selection.IsValid())
+	{
+		ShowNotifyInDetailsView(Selection->Name);
+	}
 }
 
-void SSkeletonSlotNames::OnDeleteSlotName()
+void SSkeletonSlotNames::OnSaveSkeleton()
 {
-	TArray< TSharedPtr< FDisplayedSlotNameInfo > > SelectedRows = SlotNameListView->GetSelectedItems();
-
-	const FScopedTransaction Transaction( LOCTEXT("DeleteSlotName", "Delete Slot Name") );
-
-	// this one deletes all SlotName with same name. 
-	TArray<FName> SelectedSlotNames;
-	TargetSkeleton->Modify();
-
-	for(int Selection = 0; Selection < SelectedRows.Num(); ++Selection)
+	if (TargetSkeleton)
 	{
-		FName& SelectedName = SelectedRows[Selection]->Name;
-		TargetSkeleton->RemoveSlotNodeName(SelectedName);
-		SelectedSlotNames.Add(SelectedName);
+		TArray< UPackage* > PackagesToSave;
+		PackagesToSave.Add(TargetSkeleton->GetOutermost());
+
+		FEditorFileUtils::PromptForCheckoutAndSave(PackagesToSave, /*bCheckDirty=*/ false, /*bPromptToSave=*/ false);
 	}
+}
 
-	TArray<FAssetData> CompatibleAnimMontages;
-	GetCompatibleAnimMontage(CompatibleAnimMontages);
+void SSkeletonSlotNames::OnAddSlot()
+{
+	TSharedRef<STextEntryPopup> TextEntry =
+		SNew(STextEntryPopup)
+		.Label(LOCTEXT("NewSlotName_AskSlotName", "New Slot Name"))
+		.OnTextCommitted(this, &SSkeletonSlotNames::AddSlotPopUpOnCommit);
+
+	// Show dialog to enter new track name
+	FSlateApplication::Get().PushMenu(
+		SharedThis(this),
+		TextEntry,
+		FSlateApplication::Get().GetCursorPos(),
+		FPopupTransitionEffect(FPopupTransitionEffect::TypeInPopup)
+		);
+
+	TextEntry->FocusDefaultWidget();
+}
 
-	int32 NumAnimationsModified = 0;
+void SSkeletonSlotNames::OnAddGroup()
+{
+	TSharedRef<STextEntryPopup> TextEntry =
+		SNew(STextEntryPopup)
+		.Label(LOCTEXT("NewGroupName_AskGroupName", "New Group Name"))
+		.OnTextCommitted(this, &SSkeletonSlotNames::AddGroupPopUpOnCommit);
+
+	// Show dialog to enter new track name
+	FSlateApplication::Get().PushMenu(
+		SharedThis(this),
+		TextEntry,
+		FSlateApplication::Get().GetCursorPos(),
+		FPopupTransitionEffect(FPopupTransitionEffect::TypeInPopup)
+		);
+
+	TextEntry->FocusDefaultWidget();
+}
 
-	for( int32 AssetIndex = 0; AssetIndex < CompatibleAnimMontages.Num(); ++AssetIndex )
+void SSkeletonSlotNames::AddSlotPopUpOnCommit(const FText & InNewSlotText, ETextCommit::Type CommitInfo)
+{
+	if (!InNewSlotText.IsEmpty())
 	{
-		const FAssetData& PossibleAnimAsset = CompatibleAnimMontages[AssetIndex];
-		UAnimMontage* Montage = CastChecked<UAnimMontage>(PossibleAnimAsset.GetAsset());
-	
-		for (auto & SlotAnim : Montage->SlotAnimTracks)
+		const FScopedTransaction Transaction(LOCTEXT("NewSlotName_AddSlotName", "Add New Slot Node Name"));
+
+		FName NewSlotName = FName(*InNewSlotText.ToString());
+		// Keep slot and group names unique
+		if (!TargetSkeleton->ContainsSlotName(NewSlotName) && (TargetSkeleton->FindAnimSlotGroup(NewSlotName) == NULL))
 		{
-			if ( SelectedSlotNames.Contains(SlotAnim.SlotName) )
-			{
-				// found one, set transaction mark
-				Montage->Modify();
+			TArray< TSharedPtr< FDisplayedSlotNameInfo > > SelectedItems = SlotNameListView->GetSelectedItems();
+			bool bHasSelectedItem = (SelectedItems.Num() > 0);
+			bool bShowGroupItem = bHasSelectedItem && SelectedItems[0].Get()->bIsGroupItem;
 
-				// if this name is deleted, just clear it
-				// I don't think changing to something else will help
-				SlotAnim.SlotName = NAME_None;
-				// although people shouldn't have same slot name in the same montage, 
-				// just in case anybody does, it will loop
-				// this might cause issue because Modify will get called again.
+			TargetSkeleton->SetSlotGroupName(NewSlotName, bShowGroupItem ? SelectedItems[0].Get()->Name : FAnimSlotGroup::DefaultGroupName);
+			TargetSkeleton->Modify(true);
 
-				++NumAnimationsModified;
-			}
+			RefreshSlotNameListWithFilter();
+		}
+
+		// Highlight newly created item.
+		TSharedPtr< FDisplayedSlotNameInfo > Item = FindItemNamed(NewSlotName);
+		if (Item.IsValid())
+		{
+			SlotNameListView->SetSelection(Item);
 		}
-	}
 
-// 	TMultiMap<class UAnimBlueprint *, class UAnimGraphNode_Slot *> AssetsToFix;
-// 	GetCompatibleAnimBlueprints( AssetsToFix );
+		FSlateApplication::Get().DismissAllMenus();
+	}
+}
 
-	if(NumAnimationsModified > 0)
+void SSkeletonSlotNames::AddGroupPopUpOnCommit(const FText & InNewGroupText, ETextCommit::Type CommitInfo)
+{
+	if (!InNewGroupText.IsEmpty())
 	{
-		// Tell the user that the socket is a duplicate
-		FFormatNamedArguments Args;
-		Args.Add( TEXT("NumAnimationsModified"), NumAnimationsModified );
-		FNotificationInfo Info( FText::Format( LOCTEXT( "SlotNamesDeleted", "{NumAnimationsModified} animation(s) modified to delete notifications" ), Args ) );
+		const FScopedTransaction Transaction(LOCTEXT("NewGroupName_AddGroupName", "Add New Slot Group Name"));
 
-		Info.bUseLargeFont = false;
-		Info.ExpireDuration = 5.0f;
+		FName NewGroupName = FName(*InNewGroupText.ToString());
+		// Keep slot and group names unique
+		if (!TargetSkeleton->ContainsSlotName(NewGroupName) && TargetSkeleton->AddSlotGroupName(NewGroupName))
+		{
+			TargetSkeleton->Modify(true);
+			RefreshSlotNameListWithFilter();
+		}
 
-		NotifyUser( Info );
-	}
+		// Highlight newly created item.
+		TSharedPtr< FDisplayedSlotNameInfo > Item = FindItemNamed(NewGroupName);
+		if (Item.IsValid())
+		{
+			SlotNameListView->SetSelection(Item);
+		}
 
-	CreateSlotNameList( NameFilterBox->GetText().ToString() );
+		FSlateApplication::Get().DismissAllMenus();
+	}
 }
 
 // void SSkeletonSlotNames::GetCompatibleAnimBlueprints( TMultiMap<class UAnimBlueprint *, class UAnimGraphNode_Slot *>& OutAssets )
@@ -313,123 +457,59 @@ void SSkeletonSlotNames::OnDeleteSlotName()
 // 	}
 // }
 
-void SSkeletonSlotNames::OnRenameSlotName()
+void SSkeletonSlotNames::RefreshSlotNameListWithFilter()
 {
-	TArray< TSharedPtr< FDisplayedSlotNameInfo > > SelectedRows = SlotNameListView->GetSelectedItems();
-
-	check(SelectedRows.Num() == 1); // Should be guaranteed by CanPerformRename
-
-	SelectedRows[0]->InlineEditableText->EnterEditingMode();
+	CreateSlotNameList( NameFilterBox->GetText().ToString() );
 }
 
-bool SSkeletonSlotNames::OnVerifyNotifyNameCommit( const FText& NewName, FText& OutErrorMessage, TSharedPtr<FDisplayedSlotNameInfo> Item )
+void SSkeletonSlotNames::CreateSlotNameList(const FString& SearchText)
 {
-	bool bValid(true);
-
-	if(NewName.IsEmpty())
-	{
-		OutErrorMessage = LOCTEXT( "NameMissing_Error", "You must provide a name." );
-		bValid = false;
-	}
+	NotifyList.Empty();
 
-	FName NotifyName( *NewName.ToString() );
-	if(NotifyName != Item->Name)
+	if ( TargetSkeleton )
 	{
-		if(TargetSkeleton->DoesHaveSlotNodeName(NotifyName))
+		const TArray<FAnimSlotGroup>& SlotGroups = TargetSkeleton->GetSlotGroups();
+		for (auto SlotGroup : SlotGroups)
 		{
-			OutErrorMessage = FText::Format( LOCTEXT("AlreadyInUseMessage", "'{0}' is already in use."), NewName );
-			bValid = false;
-		}
-	}
-
-	return bValid;
-}
-
-void SSkeletonSlotNames::OnNotifyNameCommitted( const FText& NewName, ETextCommit::Type, TSharedPtr<FDisplayedSlotNameInfo> Item )
-{
-	const FScopedTransaction Transaction( LOCTEXT("RenameSlotName", "Rename Anim Notify") );
-
-	FName NewSlotName = FName( *NewName.ToString() );
-	FName SlotToRename = Item->Name;
+			const FName& GroupName = SlotGroup.GroupName;
+			
+			TSharedRef<FDisplayedSlotNameInfo> GroupItem = FDisplayedSlotNameInfo::Make(GroupName, true);
+			SlotNameListView->SetItemExpansion(GroupItem, true);
+			NotifyList.Add(GroupItem);
 
-	// rename notify in skeleton
-	TargetSkeleton->Modify();
-
-	TArray<FAssetData> CompatibleAnimMontages;
-	GetCompatibleAnimMontage(CompatibleAnimMontages);
-
-	int32 NumAnimationsModified = 0;
-
-	for(int32 AssetIndex = 0; AssetIndex < CompatibleAnimMontages.Num(); ++AssetIndex)
-	{
-		const FAssetData& PossibleAnimAsset = CompatibleAnimMontages[AssetIndex];
-		UAnimMontage* Montage = Cast<UAnimMontage>(PossibleAnimAsset.GetAsset());
-
-		for (auto & SlotAnim : Montage->SlotAnimTracks)
-		{
-			if(SlotAnim.SlotName == SlotToRename)
+			for (auto SlotName : SlotGroup.SlotNames)
 			{
-				// found one, set transaction mark
-				Montage->Modify();
-
-				// if this name is deleted, just clear it
-				// I don't think changing to something else will help
-				SlotAnim.SlotName = NewSlotName;
-				// although people shouldn't have same slot name in the same montage, 
-				// just in case anybody does, it will loop
-				// this might cause issue because Modify will get called again.
-
-				++NumAnimationsModified;
+				if (SearchText.IsEmpty() || GroupName.ToString().Contains(SearchText) || SlotName.ToString().Contains(SearchText))
+				{
+					TSharedRef<FDisplayedSlotNameInfo> SlotItem = FDisplayedSlotNameInfo::Make(SlotName, false);
+					SlotNameListView->SetItemExpansion(SlotItem, true);
+					NotifyList[NotifyList.Num() - 1]->Children.Add(SlotItem);
+				}
 			}
 		}
 	}
 
-	if(NumAnimationsModified > 0)
-	{
-		// Tell the user that the socket is a duplicate
-		FFormatNamedArguments Args;
-		Args.Add( TEXT("NumAnimationsModified"), NumAnimationsModified );
-		FNotificationInfo Info( FText::Format( LOCTEXT( "SlotNamesRenamed", "{NumAnimationsModified} animation(s) modified to rename notification" ), Args ) );
-
-		Info.bUseLargeFont = false;
-		Info.ExpireDuration = 5.0f;
-
-		NotifyUser( Info );
-	}
-
-	RefreshSlotNameListWithFilter();
-}
-
-void SSkeletonSlotNames::RefreshSlotNameListWithFilter()
-{
-	CreateSlotNameList( NameFilterBox->GetText().ToString() );
+	SlotNameListView->RequestTreeRefresh();
 }
 
-void SSkeletonSlotNames::CreateSlotNameList( const FString& SearchText )
+TSharedPtr< FDisplayedSlotNameInfo > SSkeletonSlotNames::FindItemNamed(FName ItemName) const
 {
-	NotifyList.Empty();
-
-	if ( TargetSkeleton )
+	for (auto SlotGroupItem : NotifyList)
 	{
-		const TArray<FName> & SlotNodeNames = TargetSkeleton->GetSlotNodeNames();
-		for(int i = 0; i < SlotNodeNames.Num(); ++i)
+		if (SlotGroupItem->Name == ItemName)
 		{
-			const FName& NotifyName = SlotNodeNames[i];
-			if ( !SearchText.IsEmpty() )
-			{
-				if ( NotifyName.ToString().Contains( SearchText ) )
-				{
-					NotifyList.Add( FDisplayedSlotNameInfo::Make( NotifyName ) );
-				}
-			}
-			else
+			return SlotGroupItem;
+		}
+		for (auto SlotItem : SlotGroupItem->Children)
+		{
+			if (SlotItem->Name == ItemName)
 			{
-				NotifyList.Add( FDisplayedSlotNameInfo::Make( NotifyName ) );
+				return SlotItem;
 			}
 		}
 	}
 
-	SlotNameListView->RequestListRefresh();
+	return NULL;
 }
 
 void SSkeletonSlotNames::ShowNotifyInDetailsView(FName NotifyName)
diff --git a/Engine/Source/Editor/Persona/Private/SSkeletonSlotNames.h b/Engine/Source/Editor/Persona/Private/SSkeletonSlotNames.h
--- a/Engine/Source/Editor/Persona/Private/SSkeletonSlotNames.h
+++ b/Engine/Source/Editor/Persona/Private/SSkeletonSlotNames.h
@@ -30,19 +30,24 @@ class FDisplayedSlotNameInfo
 public:
 	FName Name;
 
+	bool bIsGroupItem;
+
+	TArray< TSharedPtr<FDisplayedSlotNameInfo> > Children;
+
 	/** Handle to editable text block for rename */
 	TSharedPtr<SInlineEditableTextBlock> InlineEditableText;
 
 	/** Static function for creating a new item, but ensures that you can only have a TSharedRef to one */
-	static TSharedRef<FDisplayedSlotNameInfo> Make(const FName& NotifyName)
+	static TSharedRef<FDisplayedSlotNameInfo> Make(const FName& InItemName, bool InbIsGroupItem)
 	{
-		return MakeShareable(new FDisplayedSlotNameInfo(NotifyName));
+		return MakeShareable(new FDisplayedSlotNameInfo(InItemName, InbIsGroupItem));
 	}
 
 protected:
 	/** Hidden constructor, always use Make above */
-	FDisplayedSlotNameInfo(const FName& InNotifyName)
-		: Name( InNotifyName )
+	FDisplayedSlotNameInfo(const FName& InItemName, bool InbIsGroupItem)
+		: Name(InItemName)
+		, bIsGroupItem(InbIsGroupItem)
 	{}
 
 	/** Hidden constructor, always use Make above */
@@ -50,7 +55,7 @@ protected:
 };
 
 /** Widgets list type */
-typedef SListView< TSharedPtr<FDisplayedSlotNameInfo> > SSlotNameListType;
+typedef STreeView< TSharedPtr<FDisplayedSlotNameInfo> > SSlotNameListType;
 
 class SSkeletonSlotNames : public SCompoundWidget, public FGCObject
 {
@@ -86,14 +91,6 @@ public:
 	virtual void AddReferencedObjects( FReferenceCollector& Collector ) override;
 	// FGCObject interface end
 
-	/** When user attempts to commit the name of a track*/
-	bool OnVerifyNotifyNameCommit( const FText& NewName, FText& OutErrorMessage, TSharedPtr<FDisplayedSlotNameInfo> Item );
-
-	/** When user commits the name of a track*/
-	void OnNotifyNameCommitted( const FText& NewName, ETextCommit::Type, TSharedPtr<FDisplayedSlotNameInfo> Item );
-
-	/** Dummy handler to stop editable text boxes swallowing our list selected events */
-	bool IsSelected(){return false;}
 private:
 
 	/** Called when the user changes the contents of the search box */
@@ -105,27 +102,35 @@ private:
 	/** Delegate handler for generating rows in SlotNameListView */ 
 	TSharedRef<ITableRow> GenerateNotifyRow( TSharedPtr<FDisplayedSlotNameInfo> InInfo, const TSharedRef<STableViewBase>& OwnerTable );
 
+	/** Get all children for a given entry in the list */
+	void GetChildrenForInfo(TSharedPtr<FDisplayedSlotNameInfo> InInfo, TArray< TSharedPtr<FDisplayedSlotNameInfo> >& OutChildren);
+
 	/** Delegate handler called when the user right clicks in SlotNameListView */
 	TSharedPtr<SWidget> OnGetContextMenuContent() const;
 
 	/** Delegate handler for when the user selects something in SlotNameListView */
 	void OnNotifySelectionChanged( TSharedPtr<FDisplayedSlotNameInfo> Selection, ESelectInfo::Type SelectInfo );
 
-	/** Delegate handler for determining whether we can show the delete menu options */
-	bool CanPerformDelete() const;
+	// Save Skeleton
+	void OnSaveSkeleton();
 
-	/** Delegate handler for deleting anim SlotName */
-	void OnDeleteSlotName();
+	// Add a new Slot
+	void OnAddSlot();
+	void AddSlotPopUpOnCommit(const FText& InNewSlotName, ETextCommit::Type CommitInfo);
 
-	/** Delegate handler for determining whether we can show the rename menu options */
-	bool CanPerformRename() const;
+	// Add a new Group
+	void OnAddGroup();
+	void AddGroupPopUpOnCommit(const FText& InNewGroupName, ETextCommit::Type CommitInfo);
 
-	/** Delegate handler for renaming anim SlotName */
-	void OnRenameSlotName();
+	// Set Slot Group
+	void FillSetSlotGroupSubMenu(FMenuBuilder& MenuBuilder);
+	void ContextMenuOnSetSlot(FName InNewGroupName);
 
 	/** Wrapper that populates SlotNameListView using current filter test */
 	void RefreshSlotNameListWithFilter();
 
+	TSharedPtr< FDisplayedSlotNameInfo > FindItemNamed(FName ItemName) const;
+
 	/** Populates SlotNameListView based on the skeletons SlotName and the supplied filter text */
 	void CreateSlotNameList( const FString& SearchText = FString("") );
 
diff --git a/Engine/Source/Editor/Persona/Private/Shared/PersonaMode.cpp b/Engine/Source/Editor/Persona/Private/Shared/PersonaMode.cpp
--- a/Engine/Source/Editor/Persona/Private/Shared/PersonaMode.cpp
+++ b/Engine/Source/Editor/Persona/Private/Shared/PersonaMode.cpp
@@ -6,7 +6,6 @@
 #include "SSkeletonAnimNotifies.h"
 #include "IDocumentation.h"
 #include "SAnimBlueprintParentPlayerList.h"
-#include "SSkeletonSlotGroupNames.h"
 #include "SSkeletonSlotNames.h"
 #include "SSkeletonSmartNameManager.h"
 
@@ -67,7 +66,6 @@ FPersonaAppMode::FPersonaAppMode(TSharedPtr<class FPersona> InPersona, FName InM
 	PersonaTabFactories.RegisterFactory(MakeShareable(new FSkeletonAnimNotifiesSummoner(InPersona)));
 	PersonaTabFactories.RegisterFactory(MakeShareable(new FRetargetManagerTabSummoner(InPersona)));
 	PersonaTabFactories.RegisterFactory(MakeShareable(new FSkeletonSlotNamesSummoner(InPersona)));
-	PersonaTabFactories.RegisterFactory(MakeShareable(new FSkeletonSlotGroupNamesSummoner(InPersona)));
 	PersonaTabFactories.RegisterFactory(MakeShareable(new FSkeletonCurveNameManagerSummoner(InPersona)));
 }
 
diff --git a/Engine/Source/Editor/UnrealEd/Private/Factories/AnimMontageFactory.cpp b/Engine/Source/Editor/UnrealEd/Private/Factories/AnimMontageFactory.cpp
--- a/Engine/Source/Editor/UnrealEd/Private/Factories/AnimMontageFactory.cpp
+++ b/Engine/Source/Editor/UnrealEd/Private/Factories/AnimMontageFactory.cpp
@@ -84,7 +84,7 @@ UObject* UAnimMontageFactory::FactoryCreateNew(UClass* Class, UObject* InParent,
 			FSlotAnimationTrack NewTrack;
 			NewTrack.AnimTrack.AnimSegments.Add(NewSegment);
 
-			NewTrack.SlotName = FName("");
+			NewTrack.SlotName = FAnimSlotGroup::DefaultSlotName;
 			AnimMontage->SlotAnimTracks.Add( NewTrack );
 			AnimMontage->SetSequenceLength(SourceAnimation->SequenceLength);
 		}
diff --git a/Engine/Source/Runtime/Engine/Classes/Animation/AnimInstance.h b/Engine/Source/Runtime/Engine/Classes/Animation/AnimInstance.h
--- a/Engine/Source/Runtime/Engine/Classes/Animation/AnimInstance.h
+++ b/Engine/Source/Runtime/Engine/Classes/Animation/AnimInstance.h
@@ -472,6 +472,9 @@ protected:
 	/** Stop all montages that are active **/
 	void StopAllMontages(float BlendOut);
 
+	/** Stop all active montages belonging to 'InGroupName' */
+	void StopAllMontagesByGroupName(FName InGroupName, float BlendOutTime);
+
 	/** Update weight of montages  **/
 	virtual void Montage_UpdateWeight(float DeltaSeconds);
 	/** Advance montages **/
diff --git a/Engine/Source/Runtime/Engine/Classes/Animation/AnimMontage.h b/Engine/Source/Runtime/Engine/Classes/Animation/AnimMontage.h
--- a/Engine/Source/Runtime/Engine/Classes/Animation/AnimMontage.h
+++ b/Engine/Source/Runtime/Engine/Classes/Animation/AnimMontage.h
@@ -44,7 +44,6 @@ struct FCompositeSection
 		, NextSectionName(NAME_None)
 	{
 	}
-
 };
 
 /**
@@ -63,9 +62,8 @@ struct FSlotAnimationTrack
 	FAnimTrack AnimTrack;
 
 	FSlotAnimationTrack()
-		: SlotName(NAME_None)
+		: SlotName(FAnimSlotGroup::DefaultSlotName)
 	{}
-
 };
 
 /** 
@@ -390,6 +388,12 @@ public:
 	 */
 	FTransform ExtractRootMotionFromTrackRange(float StartTrackPosition, float EndTrackPosition) const;
 
+	/** Get the Montage's Group Name. This is the group from the first slot.  */
+	ENGINE_API FName GetGroupName() const;
+
+	/** true if valid, false otherwise. Will log warning if not valid. */
+	bool HasValidSlotSetup() const;
+
 private:
 	/** 
 	 * Utility function to check if CurrentTime is between FirstIndex and SecondIndex of CompositeSections
diff --git a/Engine/Source/Runtime/Engine/Classes/Animation/AnimNode_Slot.h b/Engine/Source/Runtime/Engine/Classes/Animation/AnimNode_Slot.h
--- a/Engine/Source/Runtime/Engine/Classes/Animation/AnimNode_Slot.h
+++ b/Engine/Source/Runtime/Engine/Classes/Animation/AnimNode_Slot.h
@@ -20,11 +20,6 @@ struct ENGINE_API FAnimNode_Slot : public FAnimNode_Base
 	UPROPERTY(EditAnywhere, BlueprintReadWrite, Category=Settings, meta=(CustomizeProperty))
 	FName SlotName;
 
-	// The group name of this slot, exposed to gameplay code, etc...
-	// you can only play one montage per group
-	UPROPERTY(EditAnywhere, BlueprintReadWrite, Category=Settings, meta=(CustomizeProperty))
-	FName GroupName;
-
 protected:
 	/** Weight of Source Branch. This is the weight of the input pose coming from children.
 	This is different than (1.f - SourceWeight) since the Slot can play additive animations, which are overlayed on top of the source pose. */
diff --git a/Engine/Source/Runtime/Engine/Classes/Animation/Skeleton.h b/Engine/Source/Runtime/Engine/Classes/Animation/Skeleton.h
--- a/Engine/Source/Runtime/Engine/Classes/Animation/Skeleton.h
+++ b/Engine/Source/Runtime/Engine/Classes/Animation/Skeleton.h
@@ -186,6 +186,33 @@ struct FRigConfiguration
 	UPROPERTY()
 	TArray<FNameMapping> BoneMappingTable;
 };
+
+USTRUCT()
+struct FAnimSlotGroup
+{
+	GENERATED_USTRUCT_BODY()
+
+public:
+	static ENGINE_API const FName DefaultGroupName;
+	static ENGINE_API const FName DefaultSlotName;
+
+	UPROPERTY()
+	FName GroupName;
+
+	UPROPERTY()
+	TArray<FName> SlotNames;
+
+	FAnimSlotGroup()
+		: GroupName(DefaultGroupName)
+	{
+	}
+
+	FAnimSlotGroup(FName InGroupName)
+		: GroupName(InGroupName)
+	{
+	}
+};
+
 /**
  *	USkeleton : that links between mesh and animation
  *		- Bone hierarchy for animations
@@ -247,20 +274,35 @@ public:
 	UPROPERTY()
 	FSmartNameContainer SmartNames;
 
+	/************************************************************************/
+	/* Slot Groups */
+	/************************************************************************/
+private:
+	// serialized slot groups and slot names.
+	UPROPERTY()
+	TArray<FAnimSlotGroup> SlotGroups;
+
+	/** SlotName to GroupName TMap, only at runtime, not serialized. **/
+	TMap<FName, FName> SlotToGroupNameMap;
+
+	void BuildSlotToGroupMap();
+
+public:
+	ENGINE_API FAnimSlotGroup* FindAnimSlotGroup(const FName& InGroupName);
+	ENGINE_API const TArray<FAnimSlotGroup>& GetSlotGroups() const;
+	ENGINE_API bool ContainsSlotName(const FName& InSlotName);
+	ENGINE_API void RegisterSlotNode(const FName& InSlotName);
+	ENGINE_API void SetSlotGroupName(const FName& InSlotName, const FName& InGroupName);
+	/** Returns true if Group is added, false if it already exists */
+	ENGINE_API bool AddSlotGroupName(const FName& InNewGroupName);
+	ENGINE_API FName GetSlotGroupName(const FName& InSlotName) const;
+
 #if WITH_EDITORONLY_DATA
 private:
 	/** The default skeletal mesh to use when previewing this skeleton */
 	UPROPERTY(duplicatetransient, AssetRegistrySearchable)
 	TAssetPtr<class USkeletalMesh> PreviewSkeletalMesh;
 
-	/** @todo collection of slotnode names available for this skeleton */
-	UPROPERTY()
-	TArray<FName> SlotNodeNames;
-
-	/** @todo collection of graoup names for slot node available for this skeleton */
-	UPROPERTY()
-	TArray<FName> SlotGroupNames;
-
 	UPROPERTY()
 	FRigConfiguration RigConfig;
 
@@ -334,17 +376,6 @@ public:
 	// @todo document
 	ENGINE_API void AddNewAnimationNotify(FName NewAnimNotifyName);
 
-	// @todo document
-	ENGINE_API void AddSlotNodeName(FName SlotNodeName);
-	ENGINE_API void RemoveSlotNodeName(FName SlotNodeName);
-	ENGINE_API bool DoesHaveSlotNodeName(FName SlotNodeName) const;
-	ENGINE_API const TArray<FName> & GetSlotNodeNames() const;
-
-	ENGINE_API void AddSlotGroupName(FName GroupName);
-	ENGINE_API void RemoveSlotGroupName(FName GroupName);
-	ENGINE_API bool DoesHaveSlotGroupName(FName GroupName) const;
-	ENGINE_API const TArray<FName> & GetSlotGroupNames() const;
-
 	/** Returns the skeletons preview mesh, loading it if necessary */
 	ENGINE_API USkeletalMesh* GetPreviewMesh(bool bFindIfNotSet=false);
 	ENGINE_API USkeletalMesh* GetPreviewMesh() const;
@@ -644,9 +675,6 @@ public:
 	ENGINE_API URig * GetRig() const;
 #endif
 
-public:
-	// this should be outside of editor because slot node is initialize to it
-	ENGINE_API static const FName DefaultSlotGroupName;
 private:
 	void RegenerateGuid();
 };
diff --git a/Engine/Source/Runtime/Engine/Private/Animation/AnimInstance.cpp b/Engine/Source/Runtime/Engine/Private/Animation/AnimInstance.cpp
--- a/Engine/Source/Runtime/Engine/Private/Animation/AnimInstance.cpp
+++ b/Engine/Source/Runtime/Engine/Private/Animation/AnimInstance.cpp
@@ -1691,14 +1691,23 @@ bool UAnimInstance::IsPlayingSlotAnimation(UAnimSequenceBase* Asset, FName SlotN
 /** Play a Montage. Returns Length of Montage in seconds. Returns 0.f if failed to play. */
 float UAnimInstance::Montage_Play(UAnimMontage * MontageToPlay, float InPlayRate/*= 1.f*/)
 {
-	if (MontageToPlay && (MontageToPlay->SequenceLength > 0.f))
+	if (MontageToPlay && (MontageToPlay->SequenceLength > 0.f) && MontageToPlay->HasValidSlotSetup())
 	{
 		if (CurrentSkeleton->IsCompatible(MontageToPlay->GetSkeleton()))
 		{
-			//@fixme laurent -- only stop montages that belong to same group. 
-			// or to ensure single root motion montage playing.
-			// when stopping old animations, make sure it does give current new blendintime to blend out
-			StopAllMontages(MontageToPlay->BlendInTime);
+			// Enforce 'a single montage at once per group' rule
+			FName NewMontageGroupName = MontageToPlay->GetGroupName();
+			StopAllMontagesByGroupName(NewMontageGroupName, MontageToPlay->BlendInTime);
+
+			// Enforce 'a single root motion montage at once' rule.
+			if (MontageToPlay->bEnableRootMotionTranslation || MontageToPlay->bEnableRootMotionRotation)
+			{
+				FAnimMontageInstance* ActiveRootMotionMontageInstance = GetRootMotionMontageInstance();
+				if (ActiveRootMotionMontageInstance)
+				{
+					ActiveRootMotionMontageInstance->Stop(MontageToPlay->BlendInTime);
+				}
+			}
 
 			FAnimMontageInstance * NewInstance = new FAnimMontageInstance(this);
 			check(NewInstance);
@@ -2162,6 +2171,18 @@ void UAnimInstance::StopAllMontages(float BlendOut)
 	}
 }
 
+void UAnimInstance::StopAllMontagesByGroupName(FName InGroupName, float BlendOutTime)
+{
+	for (int32 InstanceIndex = MontageInstances.Num() - 1; InstanceIndex >= 0; InstanceIndex--)
+	{
+		FAnimMontageInstance * MontageInstance = MontageInstances[InstanceIndex];
+		if (MontageInstance && MontageInstance->IsActive() && MontageInstance->Montage && (MontageInstance->Montage->GetGroupName() == InGroupName))
+		{
+			MontageInstances[InstanceIndex]->Stop(BlendOutTime, true);
+		}
+	}
+}
+
 void UAnimInstance::OnMontageInstanceStopped(FAnimMontageInstance & StoppedMontageInstance)
 {
 	UAnimMontage * MontageStopped = StoppedMontageInstance.Montage;
diff --git a/Engine/Source/Runtime/Engine/Private/Animation/AnimMontage.cpp b/Engine/Source/Runtime/Engine/Private/Animation/AnimMontage.cpp
--- a/Engine/Source/Runtime/Engine/Private/Animation/AnimMontage.cpp
+++ b/Engine/Source/Runtime/Engine/Private/Animation/AnimMontage.cpp
@@ -411,21 +411,36 @@ void UAnimMontage::PostLoad()
 	}
 
 	// verify if skeleton matches, otherwise clear it, this can happen if anim sequence has been modified when this hasn't been loaded. 
-	USkeleton* MySkeleton = GetSkeleton();
-	for (int32 I=0; I<SlotAnimTracks.Num(); ++I)
 	{
-		if ( SlotAnimTracks[I].AnimTrack.AnimSegments.Num() > 0 )
+		USkeleton* MySkeleton = GetSkeleton();
+		for (int32 I=0; I<SlotAnimTracks.Num(); ++I)
 		{
-			UAnimSequence * Sequence = Cast<UAnimSequence>(SlotAnimTracks[I].AnimTrack.AnimSegments[0].AnimReference);
-			if ( Sequence && Sequence->GetSkeleton() != MySkeleton )
+			if ( SlotAnimTracks[I].AnimTrack.AnimSegments.Num() > 0 )
 			{
-				SlotAnimTracks[I].AnimTrack.AnimSegments[0].AnimReference = 0;
-				MarkPackageDirty();
-				break;
+				UAnimSequence * Sequence = Cast<UAnimSequence>(SlotAnimTracks[I].AnimTrack.AnimSegments[0].AnimReference);
+				if ( Sequence && Sequence->GetSkeleton() != MySkeleton )
+				{
+					SlotAnimTracks[I].AnimTrack.AnimSegments[0].AnimReference = 0;
+					MarkPackageDirty();
+					break;
+				}
 			}
 		}
 	}
 #endif // WITH_EDITORONLY_DATA
+
+	// Register Slots w/ Skeleton
+	{
+		USkeleton* MySkeleton = GetSkeleton();
+		if (MySkeleton)
+		{
+			for (int32 SlotIndex = 0; SlotIndex < SlotAnimTracks.Num(); SlotIndex++)
+			{
+				FName SlotName = SlotAnimTracks[SlotIndex].SlotName;
+				MySkeleton->RegisterSlotNode(SlotName);
+			}
+		}
+	}
 }
 
 void UAnimMontage::SortAnimBranchingPointByTime()
@@ -608,6 +623,64 @@ FTransform UAnimMontage::ExtractRootMotionFromTrackRange(float StartTrackPositio
 	return RootMotion.RootMotionTransform;
 }
 
+/** Get Montage's Group Name */
+FName UAnimMontage::GetGroupName() const
+{
+	USkeleton* MySkeleton = GetSkeleton();
+	if (MySkeleton && (SlotAnimTracks.Num() > 0))
+	{
+		return MySkeleton->GetSlotGroupName(SlotAnimTracks[0].SlotName);
+	}
+
+	return FAnimSlotGroup::DefaultGroupName;
+}
+
+bool UAnimMontage::HasValidSlotSetup() const
+{
+	// We only need to worry about this if we have multiple tracks.
+	// Montages with a single track will always have a valid slot setup.
+	int32 NumAnimTracks = SlotAnimTracks.Num();
+	if (NumAnimTracks > 1)
+	{
+		USkeleton* MySkeleton = GetSkeleton();
+		if (MySkeleton)
+		{
+			FName MontageGroupName = GetGroupName();
+			TArray<FName> UniqueSlotNameList;
+			UniqueSlotNameList.Add(SlotAnimTracks[0].SlotName);
+
+			for (int32 TrackIndex = 1; TrackIndex < NumAnimTracks; TrackIndex++)
+			{
+				// Verify that slot names are unique.
+				FName CurrentSlotName = SlotAnimTracks[TrackIndex].SlotName;
+				bool bSlotNameAlreadyInUse = UniqueSlotNameList.Contains(CurrentSlotName);
+				if (!bSlotNameAlreadyInUse)
+				{
+					UniqueSlotNameList.Add(CurrentSlotName);
+				}
+				else
+				{
+					UE_LOG(LogAnimation, Warning, TEXT("Montage '%s' not properly setup. Slot named '%s' is already used in this Montage. All slots must be unique"),
+						*GetFullName(), *CurrentSlotName.ToString());
+					return false;
+				}
+
+				// Verify that all slots belong to the same group.
+				FName CurrentSlotGroupName = MySkeleton->GetSlotGroupName(CurrentSlotName);
+				bool bDifferentGroupName = (CurrentSlotGroupName != MontageGroupName);
+				if (bDifferentGroupName)
+				{
+					UE_LOG(LogAnimation, Warning, TEXT("Montage '%s' not properly setup. Slot's group '%s' is different than the Montage's group '%s'. All slots must belong to the same group."),
+						*GetFullName(), *CurrentSlotGroupName.ToString(), *MontageGroupName.ToString());
+					return false;
+				}
+			}
+		}
+	}
+
+	return true;
+}
+
 void UAnimMontage::EvaluateCurveData(class UAnimInstance* Instance, float CurrentTime, float BlendWeight ) const
 {
 	Super::EvaluateCurveData(Instance, CurrentTime, BlendWeight);
diff --git a/Engine/Source/Runtime/Engine/Private/Animation/AnimNode_Slot.cpp b/Engine/Source/Runtime/Engine/Private/Animation/AnimNode_Slot.cpp
--- a/Engine/Source/Runtime/Engine/Private/Animation/AnimNode_Slot.cpp
+++ b/Engine/Source/Runtime/Engine/Private/Animation/AnimNode_Slot.cpp
@@ -64,7 +64,7 @@ void FAnimNode_Slot::GatherDebugData(FNodeDebugData& DebugData)
 {
 	FString DebugLine = DebugData.GetNodeName(this);
 	DebugLine += FString::Printf(TEXT("(Slot Name: '%s' Weight:%.1f%%)"), *SlotName.ToString(), SlotNodeWeight*100.f);
-
+	
 	bool const bIsPoseSource = (SourceWeight <= ZERO_ANIMWEIGHT_THRESH);
 	DebugData.AddDebugItem(DebugLine, bIsPoseSource);
 	Source.GatherDebugData(DebugData.BranchFlow(SourceWeight));
@@ -90,8 +90,8 @@ void FAnimNode_Slot::GatherDebugData(FNodeDebugData& DebugData)
 }
 
 FAnimNode_Slot::FAnimNode_Slot()
-: GroupName(USkeleton::DefaultSlotGroupName)
-, SourceWeight(0.f)
-, SlotNodeWeight(0.f)
+	: SlotName(FAnimSlotGroup::DefaultSlotName)
+	, SourceWeight(0.f)
+	, SlotNodeWeight(0.f)
 {
 }
diff --git a/Engine/Source/Runtime/Engine/Private/Animation/Skeleton.cpp b/Engine/Source/Runtime/Engine/Private/Animation/Skeleton.cpp
--- a/Engine/Source/Runtime/Engine/Private/Animation/Skeleton.cpp
+++ b/Engine/Source/Runtime/Engine/Private/Animation/Skeleton.cpp
@@ -27,7 +27,8 @@ const FName USkeleton::RigTag = FName(TEXT("Rig"));
 // Names of smartname containers for skeleton properties
 const FName USkeleton::AnimCurveMappingName = FName(TEXT("AnimationCurves"));
 
-const FName USkeleton::DefaultSlotGroupName = FName(TEXT("Default"));
+const FName FAnimSlotGroup::DefaultGroupName = FName(TEXT("DefaultGroup"));
+const FName FAnimSlotGroup::DefaultSlotName = FName(TEXT("DefaultSlot"));
 
 FArchive& operator<<(FArchive& Ar, FReferencePose & P)
 {
@@ -48,12 +49,6 @@ USkeleton::USkeleton(const FObjectInitializer& ObjectInitializer)
 {
 	// Make sure we have somewhere for curve names.
 	SmartNames.AddContainer(AnimCurveMappingName);
-
-#if WITH_EDITORONLY_DATA
-	SlotGroupNames.Empty();
-	SlotGroupNames.Add(DefaultSlotGroupName);
-
-#endif
 }
 
 void USkeleton::PostInitProperties()
@@ -78,6 +73,9 @@ void USkeleton::PostLoad()
 		ConvertToFReferenceSkeleton();
 	}
 
+	// Build look up table between Slot nodes and their Group.
+	BuildSlotToGroupMap();
+
 	// catch any case if guid isn't valid
 	check(Guid.IsValid());
 }
@@ -910,50 +908,6 @@ void USkeleton::UnregisterOnSkeletonHierarchyChanged(void* Unregister)
 
 #endif
 
-void USkeleton::AddSlotNodeName(FName SlotNodeName)
-{
-	SlotNodeNames.AddUnique(SlotNodeName);
-}
-
-void USkeleton::RemoveSlotNodeName(FName SlotNodeName)
-{
-	SlotNodeNames.Remove(SlotNodeName);
-}
-
-bool USkeleton::DoesHaveSlotNodeName(FName SlotNodeName) const
-{
-	return SlotNodeNames.Contains(SlotNodeName);
-}
-
-const TArray<FName> & USkeleton::GetSlotNodeNames() const
-{
-	return SlotNodeNames;
-}
-
-void USkeleton::AddSlotGroupName(FName GroupName)
-{
-	SlotGroupNames.AddUnique(GroupName);
-}
-
-void USkeleton::RemoveSlotGroupName(FName GroupName)
-{
-	// we can't delete the deafult one
-	if (GroupName != DefaultSlotGroupName)
-	{
-		SlotGroupNames.Remove(GroupName);
-	}
-}
-
-const TArray<FName> & USkeleton::GetSlotGroupNames() const
-{
-	return SlotGroupNames;
-}
-
-bool USkeleton::DoesHaveSlotGroupName(FName GroupName) const
-{
-	return SlotGroupNames.Contains(GroupName);
-}
-
 bool USkeleton::AddSmartnameAndModify(FName ContainerName, FName NewName, FSmartNameMapping::UID& NewUid)
 {
 	bool Successful = false;
@@ -1001,6 +955,111 @@ void USkeleton::RemoveSmartnameAndModify(FName ContainerName, FSmartNameMapping:
 
 #endif // WITH_EDITORONLY_DATA
 
+const TArray<FAnimSlotGroup>& USkeleton::GetSlotGroups() const
+{
+	return SlotGroups;
+}
+
+void USkeleton::BuildSlotToGroupMap()
+{
+	SlotToGroupNameMap.Empty();
+
+	for (auto SlotGroup : SlotGroups)
+	{
+		for (auto SlotName : SlotGroup.SlotNames)
+		{
+			SlotToGroupNameMap.Add(SlotName, SlotGroup.GroupName);
+		}
+	}
+}
+
+FAnimSlotGroup* USkeleton::FindAnimSlotGroup(const FName& InGroupName)
+{
+	for (int32 GroupIndex = 0; GroupIndex < SlotGroups.Num(); GroupIndex++)
+	{
+		FAnimSlotGroup& SlotGroup = SlotGroups[GroupIndex];
+		if (SlotGroup.GroupName == InGroupName)
+		{
+			return &SlotGroup;
+		}
+	}
+	return NULL;
+}
+
+bool USkeleton::ContainsSlotName(const FName& InSlotName) 
+{
+	return SlotToGroupNameMap.Contains(InSlotName);
+}
+
+void USkeleton::RegisterSlotNode(const FName& InSlotName)
+{
+	// verify the slot name exists, if not create it in the default group.
+	if (!ContainsSlotName(InSlotName))
+	{
+		SetSlotGroupName(InSlotName, FAnimSlotGroup::DefaultGroupName);
+	}
+}
+
+void USkeleton::SetSlotGroupName(const FName& InSlotName, const FName& InGroupName)
+{
+// See if Slot already exists and belongs to a group.
+	const FName* FoundGroupNamePtr = SlotToGroupNameMap.Find(InSlotName);
+
+	// If slot exists, but is not in the right group, remove it from there
+	if (FoundGroupNamePtr && ((*FoundGroupNamePtr) != InGroupName))
+	{
+		FAnimSlotGroup* OldSlotGroupPtr = FindAnimSlotGroup(*FoundGroupNamePtr);
+		if (OldSlotGroupPtr)
+		{
+			OldSlotGroupPtr->SlotNames.RemoveSingleSwap(InSlotName);
+		}
+	}
+
+	// Add the slot to the right group if it's not
+	if ((FoundGroupNamePtr == NULL) || (*FoundGroupNamePtr != InGroupName))
+	{
+		// If the SlotGroup does not exist, create it.
+		FAnimSlotGroup* SlotGroupPtr = FindAnimSlotGroup(InGroupName);
+		if (SlotGroupPtr == NULL)
+		{
+			SlotGroups.AddZeroed(1);
+			SlotGroupPtr = &SlotGroups.Last();
+			SlotGroupPtr->GroupName = InGroupName;
+		}
+		// Add Slot to group.
+		SlotGroupPtr->SlotNames.Add(InSlotName);
+		// Keep our TMap up to date.
+		SlotToGroupNameMap.Add(InSlotName, InGroupName);
+	}
+}
+
+bool USkeleton::AddSlotGroupName(const FName& InNewGroupName)
+{
+	FAnimSlotGroup* ExistingSlotGroupPtr = FindAnimSlotGroup(InNewGroupName);
+	if (ExistingSlotGroupPtr == NULL)
+	{
+		// if not found, create a new one.
+		SlotGroups.AddZeroed(1);
+		ExistingSlotGroupPtr = &SlotGroups.Last();
+		ExistingSlotGroupPtr->GroupName = InNewGroupName;
+		return true;
+	}
+
+	return false;
+}
+
+FName USkeleton::GetSlotGroupName(const FName& InSlotName) const
+{
+	const FName* FoundGroupNamePtr = SlotToGroupNameMap.Find(InSlotName);
+	if (FoundGroupNamePtr)
+	{
+		return *FoundGroupNamePtr;
+	}
+
+	// If Group name cannot be found, use DefaultSlotGroupName.
+	return FAnimSlotGroup::DefaultGroupName;
+}
+
 void USkeleton::RegenerateGuid()
 {
 	Guid = FGuid::NewGuid();
