Fixes a crash when compiling an AnimBlueprint while the Property Matrix is open for the animation instance.

diff --git a/Engine/Source/Editor/Persona/Private/Persona.cpp b/Engine/Source/Editor/Persona/Private/Persona.cpp
--- a/Engine/Source/Editor/Persona/Private/Persona.cpp
+++ b/Engine/Source/Editor/Persona/Private/Persona.cpp
@@ -1466,6 +1466,9 @@ void FPersona::Compile()
 	UObject* CurrentDebugObject = GetBlueprintObj()->GetObjectBeingDebugged();
 	const bool bIsDebuggingPreview = (PreviewComponent != NULL) && PreviewComponent->IsAnimBlueprintInstanced() && (PreviewComponent->AnimScriptInstance == CurrentDebugObject);
 
+	// Force close any asset editors that are using the AnimScriptInstance (such as the Property Matrix), the class will be garbage collected
+	FAssetEditorManager::Get().CloseOtherEditors(PreviewComponent->AnimScriptInstance, nullptr);
+	
 	// Compile the blueprint
 	FBlueprintEditor::Compile();
 
