Persona: Fixed non-unity build

diff --git a/Engine/Source/Editor/Persona/Private/SAnimMontagePanel.cpp b/Engine/Source/Editor/Persona/Private/SAnimMontagePanel.cpp
--- a/Engine/Source/Editor/Persona/Private/SAnimMontagePanel.cpp
+++ b/Engine/Source/Editor/Persona/Private/SAnimMontagePanel.cpp
@@ -1,8 +1,6 @@
 // Copyright 1998-2014 Epic Games, Inc. All Rights Reserved.
 
-
 #include "PersonaPrivatePCH.h"
-
 #include "SAnimMontagePanel.h"
 #include "ScopedTransaction.h"
 #include "SCurveEditor.h"
@@ -12,9 +10,12 @@
 #include "Editor/UnrealEd/Public/DragAndDrop/AssetDragDropOp.h"
 #include "SExpandableArea.h"
 #include "STextEntryPopup.h"
+#include "STextComboBox.h"
+
 
 #define LOCTEXT_NAMESPACE "AnimMontagePanel"
 
+
 //////////////////////////////////////////////////////////////////////////
 // SMontageBranchingPointNode
 
