Montage Improvements
- Allow relative and proportional linking for montage elements
- Show trigger thresh of notifies in details panel
- Fixup segments on reimport to stop t-pose or frozen animations
- Only automatic improvements in this CL, requring no resave for assets

diff --git a/Engine/Source/Editor/AnimGraph/Private/AnimPreviewInstance.cpp b/Engine/Source/Editor/AnimGraph/Private/AnimPreviewInstance.cpp
--- a/Engine/Source/Editor/AnimGraph/Private/AnimPreviewInstance.cpp
+++ b/Engine/Source/Editor/AnimGraph/Private/AnimPreviewInstance.cpp
@@ -259,7 +259,7 @@ void UAnimPreviewInstance::MontagePreview_StepForward()
 					SetPosition(StoppedAt, false);
 				}
 				int32 LastPreviewSectionIdx = MontagePreview_FindLastSection(MontagePreviewStartSectionIdx);
-				if (FMath::Abs(CurrentTime - (Montage->CompositeSections[LastPreviewSectionIdx].StartTime + Montage->GetSectionLength(LastPreviewSectionIdx))) <= MontagePreview_CalculateStepLength())
+				if (FMath::Abs(CurrentTime - (Montage->CompositeSections[LastPreviewSectionIdx].GetTime() + Montage->GetSectionLength(LastPreviewSectionIdx))) <= MontagePreview_CalculateStepLength())
 				{
 					// we're at the end, jump right to the end
 					Montage_JumpToSectionsEnd(Montage->GetSectionName(LastPreviewSectionIdx));
@@ -302,12 +302,12 @@ void UAnimPreviewInstance::MontagePreview_StepBackward()
 					SetPosition(StoppedAt, false);
 				}
 				int32 LastPreviewSectionIdx = MontagePreview_FindLastSection(MontagePreviewStartSectionIdx);
-				if (FMath::Abs(CurrentTime - (Montage->CompositeSections[LastPreviewSectionIdx].StartTime + Montage->GetSectionLength(LastPreviewSectionIdx))) <= MontagePreview_CalculateStepLength())
+				if (FMath::Abs(CurrentTime - (Montage->CompositeSections[LastPreviewSectionIdx].GetTime() + Montage->GetSectionLength(LastPreviewSectionIdx))) <= MontagePreview_CalculateStepLength())
 				{
 					// special case as we could stop at the end of our last section which is also beginning of following section - we don't want to get stuck there, but be inside of our starting section
 					Montage_JumpToSection(Montage->GetSectionName(LastPreviewSectionIdx));
 				}
-				else if (FMath::Abs(CurrentTime - Montage->CompositeSections[MontagePreviewStartSectionIdx].StartTime) <= MontagePreview_CalculateStepLength())
+				else if (FMath::Abs(CurrentTime - Montage->CompositeSections[MontagePreviewStartSectionIdx].GetTime()) <= MontagePreview_CalculateStepLength())
 				{
 					// we're at the end of playing backward, jump right to the end
 					Montage_JumpToSectionsEnd(Montage->GetSectionName(MontagePreviewStartSectionIdx));
diff --git a/Engine/Source/Editor/Persona/Private/AnimNotifyDetails.cpp b/Engine/Source/Editor/Persona/Private/AnimNotifyDetails.cpp
--- a/Engine/Source/Editor/Persona/Private/AnimNotifyDetails.cpp
+++ b/Engine/Source/Editor/Persona/Private/AnimNotifyDetails.cpp
@@ -10,6 +10,7 @@
 #include "DetailCategoryBuilder.h"
 #include "IDetailsView.h"
 #include "AssetSearchBoxUtilPersona.h"
+#include "STextComboBox.h"
 
 TSharedRef<IDetailCustomization> FAnimNotifyDetails::MakeInstance()
 {
@@ -18,135 +19,321 @@ TSharedRef<IDetailCustomization> FAnimNotifyDetails::MakeInstance()
 
 void FAnimNotifyDetails::CustomizeDetails( IDetailLayoutBuilder& DetailBuilder )
 {
-	const UClass* DetailObjectClass = DetailBuilder.GetDetailsView().GetBaseClass();
-	if (DetailObjectClass)
+	const UClass* DetailObjectClass = nullptr;
+	const UClass* BaseClass = DetailBuilder.GetDetailsView().GetBaseClass();
+	TArray<TWeakObjectPtr<UObject>> SelectedObjects;
+	TArray<UClass*> NotifyClasses;
+	DetailBuilder.GetObjectsBeingCustomized(SelectedObjects);
+
+	check(SelectedObjects.Num() > 0);
+	UEditorNotifyObject* EditorObject = Cast<UEditorNotifyObject>(SelectedObjects[0].Get());
+	check(EditorObject);
+	UpdateSlotNames(EditorObject->AnimObject);
+
+	// Hide notify objects that aren't set
+	UObject* NotifyPtr = nullptr;
+	FString NotifyClassName;
+	TSharedRef<IPropertyHandle> NotifyPropHandle = DetailBuilder.GetProperty(TEXT("Event.Notify"));
+	NotifyPropHandle->GetValue(NotifyPtr);
+
+	IDetailCategoryBuilder& AnimNotifyCategory = DetailBuilder.EditCategory(TEXT("AnimNotify"), TEXT(""), ECategoryPriority::TypeSpecific);
+
+	// Check existence of notify, get rid of the property if not set
+	if(!NotifyPtr)
 	{
-		if (DetailObjectClass->GetName().Find(TEXT("AnimNotify_PlayParticleEffect")) != INDEX_NONE)
-		{
-			AddBoneNameProperty(DetailBuilder, DetailObjectClass, TEXT("SocketName"), TEXT("AnimNotify"));
-		}
-		else if (DetailObjectClass->GetName().Find(TEXT("AnimNotify_PlaySound")) != INDEX_NONE)
+		DetailBuilder.HideProperty(TEXT("Event.Notify"));
+
+		NotifyPropHandle = DetailBuilder.GetProperty(TEXT("Event.NotifyStateClass"));
+		NotifyPropHandle->GetValue(NotifyPtr);
+
+		// Check existence of notify state, get rid of the property if not set
+		if(!NotifyPtr)
 		{
-			AddBoneNameProperty(DetailBuilder, DetailObjectClass, TEXT("AttachName"), TEXT("AnimNotify"));
+			DetailBuilder.HideProperty(TEXT("Event.NotifyStateClass"));
+			DetailBuilder.HideProperty(TEXT("Event.EndLink"));
 		}
-		else if (DetailObjectClass->GetName().Find(TEXT("AnimNotifyState_Trail")) != INDEX_NONE)
+		else
 		{
-			AddPropertyDefault(DetailBuilder, DetailObjectClass, TEXT("PSTemplate"), TEXT("Trail"));	
-			AddBoneNameProperty(DetailBuilder, DetailObjectClass, TEXT("FirstSocketName"), TEXT("Trail"));
-			AddBoneNameProperty(DetailBuilder, DetailObjectClass, TEXT("SecondSocketName"), TEXT("Trail"));
-			AddPropertyDefault(DetailBuilder, DetailObjectClass, TEXT("WidthScaleMode"), TEXT("Trail"));
-			AddCurveNameProperty(DetailBuilder, DetailObjectClass, TEXT("WidthScaleCurve"), TEXT("Trail"));
+			DetailObjectClass = NotifyPtr->GetClass();
+
+			// Get rid of the class selector in the details panel. It's not necessary for notifies
+			ClearInstancedSelectionDropDown(AnimNotifyCategory, NotifyPropHandle);
 		}
-		else if (DetailObjectClass->GetName().Find(TEXT("AnimNotifyState_TimedParticleEffect")) != INDEX_NONE)
+	}
+	else
+	{
+		DetailObjectClass = NotifyPtr->GetClass();
+		// Get rid of the class selector in the details panel. It's not necessary for notifies
+		ClearInstancedSelectionDropDown(AnimNotifyCategory, NotifyPropHandle);
+
+		// No state present, hide the entry
+		DetailBuilder.HideProperty(TEXT("Event.NotifyStateClass"));
+	}
+
+	TSharedRef<IPropertyHandle> TopHandle = DetailBuilder.GetProperty(TEXT("Event"));
+	if(Cast<UAnimMontage>(EditorObject->AnimObject))
+	{
+		CustomizeLinkProperties(DetailBuilder, TopHandle, EditorObject);
+	}
+	else
+	{
+		// No montage, hide link properties
+		HideLinkProperties(DetailBuilder, TopHandle);
+	}
+
+	// Customizations do not run for instanced properties, so we have to resolve the properties and then
+	// customize them here instead.
+	if(NotifyPropHandle->IsValidHandle())
+	{
+		uint32 NumChildren = 0;
+		NotifyPropHandle->GetNumChildren(NumChildren);
+		if(NumChildren > 0)
 		{
-			AddBoneNameProperty(DetailBuilder, DetailObjectClass, TEXT("SocketName"), TEXT("ParticleSystem"));
+			TSharedPtr<IPropertyHandle> BaseHandle = NotifyPropHandle->GetChildHandle(0);
+			DetailBuilder.HideProperty(NotifyPropHandle);
+
+			BaseHandle->GetNumChildren(NumChildren);
+			DetailBuilder.HideProperty(BaseHandle);
+
+			for(uint32 ChildIdx = 0 ; ChildIdx < NumChildren ; ++ChildIdx)
+			{
+				TSharedPtr<IPropertyHandle> NotifyProperty = BaseHandle->GetChildHandle(ChildIdx);
+
+				if(!CustomizeProperty(AnimNotifyCategory, NotifyPtr, NotifyProperty))
+				{
+					AnimNotifyCategory.AddProperty(NotifyProperty);
+				}
+			}
 		}
 	}
 }
 
-void FAnimNotifyDetails::AddPropertyDefault(IDetailLayoutBuilder& DetailBuilder, const UClass* PropertyClass, const TCHAR* PropertyName, const TCHAR* CategoryName)
+void FAnimNotifyDetails::AddBoneNameProperty(IDetailCategoryBuilder& CategoryBuilder, UObject* Notify,  TSharedPtr<IPropertyHandle> Property)
 {
-	IDetailCategoryBuilder& AnimNotifyCategory = DetailBuilder.EditCategory(CategoryName, TEXT(""), ECategoryPriority::TypeSpecific);
-	TSharedPtr<IPropertyHandle> Property = DetailBuilder.GetProperty(FName(PropertyName), PropertyClass);
-	if (Property.IsValid())
+	int32 PropIndex = NameProperties.Num();
+
+	if(Notify && Property->IsValidHandle())
 	{
-		AnimNotifyCategory.AddProperty(Property);
+		NameProperties.Add(Property);
+		// get all the possible suggestions for the bones and sockets.
+		if(const UAnimationAsset* AnimAsset = Cast<const UAnimationAsset>(Notify->GetOuter()))
+		{
+			if(const USkeleton* Skeleton = AnimAsset->GetSkeleton())
+			{
+				CategoryBuilder.AddProperty(Property.ToSharedRef())
+					.CustomWidget()
+					.NameContent()
+					[
+						SNew(SHorizontalBox)
+						+ SHorizontalBox::Slot()
+						.Padding(FMargin(2, 1, 0, 1))
+						[
+							SNew(STextBlock)
+							.Text(Property->GetPropertyDisplayName())
+							.Font(FEditorStyle::GetFontStyle(TEXT("PropertyWindow.NormalFont")))
+						]
+					]
+				.ValueContent()
+					[
+						SNew(SAssetSearchBoxForBones, Skeleton, Property)
+						.IncludeSocketsForSuggestions(true)
+						.MustMatchPossibleSuggestions(true)
+						.HintText(NSLOCTEXT("AnimNotifyDetails", "Hint Text", "Bone Name..."))
+						.OnTextCommitted(this, &FAnimNotifyDetails::OnSearchBoxCommitted, PropIndex)
+					];
+			}
+		}
 	}
 }
 
-void FAnimNotifyDetails::AddBoneNameProperty(IDetailLayoutBuilder& DetailBuilder, const UClass* PropertyClass, const TCHAR* PropertyName, const TCHAR* CategoryName)
+void FAnimNotifyDetails::AddCurveNameProperty(IDetailCategoryBuilder& CategoryBuilder, UObject* Notify, TSharedPtr<IPropertyHandle> Property)
 {
-	IDetailCategoryBuilder& AnimNotifyCategory = DetailBuilder.EditCategory(CategoryName, TEXT(""), ECategoryPriority::TypeSpecific);
 	int32 PropIndex = NameProperties.Num();
-	TSharedPtr<IPropertyHandle> BoneNameProperty = DetailBuilder.GetProperty(FName(PropertyName), PropertyClass);
-
-	if (BoneNameProperty.IsValid() && BoneNameProperty->GetProperty())
+	
+	if(Notify && Property->IsValidHandle())
 	{
-		NameProperties.Add(BoneNameProperty);
-		// get all the possible suggestions for the bones and sockets.
-		const TArray< TWeakObjectPtr<UObject> >& SelectedObjects = DetailBuilder.GetDetailsView().GetSelectedObjects();
-		if (SelectedObjects.Num() == 1)
+		NameProperties.Add(Property);
+
+		if(const UAnimationAsset* AnimAsset = Cast<const UAnimationAsset>(Notify->GetOuter()))
 		{
-			TWeakObjectPtr<UObject> SelectedObject = SelectedObjects[0];
-			const UObject* OuterObject = SelectedObject->GetOuter();
-			if (const UAnimationAsset* AnimAsset = Cast<const UAnimationAsset>(OuterObject))
+			if(const USkeleton* Skeleton = AnimAsset->GetSkeleton())
 			{
-				if (const USkeleton* Skeleton = AnimAsset->GetSkeleton())
-				{
-					AnimNotifyCategory.AddProperty(BoneNameProperty.ToSharedRef())
-						.CustomWidget()
-						.NameContent()
+				CategoryBuilder.AddProperty(Property.ToSharedRef())
+					.CustomWidget()
+					.NameContent()
+					[
+						SNew(SHorizontalBox)
+						+ SHorizontalBox::Slot()
+						.Padding(FMargin(2, 1, 0, 1))
 						[
-							SNew(SHorizontalBox)
-							+ SHorizontalBox::Slot()
-							.Padding(FMargin(2, 1, 0, 1))
-							[
-								SNew(STextBlock)
-								.Text(BoneNameProperty->GetPropertyDisplayName())
-								.Font(FEditorStyle::GetFontStyle(TEXT("PropertyWindow.NormalFont")))
-							]
+							SNew(STextBlock)
+							.Text(Property->GetPropertyDisplayName())
+							.Font(FEditorStyle::GetFontStyle(TEXT("PropertyWindow.NormalFont")))
 						]
-					.ValueContent()
-						[
-							SNew(SAssetSearchBoxForBones, Skeleton, BoneNameProperty)
-							.IncludeSocketsForSuggestions(true)
-							.MustMatchPossibleSuggestions(true)
-							.HintText(NSLOCTEXT("AnimNotifyDetails", "Hint Text", "Bone Name..."))
-							.OnTextCommitted(this, &FAnimNotifyDetails::OnSearchBoxCommitted, PropIndex)
-						];
-				}
+					]
+				.ValueContent()
+					[
+						SNew(SAssetSearchBoxForCurves, Skeleton, Property)
+						.IncludeSocketsForSuggestions(true)
+						.MustMatchPossibleSuggestions(true)
+						.HintText(NSLOCTEXT("AnimNotifyDetails", "Curve Name Hint Text", "Curve Name..."))
+						.OnTextCommitted(this, &FAnimNotifyDetails::OnSearchBoxCommitted, PropIndex)
+					];
 			}
 		}
 	}
 }
 
-void FAnimNotifyDetails::AddCurveNameProperty(IDetailLayoutBuilder& DetailBuilder, const UClass* PropertyClass, const TCHAR* PropertyName, const TCHAR* CategoryName)
+void FAnimNotifyDetails::OnSearchBoxCommitted(const FText& InSearchText, ETextCommit::Type CommitInfo, int32 PropertyIndex )
 {
-	IDetailCategoryBuilder& AnimNotifyCategory = DetailBuilder.EditCategory(CategoryName, TEXT(""), ECategoryPriority::TypeSpecific);
-	int32 PropIndex = NameProperties.Num();
-	TSharedPtr<IPropertyHandle> CurveNameProperty = DetailBuilder.GetProperty(FName(PropertyName), PropertyClass);
+	NameProperties[PropertyIndex]->SetValue( InSearchText.ToString() );
+}
 
-	if (CurveNameProperty.IsValid() && CurveNameProperty->GetProperty())
+void FAnimNotifyDetails::ClearInstancedSelectionDropDown(IDetailCategoryBuilder& CategoryBuilder, TSharedRef<IPropertyHandle> PropHandle, bool bShowChildren /*= true*/)
+{
+	IDetailPropertyRow& PropRow = CategoryBuilder.AddProperty(PropHandle);
+	
+	PropRow
+	.OverrideResetToDefault(false, FSimpleDelegate())
+	.CustomWidget(bShowChildren)
+	.NameContent()
+	[
+		PropHandle->CreatePropertyNameWidget(TEXT(""), TEXT(""), false)
+	]
+	.ValueContent()
+	[
+		SNullWidget::NullWidget
+	];
+}
+
+void FAnimNotifyDetails::CustomizeLinkProperties(IDetailLayoutBuilder& Builder, TSharedRef<IPropertyHandle> NotifyProperty, UEditorNotifyObject* EditorObject)
+{
+	uint32 NumChildProperties = 0;
+	NotifyProperty->GetNumChildren(NumChildProperties);
+
+	if(NumChildProperties > 0)
 	{
-		NameProperties.Add(CurveNameProperty);
-		// get all the possible suggestions for the curves.
-		const TArray< TWeakObjectPtr<UObject> >& SelectedObjects = DetailBuilder.GetDetailsView().GetSelectedObjects();
-		if (SelectedObjects.Num() == 1)
+		IDetailCategoryBuilder& LinkCategory = Builder.EditCategory(TEXT("AnimLink"));
+		for(uint32 ChildIdx = 0 ; ChildIdx < NumChildProperties ; ++ChildIdx)
 		{
-			TWeakObjectPtr<UObject> SelectedObject = SelectedObjects[0];
-			const UObject* OuterObject = SelectedObject->GetOuter();
-			if (const UAnimationAsset* AnimAsset = Cast<const UAnimationAsset>(OuterObject))
+			TSharedPtr<IPropertyHandle> ChildHandle = NotifyProperty->GetChildHandle(ChildIdx);
+			FString OuterFieldType = ChildHandle->GetProperty()->GetOuterUField()->GetName();
+
+			if(ChildHandle->GetProperty()->GetName() == GET_MEMBER_NAME_CHECKED(FAnimNotifyEvent, EndLink).ToString()
+			   || OuterFieldType == FString(TEXT("AnimLinkableElement")))
 			{
-				if (const USkeleton* Skeleton = AnimAsset->GetSkeleton())
+				// If we get a slot index property replace it with a dropdown showing the names of the 
+				// slots, as the indices are hidden from the user.
+				if(ChildHandle->GetProperty()->GetName() == TEXT("SlotIndex"))
 				{
-					AnimNotifyCategory.AddProperty(CurveNameProperty.ToSharedRef())
+					int32 SlotIdx = INDEX_NONE;
+					ChildHandle->GetValue(SlotIdx);
+
+					LinkCategory.AddProperty(ChildHandle)
 						.CustomWidget()
 						.NameContent()
 						[
-							SNew(SHorizontalBox)
-							+ SHorizontalBox::Slot()
-							.Padding(FMargin(2, 1, 0, 1))
-							[
-								SNew(STextBlock)
-								.Text(CurveNameProperty->GetPropertyDisplayName())
-								.Font(FEditorStyle::GetFontStyle(TEXT("PropertyWindow.NormalFont")))
-							]
+							ChildHandle->CreatePropertyNameWidget(NSLOCTEXT("NotifyDetails", "SlotIndexName", "Slot").ToString())
 						]
-					.ValueContent()
+						.ValueContent()
 						[
-							SNew(SAssetSearchBoxForCurves, Skeleton, CurveNameProperty)
-							.IncludeSocketsForSuggestions(true)
-							.MustMatchPossibleSuggestions(true)
-							.HintText(NSLOCTEXT("AnimNotifyDetails", "Curve Name Hint Text", "Curve Name..."))
-							.OnTextCommitted(this, &FAnimNotifyDetails::OnSearchBoxCommitted, PropIndex)
+							SNew(STextComboBox)
+							.OptionsSource(&SlotNameItems)
+							.OnSelectionChanged(this, &FAnimNotifyDetails::OnSlotSelected, ChildHandle)
+							.OnComboBoxOpening(this, &FAnimNotifyDetails::UpdateSlotNames, EditorObject->AnimObject)
+							.InitiallySelectedItem(SlotNameItems[SlotIdx])
 						];
 				}
+				else
+				{
+					LinkCategory.AddProperty(ChildHandle);
+				}
 			}
 		}
 	}
 }
 
-void FAnimNotifyDetails::OnSearchBoxCommitted(const FText& InSearchText, ETextCommit::Type CommitInfo, int32 PropertyIndex )
+void FAnimNotifyDetails::HideLinkProperties(IDetailLayoutBuilder& Builder, TSharedRef<IPropertyHandle> NotifyProperty)
 {
-	NameProperties[PropertyIndex]->SetValue( InSearchText.ToString() );
+	uint32 NumChildProperties = 0;
+	NotifyProperty->GetNumChildren(NumChildProperties);
+
+	if(NumChildProperties > 0)
+	{
+		for(uint32 ChildIdx = 0 ; ChildIdx < NumChildProperties ; ++ChildIdx)
+		{
+			TSharedPtr<IPropertyHandle> ChildHandle = NotifyProperty->GetChildHandle(ChildIdx);
+			FString OuterFieldType = ChildHandle->GetProperty()->GetOuterUField()->GetName();
+			if(ChildHandle->GetProperty()->GetName() == GET_MEMBER_NAME_CHECKED(FAnimNotifyEvent, EndLink).ToString()
+			   || OuterFieldType == FString(TEXT("AnimLinkableElement")))
+			{
+				Builder.HideProperty(ChildHandle);
+			}
+		}
+	}
+}
+
+bool FAnimNotifyDetails::CustomizeProperty(IDetailCategoryBuilder& CategoryBuilder, UObject* Notify, TSharedPtr<IPropertyHandle> Property)
+{
+	if(Notify && Notify->GetClass() && Property->IsValidHandle())
+	{
+		FString ClassName = Notify->GetClass()->GetName();
+		FString PropertyName = Property->GetProperty()->GetName();
+
+		if(ClassName.Find(TEXT("AnimNotify_PlayParticleEffect")) != INDEX_NONE && PropertyName == TEXT("SocketName"))
+		{
+			AddBoneNameProperty(CategoryBuilder, Notify, Property);
+			return true;
+		}
+		else if(ClassName.Find(TEXT("AnimNotifyState_TimedParticleEffect")) != INDEX_NONE && PropertyName == TEXT("SocketName"))
+		{
+			AddBoneNameProperty(CategoryBuilder, Notify, Property);
+			return true;
+		}
+		else if(ClassName.Find(TEXT("AnimNotify_PlaySound")) != INDEX_NONE && PropertyName == TEXT("AttachName"))
+		{
+			AddBoneNameProperty(CategoryBuilder, Notify, Property);
+			return true;
+		}
+		else if (ClassName.Find(TEXT("AnimNotifyState_Trail")) != INDEX_NONE)
+		{
+			if(PropertyName == TEXT("FirstSocketName") || PropertyName == TEXT("SecondSocketName"))
+			{
+				AddBoneNameProperty(CategoryBuilder, Notify, Property);
+				return true;
+			}
+			else if(PropertyName == TEXT("WidthScaleCurve"))
+			{
+				AddCurveNameProperty(CategoryBuilder, Notify, Property);
+				return true;
+			}
+		}
+	}
+	return false;
+}
+
+void FAnimNotifyDetails::UpdateSlotNames(UAnimSequenceBase* AnimObject)
+{
+	if(UAnimMontage* MontageObj = Cast<UAnimMontage>(AnimObject))
+	{
+		for(FSlotAnimationTrack& Slot : MontageObj->SlotAnimTracks)
+		{
+			if(!SlotNameItems.ContainsByPredicate([&Slot](TSharedPtr<FString>& Item){return Slot.SlotName.ToString() == *Item;}))
+			{
+				SlotNameItems.Add(MakeShareable(new FString(*Slot.SlotName.ToString()))); 
+			}
+		}
+	}
+}
+
+void FAnimNotifyDetails::OnSlotSelected(TSharedPtr<FString> SlotName, ESelectInfo::Type SelectInfo, TSharedPtr<IPropertyHandle> Property)
+{
+	if(SelectInfo != ESelectInfo::Direct && Property->IsValidHandle())
+	{
+		int32 NewIndex = SlotNameItems.Find(SlotName);
+		if(NewIndex != INDEX_NONE)
+		{
+			Property->SetValue(NewIndex);
+		}
+	}
 }
diff --git a/Engine/Source/Editor/Persona/Private/AnimNotifyDetails.h b/Engine/Source/Editor/Persona/Private/AnimNotifyDetails.h
--- a/Engine/Source/Editor/Persona/Private/AnimNotifyDetails.h
+++ b/Engine/Source/Editor/Persona/Private/AnimNotifyDetails.h
@@ -4,6 +4,8 @@
 
 #include "PropertyEditorModule.h"
 
+class IDetailCategoryBuilder;
+
 class FAnimNotifyDetails : public IDetailCustomization
 {
 public:
@@ -17,19 +19,66 @@ private:
 
 	/** Array of pointers to the name properties */
 	TArray<TSharedPtr<IPropertyHandle>> NameProperties;
+	/** List of slot names for selection dropdown, in same order as slots in the montage */
+	TArray<TSharedPtr<FString>> SlotNameItems;
 
-	/** Adds a Bone Name property to the details layout. */
-	void AddBoneNameProperty(IDetailLayoutBuilder& DetailBuilder, const UClass* PropertyClass, const TCHAR* PropertyName, const TCHAR* CategoryName);
+	/** Customize a notify property that is inside an instanced property
+	 *	@param CategoryBuilder The category to place the property into
+	 *	@param Notify the UAnimNotify or UAnimNotifyState object for the property
+	 *	@param Property Handle to the property to customize
+	 */
+	bool CustomizeProperty(IDetailCategoryBuilder& CategoryBuilder, UObject* Notify, TSharedPtr<IPropertyHandle> Property);
 
-	/** Adds a Curve Name property to the details layout. */
-	void AddCurveNameProperty(IDetailLayoutBuilder& DetailBuilder, const UClass* PropertyClass, const TCHAR* PropertyName, const TCHAR* CategoryName);
+	/** Adds a Bone Name property to the details layout
+	 *	@param CategoryBuilder The category to add the property to
+	 *	@param Notify the UAnimNotify or UAnimNotifyState object for the property
+	 *	@param Property Handle to the property to customize
+	 */
+	void AddBoneNameProperty(IDetailCategoryBuilder& CategoryBuilder, UObject* Notify,  TSharedPtr<IPropertyHandle> Property);
 
-	/** Adds the default layout for a property. */
-	void AddPropertyDefault(IDetailLayoutBuilder& DetailBuilder, const UClass* PropertyClass, const TCHAR* PropertyName, const TCHAR* CategoryName);
+	/** Adds a Curve Name property to the details layout
+	 *	@param CategoryBuilder The category to add the property to
+	 *	@param Notify the UAnimNotify or UAnimNotifyState object for the property
+	 *	@param Property Handle to the property to customize
+	 */
+	void AddCurveNameProperty(IDetailCategoryBuilder& CategoryBuilder, UObject* Notify,  TSharedPtr<IPropertyHandle> Property);
 
-	/** Handle the committed text */
+	/** Handles search box commit for name properties
+	 *	@param InSearchText Text that was committed
+	 *	@param CommitInfo Commit method
+	 *	@param PropertyIndex Index of the name properties in internal array
+	 */
 	void OnSearchBoxCommitted(const FText& InSearchText, ETextCommit::Type CommitInfo, int32 PropertyIndex);
 
 	/** Get the search suggestions */
 	TArray<FString> GetSearchSuggestions() const;
+
+	/** Removes the dropdown selection of instanced objects from the header for the property
+	 *	@param CategoryBuilder The category builder for the property
+	 *	@param PropHandle Handle to the property to modify
+	 *	@param bShowChildren Whether or not to show the instanced object children
+	 */
+	void ClearInstancedSelectionDropDown(IDetailCategoryBuilder& CategoryBuilder, TSharedRef<IPropertyHandle> PropHandle, bool bShowChildren = true);
+
+	/** Move properties representing notify linking information into their own category
+	 *	@param Builder Layout builder for this object
+	 *	@param NotifyProperty The property that contains the linking properties
+	 */
+	void CustomizeLinkProperties(IDetailLayoutBuilder& Builder, TSharedRef<IPropertyHandle> NotifyProperty, UEditorNotifyObject* EditorObject);
+
+	/** Hide any properties relating to notify linking
+	 *	@param Builder Layout builder for this object
+	 *	@param NotifyProperty The notify property that contains the linking properties
+	 */
+	void HideLinkProperties(IDetailLayoutBuilder& Builder, TSharedRef<IPropertyHandle> NotifyProperty);
+
+	/** Updates the list of slot names used for combo box
+	 *	@param AnimObject Object being edited, where we will search for slots
+	 */
+	void UpdateSlotNames(UAnimSequenceBase* AnimObject);
+
+	/** Called when the user selects a slot
+	 *	@param Index Index of the new slot
+	 */
+	void OnSlotSelected(TSharedPtr<FString> SlotName, ESelectInfo::Type SelectInfo, TSharedPtr<IPropertyHandle> Property);
 };
\ No newline at end of file
diff --git a/Engine/Source/Editor/Persona/Private/Persona.cpp b/Engine/Source/Editor/Persona/Private/Persona.cpp
--- a/Engine/Source/Editor/Persona/Private/Persona.cpp
+++ b/Engine/Source/Editor/Persona/Private/Persona.cpp
@@ -943,14 +943,41 @@ void FPersona::OnPostImport(UFactory* InFactory, UObject* InObject)
 
 void FPersona::ConditionalRefreshEditor(UObject* InObject)
 {
+	bool bInterestingAsset = true;
 	// Ignore if this is regarding a different object
 	if(InObject != TargetSkeleton && InObject != TargetSkeleton->GetPreviewMesh() && InObject != GetAnimationAssetBeingEdited())
 	{
-		return;
+		bInterestingAsset = false;
+	}
+
+	// Check that we aren't a montage that uses an incoming animation
+	if(UAnimMontage* Montage = Cast<UAnimMontage>(GetAnimationAssetBeingEdited()))
+	{
+		for(FSlotAnimationTrack& Slot : Montage->SlotAnimTracks)
+		{
+			if(bInterestingAsset)
+			{
+				break;
+			}
+
+			for(FAnimSegment& Segment : Slot.AnimTrack.AnimSegments)
+			{
+				if(Segment.AnimReference == InObject)
+				{
+					bInterestingAsset = true;
+					break;
+				}
+			}
+		}
 	}
 
-	RefreshViewport();
-	ReinitMode();
+	if(bInterestingAsset)
+	{
+		RefreshViewport();
+		ReinitMode();
+		
+		OnPersonaRefresh.Broadcast();
+	}
 }
 
 /** Called when graph editor focus is changed */
diff --git a/Engine/Source/Editor/Persona/Private/Persona.h b/Engine/Source/Editor/Persona/Private/Persona.h
--- a/Engine/Source/Editor/Persona/Private/Persona.h
+++ b/Engine/Source/Editor/Persona/Private/Persona.h
@@ -316,8 +316,25 @@ private:
 	DECLARE_MULTICAST_DELEGATE_OneParam(FOnCreateViewport, TWeakPtr<class SAnimationEditorViewportTabBody>)
 	// Called when generic delete happens
 	DECLARE_MULTICAST_DELEGATE( FOnGenericDelete );
+	// Called when Persona refreshes
+	DECLARE_MULTICAST_DELEGATE( FOnPersonaRefreshMulticaster );
 public:
 
+	// Persona refreshed
+	typedef FOnPersonaRefreshMulticaster::FDelegate FOnPersonaRefresh;
+
+	/** Registers a delegate to be called when Persona is Refreshed */
+	void RegisterOnPersonaRefresh(const FOnPersonaRefresh& Delegate)
+	{
+		OnPersonaRefresh.Add(Delegate);
+	}
+
+	/** Unregisters refresh delegate */
+	void UnregisterOnPersonaRefresh(SWidget* Widget)
+	{
+		OnPersonaRefresh.RemoveAll(Widget);
+	}
+
 	// anim changed 
 	typedef FOnAnimChangedMulticaster::FDelegate FOnAnimChanged;
 
@@ -503,6 +520,9 @@ protected:
 
 protected:
 
+	/** Called when persona is refreshed through an external action (reimport etc) */
+	FOnPersonaRefreshMulticaster OnPersonaRefresh;
+
 	/** Delegate called after an undo operation for child widgets to refresh */
 	FOnPostUndoMulticaster OnPostUndo;	
 
diff --git a/Engine/Source/Editor/Persona/Private/PersonaModule.cpp b/Engine/Source/Editor/Persona/Private/PersonaModule.cpp
--- a/Engine/Source/Editor/Persona/Private/PersonaModule.cpp
+++ b/Engine/Source/Editor/Persona/Private/PersonaModule.cpp
@@ -58,8 +58,7 @@ void FPersonaModule::StartupModule()
 	{
 		FPropertyEditorModule& PropertyModule = FModuleManager::LoadModuleChecked<FPropertyEditorModule>("PropertyEditor");
 		PropertyModule.RegisterCustomClassLayout( "SkeletalMeshSocket", FOnGetDetailCustomizationInstance::CreateStatic( &FSkeletalMeshSocketDetails::MakeInstance ) );
-		PropertyModule.RegisterCustomClassLayout( "AnimNotify", FOnGetDetailCustomizationInstance::CreateStatic( &FAnimNotifyDetails::MakeInstance ) );
-		PropertyModule.RegisterCustomClassLayout( "AnimNotifyState", FOnGetDetailCustomizationInstance::CreateStatic(&FAnimNotifyDetails::MakeInstance));
+		PropertyModule.RegisterCustomClassLayout( "EditorNotifyObject", FOnGetDetailCustomizationInstance::CreateStatic(&FAnimNotifyDetails::MakeInstance));
 		PropertyModule.RegisterCustomClassLayout( "AnimGraphNode_Base", FOnGetDetailCustomizationInstance::CreateStatic( &FAnimGraphNodeDetails::MakeInstance ) );
 		PropertyModule.RegisterCustomClassLayout( "AnimInstance", FOnGetDetailCustomizationInstance::CreateStatic(&FAnimInstanceDetails::MakeInstance));
 
@@ -78,8 +77,7 @@ void FPersonaModule::ShutdownModule()
 	{
 		FPropertyEditorModule& PropertyModule = FModuleManager::LoadModuleChecked<FPropertyEditorModule>("PropertyEditor");
 		PropertyModule.UnregisterCustomClassLayout("SkeletalMeshSocket");
-		PropertyModule.UnregisterCustomClassLayout("AnimNotify");
-		PropertyModule.UnregisterCustomClassLayout("AnimNotifyState");
+		PropertyModule.UnregisterCustomClassLayout("EditorNotifyObject");
 		PropertyModule.UnregisterCustomClassLayout("AnimGraphNode_Base");
 
 		PropertyModule.UnregisterCustomPropertyTypeLayout("InputScaleBias");
diff --git a/Engine/Source/Editor/Persona/Private/SAnimMontagePanel.cpp b/Engine/Source/Editor/Persona/Private/SAnimMontagePanel.cpp
--- a/Engine/Source/Editor/Persona/Private/SAnimMontagePanel.cpp
+++ b/Engine/Source/Editor/Persona/Private/SAnimMontagePanel.cpp
@@ -262,7 +262,8 @@ int32 SMontageBranchingPointNode::OnPaint(const FPaintArgs& Args, const FGeometr
 
 	if(BranchingPoint->TriggerTimeOffset != 0.f) //Do we have an offset to render?
 	{
-		if(BranchingPoint->DisplayTime != 0.f && BranchingPoint->DisplayTime != Montage->SequenceLength) //Don't render offset when we are at the start/end of the sequence, doesn't help the user
+		float BranchingPointTime = BranchingPoint->GetTime();
+		if(BranchingPointTime != 0.f && BranchingPointTime != Montage->SequenceLength) //Don't render offset when we are at the start/end of the sequence, doesn't help the user
 		{
 			// ScrubHandle
 			FVector2D MarkerPosition = TextPos;
@@ -638,6 +639,17 @@ void SAnimMontagePanel::SummonTrackContextMenu( FMenuBuilder& MenuBuilder, float
 	}
 	MenuBuilder.EndSection();
 
+	MenuBuilder.BeginSection("AnimMontageElementBulkActions", LOCTEXT("BulkLinkActions", "Bulk Link Actions"));
+	{
+		MenuBuilder.AddSubMenu(LOCTEXT("SetElementLink_SubMenu", "Set Elements to..."), LOCTEXT("SetElementLink_TooTip", "Sets all montage elements (Branching Points, Sections, Notifies) to a chosen link type."), FNewMenuDelegate::CreateSP(this, &SAnimMontagePanel::FillElementSubMenuForTimes));
+
+		if(Montage->SlotAnimTracks.Num() > 1)
+		{
+			MenuBuilder.AddSubMenu(LOCTEXT("SetToSlotMenu", "Link all Elements to Slot..."), LOCTEXT("SetToSlotMenuToolTip", "Link all elements to a selected slot"), FNewMenuDelegate::CreateSP(this, &SAnimMontagePanel::FillSlotSubMenu));
+		}
+	}
+	MenuBuilder.EndSection();
+
 	LastContextHeading.Empty();
 }
 
@@ -655,6 +667,22 @@ void SAnimMontagePanel::SummonBranchNodeContextMenu( FMenuBuilder& MenuBuilder,
 	LastContextHeading = "Branch Point";
 }
 
+void SAnimMontagePanel::FillElementSubMenuForTimes(FMenuBuilder& MenuBuilder)
+{
+	MenuBuilder.AddMenuEntry(LOCTEXT("SubLinkAbs", "Absolute"), LOCTEXT("SubLinkAbs", "Set all elements to absolute link"), FSlateIcon(), FUIAction(FExecuteAction::CreateSP(this, &SAnimMontagePanel::OnSetElementsToLinkMode, EAnimLinkMethod::Absolute)));
+	MenuBuilder.AddMenuEntry(LOCTEXT("SubLinkRel", "Relative"), LOCTEXT("SubLinkRel", "Set all elements to relative link"), FSlateIcon(), FUIAction(FExecuteAction::CreateSP(this, &SAnimMontagePanel::OnSetElementsToLinkMode, EAnimLinkMethod::Relative)));
+	MenuBuilder.AddMenuEntry(LOCTEXT("SubLinkPro", "Proportional"), LOCTEXT("SubLinkPro", "Set all elements to proportional link"), FSlateIcon(), FUIAction(FExecuteAction::CreateSP(this, &SAnimMontagePanel::OnSetElementsToLinkMode, EAnimLinkMethod::Proportional)));
+}
+
+void SAnimMontagePanel::FillSlotSubMenu(FMenuBuilder& Menubuilder)
+{
+	for(int32 SlotIdx = 0 ; SlotIdx < Montage->SlotAnimTracks.Num() ; ++SlotIdx)
+	{
+		FSlotAnimationTrack& Slot = Montage->SlotAnimTracks[SlotIdx];
+		Menubuilder.AddMenuEntry(FText::Format(LOCTEXT("SubSlotMenuNameEntry", "{SlotName}"), FText::FromString(Slot.SlotName.ToString())), LOCTEXT("SubSlotEntry", "Set to link to this slot"), FSlateIcon(), FUIAction(FExecuteAction::CreateSP(this, &SAnimMontagePanel::OnSetElementsToSlot, SlotIdx)));
+	}
+}
+
 /** Slots */
 void SAnimMontagePanel::OnNewSlotClicked()
 {
@@ -953,4 +981,73 @@ void SAnimMontagePanel::UpdateSlotGroupWarningVisibility()
 	}
 }
 
+void SAnimMontagePanel::OnSetElementsToLinkMode(EAnimLinkMethod::Type NewLinkMethod)
+{
+	TArray<FAnimLinkableElement*> Elements;
+	CollectLinkableElements(Elements);
+
+	for(FAnimLinkableElement* Element : Elements)
+	{
+		Element->ChangeLinkMethod(NewLinkMethod);
+	}
+
+	// Handle notify state links
+	for(FAnimNotifyEvent& Notify : Montage->Notifies)
+	{
+		if(Notify.GetDuration() > 0.0f)
+		{
+			// Always keep link methods in sync between notifies and duration links
+			if(Notify.GetLinkMethod() != Notify.EndLink.GetLinkMethod())
+			{
+				Notify.EndLink.ChangeLinkMethod(Notify.GetLinkMethod());
+			}
+		}
+	}
+}
+
+void SAnimMontagePanel::OnSetElementsToSlot(int32 SlotIndex)
+{
+	TArray<FAnimLinkableElement*> Elements;
+	CollectLinkableElements(Elements);
+
+	for(FAnimLinkableElement* Element : Elements)
+	{
+		Element->ChangeSlotIndex(SlotIndex);
+	}
+
+	// Handle notify state links
+	for(FAnimNotifyEvent& Notify : Montage->Notifies)
+	{
+		if(Notify.GetDuration() > 0.0f)
+		{
+			// Always keep link methods in sync between notifies and duration links
+			if(Notify.GetSlotIndex() != Notify.EndLink.GetSlotIndex())
+			{
+				Notify.EndLink.ChangeSlotIndex(Notify.GetSlotIndex());
+			}
+		}
+	}
+}
+
+void SAnimMontagePanel::CollectLinkableElements(TArray<FAnimLinkableElement*> &Elements)
+{
+	for(auto& Composite : Montage->CompositeSections)
+	{
+		FAnimLinkableElement* Element = &Composite;
+		Elements.Add(Element);
+	}
+
+	for(auto& BranchingPoint : Montage->BranchingPoints)
+	{
+		FAnimLinkableElement* Element = &BranchingPoint;
+		Elements.Add(Element);
+	}
+
+	for(auto& Notify : Montage->Notifies)
+	{
+		FAnimLinkableElement* Element = &Notify;
+		Elements.Add(Element);
+	}
+}
+
 #undef LOCTEXT_NAMESPACE
\ No newline at end of file
diff --git a/Engine/Source/Editor/Persona/Private/SAnimMontagePanel.h b/Engine/Source/Editor/Persona/Private/SAnimMontagePanel.h
--- a/Engine/Source/Editor/Persona/Private/SAnimMontagePanel.h
+++ b/Engine/Source/Editor/Persona/Private/SAnimMontagePanel.h
@@ -67,6 +67,8 @@ public:
 	void SummonTrackContextMenu( FMenuBuilder& MenuBuilder, float DataPosX, int32 SectionIndex, int32 AnimSlotIndex );
 	void SummonSectionContextMenu( FMenuBuilder& MenuBuilder, int32 SectionIndex );
 	void SummonBranchNodeContextMenu( FMenuBuilder& MenuBuilder, int32 BranchIndex );
+	void FillElementSubMenuForTimes( FMenuBuilder& MenuBuilder );
+	void FillSlotSubMenu( FMenuBuilder& Menubuilder );
 
 	void OnNewSlotClicked();
 	void OnNewSectionClicked(float DataPosX);
@@ -80,6 +82,24 @@ public:
 
 	void ClearSelected();
 
+	// helper method to check whether the slot name is empty or not. If empty, shows an error message to provide a valid name
+	void CheckSlotName(const FText& SlotName, int32 SlotNodeIndex, bool bShouldCheckCollapsed = false) const;
+
+	// check the slot name whether valid or not while the user is typing
+	void OnSlotNameChanged(const FText& NewText, int32 SlotNodeIndex);
+
+	// get slot name from a montage editor and check the slot name whether valid or not
+	FText GetMontageSlotName(int32 SlotIndex) const;
+	
+	// Context menu callback to set all elements in the montage to a given link method
+	void OnSetElementsToLinkMode(EAnimLinkMethod::Type NewLinkMethod);
+
+	// Fills the given array with all linkable elements in the montage
+	void CollectLinkableElements(TArray<FAnimLinkableElement*>& Elements);
+
+	// Context menu callback to set all elements in the montage to a given slot
+	void OnSetElementsToSlot(int32 SlotIndex);
+
 private:
 	TWeakPtr<FPersona> Persona;
 	TWeakPtr<SMontageEditor>	MontageEditor;
@@ -112,7 +132,7 @@ private:
 	FReply OnOpenAnimSlotManager();
 	void RefreshComboLists(bool bOnlyRefreshIfDifferent = false);
 	void UpdateSlotGroupWarningVisibility();
-
+	
 	void CreateNewSlot(const FText& NewSlotName, ETextCommit::Type CommitInfo);
 	void CreateNewSection(const FText& NewSectionName, ETextCommit::Type CommitInfo, float StartTime);
 	void CreateNewBranch(const FText& NewBranchName, ETextCommit::Type CommitInfo, float StartTime);
diff --git a/Engine/Source/Editor/Persona/Private/SAnimNotifyPanel.cpp b/Engine/Source/Editor/Persona/Private/SAnimNotifyPanel.cpp
--- a/Engine/Source/Editor/Persona/Private/SAnimNotifyPanel.cpp
+++ b/Engine/Source/Editor/Persona/Private/SAnimNotifyPanel.cpp
@@ -47,6 +47,7 @@ public:
 		, _OnNodeDragStarted()
 		, _OnUpdatePanel()
 		, _PanTrackRequest()
+		, _OnDeselectAllNotifies()
 		, _ViewInputMin()
 		, _ViewInputMax()
 		, _MarkerBars()
@@ -57,6 +58,7 @@ public:
 	SLATE_EVENT( FOnNotifyNodeDragStarted, OnNodeDragStarted )
 	SLATE_EVENT( FOnUpdatePanel, OnUpdatePanel )
 	SLATE_EVENT( FPanTrackRequest, PanTrackRequest )
+	SLATE_EVENT( FDeselectAllNotifies, OnDeselectAllNotifies)
 	SLATE_ATTRIBUTE( float, ViewInputMin )
 	SLATE_ATTRIBUTE( float, ViewInputMax )
 	SLATE_ATTRIBUTE( TArray<FTrackMarkerBar>, MarkerBars)
@@ -178,6 +180,9 @@ private:
 	/** Marker bars for snapping to when dragging the markers in a state notify node */
 	TAttribute<TArray<FTrackMarkerBar>> MarkerBars;
 
+	/** Delegate to deselect notifies and clear the details panel */
+	FDeselectAllNotifies OnDeselectAllNotifies;
+
 	friend class SAnimNotifyTrack;
 };
 
@@ -265,9 +270,11 @@ public:
 	// Adds our current selection to the provided set
 	void AppendSelectionToSet(FGraphPanelSelectionSet& SelectionSet);
 	// Adds our current selection to the provided array
-	void AppendSelectionToArray(TArray<const FAnimNotifyEvent*>& Selection) const;
+	void AppendSelectionToArray(TArray<FAnimNotifyEvent*>& Selection) const;
 	// Gets the currently selected SAnimNotifyNode instances
 	void AppendSelectedNodeWidgetsToArray(TArray<TSharedPtr<SAnimNotifyNode>>& NodeArray) const;
+	// Gets the indices of the selected notifies
+	TArray<int32> GetSelectedNotifyIndices() const {return SelectedNotifyIndices;}
 
 	/**
 	* Deselects all currently selected notify nodes
@@ -634,7 +641,7 @@ public:
 			// Always clamp the Y to the current track
 			SelectionBeginPosition.Y = SelectionPositionClampInfo->TrackPos;
 
-			if(CurrentNode.IsValid() && CurrentEvent->Duration > 0)
+			if(CurrentNode.IsValid() && CurrentEvent->GetDuration() > 0)
 			{
 				// If we didn't snap the beginning of the node, attempt to snap the end
 				if(SnapX == -1.0f)
@@ -660,7 +667,8 @@ public:
 						if(SnapMovement == 0.0f)
 						{
 							SnapMovement = SnapX - TrackNodeEndPos.X;
-							SnapTime = TrackScaleInfo.LocalXToInput(SnapX) - CurrentEvent->Duration;
+							SnapTime = TrackScaleInfo.LocalXToInput(SnapX) - CurrentEvent->GetDuration();
+							CurrentNode->SetLastSnappedTime(SnapTime);
 							SnappedNode = CurrentNode;
 						}
 					}
@@ -806,24 +814,28 @@ public:
 		float BeginTime = MAX_flt;
 		for(TSharedPtr<SAnimNotifyNode> Node : NotifyNodes)
 		{
-			if(Node->NotifyEvent->DisplayTime < BeginTime)
+			float NotifyTime = Node->NotifyEvent->GetTime();
+
+			if(NotifyTime < BeginTime)
 			{
-				BeginTime = Node->NotifyEvent->DisplayTime;
+				BeginTime = NotifyTime;
 			}
 		}
 
 		// Initialise node data
 		for(TSharedPtr<SAnimNotifyNode> Node : NotifyNodes)
 		{
+			float NotifyTime = Node->NotifyEvent->GetTime();
+
 			Node->ClearLastSnappedTime();
-			Operation->NodeTimeOffsets.Add(Node->NotifyEvent->DisplayTime - BeginTime);
-			Operation->NodeTimes.Add(Node->NotifyEvent->DisplayTime);
+			Operation->NodeTimeOffsets.Add(NotifyTime - BeginTime);
+			Operation->NodeTimes.Add(NotifyTime);
 			Operation->NodeXOffsets.Add(Node->GetNotifyPositionOffset().X);
 
 			// Calculate the time length of the selection. Because it is possible to have states
 			// with arbitrary durations we need to search all of the nodes and find the furthest
 			// possible point
-			Operation->SelectionTimeLength = FMath::Max(Operation->SelectionTimeLength, Node->NotifyEvent->DisplayTime + Node->NotifyEvent->Duration - BeginTime);
+			Operation->SelectionTimeLength = FMath::Max(Operation->SelectionTimeLength, NotifyTime + Node->NotifyEvent->GetDuration() - BeginTime);
 		}
 
 		Operation->Construct();
@@ -901,7 +913,7 @@ void SAnimNotifyNode::Construct(const FArguments& InArgs)
 	CurrentDragHandle = ENotifyStateHandleHit::None;
 	bDrawTooltipToRight = true;
 	bSelected = false;
-
+	
 	// Cache notify name for blueprint / Native notifies.
 	if(NotifyEvent->Notify)
 	{
@@ -914,6 +926,7 @@ void SAnimNotifyNode::Construct(const FArguments& InArgs)
 
 	OnNodeDragStarted = InArgs._OnNodeDragStarted;
 	PanTrackRequest = InArgs._PanTrackRequest;
+	OnDeselectAllNotifies = InArgs._OnDeselectAllNotifies;
 
 	ViewInputMin = InArgs._ViewInputMin;
 	ViewInputMax = InArgs._ViewInputMax;
@@ -977,14 +990,14 @@ FText SAnimNotifyNode::GetNotifyText() const
 
 FText SAnimNotifyNode::GetNodeTooltip() const
 {
-	const float Time = NotifyEvent->DisplayTime;
+	const float Time = NotifyEvent->GetTime();
 	const float Percentage = Time / Sequence->SequenceLength;
 	const FText Frame = FText::AsNumber( (int32)(Percentage * Sequence->GetNumberOfFrames()) );
 	const FText Seconds = FText::AsNumber( Time );
 
-	if (NotifyEvent->Duration > 0.0f)
+	if (NotifyEvent->GetDuration() > 0.0f)
 	{
-		const FText Duration = FText::AsNumber( NotifyEvent->Duration );
+		const FText Duration = FText::AsNumber( NotifyEvent->GetDuration() );
 		return FText::Format( LOCTEXT("NodeToolTipLong", "@ {0} sec (frame {1}) for {2} sec"), Seconds, Frame, Duration );
 	}
 
@@ -1067,8 +1080,8 @@ void SAnimNotifyNode::UpdateSizeAndPosition(const FGeometry& AllottedGeometry)
 	// Cache the geometry information, the alloted geometry is the same size as the track.
 	CachedAllotedGeometrySize = AllottedGeometry.Size;
 
-	NotifyTimePositionX = ScaleInfo.InputToLocalX(NotifyEvent->DisplayTime);
-	NotifyDurationSizeX = ScaleInfo.PixelsPerInput * NotifyEvent->Duration;
+	NotifyTimePositionX = ScaleInfo.InputToLocalX(NotifyEvent->GetTime());
+	NotifyDurationSizeX = ScaleInfo.PixelsPerInput * NotifyEvent->GetDuration();
 
 	const TSharedRef< FSlateFontMeasure > FontMeasureService = FSlateApplication::Get().GetRenderer()->GetFontMeasureService();
 	TextSize = FontMeasureService->Measure( GetNotifyText(), Font );
@@ -1156,7 +1169,7 @@ int32 SAnimNotifyNode::OnPaint(const FPaintArgs& Args, const FGeometry& Allotted
 		// Render offsets if necessary
 		if(NotifyEvent->EndTriggerTimeOffset != 0.f) //Do we have an offset to render?
 		{
-			float EndTime = NotifyEvent->DisplayTime + NotifyEvent->Duration;
+			float EndTime = NotifyEvent->GetTime() + NotifyEvent->GetDuration();
 			if(EndTime != Sequence->SequenceLength) //Don't render offset when we are at the end of the sequence, doesnt help the user
 			{
 				// ScrubHandle
@@ -1234,7 +1247,8 @@ int32 SAnimNotifyNode::OnPaint(const FPaintArgs& Args, const FGeometry& Allotted
 
 	if(NotifyEvent->TriggerTimeOffset != 0.f) //Do we have an offset to render?
 	{
-		if(NotifyEvent->DisplayTime != 0.f && NotifyEvent->DisplayTime != Sequence->SequenceLength) //Don't render offset when we are at the start/end of the sequence, doesnt help the user
+		float NotifyTime = NotifyEvent->GetTime();
+		if(NotifyTime != 0.f && NotifyTime != Sequence->SequenceLength) //Don't render offset when we are at the start/end of the sequence, doesnt help the user
 		{
 			float HandleCentre = NotifyScrubHandleCentre;
 			float &Offset = NotifyEvent->TriggerTimeOffset;
@@ -1260,12 +1274,12 @@ FReply SAnimNotifyNode::OnMouseMove( const FGeometry& MyGeometry, const FPointer
 	if(CurrentDragHandle == ENotifyStateHandleHit::Start)
 	{
 		// Check track bounds
-		float OldDisplayTime = NotifyEvent->DisplayTime;
+		float OldDisplayTime = NotifyEvent->GetTime();
 
 		if(MouseEvent.GetScreenSpacePosition().X >= TrackScreenSpaceXPosition && MouseEvent.GetScreenSpacePosition().X <= TrackScreenSpaceXPosition + CachedAllotedGeometrySize.X)
 		{
 			float NewDisplayTime = ScaleInfo.LocalXToInput((MouseEvent.GetScreenSpacePosition() - MyGeometry.AbsolutePosition + MyGeometry.Position).X);
-			float NewDuration = NotifyEvent->Duration + OldDisplayTime - NewDisplayTime;
+			float NewDuration = NotifyEvent->GetDuration() + OldDisplayTime - NewDisplayTime;
 
 			// Check to make sure the duration is not less than the minimum allowed
 			if(NewDuration < MinimumStateDuration)
@@ -1273,10 +1287,10 @@ FReply SAnimNotifyNode::OnMouseMove( const FGeometry& MyGeometry, const FPointer
 				NewDisplayTime -= MinimumStateDuration - NewDuration;
 			}
 
-			NotifyEvent->DisplayTime = NewDisplayTime;
-			NotifyEvent->Duration += OldDisplayTime - NotifyEvent->DisplayTime;
+			NotifyEvent->SetTime(NewDisplayTime);
+			NotifyEvent->SetDuration(NotifyEvent->GetDuration() + OldDisplayTime - NotifyEvent->GetTime());
 		}
-		else if(NotifyEvent->Duration > MinimumStateDuration)
+		else if(NotifyEvent->GetDuration() > MinimumStateDuration)
 		{
 			float Overflow = HandleOverflowPan(MouseEvent.GetScreenSpacePosition(), TrackScreenSpaceXPosition);
 
@@ -1284,21 +1298,21 @@ FReply SAnimNotifyNode::OnMouseMove( const FGeometry& MyGeometry, const FPointer
 			ScaleInfo.ViewMinInput = ViewInputMin.Get();
 			ScaleInfo.ViewMaxInput = ViewInputMax.Get();
 
-			NotifyEvent->DisplayTime = ScaleInfo.LocalXToInput(Overflow < 0.0f ? 0.0f : CachedAllotedGeometrySize.X);
-			NotifyEvent->Duration += OldDisplayTime - NotifyEvent->DisplayTime;
+			NotifyEvent->SetTime(ScaleInfo.LocalXToInput(Overflow < 0.0f ? 0.0f : CachedAllotedGeometrySize.X));
+			NotifyEvent->SetDuration(NotifyEvent->GetDuration() + OldDisplayTime - NotifyEvent->GetTime());
 
 			// Adjust incase we went under the minimum
-			if(NotifyEvent->Duration < MinimumStateDuration)
+			if(NotifyEvent->GetDuration() < MinimumStateDuration)
 			{
-				float EndTimeBefore = NotifyEvent->DisplayTime + NotifyEvent->Duration;
-				NotifyEvent->DisplayTime += NotifyEvent->Duration - MinimumStateDuration;
-				NotifyEvent->Duration = MinimumStateDuration;
-				float EndTimeAfter = NotifyEvent->DisplayTime + NotifyEvent->Duration;
+				float EndTimeBefore = NotifyEvent->GetTime() + NotifyEvent->GetDuration();
+				NotifyEvent->SetTime(NotifyEvent->GetTime() + NotifyEvent->GetDuration() - MinimumStateDuration);
+				NotifyEvent->SetDuration(MinimumStateDuration);
+				float EndTimeAfter = NotifyEvent->GetTime() + NotifyEvent->GetDuration();
 			}
 		}
 
 		// Now we know where the marker should be, look for possible snaps on montage marker bars
-		float NodePositionX = ScaleInfo.InputToLocalX(NotifyEvent->DisplayTime);
+		float NodePositionX = ScaleInfo.InputToLocalX(NotifyEvent->GetTime());
 		float MarkerSnap = GetScrubHandleSnapPosition(NodePositionX, ENotifyStateHandleHit::Start);
 		if(MarkerSnap != -1.0f)
 		{
@@ -1308,9 +1322,9 @@ FReply SAnimNotifyNode::OnMouseMove( const FGeometry& MyGeometry, const FPointer
 			NodePositionX = MarkerSnap;
 
 			// Adjust our start marker
-			OldDisplayTime = NotifyEvent->DisplayTime;
-			NotifyEvent->DisplayTime = ScaleInfo.LocalXToInput(NodePositionX);
-			NotifyEvent->Duration += OldDisplayTime - NotifyEvent->DisplayTime;
+			OldDisplayTime = NotifyEvent->GetTime();
+			NotifyEvent->SetTime(ScaleInfo.LocalXToInput(NodePositionX));
+			NotifyEvent->SetDuration(NotifyEvent->GetDuration() + OldDisplayTime - NotifyEvent->GetTime());
 		}
 		else
 		{
@@ -1321,11 +1335,11 @@ FReply SAnimNotifyNode::OnMouseMove( const FGeometry& MyGeometry, const FPointer
 	{
 		if(MouseEvent.GetScreenSpacePosition().X >= TrackScreenSpaceXPosition && MouseEvent.GetScreenSpacePosition().X <= TrackScreenSpaceXPosition + CachedAllotedGeometrySize.X)
 		{
-			float NewDuration = ScaleInfo.LocalXToInput((MouseEvent.GetScreenSpacePosition() - MyGeometry.AbsolutePosition + MyGeometry.Position).X) - NotifyEvent->DisplayTime;
+			float NewDuration = ScaleInfo.LocalXToInput((MouseEvent.GetScreenSpacePosition() - MyGeometry.AbsolutePosition + MyGeometry.Position).X) - NotifyEvent->GetTime();
 
-			NotifyEvent->Duration = FMath::Max(NewDuration, MinimumStateDuration);
+			NotifyEvent->SetDuration(FMath::Max(NewDuration, MinimumStateDuration));
 		}
-		else if(NotifyEvent->Duration > MinimumStateDuration)
+		else if(NotifyEvent->GetDuration() > MinimumStateDuration)
 		{
 			float Overflow = HandleOverflowPan(MouseEvent.GetScreenSpacePosition(), TrackScreenSpaceXPosition);
 
@@ -1333,11 +1347,11 @@ FReply SAnimNotifyNode::OnMouseMove( const FGeometry& MyGeometry, const FPointer
 			ScaleInfo.ViewMinInput = ViewInputMin.Get();
 			ScaleInfo.ViewMaxInput = ViewInputMax.Get();
 
-			NotifyEvent->Duration = FMath::Max(ScaleInfo.LocalXToInput(Overflow > 0.0f ? CachedAllotedGeometrySize.X : 0.0f) - NotifyEvent->DisplayTime, MinimumStateDuration);
+			NotifyEvent->SetDuration(FMath::Max(ScaleInfo.LocalXToInput(Overflow > 0.0f ? CachedAllotedGeometrySize.X : 0.0f) - NotifyEvent->GetTime(), MinimumStateDuration));
 		}
 
 		// Now we know where the scrub handle should be, look for possible snaps on montage marker bars
-		float NodePositionX = ScaleInfo.InputToLocalX(NotifyEvent->DisplayTime + NotifyEvent->Duration);
+		float NodePositionX = ScaleInfo.InputToLocalX(NotifyEvent->GetTime() + NotifyEvent->GetDuration());
 		float MarkerSnap = GetScrubHandleSnapPosition(NodePositionX, ENotifyStateHandleHit::End);
 		if(MarkerSnap != -1.0f)
 		{
@@ -1347,7 +1361,7 @@ FReply SAnimNotifyNode::OnMouseMove( const FGeometry& MyGeometry, const FPointer
 			NodePositionX = MarkerSnap;
 
 			// Adjust our end marker
-			NotifyEvent->Duration = ScaleInfo.LocalXToInput(NodePositionX) - NotifyEvent->DisplayTime;
+			NotifyEvent->SetDuration(ScaleInfo.LocalXToInput(NodePositionX) - NotifyEvent->GetTime());
 		}
 		else
 		{
@@ -1366,6 +1380,7 @@ FReply SAnimNotifyNode::OnMouseButtonUp( const FGeometry& MyGeometry, const FPoi
 	{
 		// Clear the drag marker and give the mouse back
 		CurrentDragHandle = ENotifyStateHandleHit::None;
+		OnDeselectAllNotifies.ExecuteIfBound();
 		return FReply::Handled().ReleaseMouseCapture();
 	}
 
@@ -1835,10 +1850,20 @@ void SAnimNotifyTrack::CreateNewNotifyAtCursor(FString NewNotifyName, UClass* No
 	Sequence->Modify();
 
 	// Insert a new notify record and spawn the new notify object
-	int32 NewNotifyIndex = Sequence->Notifies.AddZeroed();
+	int32 NewNotifyIndex = Sequence->Notifies.Add(FAnimNotifyEvent());
 	FAnimNotifyEvent& NewEvent = Sequence->Notifies[NewNotifyIndex];
 	NewEvent.NotifyName = FName(*NewNotifyName);
-	NewEvent.DisplayTime = NewTime;
+
+	UAnimMontage* SequenceAsMontage = Cast<UAnimMontage>(Sequence);
+	if(SequenceAsMontage)
+	{
+		NewEvent.LinkMontage(SequenceAsMontage, NewTime);
+	}
+	else
+	{
+		NewEvent.LinkSequence(Sequence, NewTime);
+	}
+
 	NewEvent.TriggerTimeOffset = GetTriggerTimeOffsetForType(Sequence->CalculateOffsetForNotify(NewTime));
 	NewEvent.TrackIndex = TrackIndex;
 
@@ -1851,8 +1876,16 @@ void SAnimNotifyTrack::CreateNewNotifyAtCursor(FString NewNotifyName, UClass* No
 		// Set default duration to 1 frame for AnimNotifyState.
 		if( NewEvent.NotifyStateClass )
 		{
-			NewEvent.Duration = 1 / 30.f;
 			NewEvent.NotifyName = FName(*NewEvent.NotifyStateClass->GetNotifyName());
+			NewEvent.SetDuration(1 / 30.f);
+			if(SequenceAsMontage)
+			{
+				NewEvent.EndLink.LinkMontage(SequenceAsMontage, NewEvent.EndLink.GetTime());
+			}
+			else
+			{
+				NewEvent.EndLink.LinkSequence(Sequence, NewEvent.EndLink.GetTime());
+			}
 		}
 		else
 		{
@@ -2060,7 +2093,7 @@ TSharedPtr<SWidget> SAnimNotifyTrack::SummonContextMenu(const FGeometry& MyGeome
 				NewAction.CanExecuteAction.BindRaw(
 					this, &SAnimNotifyTrack::IsSingleNodeSelected);
 
-				FText SetAnimStateDurationText = FText::Format( LOCTEXT("SetAnimStateDuration", "Set AnimNotifyState duration ({0})"), FText::AsNumber( AnimNotifies[NotifyIndex]->Duration ) );
+				FText SetAnimStateDurationText = FText::Format( LOCTEXT("SetAnimStateDuration", "Set AnimNotifyState duration ({0})"), FText::AsNumber( AnimNotifies[NotifyIndex]->GetDuration() ) );
 				MenuBuilder.AddMenuEntry(SetAnimStateDurationText, LOCTEXT("SetAnimStateDuration_ToolTip", "The duration of this AnimNotifyState"), FSlateIcon(), NewAction);
 			}
 
@@ -2273,19 +2306,19 @@ void SAnimNotifyTrack::OnPasteNotify(float TimeToPasteAt, ENotifyPasteMultipleMo
 
 				if (TimeToPasteAt >= 0.f)
 				{
-					Shift = TimeToPasteAt - (NewNotify.DisplayTime * ScaleMultiplier);
-					NewNotify.DisplayTime = TimeToPasteAt;
+					Shift = TimeToPasteAt - (NewNotify.GetTime() * ScaleMultiplier);
+					NewNotify.SetTime(TimeToPasteAt);
 
 					// Don't want to paste subsequent nodes at this position
 					TimeToPasteAt = -1.0f;
 				}
 				else
 				{
-					NewNotify.DisplayTime = (NewNotify.DisplayTime * ScaleMultiplier) + Shift;
+					NewNotify.SetTime((NewNotify.GetTime() * ScaleMultiplier) + Shift);
 					// Clamp into the bounds of the sequence
-					NewNotify.DisplayTime = FMath::Clamp(NewNotify.DisplayTime, 0.0f, Sequence->SequenceLength);
+					NewNotify.SetTime(FMath::Clamp(NewNotify.GetTime(), 0.0f, Sequence->SequenceLength));
 				}
-				NewNotify.TriggerTimeOffset = GetTriggerTimeOffsetForType(Sequence->CalculateOffsetForNotify(NewNotify.DisplayTime));
+				NewNotify.TriggerTimeOffset = GetTriggerTimeOffsetForType(Sequence->CalculateOffsetForNotify(NewNotify.GetTime()));
 				NewNotify.TrackIndex = TrackIndex;
 
 				// Get anim notify ptr of the new notify event array element (This ptr is the same as the source anim notify ptr)
@@ -2312,9 +2345,9 @@ void SAnimNotifyTrack::OnPasteNotify(float TimeToPasteAt, ENotifyPasteMultipleMo
 						NewNotify.NotifyStateClass = NewAnimNotifyStatePtr;
 					}
 
-					if ((NewNotify.DisplayTime + NewNotify.Duration) > Sequence->SequenceLength)
+					if ((NewNotify.GetTime() + NewNotify.GetDuration()) > Sequence->SequenceLength)
 					{
-						NewNotify.Duration = Sequence->SequenceLength - NewNotify.DisplayTime;
+						NewNotify.SetDuration(Sequence->SequenceLength - NewNotify.GetTime());
 					}
 				}
 			}
@@ -2404,7 +2437,7 @@ void SAnimNotifyTrack::OnSetDurationNotifyClicked(int32 NotifyIndex)
 {
 	if (AnimNotifies.IsValidIndex(NotifyIndex))
 	{
-		FString DefaultText = FString::Printf(TEXT("%f"), AnimNotifies[NotifyIndex]->Duration);
+		FString DefaultText = FString::Printf(TEXT("%f"), AnimNotifies[NotifyIndex]->GetDuration());
 
 		// Show dialog to enter weight
 		TSharedRef<STextEntryPopup> TextEntry =
@@ -2431,9 +2464,9 @@ void SAnimNotifyTrack::SetDuration(const FText& DurationText, ETextCommit::Type
 		if ( AnimNotifies.IsValidIndex(NotifyIndex) )
 		{
 			float NewDuration = FMath::Max(FCString::Atof( *DurationText.ToString() ), SAnimNotifyNode::MinimumStateDuration);
-			float MaxDuration = Sequence->SequenceLength - AnimNotifies[NotifyIndex]->DisplayTime;
+			float MaxDuration = Sequence->SequenceLength - AnimNotifies[NotifyIndex]->GetTime();
 			NewDuration = FMath::Min(NewDuration, MaxDuration);
-			AnimNotifies[NotifyIndex]->Duration = NewDuration;
+			AnimNotifies[NotifyIndex]->SetDuration(NewDuration);
 
 			// If we have a delegate bound to refresh the offsets, call it.
 			// This is used by the montage editor to keep the offsets up to date.
@@ -2507,6 +2540,7 @@ void SAnimNotifyTrack::Update()
 				.ViewInputMin(ViewInputMin)
 				.ViewInputMax(ViewInputMax)
 				.MarkerBars(MarkerBars)
+				.OnDeselectAllNotifies(OnDeselectAllNotifies)
 			];
 
 			NotifyNodes.Add(AnimNotifyNode);
@@ -2540,7 +2574,7 @@ FReply SAnimNotifyTrack::OnNotifyNodeDragStarted( TSharedRef<SAnimNotifyNode> No
 	{
 		FAnimNotifyEvent* First = AnimNotifies[A];
 		FAnimNotifyEvent* Second = AnimNotifies[B];
-		return First->DisplayTime < Second->DisplayTime;
+		return First->GetTime() < Second->GetTime();
 	});
 
 	// If we're dragging one of the direction markers we don't need to call any further as we don't want the drag drop op
@@ -2614,20 +2648,24 @@ void SAnimNotifyTrack::HandleNodeDrop(TSharedPtr<SAnimNotifyNode> Node, float Of
 	float LocalX = GetCachedGeometry().AbsoluteToLocal(Node->GetScreenPosition() + Offset).X;
 	float Time = GetCachedScaleInfo().LocalXToInput(LocalX);
 	FAnimNotifyEvent* DroppedEvent = Node->NotifyEvent;
+	float EventDuration = DroppedEvent->GetDuration();
+
+	UAnimMontage* MontageObj = Cast<UAnimMontage>(Sequence);
 
 	if(Node->GetLastSnappedTime() != -1.0f)
 	{
-		DroppedEvent->DisplayTime = Node->GetLastSnappedTime();
+		DroppedEvent->LinkMontage(MontageObj, Node->GetLastSnappedTime(), DroppedEvent->GetSlotIndex());
 	}
 	else
 	{
-		DroppedEvent->DisplayTime = Time;
+		DroppedEvent->LinkMontage(MontageObj, Time, DroppedEvent->GetSlotIndex());
 	}
-	DroppedEvent->RefreshTriggerOffset(Sequence->CalculateOffsetForNotify(DroppedEvent->DisplayTime));
+	DroppedEvent->RefreshTriggerOffset(Sequence->CalculateOffsetForNotify(DroppedEvent->GetTime()));
 
-	if(DroppedEvent->Duration > 0.0f)
+	if(EventDuration > 0.0f)
 	{
-		DroppedEvent->RefreshEndTriggerOffset(Sequence->CalculateOffsetForNotify(DroppedEvent->DisplayTime + DroppedEvent->Duration));
+		DroppedEvent->EndLink.LinkMontage(MontageObj, DroppedEvent->GetTime() + EventDuration, DroppedEvent->GetSlotIndex());
+		DroppedEvent->RefreshEndTriggerOffset(Sequence->CalculateOffsetForNotify(DroppedEvent->EndLink.GetTime()/*DroppedEvent->GetTime()+ DroppedEvent->GetDuration()*/));
 	}
 	else
 	{
@@ -2675,7 +2713,7 @@ void SAnimNotifyTrack::AppendSelectionToSet(FGraphPanelSelectionSet& SelectionSe
 	}
 }
 
-void SAnimNotifyTrack::AppendSelectionToArray(TArray<const FAnimNotifyEvent*>& Selection) const
+void SAnimNotifyTrack::AppendSelectionToArray(TArray<FAnimNotifyEvent*>& Selection) const
 {
 	for(int32 Idx : SelectedNotifyIndices)
 	{
@@ -2703,12 +2741,12 @@ void SAnimNotifyTrack::PasteSingleNotify(FString& NotifyString, float PasteTime)
 
 		if(PasteTime != -1.0f)
 		{
-			NewNotify.DisplayTime = PasteTime;
+			NewNotify.SetTime(PasteTime);
 		}
 
 		// Make sure the notify is within the track area
-		NewNotify.DisplayTime = FMath::Clamp(NewNotify.DisplayTime, 0.0f, Sequence->SequenceLength);
-		NewNotify.TriggerTimeOffset = GetTriggerTimeOffsetForType(Sequence->CalculateOffsetForNotify(NewNotify.DisplayTime));
+		NewNotify.SetTime(FMath::Clamp(NewNotify.GetTime(), 0.0f, Sequence->SequenceLength));
+		NewNotify.TriggerTimeOffset = GetTriggerTimeOffsetForType(Sequence->CalculateOffsetForNotify(NewNotify.GetTime()));
 		NewNotify.TrackIndex = TrackIndex;
 
 		if(NewNotify.Notify)
@@ -2724,8 +2762,8 @@ void SAnimNotifyTrack::PasteSingleNotify(FString& NotifyString, float PasteTime)
 			NewNotify.NotifyStateClass = NewNotifyStateObject;
 
 			// Clamp duration into the sequence
-			NewNotify.Duration = FMath::Clamp(NewNotify.Duration, 1 / 30.0f, Sequence->SequenceLength - NewNotify.DisplayTime);
-			NewNotify.EndTriggerTimeOffset = GetTriggerTimeOffsetForType(Sequence->CalculateOffsetForNotify(NewNotify.DisplayTime + NewNotify.Duration));
+			NewNotify.SetDuration(FMath::Clamp(NewNotify.GetDuration(), 1 / 30.0f, Sequence->SequenceLength - NewNotify.GetTime()));
+			NewNotify.EndTriggerTimeOffset = GetTriggerTimeOffsetForType(Sequence->CalculateOffsetForNotify(NewNotify.GetTime() + NewNotify.GetDuration()));
 		}
 	}
 	else
@@ -3278,9 +3316,21 @@ void SAnimNotifyPanel::OnTrackSelectionChanged()
 {
 	// Need to collect selection info from all tracks
 	FGraphPanelSelectionSet SelectionSet;
-	for(TSharedPtr<SAnimNotifyTrack> Track : NotifyAnimTracks)
+	TArray<UEditorNotifyObject*> NotifyObjects;
+
+	TArray<FAnimNotifyEvent*> Events;
+	for(int32 TrackIdx = 0 ; TrackIdx < NotifyAnimTracks.Num() ; ++TrackIdx)
 	{
-		Track->AppendSelectionToSet(SelectionSet);
+		TSharedPtr<SAnimNotifyTrack> Track = NotifyAnimTracks[TrackIdx];
+		TArray<int32> TrackIndices = Track->GetSelectedNotifyIndices();
+		for(int32 Idx : TrackIndices)
+		{
+			FString ObjName = MakeUniqueObjectName(GetTransientPackage(), UEditorNotifyObject::StaticClass()).ToString();
+			UEditorNotifyObject* NewNotifyObject = Cast<UEditorNotifyObject>(StaticConstructObject(UEditorNotifyObject::StaticClass(), GetTransientPackage(), FName(*ObjName), RF_Public|RF_Standalone|RF_Transient));
+			NewNotifyObject->InitFromAnim(Sequence, FOnAnimObjectChange::CreateSP(this, &SAnimNotifyPanel::OnNotifyObjectChanged));
+			NewNotifyObject->InitialiseNotify(NotifyAnimTracks.Num() - TrackIdx - 1, Idx);
+			SelectionSet.Add(NewNotifyObject);
+		}
 	}
 
 	OnSelectionChanged.ExecuteIfBound(SelectionSet);
@@ -3299,7 +3349,7 @@ void SAnimNotifyPanel::DeselectAllNotifies()
 void SAnimNotifyPanel::CopySelectedNotifiesToClipboard() const
 {
 	// Grab the selected events
-	TArray<const FAnimNotifyEvent*> SelectedNotifies;
+	TArray<FAnimNotifyEvent*> SelectedNotifies;
 	for(TSharedPtr<SAnimNotifyTrack> Track : NotifyAnimTracks)
 	{
 		Track->AppendSelectionToArray(SelectedNotifies);
@@ -3314,7 +3364,7 @@ void SAnimNotifyPanel::CopySelectedNotifiesToClipboard() const
 		// Sort by track
 		SelectedNotifies.Sort([](const FAnimNotifyEvent& A, const FAnimNotifyEvent& B)
 		{
-			return (A.TrackIndex > B.TrackIndex) || (A.TrackIndex == B.TrackIndex && A.DisplayTime < B.DisplayTime);
+			return (A.TrackIndex > B.TrackIndex) || (A.TrackIndex == B.TrackIndex && A.GetTime() < B.GetTime());
 		});
 
 		// Need to find how many tracks this selection spans and the minimum time to use as the beginning of the selection
@@ -3325,7 +3375,7 @@ void SAnimNotifyPanel::CopySelectedNotifiesToClipboard() const
 		{
 			MinTrack = FMath::Min(MinTrack, Event->TrackIndex);
 			MaxTrack = FMath::Max(MaxTrack, Event->TrackIndex);
-			MinTime = FMath::Min(MinTime, Event->DisplayTime);
+			MinTime = FMath::Min(MinTime, Event->GetTime());
 		}
 
 		int32 TrackSpan = MaxTrack - MinTrack + 1;
@@ -3664,4 +3714,9 @@ void SAnimNotifyPanel::OnGetNotifyBlueprintData(TArray<FAssetData>& OutNotifyDat
 	}
 }
 
+void SAnimNotifyPanel::OnNotifyObjectChanged(UObject* EditorBaseObj, bool bRebuild)
+{
+
+}
+
 #undef LOCTEXT_NAMESPACE
diff --git a/Engine/Source/Editor/Persona/Private/SAnimNotifyPanel.h b/Engine/Source/Editor/Persona/Private/SAnimNotifyPanel.h
--- a/Engine/Source/Editor/Persona/Private/SAnimNotifyPanel.h
+++ b/Engine/Source/Editor/Persona/Private/SAnimNotifyPanel.h
@@ -206,6 +206,8 @@ public:
 
 	void RefreshMarqueeSelectedNodes(const FGeometry& PanelGeo);
 
+	void OnNotifyObjectChanged(UObject* EditorBaseObj, bool bRebuild);
+
 private:
 	TSharedPtr<SBorder> PanelArea;
 	class UAnimSequenceBase* Sequence;
diff --git a/Engine/Source/Editor/Persona/Private/SAnimationScrubPanel.cpp b/Engine/Source/Editor/Persona/Private/SAnimationScrubPanel.cpp
--- a/Engine/Source/Editor/Persona/Private/SAnimationScrubPanel.cpp
+++ b/Engine/Source/Editor/Persona/Private/SAnimationScrubPanel.cpp
@@ -447,18 +447,19 @@ void SAnimationScrubPanel::OnCropAnimSequence( bool bFromStart, float CurrentTim
 				//Adjust anim notify time based on current time and bFromStart flag. 
 				for( int32 i=0; i<AnimSequence->Notifies.Num(); i++ )
 				{
-					float& NotifyTime = AnimSequence->Notifies[i].DisplayTime;
+					FAnimNotifyEvent& Notify = AnimSequence->Notifies[i];
+					float NotifyTime = Notify.GetTime();
 					if( bFromStart )
 					{
 						//Cropping from start to current time
 						if( NotifyTime <= CurrentTime )
 						{
-							NotifyTime = 0.0f;
-							AnimSequence->Notifies[i].TriggerTimeOffset = GetTriggerTimeOffsetForType(EAnimEventTriggerOffsets::OffsetAfter);
+							Notify.SetTime(0.0f);
+							Notify.TriggerTimeOffset = GetTriggerTimeOffsetForType(EAnimEventTriggerOffsets::OffsetAfter);
 						}
 						else
 						{
-							NotifyTime -= CurrentTime;
+							Notify.SetTime(NotifyTime - CurrentTime);
 						}
 					}
 					else
@@ -466,8 +467,8 @@ void SAnimationScrubPanel::OnCropAnimSequence( bool bFromStart, float CurrentTim
 						//Cropping from current time to the end.
 						if( NotifyTime >= CurrentTime )
 						{
-							NotifyTime = CurrentTime;
-							AnimSequence->Notifies[i].TriggerTimeOffset = GetTriggerTimeOffsetForType(EAnimEventTriggerOffsets::OffsetBefore);
+							Notify.SetTime(CurrentTime);
+							Notify.TriggerTimeOffset = GetTriggerTimeOffsetForType(EAnimEventTriggerOffsets::OffsetBefore);
 						}
 					}
 				}
diff --git a/Engine/Source/Editor/Persona/Private/SMontageEditor.cpp b/Engine/Source/Editor/Persona/Private/SMontageEditor.cpp
--- a/Engine/Source/Editor/Persona/Private/SMontageEditor.cpp
+++ b/Engine/Source/Editor/Persona/Private/SMontageEditor.cpp
@@ -37,7 +37,12 @@ void SMontageEditor::Construct(const FArguments& InArgs)
 		.Persona(InArgs._Persona)
 		);
 
-	PersonaPtr.Pin()->RegisterOnPostUndo(FPersona::FOnPostUndo::CreateSP( this, &SMontageEditor::PostUndo ) );	
+	TSharedPtr<FPersona> SharedPersona = PersonaPtr.Pin();
+	if(SharedPersona.IsValid())
+	{
+		SharedPersona->RegisterOnPostUndo(FPersona::FOnPostUndo::CreateSP( this, &SMontageEditor::PostUndo ) );
+		SharedPersona->RegisterOnPersonaRefresh(FPersona::FOnPersonaRefresh::CreateSP(this, &SMontageEditor::RebuildMontagePanel));
+	}
 
 	EditorPanels->AddSlot()
 	.AutoHeight()
@@ -110,9 +115,11 @@ void SMontageEditor::Construct(const FArguments& InArgs)
 
 SMontageEditor::~SMontageEditor()
 {
-	if (PersonaPtr.IsValid())
+	TSharedPtr<FPersona> SharedPersona = PersonaPtr.Pin();
+	if (SharedPersona.IsValid())
 	{
-		PersonaPtr.Pin()->UnregisterOnPostUndo(this);
+		SharedPersona->UnregisterOnPostUndo(this);
+		SharedPersona->UnregisterOnPersonaRefresh(this);
 	}
 }
 
@@ -166,13 +173,13 @@ void SMontageEditor::RefreshNotifyTriggerOffsets()
 		FAnimNotifyEvent& Notify = (*Iter);
 
 		// Offset for the beginning of a notify
-		EAnimEventTriggerOffsets::Type PredictedOffset = MontageObj->CalculateOffsetForNotify(Notify.DisplayTime);
+		EAnimEventTriggerOffsets::Type PredictedOffset = MontageObj->CalculateOffsetForNotify(Notify.GetTime());
 		Notify.RefreshTriggerOffset(PredictedOffset);
 
 		// Offset for the end of a notify state if necessary
-		if(Notify.Duration > 0.0f)
+		if(Notify.GetDuration() > 0.0f)
 		{
-			PredictedOffset = MontageObj->CalculateOffsetForNotify(Notify.DisplayTime + Notify.Duration);
+			PredictedOffset = MontageObj->CalculateOffsetForNotify(Notify.GetTime() + Notify.GetDuration());
 			Notify.RefreshEndTriggerOffset(PredictedOffset);
 		}
 		else
@@ -185,7 +192,7 @@ void SMontageEditor::RefreshNotifyTriggerOffsets()
 	{
 		FBranchingPoint& BranchPoint = (*Iter);
 
-		EAnimEventTriggerOffsets::Type PredictedOffset = MontageObj->CalculateOffsetForBranchingPoint(BranchPoint.DisplayTime);
+		EAnimEventTriggerOffsets::Type PredictedOffset = MontageObj->CalculateOffsetForBranchingPoint(BranchPoint.GetTime());
 		BranchPoint.RefreshTriggerOffset(PredictedOffset);
 	}
 }
@@ -194,7 +201,7 @@ bool SMontageEditor::GetSectionTime( int32 SectionIndex, float &OutTime ) const
 {
 	if (MontageObj != NULL && MontageObj->CompositeSections.IsValidIndex(SectionIndex))
 	{
-		OutTime = MontageObj->CompositeSections[SectionIndex].StartTime;
+		OutTime = MontageObj->CompositeSections[SectionIndex].GetTime();
 		return true;
 	}
 
@@ -221,7 +228,7 @@ TArray<float> SMontageEditor::GetSectionStartTimes() const
 	{
 		for( int32 I=0; I < MontageObj->CompositeSections.Num(); I++)
 		{
-			Times.Add(MontageObj->CompositeSections[I].StartTime);
+			Times.Add(MontageObj->CompositeSections[I].GetTime());
 		}
 	}
 	return Times;
@@ -235,14 +242,14 @@ TArray<FTrackMarkerBar> SMontageEditor::GetMarkerBarInformation() const
 		for( int32 I=0; I < MontageObj->CompositeSections.Num(); I++)
 		{
 			FTrackMarkerBar Bar;
-			Bar.Time = MontageObj->CompositeSections[I].StartTime;
+			Bar.Time = MontageObj->CompositeSections[I].GetTime();
 			Bar.DrawColour = FLinearColor(0.f,1.f,0.f);
 			MarkerBars.Add(Bar);
 		}
 		for( int32 I = 0; I < MontageObj->BranchingPoints.Num(); ++I )
 		{
 			FTrackMarkerBar Bar;
-			Bar.Time = MontageObj->BranchingPoints[I].DisplayTime;
+			Bar.Time = MontageObj->BranchingPoints[I].GetTime();
 			Bar.DrawColour = FLinearColor(0.8f,0.f,0.8f);
 			MarkerBars.Add(Bar);
 		}
@@ -278,7 +285,8 @@ void SMontageEditor::OnEditSectionTime( int32 SectionIndex, float NewTime)
 		}
 		bDragging = true;
 
-		MontageObj->CompositeSections[SectionIndex].StartTime = NewTime;
+		MontageObj->CompositeSections[SectionIndex].SetTime(NewTime);
+		MontageObj->CompositeSections[SectionIndex].LinkMontage(MontageObj, NewTime);
 	}
 }
 void SMontageEditor::OnEditSectionTimeFinish( int32 SectionIndex )
@@ -349,6 +357,8 @@ void SMontageEditor::CollapseMontage()
 		MontageObj->SlotAnimTracks[i].AnimTrack.CollapseAnimSegments();
 	}
 
+	MontageObj->UpdateLinkableElements();
+
 	RecalculateSequenceLength();
 }
 
@@ -362,6 +372,8 @@ void SMontageEditor::SortAndUpdateMontage()
 	
 	SortAnimSegments();
 
+	MontageObj->UpdateLinkableElements();
+
 	RecalculateSequenceLength();
 
 	SortSections();
@@ -424,7 +436,7 @@ void SMontageEditor::SortSections()
 	{
 		bool operator()( const FCompositeSection &A, const FCompositeSection &B ) const
 		{
-			return A.StartTime < B.StartTime;
+			return A.GetTime() < B.GetTime();
 		}
 	};
 	if (MontageObj != NULL)
@@ -441,16 +453,16 @@ void SMontageEditor::EnsureStartingSection()
 	if(MontageObj->CompositeSections.Num() <= 0)
 	{
 		FCompositeSection NewSection;
-		NewSection.StartTime = 0.f;
+		NewSection.SetTime(0.0f);
 		NewSection.SectionName = FName(TEXT("Default"));
 		MontageObj->CompositeSections.Add(NewSection);
 		MontageObj->MarkPackageDirty();
 	}
 
 	check(MontageObj->CompositeSections.Num() > 0);
-	if(MontageObj->CompositeSections[0].StartTime > 0.f)
+	if(MontageObj->CompositeSections[0].GetTime() > 0.0f)
 	{
-		MontageObj->CompositeSections[0].StartTime = 0.f;
+		MontageObj->CompositeSections[0].SetTime(0.0f);
 		MontageObj->MarkPackageDirty();
 	}
 }
@@ -475,30 +487,33 @@ bool SMontageEditor::ClampToEndTime(float NewEndTime)
 
 		for(int32 i=0; i < MontageObj->CompositeSections.Num(); i++)
 		{
-			if(MontageObj->CompositeSections[i].StartTime > NewEndTime)
+			if(MontageObj->CompositeSections[i].GetTime() > NewEndTime)
 			{
-				MontageObj->CompositeSections[i].StartTime *= ratio;
+				float CurrentTime = MontageObj->CompositeSections[i].GetTime();
+				MontageObj->CompositeSections[i].SetTime(CurrentTime * ratio);
 			}
 		}
 
 		for(int32 i=0; i < MontageObj->BranchingPoints.Num(); i++)
 		{
 			FBranchingPoint& BranchPoint = MontageObj->BranchingPoints[i];
-			if(BranchPoint.DisplayTime > NewEndTime)
+			float BranchingPointTime = BranchPoint.GetTime();
+			if(BranchingPointTime > NewEndTime)
 			{
-				BranchPoint.DisplayTime *= ratio;
-				BranchPoint.TriggerTimeOffset = GetTriggerTimeOffsetForType(MontageObj->CalculateOffsetForBranchingPoint(BranchPoint.DisplayTime));
+				BranchPoint.SetTime(BranchingPointTime * ratio);
+				BranchPoint.TriggerTimeOffset = GetTriggerTimeOffsetForType(MontageObj->CalculateOffsetForBranchingPoint(BranchPoint.GetTime()));
 			}
 		}
 
 		for(int32 i=0; i < MontageObj->Notifies.Num(); i++)
 		{
 			FAnimNotifyEvent& Notify = MontageObj->Notifies[i];
+			float NotifyTime = Notify.GetTime();
 
-			if(Notify.DisplayTime >= NewEndTime)
+			if(NotifyTime >= NewEndTime)
 			{
-				Notify.DisplayTime *= ratio;
-				Notify.TriggerTimeOffset = GetTriggerTimeOffsetForType(MontageObj->CalculateOffsetForNotify(Notify.DisplayTime));
+				Notify.SetTime(NotifyTime * ratio);
+				Notify.TriggerTimeOffset = GetTriggerTimeOffsetForType(MontageObj->CalculateOffsetForNotify(Notify.GetTime()));
 			}
 		}
 	}
@@ -548,7 +563,7 @@ float SMontageEditor::GetBranchPointStartPos(int32 BranchPointIndex) const
 {
 	if (ValidBranch(BranchPointIndex))
 	{
-		return MontageObj->BranchingPoints[BranchPointIndex].DisplayTime;
+		return MontageObj->BranchingPoints[BranchPointIndex].GetTime();
 	}
 	return 0.f;
 }
@@ -573,7 +588,7 @@ void SMontageEditor::SetBranchPointStartPos(float NewStartPos, int32 BranchPoint
 
 		FBranchingPoint& BranchPoint = MontageObj->BranchingPoints[BranchPointIndex];
 
-		float OriginalPosition = BranchPoint.DisplayTime;
+		float OriginalPosition = BranchPoint.GetTime();
 		float TempPosition = NewStartPos;
 		int32 MovementDirection = 1;
 
@@ -588,7 +603,7 @@ void SMontageEditor::SetBranchPointStartPos(float NewStartPos, int32 BranchPoint
 				continue;
 			}
 
-			float BranchTime = MontageObj->BranchingPoints[BranchIndex].DisplayTime;
+			float BranchTime = MontageObj->BranchingPoints[BranchIndex].GetTime();
 			float LowerBound = BranchTime - MinimumBranchPointSeperation;
 			float UpperBound = BranchTime + MinimumBranchPointSeperation;
 
@@ -634,8 +649,13 @@ void SMontageEditor::SetBranchPointStartPos(float NewStartPos, int32 BranchPoint
 			NewStartPos = OriginalPosition;
 		}
 
-		BranchPoint.DisplayTime = NewStartPos;
-		BranchPoint.RefreshTriggerOffset(MontageObj->CalculateOffsetForBranchingPoint(BranchPoint.DisplayTime));
+		BranchPoint.LinkMontage(MontageObj, NewStartPos, BranchPoint.GetSlotIndex());
+		BranchPoint.RefreshTriggerOffset(MontageObj->CalculateOffsetForBranchingPoint(BranchPoint.GetTime()));
+
+		if(AnimMontagePanel.IsValid())
+		{
+			AnimMontagePanel->ShowBranchPointInDetailsView(BranchPointIndex);
+		}
 	}
 }
 
@@ -659,8 +679,8 @@ void SMontageEditor::AddBranchPoint(float StartTime, FString EventName)
 		MontageObj->Modify();
 
 		FBranchingPoint NewBranch;
-		NewBranch.DisplayTime=StartTime;
-		NewBranch.TriggerTimeOffset = GetTriggerTimeOffsetForType(MontageObj->CalculateOffsetForBranchingPoint(NewBranch.DisplayTime));
+		NewBranch.LinkMontage(MontageObj, StartTime);
+		NewBranch.TriggerTimeOffset = GetTriggerTimeOffsetForType(MontageObj->CalculateOffsetForBranchingPoint(NewBranch.GetTime()));
 		NewBranch.EventName=FName(*EventName);
 		MontageObj->BranchingPoints.Add(NewBranch);
 		MontageObj->MarkPackageDirty();
diff --git a/Engine/Source/Editor/UnrealEd/Classes/Animation/EditorNotifyObject.h b/Engine/Source/Editor/UnrealEd/Classes/Animation/EditorNotifyObject.h
new file mode 100644
--- /dev/null
+++ b/Engine/Source/Editor/UnrealEd/Classes/Animation/EditorNotifyObject.h
@@ -0,0 +1,37 @@
+// Copyright 1998-2014 Epic Games, Inc. All Rights Reserved.
+
+//////////////////////////////////////////////////////////////////////////
+// Proxy object for displaying notifies in the details panel with
+// event data included alongside UAnimNotify 
+//////////////////////////////////////////////////////////////////////////
+
+#pragma once
+
+#include "EditorNotifyObject.generated.h"
+
+class UEditorNotifyEntry;
+
+UCLASS(MinimalAPI)
+class UEditorNotifyObject : public UEditorAnimBaseObj
+{
+	GENERATED_UCLASS_BODY()
+
+	/** Set up the editor object
+	 *	@param TrackIdx Index of the track the notify belongs to
+	 *	@param NotifyIndex index of the notify within the track
+	 */
+	virtual void InitialiseNotify(int32 TrackIdx, int32 NotifyIndex);
+	
+	/** Copy changes made to the event object back to the montage asset */
+	virtual bool ApplyChangesToMontage();
+
+	/** Index of the notify within it's track */
+	int32 NotifyIndex;
+
+	/** Index of the track the notify is in */
+	int32 TrackIndex;
+
+	/** The notify event to modify */
+	UPROPERTY(EditAnywhere, Category=Event)
+	FAnimNotifyEvent Event;
+};
\ No newline at end of file
diff --git a/Engine/Source/Editor/UnrealEd/Private/Animation/EditorAnimSegment.cpp b/Engine/Source/Editor/UnrealEd/Private/Animation/EditorAnimSegment.cpp
--- a/Engine/Source/Editor/UnrealEd/Private/Animation/EditorAnimSegment.cpp
+++ b/Engine/Source/Editor/UnrealEd/Private/Animation/EditorAnimSegment.cpp
@@ -5,6 +5,8 @@
 =============================================================================*/ 
 
 #include "UnrealEd.h"
+#include "UnrealType.h"
+#include "List.h"
 
 
 
@@ -38,6 +40,15 @@ bool UEditorAnimSegment::ApplyChangesToMontage()
 			if (AnimSegment.AnimReference && AnimSegment.AnimReference->GetSkeleton() == Montage->GetSkeleton())
 			{
 				Montage->SlotAnimTracks[AnimSlotIndex].AnimTrack.AnimSegments[AnimSegmentIndex] = AnimSegment;
+				Montage->UpdateLinkableElements(AnimSlotIndex, AnimSegmentIndex);
+
+				// Need to update linkable elements for segments further along.
+				int32 NumSegments = Montage->SlotAnimTracks[AnimSlotIndex].AnimTrack.AnimSegments.Num();
+				for(int32 SegmentIdx = AnimSegmentIndex + 1 ; SegmentIdx < NumSegments ; ++SegmentIdx)
+				{
+					Montage->UpdateLinkableElements(AnimSlotIndex, SegmentIdx);
+				}
+
 				return true;
 			}
 			else
diff --git a/Engine/Source/Editor/UnrealEd/Private/Animation/EditorBranchPoint.cpp b/Engine/Source/Editor/UnrealEd/Private/Animation/EditorBranchPoint.cpp
--- a/Engine/Source/Editor/UnrealEd/Private/Animation/EditorBranchPoint.cpp
+++ b/Engine/Source/Editor/UnrealEd/Private/Animation/EditorBranchPoint.cpp
@@ -33,6 +33,7 @@ bool UEditorBranchPoint::ApplyChangesToMontage()
 	{
 		if(Montage->BranchingPoints.IsValidIndex(BranchIndex))
 		{
+			BranchingPoint.OnChanged(BranchingPoint.GetTime());
 			Montage->BranchingPoints[BranchIndex] = BranchingPoint;
 			return true;
 		}
diff --git a/Engine/Source/Editor/UnrealEd/Private/Animation/EditorCompositeSection.cpp b/Engine/Source/Editor/UnrealEd/Private/Animation/EditorCompositeSection.cpp
--- a/Engine/Source/Editor/UnrealEd/Private/Animation/EditorCompositeSection.cpp
+++ b/Engine/Source/Editor/UnrealEd/Private/Animation/EditorCompositeSection.cpp
@@ -32,6 +32,7 @@ bool UEditorCompositeSection::ApplyChangesToMontage()
 	{
 		if(Montage->CompositeSections.IsValidIndex(SectionIndex))
 		{
+			CompositeSection.OnChanged(CompositeSection.GetTime());
 			Montage->CompositeSections[SectionIndex] = CompositeSection;
 			return true;
 		}
diff --git a/Engine/Source/Editor/UnrealEd/Private/Animation/EditorNotifyObject.cpp b/Engine/Source/Editor/UnrealEd/Private/Animation/EditorNotifyObject.cpp
new file mode 100644
--- /dev/null
+++ b/Engine/Source/Editor/UnrealEd/Private/Animation/EditorNotifyObject.cpp
@@ -0,0 +1,46 @@
+// Copyright 1998-2014 Epic Games, Inc. All Rights Reserved.
+
+#include "UnrealEd.h"
+
+UEditorNotifyObject::UEditorNotifyObject(const FObjectInitializer& ObjectInitializer)
+	: Super(ObjectInitializer)
+{
+
+}
+
+bool UEditorNotifyObject::ApplyChangesToMontage()
+{
+	if(AnimObject)
+	{
+		check(AnimObject->AnimNotifyTracks.IsValidIndex(TrackIndex));
+		FAnimNotifyTrack* Track = &AnimObject->AnimNotifyTracks[TrackIndex];
+		check(Track->Notifies.IsValidIndex(NotifyIndex));
+		FAnimNotifyEvent* ActualNotify = Track->Notifies[NotifyIndex];
+		Event.OnChanged(Event.GetTime());
+
+		// If we have a duration this is a state notify
+		if(Event.GetDuration() > 0.0f)
+		{
+			Event.EndLink.OnChanged(Event.EndLink.GetTime());
+
+			// Always keep link methods in sync between notifies and duration links
+			if(Event.GetLinkMethod() != Event.EndLink.GetLinkMethod())
+			{
+				Event.EndLink.ChangeLinkMethod(Event.GetLinkMethod());
+			}
+		}
+		*ActualNotify = Event;
+	}
+
+	return true;
+}
+
+void UEditorNotifyObject::InitialiseNotify(int32 InTrackIdx, int32 InNotifyIndex)
+{
+	if(AnimObject)
+	{
+		Event = *AnimObject->AnimNotifyTracks[InTrackIdx].Notifies[InNotifyIndex];
+		NotifyIndex = InNotifyIndex;
+		TrackIndex = InTrackIdx;
+	}
+}
diff --git a/Engine/Source/Runtime/Engine/Classes/Animation/AnimCompositeBase.h b/Engine/Source/Runtime/Engine/Classes/Animation/AnimCompositeBase.h
--- a/Engine/Source/Runtime/Engine/Classes/Animation/AnimCompositeBase.h
+++ b/Engine/Source/Runtime/Engine/Classes/Animation/AnimCompositeBase.h
@@ -168,7 +168,15 @@ struct FAnimTrack
 	 */
 	void GetRootMotionExtractionStepsForTrackRange(TArray<FRootMotionExtractionStep> & RootMotionExtractionSteps, const float StartTrackPosition, const float EndTrackPosition) const;
 
-	
+	/** Ensure segment times are correctly formed (no gaps and no extra time at the end of the anim reference) */
+	void ValidateSegmentTimes();
+
+	/** Gets the index of the segment at the given absolute montage time. */	
+	int32 GetSegmentIndexAtTime(float InTime);
+
+	/** Get the segment at the given absolute montage time */
+	FAnimSegment* GetSegmentAtTime(float InTime);
+
 #if WITH_EDITOR
 	bool GetAllAnimationSequencesReferred(TArray<UAnimSequence*>& AnimationSequences) const;
 	void ReplaceReferredAnimations(const TMap<UAnimSequence*, UAnimSequence*>& ReplacementMap);
diff --git a/Engine/Source/Runtime/Engine/Classes/Animation/AnimLinkableElement.cpp b/Engine/Source/Runtime/Engine/Classes/Animation/AnimLinkableElement.cpp
new file mode 100644
--- /dev/null
+++ b/Engine/Source/Runtime/Engine/Classes/Animation/AnimLinkableElement.cpp
@@ -0,0 +1,367 @@
+// Copyright 1998-2014 Epic Games, Inc. All Rights Reserved.
+
+#include "EnginePrivate.h"
+#include "AnimLinkableElement.h"
+#include "AnimMontage.h"
+#include "AnimSequenceBase.h"
+
+void FAnimLinkableElement::LinkMontage(UAnimMontage* Montage, float AbsMontageTime, int32 InSlotIndex)
+{
+	if(Montage && Montage->SlotAnimTracks.Num() > 0)
+	{
+		LinkedMontage = Montage;
+
+		SlotIndex = InSlotIndex;
+		FSlotAnimationTrack& Slot = Montage->SlotAnimTracks[SlotIndex];
+
+		SegmentIndex = Slot.AnimTrack.GetSegmentIndexAtTime(AbsMontageTime);
+		if(SegmentIndex != INDEX_NONE)
+		{
+			FAnimSegment& Segment = Slot.AnimTrack.AnimSegments[SegmentIndex];
+			LinkedSequence = Segment.AnimReference;
+			SegmentBeginTime = Segment.StartPos;
+			SegmentLength = Segment.GetLength();
+
+			SetTime_Internal(AbsMontageTime);
+		}
+		else
+		{
+			// Nothing to link to
+			// We have no segment to link to, we need to clear our the segment data
+			// and give ourselves an absolute time
+			LinkValue = AbsMontageTime;
+			LinkedSequence = nullptr;
+			SegmentBeginTime = -1.0f;
+			SegmentLength = -1.0f;
+			LinkMethod = EAnimLinkMethod::Absolute;
+			CachedLinkMethod = LinkMethod;
+		}
+	}
+}
+
+void FAnimLinkableElement::LinkSequence(UAnimSequenceBase* Sequence, float AbsSequenceTime)
+{
+	if(Sequence && Sequence->SequenceLength > 0)
+	{
+		LinkedMontage = nullptr;
+		LinkedSequence = Sequence;
+		SegmentIndex = 0;
+
+		SegmentBeginTime = 0.0f;
+		SegmentLength = Sequence->SequenceLength;
+
+		SetTime(AbsSequenceTime);
+	}
+}
+
+void FAnimLinkableElement::Clear()
+{
+	LinkedSequence = nullptr;
+	SegmentBeginTime = -1.0f;
+	SegmentLength = -1.0f;
+	LinkMethod = EAnimLinkMethod::Absolute;
+	CachedLinkMethod = LinkMethod;
+}
+
+void FAnimLinkableElement::Update()
+{
+	if(LinkedMontage)
+	{
+		check(LinkedMontage->SlotAnimTracks.IsValidIndex(SlotIndex));
+		FSlotAnimationTrack& Slot = LinkedMontage->SlotAnimTracks[SlotIndex];
+		float CurrentTime = GetTime();
+
+		// If we don't have a segment, check to see if one has been added.
+		if(SegmentIndex == INDEX_NONE)
+		{
+			SegmentIndex = Slot.AnimTrack.GetSegmentIndexAtTime(CurrentTime);
+		}
+
+		if(SegmentIndex != INDEX_NONE)
+		{
+			// Update timing info from current segment
+			FAnimSegment& Segment = Slot.AnimTrack.AnimSegments[SegmentIndex];
+			LinkedSequence = Segment.AnimReference;
+			SegmentBeginTime = Segment.StartPos;
+			SegmentLength = Segment.GetLength();
+
+			// Handle Relative link mode, make sure to stay within the linked segment
+			if(CachedLinkMethod == EAnimLinkMethod::Relative)
+			{
+				float SegmentEnd = SegmentBeginTime + SegmentLength;
+				if(GetTime() > SegmentEnd)
+				{
+					SetTime(SegmentEnd);
+				}
+			}
+
+			// Relink if necessary
+			ConditionalRelink();
+		}
+	}
+}
+
+void FAnimLinkableElement::OnChanged(float NewMontageTime)
+{
+	// Only update linkage in a montage.
+	if(!LinkedMontage)
+	{
+		return;
+	}
+
+	SlotIndex = FMath::Clamp(SlotIndex, 0, LinkedMontage->SlotAnimTracks.Num());
+
+	// If the link method changed, transform the link value
+	if(CachedLinkMethod != LinkMethod)
+	{
+		float AbsTime = -1.0f;
+		switch(CachedLinkMethod)
+		{
+			case EAnimLinkMethod::Absolute:
+				AbsTime = LinkValue;
+				break;
+			case EAnimLinkMethod::Relative:
+				AbsTime = GetTimeFromRelative(EAnimLinkMethod::Absolute);
+				break;
+			case EAnimLinkMethod::Proportional:
+				AbsTime = GetTimeFromProportional(EAnimLinkMethod::Absolute);
+				break;
+		}
+		check(AbsTime != -1.0f);
+		CachedLinkMethod = LinkMethod;
+
+		// We aren't changing the time, just transforming it so use internal settime
+		SetTime_Internal(AbsTime);
+		NewMontageTime = AbsTime;
+	}
+
+	FSlotAnimationTrack& Slot = LinkedMontage->SlotAnimTracks[SlotIndex];
+
+	SegmentIndex = Slot.AnimTrack.GetSegmentIndexAtTime(NewMontageTime);
+	if(SegmentIndex != INDEX_NONE)
+	{
+		// Update to the detected segment
+		FAnimSegment& Segment = Slot.AnimTrack.AnimSegments[SegmentIndex];
+		LinkedSequence = Segment.AnimReference;
+		SegmentBeginTime = Segment.StartPos;
+		SegmentLength = Segment.GetLength();
+
+		SetTime(NewMontageTime);
+	}
+	else if(!LinkedSequence)
+	{
+		// We have no segment to link to, we need to clear our the segment data
+		// and give ourselves an absolute time
+		LinkValue = NewMontageTime;
+		Clear();
+	}
+}
+
+FAnimSegment* FAnimLinkableElement::GetSegmentAtCurrentTime()
+{
+	FAnimSegment* Result = nullptr;
+	if(LinkedMontage)
+	{
+		Result = LinkedMontage->SlotAnimTracks[SlotIndex].AnimTrack.GetSegmentAtTime(GetTime());
+	}
+	return Result;
+}
+
+float FAnimLinkableElement::GetTime(EAnimLinkMethod::Type ReferenceFrame /*= EMontageLinkMethod::Absolute*/) const
+{
+	if(ReferenceFrame != CachedLinkMethod)
+	{
+		switch(CachedLinkMethod)
+		{
+			case EAnimLinkMethod::Absolute:
+			{
+				return GetTimeFromAbsolute(ReferenceFrame);
+			}
+				break;
+			case EAnimLinkMethod::Relative:
+			{
+				return GetTimeFromRelative(ReferenceFrame);
+			}
+				break;
+			case EAnimLinkMethod::Proportional:
+			{
+				return GetTimeFromProportional(ReferenceFrame);
+			}
+				break;
+		}
+	}
+	return LinkValue;
+}
+
+void FAnimLinkableElement::SetTime(float NewTime, EAnimLinkMethod::Type ReferenceFrame /*= EMontageLinkMethod::Absolute*/)
+{
+	SetTime_Internal(NewTime, ReferenceFrame);
+}
+
+float FAnimLinkableElement::GetTimeFromAbsolute(EAnimLinkMethod::Type ReferenceFrame) const
+{
+	switch(ReferenceFrame)
+	{
+		case EAnimLinkMethod::Relative:
+		{
+			return LinkValue - SegmentBeginTime;
+		}
+			break;
+		case EAnimLinkMethod::Proportional:
+		{
+			return (LinkValue - SegmentBeginTime) / SegmentLength;
+		}
+			break;
+	}
+	return -1.0f;
+}
+
+float FAnimLinkableElement::GetTimeFromRelative(EAnimLinkMethod::Type ReferenceFrame) const
+{
+	switch(ReferenceFrame)
+	{
+		case EAnimLinkMethod::Absolute:
+		{
+			return SegmentBeginTime + LinkValue;
+		}
+			break;
+		case EAnimLinkMethod::Proportional:
+		{
+			return LinkValue / SegmentLength;
+		}
+			break;
+	}
+	return -1.0f;
+}
+
+float FAnimLinkableElement::GetTimeFromProportional(EAnimLinkMethod::Type ReferenceFrame) const
+{
+	switch(ReferenceFrame)
+	{
+		case EAnimLinkMethod::Absolute:
+		{
+			return SegmentBeginTime + LinkValue * SegmentLength;
+		}
+			break;
+		case EAnimLinkMethod::Relative:
+		{
+			return LinkValue * SegmentLength;
+		}
+			break;
+	}
+	return -1.0f;
+}
+
+void FAnimLinkableElement::SetTimeFromAbsolute(float NewTime, EAnimLinkMethod::Type ReferenceFrame)
+{
+	switch(ReferenceFrame)
+	{
+		case EAnimLinkMethod::Relative:
+		{
+			LinkValue = SegmentBeginTime + NewTime;
+		}
+			break;
+		case EAnimLinkMethod::Proportional:
+		{
+			LinkValue = SegmentBeginTime + SegmentLength * NewTime;
+		}
+			break;
+	}
+}
+
+void FAnimLinkableElement::SetTimeFromRelative(float NewTime, EAnimLinkMethod::Type ReferenceFrame)
+{
+	switch(ReferenceFrame)
+	{
+		case EAnimLinkMethod::Absolute:
+		{
+			LinkValue = NewTime - SegmentBeginTime;
+		}
+			break;
+		case EAnimLinkMethod::Proportional:
+		{
+			LinkValue = NewTime * SegmentLength;
+		}
+			break;
+	}
+}
+
+void FAnimLinkableElement::SetTimeFromProportional(float NewTime, EAnimLinkMethod::Type ReferenceFrame)
+{
+	switch(ReferenceFrame)
+	{
+		case EAnimLinkMethod::Absolute:
+		{
+			LinkValue = (NewTime - SegmentBeginTime) / SegmentLength;
+		}
+			break;
+		case EAnimLinkMethod::Relative:
+		{
+			LinkValue = NewTime / SegmentLength;
+		}
+			break;
+	}
+}
+
+void FAnimLinkableElement::ChangeLinkMethod(EAnimLinkMethod::Type NewLinkMethod)
+{
+	if(NewLinkMethod != LinkMethod)
+	{
+		// Switch to the new link method and resolve it
+		LinkMethod = NewLinkMethod;
+		OnChanged(GetTime());
+	}
+}
+
+void FAnimLinkableElement::ChangeSlotIndex(int32 NewSlotIndex)
+{
+	if(LinkedMontage)
+	{
+		LinkMontage(LinkedMontage, GetTime(), NewSlotIndex);
+	}
+}
+
+void FAnimLinkableElement::SetTime_Internal(float NewTime, EAnimLinkMethod::Type ReferenceFrame)
+{
+	if(ReferenceFrame != CachedLinkMethod)
+	{
+		switch(CachedLinkMethod)
+		{
+			case EAnimLinkMethod::Absolute:
+			{
+				SetTimeFromAbsolute(NewTime, ReferenceFrame);
+			}
+				break;
+			case EAnimLinkMethod::Relative:
+			{
+				SetTimeFromRelative(NewTime, ReferenceFrame);
+			}
+				break;
+			case EAnimLinkMethod::Proportional:
+			{
+				SetTimeFromProportional(NewTime, ReferenceFrame);
+			}
+				break;
+		}
+	}
+	else
+	{
+		LinkValue = NewTime;
+	}
+}
+
+void FAnimLinkableElement::ConditionalRelink()
+{
+	float CurrentAbsTime = GetTime();
+	if(CurrentAbsTime < SegmentBeginTime || CurrentAbsTime > SegmentBeginTime + SegmentLength)
+	{
+		if(LinkedMontage)
+		{
+			LinkMontage(LinkedMontage, CurrentAbsTime, SlotIndex);
+		}
+		else if(LinkedSequence)
+		{
+			LinkSequence(LinkedSequence, CurrentAbsTime);
+		}
+	}
+}
diff --git a/Engine/Source/Runtime/Engine/Classes/Animation/AnimLinkableElement.h b/Engine/Source/Runtime/Engine/Classes/Animation/AnimLinkableElement.h
new file mode 100644
--- /dev/null
+++ b/Engine/Source/Runtime/Engine/Classes/Animation/AnimLinkableElement.h
@@ -0,0 +1,190 @@
+// Copyright 1998-2014 Epic Games, Inc. All Rights Reserved.
+
+#pragma once
+
+#include "AnimLinkableElement.generated.h"
+
+class UAnimMontage;
+class UAnimSequenceBase;
+struct FAnimSegment;
+
+/** Supported types of time for a linked element */
+UENUM()
+namespace EAnimLinkMethod
+{
+	enum Type
+	{
+		Absolute,		// Element stays at a specific time without moving
+		Relative,		// Element moves with its segment, but not when the segment changes size
+		Proportional	// Element moves with its segment and will stay at a certain proportion through the segment
+	};
+}
+
+/** Used to describe an element that can be linked to a segment in a montage or sequence.
+ *	Usage: 
+ *		Inherit from FAnimLinkableElement and make sure to call LinkMontage or LinkSequence
+ *		both on creation and on loading the element. From there SetTime and GetTime should be
+ *		used to control where this element is in the montage or sequence.
+ *	
+ *		For more advanced usage, see this implementation used in FAnimNotifyEvent where
+ *		we have a secondary link to handle a duration
+ *		@see FAnimNotifyEvent
+ */
+USTRUCT()
+struct FAnimLinkableElement
+{
+	GENERATED_USTRUCT_BODY()
+
+	FAnimLinkableElement()
+	: LinkedMontage(nullptr)
+	, SlotIndex(INDEX_NONE)
+	, SegmentIndex(INDEX_NONE)
+	, LinkMethod(EAnimLinkMethod::Absolute)
+	, CachedLinkMethod(LinkMethod)
+	, LinkedSequence(nullptr)
+	{
+	}
+
+	/** Update the timing information for this element. Will not search for a new segment unless this element does not have one */
+	ENGINE_API void Update();
+
+	/** Link this element to a montage 
+	 * @param Montage The montage to link to
+	 * @param AbsMontageTime The time in the montage that this element should be placed at
+	 * @param InSlotIndex The slot in the montage to detect segments in
+	 */
+	ENGINE_API void LinkMontage(UAnimMontage* Montage, float AbsMontageTime, int32 InSlotIndex = 0);
+
+	/** Link this element to a Sequence, Just setting basic data as sequences don't need full linking 
+	 * @param Sequence The sequence to link to
+	 * @param AbsSequenceTime The time in the sequence that this element should be placed at
+	 */
+	ENGINE_API void LinkSequence(UAnimSequenceBase* Sequence, float AbsSequenceTime);
+
+	/** Clear the linking information from this element, leaves montage link intact */
+	ENGINE_API void Clear();
+
+	/** Called when the properties of this element are changed */
+	ENGINE_API void OnChanged(float NewMontageTime);
+
+	/** Gets the current time for this element 
+	 * @param ReferenceFrame What kind of time to return
+	 */
+	ENGINE_API float GetTime(EAnimLinkMethod::Type ReferenceFrame = EAnimLinkMethod::Absolute) const;
+
+	/** Sets the time of this element
+	 * @param NewTime The time to set this element to
+	 * @param ReferenceFrame The kind of time being passed to this method
+	 */
+	ENGINE_API virtual void SetTime(float NewTime, EAnimLinkMethod::Type ReferenceFrame = EAnimLinkMethod::Absolute);
+
+	/** Gets the sequence this element is linked to */
+	ENGINE_API const UAnimSequenceBase* GetLinkedSequence() {return LinkedSequence;}
+
+	/** Changes the way this element is linked to its segment
+	 * @param NewLinkMethod The new linking method to use
+	 */
+	ENGINE_API void ChangeLinkMethod(EAnimLinkMethod::Type NewLinkMethod);
+
+	/** Change the montage slot that this element is linked to
+	 * @param NewSlotIndex The new slot to link to
+	 */
+	ENGINE_API void ChangeSlotIndex(int32 NewSlotIndex);
+
+	/** Get the method used to link this element to its segment */
+	ENGINE_API EAnimLinkMethod::Type GetLinkMethod() const {return LinkMethod;}
+
+	/** Get the slot index this element is currently linked to */
+	ENGINE_API int32 GetSlotIndex() const {return SlotIndex;}
+	
+	/** Get the index of the segment this element is currently linked to */
+	ENGINE_API int32 GetSegmentIndex() const {return SegmentIndex;}
+
+protected:
+
+	/** Gets the segment in the current montage in the current slot that is at the time of this element */
+	FAnimSegment* GetSegmentAtCurrentTime();
+
+	/** The montage that this element is currently linked to */
+	UPROPERTY()
+	UAnimMontage* LinkedMontage;
+
+	/** The slot index we are currently using within LinkedMontage */
+	UPROPERTY(EditAnywhere, Category=AnimLink)
+	int32 SlotIndex;
+
+	/** The index of the segment we are linked to within the slot we are using */
+	UPROPERTY()
+	int32 SegmentIndex;
+
+	/** The method we are using to calculate our times */
+	UPROPERTY(EditAnywhere, Category=AnimLink)
+	TEnumAsByte<EAnimLinkMethod::Type> LinkMethod;
+
+	/** Cached link method used to transform the time when LinkMethod changes, always relates to the currently stored time */
+	UPROPERTY()
+	TEnumAsByte<EAnimLinkMethod::Type> CachedLinkMethod;
+
+	/** The absolute time in the montage that our currently linked segment begins */
+	UPROPERTY()
+	float SegmentBeginTime;
+
+	/** The absolute length of our currently linked segment */
+	UPROPERTY()
+	float SegmentLength;
+
+	/** The time of this montage. This will differ depending upon the method we are using to link the time for this element */
+	UPROPERTY()
+	float LinkValue;
+
+	/** 
+	 * The Animation Sequence that this montage element will link to, when the sequence changes
+	 * in either length or rate; the element will correctly place itself in relation to the sequence
+	 */
+	UPROPERTY(VisibleAnywhere, AdvancedDisplay, Category=AnimLink)
+	UAnimSequenceBase* LinkedSequence;
+
+private:
+
+	/** Relinks this element if it is outside of its segments */
+	void ConditionalRelink();
+
+	/** Gets a time from an absolute LinkValue
+	 * @param NewLinkMethod The type of time to return
+	 */
+	float GetTimeFromAbsolute(EAnimLinkMethod::Type ReferenceFrame) const ;
+
+	/** Gets a time from a relative LinkValue
+	 * @param NewLinkMethod The type of time to return
+	 */
+	float GetTimeFromRelative(EAnimLinkMethod::Type ReferenceFrame) const ;
+
+	/** Gets a time from a proportional LinkValue
+	 * @param NewLinkMethod The type of time to return
+	 */
+	float GetTimeFromProportional(EAnimLinkMethod::Type ReferenceFrame) const;
+
+	/** Set the LinkValue to a converted absolute time
+	 * @param NewTime The new time to set
+	 * @param NewLinkMethod The type of time passed to this method
+	 */
+	void SetTimeFromAbsolute(float NewTime, EAnimLinkMethod::Type ReferenceFrame);
+
+	/** Set the LinkValue to a converted relative time
+	 * @param NewTime The new time to set
+	 * @param NewLinkMethod The type of time passed to this method
+	 */
+	void SetTimeFromRelative(float NewTime, EAnimLinkMethod::Type ReferenceFrame);
+
+	/** Set the LinkValue to a converted proportional time
+	 * @param NewTime The new time to set
+	 * @param NewLinkMethod The type of time passed to this method
+	 */
+	void SetTimeFromProportional(float NewTime, EAnimLinkMethod::Type ReferenceFrame);
+
+	/** Sets the LinkValue based on NewTime. Internal version to ensure overridden behavior isn't used when transforming a time
+	 * @param NewTime New time to set
+	 * @param ReferenceFrame The kind of time passed in to this method
+	 */
+	void SetTime_Internal(float NewTime, EAnimLinkMethod::Type ReferenceFrame = EAnimLinkMethod::Absolute);
+};
diff --git a/Engine/Source/Runtime/Engine/Classes/Animation/AnimMontage.h b/Engine/Source/Runtime/Engine/Classes/Animation/AnimMontage.h
--- a/Engine/Source/Runtime/Engine/Classes/Animation/AnimMontage.h
+++ b/Engine/Source/Runtime/Engine/Classes/Animation/AnimMontage.h
@@ -11,13 +11,15 @@
 #include "Animation/AnimInstance.h"
 #include "AnimMontage.generated.h"
 
+class UAnimMontage;
+
 /**
  * Section data for each track. Reference of data will be stored in the child class for the way they want
  * AnimComposite vs AnimMontage have different requirement for the actual data reference
  * This only contains composite section information. (vertical sequences)
  */
 USTRUCT()
-struct FCompositeSection
+struct FCompositeSection : public FAnimLinkableElement
 {
 	GENERATED_USTRUCT_BODY()
 
@@ -30,17 +32,17 @@ struct FCompositeSection
 	float StarTime_DEPRECATED;
 
 	/** Start Time **/
-	UPROPERTY(EditAnywhere, Category=Section)
-	float StartTime;
+	UPROPERTY()
+	float StartTime_DEPRECATED;
 
 	/** Should this animation loop. */
-	UPROPERTY(EditAnywhere, Category=Section)
+	UPROPERTY()
 	FName NextSectionName;
 
 	FCompositeSection()
-		: SectionName(NAME_None)
+		: FAnimLinkableElement()
+		, SectionName(NAME_None)
 		, StarTime_DEPRECATED(0.f)
-		, StartTime(0.f)
 		, NextSectionName(NAME_None)
 	{
 	}
@@ -74,22 +76,22 @@ struct FSlotAnimationTrack
  * Use notifies if possible
  */
 USTRUCT()
-struct FBranchingPoint
+struct FBranchingPoint : public FAnimLinkableElement
 {
 	GENERATED_USTRUCT_BODY()
 
 	UPROPERTY(EditAnywhere, Category=BranchingPoint)
 	FName EventName;
 
-	UPROPERTY(EditAnywhere, Category=BranchingPoint)
-	float DisplayTime;
+	UPROPERTY()
+	float DisplayTime_DEPRECATED;
 
 	/** An offset from the DisplayTime to the actual time we will trigger the notify, as we cannot always trigger it exactly at the time the user wants */
 	UPROPERTY()
 	float TriggerTimeOffset;
 
 	/** Returns the time this branching point should be triggered */
-	float GetTriggerTime() const { return DisplayTime + TriggerTimeOffset; }
+	float GetTriggerTime() const { return GetTime() + TriggerTimeOffset; }
 
 	/** Updates trigger offset based on a combination of predicted offset and current offset */
 	ENGINE_API void RefreshTriggerOffset(EAnimEventTriggerOffsets::Type PredictedOffsetType);
@@ -330,6 +332,16 @@ public:
 
 	/** Calculates what (if any) offset should be applied to the trigger time of a branch point given its display time */
 	ENGINE_API EAnimEventTriggerOffsets::Type CalculateOffsetForBranchingPoint(float BranchingPointDisplayTime) const;
+
+	/** Update all linkable elements contained in the montage */
+	ENGINE_API void UpdateLinkableElements();
+
+	/** Update linkable elements that rely on a specific segment. This will update linkable elements for the segment specified
+	 *	and elements linked to segments after the segment specified
+	 *	@param SlotIdx The slot that the segment is contained in
+	 *	@param SegmentIdx The index of the segment within the specified slot
+	 */
+	ENGINE_API void UpdateLinkableElements(int32 SlotIdx, int32 SegmentIdx);
 #endif
 
 	/** Find First BranchingPoint between ]StartTrackPos,EndTrackPos] 
@@ -430,6 +442,8 @@ private:
 
 #endif	//WITH_EDITOR
 
+private:
+
 	/** Sort CompositeSections in the order of StartPos */
 	void SortAnimBranchingPointByTime();
 
diff --git a/Engine/Source/Runtime/Engine/Classes/Animation/AnimSequenceBase.h b/Engine/Source/Runtime/Engine/Classes/Animation/AnimSequenceBase.h
--- a/Engine/Source/Runtime/Engine/Classes/Animation/AnimSequenceBase.h
+++ b/Engine/Source/Runtime/Engine/Classes/Animation/AnimSequenceBase.h
@@ -11,6 +11,7 @@
 #include "SmartName.h"
 #include "Skeleton.h"
 #include "Curves/CurveBase.h"
+#include "AnimLinkableElement.h"
 #include "AnimSequenceBase.generated.h"
 
 #define DEFAULT_SAMPLERATE			30.f
@@ -33,13 +34,13 @@ ENGINE_API float GetTriggerTimeOffsetForType(EAnimEventTriggerOffsets::Type Offs
  * which has its Notify method called and passed to the animation.
  */
 USTRUCT()
-struct FAnimNotifyEvent
+struct FAnimNotifyEvent : public FAnimLinkableElement
 {
 	GENERATED_USTRUCT_BODY()
 
 	/** The user requested time for this notify */
 	UPROPERTY()
-	float DisplayTime;
+	float DisplayTime_DEPRECATED;
 
 	/** An offset from the DisplayTime to the actual time we will trigger the notify, as we cannot always trigger it exactly at the time the user wants */
 	UPROPERTY()
@@ -52,18 +53,22 @@ struct FAnimNotifyEvent
 	UPROPERTY(EditAnywhere, BlueprintReadWrite, Category=AnimNotifyEvent)
 	float TriggerWeightThreshold;
 
-	UPROPERTY(EditAnywhere, BlueprintReadWrite, Category=AnimNotifyEvent)
+	UPROPERTY()
 	FName NotifyName;
 
-	UPROPERTY(EditAnywhere, BlueprintReadWrite, Category=AnimNotifyEvent)
+	UPROPERTY(EditAnywhere, Instanced, BlueprintReadWrite, Category=AnimNotifyEvent)
 	class UAnimNotify * Notify;
 
-	UPROPERTY(EditAnywhere, BlueprintReadWrite, Category=AnimNotifyEvent)
+	UPROPERTY(EditAnywhere, Instanced, BlueprintReadWrite, Category=AnimNotifyEvent)
 	class UAnimNotifyState * NotifyStateClass;
 
 	UPROPERTY()
 	float Duration;
 
+	/** Linkable element to use for the end handle representing a notify state duration */
+	UPROPERTY()
+	FAnimLinkableElement EndLink;
+
 #if WITH_EDITORONLY_DATA
 	/** Color of Notify in editor */
 	UPROPERTY()
@@ -76,7 +81,7 @@ struct FAnimNotifyEvent
 #endif // WITH_EDITORONLY_DATA
 
 	FAnimNotifyEvent()
-		: DisplayTime(0)
+		: DisplayTime_DEPRECATED(0)
 		, TriggerTimeOffset(0)
 		, EndTriggerTimeOffset(0)
 		, TriggerWeightThreshold(ZERO_ANIMWEIGHT_THRESH)
@@ -103,6 +108,10 @@ struct FAnimNotifyEvent
 	/** Returns the actual end time for a state notify. In some cases this may be different from DisplayTime + Duration */
 	ENGINE_API float GetEndTriggerTime() const;
 
+	ENGINE_API float GetDuration() const;
+
+	ENGINE_API void SetDuration(float NewDuration);
+
 	/** Returns true if this is blueprint derived notifies **/
 	bool IsBlueprintNotify() const
 	{
@@ -117,6 +126,8 @@ struct FAnimNotifyEvent
 			(!IsBlueprintNotify() && NotifyName == Other.NotifyName)
 			);
 	}
+
+	ENGINE_API virtual void SetTime(float NewTime, EAnimLinkMethod::Type ReferenceFrame = EAnimLinkMethod::Absolute) override;
 };
 
 /**
diff --git a/Engine/Source/Runtime/Engine/Private/Animation/AnimCompositeBase.cpp b/Engine/Source/Runtime/Engine/Private/Animation/AnimCompositeBase.cpp
--- a/Engine/Source/Runtime/Engine/Private/Animation/AnimCompositeBase.cpp
+++ b/Engine/Source/Runtime/Engine/Private/Animation/AnimCompositeBase.cpp
@@ -318,6 +318,54 @@ int32 FAnimTrack::GetTrackAdditiveType() const
 	return -1;
 }
 
+void FAnimTrack::ValidateSegmentTimes()
+{
+	// rearrange, make sure there are no gaps between and all start times are correctly set
+	AnimSegments[0].StartPos = 0.0f;
+	for(int32 J = 0; J < AnimSegments.Num(); J++)
+	{
+		FAnimSegment& Segment = AnimSegments[J];
+		if(J > 0)
+		{
+			Segment.StartPos = AnimSegments[J - 1].StartPos + AnimSegments[J - 1].GetLength();
+		}
+
+		if(Segment.AnimReference && Segment.AnimEndTime > Segment.AnimReference->SequenceLength)
+		{
+			Segment.AnimEndTime = Segment.AnimReference->SequenceLength;
+		}
+	}
+}
+
+FAnimSegment* FAnimTrack::GetSegmentAtTime(float InTime)
+{
+	FAnimSegment* Result = nullptr;
+	for(FAnimSegment& Segment : AnimSegments)
+	{
+		if(Segment.AnimStartTime <= InTime && InTime <= Segment.StartPos + Segment.GetLength())
+		{
+			Result = &Segment;
+			break;
+		}
+	}
+	return Result;
+}
+
+int32 FAnimTrack::GetSegmentIndexAtTime(float InTime)
+{
+	int32 Result = INDEX_NONE;
+	for(int32 Idx = 0 ; Idx < AnimSegments.Num() ; ++Idx)
+	{
+		FAnimSegment& Segment = AnimSegments[Idx];
+		if(Segment.AnimStartTime <= InTime && InTime <= Segment.StartPos + Segment.GetLength())
+		{
+			Result = Idx;
+			break;
+		}
+	}
+	return Result;
+}
+
 #if WITH_EDITOR
 bool FAnimTrack::GetAllAnimationSequencesReferred(TArray<UAnimSequence*>& AnimationSequences) const
 {
@@ -406,13 +454,7 @@ void FAnimTrack::SortAnimSegments()
 
 		AnimSegments.Sort( FCompareSegments() );
 
-		// rearrange, make sure there are no gaps between and all start times are correctly set
-		AnimSegments[0].StartPos = 0.0f;
-
-		for ( int32 J=1; J < AnimSegments.Num(); J++ )
-		{
-			AnimSegments[J].StartPos = AnimSegments[J-1].StartPos + AnimSegments[J-1].GetLength();
-		}
+		ValidateSegmentTimes();
 	}
 }
 #endif
diff --git a/Engine/Source/Runtime/Engine/Private/Animation/AnimInstance.cpp b/Engine/Source/Runtime/Engine/Private/Animation/AnimInstance.cpp
--- a/Engine/Source/Runtime/Engine/Private/Animation/AnimInstance.cpp
+++ b/Engine/Source/Runtime/Engine/Private/Animation/AnimInstance.cpp
@@ -713,7 +713,7 @@ void UAnimInstance::DisplayDebug(class UCanvas* Canvas, const FDebugDisplayInfo&
 
 			const FAnimNotifyEvent& NotifyState = ActiveAnimNotifyState[NotifyIndex];
 
-			FString NotifyEntry = FString::Printf(TEXT("%i) %s Class: %s Dur:%.3f"), NotifyIndex, *NotifyState.NotifyName.ToString(), *NotifyState.NotifyStateClass->GetName(), NotifyState.Duration);
+			FString NotifyEntry = FString::Printf(TEXT("%i) %s Class: %s Dur:%.3f"), NotifyIndex, *NotifyState.NotifyName.ToString(), *NotifyState.NotifyStateClass->GetName(), NotifyState.GetDuration());
 			Canvas->DrawText(RenderFont, NotifyEntry, Indent, YPos, 1.f, 1.f, RenderInfo);
 			YPos += YL;
 		}
@@ -1164,7 +1164,7 @@ void UAnimInstance::TriggerAnimNotifies(float DeltaSeconds)
 	for (int32 Index = 0; Index < NotifyStateBeginEvent.Num(); Index++)
 	{
 		const FAnimNotifyEvent* AnimNotifyEvent = NotifyStateBeginEvent[Index];
-		AnimNotifyEvent->NotifyStateClass->NotifyBegin(SkelMeshComp, Cast<UAnimSequenceBase>(AnimNotifyEvent->NotifyStateClass->GetOuter()), AnimNotifyEvent->Duration);
+		AnimNotifyEvent->NotifyStateClass->NotifyBegin(SkelMeshComp, Cast<UAnimSequenceBase>(AnimNotifyEvent->NotifyStateClass->GetOuter()), AnimNotifyEvent->GetDuration());
 	}
 
 	// Switch our arrays.
@@ -1645,7 +1645,7 @@ float UAnimInstance::PlaySlotAnimation(UAnimSequenceBase* Asset, FName SlotNodeN
 		
 	FCompositeSection NewSection;
 	NewSection.SectionName = TEXT("Default");
-	NewSection.StartTime = 0.f;
+	NewSection.SetTime(0.0f);
 
 	// add new section
 	NewMontage->CompositeSections.Add(NewSection);
diff --git a/Engine/Source/Runtime/Engine/Private/Animation/AnimMontage.cpp b/Engine/Source/Runtime/Engine/Private/Animation/AnimMontage.cpp
--- a/Engine/Source/Runtime/Engine/Private/Animation/AnimMontage.cpp
+++ b/Engine/Source/Runtime/Engine/Private/Animation/AnimMontage.cpp
@@ -112,7 +112,7 @@ bool UAnimMontage::IsWithinPos(int32 FirstIndex, int32 SecondIndex, float Curren
 	float EndTime;
 	if ( CompositeSections.IsValidIndex(FirstIndex) )
 	{
-		StartTime = CompositeSections[FirstIndex].StartTime;
+		StartTime = CompositeSections[FirstIndex].GetTime();
 	}
 	else // if first index isn't valid, set to be 0.f, so it starts from reset
 	{
@@ -121,7 +121,7 @@ bool UAnimMontage::IsWithinPos(int32 FirstIndex, int32 SecondIndex, float Curren
 
 	if ( CompositeSections.IsValidIndex(SecondIndex) )
 	{
-		EndTime = CompositeSections[SecondIndex].StartTime;
+		EndTime = CompositeSections[SecondIndex].GetTime();
 	}
 	else // if end index isn't valid, set to be BIG_NUMBER
 	{
@@ -135,7 +135,7 @@ bool UAnimMontage::IsWithinPos(int32 FirstIndex, int32 SecondIndex, float Curren
 
 float UAnimMontage::CalculatePos(FCompositeSection &Section, float PosWithinCompositeSection) const
 {
-	float Offset = Section.StartTime;
+	float Offset = Section.GetTime();
 	Offset += PosWithinCompositeSection;
 	// @todo anim
 	return Offset;
@@ -164,7 +164,7 @@ int32 UAnimMontage::GetAnimCompositeSectionIndexFromPos(float CurrentTime, float
 		// if within
 		if (IsWithinPos(I, I+1, CurrentTime))
 		{
-			PosWithinCompositeSection = CurrentTime - CompositeSections[I].StartTime;
+			PosWithinCompositeSection = CurrentTime - CompositeSections[I].GetTime();
 			return I;
 		}
 	}
@@ -179,7 +179,7 @@ float UAnimMontage::GetSectionTimeLeftFromPos(float Position)
 	{
 		if( IsValidSectionIndex(SectionID+1) )
 		{
-			return (GetAnimCompositeSection(SectionID+1).StartTime - Position);
+			return (GetAnimCompositeSection(SectionID+1).GetTime() - Position);
 		}
 		else
 		{
@@ -243,12 +243,12 @@ void UAnimMontage::GetSectionStartAndEndTime(int32 SectionIndex, float& OutStart
 	OutEndTime = SequenceLength;
 	if ( IsValidSectionIndex(SectionIndex) )
 	{
-		OutStartTime = GetAnimCompositeSection(SectionIndex).StartTime;		
+		OutStartTime = GetAnimCompositeSection(SectionIndex).GetTime();		
 	}
 
 	if ( IsValidSectionIndex(SectionIndex + 1))
 	{
-		OutEndTime = GetAnimCompositeSection(SectionIndex + 1).StartTime;		
+		OutEndTime = GetAnimCompositeSection(SectionIndex + 1).GetTime();		
 	}
 }
 
@@ -258,12 +258,12 @@ float UAnimMontage::GetSectionLength(int32 SectionIndex) const
 	float EndTime = SequenceLength;
 	if ( IsValidSectionIndex(SectionIndex) )
 	{
-		StartTime = GetAnimCompositeSection(SectionIndex).StartTime;		
+		StartTime = GetAnimCompositeSection(SectionIndex).GetTime();		
 	}
 
 	if ( IsValidSectionIndex(SectionIndex + 1))
 	{
-		EndTime = GetAnimCompositeSection(SectionIndex + 1).StartTime;		
+		EndTime = GetAnimCompositeSection(SectionIndex + 1).GetTime();		
 	}
 
 	return EndTime - StartTime;
@@ -293,7 +293,7 @@ int32 UAnimMontage::AddAnimCompositeSection(FName InSectionName, float StartTime
 		return INDEX_NONE;
 	}
 
-	NewSection.StartTime = StartTime;
+	NewSection.LinkMontage(this, StartTime);
 
 	// we'd like to sort them in the order of time
 	int32 NewSectionIndex = CompositeSections.Add(NewSection);
@@ -329,7 +329,7 @@ void UAnimMontage::SortAnimCompositeSectionByPos()
 	{
 		FORCEINLINE bool operator()( const FCompositeSection &A, const FCompositeSection &B ) const
 		{
-			return A.StartTime < B.StartTime;
+			return A.GetTime() < B.GetTime();
 		}
 	};
 	CompositeSections.Sort( FCompareFCompositeSection() );
@@ -359,15 +359,40 @@ void UAnimMontage::PostLoad()
 				Segment.AnimEndOffset_DEPRECATED = 0.f;
 			}
 		}
+		Track.ValidateSegmentTimes();
 			}
 
-	for ( auto CompositeIter = CompositeSections.CreateIterator(); CompositeIter; ++CompositeIter )
+	int32 Ver = GetLinker()->UE4Ver();
+
+	for (auto& Composite : CompositeSections)
 	{
-		FCompositeSection & Composite = (*CompositeIter);
-		if (Composite.StarTime_DEPRECATED!=0.f)
+		if (Composite.StarTime_DEPRECATED != 0.0f)
 		{
-			Composite.StartTime = Composite.StarTime_DEPRECATED;
-			Composite.StarTime_DEPRECATED = 0.f;
+			Composite.SetTime(Composite.StarTime_DEPRECATED);
+			Composite.StarTime_DEPRECATED = 0.0f;
+		}
+
+		if(Composite.StartTime_DEPRECATED != 0.0f)
+		{
+			Composite.Clear();
+			Composite.LinkMontage(this, Composite.StartTime_DEPRECATED);
+		}
+		else
+	{
+			Composite.LinkMontage(this, Composite.GetTime());
+		}
+	}
+
+	for(auto& BranchingPoint : BranchingPoints)
+	{
+		if(BranchingPoint.DisplayTime_DEPRECATED != 0.0f)
+		{
+			BranchingPoint.Clear();
+			BranchingPoint.LinkMontage(this, BranchingPoint.DisplayTime_DEPRECATED);
+		}
+		else
+		{
+			BranchingPoint.LinkMontage(this, BranchingPoint.GetTime());
 		}
 	}
 
@@ -412,21 +437,21 @@ void UAnimMontage::PostLoad()
 
 	// verify if skeleton matches, otherwise clear it, this can happen if anim sequence has been modified when this hasn't been loaded. 
 	{
-		USkeleton* MySkeleton = GetSkeleton();
-		for (int32 I=0; I<SlotAnimTracks.Num(); ++I)
+	USkeleton* MySkeleton = GetSkeleton();
+	for (int32 I=0; I<SlotAnimTracks.Num(); ++I)
+	{
+		if ( SlotAnimTracks[I].AnimTrack.AnimSegments.Num() > 0 )
 		{
-			if ( SlotAnimTracks[I].AnimTrack.AnimSegments.Num() > 0 )
+			UAnimSequence * Sequence = Cast<UAnimSequence>(SlotAnimTracks[I].AnimTrack.AnimSegments[0].AnimReference);
+			if ( Sequence && Sequence->GetSkeleton() != MySkeleton )
 			{
-				UAnimSequence * Sequence = Cast<UAnimSequence>(SlotAnimTracks[I].AnimTrack.AnimSegments[0].AnimReference);
-				if ( Sequence && Sequence->GetSkeleton() != MySkeleton )
-				{
-					SlotAnimTracks[I].AnimTrack.AnimSegments[0].AnimReference = 0;
-					MarkPackageDirty();
-					break;
-				}
+				SlotAnimTracks[I].AnimTrack.AnimSegments[0].AnimReference = 0;
+				MarkPackageDirty();
+				break;
 			}
 		}
 	}
+	}
 #endif // WITH_EDITORONLY_DATA
 
 	// Register Slots w/ Skeleton
@@ -441,6 +466,24 @@ void UAnimMontage::PostLoad()
 			}
 		}
 	}
+
+	for(FAnimNotifyEvent& Notify : Notifies)
+	{
+		if(Notify.DisplayTime_DEPRECATED != 0.0f)
+		{
+			Notify.Clear();
+			Notify.LinkMontage(this, Notify.DisplayTime_DEPRECATED);
+		}
+		else
+		{
+			Notify.LinkMontage(this, Notify.GetTime());
+		}
+
+		if(Notify.Duration != 0.0f)
+		{
+			Notify.EndLink.LinkMontage(this, Notify.GetTime() + Notify.Duration);
+		}
+	}
 }
 
 void UAnimMontage::SortAnimBranchingPointByTime()
@@ -537,7 +580,7 @@ EAnimEventTriggerOffsets::Type UAnimMontage::CalculateOffsetFromSections(float T
 {
 	for(auto Iter = CompositeSections.CreateConstIterator(); Iter; ++Iter)
 	{
-		float SectionTime = Iter->StartTime;
+		float SectionTime = Iter->GetTime();
 		if(SectionTime == Time)
 		{
 			return EAnimEventTriggerOffsets::OffsetBefore;
@@ -567,7 +610,7 @@ EAnimEventTriggerOffsets::Type UAnimMontage::CalculateOffsetForNotify(float Noti
 		{
 			for(auto Iter = BranchingPoints.CreateConstIterator(); Iter; ++Iter)
 			{
-				float BranchTime = Iter->DisplayTime;
+				float BranchTime = Iter->GetTime();
 				if(BranchTime == NotifyDisplayTime)
 				{
 					return EAnimEventTriggerOffsets::OffsetBefore;
@@ -1009,7 +1052,7 @@ bool FAnimMontageInstance::SimulateAdvance(float DeltaTime, float& InOutPosition
 			float StoppingPosInSection = CurrentSectionLength;
 
 			// Also look for a branching point. If we have one, stop there first to handle it.
-			const float CurrentSectionEndPos = CurrentSection.StartTime + CurrentSectionLength;
+			const float CurrentSectionEndPos = CurrentSection.GetTime() + CurrentSectionLength;
 
 			// Update position within current section.
 			// Note that we explicitly disallow looping there, we handle it ourselves to simplify code by not having to handle time wrapping around in multiple places.
@@ -1023,7 +1066,7 @@ bool FAnimMontageInstance::SimulateAdvance(float DeltaTime, float& InOutPosition
 			DesiredDeltaMove -= ActualDeltaMove;
 
 			float PrevPosition = Position;
-			InOutPosition = FMath::Clamp<float>(InOutPosition + ActualDeltaPos, CurrentSection.StartTime, CurrentSectionEndPos);
+			InOutPosition = FMath::Clamp<float>(InOutPosition + ActualDeltaPos, CurrentSection.GetTime(), CurrentSectionEndPos);
 
 			if( FMath::Abs(ActualDeltaMove) > 0.f )
 			{
@@ -1110,12 +1153,12 @@ void FAnimMontageInstance::Advance(float DeltaTime, struct FRootMotionMovementPa
 					float StoppingPosInSection = CurrentSectionLength;
 
 					// Also look for a branching point. If we have one, stop there first to handle it.
-					const float CurrentSectionEndPos = CurrentSection.StartTime + CurrentSectionLength;
-					const FBranchingPoint* NextBranchingPoint = Montage->FindFirstBranchingPoint(Position, FMath::Clamp(Position + DesiredDeltaMove, CurrentSection.StartTime, CurrentSectionEndPos));
+					const float CurrentSectionEndPos = CurrentSection.GetTime() + CurrentSectionLength;
+					const FBranchingPoint* NextBranchingPoint = Montage->FindFirstBranchingPoint(Position, FMath::Clamp(Position + DesiredDeltaMove, CurrentSection.GetTime(), CurrentSectionEndPos));
 					if( NextBranchingPoint )
 					{
 						// get the first one to see if it's less than AnimEnd
-						StoppingPosInSection = NextBranchingPoint->GetTriggerTime() - CurrentSection.StartTime;
+						StoppingPosInSection = NextBranchingPoint->GetTriggerTime() - CurrentSection.GetTime();
 					}
 
 					// Update position within current section.
@@ -1130,7 +1173,7 @@ void FAnimMontageInstance::Advance(float DeltaTime, struct FRootMotionMovementPa
 					DesiredDeltaMove -= ActualDeltaMove;
 
 					float PrevPosition = Position;
-					Position = FMath::Clamp<float>(Position + ActualDeltaPos, CurrentSection.StartTime, CurrentSectionEndPos);
+					Position = FMath::Clamp<float>(Position + ActualDeltaPos, CurrentSection.GetTime(), CurrentSectionEndPos);
 
 					if( FMath::Abs(ActualDeltaMove) > 0.f )
 					{
@@ -1300,4 +1343,71 @@ void UAnimMontage::ReplaceReferredAnimations(const TMap<UAnimSequence*, UAnimSeq
 		Track.AnimTrack.ReplaceReferredAnimations(ReplacementMap);
 	}
 }
+
+ENGINE_API void UAnimMontage::UpdateLinkableElements()
+{
+	// Update all linkable elements
+	for(FCompositeSection& Section : CompositeSections)
+	{
+		Section.Update();
+	}
+
+	for(FBranchingPoint& BranchPoint : BranchingPoints)
+	{
+		BranchPoint.Update();
+
+		// Refresh branching point trigger offsets
+		BranchPoint.RefreshTriggerOffset(CalculateOffsetForBranchingPoint(BranchPoint.GetTime()));
+	}
+
+	for(FAnimNotifyEvent& Notify : Notifies)
+	{
+		Notify.Update();
+		Notify.RefreshTriggerOffset(CalculateOffsetForNotify(Notify.GetTime()));
+
+		Notify.EndLink.Update();
+		Notify.RefreshEndTriggerOffset(CalculateOffsetForNotify(Notify.EndLink.GetTime()));
+	}
+}
+
+ENGINE_API void UAnimMontage::UpdateLinkableElements(int32 SlotIdx, int32 SegmentIdx)
+{
+	FAnimSegment* UpdatedSegment = &SlotAnimTracks[SlotIdx].AnimTrack.AnimSegments[SegmentIdx];
+
+	for(FCompositeSection& Section : CompositeSections)
+	{
+		if(Section.GetSlotIndex() == SlotIdx && Section.GetSegmentIndex() == SegmentIdx)
+		{
+			// Update the link
+			Section.Update();
+		}
+	}
+
+	for(FBranchingPoint& BranchPoint : BranchingPoints)
+	{
+		if(BranchPoint.GetSlotIndex() == SlotIdx && BranchPoint.GetSegmentIndex() == SegmentIdx)
+		{
+			BranchPoint.Update();
+
+			// Refresh branching point trigger offsets
+			BranchPoint.RefreshTriggerOffset(CalculateOffsetForBranchingPoint(BranchPoint.GetTime()));
+		}
+	}
+
+	for(FAnimNotifyEvent& Notify : Notifies)
+	{
+		if(Notify.GetSlotIndex() == SlotIdx && Notify.GetSegmentIndex() == SegmentIdx)
+		{
+			Notify.Update();
+			Notify.RefreshTriggerOffset(CalculateOffsetForNotify(Notify.GetTime()));
+		}
+
+		if(Notify.EndLink.GetSlotIndex() == SlotIdx && Notify.EndLink.GetSegmentIndex() == SegmentIdx)
+		{
+			Notify.EndLink.Update();
+			Notify.RefreshEndTriggerOffset(CalculateOffsetForNotify(Notify.EndLink.GetTime()));
+		}
+	}
+}
+
 #endif
\ No newline at end of file
diff --git a/Engine/Source/Runtime/Engine/Private/Animation/AnimSequence.cpp b/Engine/Source/Runtime/Engine/Private/Animation/AnimSequence.cpp
--- a/Engine/Source/Runtime/Engine/Private/Animation/AnimSequence.cpp
+++ b/Engine/Source/Runtime/Engine/Private/Animation/AnimSequence.cpp
@@ -429,6 +429,24 @@ void UAnimSequence::PostLoad()
  			}
  		}
 	}
+
+	for(FAnimNotifyEvent& Notify : Notifies)
+	{
+		if(Notify.DisplayTime_DEPRECATED != 0.0f)
+		{
+			Notify.Clear();
+			Notify.LinkSequence(this, Notify.DisplayTime_DEPRECATED);
+		}
+		else
+		{
+			Notify.LinkSequence(this, Notify.GetTime());
+		}
+
+		if(Notify.Duration != 0.0f)
+		{
+			Notify.EndLink.LinkSequence(this, Notify.GetTime() + Notify.Duration);
+		}
+	}
 }
 
 #if WITH_EDITOR
@@ -1537,7 +1555,7 @@ bool UAnimSequence::CopyNotifies(UAnimSequence* SourceAnimSeq, UAnimSequence* De
 		{
 			// If a notify is found which occurs off the end of the destination sequence, prompt the user to continue.
 			const FAnimNotifyEvent& SrcNotifyEvent = SourceAnimSeq->Notifies[NotifyIndex];
-			if( SrcNotifyEvent.DisplayTime > DestAnimSeq->SequenceLength )
+			if( SrcNotifyEvent.GetTime() > DestAnimSeq->SequenceLength )
 			{
 				const bool bProceed = EAppReturnType::Yes == FMessageDialog::Open( EAppMsgType::YesNo, NSLOCTEXT("UnrealEd", "SomeNotifiesWillNotBeCopiedQ", "Some notifies will not be copied because the destination sequence is not long enough.  Proceed?") );
 				if( !bProceed )
@@ -1574,7 +1592,7 @@ bool UAnimSequence::CopyNotifies(UAnimSequence* SourceAnimSeq, UAnimSequence* De
 		const FAnimNotifyEvent& SrcNotifyEvent = SourceAnimSeq->Notifies[NotifyIndex];
 
 		// Skip notifies which occur at times later than the destination sequence is long.
-		if( SrcNotifyEvent.DisplayTime > DestAnimSeq->SequenceLength )
+		if( SrcNotifyEvent.GetTime() > DestAnimSeq->SequenceLength )
 		{
 			continue;
 		}
@@ -1583,7 +1601,7 @@ bool UAnimSequence::CopyNotifies(UAnimSequence* SourceAnimSeq, UAnimSequence* De
 		// to insert the new notify.
 		int32 NewNotifyIndex = 0;
 		while( NewNotifyIndex < DestAnimSeq->Notifies.Num()
-			&& DestAnimSeq->Notifies[NewNotifyIndex].DisplayTime <= SrcNotifyEvent.DisplayTime )
+			&& DestAnimSeq->Notifies[NewNotifyIndex].GetTime() <= SrcNotifyEvent.GetTime() )
 		{
 			++NewNotifyIndex;
 		}
@@ -1596,10 +1614,10 @@ bool UAnimSequence::CopyNotifies(UAnimSequence* SourceAnimSeq, UAnimSequence* De
 
 		// Copy time and comment.
 		FAnimNotifyEvent& Notify = DestAnimSeq->Notifies[NewNotifyIndex];
-		Notify.DisplayTime = SrcNotifyEvent.DisplayTime;
-		Notify.TriggerTimeOffset = GetTriggerTimeOffsetForType( DestAnimSeq->CalculateOffsetForNotify(Notify.DisplayTime));
+		Notify.SetTime(SrcNotifyEvent.GetTime());
+		Notify.TriggerTimeOffset = GetTriggerTimeOffsetForType( DestAnimSeq->CalculateOffsetForNotify(Notify.GetTime()));
 		Notify.NotifyName = SrcNotifyEvent.NotifyName;
-		Notify.Duration = SrcNotifyEvent.Duration;
+		Notify.SetDuration(SrcNotifyEvent.GetDuration());
 
 		// Copy the notify itself, and point the new one at it.
 		if( SrcNotifyEvent.Notify )
diff --git a/Engine/Source/Runtime/Engine/Private/Animation/AnimSequenceBase.cpp b/Engine/Source/Runtime/Engine/Private/Animation/AnimSequenceBase.cpp
--- a/Engine/Source/Runtime/Engine/Private/Animation/AnimSequenceBase.cpp
+++ b/Engine/Source/Runtime/Engine/Private/Animation/AnimSequenceBase.cpp
@@ -58,12 +58,29 @@ void FAnimNotifyEvent::RefreshEndTriggerOffset( EAnimEventTriggerOffsets::Type P
 
 float FAnimNotifyEvent::GetTriggerTime() const
 {
-	return DisplayTime + TriggerTimeOffset;
+	return GetTime() + TriggerTimeOffset;
 }
 
 float FAnimNotifyEvent::GetEndTriggerTime() const
 {
-	return GetTriggerTime() + Duration + EndTriggerTimeOffset;
+	return GetTriggerTime() + GetDuration() + EndTriggerTimeOffset;
+}
+
+float FAnimNotifyEvent::GetDuration() const
+{
+	return NotifyStateClass ? EndLink.GetTime() - GetTime() : 0.0f;
+}
+
+void FAnimNotifyEvent::SetDuration(float NewDuration)
+{
+	Duration = NewDuration;
+	EndLink.SetTime(GetTime() + Duration);
+}
+
+void FAnimNotifyEvent::SetTime(float NewTime, EAnimLinkMethod::Type ReferenceFrame /*= EAnimLinkMethod::Absolute*/)
+{
+	FAnimLinkableElement::SetTime(NewTime, ReferenceFrame);
+	SetDuration(Duration);
 }
 
 /////////////////////////////////////////////////////
@@ -522,9 +539,9 @@ void UAnimSequenceBase::ClampNotifiesAtEndOfSequence()
 	const float NotifyClampTime = SequenceLength - 0.01f; //Slight offset so that notify is still draggable
 	for(int i = 0; i < Notifies.Num(); ++ i)
 	{
-		if(Notifies[i].DisplayTime >= SequenceLength)
+		if(Notifies[i].GetTime() >= SequenceLength)
 		{
-			Notifies[i].DisplayTime = NotifyClampTime;
+			Notifies[i].SetTime(NotifyClampTime);
 			Notifies[i].TriggerTimeOffset = GetTriggerTimeOffsetForType(EAnimEventTriggerOffsets::OffsetBefore);
 		}
 	}
diff --git a/Engine/Source/Runtime/GameplayAbilities/Private/AbilitySystemComponent_Abilities.cpp b/Engine/Source/Runtime/GameplayAbilities/Private/AbilitySystemComponent_Abilities.cpp
--- a/Engine/Source/Runtime/GameplayAbilities/Private/AbilitySystemComponent_Abilities.cpp
+++ b/Engine/Source/Runtime/GameplayAbilities/Private/AbilitySystemComponent_Abilities.cpp
@@ -1654,12 +1654,12 @@ float UAbilitySystemComponent::GetCurrentMontageSectionLength() const
 			// If we have another section after us, then take delta between both start times.
 			if (CurrentSectionID < (CompositeSections.Num() - 1))
 			{
-				return (CompositeSections[CurrentSectionID + 1].StartTime - CompositeSections[CurrentSectionID].StartTime);
+				return (CompositeSections[CurrentSectionID + 1].GetTime() - CompositeSections[CurrentSectionID].GetTime());
 			}
 			// Otherwise we are the last section, so take delta with Montage total time.
 			else
 			{
-				return (CurrentAnimMontage->SequenceLength - CompositeSections[CurrentSectionID].StartTime);
+				return (CurrentAnimMontage->SequenceLength - CompositeSections[CurrentSectionID].GetTime());
 			}
 		}
 
