Montage improvements - Fixed section bars moving when clicking near them.

diff --git a/Engine/Source/Editor/Persona/Private/STrack.cpp b/Engine/Source/Editor/Persona/Private/STrack.cpp
--- a/Engine/Source/Editor/Persona/Private/STrack.cpp
+++ b/Engine/Source/Editor/Persona/Private/STrack.cpp
@@ -552,6 +552,16 @@ FReply STrack::OnDragOver( const FGeometry& MyGeometry, const FDragDropEvent& Dr
 	return FReply::Unhandled();
 }
 
+FReply STrack::OnDragDetected(const FGeometry& MyGeometry, const FPointerEvent& MouseEvent)
+{
+	if(DraggableBarIndex != INDEX_NONE)
+	{
+		bDraggingBar = true;
+		return FReply::Handled();
+	}
+	return FReply::Unhandled();
+}
+
 float STrack::GetNodeDragDropDataPos( const FGeometry& MyGeometry, const FDragDropEvent& DragDropEvent )
 {
 	float DataPos = 0.f;
@@ -621,6 +631,9 @@ FReply STrack::OnMouseMove( const FGeometry& MyGeometry, const FPointerEvent& Mo
 		FVector2D CursorPos = MyGeometry.AbsoluteToLocal( MouseEvent.GetScreenSpacePosition() );
 		float NewDataPos = FMath::Clamp( LocalToDataX(CursorPos.X, MyGeometry), TrackMinValue.Get(), TrackMaxValue.Get() );
 		OnBarDrag.Execute(DraggableBarIndex, NewDataPos);
+
+		// Update details panel
+		OnBarClicked.ExecuteIfBound(DraggableBarIndex);
 	}
 	else if(DraggableBars.IsBound())
 	{
@@ -640,7 +653,6 @@ FReply STrack::OnMouseButtonDown( const FGeometry& MyGeometry, const FPointerEve
 			if(!bDraggingBar)
 			{
 				OnBarClicked.ExecuteIfBound(DraggableBarIndex);
-				bDraggingBar = true;
 				return FReply::Handled().DetectDrag( SharedThis(this), EKeys::LeftMouseButton );
 			}
 		}
@@ -676,6 +688,7 @@ FReply STrack::OnMouseButtonUp( const FGeometry& MyGeometry, const FPointerEvent
 		if(bDraggingBar)
 		{
 			OnBarDrop.Execute(DraggableBarIndex);
+			OnBarClicked.ExecuteIfBound(DraggableBarIndex);
 		}
 
 		/** Bar dragging */
@@ -821,5 +834,4 @@ void STrack::AddTrackNode( TSharedRef<STrackNode> Node )
 	TrackNodes.Add(Node);
 }
 
-
-#undef LOCTEXT_NAMESPACE
\ No newline at end of file
+#undef LOCTEXT_NAMESPACE
diff --git a/Engine/Source/Editor/Persona/Private/STrack.h b/Engine/Source/Editor/Persona/Private/STrack.h
--- a/Engine/Source/Editor/Persona/Private/STrack.h
+++ b/Engine/Source/Editor/Persona/Private/STrack.h
@@ -242,6 +242,7 @@ public:
 
 	virtual FReply	OnDrop( const FGeometry& MyGeometry, const FDragDropEvent& DragDropEvent ) override;	
 	virtual FReply	OnDragOver( const FGeometry& MyGeometry, const FDragDropEvent& DragDropEvent ) override;
+	virtual FReply	OnDragDetected( const FGeometry& MyGeometry, const FPointerEvent& MouseEvent ) override;
 
 	int32			GetHitNode(const FGeometry& MyGeometry, const FVector2D& CursorPosition);
 
