[UE-3526] - Added pre-compile broadcaster to editor. Persona catches this and checks the compile target against the current preview active list. If there's a clash it performs a cleanup.

diff --git a/Engine/Source/Editor/Persona/Private/Persona.cpp b/Engine/Source/Editor/Persona/Private/Persona.cpp
--- a/Engine/Source/Editor/Persona/Private/Persona.cpp
+++ b/Engine/Source/Editor/Persona/Private/Persona.cpp
@@ -188,10 +188,14 @@ FPersona::FPersona()
 	// Register to be notified when properties are edited
 	OnPropertyChangedHandle = FCoreUObjectDelegates::FOnObjectPropertyChanged::FDelegate::CreateRaw(this, &FPersona::OnPropertyChanged);
 	FCoreUObjectDelegates::OnObjectPropertyChanged.Add(OnPropertyChangedHandle);
+
+	GEditor->OnBlueprintPreCompile().AddRaw(this, &FPersona::OnBlueprintPreCompile);
 }
 
 FPersona::~FPersona()
 {
+	GEditor->OnBlueprintPreCompile().RemoveAll(this);
+
 	FEditorDelegates::OnAssetPostImport.RemoveAll(this);
 	FReimportManager::Instance()->OnPostReimport().RemoveAll(this);
 
@@ -723,6 +727,8 @@ void FPersona::InitPersona(const EToolkitMode::Type Mode, const TSharedPtr< clas
 
 	// Register post import callback to catch animation imports when we have the asset open (we need to reinit)
 	FEditorDelegates::OnAssetPostImport.AddRaw(this, &FPersona::OnPostImport);
+
+	
 }
 	
 TSharedRef< SWidget > FPersona::GenerateCreateAssetMenu( USkeleton* Skeleton ) const
@@ -2867,6 +2873,28 @@ TStatId FPersona::GetStatId() const
 	RETURN_QUICK_DECLARE_CYCLE_STAT(FPersona, STATGROUP_Tickables);
 }
 
+void FPersona::OnBlueprintPreCompile(UBlueprint* BlueprintToCompile)
+{
+	if(PreviewComponent && PreviewComponent->PreviewInstance)
+	{
+		// If we are compiling an anim notify state the class will soon be sanitized and 
+		// if an anim instance is running a state when that happens it will likely
+		// crash, so we end any states that are about to compile.
+		UAnimPreviewInstance* Instance = PreviewComponent->PreviewInstance;
+		USkeletalMeshComponent* SkelMeshComp = Instance->GetSkelMeshComponent();
+
+		for(int32 Idx = Instance->ActiveAnimNotifyState.Num() - 1 ; Idx >= 0 ; --Idx)
+		{
+			FAnimNotifyEvent& Event = Instance->ActiveAnimNotifyState[Idx];
+			if(Event.NotifyStateClass->GetClass() == BlueprintToCompile->GeneratedClass)
+			{
+				Event.NotifyStateClass->NotifyEnd(SkelMeshComp, Cast<UAnimSequenceBase>(Event.NotifyStateClass->GetOuter()));
+				Instance->ActiveAnimNotifyState.RemoveAt(Idx);
+			}
+		}
+	}
+}
+
 static class FMeshHierarchyCmd : private FSelfRegisteringExec
 {
 public:
diff --git a/Engine/Source/Editor/Persona/Private/Persona.h b/Engine/Source/Editor/Persona/Private/Persona.h
--- a/Engine/Source/Editor/Persona/Private/Persona.h
+++ b/Engine/Source/Editor/Persona/Private/Persona.h
@@ -600,6 +600,9 @@ private:
 	/** Returns the editor objects that are applicable for our current mode (e.g mesh, animation etc) */
 	TArray<UObject*> GetEditorObjectsForMode(FName Mode) const;
 
+	/** Called immediately prior to a blueprint compilation */
+	void OnBlueprintPreCompile(UBlueprint* BlueprintToCompile);
+
 	/** The extender to pass to the level editor to extend it's window menu */
 	TSharedPtr<FExtender> MenuExtender;
 
diff --git a/Engine/Source/Editor/UnrealEd/Classes/Editor/EditorEngine.h b/Engine/Source/Editor/UnrealEd/Classes/Editor/EditorEngine.h
--- a/Engine/Source/Editor/UnrealEd/Classes/Editor/EditorEngine.h
+++ b/Engine/Source/Editor/UnrealEd/Classes/Editor/EditorEngine.h
@@ -533,6 +533,13 @@ public:
 	/** Annotation to track which PIE/SIE (PlayWorld) UObjects have counterparts in the EditorWorld **/
 	class FUObjectAnnotationSparseBool ObjectsThatExistInEditorWorld;
 
+	/** Called prior to a Blueprint compile */
+	DECLARE_EVENT_OneParam( UEditorEngine, FBlueprintPreCompileEvent, UBlueprint* );
+	FBlueprintPreCompileEvent& OnBlueprintPreCompile() { return BlueprintPreCompileEvent; }
+
+	/** Broadcasts that a Blueprint is about to be compiled */
+	void BroadcastBlueprintPreCompile(UBlueprint* BlueprintToCompile) { BlueprintPreCompileEvent.Broadcast(BlueprintToCompile); }
+
 	/** Called when a Blueprint compile is completed. */
 	DECLARE_EVENT( UEditorEngine, FBlueprintCompiledEvent );
 	FBlueprintCompiledEvent& OnBlueprintCompiled() { return BlueprintCompiledEvent; }
@@ -2575,6 +2582,9 @@ private:
 
 private:
 
+	/** Delegate broadcast just before a blueprint is compiled */
+	FBlueprintPreCompileEvent BlueprintPreCompileEvent;
+
 	/** Delegate broadcast when blueprint is compiled */
 	FBlueprintCompiledEvent BlueprintCompiledEvent;
 
diff --git a/Engine/Source/Editor/UnrealEd/Private/Kismet2/Kismet2.cpp b/Engine/Source/Editor/UnrealEd/Private/Kismet2/Kismet2.cpp
--- a/Engine/Source/Editor/UnrealEd/Private/Kismet2/Kismet2.cpp
+++ b/Engine/Source/Editor/UnrealEd/Private/Kismet2/Kismet2.cpp
@@ -338,6 +338,14 @@ UBlueprint* FKismetEditorUtilities::CreateBlueprint(UClass* ParentClass, UObject
 
 void FKismetEditorUtilities::CompileBlueprint(UBlueprint* BlueprintObj, bool bIsRegeneratingOnLoad, bool bSkipGarbageCollection, bool bSaveIntermediateProducts, FCompilerResultsLog* pResults)
 {
+	// Broadcast pre-compile
+#if WITH_EDITOR
+	if(GEditor && GIsEditor)
+	{
+		GEditor->BroadcastBlueprintPreCompile(BlueprintObj);
+	}
+#endif
+
 	// Reset the flag, so if the user tries to use PIE it will warn them if the BP did not compile
 	BlueprintObj->bDisplayCompilePIEWarning = true;
 
