[UE-4046] - Added explicit invoke to override editor tab if the target document is a derived anim blueprint

diff --git a/Engine/Source/Editor/Persona/Private/BlueprintMode/AnimBlueprintMode.cpp b/Engine/Source/Editor/Persona/Private/BlueprintMode/AnimBlueprintMode.cpp
--- a/Engine/Source/Editor/Persona/Private/BlueprintMode/AnimBlueprintMode.cpp
+++ b/Engine/Source/Editor/Persona/Private/BlueprintMode/AnimBlueprintMode.cpp
@@ -150,6 +150,12 @@ void FAnimBlueprintEditAppMode::PostActivateMode()
 		{
 			AnimBlueprint->SetObjectBeingDebugged(PreviewComponent->AnimScriptInstance);
 		}
+
+		// If we are a derived anim blueprint always show the overrides tab
+		if(UAnimBlueprint::FindRootAnimBlueprint(AnimBlueprint))
+		{
+			Persona->GetTabManager()->InvokeTab(FPersonaTabs::AnimBlueprintParentPlayerEditorID);
+		}
 	}
 
 	FBlueprintEditorApplicationMode::PostActivateMode();
