Fix for crashes when deleting segments or slots with linked montage elements.
[OR-89], [OR-90]

diff --git a/Engine/Source/Editor/Persona/Private/SAnimMontagePanel.cpp b/Engine/Source/Editor/Persona/Private/SAnimMontagePanel.cpp
--- a/Engine/Source/Editor/Persona/Private/SAnimMontagePanel.cpp
+++ b/Engine/Source/Editor/Persona/Private/SAnimMontagePanel.cpp
@@ -520,6 +520,7 @@ void SAnimMontagePanel::Update()
 					.OnAnimSegmentNodeClicked( this, &SAnimMontagePanel::ShowSegmentInDetailsView, SlotAnimIdx )
 					.OnPreAnimUpdate( Editor, &SMontageEditor::PreAnimUpdate )
 					.OnPostAnimUpdate( Editor, &SMontageEditor::PostAnimUpdate )
+					.OnAnimSegmentRemoved(this, &SAnimMontagePanel::OnAnimSegmentRemoved, SlotAnimIdx )
 					.OnBarDrag(Editor, &SMontageEditor::OnEditSectionTime)
 					.OnBarDrop(Editor, &SMontageEditor::OnEditSectionTimeFinish)
 					.OnBarClicked(SharedThis(this), &SAnimMontagePanel::ShowSectionInDetailsView)
@@ -1053,4 +1054,35 @@ void SAnimMontagePanel::CollectLinkableElements(TArray<FAnimLinkableElement*> &E
 	}
 }
 
+void SAnimMontagePanel::OnAnimSegmentRemoved(int32 SegmentIndex, int32 SlotIndex)
+{
+	TArray<FAnimLinkableElement*> LinkableElements;
+	CollectLinkableElements(LinkableElements);
+
+	for(FAnimNotifyEvent& Notify : Montage->Notifies)
+	{
+		if(Notify.NotifyStateClass)
+		{
+			LinkableElements.Add(&Notify.EndLink);
+		}
+	}
+
+	// Go through the linkable elements and fix the indices.
+	// BG TODO: Once we can identify moved segments, remove
+	for(FAnimLinkableElement* Element : LinkableElements)
+	{
+		if(Element->GetSlotIndex() == SlotIndex)
+		{
+			if(Element->GetSegmentIndex() == SegmentIndex)
+			{
+				Element->Clear();
+			}
+			else if(Element->GetSegmentIndex() > SegmentIndex)
+			{
+				Element->SetSegmentIndex(Element->GetSegmentIndex() - 1);
+			}
+		}
+	}
+}
+
 #undef LOCTEXT_NAMESPACE
\ No newline at end of file
diff --git a/Engine/Source/Editor/Persona/Private/SAnimMontagePanel.h b/Engine/Source/Editor/Persona/Private/SAnimMontagePanel.h
--- a/Engine/Source/Editor/Persona/Private/SAnimMontagePanel.h
+++ b/Engine/Source/Editor/Persona/Private/SAnimMontagePanel.h
@@ -137,4 +137,9 @@ private:
 	void CreateNewSection(const FText& NewSectionName, ETextCommit::Type CommitInfo, float StartTime);
 	void CreateNewBranch(const FText& NewBranchName, ETextCommit::Type CommitInfo, float StartTime);
 	void RenameBranchPoint(const FText &NewName, ETextCommit::Type CommitInfo, int32 BranchPointIndex);
+
+	/** Called when a segment is removed from a track, so we can adjust the indices in linkable elements
+	 *	@param SegmentIndex - Index of the removed segment
+	 */
+	void OnAnimSegmentRemoved(int32 SegmentIndex, int32 SlotIndex);
 };
\ No newline at end of file
diff --git a/Engine/Source/Editor/Persona/Private/SAnimSegmentsPanel.cpp b/Engine/Source/Editor/Persona/Private/SAnimSegmentsPanel.cpp
--- a/Engine/Source/Editor/Persona/Private/SAnimSegmentsPanel.cpp
+++ b/Engine/Source/Editor/Persona/Private/SAnimSegmentsPanel.cpp
@@ -24,6 +24,7 @@ void SAnimSegmentsPanel::Construct(const FArguments& InArgs)
 	OnAnimSegmentNodeClickedDelegate = InArgs._OnAnimSegmentNodeClicked;
 	OnPreAnimUpdateDelegate			 = InArgs._OnPreAnimUpdate;
 	OnPostAnimUpdateDelegate		 = InArgs._OnPostAnimUpdate;
+	OnAnimSegmentRemovedDelegate	 = InArgs._OnAnimSegmentRemoved;
 
 	// Animation Segment tracks
 	TArray<TSharedPtr<STrack>> AnimSTracks;
@@ -214,6 +215,7 @@ void SAnimSegmentsPanel::RemoveAnimSegment(int32 AnimSegmentIndex)
 
 		AnimTrack->AnimSegments.RemoveAt(AnimSegmentIndex);
 
+		OnAnimSegmentRemovedDelegate.ExecuteIfBound(AnimSegmentIndex);
 		OnPostAnimUpdateDelegate.Execute();
 	}
 }
diff --git a/Engine/Source/Editor/Persona/Private/SAnimSegmentsPanel.h b/Engine/Source/Editor/Persona/Private/SAnimSegmentsPanel.h
--- a/Engine/Source/Editor/Persona/Private/SAnimSegmentsPanel.h
+++ b/Engine/Source/Editor/Persona/Private/SAnimSegmentsPanel.h
@@ -10,6 +10,7 @@
 DECLARE_DELEGATE( FOnPreAnimUpdate )
 DECLARE_DELEGATE( FOnPostAnimUpdate )
 DECLARE_DELEGATE_OneParam( FOnAnimSegmentNodeClicked, int32 )
+DECLARE_DELEGATE_OneParam( FOnAnimSegmentRemoved, int32 )
 
 //////////////////////////////////////////////////////////////////////////
 // SAnimSegmentsPanel
@@ -41,6 +42,7 @@ public:
 	SLATE_EVENT( FOnAnimSegmentNodeClicked, OnAnimSegmentNodeClicked )
 	SLATE_EVENT( FOnPreAnimUpdate,			OnPreAnimUpdate )
 	SLATE_EVENT( FOnPostAnimUpdate,			OnPostAnimUpdate )
+	SLATE_EVENT( FOnAnimSegmentRemoved,		OnAnimSegmentRemoved )
 
 	SLATE_EVENT( FOnBarDrag,				OnBarDrag)
 	SLATE_EVENT( FOnBarDrop,				OnBarDrop)
@@ -57,6 +59,7 @@ private:
 	FOnPreAnimUpdate			OnPreAnimUpdateDelegate;
 	FOnPostAnimUpdate			OnPostAnimUpdateDelegate;
 	FOnAnimSegmentNodeClicked	OnAnimSegmentNodeClickedDelegate;
+	FOnAnimSegmentRemoved		OnAnimSegmentRemovedDelegate;
 
 	enum ETrackViewStyle
 	{
diff --git a/Engine/Source/Runtime/Engine/Classes/Animation/AnimLinkableElement.cpp b/Engine/Source/Runtime/Engine/Classes/Animation/AnimLinkableElement.cpp
--- a/Engine/Source/Runtime/Engine/Classes/Animation/AnimLinkableElement.cpp
+++ b/Engine/Source/Runtime/Engine/Classes/Animation/AnimLinkableElement.cpp
@@ -56,23 +56,22 @@ void FAnimLinkableElement::LinkSequence(UAnimSequenceBase* Sequence, float AbsSe
 
 void FAnimLinkableElement::Clear()
 {
+	ChangeLinkMethod(EAnimLinkMethod::Absolute);
 	LinkedSequence = nullptr;
 	SegmentBeginTime = -1.0f;
 	SegmentLength = -1.0f;
-	LinkMethod = EAnimLinkMethod::Absolute;
-	CachedLinkMethod = LinkMethod;
+	SegmentIndex = INDEX_NONE;
 }
 
 void FAnimLinkableElement::Update()
 {
-	if(LinkedMontage)
+	if(LinkedMontage && LinkedMontage->SlotAnimTracks.IsValidIndex(SlotIndex))
 	{
-		check(LinkedMontage->SlotAnimTracks.IsValidIndex(SlotIndex));
 		FSlotAnimationTrack& Slot = LinkedMontage->SlotAnimTracks[SlotIndex];
 		float CurrentTime = GetTime();
 
 		// If we don't have a segment, check to see if one has been added.
-		if(SegmentIndex == INDEX_NONE)
+		if(SegmentIndex == INDEX_NONE || !Slot.AnimTrack.AnimSegments.IsValidIndex(SegmentIndex))
 		{
 			SegmentIndex = Slot.AnimTrack.GetSegmentIndexAtTime(CurrentTime);
 		}
@@ -148,7 +147,7 @@ void FAnimLinkableElement::OnChanged(float NewMontageTime)
 
 		SetTime(NewMontageTime);
 	}
-	else if(!LinkedSequence)
+	else if(!(LinkedSequence && !LinkedMontage))
 	{
 		// We have no segment to link to, we need to clear our the segment data
 		// and give ourselves an absolute time
diff --git a/Engine/Source/Runtime/Engine/Classes/Animation/AnimLinkableElement.h b/Engine/Source/Runtime/Engine/Classes/Animation/AnimLinkableElement.h
--- a/Engine/Source/Runtime/Engine/Classes/Animation/AnimLinkableElement.h
+++ b/Engine/Source/Runtime/Engine/Classes/Animation/AnimLinkableElement.h
@@ -100,6 +100,11 @@ struct FAnimLinkableElement
 	/** Get the index of the segment this element is currently linked to */
 	ENGINE_API int32 GetSegmentIndex() const {return SegmentIndex;}
 
+	/** Directly sets segment index
+	 *	@param NewSegmentIndex New segment index
+	 */
+	ENGINE_API void SetSegmentIndex(int32 NewSegmentIndex) {SegmentIndex = NewSegmentIndex;}
+
 protected:
 
 	/** Gets the segment in the current montage in the current slot that is at the time of this element */
