Fixed a bug which cloth asset values reset to 0 when reopening a project after a save

diff --git a/Engine/Source/Editor/Persona/Private/PersonaMeshDetails.cpp b/Engine/Source/Editor/Persona/Private/PersonaMeshDetails.cpp
--- a/Engine/Source/Editor/Persona/Private/PersonaMeshDetails.cpp
+++ b/Engine/Source/Editor/Persona/Private/PersonaMeshDetails.cpp
@@ -1383,11 +1383,7 @@ void FPersonaMeshDetails::OnGenerateElementForClothingAsset( TSharedRef<IPropert
 	USkeletalMesh* SkelMesh = PersonaPtr->GetMesh();
 	FClothingAssetData& AssetData = SkelMesh->ClothingAssets[ElementIndex];
 
-	// if properties are not changed by a user through UE4 editor, loads original properties from the asset
-	if (!AssetData.bClothPropertiesChanged)
-	{
-		ApexClothingUtils::GetPhysicsPropertiesFromApexAsset(AssetData.ApexClothingAsset->GetAsset(), AssetData.PhysicsProperties);
-	}
+	ApexClothingUtils::GetPhysicsPropertiesFromApexAsset(AssetData.ApexClothingAsset->GetAsset(), AssetData.PhysicsProperties);
 
 	// cloth physics properties
 	TSharedRef<IPropertyHandle> ClothPhysicsProperties = StructProperty->GetChildHandle(FName("PhysicsProperties")).ToSharedRef();
