Fixed Persona crash when invoking the preview tooltip without a preview mesh. Viewport hidden and message about the preview mesh displayed

diff --git a/Engine/Source/Editor/Persona/Private/SAnimationSequenceBrowser.cpp b/Engine/Source/Editor/Persona/Private/SAnimationSequenceBrowser.cpp
--- a/Engine/Source/Editor/Persona/Private/SAnimationSequenceBrowser.cpp
+++ b/Engine/Source/Editor/Persona/Private/SAnimationSequenceBrowser.cpp
@@ -803,7 +803,21 @@ TSharedRef<SToolTip> SAnimationSequenceBrowser::CreateCustomAssetToolTip(FAssetD
 					.Padding(6)
 					.BorderImage(FEditorStyle::GetBrush("ContentBrowser.TileViewTooltip.ContentBorder"))
 					[
-						ViewportWidget.ToSharedRef()
+						SNew(SOverlay)
+						+SOverlay::Slot()
+						.VAlign(VAlign_Center)
+						.HAlign(HAlign_Center)
+						[
+							SNew(STextBlock)
+							.Text(FString(TEXT("No Preview Mesh")))
+						]
+
+						+ SOverlay::Slot()
+						.VAlign(VAlign_Center)
+						.HAlign(HAlign_Center)
+						[
+							ViewportWidget.ToSharedRef()
+						]
 					]
 				]
 			]
@@ -881,20 +895,29 @@ bool SAnimationSequenceBrowser::OnVisualizeAssetToolTip(const TSharedPtr<SWidget
 		USkeleton* Skeleton = Asset->GetSkeleton();
 		
 		MeshToUse = Skeleton->GetAssetPreviewMesh(Asset);
-		check(MeshToUse);
-		if(PreviewComponent->SkeletalMesh != MeshToUse)
+		
+		if(MeshToUse)
 		{
-			PreviewComponent->SetSkeletalMesh(MeshToUse);
-		}
+			if(PreviewComponent->SkeletalMesh != MeshToUse)
+			{
+				PreviewComponent->SetSkeletalMesh(MeshToUse);
+			}
+
+			PreviewComponent->EnablePreview(true, Asset, NULL);
+			PreviewComponent->PreviewInstance->PlayAnim(true);
 
-		PreviewComponent->EnablePreview(true, Asset, NULL);
-		PreviewComponent->PreviewInstance->PlayAnim(true);
+			float HalfFov = FMath::DegreesToRadians(ViewportClient->ViewFOV) / 2.0f;
+			float TargetDist = MeshToUse->Bounds.SphereRadius / FMath::Tan(HalfFov);
 
-		float HalfFov = FMath::DegreesToRadians(ViewportClient->ViewFOV) / 2.0f;
-		float TargetDist = MeshToUse->Bounds.SphereRadius / FMath::Tan(HalfFov);
+			ViewportClient->SetViewRotation(FRotator(0.0f, -45.0f, 0.0f));
+			ViewportClient->SetViewLocationForOrbiting(FVector(0.0f, 0.0f, MeshToUse->Bounds.BoxExtent.Z / 2.0f), TargetDist);
 
-		ViewportClient->SetViewRotation(FRotator(0.0f, -45.0f, 0.0f));
-		ViewportClient->SetViewLocationForOrbiting(FVector(0.0f, 0.0f, MeshToUse->Bounds.BoxExtent.Z / 2.0f), TargetDist);
+			ViewportWidget->SetVisibility(EVisibility::Visible);
+		}
+		else
+		{
+			ViewportWidget->SetVisibility(EVisibility::Hidden);
+		}
 	}
 
 	// We return false here as we aren't visualizing the tooltip - just detecting when it is about to be shown.
