Fixed crash with copy/paste notifies

diff --git a/Engine/Source/Editor/Persona/Private/SAnimNotifyPanel.cpp b/Engine/Source/Editor/Persona/Private/SAnimNotifyPanel.cpp
--- a/Engine/Source/Editor/Persona/Private/SAnimNotifyPanel.cpp
+++ b/Engine/Source/Editor/Persona/Private/SAnimNotifyPanel.cpp
@@ -2729,7 +2729,7 @@ void SAnimNotifyTrack::AppendSelectionToArray(TArray<FAnimNotifyEvent*>& Selecti
 
 void SAnimNotifyTrack::PasteSingleNotify(FString& NotifyString, float PasteTime)
 {
-	int32 NewIdx = Sequence->Notifies.AddZeroed();
+	int32 NewIdx = Sequence->Notifies.Add(FAnimNotifyEvent());
 	UArrayProperty* ArrayProperty = NULL;
 	uint8* PropertyData = Sequence->FindNotifyPropertyData(NewIdx, ArrayProperty);
 
