Fix for time being incorrect when pasting a linkable montage element in Persona.

diff --git a/Engine/Source/Editor/Persona/Private/SAnimNotifyPanel.cpp b/Engine/Source/Editor/Persona/Private/SAnimNotifyPanel.cpp
--- a/Engine/Source/Editor/Persona/Private/SAnimNotifyPanel.cpp
+++ b/Engine/Source/Editor/Persona/Private/SAnimNotifyPanel.cpp
@@ -2765,6 +2765,8 @@ void SAnimNotifyTrack::PasteSingleNotify(FString& NotifyString, float PasteTime)
 			NewNotify.SetDuration(FMath::Clamp(NewNotify.GetDuration(), 1 / 30.0f, Sequence->SequenceLength - NewNotify.GetTime()));
 			NewNotify.EndTriggerTimeOffset = GetTriggerTimeOffsetForType(Sequence->CalculateOffsetForNotify(NewNotify.GetTime() + NewNotify.GetDuration()));
 		}
+
+		NewNotify.ConditionalRelink();
 	}
 	else
 	{
@@ -3403,6 +3405,7 @@ void SAnimNotifyPanel::CopySelectedNotifiesToClipboard() const
 			StrValue += "\n";
 			UArrayProperty* ArrayProperty = NULL;
 			uint8* PropertyData = Sequence->FindNotifyPropertyData(Index, ArrayProperty);
+			StrValue += FString::Printf(TEXT("AbsTime=%f,"), Event->GetTime());
 			if(PropertyData && ArrayProperty)
 			{
 				ArrayProperty->Inner->ExportTextItem(StrValue, PropertyData, PropertyData, Sequence, PPF_Copy);
@@ -3465,8 +3468,11 @@ void SAnimNotifyPanel::OnPasteNotifies(SAnimNotifyTrack* RequestTrack, float Cli
 			int32 OriginalTrack;
 			float OrigTime;
 			float PasteTime = -1.0f;
-			if(FParse::Value(*CurrentLine, TEXT("TrackIndex="), OriginalTrack) && FParse::Value(*CurrentLine, TEXT("DisplayTime="), OrigTime))
+			if(FParse::Value(*CurrentLine, TEXT("TrackIndex="), OriginalTrack) && FParse::Value(*CurrentLine, TEXT("AbsTime="), OrigTime))
 			{
+				FString NotifyExportString;
+				CurrentLine.Split(TEXT(","), NULL, &NotifyExportString);
+
 				// Store the first track so we know where to place notifies
 				if(FirstTrack < 0)
 				{
@@ -3480,7 +3486,7 @@ void SAnimNotifyPanel::OnPasteNotifies(SAnimNotifyTrack* RequestTrack, float Cli
 				// Have to invert the index here as tracks are stored in reverse
 				TSharedPtr<SAnimNotifyTrack> TrackToUse = NotifyAnimTracks[NotifyAnimTracks.Num() - 1 - (PasteIdx + TrackOffset)];
 
-				TrackToUse->PasteSingleNotify(CurrentLine, TimeToPaste);
+				TrackToUse->PasteSingleNotify(NotifyExportString, TimeToPaste);
 			}
 		}
 	}
diff --git a/Engine/Source/Runtime/Engine/Classes/Animation/AnimLinkableElement.h b/Engine/Source/Runtime/Engine/Classes/Animation/AnimLinkableElement.h
--- a/Engine/Source/Runtime/Engine/Classes/Animation/AnimLinkableElement.h
+++ b/Engine/Source/Runtime/Engine/Classes/Animation/AnimLinkableElement.h
@@ -105,6 +105,9 @@ struct FAnimLinkableElement
 	 */
 	ENGINE_API void SetSegmentIndex(int32 NewSegmentIndex) {SegmentIndex = NewSegmentIndex;}
 
+	/** Relinks this element if it is outside of its segments */
+	ENGINE_API void ConditionalRelink();
+
 protected:
 
 	/** Gets the segment in the current montage in the current slot that is at the time of this element */
@@ -151,9 +154,6 @@ protected:
 
 private:
 
-	/** Relinks this element if it is outside of its segments */
-	void ConditionalRelink();
-
 	/** Gets a time from an absolute LinkValue
 	 * @param NewLinkMethod The type of time to return
 	 */
