Add context menu option for opening source blueprints for notifies

diff --git a/Engine/Source/Editor/Persona/Private/SAnimNotifyPanel.cpp b/Engine/Source/Editor/Persona/Private/SAnimNotifyPanel.cpp
--- a/Engine/Source/Editor/Persona/Private/SAnimNotifyPanel.cpp
+++ b/Engine/Source/Editor/Persona/Private/SAnimNotifyPanel.cpp
@@ -12,6 +12,7 @@
 #include "AssetSelection.h"
 #include "STextEntryPopup.h"
 #include "SExpandableArea.h"
+#include "Toolkits/AssetEditorManager.h"
 
 // Track Panel drawing
 const float NotificationTrackHeight = 20.0f;
@@ -335,6 +336,9 @@ protected:
 	/** Provides direct access to the notify menu from the context menu */
 	void OnManageNotifies();
 
+	/** Opens the supplied blueprint in an editor */
+	void OnOpenNotifySource(UBlueprint* InSourceBlueprint) const;
+
 	/**
 	 * Update the nodes to match the data that the panel is observing
 	 */
@@ -2179,6 +2183,27 @@ TSharedPtr<SWidget> SAnimNotifyTrack::SummonContextMenu(const FGeometry& MyGeome
 	}
 	MenuBuilder.EndSection(); //AnimEdit
 
+	if (NotifyIndex != INDEX_NONE)
+	{
+		FAnimNotifyEvent* Notify = AnimNotifies[NotifyIndex];
+		UObject* NotifyObject = Notify->Notify;
+		NotifyObject = NotifyObject ? NotifyObject : Notify->NotifyStateClass;
+
+		if (NotifyObject && Cast<UBlueprintGeneratedClass>(NotifyObject->GetClass()))
+		{
+			if (UBlueprint* Blueprint = Cast<UBlueprint>(NotifyObject->GetClass()->ClassGeneratedBy))
+			{
+				MenuBuilder.BeginSection("ViewSource", LOCTEXT("NotifyViewHeading", "View"));
+
+				NewAction.ExecuteAction.BindRaw(
+					this, &SAnimNotifyTrack::OnOpenNotifySource, Blueprint);
+				MenuBuilder.AddMenuEntry(LOCTEXT("OpenNotifyBlueprint", "Open Notify Blueprint"), LOCTEXT("OpenNotifyBlueprint", "Opens the source blueprint for this notify"), FSlateIcon(), NewAction);
+
+				MenuBuilder.EndSection(); //ViewSource
+			}
+		}
+	}
+
 	// Display the newly built menu
 	TWeakPtr<SWindow> ContextMenuWindow =
 		FSlateApplication::Get().PushMenu( SharedThis( this ), MenuBuilder.MakeWidget(), CursorPos, FPopupTransitionEffect( FPopupTransitionEffect::ContextMenu ));
@@ -2373,6 +2398,11 @@ void SAnimNotifyTrack::OnManageNotifies()
 	}
 }
 
+void SAnimNotifyTrack::OnOpenNotifySource(UBlueprint* InSourceBlueprint) const
+{
+	FAssetEditorManager::Get().OpenEditorForAsset(InSourceBlueprint);
+}
+
 void SAnimNotifyTrack::SetTriggerWeight(const FText& TriggerWeight, ETextCommit::Type CommitInfo, int32 NotifyIndex)
 {
 	if ( CommitInfo == ETextCommit::OnEnter || CommitInfo == ETextCommit::OnUserMovedFocus )
