Fixed for element movement issue in anim sequences after montage improvement changes.

diff --git a/Engine/Source/Editor/Persona/Private/SAnimNotifyPanel.cpp b/Engine/Source/Editor/Persona/Private/SAnimNotifyPanel.cpp
--- a/Engine/Source/Editor/Persona/Private/SAnimNotifyPanel.cpp
+++ b/Engine/Source/Editor/Persona/Private/SAnimNotifyPanel.cpp
@@ -1858,16 +1858,7 @@ void SAnimNotifyTrack::CreateNewNotifyAtCursor(FString NewNotifyName, UClass* No
 	FAnimNotifyEvent& NewEvent = Sequence->Notifies[NewNotifyIndex];
 	NewEvent.NotifyName = FName(*NewNotifyName);
 
-	UAnimMontage* SequenceAsMontage = Cast<UAnimMontage>(Sequence);
-	if(SequenceAsMontage)
-	{
-		NewEvent.LinkMontage(SequenceAsMontage, NewTime);
-	}
-	else
-	{
-		NewEvent.LinkSequence(Sequence, NewTime);
-	}
-
+	NewEvent.Link(Sequence, NewTime);
 	NewEvent.TriggerTimeOffset = GetTriggerTimeOffsetForType(Sequence->CalculateOffsetForNotify(NewTime));
 	NewEvent.TrackIndex = TrackIndex;
 
@@ -1882,14 +1873,7 @@ void SAnimNotifyTrack::CreateNewNotifyAtCursor(FString NewNotifyName, UClass* No
 		{
 			NewEvent.NotifyName = FName(*NewEvent.NotifyStateClass->GetNotifyName());
 			NewEvent.SetDuration(1 / 30.f);
-			if(SequenceAsMontage)
-			{
-				NewEvent.EndLink.LinkMontage(SequenceAsMontage, NewEvent.EndLink.GetTime());
-			}
-			else
-			{
-				NewEvent.EndLink.LinkSequence(Sequence, NewEvent.EndLink.GetTime());
-			}
+			NewEvent.EndLink.Link(Sequence, NewEvent.EndLink.GetTime());
 		}
 		else
 		{
@@ -2680,22 +2664,20 @@ void SAnimNotifyTrack::HandleNodeDrop(TSharedPtr<SAnimNotifyNode> Node, float Of
 	FAnimNotifyEvent* DroppedEvent = Node->NotifyEvent;
 	float EventDuration = DroppedEvent->GetDuration();
 
-	UAnimMontage* MontageObj = Cast<UAnimMontage>(Sequence);
-
 	if(Node->GetLastSnappedTime() != -1.0f)
 	{
-		DroppedEvent->LinkMontage(MontageObj, Node->GetLastSnappedTime(), DroppedEvent->GetSlotIndex());
+		DroppedEvent->Link(Sequence, Node->GetLastSnappedTime(), DroppedEvent->GetSlotIndex());
 	}
 	else
 	{
-		DroppedEvent->LinkMontage(MontageObj, Time, DroppedEvent->GetSlotIndex());
+		DroppedEvent->Link(Sequence, Time, DroppedEvent->GetSlotIndex());
 	}
 	DroppedEvent->RefreshTriggerOffset(Sequence->CalculateOffsetForNotify(DroppedEvent->GetTime()));
 
 	if(EventDuration > 0.0f)
 	{
-		DroppedEvent->EndLink.LinkMontage(MontageObj, DroppedEvent->GetTime() + EventDuration, DroppedEvent->GetSlotIndex());
-		DroppedEvent->RefreshEndTriggerOffset(Sequence->CalculateOffsetForNotify(DroppedEvent->EndLink.GetTime()/*DroppedEvent->GetTime()+ DroppedEvent->GetDuration()*/));
+		DroppedEvent->EndLink.Link(Sequence, DroppedEvent->GetTime() + EventDuration, DroppedEvent->GetSlotIndex());
+		DroppedEvent->RefreshEndTriggerOffset(Sequence->CalculateOffsetForNotify(DroppedEvent->EndLink.GetTime()));
 	}
 	else
 	{
diff --git a/Engine/Source/Runtime/Engine/Classes/Animation/AnimLinkableElement.cpp b/Engine/Source/Runtime/Engine/Classes/Animation/AnimLinkableElement.cpp
--- a/Engine/Source/Runtime/Engine/Classes/Animation/AnimLinkableElement.cpp
+++ b/Engine/Source/Runtime/Engine/Classes/Animation/AnimLinkableElement.cpp
@@ -364,3 +364,15 @@ void FAnimLinkableElement::ConditionalRelink()
 		}
 	}
 }
+
+void FAnimLinkableElement::Link(UAnimSequenceBase* AnimObject, float AbsTime, int32 InSlotIndex /*= 0*/)
+{
+	if(UAnimMontage* Montage = Cast<UAnimMontage>(AnimObject))
+	{
+		LinkMontage(Montage, AbsTime, InSlotIndex);
+	}
+	else
+	{
+		LinkSequence(AnimObject, AbsTime);
+	}
+}
diff --git a/Engine/Source/Runtime/Engine/Classes/Animation/AnimLinkableElement.h b/Engine/Source/Runtime/Engine/Classes/Animation/AnimLinkableElement.h
--- a/Engine/Source/Runtime/Engine/Classes/Animation/AnimLinkableElement.h
+++ b/Engine/Source/Runtime/Engine/Classes/Animation/AnimLinkableElement.h
@@ -48,6 +48,13 @@ struct FAnimLinkableElement
 	/** Update the timing information for this element. Will not search for a new segment unless this element does not have one */
 	ENGINE_API void Update();
 
+	/** Link this element to an animation object (Sequence or Montage)
+	 *	@param AnimObject The object to link to
+	 *	@param AbsTime The absolute time to place this element
+	 *	@param InSlotIndex Slot index for montages (ignored otherwise)
+	 */
+	ENGINE_API void Link(UAnimSequenceBase* AnimObject, float AbsTime, int32 InSlotIndex = 0);
+
 	/** Link this element to a montage 
 	 * @param Montage The montage to link to
 	 * @param AbsMontageTime The time in the montage that this element should be placed at
