Missing file

diff --git a/Engine/Source/Editor/Persona/Private/SAnimCurveEd.h b/Engine/Source/Editor/Persona/Private/SAnimCurveEd.h
new file mode 100644
--- /dev/null
+++ b/Engine/Source/Editor/Persona/Private/SAnimCurveEd.h
@@ -0,0 +1,61 @@
+// Copyright 1998-2014 Epic Games, Inc. All Rights Reserved.
+
+
+#pragma once
+
+#include "Persona.h"
+#include "SCurveEditor.h"
+
+DECLARE_DELEGATE_OneParam(FOnSelectionChanged, const FGraphPanelSelectionSet&)
+DECLARE_DELEGATE(FOnUpdatePanel)
+DECLARE_DELEGATE_RetVal(float, FOnGetScrubValue)
+
+//////////////////////////////////////////////////////////////////////////
+//  SAnimCurveEd : anim curve editor
+
+class SAnimCurveEd: public SCurveEditor
+{
+public:
+	SLATE_BEGIN_ARGS(SAnimCurveEd)
+		: _ViewMinInput(0.0f)
+		, _ViewMaxInput(10.0f)
+		, _TimelineLength(5.0f)
+		, _DrawCurve(true)
+		, _HideUI(true)
+		, _OnGetScrubValue()
+	{}
+
+	SLATE_ATTRIBUTE(float, ViewMinInput)
+	SLATE_ATTRIBUTE(float, ViewMaxInput)
+	SLATE_ATTRIBUTE(TOptional<float>, DataMinInput)
+	SLATE_ATTRIBUTE(TOptional<float>, DataMaxInput)
+	SLATE_ATTRIBUTE(float, TimelineLength)
+	SLATE_ATTRIBUTE(int32, NumberOfKeys)
+	SLATE_ATTRIBUTE(FVector2D, DesiredSize)
+	SLATE_ARGUMENT(bool, DrawCurve)
+	SLATE_ARGUMENT(bool, HideUI)
+	SLATE_EVENT(FOnGetScrubValue, OnGetScrubValue)
+	SLATE_EVENT(FOnSetInputViewRange, OnSetInputViewRange)
+SLATE_END_ARGS()
+
+void Construct(const FArguments& InArgs);
+
+protected:
+	// SWidget interface
+	virtual int32 OnPaint(const FPaintArgs& Args, const FGeometry& AllottedGeometry, const FSlateRect& MyClippingRect, FSlateWindowElementList& OutDrawElements, int32 LayerId, const FWidgetStyle& InWidgetStyle, bool bParentEnabled) const override;
+	virtual FReply OnMouseWheel(const FGeometry& MyGeometry, const FPointerEvent& MouseEvent);
+	virtual FCursorReply OnCursorQuery(const FGeometry& MyGeometry, const FPointerEvent& CursorEvent) const override;
+	// SWidget interface
+
+	// SCurveEditor interface
+	virtual void SetDefaultOutput(const float MinZoomRange) override;
+	virtual float GetTimeStep(FTrackScaleInfo &ScaleInfo) const override;
+	// SCurveEditor interface
+
+private:
+	// scrub value grabber
+	FOnGetScrubValue	OnGetScrubValue;
+	// @todo redo this code, all mess and dangerous
+	TAttribute<int32>	NumberOfKeys;
+};
+
