Missing cpp file.

diff --git a/Engine/Source/Editor/Persona/Private/SAnimCurveEd.cpp b/Engine/Source/Editor/Persona/Private/SAnimCurveEd.cpp
new file mode 100644
--- /dev/null
+++ b/Engine/Source/Editor/Persona/Private/SAnimCurveEd.cpp
@@ -0,0 +1,131 @@
+// Copyright 1998-2014 Epic Games, Inc. All Rights Reserved.
+
+
+#include "PersonaPrivatePCH.h"
+
+#include "SAnimCurveEd.h"
+
+#define LOCTEXT_NAMESPACE "AnimCurveEd"
+
+//////////////////////////////////////////////////////////////////////////
+//  SAnimCurveEd : anim curve editor
+
+float SAnimCurveEd::GetTimeStep(FTrackScaleInfo &ScaleInfo)const
+{
+	if(NumberOfKeys.Get())
+	{
+		int32 Divider = SScrubWidget::GetDivider(ViewMinInput.Get(), ViewMaxInput.Get(), ScaleInfo.WidgetSize, TimelineLength.Get(), NumberOfKeys.Get());
+
+		float TimePerKey;
+
+		if(NumberOfKeys.Get() != 0.f)
+		{
+			TimePerKey = TimelineLength.Get()/(float)NumberOfKeys.Get();
+		}
+		else
+		{
+			TimePerKey = 1.f;
+		}
+
+		return TimePerKey * Divider;
+	}
+
+	return 0.0f;
+}
+
+int32 SAnimCurveEd::OnPaint(const FPaintArgs& Args, const FGeometry& AllottedGeometry, const FSlateRect& MyClippingRect, FSlateWindowElementList& OutDrawElements, int32 LayerId, const FWidgetStyle& InWidgetStyle, bool bParentEnabled) const
+{
+	int32 NewLayerId = SCurveEditor::OnPaint(Args, AllottedGeometry, MyClippingRect, OutDrawElements, LayerId, InWidgetStyle, bParentEnabled) + 1;
+
+	float Value = 0.f;
+
+	if(OnGetScrubValue.IsBound())
+	{
+		Value = OnGetScrubValue.Execute();
+	}
+
+	FPaintGeometry MyGeometry = AllottedGeometry.ToPaintGeometry();
+
+	// scale info
+	FTrackScaleInfo ScaleInfo(ViewMinInput.Get(), ViewMaxInput.Get(), 0.f, 0.f, AllottedGeometry.Size);
+	float XPos = ScaleInfo.InputToLocalX(Value);
+
+	TArray<FVector2D> LinePoints;
+	LinePoints.Add(FVector2D(XPos-1, 0.f));
+	LinePoints.Add(FVector2D(XPos+1, AllottedGeometry.Size.Y));
+
+
+	FSlateDrawElement::MakeLines(
+		OutDrawElements,
+		NewLayerId,
+		MyGeometry,
+		LinePoints,
+		MyClippingRect,
+		ESlateDrawEffect::None,
+		FLinearColor::Red
+		);
+
+	// now draw scrub with new layer ID + 1;
+	return NewLayerId;
+}
+
+FReply SAnimCurveEd::OnMouseWheel(const FGeometry& MyGeometry, const FPointerEvent& MouseEvent)
+{
+	const float ZoomDelta = -0.1f * MouseEvent.GetWheelDelta();
+
+	const FVector2D WidgetSpace = MyGeometry.AbsoluteToLocal(MouseEvent.GetScreenSpacePosition());
+	const float ZoomRatio = FMath::Clamp((WidgetSpace.X / MyGeometry.Size.X), 0.f, 1.f);
+
+	{
+		const float InputViewSize = ViewMaxInput.Get() - ViewMinInput.Get();
+		const float InputChange = InputViewSize * ZoomDelta;
+
+		float NewViewMinInput = ViewMinInput.Get() - (InputChange * ZoomRatio);
+		float NewViewMaxInput = ViewMaxInput.Get() + (InputChange * (1.f - ZoomRatio));
+
+		SetInputMinMax(NewViewMinInput, NewViewMaxInput);
+	}
+
+	return FReply::Handled();
+}
+
+FCursorReply SAnimCurveEd::OnCursorQuery(const FGeometry& MyGeometry, const FPointerEvent& CursorEvent) const
+{
+	if(ViewMinInput.Get() > 0.f || ViewMaxInput.Get() < TimelineLength.Get())
+	{
+		return FCursorReply::Cursor(EMouseCursor::GrabHand);
+	}
+
+	return FCursorReply::Unhandled();
+}
+
+void SAnimCurveEd::Construct(const FArguments& InArgs)
+{
+	OnGetScrubValue = InArgs._OnGetScrubValue;
+	NumberOfKeys = InArgs._NumberOfKeys;
+
+	SCurveEditor::Construct(SCurveEditor::FArguments()
+		.ViewMinInput(InArgs._ViewMinInput)
+		.ViewMaxInput(InArgs._ViewMaxInput)
+		.DataMinInput(InArgs._DataMinInput)
+		.DataMaxInput(InArgs._DataMaxInput)
+		.ViewMinOutput(0.f)
+		.ViewMaxOutput(1.f)
+		.ZoomToFitVertical(true)
+		.ZoomToFitHorizontal(false)
+		.TimelineLength(InArgs._TimelineLength)
+		.DrawCurve(InArgs._DrawCurve)
+		.HideUI(InArgs._HideUI)
+		.AllowZoomOutput(false)
+		.DesiredSize(InArgs._DesiredSize)
+		.OnSetInputViewRange(InArgs._OnSetInputViewRange));
+}
+
+void SAnimCurveEd::SetDefaultOutput(const float MinZoomRange)
+{
+	const float NewMinOutput = (ViewMinOutput.Get());
+	const float NewMaxOutput = (ViewMaxOutput.Get() + MinZoomRange);
+
+	SetOutputMinMax(NewMinOutput, NewMaxOutput);
+}
+#undef LOCTEXT_NAMESPACE
