Adding scale editing - yay!

diff --git a/Engine/Source/Editor/Persona/Private/AnimationEditorViewportClient.cpp b/Engine/Source/Editor/Persona/Private/AnimationEditorViewportClient.cpp
--- a/Engine/Source/Editor/Persona/Private/AnimationEditorViewportClient.cpp
+++ b/Engine/Source/Editor/Persona/Private/AnimationEditorViewportClient.cpp
@@ -1108,7 +1108,8 @@ bool FAnimationViewportClient::InputWidgetDelta( FViewport* Viewport, EAxisList:
 		{
 			FTransform CurrentSkelControlTM(
 				SelectedSocket ? SelectedSocket->RelativeRotation : SkelControl->Rotation,
-				SelectedSocket ? SelectedSocket->RelativeLocation : SkelControl->Translation);
+				SelectedSocket ? SelectedSocket->RelativeLocation : SkelControl->Translation,
+				SelectedSocket ? SelectedSocket->RelativeScale : SkelControl->Scale);
 
 			FTransform BaseTM;
 
@@ -1122,10 +1123,11 @@ bool FAnimationViewportClient::InputWidgetDelta( FViewport* Viewport, EAxisList:
 			}
 
 			// Remove SkelControl's orientation from BoneMatrix, as we need to translate/rotate in the non-SkelControlled space
-			BaseTM = BaseTM.Inverse() * CurrentSkelControlTM;
+			BaseTM = BaseTM.GetRelativeTransformReverse(CurrentSkelControlTM);
 
 			const bool bDoRotation    = WidgetMode == FWidget::WM_Rotate    || WidgetMode == FWidget::WM_TranslateRotateZ;
 			const bool bDoTranslation = WidgetMode == FWidget::WM_Translate || WidgetMode == FWidget::WM_TranslateRotateZ;
+			const bool bDoScale = WidgetMode == FWidget::WM_Scale;
 
 			if (bDoRotation)
 			{
@@ -1164,9 +1166,30 @@ bool FAnimationViewportClient::InputWidgetDelta( FViewport* Viewport, EAxisList:
 					SkelControl->Translation += BoneSpaceOffset;
 				}
 			}
+			if(bDoScale)
+			{
+				FVector4 BoneSpaceScaleOffset;
+
+				if (ModeTools->GetCoordSystem() == COORD_World)
+				{
+					BoneSpaceScaleOffset = BaseTM.TransformVector(Scale);
+				}
+				else
+				{
+					BoneSpaceScaleOffset = Scale;
+				}
 
-				//@TODO: ANIMATION: Add scaling support here
+				if(SelectedSocket)
+				{
+					SelectedSocket->RelativeScale += BoneSpaceScaleOffset;
+				}
+				else
+				{
+					SkelControl->Scale += BoneSpaceScaleOffset;
+				}
 			}
+			
+		}
 		else if( WindActor.IsValid() )
 		{
 			if (WidgetMode == FWidget::WM_Rotate)
@@ -1436,8 +1459,7 @@ void FAnimationViewportClient::SetWidgetMode(FWidget::EWidgetMode InMode)
 
 bool FAnimationViewportClient::CanSetWidgetMode(FWidget::EWidgetMode NewMode) const
 {
-	//@TODO: ANIMATION: Add scaling support here
-	return (NewMode != FWidget::WM_Scale);
+	return true;
 }
 
 void FAnimationViewportClient::SetLocalAxesMode(ELocalAxesMode::Type AxesMode)
diff --git a/Engine/Source/Editor/Persona/Private/SAnimTrackCurvePanel.cpp b/Engine/Source/Editor/Persona/Private/SAnimTrackCurvePanel.cpp
--- a/Engine/Source/Editor/Persona/Private/SAnimTrackCurvePanel.cpp
+++ b/Engine/Source/Editor/Persona/Private/SAnimTrackCurvePanel.cpp
@@ -17,19 +17,30 @@
 //////////////////////////////////////////////////////////////////////////
 //  FAnimTrackCurveInterface interface 
 
+namespace ETransformCurve
+{
+	enum Type
+	{
+		Translation = 0,
+		Rotation,
+		Scale,
+		Max
+	};
+}
+
 /** Interface you implement if you want the CurveEditors to be able to edit curves on you */
 class FAnimTrackCurveBaseInterface : public FCurveOwnerInterface
 {
 public:
 	TWeakObjectPtr<UAnimSequence>	BaseSequence;
 	FAnimCurveBase*	CurveData;
-	bool			bRotationCurve;
+	ETransformCurve::Type	CurveType;
 
 public:
-	FAnimTrackCurveBaseInterface(UAnimSequence * BaseSeq, FAnimCurveBase*	InData, bool bInRotationCurve)
+	FAnimTrackCurveBaseInterface(UAnimSequence * BaseSeq, FAnimCurveBase*	InData, ETransformCurve::Type InCurveType)
 		: BaseSequence(BaseSeq)
 		, CurveData(InData)
-		, bRotationCurve(bInRotationCurve)
+		, CurveType(InCurveType)
 	{
 		// they should be valid
 		check (BaseSequence.IsValid());
@@ -38,7 +49,8 @@ public:
 
 	FName GetDisplayCurveName(int32 Index) const
 	{
-		if(bRotationCurve)
+		// rotation curve
+		if(CurveType == ETransformCurve::Rotation )
 		{
 			switch(Index)
 			{
@@ -128,7 +140,7 @@ public:
 				{
 					FName DisplayName = CurveName;
 
-					return FText::FromString(FString::Printf(TEXT("%s(%c)"), *DisplayName.ToString(), (bRotationCurve)? TCHAR('R') : TCHAR('T')));
+					return FText::FromString(FString::Printf(TEXT("%s(%c)"), *DisplayName.ToString(), GetCurveTypeCharacter()));
 				}
 			}
 		}
@@ -136,6 +148,19 @@ public:
 		return FText::GetEmpty();
 	}
 
+	TCHAR GetCurveTypeCharacter() const
+	{
+		switch (CurveType)
+		{
+		case ETransformCurve::Translation:
+			return TCHAR('T');
+		case ETransformCurve::Rotation:
+			return TCHAR('R');
+		}
+
+		return TCHAR('S');
+	}
+
 	virtual void OnCurveChanged() override
 	{
 	}
@@ -149,20 +174,14 @@ class STransformCurveEdTrack : public SCompoundWidget
 {
 private:
 	/** Pointer to notify panel for drawing*/
-	TSharedPtr<class SCurveEditor>			CurveEditors[2];
+	TSharedPtr<class SCurveEditor>			CurveEditors[ETransformCurve::Max];
 
 	/** Name of curve it's editing - CurveName should be unique within this tracks**/
-	FAnimTrackCurveBaseInterface	*		CurveInterfaces[2];
+	FAnimTrackCurveBaseInterface	*		CurveInterfaces[ETransformCurve::Max];
 
 	/** Curve Panel Ptr **/
 	TWeakPtr<SAnimTrackCurvePanel>			PanelPtr;
 
-	enum EType
-	{
-		Translation = 0,
-		Rotation,
-	};
-
 private:
 	USkeleton::AnimCurveUID	CurveUid;
 
@@ -206,7 +225,7 @@ public:
 	FVector2D GetDesiredSize() const;
 
 	// Bound to attribute for curve name, uses curve interface to request from skeleton
-	FText GetCurveName(USkeleton::AnimCurveUID Uid, EType Type) const;
+	FText GetCurveName(USkeleton::AnimCurveUID Uid, ETransformCurve::Type Type) const;
 };
 
 //////////////////////////////////////////////////////////////////////////
@@ -228,8 +247,9 @@ void STransformCurveEdTrack::Construct(const FArguments& InArgs)
 
 	CurveUid = InArgs._CurveUid;
 
-	CurveInterfaces[EType::Translation] = new FAnimTrackCurveBaseInterface(Sequence, &Curve->TranslationCurve, false);
-	CurveInterfaces[EType::Rotation] = new FAnimTrackCurveBaseInterface(Sequence, &Curve->RotationCurve, true);
+	CurveInterfaces[ETransformCurve::Translation] = new FAnimTrackCurveBaseInterface(Sequence, &Curve->TranslationCurve, ETransformCurve::Translation);
+	CurveInterfaces[ETransformCurve::Rotation] = new FAnimTrackCurveBaseInterface(Sequence, &Curve->RotationCurve, ETransformCurve::Rotation);
+	CurveInterfaces[ETransformCurve::Scale] = new FAnimTrackCurveBaseInterface(Sequence, &Curve->ScaleCurve, ETransformCurve::Scale);
 	int32 NumberOfKeys = Sequence->GetNumberOfFrames();
 	//////////////////////////////
 	
@@ -248,7 +268,7 @@ void STransformCurveEdTrack::Construct(const FArguments& InArgs)
 				.FillWidth(1)
 				[
 					// Notification editor panel
-					SAssignNew(CurveEditors[EType::Translation], SAnimCurveEd)
+					SAssignNew(CurveEditors[ETransformCurve::Translation], SAnimCurveEd)
 					.ViewMinInput(InArgs._ViewInputMin)
 					.ViewMaxInput(InArgs._ViewInputMax)
 					.DataMinInput(0.f)
@@ -283,7 +303,7 @@ void STransformCurveEdTrack::Construct(const FArguments& InArgs)
 							.Font(FEditorStyle::GetFontStyle("CurveEd.InfoFont"))
 							.SelectAllTextWhenFocused(true)
 							.IsReadOnly(true)
-							.Text(this, &STransformCurveEdTrack::GetCurveName, Curve->CurveUid, EType::Translation)
+							.Text(this, &STransformCurveEdTrack::GetCurveName, Curve->CurveUid, ETransformCurve::Translation)
 						]
 
 						+SHorizontalBox::Slot()
@@ -313,7 +333,7 @@ void STransformCurveEdTrack::Construct(const FArguments& InArgs)
 				.FillWidth(1)
 				[
 					// Notification editor panel
-					SAssignNew(CurveEditors[EType::Rotation], SAnimCurveEd)
+					SAssignNew(CurveEditors[ETransformCurve::Rotation], SAnimCurveEd)
 					.ViewMinInput(InArgs._ViewInputMin)
 					.ViewMaxInput(InArgs._ViewInputMax)
 					.DataMinInput(0.f)
@@ -349,18 +369,85 @@ void STransformCurveEdTrack::Construct(const FArguments& InArgs)
 							.Font(FEditorStyle::GetFontStyle("CurveEd.InfoFont"))
 							.SelectAllTextWhenFocused(true)
 							.IsReadOnly(true)
-							.Text(this, &STransformCurveEdTrack::GetCurveName, Curve->CurveUid, EType::Rotation)
+							.Text(this, &STransformCurveEdTrack::GetCurveName, Curve->CurveUid, ETransformCurve::Rotation)
+						]
+					]
+				]
+			]
+
+			+SVerticalBox::Slot()
+			[
+				SNew(SHorizontalBox)
+
+				+SHorizontalBox::Slot()
+				.FillWidth(1)
+				[
+					// Notification editor panel
+					SAssignNew(CurveEditors[ETransformCurve::Scale], SAnimCurveEd)
+					.ViewMinInput(InArgs._ViewInputMin)
+					.ViewMaxInput(InArgs._ViewInputMax)
+					.DataMinInput(0.f)
+					.DataMaxInput(Sequence->SequenceLength)
+					// @fixme fix this to delegate
+					.TimelineLength(Sequence->SequenceLength)
+					.NumberOfKeys(NumberOfKeys)
+					.DesiredSize(this, &STransformCurveEdTrack::GetDesiredSize)
+					.OnSetInputViewRange(InArgs._OnSetInputViewRange)
+					.OnGetScrubValue(InArgs._OnGetScrubValue)
+				]
+
+				+SHorizontalBox::Slot()
+				.AutoWidth()
+				.Padding(2.f)
+				[
+					SNew(SBox)
+					.WidthOverride(InArgs._WidgetWidth)
+					[
+						SNew(SHorizontalBox)
+
+						+SHorizontalBox::Slot()
+						.FillWidth(1)
+						.VAlign(VAlign_Center)
+						.HAlign(HAlign_Left)
+						.Padding(5, 0, 0, 0)
+						[
+							// Name of track
+							SNew(SEditableText)
+							.MinDesiredWidth(64.0f)
+							.IsEnabled(true)
+							.Font(FEditorStyle::GetFontStyle("CurveEd.InfoFont"))
+							.SelectAllTextWhenFocused(true)
+							.IsReadOnly(true)
+							.Text(this, &STransformCurveEdTrack::GetCurveName, Curve->CurveUid, ETransformCurve::Scale)
+						]
+
+						+SHorizontalBox::Slot()
+						.Padding(FMargin(0.0f, 5.0f, 0.0f, 5.0f))
+						.AutoWidth()
+						.VAlign(VAlign_Top)
+						[
+							SNew(SButton)
+							.ToolTipText(LOCTEXT("DisplayTrackOptionsMenuTooltip", "Display track options menu"))
+							.OnClicked(this, &STransformCurveEdTrack::OnContextMenu)
+							.Content()
+							[
+								SNew(SImage)
+								.Image(FEditorStyle::GetBrush("ComboButton.Arrow"))
+								.ColorAndOpacity(FSlateColor::UseForeground())
+							]
 						]
 					]
 				]
 			]
 
+
 		]
 	];
 	
 	//Inform track widget about the curve and whether it is editable or not.
-	CurveEditors[EType::Translation]->SetCurveOwner(CurveInterfaces[EType::Translation], true);
-	CurveEditors[EType::Rotation]->SetCurveOwner(CurveInterfaces[EType::Rotation], true);
+	CurveEditors[ETransformCurve::Translation]->SetCurveOwner(CurveInterfaces[ETransformCurve::Translation], true);
+	CurveEditors[ETransformCurve::Rotation]->SetCurveOwner(CurveInterfaces[ETransformCurve::Rotation], true);
+	CurveEditors[ETransformCurve::Scale]->SetCurveOwner(CurveInterfaces[ETransformCurve::Scale], true);
 }
 
 /** return a widget */
@@ -371,10 +458,9 @@ const FSlateBrush* STransformCurveEdTrack::GetExpandContent() const
 
 STransformCurveEdTrack::~STransformCurveEdTrack()
 {
-	// @fixme - check -is this okay way of doing it?
-	// @fix this - funny
-	delete CurveInterfaces[EType::Translation];
-	delete CurveInterfaces[EType::Rotation];
+	delete CurveInterfaces[ETransformCurve::Translation];
+	delete CurveInterfaces[ETransformCurve::Rotation];
+	delete CurveInterfaces[ETransformCurve::Scale];
 }
 
 FReply STransformCurveEdTrack::OnContextMenu()
@@ -397,7 +483,7 @@ FVector2D STransformCurveEdTrack::GetDesiredSize() const
 	return FVector2D(128.f, 128.f);
 }
 
-FText STransformCurveEdTrack::GetCurveName(USkeleton::AnimCurveUID Uid, EType Type) const
+FText STransformCurveEdTrack::GetCurveName(USkeleton::AnimCurveUID Uid, ETransformCurve::Type Type) const
 {
 	return CurveInterfaces[Type]->GetCurveName(Uid);
 }
