UE-5202 - Fix for notify binding field not appearing in transition rules due to visibility changes in montage improvements

diff --git a/Engine/Source/Editor/Persona/Private/AnimNotifyDetails.cpp b/Engine/Source/Editor/Persona/Private/AnimNotifyDetails.cpp
--- a/Engine/Source/Editor/Persona/Private/AnimNotifyDetails.cpp
+++ b/Engine/Source/Editor/Persona/Private/AnimNotifyDetails.cpp
@@ -35,6 +35,9 @@ void FAnimNotifyDetails::CustomizeDetails( IDetailLayoutBuilder& DetailBuilder )
 	FString NotifyClassName;
 	TSharedRef<IPropertyHandle> NotifyPropHandle = DetailBuilder.GetProperty(TEXT("Event.Notify"));
 	NotifyPropHandle->GetValue(NotifyPtr);
+	
+	// Don't want to edit the notify name here.
+	DetailBuilder.HideProperty(TEXT("Event.NotifyName"));
 
 	IDetailCategoryBuilder& AnimNotifyCategory = DetailBuilder.EditCategory(TEXT("AnimNotify"), TEXT(""), ECategoryPriority::TypeSpecific);
 
diff --git a/Engine/Source/Runtime/Engine/Classes/Animation/AnimSequenceBase.h b/Engine/Source/Runtime/Engine/Classes/Animation/AnimSequenceBase.h
--- a/Engine/Source/Runtime/Engine/Classes/Animation/AnimSequenceBase.h
+++ b/Engine/Source/Runtime/Engine/Classes/Animation/AnimSequenceBase.h
@@ -53,7 +53,7 @@ struct FAnimNotifyEvent : public FAnimLinkableElement
 	UPROPERTY(EditAnywhere, BlueprintReadWrite, Category=AnimNotifyEvent)
 	float TriggerWeightThreshold;
 
-	UPROPERTY()
+	UPROPERTY(EditAnywhere, BlueprintReadOnly, Category=AnimNotifyEvent)
 	FName NotifyName;
 
 	UPROPERTY(EditAnywhere, Instanced, BlueprintReadWrite, Category=AnimNotifyEvent)
