to support direct manipulation of Pin values from anim node in Bone controller preview mode

diff --git a/Engine/Source/Editor/Persona/Private/AnimationEditorViewportClient.cpp b/Engine/Source/Editor/Persona/Private/AnimationEditorViewportClient.cpp
--- a/Engine/Source/Editor/Persona/Private/AnimationEditorViewportClient.cpp
+++ b/Engine/Source/Editor/Persona/Private/AnimationEditorViewportClient.cpp
@@ -750,6 +750,29 @@ void FAnimationViewportClient::PostCompile()
 {
 	// reset the selected anim graph to copy 
 	SelectedSkelControlAnimGraph.Reset();
+
+	// if the user manipulated Pin values directly from the node, then should copy updated values to the internal node to retain data consistency
+	UEdGraph* FocusedGraph = PersonaPtr.Pin()->GetFocusedGraph();
+
+	// @fixme : fix this to better way than looping all nodes
+	if (FocusedGraph)
+	{
+		// find UAnimGraphNode_SkeletalControlBase
+		for (UEdGraphNode* Node : FocusedGraph->Nodes)
+		{
+			UAnimGraphNode_SkeletalControlBase* AnimGraphNode = Cast<UAnimGraphNode_SkeletalControlBase>(Node);
+
+			if (AnimGraphNode)
+			{
+				FAnimNode_SkeletalControlBase* AnimNode = FindSkeletalControlAnimNode(AnimGraphNode);
+
+				if (AnimNode)
+				{
+					AnimGraphNode->CopyNodeDataFrom(AnimNode);
+				}
+			}
+		}
+	}
 }
 
 void FAnimationViewportClient::Tick(float DeltaSeconds) 
