Rather than accessing SmallLayoutFont from a static variable, use a static function. On OS X the module loading order is very different & statics are initialised on load, so this was being constructed long before the engine & UObject system were initialised causing an assert in the UObject construction.

diff --git a/Engine/Source/Editor/Persona/Private/SAnimationBlendSpaceBase.cpp b/Engine/Source/Editor/Persona/Private/SAnimationBlendSpaceBase.cpp
--- a/Engine/Source/Editor/Persona/Private/SAnimationBlendSpaceBase.cpp
+++ b/Engine/Source/Editor/Persona/Private/SAnimationBlendSpaceBase.cpp
@@ -18,7 +18,11 @@
 
 DEFINE_LOG_CATEGORY_STATIC(LogBlendSpaceBase, Log, All);
 
-FSlateFontInfo SmallLayoutFont( FPaths::EngineContentDir() / TEXT("Slate/Fonts/Roboto-Regular.ttf"), 10 );
+static FSlateFontInfo GetSmallLayoutFont()
+{
+	static FSlateFontInfo SmallLayoutFont( FPaths::EngineContentDir() / TEXT("Slate/Fonts/Roboto-Regular.ttf"), 10 );
+	return SmallLayoutFont;
+}
 
 #define LOCTEXT_NAMESPACE "BlendSpaceEditorBase"
 
@@ -305,7 +309,7 @@ void SBlendSpaceParameterWidget::ConstructParameterPanel()
 						[
 							SNew(STextBlock)
 							.Text(AxisNameLabel)
-							.Font(SmallLayoutFont)
+							.Font(GetSmallLayoutFont())
 						]
 					]
 
@@ -334,7 +338,7 @@ void SBlendSpaceParameterWidget::ConstructParameterPanel()
 							[
 								SNew(STextBlock)
 								.Text(AxisRangeLabel)
-								.Font(SmallLayoutFont)
+								.Font(GetSmallLayoutFont())
 							]
 						]
 
@@ -356,7 +360,7 @@ void SBlendSpaceParameterWidget::ConstructParameterPanel()
 									.VAlign(VAlign_Center)
 									[
 										SNew(STextBlock)
-										.Font( SmallLayoutFont )
+										.Font( GetSmallLayoutFont() )
 										.Text(LOCTEXT("MinValue - MaxValue Separator", " - "))
 									]
 
@@ -387,7 +391,7 @@ void SBlendSpaceParameterWidget::ConstructParameterPanel()
 							[
 								SNew(STextBlock)
 								.Text(AxisGridCountLabel)
-								.Font(SmallLayoutFont)
+								.Font(GetSmallLayoutFont())
 							]
 						]
 
@@ -514,13 +518,13 @@ void SBlendSpaceSamplesWidget::ConstructSamplesPanel()
 				+SHorizontalBox::Slot() .HAlign(HAlign_Right)
 				[
 					SNew(STextBlock)
-					.Font( SmallLayoutFont )
+					.Font( GetSmallLayoutFont() )
 					.Text(FString::Printf(TEXT("%3d] "), I+1))
 				]
 				+SHorizontalBox::Slot() .FillWidth(1) .HAlign(HAlign_Left)
 				[
 					SNew(SEditableText)
-					.Font( SmallLayoutFont )
+					.Font( GetSmallLayoutFont() )
 					.Text( FText::FromName(Sample.Animation->GetFName() ) )
 				]
 				+SHorizontalBox::Slot().AutoWidth() .HAlign(HAlign_Center)
