Add warning that changes to the preview properties wont be saved into the asset

diff --git a/Engine/Source/Editor/EditorStyle/Private/SlateEditorStyle.cpp b/Engine/Source/Editor/EditorStyle/Private/SlateEditorStyle.cpp
--- a/Engine/Source/Editor/EditorStyle/Private/SlateEditorStyle.cpp
+++ b/Engine/Source/Editor/EditorStyle/Private/SlateEditorStyle.cpp
@@ -4485,6 +4485,10 @@ void FSlateEditorStyle::FStyle::SetupPersonaStyle()
 
 		// Anim Notify Editor
 		Set("AnimNotifyEditor.BranchingPoint", new IMAGE_BRUSH("Persona/NotifyEditor/BranchingPoints_24x", Icon24x24));
+
+		// AnimBlueprint Preview Warning Background
+		FSlateColor PreviewPropertiesWarningColour(FLinearColor(0.1f, 0.2f, 0.5f));
+		Set("Persona.PreviewPropertiesWarning", new BOX_BRUSH("Common/RoundedSelection_16x", 4.0f / 16.0f, PreviewPropertiesWarningColour));
 	}
 
 	// Kismet 2
diff --git a/Engine/Source/Editor/Persona/Private/Shared/PersonaMode.cpp b/Engine/Source/Editor/Persona/Private/Shared/PersonaMode.cpp
--- a/Engine/Source/Editor/Persona/Private/Shared/PersonaMode.cpp
+++ b/Engine/Source/Editor/Persona/Private/Shared/PersonaMode.cpp
@@ -329,7 +329,25 @@ TSharedRef<SWidget> FAnimBlueprintDefaultsEditorSummoner::CreateTabBody(const FW
 					.BorderImage( FEditorStyle::GetBrush("NoBorder") )
 					.Visibility( this, &FAnimBlueprintDefaultsEditorSummoner::IsEditorVisible, EAnimBlueprintEditorMode::PreviewMode)
 					[
-						PersonaPtr->GetPreviewEditor()
+						SNew(SVerticalBox)
+						+SVerticalBox::Slot()
+						.AutoHeight()
+						.Padding(0.f, 8.f, 0.f, 0.f)
+						[
+							SNew(SBorder)
+							.BorderImage(FEditorStyle::GetBrush("Persona.PreviewPropertiesWarning"))
+							[
+								SNew(STextBlock)
+								.Text(LOCTEXT("AnimBlueprintEditPreviewText", "Changes to preview options are not saved in the asset."))
+								.Font(FEditorStyle::GetFontStyle("PropertyWindow.NormalFont"))
+								.ShadowColorAndOpacity(FLinearColor::Black.CopyWithNewOpacity(0.3f))
+								.ShadowOffset(FVector2D::UnitVector)
+							]
+						]
+						+SVerticalBox::Slot()
+						[
+							PersonaPtr->GetPreviewEditor()
+						]
 					]
 				]
 				+SOverlay::Slot()
