Fix crash when changing paramater range in blend space 1D editor

diff --git a/Engine/Source/Editor/Persona/Private/SAnimationBlendSpace1D.cpp b/Engine/Source/Editor/Persona/Private/SAnimationBlendSpace1D.cpp
--- a/Engine/Source/Editor/Persona/Private/SAnimationBlendSpace1D.cpp
+++ b/Engine/Source/Editor/Persona/Private/SAnimationBlendSpace1D.cpp
@@ -399,7 +399,10 @@ int SBlendSpace1DWidget::OnPaint( const FPaintArgs& Args, const FGeometry& Allot
 		if(CachedSamples.Num() > 0)
 		{
 			TOptional<FVector2D> PreviewBSInput = GetWidgetPosFromEditorPos(PreviewInput.Get(), WindowRect);
-			DrawSamplePoint(PreviewBSInput.GetValue(), EBlendSpaceSampleState::Highlighted, AllottedGeometry, MyClippingRect, OutDrawElements, HighlightLayer);
+			if (PreviewBSInput.IsSet())
+			{
+				DrawSamplePoint(PreviewBSInput.GetValue(), EBlendSpaceSampleState::Highlighted, AllottedGeometry, MyClippingRect, OutDrawElements, HighlightLayer);
+			}
 		}
 		
 		if(const FLineElement* LineElement = ElementGenerator.GetLineElementForBlendInput(PreviewInput.Get()))
