Fixed redundant menu showing for editing curves

diff --git a/Engine/Source/Editor/Persona/Private/SAnimTrackCurvePanel.cpp b/Engine/Source/Editor/Persona/Private/SAnimTrackCurvePanel.cpp
--- a/Engine/Source/Editor/Persona/Private/SAnimTrackCurvePanel.cpp
+++ b/Engine/Source/Editor/Persona/Private/SAnimTrackCurvePanel.cpp
@@ -420,22 +420,6 @@ void STransformCurveEdTrack::Construct(const FArguments& InArgs)
 							.IsReadOnly(true)
 							.Text(this, &STransformCurveEdTrack::GetCurveName, Curve->CurveUid, ETransformCurve::Scale)
 						]
-
-						+SHorizontalBox::Slot()
-						.Padding(FMargin(0.0f, 5.0f, 0.0f, 5.0f))
-						.AutoWidth()
-						.VAlign(VAlign_Top)
-						[
-							SNew(SButton)
-							.ToolTipText(LOCTEXT("DisplayTrackOptionsMenuTooltip", "Display track options menu"))
-							.OnClicked(this, &STransformCurveEdTrack::OnContextMenu)
-							.Content()
-							[
-								SNew(SImage)
-								.Image(FEditorStyle::GetBrush("ComboButton.Arrow"))
-								.ColorAndOpacity(FSlateColor::UseForeground())
-							]
-						]
 					]
 				]
 			]
