Fixed anim notify events not appearing in blueprint context menu in Animation Blueprints
[UE-5601]

diff --git a/Engine/Source/Editor/Persona/Private/SAnimNotifyPanel.cpp b/Engine/Source/Editor/Persona/Private/SAnimNotifyPanel.cpp
--- a/Engine/Source/Editor/Persona/Private/SAnimNotifyPanel.cpp
+++ b/Engine/Source/Editor/Persona/Private/SAnimNotifyPanel.cpp
@@ -13,6 +13,7 @@
 #include "STextEntryPopup.h"
 #include "SExpandableArea.h"
 #include "Toolkits/AssetEditorManager.h"
+#include "BlueprintActionDatabase.h"
 
 // Track Panel drawing
 const float NotificationTrackHeight = 20.0f;
@@ -2556,6 +2557,8 @@ void SAnimNotifyTrack::AddNewNotify(const FText& NewNotifyName, ETextCommit::Typ
 		FName NewName = FName( *NewNotifyName.ToString() );
 		SeqSkeleton->AddNewAnimationNotify(NewName);
 
+		FBlueprintActionDatabase::Get().RefreshAssetActions(SeqSkeleton);
+
 		CreateNewNotifyAtCursor(NewNotifyName.ToString(), (UClass*)nullptr);
 	}
 
diff --git a/Engine/Source/Editor/Persona/Private/SSkeletonAnimNotifies.cpp b/Engine/Source/Editor/Persona/Private/SSkeletonAnimNotifies.cpp
--- a/Engine/Source/Editor/Persona/Private/SSkeletonAnimNotifies.cpp
+++ b/Engine/Source/Editor/Persona/Private/SSkeletonAnimNotifies.cpp
@@ -11,6 +11,7 @@
 #include "SInlineEditableTextBlock.h"
 #include "SNotificationList.h"
 #include "NotificationManager.h"
+#include "BlueprintActionDatabase.h"
 
 #define LOCTEXT_NAMESPACE "SkeletonAnimNotifies"
 
@@ -296,6 +297,8 @@ void SSkeletonAnimNotifies::OnDeleteAnimNotify()
 		NotifyUser( Info );
 	}
 
+	FBlueprintActionDatabase::Get().RefreshAssetActions(TargetSkeleton);
+
 	CreateNotifiesList( NameFilterBox->GetText().ToString() );
 }
 
