Added sequence length getter to montage panel and passed through to segment tracks so they update with changes in the montage length.
[UE-5695]

diff --git a/Engine/Source/Editor/Persona/Private/SAnimMontagePanel.cpp b/Engine/Source/Editor/Persona/Private/SAnimMontagePanel.cpp
--- a/Engine/Source/Editor/Persona/Private/SAnimMontagePanel.cpp
+++ b/Engine/Source/Editor/Persona/Private/SAnimMontagePanel.cpp
@@ -137,7 +137,7 @@ void SAnimMontagePanel::Update()
 					.DraggableBars(Editor, &SMontageEditor::GetSectionStartTimes)
 					.DraggableBarSnapPositions(Editor, &SMontageEditor::GetAnimSegmentStartTimes)
 					.DraggableBarLabels(Editor, &SMontageEditor::GetSectionNames)
-					.TrackMaxValue(Montage->SequenceLength) // FIXME -> Make actual attribute!
+					.TrackMaxValue(this, &SAnimMontagePanel::GetSequenceLength)
 					.TrackNumDiscreteValues(Montage->GetNumberOfFrames())
 					.OnTrackRightClickContextMenu( this, &SAnimMontagePanel::SummonTrackContextMenu, static_cast<int>(INDEX_NONE))
 					.ScrubPosition( Editor, &SMontageEditor::GetScrubValue )
@@ -234,7 +234,7 @@ void SAnimMontagePanel::Update()
 					.DraggableBars(Editor, &SMontageEditor::GetSectionStartTimes)
 					.DraggableBarSnapPositions(Editor, &SMontageEditor::GetAnimSegmentStartTimes)
 					.ScrubPosition( Editor, &SMontageEditor::GetScrubValue )
-					.TrackMaxValue(Montage->SequenceLength)
+					.TrackMaxValue(this, &SAnimMontagePanel::GetSequenceLength)
 					.TrackNumDiscreteValues(Montage->GetNumberOfFrames())
 
 					.OnAnimSegmentNodeClicked( this, &SAnimMontagePanel::ShowSegmentInDetailsView, SlotAnimIdx )
@@ -659,4 +659,13 @@ void SAnimMontagePanel::OnAnimSegmentRemoved(int32 SegmentIndex, int32 SlotIndex
 	}
 }
 
+float SAnimMontagePanel::GetSequenceLength() const
+{
+	if(Montage != nullptr)
+	{
+		return Montage->SequenceLength;
+	}
+	return 0.0f;
+}
+
 #undef LOCTEXT_NAMESPACE
\ No newline at end of file
diff --git a/Engine/Source/Editor/Persona/Private/SAnimMontagePanel.h b/Engine/Source/Editor/Persona/Private/SAnimMontagePanel.h
--- a/Engine/Source/Editor/Persona/Private/SAnimMontagePanel.h
+++ b/Engine/Source/Editor/Persona/Private/SAnimMontagePanel.h
@@ -135,4 +135,9 @@ private:
 	 *	@param SegmentIndex - Index of the removed segment
 	 */
 	void OnAnimSegmentRemoved(int32 SegmentIndex, int32 SlotIndex);
+
+	/** Gets the length of the montage we are currently editing
+	 */
+	virtual float GetSequenceLength() const override;
+
 };
\ No newline at end of file
