Removed SLATE_TEXT_ATTRIBUTE from Editor specific widgets

diff --git a/Engine/Source/Editor/Kismet/Private/SBlueprintSubPalette.h b/Engine/Source/Editor/Kismet/Private/SBlueprintSubPalette.h
--- a/Engine/Source/Editor/Kismet/Private/SBlueprintSubPalette.h
+++ b/Engine/Source/Editor/Kismet/Private/SBlueprintSubPalette.h
@@ -17,7 +17,7 @@ public:
 		, _ShowFavoriteToggles(false)
 		{}
 
-		SLATE_TEXT_ATTRIBUTE(Title)
+		SLATE_ATTRIBUTE(FText, Title)
 		SLATE_ATTRIBUTE(FSlateBrush const*, Icon)
 		SLATE_ATTRIBUTE(bool, ShowFavoriteToggles)
 	SLATE_END_ARGS()
diff --git a/Engine/Source/Editor/MaterialEditor/Private/MaterialEditor.cpp b/Engine/Source/Editor/MaterialEditor/Private/MaterialEditor.cpp
--- a/Engine/Source/Editor/MaterialEditor/Private/MaterialEditor.cpp
+++ b/Engine/Source/Editor/MaterialEditor/Private/MaterialEditor.cpp
@@ -3361,9 +3361,9 @@ bool FMaterialEditor::CanDuplicateNodes() const
 	return CanCopyNodes();
 }
 
-FString FMaterialEditor::GetOriginalObjectName() const
+FText FMaterialEditor::GetOriginalObjectName() const
 {
-	return GetEditingObjects()[0]->GetName();
+	return FText::FromString(GetEditingObjects()[0]->GetName());
 }
 
 void FMaterialEditor::UpdateMaterialAfterGraphChange()
diff --git a/Engine/Source/Editor/MaterialEditor/Private/MaterialEditor.h b/Engine/Source/Editor/MaterialEditor/Private/MaterialEditor.h
--- a/Engine/Source/Editor/MaterialEditor/Private/MaterialEditor.h
+++ b/Engine/Source/Editor/MaterialEditor/Private/MaterialEditor.h
@@ -315,7 +315,7 @@ public:
 	virtual void ForceRefreshExpressionPreviews() override;
 	virtual void AddToSelection(UMaterialExpression* Expression) override;
 	virtual void DeleteSelectedNodes() override;
-	virtual FString GetOriginalObjectName() const override;
+	virtual FText GetOriginalObjectName() const override;
 	virtual void UpdateMaterialAfterGraphChange() override;
 	virtual bool CanPasteNodes() const override;
 	virtual void PasteNodesHere(const FVector2D& Location) override;
diff --git a/Engine/Source/Editor/MaterialEditor/Private/MaterialEditorUtilities.cpp b/Engine/Source/Editor/MaterialEditor/Private/MaterialEditorUtilities.cpp
--- a/Engine/Source/Editor/MaterialEditor/Private/MaterialEditorUtilities.cpp
+++ b/Engine/Source/Editor/MaterialEditor/Private/MaterialEditorUtilities.cpp
@@ -73,14 +73,14 @@ void FMaterialEditorUtilities::DeleteSelectedNodes(const class UEdGraph* Graph)
 	}
 }
 
-FString FMaterialEditorUtilities::GetOriginalObjectName(const class UEdGraph* Graph)
+FText FMaterialEditorUtilities::GetOriginalObjectName(const class UEdGraph* Graph)
 {
 	TSharedPtr<class IMaterialEditor> MaterialEditor = GetIMaterialEditorForObject(Graph);
 	if (MaterialEditor.IsValid())
 	{
 		return MaterialEditor->GetOriginalObjectName();
 	}
-	return TEXT("");
+	return FText::GetEmpty();
 }
 
 void FMaterialEditorUtilities::UpdateMaterialAfterGraphChange(const class UEdGraph* Graph)
diff --git a/Engine/Source/Editor/MaterialEditor/Private/SMaterialEditorTitleBar.h b/Engine/Source/Editor/MaterialEditor/Private/SMaterialEditorTitleBar.h
--- a/Engine/Source/Editor/MaterialEditor/Private/SMaterialEditorTitleBar.h
+++ b/Engine/Source/Editor/MaterialEditor/Private/SMaterialEditorTitleBar.h
@@ -13,7 +13,7 @@ public:
 		, _MaterialInfoList(NULL)
 	{}
 
-		SLATE_TEXT_ATTRIBUTE( TitleText )
+		SLATE_ATTRIBUTE( FText, TitleText )
 
 		SLATE_ARGUMENT( const TArray<TSharedPtr<FMaterialInfo>>*, MaterialInfoList )
 	SLATE_END_ARGS()
diff --git a/Engine/Source/Editor/MaterialEditor/Public/IMaterialEditor.h b/Engine/Source/Editor/MaterialEditor/Public/IMaterialEditor.h
--- a/Engine/Source/Editor/MaterialEditor/Public/IMaterialEditor.h
+++ b/Engine/Source/Editor/MaterialEditor/Public/IMaterialEditor.h
@@ -50,7 +50,7 @@ public:
 	/**
 	 * Gets the name of the material or material function that we are editing
 	 */
-	virtual FString GetOriginalObjectName() const {return TEXT("");}
+	virtual FText GetOriginalObjectName() const {return FText::GetEmpty();}
 
 	/**
 	 * Re-links the material and updates its representation in the editor,
diff --git a/Engine/Source/Editor/MaterialEditor/Public/MaterialEditorUtilities.h b/Engine/Source/Editor/MaterialEditor/Public/MaterialEditorUtilities.h
--- a/Engine/Source/Editor/MaterialEditor/Public/MaterialEditorUtilities.h
+++ b/Engine/Source/Editor/MaterialEditor/Public/MaterialEditorUtilities.h
@@ -73,7 +73,7 @@ public:
 	 *
 	 * @param	Graph	Graph representing a material or material function.
 	 */
-	static FString GetOriginalObjectName(const class UEdGraph* Graph);
+	static FText GetOriginalObjectName(const class UEdGraph* Graph);
 
 	/**
 	 * Re-links the material and updates its representation in the editor
diff --git a/Engine/Source/Editor/Persona/Private/SAnimViewportToolBar.cpp b/Engine/Source/Editor/Persona/Private/SAnimViewportToolBar.cpp
--- a/Engine/Source/Editor/Persona/Private/SAnimViewportToolBar.cpp
+++ b/Engine/Source/Editor/Persona/Private/SAnimViewportToolBar.cpp
@@ -783,18 +783,19 @@ FSlateColor SAnimViewportToolBar::GetFontColor() const
 	return FontColor;
 }
 
-FString SAnimViewportToolBar::GetPlaybackMenuLabel() const
+FText SAnimViewportToolBar::GetPlaybackMenuLabel() const
 {
-	FString Label = TEXT("Error");
+	FText Label = LOCTEXT("PlaybackError", "Error");
 	if (Viewport.IsValid())
 	{
 		for(int i = 0; i < EAnimationPlaybackSpeeds::NumPlaybackSpeeds; ++i)
 		{
 			if (Viewport.Pin()->IsPlaybackSpeedSelected(i))
 			{
-				Label = FString::Printf( (i == EAnimationPlaybackSpeeds::Quarter) ? 
-										 TEXT("x%.2f") : TEXT("x%.1f"), 
-										 EAnimationPlaybackSpeeds::Values[i]);
+				Label = FText::FromString(FString::Printf(
+					(i == EAnimationPlaybackSpeeds::Quarter) ? TEXT("x%.2f") : TEXT("x%.1f"), 
+					EAnimationPlaybackSpeeds::Values[i]
+					));
 				break;
 			}
 		}
diff --git a/Engine/Source/Editor/Persona/Private/SAnimViewportToolBar.h b/Engine/Source/Editor/Persona/Private/SAnimViewportToolBar.h
--- a/Engine/Source/Editor/Persona/Private/SAnimViewportToolBar.h
+++ b/Engine/Source/Editor/Persona/Private/SAnimViewportToolBar.h
@@ -94,7 +94,7 @@ private:
 	 *
 	 * @return	Label to use for this Menu
 	 */
-	FString GetPlaybackMenuLabel() const;
+	FText GetPlaybackMenuLabel() const;
 
 	/**
 	 * Returns the label for the Viewport type tool bar menu, which changes depending on the current selected type
diff --git a/Engine/Source/Editor/UnrealEd/Private/MaterialGraphNode_Root.cpp b/Engine/Source/Editor/UnrealEd/Private/MaterialGraphNode_Root.cpp
--- a/Engine/Source/Editor/UnrealEd/Private/MaterialGraphNode_Root.cpp
+++ b/Engine/Source/Editor/UnrealEd/Private/MaterialGraphNode_Root.cpp
@@ -21,7 +21,7 @@ UMaterialGraphNode_Root::UMaterialGraphNode_Root(const FObjectInitializer& Objec
 
 FText UMaterialGraphNode_Root::GetNodeTitle(ENodeTitleType::Type TitleType) const
 {
-	return FText::FromString(FMaterialEditorUtilities::GetOriginalObjectName(this->GetGraph()));
+	return FMaterialEditorUtilities::GetOriginalObjectName(this->GetGraph());
 }
 
 FLinearColor UMaterialGraphNode_Root::GetNodeTitleColor() const
diff --git a/Engine/Source/Editor/UnrealEd/Public/SEditorViewportToolBarMenu.h b/Engine/Source/Editor/UnrealEd/Public/SEditorViewportToolBarMenu.h
--- a/Engine/Source/Editor/UnrealEd/Public/SEditorViewportToolBarMenu.h
+++ b/Engine/Source/Editor/UnrealEd/Public/SEditorViewportToolBarMenu.h
@@ -24,7 +24,7 @@ public:
 		/** We need to know about the toolbar we are in */
 		SLATE_ARGUMENT( TSharedPtr<class SViewportToolBar>, ParentToolBar );
 		/** The label to show in the menu */
-		SLATE_TEXT_ATTRIBUTE( Label )
+		SLATE_ATTRIBUTE( FText, Label )
 		/** Optional icon to display next to the label */
 		SLATE_ATTRIBUTE( const FSlateBrush*, LabelIcon )
 		/** The image to show in the menu.  If both the label and image are valid, the button image is used.  Note that if this image is used, the label icon will not be displayed. */
