Converted the property editor to use FText

diff --git a/Engine/Source/Editor/LandscapeEditor/Private/LandscapeEditorDetailCustomization_CopyPaste.cpp b/Engine/Source/Editor/LandscapeEditor/Private/LandscapeEditorDetailCustomization_CopyPaste.cpp
--- a/Engine/Source/Editor/LandscapeEditor/Private/LandscapeEditorDetailCustomization_CopyPaste.cpp
+++ b/Engine/Source/Editor/LandscapeEditor/Private/LandscapeEditorDetailCustomization_CopyPaste.cpp
@@ -38,7 +38,7 @@ void FLandscapeEditorDetailCustomization_CopyPaste::CustomizeDetails(IDetailLayo
 
 	IDetailCategoryBuilder& ToolsCategory = DetailBuilder.EditCategory("Tool Settings");
 
-	ToolsCategory.AddCustomRow("Copy Data to Gizmo")
+	ToolsCategory.AddCustomRow(LOCTEXT("CopyToGizmo", "Copy Data to Gizmo"))
 	[
 		SNew(SButton)
 		.ToolTipText(LOCTEXT("CopyToGizmo.Tooltip", "Copies the data within the gizmo bounds to the gizmo taking into account any masking from selected regions."))
@@ -47,7 +47,7 @@ void FLandscapeEditorDetailCustomization_CopyPaste::CustomizeDetails(IDetailLayo
 		.OnClicked_Static(&FLandscapeEditorDetailCustomization_CopyPaste::OnCopyToGizmoButtonClicked)
 	];
 
-	ToolsCategory.AddCustomRow("Fit Gizmo to Selected Regions")
+	ToolsCategory.AddCustomRow(LOCTEXT("FitGizmoToSelection", "Fit Gizmo to Selected Regions"))
 	[
 		SNew(SButton)
 		.ToolTipText(LOCTEXT("FitGizmoToSelection.Tooltip", "Positions and resizes the gizmo so that it completely encompasses all region selections."))
@@ -56,7 +56,7 @@ void FLandscapeEditorDetailCustomization_CopyPaste::CustomizeDetails(IDetailLayo
 		.OnClicked_Static(&FLandscapeEditorDetailCustomization_CopyPaste::OnFitGizmoToSelectionButtonClicked)
 	];
 
-	ToolsCategory.AddCustomRow("Fit Height Values to Gizmo Size")
+	ToolsCategory.AddCustomRow(LOCTEXT("FitHeightsToGizmo", "Fit Height Values to Gizmo Size"))
 	[
 		SNew(SButton)
 		.ToolTipText(LOCTEXT("FitHeightsToGizmo.Tooltip", "Scales the data in the gizmo to fit the gizmo's Z size"))
@@ -65,7 +65,7 @@ void FLandscapeEditorDetailCustomization_CopyPaste::CustomizeDetails(IDetailLayo
 		.OnClicked_Static(&FLandscapeEditorDetailCustomization_CopyPaste::OnFitHeightsToGizmoButtonClicked)
 	];
 
-	ToolsCategory.AddCustomRow("Clear Gizmo Data")
+	ToolsCategory.AddCustomRow(LOCTEXT("ClearGizmoData", "Clear Gizmo Data"))
 	[
 		SNew(SButton)
 		.ToolTipText(LOCTEXT("ClearGizmoData.Tooltip", "Clears the gizmo of any copied data."))
@@ -75,7 +75,7 @@ void FLandscapeEditorDetailCustomization_CopyPaste::CustomizeDetails(IDetailLayo
 	];
 
 
-	IDetailGroup& GizmoImportExportGroup = ToolsCategory.AddGroup("Gizmo Import / Export", LOCTEXT("ImportExportTitle", "Gizmo Import / Export").ToString(), true);
+	IDetailGroup& GizmoImportExportGroup = ToolsCategory.AddGroup("Gizmo Import / Export", LOCTEXT("ImportExportTitle", "Gizmo Import / Export"), true);
 
 	TSharedRef<IPropertyHandle> PropertyHandle_Heightmap = DetailBuilder.GetProperty(GET_MEMBER_NAME_CHECKED(ULandscapeEditorObject, GizmoHeightmapFilenameString));
 	DetailBuilder.HideProperty(PropertyHandle_Heightmap);
@@ -166,7 +166,7 @@ void FLandscapeEditorDetailCustomization_CopyPaste::CustomizeDetails(IDetailLayo
 	GizmoImportExportGroup.AddPropertyRow(PropertyHandle_ImportLayers);
 
 	GizmoImportExportGroup.AddWidgetRow()
-	.FilterString("")
+	.FilterString(LOCTEXT("GizmoImportExport", "ImportExport"))
 	[
 		SNew(SHorizontalBox)
 		+ SHorizontalBox::Slot()
diff --git a/Engine/Source/Editor/LandscapeEditor/Private/LandscapeEditorDetailCustomization_MiscTools.cpp b/Engine/Source/Editor/LandscapeEditor/Private/LandscapeEditorDetailCustomization_MiscTools.cpp
--- a/Engine/Source/Editor/LandscapeEditor/Private/LandscapeEditorDetailCustomization_MiscTools.cpp
+++ b/Engine/Source/Editor/LandscapeEditor/Private/LandscapeEditorDetailCustomization_MiscTools.cpp
@@ -29,7 +29,7 @@ void FLandscapeEditorDetailCustomization_MiscTools::CustomizeDetails(IDetailLayo
 
 	if (IsBrushSetActive("BrushSet_Component"))
 	{
-		ToolsCategory.AddCustomRow("Clear Component Selection")
+		ToolsCategory.AddCustomRow(LOCTEXT("Component.ClearSelection", "Clear Component Selection"))
 		.Visibility(TAttribute<EVisibility>::Create(TAttribute<EVisibility>::FGetter::CreateStatic(&FLandscapeEditorDetailCustomization_MiscTools::GetClearComponentSelectionVisibility)))
 		[
 			SNew(SButton)
@@ -41,7 +41,7 @@ void FLandscapeEditorDetailCustomization_MiscTools::CustomizeDetails(IDetailLayo
 
 	//if (IsToolActive("Mask"))
 	{
-		ToolsCategory.AddCustomRow("Clear Region Selection")
+		ToolsCategory.AddCustomRow(LOCTEXT("Mask.ClearSelection", "Clear Region Selection"))
 		.Visibility(TAttribute<EVisibility>::Create(TAttribute<EVisibility>::FGetter::CreateStatic(&FLandscapeEditorDetailCustomization_MiscTools::GetClearRegionSelectionVisibility)))
 		[
 			SNew(SButton)
@@ -53,7 +53,7 @@ void FLandscapeEditorDetailCustomization_MiscTools::CustomizeDetails(IDetailLayo
 
 	if (IsToolActive("Splines"))
 	{
-		ToolsCategory.AddCustomRow("Apply Splines")
+		ToolsCategory.AddCustomRow(LOCTEXT("ApplySplinesLabel", "Apply Splines"))
 		[
 			SNew(SHorizontalBox)
 			+ SHorizontalBox::Slot()
@@ -65,7 +65,7 @@ void FLandscapeEditorDetailCustomization_MiscTools::CustomizeDetails(IDetailLayo
 				.Text(LOCTEXT("Spline.ApplySplines", "Deform Landscape to Splines:"))
 			]
 		];
-		ToolsCategory.AddCustomRow("Apply Splines")
+		ToolsCategory.AddCustomRow(LOCTEXT("ApplySplinesLabel", "Apply Splines"))
 		[
 			SNew(SHorizontalBox)
 			+ SHorizontalBox::Slot()
@@ -85,7 +85,7 @@ void FLandscapeEditorDetailCustomization_MiscTools::CustomizeDetails(IDetailLayo
 				.OnClicked_Static(&FLandscapeEditorDetailCustomization_MiscTools::OnApplySelectedSplinesButtonClicked)
 			]
 		];
-		ToolsCategory.AddCustomRow("Use Auto Rotate Control Point")
+		ToolsCategory.AddCustomRow(LOCTEXT("Spline.bUseAutoRotateControlPoint.Selected", "Use Auto Rotate Control Point"))
 		[
 			SNew(SHorizontalBox)
 			+ SHorizontalBox::Slot()
@@ -105,7 +105,7 @@ void FLandscapeEditorDetailCustomization_MiscTools::CustomizeDetails(IDetailLayo
 
 	if (IsToolActive("Ramp"))
 	{
-		ToolsCategory.AddCustomRow("Ramp")
+		ToolsCategory.AddCustomRow(LOCTEXT("RampLabel", "Ramp"))
 		[
 			SNew(SHorizontalBox)
 			+ SHorizontalBox::Slot()
@@ -117,7 +117,7 @@ void FLandscapeEditorDetailCustomization_MiscTools::CustomizeDetails(IDetailLayo
 				.Text(LOCTEXT("Ramp.Hint", "Click to add ramp points, then press \"Add Ramp\"."))
 			]
 		];
-		ToolsCategory.AddCustomRow("Apply Ramp")
+		ToolsCategory.AddCustomRow(LOCTEXT("ApplyRampLabel", "Apply Ramp"))
 		[
 			SNew(SBox)
 			.Padding(FMargin(0, 0, 12, 0)) // Line up with the other properties due to having no reset to default button
diff --git a/Engine/Source/Editor/LandscapeEditor/Private/LandscapeEditorDetailCustomization_NewLandscape.cpp b/Engine/Source/Editor/LandscapeEditor/Private/LandscapeEditorDetailCustomization_NewLandscape.cpp
--- a/Engine/Source/Editor/LandscapeEditor/Private/LandscapeEditorDetailCustomization_NewLandscape.cpp
+++ b/Engine/Source/Editor/LandscapeEditor/Private/LandscapeEditorDetailCustomization_NewLandscape.cpp
@@ -52,7 +52,7 @@ void FLandscapeEditorDetailCustomization_NewLandscape::CustomizeDetails(IDetailL
 
 	IDetailCategoryBuilder& NewLandscapeCategory = DetailBuilder.EditCategory("New Landscape");
 
-	NewLandscapeCategory.AddCustomRow("")
+	NewLandscapeCategory.AddCustomRow(FText::GetEmpty())
 	[
 		SNew(SUniformGridPanel)
 		.SlotPadding(FMargin(10, 2))
@@ -133,7 +133,7 @@ void FLandscapeEditorDetailCustomization_NewLandscape::CustomizeDetails(IDetailL
 		]
 	];
 
-	NewLandscapeCategory.AddCustomRow("Heightmap Resolution")
+	NewLandscapeCategory.AddCustomRow(LOCTEXT("HeightmapResolution", "Heightmap Resolution"))
 	.Visibility(TAttribute<EVisibility>::Create(TAttribute<EVisibility>::FGetter::CreateStatic(&GetVisibilityOnlyInNewLandscapeMode, ENewLandscapePreviewMode::ImportLandscape)))
 	.NameContent()
 	[
@@ -166,7 +166,7 @@ void FLandscapeEditorDetailCustomization_NewLandscape::CustomizeDetails(IDetailL
 	TSharedRef<IPropertyHandle> PropertyHandle_Material = DetailBuilder.GetProperty(GET_MEMBER_NAME_CHECKED(ULandscapeEditorObject, NewLandscape_Material));
 	NewLandscapeCategory.AddProperty(PropertyHandle_Material);
 
-	NewLandscapeCategory.AddCustomRow("Layers")
+	NewLandscapeCategory.AddCustomRow(LOCTEXT("LayersLabel", "Layers"))
 	.Visibility(TAttribute<EVisibility>(this, &FLandscapeEditorDetailCustomization_NewLandscape::GetMaterialTipVisibility))
 	[
 		SNew(SHorizontalBox)
@@ -351,7 +351,7 @@ void FLandscapeEditorDetailCustomization_NewLandscape::CustomizeDetails(IDetailL
 		]
 	];
 
-	NewLandscapeCategory.AddCustomRow("Overall Resolution")
+	NewLandscapeCategory.AddCustomRow(LOCTEXT("Resolution", "Overall Resolution"))
 	.RowTag("LandscapeEditor.OverallResolution")
 	.NameContent()
 	[
@@ -410,7 +410,7 @@ void FLandscapeEditorDetailCustomization_NewLandscape::CustomizeDetails(IDetailL
 		]
 	];
 
-	NewLandscapeCategory.AddCustomRow("Total Components")
+	NewLandscapeCategory.AddCustomRow(LOCTEXT("TotalComponents", "Total Components"))
 	.RowTag("LandscapeEditor.TotalComponents")
 	.NameContent()
 	[
@@ -436,7 +436,7 @@ void FLandscapeEditorDetailCustomization_NewLandscape::CustomizeDetails(IDetailL
 		]
 	];
 
-	NewLandscapeCategory.AddCustomRow("")
+	NewLandscapeCategory.AddCustomRow(FText::GetEmpty())
 	[
 		SNew(SHorizontalBox)
 		+ SHorizontalBox::Slot()
@@ -1377,7 +1377,7 @@ void FLandscapeEditorStructCustomization_FLandscapeImportLayer::CustomizeChildre
 	Result = PropertyHandle_ThumbnailMIC->GetValue(ThumbnailMIC);
 	checkSlow(Result == FPropertyAccess::Success);
 
-	ChildBuilder.AddChildContent(LayerNameText.ToString())
+	ChildBuilder.AddChildContent(LayerNameText)
 	.NameContent()
 	[
 		SNew(SHorizontalBox)
diff --git a/Engine/Source/Editor/LandscapeEditor/Private/LandscapeEditorDetailCustomization_ResizeLandscape.cpp b/Engine/Source/Editor/LandscapeEditor/Private/LandscapeEditorDetailCustomization_ResizeLandscape.cpp
--- a/Engine/Source/Editor/LandscapeEditor/Private/LandscapeEditorDetailCustomization_ResizeLandscape.cpp
+++ b/Engine/Source/Editor/LandscapeEditor/Private/LandscapeEditorDetailCustomization_ResizeLandscape.cpp
@@ -46,7 +46,7 @@ void FLandscapeEditorDetailCustomization_ResizeLandscape::CustomizeDetails(IDeta
 
 	IDetailCategoryBuilder& ResizeLandscapeCategory = DetailBuilder.EditCategory("Change Component Size");
 
-	ResizeLandscapeCategory.AddCustomRow("Original New")
+	ResizeLandscapeCategory.AddCustomRow(LOCTEXT("OriginalNewLabel", "Original New"))
 	//.NameContent()
 	//[
 	//]
@@ -200,7 +200,7 @@ void FLandscapeEditorDetailCustomization_ResizeLandscape::CustomizeDetails(IDeta
 		]
 	];
 
-	ResizeLandscapeCategory.AddCustomRow("Overall Resolution")
+	ResizeLandscapeCategory.AddCustomRow(LOCTEXT("Resolution", "Overall Resolution"))
 	.NameContent()
 	[
 		SNew(SBox)
@@ -240,7 +240,7 @@ void FLandscapeEditorDetailCustomization_ResizeLandscape::CustomizeDetails(IDeta
 		]
 	];
 
-	ResizeLandscapeCategory.AddCustomRow("Total Components")
+	ResizeLandscapeCategory.AddCustomRow(LOCTEXT("TotalComponents", "Total Components"))
 	.NameContent()
 	[
 		SNew(SBox)
@@ -280,7 +280,7 @@ void FLandscapeEditorDetailCustomization_ResizeLandscape::CustomizeDetails(IDeta
 		]
 	];
 
-	ResizeLandscapeCategory.AddCustomRow("")
+	ResizeLandscapeCategory.AddCustomRow(FText::GetEmpty())
 	[
 		SNew(SHorizontalBox)
 		+ SHorizontalBox::Slot()
diff --git a/Engine/Source/Editor/LandscapeEditor/Private/LandscapeEditorDetailCustomization_TargetLayers.cpp b/Engine/Source/Editor/LandscapeEditor/Private/LandscapeEditorDetailCustomization_TargetLayers.cpp
--- a/Engine/Source/Editor/LandscapeEditor/Private/LandscapeEditorDetailCustomization_TargetLayers.cpp
+++ b/Engine/Source/Editor/LandscapeEditor/Private/LandscapeEditorDetailCustomization_TargetLayers.cpp
@@ -156,7 +156,7 @@ void FLandscapeEditorCustomNodeBuilder_TargetLayers::GenerateRow(IDetailChildren
 	
 	if (Target->TargetType != ELandscapeToolTargetType::Weightmap)
 	{
-		ChildrenBuilder.AddChildContent(Target->TargetName.ToString())
+		ChildrenBuilder.AddChildContent(Target->TargetName)
 		[
 			SNew(SLandscapeEditorSelectableBorder)
 			.Padding(0)
@@ -195,7 +195,7 @@ void FLandscapeEditorCustomNodeBuilder_TargetLayers::GenerateRow(IDetailChildren
 	}
 	else
 	{
-		ChildrenBuilder.AddChildContent(Target->TargetName.ToString())
+		ChildrenBuilder.AddChildContent(Target->TargetName)
 		[
 			SNew(SLandscapeEditorSelectableBorder)
 			.Padding(0)
diff --git a/Engine/Source/Editor/LandscapeEditor/Private/LandscapeEditorDetails.cpp b/Engine/Source/Editor/LandscapeEditor/Private/LandscapeEditorDetails.cpp
--- a/Engine/Source/Editor/LandscapeEditor/Private/LandscapeEditorDetails.cpp
+++ b/Engine/Source/Editor/LandscapeEditor/Private/LandscapeEditorDetails.cpp
@@ -73,9 +73,9 @@ void FLandscapeEditorDetails::CustomizeDetails(IDetailLayoutBuilder& DetailBuild
 	static const FLinearColor BorderColor = FLinearColor(0.2f, 0.2f, 0.2f, 0.2f);
 	static const FSlateBrush* BorderStyle = FEditorStyle::GetBrush("DetailsView.GroupSection");
 
-	IDetailCategoryBuilder& LandscapeEditorCategory = DetailBuilder.EditCategory("LandscapeEditor", TEXT(""), ECategoryPriority::TypeSpecific);
+	IDetailCategoryBuilder& LandscapeEditorCategory = DetailBuilder.EditCategory("LandscapeEditor", FText::GetEmpty(), ECategoryPriority::TypeSpecific);
 
-	LandscapeEditorCategory.AddCustomRow("")
+	LandscapeEditorCategory.AddCustomRow(FText::GetEmpty())
 	.Visibility(TAttribute<EVisibility>::Create(TAttribute<EVisibility>::FGetter::CreateStatic(&FLandscapeEditorDetails::GetTargetLandscapeSelectorVisibility)))
 	[
 		SNew(SComboButton)
@@ -123,7 +123,7 @@ void FLandscapeEditorDetails::CustomizeDetails(IDetailLayoutBuilder& DetailBuild
 			);
 	}
 
-	LandscapeEditorCategory.AddCustomRow("")
+	LandscapeEditorCategory.AddCustomRow(FText::GetEmpty())
 	.Visibility(TAttribute<EVisibility>::Create(TAttribute<EVisibility>::FGetter::CreateSP(this, &FLandscapeEditorDetails::GetToolSelectorVisibility)))
 	[
 		ToolBrushSelectorButtons.MakeWidget()
diff --git a/Engine/Source/Editor/LevelEditor/Private/LevelEditorGenericDetails.cpp b/Engine/Source/Editor/LevelEditor/Private/LevelEditorGenericDetails.cpp
--- a/Engine/Source/Editor/LevelEditor/Private/LevelEditorGenericDetails.cpp
+++ b/Engine/Source/Editor/LevelEditor/Private/LevelEditorGenericDetails.cpp
@@ -131,7 +131,7 @@ void FLevelEditorGenericDetails::AddSurfaceDetails( IDetailLayoutBuilder& Detail
 
 	const FSlateFontInfo& FontInfo = IDetailLayoutBuilder::GetDetailFont();
 
-	IDetailCategoryBuilder& MaterialsCategory = DetailBuilder.EditCategory( "Surface Materials", LOCTEXT("BSPSurfaceMaterials", "Surface Materials").ToString() );
+	IDetailCategoryBuilder& MaterialsCategory = DetailBuilder.EditCategory( "Surface Materials", LOCTEXT("BSPSurfaceMaterials", "Surface Materials") );
 	{
 		FMaterialListDelegates MaterialListDelegates;
 		MaterialListDelegates.OnGetMaterials.BindSP(this, &FLevelEditorGenericDetails::GetSelectedSurfaceMaterials);
@@ -143,8 +143,8 @@ void FLevelEditorGenericDetails::AddSurfaceDetails( IDetailLayoutBuilder& Detail
 	}
 
 	// Add a new section for static meshes
-	IDetailCategoryBuilder& BSPCategory = DetailBuilder.EditCategory( "Geometry", LOCTEXT("BSPSurfacesTitle", "Geometry").ToString() );
-	BSPCategory.AddCustomRow( TEXT("") )
+	IDetailCategoryBuilder& BSPCategory = DetailBuilder.EditCategory( "Geometry", LOCTEXT("BSPSurfacesTitle", "Geometry") );
+	BSPCategory.AddCustomRow( FText::GetEmpty() )
 	[
 		SNew( SVerticalBox )
 		+ SVerticalBox::Slot()
@@ -187,7 +187,7 @@ void FLevelEditorGenericDetails::AddSurfaceDetails( IDetailLayoutBuilder& Detail
 		]
 	];
 
-	BSPCategory.AddCustomRow( LOCTEXT("CleanBSPMaterials", "Clean Geometry Materials").ToString(), true )
+	BSPCategory.AddCustomRow( LOCTEXT("CleanBSPMaterials", "Clean Geometry Materials"), true )
 	[
 		SNew(SVerticalBox)
 		+ SVerticalBox::Slot()
@@ -206,8 +206,8 @@ void FLevelEditorGenericDetails::AddSurfaceDetails( IDetailLayoutBuilder& Detail
 		]
 	];
 
-	IDetailCategoryBuilder& SurfaceCategory = DetailBuilder.EditCategory( "Surface Properties", LOCTEXT("BSPSurfaceProperties", "Surface Properties").ToString() );
-	SurfaceCategory.AddCustomRow( LOCTEXT("BSPSurfaceProperties", "Surface Properties").ToString() )
+	IDetailCategoryBuilder& SurfaceCategory = DetailBuilder.EditCategory( "Surface Properties", LOCTEXT("BSPSurfaceProperties", "Surface Properties") );
+	SurfaceCategory.AddCustomRow( LOCTEXT("BSPSurfaceProperties", "Surface Properties") )
 	[
 		SNew(SVerticalBox)
 		+SVerticalBox::Slot()
diff --git a/Engine/Source/Editor/MaterialEditor/Private/MaterialEditorInstanceDetailCustomization.cpp b/Engine/Source/Editor/MaterialEditor/Private/MaterialEditorInstanceDetailCustomization.cpp
--- a/Engine/Source/Editor/MaterialEditor/Private/MaterialEditorInstanceDetailCustomization.cpp
+++ b/Engine/Source/Editor/MaterialEditor/Private/MaterialEditorInstanceDetailCustomization.cpp
@@ -109,7 +109,7 @@ void FMaterialInstanceParameterDetails::CustomizeDetails(IDetailLayoutBuilder& D
 
 	// Create a new category for a custom layout for the MIC parameters at the very top
 	FName GroupsCategoryName = TEXT("ParameterGroups");
-	IDetailCategoryBuilder& GroupsCategory = DetailLayout.EditCategory(GroupsCategoryName, LOCTEXT("MICParamGroupsTitle", "Parameter Groups").ToString());
+	IDetailCategoryBuilder& GroupsCategory = DetailLayout.EditCategory(GroupsCategoryName, LOCTEXT("MICParamGroupsTitle", "Parameter Groups"));
 	TSharedRef<IPropertyHandle> ParameterGroupsProperty = DetailLayout.GetProperty("ParameterGroups");
 
 	CreateGroupsWidget(ParameterGroupsProperty, GroupsCategory);
@@ -123,7 +123,7 @@ void FMaterialInstanceParameterDetails::CreateGroupsWidget(TSharedRef<IPropertyH
 	{
 		FEditorParameterGroup& ParameterGroup = MaterialEditorInstance->ParameterGroups[GroupIdx];
 
-		IDetailGroup& DetailGroup = GroupsCategory.AddGroup( ParameterGroup.GroupName, ParameterGroup.GroupName.ToString(), false, true );
+		IDetailGroup& DetailGroup = GroupsCategory.AddGroup( ParameterGroup.GroupName, FText::FromName(ParameterGroup.GroupName), false, true );
 
 		CreateSingleGroupWidget( ParameterGroup, ParameterGroupsProperty->GetChildHandle(GroupIdx), DetailGroup );
 	}
@@ -175,7 +175,7 @@ void FMaterialInstanceParameterDetails::CreateParameterValueWidget(UDEditorParam
 		IDetailPropertyRow& PropertyRow = DetailGroup.AddPropertyRow( ParameterValueProperty.ToSharedRef() );
 
 		PropertyRow
-		.DisplayName( Parameter->ParameterName.ToString() )
+		.DisplayName( FText::FromName(Parameter->ParameterName) )
 		.ToolTip( GetParameterExpressionDescription(Parameter) )
 		.EditCondition( IsParamEnabled, FOnBooleanValueChanged::CreateSP( this, &FMaterialInstanceParameterDetails::OnOverrideParameter, Parameter ) )
 		.Visibility( TAttribute<EVisibility>::Create(TAttribute<EVisibility>::FGetter::CreateSP(this, &FMaterialInstanceParameterDetails::ShouldShowExpression, Parameter)) )
@@ -202,7 +202,7 @@ void FMaterialInstanceParameterDetails::CreateMaskParameterValueWidget(UDEditorP
 		PropertyRow.OverrideResetToDefault( true, FSimpleDelegate::CreateSP( this, &FMaterialInstanceParameterDetails::ResetToDefault, Parameter ) );
 		PropertyRow.Visibility( TAttribute<EVisibility>::Create(TAttribute<EVisibility>::FGetter::CreateSP(this, &FMaterialInstanceParameterDetails::ShouldShowExpression, Parameter)) );
 
-		const FString ParameterName = Parameter->ParameterName.ToString(); 
+		const FText ParameterName = FText::FromName(Parameter->ParameterName); 
 
 		FDetailWidgetRow& CustomWidget = PropertyRow.CustomWidget();
 		CustomWidget
@@ -226,7 +226,7 @@ void FMaterialInstanceParameterDetails::CreateMaskParameterValueWidget(UDEditorP
 				.HAlign(HAlign_Left)
 				.AutoWidth()
 				[
-					RMaskProperty->CreatePropertyNameWidget( TEXT(""), TEXT( "" ), false )
+					RMaskProperty->CreatePropertyNameWidget( FText::GetEmpty(), FText::GetEmpty(), false )
 				]
 				+SHorizontalBox::Slot()
 				.HAlign(HAlign_Left)
@@ -239,7 +239,7 @@ void FMaterialInstanceParameterDetails::CreateMaskParameterValueWidget(UDEditorP
 				.Padding( FMargin( 10.0f, 0.0f, 0.0f, 0.0f ) )
 				.AutoWidth()
 				[
-					GMaskProperty->CreatePropertyNameWidget( TEXT(""), TEXT( "" ), false )
+					GMaskProperty->CreatePropertyNameWidget( FText::GetEmpty(), FText::GetEmpty(), false )
 				]
 				+SHorizontalBox::Slot()
 				.HAlign(HAlign_Left)
@@ -252,7 +252,7 @@ void FMaterialInstanceParameterDetails::CreateMaskParameterValueWidget(UDEditorP
 				.Padding( FMargin( 10.0f, 0.0f, 0.0f, 0.0f ) )
 				.AutoWidth()
 				[
-					BMaskProperty->CreatePropertyNameWidget( TEXT(""), TEXT( "" ), false )
+					BMaskProperty->CreatePropertyNameWidget( FText::GetEmpty(), FText::GetEmpty(), false )
 				]
 				+SHorizontalBox::Slot()
 				.HAlign(HAlign_Left)
@@ -265,7 +265,7 @@ void FMaterialInstanceParameterDetails::CreateMaskParameterValueWidget(UDEditorP
 				.Padding( FMargin( 10.0f, 0.0f, 0.0f, 0.0f ) )
 				.AutoWidth()
 				[
-					AMaskProperty->CreatePropertyNameWidget( TEXT(""), TEXT( "" ), false )
+					AMaskProperty->CreatePropertyNameWidget( FText::GetEmpty(), FText::GetEmpty(), false )
 				]
 				+SHorizontalBox::Slot()
 				.HAlign(HAlign_Left)
@@ -328,7 +328,7 @@ bool FMaterialInstanceParameterDetails::OnShouldSetAsset(const FAssetData& Asset
 	return true;
 }
 
-FString FMaterialInstanceParameterDetails::GetParameterExpressionDescription(UDEditorParameterValue* Parameter) const
+FText FMaterialInstanceParameterDetails::GetParameterExpressionDescription(UDEditorParameterValue* Parameter) const
 {
 	UMaterial* BaseMaterial = MaterialEditorInstance->SourceInstance->GetMaterial();
 	if ( BaseMaterial )
@@ -336,23 +336,23 @@ FString FMaterialInstanceParameterDetails::GetParameterExpressionDescription(UDE
 		UMaterialExpressionParameter* Expression = BaseMaterial->FindExpressionByGUID<UMaterialExpressionParameter>( Parameter->ExpressionId );
 		if ( Expression )
 		{
-			return Expression->Desc;
+			return FText::FromString(Expression->Desc);
 		}
 
 		UMaterialExpressionTextureSampleParameter* TextureExpression = BaseMaterial->FindExpressionByGUID<UMaterialExpressionTextureSampleParameter>( Parameter->ExpressionId );
 		if ( TextureExpression )
 		{
-			return TextureExpression->Desc;
+			return FText::FromString(TextureExpression->Desc);
 		}
 
 		UMaterialExpressionFontSampleParameter* FontExpression = BaseMaterial->FindExpressionByGUID<UMaterialExpressionFontSampleParameter>( Parameter->ExpressionId );
 		if ( FontExpression )
 		{
-			return FontExpression->Desc;
+			return FText::FromString(FontExpression->Desc);
 		}
 	}
 
-	return "";
+	return FText::GetEmpty();
 }
 
 void FMaterialInstanceParameterDetails::ResetToDefault( class UDEditorParameterValue* Parameter )
@@ -449,7 +449,7 @@ void FMaterialInstanceParameterDetails::CreateBasePropertyOverrideWidgets(IDetai
 	IDetailCategoryBuilder& DetailCategory = DetailLayout.EditCategory(NAME_None);
 	
 	static FName GroupName(TEXT("BasePropertyOverrideGroup"));
-	IDetailGroup& BasePropertyOverrideGroup = DetailCategory.AddGroup(GroupName, LOCTEXT("BasePropertyOverrideGroup", "Material Property Overrides").ToString(), false, false);
+	IDetailGroup& BasePropertyOverrideGroup = DetailCategory.AddGroup(GroupName, LOCTEXT("BasePropertyOverrideGroup", "Material Property Overrides"), false, false);
 
 	TAttribute<bool> IsOverrideOpacityClipMaskValueEnabled = TAttribute<bool>::Create(TAttribute<bool>::FGetter::CreateSP(this, &FMaterialInstanceParameterDetails::OverrideOpacityClipMaskValueEnabled));
 	TAttribute<bool> IsOverrideBlendModeEnabled = TAttribute<bool>::Create(TAttribute<bool>::FGetter::CreateSP(this, &FMaterialInstanceParameterDetails::OverrideBlendModeEnabled));
diff --git a/Engine/Source/Editor/MaterialEditor/Private/MaterialEditorInstanceDetailCustomization.h b/Engine/Source/Editor/MaterialEditor/Private/MaterialEditorInstanceDetailCustomization.h
--- a/Engine/Source/Editor/MaterialEditor/Private/MaterialEditorInstanceDetailCustomization.h
+++ b/Engine/Source/Editor/MaterialEditor/Private/MaterialEditorInstanceDetailCustomization.h
@@ -45,7 +45,7 @@ private:
 	bool IsOverriddenExpression(class UDEditorParameterValue* Parameter);
 
 	/** Gets the expression description of this parameter from the the base material */
-	FString GetParameterExpressionDescription(class UDEditorParameterValue* Parameter) const;
+	FText GetParameterExpressionDescription(class UDEditorParameterValue* Parameter) const;
 
 	/**
 	 * Called when a parameter is overridden;
diff --git a/Engine/Source/Editor/Persona/Private/AnimGraphNodeDetails.cpp b/Engine/Source/Editor/Persona/Private/AnimGraphNodeDetails.cpp
--- a/Engine/Source/Editor/Persona/Private/AnimGraphNodeDetails.cpp
+++ b/Engine/Source/Editor/Persona/Private/AnimGraphNodeDetails.cpp
@@ -242,7 +242,7 @@ void FAnimGraphNodeDetails::AbortDisplayOfAllNodes(TArray< TWeakObjectPtr<UObjec
 {
 	// Display a warning message
 	IDetailCategoryBuilder& ErrorCategory = DetailBuilder.EditCategory("Animation Nodes");
-	ErrorCategory.AddCustomRow( LOCTEXT("ErrorRow", "Error").ToString() )
+	ErrorCategory.AddCustomRow( LOCTEXT("ErrorRow", "Error") )
 	[
 		SNew(STextBlock)
 		.Text(LOCTEXT("MultiSelectNotSupported", "Multiple nodes selected"))
@@ -468,7 +468,7 @@ void FBoneReferenceCustomization::CustomizeHeader( TSharedRef<IPropertyHandle> S
 		[
 			SNew(SBoneSelectionWidget)
 			.Skeleton(TargetSkeleton)
-			.Tooltip(FText::FromString(StructPropertyHandle->GetToolTipText()))
+			.Tooltip(StructPropertyHandle->GetToolTipText())
 			.OnBoneSelectionChanged(this, &FBoneReferenceCustomization::OnBoneSelectionChanged)
 			.OnGetSelectedBone(this, &FBoneReferenceCustomization::GetSelectedBone)
 		];
@@ -545,7 +545,7 @@ void FAnimGraphParentPlayerDetails::CustomizeDetails(class IDetailLayoutBuilder&
 		(*ExistingGraphEntry)->Children.Add(NodeEntry);
 	}
 
-	FDetailWidgetRow& Row = Category.AddCustomRow("");
+	FDetailWidgetRow& Row = Category.AddCustomRow(FText::GetEmpty());
 	TSharedRef<STreeView<TSharedPtr<FPlayerTreeViewEntry>>> TreeView = SNew(STreeView<TSharedPtr<FPlayerTreeViewEntry>>)
 		.SelectionMode(ESelectionMode::None)
 		.OnGenerateRow(this, &FAnimGraphParentPlayerDetails::OnGenerateRow)
diff --git a/Engine/Source/Editor/Persona/Private/AnimNotifyDetails.cpp b/Engine/Source/Editor/Persona/Private/AnimNotifyDetails.cpp
--- a/Engine/Source/Editor/Persona/Private/AnimNotifyDetails.cpp
+++ b/Engine/Source/Editor/Persona/Private/AnimNotifyDetails.cpp
@@ -39,7 +39,7 @@ void FAnimNotifyDetails::CustomizeDetails( IDetailLayoutBuilder& DetailBuilder )
 	// Don't want to edit the notify name here.
 	DetailBuilder.HideProperty(TEXT("Event.NotifyName"));
 
-	IDetailCategoryBuilder& AnimNotifyCategory = DetailBuilder.EditCategory(TEXT("AnimNotify"), TEXT(""), ECategoryPriority::TypeSpecific);
+	IDetailCategoryBuilder& AnimNotifyCategory = DetailBuilder.EditCategory(TEXT("AnimNotify"), FText::GetEmpty(), ECategoryPriority::TypeSpecific);
 
 	// Check existence of notify, get rid of the property if not set
 	if(!NotifyPtr)
@@ -201,7 +201,7 @@ void FAnimNotifyDetails::ClearInstancedSelectionDropDown(IDetailCategoryBuilder&
 	.CustomWidget(bShowChildren)
 	.NameContent()
 	[
-		PropHandle->CreatePropertyNameWidget(TEXT(""), TEXT(""), false)
+		PropHandle->CreatePropertyNameWidget(FText::GetEmpty(), FText::GetEmpty(), false)
 	]
 	.ValueContent()
 	[
@@ -236,7 +236,7 @@ void FAnimNotifyDetails::CustomizeLinkProperties(IDetailLayoutBuilder& Builder,
 						.CustomWidget()
 						.NameContent()
 						[
-							ChildHandle->CreatePropertyNameWidget(NSLOCTEXT("NotifyDetails", "SlotIndexName", "Slot").ToString())
+							ChildHandle->CreatePropertyNameWidget(NSLOCTEXT("NotifyDetails", "SlotIndexName", "Slot"))
 						]
 						.ValueContent()
 						[
diff --git a/Engine/Source/Editor/Persona/Private/AssetSearchBoxUtilPersona.cpp b/Engine/Source/Editor/Persona/Private/AssetSearchBoxUtilPersona.cpp
--- a/Engine/Source/Editor/Persona/Private/AssetSearchBoxUtilPersona.cpp
+++ b/Engine/Source/Editor/Persona/Private/AssetSearchBoxUtilPersona.cpp
@@ -16,7 +16,7 @@ void SAssetSearchBoxForBones::Construct( const FArguments& InArgs, const class U
 	check(Outer);
 
 	// get the currently chosen bone/socket, if any
-	const FString PropertyName = BoneNameProperty->GetPropertyDisplayName();
+	const FText PropertyName = BoneNameProperty->GetPropertyDisplayName();
 	FString CurValue;
 	BoneNameProperty->GetValue(CurValue);
 	if( CurValue == FString("None") )
@@ -75,7 +75,7 @@ void SAssetSearchBoxForCurves::Construct(const FArguments& InArgs, const class U
 	check(InSkeleton);
 
 	// get the currently chosen curve, if any
-	const FString PropertyName = CurveNameProperty->GetPropertyDisplayName();
+	const FText PropertyName = CurveNameProperty->GetPropertyDisplayName();
 	FString CurValue;
 	CurveNameProperty->GetValue(CurValue);
 	if (CurValue == FString("None"))
diff --git a/Engine/Source/Editor/Persona/Private/Customization/AnimGraphNodeSlotDetails.cpp b/Engine/Source/Editor/Persona/Private/Customization/AnimGraphNodeSlotDetails.cpp
--- a/Engine/Source/Editor/Persona/Private/Customization/AnimGraphNodeSlotDetails.cpp
+++ b/Engine/Source/Editor/Persona/Private/Customization/AnimGraphNodeSlotDetails.cpp
@@ -67,7 +67,7 @@ void FAnimGraphNodeSlotDetails::CustomizeDetails(class IDetailLayoutBuilder& Det
 
 	// create combo box
 	IDetailCategoryBuilder& SlotNameGroup = DetailBuilder.EditCategory(TEXT("Settings"));
-	SlotNameGroup.AddCustomRow(LOCTEXT("SlotNameTitleLabel", "Slot Name").ToString())
+	SlotNameGroup.AddCustomRow(LOCTEXT("SlotNameTitleLabel", "Slot Name"))
 	.NameContent()
 	[
 		SlotNodePropertyNameWidget
diff --git a/Engine/Source/Editor/Persona/Private/PersonaMeshDetails.cpp b/Engine/Source/Editor/Persona/Private/PersonaMeshDetails.cpp
--- a/Engine/Source/Editor/Persona/Private/PersonaMeshDetails.cpp
+++ b/Engine/Source/Editor/Persona/Private/PersonaMeshDetails.cpp
@@ -96,7 +96,7 @@ BEGIN_SLATE_FUNCTION_BUILD_OPTIMIZATION
 void FSkelMeshReductionSettingsLayout::GenerateChildContent(IDetailChildrenBuilder& ChildrenBuilder)
 {
 	{
-		ChildrenBuilder.AddChildContent(LOCTEXT("PercentTriangles", "Percent Triangles").ToString())
+		ChildrenBuilder.AddChildContent(LOCTEXT("PercentTriangles", "Percent Triangles"))
 			.NameContent()
 			[
 				SNew(STextBlock)
@@ -116,7 +116,7 @@ void FSkelMeshReductionSettingsLayout::GenerateChildContent(IDetailChildrenBuild
 	}
 
 	{
-		ChildrenBuilder.AddChildContent(LOCTEXT("MaxDeviation", "Max Deviation").ToString())
+		ChildrenBuilder.AddChildContent(LOCTEXT("MaxDeviation", "Max Deviation"))
 			.NameContent()
 			[
 				SNew(STextBlock)
@@ -136,7 +136,7 @@ void FSkelMeshReductionSettingsLayout::GenerateChildContent(IDetailChildrenBuild
 	}
 
 	{
-		ChildrenBuilder.AddChildContent(LOCTEXT("Silhouette_MeshSimplification", "Silhouette").ToString())
+		ChildrenBuilder.AddChildContent(LOCTEXT("Silhouette_MeshSimplification", "Silhouette"))
 			.NameContent()
 			[
 				SNew(STextBlock)
@@ -156,7 +156,7 @@ void FSkelMeshReductionSettingsLayout::GenerateChildContent(IDetailChildrenBuild
 	}
 
 	{
-		ChildrenBuilder.AddChildContent(LOCTEXT("Texture_MeshSimplification", "Texture").ToString())
+		ChildrenBuilder.AddChildContent(LOCTEXT("Texture_MeshSimplification", "Texture"))
 			.NameContent()
 			[
 				SNew(STextBlock)
@@ -176,7 +176,7 @@ void FSkelMeshReductionSettingsLayout::GenerateChildContent(IDetailChildrenBuild
 	}
 
 	{
-		ChildrenBuilder.AddChildContent(LOCTEXT("Shading_MeshSimplification", "Shading").ToString())
+		ChildrenBuilder.AddChildContent(LOCTEXT("Shading_MeshSimplification", "Shading"))
 			.NameContent()
 			[
 				SNew(STextBlock)
@@ -196,7 +196,7 @@ void FSkelMeshReductionSettingsLayout::GenerateChildContent(IDetailChildrenBuild
 	}
 
 	{
-		ChildrenBuilder.AddChildContent(LOCTEXT("Skinning_MeshSimplification", "Skinning").ToString())
+		ChildrenBuilder.AddChildContent(LOCTEXT("Skinning_MeshSimplification", "Skinning"))
 			.NameContent()
 			[
 				SNew(STextBlock)
@@ -216,7 +216,7 @@ void FSkelMeshReductionSettingsLayout::GenerateChildContent(IDetailChildrenBuild
 	}
 
 	{
-		ChildrenBuilder.AddChildContent(LOCTEXT("WeldingThreshold", "Welding Threshold").ToString())
+		ChildrenBuilder.AddChildContent(LOCTEXT("WeldingThreshold", "Welding Threshold"))
 			.NameContent()
 			[
 				SNew(STextBlock)
@@ -236,7 +236,7 @@ void FSkelMeshReductionSettingsLayout::GenerateChildContent(IDetailChildrenBuild
 	}
 
 	{
-		ChildrenBuilder.AddChildContent(LOCTEXT("RecomputeNormals", "Recompute Normals").ToString())
+		ChildrenBuilder.AddChildContent(LOCTEXT("RecomputeNormals", "Recompute Normals"))
 			.NameContent()
 			[
 				SNew(STextBlock)
@@ -253,7 +253,7 @@ void FSkelMeshReductionSettingsLayout::GenerateChildContent(IDetailChildrenBuild
 	}
 
 	{
-		ChildrenBuilder.AddChildContent(LOCTEXT("HardEdgeAngle", "Hard Edge Angle").ToString())
+		ChildrenBuilder.AddChildContent(LOCTEXT("HardEdgeAngle", "Hard Edge Angle"))
 			.NameContent()
 			[
 				SNew(STextBlock)
@@ -280,7 +280,7 @@ void FSkelMeshReductionSettingsLayout::GenerateChildContent(IDetailChildrenBuild
 	}
 
 	{
-		ChildrenBuilder.AddChildContent(LOCTEXT("MaxBonesPerVertex", "Max Bones Per Vertex").ToString())
+		ChildrenBuilder.AddChildContent(LOCTEXT("MaxBonesPerVertex", "Max Bones Per Vertex"))
 			.NameContent()
 			[
 				SNew(STextBlock)
@@ -300,7 +300,7 @@ void FSkelMeshReductionSettingsLayout::GenerateChildContent(IDetailChildrenBuild
 	}
 
 	{
-		ChildrenBuilder.AddChildContent(LOCTEXT("ApplyChanges", "Apply Changes").ToString())
+		ChildrenBuilder.AddChildContent(LOCTEXT("ApplyChanges", "Apply Changes"))
 			.ValueContent()
 			.HAlign(HAlign_Left)
 			[
@@ -494,7 +494,7 @@ void FPersonaMeshDetails::AddLODLevelCategories(IDetailLayoutBuilder& DetailLayo
 			FString CategoryName = FString(TEXT("LOD"));
 			CategoryName.AppendInt(LODIndex);
 
-			FString LODLevelString = FText::Format(LOCTEXT("LODLevel", "LOD{0}"), FText::AsNumber(LODIndex)).ToString();
+			FText LODLevelString = FText::Format(LOCTEXT("LODLevel", "LOD{0}"), FText::AsNumber(LODIndex));
 
 			IDetailCategoryBuilder& LODCategory = DetailLayout.EditCategory(*CategoryName, LODLevelString, ECategoryPriority::Important);
 
@@ -519,7 +519,7 @@ void FPersonaMeshDetails::AddLODLevelCategories(IDetailLayoutBuilder& DetailLayo
 
 				// changing property name to "LOD Info" because it shows only array index
 				TSharedPtr<IPropertyHandle> LODInfoChild = LODInfoProperty->GetChildHandle(LODIndex);
-				LODInfoChild->CreatePropertyNameWidget(FString(TEXT("LOD Info")));
+				LODInfoChild->CreatePropertyNameWidget(LOCTEXT("LODInfoLabel", "LOD Info"));
 				LODCategory.AddProperty(LODInfoChild);
 			}
 
@@ -538,9 +538,9 @@ void FPersonaMeshDetails::CustomizeLODSettingsCategories(IDetailLayoutBuilder& D
 
 	UpdateLODNames();
 
-	IDetailCategoryBuilder& LODSettingsCategory = DetailLayout.EditCategory("LodSettings", LOCTEXT("LodSettingsCategory", "LOD Settings").ToString(), ECategoryPriority::TypeSpecific);
+	IDetailCategoryBuilder& LODSettingsCategory = DetailLayout.EditCategory("LodSettings", LOCTEXT("LodSettingsCategory", "LOD Settings"), ECategoryPriority::TypeSpecific);
 
-	LODSettingsCategory.AddCustomRow(LOCTEXT("LODImport", "LOD Import").ToString())
+	LODSettingsCategory.AddCustomRow(LOCTEXT("LODImport", "LOD Import"))
 		.NameContent()
 		[
 			SNew(STextBlock)
@@ -558,7 +558,7 @@ void FPersonaMeshDetails::CustomizeLODSettingsCategories(IDetailLayoutBuilder& D
 
 	// Add Number of LODs slider.
 	const int32 MinAllowedLOD = 1;
-	LODSettingsCategory.AddCustomRow(LOCTEXT("NumberOfLODs", "Number of LODs").ToString())
+	LODSettingsCategory.AddCustomRow(LOCTEXT("NumberOfLODs", "Number of LODs"))
 		.NameContent()
 		[
 			SNew(STextBlock)
@@ -578,7 +578,7 @@ void FPersonaMeshDetails::CustomizeLODSettingsCategories(IDetailLayoutBuilder& D
 			.IsEnabled(IsAutoMeshReductionAvailable())
 		];
 
-	LODSettingsCategory.AddCustomRow(LOCTEXT("ApplyChanges", "Apply Changes").ToString())
+	LODSettingsCategory.AddCustomRow(LOCTEXT("ApplyChanges", "Apply Changes"))
 		.ValueContent()
 		.HAlign(HAlign_Left)
 		[
@@ -824,12 +824,12 @@ void FPersonaMeshDetails::CustomizeDetails( IDetailLayoutBuilder& DetailLayout )
 	CustomizeLODSettingsCategories(DetailLayout);
 
 #if WITH_APEX_CLOTHING
-	IDetailCategoryBuilder& ClothingCategory = DetailLayout.EditCategory("Clothing", TEXT(""), ECategoryPriority::TypeSpecific);
+	IDetailCategoryBuilder& ClothingCategory = DetailLayout.EditCategory("Clothing", FText::GetEmpty(), ECategoryPriority::TypeSpecific);
 	CustomizeClothingProperties(DetailLayout,ClothingCategory);
 #endif// #if WITH_APEX_CLOTHING
 
-	IDetailCategoryBuilder& AdditionalMeshCategory = DetailLayout.EditCategory("AdditionalBodyPart", LOCTEXT("AdditionalMeshesCollapsable", "Additional Body Part").ToString(), ECategoryPriority::TypeSpecific);
-	AdditionalMeshCategory.AddCustomRow("")
+	IDetailCategoryBuilder& AdditionalMeshCategory = DetailLayout.EditCategory("AdditionalBodyPart", LOCTEXT("AdditionalMeshesCollapsable", "Additional Body Part"), ECategoryPriority::TypeSpecific);
+	AdditionalMeshCategory.AddCustomRow(FText::GetEmpty())
 	[
 		SNew(SAdditionalMeshesEditor, PersonaPtr)
 	];
@@ -865,7 +865,7 @@ void FPersonaMeshDetails::HideUnnecessaryProperties(IDetailLayoutBuilder& Detail
 	DetailLayout.HideProperty(MaterialsProperty);
 
 	// hide all properties in Mirroring category to hide Mirroring category itself
-	IDetailCategoryBuilder& MirroringCategory = DetailLayout.EditCategory("Mirroring", TEXT(""), ECategoryPriority::Default);
+	IDetailCategoryBuilder& MirroringCategory = DetailLayout.EditCategory("Mirroring", FText::GetEmpty(), ECategoryPriority::Default);
 	TArray<TSharedRef<IPropertyHandle>> MirroringProperties;
 	MirroringCategory.GetDefaultProperties(MirroringProperties);
 	for (int32 MirrorPropertyIdx = 0; MirrorPropertyIdx < MirroringProperties.Num(); MirrorPropertyIdx++)
@@ -1282,7 +1282,7 @@ void FPersonaMeshDetails::CustomizeClothingProperties(IDetailLayoutBuilder& Deta
 	}
 
 	// Button to add a new clothing file
-	ClothingFilesCategory.AddCustomRow( LOCTEXT("AddAPEXClothingFileFilterString", "Add APEX clothing file").ToString())
+	ClothingFilesCategory.AddCustomRow( LOCTEXT("AddAPEXClothingFileFilterString", "Add APEX clothing file"))
 	[
 		SNew(SHorizontalBox)
 		 
@@ -1309,7 +1309,7 @@ void FPersonaMeshDetails::OnGenerateElementForClothingAsset( TSharedRef<IPropert
 	const FSlateFontInfo DetailFontInfo = IDetailLayoutBuilder::GetDetailFont();
 
 	// Remove and reimport asset buttons
-	ChildrenBuilder.AddChildContent( TEXT("") ) 
+	ChildrenBuilder.AddChildContent( FText::GetEmpty() ) 
 	[
 		SNew(SHorizontalBox)
 		+ SHorizontalBox::Slot()
diff --git a/Engine/Source/Editor/Persona/Private/SkeletalMeshSocketDetails.cpp b/Engine/Source/Editor/Persona/Private/SkeletalMeshSocketDetails.cpp
--- a/Engine/Source/Editor/Persona/Private/SkeletalMeshSocketDetails.cpp
+++ b/Engine/Source/Editor/Persona/Private/SkeletalMeshSocketDetails.cpp
@@ -22,7 +22,7 @@ void FSkeletalMeshSocketDetails::CustomizeDetails( IDetailLayoutBuilder& DetailB
 {
 	TargetSocket = NULL;
 
-	IDetailCategoryBuilder& SocketCategory = DetailBuilder.EditCategory( TEXT("SkeletalMeshSocket"), TEXT(""), ECategoryPriority::TypeSpecific );
+	IDetailCategoryBuilder& SocketCategory = DetailBuilder.EditCategory( TEXT("SkeletalMeshSocket"), FText::GetEmpty(), ECategoryPriority::TypeSpecific );
 
 	SocketNameProperty = DetailBuilder.GetProperty( TEXT("SocketName") );
 	if( SocketNameProperty.IsValid() && SocketNameProperty->GetProperty() )
@@ -39,7 +39,7 @@ void FSkeletalMeshSocketDetails::CustomizeDetails( IDetailLayoutBuilder& DetailB
 		}
 
 		// get the currently chosen bone/socket, if any
-		const FString PropertyName = SocketNameProperty->GetPropertyDisplayName();
+		const FText PropertyName = SocketNameProperty->GetPropertyDisplayName();
 		FString CurValue;
 		SocketNameProperty->GetValue(CurValue);
 		if( CurValue == FString("None") )
diff --git a/Engine/Source/Editor/PropertyEditor/Private/CategoryPropertyNode.cpp b/Engine/Source/Editor/PropertyEditor/Private/CategoryPropertyNode.cpp
--- a/Engine/Source/Editor/PropertyEditor/Private/CategoryPropertyNode.cpp
+++ b/Engine/Source/Editor/PropertyEditor/Private/CategoryPropertyNode.cpp
@@ -23,24 +23,22 @@ bool FCategoryPropertyNode::IsSubcategory() const
 	return GetParentNode() != NULL && const_cast<FPropertyNode*>( GetParentNode() )->AsCategoryNode() != NULL;
 }
 
-FString FCategoryPropertyNode::GetDisplayName() const 
+FText FCategoryPropertyNode::GetDisplayName() const 
 {
 	FString SubcategoryName = GetSubcategoryName();
-	//if in "readable display name mode" return that
 	if ( FPropertySettings::Get().ShowFriendlyPropertyNames() )
 	{
-		FString DisplaySubcategoryName = FName::NameToDisplayString( SubcategoryName, false );
-
 		//if there is a localized version, return that
 		FText LocalizedName;
 		if ( FText::FindText( TEXT("UObjectCategories"), *SubcategoryName, /*OUT*/LocalizedName ) )
 		{
-			return LocalizedName.ToString();
+			return LocalizedName;
 		}
 
-		return DisplaySubcategoryName;
+		//if in "readable display name mode" return that
+		return FText::FromString( FName::NameToDisplayString( SubcategoryName, false ) );
 	}
-	return SubcategoryName;
+	return FText::FromString( SubcategoryName );
 }	
 /**
  * Overridden function for special setup
diff --git a/Engine/Source/Editor/PropertyEditor/Private/CategoryPropertyNode.h b/Engine/Source/Editor/PropertyEditor/Private/CategoryPropertyNode.h
--- a/Engine/Source/Editor/PropertyEditor/Private/CategoryPropertyNode.h
+++ b/Engine/Source/Editor/PropertyEditor/Private/CategoryPropertyNode.h
@@ -26,7 +26,7 @@ public:
 	 */
 	bool IsSubcategory() const;
 
-	virtual FString GetDisplayName() const override;
+	virtual FText GetDisplayName() const override;
 
 protected:
 	/**
diff --git a/Engine/Source/Editor/PropertyEditor/Private/CustomChildBuilder.cpp b/Engine/Source/Editor/PropertyEditor/Private/CustomChildBuilder.cpp
--- a/Engine/Source/Editor/PropertyEditor/Private/CustomChildBuilder.cpp
+++ b/Engine/Source/Editor/PropertyEditor/Private/CustomChildBuilder.cpp
@@ -15,7 +15,7 @@ IDetailChildrenBuilder& FCustomChildrenBuilder::AddChildCustomBuilder( TSharedRe
 	return *this;
 }
 
-IDetailGroup& FCustomChildrenBuilder::AddChildGroup( FName GroupName, const FString& LocalizedDisplayName )
+IDetailGroup& FCustomChildrenBuilder::AddChildGroup( FName GroupName, const FText& LocalizedDisplayName )
 {
 	FDetailLayoutCustomization NewCustomization;
 	NewCustomization.DetailGroup = MakeShareable( new FDetailGroup( GroupName, ParentCategory.Pin().ToSharedRef(), LocalizedDisplayName ) );
@@ -25,7 +25,7 @@ IDetailGroup& FCustomChildrenBuilder::AddChildGroup( FName GroupName, const FStr
 	return *NewCustomization.DetailGroup;
 }
 
-FDetailWidgetRow& FCustomChildrenBuilder::AddChildContent( const FString& SearchString )
+FDetailWidgetRow& FCustomChildrenBuilder::AddChildContent( const FText& SearchString )
 {
 	TSharedRef<FDetailWidgetRow> NewRow = MakeShareable( new FDetailWidgetRow );
 	FDetailLayoutCustomization NewCustomization;
diff --git a/Engine/Source/Editor/PropertyEditor/Private/CustomChildBuilder.h b/Engine/Source/Editor/PropertyEditor/Private/CustomChildBuilder.h
--- a/Engine/Source/Editor/PropertyEditor/Private/CustomChildBuilder.h
+++ b/Engine/Source/Editor/PropertyEditor/Private/CustomChildBuilder.h
@@ -10,8 +10,8 @@ public:
 	{}
 	
 	virtual IDetailChildrenBuilder& AddChildCustomBuilder( TSharedRef<class IDetailCustomNodeBuilder> InCustomBuilder ) override;
-	virtual IDetailGroup& AddChildGroup( FName GroupName, const FString& LocalizedDisplayName ) override;
-	virtual FDetailWidgetRow& AddChildContent( const FString& SearchString ) override;
+	virtual IDetailGroup& AddChildGroup( FName GroupName, const FText& LocalizedDisplayName ) override;
+	virtual FDetailWidgetRow& AddChildContent( const FText& SearchString ) override;
 	virtual IDetailPropertyRow& AddChildProperty( TSharedRef<IPropertyHandle> PropertyHandle ) override;
 	virtual TSharedRef<SWidget> GenerateStructValueWidget( TSharedRef<IPropertyHandle> StructPropertyHandle );
 
diff --git a/Engine/Source/Editor/PropertyEditor/Private/DetailCategoryBuilderImpl.cpp b/Engine/Source/Editor/PropertyEditor/Private/DetailCategoryBuilderImpl.cpp
--- a/Engine/Source/Editor/PropertyEditor/Private/DetailCategoryBuilderImpl.cpp
+++ b/Engine/Source/Editor/PropertyEditor/Private/DetailCategoryBuilderImpl.cpp
@@ -101,7 +101,7 @@ FDetailCategoryImpl::~FDetailCategoryImpl()
 }
 
 
-FDetailWidgetRow& FDetailCategoryImpl::AddCustomRow( const FString& FilterString, bool bForAdvanced )
+FDetailWidgetRow& FDetailCategoryImpl::AddCustomRow( const FText& FilterString, bool bForAdvanced )
 {
 	FDetailLayoutCustomization NewCustomization;
 	NewCustomization.WidgetDecl = MakeShareable( new FDetailWidgetRow );
@@ -121,7 +121,7 @@ void FDetailCategoryImpl::AddCustomBuilder( TSharedRef<IDetailCustomNodeBuilder>
 	AddCustomLayout( NewCustomization, bForAdvanced );
 }
 
-IDetailGroup& FDetailCategoryImpl::AddGroup( FName GroupName, const FString& LocalizedDisplayName, bool bForAdvanced, bool bStartExpanded )
+IDetailGroup& FDetailCategoryImpl::AddGroup( FName GroupName, const FText& LocalizedDisplayName, bool bForAdvanced, bool bStartExpanded )
 {
 	FDetailLayoutCustomization NewCustomization;
 	NewCustomization.DetailGroup = MakeShareable( new FDetailGroup( GroupName, AsShared(), LocalizedDisplayName, bStartExpanded ) );
@@ -439,7 +439,7 @@ bool FDetailCategoryImpl::ContainsOnlyAdvanced() const
 	return SimpleChildNodes.Num() == 0 && AdvancedChildNodes.Num() > 0;
 }
 
-void FDetailCategoryImpl::SetDisplayName( FName InCategoryName, const FString& LocalizedNameOverride )
+void FDetailCategoryImpl::SetDisplayName( FName InCategoryName, const FText& LocalizedNameOverride )
 {
 	if( !LocalizedNameOverride.IsEmpty() )
 	{
@@ -467,11 +467,11 @@ void FDetailCategoryImpl::SetDisplayName( FName InCategoryName, const FString& L
 
 		if ( FoundText )
 		{
-			DisplayName = DisplayNameText.ToString();
+			DisplayName = DisplayNameText;
 		}
 		else
 		{
-			DisplayName = SourceCategoryStr;
+			DisplayName = FText::FromString( SourceCategoryStr );
 		}
 	}
 }
diff --git a/Engine/Source/Editor/PropertyEditor/Private/DetailCategoryBuilderImpl.h b/Engine/Source/Editor/PropertyEditor/Private/DetailCategoryBuilderImpl.h
--- a/Engine/Source/Editor/PropertyEditor/Private/DetailCategoryBuilderImpl.h
+++ b/Engine/Source/Editor/PropertyEditor/Private/DetailCategoryBuilderImpl.h
@@ -138,11 +138,11 @@ public:
 	virtual IDetailPropertyRow& AddProperty( TSharedPtr<IPropertyHandle> PropertyHandle,  EPropertyLocation::Type Location = EPropertyLocation::Default ) override;
 	virtual IDetailPropertyRow* AddExternalProperty( const TArray<UObject*>& Objects, FName PropertyName, EPropertyLocation::Type Location = EPropertyLocation::Default ) override;
 	virtual IDetailLayoutBuilder& GetParentLayout() const override { return *DetailLayoutBuilder.Pin(); }
-	virtual FDetailWidgetRow& AddCustomRow( const FString& FilterString, bool bForAdvanced = false ) override;
+	virtual FDetailWidgetRow& AddCustomRow( const FText& FilterString, bool bForAdvanced = false ) override;
 	virtual void AddCustomBuilder( TSharedRef<IDetailCustomNodeBuilder> InCustomBuilder, bool bForAdvanced = false ) override;
-	virtual IDetailGroup& AddGroup( FName GroupName, const FString& LocalizedDisplayName, bool bForAdvanced = false, bool bStartExpanded = false ) override;
+	virtual IDetailGroup& AddGroup( FName GroupName, const FText& LocalizedDisplayName, bool bForAdvanced = false, bool bStartExpanded = false ) override;
 	virtual void GetDefaultProperties( TArray<TSharedRef<IPropertyHandle> >& OutAllProperties, bool bSimpleProperties = true, bool bAdvancedProperties = true ) override;
-	virtual const FString& GetDisplayName() const override { return DisplayName; }
+	virtual const FText& GetDisplayName() const override { return DisplayName; }
 
 	/** IDetailTreeNode interface */
 	virtual IDetailsViewPrivate& GetDetailsView() const override{ return DetailLayoutBuilder.Pin()->GetDetailsView(); }
@@ -194,7 +194,7 @@ public:
 	 * @param CategoryName			Base category name to use if no localized override is set
 	 * @param LocalizedNameOverride	The localized name override to use (can be empty)
 	 */
-	void SetDisplayName( FName CategoryName, const FString& LocalizedNameOverride );
+	void SetDisplayName( FName CategoryName, const FText& LocalizedNameOverride );
 
 	/**
 	 * Request that a child node of this category be expanded or collapsed
@@ -346,7 +346,7 @@ private:
 	/** Delegate called when expansion of the category changes */
 	FOnBooleanValueChanged OnExpansionChangedDelegate;
 	/** The display name of the category */
-	FString DisplayName;
+	FText DisplayName;
 	/** The path name of the category */
 	FString CategoryPathName;
 	/** Custom header content displayed to the right of the category name */
diff --git a/Engine/Source/Editor/PropertyEditor/Private/DetailCategoryGroupNode.cpp b/Engine/Source/Editor/PropertyEditor/Private/DetailCategoryGroupNode.cpp
--- a/Engine/Source/Editor/PropertyEditor/Private/DetailCategoryGroupNode.cpp
+++ b/Engine/Source/Editor/PropertyEditor/Private/DetailCategoryGroupNode.cpp
@@ -104,7 +104,7 @@ TSharedRef< ITableRow > FDetailCategoryGroupNode::GenerateNodeWidget( const TSha
 {
 	return
 		SNew( SDetailCategoryTableRow, AsShared(), OwnerTable )
-		.DisplayName( GroupName.ToString() )
+		.DisplayName( FText::FromName(GroupName) )
 		.InnerCategory( true );
 }
 
diff --git a/Engine/Source/Editor/PropertyEditor/Private/DetailCategoryGroupNode.h b/Engine/Source/Editor/PropertyEditor/Private/DetailCategoryGroupNode.h
--- a/Engine/Source/Editor/PropertyEditor/Private/DetailCategoryGroupNode.h
+++ b/Engine/Source/Editor/PropertyEditor/Private/DetailCategoryGroupNode.h
@@ -10,7 +10,7 @@ public:
 	SLATE_BEGIN_ARGS( SDetailCategoryTableRow )
 		: _InnerCategory( false )
 	{}
-		SLATE_TEXT_ARGUMENT( DisplayName )
+		SLATE_ARGUMENT( FText, DisplayName )
 		SLATE_ARGUMENT( bool, InnerCategory )
 		SLATE_ARGUMENT( TSharedPtr<SWidget>, HeaderContent )
 	SLATE_END_ARGS()
diff --git a/Engine/Source/Editor/PropertyEditor/Private/DetailGroup.cpp b/Engine/Source/Editor/PropertyEditor/Private/DetailGroup.cpp
--- a/Engine/Source/Editor/PropertyEditor/Private/DetailGroup.cpp
+++ b/Engine/Source/Editor/PropertyEditor/Private/DetailGroup.cpp
@@ -6,7 +6,7 @@
 #include "DetailPropertyRow.h"
 #include "DetailItemNode.h"
 
-FDetailGroup::FDetailGroup( const FName InGroupName, TSharedRef<FDetailCategoryImpl> InParentCategory, const FString& InLocalizedDisplayName, const bool bInStartExpanded )
+FDetailGroup::FDetailGroup( const FName InGroupName, TSharedRef<FDetailCategoryImpl> InParentCategory, const FText& InLocalizedDisplayName, const bool bInStartExpanded )
 	: ParentCategory( InParentCategory )
 	, LocalizedDisplayName( InLocalizedDisplayName )
 	, GroupName( InGroupName )
diff --git a/Engine/Source/Editor/PropertyEditor/Private/DetailGroup.h b/Engine/Source/Editor/PropertyEditor/Private/DetailGroup.h
--- a/Engine/Source/Editor/PropertyEditor/Private/DetailGroup.h
+++ b/Engine/Source/Editor/PropertyEditor/Private/DetailGroup.h
@@ -5,7 +5,7 @@
 class FDetailGroup : public IDetailGroup, public TSharedFromThis<FDetailGroup>
 {
 public:
-	FDetailGroup( const FName InGroupName, TSharedRef<FDetailCategoryImpl> InParentCategory, const FString& InLocalizedDisplayName, const bool bStartExpanded = false );
+	FDetailGroup( const FName InGroupName, TSharedRef<FDetailCategoryImpl> InParentCategory, const FText& InLocalizedDisplayName, const bool bStartExpanded = false );
 
 	/** IDetailGroup interface */     
 	virtual FDetailWidgetRow& HeaderRow() override;
@@ -72,7 +72,7 @@ private:
 	/** Whether or not our parent is enabled */
 	TAttribute<bool> IsParentEnabled;
 	/** Display name of this group */
-	FString LocalizedDisplayName;
+	FText LocalizedDisplayName;
 	/** Name identifier of this group */
 	FName GroupName;
 	/** Whether the detail group should start expanded or not */
diff --git a/Engine/Source/Editor/PropertyEditor/Private/DetailItemNode.cpp b/Engine/Source/Editor/PropertyEditor/Private/DetailItemNode.cpp
--- a/Engine/Source/Editor/PropertyEditor/Private/DetailItemNode.cpp
+++ b/Engine/Source/Editor/PropertyEditor/Private/DetailItemNode.cpp
@@ -148,7 +148,7 @@ TSharedRef< ITableRow > FDetailItemNode::GenerateNodeWidget( const TSharedRef<ST
 	{
 		if (Customization.IsValidCustomization() && Customization.GetPropertyNode().IsValid())
 		{
-			TagMeta.Tag = *FString::Printf(TEXT("DetailRowItem.%s"), *Customization.GetPropertyNode()->GetDisplayName());
+			TagMeta.Tag = *FString::Printf(TEXT("DetailRowItem.%s"), *Customization.GetPropertyNode()->GetDisplayName().ToString());
 		}
 		else if (Customization.HasCustomWidget() )
 		{
@@ -325,7 +325,7 @@ static bool PassesAllFilters( const FDetailLayoutCustomization& InCustomization,
 		}
 		else if( InCustomization.HasCustomWidget() )
 		{
-			const bool bPassesTextFilter = Local::StringPassesFilter(InFilter, InCustomization.WidgetDecl->FilterTextString);
+			const bool bPassesTextFilter = Local::StringPassesFilter(InFilter, InCustomization.WidgetDecl->FilterTextString.ToString());
 
 			bPassesAllFilters = bPassesTextFilter || bPassesCategoryFilter;
 		}
@@ -402,7 +402,7 @@ FPropertyPath FDetailItemNode::GetPropertyPath() const
 
 void FDetailItemNode::FilterNode( const FDetailFilter& InFilter )
 {
-	bShouldBeVisibleDueToFiltering = PassesAllFilters( Customization, InFilter, ParentCategory.Pin()->GetDisplayName() );
+	bShouldBeVisibleDueToFiltering = PassesAllFilters( Customization, InFilter, ParentCategory.Pin()->GetDisplayName().ToString() );
 
 	bShouldBeVisibleDueToChildFiltering = false;
 
diff --git a/Engine/Source/Editor/PropertyEditor/Private/DetailLayoutBuilderImpl.cpp b/Engine/Source/Editor/PropertyEditor/Private/DetailLayoutBuilderImpl.cpp
--- a/Engine/Source/Editor/PropertyEditor/Private/DetailLayoutBuilderImpl.cpp
+++ b/Engine/Source/Editor/PropertyEditor/Private/DetailLayoutBuilderImpl.cpp
@@ -16,15 +16,15 @@ FDetailLayoutBuilderImpl::FDetailLayoutBuilderImpl(FClassToPropertyMap& InProper
 
 }
 
-IDetailCategoryBuilder& FDetailLayoutBuilderImpl::EditCategory( FName CategoryName, const FString& NewLocalizedDisplayName, ECategoryPriority::Type CategoryType )
+IDetailCategoryBuilder& FDetailLayoutBuilderImpl::EditCategory( FName CategoryName, const FText& NewLocalizedDisplayName, ECategoryPriority::Type CategoryType )
 {
-	FString LocalizedDisplayName = NewLocalizedDisplayName;
+	FText LocalizedDisplayName = NewLocalizedDisplayName;
 
 	// Use a generic name if one was not specified
 	if( CategoryName == NAME_None )
 	{
-		static const FString GeneralString = NSLOCTEXT("DetailLayoutBuilderImpl", "General", "General").ToString();
-		static const FName GeneralName = *GeneralString;
+		static const FText GeneralString = NSLOCTEXT("DetailLayoutBuilderImpl", "General", "General");
+		static const FName GeneralName = *GeneralString.ToString();
 
 		CategoryName = GeneralName;
 		LocalizedDisplayName = GeneralString;
@@ -115,7 +115,7 @@ FDetailCategoryImpl& FDetailLayoutBuilderImpl::DefaultCategory( FName CategoryNa
 	}
 	
 
-	CategoryImpl->SetDisplayName( CategoryName, TEXT("") );
+	CategoryImpl->SetDisplayName( CategoryName, FText::GetEmpty() );
 	return *CategoryImpl;
 }
 
diff --git a/Engine/Source/Editor/PropertyEditor/Private/DetailLayoutBuilderImpl.h b/Engine/Source/Editor/PropertyEditor/Private/DetailLayoutBuilderImpl.h
--- a/Engine/Source/Editor/PropertyEditor/Private/DetailLayoutBuilderImpl.h
+++ b/Engine/Source/Editor/PropertyEditor/Private/DetailLayoutBuilderImpl.h
@@ -11,7 +11,7 @@ public:
 	/** IDetailLayoutBuilder Interface */
 	virtual const IDetailsView& GetDetailsView() const override;
 	virtual void GetObjectsBeingCustomized( TArray< TWeakObjectPtr<UObject> >& OutObjects ) const override;
-	virtual IDetailCategoryBuilder& EditCategory( FName CategoryName, const FString& NewLocalizedDisplayName = TEXT(""), ECategoryPriority::Type CategoryType = ECategoryPriority::Default );
+	virtual IDetailCategoryBuilder& EditCategory( FName CategoryName, const FText& NewLocalizedDisplayName = FText::GetEmpty(), ECategoryPriority::Type CategoryType = ECategoryPriority::Default );
 	virtual TSharedRef<IPropertyHandle> GetProperty( const FName PropertyPath, const UClass* ClassOutermost, FName InInstanceName ) override;
 	virtual void HideProperty( const TSharedPtr<IPropertyHandle> Property ) override;
 	virtual void HideProperty( FName PropertyPath, const UClass* ClassOutermost = NULL, FName InstanceName = NAME_None ) override;
diff --git a/Engine/Source/Editor/PropertyEditor/Private/DetailPropertyRow.cpp b/Engine/Source/Editor/PropertyEditor/Private/DetailPropertyRow.cpp
--- a/Engine/Source/Editor/PropertyEditor/Private/DetailPropertyRow.cpp
+++ b/Engine/Source/Editor/PropertyEditor/Private/DetailPropertyRow.cpp
@@ -52,7 +52,7 @@ FDetailPropertyRow::FDetailPropertyRow(TSharedPtr<FPropertyNode> InPropertyNode,
 	}
 }
 
-IDetailPropertyRow& FDetailPropertyRow::DisplayName( const FString& InDisplayName )
+IDetailPropertyRow& FDetailPropertyRow::DisplayName( const FText& InDisplayName )
 {
 	if (PropertyNode.IsValid())
 	{
@@ -61,7 +61,7 @@ IDetailPropertyRow& FDetailPropertyRow::DisplayName( const FString& InDisplayNam
 	return *this;
 }
 
-IDetailPropertyRow& FDetailPropertyRow::ToolTip( const FString& InToolTip )
+IDetailPropertyRow& FDetailPropertyRow::ToolTip( const FText& InToolTip )
 {
 	if (PropertyNode.IsValid())
 	{
diff --git a/Engine/Source/Editor/PropertyEditor/Private/DetailPropertyRow.h b/Engine/Source/Editor/PropertyEditor/Private/DetailPropertyRow.h
--- a/Engine/Source/Editor/PropertyEditor/Private/DetailPropertyRow.h
+++ b/Engine/Source/Editor/PropertyEditor/Private/DetailPropertyRow.h
@@ -8,8 +8,8 @@ public:
 	FDetailPropertyRow( TSharedPtr<FPropertyNode> InPropertyNode, TSharedRef<FDetailCategoryImpl> InParentCategory, TSharedPtr<FObjectPropertyNode> InExternalRootNode = NULL);
 
 	/** IDetailPropertyRow interface */
-	virtual IDetailPropertyRow& DisplayName( const FString& InDisplayName ) override;
-	virtual IDetailPropertyRow& ToolTip( const FString& InToolTip ) override;
+	virtual IDetailPropertyRow& DisplayName( const FText& InDisplayName ) override;
+	virtual IDetailPropertyRow& ToolTip( const FText& InToolTip ) override;
 	virtual IDetailPropertyRow& ShowPropertyButtons( bool bInShowPropertyButtons ) override;
 	virtual IDetailPropertyRow& EditCondition( TAttribute<bool> EditConditionValue, FOnBooleanValueChanged OnEditConditionValueChanged ) override;
 	virtual IDetailPropertyRow& IsEnabled( TAttribute<bool> InIsEnabled ) override;
diff --git a/Engine/Source/Editor/PropertyEditor/Private/ItemPropertyNode.cpp b/Engine/Source/Editor/PropertyEditor/Private/ItemPropertyNode.cpp
--- a/Engine/Source/Editor/PropertyEditor/Private/ItemPropertyNode.cpp
+++ b/Engine/Source/Editor/PropertyEditor/Private/ItemPropertyNode.cpp
@@ -284,16 +284,16 @@ void FItemPropertyNode::InitChildNodes()
 	}
 }
 
-void FItemPropertyNode::SetDisplayNameOverride( const FString& InDisplayNameOverride )
+void FItemPropertyNode::SetDisplayNameOverride( const FText& InDisplayNameOverride )
 {
 	DisplayNameOverride = InDisplayNameOverride;
 }
 
-FString FItemPropertyNode::GetDisplayName() const
+FText FItemPropertyNode::GetDisplayName() const
 {
-	FString FinalDisplayName;
+	FText FinalDisplayName;
 
-	if( DisplayNameOverride.Len() > 0 )
+	if( !DisplayNameOverride.IsEmpty() )
 	{
 		FinalDisplayName = DisplayNameOverride;
 	}
@@ -306,8 +306,8 @@ FString FItemPropertyNode::GetDisplayName() const
 			if ( FPropertySettings::Get().ShowFriendlyPropertyNames() )
 			{
 				//We are in "readable display name mode"../ Make a nice name
-				FinalDisplayName = PropertyPtr->GetDisplayNameText().ToString();
-				if ( FinalDisplayName.Len() == 0 )
+				FinalDisplayName = PropertyPtr->GetDisplayNameText();
+				if ( FinalDisplayName.IsEmpty() )
 				{
 					FString PropertyDisplayName;
 					bool bIsBoolProperty = Cast<const UBoolProperty>(PropertyPtr) != NULL;
@@ -340,12 +340,12 @@ FString FItemPropertyNode::GetDisplayName() const
 						PropertyDisplayName = FName::NameToDisplayString( PropertyDisplayName, bIsBoolProperty );
 					}
 
-					FinalDisplayName = PropertyDisplayName;
+					FinalDisplayName = FText::FromString( PropertyDisplayName );
 				}
 			}
 			else
 			{
-				FinalDisplayName =  PropertyPtr->GetName();
+				FinalDisplayName =  FText::FromString( PropertyPtr->GetName() );
 			}
 		}
 		else
@@ -360,29 +360,29 @@ FString FItemPropertyNode::GetDisplayName() const
 			// This item is a member of an array, its display name is its index 
 			if ( PropertyPtr == NULL || ArraySizeEnum == NULL )
 			{
-				FinalDisplayName = *FString::Printf(TEXT("%i"), GetArrayIndex() );
+				FinalDisplayName = FText::AsNumber( GetArrayIndex() );
 			}
 			else
 			{
-				FinalDisplayName = *FString::Printf(TEXT("%s"), *ArraySizeEnum->GetEnumName(GetArrayIndex()));
+				FString TempDisplayName = ArraySizeEnum->GetEnumName(GetArrayIndex());
 				//fixup the display name if we have displayname metadata
-				AdjustEnumPropDisplayName(ArraySizeEnum, FinalDisplayName);
+				AdjustEnumPropDisplayName(ArraySizeEnum, TempDisplayName);
+				FinalDisplayName = FText::FromString(TempDisplayName); // todo: should this be using ArraySizeEnum->GetEnumText?
 			}
 		}
 	}
 	
 	return FinalDisplayName;
-
 }
 
-void FItemPropertyNode::SetToolTipOverride( const FString& InToolTipOverride )
+void FItemPropertyNode::SetToolTipOverride( const FText& InToolTipOverride )
 {
 	ToolTipOverride = InToolTipOverride;
 }
 
-FString FItemPropertyNode::GetToolTipText() const
+FText FItemPropertyNode::GetToolTipText() const
 {
-	if(ToolTipOverride.Len() > 0)
+	if(!ToolTipOverride.IsEmpty())
 	{
 		return ToolTipOverride;
 	}
diff --git a/Engine/Source/Editor/PropertyEditor/Private/ItemPropertyNode.h b/Engine/Source/Editor/PropertyEditor/Private/ItemPropertyNode.h
--- a/Engine/Source/Editor/PropertyEditor/Private/ItemPropertyNode.h
+++ b/Engine/Source/Editor/PropertyEditor/Private/ItemPropertyNode.h
@@ -31,22 +31,22 @@ public:
 	virtual FItemPropertyNode* AsItemPropertyNode() override { return this; }
 
 	/** Display name override to use instead of the property name */
-	void SetDisplayNameOverride( const FString& InDisplayNameOverride ) override;
+	void SetDisplayNameOverride( const FText& InDisplayNameOverride ) override;
 
 	/**
 	 * @return The formatted display name for the property in this node                                                              
 	 */
-	FString GetDisplayName() const override;
+	virtual FText GetDisplayName() const override;
 
 	/**
 	 * Sets the tooltip override to use instead of the property tooltip
 	 */
-	virtual void SetToolTipOverride( const FString& InToolTipOverride ) override;
+	virtual void SetToolTipOverride( const FText& InToolTipOverride ) override;
 
 	/**
 	 * @return The tooltip for the property in this node                                                              
 	 */
-	virtual FString GetToolTipText() const override;
+	virtual FText GetToolTipText() const override;
 
 protected:
 	/**
@@ -61,7 +61,7 @@ protected:
 
 private:
 	/** Display name override to use instead of the property name */
-	FString DisplayNameOverride;
+	FText DisplayNameOverride;
 	/** Tooltip override to use instead of the property tooltip */
-	FString ToolTipOverride;
+	FText ToolTipOverride;
 };
diff --git a/Engine/Source/Editor/PropertyEditor/Private/Presentation/PropertyEditor/PropertyEditor.cpp b/Engine/Source/Editor/PropertyEditor/Private/Presentation/PropertyEditor/PropertyEditor.cpp
--- a/Engine/Source/Editor/PropertyEditor/Private/Presentation/PropertyEditor/PropertyEditor.cpp
+++ b/Engine/Source/Editor/PropertyEditor/Private/Presentation/PropertyEditor/PropertyEditor.cpp
@@ -52,7 +52,7 @@ FPropertyEditor::FPropertyEditor( const TSharedRef<FPropertyNode>& InPropertyNod
 }
 
 
-FString FPropertyEditor::GetDisplayName() const
+FText FPropertyEditor::GetDisplayName() const
 {
 	FCategoryPropertyNode* CategoryNode = PropertyNode->AsCategoryNode();
 	FItemPropertyNode* ItemPropertyNode = PropertyNode->AsItemPropertyNode();
@@ -69,12 +69,13 @@ FString FPropertyEditor::GetDisplayName() const
 	{
 		FString DisplayName;
 		PropertyNode->GetQualifiedName( DisplayName, true );
-
-		return DisplayName;
+		return FText::FromString(DisplayName);
 	}
+
+	return FText::GetEmpty();
 }
 
-FString FPropertyEditor::GetToolTipText() const
+FText FPropertyEditor::GetToolTipText() const
 {
 	return PropertyNode->GetToolTipText();
 }
diff --git a/Engine/Source/Editor/PropertyEditor/Private/Presentation/PropertyEditor/PropertyEditor.h b/Engine/Source/Editor/PropertyEditor/Private/Presentation/PropertyEditor/PropertyEditor.h
--- a/Engine/Source/Editor/PropertyEditor/Private/Presentation/PropertyEditor/PropertyEditor.h
+++ b/Engine/Source/Editor/PropertyEditor/Private/Presentation/PropertyEditor/PropertyEditor.h
@@ -50,7 +50,7 @@ public:
 	bool IsEditConditionMet() const;
 
 	/**	@return The tooltip for this property editor */
-	FString GetToolTipText() const;
+	FText GetToolTipText() const;
 
 	/** @return The documentation link for this property */
 	FString GetDocumentationLink() const;
@@ -59,7 +59,7 @@ public:
 	FString GetDocumentationExcerptName() const;
 
 	/**	@return The display name to be used for the property */
-	FString GetDisplayName() const;
+	FText GetDisplayName() const;
 
 	/** @return Whether or not resetting this property to its default value is a valid and worthwhile operation */
 	bool IsResetToDefaultAvailable() const;
diff --git a/Engine/Source/Editor/PropertyEditor/Private/PropertyCustomizationHelpers.cpp b/Engine/Source/Editor/PropertyEditor/Private/PropertyCustomizationHelpers.cpp
--- a/Engine/Source/Editor/PropertyEditor/Private/PropertyCustomizationHelpers.cpp
+++ b/Engine/Source/Editor/PropertyEditor/Private/PropertyCustomizationHelpers.cpp
@@ -30,7 +30,7 @@ namespace PropertyCustomizationHelpers
 			, _Image( FEditorStyle::GetBrush("Default") )
 			, _IsFocusable( true )
 		{}
-			SLATE_TEXT_ARGUMENT( Text )
+			SLATE_ARGUMENT( FText, Text )
 			SLATE_ARGUMENT( const FSlateBrush*, Image )
 			SLATE_EVENT( FSimpleDelegate, OnClickAction )
 
@@ -421,7 +421,7 @@ void SProperty::Construct( const FArguments& InArgs, TSharedPtr<IPropertyHandle>
 {
 	TSharedPtr<SWidget> ChildSlotContent;
 
-	const FString& DisplayName = InArgs._DisplayName.Get();
+	const FText& DisplayName = InArgs._DisplayName.Get();
 
 	PropertyHandle = InPropertyHandle;
 
@@ -948,7 +948,7 @@ void FMaterialList::GenerateChildContent( IDetailChildrenBuilder& ChildrenBuilde
 				// If we are currently displaying an expanded set of materials for an element add a link to collapse all of them
 				if( bWantToDisplayAllMaterials )
 				{
-					FDetailWidgetRow& ChildRow = ChildrenBuilder.AddChildContent( LOCTEXT( "HideAllMaterialSearchString", "Hide All Materials").ToString() );
+					FDetailWidgetRow& ChildRow = ChildrenBuilder.AddChildContent( LOCTEXT( "HideAllMaterialSearchString", "Hide All Materials") );
 
 					FFormatNamedArguments Arguments;
 					Arguments.Add(TEXT("ElementSlot"), CurrentSlot);
@@ -972,7 +972,7 @@ void FMaterialList::GenerateChildContent( IDetailChildrenBuilder& ChildrenBuilde
 					// The current slot has multiple elements to view
 					bDisplayAllMaterialsInSlot = false;
 
-					FDetailWidgetRow& ChildRow = ChildrenBuilder.AddChildContent( TEXT("") );
+					FDetailWidgetRow& ChildRow = ChildrenBuilder.AddChildContent( FText::GetEmpty() );
 
 					AddMaterialItem( ChildRow, CurrentSlot, FMaterialListItem( NULL, CurrentSlot, true ), !bDisplayAllMaterialsInSlot );
 				}
@@ -986,7 +986,7 @@ void FMaterialList::GenerateChildContent( IDetailChildrenBuilder& ChildrenBuilde
 			// Display each thumbnail element unless we shouldn't display multiple materials for one slot
 			if( bDisplayAllMaterialsInSlot )
 			{
-				FDetailWidgetRow& ChildRow = ChildrenBuilder.AddChildContent( Material.Material.IsValid()? Material.Material->GetName() : TEXT("") );
+				FDetailWidgetRow& ChildRow = ChildrenBuilder.AddChildContent( Material.Material.IsValid()? FText::FromString(Material.Material->GetName()) : FText::GetEmpty() );
 
 				AddMaterialItem( ChildRow, CurrentSlot, Material, !bDisplayAllMaterialsInSlot );
 			}
@@ -994,7 +994,7 @@ void FMaterialList::GenerateChildContent( IDetailChildrenBuilder& ChildrenBuilde
 	}
 	else
 	{
-		FDetailWidgetRow& ChildRow = ChildrenBuilder.AddChildContent( LOCTEXT("NoMaterials", "No Materials").ToString() );
+		FDetailWidgetRow& ChildRow = ChildrenBuilder.AddChildContent( LOCTEXT("NoMaterials", "No Materials") );
 
 		ChildRow
 		[
diff --git a/Engine/Source/Editor/PropertyEditor/Private/PropertyEditorHelpers.cpp b/Engine/Source/Editor/PropertyEditor/Private/PropertyEditorHelpers.cpp
--- a/Engine/Source/Editor/PropertyEditor/Private/PropertyEditorHelpers.cpp
+++ b/Engine/Source/Editor/PropertyEditor/Private/PropertyEditorHelpers.cpp
@@ -51,7 +51,7 @@ void SPropertyNameWidget::Construct( const FArguments& InArgs, TSharedPtr<FPrope
 				SNew( SPropertyEditorTitle, PropertyEditor.ToSharedRef() )
 				.StaticDisplayName( PropertyEditor->GetDisplayName() )
 				.OnDoubleClicked( InArgs._OnDoubleClicked )
-                .ToolTip( IDocumentation::Get()->CreateToolTip( FText::FromString( PropertyEditor->GetToolTipText() ), NULL, PropertyEditor->GetDocumentationLink(), PropertyEditor->GetDocumentationExcerptName() ) )
+                .ToolTip( IDocumentation::Get()->CreateToolTip( PropertyEditor->GetToolTipText(), NULL, PropertyEditor->GetDocumentationLink(), PropertyEditor->GetDocumentationExcerptName() ) )
 			]
 		]
 	
@@ -79,7 +79,7 @@ void SPropertyValueWidget::Construct( const FArguments& InArgs, TSharedPtr<FProp
 
 	ValueEditorWidget = ConstructPropertyEditorWidget( PropertyEditor, InPropertyUtilities );
 
-	ValueEditorWidget->SetToolTipText( TAttribute<FString>(PropertyEditor->GetToolTipText()) );
+	ValueEditorWidget->SetToolTipText( PropertyEditor->GetToolTipText() );
 
 
 	if( InArgs._ShowPropertyButtons )
@@ -354,16 +354,14 @@ namespace PropertyEditorHelpers
 			&&	(bAllowAbstract || !CheckClass->HasAnyClassFlags(CLASS_Abstract));
 	}
 
-	FString GetToolTipText( const UProperty* const Property )
+	FText GetToolTipText( const UProperty* const Property )
 	{
 		if( Property )
 		{
-			FString ToolTipText = Property->GetToolTipText().ToString();
-
-			return ToolTipText;
+			return Property->GetToolTipText();
 		}
 
-		return FString();
+		return FText::GetEmpty();
 	}
 
 	FString GetDocumentationLink( const UProperty* const Property )
diff --git a/Engine/Source/Editor/PropertyEditor/Private/PropertyEditorHelpers.h b/Engine/Source/Editor/PropertyEditor/Private/PropertyEditorHelpers.h
--- a/Engine/Source/Editor/PropertyEditor/Private/PropertyEditorHelpers.h
+++ b/Engine/Source/Editor/PropertyEditor/Private/PropertyEditorHelpers.h
@@ -141,7 +141,7 @@ namespace PropertyEditorHelpers
 	/**
 	 * @return The text that represents the specified properties tooltip
 	 */
-	FString GetToolTipText( const UProperty* const Property );
+	FText GetToolTipText( const UProperty* const Property );
 
 	/**
 	 * @return The link to the documentation that describes this property in detail
diff --git a/Engine/Source/Editor/PropertyEditor/Private/PropertyEditorModule.cpp b/Engine/Source/Editor/PropertyEditor/Private/PropertyEditorModule.cpp
--- a/Engine/Source/Editor/PropertyEditor/Private/PropertyEditorModule.cpp
+++ b/Engine/Source/Editor/PropertyEditor/Private/PropertyEditorModule.cpp
@@ -664,7 +664,7 @@ FPropertyTypeLayoutCallback FPropertyEditorModule::GetPropertyTypeCustomization(
 	return FPropertyTypeLayoutCallback();
 }
 
-TSharedRef<class IStructureDetailsView> FPropertyEditorModule::CreateStructureDetailView(const FDetailsViewArgs& DetailsViewArgs, TSharedPtr<FStructOnScope> StructData, bool bShowObjects, const FString& CustomName)
+TSharedRef<class IStructureDetailsView> FPropertyEditorModule::CreateStructureDetailView(const FDetailsViewArgs& DetailsViewArgs, TSharedPtr<FStructOnScope> StructData, bool bShowObjects, const FText& CustomName)
 {
 	TSharedRef<SStructureDetailsView> DetailView =
 		SNew(SStructureDetailsView)
diff --git a/Engine/Source/Editor/PropertyEditor/Private/PropertyHandleImpl.cpp b/Engine/Source/Editor/PropertyEditor/Private/PropertyHandleImpl.cpp
--- a/Engine/Source/Editor/PropertyEditor/Private/PropertyHandleImpl.cpp
+++ b/Engine/Source/Editor/PropertyEditor/Private/PropertyHandleImpl.cpp
@@ -378,7 +378,7 @@ FPropertyAccess::Result FPropertyValueImpl::ImportText( const TArray<FObjectBase
 					// Begin a transaction only if we need to call PreChange
 					if (GEditor && bTransactable)
 					{
-						GEditor->BeginTransaction(TEXT("PropertyEditor"), FText::Format(NSLOCTEXT("PropertyEditor", "EditPropertyTransaction", "Edit {0}"), FText::FromString(InPropertyNode->GetDisplayName())), NodeProperty);
+						GEditor->BeginTransaction(TEXT("PropertyEditor"), FText::Format(NSLOCTEXT("PropertyEditor", "EditPropertyTransaction", "Edit {0}"), InPropertyNode->GetDisplayName()), NodeProperty);
 					}
 				}
 
@@ -844,7 +844,7 @@ FPropertyAccess::Result FPropertyValueImpl::OnUseSelected()
 					// Warn that some object assignments failed.
 					FMessageDialog::Open( EAppMsgType::Ok, FText::Format(
 						NSLOCTEXT("UnrealEd", "ObjectAssignmentsFailed", "Failed to assign {0} to the {1} property, see log for details."),
-						FText::FromString(SelectedObject->GetPathName()), FText::FromString(PropertyNodePin->GetDisplayName())) );
+						FText::FromString(SelectedObject->GetPathName()), PropertyNodePin->GetDisplayName()) );
 				}
 				else
 				{
@@ -1426,9 +1426,9 @@ bool FPropertyValueImpl::HasValidProperty() const
 	return PropertyNode.IsValid();
 }
 
-FString FPropertyValueImpl::GetDisplayName() const
+FText FPropertyValueImpl::GetDisplayName() const
 {
-	return PropertyNode.IsValid() ? PropertyNode.Pin()->GetDisplayName() : TEXT("");
+	return PropertyNode.IsValid() ? PropertyNode.Pin()->GetDisplayName() : FText::GetEmpty();
 }
 
 #define IMPLEMENT_PROPERTY_ACCESSOR( ValueType ) \
@@ -1466,7 +1466,7 @@ bool FPropertyHandleBase::IsValidHandle() const
 	return Implementation->HasValidProperty();
 }
 
-FString FPropertyHandleBase::GetPropertyDisplayName() const
+FText FPropertyHandleBase::GetPropertyDisplayName() const
 {
 	return Implementation->GetDisplayName();
 }
@@ -1514,16 +1514,16 @@ FString FPropertyHandleBase::GeneratePathToProperty() const
 
 }
 
-TSharedRef<SWidget> FPropertyHandleBase::CreatePropertyNameWidget( const FString& NameOverride, const FString& ToolTipOverride, bool bDisplayResetToDefault, bool bDisplayText, bool bDisplayThumbnail ) const
+TSharedRef<SWidget> FPropertyHandleBase::CreatePropertyNameWidget( const FText& NameOverride, const FText& ToolTipOverride, bool bDisplayResetToDefault, bool bDisplayText, bool bDisplayThumbnail ) const
 {
 	if( Implementation.IsValid() && Implementation->GetPropertyNode().IsValid() )
 	{
-		if( NameOverride.Len() > 0 )
+		if( !NameOverride.IsEmpty() )
 		{
 			Implementation->GetPropertyNode()->SetDisplayNameOverride( NameOverride );
 		}
 
-		if( ToolTipOverride.Len() > 0 )
+		if( !ToolTipOverride.IsEmpty() )
 		{
 			Implementation->GetPropertyNode()->SetToolTipOverride( ToolTipOverride );
 		}
@@ -1762,7 +1762,7 @@ UClass* FPropertyHandleBase::GetClassMetaData(const FName& Key) const
 	return (MetaDataProperty) ? MetaDataProperty->GetClassMetaData(Key) : nullptr;
 }
 
-FString FPropertyHandleBase::GetToolTipText() const
+FText FPropertyHandleBase::GetToolTipText() const
 {
 	TSharedPtr<FPropertyNode> PropertyNode = Implementation->GetPropertyNode();
 	if( PropertyNode.IsValid() )
@@ -1770,10 +1770,10 @@ FString FPropertyHandleBase::GetToolTipText() const
 		return PropertyNode->GetToolTipText();
 	}
 
-	return FString();
+	return FText::GetEmpty();
 }
 
-void FPropertyHandleBase::SetToolTipText( const FString& ToolTip )
+void FPropertyHandleBase::SetToolTipText( const FText& ToolTip )
 {
 	TSharedPtr<FPropertyNode> PropertyNode = Implementation->GetPropertyNode();
 	if (PropertyNode.IsValid())
diff --git a/Engine/Source/Editor/PropertyEditor/Private/PropertyHandleImpl.h b/Engine/Source/Editor/PropertyEditor/Private/PropertyHandleImpl.h
--- a/Engine/Source/Editor/PropertyEditor/Private/PropertyHandleImpl.h
+++ b/Engine/Source/Editor/PropertyEditor/Private/PropertyHandleImpl.h
@@ -300,7 +300,7 @@ public:
 	/**
 	 * @return The display name of the property
 	 */
-	FString GetDisplayName() const;
+	FText GetDisplayName() const;
 
 	/** @return The notify hook being used */
 	FNotifyHook* GetNotifyHook() const { return NotifyHook; }
@@ -367,14 +367,14 @@ public:
 
 	/** IPropertyHandle interface */
 	virtual bool IsValidHandle() const override;
-	virtual FString GetPropertyDisplayName() const override;
+	virtual FText GetPropertyDisplayName() const override;
 	virtual void ResetToDefault() override;
 	virtual bool DiffersFromDefault() const override;
 	virtual FText GetResetToDefaultLabel() const override;
 	virtual void MarkHiddenByCustomization() override;
 	virtual bool IsCustomized() const override;
 	virtual FString GeneratePathToProperty() const override;
-	virtual TSharedRef<SWidget> CreatePropertyNameWidget( const FString& NameOverride = TEXT(""), const FString& ToolTipOverride = TEXT(""), bool bDisplayResetToDefault = false, bool bDisplayText = true, bool bDisplayThumbnail = true ) const override;
+	virtual TSharedRef<SWidget> CreatePropertyNameWidget( const FText& NameOverride = FText::GetEmpty(), const FText& ToolTipOverride = FText::GetEmpty(), bool bDisplayResetToDefault = false, bool bDisplayText = true, bool bDisplayThumbnail = true ) const override;
 	virtual TSharedRef<SWidget> CreatePropertyValueWidget() const override;
 	virtual bool IsEditConst() const override;
 	virtual void SetOnPropertyValueChanged( const FSimpleDelegate& InOnPropertyValueChanged ) override;
@@ -402,8 +402,8 @@ public:
 	virtual int32 GetINTMetaData(const FName& Key) const override;
 	virtual float GetFLOATMetaData(const FName& Key) const override;
 	virtual UClass* GetClassMetaData(const FName& Key) const override;
-	virtual FString GetToolTipText() const override;
-	virtual void SetToolTipText(const FString& ToolTip) override;
+	virtual FText GetToolTipText() const override;
+	virtual void SetToolTipText(const FText& ToolTip) override;
 	virtual FPropertyAccess::Result SetPerObjectValues( const TArray<FString>& InPerObjectValues,  EPropertyValueSetFlags::Type Flags = EPropertyValueSetFlags::DefaultFlags ) override;
 	virtual FPropertyAccess::Result GetPerObjectValues( TArray<FString>& OutPerObjectValues ) override;
 	virtual bool GeneratePossibleValues(TArray< TSharedPtr<FString> >& OutOptionStrings, TArray< TSharedPtr<FString> >& OutToolTips, TArray<bool>& OutRestrictedItems) override;
diff --git a/Engine/Source/Editor/PropertyEditor/Private/PropertyNode.cpp b/Engine/Source/Editor/PropertyEditor/Private/PropertyNode.cpp
--- a/Engine/Source/Editor/PropertyEditor/Private/PropertyNode.cpp
+++ b/Engine/Source/Editor/PropertyEditor/Private/PropertyNode.cpp
@@ -1447,21 +1447,21 @@ FString FPropertyNode::GetDefaultValueAsString()
 FText FPropertyNode::GetResetToDefaultLabel()
 {
 	FString DefaultValue = GetDefaultValueAsString();
-	FString OutLabel = GetDisplayName();
+	FText OutLabel = GetDisplayName();
 	if ( DefaultValue.Len() )
 	{
 		const int32 MaxValueLen = 60;
 
-		OutLabel += TEXT(": ");
 		if ( DefaultValue.Len() > MaxValueLen )
 		{
 			DefaultValue = DefaultValue.Left( MaxValueLen );
 			DefaultValue += TEXT( "..." );
 		}
-		OutLabel += DefaultValue;
+
+		return FText::Format(NSLOCTEXT("FPropertyNode", "ResetToDefaultLabelFmt", "{0}: {1}"), OutLabel, FText::FromString(DefaultValue));
 	}
 
-	return FText::FromString( OutLabel );
+	return OutLabel;
 }
 
 void FPropertyNode::ResetToDefault( FNotifyHook* InNotifyHook )
@@ -1681,13 +1681,14 @@ void FPropertyNode::FilterNodes( const TArray<FString>& InFilterStrings, const b
 		bool bPassedFilter = false;	//assuming that we aren't filtered
 
 		//see if this is a filter-able primitive
-		FString DisplayName = GetDisplayName();
+		FText DisplayName = GetDisplayName();
+		const FString& DisplayNameStr = DisplayName.ToString();
 		TArray <FString> AcceptableNames;
-		AcceptableNames.Add(DisplayName);
+		AcceptableNames.Add(DisplayNameStr);
 
 		//get the basic name as well of the property
 		UProperty* TheProperty = GetProperty();
-		if (TheProperty && (TheProperty->GetName() != DisplayName))
+		if (TheProperty && (TheProperty->GetName() != DisplayNameStr))
 		{
 			AcceptableNames.Add(TheProperty->GetName());
 		}
diff --git a/Engine/Source/Editor/PropertyEditor/Private/PropertyNode.h b/Engine/Source/Editor/PropertyEditor/Private/PropertyNode.h
--- a/Engine/Source/Editor/PropertyEditor/Private/PropertyNode.h
+++ b/Engine/Source/Editor/PropertyEditor/Private/PropertyNode.h
@@ -405,22 +405,22 @@ public:
 	/**
 	 * Sets the display name override to use instead of the display name
 	 */
-	virtual void SetDisplayNameOverride( const FString& InDisplayNameOverride ) {}
+	virtual void SetDisplayNameOverride( const FText& InDisplayNameOverride ) {}
 
 	/**
 	 * @return The formatted display name for the property in this node                                                              
 	 */
-	virtual FString GetDisplayName() const { return FString(); }
+	virtual FText GetDisplayName() const { return FText::GetEmpty(); }
 
 	/**
 	 * Sets the tooltip override to use instead of the property tooltip
 	 */
-	virtual void SetToolTipOverride( const FString& InToolTipOverride ) {}
+	virtual void SetToolTipOverride( const FText& InToolTipOverride ) {}
 
 	/**
 	 * @return The tooltip for the property in this node                                                              
 	 */
-	virtual FString GetToolTipText() const { return FString(); }
+	virtual FText GetToolTipText() const { return FText::GetEmpty(); }
 
 	/**
 	 * If there is a property, sees if it matches.  Otherwise sees if the entire parent structure matches
diff --git a/Engine/Source/Editor/PropertyEditor/Private/SSingleProperty.h b/Engine/Source/Editor/PropertyEditor/Private/SSingleProperty.h
--- a/Engine/Source/Editor/PropertyEditor/Private/SSingleProperty.h
+++ b/Engine/Source/Editor/PropertyEditor/Private/SSingleProperty.h
@@ -17,7 +17,7 @@ public:
 		SLATE_ARGUMENT( FNotifyHook*, NotifyHook )
 		SLATE_ARGUMENT( FSlateFontInfo, PropertyFont )
 		SLATE_ARGUMENT( EPropertyNamePlacement::Type, NamePlacement )
-		SLATE_ARGUMENT( FString, NameOverride )
+		SLATE_ARGUMENT( FText, NameOverride )
 	SLATE_END_ARGS()	
 
 	void Construct( const FArguments& InArgs );
@@ -63,7 +63,7 @@ private:
 	/** Property utilities for handling common functionality of property editors */
 	TSharedPtr<class FSinglePropertyUtilities> PropertyUtilities;
 	/** Name override to display instead of the property name */
-	FString NameOverride;
+	FText NameOverride;
 	/** Font to use */
 	FSlateFontInfo PropertyFont;
 	/** Notify hook to use when editing values */
diff --git a/Engine/Source/Editor/PropertyEditor/Private/SStructureDetailsView.h b/Engine/Source/Editor/PropertyEditor/Private/SStructureDetailsView.h
--- a/Engine/Source/Editor/PropertyEditor/Private/SStructureDetailsView.h
+++ b/Engine/Source/Editor/PropertyEditor/Private/SStructureDetailsView.h
@@ -8,7 +8,7 @@ class SStructureDetailsView : public SDetailsViewBase, public IStructureDetailsV
 {
 	TSharedPtr<class FStructOnScope> StructData;
 	TSharedPtr<class FStructurePropertyNode> RootNode;
-	FString CustomName;
+	FText CustomName;
 public:
 
 	SLATE_BEGIN_ARGS(SStructureDetailsView)
@@ -16,7 +16,7 @@ public:
 		{}
 		/** The user defined args for the details view */
 		SLATE_ARGUMENT(FDetailsViewArgs, DetailsViewArgs)
-		SLATE_ARGUMENT(FString, CustomName)
+		SLATE_ARGUMENT(FText, CustomName)
 	SLATE_END_ARGS()
 
 	~SStructureDetailsView();
diff --git a/Engine/Source/Editor/PropertyEditor/Private/UserInterface/PropertyEditor/SPropertyEditorNumeric.h b/Engine/Source/Editor/PropertyEditor/Private/UserInterface/PropertyEditor/SPropertyEditorNumeric.h
--- a/Engine/Source/Editor/PropertyEditor/Private/UserInterface/PropertyEditor/SPropertyEditorNumeric.h
+++ b/Engine/Source/Editor/PropertyEditor/Private/UserInterface/PropertyEditor/SPropertyEditorNumeric.h
@@ -231,7 +231,7 @@ private:
 	{
 		bIsUsingSlider = true;
 
-		GEditor->BeginTransaction( TEXT("PropertyEditor"), FText::Format( NSLOCTEXT("PropertyEditor", "SetNumericPropertyTransaction", "Edit {0}"), FText::FromString( PropertyEditor->GetDisplayName() ) ), PropertyEditor->GetPropertyHandle()->GetProperty() );
+		GEditor->BeginTransaction( TEXT("PropertyEditor"), FText::Format( NSLOCTEXT("PropertyEditor", "SetNumericPropertyTransaction", "Edit {0}"), PropertyEditor->GetDisplayName() ), PropertyEditor->GetPropertyHandle()->GetProperty() );
 	}
 
 
diff --git a/Engine/Source/Editor/PropertyEditor/Private/UserInterface/PropertyEditor/SPropertyEditorTableRow.cpp b/Engine/Source/Editor/PropertyEditor/Private/UserInterface/PropertyEditor/SPropertyEditorTableRow.cpp
--- a/Engine/Source/Editor/PropertyEditor/Private/UserInterface/PropertyEditor/SPropertyEditorTableRow.cpp
+++ b/Engine/Source/Editor/PropertyEditor/Private/UserInterface/PropertyEditor/SPropertyEditorTableRow.cpp
@@ -34,7 +34,7 @@ void SPropertyEditorTableRow::Construct( const FArguments& InArgs, const TShared
 
 	PropertyPath = FPropertyNode::CreatePropertyPath( PropertyEditor->GetPropertyNode() );
 
-	this->SetToolTipText(TAttribute<FString>(PropertyEditor->GetToolTipText()));
+	this->SetToolTipText(PropertyEditor->GetToolTipText());
 
 	SMultiColumnTableRow< TSharedPtr<FPropertyNode*> >::Construct( FSuperRowType::FArguments(), InOwnerTable );
 }
@@ -289,7 +289,7 @@ TSharedRef<SWidget> SPropertyEditorTableRow::ConstructPropertyEditorWidget()
 		PropertyWidget = SNew( SPropertyEditor, PropertyEditorRef );
 	}
 
-	PropertyWidget->SetToolTipText( TAttribute<FString>(PropertyEditor->GetToolTipText()) );
+	PropertyWidget->SetToolTipText( PropertyEditor->GetToolTipText() );
 
 	return PropertyWidget.ToSharedRef();
 }
\ No newline at end of file
diff --git a/Engine/Source/Editor/PropertyEditor/Private/UserInterface/PropertyEditor/SPropertyEditorTitle.h b/Engine/Source/Editor/PropertyEditor/Private/UserInterface/PropertyEditor/SPropertyEditorTitle.h
--- a/Engine/Source/Editor/PropertyEditor/Private/UserInterface/PropertyEditor/SPropertyEditorTitle.h
+++ b/Engine/Source/Editor/PropertyEditor/Private/UserInterface/PropertyEditor/SPropertyEditorTitle.h
@@ -16,7 +16,7 @@ public:
 		, _PropertyFont( FEditorStyle::GetFontStyle( PropertyEditorConstants::PropertyFontStyle ) )
 		, _CategoryFont( FEditorStyle::GetFontStyle( PropertyEditorConstants::CategoryFontStyle ) )
 		{}
-		SLATE_ARGUMENT( FString, StaticDisplayName )
+		SLATE_ARGUMENT( FText, StaticDisplayName )
 		SLATE_ATTRIBUTE( FSlateFontInfo, PropertyFont )
 		SLATE_ATTRIBUTE( FSlateFontInfo, CategoryFont )
 		SLATE_EVENT( FOnClicked, OnDoubleClicked )
@@ -32,7 +32,7 @@ public:
 		FItemPropertyNode* ItemPropertyNode = PropertyNode->AsItemPropertyNode();
 
 		TSharedPtr<SWidget> NameWidget;
-		const bool bHasStaticName = InArgs._StaticDisplayName.Len() > 0;
+		const bool bHasStaticName = !InArgs._StaticDisplayName.IsEmpty();
 		if ( CategoryNode != NULL )
 		{
 			NameWidget =
diff --git a/Engine/Source/Editor/PropertyEditor/Private/UserInterface/PropertyTable/SColumnHeader.h b/Engine/Source/Editor/PropertyEditor/Private/UserInterface/PropertyTable/SColumnHeader.h
--- a/Engine/Source/Editor/PropertyEditor/Private/UserInterface/PropertyTable/SColumnHeader.h
+++ b/Engine/Source/Editor/PropertyEditor/Private/UserInterface/PropertyTable/SColumnHeader.h
@@ -61,7 +61,7 @@ protected:
 		else if ( PropertyPath.IsValid() && PropertyPath->GetNumProperties() > 0 )
 		{
 			UProperty* Property = PropertyPath->GetLeafMostProperty().Property.Get();
-			const FText ToolTipText = FText::FromString( PropertyEditorHelpers::GetToolTipText( Property ) );
+			const FText ToolTipText = PropertyEditorHelpers::GetToolTipText( Property );
 			const FString DocumentationLink = PropertyEditorHelpers::GetDocumentationLink( Property );
 			const FString DocumentationExcerptName = PropertyEditorHelpers::GetDocumentationExcerptName( Property );
 			ToolTip = IDocumentation::Get()->CreateToolTip( ToolTipText, NULL, DocumentationLink, DocumentationExcerptName );
diff --git a/Engine/Source/Editor/PropertyEditor/Private/UserInterface/PropertyTable/TextPropertyTableCellPresenter.cpp b/Engine/Source/Editor/PropertyEditor/Private/UserInterface/PropertyTable/TextPropertyTableCellPresenter.cpp
--- a/Engine/Source/Editor/PropertyEditor/Private/UserInterface/PropertyTable/TextPropertyTableCellPresenter.cpp
+++ b/Engine/Source/Editor/PropertyEditor/Private/UserInterface/PropertyTable/TextPropertyTableCellPresenter.cpp
@@ -156,7 +156,7 @@ TSharedRef<SWidget> FTextPropertyTableCellPresenter::ConstructEditModeCellWidget
 		HasReadOnlyEditingWidget = true;
 	}
 
-	PropertyWidget->SetToolTipText( TAttribute<FString>(PropertyEditor->GetToolTipText()) );
+	PropertyWidget->SetToolTipText( PropertyEditor->GetToolTipText() );
 
 	return PropertyWidget.ToSharedRef();
 }
diff --git a/Engine/Source/Editor/PropertyEditor/Private/UserInterface/PropertyTree/SPropertyTreeCategoryRow.h b/Engine/Source/Editor/PropertyEditor/Private/UserInterface/PropertyTree/SPropertyTreeCategoryRow.h
--- a/Engine/Source/Editor/PropertyEditor/Private/UserInterface/PropertyTree/SPropertyTreeCategoryRow.h
+++ b/Engine/Source/Editor/PropertyEditor/Private/UserInterface/PropertyTree/SPropertyTreeCategoryRow.h
@@ -10,7 +10,7 @@ public:
 	SLATE_BEGIN_ARGS( SPropertyTreeCategoryRow )
 		: _DisplayName()
 	{}
-		SLATE_ARGUMENT( FString, DisplayName )
+		SLATE_ARGUMENT( FText, DisplayName )
 
 	SLATE_END_ARGS()
 
diff --git a/Engine/Source/Editor/PropertyEditor/Public/DetailCategoryBuilder.h b/Engine/Source/Editor/PropertyEditor/Public/DetailCategoryBuilder.h
--- a/Engine/Source/Editor/PropertyEditor/Public/DetailCategoryBuilder.h
+++ b/Engine/Source/Editor/PropertyEditor/Public/DetailCategoryBuilder.h
@@ -86,7 +86,7 @@ public:
 	 * @param FilterString	 A string which is used to filter this custom row when a user types into the details panel search box
 	 * @param bForAdvanced	 Whether the widget should appear in the advanced section
 	 */
-	virtual FDetailWidgetRow& AddCustomRow( const FString& FilterString, bool bForAdvanced = false ) = 0;
+	virtual FDetailWidgetRow& AddCustomRow( const FText& FilterString, bool bForAdvanced = false ) = 0;
 
 	/**
 	 * Adds a custom builder to the category (for more complicated layouts)
@@ -104,7 +104,7 @@ public:
 	 * @param true if the group should appear in the advanced section of the category
 	 * @param true if the group should start expanded
 	 */
-	virtual IDetailGroup& AddGroup( FName GroupName, const FString& LocalizedDisplayName, bool bForAdvanced = false, bool bStartExpanded = false ) = 0;
+	virtual IDetailGroup& AddGroup( FName GroupName, const FText& LocalizedDisplayName, bool bForAdvanced = false, bool bStartExpanded = false ) = 0;
 
 	/**
 	 * Gets the default propeties of this category
@@ -123,5 +123,5 @@ public:
 	/**
 	 * @return The localized display name of the category
 	 */
-	virtual const FString& GetDisplayName() const = 0;
+	virtual const FText& GetDisplayName() const = 0;
 };
\ No newline at end of file
diff --git a/Engine/Source/Editor/PropertyEditor/Public/DetailLayoutBuilder.h b/Engine/Source/Editor/PropertyEditor/Public/DetailLayoutBuilder.h
--- a/Engine/Source/Editor/PropertyEditor/Public/DetailLayoutBuilder.h
+++ b/Engine/Source/Editor/PropertyEditor/Public/DetailLayoutBuilder.h
@@ -71,7 +71,7 @@ public:
 	 * @param NewLocalizedDisplayName	The new display name of the category (optional)
 	 * @param CategoryType				Category type to define sort order.  Category display order is sorted by this type (optional)
 	 */
-	virtual IDetailCategoryBuilder& EditCategory( FName CategoryName, const FString& NewLocalizedDisplayName = TEXT(""), ECategoryPriority::Type CategoryType = ECategoryPriority::Default ) = 0;
+	virtual IDetailCategoryBuilder& EditCategory( FName CategoryName, const FText& NewLocalizedDisplayName = FText::GetEmpty(), ECategoryPriority::Type CategoryType = ECategoryPriority::Default ) = 0;
 
 	
 	/**
diff --git a/Engine/Source/Editor/PropertyEditor/Public/DetailWidgetRow.h b/Engine/Source/Editor/PropertyEditor/Public/DetailWidgetRow.h
--- a/Engine/Source/Editor/PropertyEditor/Public/DetailWidgetRow.h
+++ b/Engine/Source/Editor/PropertyEditor/Public/DetailWidgetRow.h
@@ -110,7 +110,7 @@ public:
 	/**
 	 * Sets a string which should be used to filter the content when a user searches
 	 */
-	FDetailWidgetRow& FilterString( const FString& InFilterString )
+	FDetailWidgetRow& FilterString( const FText& InFilterString )
 	{
 		FilterTextString = InFilterString;
 		return *this;
@@ -195,7 +195,7 @@ public:
 	/** IsEnabled of the row */
 	TAttribute<bool> IsEnabledAttr;
 	/** String to filter with */
-	FString FilterTextString;
+	FText FilterTextString;
 	/** Action for coping data on this row */
 	FUIAction CopyMenuAction;
 	/** Action for pasting data on this row */
diff --git a/Engine/Source/Editor/PropertyEditor/Public/IDetailChildrenBuilder.h b/Engine/Source/Editor/PropertyEditor/Public/IDetailChildrenBuilder.h
--- a/Engine/Source/Editor/PropertyEditor/Public/IDetailChildrenBuilder.h
+++ b/Engine/Source/Editor/PropertyEditor/Public/IDetailChildrenBuilder.h
@@ -27,7 +27,7 @@ public:
 	 * @param LocalizedDisplayName	The display name of the group
 	 * @param true if the group should appear in the advanced section of the category
 	 */
-	virtual class IDetailGroup& AddChildGroup( FName GroupName, const FString& LocalizedDisplayName ) = 0;
+	virtual class IDetailGroup& AddChildGroup( FName GroupName, const FText& LocalizedDisplayName ) = 0;
 
 	/**
 	 * Adds new custom content as a child to the struct
@@ -35,7 +35,7 @@ public:
 	 * @param SearchString	Search string that will be matched when users search in the details panel.  If the search string doesnt match what the user types, this row will be hidden
 	 * @return A row that accepts widgets
 	 */
-	virtual class FDetailWidgetRow& AddChildContent( const FString& SearchString ) = 0;
+	virtual class FDetailWidgetRow& AddChildContent( const FText& SearchString ) = 0;
 
 	/**
 	 * Adds a property to the struct
diff --git a/Engine/Source/Editor/PropertyEditor/Public/IDetailPropertyRow.h b/Engine/Source/Editor/PropertyEditor/Public/IDetailPropertyRow.h
--- a/Engine/Source/Editor/PropertyEditor/Public/IDetailPropertyRow.h
+++ b/Engine/Source/Editor/PropertyEditor/Public/IDetailPropertyRow.h
@@ -15,14 +15,14 @@ public:
 	 *
 	 * @param InDisplayName	The name of the property
 	 */
-	virtual IDetailPropertyRow& DisplayName( const FString& InDisplayName ) = 0;
+	virtual IDetailPropertyRow& DisplayName( const FText& InDisplayName ) = 0;
 
 	/**
 	 * Sets the localized tooltip of the property
 	 *
 	 * @param InToolTip	The tooltip of the property
 	 */
-	virtual IDetailPropertyRow& ToolTip( const FString& InToolTip ) = 0;
+	virtual IDetailPropertyRow& ToolTip( const FText& InToolTip ) = 0;
 
 	/**
 	 * Sets whether or not we show the default property editing buttons for this property
diff --git a/Engine/Source/Editor/PropertyEditor/Public/ISinglePropertyView.h b/Engine/Source/Editor/PropertyEditor/Public/ISinglePropertyView.h
--- a/Engine/Source/Editor/PropertyEditor/Public/ISinglePropertyView.h
+++ b/Engine/Source/Editor/PropertyEditor/Public/ISinglePropertyView.h
@@ -9,7 +9,7 @@
 struct FSinglePropertyParams
 {
 	/** Override for the property name that will be displayed instead of the property name */
-	FString NameOverride;
+	FText NameOverride;
 
 	/** Font to use instead of the default property font */
 	FSlateFontInfo Font;
@@ -21,7 +21,7 @@ struct FSinglePropertyParams
 	EPropertyNamePlacement::Type NamePlacement;
 		
 	FSinglePropertyParams()
-		: NameOverride()
+		: NameOverride(FText::GetEmpty())
 		, Font()
 		, NotifyHook( NULL )
 		, NamePlacement( EPropertyNamePlacement::Left )
diff --git a/Engine/Source/Editor/PropertyEditor/Public/PropertyCustomizationHelpers.h b/Engine/Source/Editor/PropertyEditor/Public/PropertyCustomizationHelpers.h
--- a/Engine/Source/Editor/PropertyEditor/Public/PropertyCustomizationHelpers.h
+++ b/Engine/Source/Editor/PropertyEditor/Public/PropertyCustomizationHelpers.h
@@ -162,7 +162,7 @@ public:
 		, _DisplayResetToDefault( true )
 		{}
 		/** The display name to use in the default property widget */
-		SLATE_ATTRIBUTE( FString, DisplayName )
+		SLATE_ATTRIBUTE( FText, DisplayName )
 		/** Whether or not to display the property name */
 		SLATE_ARGUMENT( bool, ShouldDisplayName )
 		/** The widget to display for this property instead of the default */
@@ -228,7 +228,7 @@ public:
 		ArrayProperty->SetOnNumElementsChanged( Empty );
 	}
 
-	void SetDisplayName( const FString& InDisplayName )
+	void SetDisplayName( const FText& InDisplayName )
 	{
 		DisplayName = InDisplayName;
 	}
@@ -256,10 +256,10 @@ public:
 			const bool bDisplayResetToDefaultInNameContent = false;
 			TSharedPtr<SHorizontalBox> ContentHorizontalBox;
 			NodeRow
-			.FilterString(DisplayName.Len() > 0 ? DisplayName : BaseProperty->GetPropertyDisplayName())
+			.FilterString(!DisplayName.IsEmpty() ? DisplayName : BaseProperty->GetPropertyDisplayName())
 			.NameContent()
 			[
-				BaseProperty->CreatePropertyNameWidget(DisplayName, TEXT(""), bDisplayResetToDefaultInNameContent)
+				BaseProperty->CreatePropertyNameWidget(DisplayName, FText::GetEmpty(), bDisplayResetToDefaultInNameContent)
 			]
 			.ValueContent()
 			[
@@ -310,7 +310,7 @@ protected:
 		OnRebuildChildren.ExecuteIfBound();
 	}
 private:
-	FString DisplayName;
+	FText DisplayName;
 	FOnGenerateArrayElementWidget OnGenerateArrayElementWidgetDelegate;
 	TSharedPtr<IPropertyHandleArray> ArrayProperty;
 	TSharedRef<IPropertyHandle> BaseProperty;
diff --git a/Engine/Source/Editor/PropertyEditor/Public/PropertyEditorModule.h b/Engine/Source/Editor/PropertyEditor/Public/PropertyEditorModule.h
--- a/Engine/Source/Editor/PropertyEditor/Public/PropertyEditorModule.h
+++ b/Engine/Source/Editor/PropertyEditor/Public/PropertyEditorModule.h
@@ -232,7 +232,7 @@ public:
 	 */
 	virtual TSharedPtr<class ISinglePropertyView> CreateSingleProperty( UObject* InObject, FName InPropertyName, const struct FSinglePropertyParams& InitParams );
 
-	virtual TSharedRef<class IStructureDetailsView> CreateStructureDetailView(const struct FDetailsViewArgs& DetailsViewArgs, TSharedPtr<class FStructOnScope> StructData, bool bShowObjects, const FString& CustomName = FString());
+	virtual TSharedRef<class IStructureDetailsView> CreateStructureDetailView(const struct FDetailsViewArgs& DetailsViewArgs, TSharedPtr<class FStructOnScope> StructData, bool bShowObjects, const FText& CustomName = FText::GetEmpty());
 
 	/**
 	 * Creates a property change listener that notifies users via a  delegate when a property on an object changes
