support WER keys to change gizmo mode for Bone Controller Preview in Anim blueprint

diff --git a/Engine/Source/Editor/AnimGraph/Classes/AnimGraphNode_ModifyBone.h b/Engine/Source/Editor/AnimGraph/Classes/AnimGraphNode_ModifyBone.h
--- a/Engine/Source/Editor/AnimGraph/Classes/AnimGraphNode_ModifyBone.h
+++ b/Engine/Source/Editor/AnimGraph/Classes/AnimGraphNode_ModifyBone.h
@@ -28,7 +28,8 @@ protected:
 
 	ANIMGRAPH_API virtual FVector GetWidgetLocation(const USkeletalMeshComponent* SkelComp, struct FAnimNode_SkeletalControlBase* AnimNode) override;
 	ANIMGRAPH_API virtual int32 GetWidgetMode(const USkeletalMeshComponent* SkelComp) override;
-	ANIMGRAPH_API virtual int32 ChangeWidgetMode(const USkeletalMeshComponent* SkelComp, int32 InCurWidgetMode) override;
+	ANIMGRAPH_API virtual int32 ChangeToNextWidgetMode(const USkeletalMeshComponent* SkelComp, int32 InCurWidgetMode) override;
+	ANIMGRAPH_API virtual bool SetWidgetMode(const USkeletalMeshComponent* SkelComp, int32 InWidgetMode) override;
 
 	ANIMGRAPH_API virtual FName FindSelectedBone() override;
 
diff --git a/Engine/Source/Editor/AnimGraph/Classes/AnimGraphNode_SkeletalControlBase.h b/Engine/Source/Editor/AnimGraph/Classes/AnimGraphNode_SkeletalControlBase.h
--- a/Engine/Source/Editor/AnimGraph/Classes/AnimGraphNode_SkeletalControlBase.h
+++ b/Engine/Source/Editor/AnimGraph/Classes/AnimGraphNode_SkeletalControlBase.h
@@ -42,8 +42,11 @@ public:
 
 	// return current widget mode this anim graph node supports
 	ANIMGRAPH_API virtual int32 GetWidgetMode(const USkeletalMeshComponent* SkelComp);
-	// called when the user changed widget mode
-	ANIMGRAPH_API virtual int32 ChangeWidgetMode(const USkeletalMeshComponent* SkelComp, int32 CurWidgetMode);
+	// called when the user changed widget mode by pressing "Space" key
+	ANIMGRAPH_API virtual int32 ChangeToNextWidgetMode(const USkeletalMeshComponent* SkelComp, int32 CurWidgetMode);
+	// called when the user set widget mode directly, returns true if InWidgetMode is available
+	ANIMGRAPH_API virtual bool SetWidgetMode(const USkeletalMeshComponent* SkelComp, int32 InWidgetMode){ return false; }
+
 	// 
 	ANIMGRAPH_API virtual FName FindSelectedBone();
 
diff --git a/Engine/Source/Editor/AnimGraph/Private/AnimGraphNode_ModifyBone.cpp b/Engine/Source/Editor/AnimGraph/Private/AnimGraphNode_ModifyBone.cpp
--- a/Engine/Source/Editor/AnimGraph/Private/AnimGraphNode_ModifyBone.cpp
+++ b/Engine/Source/Editor/AnimGraph/Private/AnimGraphNode_ModifyBone.cpp
@@ -168,7 +168,7 @@ int32 UAnimGraphNode_ModifyBone::GetWidgetMode(const USkeletalMeshComponent* Ske
 	return (int32)FWidget::WM_None;
 }
 
-int32 UAnimGraphNode_ModifyBone::ChangeWidgetMode(const USkeletalMeshComponent* SkelComp, int32 InCurWidgetMode)
+int32 UAnimGraphNode_ModifyBone::ChangeToNextWidgetMode(const USkeletalMeshComponent* SkelComp, int32 InCurWidgetMode)
 {
 	int32 NextWidgetMode = GetNextWidgetMode(InCurWidgetMode);
 	CurWidgetMode = FindValidWidgetMode(NextWidgetMode);
@@ -176,6 +176,18 @@ int32 UAnimGraphNode_ModifyBone::ChangeWidgetMode(const USkeletalMeshComponent*
 	return CurWidgetMode;
 }
 
+bool UAnimGraphNode_ModifyBone::SetWidgetMode(const USkeletalMeshComponent* SkelComp, int32 InWidgetMode)
+{
+	// if InWidgetMode is available 
+	if (FindValidWidgetMode(InWidgetMode) == InWidgetMode)
+	{
+		CurWidgetMode = InWidgetMode;
+		return true;
+	}
+
+	return false;
+}
+
 FName UAnimGraphNode_ModifyBone::FindSelectedBone()
 {
 	return Node.BoneToModify.BoneName;
diff --git a/Engine/Source/Editor/AnimGraph/Private/AnimGraphNode_SkeletalControlBase.cpp b/Engine/Source/Editor/AnimGraph/Private/AnimGraphNode_SkeletalControlBase.cpp
--- a/Engine/Source/Editor/AnimGraph/Private/AnimGraphNode_SkeletalControlBase.cpp
+++ b/Engine/Source/Editor/AnimGraph/Private/AnimGraphNode_SkeletalControlBase.cpp
@@ -21,7 +21,7 @@ int32 UAnimGraphNode_SkeletalControlBase::GetWidgetMode(const USkeletalMeshCompo
 	return  (int32)FWidget::EWidgetMode::WM_None;
 }
 
-int32 UAnimGraphNode_SkeletalControlBase::ChangeWidgetMode(const USkeletalMeshComponent* SkelComp, int32 CurWidgetMode)
+int32 UAnimGraphNode_SkeletalControlBase::ChangeToNextWidgetMode(const USkeletalMeshComponent* SkelComp, int32 CurWidgetMode)
 {
 	return GetWidgetMode(SkelComp);
 }
diff --git a/Engine/Source/Editor/Persona/Private/AnimationEditorViewportClient.cpp b/Engine/Source/Editor/Persona/Private/AnimationEditorViewportClient.cpp
--- a/Engine/Source/Editor/Persona/Private/AnimationEditorViewportClient.cpp
+++ b/Engine/Source/Editor/Persona/Private/AnimationEditorViewportClient.cpp
@@ -1657,7 +1657,7 @@ bool FAnimationViewportClient::InputKey( FViewport* Viewport, int32 ControllerId
 		{
 			if (SelectedSkelControlAnimGraph.IsValid())
 			{
-				WidgetMode = (FWidget::EWidgetMode)SelectedSkelControlAnimGraph->ChangeWidgetMode(PreviewSkelMeshComp.Get(), WidgetMode);
+				WidgetMode = (FWidget::EWidgetMode)SelectedSkelControlAnimGraph->ChangeToNextWidgetMode(PreviewSkelMeshComp.Get(), WidgetMode);
 				if (WidgetMode == FWidget::WM_Scale)
 				{
 					SetWidgetCoordSystemSpace(COORD_Local);
@@ -1697,7 +1697,27 @@ bool FAnimationViewportClient::InputKey( FViewport* Viewport, int32 ControllerId
 
 void FAnimationViewportClient::SetWidgetMode(FWidget::EWidgetMode InMode)
 {
-	WidgetMode = InMode;
+	bool bAnimBPMode = PersonaPtr.Pin()->IsModeCurrent(FPersonaModes::AnimBlueprintEditMode);
+
+	// support WER keys to change gizmo mode for Bone Controller preivew in anim blueprint
+	if (bAnimBPMode)
+	{
+		if (!bManipulating)
+		{
+			int32 SelectedBone = FindSelectedBone();
+			if (SelectedBone >= 0 && SelectedSkelControlAnimGraph.IsValid())
+			{
+				if (SelectedSkelControlAnimGraph->SetWidgetMode(PreviewSkelMeshComp.Get(), InMode))
+				{
+					WidgetMode = InMode;
+				}
+			}
+		}
+	}
+	else
+	{
+		WidgetMode = InMode;
+	}
 
 	// Force viewport to redraw
 	Viewport->Invalidate();
