#Anim: persona: Fixed space doesn't toggle to scale

diff --git a/Engine/Source/Editor/Persona/Private/AnimationEditorViewportClient.cpp b/Engine/Source/Editor/Persona/Private/AnimationEditorViewportClient.cpp
--- a/Engine/Source/Editor/Persona/Private/AnimationEditorViewportClient.cpp
+++ b/Engine/Source/Editor/Persona/Private/AnimationEditorViewportClient.cpp
@@ -1669,6 +1669,10 @@ bool FAnimationViewportClient::InputKey( FViewport* Viewport, int32 ControllerId
 			}
 			else
 			if (WidgetMode == FWidget::WM_Rotate)
+			{
+				WidgetMode = FWidget::WM_Scale;
+			}
+			else if (WidgetMode == FWidget::WM_Scale)
 			{
 				WidgetMode = FWidget::WM_Translate;
 			}
