Fix issue where toggling editor doesn't work

diff --git a/Engine/Source/Editor/Persona/Private/SAnimTrackCurvePanel.cpp b/Engine/Source/Editor/Persona/Private/SAnimTrackCurvePanel.cpp
--- a/Engine/Source/Editor/Persona/Private/SAnimTrackCurvePanel.cpp
+++ b/Engine/Source/Editor/Persona/Private/SAnimTrackCurvePanel.cpp
@@ -764,6 +764,8 @@ void SAnimTrackCurvePanel::ToggleEditability(ECheckBoxState NewType, USkeleton::
 			Curve->SetCurveTypeFlag(ACF_Editable, bEdit);
 		}
 	}
+
+	UpdatePanel();
 }
 
 FReply		SAnimTrackCurvePanel::RefreshPanel()
