Fix for crash when selecting anim blueprint nodes with preview ability

diff --git a/Engine/Source/Editor/Persona/Private/AnimationEditorViewportClient.cpp b/Engine/Source/Editor/Persona/Private/AnimationEditorViewportClient.cpp
--- a/Engine/Source/Editor/Persona/Private/AnimationEditorViewportClient.cpp
+++ b/Engine/Source/Editor/Persona/Private/AnimationEditorViewportClient.cpp
@@ -618,16 +618,19 @@ FAnimNode_SkeletalControlBase* FAnimationViewportClient::FindSkeletalControlAnim
 	{
 		// find an anim node index from debug data
 		UAnimBlueprintGeneratedClass* AnimBlueprintClass = Cast<UAnimBlueprintGeneratedClass>(PreviewSkelMeshComp->GetAnimInstance()->GetClass());
-		FAnimBlueprintDebugData& DebugData = AnimBlueprintClass->GetAnimBlueprintDebugData();
-		int32* IndexPtr = DebugData.NodePropertyToIndexMap.Find(AnimGraphNode);
-
-		if (IndexPtr)
+		if (AnimBlueprintClass)
 		{
-			int32 AnimNodeIndex = *IndexPtr;
-			// reverse node index temporarily because of a bug in NodeGuidToIndexMap
-			AnimNodeIndex = AnimBlueprintClass->AnimNodeProperties.Num() - AnimNodeIndex - 1;
+			FAnimBlueprintDebugData& DebugData = AnimBlueprintClass->GetAnimBlueprintDebugData();
+			int32* IndexPtr = DebugData.NodePropertyToIndexMap.Find(AnimGraphNode);
+
+			if (IndexPtr)
+			{
+				int32 AnimNodeIndex = *IndexPtr;
+				// reverse node index temporarily because of a bug in NodeGuidToIndexMap
+				AnimNodeIndex = AnimBlueprintClass->AnimNodeProperties.Num() - AnimNodeIndex - 1;
 
-			AnimNode = AnimBlueprintClass->AnimNodeProperties[AnimNodeIndex]->ContainerPtrToValuePtr<FAnimNode_SkeletalControlBase>(PreviewSkelMeshComp->GetAnimInstance());
+				AnimNode = AnimBlueprintClass->AnimNodeProperties[AnimNodeIndex]->ContainerPtrToValuePtr<FAnimNode_SkeletalControlBase>(PreviewSkelMeshComp->GetAnimInstance());
+			}
 		}
 	}
 
