Fixed blendspace input preview not working when opened first

https://jira.ol.epicgames.net/browse/UE-7159

diff --git a/Engine/Source/Editor/Persona/Private/SAnimationBlendSpaceBase.cpp b/Engine/Source/Editor/Persona/Private/SAnimationBlendSpaceBase.cpp
--- a/Engine/Source/Editor/Persona/Private/SAnimationBlendSpaceBase.cpp
+++ b/Engine/Source/Editor/Persona/Private/SAnimationBlendSpaceBase.cpp
@@ -1184,6 +1184,9 @@ void SBlendSpaceEditorBase::Construct(const FArguments& InArgs)
 	BlendSpace = InArgs._BlendSpace;
 	PersonaPtr = InArgs._Persona;
 	PersonaPtr.Pin()->RegisterOnPostUndo(FPersona::FOnPostUndo::CreateSP( this, &SBlendSpaceEditorBase::PostUndo ) );
+
+	bIsActiveTimerRegistered = true;
+	RegisterActiveTimer(0.f, FWidgetActiveTimerDelegate::CreateSP(this, &SBlendSpaceEditorBase::UpdatePreview));
 }
 
 EActiveTimerReturnType SBlendSpaceEditorBase::UpdatePreview( double InCurrentTime, float InDeltaTime )
