Animation: Allow deselecting bones by clicking on a bone not in the mesh

diff --git a/Engine/Source/Editor/Persona/Private/Persona.cpp b/Engine/Source/Editor/Persona/Private/Persona.cpp
--- a/Engine/Source/Editor/Persona/Private/Persona.cpp
+++ b/Engine/Source/Editor/Persona/Private/Persona.cpp
@@ -1983,30 +1983,30 @@ void FPersona::SetSelectedBone(USkeleton* InTargetSkeleton, const FName& BoneNam
 	{
 		if (UDebugSkelMeshComponent* Preview = GetPreviewMeshComponent())
 		{
-			// need to get mesh bone base since BonesOfInterest is saved in SkeletalMeshComponent
-			// and it is used by renderer. It is not Skeleton base
-			const int32 MeshBoneIndex = Preview->GetBoneIndex(BoneName);
+			Preview->BonesOfInterest.Empty();
+			ClearSelectedSocket();
 
-			if (MeshBoneIndex != INDEX_NONE)
+			// Add in bone of interest only if we have a preview instance set-up
+			if (Preview->PreviewInstance != NULL)
 			{
-				Preview->BonesOfInterest.Empty();
-				ClearSelectedSocket();
+				// need to get mesh bone base since BonesOfInterest is saved in SkeletalMeshComponent
+				// and it is used by renderer. It is not Skeleton base
+				const int32 MeshBoneIndex = Preview->GetBoneIndex(BoneName);
 
-				if (Preview->PreviewInstance != NULL)
+				if (MeshBoneIndex != INDEX_NONE)
 				{
-					// Add in bone of interest only if we have a preview instance set-up
 					Preview->BonesOfInterest.Add(MeshBoneIndex);
+				}
 
-					if ( bRebroadcast )
-					{
-						// Broadcast that a bone has been selected
-						OnBoneSelected.Broadcast( BoneName );
-					}
+				if ( bRebroadcast )
+				{
+					// Broadcast that a bone has been selected
+					OnBoneSelected.Broadcast( BoneName );
+				}
 
-					if( Viewport.IsValid() )
-					{
-						Viewport.Pin()->GetLevelViewportClient().Invalidate();
-					}
+				if( Viewport.IsValid() )
+				{
+					Viewport.Pin()->GetLevelViewportClient().Invalidate();
 				}
 			}
 		}
