Temporary fix for issue with attached assets not showing in Persona

diff --git a/Engine/Source/Editor/Persona/Private/Persona.cpp b/Engine/Source/Editor/Persona/Private/Persona.cpp
--- a/Engine/Source/Editor/Persona/Private/Persona.cpp
+++ b/Engine/Source/Editor/Persona/Private/Persona.cpp
@@ -190,6 +190,10 @@ FPersona::FPersona()
 	OnPropertyChangedHandleDelegateHandle = FCoreUObjectDelegates::OnObjectPropertyChanged.Add(OnPropertyChangedHandle);
 
 	GEditor->OnBlueprintPreCompile().AddRaw(this, &FPersona::OnBlueprintPreCompile);
+
+	//Temporary fix for missing attached assets - MDW
+	PreviewScene.GetWorld()->GetWorldSettings()->SetIsTemporarilyHiddenInEditor(false);
+
 }
 
 FPersona::~FPersona()
