Fixed blueprint notifies deriving from native classes that derive from notify classes not showing in add notify menu.

diff --git a/Engine/Source/Editor/Persona/Private/SAnimNotifyPanel.cpp b/Engine/Source/Editor/Persona/Private/SAnimNotifyPanel.cpp
--- a/Engine/Source/Editor/Persona/Private/SAnimNotifyPanel.cpp
+++ b/Engine/Source/Editor/Persona/Private/SAnimNotifyPanel.cpp
@@ -204,6 +204,13 @@ public:
 		, _TrackColor(FLinearColor::White)
 		, _OnSelectionChanged()
 		, _OnUpdatePanel()
+		, _OnGetNotifyBlueprintData()
+		, _OnGetNotifyStateBlueprintData()
+		, _OnGetNotifyNativeClasses()
+		, _OnGetNotifyStateNativeClasses()
+		, _OnGetScrubValue()
+		, _OnGetDraggedNodePos()
+		, _OnNodeDragStarted()
 		, _OnRequestTrackPan()
 		, _OnRequestOffsetRefresh()
 		, _OnDeleteNotify()
@@ -225,6 +232,8 @@ public:
 		SLATE_EVENT( FOnUpdatePanel, OnUpdatePanel )
 		SLATE_EVENT( FOnGetBlueprintNotifyData, OnGetNotifyBlueprintData )
 		SLATE_EVENT( FOnGetBlueprintNotifyData, OnGetNotifyStateBlueprintData )
+		SLATE_EVENT( FOnGetNativeNotifyClasses, OnGetNotifyNativeClasses )
+		SLATE_EVENT( FOnGetNativeNotifyClasses, OnGetNotifyStateNativeClasses )
 		SLATE_EVENT( FOnGetScrubValue, OnGetScrubValue )
 		SLATE_EVENT( FOnGetDraggedNodePos, OnGetDraggedNodePos )
 		SLATE_EVENT( FOnNotifyNodesDragStarted, OnNodeDragStarted )
@@ -428,6 +437,8 @@ protected:
 	FOnUpdatePanel							OnUpdatePanel;
 	FOnGetBlueprintNotifyData				OnGetNotifyBlueprintData;
 	FOnGetBlueprintNotifyData				OnGetNotifyStateBlueprintData;
+	FOnGetNativeNotifyClasses				OnGetNotifyNativeClasses;
+	FOnGetNativeNotifyClasses				OnGetNotifyStateNativeClasses;
 	FOnGetScrubValue						OnGetScrubValue;
 	FOnGetDraggedNodePos					OnGetDraggedNodePos;
 	FOnNotifyNodesDragStarted				OnNodeDragStarted;
@@ -502,6 +513,8 @@ public:
 	SLATE_EVENT( FOnUpdatePanel, OnUpdatePanel )
 	SLATE_EVENT( FOnGetBlueprintNotifyData, OnGetNotifyBlueprintData )
 	SLATE_EVENT( FOnGetBlueprintNotifyData, OnGetNotifyStateBlueprintData )
+	SLATE_EVENT( FOnGetNativeNotifyClasses, OnGetNotifyNativeClasses )
+	SLATE_EVENT( FOnGetNativeNotifyClasses, OnGetNotifyStateNativeClasses )
 	SLATE_EVENT( FOnNotifyNodesDragStarted, OnNodeDragStarted )
 	SLATE_EVENT( FRefreshOffsetsRequest, OnRequestRefreshOffsets )
 	SLATE_EVENT( FDeleteNotify, OnDeleteNotify )
@@ -1542,6 +1555,8 @@ void SAnimNotifyTrack::Construct(const FArguments& InArgs)
 	OnUpdatePanel = InArgs._OnUpdatePanel;
 	OnGetNotifyBlueprintData = InArgs._OnGetNotifyBlueprintData;
 	OnGetNotifyStateBlueprintData = InArgs._OnGetNotifyStateBlueprintData;
+	OnGetNotifyNativeClasses = InArgs._OnGetNotifyNativeClasses;
+	OnGetNotifyStateNativeClasses = InArgs._OnGetNotifyStateNativeClasses;
 	TrackIndex = InArgs._TrackIndex;
 	OnGetScrubValue = InArgs._OnGetScrubValue;
 	OnGetDraggedNodePos = InArgs._OnGetDraggedNodePos;
@@ -1747,6 +1762,11 @@ void SAnimNotifyTrack::FillNewNotifyStateMenu(FMenuBuilder& MenuBuilder)
 {
 	const static FText Description = LOCTEXT("AddsAnExistingAnimNotify", "Add an existing notify");
 
+	// Run the native query first to update the allowed classes for blueprints.
+	TArray<UClass*> NotifyStateClasses;
+	OnGetNotifyStateNativeClasses.ExecuteIfBound(NotifyStateClasses);
+
+	// Collect blueprint notify data
 	TArray<FAssetData> NotifyAssetData;
 	TArray<BlueprintNotifyMenuInfo> NotifyMenuData;
 	OnGetNotifyStateBlueprintData.ExecuteIfBound(NotifyAssetData);
@@ -1767,25 +1787,19 @@ void SAnimNotifyTrack::FillNewNotifyStateMenu(FMenuBuilder& MenuBuilder)
 
 	MenuBuilder.BeginSection("NativeNotifyStates", LOCTEXT("NewStateNotifyMenu_Native", "Native Notify States"));
 	{
-		for(TObjectIterator<UClass> It ; It ; ++It)
+		for(UClass* Class : NotifyStateClasses)
 		{
-			UClass* Class = *It;
-
-			if(Class->IsChildOf(UAnimNotifyState::StaticClass()) && !Class->HasAllClassFlags(CLASS_Abstract) && !Class->IsInBlueprint())
-			{
-				// Found a native anim notify class (that isn't UAnimNotify)
-				const FText Description = LOCTEXT("NewNotifyStateSubMenu_NativeToolTip", "Add an existing native notify state");
-				FString Label = Class->GetDisplayNameText().ToString();
-				const FText LabelText = FText::FromString(Label);
+			const FText Description = LOCTEXT("NewNotifyStateSubMenu_NativeToolTip", "Add an existing native notify state");
+			FString Label = Class->GetDisplayNameText().ToString();
+			const FText LabelText = FText::FromString(Label);
 
-				FUIAction UIAction;
-				UIAction.ExecuteAction.BindRaw(
-					this, &SAnimNotifyTrack::CreateNewNotifyAtCursor,
-					Label,
-					Class);
+			FUIAction UIAction;
+			UIAction.ExecuteAction.BindRaw(
+				this, &SAnimNotifyTrack::CreateNewNotifyAtCursor,
+				Label,
+				Class);
 
-				MenuBuilder.AddMenuEntry(LabelText, Description, FSlateIcon(), UIAction);
-			}
+			MenuBuilder.AddMenuEntry(LabelText, Description, FSlateIcon(), UIAction);
 		}
 	}
 	MenuBuilder.EndSection();
@@ -1795,6 +1809,9 @@ void SAnimNotifyTrack::FillNewNotifyMenu(FMenuBuilder& MenuBuilder)
 {
 	const static FText Description = LOCTEXT("NewNotifySubMenu_ToolTip", "Add an existing notify");
 
+	TArray<UClass*> NativeNotifyClasses;
+	OnGetNotifyNativeClasses.ExecuteIfBound(NativeNotifyClasses);
+
 	TArray<FAssetData> NotifyAssetData;
 	TArray<BlueprintNotifyMenuInfo> NotifyMenuData;
 	OnGetNotifyBlueprintData.ExecuteIfBound(NotifyAssetData);
@@ -1815,25 +1832,19 @@ void SAnimNotifyTrack::FillNewNotifyMenu(FMenuBuilder& MenuBuilder)
 
 	MenuBuilder.BeginSection("NativeNotifies", LOCTEXT("NewNotifyMenu_Native", "Native Notifies"));
 	{
-		for(TObjectIterator<UClass> It ; It ; ++It)
+		for(UClass* Class : NativeNotifyClasses)
 		{
-			UClass* Class = *It;
-
-			if(Class->IsChildOf(UAnimNotify::StaticClass()) && !Class->HasAllClassFlags(CLASS_Abstract) && !Class->IsInBlueprint())
-			{
-				// Found a native anim notify class (that isn't UAnimNotify)
-				const FText Description = LOCTEXT("NewNotifySubMenu_NativeToolTip", "Add an existing native notify");
-				FString Label = Class->GetName();
-				const FText LabelText = FText::FromString(Label);
+			const FText Description = LOCTEXT("NewNotifySubMenu_NativeToolTip", "Add an existing native notify");
+			FString Label = Class->GetName();
+			const FText LabelText = FText::FromString(Label);
 
-				FUIAction UIAction;
-				UIAction.ExecuteAction.BindRaw(
-					this, &SAnimNotifyTrack::CreateNewNotifyAtCursor,
-					Label,
-					Class);
+			FUIAction UIAction;
+			UIAction.ExecuteAction.BindRaw(
+				this, &SAnimNotifyTrack::CreateNewNotifyAtCursor,
+				Label,
+				Class);
 
-				MenuBuilder.AddMenuEntry(LabelText, Description, FSlateIcon(), UIAction);
-			}
+			MenuBuilder.AddMenuEntry(LabelText, Description, FSlateIcon(), UIAction);
 		}
 	}
 	MenuBuilder.EndSection();
@@ -2947,6 +2958,8 @@ void SNotifyEdTrack::Construct(const FArguments& InArgs)
 				.OnUpdatePanel(InArgs._OnUpdatePanel)
 				.OnGetNotifyBlueprintData(InArgs._OnGetNotifyBlueprintData)
 				.OnGetNotifyStateBlueprintData(InArgs._OnGetNotifyStateBlueprintData)
+				.OnGetNotifyNativeClasses(InArgs._OnGetNotifyNativeClasses)
+				.OnGetNotifyStateNativeClasses(InArgs._OnGetNotifyStateNativeClasses)
 				.OnGetScrubValue(InArgs._OnGetScrubValue)
 				.OnGetDraggedNodePos(InArgs._OnGetDraggedNodePos)
 				.OnNodeDragStarted(InArgs._OnNodeDragStarted)
@@ -3246,6 +3259,8 @@ void SAnimNotifyPanel::RefreshNotifyTracks()
 			.OnUpdatePanel(this, &SAnimNotifyPanel::Update)
 			.OnGetNotifyBlueprintData(this, &SAnimNotifyPanel::OnGetNotifyBlueprintData, &NotifyClassNames)
 			.OnGetNotifyStateBlueprintData(this, &SAnimNotifyPanel::OnGetNotifyBlueprintData, &NotifyStateClassNames)
+			.OnGetNotifyNativeClasses(this, &SAnimNotifyPanel::OnGetNativeNotifyData, UAnimNotify::StaticClass(), &NotifyClassNames)
+			.OnGetNotifyStateNativeClasses(this, &SAnimNotifyPanel::OnGetNativeNotifyData, UAnimNotifyState::StaticClass(), &NotifyStateClassNames)
 			.OnSelectionChanged(this, &SAnimNotifyPanel::OnTrackSelectionChanged)
 			.OnNodeDragStarted(this, &SAnimNotifyPanel::OnNotifyNodeDragStarted)
 			.MarkerBars(MarkerBars)
@@ -3764,13 +3779,23 @@ void SAnimNotifyPanel::OnGetNotifyBlueprintData(TArray<FAssetData>& OutNotifyDat
 		CurrentClassCount = InOutAllowedClassNames->Num();
 	}
 
-	// One less here to include the base class originally added
-	if(FoundClasses.Num() < InOutAllowedClassNames->Num() - 1)
+	// Count native classes, so we don't remove them from the list
+	int32 NumNativeClasses = 0;
+	for(FString& AllowedClass : *InOutAllowedClassNames)
+	{
+		if(!AllowedClass.EndsWith(FString(TEXT("_C'"))))
+		{
+			++NumNativeClasses;
+		}
+	}
+
+	if(FoundClasses.Num() < InOutAllowedClassNames->Num() - NumNativeClasses)
 	{
 		// Less classes found, some may have been deleted or reparented
 		for(int32 ClassIndex = InOutAllowedClassNames->Num() - 1 ; ClassIndex >= 0 ; --ClassIndex)
 		{
-			if(!FoundClasses.Contains((*InOutAllowedClassNames)[ClassIndex]))
+			FString& ClassName = (*InOutAllowedClassNames)[ClassIndex];
+			if(ClassName.EndsWith(FString(TEXT("_C'"))) && !FoundClasses.Contains(ClassName))
 			{
 				InOutAllowedClassNames->RemoveAt(ClassIndex);
 			}
@@ -3778,6 +3803,22 @@ void SAnimNotifyPanel::OnGetNotifyBlueprintData(TArray<FAssetData>& OutNotifyDat
 	}
 }
 
+void SAnimNotifyPanel::OnGetNativeNotifyData(TArray<UClass*>& OutClasses, UClass* NotifyOutermost, TArray<FString>* OutAllowedBlueprintClassNames)
+{
+	for(TObjectIterator<UClass> It ; It ; ++It)
+	{
+		UClass* Class = *It;
+
+		if(Class->IsChildOf(NotifyOutermost) && Class->HasAllClassFlags(CLASS_Native) && !Class->HasAllClassFlags(CLASS_Abstract) && !Class->IsInBlueprint())
+		{
+			OutClasses.Add(Class);
+			// Form class name to search later
+			FString ClassName = FString::Printf(TEXT("%s'%s'"), *Class->GetClass()->GetName(), *Class->GetPathName());
+			OutAllowedBlueprintClassNames->AddUnique(ClassName);
+		}
+	}
+}
+
 void SAnimNotifyPanel::OnNotifyObjectChanged(UObject* EditorBaseObj, bool bRebuild)
 {
 
diff --git a/Engine/Source/Editor/Persona/Private/SAnimNotifyPanel.h b/Engine/Source/Editor/Persona/Private/SAnimNotifyPanel.h
--- a/Engine/Source/Editor/Persona/Private/SAnimNotifyPanel.h
+++ b/Engine/Source/Editor/Persona/Private/SAnimNotifyPanel.h
@@ -19,6 +19,7 @@ DECLARE_DELEGATE( FDeleteNotify )
 DECLARE_DELEGATE( FDeselectAllNotifies )
 DECLARE_DELEGATE( FCopyNotifies )
 DECLARE_DELEGATE_OneParam( FOnGetBlueprintNotifyData, TArray<FAssetData>& )
+DECLARE_DELEGATE_OneParam( FOnGetNativeNotifyClasses, TArray<UClass*>&)
 
 class SAnimNotifyNode;
 class SAnimNotifyTrack;
@@ -266,6 +267,14 @@ private:
 	 */
 	void OnGetNotifyBlueprintData(TArray<FAssetData>& OutNotifyData, TArray<FString>* InOutAllowedClassNames);
 
+	/** Find classes that inherit from NotifyOutermost and add correctly formatted class name to OutAllowedBlueprintClassNames
+	 *  to allow us to find blueprints inherited from those types without loading the blueprints.
+	 *	@param OutClasses Array of classes that inherit from NotifyOutermost
+	 *	@param NotifyOutermost Outermost notify class to detect children of
+	 *	@param OutAllowedBlueprintClassNames list of class names to add the native class names to
+	 */
+	void OnGetNativeNotifyData(TArray<UClass*>& OutClasses, UClass* NotifyOutermost, TArray<FString>* OutAllowedBlueprintClassNames);
+
 	/** Persona reference **/
 	TWeakPtr<FPersona> PersonaPtr;
 
