FCurveOwnerInterface - Add a method for validating that a previously retrieved curve is still valid for editing.  This is needed to fix cases where curves become invalid due to undo/redo.

diff --git a/Engine/Source/Editor/DetailCustomizations/Private/CurveColorCustomization.cpp b/Engine/Source/Editor/DetailCustomizations/Private/CurveColorCustomization.cpp
--- a/Engine/Source/Editor/DetailCustomizations/Private/CurveColorCustomization.cpp
+++ b/Engine/Source/Editor/DetailCustomizations/Private/CurveColorCustomization.cpp
@@ -236,6 +236,15 @@ bool FCurveColorCustomization::HasAnyAlphaKeys() const
 	return false;
 }
 
+bool FCurveColorCustomization::IsValidCurve( FRichCurveEditInfo CurveInfo )
+{
+	return 
+		CurveInfo.CurveToEdit == &RuntimeCurve->ColorCurves[0] ||
+		CurveInfo.CurveToEdit == &RuntimeCurve->ColorCurves[1] ||
+		CurveInfo.CurveToEdit == &RuntimeCurve->ColorCurves[2] ||
+		CurveInfo.CurveToEdit == &RuntimeCurve->ColorCurves[3];
+}
+
 float FCurveColorCustomization::GetTimelineLength() const
 {
 	return 0.f;
diff --git a/Engine/Source/Editor/DetailCustomizations/Private/CurveColorCustomization.h b/Engine/Source/Editor/DetailCustomizations/Private/CurveColorCustomization.h
--- a/Engine/Source/Editor/DetailCustomizations/Private/CurveColorCustomization.h
+++ b/Engine/Source/Editor/DetailCustomizations/Private/CurveColorCustomization.h
@@ -37,6 +37,7 @@ public:
 	virtual bool IsLinearColorCurve() const override { return true; }
 	virtual FLinearColor GetLinearColorValue(float InTime) const override;
 	virtual bool HasAnyAlphaKeys() const override;
+	virtual bool IsValidCurve( FRichCurveEditInfo CurveInfo );
 
 private:
 	/**
diff --git a/Engine/Source/Editor/DetailCustomizations/Private/CurveStructCustomization.cpp b/Engine/Source/Editor/DetailCustomizations/Private/CurveStructCustomization.cpp
--- a/Engine/Source/Editor/DetailCustomizations/Private/CurveStructCustomization.cpp
+++ b/Engine/Source/Editor/DetailCustomizations/Private/CurveStructCustomization.cpp
@@ -229,6 +229,11 @@ void FCurveStructCustomization::OnCurveChanged()
 	StructPropertyHandle->NotifyPostChange();
 }
 
+bool FCurveStructCustomization::IsValidCurve( FRichCurveEditInfo CurveInfo )
+{
+	return CurveInfo.CurveToEdit == &RuntimeCurve->EditorCurveData;
+}
+
 float FCurveStructCustomization::GetTimelineLength() const
 {
 	return 0.f;
diff --git a/Engine/Source/Editor/DetailCustomizations/Private/CurveStructCustomization.h b/Engine/Source/Editor/DetailCustomizations/Private/CurveStructCustomization.h
--- a/Engine/Source/Editor/DetailCustomizations/Private/CurveStructCustomization.h
+++ b/Engine/Source/Editor/DetailCustomizations/Private/CurveStructCustomization.h
@@ -34,6 +34,7 @@ public:
 	virtual void ModifyOwner() override;
 	virtual void MakeTransactional() override;
 	virtual void OnCurveChanged() override;
+	virtual bool IsValidCurve( FRichCurveEditInfo CurveInfo );
 
 private:
 	/**
diff --git a/Engine/Source/Editor/DetailCustomizations/Private/WheeledVehicleMovementComponent4WDetails.cpp b/Engine/Source/Editor/DetailCustomizations/Private/WheeledVehicleMovementComponent4WDetails.cpp
--- a/Engine/Source/Editor/DetailCustomizations/Private/WheeledVehicleMovementComponent4WDetails.cpp
+++ b/Engine/Source/Editor/DetailCustomizations/Private/WheeledVehicleMovementComponent4WDetails.cpp
@@ -130,6 +130,11 @@ void FWheeledVehicleMovementComponent4WDetails::FSteeringCurveEditor::MakeTransa
 	}
 }
 
+bool FWheeledVehicleMovementComponent4WDetails::FSteeringCurveEditor::IsValidCurve( FRichCurveEditInfo CurveInfo )
+{
+	return CurveInfo.CurveToEdit == &VehicleComponent->SteeringCurve.EditorCurveData;
+}
+
 FWheeledVehicleMovementComponent4WDetails::FSteeringCurveEditor::FSteeringCurveEditor(UWheeledVehicleMovementComponent4W * InVehicle)
 {
 	VehicleComponent = InVehicle;
@@ -174,6 +179,11 @@ void FWheeledVehicleMovementComponent4WDetails::FTorqueCurveEditor::MakeTransact
 	}
 }
 
+bool FWheeledVehicleMovementComponent4WDetails::FTorqueCurveEditor::IsValidCurve( FRichCurveEditInfo CurveInfo )
+{
+	return CurveInfo.CurveToEdit == &VehicleComponent->EngineSetup.TorqueCurve.EditorCurveData;
+}
+
 FWheeledVehicleMovementComponent4WDetails::FTorqueCurveEditor::FTorqueCurveEditor(UWheeledVehicleMovementComponent4W * InVehicle)
 {
 	VehicleComponent = InVehicle;
diff --git a/Engine/Source/Editor/DetailCustomizations/Private/WheeledVehicleMovementComponent4WDetails.h b/Engine/Source/Editor/DetailCustomizations/Private/WheeledVehicleMovementComponent4WDetails.h
--- a/Engine/Source/Editor/DetailCustomizations/Private/WheeledVehicleMovementComponent4WDetails.h
+++ b/Engine/Source/Editor/DetailCustomizations/Private/WheeledVehicleMovementComponent4WDetails.h
@@ -27,6 +27,7 @@ private:
 		virtual void ModifyOwner() override;
 		virtual void MakeTransactional() override;
 		virtual void OnCurveChanged() override { }
+		virtual bool IsValidCurve(FRichCurveEditInfo CurveInfo) override;
 
 	private:
 		UWheeledVehicleMovementComponent4W * VehicleComponent;
@@ -44,6 +45,7 @@ private:
 		virtual void ModifyOwner() override;
 		virtual void MakeTransactional() override;
 		virtual void OnCurveChanged() override { }
+		virtual bool IsValidCurve( FRichCurveEditInfo CurveInfo ) override;
 
 	private:
 		UWheeledVehicleMovementComponent4W * VehicleComponent;
diff --git a/Engine/Source/Editor/Persona/Private/SAnimCurvePanel.cpp b/Engine/Source/Editor/Persona/Private/SAnimCurvePanel.cpp
--- a/Engine/Source/Editor/Persona/Private/SAnimCurvePanel.cpp
+++ b/Engine/Source/Editor/Persona/Private/SAnimCurvePanel.cpp
@@ -22,13 +22,15 @@ class FAnimCurveBaseInterface : public FCurveOwnerInterface
 {
 public:
 	TWeakObjectPtr<UAnimSequenceBase>	BaseSequence;
+	USkeleton::AnimCurveUID CurveUID;
 	FAnimCurveBase*	CurveData;
 
 public:
-	FAnimCurveBaseInterface(UAnimSequenceBase * BaseSeq, FAnimCurveBase*	InData)
+	FAnimCurveBaseInterface( UAnimSequenceBase * BaseSeq, USkeleton::AnimCurveUID InCurveUID)
 		: BaseSequence(BaseSeq)
-		, CurveData(InData)
+		, CurveUID(InCurveUID)
 	{
+		CurveData = BaseSequence.Get()->RawCurveData.GetCurveData( CurveUID );
 		// they should be valid
 		check (BaseSequence.IsValid());
 		check (CurveData);
@@ -103,6 +105,15 @@ public:
 	virtual void OnCurveChanged() override
 	{
 	}
+
+	virtual bool IsValidCurve( FRichCurveEditInfo CurveInfo ) override
+	{
+		// Get the curve with the ID directly from the sequence and compare it since undo/redo can cause previously
+		// used curves to become invalid.
+		FAnimCurveBase* CurrentCurveData = BaseSequence.Get()->RawCurveData.GetCurveData( CurveUID );
+		return CurrentCurveData != nullptr &&
+			CurveInfo.CurveToEdit == &((FFloatCurve*)CurrentCurveData)->FloatCurve;
+	}
 };
 //////////////////////////////////////////////////////////////////////////
 //  SCurveEd Track : each track for curve editing 
@@ -197,7 +208,7 @@ void SCurveEdTrack::Construct(const FArguments& InArgs)
 	FAnimCurveBase * Curve = Sequence->RawCurveData.GetCurveData(InArgs._CurveUid);
 	check (Curve);
 
-	CurveInterface = new FAnimCurveBaseInterface(Sequence, Curve);
+	CurveInterface = new FAnimCurveBaseInterface(Sequence, InArgs._CurveUid);
 	int32 NumberOfKeys = Sequence->GetNumberOfFrames();
 	//////////////////////////////
 	
diff --git a/Engine/Source/Editor/Persona/Private/SAnimTrackCurvePanel.cpp b/Engine/Source/Editor/Persona/Private/SAnimTrackCurvePanel.cpp
--- a/Engine/Source/Editor/Persona/Private/SAnimTrackCurvePanel.cpp
+++ b/Engine/Source/Editor/Persona/Private/SAnimTrackCurvePanel.cpp
@@ -33,15 +33,17 @@ class FAnimTrackCurveBaseInterface : public FCurveOwnerInterface
 {
 public:
 	TWeakObjectPtr<UAnimSequence>	BaseSequence;
+	USkeleton::AnimCurveUID CurveUID;
 	FAnimCurveBase*	CurveData;
 	ETransformCurve::Type	CurveType;
 
 public:
-	FAnimTrackCurveBaseInterface(UAnimSequence * BaseSeq, FAnimCurveBase*	InData, ETransformCurve::Type InCurveType)
+	FAnimTrackCurveBaseInterface(UAnimSequence * BaseSeq, USkeleton::AnimCurveUID InCurveUID, ETransformCurve::Type InCurveType)
 		: BaseSequence(BaseSeq)
-		, CurveData(InData)
+		, CurveUID(InCurveUID)
 		, CurveType(InCurveType)
 	{
+		CurveData = GetCurveDataFromSequence();
 		// they should be valid
 		check (BaseSequence.IsValid());
 		check (CurveData);
@@ -164,6 +166,41 @@ public:
 	virtual void OnCurveChanged() override
 	{
 	}
+
+	virtual bool IsValidCurve( FRichCurveEditInfo CurveInfo ) override
+	{
+		// Get the curves with the ID directly from the sequence and compare it since undo/redo can cause previously
+		// used curves to become invalid.
+		FAnimCurveBase* CurrentCurveData = GetCurveDataFromSequence();
+		if (CurrentCurveData != nullptr)
+		{
+			FVectorCurve * CurrentVectorCurveData = (FVectorCurve*)(CurrentCurveData);
+			return
+				CurveInfo.CurveToEdit == &CurrentVectorCurveData->FloatCurves[0] ||
+				CurveInfo.CurveToEdit == &CurrentVectorCurveData->FloatCurves[1] ||
+				CurveInfo.CurveToEdit == &CurrentVectorCurveData->FloatCurves[2];
+		}
+		return false;
+	}
+
+private:
+	FAnimCurveBase*	GetCurveDataFromSequence()
+	{
+		FTransformCurve* Curve = static_cast<FTransformCurve*> (BaseSequence.Get()->RawCurveData.GetCurveData( CurveUID, FRawCurveTracks::TransformType));
+		if (Curve)
+		{
+			switch ( CurveType )
+			{
+			case ETransformCurve::Translation:
+				return &Curve->TranslationCurve;
+			case ETransformCurve::Rotation:
+				return &Curve->RotationCurve;
+			case ETransformCurve::Scale:
+				return &Curve->ScaleCurve;
+			}
+		}
+		return nullptr;
+	}
 };
 
 //////////////////////////////////////////////////////////////////////////
@@ -247,9 +284,9 @@ void STransformCurveEdTrack::Construct(const FArguments& InArgs)
 
 	CurveUid = InArgs._CurveUid;
 
-	CurveInterfaces[ETransformCurve::Translation] = new FAnimTrackCurveBaseInterface(Sequence, &Curve->TranslationCurve, ETransformCurve::Translation);
-	CurveInterfaces[ETransformCurve::Rotation] = new FAnimTrackCurveBaseInterface(Sequence, &Curve->RotationCurve, ETransformCurve::Rotation);
-	CurveInterfaces[ETransformCurve::Scale] = new FAnimTrackCurveBaseInterface(Sequence, &Curve->ScaleCurve, ETransformCurve::Scale);
+	CurveInterfaces[ETransformCurve::Translation] = new FAnimTrackCurveBaseInterface(Sequence, CurveUid, ETransformCurve::Translation);
+	CurveInterfaces[ETransformCurve::Rotation] = new FAnimTrackCurveBaseInterface(Sequence, CurveUid, ETransformCurve::Rotation);
+	CurveInterfaces[ETransformCurve::Scale] = new FAnimTrackCurveBaseInterface(Sequence, CurveUid, ETransformCurve::Scale);
 	int32 NumberOfKeys = Sequence->GetNumberOfFrames();
 	//////////////////////////////
 	
diff --git a/Engine/Source/Editor/UnrealEd/Private/SCurveEditor.cpp b/Engine/Source/Editor/UnrealEd/Private/SCurveEditor.cpp
--- a/Engine/Source/Editor/UnrealEd/Private/SCurveEditor.cpp
+++ b/Engine/Source/Editor/UnrealEd/Private/SCurveEditor.cpp
@@ -1983,11 +1983,11 @@ TSharedPtr<FUICommandList> SCurveEditor::GetCommands()
 bool SCurveEditor::IsValidCurve( FRichCurve* Curve ) const
 {
 	bool bIsValid = false;
-	if(Curve)
+	if(Curve && CurveOwner)
 	{
 		for(auto CurveViewModel : CurveViewModels)
 		{
-			if(CurveViewModel->CurveInfo.CurveToEdit == Curve)
+			if(CurveViewModel->CurveInfo.CurveToEdit == Curve && CurveOwner->IsValidCurve(CurveViewModel->CurveInfo))
 			{
 				bIsValid = true;
 				break;
diff --git a/Engine/Source/Runtime/Engine/Classes/Curves/CurveBase.h b/Engine/Source/Runtime/Engine/Classes/Curves/CurveBase.h
--- a/Engine/Source/Runtime/Engine/Classes/Curves/CurveBase.h
+++ b/Engine/Source/Runtime/Engine/Classes/Curves/CurveBase.h
@@ -432,6 +432,9 @@ public:
 
 	/** @return True if the curve has any alpha keys */
 	virtual bool HasAnyAlphaKeys() const { return false; }
+
+	/** Validates that a previously retrieved curve is still valid for editing. */
+	virtual bool IsValidCurve(FRichCurveEditInfo CurveInfo) = 0;
 };
 
 
@@ -479,6 +482,8 @@ public:
 	virtual void ModifyOwner() override;
 	virtual void MakeTransactional() override;
 	virtual void OnCurveChanged() override;
+
+	virtual bool IsValidCurve( FRichCurveEditInfo CurveInfo ) { return false; };
 	// End FCurveOwnerInterface
 
 	// Begin UCurveBase interface
diff --git a/Engine/Source/Runtime/Engine/Classes/Curves/CurveFloat.h b/Engine/Source/Runtime/Engine/Classes/Curves/CurveFloat.h
--- a/Engine/Source/Runtime/Engine/Classes/Curves/CurveFloat.h
+++ b/Engine/Source/Runtime/Engine/Classes/Curves/CurveFloat.h
@@ -43,6 +43,7 @@ class UCurveFloat : public UCurveBase
 	// Begin FCurveOwnerInterface
 	virtual TArray<FRichCurveEditInfoConst> GetCurves() const override;
 	virtual TArray<FRichCurveEditInfo> GetCurves() override;
+	virtual bool IsValidCurve( FRichCurveEditInfo CurveInfo ) override;
 
 	/** Determine if Curve is the same */
 	ENGINE_API bool operator == (const UCurveFloat& Curve) const;
diff --git a/Engine/Source/Runtime/Engine/Classes/Curves/CurveLinearColor.h b/Engine/Source/Runtime/Engine/Classes/Curves/CurveLinearColor.h
--- a/Engine/Source/Runtime/Engine/Classes/Curves/CurveLinearColor.h
+++ b/Engine/Source/Runtime/Engine/Classes/Curves/CurveLinearColor.h
@@ -36,6 +36,8 @@ class ENGINE_API UCurveLinearColor : public UCurveBase
 	virtual FLinearColor GetLinearColorValue(float InTime) const override;
 
 	bool HasAnyAlphaKeys() const override { return FloatCurves[3].GetNumKeys() > 0; }
+
+	virtual bool IsValidCurve( FRichCurveEditInfo CurveInfo ) override;
 	// End FCurveOwnerInterface
 
 	/** Determine if Curve is the same */
diff --git a/Engine/Source/Runtime/Engine/Classes/Curves/CurveVector.h b/Engine/Source/Runtime/Engine/Classes/Curves/CurveVector.h
--- a/Engine/Source/Runtime/Engine/Classes/Curves/CurveVector.h
+++ b/Engine/Source/Runtime/Engine/Classes/Curves/CurveVector.h
@@ -23,5 +23,7 @@ class UCurveVector : public UCurveBase
 
 	/** Determine if Curve is the same */
 	ENGINE_API bool operator == (const UCurveVector& Curve) const;
+
+	virtual bool IsValidCurve( FRichCurveEditInfo CurveInfo ) override;
 };
 
diff --git a/Engine/Source/Runtime/Engine/Private/CurveFloat.cpp b/Engine/Source/Runtime/Engine/Private/CurveFloat.cpp
--- a/Engine/Source/Runtime/Engine/Private/CurveFloat.cpp
+++ b/Engine/Source/Runtime/Engine/Private/CurveFloat.cpp
@@ -59,6 +59,11 @@ TArray<FRichCurveEditInfo> UCurveFloat::GetCurves()
 	return Curves;
 }
 
+bool UCurveFloat::IsValidCurve( FRichCurveEditInfo CurveInfo )
+{
+	return CurveInfo.CurveToEdit == &FloatCurve;
+}
+
 bool UCurveFloat::operator==( const UCurveFloat& Curve ) const
 {
 	return bIsEventCurve == Curve.bIsEventCurve && FloatCurve == Curve.FloatCurve;
diff --git a/Engine/Source/Runtime/Engine/Private/CurveLinearColor.cpp b/Engine/Source/Runtime/Engine/Private/CurveLinearColor.cpp
--- a/Engine/Source/Runtime/Engine/Private/CurveLinearColor.cpp
+++ b/Engine/Source/Runtime/Engine/Private/CurveLinearColor.cpp
@@ -84,3 +84,10 @@ bool UCurveLinearColor::operator==( const UCurveLinearColor& Curve ) const
 	return (FloatCurves[0] == Curve.FloatCurves[0]) && (FloatCurves[1] == Curve.FloatCurves[1]) && (FloatCurves[2] == Curve.FloatCurves[2]) && (FloatCurves[3] == Curve.FloatCurves[3]) ;
 }
 
+bool UCurveLinearColor::IsValidCurve( FRichCurveEditInfo CurveInfo )
+{
+	return CurveInfo.CurveToEdit == &FloatCurves[0] ||
+		CurveInfo.CurveToEdit == &FloatCurves[1] ||
+		CurveInfo.CurveToEdit == &FloatCurves[2];
+}
+
diff --git a/Engine/Source/Runtime/Engine/Private/CurveVector.cpp b/Engine/Source/Runtime/Engine/Private/CurveVector.cpp
--- a/Engine/Source/Runtime/Engine/Private/CurveVector.cpp
+++ b/Engine/Source/Runtime/Engine/Private/CurveVector.cpp
@@ -49,3 +49,10 @@ bool UCurveVector::operator==( const UCurveVector& Curve ) const
 	return (FloatCurves[0] == Curve.FloatCurves[0]) && (FloatCurves[1] == Curve.FloatCurves[1]) && (FloatCurves[2] == Curve.FloatCurves[2]) ;
 }
 
+bool UCurveVector::IsValidCurve( FRichCurveEditInfo CurveInfo )
+{
+	return CurveInfo.CurveToEdit == &FloatCurves[0] ||
+		CurveInfo.CurveToEdit == &FloatCurves[1] ||
+		CurveInfo.CurveToEdit == &FloatCurves[2];
+}
+
