Fix floor/sky disappearing in Persona

diff --git a/Engine/Source/Editor/Persona/Private/AnimationEditorPreviewScene.cpp b/Engine/Source/Editor/Persona/Private/AnimationEditorPreviewScene.cpp
new file mode 100644
--- /dev/null
+++ b/Engine/Source/Editor/Persona/Private/AnimationEditorPreviewScene.cpp
@@ -0,0 +1,67 @@
+// Copyright 1998-2015 Epic Games, Inc. All Rights Reserved.
+
+#include "PersonaPrivatePCH.h"
+
+#include "AnimationEditorPreviewScene.h"
+
+/////////////////////////////////////////////////////////////////////////
+// FAnimationEditorPreviewScene
+
+FAnimationEditorPreviewScene::FAnimationEditorPreviewScene(ConstructionValues CVS) : FPreviewScene(CVS)
+{
+	// set light options 
+	DirectionalLight->SetRelativeLocation(FVector(-1024.f, 1024.f, 2048.f));
+	DirectionalLight->SetRelativeScale3D(FVector(15.f));
+	DirectionalLight->Mobility = EComponentMobility::Stationary;
+	DirectionalLight->DynamicShadowDistanceStationaryLight = 3000.f;
+	DirectionalLight->bPrecomputedLightingIsValid = false;
+	SetLightBrightness(4.f);
+	//InSetSkyBrightness(0.5f);
+	DirectionalLight->InvalidateLightingCache();
+	DirectionalLight->RecreateRenderState_Concurrent();
+
+	// A background sky sphere
+	EditorSkyComp = NewObject<UStaticMeshComponent>(GetTransientPackage(), TEXT("EditorSkyComp"));
+	UStaticMesh * StaticMesh = LoadObject<UStaticMesh>(NULL, TEXT("/Engine/MapTemplates/Sky/SM_SkySphere.SM_SkySphere"), NULL, LOAD_None, NULL);
+	check (StaticMesh);
+	EditorSkyComp->SetStaticMesh( StaticMesh );
+	UMaterial* SkyMaterial= LoadObject<UMaterial>(NULL, TEXT("/Engine/EditorMaterials/PersonaSky.PersonaSky"), NULL, LOAD_None, NULL);
+	check (SkyMaterial);
+	EditorSkyComp->SetMaterial(0, SkyMaterial);
+	const float SkySphereScale = 1000.f;
+	const FTransform SkyTransform(FRotator(0,0,0), FVector(0,0,0), FVector(SkySphereScale));
+	AddComponent(EditorSkyComp, SkyTransform);
+
+	// now add height fog component
+
+	EditorHeightFogComponent = NewObject<UExponentialHeightFogComponent>(GetTransientPackage(), TEXT("EditorHeightFogComponent"));
+	
+	EditorHeightFogComponent->FogDensity=0.00075f;
+	EditorHeightFogComponent->FogInscatteringColor=FLinearColor(3.f,4.f,6.f,0.f)*0.3f;
+	EditorHeightFogComponent->DirectionalInscatteringExponent=16.f;
+	EditorHeightFogComponent->DirectionalInscatteringColor=FLinearColor(1.1f,0.9f,0.538427f,0.f);
+	EditorHeightFogComponent->FogHeightFalloff=0.01f;
+	EditorHeightFogComponent->StartDistance=15000.f;
+	const FTransform FogTransform(FRotator(0,0,0), FVector(3824.f,34248.f,50000.f), FVector(80.f));
+	AddComponent(EditorHeightFogComponent, FogTransform);
+
+	// add capture component for reflection
+	USphereReflectionCaptureComponent* CaptureComponent = NewObject<USphereReflectionCaptureComponent>(GetTransientPackage(), TEXT("CaptureComponent"));
+ 
+ 	const FTransform CaptureTransform(FRotator(0,0,0), FVector(0.f,0.f,100.f), FVector(1.f));
+ 	AddComponent(CaptureComponent, CaptureTransform);
+	CaptureComponent->SetCaptureIsDirty();
+	CaptureComponent->UpdateReflectionCaptureContents(GetWorld());
+
+	// now add floor
+ 	UStaticMesh* FloorMesh = LoadObject<UStaticMesh>(NULL, TEXT("/Engine/EditorMeshes/PhAT_FloorBox.PhAT_FloorBox"), NULL, LOAD_None, NULL);
+ 	check(FloorMesh);
+	EditorFloorComp = NewObject<UStaticMeshComponent>(GetTransientPackage(), TEXT("EditorFloorComp"));
+ 	EditorFloorComp->SetStaticMesh( FloorMesh );
+ 	AddComponent(EditorFloorComp, FTransform::Identity);
+ 	EditorFloorComp->SetRelativeScale3D(FVector(3.f, 3.f, 1.f));
+	UMaterial* Material= LoadObject<UMaterial>(NULL, TEXT("/Engine/EditorMaterials/PersonaFloorMat.PersonaFloorMat"), NULL, LOAD_None, NULL);
+	check(Material);
+	EditorFloorComp->SetMaterial( 0, Material );
+}
+
diff --git a/Engine/Source/Editor/Persona/Private/AnimationEditorPreviewScene.h b/Engine/Source/Editor/Persona/Private/AnimationEditorPreviewScene.h
new file mode 100644
--- /dev/null
+++ b/Engine/Source/Editor/Persona/Private/AnimationEditorPreviewScene.h
@@ -0,0 +1,19 @@
+// Copyright 1998-2015 Epic Games, Inc. All Rights Reserved.
+
+#pragma once
+
+#include "PreviewScene.h"
+
+class UStaticMeshComponent;
+class UExponentialHeightFogComponent;
+
+class FAnimationEditorPreviewScene : public FPreviewScene
+{
+public:
+	FAnimationEditorPreviewScene(ConstructionValues CVS);
+
+	/** Editor accessory components **/
+	UStaticMeshComponent*				EditorFloorComp;
+	UStaticMeshComponent*				EditorSkyComp;
+	UExponentialHeightFogComponent*		EditorHeightFogComponent;
+};
\ No newline at end of file
diff --git a/Engine/Source/Editor/Persona/Private/AnimationEditorViewportClient.cpp b/Engine/Source/Editor/Persona/Private/AnimationEditorViewportClient.cpp
--- a/Engine/Source/Editor/Persona/Private/AnimationEditorViewportClient.cpp
+++ b/Engine/Source/Editor/Persona/Private/AnimationEditorViewportClient.cpp
@@ -2,7 +2,7 @@
 
 #include "PersonaPrivatePCH.h"
 
-#include "PreviewScene.h"
+#include "AnimationEditorPreviewScene.h"
 #include "SAnimationEditorViewport.h"
 #include "Runtime/Engine/Public/Slate/SceneViewport.h"
 #include "SAnimViewportToolBar.h"
@@ -79,7 +79,7 @@ IMPLEMENT_HIT_PROXY( HPersonaBoneProxy, HHitProxy );
 /////////////////////////////////////////////////////////////////////////
 // FAnimationViewportClient
 
-FAnimationViewportClient::FAnimationViewportClient(FPreviewScene& InPreviewScene, TWeakPtr<FPersona> InPersonaPtr, const TSharedRef<SAnimationEditorViewport>& InAnimationEditorViewport)
+FAnimationViewportClient::FAnimationViewportClient(FAnimationEditorPreviewScene& InPreviewScene, TWeakPtr<FPersona> InPersonaPtr, const TSharedRef<SAnimationEditorViewport>& InAnimationEditorViewport)
 	: FEditorViewportClient(nullptr, &InPreviewScene, StaticCastSharedRef<SEditorViewport>(InAnimationEditorViewport))
 	, PersonaPtr( InPersonaPtr )
 	, bManipulating(false)
@@ -118,65 +118,10 @@ FAnimationViewportClient::FAnimationViewportClient(FPreviewScene& InPreviewScene
 		SetRealtime(false,true); // We are PIE, don't start in realtime mode
 	}
 
-	// set light options 
-	PreviewScene->DirectionalLight->SetRelativeLocation(FVector(-1024.f, 1024.f, 2048.f));
-	PreviewScene->DirectionalLight->SetRelativeScale3D(FVector(15.f));
-	PreviewScene->DirectionalLight->Mobility = EComponentMobility::Stationary;
-	PreviewScene->DirectionalLight->DynamicShadowDistanceStationaryLight = 3000.f;
-	PreviewScene->DirectionalLight->bPrecomputedLightingIsValid = false;
-	PreviewScene->SetLightBrightness(4.f);
-	//InPreviewScene->SetSkyBrightness(0.5f);
-	PreviewScene->DirectionalLight->InvalidateLightingCache();
-	PreviewScene->DirectionalLight->RecreateRenderState_Concurrent();
-
-	// A background sky sphere
-	EditorSkyComp = NewObject<UStaticMeshComponent>(GetTransientPackage(), TEXT("EditorSkyComp"));
-	UStaticMesh * StaticMesh = LoadObject<UStaticMesh>(NULL, TEXT("/Engine/MapTemplates/Sky/SM_SkySphere.SM_SkySphere"), NULL, LOAD_None, NULL);
-	check (StaticMesh);
-	EditorSkyComp->SetStaticMesh( StaticMesh );
-	UMaterial* SkyMaterial= LoadObject<UMaterial>(NULL, TEXT("/Engine/EditorMaterials/PersonaSky.PersonaSky"), NULL, LOAD_None, NULL);
-	check (SkyMaterial);
-	EditorSkyComp->SetMaterial(0, SkyMaterial);
-	const float SkySphereScale = 1000.f;
-	const FTransform SkyTransform(FRotator(0,0,0), FVector(0,0,0), FVector(SkySphereScale));
-	PreviewScene->AddComponent(EditorSkyComp, SkyTransform);
-
-	// now add height fog component
-
-	EditorHeightFogComponent = NewObject<UExponentialHeightFogComponent>(GetTransientPackage(), TEXT("EditorHeightFogComponent"));
-	
-	EditorHeightFogComponent->FogDensity=0.00075f;
-	EditorHeightFogComponent->FogInscatteringColor=FLinearColor(3.f,4.f,6.f,0.f)*0.3f;
-	EditorHeightFogComponent->DirectionalInscatteringExponent=16.f;
-	EditorHeightFogComponent->DirectionalInscatteringColor=FLinearColor(1.1f,0.9f,0.538427f,0.f);
-	EditorHeightFogComponent->FogHeightFalloff=0.01f;
-	EditorHeightFogComponent->StartDistance=15000.f;
-	const FTransform FogTransform(FRotator(0,0,0), FVector(3824.f,34248.f,50000.f), FVector(80.f));
-	PreviewScene->AddComponent(EditorHeightFogComponent, FogTransform);
-
-	// add capture component for reflection
-	USphereReflectionCaptureComponent* CaptureComponent = NewObject<USphereReflectionCaptureComponent>(GetTransientPackage(), TEXT("CaptureComponent"));
- 
- 	const FTransform CaptureTransform(FRotator(0,0,0), FVector(0.f,0.f,100.f), FVector(1.f));
- 	PreviewScene->AddComponent(CaptureComponent, CaptureTransform);
-	CaptureComponent->SetCaptureIsDirty();
-	CaptureComponent->UpdateReflectionCaptureContents(PreviewScene->GetWorld());
-
-	// now add floor
- 	UStaticMesh* FloorMesh = LoadObject<UStaticMesh>(NULL, TEXT("/Engine/EditorMeshes/PhAT_FloorBox.PhAT_FloorBox"), NULL, LOAD_None, NULL);
- 	check(FloorMesh);
-	EditorFloorComp = NewObject<UStaticMeshComponent>(GetTransientPackage(), TEXT("EditorFloorComp"));
- 	EditorFloorComp->SetStaticMesh( FloorMesh );
- 	PreviewScene->AddComponent(EditorFloorComp, FTransform::Identity);
- 	EditorFloorComp->SetRelativeScale3D(FVector(3.f, 3.f, 1.f));
-	UMaterial* Material= LoadObject<UMaterial>(NULL, TEXT("/Engine/EditorMaterials/PersonaFloorMat.PersonaFloorMat"), NULL, LOAD_None, NULL);
-	check(Material);
-	EditorFloorComp->SetMaterial( 0, Material );
-
 	// set visibility
-	EditorSkyComp->SetVisibility(ConfigOption->bShowSky);
-	EditorHeightFogComponent->SetVisibility(ConfigOption->bShowSky);
- 	EditorFloorComp->SetVisibility(ConfigOption->bShowFloor);
+	GetAnimPreviewScene()->EditorSkyComp->SetVisibility(ConfigOption->bShowSky);
+	GetAnimPreviewScene()->EditorHeightFogComponent->SetVisibility(ConfigOption->bShowSky);
+	GetAnimPreviewScene()->EditorFloorComp->SetVisibility(ConfigOption->bShowFloor);
 
 	ViewFOV = FMath::Clamp<float>(ConfigOption->ViewFOV, FOVMin, FOVMax);
 
@@ -204,20 +149,6 @@ FAnimationViewportClient::FAnimationViewportClient(FPreviewScene& InPreviewScene
 
 FAnimationViewportClient::~FAnimationViewportClient()
 {
-	if (PersonaPtr.IsValid())
-	{
-		if (EditorFloorComp)
-		{
-			PreviewScene->RemoveComponent(EditorFloorComp);
-			EditorFloorComp = NULL;
-		}
-
-		if (EditorSkyComp)
-		{
-			PreviewScene->RemoveComponent(EditorSkyComp);
-			EditorSkyComp = NULL;
-		}
-	}
 }
 
 FLinearColor FAnimationViewportClient::GetBackgroundColor() const
@@ -295,36 +226,36 @@ bool FAnimationViewportClient::IsShowingGrid() const
 
 void FAnimationViewportClient::OnToggleShowFloor()
 {
-	if (EditorFloorComp)
+	if (GetAnimPreviewScene()->EditorFloorComp)
 	{
-		EditorFloorComp->SetVisibility(!EditorFloorComp->bVisible);
+		GetAnimPreviewScene()->EditorFloorComp->SetVisibility(!GetAnimPreviewScene()->EditorFloorComp->bVisible);
 		Invalidate();
 	}
 
-	ConfigOption->SetShowFloor(EditorFloorComp->bVisible);
+	ConfigOption->SetShowFloor(GetAnimPreviewScene()->EditorFloorComp->bVisible);
 }
 
 bool FAnimationViewportClient::IsShowingFloor() const
 {
-	return (EditorFloorComp)? EditorFloorComp->bVisible : false;
+	return (GetAnimPreviewScene()->EditorFloorComp) ? GetAnimPreviewScene()->EditorFloorComp->bVisible : false;
 }
 
 void FAnimationViewportClient::OnToggleShowSky()
 {
-	if (EditorSkyComp)
+	if (GetAnimPreviewScene()->EditorSkyComp)
 	{
-		bool bNewVisibility = !EditorSkyComp->bVisible;
-		EditorSkyComp->SetVisibility(bNewVisibility);
-		EditorHeightFogComponent->SetVisibility(bNewVisibility);
+		bool bNewVisibility = !GetAnimPreviewScene()->EditorSkyComp->bVisible;
+		GetAnimPreviewScene()->EditorSkyComp->SetVisibility(bNewVisibility);
+		GetAnimPreviewScene()->EditorHeightFogComponent->SetVisibility(bNewVisibility);
 		Invalidate();
 	}
 
-	ConfigOption->SetShowSky(EditorSkyComp->bVisible);
+	ConfigOption->SetShowSky(GetAnimPreviewScene()->EditorSkyComp->bVisible);
 }
 
 bool FAnimationViewportClient::IsShowingSky() const
 {
-	return (EditorSkyComp)? EditorSkyComp->bVisible : false;
+	return (GetAnimPreviewScene()->EditorSkyComp) ? GetAnimPreviewScene()->EditorSkyComp->bVisible : false;
 }
 
 void FAnimationViewportClient::OnToggleMuteAudio()
@@ -802,7 +733,7 @@ void FAnimationViewportClient::Tick(float DeltaSeconds)
 	if (PreviewComp)
 	{
 		// Handle updating the preview component to represent the effects of root motion	
-		FBoxSphereBounds Bounds = EditorFloorComp->CalcBounds(EditorFloorComp->GetRelativeTransform());
+		FBoxSphereBounds Bounds = GetAnimPreviewScene()->EditorFloorComp->CalcBounds(GetAnimPreviewScene()->EditorFloorComp->GetRelativeTransform());
 		PreviewComp->ConsumeRootMotion(Bounds.GetBox().Min, Bounds.GetBox().Max);
 	}	
 
@@ -2115,9 +2046,9 @@ void FAnimationViewportClient::UpdateCameraSetup()
 		// Move the floor to the bottom of the bounding box of the mesh, rather than on the origin
 		FVector Bottom = PreviewSkelMeshComp->Bounds.GetBoxExtrema(0);
 
-		if ( EditorFloorComp )
+		if (GetAnimPreviewScene()->EditorFloorComp)
 		{
-			EditorFloorComp->SetWorldTransform( FTransform( FQuat::Identity, FVector( 0.0f, 0.0f, Bottom.Z + GetFloorOffset() ), FVector( 3.0f, 3.0f, 1.0f ) ) );
+			GetAnimPreviewScene()->EditorFloorComp->SetWorldTransform(FTransform(FQuat::Identity, FVector(0.0f, 0.0f, Bottom.Z + GetFloorOffset()), FVector(3.0f, 3.0f, 1.0f)));
 		}
 	}
 }
@@ -2368,5 +2299,10 @@ int32 FAnimationViewportClient::GetShowMeshStats() const
 	return ConfigOption->ShowMeshStats;
 }
 
+FAnimationEditorPreviewScene* FAnimationViewportClient::GetAnimPreviewScene() const
+{
+	return static_cast<FAnimationEditorPreviewScene*>(PreviewScene);
+}
+
 #undef LOCTEXT_NAMESPACE
 
diff --git a/Engine/Source/Editor/Persona/Private/AnimationEditorViewportClient.h b/Engine/Source/Editor/Persona/Private/AnimationEditorViewportClient.h
--- a/Engine/Source/Editor/Persona/Private/AnimationEditorViewportClient.h
+++ b/Engine/Source/Editor/Persona/Private/AnimationEditorViewportClient.h
@@ -5,6 +5,8 @@
 #include "EditorViewportClient.h"
 #include "Engine/WindDirectionalSource.h"
 
+class FAnimationEditorPreviewScene;
+
 //////////////////////////////////////////////////////////////////////////
 // ELocalAxesMode
 
@@ -50,7 +52,7 @@ protected:
 	void DisplayInfo( FCanvas* Canvas, FSceneView* View, bool bDisplayAllInfo );
 
 public:
-	FAnimationViewportClient(FPreviewScene& InPreviewScene, TWeakPtr<FPersona> InPersonaPtr, const TSharedRef<SAnimationEditorViewport>& InAnimationEditorViewport);
+	FAnimationViewportClient(FAnimationEditorPreviewScene& InPreviewScene, TWeakPtr<FPersona> InPersonaPtr, const TSharedRef<SAnimationEditorViewport>& InAnimationEditorViewport);
 	virtual ~FAnimationViewportClient();
 
 	// FEditorViewportClient interface
@@ -272,11 +274,6 @@ private:
 		MaxGridSize	= 50,
 	};
 
-	/** Editor accessory components **/
-	UStaticMeshComponent* EditorFloorComp;
-	UStaticMeshComponent* EditorSkyComp;
-	UExponentialHeightFogComponent* EditorHeightFogComponent;
-
 	/** Wind actor used for preview */
 	TWeakObjectPtr<AWindDirectionalSource> SelectedWindActor;
 
@@ -325,4 +322,6 @@ private:
 	TWeakObjectPtr<class UAnimGraphNode_SkeletalControlBase> SelectedSkelControlAnimGraph;
 	// to check whether we should update literal values in selected AnimGraphNode
 	bool bShouldUpdateDefaultValues;
+
+	FAnimationEditorPreviewScene* GetAnimPreviewScene() const;
 };
diff --git a/Engine/Source/Editor/Persona/Private/Persona.h b/Engine/Source/Editor/Persona/Private/Persona.h
--- a/Engine/Source/Editor/Persona/Private/Persona.h
+++ b/Engine/Source/Editor/Persona/Private/Persona.h
@@ -4,7 +4,7 @@
 
 #include "Editor/Kismet/Public/BlueprintEditor.h"
 
-#include "PreviewScene.h"
+#include "AnimationEditorPreviewScene.h"
 
 #include "PersonaModule.h"
 #include "PersonaCommands.h"
@@ -162,7 +162,7 @@ public:
 	void CleanupComponent(USceneComponent* Component);
 
 	/** Returns the editors preview scene */
-	FPreviewScene& GetPreviewScene() { return PreviewScene; }
+	FAnimationEditorPreviewScene& GetPreviewScene() { return PreviewScene; }
 
 	/** Gets called when it's clicked via mode tab - Reinitialize the mode **/
 	void ReinitMode();
@@ -627,7 +627,7 @@ private:
 	TSharedPtr<FExtender> ToolbarExtender;
 
 	/** Preview scene for the editor */
-	FPreviewScene PreviewScene;
+	FAnimationEditorPreviewScene PreviewScene;
 
 	/** Brush to use as a dirty marker for assets */
 	const FSlateBrush* AssetDirtyBrush;
