Fix missing floor/sky in Persona when 2 copies of Persona are open

diff --git a/Engine/Source/Editor/Persona/Private/AnimationEditorPreviewScene.cpp b/Engine/Source/Editor/Persona/Private/AnimationEditorPreviewScene.cpp
--- a/Engine/Source/Editor/Persona/Private/AnimationEditorPreviewScene.cpp
+++ b/Engine/Source/Editor/Persona/Private/AnimationEditorPreviewScene.cpp
@@ -21,7 +21,7 @@ FAnimationEditorPreviewScene::FAnimationEditorPreviewScene(ConstructionValues CV
 	DirectionalLight->RecreateRenderState_Concurrent();
 
 	// A background sky sphere
-	EditorSkyComp = NewObject<UStaticMeshComponent>(GetTransientPackage(), TEXT("EditorSkyComp"));
+	EditorSkyComp = NewObject<UStaticMeshComponent>(GetTransientPackage());
 	UStaticMesh * StaticMesh = LoadObject<UStaticMesh>(NULL, TEXT("/Engine/MapTemplates/Sky/SM_SkySphere.SM_SkySphere"), NULL, LOAD_None, NULL);
 	check (StaticMesh);
 	EditorSkyComp->SetStaticMesh( StaticMesh );
@@ -34,7 +34,7 @@ FAnimationEditorPreviewScene::FAnimationEditorPreviewScene(ConstructionValues CV
 
 	// now add height fog component
 
-	EditorHeightFogComponent = NewObject<UExponentialHeightFogComponent>(GetTransientPackage(), TEXT("EditorHeightFogComponent"));
+	EditorHeightFogComponent = NewObject<UExponentialHeightFogComponent>(GetTransientPackage());
 	
 	EditorHeightFogComponent->FogDensity=0.00075f;
 	EditorHeightFogComponent->FogInscatteringColor=FLinearColor(3.f,4.f,6.f,0.f)*0.3f;
@@ -46,7 +46,7 @@ FAnimationEditorPreviewScene::FAnimationEditorPreviewScene(ConstructionValues CV
 	AddComponent(EditorHeightFogComponent, FogTransform);
 
 	// add capture component for reflection
-	USphereReflectionCaptureComponent* CaptureComponent = NewObject<USphereReflectionCaptureComponent>(GetTransientPackage(), TEXT("CaptureComponent"));
+	USphereReflectionCaptureComponent* CaptureComponent = NewObject<USphereReflectionCaptureComponent>(GetTransientPackage());
  
  	const FTransform CaptureTransform(FRotator(0,0,0), FVector(0.f,0.f,100.f), FVector(1.f));
  	AddComponent(CaptureComponent, CaptureTransform);
@@ -56,7 +56,7 @@ FAnimationEditorPreviewScene::FAnimationEditorPreviewScene(ConstructionValues CV
 	// now add floor
  	UStaticMesh* FloorMesh = LoadObject<UStaticMesh>(NULL, TEXT("/Engine/EditorMeshes/PhAT_FloorBox.PhAT_FloorBox"), NULL, LOAD_None, NULL);
  	check(FloorMesh);
-	EditorFloorComp = NewObject<UStaticMeshComponent>(GetTransientPackage(), TEXT("EditorFloorComp"));
+	EditorFloorComp = NewObject<UStaticMeshComponent>(GetTransientPackage());
  	EditorFloorComp->SetStaticMesh( FloorMesh );
  	AddComponent(EditorFloorComp, FTransform::Identity);
  	EditorFloorComp->SetRelativeScale3D(FVector(3.f, 3.f, 1.f));
