#ANIM: CurveEditor: Display Total number of curves
Merging

//depot/UE4-GDC2015/Engine/Source/Editor/Persona/Private/SAnimCurvePanel.cpp

to //depot/UE4/Engine/Source/Editor/Persona/Private/SAnimCurvePanel.cpp

diff --git a/Engine/Source/Editor/Persona/Private/SAnimCurvePanel.cpp b/Engine/Source/Editor/Persona/Private/SAnimCurvePanel.cpp
--- a/Engine/Source/Editor/Persona/Private/SAnimCurvePanel.cpp
+++ b/Engine/Source/Editor/Persona/Private/SAnimCurvePanel.cpp
@@ -529,6 +529,17 @@ void SAnimCurvePanel::Construct(const FArguments& InArgs)
 						.OnGetMenuContent(this, &SAnimCurvePanel::GenerateCurveList)
 					]
 
+					+SHorizontalBox::Slot()
+					.HAlign(HAlign_Center)
+					.VAlign(VAlign_Center)
+					.Padding(5,0)
+					.AutoWidth()
+					[
+						SNew(STextBlock)
+						.Font(FEditorStyle::GetFontStyle("CurveEd.InfoFont"))
+						.Text(FText::FromString(FString::Printf(TEXT(" Total Number : %d "), Sequence->RawCurveData.FloatCurves.Num())))
+					]
+
 					+SHorizontalBox::Slot()
 					.HAlign(HAlign_Right)
 					.Padding( 2.0f, 0.0f )
