#ANIM: renamed "Retargeting Options" to "Advanced Options"

https://udn.unrealengine.com/questions/235807/preview-skel-meshes-in-persona-viewer.html

diff --git a/Engine/Source/Editor/Persona/Private/SSkeletonTree.cpp b/Engine/Source/Editor/Persona/Private/SSkeletonTree.cpp
--- a/Engine/Source/Editor/Persona/Private/SSkeletonTree.cpp
+++ b/Engine/Source/Editor/Persona/Private/SSkeletonTree.cpp
@@ -792,7 +792,7 @@ void SSkeletonTree::Construct(const FArguments& InArgs)
 
 	BoneFilter = EBoneFilter::All;
 	SocketFilter = ESocketFilter::Active;
-	bShowingRetargetingOptions = false;
+	bShowingAdvancedOptions = false;
 
 	PersonaPtr = InArgs._Persona;
 	IsEditable = InArgs._IsEditable;
@@ -872,12 +872,12 @@ void SSkeletonTree::Construct(const FArguments& InArgs)
 			.AutoWidth()
 			[
 				SNew(SCheckBox)
-				.IsChecked(this, &SSkeletonTree::IsShowingRetargetingOptions)
+				.IsChecked(this, &SSkeletonTree::IsShowingAdvancedOptions)
 				.ToolTipText(LOCTEXT("SocketFilterToolTip", "Change which types of sockets are shown"))
-				.OnCheckStateChanged(this, &SSkeletonTree::OnChangeShowingRetargetingOptions)
+				.OnCheckStateChanged(this, &SSkeletonTree::OnChangeShowingAdvancedOptions)
 				[
 					SNew(STextBlock)
-						.Text(LOCTEXT("ShowRetargetingOptions", "Show Retargeting Options"))
+						.Text(LOCTEXT("ShowAdvancedOptions", "Show Advanced Options"))
 				]
 			]
 		]
@@ -1063,7 +1063,7 @@ void SSkeletonTree::CreateTreeColumns()
 		.DefaultLabel(LOCTEXT("SkeletonBoneNameLabel", "Name"))
 		.FillWidth(0.75f);
 
-	if (bShowingRetargetingOptions)
+	if (bShowingAdvancedOptions)
 	{
 		TreeHeaderRow->AddColumn(
 			SHeaderRow::Column(ColumnID_RetargetingLabel)
@@ -1407,7 +1407,7 @@ TSharedPtr< SWidget > SSkeletonTree::CreateContextMenu()
 
 			MenuBuilder.EndSection();
 
-			if(bShowingRetargetingOptions)
+			if(bShowingAdvancedOptions)
 			{
 				MenuBuilder.BeginSection("SkeletonTreeBoneTranslationRetargeting", LOCTEXT("BoneTranslationRetargetingHeader", "Bone Translation Retargeting"));
 				{
@@ -2545,15 +2545,15 @@ void SSkeletonTree::AddAttachedAssets( const FPreviewAssetAttachContainer& Attac
 	}
 }
 
-void SSkeletonTree::OnChangeShowingRetargetingOptions(ECheckBoxState NewState)
+void SSkeletonTree::OnChangeShowingAdvancedOptions(ECheckBoxState NewState)
 {
-	bShowingRetargetingOptions = NewState == ECheckBoxState::Checked;
+	bShowingAdvancedOptions = NewState == ECheckBoxState::Checked;
 	CreateTreeColumns();
 }
 
-ECheckBoxState SSkeletonTree::IsShowingRetargetingOptions() const
+ECheckBoxState SSkeletonTree::IsShowingAdvancedOptions() const
 {
-	return bShowingRetargetingOptions ? ECheckBoxState::Checked : ECheckBoxState::Unchecked;
+	return bShowingAdvancedOptions ? ECheckBoxState::Checked : ECheckBoxState::Unchecked;
 }
 
 
diff --git a/Engine/Source/Editor/Persona/Private/SSkeletonTree.h b/Engine/Source/Editor/Persona/Private/SSkeletonTree.h
--- a/Engine/Source/Editor/Persona/Private/SSkeletonTree.h
+++ b/Engine/Source/Editor/Persona/Private/SSkeletonTree.h
@@ -535,10 +535,10 @@ private:
 	bool IsAddingSocketsAllowed() const;
 
 	/** Handler for "Show Retargeting Options" check box IsChecked functionality */
-	ECheckBoxState IsShowingRetargetingOptions() const;
+	ECheckBoxState IsShowingAdvancedOptions() const;
 
 	/**  Handler for when we change the "Show Retargeting Options" check box */
-	void OnChangeShowingRetargetingOptions(ECheckBoxState NewState);
+	void OnChangeShowingAdvancedOptions(ECheckBoxState NewState);
 
 	/** This replicates the socket filter to the previewcomponent so that the viewport can use the same settings */
 	void SetPreviewComponentSocketFilter() const;
@@ -597,7 +597,7 @@ private:
 	/** Current type of sockets to show */
 	ESocketFilter::Type SocketFilter;
 
-	bool bShowingRetargetingOptions;
+	bool bShowingAdvancedOptions;
 
 	/** Points to an item that is being requested to be renamed */
 	TSharedPtr<FDisplayedTreeRowInfo> DeferredRenameRequest;
