#ANIM: Fixed light glitchy when you have enable cast shadow

diff --git a/Engine/Source/Editor/Persona/Private/AnimationEditorPreviewScene.cpp b/Engine/Source/Editor/Persona/Private/AnimationEditorPreviewScene.cpp
--- a/Engine/Source/Editor/Persona/Private/AnimationEditorPreviewScene.cpp
+++ b/Engine/Source/Editor/Persona/Private/AnimationEditorPreviewScene.cpp
@@ -12,7 +12,7 @@ FAnimationEditorPreviewScene::FAnimationEditorPreviewScene(ConstructionValues CV
 	// set light options 
 	DirectionalLight->SetRelativeLocation(FVector(-1024.f, 1024.f, 2048.f));
 	DirectionalLight->SetRelativeScale3D(FVector(15.f));
-	DirectionalLight->Mobility = EComponentMobility::Stationary;
+	DirectionalLight->Mobility = EComponentMobility::Movable;
 	DirectionalLight->DynamicShadowDistanceStationaryLight = 3000.f;
 	DirectionalLight->bPrecomputedLightingIsValid = false;
 	SetLightBrightness(4.f);
