[INTEGRATE] Change 2442741 by [REDACTED_PERSON][REDACTED_HANDLE] on 2015/02/11 20:45:29

	Persona: Rename Persona 'Defaults' panel to Preview panel and rebuild it on a SSingleObjectDetailsPanel instead of an SKismetInspector, fixing selection update issues

	Blueprint Editor: Remove some now-dead methods and replace some raw strings with GET_MEMBER_NAME_CHECKED references

	Persona: Prevent editing preview properties when stopped at a breakpoint, as it was outright dangerous if the parent class had a Pre- or PostEditChange implementation, and otherwise generally unwise

	Persona: Hide back/fowards buttons in SAnimationSequenceBrowser when in Anim Blueprint edit mode, as they take up a lot of vertical space but provide no value

	[UE-9498]

diff --git a/Engine/Source/Editor/Kismet/Private/BlueprintEditor.cpp b/Engine/Source/Editor/Kismet/Private/BlueprintEditor.cpp
--- a/Engine/Source/Editor/Kismet/Private/BlueprintEditor.cpp
+++ b/Engine/Source/Editor/Kismet/Private/BlueprintEditor.cpp
@@ -72,6 +72,7 @@
 #include "AnimationTransitionGraph.h"
 #include "BlueprintEditorModes.h"
 #include "BlueprintEditorSettings.h"
+#include "K2Node_SwitchString.h"
 
 #include "EngineAnalytics.h"
 #include "IAnalyticsProvider.h"
@@ -510,11 +511,6 @@ bool FBlueprintEditor::IsCompilingEnabled() const
 	return Blueprint && Blueprint->BlueprintType != BPTYPE_MacroLibrary && InEditingMode();
 }
 
-bool FBlueprintEditor::IsPropertyEditingEnabled() const
-{
-	return true;
-}
-
 bool FBlueprintEditor::InDebuggingMode() const
 {
 	return GEditor->PlayWorld != NULL;
@@ -6609,10 +6605,13 @@ void FBlueprintEditor::NotifyPostChange(const FPropertyChangedEvent& PropertyCha
 void FBlueprintEditor::OnFinishedChangingProperties(const FPropertyChangedEvent& PropertyChangedEvent)
 {
 	FName PropertyName = (PropertyChangedEvent.Property != NULL) ? PropertyChangedEvent.Property->GetFName() : NAME_None;
-	if (PropertyName == TEXT("bHasDefaultPin") ||
-		PropertyName == TEXT("StartIndex") ||
-		PropertyName == TEXT("PinNames") ||
-		PropertyName == TEXT("bIsCaseSensitive"))
+
+	//@TODO: This code does not belong here (might not even be necessary anymore as they seem to have PostEditChangeProperty impls now)!
+	if ((PropertyName == GET_MEMBER_NAME_CHECKED(UK2Node_Switch, bHasDefaultPin)) ||
+		(PropertyName == GET_MEMBER_NAME_CHECKED(UK2Node_SwitchInteger, StartIndex)) ||
+		(PropertyName == GET_MEMBER_NAME_CHECKED(UK2Node_SwitchString, PinNames)) ||
+		(PropertyName == GET_MEMBER_NAME_CHECKED(UK2Node_SwitchName, PinNames)) ||
+		(PropertyName == GET_MEMBER_NAME_CHECKED(UK2Node_SwitchString, bIsCaseSensitive)))
 	{
 		DocumentManager->RefreshAllTabs();
 	}
diff --git a/Engine/Source/Editor/Kismet/Public/BlueprintEditor.h b/Engine/Source/Editor/Kismet/Public/BlueprintEditor.h
--- a/Engine/Source/Editor/Kismet/Public/BlueprintEditor.h
+++ b/Engine/Source/Editor/Kismet/Public/BlueprintEditor.h
@@ -275,9 +275,6 @@ public:
 	/** Returns true if able to compile */
 	bool IsCompilingEnabled() const;
 
-	/** Returns true if property editing is allowed */
-	bool IsPropertyEditingEnabled() const;
-
 	/** Returns true if the parent class is also a Blueprint */
 	bool IsParentClassOfObjectABlueprint( const UBlueprint* Blueprint ) const;
 
diff --git a/Engine/Source/Editor/Persona/Private/BlueprintMode/AnimBlueprintMode.cpp b/Engine/Source/Editor/Persona/Private/BlueprintMode/AnimBlueprintMode.cpp
--- a/Engine/Source/Editor/Persona/Private/BlueprintMode/AnimBlueprintMode.cpp
+++ b/Engine/Source/Editor/Persona/Private/BlueprintMode/AnimBlueprintMode.cpp
@@ -18,7 +18,7 @@
 FAnimBlueprintEditAppMode::FAnimBlueprintEditAppMode(TSharedPtr<FPersona> InPersona)
 	: FBlueprintEditorApplicationMode(StaticCastSharedPtr<FBlueprintEditor>(InPersona), FPersonaModes::AnimBlueprintEditMode, FPersonaModes::GetLocalizedMode, false, false)
 {
-	TabLayout = FTabManager::NewLayout( "Persona_AnimBlueprintEditMode_Layout_v6" )
+	TabLayout = FTabManager::NewLayout( "Persona_AnimBlueprintEditMode_Layout_v7" )
 		->AddArea
 		(
 			FTabManager::NewPrimaryArea()
@@ -54,7 +54,7 @@ FAnimBlueprintEditAppMode::FAnimBlueprintEditAppMode(TSharedPtr<FPersona> InPers
 						//	Left bottom - preview settings
 						FTabManager::NewStack()
 						->SetSizeCoefficient(0.5f)
-						->AddTab( FPersonaTabs::AnimBlueprintDefaultsEditorID, ETabState::OpenedTab )
+						->AddTab( FPersonaTabs::AnimBlueprintPreviewEditorID, ETabState::OpenedTab )
 						->AddTab( FPersonaTabs::AnimBlueprintParentPlayerEditorID, ETabState::ClosedTab )
 					)
 				)
@@ -110,7 +110,7 @@ FAnimBlueprintEditAppMode::FAnimBlueprintEditAppMode(TSharedPtr<FPersona> InPers
 	PersonaTabFactories.RegisterFactory(MakeShareable(new FAnimationAssetBrowserSummoner(InPersona)));
 	PersonaTabFactories.RegisterFactory(MakeShareable(new FPreviewViewportSummoner(InPersona)));
 	PersonaTabFactories.RegisterFactory(MakeShareable(new FSkeletonAnimNotifiesSummoner(InPersona)));
-	PersonaTabFactories.RegisterFactory(MakeShareable(new FAnimBlueprintDefaultsEditorSummoner(InPersona)));
+	PersonaTabFactories.RegisterFactory(MakeShareable(new FAnimBlueprintPreviewEditorSummoner(InPersona)));
 	PersonaTabFactories.RegisterFactory(MakeShareable(new FAnimBlueprintParentPlayerEditorSummoner(InPersona)));
 	PersonaTabFactories.RegisterFactory(MakeShareable(new FSkeletonSlotNamesSummoner(InPersona)));
 
diff --git a/Engine/Source/Editor/Persona/Private/Persona.cpp b/Engine/Source/Editor/Persona/Private/Persona.cpp
--- a/Engine/Source/Editor/Persona/Private/Persona.cpp
+++ b/Engine/Source/Editor/Persona/Private/Persona.cpp
@@ -68,6 +68,8 @@
 #include "SNotificationList.h"
 #include "NotificationManager.h"
 
+#include "Editor/KismetWidgets/Public/SSingleObjectDetailsPanel.h"
+
 #define LOCTEXT_NAMESPACE "FPersona"
 
 /////////////////////////////////////////////////////
@@ -81,6 +83,68 @@ struct FLocalCharEditorCallbacks
 	}
 };
 
+/////////////////////////////////////////////////////
+// SPersonaPreviewPropertyEditor
+
+class SPersonaPreviewPropertyEditor : public SSingleObjectDetailsPanel
+{
+public:
+	SLATE_BEGIN_ARGS(SPersonaPreviewPropertyEditor) {}
+	SLATE_END_ARGS()
+
+private:
+	// Pointer back to owning Persona editor instance (the keeper of state)
+	TWeakPtr<FPersona> PersonaPtr;
+public:
+	void Construct(const FArguments& InArgs, TSharedPtr<FPersona> InPersona)
+	{
+		PersonaPtr = InPersona;
+
+		SSingleObjectDetailsPanel::Construct(SSingleObjectDetailsPanel::FArguments(), /*bAutomaticallyObserveViaGetObjectToObserve*/ true, /*bAllowSearch*/ true);
+
+		PropertyView->SetIsPropertyEditingEnabledDelegate(FIsPropertyEditingEnabled::CreateStatic([] { return !GIntraFrameDebuggingGameThread; }));
+	}
+
+	// SSingleObjectDetailsPanel interface
+	virtual UObject* GetObjectToObserve() const override
+	{
+		if (UDebugSkelMeshComponent* PreviewComponent = PersonaPtr.Pin()->GetPreviewMeshComponent())
+		{
+			if (PreviewComponent->AnimScriptInstance != nullptr)
+			{
+				return PreviewComponent->AnimScriptInstance;
+			}
+		}
+
+		return nullptr;
+	}
+
+	virtual TSharedRef<SWidget> PopulateSlot(TSharedRef<SWidget> PropertyEditorWidget) override
+	{
+		return SNew(SVerticalBox)
+			+SVerticalBox::Slot()
+			.AutoHeight()
+			.Padding(0.f, 8.f, 0.f, 0.f)
+			[
+				SNew(SBorder)
+				.BorderImage(FEditorStyle::GetBrush("Persona.PreviewPropertiesWarning"))
+				[
+					SNew(STextBlock)
+					.Text(LOCTEXT("AnimBlueprintEditPreviewText", "Changes to preview options are not saved in the asset."))
+					.Font(FEditorStyle::GetFontStyle("PropertyWindow.NormalFont"))
+					.ShadowColorAndOpacity(FLinearColor::Black.CopyWithNewOpacity(0.3f))
+					.ShadowOffset(FVector2D::UnitVector)
+				]
+			]
+			+SVerticalBox::Slot()
+			.FillHeight(1)
+			[
+				PropertyEditorWidget
+			];
+	}
+	// End of SSingleObjectDetailsPanel interface
+};
+
 //////////////////////////////////////////////////////
 /////////////////////////////////////////////////////
 // FSkeletonEditAppMode
@@ -1246,22 +1310,6 @@ void FPersona::CreateDefaultCommands()
 		);
 }
 
-void FPersona::StartEditingDefaults(bool bAutoFocus, bool bForceRefresh)
-{
-	FBlueprintEditor::StartEditingDefaults(bAutoFocus, bForceRefresh);
-	if (PreviewComponent != NULL)
-	{
-		if (PreviewComponent->AnimScriptInstance != NULL)
-		{
-			PreviewEditor->ShowDetailsForSingleObject(PreviewComponent->AnimScriptInstance, SKismetInspector::FShowDetailsOptions(LOCTEXT("PreviewSettingsTitle", "Preview settings")));
-		}
-	}
-	if (bAutoFocus)
-	{
-		TabManager->InvokeTab(FPersonaTabs::AnimBlueprintDefaultsEditorID);
-	}
-}
-
 bool FPersona::CanSelectBone() const
 {
 	return true;
@@ -1907,25 +1955,7 @@ void FPersona::CreateDefaultTabContents(const TArray<UBlueprint*>& InBlueprints)
 {
 	FBlueprintEditor::CreateDefaultTabContents(InBlueprints);
 
-	UBlueprint* InBlueprint = InBlueprints.Num() == 1 ? InBlueprints[0] : NULL;
-
-	// Cache off whether or not this is an interface, since it is used to govern multiple widget's behavior
-	const bool bIsInterface = (InBlueprint && InBlueprint->BlueprintType == BPTYPE_Interface);
-
-	const bool bShowPublicView = false;
-	const bool bHideNameArea = true;
-
-	PreviewEditor = 
-		SNew(SKismetInspector)
-		. Kismet2(SharedThis(this))
-		. ViewIdentifier(FName("PersonaPreview"))
-		. SetNotifyHook(false)
-		. IsEnabled(!bIsInterface)
-		. ShowPublicViewControl(bShowPublicView)
-		. ShowTitleArea(false)
-		. HideNameArea(bHideNameArea)
-		. IsPropertyEditingEnabledDelegate( FIsPropertyEditingEnabled::CreateSP(this, &FPersona::IsPropertyEditingEnabled) )
-		. OnFinishedChangingProperties( FOnFinishedChangingProperties::FDelegate::CreateSP( this, &FPersona::OnFinishedChangingProperties ) );
+	PreviewEditor = SNew(SPersonaPreviewPropertyEditor, SharedThis(this));
 }
 
 FGraphAppearanceInfo FPersona::GetGraphAppearance() const
diff --git a/Engine/Source/Editor/Persona/Private/Persona.h b/Engine/Source/Editor/Persona/Private/Persona.h
--- a/Engine/Source/Editor/Persona/Private/Persona.h
+++ b/Engine/Source/Editor/Persona/Private/Persona.h
@@ -202,7 +202,7 @@ public:
 	/** Returns the image brush to use for each modes dirty marker */
 	const FSlateBrush* GetDirtyImageForMode(FName Mode) const;
 
-	TSharedRef<SKismetInspector> GetPreviewEditor() {return PreviewEditor.ToSharedRef();}
+	TSharedRef<SWidget> GetPreviewEditor() { return PreviewEditor.ToSharedRef(); }
 	/** Refresh Preview Instance Track Curves **/
 	void RefreshPreviewInstanceTrackCurves();
 
@@ -216,7 +216,6 @@ protected:
 	virtual bool CanAddPosePin() const override;
 	virtual void OnRemovePosePin() override;
 	virtual bool CanRemovePosePin() const override;
-	virtual void StartEditingDefaults(bool bAutoFocus, bool bForceRefresh = false) override;
 	virtual void Compile() override;
 	virtual void OnGraphEditorFocused(const TSharedRef<class SGraphEditor>& InGraphEditor) override;
 	virtual void OnConvertToSequenceEvaluator() override;
@@ -633,7 +632,7 @@ private:
 	const FSlateBrush* AssetDirtyBrush;
 
 	/** Preview instance inspector widget */
-	TSharedPtr<class SKismetInspector> PreviewEditor;
+	TSharedPtr<class SWidget> PreviewEditor;
 
 	/** Sequence Browser **/
 	TWeakPtr<class SAnimationSequenceBrowser> SequenceBrowser;
diff --git a/Engine/Source/Editor/Persona/Private/SAnimationSequenceBrowser.cpp b/Engine/Source/Editor/Persona/Private/SAnimationSequenceBrowser.cpp
--- a/Engine/Source/Editor/Persona/Private/SAnimationSequenceBrowser.cpp
+++ b/Engine/Source/Editor/Persona/Private/SAnimationSequenceBrowser.cpp
@@ -444,8 +444,36 @@ void SAnimationSequenceBrowser::Construct(const FArguments& InArgs)
 	Config.OnVisualizeAssetToolTip = FOnVisualizeAssetToolTip::CreateSP(this, &SAnimationSequenceBrowser::OnVisualizeAssetToolTip);
 	Config.OnAssetToolTipClosing = FOnAssetToolTipClosing::CreateSP( this, &SAnimationSequenceBrowser::OnAssetToolTipClosing );
 
-	TWeakPtr< SMenuAnchor > BackMenuAnchorPtr;
-	TWeakPtr< SMenuAnchor > FwdMenuAnchorPtr;	
+	TSharedRef< SMenuAnchor > BackMenuAnchorPtr = SNew(SMenuAnchor)
+		.Placement(MenuPlacement_BelowAnchor)
+		.OnGetMenuContent(this, &SAnimationSequenceBrowser::CreateHistoryMenu, true)
+		[
+			SNew(SButton)
+			.OnClicked(this, &SAnimationSequenceBrowser::OnGoBackInHistory)
+			.ButtonStyle(FEditorStyle::Get(), "GraphBreadcrumbButton")
+			.IsEnabled(this, &SAnimationSequenceBrowser::CanStepBackwardInHistory)
+			.ToolTipText(LOCTEXT("Backward_Tooltip", "Step backward in the asset history. Right click to see full history."))
+			[
+				SNew(SImage)
+				.Image(FEditorStyle::GetBrush("GraphBreadcrumb.BrowseBack"))
+			]
+		];
+
+	TSharedRef< SMenuAnchor > FwdMenuAnchorPtr = SNew(SMenuAnchor)
+		.Placement(MenuPlacement_BelowAnchor)
+		.OnGetMenuContent(this, &SAnimationSequenceBrowser::CreateHistoryMenu, false)
+		[
+			SNew(SButton)
+			.OnClicked(this, &SAnimationSequenceBrowser::OnGoForwardInHistory)
+			.ButtonStyle(FEditorStyle::Get(), "GraphBreadcrumbButton")
+			.IsEnabled(this, &SAnimationSequenceBrowser::CanStepForwardInHistory)
+			.ToolTipText(LOCTEXT("Forward_Tooltip", "Step forward in the asset history. Right click to see full history."))
+			[
+				SNew(SImage)
+				.Image(FEditorStyle::GetBrush("GraphBreadcrumb.BrowseForward"))
+			]
+		];
+
 	this->ChildSlot
 	[
 		SNew(SVerticalBox)
@@ -458,33 +486,22 @@ void SAnimationSequenceBrowser::Construct(const FArguments& InArgs)
 		.AutoHeight()
 		[
 			SNew(SSeparator)
+			.Visibility(this, &SAnimationSequenceBrowser::GetNonBlueprintModeVisibility)
 		]
 		+SVerticalBox::Slot()
 		.HAlign(HAlign_Right)
 		.AutoHeight()
 		[
 			SNew(SHorizontalBox)
-			+SHorizontalBox::Slot()
+			.Visibility(this, &SAnimationSequenceBrowser::GetNonBlueprintModeVisibility)
+			+ SHorizontalBox::Slot()
 			.AutoWidth()
 			[
 				SNew(SBorder)
-				.OnMouseButtonDown(this, &SAnimationSequenceBrowser::OnMouseDownHisory, BackMenuAnchorPtr)
+				.OnMouseButtonDown(this, &SAnimationSequenceBrowser::OnMouseDownHisory, TWeakPtr<SMenuAnchor>(BackMenuAnchorPtr))
 				.BorderImage( FEditorStyle::GetBrush("NoBorder") )
 				[
-					SAssignNew(BackMenuAnchorPtr, SMenuAnchor)
-					.Placement( MenuPlacement_BelowAnchor )
-					.OnGetMenuContent( this, &SAnimationSequenceBrowser::CreateHistoryMenu, true )
-					[
-						SNew(SButton)
-						.OnClicked( this, &SAnimationSequenceBrowser::OnGoBackInHistory )
-						.ButtonStyle( FEditorStyle::Get(), "GraphBreadcrumbButton" )
-						.IsEnabled(this, &SAnimationSequenceBrowser::CanStepBackwardInHistory)
-						.ToolTipText(LOCTEXT("Backward_Tooltip", "Step backward in the asset history. Right click to see full history."))
-						[
-							SNew(SImage)
-							.Image( FEditorStyle::GetBrush("GraphBreadcrumb.BrowseBack") )
-						]
-					]
+					BackMenuAnchorPtr
 				]
 			]
 
@@ -492,23 +509,10 @@ void SAnimationSequenceBrowser::Construct(const FArguments& InArgs)
 			.AutoWidth()
 			[
 				SNew(SBorder)
-				.OnMouseButtonDown(this, &SAnimationSequenceBrowser::OnMouseDownHisory, FwdMenuAnchorPtr)
+				.OnMouseButtonDown(this, &SAnimationSequenceBrowser::OnMouseDownHisory, TWeakPtr<SMenuAnchor>(FwdMenuAnchorPtr))
 				.BorderImage( FEditorStyle::GetBrush("NoBorder") )
 				[
-					SAssignNew(FwdMenuAnchorPtr, SMenuAnchor)
-					.Placement( MenuPlacement_BelowAnchor )
-					.OnGetMenuContent( this, &SAnimationSequenceBrowser::CreateHistoryMenu, false )
-					[
-						SNew(SButton)
-						.OnClicked( this, &SAnimationSequenceBrowser::OnGoForwardInHistory )
-						.ButtonStyle( FEditorStyle::Get(), "GraphBreadcrumbButton" )
-						.IsEnabled(this, &SAnimationSequenceBrowser::CanStepForwardInHistory)
-						.ToolTipText(LOCTEXT("Forward_Tooltip", "Step forward in the asset history. Right click to see full history."))
-						[
-							SNew(SImage)
-							.Image( FEditorStyle::GetBrush("GraphBreadcrumb.BrowseForward") )
-						]
-					]
+					FwdMenuAnchorPtr
 				]
 			]
 		]
@@ -990,6 +994,21 @@ bool SAnimationSequenceBrowser::IsToolTipPreviewVisible()
 	return bVisible;
 }
 
+EVisibility SAnimationSequenceBrowser::GetNonBlueprintModeVisibility() const
+{
+	if (PersonaPtr.IsValid())
+	{
+		if (PersonaPtr.Pin()->GetCurrentMode() == FPersonaModes::AnimBlueprintEditMode)
+		{
+			return EVisibility::Collapsed;
+		}
+	}
+
+	return EVisibility::Visible;
+}
+
+
+
 FAnimationAssetViewportClient::FAnimationAssetViewportClient(FPreviewScene& InPreviewScene)
 	: FEditorViewportClient(nullptr, &InPreviewScene)
 {
diff --git a/Engine/Source/Editor/Persona/Private/SAnimationSequenceBrowser.h b/Engine/Source/Editor/Persona/Private/SAnimationSequenceBrowser.h
--- a/Engine/Source/Editor/Persona/Private/SAnimationSequenceBrowser.h
+++ b/Engine/Source/Editor/Persona/Private/SAnimationSequenceBrowser.h
@@ -148,6 +148,9 @@ protected:
 	 */
 	bool IsToolTipPreviewVisible();
 
+	/** Returns visible when not in a Blueprint mode (anim mode, etc...) */
+	EVisibility GetNonBlueprintModeVisibility() const;
+protected:
 	/**
 	 * The actual viewport widget
 	 */	
@@ -173,8 +176,6 @@ protected:
 	 */
 	FPreviewScene PreviewScene;
 
-protected:
-
 	// Pointer back to persona tool that owns us
 	TWeakPtr<class FPersona> PersonaPtr;
 
diff --git a/Engine/Source/Editor/Persona/Private/Shared/PersonaMode.cpp b/Engine/Source/Editor/Persona/Private/Shared/PersonaMode.cpp
--- a/Engine/Source/Editor/Persona/Private/Shared/PersonaMode.cpp
+++ b/Engine/Source/Editor/Persona/Private/Shared/PersonaMode.cpp
@@ -22,7 +22,7 @@ const FName FPersonaTabs::RetargetManagerID("RetargetManager");
 // Anim Blueprint params
 // Explorer
 // Class Defaults
-const FName FPersonaTabs::AnimBlueprintDefaultsEditorID("AnimBlueprintDefaultsEditor");
+const FName FPersonaTabs::AnimBlueprintPreviewEditorID("AnimBlueprintPreviewEditor");
 
 const FName FPersonaTabs::AnimBlueprintParentPlayerEditorID("AnimBlueprintParentPlayerEditor");
 // Anim Document
@@ -248,24 +248,24 @@ TSharedRef<SWidget> FRetargetManagerTabSummoner::CreateTabBody(const FWorkflowTa
 
 
 /////////////////////////////////////////////////////
-// FAnimBlueprintDefaultsEditorSummoner
+// FAnimBlueprintPreviewEditorSummoner
 
 #include "SKismetInspector.h"
 
-FAnimBlueprintDefaultsEditorSummoner::FAnimBlueprintDefaultsEditorSummoner(TSharedPtr<class FAssetEditorToolkit> InHostingApp)
-	: FWorkflowTabFactory(FPersonaTabs::AnimBlueprintDefaultsEditorID, InHostingApp)
+FAnimBlueprintPreviewEditorSummoner::FAnimBlueprintPreviewEditorSummoner(TSharedPtr<class FAssetEditorToolkit> InHostingApp)
+	: FWorkflowTabFactory(FPersonaTabs::AnimBlueprintPreviewEditorID, InHostingApp)
 {
-	TabLabel = LOCTEXT("AnimBlueprintDefaultsTabTitle", "Anim Blueprint Editor");
+	TabLabel = LOCTEXT("AnimBlueprintPreviewTabTitle", "Anim Preview Editor");
 
 	bIsSingleton = true;
 
 	CurrentMode = EAnimBlueprintEditorMode::PreviewMode;
 
-	ViewMenuDescription = LOCTEXT("AnimBlueprintDefaultsView", "Defaults");
-	ViewMenuTooltip = LOCTEXT("AnimBlueprintDefaultsView_ToolTip", "Shows the animation class defaults/preview editor view");
+	ViewMenuDescription = LOCTEXT("AnimBlueprintPreviewView", "Preview");
+	ViewMenuTooltip = LOCTEXT("AnimBlueprintPreviewView_ToolTip", "Shows the animation preview editor view (as well as class defaults)");
 }
 
-TSharedRef<SWidget> FAnimBlueprintDefaultsEditorSummoner::CreateTabBody(const FWorkflowTabSpawnInfo& Info) const
+TSharedRef<SWidget> FAnimBlueprintPreviewEditorSummoner::CreateTabBody(const FWorkflowTabSpawnInfo& Info) const
 {
 	TSharedPtr<FPersona> PersonaPtr = StaticCastSharedPtr<FPersona>(HostingApp.Pin());
 
@@ -278,13 +278,13 @@ TSharedRef<SWidget> FAnimBlueprintDefaultsEditorSummoner::CreateTabBody(const FW
 				.Padding(FMargin( 0.f, 0.f, 2.f, 0.f ))
 				[
 					SNew(SBorder)
-					.BorderImage(this, &FAnimBlueprintDefaultsEditorSummoner::GetBorderBrushByMode, EAnimBlueprintEditorMode::PreviewMode )
+					.BorderImage(this, &FAnimBlueprintPreviewEditorSummoner::GetBorderBrushByMode, EAnimBlueprintEditorMode::PreviewMode)
 					.Padding(0)
 					[
 						SNew(SCheckBox)
 						.Style(FEditorStyle::Get(), "RadioButton")
-						.IsChecked( this, &FAnimBlueprintDefaultsEditorSummoner::IsChecked, EAnimBlueprintEditorMode::PreviewMode )
-						.OnCheckStateChanged( this, &FAnimBlueprintDefaultsEditorSummoner::OnCheckedChanged, EAnimBlueprintEditorMode::PreviewMode )
+						.IsChecked(this, &FAnimBlueprintPreviewEditorSummoner::IsChecked, EAnimBlueprintEditorMode::PreviewMode)
+						.OnCheckStateChanged(this, &FAnimBlueprintPreviewEditorSummoner::OnCheckedChanged, EAnimBlueprintEditorMode::PreviewMode)
 						.ToolTip(IDocumentation::Get()->CreateToolTip(	LOCTEXT("AnimBlueprintPropertyEditorPreviewMode", "Switch to editing the preview instance properties"),
 																		NULL,
 																		TEXT("Shared/Editors/Persona"),
@@ -300,13 +300,13 @@ TSharedRef<SWidget> FAnimBlueprintDefaultsEditorSummoner::CreateTabBody(const FW
 				.Padding(FMargin( 2.f, 0.f, 0.f, 0.f ))
 				[
 					SNew(SBorder)
-					.BorderImage(this, &FAnimBlueprintDefaultsEditorSummoner::GetBorderBrushByMode, EAnimBlueprintEditorMode::DefaultsMode )
+					.BorderImage(this, &FAnimBlueprintPreviewEditorSummoner::GetBorderBrushByMode, EAnimBlueprintEditorMode::DefaultsMode)
 					.Padding(0)
 					[
 						SNew(SCheckBox)
 						.Style(FEditorStyle::Get(), "RadioButton")
-						.IsChecked( this, &FAnimBlueprintDefaultsEditorSummoner::IsChecked, EAnimBlueprintEditorMode::DefaultsMode )
-						.OnCheckStateChanged( this, &FAnimBlueprintDefaultsEditorSummoner::OnCheckedChanged, EAnimBlueprintEditorMode::DefaultsMode )
+						.IsChecked(this, &FAnimBlueprintPreviewEditorSummoner::IsChecked, EAnimBlueprintEditorMode::DefaultsMode)
+						.OnCheckStateChanged(this, &FAnimBlueprintPreviewEditorSummoner::OnCheckedChanged, EAnimBlueprintEditorMode::DefaultsMode)
 						.ToolTip(IDocumentation::Get()->CreateToolTip(	LOCTEXT("AnimBlueprintPropertyEditorDefaultMode", "Switch to editing the class defaults"),
 																		NULL,
 																		TEXT("Shared/Editors/Persona"),
@@ -327,27 +327,9 @@ TSharedRef<SWidget> FAnimBlueprintDefaultsEditorSummoner::CreateTabBody(const FW
 					SNew(SBorder)
 					.Padding(0)
 					.BorderImage( FEditorStyle::GetBrush("NoBorder") )
-					.Visibility( this, &FAnimBlueprintDefaultsEditorSummoner::IsEditorVisible, EAnimBlueprintEditorMode::PreviewMode)
+					.Visibility(this, &FAnimBlueprintPreviewEditorSummoner::IsEditorVisible, EAnimBlueprintEditorMode::PreviewMode)
 					[
-						SNew(SVerticalBox)
-						+SVerticalBox::Slot()
-						.AutoHeight()
-						.Padding(0.f, 8.f, 0.f, 0.f)
-						[
-							SNew(SBorder)
-							.BorderImage(FEditorStyle::GetBrush("Persona.PreviewPropertiesWarning"))
-							[
-								SNew(STextBlock)
-								.Text(LOCTEXT("AnimBlueprintEditPreviewText", "Changes to preview options are not saved in the asset."))
-								.Font(FEditorStyle::GetFontStyle("PropertyWindow.NormalFont"))
-								.ShadowColorAndOpacity(FLinearColor::Black.CopyWithNewOpacity(0.3f))
-								.ShadowOffset(FVector2D::UnitVector)
-							]
-						]
-						+SVerticalBox::Slot()
-						[
-							PersonaPtr->GetPreviewEditor()
-						]
+						PersonaPtr->GetPreviewEditor()
 					]
 				]
 				+SOverlay::Slot()
@@ -355,7 +337,7 @@ TSharedRef<SWidget> FAnimBlueprintDefaultsEditorSummoner::CreateTabBody(const FW
 					SNew(SBorder)
 					.Padding(0)
 					.BorderImage( FEditorStyle::GetBrush("NoBorder") )
-					.Visibility( this, &FAnimBlueprintDefaultsEditorSummoner::IsEditorVisible, EAnimBlueprintEditorMode::DefaultsMode)
+					.Visibility(this, &FAnimBlueprintPreviewEditorSummoner::IsEditorVisible, EAnimBlueprintEditorMode::DefaultsMode)
 					[
 						PersonaPtr->GetDefaultEditor()
 					]
@@ -363,22 +345,22 @@ TSharedRef<SWidget> FAnimBlueprintDefaultsEditorSummoner::CreateTabBody(const FW
 			];
 }
 
-FText FAnimBlueprintDefaultsEditorSummoner::GetTabToolTipText(const FWorkflowTabSpawnInfo& Info) const
+FText FAnimBlueprintPreviewEditorSummoner::GetTabToolTipText(const FWorkflowTabSpawnInfo& Info) const
 {
-	return LOCTEXT("AnimBlueprintDefaultsEditorTooltip", "The editor lets you set the default values for all variables in your Blueprint or lets you change the values of the preview instance, depending on mode");
+	return LOCTEXT("AnimBlueprintPreviewEditorTooltip", "The editor lets you change the values of the preview instance");
 }
 
-EVisibility FAnimBlueprintDefaultsEditorSummoner::IsEditorVisible(EAnimBlueprintEditorMode::Type Mode) const
+EVisibility FAnimBlueprintPreviewEditorSummoner::IsEditorVisible(EAnimBlueprintEditorMode::Type Mode) const
 {
 	return CurrentMode == Mode ? EVisibility::Visible: EVisibility::Hidden;
 }
 
-ECheckBoxState FAnimBlueprintDefaultsEditorSummoner::IsChecked(EAnimBlueprintEditorMode::Type Mode) const
+ECheckBoxState FAnimBlueprintPreviewEditorSummoner::IsChecked(EAnimBlueprintEditorMode::Type Mode) const
 {
 	return CurrentMode == Mode ? ECheckBoxState::Checked: ECheckBoxState::Unchecked;
 }
 
-const FSlateBrush* FAnimBlueprintDefaultsEditorSummoner::GetBorderBrushByMode(EAnimBlueprintEditorMode::Type Mode) const
+const FSlateBrush* FAnimBlueprintPreviewEditorSummoner::GetBorderBrushByMode(EAnimBlueprintEditorMode::Type Mode) const
 {
 	if(Mode == CurrentMode)
 	{
@@ -390,7 +372,7 @@ const FSlateBrush* FAnimBlueprintDefaultsEditorSummoner::GetBorderBrushByMode(EA
 	}
 }
 
-void FAnimBlueprintDefaultsEditorSummoner::OnCheckedChanged(ECheckBoxState NewType, EAnimBlueprintEditorMode::Type Mode)
+void FAnimBlueprintPreviewEditorSummoner::OnCheckedChanged(ECheckBoxState NewType, EAnimBlueprintEditorMode::Type Mode)
 {
 	if(NewType == ECheckBoxState::Checked)
 	{
diff --git a/Engine/Source/Editor/Persona/Private/Shared/PersonaMode.h b/Engine/Source/Editor/Persona/Private/Shared/PersonaMode.h
--- a/Engine/Source/Editor/Persona/Private/Shared/PersonaMode.h
+++ b/Engine/Source/Editor/Persona/Private/Shared/PersonaMode.h
@@ -22,7 +22,7 @@ struct FPersonaTabs
 	// Anim Blueprint Params
 	// Explorer
 	// Class Defaults
-	static const FName AnimBlueprintDefaultsEditorID;
+	static const FName AnimBlueprintPreviewEditorID;
 	static const FName AnimBlueprintParentPlayerEditorID;
 	// Anim Document
 	static const FName ScrubberID;
@@ -221,7 +221,7 @@ public:
 };
 
 /////////////////////////////////////////////////////
-// FAnimBlueprintDefaultsEditorSummoner
+// FAnimBlueprintPreviewEditorSummoner
 
 namespace EAnimBlueprintEditorMode
 {
@@ -232,10 +232,10 @@ namespace EAnimBlueprintEditorMode
 	};
 }
 
-struct FAnimBlueprintDefaultsEditorSummoner : public FWorkflowTabFactory
+struct FAnimBlueprintPreviewEditorSummoner : public FWorkflowTabFactory
 {
 public:
-	FAnimBlueprintDefaultsEditorSummoner(TSharedPtr<class FAssetEditorToolkit> InHostingApp);
+	FAnimBlueprintPreviewEditorSummoner(TSharedPtr<class FAssetEditorToolkit> InHostingApp);
 
 	virtual TSharedRef<SWidget> CreateTabBody(const FWorkflowTabSpawnInfo& Info) const override;
 
