Fixed improper FText use for an enum in skeletal animation code.

diff --git a/Engine/Source/Editor/Persona/Private/SSkeletonTree.cpp b/Engine/Source/Editor/Persona/Private/SSkeletonTree.cpp
--- a/Engine/Source/Editor/Persona/Private/SSkeletonTree.cpp
+++ b/Engine/Source/Editor/Persona/Private/SSkeletonTree.cpp
@@ -337,17 +337,20 @@ TSharedRef< SWidget > FDisplayedMeshBoneInfo::CreateBoneTranslationRetargetingMo
 
 	MenuBuilder.BeginSection("BoneTranslationRetargetingMode", LOCTEXT( "BoneTranslationRetargetingModeMenuHeading", "Bone Translation Retargeting Mode" ) );
 	{
+		UEnum* const Enum = FindObject<UEnum>(ANY_PACKAGE, TEXT("EBoneTranslationRetargetingMode"), true);	
+		check(Enum);
+
 		FUIAction ActionRetargetingAnimation = FUIAction(FExecuteAction::CreateSP(this, &FDisplayedMeshBoneInfo::SetBoneTranslationRetargetingMode, EBoneTranslationRetargetingMode::Animation));
-		MenuBuilder.AddMenuEntry( FText::FromString( TargetSkeleton->GetRetargetingModeString(EBoneTranslationRetargetingMode::Animation) ), LOCTEXT( "BoneTranslationRetargetingAnimationToolTip", "Use translation from animation." ), FSlateIcon(), ActionRetargetingAnimation);
+		MenuBuilder.AddMenuEntry( Enum->GetDisplayNameText(EBoneTranslationRetargetingMode::Animation), LOCTEXT( "BoneTranslationRetargetingAnimationToolTip", "Use translation from animation." ), FSlateIcon(), ActionRetargetingAnimation);
 
 		FUIAction ActionRetargetingSkeleton = FUIAction(FExecuteAction::CreateSP(this, &FDisplayedMeshBoneInfo::SetBoneTranslationRetargetingMode, EBoneTranslationRetargetingMode::Skeleton));
-		MenuBuilder.AddMenuEntry( FText::FromString( TargetSkeleton->GetRetargetingModeString(EBoneTranslationRetargetingMode::Skeleton) ), LOCTEXT( "BoneTranslationRetargetingSkeletonToolTip", "Use translation from Skeleton." ), FSlateIcon(), ActionRetargetingSkeleton);
+		MenuBuilder.AddMenuEntry( Enum->GetDisplayNameText(EBoneTranslationRetargetingMode::Skeleton), LOCTEXT( "BoneTranslationRetargetingSkeletonToolTip", "Use translation from Skeleton." ), FSlateIcon(), ActionRetargetingSkeleton);
 
 		FUIAction ActionRetargetingLengthScale = FUIAction(FExecuteAction::CreateSP(this, &FDisplayedMeshBoneInfo::SetBoneTranslationRetargetingMode, EBoneTranslationRetargetingMode::AnimationScaled));
-		MenuBuilder.AddMenuEntry( FText::FromString( TargetSkeleton->GetRetargetingModeString(EBoneTranslationRetargetingMode::AnimationScaled) ), LOCTEXT( "BoneTranslationRetargetingAnimationScaledToolTip", "Use translation from animation, scale length by Skeleton's proportions." ), FSlateIcon(), ActionRetargetingLengthScale);
+		MenuBuilder.AddMenuEntry( Enum->GetDisplayNameText(EBoneTranslationRetargetingMode::AnimationScaled), LOCTEXT( "BoneTranslationRetargetingAnimationScaledToolTip", "Use translation from animation, scale length by Skeleton's proportions." ), FSlateIcon(), ActionRetargetingLengthScale);
 
 		FUIAction ActionRetargetingAnimationRelative = FUIAction(FExecuteAction::CreateSP(this, &FDisplayedMeshBoneInfo::SetBoneTranslationRetargetingMode, EBoneTranslationRetargetingMode::AnimationRelative));
-		MenuBuilder.AddMenuEntry(FText::FromString(TargetSkeleton->GetRetargetingModeString(EBoneTranslationRetargetingMode::AnimationRelative)), LOCTEXT("BoneTranslationRetargetingAnimationRelativeToolTip", "Use relative translation from animation similar to an additive animation."), FSlateIcon(), ActionRetargetingAnimationRelative);
+		MenuBuilder.AddMenuEntry( Enum->GetDisplayNameText(EBoneTranslationRetargetingMode::AnimationRelative), LOCTEXT("BoneTranslationRetargetingAnimationRelativeToolTip", "Use relative translation from animation similar to an additive animation."), FSlateIcon(), ActionRetargetingAnimationRelative);
 	}
 	MenuBuilder.EndSection();
 
@@ -360,7 +363,11 @@ FText FDisplayedMeshBoneInfo::GetTranslationRetargetingModeMenuTitle() const
 	if( BoneIndex != INDEX_NONE )
 	{
 		const EBoneTranslationRetargetingMode::Type RetargetingMode = TargetSkeleton->GetBoneTranslationRetargetingMode(BoneIndex);
-		return FText::FromString(TargetSkeleton->GetRetargetingModeString(RetargetingMode));
+		UEnum* const Enum = FindObject<UEnum>(ANY_PACKAGE, TEXT("EBoneTranslationRetargetingMode"), true);	
+		if (Enum)
+		{
+			return Enum->GetDisplayNameText(RetargetingMode);
+		}
 	}
 
 	return LOCTEXT("None", "None");
diff --git a/Engine/Source/Runtime/Engine/Classes/Animation/Skeleton.h b/Engine/Source/Runtime/Engine/Classes/Animation/Skeleton.h
--- a/Engine/Source/Runtime/Engine/Classes/Animation/Skeleton.h
+++ b/Engine/Source/Runtime/Engine/Classes/Animation/Skeleton.h
@@ -574,8 +574,6 @@ public:
 		return BoneTree[BoneTreeIdx].TranslationRetargetingMode;
 	}
 
-	ENGINE_API FString GetRetargetingModeString(const EBoneTranslationRetargetingMode::Type & RetargetingMode) const;
-	
 	/** 
 	 * Rebuild Look up between SkelMesh to BoneTree - this should only get called when SkelMesh is re-imported or so, where the mapping may be no longer valid
 	 *
diff --git a/Engine/Source/Runtime/Engine/Private/Animation/Skeleton.cpp b/Engine/Source/Runtime/Engine/Private/Animation/Skeleton.cpp
--- a/Engine/Source/Runtime/Engine/Private/Animation/Skeleton.cpp
+++ b/Engine/Source/Runtime/Engine/Private/Animation/Skeleton.cpp
@@ -300,44 +300,12 @@ bool USkeleton::IsCompatibleMesh(USkeletalMesh* InSkelMesh) const
 	return (NumOfBoneMatches > 0);
 }
 
-FString USkeleton::GetRetargetingModeString(const EBoneTranslationRetargetingMode::Type & RetargetingMode) const
-{
-	FText ModeNameText;
-
-	switch( RetargetingMode )
-	{
-	case EBoneTranslationRetargetingMode::Animation :
-		ModeNameText = LOCTEXT( "BoneRetargetingModeAnimation", "Animation" );
-		break;
-
-	case EBoneTranslationRetargetingMode::Skeleton :
-		ModeNameText = LOCTEXT( "BoneRetargetingModeSkeleton", "Skeleton" );
-		break;
-
-	case EBoneTranslationRetargetingMode::AnimationScaled :
-		ModeNameText = LOCTEXT( "BoneRetargetingMode", "AnimationScaled" );
-		break;
-
-	case EBoneTranslationRetargetingMode::AnimationRelative:
-		ModeNameText = LOCTEXT("BoneRetargetingModeAnimationRelative", "AnimationRelative");
-		break;
-
-	default:
-		// Unknown mode
-		check( 0 );
-		break;
-	}
-
-	return ModeNameText.ToString();
-}
-
 void USkeleton::ClearCacheData()
 {
 	LinkupCache.Empty();
 	SkelMesh2LinkupCache.Empty();
 }
 
-
 int32 USkeleton::GetMeshLinkupIndex(const USkeletalMesh* InSkelMesh)
 {
 	const int32* IndexPtr = SkelMesh2LinkupCache.Find(InSkelMesh);
