Fix bugs in notify track panning introduced by marquee selection

diff --git a/Engine/Source/Editor/Persona/Private/SAnimNotifyPanel.cpp b/Engine/Source/Editor/Persona/Private/SAnimNotifyPanel.cpp
--- a/Engine/Source/Editor/Persona/Private/SAnimNotifyPanel.cpp
+++ b/Engine/Source/Editor/Persona/Private/SAnimNotifyPanel.cpp
@@ -3634,10 +3634,11 @@ FReply SAnimNotifyPanel::OnMouseButtonUp(const FGeometry& MyGeometry, const FPoi
 	if(Marquee.bActive)
 	{
 		OnTrackSelectionChanged();
+		Marquee = FNotifyMarqueeOperation();
+		return FReply::Handled().ReleaseMouseCapture();
 	}
-	Marquee = FNotifyMarqueeOperation();
 
-	return FReply::Handled().ReleaseMouseCapture();
+	return SAnimTrackPanel::OnMouseButtonUp(MyGeometry, MouseEvent);
 }
 
 FReply SAnimNotifyPanel::OnMouseMove(const FGeometry& MyGeometry, const FPointerEvent& MouseEvent)
@@ -3655,7 +3656,7 @@ FReply SAnimNotifyPanel::OnMouseMove(const FGeometry& MyGeometry, const FPointer
 		return FReply::Handled();
 	}
 
-	return FReply::Unhandled();
+	return BaseReply;
 }
 
 int32 SAnimNotifyPanel::OnPaint(const FPaintArgs& Args, const FGeometry& AllottedGeometry, const FSlateRect& MyClippingRect, FSlateWindowElementList& OutDrawElements, int32 LayerId, const FWidgetStyle& InWidgetStyle, bool bParentEnabled) const
