Fix for Preview Asset Menu for socket includes incompatible items

diff --git a/Engine/Source/Editor/Persona/Private/SSkeletonTree.cpp b/Engine/Source/Editor/Persona/Private/SSkeletonTree.cpp
--- a/Engine/Source/Editor/Persona/Private/SSkeletonTree.cpp
+++ b/Engine/Source/Editor/Persona/Private/SSkeletonTree.cpp
@@ -1885,6 +1885,8 @@ void SSkeletonTree::FillAttachAssetSubmenu(FMenuBuilder& MenuBuilder, const FDis
 
 	TArray<UClass*> FilterClasses = FComponentAssetBrokerage::GetSupportedAssets(USceneComponent::StaticClass());
 
+	FilterClasses.RemoveSingleSwap(UBlueprint::StaticClass(), false); //Child actor components broker gives us blueprints which breaks this
+
 	FAssetPickerConfig AssetPickerConfig;
 	AssetPickerConfig.Filter.bRecursiveClasses = true;
 
