#PERSONA: ANIM: Keep animation change when mesh is swapped.

When mesh is swapped, right now it goes back to ref pose, which is very annoying.

diff --git a/Engine/Source/Editor/Persona/Private/Persona.cpp b/Engine/Source/Editor/Persona/Private/Persona.cpp
--- a/Engine/Source/Editor/Persona/Private/Persona.cpp
+++ b/Engine/Source/Editor/Persona/Private/Persona.cpp
@@ -1847,7 +1847,27 @@ void FPersona::SetPreviewMesh(USkeletalMesh* NewPreviewMesh)
 				AddEditingObject(NewPreviewMesh);
 			}
 
+			// setting skeletalmesh unregister/re-register, 
+			// so I have to save the animation settings and resetting after setting mesh
+			UAnimationAsset* AnimAssetToPlay = NULL;
+			float PlayPosition = 0.f;
+			bool bPlaying = false;
+			bool bNeedsToCopyAnimationData = PreviewComponent->AnimScriptInstance && PreviewComponent->AnimScriptInstance == PreviewComponent->PreviewInstance;
+			if(bNeedsToCopyAnimationData)
+			{
+				AnimAssetToPlay = PreviewComponent->PreviewInstance->CurrentAsset;
+				PlayPosition = PreviewComponent->PreviewInstance->CurrentTime;
+				bPlaying = PreviewComponent->PreviewInstance->bPlaying;
+			}
+
 			PreviewComponent->SetSkeletalMesh(NewPreviewMesh);
+
+			if(bNeedsToCopyAnimationData)
+			{
+				SetPreviewAnimationAsset(AnimAssetToPlay);
+				PreviewComponent->PreviewInstance->SetPosition(PlayPosition);
+				PreviewComponent->PreviewInstance->bPlaying = bPlaying;
+			}
 		}
 		else
 		{
