Fix crash when "modify bone" node bone name is set to None

diff --git a/Engine/Source/Editor/Persona/Private/AnimationEditorViewportClient.cpp b/Engine/Source/Editor/Persona/Private/AnimationEditorViewportClient.cpp
--- a/Engine/Source/Editor/Persona/Private/AnimationEditorViewportClient.cpp
+++ b/Engine/Source/Editor/Persona/Private/AnimationEditorViewportClient.cpp
@@ -1513,12 +1513,13 @@ FMatrix FAnimationViewportClient::GetWidgetCoordSystem() const
 		{
 			FName BoneName = SelectedSkelControlAnimGraph->FindSelectedBone();
 			int32 BoneIndex = PreviewSkelMeshComp->GetBoneIndex(BoneName);
-			FTransform BoneMatrix = PreviewSkelMeshComp->GetBoneTransform(BoneIndex);
-
-			return BoneMatrix.ToMatrixNoScale().RemoveTranslation();
+			if (BoneIndex != INDEX_NONE)
+			{
+				FTransform BoneMatrix = PreviewSkelMeshComp->GetBoneTransform(BoneIndex);
+				return BoneMatrix.ToMatrixNoScale().RemoveTranslation();
+			}
 		}
-		else
-		if ( PreviewSkelMeshComp->BonesOfInterest.Num() > 0 )
+		else if ( PreviewSkelMeshComp->BonesOfInterest.Num() > 0 )
 		{
 			int32 BoneIndex = PreviewSkelMeshComp->BonesOfInterest.Last();
 
