Fixed notify deletion picking the first notify with a matching name for event notifies.

diff --git a/Engine/Source/Editor/Persona/Private/SAnimNotifyPanel.cpp b/Engine/Source/Editor/Persona/Private/SAnimNotifyPanel.cpp
--- a/Engine/Source/Editor/Persona/Private/SAnimNotifyPanel.cpp
+++ b/Engine/Source/Editor/Persona/Private/SAnimNotifyPanel.cpp
@@ -3177,7 +3177,7 @@ void SAnimNotifyPanel::DeleteNotify(FAnimNotifyEvent* Notify)
 {
 	for (int32 I=0; I<Sequence->Notifies.Num(); ++I)
 	{
-		if (*Notify == Sequence->Notifies[I])
+		if (Notify == &(Sequence->Notifies[I]))
 		{
 			Sequence->Notifies.RemoveAt(I);
 			Sequence->MarkPackageDirty();
