Persona - Prevent a crash when renaming animation assets.

diff --git a/Engine/Source/Editor/Persona/Private/SAnimationSequenceBrowser.cpp b/Engine/Source/Editor/Persona/Private/SAnimationSequenceBrowser.cpp
--- a/Engine/Source/Editor/Persona/Private/SAnimationSequenceBrowser.cpp
+++ b/Engine/Source/Editor/Persona/Private/SAnimationSequenceBrowser.cpp
@@ -894,13 +894,11 @@ bool SAnimationSequenceBrowser::OnVisualizeAssetToolTip(const TSharedPtr<SWidget
 	// Resolve the asset
 	USkeletalMesh* MeshToUse = nullptr;
 	UClass* AssetClass = FindObject<UClass>(ANY_PACKAGE, *AssetData.AssetClass.ToString());
-	if(AssetClass->IsChildOf(UAnimationAsset::StaticClass()))
+	if(AssetClass->IsChildOf(UAnimationAsset::StaticClass()) && AssetData.GetAsset())
 	{
 		// Set up the viewport to show the asset. Catching the visualize allows us to use
 		// one viewport between all of the assets in the sequence browser.
 		UAnimationAsset* Asset = StaticCast<UAnimationAsset*>(AssetData.GetAsset());
-		check(Asset);
-
 		USkeleton* Skeleton = Asset->GetSkeleton();
 		
 		MeshToUse = Skeleton->GetAssetPreviewMesh(Asset);
