Fix undo when changing a material on a skeletal mesh

diff --git a/Engine/Source/Editor/Persona/Private/PersonaMeshDetails.cpp b/Engine/Source/Editor/Persona/Private/PersonaMeshDetails.cpp
--- a/Engine/Source/Editor/Persona/Private/PersonaMeshDetails.cpp
+++ b/Engine/Source/Editor/Persona/Private/PersonaMeshDetails.cpp
@@ -1231,7 +1231,7 @@ void FPersonaMeshDetails::OnMaterialChanged(UMaterialInterface* NewMaterial, UMa
 				GEditor->BeginTransaction(LOCTEXT("PersonaReplaceMaterial", "Replace material on mesh"));
 				bMadeTransaction = true;
 			}
-
+			Mesh->Modify();
 			Mesh->Materials[MaterialIndex].MaterialInterface = NewMaterial;
 		}
 	}
@@ -1274,6 +1274,7 @@ void FPersonaMeshDetails::OnMaterialChanged(UMaterialInterface* NewMaterial, UMa
 				GEditor->BeginTransaction(LOCTEXT("PersonaReplaceMaterial", "Replace material on mesh"));
 				bMadeTransaction = true;
 			}
+			Mesh->Modify();
 			Mesh->Materials[MaterialIndex].MaterialInterface = NewMaterial;
 		}
 	}
