#ANIM: RETARGETING: Make sure to return the node name

if any of them matches, list them out

diff --git a/Engine/Source/Editor/Persona/Private/SRigWindow.cpp b/Engine/Source/Editor/Persona/Private/SRigWindow.cpp
--- a/Engine/Source/Editor/Persona/Private/SRigWindow.cpp
+++ b/Engine/Source/Editor/Persona/Private/SRigWindow.cpp
@@ -346,18 +346,11 @@ void SRigWindow::CreateBoneMappingList( const FString& SearchText)
 			{
 				if(bDoFiltering)
 				{
-					if(!Name.ToString().Contains(SearchText))
+					// make sure it doens't fit any of them
+					if(!Name.ToString().Contains(SearchText) && !DisplayName.Contains(SearchText) && !BoneName.ToString().Contains(SearchText))
 					{
 						continue; // Skip items that don't match our filter
 					}
-					if(!DisplayName.Contains(SearchText))
-					{
-						continue;
-					}
-					if(!BoneName.ToString().Contains(SearchText))
-					{
-						continue;
-					}
 				}
 
 				TSharedRef<FDisplayedBoneMappingInfo> Info = FDisplayedBoneMappingInfo::Make(Name, DisplayName, Skeleton);
