Fixed undo when altering the size of notify state nodes using the drag markers.
#UE-11964

diff --git a/Engine/Source/Editor/Persona/Private/SAnimNotifyPanel.cpp b/Engine/Source/Editor/Persona/Private/SAnimNotifyPanel.cpp
--- a/Engine/Source/Editor/Persona/Private/SAnimNotifyPanel.cpp
+++ b/Engine/Source/Editor/Persona/Private/SAnimNotifyPanel.cpp
@@ -73,7 +73,9 @@ public:
 	virtual FReply OnMouseMove( const FGeometry& MyGeometry, const FPointerEvent& MouseEvent) override;
 	virtual FReply OnMouseButtonUp( const FGeometry& MyGeometry, const FPointerEvent& MouseEvent ) override;
 	virtual void Tick( const FGeometry& AllottedGeometry, const double InCurrentTime, const float InDeltaTime ) override;
-	
+	virtual void OnFocusLost(const FFocusEvent& InFocusEvent) override;
+	virtual bool SupportsKeyboardFocus() const override;
+
 	// End of SWidget interface
 
 	// SNodePanel::SNode interface
@@ -157,6 +159,9 @@ private:
 	bool						bBeingDragged;
 	bool						bSelected;
 
+	// Index for undo transactions for dragging, as a check to make sure it's active
+	int32						DragMarkerTransactionIdx;
+
 	/** The scrub handle currently being dragged, if any */
 	ENotifyStateHandleHit::Type CurrentDragHandle;
 	
@@ -932,6 +937,7 @@ void SAnimNotifyNode::Construct(const FArguments& InArgs)
 	CurrentDragHandle = ENotifyStateHandleHit::None;
 	bDrawTooltipToRight = true;
 	bSelected = false;
+	DragMarkerTransactionIdx = INDEX_NONE;
 	
 	// Cache notify name for blueprint / Native notifies.
 	if(NotifyEvent->Notify)
@@ -973,11 +979,15 @@ FReply SAnimNotifyNode::OnDragDetected( const FGeometry& MyGeometry, const FPoin
 			bDragOnMarker = true;
 			bBeingDragged = false;
 			CurrentDragHandle = MarkerHit;
+
+			// Modify the owning sequence as we're now dragging the marker and begin a transaction
+			check(DragMarkerTransactionIdx == INDEX_NONE);
+			DragMarkerTransactionIdx = GEditor->BeginTransaction(NSLOCTEXT("AnimNotifyNode", "StateNodeDragTransation", "Drag State Node Marker"));
+			Sequence->Modify();
 		}
 	}
 
 	return OnNodeDragStarted.Execute(SharedThis(this), MouseEvent, ScreenNodePosition, bDragOnMarker);
-
 }
 
 FLinearColor SAnimNotifyNode::GetNotifyColor() const
@@ -1333,9 +1343,6 @@ FReply SAnimNotifyNode::OnMouseMove( const FGeometry& MyGeometry, const FPointer
 
 		if(MouseEvent.GetScreenSpacePosition().X >= TrackScreenSpaceXPosition && MouseEvent.GetScreenSpacePosition().X <= TrackScreenSpaceXPosition + CachedAllotedGeometrySize.X)
 		{
-			const FScopedTransaction Transaction( LOCTEXT("AddNotifyEvent", "Edit Anim Notify") );
-			Sequence->Modify();
-
 			float NewDisplayTime = ScaleInfo.LocalXToInput((MouseEvent.GetScreenSpacePosition() - MyGeometry.AbsolutePosition + MyGeometry.Position).X);
 			float NewDuration = NotifyEvent->GetDuration() + OldDisplayTime - NewDisplayTime;
 
@@ -1350,9 +1357,6 @@ FReply SAnimNotifyNode::OnMouseMove( const FGeometry& MyGeometry, const FPointer
 		}
 		else if(NotifyEvent->GetDuration() > MinimumStateDuration)
 		{
-			const FScopedTransaction Transaction(LOCTEXT("AddNotifyEvent", "Edit Anim Notify"));
-			Sequence->Modify();
-
 			float Overflow = HandleOverflowPan(MouseEvent.GetScreenSpacePosition(), TrackScreenSpaceXPosition);
 
 			// Update scale info to the new view inputs after panning
@@ -1396,18 +1400,12 @@ FReply SAnimNotifyNode::OnMouseMove( const FGeometry& MyGeometry, const FPointer
 	{
 		if(MouseEvent.GetScreenSpacePosition().X >= TrackScreenSpaceXPosition && MouseEvent.GetScreenSpacePosition().X <= TrackScreenSpaceXPosition + CachedAllotedGeometrySize.X)
 		{
-			const FScopedTransaction Transaction(LOCTEXT("AddNotifyEvent", "Edit Anim Notify"));
-			Sequence->Modify();
-
 			float NewDuration = ScaleInfo.LocalXToInput((MouseEvent.GetScreenSpacePosition() - MyGeometry.AbsolutePosition + MyGeometry.Position).X) - NotifyEvent->GetTime();
 
 			NotifyEvent->SetDuration(FMath::Max(NewDuration, MinimumStateDuration));
 		}
 		else if(NotifyEvent->GetDuration() > MinimumStateDuration)
 		{
-			const FScopedTransaction Transaction(LOCTEXT("AddNotifyEvent", "Edit Anim Notify"));
-			Sequence->Modify();
-
 			float Overflow = HandleOverflowPan(MouseEvent.GetScreenSpacePosition(), TrackScreenSpaceXPosition);
 
 			// Update scale info to the new view inputs after panning
@@ -1448,6 +1446,12 @@ FReply SAnimNotifyNode::OnMouseButtonUp( const FGeometry& MyGeometry, const FPoi
 		// Clear the drag marker and give the mouse back
 		CurrentDragHandle = ENotifyStateHandleHit::None;
 		OnDeselectAllNotifies.ExecuteIfBound();
+
+		// End drag transaction before handing mouse back
+		check(DragMarkerTransactionIdx != INDEX_NONE);
+		GEditor->EndTransaction();
+		DragMarkerTransactionIdx = INDEX_NONE;
+
 		return FReply::Handled().ReleaseMouseCapture();
 	}
 
@@ -1550,6 +1554,27 @@ void SAnimNotifyNode::Tick( const FGeometry& AllottedGeometry, const double InCu
 	ScreenPosition = AllottedGeometry.AbsolutePosition;
 }
 
+void SAnimNotifyNode::OnFocusLost(const FFocusEvent& InFocusEvent)
+{
+	if(CurrentDragHandle != ENotifyStateHandleHit::None)
+	{
+		// Lost focus while dragging a state node, clear the drag and end the current transaction
+		CurrentDragHandle = ENotifyStateHandleHit::None;
+		OnDeselectAllNotifies.ExecuteIfBound();
+		
+		check(DragMarkerTransactionIdx != INDEX_NONE);
+		GEditor->EndTransaction();
+		DragMarkerTransactionIdx = INDEX_NONE;
+	}
+}
+
+bool SAnimNotifyNode::SupportsKeyboardFocus() const
+{
+	// Need to support focus on the node so we can end drag transactions if the user alt-tabs
+	// from the editor while in the proceess of dragging a state notify duration marker.
+	return true;
+}
+
 //////////////////////////////////////////////////////////////////////////
 // SAnimNotifyTrack
 
