Fixed BP anim notify name update not getting called when the notify is applied to the asset
#UE-9321

diff --git a/Engine/Source/Editor/Persona/Private/SAnimNotifyPanel.cpp b/Engine/Source/Editor/Persona/Private/SAnimNotifyPanel.cpp
--- a/Engine/Source/Editor/Persona/Private/SAnimNotifyPanel.cpp
+++ b/Engine/Source/Editor/Persona/Private/SAnimNotifyPanel.cpp
@@ -262,6 +262,11 @@ public:
 	virtual FCursorReply OnCursorQuery(const FGeometry& MyGeometry, const FPointerEvent& CursorEvent) const override;
 	// End of SWidget interface
 
+	/**
+	 * Update the nodes to match the data that the panel is observing
+	 */
+	void Update();
+
 	/** Returns the cached rendering geometry of this track */
 	const FGeometry& GetCachedGeometry() const { return CachedGeometry; }
 
@@ -355,11 +360,6 @@ protected:
 	/** Opens the supplied blueprint in an editor */
 	void OnOpenNotifySource(UBlueprint* InSourceBlueprint) const;
 
-	/**
-	 * Update the nodes to match the data that the panel is observing
-	 */
-	void Update();
-
 	/**
 	 * Selects a notify node in the graph. Supports multi selection
 	 * @param NotifyIndex - Index of the notify node to select.
@@ -3857,7 +3857,15 @@ void SAnimNotifyPanel::OnGetNativeNotifyData(TArray<UClass*>& OutClasses, UClass
 
 void SAnimNotifyPanel::OnNotifyObjectChanged(UObject* EditorBaseObj, bool bRebuild)
 {
-
+	if(UEditorNotifyObject* NotifyObject = Cast<UEditorNotifyObject>(EditorBaseObj))
+	{
+		// TODO: We should really un-invert these.
+		int32 WidgetTrackIdx = NotifyAnimTracks.Num() - NotifyObject->TrackIndex - 1;
+		if(NotifyAnimTracks.IsValidIndex(WidgetTrackIdx))
+		{
+			NotifyAnimTracks[WidgetTrackIdx]->Update();
+		}
+	}
 }
 
 #undef LOCTEXT_NAMESPACE
