Fixed private / non-visible anim notify properties appearing in details panels.
#UE-2811

diff --git a/Engine/Source/Editor/Persona/Private/AnimNotifyDetails.cpp b/Engine/Source/Editor/Persona/Private/AnimNotifyDetails.cpp
--- a/Engine/Source/Editor/Persona/Private/AnimNotifyDetails.cpp
+++ b/Engine/Source/Editor/Persona/Private/AnimNotifyDetails.cpp
@@ -101,10 +101,14 @@ void FAnimNotifyDetails::CustomizeDetails( IDetailLayoutBuilder& DetailBuilder )
 			for(uint32 ChildIdx = 0 ; ChildIdx < NumChildren ; ++ChildIdx)
 			{
 				TSharedPtr<IPropertyHandle> NotifyProperty = BaseHandle->GetChildHandle(ChildIdx);
-
-				if(!CustomizeProperty(AnimNotifyCategory, NotifyPtr, NotifyProperty))
+				UProperty* Prop = NotifyProperty->GetProperty();
+				
+				if(Prop && !Prop->HasAnyPropertyFlags(CPF_DisableEditOnInstance))
 				{
-					AnimNotifyCategory.AddProperty(NotifyProperty);
+					if(!CustomizeProperty(AnimNotifyCategory, NotifyPtr, NotifyProperty))
+					{
+						AnimNotifyCategory.AddProperty(NotifyProperty);
+					}
 				}
 			}
 		}
