#ANIM: PERSONA: added search for anim asset

https://jira.ol.epicgames.net/browse/UE-2084

diff --git a/Engine/Source/Editor/Persona/Private/AnimationMode/AnimationMode.cpp b/Engine/Source/Editor/Persona/Private/AnimationMode/AnimationMode.cpp
--- a/Engine/Source/Editor/Persona/Private/AnimationMode/AnimationMode.cpp
+++ b/Engine/Source/Editor/Persona/Private/AnimationMode/AnimationMode.cpp
@@ -28,7 +28,7 @@ public:
 	{
 		PersonaPtr = InPersona;
 
-		SSingleObjectDetailsPanel::Construct(SSingleObjectDetailsPanel::FArguments());
+		SSingleObjectDetailsPanel::Construct(SSingleObjectDetailsPanel::FArguments(), true, true);
 	}
 
 	virtual EVisibility GetAssetDisplayNameVisibility() const
