Fix performance drop when multiple copies of Persona are open.

diff --git a/Engine/Source/Editor/Persona/Private/SAnimationEditorViewport.cpp b/Engine/Source/Editor/Persona/Private/SAnimationEditorViewport.cpp
--- a/Engine/Source/Editor/Persona/Private/SAnimationEditorViewport.cpp
+++ b/Engine/Source/Editor/Persona/Private/SAnimationEditorViewport.cpp
@@ -51,6 +51,8 @@ void SAnimationEditorViewport::Construct(const FArguments& InArgs, TSharedPtr<cl
 			.IsEnabled(FSlateApplication::Get().GetNormalExecutionAttribute())
 			.AddMetaData<FTagMetaData>(TEXT("Persona.Viewport"))
 		);
+
+	Client->VisibilityDelegate.BindSP(this, &SAnimationEditorViewport::IsVisible);
 }
 
 TSharedRef<FEditorViewportClient> SAnimationEditorViewport::MakeEditorViewportClient()
@@ -382,7 +384,6 @@ void SAnimationEditorViewportTabBody::Construct(const FArguments& InArgs)
 	}
 
 	LevelViewportClient = ViewportWidget->GetViewportClient();
-	LevelViewportClient->VisibilityDelegate.BindSP(this, &SAnimationEditorViewportTabBody::IsVisible);
 
 	TSharedRef<FAnimationViewportClient> AnimViewportClient = StaticCastSharedRef<FAnimationViewportClient>(LevelViewportClient.ToSharedRef());
 
