Added context menu options to notify nodes in Persona to allow direct setting of Frame and Time for each notify.
#UE-5416

diff --git a/Engine/Source/Editor/Persona/Private/SAnimNotifyPanel.cpp b/Engine/Source/Editor/Persona/Private/SAnimNotifyPanel.cpp
--- a/Engine/Source/Editor/Persona/Private/SAnimNotifyPanel.cpp
+++ b/Engine/Source/Editor/Persona/Private/SAnimNotifyPanel.cpp
@@ -326,7 +326,11 @@ protected:
 	void AddNewNotify(const FText& NewNotifyName, ETextCommit::Type CommitInfo);
 	void OnSetTriggerWeightNotifyClicked(int32 NotifyIndex);
 	void SetTriggerWeight(const FText& TriggerWeight, ETextCommit::Type CommitInfo, int32 NotifyIndex);
-	
+	void OnSetNotifyTimeClicked(int32 NotifyIndex);
+	void SetNotifyTime(const FText& NotifyTimeText, ETextCommit::Type CommitInfo, int32 NotifyIndex);
+	void OnSetNotifyFrameClicked(int32 NotifyIndex);
+	void SetNotifyFrame(const FText& NotifyFrameText, ETextCommit::Type CommitInfo, int32 NotifyIndex);
+
 	// Whether we have one node selected
 	bool IsSingleNodeSelected();
 	// Checks the clipboard for an anim notify buffer, and returns whether there's only one notify
@@ -1022,8 +1026,7 @@ FText SAnimNotifyNode::GetNotifyText() const
 FText SAnimNotifyNode::GetNodeTooltip() const
 {
 	const float Time = NotifyEvent->GetTime();
-	const float Percentage = Time / Sequence->SequenceLength;
-	const FText Frame = FText::AsNumber( (int32)(Percentage * Sequence->GetNumberOfFrames()) );
+	const FText Frame = FText::AsNumber(Sequence->GetFrameAtTime(NotifyEvent->GetTime()));
 	const FText Seconds = FText::AsNumber( Time );
 
 	FText ToolTipText;
@@ -2150,15 +2153,35 @@ TSharedPtr<SWidget> SAnimNotifyTrack::SummonContextMenu(const FGeometry& MyGeome
 				SelectNotifyNode(NotifyIndex, MouseEvent.IsControlDown());
 			}
 
+			FAnimNotifyEvent* Event = AnimNotifies[NotifyIndex];
+
+			FNumberFormattingOptions Options;
+			Options.MinimumFractionalDigits = 5;
+
+			// Add item to directly set notify time
+			const FText CurrentTime = FText::AsNumber(Event->GetTime(), &Options);
+			const FText TimeMenuText = FText::Format(LOCTEXT("TimeMenuText", "Set Notify Begin Time: {0}..."), CurrentTime);
+
+			NewAction.ExecuteAction.BindRaw(this, &SAnimNotifyTrack::OnSetNotifyTimeClicked, NotifyIndex);
+			NewAction.CanExecuteAction.BindRaw(this, &SAnimNotifyTrack::IsSingleNodeSelected);
+
+			MenuBuilder.AddMenuEntry(TimeMenuText, LOCTEXT("SetTimeToolTip", "Set the time of this notify directly"), FSlateIcon(), NewAction);
+
+			// Add item to directly set notify frame
+			const FText Frame = FText::AsNumber(Sequence->GetFrameAtTime(Event->GetTime()));
+			const FText FrameMenuText = FText::Format(LOCTEXT("FrameMenuText", "Set Notify Frame: {0}..."), Frame);
+
+			NewAction.ExecuteAction.BindRaw(this, &SAnimNotifyTrack::OnSetNotifyFrameClicked, NotifyIndex);
+			NewAction.CanExecuteAction.BindRaw(this, &SAnimNotifyTrack::IsSingleNodeSelected);
+
+			MenuBuilder.AddMenuEntry(FrameMenuText, LOCTEXT("SetFrameToolTip", "Set the frame of this notify directly"), FSlateIcon(), NewAction);
+
 			// add menu to get threshold weight for triggering this notify
 			NewAction.ExecuteAction.BindRaw(
 				this, &SAnimNotifyTrack::OnSetTriggerWeightNotifyClicked, NotifyIndex);
 			NewAction.CanExecuteAction.BindRaw(
 				this, &SAnimNotifyTrack::IsSingleNodeSelected);
 
-			FNumberFormattingOptions Options;
-			Options.MinimumFractionalDigits = 5;
-
 			const FText Threshold = FText::AsNumber( AnimNotifies[NotifyIndex]->TriggerWeightThreshold, &Options );
 			const FText MinTriggerWeightText = FText::Format( LOCTEXT("MinTriggerWeight", "Min Trigger Weight: {0}..."), Threshold );
 			MenuBuilder.AddMenuEntry(MinTriggerWeightText, LOCTEXT("MinTriggerWeightToolTip", "The minimum weight to trigger this notify"), FSlateIcon(), NewAction);
@@ -2956,6 +2979,107 @@ void SAnimNotifyTrack::GetNotifyMenuData(TArray<FAssetData>& NotifyAssetData, TA
 	});
 }
 
+void SAnimNotifyTrack::OnSetNotifyTimeClicked(int32 NotifyIndex)
+{
+	if(AnimNotifies.IsValidIndex(NotifyIndex))
+	{
+		FString DefaultText = FString::Printf(TEXT("%0.6f"), AnimNotifies[NotifyIndex]->GetTime());
+
+		// Show dialog to enter time
+		TSharedRef<STextEntryPopup> TextEntry =
+			SNew(STextEntryPopup)
+			.Label(LOCTEXT("NotifyTimeClickedLabel", "Notify Time"))
+			.DefaultText(FText::FromString(DefaultText))
+			.OnTextCommitted(this, &SAnimNotifyTrack::SetNotifyTime, NotifyIndex);
+
+		FSlateApplication::Get().PushMenu(
+			AsShared(), // Menu being summoned from a menu that is closing: Parent widget should be k2 not the menu thats open or it will be closed when the menu is dismissed
+			TextEntry,
+			FSlateApplication::Get().GetCursorPos(),
+			FPopupTransitionEffect(FPopupTransitionEffect::TypeInPopup)
+			);
+
+		TextEntry->FocusDefaultWidget();
+	}
+}
+
+void SAnimNotifyTrack::SetNotifyTime(const FText& NotifyTimeText, ETextCommit::Type CommitInfo, int32 NotifyIndex)
+{
+	if(CommitInfo == ETextCommit::OnEnter || CommitInfo == ETextCommit::OnUserMovedFocus)
+	{
+		if(AnimNotifies.IsValidIndex(NotifyIndex))
+		{
+			FAnimNotifyEvent* Event = AnimNotifies[NotifyIndex];
+
+			float NewTime = FMath::Clamp(FCString::Atof(*NotifyTimeText.ToString()), 0.0f, Sequence->SequenceLength - Event->GetDuration());
+
+			Event->SetTime(NewTime);
+
+			Event->RefreshTriggerOffset(Sequence->CalculateOffsetForNotify(Event->GetTime()));
+			if(Event->GetDuration() > 0.0f)
+			{
+				Event->RefreshEndTriggerOffset(Sequence->CalculateOffsetForNotify(Event->GetTime() + Event->GetDuration()));
+			}
+		}
+	}
+
+	FSlateApplication::Get().DismissAllMenus();
+}
+
+void SAnimNotifyTrack::OnSetNotifyFrameClicked(int32 NotifyIndex)
+{
+	if(AnimNotifies.IsValidIndex(NotifyIndex))
+	{
+		FAnimNotifyEvent* Event = AnimNotifies[NotifyIndex];
+
+		//const float Time = Event->GetTime();
+		//const float Percentage = Time / Sequence->SequenceLength;
+		const FText Frame = FText::AsNumber(Sequence->GetFrameAtTime(Event->GetTime()));
+
+		FString DefaultText = FString::Printf(TEXT("%s"), *Frame.ToString());
+
+		// Show dialog to enter frame
+		TSharedRef<STextEntryPopup> TextEntry =
+			SNew(STextEntryPopup)
+			.Label(LOCTEXT("NotifyFrameClickedLabel", "Notify Frame"))
+			.DefaultText(FText::FromString(DefaultText))
+			.OnTextCommitted(this, &SAnimNotifyTrack::SetNotifyFrame, NotifyIndex);
+
+		FSlateApplication::Get().PushMenu(
+			AsShared(), // Menu being summoned from a menu that is closing: Parent widget should be k2 not the menu thats open or it will be closed when the menu is dismissed
+			TextEntry,
+			FSlateApplication::Get().GetCursorPos(),
+			FPopupTransitionEffect(FPopupTransitionEffect::TypeInPopup)
+			);
+
+		TextEntry->FocusDefaultWidget();
+	}
+}
+
+void SAnimNotifyTrack::SetNotifyFrame(const FText& NotifyFrameText, ETextCommit::Type CommitInfo, int32 NotifyIndex)
+{
+	if(CommitInfo == ETextCommit::OnEnter || CommitInfo == ETextCommit::OnUserMovedFocus)
+	{
+		if(AnimNotifies.IsValidIndex(NotifyIndex))
+		{
+			FAnimNotifyEvent* Event = AnimNotifies[NotifyIndex];
+
+			int32 Frame = FCString::Atof(*NotifyFrameText.ToString());
+			float NewTime = FMath::Clamp(Sequence->GetTimeAtFrame(Frame), 0.0f, Sequence->SequenceLength - Event->Duration);
+
+			Event->SetTime(NewTime);
+
+			Event->RefreshTriggerOffset(Sequence->CalculateOffsetForNotify(Event->GetTime()));
+			if(Event->GetDuration() > 0.0f)
+			{
+				Event->RefreshEndTriggerOffset(Sequence->CalculateOffsetForNotify(Event->GetTime() + Event->GetDuration()));
+			}
+		}
+	}
+
+	FSlateApplication::Get().DismissAllMenus();
+}
+
 //////////////////////////////////////////////////////////////////////////
 // SSequenceEdTrack
 
diff --git a/Engine/Source/Runtime/Engine/Classes/Animation/AnimSequence.h b/Engine/Source/Runtime/Engine/Classes/Animation/AnimSequence.h
--- a/Engine/Source/Runtime/Engine/Classes/Animation/AnimSequence.h
+++ b/Engine/Source/Runtime/Engine/Classes/Animation/AnimSequence.h
@@ -561,7 +561,7 @@ public:
 #if WITH_EDITOR
 	ENGINE_API virtual bool GetAllAnimationSequencesReferred(TArray<UAnimSequence*>& AnimationSequences) override;
 	ENGINE_API virtual void ReplaceReferredAnimations(const TMap<UAnimSequence*, UAnimSequence*>& ReplacementMap) override;
-	ENGINE_API virtual int32 GetNumberOfFrames() override { return NumFrames; }
+	ENGINE_API virtual int32 GetNumberOfFrames() const override { return NumFrames; }
 #endif
 	// End of UAnimationAsset interface
 
diff --git a/Engine/Source/Runtime/Engine/Classes/Animation/AnimSequenceBase.h b/Engine/Source/Runtime/Engine/Classes/Animation/AnimSequenceBase.h
--- a/Engine/Source/Runtime/Engine/Classes/Animation/AnimSequenceBase.h
+++ b/Engine/Source/Runtime/Engine/Classes/Animation/AnimSequenceBase.h
@@ -624,7 +624,13 @@ class UAnimSequenceBase : public UAnimationAsset
 
 #if WITH_EDITOR
 	/** Return Number of Frames **/
-	virtual int32 GetNumberOfFrames();
+	virtual int32 GetNumberOfFrames() const;
+
+	/** Get the frame number for the provided time */
+	virtual int32 GetFrameAtTime(const float Time) const;
+
+	/** Get the time at the given frame */
+	virtual float GetTimeAtFrame(const int32 Frame) const;
 
 	// update anim notify track cache
 	ENGINE_API void UpdateAnimNotifyTrackCache();
diff --git a/Engine/Source/Runtime/Engine/Private/Animation/AnimSequenceBase.cpp b/Engine/Source/Runtime/Engine/Private/Animation/AnimSequenceBase.cpp
--- a/Engine/Source/Runtime/Engine/Private/Animation/AnimSequenceBase.cpp
+++ b/Engine/Source/Runtime/Engine/Private/Animation/AnimSequenceBase.cpp
@@ -792,11 +792,23 @@ void UAnimSequenceBase::InitializeNotifyTrack()
 	}
 }
 
-int32 UAnimSequenceBase::GetNumberOfFrames()
+int32 UAnimSequenceBase::GetNumberOfFrames() const
 {
 	return (SequenceLength/0.033f);
 }
 
+int32 UAnimSequenceBase::GetFrameAtTime(const float Time) const
+{
+	const float Frac = Time / SequenceLength;
+	return FMath::FloorToInt(Frac * GetNumberOfFrames());
+}
+
+float UAnimSequenceBase::GetTimeAtFrame(const int32 Frame) const
+{
+	const float FrameTime = SequenceLength / GetNumberOfFrames();
+	return FrameTime * Frame;
+}
+
 void UAnimSequenceBase::RegisterOnNotifyChanged(const FOnNotifyChanged& Delegate)
 {
 	OnNotifyChanged.Add(Delegate);
