Added resize cursor when hovering over scrub handles on notify state nodes in Persona
#UE-2120

diff --git a/Engine/Source/Editor/Persona/Private/SAnimNotifyPanel.cpp b/Engine/Source/Editor/Persona/Private/SAnimNotifyPanel.cpp
--- a/Engine/Source/Editor/Persona/Private/SAnimNotifyPanel.cpp
+++ b/Engine/Source/Editor/Persona/Private/SAnimNotifyPanel.cpp
@@ -75,7 +75,7 @@ public:
 	virtual void Tick( const FGeometry& AllottedGeometry, const double InCurrentTime, const float InDeltaTime ) override;
 	virtual void OnFocusLost(const FFocusEvent& InFocusEvent) override;
 	virtual bool SupportsKeyboardFocus() const override;
-
+	virtual FCursorReply OnCursorQuery(const FGeometry& MyGeometry, const FPointerEvent& CursorEvent) const override;
 	// End of SWidget interface
 
 	// SNodePanel::SNode interface
@@ -108,7 +108,7 @@ public:
 	void DropCancelled();
 
 	/** Returns the size of this notifies duration in screen space */
-	float GetDurationSize() { return NotifyDurationSizeX;}
+	float GetDurationSize() const { return NotifyDurationSizeX;}
 
 	/** Sets the position the mouse was at when this node was last hit */
 	void SetLastMouseDownPosition(const FVector2D& CursorPosition) {LastMouseDownPosition = CursorPosition;}
@@ -1578,6 +1578,25 @@ bool SAnimNotifyNode::SupportsKeyboardFocus() const
 	return true;
 }
 
+FCursorReply SAnimNotifyNode::OnCursorQuery(const FGeometry& MyGeometry, const FPointerEvent& CursorEvent) const
+{
+	// Show resize cursor if the cursor is hoverring over either of the scrub handles of a notify state node
+	if(IsHovered() && GetDurationSize() > 0.0f)
+	{
+		FVector2D RelMouseLocation = MyGeometry.AbsoluteToLocal(CursorEvent.GetScreenSpacePosition());
+
+		const float HandleHalfWidth = ScrubHandleSize.X / 2.0f;
+		const float DistFromFirstHandle = FMath::Abs(RelMouseLocation.X - NotifyScrubHandleCentre);
+		const float DistFromSecondHandle = FMath::Abs(RelMouseLocation.X - (NotifyScrubHandleCentre + NotifyDurationSizeX));
+
+		if(DistFromFirstHandle < HandleHalfWidth || DistFromSecondHandle < HandleHalfWidth || CurrentDragHandle != ENotifyStateHandleHit::None)
+		{
+			return FCursorReply::Cursor(EMouseCursor::ResizeLeftRight);
+		}
+	}
+	return FCursorReply::Unhandled();
+}
+
 //////////////////////////////////////////////////////////////////////////
 // SAnimNotifyTrack
 
