Persona: Fixed a few whitespace/formatting issues

diff --git a/Engine/Source/Editor/Persona/Private/SAnimationEditorViewport.cpp b/Engine/Source/Editor/Persona/Private/SAnimationEditorViewport.cpp
--- a/Engine/Source/Editor/Persona/Private/SAnimationEditorViewport.cpp
+++ b/Engine/Source/Editor/Persona/Private/SAnimationEditorViewport.cpp
@@ -784,8 +784,7 @@ int32 SAnimationEditorViewportTabBody::GetLODModelCount() const
 
 void SAnimationEditorViewportTabBody::OnShowBones()
 {
-	UDebugSkelMeshComponent* PreviewComponent = PersonaPtr.Pin()->PreviewComponent;
-	if( PreviewComponent != NULL )
+	if (UDebugSkelMeshComponent* PreviewComponent = PersonaPtr.Pin()->PreviewComponent)
 	{
 		PreviewComponent->bDisplayBones = !PreviewComponent->bDisplayBones;
 		PreviewComponent->MarkRenderStateDirty();
@@ -795,8 +794,7 @@ void SAnimationEditorViewportTabBody::OnShowBones()
 
 void SAnimationEditorViewportTabBody::OnShowMorphTargets()
 {
-	UDebugSkelMeshComponent* PreviewComponent = PersonaPtr.Pin()->PreviewComponent;
-	if(PreviewComponent != NULL)
+	if (UDebugSkelMeshComponent* PreviewComponent = PersonaPtr.Pin()->PreviewComponent)
 	{
 		PreviewComponent->bDisableMorphTarget = !PreviewComponent->bDisableMorphTarget;
 		PreviewComponent->MarkRenderStateDirty();
@@ -806,8 +804,7 @@ void SAnimationEditorViewportTabBody::OnShowMorphTargets()
 
 void SAnimationEditorViewportTabBody::OnShowBoneNames()
 {
-	UDebugSkelMeshComponent* PreviewComponent = PersonaPtr.Pin()->PreviewComponent;
-	if( PreviewComponent != NULL )
+	if (UDebugSkelMeshComponent* PreviewComponent = PersonaPtr.Pin()->PreviewComponent)
 	{
 		PreviewComponent->bShowBoneNames = !PreviewComponent->bShowBoneNames;
 		PreviewComponent->MarkRenderStateDirty();
@@ -817,8 +814,7 @@ void SAnimationEditorViewportTabBody::OnShowBoneNames()
 
 void SAnimationEditorViewportTabBody::OnShowRawAnimation()
 {
-	UDebugSkelMeshComponent* PreviewComponent = PersonaPtr.Pin()->PreviewComponent;
-	if( PreviewComponent != NULL )
+	if (UDebugSkelMeshComponent* PreviewComponent = PersonaPtr.Pin()->PreviewComponent)
 	{
 		PreviewComponent->bDisplayRawAnimation = !PreviewComponent->bDisplayRawAnimation;
 		PreviewComponent->MarkRenderStateDirty();
@@ -827,8 +823,7 @@ void SAnimationEditorViewportTabBody::OnShowRawAnimation()
 
 void SAnimationEditorViewportTabBody::OnShowNonRetargetedAnimation()
 {
-	UDebugSkelMeshComponent* PreviewComponent = PersonaPtr.Pin()->PreviewComponent;
-	if( PreviewComponent != NULL )
+	if (UDebugSkelMeshComponent* PreviewComponent = PersonaPtr.Pin()->PreviewComponent)
 	{
 		PreviewComponent->bDisplayNonRetargetedPose = !PreviewComponent->bDisplayNonRetargetedPose;
 		PreviewComponent->MarkRenderStateDirty();
@@ -837,8 +832,7 @@ void SAnimationEditorViewportTabBody::OnShowNonRetargetedAnimation()
 
 void SAnimationEditorViewportTabBody::OnShowSourceRawAnimation()
 {
-	UDebugSkelMeshComponent* PreviewComponent = PersonaPtr.Pin()->PreviewComponent;
-	if(PreviewComponent != NULL)
+	if (UDebugSkelMeshComponent* PreviewComponent = PersonaPtr.Pin()->PreviewComponent)
 	{
 		PreviewComponent->bDisplaySourceAnimation = !PreviewComponent->bDisplaySourceAnimation;
 		PreviewComponent->MarkRenderStateDirty();
@@ -847,8 +841,7 @@ void SAnimationEditorViewportTabBody::OnShowSourceRawAnimation()
 
 void SAnimationEditorViewportTabBody::OnShowBakedAnimation()
 {
-	UDebugSkelMeshComponent* PreviewComponent = PersonaPtr.Pin()->PreviewComponent;
-	if(PreviewComponent != NULL)
+	if (UDebugSkelMeshComponent* PreviewComponent = PersonaPtr.Pin()->PreviewComponent)
 	{
 		PreviewComponent->bDisplayBakedAnimation = !PreviewComponent->bDisplayBakedAnimation;
 		PreviewComponent->MarkRenderStateDirty();
@@ -857,8 +850,7 @@ void SAnimationEditorViewportTabBody::OnShowBakedAnimation()
 
 void SAnimationEditorViewportTabBody::OnShowAdditiveBase()
 {
-	UDebugSkelMeshComponent* PreviewComponent = PersonaPtr.Pin()->PreviewComponent;
-	if( PreviewComponent != NULL )
+	if (UDebugSkelMeshComponent* PreviewComponent = PersonaPtr.Pin()->PreviewComponent)
 	{
 		PreviewComponent->bDisplayAdditiveBasePose = !PreviewComponent->bDisplayAdditiveBasePose;
 		PreviewComponent->MarkRenderStateDirty();
@@ -1103,8 +1095,7 @@ bool SAnimationEditorViewportTabBody::IsShowBoundEnabled() const
 
 void SAnimationEditorViewportTabBody::ToggleShowPreviewMesh()
 {
-	UDebugSkelMeshComponent* PreviewComponent = PersonaPtr.Pin()->PreviewComponent;
-	if(PreviewComponent != NULL)
+	if (UDebugSkelMeshComponent* PreviewComponent = PersonaPtr.Pin()->PreviewComponent)
 	{
 		bool bCurrentlyVisible = IsShowPreviewMeshEnabled();
 		PreviewComponent->SetVisibility(!bCurrentlyVisible);
@@ -1185,9 +1176,9 @@ void SAnimationEditorViewportTabBody::PopulateNumUVChannels()
 {
 	NumUVChannels.Empty();
 
-	if(UDebugSkelMeshComponent* PreviewComponent = PersonaPtr.Pin()->PreviewComponent)
+	if (UDebugSkelMeshComponent* PreviewComponent = PersonaPtr.Pin()->PreviewComponent)
 	{
-		if(FSkeletalMeshResource* MeshResource = PreviewComponent->GetSkeletalMeshResource())
+		if (FSkeletalMeshResource* MeshResource = PreviewComponent->GetSkeletalMeshResource())
 		{
 			int32 NumLods = MeshResource->LODModels.Num();
 			NumUVChannels.AddZeroed(NumLods);
@@ -1206,27 +1197,27 @@ void SAnimationEditorViewportTabBody::PopulateUVChoices()
 	// Fill out the UV channels combo.
 	UVChannels.Empty();
 	
-	if( UDebugSkelMeshComponent* PreviewComponent = PersonaPtr.Pin()->PreviewComponent )
+	if (UDebugSkelMeshComponent* PreviewComponent = PersonaPtr.Pin()->PreviewComponent)
 	{
 		int32 CurrentLOD = FMath::Clamp(PreviewComponent->ForcedLodModel - 1, 0, NumUVChannels.Num() - 1);
 
-		if(NumUVChannels.IsValidIndex(CurrentLOD))
+		if (NumUVChannels.IsValidIndex(CurrentLOD))
 		{
-			for(int32 UVChannelID = 0; UVChannelID < NumUVChannels[CurrentLOD]; ++UVChannelID)
+			for (int32 UVChannelID = 0; UVChannelID < NumUVChannels[CurrentLOD]; ++UVChannelID)
 			{
 				UVChannels.Add( MakeShareable( new FString( FText::Format( NSLOCTEXT("AnimationEditorViewport", "UVChannel_ID", "UV Channel {0}"), FText::AsNumber( UVChannelID ) ).ToString() ) ) );
 			}
 
 			TSharedRef<FAnimationViewportClient> AnimViewportClient = StaticCastSharedRef<FAnimationViewportClient>(LevelViewportClient.ToSharedRef());
 			int32 CurrentUVChannel = AnimViewportClient->GetUVChannelToDraw();
-			if(!UVChannels.IsValidIndex(CurrentUVChannel))
+			if (!UVChannels.IsValidIndex(CurrentUVChannel))
 			{
 				CurrentUVChannel = 0;
 			}
 
 			AnimViewportClient->SetUVChannelToDraw(CurrentUVChannel);
 
-			if(UVChannelCombo.IsValid() && UVChannels.IsValidIndex(CurrentUVChannel))
+			if (UVChannelCombo.IsValid() && UVChannels.IsValidIndex(CurrentUVChannel))
 			{
 				UVChannelCombo->SetSelectedItem(UVChannels[CurrentUVChannel]);
 			}
@@ -1237,11 +1228,11 @@ void SAnimationEditorViewportTabBody::PopulateUVChoices()
 void SAnimationEditorViewportTabBody::UpdateScrubPanel(UAnimationAsset* AnimAsset)
 {
 	// We might not have a scrub panel if we're in animation mode.
-	if(ScrubPanelContainer.IsValid())
+	if (ScrubPanelContainer.IsValid())
 	{
 		ScrubPanelContainer->ClearChildren();
 		bool bUseDefaultScrubPanel = true;
-		if(UAnimMontage* Montage = Cast<UAnimMontage>(AnimAsset))
+		if (UAnimMontage* Montage = Cast<UAnimMontage>(AnimAsset))
 		{
 			ScrubPanelContainer->AddSlot()
 				.AutoHeight()
@@ -1271,8 +1262,7 @@ void SAnimationEditorViewportTabBody::UpdateScrubPanel(UAnimationAsset* AnimAsse
 
 float SAnimationEditorViewportTabBody::GetViewMinInput() const
 {
-	UDebugSkelMeshComponent* PreviewComponent = PersonaPtr.Pin()->PreviewComponent;
-	if (PreviewComponent != NULL)
+	if (UDebugSkelMeshComponent* PreviewComponent = PersonaPtr.Pin()->PreviewComponent)
 	{
 		UObject* PreviewAsset = PersonaPtr.Pin()->GetPreviewAnimationAsset();
 		if (PreviewAsset != NULL)
@@ -1771,9 +1761,7 @@ void SAnimationEditorViewportTabBody::OnShowClothBackstops()
 
 bool SAnimationEditorViewportTabBody::IsShowingClothBackstops() const
 {
-	UDebugSkelMeshComponent* PreviewComponent = PersonaPtr.Pin()->PreviewComponent;
-
-	if( PreviewComponent )
+	if (UDebugSkelMeshComponent* PreviewComponent = PersonaPtr.Pin()->PreviewComponent)
 	{
 		return PreviewComponent->bDisplayClothBackstops;
 	}
@@ -1783,9 +1771,7 @@ bool SAnimationEditorViewportTabBody::IsShowingClothBackstops() const
 
 void SAnimationEditorViewportTabBody::OnShowClothFixedVertices()
 {
-	UDebugSkelMeshComponent* PreviewComponent = PersonaPtr.Pin()->PreviewComponent;
-
-	if( PreviewComponent )
+	if (UDebugSkelMeshComponent* PreviewComponent = PersonaPtr.Pin()->PreviewComponent)
 	{
 		PreviewComponent->bDisplayClothFixedVertices = !PreviewComponent->bDisplayClothFixedVertices;
 		RefreshViewport();
@@ -1794,9 +1780,7 @@ void SAnimationEditorViewportTabBody::OnShowClothFixedVertices()
 
 bool SAnimationEditorViewportTabBody::IsShowingClothFixedVertices() const
 {
-	UDebugSkelMeshComponent* PreviewComponent = PersonaPtr.Pin()->PreviewComponent;
-
-	if( PreviewComponent )
+	if (UDebugSkelMeshComponent* PreviewComponent = PersonaPtr.Pin()->PreviewComponent)
 	{
 		return PreviewComponent->bDisplayClothFixedVertices;
 	}
@@ -1853,7 +1837,7 @@ EVisibility SAnimationEditorViewportTabBody::GetViewportCornerImageVisibility()
 	return Persona->Recorder.InRecording()? EVisibility::Visible : EVisibility::Collapsed;
 }
 
-const FSlateBrush * SAnimationEditorViewportTabBody::GetViewportCornerImage() const
+const FSlateBrush* SAnimationEditorViewportTabBody::GetViewportCornerImage() const
 {
 	static int32 Count=0;
 
