Option in Persona to toggle auto align floor to mesh bounds

diff --git a/Engine/Source/Editor/Persona/Private/AnimViewportShowCommands.cpp b/Engine/Source/Editor/Persona/Private/AnimViewportShowCommands.cpp
--- a/Engine/Source/Editor/Persona/Private/AnimViewportShowCommands.cpp
+++ b/Engine/Source/Editor/Persona/Private/AnimViewportShowCommands.cpp
@@ -12,6 +12,7 @@ void FAnimViewportShowCommands::RegisterCommands()
 	UI_COMMAND( ToggleFloor, "Floor", "Display Floor", EUserInterfaceActionType::ToggleButton, FInputChord() );
 	UI_COMMAND( ToggleSky, "Sky", "Display Sky", EUserInterfaceActionType::ToggleButton, FInputChord() );
 
+	UI_COMMAND( AutoAlignFloorToMesh, "Auto Align Floor To Mesh", "Auto aligns floor to mesh bounds", EUserInterfaceActionType::ToggleButton, FInputChord() );
 	UI_COMMAND( MuteAudio, "Mute Audio", "Mutes audio from the preview", EUserInterfaceActionType::ToggleButton, FInputChord() );
 
 	UI_COMMAND(ProcessRootMotion, "Process Root Motion", "Moves preview based on animation root motion", EUserInterfaceActionType::ToggleButton, FInputChord());
diff --git a/Engine/Source/Editor/Persona/Private/AnimViewportShowCommands.h b/Engine/Source/Editor/Persona/Private/AnimViewportShowCommands.h
--- a/Engine/Source/Editor/Persona/Private/AnimViewportShowCommands.h
+++ b/Engine/Source/Editor/Persona/Private/AnimViewportShowCommands.h
@@ -30,6 +30,9 @@ public:
 	/** Option to show floor */
 	TSharedPtr< FUICommandInfo > ToggleSky;
 
+	/** Option to align floor to Mesh */
+	TSharedPtr< FUICommandInfo > AutoAlignFloorToMesh;
+
 	/** Option to mute audio in the viewport */
 	TSharedPtr< FUICommandInfo > MuteAudio;
 
diff --git a/Engine/Source/Editor/Persona/Private/AnimationEditorViewportClient.cpp b/Engine/Source/Editor/Persona/Private/AnimationEditorViewportClient.cpp
--- a/Engine/Source/Editor/Persona/Private/AnimationEditorViewportClient.cpp
+++ b/Engine/Source/Editor/Persona/Private/AnimationEditorViewportClient.cpp
@@ -139,6 +139,8 @@ FAnimationViewportClient::FAnimationViewportClient(FAnimationEditorPreviewScene&
 	bDrawUVs = false;
 	UVChannelToDraw = 0;
 
+	bAutoAlignFloor = true;
+
 	// Set audio mute option
 	UWorld* World = PreviewScene->GetWorld();
 	if(World)
@@ -240,6 +242,17 @@ bool FAnimationViewportClient::IsShowingFloor() const
 	return (GetAnimPreviewScene()->EditorFloorComp) ? GetAnimPreviewScene()->EditorFloorComp->bVisible : false;
 }
 
+void FAnimationViewportClient::OnToggleAutoAlignFloor()
+{
+	bAutoAlignFloor = !bAutoAlignFloor;
+	UpdateCameraSetup();
+}
+
+bool FAnimationViewportClient::IsAutoAlignFloor() const
+{
+	return bAutoAlignFloor;
+}
+
 void FAnimationViewportClient::OnToggleShowSky()
 {
 	if (GetAnimPreviewScene()->EditorSkyComp)
@@ -2044,7 +2057,12 @@ void FAnimationViewportClient::UpdateCameraSetup()
 
 		if (GetAnimPreviewScene()->EditorFloorComp)
 		{
-			GetAnimPreviewScene()->EditorFloorComp->SetWorldTransform(FTransform(FQuat::Identity, FVector(0.0f, 0.0f, Bottom.Z + GetFloorOffset()), FVector(3.0f, 3.0f, 1.0f)));
+			FVector FloorPos(0.f, 0.f, 0.f);
+			if (bAutoAlignFloor)
+			{
+				FloorPos.Z = GetFloorOffset() + Bottom.Z;
+			}
+			GetAnimPreviewScene()->EditorFloorComp->SetWorldTransform(FTransform(FQuat::Identity, FloorPos, FVector(3.0f, 3.0f, 1.0f)));
 		}
 	}
 }
diff --git a/Engine/Source/Editor/Persona/Private/AnimationEditorViewportClient.h b/Engine/Source/Editor/Persona/Private/AnimationEditorViewportClient.h
--- a/Engine/Source/Editor/Persona/Private/AnimationEditorViewportClient.h
+++ b/Engine/Source/Editor/Persona/Private/AnimationEditorViewportClient.h
@@ -104,6 +104,12 @@ public:
 	/** Function to check whether grid is displayed or not */
 	bool IsShowingFloor() const;
 
+	/** Function to enable/disable floor auto align */
+	void OnToggleAutoAlignFloor();
+
+	/** Function to check whether floor is auto align or not */
+	bool IsAutoAlignFloor() const;
+
 	/** Function to show/hide Sky in the viewport */
 	void OnToggleShowSky();
 
@@ -240,6 +246,9 @@ private:
 	/** add follow option @todo change to enum later - we share editorviewportclient, which is only problem*/
 	bool bCameraFollow;
 
+	/** Should we auto align floor to mesh bounds */
+	bool bAutoAlignFloor;
+
 	/** True when we're in an editor transaction (moving/rotating sockets) */
 	bool bInTransaction;
 
diff --git a/Engine/Source/Editor/Persona/Private/SAnimViewportToolBar.cpp b/Engine/Source/Editor/Persona/Private/SAnimViewportToolBar.cpp
--- a/Engine/Source/Editor/Persona/Private/SAnimViewportToolBar.cpp
+++ b/Engine/Source/Editor/Persona/Private/SAnimViewportToolBar.cpp
@@ -539,7 +539,8 @@ void SAnimViewportToolBar::FillShowSceneMenu(FMenuBuilder& MenuBuilder) const
 	{
 		MenuBuilder.AddMenuEntry(Actions.ToggleFloor);
 		MenuBuilder.AddMenuEntry(Actions.ToggleSky);
-		}
+		MenuBuilder.AddMenuEntry(Actions.AutoAlignFloorToMesh);
+	}
 	MenuBuilder.EndSection();
 
 	MenuBuilder.BeginSection("AnimViewportFloorOffset", LOCTEXT("Viewport_FloorOffsetLabel", "Floor Height Offset"));
diff --git a/Engine/Source/Editor/Persona/Private/SAnimationEditorViewport.cpp b/Engine/Source/Editor/Persona/Private/SAnimationEditorViewport.cpp
--- a/Engine/Source/Editor/Persona/Private/SAnimationEditorViewport.cpp
+++ b/Engine/Source/Editor/Persona/Private/SAnimationEditorViewport.cpp
@@ -714,6 +714,12 @@ void SAnimationEditorViewportTabBody::BindCommands()
 		FCanExecuteAction(),
 		FIsActionChecked::CreateSP(this, &SAnimationEditorViewportTabBody::IsShowingFloor));
 
+	CommandList.MapAction(
+		ViewportShowMenuCommands.AutoAlignFloorToMesh,
+		FExecuteAction::CreateSP(this, &SAnimationEditorViewportTabBody::OnToggleAutoAlignFloor),
+		FCanExecuteAction(),
+		FIsActionChecked::CreateSP(this, &SAnimationEditorViewportTabBody::IsAutoAlignFloor));
+
 	CommandList.MapAction( 
 		ViewportShowMenuCommands.ToggleSky,
 		FExecuteAction::CreateSP(this, &SAnimationEditorViewportTabBody::OnShowSky),
@@ -993,6 +999,18 @@ bool SAnimationEditorViewportTabBody::IsShowingFloor() const
 	return AnimViewportClient->IsShowingFloor();
 }
 
+void SAnimationEditorViewportTabBody::OnToggleAutoAlignFloor()
+{
+	TSharedRef<FAnimationViewportClient> AnimViewportClient = StaticCastSharedRef<FAnimationViewportClient>(LevelViewportClient.ToSharedRef());
+	AnimViewportClient->OnToggleAutoAlignFloor();
+}
+
+bool SAnimationEditorViewportTabBody::IsAutoAlignFloor() const
+{
+	TSharedRef<FAnimationViewportClient> AnimViewportClient = StaticCastSharedRef<FAnimationViewportClient>(LevelViewportClient.ToSharedRef());
+	return AnimViewportClient->IsAutoAlignFloor();
+}
+
 void SAnimationEditorViewportTabBody::OnShowSky()
 {
 	TSharedRef<FAnimationViewportClient> AnimViewportClient = StaticCastSharedRef<FAnimationViewportClient>(LevelViewportClient.ToSharedRef());
diff --git a/Engine/Source/Editor/Persona/Private/SAnimationEditorViewport.h b/Engine/Source/Editor/Persona/Private/SAnimationEditorViewport.h
--- a/Engine/Source/Editor/Persona/Private/SAnimationEditorViewport.h
+++ b/Engine/Source/Editor/Persona/Private/SAnimationEditorViewport.h
@@ -159,6 +159,9 @@ public:
 	/** Function to check whether floor is displayed or not */
 	bool IsShowingFloor() const;
 
+	/** Function to check whether floor is auto aligned or not */
+	bool IsAutoAlignFloor() const;
+
 	/** Clears our reference to Persona, also cleaning up anything that depends on Persona first */
 	void CleanupPersonaReferences();
 
@@ -274,6 +277,9 @@ private:
 	/** Function to show/hide floor in the viewport */
 	void OnShowFloor();
 
+	/** Function to enable/disable floor auto align */
+	void OnToggleAutoAlignFloor();
+
 	/** Function to show/hide floor in the viewport */
 	void OnShowSky();
 
