Added scroll bar to notify panel and hooked up with other sections of the editor panel (curve editor etc.).
#UE-4077

diff --git a/Engine/Source/Editor/Persona/Private/SAnimNotifyPanel.cpp b/Engine/Source/Editor/Persona/Private/SAnimNotifyPanel.cpp
--- a/Engine/Source/Editor/Persona/Private/SAnimNotifyPanel.cpp
+++ b/Engine/Source/Editor/Persona/Private/SAnimNotifyPanel.cpp
@@ -3253,14 +3253,40 @@ void SAnimNotifyPanel::Construct(const FArguments& InArgs)
 			.AddMetaData<FTagMetaData>(TEXT("AnimNotify.Notify"))
 			.BodyContent()
 			[
-				SAssignNew( PanelArea, SBorder )
-				.BorderImage( FEditorStyle::GetBrush("NoBorder") )
-				.Padding( FMargin(2.0f, 2.0f) )
-				.ColorAndOpacity( FLinearColor::White )
+				SNew(SVerticalBox)
+				+SVerticalBox::Slot()
+				.FillHeight(1)
+				[
+					SAssignNew(PanelArea, SBorder)
+					.BorderImage(FEditorStyle::GetBrush("NoBorder"))
+					.Padding(FMargin(2.0f, 2.0f))
+					.ColorAndOpacity(FLinearColor::White)
+				]
+				+ SVerticalBox::Slot()
+				.AutoHeight()
+				[
+					SNew(SHorizontalBox)
+					+SHorizontalBox::Slot()
+					.FillWidth(1)
+					[
+						SAssignNew(NotifyTrackScrollBar, SScrollBar)
+						.Orientation(EOrientation::Orient_Horizontal)
+						.AlwaysShowScrollbar(true)
+						.OnUserScrolled(this, &SAnimNotifyPanel::OnNotifyTrackScrolled)
+					]
+					+SHorizontalBox::Slot()
+					.AutoWidth()
+					[
+						SNew(SBox)
+						.WidthOverride(WidgetWidth)
+					]
+				]
 			]
 		]
 	];
 
+	InputViewRangeChanged(ViewInputMin.Get(), ViewInputMax.Get());
+
 	OnPropertyChangedHandle = FCoreUObjectDelegates::FOnObjectPropertyChanged::FDelegate::CreateSP(this, &SAnimNotifyPanel::OnPropertyChanged);
 	OnPropertyChangedHandleDelegateHandle = FCoreUObjectDelegates::OnObjectPropertyChanged.Add(OnPropertyChangedHandle);
 
@@ -3447,7 +3473,7 @@ void SAnimNotifyPanel::RefreshNotifyTracks()
 			.OnDeselectAllNotifies(this, &SAnimNotifyPanel::DeselectAllNotifies)
 			.OnCopyNotifies(this, &SAnimNotifyPanel::CopySelectedNotifiesToClipboard)
 			.OnPasteNotifies(this, &SAnimNotifyPanel::OnPasteNotifies)
-			.OnSetInputViewRange(OnSetInputViewRange)
+			.OnSetInputViewRange(this, &SAnimNotifyPanel::InputViewRangeChanged)
 		];
 
 		NotifyAnimTracks.Add(EdTrack->NotifyTrack);
@@ -4011,4 +4037,26 @@ void SAnimNotifyPanel::OnNotifyObjectChanged(UObject* EditorBaseObj, bool bRebui
 	}
 }
 
+void SAnimNotifyPanel::OnNotifyTrackScrolled(float InScrollOffsetFraction)
+{
+	float Ratio = (ViewInputMax.Get() - ViewInputMin.Get()) / Sequence->SequenceLength;
+	float MaxOffset = (Ratio < 1.0f) ? 1.0f - Ratio : 0.0f;
+	InScrollOffsetFraction = FMath::Clamp(InScrollOffsetFraction, 0.0f, MaxOffset);
+
+	// Calculate new view ranges
+	float NewMin = InScrollOffsetFraction * Sequence->SequenceLength;
+	float NewMax = (InScrollOffsetFraction + Ratio) * Sequence->SequenceLength;
+	
+	InputViewRangeChanged(NewMin, NewMax);
+}
+
+void SAnimNotifyPanel::InputViewRangeChanged(float ViewMin, float ViewMax)
+{
+	float Ratio = (ViewMax - ViewMin) / Sequence->SequenceLength;
+	float OffsetFraction = ViewMin / Sequence->SequenceLength;
+	NotifyTrackScrollBar->SetState(OffsetFraction, Ratio);
+
+	SAnimTrackPanel::InputViewRangeChanged(ViewMin, ViewMax);
+}
+
 #undef LOCTEXT_NAMESPACE
diff --git a/Engine/Source/Editor/Persona/Private/SAnimNotifyPanel.h b/Engine/Source/Editor/Persona/Private/SAnimNotifyPanel.h
--- a/Engine/Source/Editor/Persona/Private/SAnimNotifyPanel.h
+++ b/Engine/Source/Editor/Persona/Private/SAnimNotifyPanel.h
@@ -211,6 +211,7 @@ public:
 
 private:
 	TSharedPtr<SBorder> PanelArea;
+	TSharedPtr<SScrollBar> NotifyTrackScrollBar;
 	class UAnimSequenceBase* Sequence;
 	float WidgetWidth;
 	TAttribute<float> CurrentPosition;
@@ -275,6 +276,10 @@ private:
 	 */
 	void OnGetNativeNotifyData(TArray<UClass*>& OutClasses, UClass* NotifyOutermost, TArray<FString>* OutAllowedBlueprintClassNames);
 
+	void OnNotifyTrackScrolled(float InScrollOffsetFraction);
+
+	virtual void InputViewRangeChanged(float ViewMin, float ViewMax) override;
+
 	/** Persona reference **/
 	TWeakPtr<FPersona> PersonaPtr;
 
diff --git a/Engine/Source/Editor/Persona/Private/SAnimTrackPanel.cpp b/Engine/Source/Editor/Persona/Private/SAnimTrackPanel.cpp
--- a/Engine/Source/Editor/Persona/Private/SAnimTrackPanel.cpp
+++ b/Engine/Source/Editor/Persona/Private/SAnimTrackPanel.cpp
@@ -87,7 +87,7 @@ FReply SAnimTrackPanel::OnMouseWheel( const FGeometry& MyGeometry, const FPointe
 		float ViewMinInput = ViewInputMin.Get() - (InputChange * ZoomRatio);
 		float ViewMaxInput = ViewInputMax.Get() + (InputChange * (1.f - ZoomRatio));
 		
-		OnSetInputViewRange.Execute(ViewMinInput, ViewMaxInput);
+		InputViewRangeChanged(ViewMinInput, ViewMaxInput);
 	}
 
 	return FReply::Handled();
@@ -160,7 +160,7 @@ FReply SAnimTrackPanel::OnMouseMove( const FGeometry& InMyGeometry, const FPoint
 				NewViewInputMin = NewViewInputMax - ScaleInfo.ViewInputRange;
 			}
 
-			OnSetInputViewRange.Execute(NewViewInputMin, NewViewInputMax);
+			InputViewRangeChanged(NewViewInputMin, NewViewInputMax);
 
 			UE_LOG(LogAnimation, Log, TEXT("MouseMove (Panning) %0.2f, %0.2f"), ViewInputMin.Get(), ViewInputMax.Get());
 			return FReply::Handled();
@@ -192,7 +192,12 @@ void SAnimTrackPanel::PanInputViewRange(int32 ScreenDelta, FVector2D ScreenViewS
 		NewViewInputMin = NewViewInputMax - ScaleInfo.ViewInputRange;
 	}
 
-	OnSetInputViewRange.Execute(NewViewInputMin, NewViewInputMax);
+	InputViewRangeChanged(NewViewInputMin, NewViewInputMax);
+}
+
+void SAnimTrackPanel::InputViewRangeChanged(float ViewMin, float ViewMax)
+{
+	OnSetInputViewRange.Execute(ViewMin, ViewMax);
 }
 
 #undef LOCTEXT_NAMESPACE
diff --git a/Engine/Source/Editor/Persona/Private/SAnimTrackPanel.h b/Engine/Source/Editor/Persona/Private/SAnimTrackPanel.h
--- a/Engine/Source/Editor/Persona/Private/SAnimTrackPanel.h
+++ b/Engine/Source/Editor/Persona/Private/SAnimTrackPanel.h
@@ -94,8 +94,13 @@ public:
 	void PanInputViewRange(int32 ScreenDelta, FVector2D ScreenViewSize);
 
 	virtual float GetSequenceLength() const {return 0.0f;}
+
 protected:
 
+	// Calls out to notify of a range change, can be overridden by derived classes to respond
+	// but they must call this version too after processing range changes
+	virtual void InputViewRangeChanged(float ViewMin, float ViewMax);
+
 	/** Create a 2 column widget */
 	TSharedRef<class S2ColumnWidget> Create2ColumnWidget( TSharedRef<SVerticalBox> Parent );	
 
