Anim Notify Editor changes:

- Added "Replace with" option in right click menu for existing notify nodes.
- Fixed issue with deleting multiple animations at the same time deleting the wrong notifies

diff --git a/Engine/Source/Editor/Persona/Private/SAnimNotifyPanel.cpp b/Engine/Source/Editor/Persona/Private/SAnimNotifyPanel.cpp
--- a/Engine/Source/Editor/Persona/Private/SAnimNotifyPanel.cpp
+++ b/Engine/Source/Editor/Persona/Private/SAnimNotifyPanel.cpp
@@ -219,6 +219,9 @@ public:
 		, _OnRequestTrackPan()
 		, _OnRequestOffsetRefresh()
 		, _OnDeleteNotify()
+		, _OnGetIsAnimNotifySelectionValidForReplacement()
+		, _OnReplaceSelectedWithNotify()
+		, _OnReplaceSelectedWithBlueprintNotify()
 		, _OnDeselectAllNotifies()
 		, _OnCopyNotifies()
 		, _OnPasteNotifies()
@@ -245,6 +248,9 @@ public:
 		SLATE_EVENT( FPanTrackRequest, OnRequestTrackPan )
 		SLATE_EVENT( FRefreshOffsetsRequest, OnRequestOffsetRefresh )
 		SLATE_EVENT( FDeleteNotify, OnDeleteNotify )
+		SLATE_EVENT( FOnGetIsAnimNotifySelectionValidForReplacement, OnGetIsAnimNotifySelectionValidForReplacement)
+		SLATE_EVENT( FReplaceWithNotify, OnReplaceSelectedWithNotify )
+		SLATE_EVENT( FReplaceWithBlueprintNotify, OnReplaceSelectedWithBlueprintNotify)
 		SLATE_EVENT( FDeselectAllNotifies, OnDeselectAllNotifies)
 		SLATE_EVENT( FCopyNotifies, OnCopyNotifies )
 		SLATE_EVENT( FPasteNotifies, OnPasteNotifies )
@@ -313,19 +319,37 @@ public:
 	// Uses the given track space rect and marquee information to refresh selection information
 	void RefreshMarqueeSelectedNodes(FSlateRect& Rect, FNotifyMarqueeOperation& Marquee);
 
+	// Create new notifies
+	FAnimNotifyEvent& CreateNewBlueprintNotify(FString NewNotifyName, FString BlueprintPath, float StartTime);
+	FAnimNotifyEvent& CreateNewNotify(FString NewNotifyName, UClass* NotifyClass, float StartTime);
+
+	// Get the Blueprint Class from the path of the Blueprint
+	static TSubclassOf<UObject> GetBlueprintClassFromPath(FString BlueprintPath);
+
+	// Get the default Notify Name for a given blueprint notify asset
+	FString MakeBlueprintNotifyName(FAssetData& NotifyAssetData);
+
 protected:
 
 	void CreateCommands();
 
 	// Build up a "New Notify..." menu
-	void FillNewNotifyMenu(FMenuBuilder& MenuBuilder);
-	void FillNewNotifyStateMenu(FMenuBuilder& MenuBuilder);
+	void FillNewNotifyMenu(FMenuBuilder& MenuBuilderbool, bool bIsReplaceWithMenu = false);
+	void FillNewNotifyStateMenu(FMenuBuilder& MenuBuilder, bool bIsReplaceWithMenu  = false);
+
+	// New notify functions
 	void CreateNewBlueprintNotifyAtCursor(FString NewNotifyName, FString BlueprintPath);
 	void CreateNewNotifyAtCursor(FString NewNotifyName, UClass* NotifyClass);
 	void OnNewNotifyClicked();
 	void AddNewNotify(const FText& NewNotifyName, ETextCommit::Type CommitInfo);
+
+	// Trigger weight functions
 	void OnSetTriggerWeightNotifyClicked(int32 NotifyIndex);
 	void SetTriggerWeight(const FText& TriggerWeight, ETextCommit::Type CommitInfo, int32 NotifyIndex);
+
+	// "Replace with... " commands
+	void ReplaceSelectedWithBlueprintNotify(FString NewNotifyName, FString BlueprintPath);
+	void ReplaceSelectedWithNotify(FString NewNotifyName, UClass* NotifyClass);
 	void OnSetNotifyTimeClicked(int32 NotifyIndex);
 	void SetNotifyTime(const FText& NotifyTimeText, ETextCommit::Type CommitInfo, int32 NotifyIndex);
 	void OnSetNotifyFrameClicked(int32 NotifyIndex);
@@ -463,6 +487,11 @@ protected:
 	/** Delegate to call when deleting notifies */
 	FDeleteNotify							OnDeleteNotify;
 
+	/** Delegates to call when replacing notifies */
+	FOnGetIsAnimNotifySelectionValidForReplacement OnGetIsAnimNotifySelectionValidforReplacement;
+	FReplaceWithNotify						OnReplaceSelectedWithNotify;
+	FReplaceWithBlueprintNotify				OnReplaceSelectedWithBlueprintNotify;
+
 	TSharedPtr<SBorder>						TrackArea;
 	TWeakPtr<FPersona>						PersonaPtr;
 
@@ -1625,6 +1654,9 @@ void SAnimNotifyTrack::Construct(const FArguments& InArgs)
 	OnRequestTrackPan = InArgs._OnRequestTrackPan;
 	OnRequestRefreshOffsets = InArgs._OnRequestOffsetRefresh;
 	OnDeleteNotify = InArgs._OnDeleteNotify;
+	OnGetIsAnimNotifySelectionValidforReplacement = InArgs._OnGetIsAnimNotifySelectionValidForReplacement;
+	OnReplaceSelectedWithNotify = InArgs._OnReplaceSelectedWithNotify;
+	OnReplaceSelectedWithBlueprintNotify = InArgs._OnReplaceSelectedWithBlueprintNotify;
 	OnDeselectAllNotifies = InArgs._OnDeselectAllNotifies;
 	OnCopyNotifies = InArgs._OnCopyNotifies;
 	OnPasteNotifies = InArgs._OnPasteNotifies;
@@ -1817,7 +1849,8 @@ FCursorReply SAnimNotifyTrack::OnCursorQuery(const FGeometry& MyGeometry, const
 	return FCursorReply::Unhandled();
 }
 
-void SAnimNotifyTrack::FillNewNotifyStateMenu(FMenuBuilder& MenuBuilder)
+// Fill "new notify state" menu, or "replace with notify state menu"
+void SAnimNotifyTrack::FillNewNotifyStateMenu(FMenuBuilder& MenuBuilder, bool bIsReplaceWithMenu /* = false */)
 {
 	// Run the native query first to update the allowed classes for blueprints.
 	TArray<UClass*> NotifyStateClasses;
@@ -1834,12 +1867,24 @@ void SAnimNotifyTrack::FillNewNotifyStateMenu(FMenuBuilder& MenuBuilder)
 		const FText LabelText = FText::FromString(NotifyData.NotifyName);
 
 		FUIAction UIAction;
-		UIAction.ExecuteAction.BindRaw(
-			this, &SAnimNotifyTrack::CreateNewBlueprintNotifyAtCursor,
-			NotifyData.NotifyName,
-			NotifyData.BlueprintPath);
+		FText Description = FText::GetEmpty();
+		if (!bIsReplaceWithMenu)
+		{
+			Description = LOCTEXT("AddsAnExistingAnimNotify", "Add an existing notify");
+			UIAction.ExecuteAction.BindRaw(
+				this, &SAnimNotifyTrack::CreateNewBlueprintNotifyAtCursor,
+				NotifyData.NotifyName,
+				NotifyData.BlueprintPath);
+		}
+		else
+		{
+			Description = LOCTEXT("ReplaceWithAnExistingAnimNotify", "Replace with an existing notify");
+			UIAction.ExecuteAction.BindRaw(
+				this, &SAnimNotifyTrack::ReplaceSelectedWithBlueprintNotify,
+				NotifyData.NotifyName,
+				NotifyData.BlueprintPath);
+		}
 
-		const static FText Description = LOCTEXT("AddsAnExistingAnimNotify", "Add an existing notify");
 		MenuBuilder.AddMenuEntry(LabelText, Description, FSlateIcon(), UIAction);
 	}
 
@@ -1852,10 +1897,20 @@ void SAnimNotifyTrack::FillNewNotifyStateMenu(FMenuBuilder& MenuBuilder)
 			const FText LabelText = FText::FromString(Label);
 
 			FUIAction UIAction;
-			UIAction.ExecuteAction.BindRaw(
-				this, &SAnimNotifyTrack::CreateNewNotifyAtCursor,
-				Label,
-				Class);
+			if (!bIsReplaceWithMenu)
+			{
+				UIAction.ExecuteAction.BindRaw(
+					this, &SAnimNotifyTrack::CreateNewNotifyAtCursor,
+					Label,
+					Class);
+			}
+			else
+			{
+				UIAction.ExecuteAction.BindRaw(
+					this, &SAnimNotifyTrack::ReplaceSelectedWithNotify,
+					Label,
+					Class);
+			}
 
 			MenuBuilder.AddMenuEntry(LabelText, Description, FSlateIcon(), UIAction);
 		}
@@ -1863,7 +1918,8 @@ void SAnimNotifyTrack::FillNewNotifyStateMenu(FMenuBuilder& MenuBuilder)
 	MenuBuilder.EndSection();
 }
 
-void SAnimNotifyTrack::FillNewNotifyMenu(FMenuBuilder& MenuBuilder)
+// Fill "new notify" menu, or "replace with notify menu"
+void SAnimNotifyTrack::FillNewNotifyMenu(FMenuBuilder& MenuBuilder, bool bIsReplaceWithMenu /* = false */)
 {
 	TArray<UClass*> NativeNotifyClasses;
 	OnGetNotifyNativeClasses.ExecuteIfBound(NativeNotifyClasses);
@@ -1878,12 +1934,24 @@ void SAnimNotifyTrack::FillNewNotifyMenu(FMenuBuilder& MenuBuilder)
 		const FText LabelText = FText::FromString( NotifyData.NotifyName );
 
 		FUIAction UIAction;
-		UIAction.ExecuteAction.BindRaw(
-			this, &SAnimNotifyTrack::CreateNewBlueprintNotifyAtCursor,
-			NotifyData.NotifyName,
-			NotifyData.BlueprintPath);
+		FText Description = FText::GetEmpty();
+		if (!bIsReplaceWithMenu)
+		{
+			Description = LOCTEXT("NewNotifySubMenu_ToolTip", "Add an existing notify");
+			UIAction.ExecuteAction.BindRaw(
+				this, &SAnimNotifyTrack::CreateNewBlueprintNotifyAtCursor,
+				NotifyData.NotifyName,
+				NotifyData.BlueprintPath);
+		}
+		else
+		{
+			Description = LOCTEXT("ReplaceWithNotifySubMenu_ToolTip", "Replace with an existing notify");
+			UIAction.ExecuteAction.BindRaw(
+				this, &SAnimNotifyTrack::ReplaceSelectedWithBlueprintNotify,
+				NotifyData.NotifyName,
+				NotifyData.BlueprintPath);
+		}
 		
-		const static FText Description = LOCTEXT("NewNotifySubMenu_ToolTip", "Add an existing notify");
 		MenuBuilder.AddMenuEntry(LabelText, Description, FSlateIcon(), UIAction);
 	}
 
@@ -1891,15 +1959,27 @@ void SAnimNotifyTrack::FillNewNotifyMenu(FMenuBuilder& MenuBuilder)
 	{
 		for(UClass* Class : NativeNotifyClasses)
 		{
-			const FText Description = LOCTEXT("NewNotifySubMenu_NativeToolTip", "Add an existing native notify");
 			FString Label = Class->GetName();
 			const FText LabelText = FText::FromString(Label);
 
 			FUIAction UIAction;
-			UIAction.ExecuteAction.BindRaw(
-				this, &SAnimNotifyTrack::CreateNewNotifyAtCursor,
-				Label,
-				Class);
+			FText Description = FText::GetEmpty();
+			if (!bIsReplaceWithMenu)
+			{
+				Description = LOCTEXT("NewNotifySubMenu_NativeToolTip", "Add an existing native notify");
+				UIAction.ExecuteAction.BindRaw(
+					this, &SAnimNotifyTrack::CreateNewNotifyAtCursor,
+					Label,
+					Class);
+			}
+			else
+			{
+				Description = LOCTEXT("ReplaceWithNotifySubMenu_NativeToolTip", "Replace with an existing native notify");
+				UIAction.ExecuteAction.BindRaw(
+					this, &SAnimNotifyTrack::ReplaceSelectedWithNotify,
+					Label,
+					Class);
+			}
 
 			MenuBuilder.AddMenuEntry(LabelText, Description, FSlateIcon(), UIAction);
 		}
@@ -1917,13 +1997,31 @@ void SAnimNotifyTrack::FillNewNotifyMenu(FMenuBuilder& MenuBuilder)
 				FName NotifyName = SeqSkeleton->AnimationNotifies[I];
 				FString Label = NotifyName.ToString();
 
-				const FText Description = LOCTEXT("NewNotifySubMenu_ToolTip", "Add an existing notify");
+				FText Description = FText::GetEmpty();
+				if (!bIsReplaceWithMenu)
+				{
+					Description = LOCTEXT("NewNotifySubMenu_ToolTip", "Add an existing notify");
+				}
+				else
+				{
+					Description = LOCTEXT("ReplaceWithNotifySubMenu_ToolTip", "Replace with an existing notify");
+				}
 
 				FUIAction UIAction;
-				UIAction.ExecuteAction.BindRaw(
-					this, &SAnimNotifyTrack::CreateNewNotifyAtCursor,
-					Label, 
-					(UClass*)nullptr);
+				if (!bIsReplaceWithMenu)
+				{
+					UIAction.ExecuteAction.BindRaw(
+						this, &SAnimNotifyTrack::CreateNewNotifyAtCursor,
+						Label,
+						(UClass*)nullptr);
+				}
+				else
+				{
+					UIAction.ExecuteAction.BindRaw(
+						this, &SAnimNotifyTrack::ReplaceSelectedWithNotify,
+						Label,
+						(UClass*)nullptr);
+				}
 
 				MenuBuilder.AddMenuEntry( FText::FromString( Label ), Description, FSlateIcon(), UIAction);
 			}
@@ -1931,41 +2029,35 @@ void SAnimNotifyTrack::FillNewNotifyMenu(FMenuBuilder& MenuBuilder)
 	}
 	MenuBuilder.EndSection();
 
-	MenuBuilder.BeginSection("AnimNotifyCreateNew");
+	if (!bIsReplaceWithMenu)
 	{
-		FUIAction UIAction;
-		UIAction.ExecuteAction.BindRaw(
-			this, &SAnimNotifyTrack::OnNewNotifyClicked);
-		MenuBuilder.AddMenuEntry(LOCTEXT("NewNotify", "New Notify"), LOCTEXT("NewNotifyToolTip", "Create a new animation notify"), FSlateIcon(), UIAction);
+		MenuBuilder.BeginSection("AnimNotifyCreateNew");
+		{
+			FUIAction UIAction;
+			UIAction.ExecuteAction.BindRaw(
+				this, &SAnimNotifyTrack::OnNewNotifyClicked);
+			MenuBuilder.AddMenuEntry(LOCTEXT("NewNotify", "New Notify"), LOCTEXT("NewNotifyToolTip", "Create a new animation notify"), FSlateIcon(), UIAction);
+		}
+		MenuBuilder.EndSection();
 	}
-	MenuBuilder.EndSection();
 }
 
-void SAnimNotifyTrack::CreateNewBlueprintNotifyAtCursor(FString NewNotifyName, FString BlueprintPath)
+FAnimNotifyEvent& SAnimNotifyTrack::CreateNewBlueprintNotify(FString NewNotifyName, FString BlueprintPath, float StartTime)
 {
-	TSubclassOf<UObject> BlueprintClass = NULL;
-	if( !BlueprintPath.IsEmpty() )
-	{
-		UBlueprint* BlueprintLibPtr = LoadObject<UBlueprint>(NULL, *BlueprintPath, NULL, 0, NULL);
-		BlueprintClass = Cast<UClass>(BlueprintLibPtr->GeneratedClass);
-	}
-	check(BlueprintClass)
-	CreateNewNotifyAtCursor(NewNotifyName, BlueprintClass);
+	TSubclassOf<UObject> BlueprintClass = GetBlueprintClassFromPath(BlueprintPath);
+	check(BlueprintClass);
+	return CreateNewNotify(NewNotifyName, BlueprintClass, StartTime);
 }
 
-void SAnimNotifyTrack::CreateNewNotifyAtCursor(FString NewNotifyName, UClass* NotifyClass)
+FAnimNotifyEvent& SAnimNotifyTrack::CreateNewNotify(FString NewNotifyName, UClass* NotifyClass, float StartTime)
 {
-	const float NewTime = LastClickedTime;
-	const FScopedTransaction Transaction( LOCTEXT("AddNotifyEvent", "Add Anim Notify") );
-	Sequence->Modify();
-
 	// Insert a new notify record and spawn the new notify object
 	int32 NewNotifyIndex = Sequence->Notifies.Add(FAnimNotifyEvent());
 	FAnimNotifyEvent& NewEvent = Sequence->Notifies[NewNotifyIndex];
 	NewEvent.NotifyName = FName(*NewNotifyName);
 
-	NewEvent.Link(Sequence, NewTime);
-	NewEvent.TriggerTimeOffset = GetTriggerTimeOffsetForType(Sequence->CalculateOffsetForNotify(NewTime));
+	NewEvent.Link(Sequence, StartTime);
+	NewEvent.TriggerTimeOffset = GetTriggerTimeOffsetForType(Sequence->CalculateOffsetForNotify(StartTime));
 	NewEvent.TrackIndex = TrackIndex;
 
 	if( NotifyClass )
@@ -2024,9 +2116,47 @@ void SAnimNotifyTrack::CreateNewNotifyAtCursor(FString NewNotifyName, UClass* No
 	}
 
 	Sequence->MarkPackageDirty();
+
+	return NewEvent;
+}
+
+void SAnimNotifyTrack::CreateNewBlueprintNotifyAtCursor(FString NewNotifyName, FString BlueprintPath)
+{
+	TSubclassOf<UObject> BlueprintClass = GetBlueprintClassFromPath(BlueprintPath);
+	check(BlueprintClass);
+	CreateNewNotifyAtCursor(NewNotifyName, BlueprintClass);
+}
+
+void SAnimNotifyTrack::CreateNewNotifyAtCursor(FString NewNotifyName, UClass* NotifyClass)
+{
+	const FScopedTransaction Transaction(LOCTEXT("AddNotifyEvent", "Add Anim Notify"));
+	Sequence->Modify();
+	CreateNewNotify(NewNotifyName, NotifyClass, LastClickedTime);
 	OnUpdatePanel.ExecuteIfBound();
 }
 
+void SAnimNotifyTrack::ReplaceSelectedWithBlueprintNotify(FString NewNotifyName, FString BlueprintPath)
+{
+	OnReplaceSelectedWithBlueprintNotify.ExecuteIfBound(NewNotifyName, BlueprintPath);
+}
+
+void SAnimNotifyTrack::ReplaceSelectedWithNotify(FString NewNotifyName, UClass* NotifyClass)
+{
+	OnReplaceSelectedWithNotify.ExecuteIfBound(NewNotifyName, NotifyClass);
+}
+
+
+TSubclassOf<UObject> SAnimNotifyTrack::GetBlueprintClassFromPath(FString BlueprintPath)
+{
+	TSubclassOf<UObject> BlueprintClass = NULL;
+	if (!BlueprintPath.IsEmpty())
+	{
+		UBlueprint* BlueprintLibPtr = LoadObject<UBlueprint>(NULL, *BlueprintPath, NULL, 0, NULL);
+		BlueprintClass = Cast<UClass>(BlueprintLibPtr->GeneratedClass);
+	}
+	return BlueprintClass;
+}
+
 FReply SAnimNotifyTrack::OnMouseButtonUp( const FGeometry& MyGeometry, const FPointerEvent& MouseEvent )
 {
 	bool bLeftMouseButton =  MouseEvent.GetEffectingButton() == EKeys::LeftMouseButton;
@@ -2223,12 +2353,12 @@ TSharedPtr<SWidget> SAnimNotifyTrack::SummonContextMenu(const FGeometry& MyGeome
 			MenuBuilder.AddSubMenu(
 				NSLOCTEXT("NewNotifySubMenu", "NewNotifySubMenuAddNotify", "Add Notify..."),
 				NSLOCTEXT("NewNotifySubMenu", "NewNotifySubMenuAddNotifyToolTip", "Add AnimNotifyEvent"),
-				FNewMenuDelegate::CreateRaw( this, &SAnimNotifyTrack::FillNewNotifyMenu ) );
+				FNewMenuDelegate::CreateRaw( this, &SAnimNotifyTrack::FillNewNotifyMenu, false ) );
 
 			MenuBuilder.AddSubMenu(
 				NSLOCTEXT("NewNotifySubMenu", "NewNotifySubMenuAddNotifyState", "Add Notify State..."),
 				NSLOCTEXT("NewNotifySubMenu", "NewNotifySubMenuAddNotifyStateToolTip","Add AnimNotifyState"),
-				FNewMenuDelegate::CreateRaw( this, &SAnimNotifyTrack::FillNewNotifyStateMenu ) );
+				FNewMenuDelegate::CreateRaw( this, &SAnimNotifyTrack::FillNewNotifyStateMenu, false ) );
 
 			MenuBuilder.AddMenuEntry(
 				NSLOCTEXT("NewNotifySubMenu", "ManageNotifies", "Manage Notifies..."),
@@ -2253,6 +2383,29 @@ TSharedPtr<SWidget> SAnimNotifyTrack::SummonContextMenu(const FGeometry& MyGeome
 			// allow it to delete
 			NewAction.ExecuteAction = OnDeleteNotify;
 			MenuBuilder.AddMenuEntry(LOCTEXT("Delete", "Delete"), LOCTEXT("DeleteToolTip", "Delete animation notify"), FSlateIcon(), NewAction);
+
+			FAnimNotifyEvent* Notify = AnimNotifies[NotifyIndex];
+
+			// For the "Replace With..." menu, make sure the current AnimNotify selection is valid for replacement
+			if (OnGetIsAnimNotifySelectionValidforReplacement.IsBound() && OnGetIsAnimNotifySelectionValidforReplacement.Execute())
+			{
+				// If this is an AnimNotifyState (has duration) allow it to be replaced with other AnimNotifyStates
+				if (Notify && Notify->NotifyStateClass)
+				{
+					MenuBuilder.AddSubMenu(
+						NSLOCTEXT("NewNotifySubMenu", "NewNotifySubMenuReplaceWithNotifyState", "Replace with Notify State..."),
+						NSLOCTEXT("NewNotifySubMenu", "NewNotifySubMenuReplaceWithNotifyStateToolTip", "Replace with AnimNotifyState"),
+						FNewMenuDelegate::CreateRaw(this, &SAnimNotifyTrack::FillNewNotifyStateMenu, true));
+				}
+				// If this is a regular AnimNotify (no duration) allow it to be replaced with other AnimNotifies
+				else 
+				{
+					MenuBuilder.AddSubMenu(
+						NSLOCTEXT("NewNotifySubMenu", "NewNotifySubMenuReplaceWithNotify", "Replace with Notify..."),
+						NSLOCTEXT("NewNotifySubMenu", "NewNotifySubMenuReplaceWithNotifyToolTip", "Replace with AnimNotifyEvent"),
+						FNewMenuDelegate::CreateRaw(this, &SAnimNotifyTrack::FillNewNotifyMenu, true));
+				}
+			}
 		}
 		else
 		{
@@ -2979,6 +3132,15 @@ void SAnimNotifyTrack::RefreshMarqueeSelectedNodes(FSlateRect& Rect, FNotifyMarq
 	}
 }
 
+FString SAnimNotifyTrack::MakeBlueprintNotifyName(FAssetData& NotifyAssetData)
+{
+	FString DefaultNotifyName = NotifyAssetData.AssetName.ToString();
+	DefaultNotifyName = DefaultNotifyName.Replace(TEXT("AnimNotify_"), TEXT(""), ESearchCase::CaseSensitive);
+	DefaultNotifyName = DefaultNotifyName.Replace(TEXT("AnimNotifyState_"), TEXT(""), ESearchCase::CaseSensitive);
+
+	return DefaultNotifyName;
+}
+
 void SAnimNotifyTrack::GetNotifyMenuData(TArray<FAssetData>& NotifyAssetData, TArray<BlueprintNotifyMenuInfo>& OutNotifyMenuData)
 {
 	for(FAssetData& NotifyData : NotifyAssetData)
@@ -2987,9 +3149,7 @@ void SAnimNotifyTrack::GetNotifyMenuData(TArray<FAssetData>& NotifyAssetData, TA
 		BlueprintNotifyMenuInfo& MenuInfo = OutNotifyMenuData.Last();
 
 		MenuInfo.BlueprintPath = NotifyData.ObjectPath.ToString();
-		MenuInfo.NotifyName = NotifyData.AssetName.ToString();
-		MenuInfo.NotifyName = MenuInfo.NotifyName.Replace(TEXT("AnimNotify_"), TEXT(""), ESearchCase::CaseSensitive);
-		MenuInfo.NotifyName = MenuInfo.NotifyName.Replace(TEXT("AnimNotifyState_"), TEXT(""), ESearchCase::CaseSensitive);
+		MenuInfo.NotifyName = MakeBlueprintNotifyName(NotifyData);
 	}
 
 	OutNotifyMenuData.Sort([](const BlueprintNotifyMenuInfo& A, const BlueprintNotifyMenuInfo& B)
@@ -3146,6 +3306,9 @@ void SNotifyEdTrack::Construct(const FArguments& InArgs)
 				.OnRequestTrackPan(FPanTrackRequest::CreateSP(PanelRef, &SAnimNotifyPanel::PanInputViewRange))
 				.OnRequestOffsetRefresh(InArgs._OnRequestRefreshOffsets)
 				.OnDeleteNotify(InArgs._OnDeleteNotify)
+				.OnGetIsAnimNotifySelectionValidForReplacement(PanelRef, &SAnimNotifyPanel::IsNotifySelectionValidForReplacement)
+				.OnReplaceSelectedWithNotify(PanelRef, &SAnimNotifyPanel::OnReplaceSelectedWithNotify)
+				.OnReplaceSelectedWithBlueprintNotify(PanelRef, &SAnimNotifyPanel::OnReplaceSelectedWithNotifyBlueprint)
 				.OnDeselectAllNotifies(InArgs._OnDeselectAllNotifies)
 				.OnCopyNotifies(InArgs._OnCopyNotifies)
 				.OnPasteNotifies(InArgs._OnPasteNotifies)
@@ -3554,19 +3717,23 @@ void SAnimNotifyPanel::OnDeletePressed()
 	// so don't delete anything when the key is pressed.
 	if(HasKeyboardFocus() || HasFocusedDescendants()) 
 	{
-		TArray<TSharedPtr<SAnimNotifyNode>> SelectedNodes;
-		for(TSharedPtr<SAnimNotifyTrack> Track : NotifyAnimTracks)
+		TArray<FAnimNotifyEvent*> SelectedNotifies;
+		for (TSharedPtr<SAnimNotifyTrack> Track : NotifyAnimTracks)
 		{
-			Track->AppendSelectedNodeWidgetsToArray(SelectedNodes);
+			Track->AppendSelectionToArray(SelectedNotifies);
 		}
 
-		if(SelectedNodes.Num() > 0)
+		if(SelectedNotifies.Num() > 0)
 		{
 			FScopedTransaction Transaction(LOCTEXT("DeleteNotifies", "Delete Animation Notifies"));
 			Sequence->Modify(true);
-			for(TSharedPtr<SAnimNotifyNode> Node : SelectedNodes)
+
+			// Delete from highest index to lowest
+			SelectedNotifies.Sort();
+			for (int NotifyIndex = SelectedNotifies.Num()-1; NotifyIndex >= 0 ; --NotifyIndex)
 			{
-				DeleteNotify(Node->NotifyEvent);
+				FAnimNotifyEvent* Notify = SelectedNotifies[NotifyIndex];
+				DeleteNotify(Notify);
 			}
 		}
 
@@ -3693,6 +3860,113 @@ void SAnimNotifyPanel::CopySelectedNotifiesToClipboard() const
 	}
 }
 
+bool SAnimNotifyPanel::IsNotifySelectionValidForReplacement()
+{
+	// Grab the selected events
+	TArray<FAnimNotifyEvent*> SelectedNotifies;
+	for (TSharedPtr<SAnimNotifyTrack> Track : NotifyAnimTracks)
+	{
+		Track->AppendSelectionToArray(SelectedNotifies);
+	}
+
+	bool bSelectionContainsAnimNotify = false;
+	bool bSelectionContainsAnimNotifyState = false;
+	for (FAnimNotifyEvent* NotifyEvent : SelectedNotifies)
+	{
+		if (NotifyEvent)
+		{
+			if (NotifyEvent->Notify)
+			{
+				bSelectionContainsAnimNotify = true;
+			}
+			else if (NotifyEvent->NotifyStateClass)
+			{
+				bSelectionContainsAnimNotifyState = true;
+			}
+			// Custom AnimNotifies have no class, but they are like AnimNotify class notifies in that they have no duration
+			else
+			{
+				bSelectionContainsAnimNotify = true;
+			}
+		}
+	}
+
+	// Only allow replacement for selections that contain _only_ AnimNotifies, or _only_ AnimNotifyStates, but not both
+	// (Want to disallow replacement of AnimNotify with AnimNotifyState, and vice-versa)
+	bool bIsValidSelection = bSelectionContainsAnimNotify != bSelectionContainsAnimNotifyState;
+
+	return bIsValidSelection;
+}
+
+
+void SAnimNotifyPanel::OnReplaceSelectedWithNotify(FString NewNotifyName, UClass* NewNotifyClass)
+{
+	TArray<FAnimNotifyEvent*> SelectedNotifies;
+	for (TSharedPtr<SAnimNotifyTrack> Track : NotifyAnimTracks)
+	{
+		Track->AppendSelectionToArray(SelectedNotifies);
+	}
+
+	// Sort these since order is important for deletion
+	SelectedNotifies.Sort();
+
+	const FScopedTransaction Transaction(LOCTEXT("AddNotifyEvent", "Replace Anim Notify"));
+	Sequence->Modify(true);
+
+	for (int NotifyIndex = SelectedNotifies.Num()-1; NotifyIndex >= 0; --NotifyIndex)
+	{
+		FAnimNotifyEvent* OldEvent = SelectedNotifies[NotifyIndex];
+		if (OldEvent)
+		{
+			float BeginTime = OldEvent->GetTime();
+			float Length = OldEvent->GetDuration();
+			int32 TargetTrackIndex = (NotifyAnimTracks.Num()-1) - OldEvent->TrackIndex;
+			float TriggerTimeOffset = OldEvent->TriggerTimeOffset;
+			float EndTriggerTimeOffset = OldEvent->EndTriggerTimeOffset;
+			int32 SlotIndex = OldEvent->GetSlotIndex();
+			int32 EndSlotIndex = OldEvent->EndLink.GetSlotIndex();
+			int32 SegmentIndex = OldEvent->GetSegmentIndex();
+			int32 EndSegmentIndex = OldEvent->GetSegmentIndex();
+			EAnimLinkMethod::Type LinkMethod = OldEvent->GetLinkMethod();
+			EAnimLinkMethod::Type EndLinkMethod = OldEvent->EndLink.GetLinkMethod();
+
+			// Delete old one before creating new one to avoid potential array re-allocation when array temporarily increases by 1 in size
+			DeleteNotify(OldEvent);
+			FAnimNotifyEvent& NewEvent = NotifyAnimTracks[TargetTrackIndex]->CreateNewNotify(NewNotifyName, NewNotifyClass, BeginTime);
+
+			NewEvent.TriggerTimeOffset = TriggerTimeOffset;
+			NewEvent.ChangeSlotIndex(SlotIndex);
+			NewEvent.SetSegmentIndex(SegmentIndex);
+			NewEvent.ChangeLinkMethod(LinkMethod);
+
+			// For Anim Notify States, handle the end time and link
+			if (NewEvent.NotifyStateClass != nullptr)
+			{
+				NewEvent.SetDuration(Length);
+				NewEvent.EndTriggerTimeOffset = EndTriggerTimeOffset;
+				NewEvent.EndLink.ChangeSlotIndex(EndSlotIndex);
+				NewEvent.EndLink.SetSegmentIndex(EndSegmentIndex);
+				NewEvent.EndLink.ChangeLinkMethod(EndLinkMethod);
+			}
+						
+			NewEvent.Update();
+		}
+	}
+
+	// clear selection  
+	FGraphPanelSelectionSet ObjectSet;
+	OnSelectionChanged.ExecuteIfBound(ObjectSet);
+	// TODO: set selection to new notifies?
+	// update the panel
+	Update();
+}
+
+void SAnimNotifyPanel::OnReplaceSelectedWithNotifyBlueprint(FString NewBlueprintNotifyName, FString NewBlueprintNotifyClass)
+{
+	TSubclassOf<UObject> BlueprintClass = SAnimNotifyTrack::GetBlueprintClassFromPath(NewBlueprintNotifyClass);
+	OnReplaceSelectedWithNotify(NewBlueprintNotifyName, BlueprintClass);
+}
+
 void SAnimNotifyPanel::OnPasteNotifies(SAnimNotifyTrack* RequestTrack, float ClickTime, ENotifyPasteMode::Type PasteMode, ENotifyPasteMultipleMode::Type MultiplePasteType)
 {
 	int32 PasteIdx = RequestTrack->GetTrackIndex();
diff --git a/Engine/Source/Editor/Persona/Private/SAnimNotifyPanel.h b/Engine/Source/Editor/Persona/Private/SAnimNotifyPanel.h
--- a/Engine/Source/Editor/Persona/Private/SAnimNotifyPanel.h
+++ b/Engine/Source/Editor/Persona/Private/SAnimNotifyPanel.h
@@ -16,6 +16,9 @@ DECLARE_DELEGATE( FOnUpdatePanel )
 DECLARE_DELEGATE_RetVal( float, FOnGetScrubValue )
 DECLARE_DELEGATE( FRefreshOffsetsRequest )
 DECLARE_DELEGATE( FDeleteNotify )
+DECLARE_DELEGATE_RetVal( bool, FOnGetIsAnimNotifySelectionValidForReplacement )
+DECLARE_DELEGATE_TwoParams( FReplaceWithNotify, FString, UClass* )
+DECLARE_DELEGATE_TwoParams( FReplaceWithBlueprintNotify, FString, FString )
 DECLARE_DELEGATE( FDeselectAllNotifies )
 DECLARE_DELEGATE( FCopyNotifies )
 DECLARE_DELEGATE_OneParam( FOnGetBlueprintNotifyData, TArray<FAssetData>& )
@@ -209,6 +212,15 @@ public:
 
 	void OnNotifyObjectChanged(UObject* EditorBaseObj, bool bRebuild);
 
+	/** Check to make sure the current AnimNotify selection is a valid selection for replacing (i.e., AnimNotifies and AnimNotifyStates aren't mixed together in the selection) */
+	bool IsNotifySelectionValidForReplacement();
+
+	/** Handler for replacing with notify */
+	void OnReplaceSelectedWithNotify(FString NewNotifyName, UClass* NewNotifyClass);
+
+	/** Handler for replacing with notify blueprint */
+	void OnReplaceSelectedWithNotifyBlueprint(FString NewBlueprintNotifyName, FString NewBlueprintNotifyClass);
+
 private:
 	TSharedPtr<SBorder> PanelArea;
 	TSharedPtr<SScrollBar> NotifyTrackScrollBar;
diff --git a/Engine/Source/Runtime/Engine/Classes/Animation/AnimSequenceBase.h b/Engine/Source/Runtime/Engine/Classes/Animation/AnimSequenceBase.h
--- a/Engine/Source/Runtime/Engine/Classes/Animation/AnimSequenceBase.h
+++ b/Engine/Source/Runtime/Engine/Classes/Animation/AnimSequenceBase.h
@@ -184,9 +184,36 @@ struct FAnimNotifyEvent : public FAnimLinkableElement
 			);
 	}
 
+	/** This can be used with the Sort() function on a TArray of FAnimNotifyEvents to sort the notifies array by time, earliest first. */
+	ENGINE_API bool operator <(const FAnimNotifyEvent& Other) const;
+
 	ENGINE_API virtual void SetTime(float NewTime, EAnimLinkMethod::Type ReferenceFrame = EAnimLinkMethod::Absolute) override;
 };
 
+// Used by UAnimSequenceBase::SortNotifies() to sort its Notifies array
+FORCEINLINE bool FAnimNotifyEvent::operator<(const FAnimNotifyEvent& Other) const
+{
+	float ATime = GetTriggerTime();
+	float BTime = Other.GetTriggerTime();
+
+#if WITH_EDITORONLY_DATA
+	// this sorting only works if it's saved in editor or loaded with editor data
+	// this is required for gameplay team to have reliable order of notifies
+	// but it was noted that this change will require to resave notifies. 
+	// once you resave, this order will be preserved
+	if (FMath::IsNearlyEqual(ATime, BTime))
+	{
+
+		// if the 2 anim notify events are the same display time sort based off of track index
+		return TrackIndex < Other.TrackIndex;
+	}
+	else
+#endif // WITH_EDITORONLY_DATA
+	{
+		return ATime < BTime;
+	}
+}
+
 /**
  * Keyframe position data for one track.  Pos(i) occurs at Time(i).  Pos.Num() always equals Time.Num().
  */
diff --git a/Engine/Source/Runtime/Engine/Private/Animation/AnimSequenceBase.cpp b/Engine/Source/Runtime/Engine/Private/Animation/AnimSequenceBase.cpp
--- a/Engine/Source/Runtime/Engine/Private/Animation/AnimSequenceBase.cpp
+++ b/Engine/Source/Runtime/Engine/Private/Animation/AnimSequenceBase.cpp
@@ -586,33 +586,8 @@ void UAnimSequenceBase::PostLoad()
 
 void UAnimSequenceBase::SortNotifies()
 {
-	struct FCompareFAnimNotifyEvent
-	{
-		FORCEINLINE bool operator()( const FAnimNotifyEvent& A, const FAnimNotifyEvent& B ) const
-		{
-			float ATime = A.GetTriggerTime();
-			float BTime = B.GetTriggerTime();
-
-#if WITH_EDITORONLY_DATA
-			// this sorting only works if it's saved in editor or loaded with editor data
-			// this is required for gameplay team to have reliable order of notifies
-			// but it was noted that this change will require to resave notifies. 
-			// once you resave, this order will be preserved
-			if(FMath::IsNearlyEqual(ATime, BTime))
-			{
-
-				// if the 2 anim notify events are the same display time sort based off of track index
-				return A.TrackIndex < B.TrackIndex;
-			}
-			else
-#endif // WITH_EDITORONLY_DATA
-			{
-				return ATime < BTime;
-			}
-		}
-	};
-
-	Notifies.Sort( FCompareFAnimNotifyEvent() );
+	// Sorts using FAnimNotifyEvent::operator<()
+	Notifies.Sort();
 }
 
 /** 
