#ANIM: Fixed issue with not opening the bone selection window after search

https://jira.ol.epicgames.net/browse/UE-6255

diff --git a/Engine/Source/Editor/Persona/Private/SRigWindow.cpp b/Engine/Source/Editor/Persona/Private/SRigWindow.cpp
--- a/Engine/Source/Editor/Persona/Private/SRigWindow.cpp
+++ b/Engine/Source/Editor/Persona/Private/SRigWindow.cpp
@@ -300,9 +300,15 @@ void SRigWindow::Construct(const FArguments& InArgs)
 
 void SRigWindow::OnFilterTextChanged( const FText& SearchText )
 {
-	FilterText = SearchText;
+	// need to make sure not to have the same text go
+	// otherwise, the widget gets recreated multiple times causing 
+	// other issue
+	if (FilterText.CompareToCaseIgnored(SearchText) != 0)
+	{
+		FilterText = SearchText;
 
-	CreateBoneMappingList( SearchText.ToString() );
+		CreateBoneMappingList(SearchText.ToString());
+	}
 }
 
 void SRigWindow::OnFilterTextCommitted( const FText& SearchText, ETextCommit::Type CommitInfo )
