Added rename and delete functionality for slot names
#UE-4688

diff --git a/Engine/Source/Editor/Persona/Private/SSkeletonSlotNames.cpp b/Engine/Source/Editor/Persona/Private/SSkeletonSlotNames.cpp
--- a/Engine/Source/Editor/Persona/Private/SSkeletonSlotNames.cpp
+++ b/Engine/Source/Editor/Persona/Private/SSkeletonSlotNames.cpp
@@ -12,6 +12,12 @@
 #include "SNotificationList.h"
 #include "NotificationManager.h"
 #include "STextEntryPopup.h"
+#include "Animation/AnimBlueprint.h"
+#include "AnimGraphNode_Slot.h"
+#include "Animation/AnimNode_Slot.h"
+#include "BlueprintEditorUtils.h"
+#include "SSlotNameReferenceWindow.h"
+#include "IMainFrameModule.h"
 
 #define LOCTEXT_NAMESPACE "SkeletonSlotNames"
 
@@ -249,10 +255,25 @@ TSharedPtr<SWidget> SSkeletonSlotNames::OnGetContextMenuContent() const
 
 	if (bShowGroupItem)
 	{
+		TSharedPtr<FDisplayedSlotNameInfo> SelectedItemPtr = SelectedItems[0];
 
+		MenuBuilder.BeginSection("SlotManagerSlotGroupActions", LOCTEXT("SlotManagerSlotGroupActions", "Slot Group Actions"));
+		// Delete Slot Group
+		{
+			FDisplayedSlotNameInfo* SlotInfo = SelectedItemPtr.Get();
+
+			FUIAction Action = FUIAction(FExecuteAction::CreateSP(this, &SSkeletonSlotNames::OnDeleteSlotGroup, SlotInfo->Name));
+			Action.CanExecuteAction.BindSP(this, &SSkeletonSlotNames::CanDeleteSlotGroup, SlotInfo->Name);
+			const FText Label = LOCTEXT("AnimSlotManagerContextMenuDeleteSlotGroupLabel", "Delete Slot Group");
+			const FText ToolTip = LOCTEXT("AnimSlotManagerContextMenuDeleteSlotGroupTooltip", "Delete this slot group.");
+			MenuBuilder.AddMenuEntry(Label, ToolTip, FSlateIcon(), Action);
+		}
+		MenuBuilder.EndSection();
 	}
 	else if (bShowSlotItem)
 	{
+		TSharedPtr<FDisplayedSlotNameInfo> SelectedItemPtr = SelectedItems[0];
+
 		MenuBuilder.BeginSection("SlotManagerSlotActions", LOCTEXT("SlotManagerSlotActions", "Slot Actions"));
 		// Set Slot Group
 		{
@@ -260,8 +281,26 @@ TSharedPtr<SWidget> SSkeletonSlotNames::OnGetContextMenuContent() const
 				FText::Format(LOCTEXT("ContextMenuSetSlotGroupLabel", "Set Slot {0} Group to"), FText::FromName(SelectedItems[0].Get()->Name)),
 				FText::Format(LOCTEXT("ContextMenuSetSlotGroupToolTip", "Set Slot {0} Group"), FText::FromName(SelectedItems[0].Get()->Name)),
 				FNewMenuDelegate::CreateRaw(this, &SSkeletonSlotNames::FillSetSlotGroupSubMenu));
-			MenuBuilder.EndSection();
 		}
+		// Rename Slot
+		{
+			FDisplayedSlotNameInfo* SlotInfo = SelectedItemPtr.Get();
+
+			FUIAction Action = FUIAction(FExecuteAction::CreateSP(this, &SSkeletonSlotNames::OnRenameSlot, SlotInfo->Name));
+			const FText Label = LOCTEXT("AnimSlotManagerContextMenuRenameSlotLabel", "Rename Slot");
+			const FText ToolTip = LOCTEXT("AnimSlotManagerContextMenuRenameSlotTooltip", "Rename this slot");
+			MenuBuilder.AddMenuEntry(Label, ToolTip, FSlateIcon(), Action);
+		}
+		// Delete Slot
+		{
+			FDisplayedSlotNameInfo* SlotInfo = SelectedItemPtr.Get();
+
+			FUIAction Action = FUIAction(FExecuteAction::CreateSP(this, &SSkeletonSlotNames::OnDeleteSlot, SlotInfo->Name));
+			const FText Label = LOCTEXT("AnimSlotManagerContextMenuDeleteSlotLabel", "Delete Slot");
+			const FText ToolTip = LOCTEXT("AnimSlotManagerContextMenuDeleteSlotTooltip", "Delete this slot.");
+			MenuBuilder.AddMenuEntry(Label, ToolTip, FSlateIcon(), Action);
+		}
+		MenuBuilder.EndSection();
 	}
 
 	MenuBuilder.BeginSection("SlotManagerGeneralActions", LOCTEXT("SlotManagerGeneralActions", "Slot Manager Actions"));
@@ -442,20 +481,28 @@ void SSkeletonSlotNames::AddGroupPopUpOnCommit(const FText & InNewGroupText, ETe
 	}
 }
 
-// void SSkeletonSlotNames::GetCompatibleAnimBlueprints( TMultiMap<class UAnimBlueprint *, class UAnimGraphNode_Slot *>& OutAssets )
-// {
-// 	// second, go through all animgraphnode_slot
-// 	for(FObjectIterator Iter(UAnimGraphNode_Slot::StaticClass()); Iter; ++Iter)
-// 	{
-// 		// if I find one, add the BP and slot nodes to the list
-// 		UAnimGraphNode_Slot * SlotNode = CastChecked<UAnimGraphNode_Slot>(Iter);
-// 		UAnimBlueprint * AnimBlueprint = Cast<UAnimBlueprint>(SlotNode->GetBlueprint());
-// 		if (AnimBlueprint)
-// 		{
-// 			OutAssets.AddUnique(AnimBlueprint, SlotNode);
-// 		}
-// 	}
-// }
+void SSkeletonSlotNames::GetCompatibleAnimBlueprints( TArray<FAssetData>& OutAssets )
+{
+	//Get the skeleton tag to search for
+	FString SkeletonExportName = FAssetData(TargetSkeleton).GetExportTextName();
+
+	// Load the asset registry module
+	FAssetRegistryModule& AssetRegistryModule = FModuleManager::LoadModuleChecked<FAssetRegistryModule>(TEXT("AssetRegistry"));
+
+	TArray<FAssetData> AssetDataList;
+	AssetRegistryModule.Get().GetAssetsByClass(UAnimBlueprint::StaticClass()->GetFName(), AssetDataList, true);
+
+	OutAssets.Empty(AssetDataList.Num());
+
+	for(FAssetData& Data : AssetDataList)
+	{
+		FString AssetSkeleton = Data.TagsAndValues.FindRef("TargetSkeleton");
+		if(AssetSkeleton == SkeletonExportName)
+		{
+			OutAssets.Add(Data);
+		}
+	}
+}
 
 void SSkeletonSlotNames::RefreshSlotNameListWithFilter()
 {
@@ -518,7 +565,7 @@ void SSkeletonSlotNames::ShowNotifyInDetailsView(FName NotifyName)
 	// we can show the list of montage that are used by this slot node?
 }
 
-void SSkeletonSlotNames::GetCompatibleAnimMontage(TArray<class FAssetData>& OutAssets)
+void SSkeletonSlotNames::GetCompatibleAnimMontages(TArray<class FAssetData>& OutAssets)
 {
 	//Get the skeleton tag to search for
 	FString SkeletonExportName = FAssetData(TargetSkeleton).GetExportTextName();
@@ -575,4 +622,452 @@ void SSkeletonSlotNames::NotifyUser( FNotificationInfo& NotificationInfo )
 		Notification->SetCompletionState( SNotificationItem::CS_Fail );
 	}
 }
+
+void SSkeletonSlotNames::OnDeleteSlot(FName SlotName)
+{
+	TArray<FAssetData> CompatibleMontages;
+	TMultiMap<UAnimBlueprint*, UAnimGraphNode_Slot*> CompatibleSlotNodes;
+	GetMontagesAndNodesUsingSlot(SlotName, CompatibleMontages, CompatibleSlotNodes);
+
+	if(CompatibleMontages.Num() > 0 || CompatibleSlotNodes.Num() > 0)
+	{
+		// We can't delete here - still have references. Give the user a chance to fix.
+		if(!ReferenceWindow.IsValid())
+		{
+			// No existing window
+			SAssignNew(ReferenceWindow, SWindow)
+				.AutoCenter(EAutoCenter::PreferredWorkArea)
+				.SizingRule(ESizingRule::Autosized)
+				.Title(LOCTEXT("ReferenceWindowTitle", "Slot References"));
+
+			ReferenceWindow->SetContent
+				(
+				SNew(SBorder)
+				.Padding(FMargin(3))
+				.BorderImage(FEditorStyle::GetBrush("ToolPanel.GroupBorder"))
+				[
+					SAssignNew(ReferenceWidget, SSlotNameReferenceWindow)
+					.ReferencingMontages(&CompatibleMontages)
+					.ReferencingNodes(&CompatibleSlotNodes)
+					.SlotName(SlotName.ToString())
+					.OperationText(LOCTEXT("DeleteOperation", "Delete"))
+					.WidgetWindow(ReferenceWindow)
+					.OnRetry(FSimpleDelegate::CreateSP(this, &SSkeletonSlotNames::RetryDeleteSlot, SlotName))
+				]
+			);
+
+			TSharedPtr<SWindow> ParentWindow;
+			IMainFrameModule& MainFrameModule = FModuleManager::LoadModuleChecked<IMainFrameModule>(TEXT("MainFrame"));
+			ParentWindow = MainFrameModule.GetParentWindow();
+			
+			FSlateApplication::Get().AddWindowAsNativeChild(ReferenceWindow.ToSharedRef(), ParentWindow.ToSharedRef());
+			ReferenceWindow->SetOnWindowClosed(FOnWindowClosed::CreateSP(this, &SSkeletonSlotNames::ReferenceWindowClosed));
+		}
+		else
+		{
+			TSharedPtr<SSlotNameReferenceWindow> RefWidgetPinned = ReferenceWidget.Pin();
+			if(RefWidgetPinned.IsValid())
+			{
+				FReferenceWindowInfo WindowInfo;
+				WindowInfo.ReferencingMontages = &CompatibleMontages;
+				WindowInfo.ReferencingNodes = &CompatibleSlotNodes;
+				WindowInfo.ItemText = FText::FromName(SlotName);
+				WindowInfo.OperationText = LOCTEXT("DeleteOperation", "Delete");
+				WindowInfo.RetryDelegate = FSimpleDelegate::CreateSP(this, &SSkeletonSlotNames::RetryDeleteSlot, SlotName);
+
+				RefWidgetPinned->UpdateInfo(WindowInfo);
+				ReferenceWindow->BringToFront();
+			}
+		}
+	}
+	else
+	{
+		DeleteSlot(SlotName);
+	}
+}
+
+void SSkeletonSlotNames::DeleteSlot(const FName& SlotName)
+{
+	if(TargetSkeleton->ContainsSlotName(SlotName))
+	{
+		TargetSkeleton->Modify();
+		TargetSkeleton->RemoveSlotName(SlotName);
+		RefreshSlotNameListWithFilter();
+	}
+}
+
+void SSkeletonSlotNames::GetAnimMontagesUsingSlot(FName SlotName, TArray<FAssetData>& OutMontages)
+{
+	TArray<FAssetData> SkeletonCompatibleMontages;
+	GetCompatibleAnimMontages(SkeletonCompatibleMontages);
+
+	for(FAssetData& MontageData : SkeletonCompatibleMontages)
+	{
+		UAnimMontage* Montage = Cast<UAnimMontage>(MontageData.GetAsset());
+
+		check(Montage);
+
+		for(FSlotAnimationTrack& SlotTrack : Montage->SlotAnimTracks)
+		{
+			if(SlotTrack.SlotName == SlotName)
+			{
+				OutMontages.Add(MontageData);
+			}
+		}
+	}
+}
+
+void SSkeletonSlotNames::GetAnimMontagesUsingSlotGroup(FName SlotGroupName, TArray<FAssetData>& OutMontages)
+{
+	if(FAnimSlotGroup* Group = TargetSkeleton->FindAnimSlotGroup(SlotGroupName))
+	{
+		for(FName& SlotName : Group->SlotNames)
+		{
+			GetAnimMontagesUsingSlot(SlotName, OutMontages);
+		}
+	}
+}
+
+void SSkeletonSlotNames::ReferenceWindowClosed(const TSharedRef<SWindow>& Window)
+{
+	ReferenceWindow = nullptr;
+}
+
+void SSkeletonSlotNames::RetryDeleteSlot(FName SlotName)
+{
+	TArray<FAssetData> CompatibleMontages;
+	TMultiMap<UAnimBlueprint*, UAnimGraphNode_Slot*> CompatibleSlotNodes;
+	GetMontagesAndNodesUsingSlot(SlotName, CompatibleMontages, CompatibleSlotNodes);
+	
+	if(CompatibleMontages.Num() > 0 || CompatibleSlotNodes.Num() > 0)
+	{
+		// Still can't delete
+		TSharedPtr<SSlotNameReferenceWindow> PinnedWidget = ReferenceWidget.Pin();
+		if(PinnedWidget.IsValid())
+		{
+			FReferenceWindowInfo WindowInfo;
+			WindowInfo.ReferencingMontages = &CompatibleMontages;
+			WindowInfo.ReferencingNodes = &CompatibleSlotNodes;
+			WindowInfo.ItemText = FText::FromName(SlotName);
+			WindowInfo.OperationText = LOCTEXT("DeleteOperation", "Delete");
+			WindowInfo.RetryDelegate = FSimpleDelegate::CreateSP(this, &SSkeletonSlotNames::RetryDeleteSlot, SlotName);
+
+			PinnedWidget->UpdateInfo(WindowInfo);
+			ReferenceWindow->BringToFront();
+		}
+	}
+	else
+	{
+		ReferenceWindow->RequestDestroyWindow();
+		DeleteSlot(SlotName);
+	}
+}
+
+void SSkeletonSlotNames::GetMontagesAndNodesUsingSlot(const FName& SlotName, TArray<FAssetData>& CompatibleMontages, TMultiMap<UAnimBlueprint*, UAnimGraphNode_Slot*> &CompatibleSlotNodes)
+{
+	FScopedSlowTask SlowTask(3, LOCTEXT("AssetReferenceSlowTaskMessage", "Checking for slot references."));
+	SlowTask.MakeDialog();
+
+	TArray<FAssetData> CompatibleBlueprints;
+	SlowTask.EnterProgressFrame(1, LOCTEXT("AssetReferenceTask_Blueprints", "Searching Blueprints"));
+	GetCompatibleAnimBlueprints(CompatibleBlueprints);
+
+	SlowTask.EnterProgressFrame(1, LOCTEXT("AssetReferenceTask_Montages", "Searching Montages"));
+	GetAnimMontagesUsingSlot(SlotName, CompatibleMontages);
+
+	SlowTask.EnterProgressFrame(1, LOCTEXT("AssetReferenceTask_Nodes", "Searching Nodes"));
+	for(FAssetData& Data : CompatibleBlueprints)
+	{
+		TArray<UEdGraph*> BPGraphs;
+		UAnimBlueprint* AnimBP = Cast<UAnimBlueprint>(Data.GetAsset());
+
+		AnimBP->GetAllGraphs(BPGraphs);
+		for(UEdGraph* Graph : BPGraphs)
+		{
+			TArray<UAnimGraphNode_Slot*> SlotNodes;
+			Graph->GetNodesOfClass(SlotNodes);
+
+			for(UAnimGraphNode_Slot* SlotNode : SlotNodes)
+			{
+				if(SlotNode->Node.SlotName == SlotName)
+				{
+					CompatibleSlotNodes.Add(AnimBP, SlotNode);
+				}
+			}
+		}
+	}
+
+	// If we end up loading in any previously unsaved assets they can add names to the list so refresh
+	RefreshSlotNameListWithFilter();
+}
+
+void SSkeletonSlotNames::GetMontagesAndNodesUsingSlotGroup(const FName& SlotGroupName, TArray<FAssetData>& OutMontages, TMultiMap<UAnimBlueprint*, UAnimGraphNode_Slot*> &OutBlueprintSlotMap)
+{
+	if(TargetSkeleton)
+	{
+		if(FAnimSlotGroup* SlotGroup = TargetSkeleton->FindAnimSlotGroup(SlotGroupName))
+		{
+			for(auto& SlotName : SlotGroup->SlotNames)
+			{
+				GetMontagesAndNodesUsingSlot(SlotName, OutMontages, OutBlueprintSlotMap);
+			}
+		}
+	}
+}
+
+void SSkeletonSlotNames::OnRenameSlotPopupCommitted(const FText & InNewSlotText, ETextCommit::Type CommitInfo, FName OldName)
+{
+	if(CommitInfo == ETextCommit::OnEnter)
+	{
+		FName NewName(*InNewSlotText.ToString());
+
+		// Need to dismiss menus early or the slow task in GetMontagesAndNodesUsingSlot will fail to show onscreen
+		FSlateApplication::Get().DismissAllMenus();
+
+		if(TargetSkeleton)
+		{
+			// Make sure the name doesn't already exist
+			if(TargetSkeleton->ContainsSlotName(NewName))
+			{
+				FNotificationInfo Notification(FText::Format(LOCTEXT("ToastRenameFailDesc", "Rename Failed! Slot name {0} already exists in the target skeleton."), FText::FromName(NewName)));
+				Notification.ExpireDuration = 3.0f;
+				Notification.bFireAndForget = true;
+
+				NotifyUser(Notification);
+
+				return;
+			}
+
+			// Validate references
+			TArray<FAssetData> CompatibleMontages;
+			TMultiMap<UAnimBlueprint*, UAnimGraphNode_Slot*> CompatibleSlotNodes;
+			GetMontagesAndNodesUsingSlot(OldName, CompatibleMontages, CompatibleSlotNodes);
+
+			if(CompatibleMontages.Num() > 0 || CompatibleSlotNodes.Num() > 0)
+			{
+				// We can't rename here - still have references. Give the user a chance to fix.
+				if(!ReferenceWindow.IsValid())
+				{
+					// No existing window
+					SAssignNew(ReferenceWindow, SWindow)
+						.AutoCenter(EAutoCenter::PreferredWorkArea)
+						.SizingRule(ESizingRule::Autosized)
+						.Title(LOCTEXT("ReferenceWindowTitle", "Slot References"));
+
+					ReferenceWindow->SetContent
+						(
+						SNew(SBorder)
+						.Padding(FMargin(3))
+						.BorderImage(FEditorStyle::GetBrush("ToolPanel.GroupBorder"))
+						[
+							SAssignNew(ReferenceWidget, SSlotNameReferenceWindow)
+							.ReferencingMontages(&CompatibleMontages)
+							.ReferencingNodes(&CompatibleSlotNodes)
+							.SlotName(OldName.ToString())
+							.OperationText(LOCTEXT("RenameOperation", "Rename"))
+							.WidgetWindow(ReferenceWindow)
+							.OnRetry(FSimpleDelegate::CreateSP(this, &SSkeletonSlotNames::RetryRenameSlot, OldName, NewName))
+						]
+					);
+
+					TSharedPtr<SWindow> ParentWindow;
+					IMainFrameModule& MainFrameModule = FModuleManager::LoadModuleChecked<IMainFrameModule>(TEXT("MainFrame"));
+					ParentWindow = MainFrameModule.GetParentWindow();
+
+					FSlateApplication::Get().AddWindowAsNativeChild(ReferenceWindow.ToSharedRef(), ParentWindow.ToSharedRef());
+					ReferenceWindow->SetOnWindowClosed(FOnWindowClosed::CreateSP(this, &SSkeletonSlotNames::ReferenceWindowClosed));
+				}
+				else
+				{
+					TSharedPtr<SSlotNameReferenceWindow> RefWidgetPinned = ReferenceWidget.Pin();
+					if(RefWidgetPinned.IsValid())
+					{
+						FReferenceWindowInfo WindowInfo;
+						WindowInfo.ReferencingMontages = &CompatibleMontages;
+						WindowInfo.ReferencingNodes = &CompatibleSlotNodes;
+						WindowInfo.ItemText = FText::FromName(OldName);
+						WindowInfo.OperationText = LOCTEXT("RenameOperation", "Rename");
+						WindowInfo.RetryDelegate = FSimpleDelegate::CreateSP(this, &SSkeletonSlotNames::RetryRenameSlot, OldName, NewName);
+
+						RefWidgetPinned->UpdateInfo(WindowInfo);
+						ReferenceWindow->BringToFront();
+					}
+				}
+			}
+			else
+			{
+				RenameSlot(OldName, NewName);
+			}
+		}
+	}
+}
+
+void SSkeletonSlotNames::OnRenameSlot(FName CurrentName)
+{
+	TSharedRef<STextEntryPopup> TextEntry =
+		SNew(STextEntryPopup)
+		.Label(LOCTEXT("RenameSlotName_AskSlotName", "New Slot Name"))
+		.OnTextCommitted(this, &SSkeletonSlotNames::OnRenameSlotPopupCommitted, CurrentName);
+
+	// Show dialog to enter new track name
+	FSlateApplication::Get().PushMenu(
+		SharedThis(this),
+		TextEntry,
+		FSlateApplication::Get().GetCursorPos(),
+		FPopupTransitionEffect(FPopupTransitionEffect::TypeInPopup)
+		);
+
+	TextEntry->FocusDefaultWidget();
+}
+
+void SSkeletonSlotNames::RenameSlot(FName CurrentName, FName NewName)
+{
+	if(TargetSkeleton->ContainsSlotName(CurrentName))
+	{
+		TargetSkeleton->Modify();
+		TargetSkeleton->RenameSlotName(CurrentName, NewName);
+		RefreshSlotNameListWithFilter();
+	}
+}
+
+void SSkeletonSlotNames::RetryRenameSlot(FName CurrentName, FName NewName)
+{
+	TArray<FAssetData> CompatibleMontages;
+	TMultiMap<UAnimBlueprint*, UAnimGraphNode_Slot*> CompatibleSlotNodes;
+	GetMontagesAndNodesUsingSlot(CurrentName, CompatibleMontages, CompatibleSlotNodes);
+
+	if(CompatibleMontages.Num() > 0 || CompatibleSlotNodes.Num() > 0)
+	{
+		// Still can't rename
+		TSharedPtr<SSlotNameReferenceWindow> PinnedWidget = ReferenceWidget.Pin();
+		if(PinnedWidget.IsValid())
+		{
+			FReferenceWindowInfo WindowInfo;
+			WindowInfo.ReferencingMontages = &CompatibleMontages;
+			WindowInfo.ReferencingNodes = &CompatibleSlotNodes;
+			WindowInfo.ItemText = FText::FromName(CurrentName);
+			WindowInfo.OperationText = LOCTEXT("DeleteOperation", "Delete");
+			WindowInfo.RetryDelegate = FSimpleDelegate::CreateSP(this, &SSkeletonSlotNames::RetryRenameSlot, CurrentName, NewName);
+
+			PinnedWidget->UpdateInfo(WindowInfo);
+			ReferenceWindow->BringToFront();
+		}
+	}
+	else
+	{
+		ReferenceWindow->RequestDestroyWindow();
+		RenameSlot(CurrentName, NewName);
+	}
+}
+
+void SSkeletonSlotNames::OnDeleteSlotGroup(FName GroupName)
+{
+	TArray<FAssetData> CompatibleMontages;
+	TMultiMap<UAnimBlueprint*, UAnimGraphNode_Slot*> CompatibleSlotNodes;
+	GetMontagesAndNodesUsingSlotGroup(GroupName, CompatibleMontages, CompatibleSlotNodes);
+
+	if(CompatibleMontages.Num() > 0 || CompatibleSlotNodes.Num() > 0)
+	{
+		// Can't delete, still referenced
+		// We can't rename here - still have references. Give the user a chance to fix.
+		if(!ReferenceWindow.IsValid())
+		{
+			// No existing window
+			SAssignNew(ReferenceWindow, SWindow)
+				.AutoCenter(EAutoCenter::PreferredWorkArea)
+				.SizingRule(ESizingRule::Autosized)
+				.Title(LOCTEXT("ReferenceWindowTitle", "Slot References"));
+
+			ReferenceWindow->SetContent
+				(
+				SNew(SBorder)
+				.Padding(FMargin(3))
+				.BorderImage(FEditorStyle::GetBrush("ToolPanel.GroupBorder"))
+				[
+					SAssignNew(ReferenceWidget, SSlotNameReferenceWindow)
+					.ReferencingMontages(&CompatibleMontages)
+					.ReferencingNodes(&CompatibleSlotNodes)
+					.SlotName(GroupName.ToString())
+					.OperationText(LOCTEXT("DeleteGroupOperation", "Delete Group"))
+					.WidgetWindow(ReferenceWindow)
+					.OnRetry(FSimpleDelegate::CreateSP(this, &SSkeletonSlotNames::RetryDeleteSlotGroup, GroupName))
+				]
+			);
+
+			TSharedPtr<SWindow> ParentWindow;
+			IMainFrameModule& MainFrameModule = FModuleManager::LoadModuleChecked<IMainFrameModule>(TEXT("MainFrame"));
+			ParentWindow = MainFrameModule.GetParentWindow();
+
+			FSlateApplication::Get().AddWindowAsNativeChild(ReferenceWindow.ToSharedRef(), ParentWindow.ToSharedRef());
+			ReferenceWindow->SetOnWindowClosed(FOnWindowClosed::CreateSP(this, &SSkeletonSlotNames::ReferenceWindowClosed));
+		}
+		else
+		{
+			TSharedPtr<SSlotNameReferenceWindow> RefWidgetPinned = ReferenceWidget.Pin();
+			if(RefWidgetPinned.IsValid())
+			{
+				FReferenceWindowInfo WindowInfo;
+				WindowInfo.ReferencingMontages = &CompatibleMontages;
+				WindowInfo.ReferencingNodes = &CompatibleSlotNodes;
+				WindowInfo.ItemText = FText::FromName(GroupName);
+				WindowInfo.OperationText = LOCTEXT("DeleteGroupOperation", "Delete Group");
+				WindowInfo.RetryDelegate = FSimpleDelegate::CreateSP(this, &SSkeletonSlotNames::RetryDeleteSlotGroup, GroupName);
+
+				RefWidgetPinned->UpdateInfo(WindowInfo);
+				ReferenceWindow->BringToFront();
+			}
+		}
+	}
+	else
+	{
+		DeleteSlotGroup(GroupName);
+	}
+}
+
+void SSkeletonSlotNames::RetryDeleteSlotGroup(FName GroupName)
+{
+	TArray<FAssetData> CompatibleMontages;
+	TMultiMap<UAnimBlueprint*, UAnimGraphNode_Slot*> CompatibleSlotNodes;
+	GetMontagesAndNodesUsingSlotGroup(GroupName, CompatibleMontages, CompatibleSlotNodes);
+
+	if(CompatibleMontages.Num() > 0 || CompatibleSlotNodes.Num() > 0)
+	{
+		// Still can't rename
+		TSharedPtr<SSlotNameReferenceWindow> PinnedWidget = ReferenceWidget.Pin();
+		if(PinnedWidget.IsValid())
+		{
+			FReferenceWindowInfo WindowInfo;
+			WindowInfo.ReferencingMontages = &CompatibleMontages;
+			WindowInfo.ReferencingNodes = &CompatibleSlotNodes;
+			WindowInfo.ItemText = FText::FromName(GroupName);
+			WindowInfo.OperationText = LOCTEXT("DeleteGroupOperation", "Delete Group");
+			WindowInfo.RetryDelegate = FSimpleDelegate::CreateSP(this, &SSkeletonSlotNames::RetryDeleteSlotGroup, GroupName);
+
+			PinnedWidget->UpdateInfo(WindowInfo);
+			ReferenceWindow->BringToFront();
+		}
+	}
+	else
+	{
+		ReferenceWindow->RequestDestroyWindow();
+		DeleteSlotGroup(GroupName);
+	}
+}
+
+bool SSkeletonSlotNames::CanDeleteSlotGroup(FName GroupName)
+{
+	static const FName DefaultGroupName("DefaultGroup");
+	return GroupName != DefaultGroupName;
+}
+
+void SSkeletonSlotNames::DeleteSlotGroup(const FName& GroupName)
+{
+	if(TargetSkeleton && TargetSkeleton->FindAnimSlotGroup(GroupName))
+	{
+		TargetSkeleton->Modify();
+		TargetSkeleton->RemoveSlotGroup(GroupName);
+		RefreshSlotNameListWithFilter();
+	}
+}
+
 #undef LOCTEXT_NAMESPACE
diff --git a/Engine/Source/Editor/Persona/Private/SSkeletonSlotNames.h b/Engine/Source/Editor/Persona/Private/SSkeletonSlotNames.h
--- a/Engine/Source/Editor/Persona/Private/SSkeletonSlotNames.h
+++ b/Engine/Source/Editor/Persona/Private/SSkeletonSlotNames.h
@@ -4,6 +4,9 @@
 
 #include "EditorObjectsTracker.h"
 #include "IDocumentation.h"
+#include "SSlotNameReferenceWindow.h"
+
+class UAnimGraphNode_Slot;
 
 #define LOCTEXT_NAMESPACE "SkeletonSlotNames"
 /////////////////////////////////////////////////////
@@ -122,6 +125,32 @@ private:
 	void OnAddGroup();
 	void AddGroupPopUpOnCommit(const FText& InNewGroupName, ETextCommit::Type CommitInfo);
 
+	// Delete a slot after checking slot name references and prompting the user to resolve
+	void OnDeleteSlot(FName SlotName);
+	// Delete a slot immediately without checking references. Call only if sure the name is unreferenced as
+	// it will be re-added on next load if it is.
+	void DeleteSlot(const FName& SlotName);
+	// Retry the validation for deleting a slot
+	void RetryDeleteSlot(FName SlotName);
+
+	// Delete a slot group after checking slot name references and prompting the user to resolve
+	void OnDeleteSlotGroup(FName GroupName);
+	// Delete a slot group immediately without checking references.
+	void DeleteSlotGroup(const FName& GroupName);
+	// Retry the validation for deleting a slot group
+	void RetryDeleteSlotGroup(FName GroupName);
+	// Context menu hook for checking whether delete group is enabled
+	bool CanDeleteSlotGroup(FName GroupName);
+
+	// Attempt to rename a slot after a name has been given - validating the new name and old references first
+	void OnRenameSlotPopupCommitted(const FText& InNewSlotText, ETextCommit::Type CommitInfo, FName OldName);
+	// Spawn popup text box for user to enter name
+	void OnRenameSlot(FName CurrentName);
+	// Rename a slot immediately without checking references
+	void RenameSlot(FName CurrentName, FName NewName);
+	// Retry the validation for renaming a slot
+	void RetryRenameSlot(FName CurrentName, FName NewName);
+
 	// Set Slot Group
 	void FillSetSlotGroupSubMenu(FMenuBuilder& MenuBuilder);
 	void ContextMenuOnSetSlot(FName InNewGroupName);
@@ -138,13 +167,26 @@ private:
 	void ShowNotifyInDetailsView( FName NotifyName );
 
 	/** Populates OutAssets with the AnimSequences that match Personas current skeleton */
-	void GetCompatibleAnimMontage( TArray<class FAssetData>& OutAssets );
+	void GetCompatibleAnimMontages( TArray<class FAssetData>& OutAssets );
+	/** Populates OutAssets with the Anim Blueprints that match Personas current skeleton */
+	void GetCompatibleAnimBlueprints( TArray<FAssetData>& OutAssets );
+
+	// Get all montages that have an anim track using the given slot
+	void GetAnimMontagesUsingSlot(FName SlotName, TArray<FAssetData>& OutMontages);
+	// Get all montages that have an anim track using the given slot group
+	void GetAnimMontagesUsingSlotGroup(FName SlotGroupName, TArray<FAssetData>& OutMontages);
 
-//	void GetCompatibleAnimBlueprints( TMultiMap<class UAnimBlueprint * , class UAnimGraphNode_Slot *>& OutAssets );
+	// Get all montages using the given slot name and a map of all blueprints/nodes also using the slot name
+	void GetMontagesAndNodesUsingSlot(const FName& SlotName, TArray<FAssetData>& CompatibleMontages, TMultiMap<UAnimBlueprint*, UAnimGraphNode_Slot*>& CompatibleSlotNodes);
+	// Get all montages using the given slot group and a map of all blueprints/nodes also using the slot group
+	void GetMontagesAndNodesUsingSlotGroup(const FName& SlotGroupName, TArray<FAssetData>& OutMontages, TMultiMap<UAnimBlueprint*, UAnimGraphNode_Slot*> &OutBlueprintSlotMap);
 
 	/** Utility function to display notifications to the user */
 	void NotifyUser( FNotificationInfo& NotificationInfo );
 
+	// Callback for popup reference window closing
+	void ReferenceWindowClosed(const TSharedRef<SWindow>& Window);
+
 	/** The owning Persona instance */
 	TWeakPtr<class FPersona> PersonaPtr;
 
@@ -165,6 +207,12 @@ private:
 
 	/** Tracks objects created for displaying in the details panel*/
 	FEditorObjectTracker EditorObjectTracker;
+
+	/** Stores the window we spawn to notify the user that references exist on deletion */
+	TSharedPtr<SWindow> ReferenceWindow;
+
+	/** The actual custom widget inside ReferenceWindow */
+	TWeakPtr<SSlotNameReferenceWindow> ReferenceWidget;
 };
 
 #undef LOCTEXT_NAMESPACE
\ No newline at end of file
diff --git a/Engine/Source/Editor/Persona/Private/SSlotNameReferenceWindow.cpp b/Engine/Source/Editor/Persona/Private/SSlotNameReferenceWindow.cpp
new file mode 100644
--- /dev/null
+++ b/Engine/Source/Editor/Persona/Private/SSlotNameReferenceWindow.cpp
@@ -0,0 +1,418 @@
+// Copyright 1998-2015 Epic Games, Inc. All Rights Reserved.
+
+#include "PersonaPrivatePCH.h"
+#include "SSlotNameReferenceWindow.h"
+#include "AnimGraphNode_Slot.h"
+#include "Persona.h"
+
+#define LOCTEXT_NAMESPACE "SkeletonSlotNames"
+
+TSharedRef<FDisplayedMontageReferenceInfo> FDisplayedMontageReferenceInfo::Make(const FAssetData& AssetData)
+{
+	return MakeShareable(new FDisplayedMontageReferenceInfo(AssetData));
+}
+
+TSharedRef<FDisplayedBlueprintReferenceInfo> FDisplayedBlueprintReferenceInfo::Make(UAnimBlueprint* Blueprint, UAnimGraphNode_Slot* SlotNode)
+{
+	return MakeShareable(new FDisplayedBlueprintReferenceInfo(Blueprint, SlotNode));
+}
+
+FDisplayedBlueprintReferenceInfo::FDisplayedBlueprintReferenceInfo(UAnimBlueprint* Blueprint, UAnimGraphNode_Slot* SlotNode)
+{
+	check(Blueprint && SlotNode);
+	BlueprintName = Blueprint->GetName();
+	GraphName = SlotNode->GetGraph()->GetName();
+	NodeName = SlotNode->GetName();
+
+	AnimBlueprint = Blueprint;
+	NodeGraph = SlotNode->GetGraph();
+	Node = SlotNode;
+}
+
+//------------------------------------------------------------------------------------------
+
+void SMontageReferenceListRow::Construct(const FArguments& InArgs, const TSharedRef<STableViewBase>& InOwnerTableView)
+{
+	check(InArgs._ReferenceInfo.IsValid());
+	ReferenceInfo = InArgs._ReferenceInfo;
+
+	SMultiColumnTableRow<TSharedPtr<FDisplayedMontageReferenceInfo>>::Construct(FSuperRowType::FArguments(), InOwnerTableView);
+}
+
+TSharedRef<SWidget> SMontageReferenceListRow::GenerateWidgetForColumn(const FName& ColumnName)
+{
+	if(ColumnName == TEXT("MontageName"))
+	{
+		return SNew(SBox)
+			.HAlign(HAlign_Left)
+			.VAlign(VAlign_Center)
+			[
+				SNew(STextBlock).Text(FText::FromString(ReferenceInfo->AssetData.AssetName.ToString()))
+			];
+	}
+	else if(ColumnName == TEXT("Asset"))
+	{
+		// Buttons to jump to / view asset
+		return SNew(SHorizontalBox)
+			+ SHorizontalBox::Slot()
+			.Padding(0,0,5,0)
+			.AutoWidth()
+			[
+				SNew(SButton)
+				.ToolTipText(LOCTEXT("MontageReferenceViewInContentBrowserToolTip", "Highlight this asset in the Content Browser."))
+				.OnClicked(this, &SMontageReferenceListRow::OnViewInContentBrowserClicked)
+				.DesiredSizeScale(FVector2D(0.75f, 0.75f))
+				.Content()
+				[
+					SNew(SImage)
+					.Image(FEditorStyle::GetBrush("PropertyWindow.Button_Browse"))
+				]
+			]
+			+ SHorizontalBox::Slot()
+			.AutoWidth()
+			[
+				SNew(SButton)
+				.ToolTipText(LOCTEXT("MontageReferenceOpenAssetToolTip", "Open the editor for this asset."))
+				.OnClicked(this, &SMontageReferenceListRow::OnOpenAssetClicked)
+				.DesiredSizeScale(FVector2D(0.75f, 0.75f))
+				.Content()
+				[
+					SNew(SImage)
+					.Image(FEditorStyle::GetBrush("SystemWideCommands.SummonOpenAssetDialog"))
+				]
+			];
+	}
+
+	return SNullWidget::NullWidget;
+}
+
+FReply SMontageReferenceListRow::OnViewInContentBrowserClicked()
+{
+	TArray<FAssetData> Objects;
+	Objects.Add(ReferenceInfo->AssetData);
+	GEditor->SyncBrowserToObjects(Objects);
+
+	return FReply::Handled();
+}
+
+FReply SMontageReferenceListRow::OnOpenAssetClicked()
+{
+	FAssetEditorManager::Get().OpenEditorForAsset(ReferenceInfo->AssetData.GetAsset());
+
+	return FReply::Handled();
+}
+
+//------------------------------------------------------------------------------------------
+
+void SBlueprintReferenceListRow::Construct(const FArguments& InArgs, const TSharedRef<STableViewBase>& InOwnerTableView)
+{
+	check(InArgs._ReferenceInfo.IsValid());
+	ReferenceInfo = InArgs._ReferenceInfo;
+
+	SMultiColumnTableRow<TSharedPtr<FDisplayedBlueprintReferenceInfo>>::Construct(FSuperRowType::FArguments(), InOwnerTableView);
+}
+
+TSharedRef<SWidget> SBlueprintReferenceListRow::GenerateWidgetForColumn(const FName& ColumnName)
+{
+	if(ColumnName == TEXT("BlueprintName"))
+	{
+		return SNew(STextBlock).Text(FText::FromString(ReferenceInfo->BlueprintName));
+	}
+	else if(ColumnName == TEXT("GraphName"))
+	{
+		return SNew(STextBlock).Text(FText::FromString(ReferenceInfo->GraphName));
+	}
+	else if(ColumnName == TEXT("NodeName"))
+	{
+		return SNew(STextBlock).Text(FText::FromString(ReferenceInfo->NodeName));
+	}
+	else if(ColumnName == TEXT("Asset"))
+	{
+		// Buttons to jump to / view asset
+		return SNew(SHorizontalBox)
+			+ SHorizontalBox::Slot()
+			.AutoWidth()
+			.Padding(0,0,5,0)
+			[
+				SNew(SButton)
+				.ToolTipText(LOCTEXT("MontageReferenceViewInContentBrowserToolTip", "Highlight this asset in the Content Browser."))
+				.OnClicked(this, &SBlueprintReferenceListRow::OnViewInContentBrowserClicked)
+				.DesiredSizeScale(FVector2D(0.75f, 0.75f))
+				.Content()
+				[
+					SNew(SImage)
+					.Image(FEditorStyle::GetBrush("PropertyWindow.Button_Browse"))
+				]
+			]
+			+ SHorizontalBox::Slot()
+			.AutoWidth()
+			[
+				SNew(SButton)
+				.ToolTipText(LOCTEXT("MontageReferenceOpenAssetToolTip", "Open the editor for this asset."))
+				.OnClicked(this, &SBlueprintReferenceListRow::OnOpenAssetClicked)
+				.DesiredSizeScale(FVector2D(0.75f, 0.75f))
+				.Content()
+				[
+					SNew(SImage)
+					.Image(FEditorStyle::GetBrush("SystemWideCommands.SummonOpenAssetDialog"))
+				]
+			];
+	}
+
+	return SNullWidget::NullWidget;
+}
+
+FReply SBlueprintReferenceListRow::OnViewInContentBrowserClicked()
+{
+	check(ReferenceInfo->AnimBlueprint);
+	TArray<UObject*> Objects;
+	UObject* BPAsObject = Cast<UObject>(ReferenceInfo->AnimBlueprint);
+	Objects.Add(BPAsObject);
+	GEditor->SyncBrowserToObjects(Objects);
+
+	return FReply::Handled();
+}
+
+FReply SBlueprintReferenceListRow::OnOpenAssetClicked()
+{
+	FAssetEditorManager::Get().OpenEditorForAsset(ReferenceInfo->AnimBlueprint);
+	
+	FPersona* BPEditor = static_cast<FPersona*>(FAssetEditorManager::Get().FindEditorForAsset(ReferenceInfo->AnimBlueprint, true));
+	if(BPEditor)
+	{
+		// Should only ever get an FPersona object back for an animBP.
+		TSharedPtr<SGraphEditor> GraphEditor = BPEditor->OpenGraphAndBringToFront(ReferenceInfo->NodeGraph);
+		if(GraphEditor.IsValid())
+		{
+			// Open the right graph and zoom in on the offending node
+			GraphEditor->JumpToNode(ReferenceInfo->Node, false);
+		}
+	}
+
+	return FReply::Handled();
+}
+
+//------------------------------------------------------------------------------------------
+
+void SSlotNameReferenceWindow::Construct(const FArguments& InArgs)
+{
+	ContainingWindow = InArgs._WidgetWindow;
+	
+	FReferenceWindowInfo WindowUpdateInfo;
+	WindowUpdateInfo.ReferencingMontages = InArgs._ReferencingMontages;
+	WindowUpdateInfo.ReferencingNodes = InArgs._ReferencingNodes;
+	WindowUpdateInfo.ItemText = FText::FromString(InArgs._SlotName);
+	WindowUpdateInfo.OperationText = InArgs._OperationText;
+	WindowUpdateInfo.RetryDelegate = InArgs._OnRetry;
+
+	UpdateInfo(WindowUpdateInfo);
+}
+
+void SSlotNameReferenceWindow::UpdateInfo(FReferenceWindowInfo& UpdatedInfo)
+{
+	// Pull any changed data out of the info struct
+	if(UpdatedInfo.ReferencingMontages)
+	{
+		// Sort for nicer display
+		UpdatedInfo.ReferencingMontages->Sort([](const FAssetData& A, const FAssetData& B)
+		{
+			return A.AssetName < B.AssetName;
+		});
+
+		ReferencingMontages.Empty(UpdatedInfo.ReferencingMontages->Num());
+		for(FAssetData& Data : *UpdatedInfo.ReferencingMontages)
+		{
+			ReferencingMontages.Add(FDisplayedMontageReferenceInfo::Make(Data));
+		}
+	}
+
+	if(UpdatedInfo.ReferencingNodes)
+	{
+		ReferencingNodes.Empty(UpdatedInfo.ReferencingNodes->Num());
+		for(auto& Pair : *UpdatedInfo.ReferencingNodes)
+		{
+			UAnimBlueprint* AnimBP = Pair.Key;
+			UAnimGraphNode_Slot* SlotNode = Pair.Value;
+			if(AnimBP && SlotNode)
+			{
+				ReferencingNodes.Add(FDisplayedBlueprintReferenceInfo::Make(AnimBP, SlotNode));
+
+				ReferencingNodes.Sort([](const TSharedPtr<FDisplayedBlueprintReferenceInfo>& A, const TSharedPtr<FDisplayedBlueprintReferenceInfo>& B)
+				{
+					return A->BlueprintName < B->BlueprintName;
+				});
+			}
+		}
+	}
+
+	if(!UpdatedInfo.OperationText.IsEmpty())
+	{
+		OperationText = UpdatedInfo.OperationText;
+	}
+
+	if(!UpdatedInfo.ItemText.IsEmpty())
+	{
+		SlotName = UpdatedInfo.ItemText.ToString();
+	}
+
+	if(UpdatedInfo.RetryDelegate.IsBound())
+	{
+		OnRetry = UpdatedInfo.RetryDelegate;
+	}
+
+	if(UpdatedInfo.bRefresh)
+	{
+		RefreshContent();
+	}
+}
+
+TSharedRef<SWidget> SSlotNameReferenceWindow::GetContent()
+{
+	TSharedPtr<SWidget> Content;
+
+	SAssignNew(Content, SScrollBox)
+	+SScrollBox::Slot()
+	[
+		SNew(SVerticalBox)
+		+ SVerticalBox::Slot()
+		.Padding(5.0f)
+		.AutoHeight()
+		[
+			SNew(STextBlock)
+			.Text(FText::Format(LOCTEXT("WindowTitle", "{0}: {1}"), OperationText, FText::FromString(SlotName)))
+			.Font(FSlateFontInfo( FPaths::EngineContentDir() / TEXT("Slate/Fonts/Roboto-Bold.ttf"), 12 ))
+		]
+		+SVerticalBox::Slot()
+		.Padding(5.0f)
+		.AutoHeight()
+		[
+			// Montage Explanation paragraph
+			SNew(STextBlock)
+			.AutoWrapText(true)
+			.Text(LOCTEXT("SlotReferenceMontageExplanation", "The following montages contain slots that reference the chosen item. Remove or change the referenced item to proceed."))
+		]
+		+SVerticalBox::Slot()
+		.Padding(5.0f)
+		.AutoHeight()
+		[
+			SNew(SBorder)
+			.BorderImage(FEditorStyle::GetBrush("ToolPanel.DarkGroupBorder"))
+			[
+				// List of montages
+				SNew(SMontageReferenceList)
+				.ItemHeight(20.0f)
+				.ListItemsSource(&ReferencingMontages)
+				.OnGenerateRow(this, &SSlotNameReferenceWindow::HandleGenerateMontageReferenceRow)
+				.SelectionMode(ESelectionMode::None)
+				.HeaderRow
+				(
+				SNew(SHeaderRow)
+				+ SHeaderRow::Column("MontageName")
+				.DefaultLabel(FText(LOCTEXT("MontageColumnDefaultLabel", "Montage Name")))
+				+ SHeaderRow::Column("Asset")
+				.DefaultLabel(FText(LOCTEXT("MontageAssetColumnDefaultLabel", "")))
+				)
+			]
+		]
+		+SVerticalBox::Slot()
+		.Padding(5.0f)
+		.AutoHeight()
+		[
+			// Anim blueprint explanation paragraph
+			SNew(STextBlock)
+			.AutoWrapText(true)
+			.Text(LOCTEXT("SlotReferenceBlueprintExplanation", "The following animation blueprints contain nodes that reference the selected item. Remove or change the nodes listed to proceed"))
+		]
+		+SVerticalBox::Slot()
+		.Padding(5.0f)
+		.AutoHeight()
+		[
+			SNew(SBorder)
+			.BorderImage(FEditorStyle::GetBrush("ToolPanel.DarkGroupBorder"))
+			[
+				// List of montages
+				SNew(SBlueprintReferenceList)
+				.ItemHeight(20.0f)
+				.ListItemsSource(&ReferencingNodes)
+				.OnGenerateRow(this, &SSlotNameReferenceWindow::HandleGenerateBlueprintReferenceRow)
+				.SelectionMode(ESelectionMode::None)
+				.HeaderRow
+				(
+				SNew(SHeaderRow)
+				+ SHeaderRow::Column("BlueprintName")
+				.DefaultLabel(FText(LOCTEXT("BlueprintColumnDefaultLabel", "Blueprint Name")))
+				+ SHeaderRow::Column("GraphName")
+				.DefaultLabel(FText(LOCTEXT("GraphColumnDefaultLabel", "Graph Name")))
+				+ SHeaderRow::Column("NodeName")
+				.DefaultLabel(FText(LOCTEXT("NodeColumnDefaultLabel", "Node Name")))
+				+ SHeaderRow::Column("Asset")
+				.DefaultLabel(FText(LOCTEXT("BPAssetColumnDefaultLabel", "")))
+				)
+			]
+		]
+		+ SVerticalBox::Slot()
+		.VAlign(VAlign_Bottom)
+		.Padding(5.0f)
+		[
+			SNew(SHorizontalBox)
+			+SHorizontalBox::Slot()
+			.HAlign(HAlign_Center)
+			[
+				SNew(SButton)
+				.Text(LOCTEXT("RetryButton", "Retry"))
+				.OnClicked(this, &SSlotNameReferenceWindow::OnRetryClicked)
+			]
+			+ SHorizontalBox::Slot()
+			.HAlign(HAlign_Center)
+			[
+				SNew(SButton)
+				.Text(LOCTEXT("CloseButton", "Close"))
+				.OnClicked(this, &SSlotNameReferenceWindow::OnCloseClicked)
+			]
+		]
+	];
+
+	return Content->AsShared();
+}
+
+void SSlotNameReferenceWindow::RefreshContent()
+{
+	this->ChildSlot
+	[
+		GetContent()
+	];
+}
+
+TSharedRef<ITableRow> SSlotNameReferenceWindow::HandleGenerateMontageReferenceRow(TSharedPtr<FDisplayedMontageReferenceInfo> Item, const TSharedRef<STableViewBase>& OwnerTable)
+{
+	return SNew(SMontageReferenceListRow, OwnerTable)
+			.ReferenceInfo(Item);
+}
+
+TSharedRef<ITableRow> SSlotNameReferenceWindow::HandleGenerateBlueprintReferenceRow(TSharedPtr<FDisplayedBlueprintReferenceInfo> Item, const TSharedRef<STableViewBase>& OwnerTable)
+{
+	return SNew(SBlueprintReferenceListRow, OwnerTable)
+			.ReferenceInfo(Item);
+}
+
+FReply SSlotNameReferenceWindow::OnRetryClicked()
+{
+	if(OnRetry.IsBound())
+	{
+		OnRetry.Execute();
+		return FReply::Handled();
+	}
+	return FReply::Unhandled();
+}
+
+FReply SSlotNameReferenceWindow::OnCloseClicked()
+{
+	if(ContainingWindow.IsValid())
+	{
+		ContainingWindow->RequestDestroyWindow();
+		return FReply::Handled();
+	}
+	return FReply::Unhandled();
+}
+
+#undef LOCTEXT_NAMESPACE
diff --git a/Engine/Source/Editor/Persona/Private/SSlotNameReferenceWindow.h b/Engine/Source/Editor/Persona/Private/SSlotNameReferenceWindow.h
new file mode 100644
--- /dev/null
+++ b/Engine/Source/Editor/Persona/Private/SSlotNameReferenceWindow.h
@@ -0,0 +1,172 @@
+// Copyright 1998-2015 Epic Games, Inc. All Rights Reserved.
+
+#pragma once
+
+#define LOCTEXT_NAMESPACE "SlotRefWindow"
+
+class UAnimBlueprint;
+class UAnimGraphNode_Slot;
+
+// This type can't exist in a SLATE_ARGUMENT macro so typedef away the template
+typedef TMultiMap<UAnimBlueprint*, UAnimGraphNode_Slot*>* NodeMapPtr;
+
+// List display helper for montage references
+class FDisplayedMontageReferenceInfo
+{
+public:
+	FAssetData AssetData;
+
+	static TSharedRef<FDisplayedMontageReferenceInfo> Make(const FAssetData& AssetData);
+
+protected:
+	/** Hidden constructor, always use Make above */
+	FDisplayedMontageReferenceInfo(const FAssetData& InData)
+		: AssetData(InData)
+	{}
+
+	/** Hidden constructor, always use Make above */
+	FDisplayedMontageReferenceInfo()
+	{}
+};
+
+// List display helper for blueprint references
+class FDisplayedBlueprintReferenceInfo
+{
+public:
+	FString BlueprintName;
+	FString GraphName;
+	FString NodeName;
+	UAnimBlueprint* AnimBlueprint;
+	UEdGraph* NodeGraph;
+	UAnimGraphNode_Base* Node;
+
+	static TSharedRef<FDisplayedBlueprintReferenceInfo> Make(UAnimBlueprint* Blueprint, UAnimGraphNode_Slot* SlotNode);
+
+protected:
+	/** Hidden constructor, always use Make above */
+	FDisplayedBlueprintReferenceInfo(UAnimBlueprint* Blueprint, UAnimGraphNode_Slot* SlotNode);
+
+	/** Hidden constructor, always use Make above */
+	FDisplayedBlueprintReferenceInfo()
+	{}
+};
+
+typedef SListView<TSharedPtr<FDisplayedMontageReferenceInfo>> SMontageReferenceList;
+typedef SListView<TSharedPtr<FDisplayedBlueprintReferenceInfo>> SBlueprintReferenceList;
+
+class SMontageReferenceListRow : public SMultiColumnTableRow < TSharedPtr<FDisplayedMontageReferenceInfo> >
+{
+public:
+	SLATE_BEGIN_ARGS(SMontageReferenceListRow)
+	{}
+	SLATE_ARGUMENT(TSharedPtr<FDisplayedMontageReferenceInfo>, ReferenceInfo)
+		SLATE_END_ARGS()
+
+public:
+	void Construct(const FArguments& InArgs, const TSharedRef<STableViewBase>& InOwnerTableView);
+
+	virtual TSharedRef<SWidget> GenerateWidgetForColumn(const FName& ColumnName) override;
+
+protected:
+
+	// Handlers for asset interaction buttons
+	FReply OnViewInContentBrowserClicked();
+	FReply OnOpenAssetClicked();
+
+	TSharedPtr<FDisplayedMontageReferenceInfo> ReferenceInfo;
+};
+
+class SBlueprintReferenceListRow : public SMultiColumnTableRow<TSharedPtr<FDisplayedBlueprintReferenceInfo>>
+{
+public:
+	SLATE_BEGIN_ARGS(SBlueprintReferenceListRow) {}
+		SLATE_ARGUMENT(TSharedPtr<FDisplayedBlueprintReferenceInfo>, ReferenceInfo)
+	SLATE_END_ARGS()
+
+public:
+	void Construct(const FArguments& InArgs, const TSharedRef<STableViewBase>& InOwnerTableView);
+
+	virtual TSharedRef<SWidget> GenerateWidgetForColumn(const FName& ColumnName) override;
+
+protected:
+
+	// Handlers for asset interaction buttons
+	FReply OnViewInContentBrowserClicked();
+	FReply OnOpenAssetClicked();
+
+	TSharedPtr<FDisplayedBlueprintReferenceInfo> ReferenceInfo;
+};
+
+// Required information to update a reference window widget, used in conjunction with SSlotNameReferenceWindow::UpdateInfo
+// any fields that do not get set by the caller will not be updated and will keep its previous value
+struct FReferenceWindowInfo
+{
+	FReferenceWindowInfo()
+		: ReferencingMontages(nullptr)
+		, ReferencingNodes(nullptr)
+		, bRefresh(true)
+	{
+
+	}
+
+	TArray<FAssetData>* ReferencingMontages;
+	NodeMapPtr ReferencingNodes;
+	FText OperationText;
+	FText ItemText;
+	FSimpleDelegate RetryDelegate;
+	bool bRefresh;
+};
+
+// Widget used to slot name references that are blocking an edit operation
+class SSlotNameReferenceWindow : public SCompoundWidget
+{
+public:
+	SLATE_BEGIN_ARGS(SSlotNameReferenceWindow) {}
+		SLATE_ARGUMENT(TArray<FAssetData>*, ReferencingMontages)
+		SLATE_ARGUMENT(NodeMapPtr, ReferencingNodes)
+		SLATE_ARGUMENT(FString, SlotName)
+		SLATE_ARGUMENT(TSharedPtr<SWindow>, WidgetWindow)
+		SLATE_ARGUMENT(FText, OperationText)
+		SLATE_EVENT(FSimpleDelegate, OnRetry)
+	SLATE_END_ARGS()
+
+public:
+	void Construct(const FArguments& InArgs);
+
+	// Refresh the child content of this widget
+	void RefreshContent();
+
+	// Update the reference arrays and optionally rebuild the widget
+	void UpdateInfo(FReferenceWindowInfo& UpdatedInfo);
+
+	// The name of the slot we're referencing
+	FString SlotName;
+
+protected:
+
+	// Rebuild the internal content of the widget
+	TSharedRef<SWidget> GetContent();
+
+	// Row generators for montage and blueprint/node data views
+	TSharedRef<ITableRow> HandleGenerateMontageReferenceRow( TSharedPtr<FDisplayedMontageReferenceInfo> Item, const TSharedRef<STableViewBase>& OwnerTable );
+	TSharedRef<ITableRow> HandleGenerateBlueprintReferenceRow( TSharedPtr<FDisplayedBlueprintReferenceInfo> Item, const TSharedRef<STableViewBase>& OwnerTable );
+
+	// Button handlers
+	FReply OnRetryClicked();
+	FReply OnCloseClicked();
+
+	// Operation type to display
+	FText OperationText;
+
+	// Objects referencing the given slot name
+	TArray<TSharedPtr<FDisplayedMontageReferenceInfo>> ReferencingMontages;
+	TArray<TSharedPtr<FDisplayedBlueprintReferenceInfo>> ReferencingNodes;
+
+	// Ptr to the window this widget resides in
+	TSharedPtr<SWindow> ContainingWindow;
+
+	// Called when the retry button is clicked. User code determines retry behaviour
+	FSimpleDelegate OnRetry;
+};
+
+#undef LOCTEXT_NAMESPACE
\ No newline at end of file
diff --git a/Engine/Source/Runtime/Engine/Classes/Animation/Skeleton.h b/Engine/Source/Runtime/Engine/Classes/Animation/Skeleton.h
--- a/Engine/Source/Runtime/Engine/Classes/Animation/Skeleton.h
+++ b/Engine/Source/Runtime/Engine/Classes/Animation/Skeleton.h
@@ -300,6 +300,13 @@ public:
 	ENGINE_API bool AddSlotGroupName(const FName& InNewGroupName);
 	ENGINE_API FName GetSlotGroupName(const FName& InSlotName) const;
 
+	// Edits/removes slot group data
+	// WARNING: Does not verify that the names aren't used anywhere - if it isn't checked
+	// by the caller the names will be recreated when referencing assets load again.
+	ENGINE_API void RemoveSlotName(const FName& InSlotName);
+	ENGINE_API void RemoveSlotGroup(const FName& InSlotName);
+	ENGINE_API void RenameSlotName(const FName& OldName, const FName& NewName);
+
 #if WITH_EDITORONLY_DATA
 private:
 	/** The default skeletal mesh to use when previewing this skeleton */
diff --git a/Engine/Source/Runtime/Engine/Private/Animation/Skeleton.cpp b/Engine/Source/Runtime/Engine/Private/Animation/Skeleton.cpp
--- a/Engine/Source/Runtime/Engine/Private/Animation/Skeleton.cpp
+++ b/Engine/Source/Runtime/Engine/Private/Animation/Skeleton.cpp
@@ -967,15 +967,10 @@ void USkeleton::BuildSlotToGroupMap()
 
 FAnimSlotGroup* USkeleton::FindAnimSlotGroup(const FName& InGroupName)
 {
-	for (int32 GroupIndex = 0; GroupIndex < SlotGroups.Num(); GroupIndex++)
+	return SlotGroups.FindByPredicate([&InGroupName](const FAnimSlotGroup& Item)
 	{
-		FAnimSlotGroup& SlotGroup = SlotGroups[GroupIndex];
-		if (SlotGroup.GroupName == InGroupName)
-		{
-			return &SlotGroup;
-		}
-	}
-	return NULL;
+		return Item.GroupName == InGroupName;
+	});
 }
 
 bool USkeleton::ContainsSlotName(const FName& InSlotName) 
@@ -1052,6 +1047,43 @@ FName USkeleton::GetSlotGroupName(const FName& InSlotName) const
 	return FAnimSlotGroup::DefaultGroupName;
 }
 
+void USkeleton::RemoveSlotName(const FName& InSlotName)
+{
+	FName GroupName = GetSlotGroupName(InSlotName);
+	
+	if(SlotToGroupNameMap.Remove(InSlotName) > 0)
+	{
+		FAnimSlotGroup* SlotGroup = FindAnimSlotGroup(GroupName);
+		SlotGroup->SlotNames.Remove(InSlotName);
+	}
+}
+
+void USkeleton::RemoveSlotGroup(const FName& InSlotGroupName)
+{
+	FAnimSlotGroup* SlotGroup = FindAnimSlotGroup(InSlotGroupName);
+	// Remove slot mappings
+	for(const FName& SlotName : SlotGroup->SlotNames)
+	{
+		SlotToGroupNameMap.Remove(SlotName);
+	}
+
+	// Remove group
+	SlotGroups.RemoveAll([&InSlotGroupName](const FAnimSlotGroup& Item)
+	{
+		return Item.GroupName == InSlotGroupName;
+	});
+}
+
+void USkeleton::RenameSlotName(const FName& OldName, const FName& NewName)
+{
+	// Can't rename a name that doesn't exist
+	check(ContainsSlotName(OldName))
+
+	FName GroupName = GetSlotGroupName(OldName);
+	RemoveSlotName(OldName);
+	SetSlotGroupName(NewName, GroupName);
+}
+
 void USkeleton::RegenerateGuid()
 {
 	Guid = FGuid::NewGuid();
@@ -1218,5 +1250,6 @@ void USkeleton::GetAssetRegistryTags(TArray<FAssetRegistryTag>& OutTags) const
 
 	OutTags.Add(FAssetRegistryTag(USkeleton::RigTag, RigFullName, FAssetRegistryTag::TT_Hidden));
 }
+
 #endif //WITH_EDITOR
 #undef LOCTEXT_NAMESPACE 
\ No newline at end of file
