Persona: Added a drop-shadow to bone names when rendered in the viewport

diff --git a/Engine/Source/Editor/Persona/Private/AnimationEditorViewportClient.cpp b/Engine/Source/Editor/Persona/Private/AnimationEditorViewportClient.cpp
--- a/Engine/Source/Editor/Persona/Private/AnimationEditorViewportClient.cpp
+++ b/Engine/Source/Editor/Persona/Private/AnimationEditorViewportClient.cpp
@@ -805,6 +805,7 @@ void FAnimationViewportClient::ShowBoneNames( FCanvas* Canvas, FSceneView* View
 				const FName BoneName = PreviewSkelMeshComp->SkeletalMesh->RefSkeleton.GetBoneName(BoneIndex);
 				const FString BoneString = FString::Printf( TEXT("%d: %s"), BoneIndex, *BoneName.ToString() );
 				FCanvasTextItem TextItem( FVector2D( XPos, YPos), FText::FromString( BoneString ), GEngine->GetSmallFont(), BoneColor );
+				TextItem.EnableShadow(FLinearColor::Black);
 				Canvas->DrawItem( TextItem );
 			}
 		}
