Filter sound assets from Persona's Preview attached assets menu

diff --git a/Engine/Source/Editor/Persona/Private/SSkeletonTree.cpp b/Engine/Source/Editor/Persona/Private/SSkeletonTree.cpp
--- a/Engine/Source/Editor/Persona/Private/SSkeletonTree.cpp
+++ b/Engine/Source/Editor/Persona/Private/SSkeletonTree.cpp
@@ -1885,7 +1885,9 @@ void SSkeletonTree::FillAttachAssetSubmenu(FMenuBuilder& MenuBuilder, const FDis
 
 	TArray<UClass*> FilterClasses = FComponentAssetBrokerage::GetSupportedAssets(USceneComponent::StaticClass());
 
-	FilterClasses.RemoveSingleSwap(UBlueprint::StaticClass(), false); //Child actor components broker gives us blueprints which breaks this
+	//Clean up the selection so it is relevant to Persona
+	FilterClasses.RemoveSingleSwap(UBlueprint::StaticClass(), false); //Child actor components broker gives us blueprints which isnt wanted
+	FilterClasses.RemoveSingleSwap(USoundBase::StaticClass(), false); //No sounds wanted
 
 	FAssetPickerConfig AssetPickerConfig;
 	AssetPickerConfig.Filter.bRecursiveClasses = true;
