Comment typo

diff --git a/Engine/Source/Editor/Persona/Private/SSkeletonTree.cpp b/Engine/Source/Editor/Persona/Private/SSkeletonTree.cpp
--- a/Engine/Source/Editor/Persona/Private/SSkeletonTree.cpp
+++ b/Engine/Source/Editor/Persona/Private/SSkeletonTree.cpp
@@ -1886,7 +1886,7 @@ void SSkeletonTree::FillAttachAssetSubmenu(FMenuBuilder& MenuBuilder, const FDis
 	TArray<UClass*> FilterClasses = FComponentAssetBrokerage::GetSupportedAssets(USceneComponent::StaticClass());
 
 	//Clean up the selection so it is relevant to Persona
-	FilterClasses.RemoveSingleSwap(UBlueprint::StaticClass(), false); //Child actor components broker gives us blueprints which isnt wanted
+	FilterClasses.RemoveSingleSwap(UBlueprint::StaticClass(), false); //Child actor components broker gives us blueprints which isn't wanted
 	FilterClasses.RemoveSingleSwap(USoundBase::StaticClass(), false); //No sounds wanted
 
 	FAssetPickerConfig AssetPickerConfig;
