#ANIM: potential fix for blendspace crash

https://jira.ol.epicgames.net/browse/UE-13031

diff --git a/Engine/Source/Editor/Persona/Private/SAnimationBlendSpace.cpp b/Engine/Source/Editor/Persona/Private/SAnimationBlendSpace.cpp
--- a/Engine/Source/Editor/Persona/Private/SAnimationBlendSpace.cpp
+++ b/Engine/Source/Editor/Persona/Private/SAnimationBlendSpace.cpp
@@ -1028,7 +1028,7 @@ int32 SBlendSpaceGridWidget::OnPaint( const FPaintArgs& Args, const FGeometry& A
 	const TArray<FTriangle*>& TriangleList = Generator.GetTriangleList();
 	// now I'd like to get all triangles all 4 points are available
 	TArray<FTriangle*> NewTriangles;
-	if ( TriangleList.Num() > 0 )
+	if ( MousePos.IsSet() && TriangleList.Num() > 0 )
 	{
 		FVector BaryCentricCoords;
 		FTriangle * HighlightTriangle = NULL;
