Fix crash when changing bone selection

diff --git a/Engine/Source/Editor/Persona/Private/Shared/BoneSelectionWidget.cpp b/Engine/Source/Editor/Persona/Private/Shared/BoneSelectionWidget.cpp
--- a/Engine/Source/Editor/Persona/Private/Shared/BoneSelectionWidget.cpp
+++ b/Engine/Source/Editor/Persona/Private/Shared/BoneSelectionWidget.cpp
@@ -182,12 +182,16 @@ void SBoneSelectionWidget::OnSelectionChanged(TSharedPtr<FBoneNameInfo> BoneInfo
 {
 	FilterText = FText::FromString("");
 
-	if(OnBoneSelectionChanged.IsBound())
+	//Because we recreate all our items on tree refresh we will get a spurious null selection event initially.
+	if (BoneInfo.IsValid())
 	{
-		OnBoneSelectionChanged.Execute(BoneInfo->BoneName);
-	}
+		if (OnBoneSelectionChanged.IsBound())
+		{
+			OnBoneSelectionChanged.Execute(BoneInfo->BoneName);
+		}
 
-	BonePickerButton->SetIsOpen(false);
+		BonePickerButton->SetIsOpen(false);
+	}
 }
 
 FText SBoneSelectionWidget::GetCurrentBoneName() const
