Added Montage timing UI, showing the triggering order of notifies, sections and branching points inside a montage.

diff --git a/Engine/Source/Editor/Persona/Private/Persona.h b/Engine/Source/Editor/Persona/Private/Persona.h
--- a/Engine/Source/Editor/Persona/Private/Persona.h
+++ b/Engine/Source/Editor/Persona/Private/Persona.h
@@ -509,6 +509,24 @@ public:
 		OnGenericDelete.RemoveAll(Widget);
 	}
 
+	/** Broadcasts section changes */
+	FSimpleMulticastDelegate OnSectionsChanged;
+
+	// Called when a section is changed
+	typedef FSimpleMulticastDelegate::FDelegate FOnSectionsChanged;
+
+	// Register a delegate to be called when a montage section changes
+	void RegisterOnSectionsChanged(const FOnSectionsChanged& Delegate)
+	{
+		OnSectionsChanged.Add(Delegate);
+	}
+
+	// Unregister a delegate to be called when a montage section changes
+	void UnregisterOnSectionsChanged(SWidget* Widget)
+	{
+		OnSectionsChanged.RemoveAll(Widget);
+	}
+
 	/** Apply Compression to list of animations */
 	void ApplyCompression(TArray<TWeakObjectPtr<UAnimSequence>>& AnimSequences);
 	/** Export to FBX files of the list of animations */
diff --git a/Engine/Source/Editor/Persona/Private/SAnimMontagePanel.cpp b/Engine/Source/Editor/Persona/Private/SAnimMontagePanel.cpp
--- a/Engine/Source/Editor/Persona/Private/SAnimMontagePanel.cpp
+++ b/Engine/Source/Editor/Persona/Private/SAnimMontagePanel.cpp
@@ -11,6 +11,7 @@
 #include "SExpandableArea.h"
 #include "STextEntryPopup.h"
 #include "STextComboBox.h"
+#include "SAnimTimingPanel.h"
 
 #define LOCTEXT_NAMESPACE "AnimMontagePanel"
 
@@ -30,6 +31,7 @@ void SAnimMontagePanel::Construct(const FArguments& InArgs)
 	Persona = InArgs._Persona;
 	Montage = InArgs._Montage;
 	MontageEditor = InArgs._MontageEditor;
+	SectionTimingNodeVisibility = InArgs._SectionTimingNodeVisibility;
 
 	//TrackStyle = TRACK_Double;
 
@@ -53,9 +55,26 @@ void SAnimMontagePanel::Construct(const FArguments& InArgs)
 		]
 	];
 
+	TSharedPtr<FPersona> SharedPersona = Persona.Pin();
+	if(SharedPersona.IsValid())
+	{
+		SharedPersona->RegisterOnChangeAnimNotifies(FPersona::FOnAnimNotifiesChanged::CreateSP(this, &SAnimMontagePanel::Update));
+		SharedPersona->RegisterOnSectionsChanged(FPersona::FOnAnimNotifiesChanged::CreateSP(this, &SAnimMontagePanel::Update));
+	}
+
 	Update();
 }
 
+SAnimMontagePanel::~SAnimMontagePanel()
+{
+	TSharedPtr<FPersona> SharedPersona = Persona.Pin();
+	if(SharedPersona.IsValid())
+	{
+		SharedPersona->UnregisterOnChangeAnimNotifies(this);
+		SharedPersona->UnregisterOnSectionsChanged(this);
+	}
+}
+
 /** This is the main function that creates the UI widgets for the montage tool.*/
 void SAnimMontagePanel::Update()
 {
@@ -127,7 +146,7 @@ void SAnimMontagePanel::Update()
 				.VAlign(VAlign_Center)
 				.Padding( FMargin(0.5f, 0.5f) )
 				[
-					SNew(STrack)
+					SAssignNew(SectionNameTrack, STrack)
 					.ViewInputMin(ViewInputMin)
 					.ViewInputMax(ViewInputMax)
 					.TrackColor( ColourTracker->GetNextColor() )
@@ -142,6 +161,8 @@ void SAnimMontagePanel::Update()
 					.OnTrackRightClickContextMenu( this, &SAnimMontagePanel::SummonTrackContextMenu, static_cast<int>(INDEX_NONE))
 					.ScrubPosition( Editor, &SMontageEditor::GetScrubValue )
 				];
+
+			RefreshTimingNodes();
 		}
 
 		// ===================================
@@ -666,4 +687,37 @@ float SAnimMontagePanel::GetSequenceLength() const
 	return 0.0f;
 }
 
+void SAnimMontagePanel::RefreshTimingNodes()
+{
+	if(SectionNameTrack.IsValid())
+	{
+		// Clear current nodes
+		SectionNameTrack->ClearTrack();
+		// Add section timing widgets
+		TArray<TSharedPtr<FTimingRelevantElementBase>> TimingElements;
+		TArray<TSharedRef<SAnimTimingNode>> SectionTimingNodes;
+		SAnimTimingPanel::GetTimingRelevantElements(Montage, TimingElements);
+		for(int32 Idx = 0 ; Idx < TimingElements.Num() ; ++Idx)
+		{
+			TSharedPtr<FTimingRelevantElementBase>& Element = TimingElements[Idx];
+			if(Element->GetType() == ETimingElementType::Section)
+			{
+				TSharedRef<SAnimTimingTrackNode> Node = SNew(SAnimTimingTrackNode)
+					.ViewInputMin(ViewInputMin)
+					.ViewInputMax(ViewInputMax)
+					.DataStartPos(Element->GetElementTime())
+					.Element(Element)
+					.bUseTooltip(false);
+
+				Node->SetVisibility(SectionTimingNodeVisibility);
+
+				SectionNameTrack->AddTrackNode
+				(
+					Node
+				);
+			}
+		}
+	}
+}
+
 #undef LOCTEXT_NAMESPACE
\ No newline at end of file
diff --git a/Engine/Source/Editor/Persona/Private/SAnimMontagePanel.h b/Engine/Source/Editor/Persona/Private/SAnimMontagePanel.h
--- a/Engine/Source/Editor/Persona/Private/SAnimMontagePanel.h
+++ b/Engine/Source/Editor/Persona/Private/SAnimMontagePanel.h
@@ -50,11 +50,14 @@ public:
 	SLATE_ATTRIBUTE( float, ViewInputMax )
 	SLATE_ATTRIBUTE( float, InputMin )
 	SLATE_ATTRIBUTE( float, InputMax )
+	SLATE_ATTRIBUTE(EVisibility, SectionTimingNodeVisibility)
 	SLATE_EVENT( FOnSetInputViewRange, OnSetInputViewRange )
 	SLATE_EVENT( FOnGetScrubValue, OnGetScrubValue )
 	SLATE_EVENT( FOnAnimObjectChange, OnMontageChange )
 	SLATE_END_ARGS()
 
+	~SAnimMontagePanel();
+
 	void Construct(const FArguments& InArgs);
 	void SetMontage(class UAnimMontage * InMontage);
 
@@ -95,6 +98,9 @@ public:
 	// Context menu callback to set all elements in the montage to a given slot
 	void OnSetElementsToSlot(int32 SlotIndex);
 
+	// Clears the name track of timing nodes and rebuilds them
+	void RefreshTimingNodes();
+
 private:
 	TWeakPtr<FPersona> Persona;
 	TWeakPtr<SMontageEditor>	MontageEditor;
@@ -116,8 +122,11 @@ private:
 
 	TArray< TSharedPtr<SImage> > SlotWarningImages;
 
+	TSharedPtr<STrack> SectionNameTrack;
+	TAttribute<EVisibility> SectionTimingNodeVisibility;
+
 	/************************************************************************/
-	/* Status Bar                                                                     */
+	/* Status Bar                                                           */
 	/************************************************************************/
 	TSharedPtr<STextBlock> StatusBarTextBlock;
 	TSharedPtr<SImage> StatusBarWarningImage;
@@ -139,5 +148,4 @@ private:
 	/** Gets the length of the montage we are currently editing
 	 */
 	virtual float GetSequenceLength() const override;
-
 };
\ No newline at end of file
diff --git a/Engine/Source/Editor/Persona/Private/SAnimMontageSectionsPanel.h b/Engine/Source/Editor/Persona/Private/SAnimMontageSectionsPanel.h
--- a/Engine/Source/Editor/Persona/Private/SAnimMontageSectionsPanel.h
+++ b/Engine/Source/Editor/Persona/Private/SAnimMontageSectionsPanel.h
@@ -8,6 +8,10 @@
 #include "SNodePanel.h"
 #include "STrack.h"
 
+// Forward declatations
+class UAnimMontage;
+class SMontageEditor;
+
 class SAnimMontageSectionsPanel : public SCompoundWidget
 {
 public:
@@ -16,39 +20,64 @@ public:
 		, _MontageEditor()
 	{}
 
-	SLATE_ARGUMENT( class UAnimMontage*, Montage)
-		SLATE_ARGUMENT( TWeakPtr<class SMontageEditor>, MontageEditor)
+		SLATE_ARGUMENT( UAnimMontage*, Montage)						// The montage asset being edited
+		SLATE_ARGUMENT( TWeakPtr<SMontageEditor>, MontageEditor )	// The montage editor that owns this widget
+
 	SLATE_END_ARGS()
 
 	void Construct(const FArguments& InArgs);
+
+	// Rebuild panel widgets
 	void Update();
 
+
 	void SetSectionPos(float NewPos, int32 SectionIdx, int32 RowIdx);
+	// Callback when mouse-drag ends for a section
 	void OnSectionDrop();
 
+	// Callback when a section in the upper display is clicked
+	// @param SectionIndex - Section that was clicked
 	void TopSectionClicked(int32 SectionIndex);
+
+	// Callback when a section in the lower display is clicked
+	// @param SectionIndex - Section that was clicked
 	void SectionClicked(int32 SectionIndex);
+	// Unlinks the requested section
+	// @param SectionIndex - Section to unlink
 	void RemoveLink(int32 SectionIndex);
 
+
+	// Set section to play sequentially (default behavior)
 	FReply MakeDefaultSequence();
+	// Completely remove section sequence data
 	FReply ClearSequence();
 
+	// Starts playing from the first section
 	FReply PreviewAllSectionsClicked();
+	// Plays the clicked section
+	// @param SectionIndex - Section to play
 	FReply PreviewSectionClicked(int32 SectionIndex);
 
 private:
 
-	TSharedPtr<SBorder> PanelArea;
+	// Returns whether or not the provided section is part of a loop
+	// @param SectionIdx - Section to check
+	bool IsLoop(int32 SectionIdx);
 
-	class UAnimMontage*			Montage;
-	TWeakPtr<SMontageEditor>	MontageEditor;
+	// Main panel area widget
+	TSharedPtr<SBorder>			PanelArea;
 
+	// The montage we are currently observing
+	UAnimMontage*				Montage;
+	// The montage editor panel we are a child of
+	TWeakPtr<SMontageEditor>	MontageEditor;
+	// Section to row idx mapping
 	TArray<TArray<int32>>		SectionMap;
 
-	STrackNodeSelectionSet TopSelectionSet;
-	STrackNodeSelectionSet SelectionSet;
+	// Selection sets 
+	STrackNodeSelectionSet		TopSelectionSet;
+	STrackNodeSelectionSet		SelectionSet;
 
-	int32 SelectedCompositeSection;
-
-	bool IsLoop(int32 SectionIdx);
+	// Currently selected section index
+	int32						SelectedCompositeSection;
 };
\ No newline at end of file
diff --git a/Engine/Source/Editor/Persona/Private/SAnimNotifyPanel.cpp b/Engine/Source/Editor/Persona/Private/SAnimNotifyPanel.cpp
--- a/Engine/Source/Editor/Persona/Private/SAnimNotifyPanel.cpp
+++ b/Engine/Source/Editor/Persona/Private/SAnimNotifyPanel.cpp
@@ -16,6 +16,7 @@
 #include "BlueprintActionDatabase.h"
 #include "Animation/AnimNotifies/AnimNotifyState.h"
 #include "Animation/AnimNotifies/AnimNotify.h"
+#include "SAnimTimingPanel.h"
 
 // Track Panel drawing
 const float NotificationTrackHeight = 20.0f;
@@ -35,10 +36,10 @@ DECLARE_DELEGATE_RetVal_FiveParams(FReply, FOnNotifyNodesDragStarted, TArray<TSh
 DECLARE_DELEGATE_RetVal( float, FOnGetDraggedNodePos )
 DECLARE_DELEGATE_TwoParams( FPanTrackRequest, int32, FVector2D)
 DECLARE_DELEGATE_FourParams(FPasteNotifies, SAnimNotifyTrack*, float, ENotifyPasteMode::Type, ENotifyPasteMultipleMode::Type)
+DECLARE_DELEGATE_RetVal_OneParam(EVisibility, FOnGetTimingNodeVisibilityForNode, TSharedPtr<SAnimNotifyNode>)
 
 class FNotifyDragDropOp;
 
-
 //////////////////////////////////////////////////////////////////////////
 // SAnimNotifyNode
 
@@ -62,10 +63,11 @@ public:
 	SLATE_EVENT( FOnNotifyNodeDragStarted, OnNodeDragStarted )
 	SLATE_EVENT( FOnUpdatePanel, OnUpdatePanel )
 	SLATE_EVENT( FPanTrackRequest, PanTrackRequest )
-	SLATE_EVENT( FDeselectAllNotifies, OnDeselectAllNotifies)
+	SLATE_EVENT( FDeselectAllNotifies, OnDeselectAllNotifies )
 	SLATE_ATTRIBUTE( float, ViewInputMin )
 	SLATE_ATTRIBUTE( float, ViewInputMax )
-	SLATE_ATTRIBUTE( TArray<FTrackMarkerBar>, MarkerBars)
+	SLATE_ATTRIBUTE( TArray<FTrackMarkerBar>, MarkerBars )
+	SLATE_ARGUMENT(TSharedPtr<SAnimTimingNode>, StateEndTimingNode)
 	SLATE_END_ARGS()
 
 	void Construct(const FArguments& Declaration);
@@ -95,7 +97,7 @@ public:
 	// End of SNodePanel::SNode
 
 	virtual FVector2D ComputeDesiredSize(float) const override;
-	virtual int32 OnPaint(const FPaintArgs& Args, const FGeometry& AllottedGeometry, const FSlateRect& MyClippingRect, FSlateWindowElementList& OutDrawElements, int32 LayerId, const FWidgetStyle& InWidgetStyle, bool bParentEnabled) const override;
+	virtual int32 OnPaint(const FPaintArgs& Args, const FGeometry& AllottedGeometry, const FSlateRect& MyClippingRect, FSlateWindowElementList& OutDrawElements, int32 LayerId, const FWidgetStyle& InWidgetStyle, bool bParentEnabled) const /*override*/;
 
 	/** Helpers to draw scrub handles and snap offsets */
 	void DrawHandleOffset( const float& Offset, const float& HandleCentre, FSlateWindowElementList& OutDrawElements, int32 MarkerLayer, const FGeometry &AllottedGeometry, const FSlateRect& MyClippingRect ) const;
@@ -193,9 +195,72 @@ private:
 	/** Delegate to deselect notifies and clear the details panel */
 	FDeselectAllNotifies OnDeselectAllNotifies;
 
+	/** Cached owning track geometry */
+	FGeometry CachedTrackGeometry;
+
+	TSharedPtr<SOverlay> EndMarkerNodeOverlay;
+
 	friend class SAnimNotifyTrack;
 };
 
+class SAnimNotifyPair : public SCompoundWidget
+{
+public:
+
+	SLATE_BEGIN_ARGS(SAnimNotifyPair)
+	{}
+
+	SLATE_NAMED_SLOT(FArguments, LeftContent)
+	
+	SLATE_ARGUMENT(TSharedPtr<SAnimNotifyNode>, Node);
+	SLATE_EVENT(FOnGetTimingNodeVisibilityForNode, OnGetTimingNodeVisibilityForNode)
+
+	SLATE_END_ARGS()
+
+		void Construct(const FArguments& InArgs);
+
+	float GetWidgetPaddingLeft();
+
+protected:
+	TSharedPtr<SWidget> PairedWidget;
+	TSharedPtr<SAnimNotifyNode> NodePtr;
+};
+
+void SAnimNotifyPair::Construct(const FArguments& InArgs)
+{
+	NodePtr = InArgs._Node;
+	PairedWidget = InArgs._LeftContent.Widget;
+	check(NodePtr.IsValid());
+	check(PairedWidget.IsValid());
+
+	float ScaleMult = 1.0f;
+	FVector2D NodeSize = NodePtr->ComputeDesiredSize(ScaleMult);
+
+	this->ChildSlot
+		[
+			SNew(SHorizontalBox)
+			+ SHorizontalBox::Slot()
+			.AutoWidth()
+			[
+				SNew(SBox)
+				.HAlign(EHorizontalAlignment::HAlign_Center)
+				.VAlign(EVerticalAlignment::VAlign_Center)
+				[
+					PairedWidget->AsShared()
+				]
+			]
+			+ SHorizontalBox::Slot()
+			[
+				NodePtr->AsShared()
+			]
+		];
+}
+
+float SAnimNotifyPair::GetWidgetPaddingLeft()
+{
+	return NodePtr->GetWidgetPosition().X - PairedWidget->GetDesiredSize().X;
+}
+
 //////////////////////////////////////////////////////////////////////////
 // SAnimNotifyTrack
 
@@ -238,6 +303,8 @@ public:
 		SLATE_ATTRIBUTE( TArray<FTrackMarkerBar>, MarkerBars)
 		SLATE_ARGUMENT( int32, TrackIndex )
 		SLATE_ARGUMENT( FLinearColor, TrackColor )
+		SLATE_ATTRIBUTE(EVisibility, QueuedNotifyTimingNodeVisibility)
+		SLATE_ATTRIBUTE(EVisibility, BranchingPointTimingNodeVisibility)
 		SLATE_EVENT(FOnTrackSelectionChanged, OnSelectionChanged)
 		SLATE_EVENT( FOnUpdatePanel, OnUpdatePanel )
 		SLATE_EVENT( FOnGetBlueprintNotifyData, OnGetNotifyBlueprintData )
@@ -256,7 +323,8 @@ public:
 		SLATE_EVENT( FDeselectAllNotifies, OnDeselectAllNotifies)
 		SLATE_EVENT( FCopyNotifies, OnCopyNotifies )
 		SLATE_EVENT( FPasteNotifies, OnPasteNotifies )
-		SLATE_EVENT(FOnSetInputViewRange, OnSetInputViewRange)
+		SLATE_EVENT( FOnSetInputViewRange, OnSetInputViewRange )
+		SLATE_EVENT( FOnGetTimingNodeVisibility, OnGetTimingNodeVisibility )
 		SLATE_END_ARGS()
 public:
 
@@ -427,6 +495,8 @@ protected:
 	// Handler that is called when the user starts dragging a node
 	FReply OnNotifyNodeDragStarted( TSharedRef<SAnimNotifyNode> NotifyNode, const FPointerEvent& MouseEvent, const FVector2D& ScreenNodePosition, const bool bDragOnMarker, int32 NotifyIndex );
 
+	const EVisibility GetTimingNodeVisibility(TSharedPtr<SAnimNotifyNode> NotifyNode);
+
 private:
 
 	// Data structure for bluprint notify context menu entries
@@ -440,13 +510,13 @@ private:
 	void GetNotifyMenuData(TArray<FAssetData>& NotifyAssetData, TArray<BlueprintNotifyMenuInfo>& OutNotifyMenuData);
 
 	// Store the tracks geometry for later use
-	void UpdateCachedGeometry(const FGeometry& InGeometry) {CachedGeometry = InGeometry;}
+	void UpdateCachedGeometry(const FGeometry& InGeometry);
 
 	// Returns the padding needed to render the notify in the correct track position
 	FMargin GetNotifyTrackPadding(int32 NotifyIndex) const
 	{
-		float LeftMargin = NotifyNodes[NotifyIndex]->GetWidgetPosition().X;
-		float RightMargin = CachedGeometry.Size.X - LeftMargin - NotifyNodes[NotifyIndex]->GetSize().X;
+		float LeftMargin = NotifyPairs[NotifyIndex]->GetWidgetPaddingLeft();
+		float RightMargin = CachedGeometry.Size.X - NotifyNodes[NotifyIndex]->GetWidgetPosition().X - NotifyNodes[NotifyIndex]->GetSize().X;
 		return FMargin(LeftMargin, 0, RightMargin, 0);
 	}
 
@@ -463,11 +533,14 @@ protected:
 
 	class UAnimSequenceBase*				Sequence; // need for menu generation of anim notifies - 
 	TArray<TSharedPtr<SAnimNotifyNode>>		NotifyNodes;
+	TArray<TSharedPtr<SAnimNotifyPair>>		NotifyPairs;
 	TArray<FAnimNotifyEvent*>				AnimNotifies;
 	TAttribute<float>						ViewInputMin;
 	TAttribute<float>						ViewInputMax;
 	TAttribute<FLinearColor>				TrackColor;
 	int32									TrackIndex;
+	TAttribute<EVisibility>					NotifyTimingNodeVisibility;
+	TAttribute<EVisibility>					BranchingPointTimingNodeVisibility;
 	FOnTrackSelectionChanged				OnSelectionChanged;
 	FOnUpdatePanel							OnUpdatePanel;
 	FOnGetBlueprintNotifyData				OnGetNotifyBlueprintData;
@@ -482,6 +555,7 @@ protected:
 	FCopyNotifies							OnCopyNotifies;
 	FPasteNotifies							OnPasteNotifies;
 	FOnSetInputViewRange					OnSetInputViewRange;
+	FOnGetTimingNodeVisibility				OnGetTimingNodeVisibility;
 
 	/** Delegate to call when offsets should be refreshed in a montage */
 	FRefreshOffsetsRequest					OnRequestRefreshOffsets;
@@ -547,6 +621,8 @@ public:
 	SLATE_ATTRIBUTE( float, ViewInputMin )
 	SLATE_ATTRIBUTE( float, ViewInputMax )
 	SLATE_ATTRIBUTE( TArray<FTrackMarkerBar>, MarkerBars)
+	SLATE_ATTRIBUTE( EVisibility, NotifyTimingNodeVisibility )
+	SLATE_ATTRIBUTE( EVisibility, BranchingPointTimingNodeVisibility )
 	SLATE_EVENT( FOnTrackSelectionChanged, OnSelectionChanged)
 	SLATE_EVENT( FOnGetScrubValue, OnGetScrubValue )
 	SLATE_EVENT( FOnGetDraggedNodePos, OnGetDraggedNodePos )
@@ -561,7 +637,8 @@ public:
 	SLATE_EVENT( FDeselectAllNotifies, OnDeselectAllNotifies)
 	SLATE_EVENT( FCopyNotifies, OnCopyNotifies )
 	SLATE_EVENT( FPasteNotifies, OnPasteNotifies )
-	SLATE_EVENT(FOnSetInputViewRange, OnSetInputViewRange)
+	SLATE_EVENT( FOnSetInputViewRange, OnSetInputViewRange )
+	SLATE_EVENT( FOnGetTimingNodeVisibility, OnGetTimingNodeVisibility )
 	SLATE_END_ARGS()
 
 	void Construct(const FArguments& InArgs);
@@ -993,6 +1070,18 @@ void SAnimNotifyNode::Construct(const FArguments& InArgs)
 
 	MarkerBars = InArgs._MarkerBars;
 
+	if(InArgs._StateEndTimingNode.IsValid())
+	{
+		// The overlay will use the desired size to calculate the notify node size,
+		// compute that once here.
+		InArgs._StateEndTimingNode->SlatePrepass(1.0f);
+		SAssignNew(EndMarkerNodeOverlay, SOverlay)
+			+ SOverlay::Slot()
+			[
+				InArgs._StateEndTimingNode.ToSharedRef()
+			];
+	}
+
 	SetToolTipText(TAttribute<FText>(this, &SAnimNotifyNode::GetNodeTooltip));
 }
 
@@ -1181,6 +1270,12 @@ void SAnimNotifyNode::UpdateSizeAndPosition(const FGeometry& AllottedGeometry)
 	WidgetX = bDrawTooltipToRight ? (NotifyTimePositionX - (NotifyHandleBoxWidth / 2.f)) : (NotifyTimePositionX - LabelWidth);
 	WidgetSize = bDrawTooltipToRight ? FVector2D(FMath::Max(LabelWidth, NotifyDurationSizeX), NotifyHeight) : FVector2D((LabelWidth + NotifyDurationSizeX), NotifyHeight);
 	WidgetSize.X += NotifyHandleBoxWidth;
+	
+	if(EndMarkerNodeOverlay.IsValid())
+	{
+		FVector2D OverlaySize = EndMarkerNodeOverlay->ComputeDesiredSize(1.0f);
+		WidgetSize.X += OverlaySize.X;
+	}
 
 	// Widget position of the notify marker
 	NotifyScrubHandleCentre = bDrawTooltipToRight ? NotifyHandleBoxWidth / 2.f : LabelWidth;
@@ -1214,6 +1309,14 @@ int32 SAnimNotifyNode::OnPaint(const FPaintArgs& Args, const FGeometry& Allotted
 	int32 TextLayerID = ScrubHandleID + 1;
 	int32 BranchPointLayerID = TextLayerID + 1;
 
+	// Paint marker node if we have one
+	if(EndMarkerNodeOverlay.IsValid())
+	{
+		FVector2D MarkerSize = EndMarkerNodeOverlay->ComputeDesiredSize(1.0f);
+		FVector2D MarkerOffset(NotifyDurationSizeX + MarkerSize.X * 0.5f + 5.0f, (NotifyHeight - MarkerSize.Y) * 0.5f);
+		EndMarkerNodeOverlay->Paint(Args.WithNewParent(this), AllottedGeometry.MakeChild(MarkerOffset, MarkerSize, 1.0f), MyClippingRect, OutDrawElements, LayerId, InWidgetStyle, bParentEnabled);
+	}
+
 	const FSlateBrush* StyleInfo = FEditorStyle::GetBrush( TEXT("SpecialEditableTextImageNormal") );
 	FSlateDrawElement::MakeBox( 
 		OutDrawElements,
@@ -1368,7 +1471,8 @@ FReply SAnimNotifyNode::OnMouseMove( const FGeometry& MyGeometry, const FPointer
 	
 	FTrackScaleInfo ScaleInfo(ViewInputMin.Get(), ViewInputMax.Get(), 0, 0, CachedAllotedGeometrySize);
 	
-	float TrackScreenSpaceXPosition = MyGeometry.AbsolutePosition.X - MyGeometry.Position.X;
+	float XPositionInTrack = MyGeometry.AbsolutePosition.X - CachedTrackGeometry.AbsolutePosition.X + ScrubHandleSize.X;
+	float TrackScreenSpaceXPosition = MyGeometry.AbsolutePosition.X - XPositionInTrack;
 
 	if(CurrentDragHandle == ENotifyStateHandleHit::Start)
 	{
@@ -1377,7 +1481,7 @@ FReply SAnimNotifyNode::OnMouseMove( const FGeometry& MyGeometry, const FPointer
 
 		if(MouseEvent.GetScreenSpacePosition().X >= TrackScreenSpaceXPosition && MouseEvent.GetScreenSpacePosition().X <= TrackScreenSpaceXPosition + CachedAllotedGeometrySize.X)
 		{
-			float NewDisplayTime = ScaleInfo.LocalXToInput((MouseEvent.GetScreenSpacePosition() - MyGeometry.AbsolutePosition + MyGeometry.Position).X);
+			float NewDisplayTime = ScaleInfo.LocalXToInput((MouseEvent.GetScreenSpacePosition() - MyGeometry.AbsolutePosition + XPositionInTrack).X);
 			float NewDuration = NotifyEvent->GetDuration() + OldDisplayTime - NewDisplayTime;
 
 			// Check to make sure the duration is not less than the minimum allowed
@@ -1434,7 +1538,7 @@ FReply SAnimNotifyNode::OnMouseMove( const FGeometry& MyGeometry, const FPointer
 	{
 		if(MouseEvent.GetScreenSpacePosition().X >= TrackScreenSpaceXPosition && MouseEvent.GetScreenSpacePosition().X <= TrackScreenSpaceXPosition + CachedAllotedGeometrySize.X)
 		{
-			float NewDuration = ScaleInfo.LocalXToInput((MouseEvent.GetScreenSpacePosition() - MyGeometry.AbsolutePosition + MyGeometry.Position).X) - NotifyEvent->GetTime();
+			float NewDuration = ScaleInfo.LocalXToInput((MouseEvent.GetScreenSpacePosition() - MyGeometry.AbsolutePosition + XPositionInTrack).X) - NotifyEvent->GetTime();
 
 			NotifyEvent->SetDuration(FMath::Max(NewDuration, MinimumStateDuration));
 		}
@@ -1663,6 +1767,7 @@ void SAnimNotifyTrack::Construct(const FArguments& InArgs)
 	OnCopyNotifies = InArgs._OnCopyNotifies;
 	OnPasteNotifies = InArgs._OnPasteNotifies;
 	OnSetInputViewRange = InArgs._OnSetInputViewRange;
+	OnGetTimingNodeVisibility = InArgs._OnGetTimingNodeVisibility;
 
 	this->ChildSlot
 	[
@@ -2810,6 +2915,7 @@ void SAnimNotifyTrack::AddNewNotify(const FText& NewNotifyName, ETextCommit::Typ
 
 void SAnimNotifyTrack::Update()
 {
+	NotifyPairs.Empty();
 	NotifyNodes.Empty();
 
 	TrackArea->SetContent(
@@ -2818,17 +2924,74 @@ void SAnimNotifyTrack::Update()
 
 	if ( AnimNotifies.Num() > 0 )
 	{
+		TArray<TSharedPtr<FTimingRelevantElementBase>> TimingElements;
+		SAnimTimingPanel::GetTimingRelevantElements(Sequence, TimingElements);
 		for (int32 NotifyIndex = 0; NotifyIndex < AnimNotifies.Num(); ++NotifyIndex)
 		{
+			TSharedPtr<FTimingRelevantElementBase> Element;
 			FAnimNotifyEvent* Event = AnimNotifies[NotifyIndex];
 
-			TSharedPtr<SAnimNotifyNode> AnimNotifyNode;
-			
-			NodeSlots->AddSlot()
-			.Padding(TAttribute<FMargin>::Create(TAttribute<FMargin>::FGetter::CreateSP(this, &SAnimNotifyTrack::GetNotifyTrackPadding, NotifyIndex)))
-			.VAlign(VAlign_Center)
-			[
-				SAssignNew( AnimNotifyNode, SAnimNotifyNode )
+			for(int32 Idx = 0 ; Idx < TimingElements.Num() ; ++Idx)
+			{
+				Element = TimingElements[Idx];
+
+				if(Element->GetType() == ETimingElementType::NotifyStateBegin
+				   || Element->GetType() == ETimingElementType::BranchPointNotify
+				   || Element->GetType() == ETimingElementType::QueuedNotify)
+				{
+					// Only the notify type will return the type flags above
+					FTimingRelevantElement_Notify* NotifyElement = static_cast<FTimingRelevantElement_Notify*>(Element.Get());
+					if(Event == &Sequence->Notifies[NotifyElement->NotifyIndex])
+					{
+						break;
+					}
+				}
+			}
+
+			TSharedPtr<SAnimNotifyNode> AnimNotifyNode = nullptr;
+			TSharedPtr<SAnimNotifyPair> NotifyPair = nullptr;
+			TSharedPtr<SAnimTimingNode> TimingNode = nullptr;
+			TSharedPtr<SAnimTimingNode> EndTimingNode = nullptr;
+
+			// Create visibility attribute to control timing node visibility for notifies
+			TAttribute<EVisibility> TimingNodeVisibility = TAttribute<EVisibility>::Create(TAttribute<EVisibility>::FGetter::CreateLambda(
+				[this]()
+				{
+					if(OnGetTimingNodeVisibility.IsBound())
+					{
+						return OnGetTimingNodeVisibility.Execute(ETimingElementType::QueuedNotify);
+					}
+					return EVisibility(EVisibility::Hidden);
+				}));
+
+			SAssignNew(TimingNode, SAnimTimingNode)
+				.InElement(Element)
+				.bUseTooltip(false)
+				.Visibility(TimingNodeVisibility);
+
+			if(Event->NotifyStateClass)
+			{
+				TSharedPtr<FTimingRelevantElementBase>* FoundStateEndElement = TimingElements.FindByPredicate([Event](TSharedPtr<FTimingRelevantElementBase>& ElementToTest)
+				{
+					if(ElementToTest.IsValid() && ElementToTest->GetType() == ETimingElementType::NotifyStateEnd)
+					{
+						FTimingRelevantElement_NotifyStateEnd* StateElement = static_cast<FTimingRelevantElement_NotifyStateEnd*>(ElementToTest.Get());
+						return &(StateElement->Sequence->Notifies[StateElement->NotifyIndex]) == Event;
+					}
+					return false;
+				});
+
+				if(FoundStateEndElement)
+				{
+					// Create an end timing node if we have a state
+					SAssignNew(EndTimingNode, SAnimTimingNode)
+						.InElement(*FoundStateEndElement)
+						.bUseTooltip(false)
+						.Visibility(TimingNodeVisibility);
+				}
+			}
+
+			SAssignNew(AnimNotifyNode, SAnimNotifyNode)
 				.Sequence(Sequence)
 				.AnimNotify(Event)
 				.OnNodeDragStarted(this, &SAnimNotifyTrack::OnNotifyNodeDragStarted, NotifyIndex)
@@ -2838,9 +3001,24 @@ void SAnimNotifyTrack::Update()
 				.ViewInputMax(ViewInputMax)
 				.MarkerBars(MarkerBars)
 				.OnDeselectAllNotifies(OnDeselectAllNotifies)
+				.StateEndTimingNode(EndTimingNode);
+
+			SAssignNew(NotifyPair, SAnimNotifyPair)
+			.LeftContent()
+			[
+				TimingNode.ToSharedRef()
+			]
+			.Node(AnimNotifyNode);
+
+			NodeSlots->AddSlot()
+			.Padding(TAttribute<FMargin>::Create(TAttribute<FMargin>::FGetter::CreateSP(this, &SAnimNotifyTrack::GetNotifyTrackPadding, NotifyIndex)))
+			.VAlign(VAlign_Center)
+			[
+				NotifyPair->AsShared()
 			];
 
 			NotifyNodes.Add(AnimNotifyNode);
+			NotifyPairs.Add(NotifyPair);
 		}
 	}
 }
@@ -2982,7 +3160,8 @@ void SAnimNotifyTrack::DisconnectSelectedNodesForDrag(TArray<TSharedPtr<SAnimNot
 	for(auto Iter = SelectedNotifyIndices.CreateIterator(); Iter; ++Iter)
 	{
 		TSharedPtr<SAnimNotifyNode> Node = NotifyNodes[*Iter];
-		NodeSlots->RemoveSlot(Node->AsShared());
+		TSharedPtr<SAnimNotifyPair> Pair = NotifyPairs[*Iter];
+		NodeSlots->RemoveSlot(Pair->AsShared());
 
 		DragNodes.Add(Node);
 	}
@@ -3251,6 +3430,30 @@ void SAnimNotifyTrack::SetNotifyFrame(const FText& NotifyFrameText, ETextCommit:
 	FSlateApplication::Get().DismissAllMenus();
 }
 
+const EVisibility SAnimNotifyTrack::GetTimingNodeVisibility(TSharedPtr<SAnimNotifyNode> NotifyNode)
+{
+	if(OnGetTimingNodeVisibility.IsBound())
+	{
+		FAnimNotifyEvent* Event = NotifyNode->NotifyEvent;
+		EMontageNotifyTickType::Type TickType = Event->MontageTickType;
+
+		return TickType == EMontageNotifyTickType::BranchingPoint ? OnGetTimingNodeVisibility.Execute(ETimingElementType::BranchPointNotify) : OnGetTimingNodeVisibility.Execute(ETimingElementType::QueuedNotify);
+	}
+
+	// No visibility defined, not visible
+	return EVisibility::Hidden;
+}
+
+void SAnimNotifyTrack::UpdateCachedGeometry(const FGeometry& InGeometry)
+{
+	CachedGeometry = InGeometry;
+
+	for(TSharedPtr<SAnimNotifyNode> Node : NotifyNodes)
+	{
+		Node->CachedTrackGeometry = InGeometry;
+	}
+}
+
 //////////////////////////////////////////////////////////////////////////
 // SSequenceEdTrack
 
@@ -3305,6 +3508,7 @@ void SNotifyEdTrack::Construct(const FArguments& InArgs)
 				.OnCopyNotifies(InArgs._OnCopyNotifies)
 				.OnPasteNotifies(InArgs._OnPasteNotifies)
 				.OnSetInputViewRange(InArgs._OnSetInputViewRange)
+				.OnGetTimingNodeVisibility(InArgs._OnGetTimingNodeVisibility)
 			]
 
 			+SHorizontalBox::Slot()
@@ -3396,6 +3600,7 @@ void SAnimNotifyPanel::Construct(const FArguments& InArgs)
 	WidgetWidth = InArgs._WidgetWidth;
 	OnGetScrubValue = InArgs._OnGetScrubValue;
 	OnRequestRefreshOffsets = InArgs._OnRequestRefreshOffsets;
+	OnGetTimingNodeVisibility = InArgs._OnGetTimingNodeVisibility;
 
 	this->ChildSlot
 	[
@@ -3547,11 +3752,11 @@ void SAnimNotifyPanel::DeleteNotify(FAnimNotifyEvent* Notify)
 
 void SAnimNotifyPanel::Update()
 {
-	PersonaPtr.Pin()->OnAnimNotifiesChanged.Broadcast();
 	if(Sequence != NULL)
 	{
 		Sequence->UpdateAnimNotifyTrackCache();
 	}
+	PersonaPtr.Pin()->OnAnimNotifiesChanged.Broadcast();
 }
 
 bool SAnimNotifyPanel::ReadNotifyPasteHeader(FString& OutPropertyString, const TCHAR*& OutBuffer, float& OutOriginalTime, float& OutOriginalLength, int32& OutTrackSpan) const
@@ -3629,6 +3834,7 @@ void SAnimNotifyPanel::RefreshNotifyTracks()
 			.OnCopyNotifies(this, &SAnimNotifyPanel::CopySelectedNotifiesToClipboard)
 			.OnPasteNotifies(this, &SAnimNotifyPanel::OnPasteNotifies)
 			.OnSetInputViewRange(this, &SAnimNotifyPanel::InputViewRangeChanged)
+			.OnGetTimingNodeVisibility(OnGetTimingNodeVisibility)
 		];
 
 		NotifyAnimTracks.Add(EdTrack->NotifyTrack);
@@ -3781,6 +3987,13 @@ void SAnimNotifyPanel::DeselectAllNotifies()
 		Track->DeselectAllNotifyNodes(false);
 	}
 
+	// Broadcast the change so the editor can update
+	TSharedPtr<FPersona> SharedPersona = PersonaPtr.Pin();
+	if(SharedPersona.IsValid())
+	{
+		SharedPersona->OnAnimNotifiesChanged.Broadcast();
+	}
+
 	OnTrackSelectionChanged();
 }
 
@@ -4048,6 +4261,13 @@ void SAnimNotifyPanel::OnPropertyChanged(UObject* ChangedObject, FPropertyChange
 				RefreshNotifyTracks();
 			}
 		}
+
+		// Broadcast the change so the editor can update
+		TSharedPtr<FPersona> SharedPersona = PersonaPtr.Pin();
+		if(SharedPersona.IsValid())
+		{
+			SharedPersona->OnAnimNotifiesChanged.Broadcast();
+		}
 	}
 }
 
@@ -4300,6 +4520,13 @@ void SAnimNotifyPanel::OnNotifyObjectChanged(UObject* EditorBaseObj, bool bRebui
 		{
 			NotifyAnimTracks[WidgetTrackIdx]->Update();
 		}
+
+		// Broadcast the change so the editor can update
+		TSharedPtr<FPersona> SharedPersona = PersonaPtr.Pin();
+		if(SharedPersona.IsValid())
+		{
+			SharedPersona->OnAnimNotifiesChanged.Broadcast();
+		}
 	}
 }
 
diff --git a/Engine/Source/Editor/Persona/Private/SAnimNotifyPanel.h b/Engine/Source/Editor/Persona/Private/SAnimNotifyPanel.h
--- a/Engine/Source/Editor/Persona/Private/SAnimNotifyPanel.h
+++ b/Engine/Source/Editor/Persona/Private/SAnimNotifyPanel.h
@@ -9,6 +9,7 @@
 #include "SCurveEditor.h"
 #include "SAnimTrackPanel.h"
 #include "SAnimEditorBase.h"
+#include "SAnimTimingPanel.h"
 
 DECLARE_DELEGATE_OneParam( FOnSelectionChanged, const FGraphPanelSelectionSet& )
 DECLARE_DELEGATE( FOnTrackSelectionChanged )
@@ -168,7 +169,9 @@ public:
 	SLATE_EVENT( FOnSetInputViewRange, OnSetInputViewRange )
 	SLATE_EVENT( FOnSelectionChanged, OnSelectionChanged )
 	SLATE_EVENT( FOnGetScrubValue, OnGetScrubValue )
-	SLATE_EVENT( FRefreshOffsetsRequest, OnRequestRefreshOffsets)
+	SLATE_EVENT( FRefreshOffsetsRequest, OnRequestRefreshOffsets )
+	SLATE_EVENT( FOnGetTimingNodeVisibility, OnGetTimingNodeVisibility )
+
 	SLATE_END_ARGS()
 
 	void Construct(const FArguments& InArgs);
@@ -229,7 +232,8 @@ private:
 	TAttribute<float> CurrentPosition;
 	FOnSelectionChanged OnSelectionChanged;
 	FOnGetScrubValue OnGetScrubValue;
-	
+	FOnGetTimingNodeVisibility OnGetTimingNodeVisibility;
+
 	/** Manager for mouse controlled marquee selection */
 	FNotifyMarqueeOperation Marquee;
 
diff --git a/Engine/Source/Editor/Persona/Private/SAnimTimingPanel.cpp b/Engine/Source/Editor/Persona/Private/SAnimTimingPanel.cpp
new file mode 100644
--- /dev/null
+++ b/Engine/Source/Editor/Persona/Private/SAnimTimingPanel.cpp
@@ -0,0 +1,559 @@
+// Copyright 1998-2015 Epic Games, Inc. All Rights Reserved.
+
+#include "PersonaPrivatePCH.h"
+#include "SAnimTimingPanel.h"
+#include "SExpandableArea.h"
+#include "STimingTrack.h"
+
+#define LOCTEXT_NAMESPACE "AnimTimingPanel"
+
+namespace AnimTimingConstants
+{
+	static const float DefaultNodeSize = 18.0f;
+	static const int32 FontSize = 10;
+}
+
+void SAnimTimingNode::Construct(const FArguments& InArgs)
+{
+	Element = InArgs._InElement;
+
+	const FSlateBrush* StyleInfo = FEditorStyle::GetBrush("ProgressBar.Background");
+	static FSlateFontInfo LabelFont( FPaths::EngineContentDir() / TEXT("Slate/Fonts/Roboto-Regular.ttf"), AnimTimingConstants::FontSize);
+
+	UPersonaOptions* EditorOptions = UPersonaOptions::StaticClass()->GetDefaultObject<UPersonaOptions>();
+	check(EditorOptions);
+
+	// Pick the colour of the node from the type of the element
+	FLinearColor NodeColour = FLinearColor::White;
+	switch(Element->GetType())
+	{
+		case ETimingElementType::QueuedNotify:
+		case ETimingElementType::NotifyStateBegin:
+		case ETimingElementType::NotifyStateEnd:
+		{
+			NodeColour = EditorOptions->NotifyTimingNodeColour;
+			break;
+		}
+		case ETimingElementType::BranchPointNotify:
+		{
+			NodeColour = EditorOptions->BranchingPointTimingNodeColour;
+			break;
+		}
+		case ETimingElementType::Section:
+		{
+			NodeColour = EditorOptions->SectionTimingNodeColour;
+			break;
+		}
+		default:
+			break;
+	}
+
+	this->ChildSlot
+	[
+		SNew(SBorder)
+		.BorderImage(StyleInfo)
+		.BorderBackgroundColor(NodeColour)
+		[
+			SNew(STextBlock)
+			.Justification(ETextJustify::Center)
+			.Text(FText::AsNumber(Element->TriggerIdx))
+			.Font(LabelFont)
+			.ColorAndOpacity(FSlateColor(FLinearColor::Black))
+		]
+	];
+
+	if(InArgs._bUseTooltip)
+	{
+		// Add asset registry tags to a text list; except skeleton as that is implied in Persona
+		TSharedRef<SVerticalBox> DescriptionBox = SNew(SVerticalBox);
+		TMap<FString, FText> DescriptionItems;
+		Element->GetDescriptionItems(DescriptionItems);
+		for(TPair<FString, FText> ItemPair : DescriptionItems)
+		{
+			DescriptionBox->AddSlot()
+				.AutoHeight()
+				.Padding(0, 0, 3, 0)
+				[
+					SNew(SHorizontalBox)
+					+ SHorizontalBox::Slot()
+					.AutoWidth()
+					[
+						SNew(STextBlock)
+						.Text(FText::Format(LOCTEXT("Item", "{0} :"), FText::FromString(ItemPair.Key)))
+						.ColorAndOpacity(FSlateColor::UseSubduedForeground())
+					]
+
+					+ SHorizontalBox::Slot()
+					.AutoWidth()
+					[
+						SNew(STextBlock)
+						.Text(ItemPair.Value)
+						.ColorAndOpacity(FSlateColor::UseForeground())
+					]
+				];
+		}
+
+		// Tooltip
+		TSharedRef<SToolTip> NodeToolTip = SNew(SToolTip)
+			.TextMargin(1)
+			.BorderImage(FEditorStyle::GetBrush("ContentBrowser.TileViewToolTip.ToolTipBorder"))
+			[
+				SNew(SBorder)
+				.Padding(3)
+				.BorderImage(FEditorStyle::GetBrush("ContentBrowser.TileViewTooltip.NonContentBorder"))
+				[
+					SNew(SVerticalBox)
+					+ SVerticalBox::Slot()
+					.AutoHeight()
+					.Padding(0, 0, 0, 3)
+					[
+						SNew(SBorder)
+						.Padding(6)
+						.BorderImage(FEditorStyle::GetBrush("ContentBrowser.TileViewTooltip.ContentBorder"))
+						[
+							SNew(SBox)
+							.HAlign(HAlign_Center)
+							[
+								SNew(STextBlock)
+								.Text(FText::FromName(Element->GetTypeName()))
+								.Font(FEditorStyle::GetFontStyle("ContentBrowser.TileViewTooltip.NameFont"))
+							]
+						]
+					]
+
+					+ SVerticalBox::Slot()
+					[
+						SNew(SHorizontalBox)
+						+ SHorizontalBox::Slot()
+						.FillWidth(1.0f)
+						[
+							SNew(SBorder)
+							.Padding(3)
+							.BorderImage(FEditorStyle::GetBrush("ContentBrowser.TileViewTooltip.ContentBorder"))
+							[
+								DescriptionBox
+							]
+						]
+					]
+				]
+			];
+		SetToolTip(NodeToolTip);
+	}
+}
+
+FVector2D SAnimTimingNode::ComputeDesiredSize(float) const
+{
+	// Desired height is always the same (a little less than the track height) but the width depends on the text we display
+	const TSharedRef< FSlateFontMeasure > FontMeasureService = FSlateApplication::Get().GetRenderer()->GetFontMeasureService();
+	static FSlateFontInfo LabelFont( FPaths::EngineContentDir() / TEXT("Slate/Fonts/Roboto-Regular.ttf"), 10 );
+	float TextWidth = FontMeasureService->Measure(FString::FromInt(Element->TriggerIdx), LabelFont).X;
+	return FVector2D(FMath::Max(AnimTimingConstants::DefaultNodeSize, TextWidth), AnimTimingConstants::DefaultNodeSize);
+}
+
+void SAnimTimingTrackNode::Construct(const FArguments& InArgs)
+{
+	TAttribute<float> TimeAttr = TAttribute<float>::Create(TAttribute<float>::FGetter::CreateSP(InArgs._Element.ToSharedRef(), &FTimingRelevantElementBase::GetElementTime));
+
+	STrackNode::Construct(STrackNode::FArguments()
+		.ViewInputMin(InArgs._ViewInputMin)
+		.ViewInputMax(InArgs._ViewInputMax)
+		.DataStartPos(TimeAttr)
+		.NodeName(InArgs._NodeName)
+		.CenterOnPosition(true)
+		.AllowDrag(false)
+		.OverrideContent()
+			[
+				SNew(SBox)
+				.HAlign(HAlign_Center)
+				[
+					SNew(SAnimTimingNode)
+					.InElement(InArgs._Element)
+					.bUseTooltip(InArgs._bUseTooltip)
+				]
+			]
+		);
+}
+
+void SAnimTimingPanel::Construct(const FArguments& InArgs)
+{
+	SAnimTrackPanel::Construct(SAnimTrackPanel::FArguments()
+		.WidgetWidth(InArgs._WidgetWidth)
+		.ViewInputMin(InArgs._ViewInputMin)
+		.ViewInputMax(InArgs._ViewInputMax)
+		.InputMin(InArgs._InputMin)
+		.InputMax(InArgs._InputMax)
+		.OnSetInputViewRange(InArgs._OnSetInputViewRange));
+
+	WeakPersona = InArgs._InWeakPersona;
+	AnimSequence = InArgs._InSequence;
+
+	check(WeakPersona.IsValid());
+	check(AnimSequence);
+
+	this->ChildSlot
+		[
+			SNew(SVerticalBox)
+			+ SVerticalBox::Slot()
+			.FillHeight(1)
+			[
+				SNew(SExpandableArea)
+				.AreaTitle(LOCTEXT("ExpandLabel", "Element Timing"))
+				.BodyContent()
+				[
+					SAssignNew(PanelArea, SBorder)
+					.BorderImage(FEditorStyle::GetBrush("NoBorder"))
+					.Padding(FMargin(2.0f, 2.0f))
+					.ColorAndOpacity(FLinearColor::White)
+				]
+			]
+		];
+
+	Update();
+
+	// Register to some delegates to update the interface
+	TSharedPtr<FPersona> SharedPersona = WeakPersona.Pin();
+	if(SharedPersona.IsValid())
+	{
+		SharedPersona->RegisterOnChangeAnimNotifies(FPersona::FOnAnimNotifiesChanged::CreateSP(this, &SAnimTimingPanel::RefreshTrackNodes));
+		SharedPersona->RegisterOnSectionsChanged(FPersona::FOnSectionsChanged::CreateSP(this, &SAnimTimingPanel::RefreshTrackNodes));
+	}
+
+	// Clear display flags
+	FMemory::Memset(bElementNodeDisplayFlags, false, sizeof(bElementNodeDisplayFlags));
+}
+
+SAnimTimingPanel::~SAnimTimingPanel()
+{
+	// Clean up our registered delegates
+	TSharedPtr<FPersona> SharedPersona = WeakPersona.Pin();
+	if(SharedPersona.IsValid())
+	{
+		SharedPersona->UnregisterOnChangeAnimNotifies(this);
+		SharedPersona->UnregisterOnSectionsChanged(this);
+	}
+}
+
+void SAnimTimingPanel::Update()
+{
+	check(PanelArea.IsValid());
+
+	TSharedPtr<SVerticalBox> TimingSlots;
+	
+	PanelArea->SetContent(SAssignNew(TimingSlots, SVerticalBox));
+
+	TSharedRef<S2ColumnWidget> TrackContainer = Create2ColumnWidget(TimingSlots.ToSharedRef());
+	TrackContainer->LeftColumn->ClearChildren();
+	TrackContainer->LeftColumn->AddSlot()
+		.AutoHeight()
+		.VAlign(VAlign_Center)
+		.Padding(FMargin(0.5f, 0.5f))
+		[
+			SAssignNew(Track, STimingTrack)
+			.ViewInputMin(ViewInputMin)
+			.ViewInputMax(ViewInputMax)
+			.TrackMinValue(InputMin)
+			.TrackMaxValue(InputMax)
+			.TrackNumDiscreteValues(AnimSequence->GetNumberOfFrames())
+		];
+
+	TrackContainer->RightColumn->ClearChildren();
+	TrackContainer->RightColumn->AddSlot()
+		.AutoHeight()
+		.VAlign(VAlign_Center)
+		.Padding(FMargin(0.5f, 0.5f))
+		[
+			SNew(SBox)
+			.HAlign(HAlign_Center)
+			[
+				SNew(SButton)
+				.ToolTipText(LOCTEXT("TrackOptionsToolTip", "Display track options menu"))
+				.OnClicked(this, &SAnimTimingPanel::OnContextMenu)
+				.Content()
+				[
+					SNew(SImage)
+					.Image(FEditorStyle::GetBrush("ComboButton.Arrow"))
+					.ColorAndOpacity(FSlateColor::UseForeground())
+				]
+			]
+		];
+
+	RefreshTrackNodes();
+}
+
+void SAnimTimingPanel::Tick(const FGeometry& AllottedGeometry, const double InCurrentTime, const float InDeltaTime)
+{
+	SAnimTrackPanel::Tick(AllottedGeometry, InCurrentTime, InDeltaTime);
+}
+
+void SAnimTimingPanel::RefreshTrackNodes()
+{
+	Elements.Empty();
+	GetTimingRelevantElements(AnimSequence, Elements);
+
+	Track->ClearTrack();
+
+	int32 NumElements = Elements.Num();
+	for(int32 ElementIdx = 0 ; ElementIdx < NumElements ; ++ElementIdx)
+	{
+		TSharedPtr<FTimingRelevantElementBase> Element = Elements[ElementIdx];
+
+		Track->AddTrackNode(
+			SNew(SAnimTimingTrackNode)
+			.ViewInputMin(ViewInputMin)
+			.ViewInputMax(ViewInputMax)
+			.DataStartPos(Element.ToSharedRef(), &FTimingRelevantElementBase::GetElementTime)
+			.NodeName(FString::FromInt(ElementIdx + 1))
+			.NodeColor(FLinearColor::Yellow)
+			.Element(Element)
+			);
+	}
+}
+
+FReply SAnimTimingPanel::OnContextMenu()
+{
+	FMenuBuilder Builder(true, nullptr);
+
+	Builder.BeginSection("TimingPanelOptions", LOCTEXT("TimingPanelOptionsHeader", "Options"));
+	{
+		Builder.AddWidget
+		(
+			SNew(SCheckBox)
+			.IsChecked(this, &SAnimTimingPanel::IsElementDisplayChecked, ETimingElementType::Section)
+			.OnCheckStateChanged(this, &SAnimTimingPanel::OnElementDisplayEnabledChanged, ETimingElementType::Section)
+			.ToolTipText(LOCTEXT("ShowSectionTimingNodes", "Show or hide the timing display for sections on the section name track"))
+			.HAlign(HAlign_Left)
+			[
+				SNew(STextBlock)
+				.Text(LOCTEXT("ToggleTimingNodes_Sections", "Show Section Timing Nodes"))
+			],
+			FText()
+		);
+
+		Builder.AddWidget
+		(
+			SNew(SCheckBox)
+			.IsChecked(this, &SAnimTimingPanel::IsElementDisplayChecked, ETimingElementType::QueuedNotify)
+			.OnCheckStateChanged(this, &SAnimTimingPanel::OnElementDisplayEnabledChanged, ETimingElementType::QueuedNotify)
+			.ToolTipText(LOCTEXT("ShowNotifyTimingNodes", "Show or hide the timing display for notifies in the notify panel"))
+			.HAlign(HAlign_Left)
+			[
+				SNew(STextBlock)
+				.Text(LOCTEXT("ToggleTimingNodes_Notifies", "Show Notify Timing Nodes"))
+			],
+			FText()
+		);
+	}
+	Builder.EndSection();
+
+	FSlateApplication::Get().PushMenu(SharedThis(this),
+		Builder.MakeWidget(),
+		FSlateApplication::Get().GetCursorPos(),
+		FPopupTransitionEffect(FPopupTransitionEffect::ContextMenu));
+
+	return FReply::Handled();
+}
+
+bool SAnimTimingPanel::IsElementDisplayEnabled(ETimingElementType::Type ElementType) const
+{
+	return bElementNodeDisplayFlags[ElementType];
+}
+
+void SAnimTimingPanel::OnElementDisplayEnabledChanged(ECheckBoxState State, ETimingElementType::Type ElementType)
+{
+	bElementNodeDisplayFlags[ElementType] = State == ECheckBoxState::Checked ? true : false;
+}
+
+ECheckBoxState SAnimTimingPanel::IsElementDisplayChecked(ETimingElementType::Type ElementType) const
+{
+	return bElementNodeDisplayFlags[ElementType] ? ECheckBoxState::Checked : ECheckBoxState::Unchecked;
+}
+
+EVisibility SAnimTimingPanel::IsElementDisplayVisible(ETimingElementType::Type ElementType) const
+{
+	return bElementNodeDisplayFlags[ElementType] ? EVisibility::Visible : EVisibility::Hidden;
+}
+
+void SAnimTimingPanel::GetTimingRelevantElements(UAnimSequenceBase* Sequence, TArray<TSharedPtr<FTimingRelevantElementBase>>& Elements)
+{
+	if(Sequence)
+	{
+		// Grab notifies
+		int32 NumNotifies = Sequence->Notifies.Num();
+		for(int32 NotifyIdx = 0 ; NotifyIdx < NumNotifies ; ++NotifyIdx)
+		{
+			FAnimNotifyEvent& Notify = Sequence->Notifies[NotifyIdx];
+
+			FTimingRelevantElement_Notify* Element = new FTimingRelevantElement_Notify;
+			Element->Sequence = Sequence;
+			Element->NotifyIndex = NotifyIdx;
+			Elements.Add(TSharedPtr<FTimingRelevantElementBase>(Element));
+
+			if(Notify.NotifyStateClass)
+			{
+				// Add the end marker
+				FTimingRelevantElement_NotifyStateEnd* EndElement = new FTimingRelevantElement_NotifyStateEnd;
+				EndElement->Sequence = Sequence;
+				EndElement->NotifyIndex = NotifyIdx;
+				Elements.Add(TSharedPtr<FTimingRelevantElementBase>(EndElement));
+			}
+		}
+
+		// Check for a montage and extract Montage elements
+		if(UAnimMontage* Montage = Cast<UAnimMontage>(Sequence))
+		{
+			// Add sections
+			int32 NumSections = Montage->CompositeSections.Num();
+
+			for(int32 SectionIdx = 0 ; SectionIdx < NumSections ; ++SectionIdx)
+			{
+				FCompositeSection& Section = Montage->CompositeSections[SectionIdx];
+
+				FTimingRelevantElement_Section* Element = new FTimingRelevantElement_Section;
+				Element->Montage = Montage;
+				Element->SectionIdx = SectionIdx;
+				Elements.Add(TSharedPtr<FTimingRelevantElementBase>(Element));
+			}
+		}
+
+		// Sort everything and give them trigger orders
+		Elements.Sort([](const TSharedPtr<FTimingRelevantElementBase>& A, const TSharedPtr<FTimingRelevantElementBase>& B)
+		{
+			return A->Compare(*B);
+		});
+
+		int32 NumElements = Elements.Num();
+		for(int32 Idx = 0 ; Idx < NumElements ; ++Idx)
+		{
+			Elements[Idx]->TriggerIdx = Idx;
+		}
+	}
+}
+
+FName FTimingRelevantElement_Section::GetTypeName()
+{
+	return FName("Montage Section");
+}
+
+float FTimingRelevantElement_Section::GetElementTime() const
+{
+	check(Montage)
+	if(Montage->CompositeSections.IsValidIndex(SectionIdx))
+	{
+		return Montage->CompositeSections[SectionIdx].GetTime();
+	}
+	return -1.0f;
+}
+
+ETimingElementType::Type FTimingRelevantElement_Section::GetType()
+{
+	return ETimingElementType::Section;
+}
+
+void FTimingRelevantElement_Section::GetDescriptionItems(TMap<FString, FText>& Items)
+{
+	check(Montage);
+	FCompositeSection& Section = Montage->CompositeSections[SectionIdx];
+
+	FNumberFormattingOptions NumberOptions;
+	NumberOptions.MinimumFractionalDigits = 3;
+
+	Items.Add(LOCTEXT("SectionName", "Name").ToString(), FText::FromName(Section.SectionName));
+	Items.Add(LOCTEXT("SectionTriggerTime", "Trigger Time").ToString(), FText::Format(LOCTEXT("SectionTriggerTime", "{0}s"), FText::AsNumber(Section.GetTime(), &NumberOptions)));
+}
+
+FName FTimingRelevantElement_Notify::GetTypeName()
+{
+	FName TypeName;
+	switch(GetType())
+	{
+		case ETimingElementType::NotifyStateBegin:
+			TypeName = FName("Notify State (Begin)");
+			break;
+		case ETimingElementType::BranchPointNotify:
+			TypeName = FName("Branching Point");
+			break;
+		default:
+			TypeName = FName("Notify");
+	}
+
+	return TypeName;
+}
+
+float FTimingRelevantElement_Notify::GetElementTime() const
+{
+	check(Sequence);
+	if(Sequence->Notifies.IsValidIndex(NotifyIndex))
+	{
+		return Sequence->Notifies[NotifyIndex].GetTriggerTime();
+	}
+	return -1.0f;
+}
+
+ETimingElementType::Type FTimingRelevantElement_Notify::GetType()
+{
+	check(Sequence);
+	FAnimNotifyEvent& Event = Sequence->Notifies[NotifyIndex];
+
+	if(Event.IsBranchingPoint())
+	{
+		return ETimingElementType::BranchPointNotify;
+	}
+	else
+	{
+		if(Event.NotifyStateClass)
+		{
+			return ETimingElementType::NotifyStateBegin;
+		}
+		else
+		{
+			return ETimingElementType::QueuedNotify;
+		}
+	}
+}
+
+void FTimingRelevantElement_Notify::GetDescriptionItems(TMap<FString, FText>& Items)
+{
+	check(Sequence);
+	FAnimNotifyEvent& Event = Sequence->Notifies[NotifyIndex];
+
+	FNumberFormattingOptions NumberOptions;
+	NumberOptions.MinimumFractionalDigits = 3;
+
+	Items.Add(LOCTEXT("NotifyName", "Name").ToString(), FText::FromName(Event.NotifyName));
+	Items.Add(LOCTEXT("NotifyTriggerTime", "Trigger Time").ToString(), FText::Format(LOCTEXT("NotifyTriggerTime_Val", "{0}s"), FText::AsNumber(Event.GetTime(), &NumberOptions)));
+
+	// +1 as we start at 1 when showing tracks to the user
+	Items.Add(LOCTEXT("TrackIdx", "Track").ToString(), FText::AsNumber(Event.TrackIndex + 1));
+
+	if(Event.NotifyStateClass)
+	{
+		Items.Add(LOCTEXT("NotifyDuration", "Duration").ToString(), FText::Format(LOCTEXT("NotifyDuration_Val", "{0}s"), FText::AsNumber(Event.GetDuration(), &NumberOptions)));
+	}
+}
+
+int32 FTimingRelevantElement_Notify::GetElementSortPriority() const
+{
+	check(Sequence);
+	FAnimNotifyEvent& Event = Sequence->Notifies[NotifyIndex];
+	return Event.TrackIndex;
+}
+
+FName FTimingRelevantElement_NotifyStateEnd::GetTypeName()
+{
+	return FName("Notify State (End)");
+}
+
+float FTimingRelevantElement_NotifyStateEnd::GetElementTime() const
+{
+	check(Sequence);
+	FAnimNotifyEvent& Event = Sequence->Notifies[NotifyIndex];
+	check(Event.NotifyStateClass);
+	return Event.GetEndTriggerTime();
+}
+
+ETimingElementType::Type FTimingRelevantElement_NotifyStateEnd::GetType()
+{
+	return ETimingElementType::NotifyStateEnd;
+}
+
+#undef LOCTEXT_NAMESPACE
\ No newline at end of file
diff --git a/Engine/Source/Editor/Persona/Private/SAnimTimingPanel.h b/Engine/Source/Editor/Persona/Private/SAnimTimingPanel.h
new file mode 100644
--- /dev/null
+++ b/Engine/Source/Editor/Persona/Private/SAnimTimingPanel.h
@@ -0,0 +1,241 @@
+// Copyright 1998-2015 Epic Games, Inc. All Rights Reserved.
+
+#pragma once
+
+#include "Persona.h"
+#include "SAnimTrackPanel.h"
+#include "SWidget.h"
+#include "STrack.h"
+
+// Forward declarations
+class UAnimSequenceBase;
+
+//////////////////////////////////////////////////////////////////////////
+// FTimingRelevantElement - data object holding timing data
+
+namespace ETimingElementType
+{
+	enum Type
+	{
+		QueuedNotify,
+		BranchPointNotify,
+		NotifyStateBegin,
+		NotifyStateEnd,
+		Section,
+		Max,
+	};
+};
+
+struct FTimingRelevantElementBase
+{
+	virtual ~FTimingRelevantElementBase()
+	{
+
+	}
+
+	virtual FName GetTypeName()
+	{
+		return FName(TEXT("BASE"));
+	}
+
+	virtual float GetElementTime() const
+	{
+		return -1.0f;
+	}
+
+	virtual int32 GetElementSortPriority() const
+	{
+		return 0;
+	}
+
+	virtual ETimingElementType::Type GetType()
+	{
+		return ETimingElementType::Max;
+	}
+
+	// Get a list of descriptions key/values to describe the element.
+	// Intended for UI/Tooltip use
+	virtual void GetDescriptionItems(TMap<FString, FText>& Items)
+	{
+
+	}
+
+	// Comparison for sorting lists of elements
+	virtual bool Compare(const FTimingRelevantElementBase& Other)
+	{
+		if(FMath::IsNearlyEqual(GetElementTime(), Other.GetElementTime(), SMALL_NUMBER))
+		{
+			return GetElementSortPriority() < Other.GetElementSortPriority();
+		}
+
+		return GetElementTime() < Other.GetElementTime();
+	}
+
+	// Where in the order for the sequence this element will trigger
+	int32 TriggerIdx;
+};
+
+// Small helper to store information about timing relevant elements (notifies, branch points, sections etc.)
+struct FTimingRelevantElement_Notify : public FTimingRelevantElementBase
+{
+	UAnimSequenceBase* Sequence;	// The sequence the notify exists within
+	int32 NotifyIndex;				// The index of the notify in the sequence
+
+	virtual FName GetTypeName() override;
+	virtual float GetElementTime() const override;
+	virtual int32 GetElementSortPriority() const override;
+	virtual ETimingElementType::Type GetType() override;
+	virtual void GetDescriptionItems(TMap<FString, FText>& Items) override;
+};
+
+// Small helper to store information about timing relevant elements (notifies, branch points, sections etc.)
+struct FTimingRelevantElement_NotifyStateEnd : public FTimingRelevantElement_Notify
+{
+	virtual FName GetTypeName() override;
+	virtual float GetElementTime() const override;
+	virtual ETimingElementType::Type GetType() override;
+};
+
+struct FTimingRelevantElement_Section : public FTimingRelevantElementBase
+{
+	UAnimMontage* Montage;	// The montage the section exists within
+	int32 SectionIdx;		// The index of the section in the montage
+
+	virtual FName GetTypeName() override;
+	virtual float GetElementTime() const override;
+	virtual ETimingElementType::Type GetType() override;
+	virtual void GetDescriptionItems(TMap<FString, FText>& Items) override;
+};
+
+// Delegate to get the visibility of a type of timing node on an external panel (not the timing track)
+DECLARE_DELEGATE_RetVal_OneParam(EVisibility, FOnGetTimingNodeVisibility, ETimingElementType::Type)
+
+//////////////////////////////////////////////////////////////////////////
+// The content of SAnimTimingTrackNode, separated to be used in non STrack widgets
+//////////////////////////////////////////////////////////////////////////
+class SAnimTimingNode : public SCompoundWidget
+{
+public:
+	SLATE_BEGIN_ARGS(SAnimTimingNode)
+		: _InElement()
+	{}
+
+	SLATE_ARGUMENT(TSharedPtr<FTimingRelevantElementBase>, InElement)
+	SLATE_ARGUMENT(bool, bUseTooltip)
+
+	SLATE_END_ARGS()
+
+	void Construct(const FArguments& InArgs);
+	virtual FVector2D ComputeDesiredSize(float) const override;
+
+protected:
+
+	// The observed element
+	TSharedPtr<FTimingRelevantElementBase> Element;
+};
+
+//////////////////////////////////////////////////////////////////////////
+// Track node containing an identifier for trigger order of a timing element
+//////////////////////////////////////////////////////////////////////////
+class SAnimTimingTrackNode : public STrackNode
+{
+public:
+	SLATE_BEGIN_ARGS(SAnimTimingTrackNode)
+	{}
+
+	SLATE_ATTRIBUTE(float, ViewInputMin)
+	SLATE_ATTRIBUTE(float, ViewInputMax)
+	SLATE_ATTRIBUTE(float, DataStartPos)
+	SLATE_ATTRIBUTE(FString, NodeName)
+	SLATE_ATTRIBUTE(FLinearColor, NodeColor)
+	SLATE_ARGUMENT(TSharedPtr<FTimingRelevantElementBase>, Element)
+	SLATE_ARGUMENT(bool, bUseTooltip)
+
+	SLATE_END_ARGS()
+
+	void Construct(const FArguments& InArgs);
+
+};
+
+//////////////////////////////////////////////////////////////////////////
+// Panel containing a track for timing nodes, to display the order
+// that events in a montage occur (Notifies, sections, branching pts etc.)
+//////////////////////////////////////////////////////////////////////////
+class SAnimTimingPanel : public SAnimTrackPanel
+{
+public:
+
+	SLATE_BEGIN_ARGS(SAnimTimingPanel)
+		: _InWeakPersona()
+		, _InSequence(nullptr)
+		, _WidgetWidth(0)
+		, _CurrentPosition(0)
+		, _ViewInputMin(0)
+		, _ViewInputMax(0)
+		, _InputMin(0)
+		, _InputMax(0)
+		, _OnSetInputViewRange()
+	{}
+
+	SLATE_ARGUMENT(TWeakPtr<FPersona>, InWeakPersona)
+	SLATE_ARGUMENT(UAnimSequenceBase*, InSequence)
+	SLATE_ARGUMENT(float, WidgetWidth)
+	SLATE_ATTRIBUTE(float, CurrentPosition)
+	SLATE_ATTRIBUTE(float, ViewInputMin)
+	SLATE_ATTRIBUTE(float, ViewInputMax)
+	SLATE_ATTRIBUTE(float, InputMin)
+	SLATE_ATTRIBUTE(float, InputMax)
+	SLATE_EVENT(FOnSetInputViewRange, OnSetInputViewRange)
+
+	SLATE_END_ARGS()
+
+	virtual ~SAnimTimingPanel();
+
+	// Construct the panel
+	// @param InArgs - Slate arguments
+	void Construct(const FArguments& InArgs);
+
+	// Updates panel widgets
+	void Update();
+
+
+	// Access to display enabled flags in multiple forms
+	bool IsElementDisplayEnabled(ETimingElementType::Type ElementType) const;
+	ECheckBoxState IsElementDisplayChecked(ETimingElementType::Type ElementType) const;
+	EVisibility IsElementDisplayVisible(ETimingElementType::Type ElementType) const;
+
+	// Callback from slate when element display flag changes
+	// @param State - New state
+	// @param ElementType - Flag to set
+	void OnElementDisplayEnabledChanged(ECheckBoxState State, ETimingElementType::Type ElementType);
+
+	// Inspects the provided sequence, collect and sorts the requested elements
+	// @param Sequence - Sequence to inspect
+	// @param Elements - Out array of elements
+	static void GetTimingRelevantElements(UAnimSequenceBase* Sequence, TArray<TSharedPtr<FTimingRelevantElementBase>>& Elements);
+
+protected:
+
+	// Context summon callback
+	FReply OnContextMenu();
+
+	// Tick the panel state
+	virtual void Tick(const FGeometry& AllottedGeometry, const double InCurrentTime, const float InDeltaTime) override;
+
+	// Clears the timing track and rebuilds the nodes
+	void RefreshTrackNodes();
+
+	// Observed timing elements
+	TArray<TSharedPtr<FTimingRelevantElementBase>> Elements;
+	// Toolkit reference
+	TWeakPtr<FPersona> WeakPersona;
+	// Anim sequence that contains the timing elements we are observing
+	UAnimSequenceBase* AnimSequence;
+	// Main panel widget
+	TSharedPtr<SBorder> PanelArea;
+	// The track to place timing nodes on
+	TSharedPtr<STrack> Track;
+
+	// Display flags for other panels
+	bool bElementNodeDisplayFlags[ETimingElementType::Max];
+};
\ No newline at end of file
diff --git a/Engine/Source/Editor/Persona/Private/SMontageEditor.cpp b/Engine/Source/Editor/Persona/Private/SMontageEditor.cpp
--- a/Engine/Source/Editor/Persona/Private/SMontageEditor.cpp
+++ b/Engine/Source/Editor/Persona/Private/SMontageEditor.cpp
@@ -14,6 +14,7 @@
 #include "SAnimNotifyPanel.h"
 #include "SAnimCurvePanel.h"
 #include "AnimPreviewInstance.h"
+#include "SAnimTimingPanel.h"
 
 #define LOCTEXT_NAMESPACE "AnimSequenceEditor"
 
@@ -30,6 +31,8 @@ void SMontageEditor::Construct(const FArguments& InArgs)
 	MontageObj = InArgs._Montage;
 	check(MontageObj);
 
+	WeakPersona = InArgs._Persona;
+
 	bDragging = false;
 	bIsActiveTimerRegistered = false;
 
@@ -44,6 +47,20 @@ void SMontageEditor::Construct(const FArguments& InArgs)
 		SharedPersona->RegisterOnPersonaRefresh(FPersona::FOnPersonaRefresh::CreateSP(this, &SMontageEditor::RebuildMontagePanel));
 	}
 
+	SAssignNew(AnimTimingPanel, SAnimTimingPanel)
+		.InWeakPersona(PersonaPtr)
+		.InSequence(MontageObj)
+		.WidgetWidth(S2ColumnWidget::DEFAULT_RIGHT_COLUMN_WIDTH)
+		.ViewInputMin(this, &SAnimEditorBase::GetViewMinInput)
+		.ViewInputMax(this, &SAnimEditorBase::GetViewMaxInput)
+		.InputMin(this, &SAnimEditorBase::GetMinInput)
+		.InputMax(this, &SAnimEditorBase::GetMaxInput)
+		.OnSetInputViewRange(this, &SAnimEditorBase::SetInputViewRange);
+
+	TAttribute<EVisibility> SectionVisibility = TAttribute<EVisibility>::Create(TAttribute<EVisibility>::FGetter::CreateSP(AnimTimingPanel.ToSharedRef(), &SAnimTimingPanel::IsElementDisplayVisible, ETimingElementType::Section));
+	TAttribute<EVisibility> NotifyVisibility = TAttribute<EVisibility>::Create(TAttribute<EVisibility>::FGetter::CreateSP(AnimTimingPanel.ToSharedRef(), &SAnimTimingPanel::IsElementDisplayVisible, ETimingElementType::QueuedNotify));
+	FOnGetTimingNodeVisibility TimingNodeVisibilityDelegate = FOnGetTimingNodeVisibility::CreateSP(AnimTimingPanel.ToSharedRef(), &SAnimTimingPanel::IsElementDisplayVisible);
+	
 	EditorPanels->AddSlot()
 	.AutoHeight()
 	.Padding(0, 10)
@@ -58,6 +75,7 @@ void SMontageEditor::Construct(const FArguments& InArgs)
 		.InputMin(this, &SAnimEditorBase::GetMinInput)
 		.InputMax(this, &SAnimEditorBase::GetMaxInput)
 		.OnSetInputViewRange(this, &SAnimEditorBase::SetInputViewRange)
+		.SectionTimingNodeVisibility(SectionVisibility)
 	];
 
 	EditorPanels->AddSlot()
@@ -69,6 +87,13 @@ void SMontageEditor::Construct(const FArguments& InArgs)
 		.MontageEditor(SharedThis(this))
 	];
 
+	EditorPanels->AddSlot()
+	.AutoHeight()
+	.Padding(0, 10)
+	[
+		AnimTimingPanel.ToSharedRef()
+	];
+
 	EditorPanels->AddSlot()
 	.AutoHeight()
 	.Padding(0, 10)
@@ -86,6 +111,7 @@ void SMontageEditor::Construct(const FArguments& InArgs)
 		.OnSelectionChanged(this, &SAnimEditorBase::OnSelectionChanged)
 		.MarkerBars(this, &SMontageEditor::GetMarkerBarInformation)
 		.OnRequestRefreshOffsets(this, &SMontageEditor::RefreshNotifyTriggerOffsets)
+		.OnGetTimingNodeVisibility(TimingNodeVisibilityDelegate)
 	];
 
 	EditorPanels->AddSlot()
@@ -268,6 +294,8 @@ void SMontageEditor::OnEditSectionTime( int32 SectionIndex, float NewTime)
 		MontageObj->CompositeSections[SectionIndex].SetTime(NewTime);
 		MontageObj->CompositeSections[SectionIndex].LinkMontage(MontageObj, NewTime);
 	}
+
+	AnimMontagePanel->RefreshTimingNodes();
 }
 void SMontageEditor::OnEditSectionTimeFinish( int32 SectionIndex )
 {
@@ -280,6 +308,12 @@ void SMontageEditor::OnEditSectionTimeFinish( int32 SectionIndex )
 		MontageObj->MarkPackageDirty();
 		AnimMontageSectionsPanel->Update();
 	}
+
+	TSharedPtr<FPersona> SharedPersona = WeakPersona.Pin();
+	if(SharedPersona.IsValid())
+	{
+		SharedPersona->OnSectionsChanged.Broadcast();
+	}
 }
 
 void SMontageEditor::PreAnimUpdate()
@@ -371,6 +405,7 @@ void SMontageEditor::SortAndUpdateMontage()
 	// Update view (this will recreate everything)
 	AnimMontagePanel->Update();
 	AnimMontageSectionsPanel->Update();
+	AnimTimingPanel->Update();
 
 	// Restart the preview instance of the montage
 	RestartPreview();
@@ -496,6 +531,7 @@ void SMontageEditor::RemoveSection(int32 SectionIndex)
 		EnsureStartingSection();
 		MontageObj->MarkPackageDirty();
 		AnimMontageSectionsPanel->Update();
+		AnimTimingPanel->Update();
 		RestartPreview();
 	}
 }
diff --git a/Engine/Source/Editor/Persona/Private/SMontageEditor.h b/Engine/Source/Editor/Persona/Private/SMontageEditor.h
--- a/Engine/Source/Editor/Persona/Private/SMontageEditor.h
+++ b/Engine/Source/Editor/Persona/Private/SMontageEditor.h
@@ -8,6 +8,13 @@
 #include "SNodePanel.h"
 #include "SAnimEditorBase.h"
 
+class SAnimMontagePanel;
+class SAnimNotifyPanel;
+class SAnimCurvePanel;
+class SAnimMontageSectionsPanel;
+class SAnimMontageScrubPanel;
+class SAnimTimingPanel;
+
 //////////////////////////////////////////////////////////////////////////
 // SMontageEditor
 
@@ -33,12 +40,14 @@ public:
 
 private:
 	/** Persona reference **/
-	TSharedPtr<class SAnimMontagePanel> AnimMontagePanel;
-	TSharedPtr<class SAnimNotifyPanel> AnimNotifyPanel;
-	TSharedPtr<class SAnimCurvePanel>	AnimCurvePanel;
-	TSharedPtr<class SAnimMontageSectionsPanel> AnimMontageSectionsPanel;
-	TSharedPtr<class SAnimMontageScrubPanel> AnimMontageScrubPanel;
-
+	TSharedPtr<SAnimMontagePanel> AnimMontagePanel;
+	TSharedPtr<SAnimNotifyPanel> AnimNotifyPanel;
+	TSharedPtr<SAnimCurvePanel>	AnimCurvePanel;
+	TSharedPtr<SAnimMontageSectionsPanel> AnimMontageSectionsPanel;
+	TSharedPtr<SAnimMontageScrubPanel> AnimMontageScrubPanel;
+	TSharedPtr<SAnimTimingPanel> AnimTimingPanel;
+
+	TWeakPtr<FPersona> WeakPersona;
 protected:
 	// Begin SAnimEditorBase interface
 	virtual TSharedRef<class SAnimationScrubPanel> ConstructAnimScrubPanel() override;
diff --git a/Engine/Source/Editor/Persona/Private/STimingTrack.cpp b/Engine/Source/Editor/Persona/Private/STimingTrack.cpp
new file mode 100644
--- /dev/null
+++ b/Engine/Source/Editor/Persona/Private/STimingTrack.cpp
@@ -0,0 +1,213 @@
+// Copyright 1998-2015 Epic Games, Inc. All Rights Reserved.
+
+#include "PersonaPrivatePCH.h"
+
+#include "STimingTrack.h"
+#include "SCurveEditor.h"
+
+void STimingTrack::Construct(const FArguments& Args)
+{
+	STrack::FArguments BaseArgs;
+	BaseArgs.ViewInputMin(Args._ViewInputMin);
+	BaseArgs.ViewInputMax(Args._ViewInputMax);
+	BaseArgs.TrackMinValue(Args._TrackMinValue);
+	BaseArgs.TrackMaxValue(Args._TrackMaxValue);
+	BaseArgs.TrackNumDiscreteValues(Args._TrackNumDiscreteValues);
+
+	STrack::Construct(BaseArgs);
+}
+
+void STimingTrack::OnArrangeChildren(const FGeometry& AllottedGeometry, FArrangedChildren& ArrangedChildren) const
+{
+	// To make this track look nice and have no overlapping nodes we're going to 
+	// treat this as a 1D collision problem and build/resolve islands until everything 
+	// is nicely resolved (or until we hit an upper limit as it can't be solved)
+
+	// Helper holding info about a single node
+	struct NodeData
+	{
+		NodeData(TSharedRef<STrackNode>& NodeRef, const FGeometry& Geometry)
+			: Node(NodeRef)
+		{
+			// Separation of the nodes on the track
+			const float NodeSeparation = 3.0f;
+
+			Node->CacheTrackGeometry(Geometry);
+			FVector2D Offset = Node->GetOffsetRelativeToParent(Geometry);
+			FVector2D Size = Node->GetSizeRelativeToParent(Geometry);
+
+			Offset.Y += (Geometry.GetLocalSize().Y - Size.Y) * 0.5f;
+
+			ActualRect = FBox2D(Offset, Offset + Size);
+			QueryRect = FBox2D(Offset - FVector2D(NodeSeparation, 0.0f), Offset + Size + FVector2D(NodeSeparation, 0.0f));
+		}
+
+		TSharedRef<STrackNode> Node;	// The node widget
+		FBox2D ActualRect;				// The actual render rect of the widget
+		FBox2D QueryRect;				// An expanded rect used to detect collisions
+	};
+
+	// Helper holding list of overlapping nodes
+	struct NodeIsland
+	{
+		TArray<NodeData*> Nodes;
+	};
+
+	int32 NumNodes = TrackNodes.Num();
+
+	TArray<NodeData> SortedNodeData;
+	SortedNodeData.Reserve(NumNodes);
+
+	// List of collision islands
+	TArray<NodeIsland> Islands;
+	// Scaling info to translate between local positions and data values
+	FTrackScaleInfo ScaleInfo(ViewInputMin.Get(), ViewInputMax.Get(), 0, 0, AllottedGeometry.Size);
+
+	for(int32 TrackIndex = 0; TrackIndex < NumNodes; ++TrackIndex)
+	{
+		TSharedRef<STrackNode> TrackNode = (TrackNodes[TrackIndex]);
+
+		SortedNodeData.Add(NodeData(TrackNode, AllottedGeometry));
+	}
+
+	const static int32 MaxRetries = 5;
+	int32 Retries = 0;
+	bool bResolved = true;
+
+	while(bResolved && Retries < MaxRetries)
+	{
+		++Retries;
+		bResolved = false;
+
+		for(int32 NodeIdx = 0; NodeIdx < NumNodes; ++NodeIdx)
+		{
+			NodeData& CurrentNode = SortedNodeData[NodeIdx];
+			// Island generation
+			NodeIsland* CurrentIsland = nullptr;
+			CurrentIsland = &Islands[Islands.AddZeroed()];
+
+			int32 Direction = -1;
+			int32 Next = NodeIdx + 1;
+			int32 HighestNode = NodeIdx;
+			FBox2D& CurrentRect = CurrentNode.ActualRect;
+			FBox2D CurrentQueryRect = CurrentNode.QueryRect;
+
+			CurrentIsland->Nodes.Add(&CurrentNode);
+
+			// Walk Nodes
+			while(Next >= 0 && Next < NumNodes)
+			{
+				NodeData& NextNode = SortedNodeData[Next];
+				FBox2D& NextRect = NextNode.ActualRect;
+				FBox2D& NextQueryRect = NextNode.QueryRect;
+				if(NextQueryRect.Intersect(CurrentQueryRect))
+				{
+					// Add to island
+					CurrentIsland->Nodes.Add(&NextNode);
+					HighestNode = Next;
+
+					// Expand the current query
+					CurrentQueryRect.Max = NextQueryRect.Max;
+				}
+				else
+				{
+					// No island, next node
+					break;
+				}
+
+				++Next;
+			}
+
+			// Skip processed nodes (those already in islands)
+			NodeIdx = HighestNode;
+		}
+
+		// Separation of the nodes on the track
+		const float NodeSeparation = 3.0f;
+
+		for(NodeIsland& Island : Islands)
+		{
+			if(Island.Nodes.Num() == 1)
+			{
+				// Keep single nodes on the data track range but skip everything else
+				FBox2D& NodeBox = Island.Nodes[0]->ActualRect;
+				FBox2D& NodeQueryRect = Island.Nodes[0]->QueryRect;
+				float Offset = 0.0f;
+				float Begin = ScaleInfo.LocalXToInput(NodeBox.Min.X);
+				float End = ScaleInfo.LocalXToInput(NodeBox.Max.X);
+				if(Begin < 0)
+				{
+					Offset = ScaleInfo.InputToLocalX(-Begin);
+				}
+				else if(End > TrackMaxValue.Get())
+				{
+					Offset = ScaleInfo.InputToLocalX(TrackMaxValue.Get() - End);
+				}
+
+				if(Offset != 0.0f)
+				{
+					NodeBox.Min.X += Offset;
+					NodeBox.Max.X += Offset;
+					NodeQueryRect.Min.X = NodeBox.Min.X - NodeSeparation;
+					NodeQueryRect.Max.X = NodeBox.Max.X + NodeSeparation;
+				}
+
+				continue;
+			}
+
+			bResolved = true;
+			// Island resolution
+			int32 NumIslandNodes = Island.Nodes.Num();
+			float Width = FMath::Max<float>((NumIslandNodes - 1), 0.0f) * NodeSeparation;
+			float Centre = 0.0f;
+
+			for(NodeData* Node : Island.Nodes)
+			{
+				Width += Node->ActualRect.GetSize().X;
+				Centre += Node->ActualRect.GetCenter().X;
+			}
+			Centre /= NumIslandNodes;
+
+			// Make sure the group stays on the track
+			float Begin = Centre - Width / 2.0f;
+			float WidthAsInput = Width / ScaleInfo.PixelsPerInput;
+			float BeginAsInput = FMath::Clamp(ScaleInfo.LocalXToInput(Begin), 0.0f, TrackMaxValue.Get() - WidthAsInput);
+			Begin = ScaleInfo.InputToLocalX(BeginAsInput);
+			
+			for(int32 NodeIdx = 0 ; NodeIdx < NumIslandNodes ; ++NodeIdx)
+			{
+				FBox2D& NodeBox = Island.Nodes[NodeIdx]->ActualRect;
+				float NodeWidth = NodeBox.GetSize().X;
+				float SeparationOffset = NodeIdx * NodeSeparation;
+				float PositionOffset = 0.0f;
+			
+				for(int32 PositionNodeIdx = 0 ; PositionNodeIdx < NodeIdx ; ++PositionNodeIdx)
+				{
+					PositionOffset += Island.Nodes[PositionNodeIdx]->ActualRect.GetSize().X;
+				}
+			
+				FBox2D& OriginalRect = Island.Nodes[NodeIdx]->ActualRect;
+				OriginalRect.Min.X = Begin + PositionOffset + SeparationOffset;
+				OriginalRect.Max.X = OriginalRect.Min.X + NodeWidth;
+			
+				// Alter Query rect for next pass
+				FBox2D& NodeQueryRect = Island.Nodes[NodeIdx]->QueryRect;
+				NodeQueryRect.Min = OriginalRect.Min - FVector2D(NodeSeparation, 0.0f);
+				NodeQueryRect.Max = OriginalRect.Max + FVector2D(NodeSeparation, 0.0f);
+			}
+		}
+	}
+
+	for(int32 TrackIndex = 0; TrackIndex < NumNodes; ++TrackIndex)
+	{
+		TSharedRef<STrackNode> TrackNode = (SortedNodeData[TrackIndex].Node);
+		if(TrackNode->IsBeingDragged())
+		{
+			continue;
+		}
+
+		FBox2D& Rect = SortedNodeData[TrackIndex].ActualRect;
+
+		ArrangedChildren.AddWidget(AllottedGeometry.MakeChild(TrackNode, Rect.Min, Rect.GetSize()));
+	}
+}
diff --git a/Engine/Source/Editor/Persona/Private/STimingTrack.h b/Engine/Source/Editor/Persona/Private/STimingTrack.h
new file mode 100644
--- /dev/null
+++ b/Engine/Source/Editor/Persona/Private/STimingTrack.h
@@ -0,0 +1,29 @@
+// Copyright 1998-2015 Epic Games, Inc. All Rights Reserved.
+
+#pragma once
+
+#include "STrack.h"
+
+//////////////////////////////////////////////////////////////////////////
+// Specialised anim track which arranges overlapping nodes into groups
+class STimingTrack : public STrack
+{
+public:
+
+	SLATE_BEGIN_ARGS(STimingTrack)
+	{}
+
+	SLATE_ATTRIBUTE(float, ViewInputMin)
+	SLATE_ATTRIBUTE(float, ViewInputMax)
+	SLATE_ATTRIBUTE(float, TrackMaxValue)
+	SLATE_ATTRIBUTE(float, TrackMinValue)
+	SLATE_ATTRIBUTE(int32, TrackNumDiscreteValues)
+
+	SLATE_END_ARGS()
+
+	void Construct(const FArguments& Args);
+
+	virtual void OnArrangeChildren( const FGeometry& AllottedGeometry, FArrangedChildren& ArrangedChildren ) const override;
+
+protected:
+};
\ No newline at end of file
diff --git a/Engine/Source/Editor/Persona/Private/STrack.cpp b/Engine/Source/Editor/Persona/Private/STrack.cpp
--- a/Engine/Source/Editor/Persona/Private/STrack.cpp
+++ b/Engine/Source/Editor/Persona/Private/STrack.cpp
@@ -90,75 +90,76 @@ void STrackNode::Construct(const FArguments& InArgs)
 	AllowDrag = InArgs._AllowDrag;
 
 	Font = FSlateFontInfo( FPaths::EngineContentDir() / TEXT("Slate/Fonts/Roboto-Regular.ttf"), 10 );
-}
-
-int32 STrackNode::OnPaint(const FPaintArgs& Args, const FGeometry& AllottedGeometry, const FSlateRect& MyClippingRect, FSlateWindowElementList& OutDrawElements, int32 LayerId, const FWidgetStyle& InWidgetStyle, bool bParentEnabled) const
-{
-	// We draw ourselves on all of the AllottedGeometry. The parent STrack should have arranged us
-	FPaintGeometry MyGeometry =	AllottedGeometry.ToPaintGeometry();
-	FVector2D DrawPos(0.f, 0.f);
-	FVector2D DrawSize = AllottedGeometry.Size;
-
-	LastSize = DrawSize; // HACK: Fixme. Need to save this off in case we are drag/dropped
-	
-	// Background
-	FLinearColor DrawColor = (IsSelected() ? SelectedNodeColor.Get() : NodeColor.Get());
 
 	const FSlateBrush* StyleInfo = FEditorStyle::GetBrush("ProgressBar.Background"); // FIXME: make slate argument for STrackNode
-	FSlateDrawElement::MakeBox( 
-		OutDrawElements,
-		LayerId++, 
-		MyGeometry, 
-		StyleInfo, 
-		MyClippingRect, 
-		ESlateDrawEffect::None,
-		DrawColor);
 
-	// Text
-	if (!NodeName.Get().IsEmpty())
+	if(InArgs._OverrideContent.Widget != SNullWidget::NullWidget)
 	{
-		MyGeometry = AllottedGeometry.ToPaintGeometry( DrawPos + FVector2D(2, 2), DrawSize );
-		FSlateDrawElement::MakeText( 
-			OutDrawElements,
-			LayerId++,
-			MyGeometry,
-			NodeName.Get(),
-			Font,
-			MyClippingRect,
-			ESlateDrawEffect::None,
-			FLinearColor::Black//FLinearColor::White//
-			);
+		bContentOverriden = true;
+		// Content was overridden by caller
+		this->ChildSlot
+		[
+			InArgs._OverrideContent.Widget
+		];
+	}
+	else
+	{
+		bContentOverriden = false;
+		// Content not overriden, use default node appearance
+		this->ChildSlot
+		[
+			SNew(SBorder)
+			.BorderImage(StyleInfo)
+			.ForegroundColor(FLinearColor::Black)
+			.BorderBackgroundColor(this, &STrackNode::GetNodeColor)
+			.Content()
+			[
+				SNew(STextBlock)
+				.Font(Font)
+				.Text(this, &STrackNode::GetNodeText)
+			]
+		];
 	}
-
-	return LayerId;
 }
 
 FVector2D STrackNode::GetOffsetRelativeToParent(const FGeometry& AllottedGeometry) const
 {
+	FVector2D Result(0.0f, 0.0f);
+	FVector2D Size = GetSizeRelativeToParent(AllottedGeometry);
 	FTrackScaleInfo ScaleInfo(ViewInputMin.Get(), ViewInputMax.Get(), 0, 0, AllottedGeometry.Size);
 
 	if(bCenterOnPosition)
 	{
-		FVector2D Size = GetSizeRelativeToParent(AllottedGeometry);
-
-		return FVector2D( ScaleInfo.InputToLocalX(DataStartPos.Get()) - (Size.X/2.f) , 0);
+		Result = FVector2D( ScaleInfo.InputToLocalX(DataStartPos.Get()) - (Size.X/2.f) , 0);
 
 	}	
-	return FVector2D( ScaleInfo.InputToLocalX(DataStartPos.Get()) , 0);
+	else
+	{
+		Result = FVector2D( ScaleInfo.InputToLocalX(DataStartPos.Get()) , 0);
+	}
+
+	return Result;
 }
 
 FVector2D STrackNode::GetSizeRelativeToParent(const FGeometry& AllottedGeometry) const
 {
-	if(DataLength.Get() > 0.f)
+	if(bContentOverriden)
 	{
-		// Scale us by data size
-		FTrackScaleInfo ScaleInfo(ViewInputMin.Get(), ViewInputMax.Get(), 0, 0, AllottedGeometry.Size);
-		return FVector2D( ScaleInfo.InputToLocalX(ViewInputMin.Get() + DataLength.Get()) , STrackDefaultHeight);
+		return ChildSlot.GetWidget()->GetDesiredSize();
 	}
 	else
 	{
-		// Use default hardcoded "knob" size
-		return FVector2D(NodeHandleWidth, NodeHandleHeight);
+		if(DataLength.Get() > 0.f)
+		{
+			// Scale us by data size
+			FTrackScaleInfo ScaleInfo(ViewInputMin.Get(), ViewInputMax.Get(), 0, 0, AllottedGeometry.Size);
+			return FVector2D(ScaleInfo.InputToLocalX(ViewInputMin.Get() + DataLength.Get()), STrackDefaultHeight);
+		}
+		else
+		{
+			// Use default hardcoded "knob" size
+			return FVector2D(NodeHandleWidth, NodeHandleHeight);
+		}
 	}
 }
 
@@ -219,6 +220,7 @@ FReply STrackNode::BeginDrag( const FGeometry& MyGeometry, const FPointerEvent&
 	FVector2D ScreenNodePosition = MyGeometry.AbsolutePosition;// + GetOffsetRelativeToParent(MyGeometry);
 	
 	bBeingDragged = true;
+	LastSize = MyGeometry.Size;
 
 	//void FTrackNodeDragDropOp(TSharedRef<STrackNode> TrackNode, const FVector2D &CursorPosition, const FVector2D &ScreenPositionOfNode)
 	return FReply::Handled().BeginDragDrop(FTrackNodeDragDropOp::New(SharedThis(this), ScreenCursorPos, ScreenNodePosition));
@@ -317,6 +319,16 @@ float STrackNode::GetDataStartPos() const
 	return 0.f;
 }
 
+FSlateColor STrackNode::GetNodeColor() const
+{
+	return IsSelected() ? FSlateColor(SelectedNodeColor.Get()) : FSlateColor(NodeColor.Get());
+}
+
+FText STrackNode::GetNodeText() const
+{
+	return FText::FromString(NodeName.Get());
+}
+
 //////////////////////////////////////////////////////////////////////////
 // STrack
 
@@ -367,7 +379,7 @@ void STrack::OnArrangeChildren( const FGeometry& AllottedGeometry, FArrangedChil
 		TrackNode->CacheTrackGeometry(AllottedGeometry);
 		FVector2D Offset = TrackNode->GetOffsetRelativeToParent(AllottedGeometry);
 		FVector2D Size = TrackNode->GetSizeRelativeToParent(AllottedGeometry);
-
+	
 		ArrangedChildren.AddWidget( AllottedGeometry.MakeChild(TrackNode, Offset, Size) );
 	}
 }
@@ -446,7 +458,7 @@ int32 STrack::OnPaint( const FPaintArgs& Args, const FGeometry& AllottedGeometry
 		{
 			//float EndPos = (DraggableBarIndex.Get().IsValidIndex(I+1) ? DraggableBarIndex.Get()(I+1)
 
-			FPaintGeometry TextGeometry = AllottedGeometry.ToPaintGeometry( FVector2D(XPos + 5.f, 5.f), AllottedGeometry.GetDrawSize() );
+			FPaintGeometry TextGeometry = AllottedGeometry.ToPaintGeometry( FVector2D(XPos + 15.f, 5.f), AllottedGeometry.GetDrawSize() );
 			FSlateDrawElement::MakeText( 
 				OutDrawElements,
 				CustomLayerId,
@@ -640,7 +652,7 @@ FReply STrack::OnMouseMove( const FGeometry& MyGeometry, const FPointerEvent& Mo
 		/** Update what bar is draggable if we arent already dragging */
 		UpdateDraggableBarIndex(MyGeometry, MouseEvent.GetScreenSpacePosition());
 	}
-	
+
 	return FReply::Unhandled();
 }
 
@@ -767,6 +779,10 @@ bool STrack::GetDraggableBarSnapPosition(const FGeometry& MyGeometry, float &Out
 
 void STrack::OnMouseLeave( const FPointerEvent& MouseEvent )
 {
+	if(bDraggingBar && OnBarDrop.IsBound())
+	{
+		OnBarDrop.Execute(DraggableBarIndex);
+	}
 	bDraggingBar = false;
 }
 
@@ -834,4 +850,9 @@ void STrack::AddTrackNode( TSharedRef<STrackNode> Node )
 	TrackNodes.Add(Node);
 }
 
+void STrack::ClearTrack()
+{
+	TrackNodes.Empty();
+}
+
 #undef LOCTEXT_NAMESPACE
diff --git a/Engine/Source/Editor/Persona/Private/STrack.h b/Engine/Source/Editor/Persona/Private/STrack.h
--- a/Engine/Source/Editor/Persona/Private/STrack.h
+++ b/Engine/Source/Editor/Persona/Private/STrack.h
@@ -2,6 +2,8 @@
 
 #pragma once
 
+#include "GraphEditor.h"
+
 //////////////////////////////////////////////////////////////////////////
 DECLARE_DELEGATE_RetVal( float, FOnGetScrubValue )
 DECLARE_DELEGATE_OneParam( FOnSelectionChanged, const FGraphPanelSelectionSet& )
@@ -63,7 +65,7 @@ protected:
 /** class STrackNode for STrack. 
  * This is Children for STrack
  */
-class STrackNode : public SLeafWidget
+class STrackNode : public SCompoundWidget 
 {
 public:
 	SLATE_BEGIN_ARGS( STrackNode )
@@ -101,6 +103,7 @@ public:
 	SLATE_EVENT( FOnNodeRightClickContextMenu, OnNodeRightClickContextMenu )
 	SLATE_EVENT( FOnTrackNodeClicked, OnTrackNodeClicked )
 	SLATE_ARGUMENT( bool, CenterOnPosition )
+	SLATE_NAMED_SLOT(FArguments, OverrideContent)
 
 	SLATE_END_ARGS()
 
@@ -117,8 +120,6 @@ public:
 	virtual FVector2D GetOffsetRelativeToParent(const FGeometry& ParentAllottedGeometry) const;
 	virtual FVector2D GetSizeRelativeToParent(const FGeometry& ParentAllottedGeometry) const;
 	
-	virtual int32 OnPaint(const FPaintArgs& Args, const FGeometry& AllottedGeometry, const FSlateRect& MyClippingRect, FSlateWindowElementList& OutDrawElements, int32 LayerId, const FWidgetStyle& InWidgetStyle, bool bParentEnabled) const override;
-
 	// drag drop relationship
 	virtual FReply OnDragDetected( const FGeometry& MyGeometry, const FPointerEvent& MouseEvent ) override;
 	virtual FReply OnMouseButtonDown( const FGeometry& MyGeometry, const FPointerEvent& MouseEvent ) override;
@@ -147,8 +148,14 @@ public:
 
 	const FGeometry& GetTrackGeometry() const {return CachedTrackGeometry;}
 
+	bool IsBeingDragged() const {return bBeingDragged;}
+
 protected:
 	
+	FSlateColor GetNodeColor() const;
+
+	FText GetNodeText() const; // Temp, remove and replace with correct attribute
+
 	void	ToggleSelect();
 	void	Select();
 	void	Deselect();
@@ -183,6 +190,7 @@ protected:
 	bool				bBeingDragged;
 	bool				bCenterOnPosition;
 	bool				AllowDrag;
+	bool				bContentOverriden;
 
 	friend class STrack;
 };
@@ -255,6 +263,7 @@ public:
 	virtual FCursorReply OnCursorQuery( const FGeometry& MyGeometry, const FPointerEvent& CursorEvent ) const override;
 
 	void				AddTrackNode( TSharedRef<STrackNode> Node );
+	void				ClearTrack();
 	FVector2D			ComputeDesiredSize(float) const override;
 	virtual FChildren*	GetChildren() override;
 
diff --git a/Engine/Source/Editor/UnrealEd/Classes/Preferences/PersonaOptions.h b/Engine/Source/Editor/UnrealEd/Classes/Preferences/PersonaOptions.h
--- a/Engine/Source/Editor/UnrealEd/Classes/Preferences/PersonaOptions.h
+++ b/Engine/Source/Editor/UnrealEd/Classes/Preferences/PersonaOptions.h
@@ -49,6 +49,15 @@ class UNREALED_API UPersonaOptions : public UObject
 	UPROPERTY(EditAnywhere, config, Category=Options)
 	uint32 DefaultLocalAxesSelection;
 
+	UPROPERTY(EditAnywhere, config, Category = Options)
+	FLinearColor SectionTimingNodeColour;
+
+	UPROPERTY(EditAnywhere, config, Category = Options)
+	FLinearColor NotifyTimingNodeColour;
+
+	UPROPERTY(EditAnywhere, config, Category = Options)
+	FLinearColor BranchingPointTimingNodeColour;
+
 public:
 	void SetViewportBackgroundColor( const FLinearColor& InViewportBackgroundColor);
 	void SetShowGrid( bool bInShowGrid );
@@ -61,5 +70,7 @@ public:
 	void SetViewFOV( float InViewFOV );
 	void SetDefaultLocalAxesSelection( uint32 InDefaultLocalAxesSelection );
 	void SetShowMeshStats( int32 InShowMeshStats );
+	void SetSectionTimingNodeColor(const FLinearColor& InColor);
+	void SetNotifyTimingNodeColor(const FLinearColor& InColor);
+	void SetBranchingPointTimingNodeColor(const FLinearColor& InColor);
 };
-
diff --git a/Engine/Source/Editor/UnrealEd/Private/PreferenceStubs.cpp b/Engine/Source/Editor/UnrealEd/Private/PreferenceStubs.cpp
--- a/Engine/Source/Editor/UnrealEd/Private/PreferenceStubs.cpp
+++ b/Engine/Source/Editor/UnrealEd/Private/PreferenceStubs.cpp
@@ -40,6 +40,10 @@ UPersonaOptions::UPersonaOptions(const FObjectInitializer& ObjectInitializer)
 	, DefaultLocalAxesSelection(2)
 {
 	ViewModeIndex = VMI_Lit;
+
+	SectionTimingNodeColour = FLinearColor(0.0f, 1.0f, 0.0f);
+	NotifyTimingNodeColour = FLinearColor(1.0f, 0.0f, 0.0f);
+	BranchingPointTimingNodeColour = FLinearColor(0.5f, 1.0f, 1.0f);
 }
 
 void UPersonaOptions::SetViewportBackgroundColor( const FLinearColor& InViewportBackgroundColor)
@@ -108,3 +112,21 @@ void UPersonaOptions::SetShowMeshStats( int32 InShowMeshStats )
 	SaveConfig();
 }
 
+
+void UPersonaOptions::SetSectionTimingNodeColor(const FLinearColor& InColor)
+{
+	SectionTimingNodeColour = InColor;
+	SaveConfig();
+}
+
+void UPersonaOptions::SetNotifyTimingNodeColor(const FLinearColor& InColor)
+{
+	NotifyTimingNodeColour = InColor;
+	SaveConfig();
+}
+
+void UPersonaOptions::SetBranchingPointTimingNodeColor(const FLinearColor& InColor)
+{
+	BranchingPointTimingNodeColour = InColor;
+	SaveConfig();
+}
diff --git a/Engine/Source/Runtime/Engine/Classes/Animation/AnimSequenceBase.h b/Engine/Source/Runtime/Engine/Classes/Animation/AnimSequenceBase.h
--- a/Engine/Source/Runtime/Engine/Classes/Animation/AnimSequenceBase.h
+++ b/Engine/Source/Runtime/Engine/Classes/Animation/AnimSequenceBase.h
@@ -119,13 +119,12 @@ struct FAnimNotifyEvent : public FAnimLinkableElement
 	/** Color of Notify in editor */
 	UPROPERTY()
 	FColor NotifyColor;
+#endif // WITH_EDITORONLY_DATA
 
-	/** Visual position of notify in the vertical 'tracks' of the notify editor panel */
+	/** 'Track' that the notify exists on, used for visual placement in editor and sorting priority in runtime */
 	UPROPERTY()
 	int32 TrackIndex;
 
-#endif // WITH_EDITORONLY_DATA
-
 	FAnimNotifyEvent()
 		: DisplayTime_DEPRECATED(0)
 		, TriggerTimeOffset(0)
@@ -196,19 +195,15 @@ FORCEINLINE bool FAnimNotifyEvent::operator<(const FAnimNotifyEvent& Other) cons
 	float ATime = GetTriggerTime();
 	float BTime = Other.GetTriggerTime();
 
-#if WITH_EDITORONLY_DATA
-	// this sorting only works if it's saved in editor or loaded with editor data
-	// this is required for gameplay team to have reliable order of notifies
-	// but it was noted that this change will require to resave notifies. 
-	// once you resave, this order will be preserved
-	if (FMath::IsNearlyEqual(ATime, BTime))
+	// When we've got the same time sort on the track index. Explicitly
+	// using SMALL_NUMBER here incase the underlying default changes as
+	// notifies can have an offset of KINDA_SMALL_NUMBER to be consider
+	// distinct
+	if (FMath::IsNearlyEqual(ATime, BTime, SMALL_NUMBER))
 	{
-
-		// if the 2 anim notify events are the same display time sort based off of track index
 		return TrackIndex < Other.TrackIndex;
 	}
 	else
-#endif // WITH_EDITORONLY_DATA
 	{
 		return ATime < BTime;
 	}
@@ -708,6 +703,7 @@ public:
 	ENGINE_API void RegisterOnNotifyChanged(const FOnNotifyChanged& Delegate);
 	ENGINE_API void UnregisterOnNotifyChanged(void* Unregister);
 	ENGINE_API virtual bool IsValidToPlay() const { return true; }
+
 #endif
 
 protected:
