Fix for notify context menu delete not working on some platforms.

diff --git a/Engine/Source/Editor/Persona/Private/SAnimNotifyPanel.cpp b/Engine/Source/Editor/Persona/Private/SAnimNotifyPanel.cpp
--- a/Engine/Source/Editor/Persona/Private/SAnimNotifyPanel.cpp
+++ b/Engine/Source/Editor/Persona/Private/SAnimNotifyPanel.cpp
@@ -3829,7 +3829,7 @@ void SAnimNotifyPanel::RefreshNotifyTracks()
 			.OnNodeDragStarted(this, &SAnimNotifyPanel::OnNotifyNodeDragStarted)
 			.MarkerBars(MarkerBars)
 			.OnRequestRefreshOffsets(OnRequestRefreshOffsets)
-			.OnDeleteNotify(this, &SAnimNotifyPanel::OnDeletePressed)
+			.OnDeleteNotify(this, &SAnimNotifyPanel::DeleteSelectedNotifies)
 			.OnDeselectAllNotifies(this, &SAnimNotifyPanel::DeselectAllNotifies)
 			.OnCopyNotifies(this, &SAnimNotifyPanel::CopySelectedNotifiesToClipboard)
 			.OnPasteNotifies(this, &SAnimNotifyPanel::OnPasteNotifies)
@@ -3915,32 +3915,37 @@ void SAnimNotifyPanel::OnDeletePressed()
 	// so don't delete anything when the key is pressed.
 	if(HasKeyboardFocus() || HasFocusedDescendants()) 
 	{
-		TArray<FAnimNotifyEvent*> SelectedNotifies;
-		for (TSharedPtr<SAnimNotifyTrack> Track : NotifyAnimTracks)
-		{
-			Track->AppendSelectionToArray(SelectedNotifies);
-		}
+		DeleteSelectedNotifies();
+	}
+}
 
-		if(SelectedNotifies.Num() > 0)
-		{
-			FScopedTransaction Transaction(LOCTEXT("DeleteNotifies", "Delete Animation Notifies"));
-			Sequence->Modify(true);
+void SAnimNotifyPanel::DeleteSelectedNotifies()
+{
+	TArray<FAnimNotifyEvent*> SelectedNotifies;
+	for(TSharedPtr<SAnimNotifyTrack> Track : NotifyAnimTracks)
+	{
+		Track->AppendSelectionToArray(SelectedNotifies);
+	}
 
-			// Delete from highest index to lowest
-			SelectedNotifies.Sort();
-			for (int NotifyIndex = SelectedNotifies.Num()-1; NotifyIndex >= 0 ; --NotifyIndex)
-			{
-				FAnimNotifyEvent* Notify = SelectedNotifies[NotifyIndex];
-				DeleteNotify(Notify);
-			}
+	if(SelectedNotifies.Num() > 0)
+	{
+		FScopedTransaction Transaction(LOCTEXT("DeleteNotifies", "Delete Animation Notifies"));
+		Sequence->Modify(true);
+
+		// Delete from highest index to lowest
+		SelectedNotifies.Sort();
+		for(int NotifyIndex = SelectedNotifies.Num() - 1; NotifyIndex >= 0 ; --NotifyIndex)
+		{
+			FAnimNotifyEvent* Notify = SelectedNotifies[NotifyIndex];
+			DeleteNotify(Notify);
 		}
+	}
 
-		// clear selection and update the panel
-		FGraphPanelSelectionSet ObjectSet;
-		OnSelectionChanged.ExecuteIfBound(ObjectSet);
+	// clear selection and update the panel
+	FGraphPanelSelectionSet ObjectSet;
+	OnSelectionChanged.ExecuteIfBound(ObjectSet);
 
-		Update();
-	}
+	Update();
 }
 
 void SAnimNotifyPanel::SetSequence(class UAnimSequenceBase *	InSequence)
diff --git a/Engine/Source/Editor/Persona/Private/SAnimNotifyPanel.h b/Engine/Source/Editor/Persona/Private/SAnimNotifyPanel.h
--- a/Engine/Source/Editor/Persona/Private/SAnimNotifyPanel.h
+++ b/Engine/Source/Editor/Persona/Private/SAnimNotifyPanel.h
@@ -257,6 +257,9 @@ private:
 	/** Handler for delete command */
 	void OnDeletePressed();
 
+	/** Deletes all currently selected notifies in the panel */
+	void DeleteSelectedNotifies();
+
 	/** We support keyboard focus to detect when we should process key commands like delete */
 	virtual bool SupportsKeyboardFocus() const override
 	{
