Fix crash in Persona when changing modes and destroying wind source actors.

diff --git a/Engine/Source/Editor/Persona/Private/AnimationEditorViewportClient.cpp b/Engine/Source/Editor/Persona/Private/AnimationEditorViewportClient.cpp
--- a/Engine/Source/Editor/Persona/Private/AnimationEditorViewportClient.cpp
+++ b/Engine/Source/Editor/Persona/Private/AnimationEditorViewportClient.cpp
@@ -131,11 +131,6 @@ FAnimationViewportClient::FAnimationViewportClient(FAnimationEditorPreviewScene&
 	// set camera mode
 	bCameraFollow = false;
 
-	//wind actor's initial position, rotation and strength
-	PrevWindLocation = FVector(100, 100, 100);
-	PrevWindRotation = FRotator(0,0,0); // roll, yaw, pitch
-	PrevWindStrength = 0.2f;
-
 	bDrawUVs = false;
 	UVChannelToDraw = 0;
 
@@ -147,6 +142,25 @@ FAnimationViewportClient::FAnimationViewportClient(FAnimationEditorPreviewScene&
 	{
 		World->bAllowAudioPlayback = !ConfigOption->bMuteAudio;
 	}
+
+	// Grab a wind actor if it exists (could have been placed in the scene before a mode transition)
+	TActorIterator<AWindDirectionalSource> WindIter(World);
+	if(WindIter)
+	{
+		AWindDirectionalSource* WindActor = *WindIter;
+		WindSourceActor = WindActor;
+
+		PrevWindLocation = WindActor->GetActorLocation();
+		PrevWindRotation = WindActor->GetActorRotation();
+		PrevWindStrength = WindActor->GetComponent()->Strength;
+	}
+	else
+	{
+		//wind actor's initial position, rotation and strength
+		PrevWindLocation = FVector(100, 100, 100);
+		PrevWindRotation = FRotator(0, 0, 0); // roll, yaw, pitch
+		PrevWindStrength = 0.2f;
+	}
 }
 
 FAnimationViewportClient::~FAnimationViewportClient()
@@ -2148,7 +2162,7 @@ void FAnimationViewportClient::EnableWindActor(bool bEnableWind)
 			WindSourceActor = CreateWindActor(World);
 		}
 	}
-	else
+	else if(WindSourceActor.IsValid())
 	{
 		PrevWindLocation = WindSourceActor->GetActorLocation();
 		PrevWindRotation = WindSourceActor->GetActorRotation();
