Fixed smart name manager widget not deleting unreferenced names when clicking delete in the context menu.

diff --git a/Engine/Source/Editor/Persona/Private/SSkeletonSmartNameManager.cpp b/Engine/Source/Editor/Persona/Private/SSkeletonSmartNameManager.cpp
--- a/Engine/Source/Editor/Persona/Private/SSkeletonSmartNameManager.cpp
+++ b/Engine/Source/Editor/Persona/Private/SSkeletonSmartNameManager.cpp
@@ -453,19 +453,21 @@ void SCurveNameManager::OnDeleteNameClicked()
 					Sequence->RawCurveData.DeleteCurveData(Uid);
 				}
 			}
-
-			// Remove names from skeleton
-			CurrentSkeleton->Modify(true);
-			FSmartNameMapping* Mapping = CurrentSkeleton->SmartNames.GetContainer(ContainerName);
-			for(FSmartNameMapping::UID Uid : SelectedUids)
-			{
-				Mapping->Remove(Uid);
-			}
-
 			GWarn->EndSlowTask();
 		}
 	}
 
+	if(SelectedUids.Num() > 0)
+	{
+		// Remove names from skeleton
+		CurrentSkeleton->Modify(true);
+		FSmartNameMapping* Mapping = CurrentSkeleton->SmartNames.GetContainer(ContainerName);
+		for(FSmartNameMapping::UID Uid : SelectedUids)
+		{
+			Mapping->Remove(Uid);
+		}
+	}
+
 	GenerateDisplayedList(CurrentFilterText);
 
 	TSharedPtr<FPersona> SharedPersona = WeakPersona.Pin();
