Animation: Minor style fixes for padding/fonts of bone references and the alpha scale/bias of a skeletal control

diff --git a/Engine/Source/Editor/Persona/Private/AnimGraphNodeDetails.cpp b/Engine/Source/Editor/Persona/Private/AnimGraphNodeDetails.cpp
--- a/Engine/Source/Editor/Persona/Private/AnimGraphNodeDetails.cpp
+++ b/Engine/Source/Editor/Persona/Private/AnimGraphNodeDetails.cpp
@@ -275,6 +275,7 @@ void FInputScaleBiasCustomization::CustomizeHeader(TSharedRef<class IPropertyHan
 	[
 		SNew(STextBlock)
 		.Text(StructPropertyHandle->GetPropertyDisplayName())
+		.Font(IDetailLayoutBuilder::GetDetailFont())
 	]
 	.ValueContent()
 	[
diff --git a/Engine/Source/Editor/Persona/Private/Shared/BoneSelectionWidget.cpp b/Engine/Source/Editor/Persona/Private/Shared/BoneSelectionWidget.cpp
--- a/Engine/Source/Editor/Persona/Private/Shared/BoneSelectionWidget.cpp
+++ b/Engine/Source/Editor/Persona/Private/Shared/BoneSelectionWidget.cpp
@@ -34,11 +34,12 @@ void SBoneSelectionWidget::Construct(const FArguments& InArgs)
 	[
 		SAssignNew(BonePickerButton, SComboButton)
 		.OnGetMenuContent(FOnGetContent::CreateSP(this, &SBoneSelectionWidget::CreateSkeletonWidgetMenu))
-		.ContentPadding(0)
+		.ContentPadding(FMargin(4.0f, 2.0f, 4.0f, 2.0f))
 		.ButtonContent()
 		[
 			SNew(STextBlock)
 			.Text(this, &SBoneSelectionWidget::GetCurrentBoneName)
+			.Font(IDetailLayoutBuilder::GetDetailFont())
 			.ToolTipText(FinalTooltip)
 		]
 	];
