Animation: Added missing include to fix non-Unity build

diff --git a/Engine/Source/Editor/Persona/Private/Shared/BoneSelectionWidget.cpp b/Engine/Source/Editor/Persona/Private/Shared/BoneSelectionWidget.cpp
--- a/Engine/Source/Editor/Persona/Private/Shared/BoneSelectionWidget.cpp
+++ b/Engine/Source/Editor/Persona/Private/Shared/BoneSelectionWidget.cpp
@@ -4,6 +4,7 @@
 #include "PersonaPrivatePCH.h"
 #include "BoneSelectionWidget.h"
 #include "SSearchBox.h"
+#include "Editor/PropertyEditor/Public/DetailLayoutBuilder.h"
 
 #define LOCTEXT_NAMESPACE "SBoneSelectionWidget"
 
