Animation: Changed the 'click to recompile AnimBP' prompt to call into Persona to do the recompile, causing the message log to be displayed/refreshed if there are warnings/errors during compilation [UE-17247]

diff --git a/Engine/Source/Editor/Persona/Private/Persona.cpp b/Engine/Source/Editor/Persona/Private/Persona.cpp
--- a/Engine/Source/Editor/Persona/Private/Persona.cpp
+++ b/Engine/Source/Editor/Persona/Private/Persona.cpp
@@ -1674,6 +1674,17 @@ bool FPersona::CanRemovePosePin() const
 	return true;
 }
 
+void FPersona::RecompileAnimBlueprintIfDirty()
+{
+	if (UBlueprint* Blueprint = GetBlueprintObj())
+	{
+		if (!Blueprint->IsUpToDate())
+		{
+			Compile();
+		}
+	}
+}
+
 void FPersona::Compile()
 {
 	// Note if we were debugging the preview
diff --git a/Engine/Source/Editor/Persona/Private/Persona.h b/Engine/Source/Editor/Persona/Private/Persona.h
--- a/Engine/Source/Editor/Persona/Private/Persona.h
+++ b/Engine/Source/Editor/Persona/Private/Persona.h
@@ -207,6 +207,8 @@ public:
 	/** Refresh Preview Instance Track Curves **/
 	void RefreshPreviewInstanceTrackCurves();
 
+	void RecompileAnimBlueprintIfDirty();
+
 protected:
 	// FBlueprintEditor interface
 	//virtual void CreateDefaultToolbar() override;
diff --git a/Engine/Source/Editor/Persona/Private/SAnimationEditorViewport.cpp b/Engine/Source/Editor/Persona/Private/SAnimationEditorViewport.cpp
--- a/Engine/Source/Editor/Persona/Private/SAnimationEditorViewport.cpp
+++ b/Engine/Source/Editor/Persona/Private/SAnimationEditorViewport.cpp
@@ -2030,16 +2030,13 @@ FReply SAnimationEditorViewportTabBody::ClickedOnViewportCornerText()
 	TSharedPtr<FPersona> Persona = PersonaPtr.Pin();
 	// if it's recording, it won't be able to see the message
 	// so disable it
-	if( Persona->Recorder.InRecording() )
+	if (Persona->Recorder.InRecording())
 	{
 		Persona->Recorder.StopRecord(true);
 	}
-	else if(UBlueprint* Blueprint = Persona->GetBlueprintObj())
+	else 
 	{
-		if(!Blueprint->IsUpToDate())
-		{
-			FKismetEditorUtilities::CompileBlueprint(Blueprint);
-		}
+		Persona->RecompileAnimBlueprintIfDirty();
 	}
 
 	return FReply::Handled();
