Fixed animation blueprint crash when compiling a referenced blueprint in another window
#UE-16822

diff --git a/Engine/Source/Editor/Persona/Private/Persona.cpp b/Engine/Source/Editor/Persona/Private/Persona.cpp
--- a/Engine/Source/Editor/Persona/Private/Persona.cpp
+++ b/Engine/Source/Editor/Persona/Private/Persona.cpp
@@ -3166,6 +3166,34 @@ void FPersona::OnBlueprintPreCompile(UBlueprint* BlueprintToCompile)
 	}
 }
 
+void FPersona::OnBlueprintChangedImpl(UBlueprint* InBlueprint, bool bIsJustBeingCompiled /*= false*/)
+{
+	FBlueprintEditor::OnBlueprintChangedImpl(InBlueprint, bIsJustBeingCompiled);
+
+	UObject* CurrentDebugObject = GetBlueprintObj()->GetObjectBeingDebugged();
+	const bool bIsDebuggingPreview = (PreviewComponent != NULL) && PreviewComponent->IsAnimBlueprintInstanced() && (PreviewComponent->AnimScriptInstance == CurrentDebugObject);
+
+	if(PreviewComponent != NULL)
+	{
+		if(PreviewComponent->AnimScriptInstance == NULL)
+		{
+			// try reinitialize animation if it doesn't exist
+			PreviewComponent->InitAnim(true);
+		}
+
+		if(bIsDebuggingPreview)
+		{
+			GetBlueprintObj()->SetObjectBeingDebugged(PreviewComponent->AnimScriptInstance);
+		}
+	}
+
+	// calls PostCompile to copy proper values between anim nodes
+	if(Viewport.IsValid())
+	{
+		Viewport.Pin()->GetAnimationViewportClient()->PostCompile();
+	}
+}
+
 static class FMeshHierarchyCmd : private FSelfRegisteringExec
 {
 public:
diff --git a/Engine/Source/Editor/Persona/Private/Persona.h b/Engine/Source/Editor/Persona/Private/Persona.h
--- a/Engine/Source/Editor/Persona/Private/Persona.h
+++ b/Engine/Source/Editor/Persona/Private/Persona.h
@@ -232,6 +232,7 @@ protected:
 	virtual FGraphAppearanceInfo GetGraphAppearance(class UEdGraph* InGraph) const override;
 	virtual bool IsEditable(UEdGraph* InGraph) const override;
 	virtual FText GetGraphDecorationString(UEdGraph* InGraph) const override;
+	virtual void OnBlueprintChangedImpl(UBlueprint* InBlueprint, bool bIsJustBeingCompiled = false) override;
 	// End of FBlueprintEditor interface
 
 	// IAssetEditorInstance interface
