Animation: Allow skeletal controls to specify the coordinate space that they would prefer the gizmo to be displayed in, and updated Modify Bone to use the new functionality

diff --git a/Engine/Source/Editor/AnimGraph/Classes/AnimGraphNode_ModifyBone.h b/Engine/Source/Editor/AnimGraph/Classes/AnimGraphNode_ModifyBone.h
--- a/Engine/Source/Editor/AnimGraph/Classes/AnimGraphNode_ModifyBone.h
+++ b/Engine/Source/Editor/AnimGraph/Classes/AnimGraphNode_ModifyBone.h
@@ -28,6 +28,7 @@ protected:
 
 	// UAnimGraphNode_SkeletalControlBase interface
 	virtual FText GetControllerDescription() const override;
+	virtual int32 GetWidgetCoordinateSystem(const USkeletalMeshComponent* SkelComp) override;
 	virtual FVector GetWidgetLocation(const USkeletalMeshComponent* SkelComp, struct FAnimNode_SkeletalControlBase* AnimNode) override;
 	virtual int32 GetWidgetMode(const USkeletalMeshComponent* SkelComp) override;
 	virtual int32 ChangeToNextWidgetMode(const USkeletalMeshComponent* SkelComp, int32 InCurWidgetMode) override;
diff --git a/Engine/Source/Editor/AnimGraph/Classes/AnimGraphNode_SkeletalControlBase.h b/Engine/Source/Editor/AnimGraph/Classes/AnimGraphNode_SkeletalControlBase.h
--- a/Engine/Source/Editor/AnimGraph/Classes/AnimGraphNode_SkeletalControlBase.h
+++ b/Engine/Source/Editor/AnimGraph/Classes/AnimGraphNode_SkeletalControlBase.h
@@ -46,6 +46,9 @@ public:
 	/** Are we currently showing this pin */
 	bool IsPinShown(const FString& PinName) const;
 
+	// Returns the coordinate system that should be used for this bone
+	virtual int32 GetWidgetCoordinateSystem(const USkeletalMeshComponent* SkelComp);
+
 	// return current widget mode this anim graph node supports
 	virtual int32 GetWidgetMode(const USkeletalMeshComponent* SkelComp);
 	// called when the user changed widget mode by pressing "Space" key
diff --git a/Engine/Source/Editor/AnimGraph/Private/AnimGraphNode_ModifyBone.cpp b/Engine/Source/Editor/AnimGraph/Private/AnimGraphNode_ModifyBone.cpp
--- a/Engine/Source/Editor/AnimGraph/Private/AnimGraphNode_ModifyBone.cpp
+++ b/Engine/Source/Editor/AnimGraph/Private/AnimGraphNode_ModifyBone.cpp
@@ -67,6 +67,36 @@ FText UAnimGraphNode_ModifyBone::GetNodeTitle(ENodeTitleType::Type TitleType) co
 	return CachedNodeTitles[TitleType];
 }
 
+int32 UAnimGraphNode_ModifyBone::GetWidgetCoordinateSystem(const USkeletalMeshComponent* SkelComp)
+{
+	EBoneControlSpace Space = BCS_BoneSpace;
+	switch (CurWidgetMode)
+	{
+	case FWidget::WM_Rotate:
+		Space = Node.RotationSpace;
+		break;
+	case FWidget::WM_Translate:
+		Space = Node.TranslationSpace;
+		break;
+	case FWidget::WM_Scale:
+		Space = Node.ScaleSpace;
+		break;
+	}
+
+	switch (Space)
+	{
+	default:
+	case BCS_ParentBoneSpace:
+		//@TODO: No good way of handling this one
+		return COORD_World;
+	case BCS_BoneSpace:
+		return COORD_Local;
+	case BCS_ComponentSpace:
+	case BCS_WorldSpace:
+		return COORD_World;
+	}
+}
+
 FVector UAnimGraphNode_ModifyBone::GetWidgetLocation(const USkeletalMeshComponent* SkelComp, FAnimNode_SkeletalControlBase* AnimNode)
 {
 	USkeleton* Skeleton = SkelComp->SkeletalMesh->Skeleton;
diff --git a/Engine/Source/Editor/AnimGraph/Private/AnimGraphNode_SkeletalControlBase.cpp b/Engine/Source/Editor/AnimGraph/Private/AnimGraphNode_SkeletalControlBase.cpp
--- a/Engine/Source/Editor/AnimGraph/Private/AnimGraphNode_SkeletalControlBase.cpp
+++ b/Engine/Source/Editor/AnimGraph/Private/AnimGraphNode_SkeletalControlBase.cpp
@@ -14,6 +14,18 @@ UAnimGraphNode_SkeletalControlBase::UAnimGraphNode_SkeletalControlBase(const FOb
 {
 }
 
+int32 UAnimGraphNode_SkeletalControlBase::GetWidgetCoordinateSystem(const USkeletalMeshComponent* SkelComp)
+{
+	if (GetWidgetMode(SkelComp) == FWidget::WM_Scale)
+	{
+		return COORD_Local;
+	}
+	else
+	{
+		return COORD_World;
+	}
+}
+
 // returns int32 instead of EWidgetMode because of compiling issue on Mac
 int32 UAnimGraphNode_SkeletalControlBase::GetWidgetMode(const USkeletalMeshComponent* SkelComp)
 {
@@ -24,7 +36,7 @@ int32 UAnimGraphNode_SkeletalControlBase::ChangeToNextWidgetMode(const USkeletal
 {
 	return GetWidgetMode(SkelComp);
 }
-// 
+
 FName UAnimGraphNode_SkeletalControlBase::FindSelectedBone()
 {
 	return NAME_None;
diff --git a/Engine/Source/Editor/Persona/Private/AnimationEditorViewportClient.cpp b/Engine/Source/Editor/Persona/Private/AnimationEditorViewportClient.cpp
--- a/Engine/Source/Editor/Persona/Private/AnimationEditorViewportClient.cpp
+++ b/Engine/Source/Editor/Persona/Private/AnimationEditorViewportClient.cpp
@@ -463,7 +463,7 @@ void FAnimationViewportClient::Draw(const FSceneView* View, FPrimitiveDrawInterf
 
 	if ( PersonaPtr.IsValid() && PreviewSkelMeshComp.IsValid() )
 	{
-		// select all nodes, and if skeletalcontrol, allow them to draw 
+		// Allow selected nodes to draw debug rendering if they support it
 		const FGraphPanelSelectionSet SelectedNodes = PersonaPtr.Pin()->GetSelectedNodes();
 		for (FGraphPanelSelectionSet::TConstIterator NodeIt(SelectedNodes); NodeIt; ++NodeIt)
 		{
@@ -602,6 +602,9 @@ void FAnimationViewportClient::FindSelectedAnimGraphNode()
 	// finding a selected node
 	if (PersonaPtr.IsValid() && PreviewSkelMeshComp.IsValid())
 	{
+		USkeletalMeshComponent* PreviewComponent = PreviewSkelMeshComp.Get();
+		check(PreviewComponent);
+
 		const FGraphPanelSelectionSet SelectedNodes = PersonaPtr.Pin()->GetSelectedNodes();
 
 		// don't support multi-selection
@@ -622,23 +625,18 @@ void FAnimationViewportClient::FindSelectedAnimGraphNode()
 					// when selected first after AnimGraph is opened, assign previous data to current node
 					if (SelectedSkelControlAnimGraph != AnimGraphNode)
 					{
+
 						if (SelectedSkelControlAnimGraph.IsValid())
 						{
-							SelectedSkelControlAnimGraph->DeselectActor(PreviewSkelMeshComp.Get());
+							SelectedSkelControlAnimGraph->DeselectActor(PreviewComponent);
 						}
 
 						// make same values to ensure data consistency
 						AnimGraphNode->CopyNodeDataTo(AnimNode);
 
-						WidgetMode = (FWidget::EWidgetMode)AnimGraphNode->GetWidgetMode(PreviewSkelMeshComp.Get());
-						if (WidgetMode == FWidget::WM_Scale)
-						{
-							SetWidgetCoordSystemSpace(COORD_Local);
-						}
-						else
-						{
-							SetWidgetCoordSystemSpace(COORD_World);
-						}
+						WidgetMode = (FWidget::EWidgetMode)AnimGraphNode->GetWidgetMode(PreviewComponent);
+						ECoordSystem DesiredCoordSystem = (ECoordSystem)AnimGraphNode->GetWidgetCoordinateSystem(PreviewComponent);
+						SetWidgetCoordSystemSpace(DesiredCoordSystem);
 					}
 					else
 					{
@@ -650,20 +648,20 @@ void FAnimationViewportClient::FindSelectedAnimGraphNode()
 						}
 					}
 
-					AnimGraphNode->MoveSelectActorLocation(PreviewSkelMeshComp.Get(), AnimNode);
+					AnimGraphNode->MoveSelectActorLocation(PreviewComponent, AnimNode);
 					SelectedSkelControlAnimGraph = AnimGraphNode;
 				}
 			}
 		}
-	}
 
-	if (!bSelected)
-	{
-		if (SelectedSkelControlAnimGraph.IsValid())
+		if (!bSelected)
 		{
-			SelectedSkelControlAnimGraph->DeselectActor(PreviewSkelMeshComp.Get());
+			if (SelectedSkelControlAnimGraph.IsValid())
+			{
+				SelectedSkelControlAnimGraph->DeselectActor(PreviewComponent);
+			}
+			SelectedSkelControlAnimGraph.Reset();
 		}
-		SelectedSkelControlAnimGraph.Reset();
 	}
 }
 
@@ -1543,9 +1541,9 @@ FMatrix FAnimationViewportClient::GetWidgetCoordSystem() const
 			int32 BoneIndex = PreviewSkelMeshComp->GetBoneIndex(BoneName);
 			if (BoneIndex != INDEX_NONE)
 			{
-			FTransform BoneMatrix = PreviewSkelMeshComp->GetBoneTransform(BoneIndex);
-			return BoneMatrix.ToMatrixNoScale().RemoveTranslation();
-		}
+				FTransform BoneMatrix = PreviewSkelMeshComp->GetBoneTransform(BoneIndex);
+				return BoneMatrix.ToMatrixNoScale().RemoveTranslation();
+			}
 		}
 		else if ( PreviewSkelMeshComp->BonesOfInterest.Num() > 0 )
 		{
