Stopped notify panel triggering a Persona UI refresh when changing any property on any UObject. Filtered to only do the refresh when we actually changed a notify.

diff --git a/Engine/Source/Editor/Persona/Private/SAnimNotifyPanel.cpp b/Engine/Source/Editor/Persona/Private/SAnimNotifyPanel.cpp
--- a/Engine/Source/Editor/Persona/Private/SAnimNotifyPanel.cpp
+++ b/Engine/Source/Editor/Persona/Private/SAnimNotifyPanel.cpp
@@ -4261,6 +4261,13 @@ void SAnimNotifyPanel::OnPasteNotifies(SAnimNotifyTrack* RequestTrack, float Cli
 
 void SAnimNotifyPanel::OnPropertyChanged(UObject* ChangedObject, FPropertyChangedEvent& PropertyEvent)
 {
+	// Bail if it isn't a notify
+	if(!ChangedObject->GetClass()->IsChildOf(UAnimNotify::StaticClass()) &&
+	   !ChangedObject->GetClass()->IsChildOf(UAnimNotifyState::StaticClass()))
+	{
+		return;
+	}
+
 	// Don't process if it's an interactive change; wait till we receive the final event.
 	if(PropertyEvent.ChangeType != EPropertyChangeType::Interactive)
 	{
