Added key handling to anim segment panel and hooked up delete to remove selected segments.

diff --git a/Engine/Source/Editor/Persona/Private/SAnimSegmentsPanel.cpp b/Engine/Source/Editor/Persona/Private/SAnimSegmentsPanel.cpp
--- a/Engine/Source/Editor/Persona/Private/SAnimSegmentsPanel.cpp
+++ b/Engine/Source/Editor/Persona/Private/SAnimSegmentsPanel.cpp
@@ -26,8 +26,14 @@ void SAnimSegmentsPanel::Construct(const FArguments& InArgs)
 	OnPostAnimUpdateDelegate		 = InArgs._OnPostAnimUpdate;
 	OnAnimSegmentRemovedDelegate	 = InArgs._OnAnimSegmentRemoved;
 
+	// Register and bind ui commands
+	FAnimSegmentsPanelCommands::Register();
+	BindCommands();
+
+	// Empty out current widget array
+	TrackWidgets.Empty();
+
 	// Animation Segment tracks
-	TArray<TSharedPtr<STrack>> AnimSTracks;
 	TArray<TSharedPtr<STrackNode>> AnimNodes;
 
 	FLinearColor SelectedColor = FLinearColor(1.0f,0.65,0.0f);
@@ -65,13 +71,13 @@ void SAnimSegmentsPanel::Construct(const FArguments& InArgs)
 				.OnTrackDragDrop( this, &SAnimSegmentsPanel::OnTrackDragDrop )
 			];
 
-		AnimSTracks.Add(AnimSegmentTrack);
+		TrackWidgets.Add(AnimSegmentTrack);
 	}
 
 	// Generate Nodes and map them to tracks
 	for ( int32 SegmentIdx=0; SegmentIdx < AnimTrack->AnimSegments.Num(); SegmentIdx++ )
 	{
-		AnimSTracks[ SegmentIdx % AnimSTracks.Num() ]->AddTrackNode(
+		TrackWidgets[ SegmentIdx % TrackWidgets.Num() ]->AddTrackNode(
 			SNew(STrackNode)
 			.ViewInputMax(this->ViewInputMax)
 			.ViewInputMin(this->ViewInputMin)
@@ -240,4 +246,50 @@ void SAnimSegmentsPanel::OnAnimSegmentNodeClicked(int32 SegmentIdx)
 {
 	OnAnimSegmentNodeClickedDelegate.ExecuteIfBound(SegmentIdx);
 }
+
+void SAnimSegmentsPanel::RemoveSelectedAnimSegments()
+{
+	TArray<int32> SelectedNodeIndices;
+	for(int i = 0 ; i < TrackWidgets.Num() ; ++i)
+	{
+		TSharedPtr<STrack> Track = TrackWidgets[i];
+		Track->GetSelectedNodeIndices(SelectedNodeIndices);
+
+		// Reverse order to preserve indices
+		for(int32 j = SelectedNodeIndices.Num() - 1 ; j >= 0 ; --j)
+		{
+			// Segments are placed on one of two tracks with the first segment always residing
+			// in track 0 - need to modify the index from track and index to data index.
+			int32 ModifiedIndex = i + 2 * SelectedNodeIndices[j];
+			RemoveAnimSegment(ModifiedIndex);
+		}
+	}
+}
+
+void SAnimSegmentsPanel::BindCommands()
+{
+	check(!UICommandList.IsValid());
+
+	UICommandList = MakeShareable(new FUICommandList);
+	const FAnimSegmentsPanelCommands& Commands = FAnimSegmentsPanelCommands::Get();
+
+	UICommandList->MapAction(
+		Commands.DeleteSegment,
+		FExecuteAction::CreateSP(this, &SAnimSegmentsPanel::RemoveSelectedAnimSegments));
+}
+
+FReply SAnimSegmentsPanel::OnKeyDown(const FGeometry& MyGeometry, const FKeyEvent& InKeyEvent)
+{
+	if(UICommandList->ProcessCommandBindings(InKeyEvent))
+	{
+		return FReply::Handled();
+	}
+	return FReply::Unhandled();
+}
+
+void FAnimSegmentsPanelCommands::RegisterCommands()
+{
+	UI_COMMAND(DeleteSegment, "Delete", "Deletes the selected segment", EUserInterfaceActionType::Button, FInputChord(EKeys::Platform_Delete));
+}
+
 #undef LOCTEXT_NAMESPACE
diff --git a/Engine/Source/Editor/Persona/Private/SAnimSegmentsPanel.h b/Engine/Source/Editor/Persona/Private/SAnimSegmentsPanel.h
--- a/Engine/Source/Editor/Persona/Private/SAnimSegmentsPanel.h
+++ b/Engine/Source/Editor/Persona/Private/SAnimSegmentsPanel.h
@@ -15,6 +15,21 @@ DECLARE_DELEGATE_OneParam( FOnAnimSegmentRemoved, int32 )
 //////////////////////////////////////////////////////////////////////////
 // SAnimSegmentsPanel
 
+// UI Commands
+class FAnimSegmentsPanelCommands : public TCommands < FAnimSegmentsPanelCommands >
+{
+public:
+	FAnimSegmentsPanelCommands()
+		: TCommands<FAnimSegmentsPanelCommands>("AnimMontagePanel", NSLOCTEXT("Contexts", "AnimMontagePanel", "Anim Montage Panel"), NAME_None, FEditorStyle::GetStyleSetName())
+	{
+
+	}
+
+	TSharedPtr<FUICommandInfo> DeleteSegment;
+
+	virtual void RegisterCommands() override;
+};
+
 // TODO: comment about segment panel
 class SAnimSegmentsPanel: public SCompoundWidget
 {
@@ -54,6 +69,12 @@ public:
 
 	void Construct(const FArguments& InArgs);
 
+	// We support keyboard focus to recieve UI key events
+	bool SupportsKeyboardFocus() const override {return true;}
+
+	// Key handler
+	FReply OnKeyDown(const FGeometry& MyGeometry, const FKeyEvent& InKeyEvent) override;
+
 private:
 
 	FOnPreAnimUpdate			OnPreAnimUpdateDelegate;
@@ -82,12 +103,21 @@ private:
 	void				SummonSegmentNodeContextMenu( FMenuBuilder& MenuBuilder, int32 AnimSegmentIndex );
 
 	void				AddAnimSegment(UAnimSequence *NewSequence, float NewStartPos );
-	void				RemoveAnimSegment(int32 AnimSegmentIndex);
+	
 	bool				DoesAnimTypeMatchTrack(UAnimSequence* NewSequence);
 
 	void				OnTrackDragDrop( TSharedPtr<FDragDropOperation> DragDropOp, float DataPos );
 	void				OnAnimSegmentNodeClicked(int32 SegmentIdx);
 
+	// Remove all selected anim segments in all segment tracks
+	void RemoveSelectedAnimSegments();
+
+	// Remove the specified anim segment
+	void RemoveAnimSegment(int32 AnimSegmentIndex);
+
+	/** Bind UI commands for this widget */
+	void				BindCommands();
+
 	ETrackViewStyle		TrackStyle;
 
 	TAttribute<float>	ViewInputMin;
@@ -96,4 +126,10 @@ private:
 	struct FAnimTrack*	AnimTrack;
 
 	bool bDragging;
+
+	/** List of UI commands for this widget */
+	TSharedPtr<FUICommandList> UICommandList;
+
+	/** List of widgets representing tracks */
+	TArray<TSharedPtr<STrack>> TrackWidgets;
 };
\ No newline at end of file
diff --git a/Engine/Source/Editor/Persona/Private/STrack.cpp b/Engine/Source/Editor/Persona/Private/STrack.cpp
--- a/Engine/Source/Editor/Persona/Private/STrack.cpp
+++ b/Engine/Source/Editor/Persona/Private/STrack.cpp
@@ -857,4 +857,17 @@ void STrack::ClearTrack()
 	TrackNodes.Empty();
 }
 
+void STrack::GetSelectedNodeIndices(TArray<int32>& OutIndices)
+{
+	OutIndices.Empty();
+	for(int32 i = 0 ; i < TrackNodes.Num() ; ++i)
+	{
+		TSharedRef<STrackNode> Node = TrackNodes[i];
+		if(Node->IsSelected())
+		{
+			OutIndices.Add(i);
+		}
+	}
+}
+
 #undef LOCTEXT_NAMESPACE
diff --git a/Engine/Source/Editor/Persona/Private/STrack.h b/Engine/Source/Editor/Persona/Private/STrack.h
--- a/Engine/Source/Editor/Persona/Private/STrack.h
+++ b/Engine/Source/Editor/Persona/Private/STrack.h
@@ -267,6 +267,8 @@ public:
 	FVector2D			ComputeDesiredSize(float) const override;
 	virtual FChildren*	GetChildren() override;
 
+	void GetSelectedNodeIndices(TArray<int32>& OutIndices);
+
 protected:
 	
 	bool									GetDraggableBarSnapPosition(const FGeometry& MyGeometry, float &OutPosition) const;
