#ANIM: Added toolbar for preview/show ref pose

- Fixed ref pose to be edit mode

diff --git a/Engine/Source/Editor/EditorStyle/Private/SlateEditorStyle.cpp b/Engine/Source/Editor/EditorStyle/Private/SlateEditorStyle.cpp
--- a/Engine/Source/Editor/EditorStyle/Private/SlateEditorStyle.cpp
+++ b/Engine/Source/Editor/EditorStyle/Private/SlateEditorStyle.cpp
@@ -4762,8 +4762,6 @@ void FSlateEditorStyle::FStyle::SetupPersonaStyle()
 			);
 
 		// persona commands
-		Set("Persona.ChangeSkeletonPreviewMesh", new IMAGE_BRUSH(TEXT("Icons/icon_Persona_PreviewMesh_40x"), Icon40x40));
-		Set("Persona.ChangeSkeletonPreviewMesh.Small", new IMAGE_BRUSH(TEXT("Icons/icon_Persona_PreviewMesh_40x"), Icon20x20));
 		Set("Persona.AnimNotifyWindow", new IMAGE_BRUSH(TEXT("Icons/icon_Persona_AnimNotift_40x"), Icon40x40));
 		Set("Persona.AnimNotifyWindow.Small", new IMAGE_BRUSH(TEXT("Icons/icon_Persona_AnimNotift_40x"), Icon20x20));
 		Set("Persona.RetargetManager", new IMAGE_BRUSH(TEXT("Icons/icon_Persona_Retarget_40x"), Icon40x40));
@@ -4797,6 +4795,12 @@ void FSlateEditorStyle::FStyle::SetupPersonaStyle()
 		Set("Persona.ApplyAnimation", new IMAGE_BRUSH(TEXT("Icons/icon_Persona_BakeAnim_40x"), Icon40x40));
 		Set("Persona.ApplyAnimation.Small", new IMAGE_BRUSH(TEXT("Icons/icon_Persona_BakeAnim_40x"), Icon20x20));
 
+		// preview set up
+		Set("Persona.TogglePreviewAsset", new IMAGE_BRUSH(TEXT("Icons/icon_Persona_PreviewAsset_40x"), Icon40x40));
+		Set("Persona.TogglePreviewAsset.Small", new IMAGE_BRUSH(TEXT("Icons/icon_Persona_PreviewAsset_40x"), Icon20x20));
+		Set("Persona.ToggleReferencePose", new IMAGE_BRUSH(TEXT("Icons/icon_Persona_ReferencePose_40x"), Icon40x40));
+		Set("Persona.ToggleReferencePose.Small", new IMAGE_BRUSH(TEXT("Icons/icon_Persona_ReferencePose_40x"), Icon20x20));
+
 		// persona extras
 		Set("Persona.ConvertAnimationGraph", new IMAGE_BRUSH("Old/Graph/ConvertIcon", Icon40x40));
 		Set("Persona.ReimportAsset", new IMAGE_BRUSH("Icons/Reimport_12x", Icon12x12));
diff --git a/Engine/Source/Editor/Persona/Private/Persona.cpp b/Engine/Source/Editor/Persona/Private/Persona.cpp
--- a/Engine/Source/Editor/Persona/Private/Persona.cpp
+++ b/Engine/Source/Editor/Persona/Private/Persona.cpp
@@ -950,7 +950,9 @@ void FPersona::ExtendDefaultPersonaToolbar()
 		{
 			ToolbarBuilder.BeginSection("Skeleton");
 			{
-				ToolbarBuilder.AddToolBarButton(FPersonaCommands::Get().ChangeSkeletonPreviewMesh, NAME_None, LOCTEXT("Toolbar_ChangePreviewMesh", "Set Preview"));
+				ToolbarBuilder.AddToolBarButton(FPersonaCommands::Get().TogglePreviewAsset, NAME_None, LOCTEXT("Toolbar_PreviewAsset", "Preview"), TAttribute<FText>(PersonaPtr, &FPersona::GetPreviewAssetTooltip));
+				ToolbarBuilder.AddToolBarButton(FPersonaCommands::Get().ToggleReferencePose, NAME_None, LOCTEXT("Toolbar_ToggleReferencePose", "Ref Pose"), LOCTEXT("Toolbar_ToggleReferencePoseTooltip", "Show Reference Pose"));
+				ToolbarBuilder.AddSeparator();
 				ToolbarBuilder.AddToolBarButton(FPersonaCommands::Get().AnimNotifyWindow);
 				ToolbarBuilder.AddToolBarButton(FPersonaCommands::Get().RetargetManager, NAME_None, LOCTEXT("Toolbar_RetargetManager", "Retarget Manager"));
 				ToolbarBuilder.AddToolBarButton(FPersonaCommands::Get().ImportMesh);
@@ -1248,6 +1250,18 @@ void FPersona::CreateDefaultCommands()
 		FCanExecuteAction::CreateSP( this, &FPersona::CanChangeSkeletonPreviewMesh )
 		);
 
+	ToolkitCommands->MapAction(FPersonaCommands::Get().ToggleReferencePose,
+	FExecuteAction::CreateSP(this, &FPersona::ShowReferencePose, true),
+	FCanExecuteAction::CreateSP(this, &FPersona::CanShowReferencePose),
+	FIsActionChecked::CreateSP(this, &FPersona::IsShowReferencePoseEnabled)
+	);
+
+	ToolkitCommands->MapAction(FPersonaCommands::Get().TogglePreviewAsset,
+	FExecuteAction::CreateSP(this, &FPersona::ShowReferencePose, false),
+	FCanExecuteAction::CreateSP(this, &FPersona::CanPreviewAsset),
+	FIsActionChecked::CreateSP(this, &FPersona::IsPreviewAssetEnabled)
+	);
+
 	ToolkitCommands->MapAction( FPersonaCommands::Get().RemoveUnusedBones,
 		FExecuteAction::CreateSP( this, &FPersona::RemoveUnusedBones ),
 		FCanExecuteAction::CreateSP( this, &FPersona::CanRemoveBones )
@@ -3194,6 +3208,95 @@ void FPersona::OnBlueprintChangedImpl(UBlueprint* InBlueprint, bool bIsJustBeing
 	}
 }
 
+void FPersona::ShowReferencePose(bool bReferencePose)
+{
+	if(PreviewComponent)
+	{
+		if(bReferencePose == false)
+		{
+			if(IsInPersonaMode(FPersonaModes::AnimBlueprintEditMode))
+			{
+				PreviewComponent->EnablePreview(false, NULL, NULL);
+
+				UAnimBlueprint* AnimBP = GetAnimBlueprint();
+				if(AnimBP)
+				{
+					PreviewComponent->SetAnimInstanceClass(AnimBP->GeneratedClass);
+				}
+			}
+			else
+			{
+				UObject* PreviewAsset = CachedPreviewAsset.IsValid()? CachedPreviewAsset.Get() : (GetAnimationAssetBeingEdited());
+				PreviewComponent->EnablePreview(true, Cast<UAnimationAsset>(PreviewAsset), NULL);
+			}
+		}
+		else
+		{
+			if (PreviewComponent->PreviewInstance && PreviewComponent->PreviewInstance->CurrentAsset)
+			{
+				CachedPreviewAsset = PreviewComponent->PreviewInstance->CurrentAsset;
+			}
+			
+			PreviewComponent->EnablePreview(true, NULL, NULL);
+		}
+	}
+}
+
+bool FPersona::CanShowReferencePose() const
+{
+	return PreviewComponent != NULL;
+}
+
+bool FPersona::IsShowReferencePoseEnabled() const
+{
+	if(PreviewComponent)
+	{
+		return PreviewComponent->IsPreviewOn() && PreviewComponent->PreviewInstance->CurrentAsset == NULL;
+	}
+	return false;
+}
+
+bool FPersona::CanPreviewAsset() const
+{
+	return CanShowReferencePose();
+}
+
+bool FPersona::IsPreviewAssetEnabled() const
+{
+	return IsShowReferencePoseEnabled() == false;
+}
+
+FText FPersona::GetPreviewAssetTooltip() const
+{
+	// if already looking at ref pose
+	if(IsShowReferencePoseEnabled())
+	{
+		FString AssetName = TEXT("None Available. Please select asset to preview.");
+
+		if(IsInPersonaMode(FPersonaModes::AnimBlueprintEditMode))
+		{
+			UAnimBlueprint* AnimBP = GetAnimBlueprint();
+			if(AnimBP)
+			{
+				AssetName = AnimBP->GetName();
+			}
+		}
+		else
+		{
+			UObject* PreviewAsset = CachedPreviewAsset.IsValid()? CachedPreviewAsset.Get() : (GetAnimationAssetBeingEdited());
+			if (PreviewAsset)
+			{
+				AssetName = PreviewAsset->GetName();
+			}
+		}
+		return FText::FromString(FString::Printf(TEXT("Preview %s"), *AssetName));
+	}
+	else
+	{
+		return FText::FromString(FString::Printf(TEXT("Currently previewing %s"), *PreviewComponent->GetPreviewText()));
+	}
+}
+
 static class FMeshHierarchyCmd : private FSelfRegisteringExec
 {
 public:
diff --git a/Engine/Source/Editor/Persona/Private/Persona.h b/Engine/Source/Editor/Persona/Private/Persona.h
--- a/Engine/Source/Editor/Persona/Private/Persona.h
+++ b/Engine/Source/Editor/Persona/Private/Persona.h
@@ -209,7 +209,16 @@ public:
 
 	void RecompileAnimBlueprintIfDirty();
 
+	/* Reference Pose Handler */
+	bool IsShowReferencePoseEnabled() const;
+	bool CanShowReferencePose() const;
+	void ShowReferencePose(bool bReferencePose);
+
 protected:
+	bool IsPreviewAssetEnabled() const;
+	bool CanPreviewAsset() const;
+	FText GetPreviewAssetTooltip() const;
+
 	// FBlueprintEditor interface
 	//virtual void CreateDefaultToolbar() override;
 	virtual void CreateDefaultCommands() override;
@@ -272,6 +281,7 @@ protected:
 
 protected:
 	USkeleton* TargetSkeleton;
+	TWeakObjectPtr<UObject> CachedPreviewAsset;
 
 public:
 	class UDebugSkelMeshComponent* PreviewComponent;
@@ -602,6 +612,7 @@ private:
 	void OnExportToFBX();
 	void OnAddLoopingInterpolation();
 	bool HasValidAnimationSequencePlaying() const;
+	/** Return true if currently in the given mode */
 	bool IsInPersonaMode(const FName InPersonaMode) const;
 
 	/** Animation Editing Features **/
diff --git a/Engine/Source/Editor/Persona/Private/PersonaCommands.cpp b/Engine/Source/Editor/Persona/Private/PersonaCommands.cpp
--- a/Engine/Source/Editor/Persona/Private/PersonaCommands.cpp
+++ b/Engine/Source/Editor/Persona/Private/PersonaCommands.cpp
@@ -35,6 +35,9 @@ void FPersonaCommands::RegisterCommands()
 	UI_COMMAND( ApplyAnimation, "Apply Animation", "Apply Additive Layer Tracks to Runtime Animation Data", EUserInterfaceActionType::Button, FInputChord() );
 
 	UI_COMMAND( UpdateSkeletonRefPose, "Update Skeleton RefPose", "Update Skeleton ref pose based on current preview mesh", EUserInterfaceActionType::Button, FInputChord() );
+
+	UI_COMMAND( ToggleReferencePose, "Toggle Reference Pose", "Toggle Reference Pose", EUserInterfaceActionType::ToggleButton, FInputChord() );
+	UI_COMMAND( TogglePreviewAsset, "Toggle Preview ASset", "Toggle Preview Asset", EUserInterfaceActionType::ToggleButton, FInputChord() );
 }
 
 #undef LOCTEXT_NAMESPACE
\ No newline at end of file
diff --git a/Engine/Source/Editor/Persona/Private/PersonaCommands.h b/Engine/Source/Editor/Persona/Private/PersonaCommands.h
--- a/Engine/Source/Editor/Persona/Private/PersonaCommands.h
+++ b/Engine/Source/Editor/Persona/Private/PersonaCommands.h
@@ -65,4 +65,8 @@ public:
 	TSharedPtr<FUICommandInfo> CreateAanimMontage;
 	// Command to allow users to remove unused bones (not referenced by any skeletalmesh) from the skeleton
 	TSharedPtr<FUICommandInfo> UpdateSkeletonRefPose;
+
+	// Toggle Reference Pose
+	TSharedPtr<FUICommandInfo> ToggleReferencePose;
+	TSharedPtr<FUICommandInfo> TogglePreviewAsset;
 };
diff --git a/Engine/Source/Editor/Persona/Private/SAnimationEditorViewport.cpp b/Engine/Source/Editor/Persona/Private/SAnimationEditorViewport.cpp
--- a/Engine/Source/Editor/Persona/Private/SAnimationEditorViewport.cpp
+++ b/Engine/Source/Editor/Persona/Private/SAnimationEditorViewport.cpp
@@ -1107,27 +1107,17 @@ bool SAnimationEditorViewportTabBody::IsPlaybackSpeedSelected(int32 PlaybackSpee
 
 void SAnimationEditorViewportTabBody::ShowReferencePose()
 {
-	UDebugSkelMeshComponent* PreviewComponent = PersonaPtr.Pin()->PreviewComponent;
-	if(PreviewComponent)
-	{
-		PreviewComponent->bForceRefpose = !PreviewComponent->bForceRefpose;
-	}
+	PersonaPtr.Pin()->ShowReferencePose(IsShowReferencePoseEnabled() == false);
 }
 
 bool SAnimationEditorViewportTabBody::CanShowReferencePose() const
 {
-	UDebugSkelMeshComponent* PreviewComponent = PersonaPtr.Pin()->PreviewComponent;
-	return PreviewComponent != NULL;
+	return PersonaPtr.Pin()->CanShowReferencePose();
 }
 
 bool SAnimationEditorViewportTabBody::IsShowReferencePoseEnabled() const
 {
-	UDebugSkelMeshComponent* PreviewComponent = PersonaPtr.Pin()->PreviewComponent;
-	if(PreviewComponent)
-	{
-		return PreviewComponent->bForceRefpose;
-	}
-	return false;
+	return PersonaPtr.Pin()->IsShowReferencePoseEnabled();
 }
 
 void SAnimationEditorViewportTabBody::ShowRetargetBasePose()
diff --git a/Engine/Source/Editor/UnrealEd/Private/Animation/DebugSkelMeshComponent.cpp b/Engine/Source/Editor/UnrealEd/Private/Animation/DebugSkelMeshComponent.cpp
--- a/Engine/Source/Editor/UnrealEd/Private/Animation/DebugSkelMeshComponent.cpp
+++ b/Engine/Source/Editor/UnrealEd/Private/Animation/DebugSkelMeshComponent.cpp
@@ -247,7 +247,7 @@ FString UDebugSkelMeshComponent::GetPreviewText() const
 
 	if (IsPreviewOn())
 	{
-		if (UBlendSpace* BlendSpace = Cast<UBlendSpace>(PreviewInstance->CurrentAsset))
+		if (UBlendSpaceBase* BlendSpace = Cast<UBlendSpaceBase>(PreviewInstance->CurrentAsset))
 		{
 			return FText::Format( LOCTEXT("BlendSpace", "Blend Space {0}"), FText::FromString(BlendSpace->GetName()) ).ToString();
 		}
@@ -255,6 +255,10 @@ FString UDebugSkelMeshComponent::GetPreviewText() const
 		{
 			return FText::Format( LOCTEXT("Montage", "Montage {0}"), FText::FromString(Montage->GetName()) ).ToString();
 		}
+		else if(UAnimComposite* Composite = Cast<UAnimComposite>(PreviewInstance->CurrentAsset))
+		{
+			return FText::Format(LOCTEXT("Composite", "Composite {0}"), FText::FromString(Composite->GetName())).ToString();
+		}
 		else if (UAnimSequence* Sequence = Cast<UAnimSequence>(PreviewInstance->CurrentAsset))
 		{
 			return FText::Format( LOCTEXT("Animation", "Animation {0}"), FText::FromString(Sequence->GetName()) ).ToString();
@@ -265,7 +269,7 @@ FString UDebugSkelMeshComponent::GetPreviewText() const
 		}
 	}
 
-	return LOCTEXT("None", "None").ToString();
+	return LOCTEXT("ReferencePose", "Reference Pose").ToString();
 
 #undef LOCTEXT_NAMESPACE
 }
