#Anim Recording: Gives default path and name

diff --git a/Engine/Source/Editor/Persona/Private/SAnimationDlgs.cpp b/Engine/Source/Editor/Persona/Private/SAnimationDlgs.cpp
--- a/Engine/Source/Editor/Persona/Private/SAnimationDlgs.cpp
+++ b/Engine/Source/Editor/Persona/Private/SAnimationDlgs.cpp
@@ -5,6 +5,7 @@
 #include "Persona.h"
 #include "SAnimationDlgs.h"
 #include "Editor/ContentBrowser/Public/ContentBrowserModule.h"
+#include "Developer/AssetTools/Public/AssetToolsModule.h"
 
 #define LOCTEXT_NAMESPACE "SAnimationDlgs"
 
@@ -28,12 +29,28 @@ void SCreateAnimationDlg::Construct(const FArguments& InArgs)
 	if (AssetPath.IsEmpty())
 	{
 		AssetPath = LastUsedAssetPath;
+		// still empty?
+		if (AssetPath.IsEmpty())
+		{
+			AssetPath = FText::FromString(TEXT("/Game"));
+		}
 	}
 	else
 	{
 		LastUsedAssetPath = AssetPath;
 	}
 
+	if (AssetName.IsEmpty())
+	{
+		// find default name for them
+		FAssetToolsModule& AssetToolsModule = FModuleManager::Get().LoadModuleChecked<FAssetToolsModule>("AssetTools");
+		FString OutPackageName, OutAssetName;
+		FString PackageName = AssetPath.ToString() + TEXT("/NewAnimation");
+
+		AssetToolsModule.Get().CreateUniqueAssetName(PackageName,  TEXT(""), OutPackageName, OutAssetName);
+		AssetName = FText::FromString(OutAssetName);
+	}
+
 	FPathPickerConfig PathPickerConfig;
 	PathPickerConfig.DefaultPath = AssetPath.ToString();
 	PathPickerConfig.OnPathSelected = FOnPathSelected::CreateSP(this, &SCreateAnimationDlg::OnPathChange);
