Use localised class name for new native notify menu display

diff --git a/Engine/Source/Editor/Persona/Private/SAnimNotifyPanel.cpp b/Engine/Source/Editor/Persona/Private/SAnimNotifyPanel.cpp
--- a/Engine/Source/Editor/Persona/Private/SAnimNotifyPanel.cpp
+++ b/Engine/Source/Editor/Persona/Private/SAnimNotifyPanel.cpp
@@ -2000,8 +2000,8 @@ void SAnimNotifyTrack::FillNewNotifyStateMenu(FMenuBuilder& MenuBuilder, bool bI
 		for(UClass* Class : NotifyStateClasses)
 		{
 			const FText Description = LOCTEXT("NewNotifyStateSubMenu_NativeToolTip", "Add an existing native notify state");
-			FString Label = Class->GetDisplayNameText().ToString();
-			const FText LabelText = FText::FromString(Label);
+			const FText LabelText = Class->GetDisplayNameText();
+			const FString Label = LabelText.ToString();
 
 			FUIAction UIAction;
 			if (!bIsReplaceWithMenu)
@@ -2066,8 +2066,8 @@ void SAnimNotifyTrack::FillNewNotifyMenu(FMenuBuilder& MenuBuilder, bool bIsRepl
 	{
 		for(UClass* Class : NativeNotifyClasses)
 		{
-			FString Label = Class->GetName();
-			const FText LabelText = FText::FromString(Label);
+			const FText LabelText = Class->GetDisplayNameText();
+			const FString Label = LabelText.ToString();
 
 			FUIAction UIAction;
 			FText Description = FText::GetEmpty();
