Fix for crash accessing property nodes from Persona preview instance. Now force initialise animation when the blueprint changes to refresh the script instance.

diff --git a/Engine/Source/Editor/Persona/Private/Persona.cpp b/Engine/Source/Editor/Persona/Private/Persona.cpp
--- a/Engine/Source/Editor/Persona/Private/Persona.cpp
+++ b/Engine/Source/Editor/Persona/Private/Persona.cpp
@@ -3189,11 +3189,9 @@ void FPersona::OnBlueprintChangedImpl(UBlueprint* InBlueprint, bool bIsJustBeing
 
 	if(PreviewComponent != NULL)
 	{
-		if(PreviewComponent->AnimScriptInstance == NULL)
-		{
-			// try reinitialize animation if it doesn't exist
-			PreviewComponent->InitAnim(true);
-		}
+		// Reinitialize the animation, anything we reference could have changed triggering
+		// the blueprint change
+		PreviewComponent->InitAnim(true);
 
 		if(bIsDebuggingPreview)
 		{
