Fixed crash when copying notifies with an invalid slot index when referenced in the destination montage
Fixed crash when accessing notifies that reference a slot that has been removed

diff --git a/Engine/Source/Editor/Persona/Private/SAnimNotifyPanel.cpp b/Engine/Source/Editor/Persona/Private/SAnimNotifyPanel.cpp
--- a/Engine/Source/Editor/Persona/Private/SAnimNotifyPanel.cpp
+++ b/Engine/Source/Editor/Persona/Private/SAnimNotifyPanel.cpp
@@ -3217,13 +3217,23 @@ void SAnimNotifyTrack::PasteSingleNotify(FString& NotifyString, float PasteTime)
 
 		FAnimNotifyEvent& NewNotify = Sequence->Notifies[NewIdx];
 
-		if(PasteTime != -1.0f)
+		// We have to link to the montage / sequence again, we need a correct time set and we could be pasting to a new montage / sequence
+		int32 NewSlotIndex = 0;
+		float NewNotifyTime = PasteTime != 1.0f ? PasteTime : NewNotify.GetTime();
+		NewNotifyTime = FMath::Clamp(NewNotifyTime, 0.0f, Sequence->SequenceLength);
+
+		if(UAnimMontage* Montage = Cast<UAnimMontage>(Sequence))
 		{
-			NewNotify.SetTime(PasteTime);
+			// We have a montage, validate slots
+			int32 OldSlotIndex = NewNotify.GetSlotIndex();
+			if(Montage->SlotAnimTracks.IsValidIndex(OldSlotIndex))
+			{
+				// Link to the same slot index
+				NewSlotIndex = OldSlotIndex;
+			}
 		}
+		NewNotify.Link(Sequence, PasteTime, NewSlotIndex);
 
-		// Make sure the notify is within the track area
-		NewNotify.SetTime(FMath::Clamp(NewNotify.GetTime(), 0.0f, Sequence->SequenceLength));
 		NewNotify.TriggerTimeOffset = GetTriggerTimeOffsetForType(Sequence->CalculateOffsetForNotify(NewNotify.GetTime()));
 		NewNotify.TrackIndex = TrackIndex;
 
@@ -3243,8 +3253,6 @@ void SAnimNotifyTrack::PasteSingleNotify(FString& NotifyString, float PasteTime)
 			NewNotify.SetDuration(FMath::Clamp(NewNotify.GetDuration(), 1 / 30.0f, Sequence->SequenceLength - NewNotify.GetTime()));
 			NewNotify.EndTriggerTimeOffset = GetTriggerTimeOffsetForType(Sequence->CalculateOffsetForNotify(NewNotify.GetTime() + NewNotify.GetDuration()));
 		}
-
-		NewNotify.ConditionalRelink();
 	}
 	else
 	{
diff --git a/Engine/Source/Editor/Persona/Private/SMontageEditor.cpp b/Engine/Source/Editor/Persona/Private/SMontageEditor.cpp
--- a/Engine/Source/Editor/Persona/Private/SMontageEditor.cpp
+++ b/Engine/Source/Editor/Persona/Private/SMontageEditor.cpp
@@ -596,6 +596,18 @@ void SMontageEditor::RemoveMontageSlot(int32 AnimSlotIndex)
 		MontageObj->SlotAnimTracks.RemoveAt(AnimSlotIndex);
 		MontageObj->MarkPackageDirty();
 		AnimMontagePanel->Update();
+
+		// Iterate the notifies and relink anything that is now invalid
+		for(FAnimNotifyEvent& Event : MontageObj->Notifies)
+		{
+			Event.ConditionalRelink();
+		}
+
+		// Do the same for sections
+		for(FCompositeSection& Section : MontageObj->CompositeSections)
+		{
+			Section.ConditionalRelink();
+		}
 	}
 }
 
diff --git a/Engine/Source/Runtime/Engine/Classes/Animation/AnimLinkableElement.cpp b/Engine/Source/Runtime/Engine/Classes/Animation/AnimLinkableElement.cpp
--- a/Engine/Source/Runtime/Engine/Classes/Animation/AnimLinkableElement.cpp
+++ b/Engine/Source/Runtime/Engine/Classes/Animation/AnimLinkableElement.cpp
@@ -349,10 +349,28 @@ void FAnimLinkableElement::SetTime_Internal(float NewTime, EAnimLinkMethod::Type
 	}
 }
 
-void FAnimLinkableElement::ConditionalRelink()
+bool FAnimLinkableElement::ConditionalRelink()
 {
+	// Check slot index if we're in a montage
+	bool bRequiresRelink = false;
+	
+	if(LinkedMontage)
+	{
+		if(!LinkedMontage->SlotAnimTracks.IsValidIndex(SlotIndex))
+		{
+			bRequiresRelink = true;
+			SlotIndex = 0;
+		}
+	}
+
+	// Check to see if we've moved to a new segment
 	float CurrentAbsTime = GetTime();
 	if(CurrentAbsTime < SegmentBeginTime || CurrentAbsTime > SegmentBeginTime + SegmentLength)
+	{
+		bRequiresRelink = true;
+	}
+
+	if(bRequiresRelink)
 	{
 		if(LinkedMontage)
 		{
@@ -363,6 +381,8 @@ void FAnimLinkableElement::ConditionalRelink()
 			LinkSequence(LinkedSequence, CurrentAbsTime);
 		}
 	}
+
+	return bRequiresRelink;
 }
 
 void FAnimLinkableElement::Link(UAnimSequenceBase* AnimObject, float AbsTime, int32 InSlotIndex /*= 0*/)
diff --git a/Engine/Source/Runtime/Engine/Classes/Animation/AnimLinkableElement.h b/Engine/Source/Runtime/Engine/Classes/Animation/AnimLinkableElement.h
--- a/Engine/Source/Runtime/Engine/Classes/Animation/AnimLinkableElement.h
+++ b/Engine/Source/Runtime/Engine/Classes/Animation/AnimLinkableElement.h
@@ -118,8 +118,8 @@ struct FAnimLinkableElement
 	 */
 	ENGINE_API void SetSegmentIndex(int32 NewSegmentIndex) {SegmentIndex = NewSegmentIndex;}
 
-	/** Relinks this element if it is outside of its segments */
-	ENGINE_API void ConditionalRelink();
+	/** Relinks this element if internal state requires relinking */
+	ENGINE_API bool ConditionalRelink();
 
 	/** Refreshes the current segment data (Begin time, length etc.) and validate the link time
 	 *  Intended to update the internal state when segment lengths/times could have changed
