Add a "Remove this LOD" button to each LOD (after 0) in Persona.

diff --git a/Engine/Source/Editor/Persona/Private/PersonaMeshDetails.cpp b/Engine/Source/Editor/Persona/Private/PersonaMeshDetails.cpp
--- a/Engine/Source/Editor/Persona/Private/PersonaMeshDetails.cpp
+++ b/Engine/Source/Editor/Persona/Private/PersonaMeshDetails.cpp
@@ -527,6 +527,22 @@ void FPersonaMeshDetails::AddLODLevelCategories(IDetailLayoutBuilder& DetailLayo
 			{
 				LODCategory.AddCustomBuilder(ReductionSettingsWidgets[LODIndex].ToSharedRef());
 			}
+
+			if (LODIndex > 0)
+			{
+				LODCategory.AddCustomRow(LOCTEXT("RemoveLODRow", "Remove LOD"))
+				.ValueContent()
+				.HAlign(HAlign_Left)
+				[
+					SNew(SButton)
+					.OnClicked(this, &FPersonaMeshDetails::RemoveOneLOD, LODIndex)
+					[
+						SNew(STextBlock)
+						.Text(LOCTEXT("RemoveLOD", "Remove this LOD"))
+						.Font(DetailLayout.GetDetailFont())
+					]
+				];
+			}
 		}
 	}
 }
@@ -630,6 +646,22 @@ FReply FPersonaMeshDetails::OnApplyChanges()
 	return FReply::Handled();
 }
 
+FReply FPersonaMeshDetails::RemoveOneLOD(int32 LODIndex)
+{
+	USkeletalMesh* SkelMesh = PersonaPtr->GetMesh();
+	check(SkelMesh);
+	check(SkelMesh->LODInfo.IsValidIndex(LODIndex));
+
+	FSkeletalMeshUpdateContext UpdateContext;
+	UpdateContext.SkeletalMesh = SkelMesh;
+	UpdateContext.AssociatedComponents.Push(PersonaPtr->PreviewComponent);
+
+	FLODUtilities::RemoveLOD(UpdateContext, LODIndex);
+
+	PersonaPtr->PersonaMeshDetailLayout->ForceRefreshDetails();
+	return FReply::Handled();
+}
+
 FText FPersonaMeshDetails::GetApplyButtonText() const
 {
 	if (IsApplyNeeded())	
diff --git a/Engine/Source/Editor/Persona/Private/PersonaMeshDetails.h b/Engine/Source/Editor/Persona/Private/PersonaMeshDetails.h
--- a/Engine/Source/Editor/Persona/Private/PersonaMeshDetails.h
+++ b/Engine/Source/Editor/Persona/Private/PersonaMeshDetails.h
@@ -201,6 +201,8 @@ private:
 	FText GetLODCountTooltip() const;
 	/** apply LOD changes if the user modified LOD reduction settings */
 	FReply OnApplyChanges();
+	/** Removes the specified lod from the skeletal mesh */
+	FReply RemoveOneLOD(int32 LODIndex);
 	/** hide properties which don't need to be showed to end users */
 	void HideUnnecessaryProperties(IDetailLayoutBuilder& DetailLayout);
 	/** clear "None" bones and remove already included bones */
