Graph Editor: Force node titles to be updated when a property on an graph node is modified

diff --git a/Engine/Source/Editor/Persona/Private/AnimGraphNodeDetails.cpp b/Engine/Source/Editor/Persona/Private/AnimGraphNodeDetails.cpp
--- a/Engine/Source/Editor/Persona/Private/AnimGraphNodeDetails.cpp
+++ b/Engine/Source/Editor/Persona/Private/AnimGraphNodeDetails.cpp
@@ -479,19 +479,14 @@ END_SLATE_FUNCTION_BUILD_OPTIMIZATION
 
 void FBoneReferenceCustomization::CustomizeChildren( TSharedRef<IPropertyHandle> StructPropertyHandle, class IDetailChildrenBuilder& ChildBuilder, IPropertyTypeCustomizationUtils& StructCustomizationUtils )
 {
-	// No child customisations as the properties are shown in the header
+	// No child customizations as the properties are shown in the header
 }
 
 void FBoneReferenceCustomization::OnBoneSelectionChanged(FName Name)
 {
 	BoneRefProperty->SetValue(Name);
-
-	// Force a refresh of cached node titles, K2Schema is no different than the AnimSchema for this purpose
-	if (const UEdGraphSchema_K2* K2Schema = GetDefault<UEdGraphSchema_K2>())
-	{
-		K2Schema->ForceVisualizationCacheClear();
-	}
 }
+
 FName FBoneReferenceCustomization::GetSelectedBone() const
 {
 	FString OutText;
diff --git a/Engine/Source/Runtime/Engine/Classes/EdGraph/EdGraphNode.h b/Engine/Source/Runtime/Engine/Classes/EdGraph/EdGraphNode.h
--- a/Engine/Source/Runtime/Engine/Classes/EdGraph/EdGraphNode.h
+++ b/Engine/Source/Runtime/Engine/Classes/EdGraph/EdGraphNode.h
@@ -165,6 +165,7 @@ public:
 	// UObject interface
 	static void AddReferencedObjects(UObject* InThis, FReferenceCollector& Collector);
 	virtual void PostLoad() override;
+	virtual void PostEditChangeProperty(FPropertyChangedEvent& PropertyChangedEvent) override;
 	// End of UObject interface
 
 	/** widget representing this node if it exists; Note: This is not safe to use in general and will be removed in the future, as there is no guarantee that only one graph editor/panel is viewing a given graph */
diff --git a/Engine/Source/Runtime/Engine/Private/EdGraph/EdGraphNode.cpp b/Engine/Source/Runtime/Engine/Private/EdGraph/EdGraphNode.cpp
--- a/Engine/Source/Runtime/Engine/Private/EdGraph/EdGraphNode.cpp
+++ b/Engine/Source/Runtime/Engine/Private/EdGraph/EdGraphNode.cpp
@@ -280,6 +280,13 @@ void UEdGraphNode::PostLoad()
 	}
 }
 
+void UEdGraphNode::PostEditChangeProperty(FPropertyChangedEvent& PropertyChangedEvent)
+{
+	Super::PostEditChangeProperty(PropertyChangedEvent);
+
+	GetSchema()->ForceVisualizationCacheClear();
+}
+
 void UEdGraphNode::CreateNewGuid()
 {
 	NodeGuid = FGuid::NewGuid();
