Fixed blendspace 2D points not allowing a Y param value to match the X value when using the type-in

diff --git a/Engine/Source/Editor/Persona/Private/SAnimationBlendSpaceBase.cpp b/Engine/Source/Editor/Persona/Private/SAnimationBlendSpaceBase.cpp
--- a/Engine/Source/Editor/Persona/Private/SAnimationBlendSpaceBase.cpp
+++ b/Engine/Source/Editor/Persona/Private/SAnimationBlendSpaceBase.cpp
@@ -227,12 +227,12 @@ private:
 		const FBlendSample& Sample = Samples[SampleIndex];
 		float NewValue = FCString::Atof(*NewText.ToString());
 
-		if (NewValue!=Sample.SampleValue.X)
+		FVector SampleVector = Sample.SampleValue;
+		if (NewValue != GetVectorValueForParam(ParamIndex, SampleVector))
 		{
-			FVector NewSampleValue = Sample.SampleValue;
-			GetVectorValueForParam(ParamIndex, NewSampleValue) = NewValue;
+			GetVectorValueForParam(ParamIndex, SampleVector) = NewValue;
 
-			if (BlendSpace->EditSample(Sample, NewSampleValue))
+			if (BlendSpace->EditSample(Sample, SampleVector))
 			{
 				//update editor widget
 				BlendSpaceEditorWidget->ResampleData();
