Improved the behavior of Ctrl+B in the Content Browser and Persona

Ctrl+B in the Content Browser will now invoke the "Show in Folder View" command to allow you to easily sync to the actual location of an asset in the cases where you're viewing it via a search or a collection.

Ctrl+B in the Persona animation sequence browser will now invoke the "Find in Content Browser" command.

Ctrl+B will no longer do anything when Persona has no asset available for a given edit mode.

diff --git a/Engine/Source/Editor/ContentBrowser/Private/AssetContextMenu.cpp b/Engine/Source/Editor/ContentBrowser/Private/AssetContextMenu.cpp
--- a/Engine/Source/Editor/ContentBrowser/Private/AssetContextMenu.cpp
+++ b/Engine/Source/Editor/ContentBrowser/Private/AssetContextMenu.cpp
@@ -11,6 +11,7 @@
 #include "Editor/PropertyEditor/Public/PropertyEditorModule.h"
 #include "Toolkits/AssetEditorManager.h"
 #include "Toolkits/ToolkitManager.h"
+#include "Toolkits/GlobalEditorCommonCommands.h"
 #include "ConsolidateWindow.h"
 #include "ReferenceViewer.h"
 #include "ISizeMapModule.h"
@@ -62,6 +63,11 @@ void FAssetContextMenu::BindCommands(TSharedPtr< FUICommandList >& Commands)
 		FIsActionChecked(),
 		FIsActionButtonVisible::CreateSP(this, &FAssetContextMenu::CanExecuteDuplicate)
 		));
+
+	Commands->MapAction(FGlobalEditorCommonCommands::Get().FindInContentBrowser, FUIAction(
+		FExecuteAction::CreateSP(this, &FAssetContextMenu::ExecuteSyncToAssetTree),
+		FCanExecuteAction::CreateSP(this, &FAssetContextMenu::CanExecuteSyncToAssetTree)
+		));
 }
 
 TSharedRef<SWidget> FAssetContextMenu::MakeContextMenu(const TArray<FAssetData>& InSelectedAssets, const FSourcesData& InSourcesData, TSharedPtr< FUICommandList > InCommandList)
@@ -278,23 +284,19 @@ void FAssetContextMenu::AddExploreMenuOptions(FMenuBuilder& MenuBuilder)
 {
 	MenuBuilder.BeginSection("AssetContextExploreMenuOptions", LOCTEXT("AssetContextExploreMenuOptionsHeading", "Explore"));
 	{
-		// Open Containing Folder
+		// Find in Content Browser
 		MenuBuilder.AddMenuEntry(
-			LOCTEXT("SyncToAssetTree", "Show in Folder View"),
-			LOCTEXT("SyncToAssetTreeTooltip", "Selects the folder that contains this asset in the Content Browser Sources Panel."),
-			FSlateIcon(),
-			FUIAction(
-				FExecuteAction::CreateSP(this, &FAssetContextMenu::ExecuteSyncToAssetTree),
-				FCanExecuteAction::CreateSP(this, &FAssetContextMenu::CanExecuteSyncToAssetTree)
-				)
+			FGlobalEditorCommonCommands::Get().FindInContentBrowser, 
+			NAME_None, 
+			LOCTEXT("ShowInFolderView", "Show in Folder View"),
+			LOCTEXT("ShowInFolderViewTooltip", "Selects the folder that contains this asset in the Content Browser Sources Panel.")
 			);
 
-
 		// Find in Explorer
 		MenuBuilder.AddMenuEntry(
 			ContentBrowserUtils::GetExploreFolderText(),
 			LOCTEXT("FindInExplorerTooltip", "Finds this asset on disk"),
-			FSlateIcon(),
+			FSlateIcon(FEditorStyle::GetStyleSetName(), "SystemWideCommands.FindInContentBrowser"),
 			FUIAction(
 				FExecuteAction::CreateSP( this, &FAssetContextMenu::ExecuteFindInExplorer ),
 				FCanExecuteAction::CreateSP( this, &FAssetContextMenu::CanExecuteFindInExplorer )
@@ -1069,7 +1071,9 @@ void FAssetContextMenu::GetSelectedAssetSourceFilePaths(TArray<FString>& OutFile
 
 void FAssetContextMenu::ExecuteSyncToAssetTree()
 {
-	OnFindInAssetTreeRequested.ExecuteIfBound(SelectedAssets);
+	// Copy this as the sync may adjust our selected assets array
+	const TArray<FAssetData> SelectedAssetsCopy = SelectedAssets;
+	OnFindInAssetTreeRequested.ExecuteIfBound(SelectedAssetsCopy);
 }
 
 void FAssetContextMenu::ExecuteFindInExplorer()
diff --git a/Engine/Source/Editor/Persona/Private/Persona.cpp b/Engine/Source/Editor/Persona/Private/Persona.cpp
--- a/Engine/Source/Editor/Persona/Private/Persona.cpp
+++ b/Engine/Source/Editor/Persona/Private/Persona.cpp
@@ -2576,7 +2576,10 @@ void FPersona::FindInContentBrowser_Execute()
 {
 	FName CurrentMode = GetCurrentMode();
 	TArray<UObject*> ObjectsToSyncTo = GetEditorObjectsForMode(CurrentMode);
-	GEditor->SyncBrowserToObjects( ObjectsToSyncTo );
+	if (ObjectsToSyncTo.Num() > 0)
+	{
+		GEditor->SyncBrowserToObjects( ObjectsToSyncTo );
+	}
 }
 
 void FPersona::OnCommandGenericDelete()
diff --git a/Engine/Source/Editor/Persona/Private/SAnimationSequenceBrowser.cpp b/Engine/Source/Editor/Persona/Private/SAnimationSequenceBrowser.cpp
--- a/Engine/Source/Editor/Persona/Private/SAnimationSequenceBrowser.cpp
+++ b/Engine/Source/Editor/Persona/Private/SAnimationSequenceBrowser.cpp
@@ -11,6 +11,7 @@
 #include "Editor/ContentBrowser/Public/ContentBrowserModule.h"
 #include "FeedbackContextEditor.h"
 #include "EditorAnimUtils.h"
+#include "Toolkits/GlobalEditorCommonCommands.h"
 #include "Editor/ContentBrowser/Public/FrontendFilterBase.h"
 #include "Runtime/AssetRegistry/Public/AssetRegistryModule.h"
 #include "SceneViewport.h"
@@ -98,8 +99,6 @@ void SAnimationSequenceBrowser::OnRequestOpenAsset(const FAssetData& AssetData,
 
 TSharedPtr<SWidget> SAnimationSequenceBrowser::OnGetAssetContextMenu(const TArray<FAssetData>& SelectedAssets)
 {
-	FMenuBuilder MenuBuilder(/*bInShouldCloseWindowAfterMenuSelection=*/ true, NULL);
-
 	bool bHasSelectedAnimSequence = false;
 	if ( SelectedAssets.Num() )
 	{
@@ -114,6 +113,8 @@ TSharedPtr<SWidget> SAnimationSequenceBrowser::OnGetAssetContextMenu(const TArra
 		}
 	}
 
+	FMenuBuilder MenuBuilder(/*bInShouldCloseWindowAfterMenuSelection=*/ true, Commands);
+
 	if(bHasSelectedAnimSequence)
 	{
 		MenuBuilder.BeginSection("AnimationSequenceOptions", NSLOCTEXT("Docking", "TabAnimationHeading", "Animation"));
@@ -173,25 +174,17 @@ TSharedPtr<SWidget> SAnimationSequenceBrowser::OnGetAssetContextMenu(const TArra
 
 	MenuBuilder.BeginSection("AnimationSequenceOptions", NSLOCTEXT("Docking", "TabOptionsHeading", "Options") );
 	{
-		MenuBuilder.AddMenuEntry(
-			LOCTEXT("GoToInContentBrowser", "Go To Content Browser"),
-			LOCTEXT("GoToInContentBrowser_ToolTip", "Select the asset in the content browser"),
-			FSlateIcon(),
-			FUIAction(
-			FExecuteAction::CreateSP( this, &SAnimationSequenceBrowser::OnGoToInContentBrowser, SelectedAssets ),
-			FCanExecuteAction()
-			)
-			);
-
 		MenuBuilder.AddMenuEntry(
 			LOCTEXT("SaveSelectedAssets", "Save"),
 			LOCTEXT("SaveSelectedAssets_ToolTip", "Save the selected assets"),
-			FSlateIcon(),
+			FSlateIcon(FEditorStyle::GetStyleSetName(), "Level.SaveIcon16x"),
 			FUIAction(
 				FExecuteAction::CreateSP( this, &SAnimationSequenceBrowser::SaveSelectedAssets, SelectedAssets),
 				FCanExecuteAction::CreateSP( this, &SAnimationSequenceBrowser::CanSaveSelectedAssets, SelectedAssets)
 				)
 			);
+
+		MenuBuilder.AddMenuEntry(FGlobalEditorCommonCommands::Get().FindInContentBrowser);
 	}
 	MenuBuilder.EndSection();
 
@@ -211,15 +204,22 @@ TSharedPtr<SWidget> SAnimationSequenceBrowser::OnGetAssetContextMenu(const TArra
 	return MenuBuilder.MakeWidget();
 }
 
-void SAnimationSequenceBrowser::OnGoToInContentBrowser(TArray<FAssetData> ObjectsToSync)
+void SAnimationSequenceBrowser::FindInContentBrowser()
 {
-	if ( ObjectsToSync.Num() > 0 )
+	TArray<FAssetData> CurrentSelection = GetCurrentSelectionDelegate.Execute();
+	if (CurrentSelection.Num() > 0)
 	{
 		FContentBrowserModule& ContentBrowserModule = FModuleManager::Get().LoadModuleChecked<FContentBrowserModule>("ContentBrowser");
-		ContentBrowserModule.Get().SyncBrowserToAssets( ObjectsToSync );
+		ContentBrowserModule.Get().SyncBrowserToAssets(CurrentSelection);
 	}
 }
 
+bool SAnimationSequenceBrowser::CanFindInContentBrowser() const
+{
+	TArray<FAssetData> CurrentSelection = GetCurrentSelectionDelegate.Execute();
+	return CurrentSelection.Num() > 0;
+}
+
 void SAnimationSequenceBrowser::GetSelectedPackages(const TArray<FAssetData>& Assets, TArray<UPackage*>& OutPackages) const
 {
 	for (int32 AssetIdx = 0; AssetIdx < Assets.Num(); ++AssetIdx)
@@ -398,6 +398,12 @@ void SAnimationSequenceBrowser::Construct(const FArguments& InArgs)
 		PersonaPtr.Pin()->SetSequenceBrowser(this);
 	}
 
+	Commands = MakeShareable(new FUICommandList());
+	Commands->MapAction(FGlobalEditorCommonCommands::Get().FindInContentBrowser, FUIAction(
+		FExecuteAction::CreateSP(this, &SAnimationSequenceBrowser::FindInContentBrowser),
+		FCanExecuteAction::CreateSP(this, &SAnimationSequenceBrowser::CanFindInContentBrowser)
+		));
+
 	CurrentAssetHistoryIndex = INDEX_NONE;
 	bTriedToCacheOrginalAsset = false;
 
@@ -525,6 +531,16 @@ void SAnimationSequenceBrowser::Construct(const FArguments& InArgs)
 	AssetRegistryTagsToIgnore.Add(GET_MEMBER_NAME_CHECKED(UAnimSequenceBase, RateScale));
 }
 
+FReply SAnimationSequenceBrowser::OnKeyDown(const FGeometry& MyGeometry, const FKeyEvent& InKeyEvent)
+{
+	if (Commands->ProcessCommandBindings(InKeyEvent))
+	{
+		return FReply::Handled();
+	}
+
+	return FReply::Unhandled();
+}
+
 void SAnimationSequenceBrowser::AddAssetToHistory(const FAssetData& AssetData)
 {
 	CacheOriginalAnimAssetHistory();
diff --git a/Engine/Source/Editor/Persona/Private/SAnimationSequenceBrowser.h b/Engine/Source/Editor/Persona/Private/SAnimationSequenceBrowser.h
--- a/Engine/Source/Editor/Persona/Private/SAnimationSequenceBrowser.h
+++ b/Engine/Source/Editor/Persona/Private/SAnimationSequenceBrowser.h
@@ -37,14 +37,18 @@ public:
 
 	void OnRequestOpenAsset(const FAssetData& AssetData, bool bFromHistory);
 
+	virtual FReply OnKeyDown(const FGeometry& MyGeometry, const FKeyEvent& InKeyEvent) override;
+
 	virtual ~SAnimationSequenceBrowser();
 
 	/** Delegate that handles creation of context menu */
 	TSharedPtr<SWidget> OnGetAssetContextMenu(const TArray<FAssetData>& SelectedAssets);
 
-	/** Delegate to handle "Go To" context menu option 
-		Sends selected assets to be highlighted in active Content Browser */
-	void OnGoToInContentBrowser(TArray<FAssetData> ObjectsToSync);
+	/** Delegate to handle "Find in Content Browser" context menu option */
+	void FindInContentBrowser();
+
+	/** Delegate to handle enabling the "Find in Content Browser" context menu option */
+	bool CanFindInContentBrowser() const;
 
 	/** Delegate to handle "Save" context menu option */
 	void SaveSelectedAssets(TArray<FAssetData> ObjectsToSave) const;
@@ -176,6 +180,11 @@ protected:
 	 */
 	FPreviewScene PreviewScene;
 
+	/**
+	 * Commands handled by this widget
+	 */
+	TSharedPtr<FUICommandList> Commands;
+
 	// Pointer back to persona tool that owns us
 	TWeakPtr<class FPersona> PersonaPtr;
 
