Comment fixup

diff --git a/Engine/Source/Editor/Persona/Private/PersonaModule.cpp b/Engine/Source/Editor/Persona/Private/PersonaModule.cpp
--- a/Engine/Source/Editor/Persona/Private/PersonaModule.cpp
+++ b/Engine/Source/Editor/Persona/Private/PersonaModule.cpp
@@ -67,7 +67,7 @@ void FPersonaModule::ShutdownModule()
 	MenuExtensibilityManager.Reset();
 	ToolBarExtensibilityManager.Reset();
 
-	// unregsiter when shut down
+	// Unregister when shut down
 	if(FModuleManager::Get().IsModuleLoaded("PropertyEditor"))
 	{
 		FPropertyEditorModule& PropertyModule = FModuleManager::LoadModuleChecked<FPropertyEditorModule>("PropertyEditor");
