[UE-19897] Implement feature to visually isolate meshes by material
Phase 1 - Added checkbox to isolate mesh sections in Persona for skeletal meshs. A notification is displayed in the viewport when this is enabled, warning the user. Static Mesh support to follow...

diff --git a/Engine/Source/Editor/Persona/Private/AnimationEditorViewportClient.cpp b/Engine/Source/Editor/Persona/Private/AnimationEditorViewportClient.cpp
--- a/Engine/Source/Editor/Persona/Private/AnimationEditorViewportClient.cpp
+++ b/Engine/Source/Editor/Persona/Private/AnimationEditorViewportClient.cpp
@@ -835,6 +835,8 @@ void FAnimationViewportClient::DisplayInfo(FCanvas* Canvas, FSceneView* View, bo
 
 	// it is weird, but unless it's completely black, it's too bright, so just making it white if only black
 	const FLinearColor TextColor = (SelectedHSVColor.B < 0.3f) ? FLinearColor::White : FLinearColor::Black;
+	const FColor HeadlineColour(255, 83, 0);
+	const FColor SubHeadlineColour(202, 66, 0);
 
 	// if not valid skeletalmesh
 	if (!PreviewSkelMeshComp.IsValid() || !PreviewSkelMeshComp->SkeletalMesh)
@@ -844,9 +846,6 @@ void FAnimationViewportClient::DisplayInfo(FCanvas* Canvas, FSceneView* View, bo
 
 	if (PreviewSkelMeshComp->SkeletalMesh->MorphTargets.Num() > 0)
 	{
-		FColor HeadlineColour(255,83,0);
-		FColor SubHeadlineColour(202,66,0);
-
 		int32 SubHeadingIndent = CurXOffset + 10;
 
 		TArray<UMaterial*> ProcessedMaterials;
@@ -1029,7 +1028,6 @@ void FAnimationViewportClient::DisplayInfo(FCanvas* Canvas, FSceneView* View, bo
 			if (PreviewSkelMeshComp->BonesOfInterest.Num() > 0)
 			{
 				int32 BoneIndex = PreviewSkelMeshComp->BonesOfInterest[0];
-				const FName BoneName = PreviewSkelMeshComp->SkeletalMesh->RefSkeleton.GetBoneName(BoneIndex);
 				FTransform ReferenceTransform = PreviewSkelMeshComp->SkeletalMesh->RefSkeleton.GetRefBonePose()[BoneIndex];
 				FTransform LocalTransform = PreviewSkelMeshComp->LocalAtoms[BoneIndex];
 				FTransform ComponentTransform = PreviewSkelMeshComp->GetSpaceBases()[BoneIndex];
@@ -1122,6 +1120,15 @@ void FAnimationViewportClient::DisplayInfo(FCanvas* Canvas, FSceneView* View, bo
 			Canvas->DrawShadowedString(CurXOffset, CurYOffset, *InfoString, GEngine->GetSmallFont(), TextColor);
 		}
 	}
+
+	if (PreviewSkelMeshComp->SectionIndexPreview != INDEX_NONE)
+	{
+		// Notify the user if they are isolating a mesh section.
+		CurYOffset += YL + 2;
+		InfoString = FString::Printf(*LOCTEXT("MeshSectionsHiddenWarning", "Mesh Sections Hidden").ToString());
+		Canvas->DrawShadowedString(CurXOffset, CurYOffset, *InfoString, GEngine->GetSmallFont(), SubHeadlineColour);
+		
+	}
 }
 
 int32 FAnimationViewportClient::FindSelectedBone() const
diff --git a/Engine/Source/Editor/Persona/Private/PersonaMeshDetails.cpp b/Engine/Source/Editor/Persona/Private/PersonaMeshDetails.cpp
--- a/Engine/Source/Editor/Persona/Private/PersonaMeshDetails.cpp
+++ b/Engine/Source/Editor/Persona/Private/PersonaMeshDetails.cpp
@@ -980,16 +980,35 @@ void FPersonaMeshDetails::OnGetMaterialsForView(IMaterialListBuilder& MaterialLi
 
 TSharedRef<SWidget> FPersonaMeshDetails::OnGenerateCustomNameWidgetsForMaterial(UMaterialInterface* Material, int32 SlotIndex, int32 LODIndex)
 {
-	return
-		SNew(SCheckBox)
-		.IsChecked(this, &FPersonaMeshDetails::IsSectionSelected, SlotIndex)
-		.OnCheckStateChanged(this, &FPersonaMeshDetails::OnSectionSelectedChanged, SlotIndex)
-		.ToolTipText(LOCTEXT("Highlight_ToolTip", "Highlights this section in the viewport"))
+	return SNew(SVerticalBox)
+		+ SVerticalBox::Slot()
+		.AutoHeight()
 		[
-			SNew(STextBlock)
-			.Font(IDetailLayoutBuilder::GetDetailFont())
-			.Text(LOCTEXT("Highlight", "Highlight"))
-			.ColorAndOpacity( FLinearColor( 0.4f, 0.4f, 0.4f, 1.0f) )
+			SNew(SCheckBox)
+			.IsChecked(this, &FPersonaMeshDetails::IsSectionSelected, SlotIndex)
+			.OnCheckStateChanged(this, &FPersonaMeshDetails::OnSectionSelectedChanged, SlotIndex)
+			.ToolTipText(LOCTEXT("Highlight_ToolTip", "Highlights this section in the viewport"))
+			[
+				SNew(STextBlock)
+				.Font(IDetailLayoutBuilder::GetDetailFont())
+				.ColorAndOpacity(FLinearColor(0.4f, 0.4f, 0.4f, 1.0f))
+				.Text(LOCTEXT("Highlight", "Highlight"))
+			]
+		]
+		+ SVerticalBox::Slot()
+		.AutoHeight()
+		.Padding(0, 2, 0, 0)
+		[
+			SNew(SCheckBox)
+			.IsChecked(this, &FPersonaMeshDetails::IsIsolateSectionEnabled, SlotIndex)
+			.OnCheckStateChanged(this, &FPersonaMeshDetails::OnSectionIsolatedChanged, SlotIndex)
+			.ToolTipText(LOCTEXT("Isolate_ToolTip", "Isolates this section in the viewport"))
+			[
+				SNew(STextBlock)
+				.Font(IDetailLayoutBuilder::GetDetailFont())
+				.ColorAndOpacity(FLinearColor(0.4f, 0.4f, 0.4f, 1.0f))
+				.Text(LOCTEXT("Isolate", "Isolate"))
+			]
 		];
 }
 
@@ -1107,8 +1126,7 @@ TSharedRef<SWidget> FPersonaMeshDetails::OnGenerateCustomMaterialWidgetsForMater
 ECheckBoxState FPersonaMeshDetails::IsSectionSelected(int32 SectionIndex) const
 {
 	ECheckBoxState State = ECheckBoxState::Unchecked;
-	USkeletalMesh* Mesh = Cast<USkeletalMesh>(SelectedObjects[0].Get());
-
+	const USkeletalMesh* Mesh = Cast<USkeletalMesh>(SelectedObjects[0].Get());
 	if (Mesh)
 	{
 		State = Mesh->SelectedEditorSection == SectionIndex ? ECheckBoxState::Checked : ECheckBoxState::Unchecked;
@@ -1120,11 +1138,20 @@ ECheckBoxState FPersonaMeshDetails::IsSectionSelected(int32 SectionIndex) const
 void FPersonaMeshDetails::OnSectionSelectedChanged(ECheckBoxState NewState, int32 SectionIndex)
 {
 	USkeletalMesh* Mesh = Cast<USkeletalMesh>(SelectedObjects[0].Get());
-	if (Mesh)
+
+	// Currently assumes that we only ever have one preview mesh in Persona.
+	UDebugSkelMeshComponent* MeshComponent = PersonaPtr->GetPreviewMeshComponent();
+
+	if (Mesh && MeshComponent)
 	{
 		if (NewState == ECheckBoxState::Checked)
 		{
 			Mesh->SelectedEditorSection = SectionIndex;
+			if (MeshComponent->SectionIndexPreview != SectionIndex)
+			{
+				// Unhide all mesh sections
+				MeshComponent->SetSectionPreview(INDEX_NONE);
+			}
 		}
 		else if (NewState == ECheckBoxState::Unchecked)
 		{
@@ -1134,11 +1161,44 @@ void FPersonaMeshDetails::OnSectionSelectedChanged(ECheckBoxState NewState, int3
 	}
 }
 
-ECheckBoxState FPersonaMeshDetails::IsShadowCastingEnabled(int32 MaterialIndex) const
+ECheckBoxState FPersonaMeshDetails::IsIsolateSectionEnabled(int32 SectionIndex) const
 {
 	ECheckBoxState State = ECheckBoxState::Unchecked;
-	USkeletalMesh* Mesh = Cast<USkeletalMesh>( SelectedObjects[0].Get() );
+	const UDebugSkelMeshComponent* MeshComponent = PersonaPtr->GetPreviewMeshComponent();
+	if (MeshComponent)
+	{
+		State = MeshComponent->SectionIndexPreview == SectionIndex ? ECheckBoxState::Checked : ECheckBoxState::Unchecked;
+	}
+	return State;
+}
+
+void FPersonaMeshDetails::OnSectionIsolatedChanged(ECheckBoxState NewState, int32 SectionIndex)
+{
+	USkeletalMesh* Mesh = Cast<USkeletalMesh>(SelectedObjects[0].Get());
+	UDebugSkelMeshComponent * MeshComponent = PersonaPtr->GetPreviewMeshComponent();
+	if (Mesh && MeshComponent)
+	{
+		if (NewState == ECheckBoxState::Checked)
+		{
+			MeshComponent->SetSectionPreview(SectionIndex); 
+			if (Mesh->SelectedEditorSection != SectionIndex)
+			{
+				Mesh->SelectedEditorSection = INDEX_NONE;
+			}
+		}
+		else if (NewState == ECheckBoxState::Unchecked)
+		{
+			MeshComponent->SetSectionPreview(INDEX_NONE);
+		}
+		PersonaPtr->RefreshViewport();
+	}
+}
+
 
+ECheckBoxState FPersonaMeshDetails::IsShadowCastingEnabled(int32 MaterialIndex) const
+{
+	ECheckBoxState State = ECheckBoxState::Unchecked;
+	const USkeletalMesh* Mesh = Cast<USkeletalMesh>( SelectedObjects[0].Get() );
 	if (Mesh && MaterialIndex < Mesh->Materials.Num())
 	{
 		State = Mesh->Materials[MaterialIndex].bEnableShadowCasting ? ECheckBoxState::Checked : ECheckBoxState::Unchecked;
diff --git a/Engine/Source/Editor/Persona/Private/PersonaMeshDetails.h b/Engine/Source/Editor/Persona/Private/PersonaMeshDetails.h
--- a/Engine/Source/Editor/Persona/Private/PersonaMeshDetails.h
+++ b/Engine/Source/Editor/Persona/Private/PersonaMeshDetails.h
@@ -154,6 +154,20 @@ private:
 	 */
 	void OnSectionSelectedChanged(ECheckBoxState NewState, int32 SectionIndex);
 
+	/**
+	* Handler for check box display based on whether the material is isolated
+	*
+	* @param SectionIndex	The material section that is being tested
+	*/
+	ECheckBoxState IsIsolateSectionEnabled(int32 SectionIndex) const;
+
+	/**
+	* Handler for changing isolated status on a material
+	*
+	* @param SectionIndex	The material section that is being tested
+	*/
+	void OnSectionIsolatedChanged(ECheckBoxState NewState, int32 SectionIndex);
+
 		/**
 	 * Handler for check box display based on whether the material has shadow casting enabled
 	 *
diff --git a/Engine/Source/Runtime/Engine/Classes/Components/SkinnedMeshComponent.h b/Engine/Source/Runtime/Engine/Classes/Components/SkinnedMeshComponent.h
--- a/Engine/Source/Runtime/Engine/Classes/Components/SkinnedMeshComponent.h
+++ b/Engine/Source/Runtime/Engine/Classes/Components/SkinnedMeshComponent.h
@@ -476,6 +476,13 @@ public:
 	 */
 	void SetForceWireframe(bool InForceWireframe);
 
+	/**
+	*	Sets the value of the SectionIndexPreview flag and reattaches the component as necessary.
+	*
+	*	@param	InSectionIndexPreview		New value of SectionIndexPreview.
+	*/
+	void SetSectionPreview(int32 InSectionIndexPreview);
+
 	/**
 	 * Function returns whether or not CPU skinning should be applied
 	 * Allows the editor to override the skinning state for editor tools
diff --git a/Engine/Source/Runtime/Engine/Classes/Engine/SkeletalMesh.h b/Engine/Source/Runtime/Engine/Classes/Engine/SkeletalMesh.h
--- a/Engine/Source/Runtime/Engine/Classes/Engine/SkeletalMesh.h
+++ b/Engine/Source/Runtime/Engine/Classes/Engine/SkeletalMesh.h
@@ -655,9 +655,10 @@ public:
 	TArray<FMatrix> RefBasesInvMatrix;    
 
 #if WITH_EDITORONLY_DATA
-	/** The section currently selected in the Editor. */
+	/** The section currently selected in the Editor. Used for highlighting */
 	UPROPERTY(transient)
 	int32 SelectedEditorSection;
+
 	/** The section currently selected for clothing. need to remember this index for reimporting cloth */
 	UPROPERTY(transient)
 	int32 SelectedClothingSection;
diff --git a/Engine/Source/Runtime/Engine/Private/Components/SkinnedMeshComponent.cpp b/Engine/Source/Runtime/Engine/Private/Components/SkinnedMeshComponent.cpp
--- a/Engine/Source/Runtime/Engine/Private/Components/SkinnedMeshComponent.cpp
+++ b/Engine/Source/Runtime/Engine/Private/Components/SkinnedMeshComponent.cpp
@@ -1093,6 +1093,17 @@ void USkinnedMeshComponent::SetForceWireframe(bool InForceWireframe)
 }
 
 
+void USkinnedMeshComponent::SetSectionPreview(int32 InSectionIndexPreview)
+{
+#if WITH_EDITORONLY_DATA
+	if (SectionIndexPreview != InSectionIndexPreview)
+	{
+		SectionIndexPreview = InSectionIndexPreview;
+		MarkRenderStateDirty();
+	}
+#endif
+}
+
 UMorphTarget* USkinnedMeshComponent::FindMorphTarget( FName MorphTargetName ) const
 {
 	if( SkeletalMesh != NULL )
