CIS - Shadow variable fix

diff --git a/Engine/Source/Editor/Persona/Private/AnimationEditorViewportClient.cpp b/Engine/Source/Editor/Persona/Private/AnimationEditorViewportClient.cpp
--- a/Engine/Source/Editor/Persona/Private/AnimationEditorViewportClient.cpp
+++ b/Engine/Source/Editor/Persona/Private/AnimationEditorViewportClient.cpp
@@ -914,7 +914,6 @@ void FAnimationViewportClient::DisplayInfo(FCanvas* Canvas, FSceneView* View, bo
 		UAnimSequence* Sequence = Cast<UAnimSequence>(PreviewInstance->CurrentAsset);
 		if ( Sequence && Sequence->DoesNeedRebake() )
 		{
-			FColor SubHeadlineColour(202, 66, 0);
 			InfoString = TEXT("Animation is being edited. To apply to raw animation data, click \"Apply\"");
 			Canvas->DrawShadowedString(CurXOffset, CurYOffset, *InfoString, GEngine->GetSmallFont(), SubHeadlineColour);
 			CurYOffset += YL + 2;
