Marker based sync
Merged from 2660856

diff --git a/Engine/Source/Editor/KismetCompiler/Private/AnimBlueprintCompiler.cpp b/Engine/Source/Editor/KismetCompiler/Private/AnimBlueprintCompiler.cpp
--- a/Engine/Source/Editor/KismetCompiler/Private/AnimBlueprintCompiler.cpp
+++ b/Engine/Source/Editor/KismetCompiler/Private/AnimBlueprintCompiler.cpp
@@ -1260,6 +1260,17 @@ void FAnimBlueprintCompiler::CleanAndSanitizeClass(UBlueprintGeneratedClass* Cla
 	}
 }
 
+void FAnimBlueprintCompiler::FinishCompilingClass(UClass* Class)
+{
+	UAnimBlueprintGeneratedClass* AnimBlueprintGeneratedClass = CastChecked<UAnimBlueprintGeneratedClass>(Class);
+	AnimBlueprintGeneratedClass->SyncGroupNames.Reset();
+	AnimBlueprintGeneratedClass->SyncGroupNames.Reserve(AnimBlueprint->Groups.Num());
+	for (const FAnimGroupInfo& GroupInfo : AnimBlueprint->Groups)
+	{
+		AnimBlueprintGeneratedClass->SyncGroupNames.Add(GroupInfo.Name);
+	}
+	Super::FinishCompilingClass(Class);
+}
 
 void FAnimBlueprintCompiler::CreateFunctionList()
 {
diff --git a/Engine/Source/Editor/KismetCompiler/Private/AnimBlueprintCompiler.h b/Engine/Source/Editor/KismetCompiler/Private/AnimBlueprintCompiler.h
--- a/Engine/Source/Editor/KismetCompiler/Private/AnimBlueprintCompiler.h
+++ b/Engine/Source/Editor/KismetCompiler/Private/AnimBlueprintCompiler.h
@@ -36,6 +36,7 @@ protected:
 	virtual void PostCompileDiagnostics() override;
 	virtual void EnsureProperGeneratedClass(UClass*& TargetClass) override;
 	virtual void CleanAndSanitizeClass(UBlueprintGeneratedClass* ClassToClean, UObject*& OldCDO) override;
+	virtual void FinishCompilingClass(UClass* Class) override;
 	// End of FKismetCompilerContext interface
 
 protected:
diff --git a/Engine/Source/Editor/Persona/Private/SAnimNotifyPanel.cpp b/Engine/Source/Editor/Persona/Private/SAnimNotifyPanel.cpp
--- a/Engine/Source/Editor/Persona/Private/SAnimNotifyPanel.cpp
+++ b/Engine/Source/Editor/Persona/Private/SAnimNotifyPanel.cpp
@@ -35,11 +35,308 @@ DECLARE_DELEGATE_RetVal_FourParams( FReply, FOnNotifyNodeDragStarted, TSharedRef
 DECLARE_DELEGATE_RetVal_FiveParams(FReply, FOnNotifyNodesDragStarted, TArray<TSharedPtr<SAnimNotifyNode>>, TSharedRef<SWidget>, const FVector2D&, const FVector2D&, const bool)
 DECLARE_DELEGATE_RetVal( float, FOnGetDraggedNodePos )
 DECLARE_DELEGATE_TwoParams( FPanTrackRequest, int32, FVector2D)
-DECLARE_DELEGATE_FourParams(FPasteNotifies, SAnimNotifyTrack*, float, ENotifyPasteMode::Type, ENotifyPasteMultipleMode::Type)
+DECLARE_DELEGATE(FCopyNodes)
+DECLARE_DELEGATE_FourParams(FPasteNodes, SAnimNotifyTrack*, float, ENotifyPasteMode::Type, ENotifyPasteMultipleMode::Type)
 DECLARE_DELEGATE_RetVal_OneParam(EVisibility, FOnGetTimingNodeVisibilityForNode, TSharedPtr<SAnimNotifyNode>)
 
 class FNotifyDragDropOp;
 
+FText MakeTooltipFromTime(const UAnimSequenceBase* InSequence, float InSeconds, float InDuration)
+{
+	const FText Frame = FText::AsNumber(InSequence->GetFrameAtTime(InSeconds));
+	const FText Seconds = FText::AsNumber(InSeconds);
+
+	if (InDuration > 0.0f)
+	{
+		const FText Duration = FText::AsNumber(InDuration);
+		return FText::Format(LOCTEXT("NodeToolTipLong", "@ {0} sec (frame {1}) for {2} sec"), Seconds, Frame, Duration);
+	}
+	else
+	{
+		return FText::Format(LOCTEXT("NodeToolTipShort", "@ {0} sec (frame {1})"), Seconds, Frame);
+	}
+}
+
+// Read common info from the clipboard
+bool ReadNotifyPasteHeader(FString& OutPropertyString, const TCHAR*& OutBuffer, float& OutOriginalTime, float& OutOriginalLength, int32& OutTrackSpan)
+{
+	OutBuffer = NULL;
+	OutOriginalTime = -1.f;
+
+	FPlatformMisc::ClipboardPaste(OutPropertyString);
+
+	if (!OutPropertyString.IsEmpty())
+	{
+		//Remove header text
+		const FString HeaderString(TEXT("COPY_ANIMNOTIFYEVENT"));
+
+		//Check for string identifier in order to determine whether the text represents an FAnimNotifyEvent.
+		if (OutPropertyString.StartsWith(HeaderString) && OutPropertyString.Len() > HeaderString.Len())
+		{
+			int32 HeaderSize = HeaderString.Len();
+			OutBuffer = *OutPropertyString;
+			OutBuffer += HeaderSize;
+
+			FString ReadLine;
+			// Read the original time from the first notify
+			FParse::Line(&OutBuffer, ReadLine);
+			FParse::Value(*ReadLine, TEXT("OriginalTime="), OutOriginalTime);
+			FParse::Value(*ReadLine, TEXT("OriginalLength="), OutOriginalLength);
+			FParse::Value(*ReadLine, TEXT("TrackSpan="), OutTrackSpan);
+			return true;
+		}
+	}
+
+	return false;
+}
+
+namespace ENodeObjectTypes
+{
+	enum Type
+	{
+		NOTIFY,
+		SYNC_MARKER
+	};
+};
+
+struct INodeObjectInterface
+{
+	virtual ENodeObjectTypes::Type GetType() const = 0;
+	virtual FAnimNotifyEvent* GetNotifyEvent() = 0;
+	virtual int GetTrackIndex() const = 0;
+	virtual float GetTime(EAnimLinkMethod::Type ReferenceFrame = EAnimLinkMethod::Absolute) const = 0;
+	virtual float GetDuration() = 0;
+	virtual FName GetName() = 0;
+	virtual TOptional<FLinearColor> GetEditorColor() = 0;
+	virtual FText GetNodeTooltip(const UAnimSequenceBase* Sequence) = 0;
+	virtual TOptional<UObject*> GetObjectBeingDisplayed() = 0;
+	virtual bool IsBranchingPoint() = 0;
+	bool operator<(const INodeObjectInterface& Rhs) const { return GetTime() < Rhs.GetTime(); }
+
+	virtual void SetTime(float Time, EAnimLinkMethod::Type ReferenceFrame = EAnimLinkMethod::Absolute) = 0;
+	virtual void SetDuration(float Duration) = 0;
+
+	virtual void HandleDrop(class UAnimSequenceBase* Sequence, float Time, int32 TrackIndex) = 0;
+	virtual void CacheName() = 0;
+
+	virtual void Delete(UAnimSequenceBase* Seq) = 0;
+
+	virtual void ExportForCopy(UAnimSequenceBase* Seq, FString& StrValue) const = 0;
+};
+
+struct FNotifyNodeInterface : public INodeObjectInterface
+{
+	FAnimNotifyEvent* NotifyEvent;
+
+	FNotifyNodeInterface(FAnimNotifyEvent* InAnimNotifyEvent) : NotifyEvent(InAnimNotifyEvent) {}
+	virtual ENodeObjectTypes::Type GetType() const override { return ENodeObjectTypes::NOTIFY; }
+	virtual FAnimNotifyEvent* GetNotifyEvent() override { return NotifyEvent; }
+	virtual int GetTrackIndex() const override{ return NotifyEvent->TrackIndex; }
+	virtual float GetTime(EAnimLinkMethod::Type ReferenceFrame = EAnimLinkMethod::Absolute) const override{ return NotifyEvent->GetTime(ReferenceFrame); }
+	virtual float GetDuration() override { return NotifyEvent->GetDuration(); }
+	virtual FName GetName() override { return NotifyEvent->NotifyName; }
+	virtual bool IsBranchingPoint() override { return NotifyEvent->IsBranchingPoint(); }
+	virtual TOptional<FLinearColor> GetEditorColor() override 
+	{
+		TOptional<FLinearColor> ReturnColour;
+		if (NotifyEvent->Notify)
+		{
+			ReturnColour = NotifyEvent->Notify->GetEditorColor();
+		}
+		else if (NotifyEvent->NotifyStateClass)
+		{
+			ReturnColour = NotifyEvent->NotifyStateClass->GetEditorColor();
+		}
+		return ReturnColour;
+	}
+
+	virtual FText GetNodeTooltip(const UAnimSequenceBase* Sequence) override
+	{
+		FText ToolTipText = MakeTooltipFromTime(Sequence, NotifyEvent->GetTime(), NotifyEvent->GetDuration());
+
+		if (NotifyEvent->IsBranchingPoint())
+		{
+			ToolTipText = FText::Format(LOCTEXT("AnimNotify_ToolTipBranchingPoint", "{0} (BranchingPoint)"), ToolTipText);
+		}
+		return ToolTipText;
+	}
+
+	virtual TOptional<UObject*> GetObjectBeingDisplayed() override
+	{
+		if (NotifyEvent->Notify)
+		{
+			return TOptional<UObject*>(NotifyEvent->Notify);
+		}
+
+		if (NotifyEvent->NotifyStateClass)
+		{
+			return TOptional<UObject*>(NotifyEvent->NotifyStateClass);
+		}
+		return TOptional<UObject*>();
+	}
+
+	virtual void SetTime(float Time, EAnimLinkMethod::Type ReferenceFrame = EAnimLinkMethod::Absolute) override { NotifyEvent->SetTime(Time, ReferenceFrame); }
+	virtual void SetDuration(float Duration) override { NotifyEvent->SetDuration(Duration); }
+
+	virtual void HandleDrop(class UAnimSequenceBase* Sequence, float Time, int32 TrackIndex) override
+	{
+		float EventDuration = NotifyEvent->GetDuration();
+
+		NotifyEvent->Link(Sequence, Time, NotifyEvent->GetSlotIndex());
+		NotifyEvent->RefreshTriggerOffset(Sequence->CalculateOffsetForNotify(NotifyEvent->GetTime()));
+
+		if (EventDuration > 0.0f)
+		{
+			NotifyEvent->EndLink.Link(Sequence, NotifyEvent->GetTime() + EventDuration, NotifyEvent->GetSlotIndex());
+			NotifyEvent->RefreshEndTriggerOffset(Sequence->CalculateOffsetForNotify(NotifyEvent->EndLink.GetTime()));
+		}
+		else
+		{
+			NotifyEvent->EndTriggerTimeOffset = 0.0f;
+		}
+
+		NotifyEvent->TrackIndex = TrackIndex;
+	}
+
+	virtual void CacheName() override 
+	{
+		if (NotifyEvent->Notify)
+		{
+			NotifyEvent->NotifyName = FName(*NotifyEvent->Notify->GetNotifyName());
+		}
+		else if (NotifyEvent->NotifyStateClass)
+		{
+			NotifyEvent->NotifyName = FName(*NotifyEvent->NotifyStateClass->GetNotifyName());
+		}
+	}
+
+	virtual void Delete(UAnimSequenceBase* Seq) override
+	{
+		for (int32 I = 0; I < Seq->Notifies.Num(); ++I)
+		{
+			if (NotifyEvent == &(Seq->Notifies[I]))
+			{
+				Seq->Notifies.RemoveAt(I);
+				Seq->MarkPackageDirty();
+				break;
+			}
+		}
+	}
+
+	virtual void ExportForCopy(UAnimSequenceBase* Seq, FString& StrValue) const override
+	{
+		int32 Index = INDEX_NONE;
+		for (int32 NotifyIdx = 0; NotifyIdx < Seq->Notifies.Num(); ++NotifyIdx)
+		{
+			if (NotifyEvent == &Seq->Notifies[NotifyIdx])
+			{
+				Index = NotifyIdx;
+				break;
+			}
+		}
+
+		check(Index != INDEX_NONE);
+
+		UArrayProperty* ArrayProperty = NULL;
+		uint8* PropertyData = Seq->FindNotifyPropertyData(Index, ArrayProperty);
+		if (PropertyData && ArrayProperty)
+		{
+			ArrayProperty->Inner->ExportTextItem(StrValue, PropertyData, PropertyData, Seq, PPF_Copy);
+		}
+	}
+};
+
+struct FSyncMarkerNodeInterface : public INodeObjectInterface
+{
+	FAnimSyncMarker* SyncMarker;
+
+	FSyncMarkerNodeInterface(FAnimSyncMarker* InSyncMarker) : SyncMarker(InSyncMarker) {}
+	virtual ENodeObjectTypes::Type GetType() const override { return ENodeObjectTypes::SYNC_MARKER; }
+	virtual FAnimNotifyEvent* GetNotifyEvent() override { return NULL; }
+	virtual int GetTrackIndex() const override{ return SyncMarker->TrackIndex; }
+	virtual float GetTime(EAnimLinkMethod::Type ReferenceFrame = EAnimLinkMethod::Absolute) const override { return SyncMarker->Time; }
+	virtual float GetDuration() override { return 0.f; }
+	virtual FName GetName() override { return SyncMarker->MarkerName; }
+	virtual bool IsBranchingPoint() override { return false; }
+	virtual TOptional<FLinearColor> GetEditorColor() override
+	{
+		return FLinearColor::Green;
+	}
+
+	virtual FText GetNodeTooltip(const UAnimSequenceBase* Sequence) override
+	{
+		return MakeTooltipFromTime(Sequence, SyncMarker->Time, 0.f);
+	}
+
+	virtual TOptional<UObject*> GetObjectBeingDisplayed() override
+	{
+		return TOptional<UObject*>();
+	}
+
+	virtual void SetTime(float Time, EAnimLinkMethod::Type ReferenceFrame = EAnimLinkMethod::Absolute) override { SyncMarker->Time = Time; }
+	virtual void SetDuration(float Duration) override {}
+
+	virtual void HandleDrop(class UAnimSequenceBase* Sequence, float Time, int32 TrackIndex) override
+	{
+		SyncMarker->Time = Time;
+		SyncMarker->TrackIndex = TrackIndex;
+	}
+
+	virtual void CacheName() override {}
+
+	virtual void Delete(UAnimSequenceBase* SeqBase) override
+	{
+		if(UAnimSequence* Seq = Cast<UAnimSequence>(SeqBase))
+		{
+			for (int32 I = 0; I < Seq->AuthoredSyncMarkers.Num(); ++I)
+			{
+				if (SyncMarker == &(Seq->AuthoredSyncMarkers[I]))
+				{
+					Seq->AuthoredSyncMarkers.RemoveAt(I);
+					Seq->MarkPackageDirty();
+					break;
+				}
+			}
+		}
+	}
+
+	virtual void ExportForCopy(UAnimSequenceBase* SeqBase, FString& StrValue) const override
+	{
+		if (UAnimSequence* Seq = Cast<UAnimSequence>(SeqBase))
+		{
+			int32 Index = INDEX_NONE;
+			for (int32 SyncMarkerIdx = 0; SyncMarkerIdx < Seq->AuthoredSyncMarkers.Num(); ++SyncMarkerIdx)
+			{
+				if (SyncMarker == &Seq->AuthoredSyncMarkers[SyncMarkerIdx])
+				{
+					Index = SyncMarkerIdx;
+					break;
+				}
+			}
+
+			check(Index != INDEX_NONE);
+
+			UArrayProperty* ArrayProperty = NULL;
+			uint8* PropertyData = Seq->FindSyncMarkerPropertyData(Index, ArrayProperty);
+			if (PropertyData && ArrayProperty)
+			{
+				ArrayProperty->Inner->ExportTextItem(StrValue, PropertyData, PropertyData, Seq, PPF_Copy);
+			}
+		}
+	}
+};
+
+// Struct that allows us to get the max value of 2 numbers at compile time
+template<int32 A, int32 B>
+struct CompileTimeMax
+{
+	enum Max{ VALUE = (A > B) ? A : B };
+};
+
+// Size of biggest object that we can store in our node, if new node interfaces are added they should be part of this calculation
+const int32 MAX_NODE_OBJECT_INTERFACE_SIZE = CompileTimeMax<sizeof(FNotifyNodeInterface), sizeof(FSyncMarkerNodeInterface)>::VALUE;
+
+
 //////////////////////////////////////////////////////////////////////////
 // SAnimNotifyNode
 
@@ -48,7 +345,8 @@ class SAnimNotifyNode : public SLeafWidget
 public:
 	SLATE_BEGIN_ARGS( SAnimNotifyNode )
 		: _Sequence()
-		, _AnimNotify()
+		, _AnimNotify(nullptr)
+		, _AnimSyncMarker(nullptr)
 		, _OnNodeDragStarted()
 		, _OnUpdatePanel()
 		, _PanTrackRequest()
@@ -60,6 +358,7 @@ public:
 	}
 	SLATE_ARGUMENT( class UAnimSequenceBase*, Sequence )
 	SLATE_ARGUMENT( FAnimNotifyEvent *, AnimNotify )
+	SLATE_ARGUMENT( FAnimSyncMarker*, AnimSyncMarker)
 	SLATE_EVENT( FOnNotifyNodeDragStarted, OnNodeDragStarted )
 	SLATE_EVENT( FOnUpdatePanel, OnUpdatePanel )
 	SLATE_EVENT( FPanTrackRequest, PanTrackRequest )
@@ -105,9 +404,20 @@ public:
 
 	FLinearColor GetNotifyColor() const;
 	FText GetNotifyText() const;
-	
-	/** The notify that is being visualized by this node */
-	FAnimNotifyEvent* NotifyEvent;
+
+	/** Node object interface */
+	INodeObjectInterface* NodeObjectInterface;
+
+	/** In object storage for our interface struct, saves us having to dynamically allocate what will be a very small struct*/
+	uint8 NodeObjectInterfaceStorage[MAX_NODE_OBJECT_INTERFACE_SIZE];
+
+	/** Helper function to create our node interface object */
+	template<typename InterfaceType, typename ParamType>
+	void MakeNodeInterface(ParamType& InParam)
+	{
+		check(sizeof(InterfaceType) <= MAX_NODE_OBJECT_INTERFACE_SIZE); //Not enough space, check definiton of MAX_NODE_OBJECT_INTERFACE_SIZE
+		NodeObjectInterface = new(NodeObjectInterfaceStorage)InterfaceType(InParam);
+	}
 
 	void DropCancelled();
 
@@ -290,14 +600,15 @@ public:
 		, _OnReplaceSelectedWithNotify()
 		, _OnReplaceSelectedWithBlueprintNotify()
 		, _OnDeselectAllNotifies()
-		, _OnCopyNotifies()
-		, _OnPasteNotifies()
+		, _OnCopyNodes()
+		, _OnPasteNodes()
 		, _OnSetInputViewRange()
 		{}
 
 		SLATE_ARGUMENT( TSharedPtr<FPersona>,	Persona )
 		SLATE_ARGUMENT( class UAnimSequenceBase*, Sequence )
 		SLATE_ARGUMENT( TArray<FAnimNotifyEvent *>, AnimNotifies )
+		SLATE_ARGUMENT( TArray<FAnimSyncMarker *>, AnimSyncMarkers)
 		SLATE_ATTRIBUTE( float, ViewInputMin )
 		SLATE_ATTRIBUTE( float, ViewInputMax )
 		SLATE_ATTRIBUTE( TArray<FTrackMarkerBar>, MarkerBars)
@@ -321,8 +632,8 @@ public:
 		SLATE_EVENT( FReplaceWithNotify, OnReplaceSelectedWithNotify )
 		SLATE_EVENT( FReplaceWithBlueprintNotify, OnReplaceSelectedWithBlueprintNotify)
 		SLATE_EVENT( FDeselectAllNotifies, OnDeselectAllNotifies)
-		SLATE_EVENT( FCopyNotifies, OnCopyNotifies )
-		SLATE_EVENT( FPasteNotifies, OnPasteNotifies )
+		SLATE_EVENT( FCopyNodes, OnCopyNodes )
+		SLATE_EVENT(FPasteNodes, OnPasteNodes)
 		SLATE_EVENT( FOnSetInputViewRange, OnSetInputViewRange )
 		SLATE_EVENT( FOnGetTimingNodeVisibility, OnGetTimingNodeVisibility )
 		SLATE_END_ARGS()
@@ -354,7 +665,7 @@ public:
 	void HandleNodeDrop(TSharedPtr<SAnimNotifyNode> Node, float Offset = 0.0f);
 
 	// Number of nodes in the track currently selected
-	int32 GetNumSelectedNodes() const { return SelectedNotifyIndices.Num(); }
+	int32 GetNumSelectedNodes() const { return SelectedNodeIndices.Num(); }
 
 	// Index of the track in the notify panel
 	int32 GetTrackIndex() const { return TrackIndex; }
@@ -368,12 +679,13 @@ public:
 	// Adds our current selection to the provided set
 	void AppendSelectionToSet(FGraphPanelSelectionSet& SelectionSet);
 	// Adds our current selection to the provided array
-	void AppendSelectionToArray(TArray<FAnimNotifyEvent*>& Selection) const;
+	void AppendSelectionToArray(TArray<INodeObjectInterface*>& Selection) const;
 	// Gets the currently selected SAnimNotifyNode instances
 	void AppendSelectedNodeWidgetsToArray(TArray<TSharedPtr<SAnimNotifyNode>>& NodeArray) const;
 	// Gets the indices of the selected notifies
-	TArray<int32> GetSelectedNotifyIndices() const {return SelectedNotifyIndices;}
+	TArray<int32> GetSelectedNotifyIndices() const {return SelectedNodeIndices;}
 
+	INodeObjectInterface* GetNodeObjectInterface(int32 NodeIndex) { return NotifyNodes[NodeIndex]->NodeObjectInterface; }
 	/**
 	* Deselects all currently selected notify nodes
 	* @param bUpdateSelectionSet - Whether we should report a selection change to the panel
@@ -386,6 +698,9 @@ public:
 	// Paste a single Notify into this track from an exported string
 	void PasteSingleNotify(FString& NotifyString, float PasteTime);
 
+	// Paste a single Sync Marker into this track from an exported string
+	void PasteSingleSyncMarker(FString& MarkerString, float PasteTime);
+
 	// Uses the given track space rect and marquee information to refresh selection information
 	void RefreshMarqueeSelectedNodes(FSlateRect& Rect, FNotifyMarqueeOperation& Marquee);
 
@@ -410,8 +725,11 @@ protected:
 	// New notify functions
 	void CreateNewBlueprintNotifyAtCursor(FString NewNotifyName, FString BlueprintPath);
 	void CreateNewNotifyAtCursor(FString NewNotifyName, UClass* NotifyClass);
+	void CreateNewSyncMarkerAtCursor(FString NewSyncMarkerName, UClass* NotifyClass);
 	void OnNewNotifyClicked();
+	void OnNewSyncMarkerClicked();
 	void AddNewNotify(const FText& NewNotifyName, ETextCommit::Type CommitInfo);
+	void AddNewSyncMarker(const FText& NewMarkerName, ETextCommit::Type CommitInfo);
 
 	// Trigger weight functions
 	void OnSetTriggerWeightNotifyClicked(int32 NotifyIndex);
@@ -420,10 +738,10 @@ protected:
 	// "Replace with... " commands
 	void ReplaceSelectedWithBlueprintNotify(FString NewNotifyName, FString BlueprintPath);
 	void ReplaceSelectedWithNotify(FString NewNotifyName, UClass* NotifyClass);
-	void OnSetNotifyTimeClicked(int32 NotifyIndex);
-	void SetNotifyTime(const FText& NotifyTimeText, ETextCommit::Type CommitInfo, int32 NotifyIndex);
-	void OnSetNotifyFrameClicked(int32 NotifyIndex);
-	void SetNotifyFrame(const FText& NotifyFrameText, ETextCommit::Type CommitInfo, int32 NotifyIndex);
+	void OnSetNodeTimeClicked(int32 NodeIndex);
+	void SetNodeTime(const FText& NodeTimeText, ETextCommit::Type CommitInfo, int32 NodeIndex);
+	void OnSetNodeFrameClicked(int32 NodeIndex);
+	void SetNodeFrame(const FText& NodeFrameText, ETextCommit::Type CommitInfo, int32 NodeIndex);
 
 	// Whether we have one node selected
 	bool IsSingleNodeSelected();
@@ -439,10 +757,6 @@ protected:
 	/** Function to check whether it is possible to paste anim notify event */
 	bool CanPasteAnimNotify() const;
 
-	/** Reads the contents of the clipboard for a notify to paste.
-	 *  If successful OutBuffer points to the start of the notify object's data */
-	bool ReadNotifyPasteHeader(FString& OutPropertyString, const TCHAR*& OutBuffer, float& OutOriginalTime, float& OutOriginalLength) const;
-
 	/** Handler for context menu paste command */
 	void OnPasteNotifyClicked(ENotifyPasteMode::Type PasteMode, ENotifyPasteMultipleMode::Type MultiplePasteType = ENotifyPasteMultipleMode::Absolute);
 
@@ -459,12 +773,12 @@ protected:
 	void OnOpenNotifySource(UBlueprint* InSourceBlueprint) const;
 
 	/**
-	 * Selects a notify node in the graph. Supports multi selection
-	 * @param NotifyIndex - Index of the notify node to select.
+	 * Selects a node on the track. Supports multi selection
+	 * @param TrackNodeIndex - Index of the node to select.
 	 * @param Append - Whether to append to to current selection or start a new one.
 	 * @param bUpdateSelection - Whether to immediately inform Persona of a selection change
 	 */
-	void SelectNotifyNode(int32 NotifyIndex, bool Append, bool bUpdateSelection = true);
+	void SelectTrackObjectNode(int32 TrackNodeIndex, bool Append, bool bUpdateSelection = true);
 	
 	/**
 	 * Toggles the selection status of a notify node, for example when
@@ -472,14 +786,14 @@ protected:
 	 * @param NotifyIndex - Index of the notify to toggle the selection status of
 	 * @param bUpdateSelection - Whether to immediately inform Persona of a selection change
 	 */
-	void ToggleNotifyNodeSelectionStatus(int32 NotifyIndex, bool bUpdateSelection = true);
+	void ToggleTrackObjectNodeSelectionStatus(int32 TrackNodeIndex, bool bUpdateSelection = true);
 
 	/**
 	 * Deselects requested notify node
 	 * @param NotifyIndex - Index of the notify node to deselect
 	 * @param bUpdateSelection - Whether to immediately inform Persona of a selection change
 	 */
-	void DeselectNotifyNode(int32 NotifyIndex, bool bUpdateSelection = true);
+	void DeselectTrackObjectNode(int32 TrackNodeIndex, bool bUpdateSelection = true);
 
 	int32 GetHitNotifyNode(const FGeometry& MyGeometry, const FVector2D& Position);
 
@@ -520,6 +834,14 @@ private:
 		return FMargin(LeftMargin, 0, RightMargin, 0);
 	}
 
+	// Returns the padding needed to render the notify in the correct track position
+	FMargin GetSyncMarkerTrackPadding(int32 SyncMarkerIndex) const
+	{
+		float LeftMargin = NotifyNodes[SyncMarkerIndex]->GetWidgetPosition().X;
+		float RightMargin = CachedGeometry.Size.X - NotifyNodes[SyncMarkerIndex]->GetWidgetPosition().X - NotifyNodes[SyncMarkerIndex]->GetSize().X;
+		return FMargin(LeftMargin, 0, RightMargin, 0);
+	}
+
 	// Builds a UObject selection set and calls the OnSelectionChanged delegate
 	void SendSelectionChanged()
 	{
@@ -535,6 +857,7 @@ protected:
 	TArray<TSharedPtr<SAnimNotifyNode>>		NotifyNodes;
 	TArray<TSharedPtr<SAnimNotifyPair>>		NotifyPairs;
 	TArray<FAnimNotifyEvent*>				AnimNotifies;
+	TArray<FAnimSyncMarker*>				AnimSyncMarkers;
 	TAttribute<float>						ViewInputMin;
 	TAttribute<float>						ViewInputMax;
 	TAttribute<FLinearColor>				TrackColor;
@@ -552,8 +875,8 @@ protected:
 	FOnNotifyNodesDragStarted				OnNodeDragStarted;
 	FPanTrackRequest						OnRequestTrackPan;
 	FDeselectAllNotifies					OnDeselectAllNotifies;
-	FCopyNotifies							OnCopyNotifies;
-	FPasteNotifies							OnPasteNotifies;
+	FCopyNodes							OnCopyNodes;
+	FPasteNodes								OnPasteNodes;
 	FOnSetInputViewRange					OnSetInputViewRange;
 	FOnGetTimingNodeVisibility				OnGetTimingNodeVisibility;
 
@@ -581,7 +904,7 @@ protected:
 	TAttribute<TArray<FTrackMarkerBar>>		MarkerBars;
 
 	/** Nodes that are currently selected */
-	TArray<int32> SelectedNotifyIndices;
+	TArray<int32> SelectedNodeIndices;
 };
 
 //////////////////////////////////////////////////////////////////////////
@@ -611,7 +934,7 @@ public:
 		, _OnUpdatePanel()
 		, _OnDeleteNotify()
 		, _OnDeselectAllNotifies()
-		, _OnCopyNotifies()
+		, _OnCopyNodes()
 		, _OnSetInputViewRange()
 	{}
 	SLATE_ARGUMENT( int32, TrackIndex )
@@ -635,8 +958,8 @@ public:
 	SLATE_EVENT( FRefreshOffsetsRequest, OnRequestRefreshOffsets )
 	SLATE_EVENT( FDeleteNotify, OnDeleteNotify )
 	SLATE_EVENT( FDeselectAllNotifies, OnDeselectAllNotifies)
-	SLATE_EVENT( FCopyNotifies, OnCopyNotifies )
-	SLATE_EVENT( FPasteNotifies, OnPasteNotifies )
+	SLATE_EVENT( FCopyNodes, OnCopyNodes )
+	SLATE_EVENT( FPasteNodes, OnPasteNodes )
 	SLATE_EVENT( FOnSetInputViewRange, OnSetInputViewRange )
 	SLATE_EVENT( FOnGetTimingNodeVisibility, OnGetTimingNodeVisibility )
 	SLATE_END_ARGS()
@@ -730,8 +1053,7 @@ public:
 		for(int32 NodeIdx = 0 ; NodeIdx < SelectedNodes.Num() ; ++NodeIdx)
 		{
 			TSharedPtr<SAnimNotifyNode> CurrentNode = SelectedNodes[NodeIdx];
-			FAnimNotifyEvent* CurrentEvent = CurrentNode->NotifyEvent;
-			
+
 			// Clear off any snap time currently stored
 			CurrentNode->ClearLastSnappedTime();
 
@@ -744,75 +1066,78 @@ public:
 			const FVector2D OriginalNodePosition = TrackNodePos;
 			float SequenceEnd = TrackScaleInfo.InputToLocalX(Sequence->SequenceLength);
 
+			// Always clamp the Y to the current track
+			SelectionBeginPosition.Y = SelectionPositionClampInfo->TrackPos;
+
 			float SnapX = GetSnapPosition(NodeClamp, TrackNodePos.X);
-			if(SnapX >= 0.f)
+			if (FAnimNotifyEvent* CurrentEvent = CurrentNode->NodeObjectInterface->GetNotifyEvent())
 			{
-				EAnimEventTriggerOffsets::Type Offset = EAnimEventTriggerOffsets::NoOffset;
-				if(SnapX == 0.0f || SnapX == SequenceEnd)
-				{
-					Offset = SnapX > 0.0f ? EAnimEventTriggerOffsets::OffsetBefore : EAnimEventTriggerOffsets::OffsetAfter;
-				}
-				else
+				if (SnapX >= 0.f)
 				{
-					Offset = (SnapX < TrackNodePos.X) ? EAnimEventTriggerOffsets::OffsetAfter : EAnimEventTriggerOffsets::OffsetBefore;
-				}
+					EAnimEventTriggerOffsets::Type Offset = EAnimEventTriggerOffsets::NoOffset;
+					if (SnapX == 0.0f || SnapX == SequenceEnd)
+					{
+						Offset = SnapX > 0.0f ? EAnimEventTriggerOffsets::OffsetBefore : EAnimEventTriggerOffsets::OffsetAfter;
+					}
+					else
+					{
+						Offset = (SnapX < TrackNodePos.X) ? EAnimEventTriggerOffsets::OffsetAfter : EAnimEventTriggerOffsets::OffsetBefore;
+					}
 
-				CurrentEvent->TriggerTimeOffset = GetTriggerTimeOffsetForType(Offset);
-				CurrentNode->SetLastSnappedTime(TrackScaleInfo.LocalXToInput(SnapX));
+					CurrentEvent->TriggerTimeOffset = GetTriggerTimeOffsetForType(Offset);
+					CurrentNode->SetLastSnappedTime(TrackScaleInfo.LocalXToInput(SnapX));
 
-				if(SnapMovement == 0.0f)
+					if (SnapMovement == 0.0f)
+					{
+						SnapMovement = SnapX - TrackNodePos.X;
+						TrackNodePos.X = SnapX;
+						SnapTime = TrackScaleInfo.LocalXToInput(SnapX);
+						SnappedNode = CurrentNode;
+					}
+					EventPosition = NodeClamp.NotifyTrack->GetCachedGeometry().LocalToAbsolute(TrackNodePos);
+				}
+				else
 				{
-					SnapMovement = SnapX - TrackNodePos.X;
-					TrackNodePos.X = SnapX;
-					SnapTime = TrackScaleInfo.LocalXToInput(SnapX);
-					SnappedNode = CurrentNode;
+					CurrentEvent->TriggerTimeOffset = GetTriggerTimeOffsetForType(EAnimEventTriggerOffsets::NoOffset);
 				}
-				EventPosition = NodeClamp.NotifyTrack->GetCachedGeometry().LocalToAbsolute(TrackNodePos);
-			}
-			else
-			{
-				CurrentEvent->TriggerTimeOffset = GetTriggerTimeOffsetForType(EAnimEventTriggerOffsets::NoOffset);
-			}
 
-			// Always clamp the Y to the current track
-			SelectionBeginPosition.Y = SelectionPositionClampInfo->TrackPos;
-
-			if(CurrentNode.IsValid() && CurrentEvent->GetDuration() > 0)
-			{
-				// If we didn't snap the beginning of the node, attempt to snap the end
-				if(SnapX == -1.0f)
+				if (CurrentNode.IsValid() && CurrentEvent->GetDuration() > 0)
 				{
-					FVector2D TrackNodeEndPos = TrackNodePos + CurrentNode->GetDurationSize();
-					SnapX = GetSnapPosition(*SelectionPositionClampInfo, TrackNodeEndPos.X);
-
-					// Only attempt to snap if the node will fit on the track
-					if(SnapX >= CurrentNode->GetDurationSize())
+					// If we didn't snap the beginning of the node, attempt to snap the end
+					if (SnapX == -1.0f)
 					{
-						EAnimEventTriggerOffsets::Type Offset = EAnimEventTriggerOffsets::NoOffset;
-						if(SnapX == SequenceEnd)
+						FVector2D TrackNodeEndPos = TrackNodePos + CurrentNode->GetDurationSize();
+						SnapX = GetSnapPosition(*SelectionPositionClampInfo, TrackNodeEndPos.X);
+
+						// Only attempt to snap if the node will fit on the track
+						if (SnapX >= CurrentNode->GetDurationSize())
 						{
-							// Only need to check the end of the sequence here; end handle can't hit the beginning
-							Offset = EAnimEventTriggerOffsets::OffsetBefore;
+							EAnimEventTriggerOffsets::Type Offset = EAnimEventTriggerOffsets::NoOffset;
+							if (SnapX == SequenceEnd)
+							{
+								// Only need to check the end of the sequence here; end handle can't hit the beginning
+								Offset = EAnimEventTriggerOffsets::OffsetBefore;
+							}
+							else
+							{
+								Offset = (SnapX < TrackNodeEndPos.X) ? EAnimEventTriggerOffsets::OffsetAfter : EAnimEventTriggerOffsets::OffsetBefore;
+							}
+							CurrentEvent->EndTriggerTimeOffset = GetTriggerTimeOffsetForType(Offset);
+
+							if (SnapMovement == 0.0f)
+							{
+								SnapMovement = SnapX - TrackNodeEndPos.X;
+								SnapTime = TrackScaleInfo.LocalXToInput(SnapX) - CurrentEvent->GetDuration();
+								CurrentNode->SetLastSnappedTime(SnapTime);
+								SnappedNode = CurrentNode;
+							}
 						}
 						else
 						{
-							Offset = (SnapX < TrackNodeEndPos.X) ? EAnimEventTriggerOffsets::OffsetAfter : EAnimEventTriggerOffsets::OffsetBefore;
-						}
-						CurrentEvent->EndTriggerTimeOffset = GetTriggerTimeOffsetForType(Offset);
-						
-						if(SnapMovement == 0.0f)
-						{
-							SnapMovement = SnapX - TrackNodeEndPos.X;
-							SnapTime = TrackScaleInfo.LocalXToInput(SnapX) - CurrentEvent->GetDuration();
-							CurrentNode->SetLastSnappedTime(SnapTime);
-							SnappedNode = CurrentNode;
+							// Remove any trigger time if we can't fit the node in.
+							CurrentEvent->EndTriggerTimeOffset = GetTriggerTimeOffsetForType(EAnimEventTriggerOffsets::NoOffset);
 						}
 					}
-					else
-					{
-						// Remove any trigger time if we can't fit the node in.
-						CurrentEvent->EndTriggerTimeOffset = GetTriggerTimeOffsetForType(EAnimEventTriggerOffsets::NoOffset);
-					}
 				}
 			}
 		}
@@ -944,13 +1269,13 @@ public:
 		Operation->MarkerBars = MarkerBars;
 		Operation->Decorator = Decorator;
 		Operation->SelectedNodes = NotifyNodes;
-		Operation->TrackSpan = NotifyNodes[0]->NotifyEvent->TrackIndex - NotifyNodes.Last()->NotifyEvent->TrackIndex;
+		Operation->TrackSpan = NotifyNodes[0]->NodeObjectInterface->GetTrackIndex() - NotifyNodes.Last()->NodeObjectInterface->GetTrackIndex();
 		
 		// Caclulate offsets for the selected nodes
 		float BeginTime = MAX_flt;
 		for(TSharedPtr<SAnimNotifyNode> Node : NotifyNodes)
 		{
-			float NotifyTime = Node->NotifyEvent->GetTime();
+			float NotifyTime = Node->NodeObjectInterface->GetTime();
 
 			if(NotifyTime < BeginTime)
 			{
@@ -961,7 +1286,7 @@ public:
 		// Initialise node data
 		for(TSharedPtr<SAnimNotifyNode> Node : NotifyNodes)
 		{
-			float NotifyTime = Node->NotifyEvent->GetTime();
+			float NotifyTime = Node->NodeObjectInterface->GetTime();
 
 			Node->ClearLastSnappedTime();
 			Operation->NodeTimeOffsets.Add(NotifyTime - BeginTime);
@@ -971,7 +1296,7 @@ public:
 			// Calculate the time length of the selection. Because it is possible to have states
 			// with arbitrary durations we need to search all of the nodes and find the furthest
 			// possible point
-			Operation->SelectionTimeLength = FMath::Max(Operation->SelectionTimeLength, NotifyTime + Node->NotifyEvent->GetDuration() - BeginTime);
+			Operation->SelectionTimeLength = FMath::Max(Operation->SelectionTimeLength, NotifyTime + Node->NodeObjectInterface->GetDuration() - BeginTime);
 		}
 
 		Operation->Construct();
@@ -1005,7 +1330,7 @@ public:
 
 		if(SelectedNodes[0].IsValid())
 		{
-			HoverText = FText::FromName( SelectedNodes[0]->NotifyEvent->NotifyName );
+			HoverText = FText::FromName( SelectedNodes[0]->NodeObjectInterface->GetName() );
 		}
 
 		return HoverText;
@@ -1043,7 +1368,6 @@ const float SAnimNotifyNode::MinimumStateDuration = (1.0f / 30.0f);
 void SAnimNotifyNode::Construct(const FArguments& InArgs)
 {
 	Sequence = InArgs._Sequence;
-	NotifyEvent = InArgs._AnimNotify;
 	Font = FSlateFontInfo( FPaths::EngineContentDir() / TEXT("Slate/Fonts/Roboto-Regular.ttf"), 10 );
 	bBeingDragged = false;
 	CurrentDragHandle = ENotifyStateHandleHit::None;
@@ -1051,15 +1375,22 @@ void SAnimNotifyNode::Construct(const FArguments& InArgs)
 	bSelected = false;
 	DragMarkerTransactionIdx = INDEX_NONE;
 	
-	// Cache notify name for blueprint / Native notifies.
-	if(NotifyEvent->Notify)
+
+	if (InArgs._AnimNotify)
+	{
+		MakeNodeInterface<FNotifyNodeInterface>(InArgs._AnimNotify);
+	}
+	else if (InArgs._AnimSyncMarker)
 	{
-		NotifyEvent->NotifyName = FName(*NotifyEvent->Notify->GetNotifyName());
+		MakeNodeInterface<FSyncMarkerNodeInterface>(InArgs._AnimSyncMarker);
 	}
-	else if(NotifyEvent->NotifyStateClass)
+	else
 	{
-		NotifyEvent->NotifyName = FName(*NotifyEvent->NotifyStateClass->GetNotifyName());
+		check(false);	// Must specify something for this node to represent
+						// Either AnimNotify or AnimSyncMarker
 	}
+	// Cache notify name for blueprint / Native notifies.
+	NodeObjectInterface->CacheName();
 
 	OnNodeDragStarted = InArgs._OnNodeDragStarted;
 	PanTrackRequest = InArgs._PanTrackRequest;
@@ -1116,21 +1447,8 @@ FReply SAnimNotifyNode::OnDragDetected( const FGeometry& MyGeometry, const FPoin
 
 FLinearColor SAnimNotifyNode::GetNotifyColor() const
 {
-	FLinearColor BaseColor;
-
-	if (NotifyEvent->Notify)
-	{
-		BaseColor = NotifyEvent->Notify->GetEditorColor();
-	}
-	else if (NotifyEvent->NotifyStateClass)
-	{
-		BaseColor = NotifyEvent->NotifyStateClass->GetEditorColor();
-	}
-	else
-	{
-		// Color for Custom notifies
-		BaseColor = FLinearColor(1, 1, 0.5f);
-	}
+	TOptional<FLinearColor> Color = NodeObjectInterface->GetEditorColor();
+	FLinearColor BaseColor = Color.Get(FLinearColor(1, 1, 0.5f));
 
 	BaseColor.A = 0.67f;
 
@@ -1140,47 +1458,19 @@ FLinearColor SAnimNotifyNode::GetNotifyColor() const
 FText SAnimNotifyNode::GetNotifyText() const
 {
 	// Combine comment from notify struct and from function on object
-	return FText::FromName( NotifyEvent->NotifyName );
+	return FText::FromName( NodeObjectInterface->GetName() );
 }
 
 FText SAnimNotifyNode::GetNodeTooltip() const
 {
-	const float Time = NotifyEvent->GetTime();
-	const FText Frame = FText::AsNumber(Sequence->GetFrameAtTime(NotifyEvent->GetTime()));
-	const FText Seconds = FText::AsNumber( Time );
-
-	FText ToolTipText;
-	if (NotifyEvent->GetDuration() > 0.0f)
-	{
-		const FText Duration = FText::AsNumber( NotifyEvent->GetDuration() );
-		ToolTipText = FText::Format( LOCTEXT("NodeToolTipLong", "@ {0} sec (frame {1}) for {2} sec"), Seconds, Frame, Duration );
-	}
-	else
-	{
-		ToolTipText = FText::Format(LOCTEXT("NodeToolTipShort", "@ {0} sec (frame {1})"), Seconds, Frame);
-	}
-
-	if (NotifyEvent->IsBranchingPoint())
-	{
-		ToolTipText = FText::Format(LOCTEXT("AnimNotify_ToolTipBranchingPoint", "{0} (BranchingPoint)"), ToolTipText);
-	}
-	return ToolTipText;
+	return NodeObjectInterface->GetNodeTooltip(Sequence);
 }
 
 /** @return the Node's position within the graph */
 UObject* SAnimNotifyNode::GetObjectBeingDisplayed() const
 {
-	if (NotifyEvent->Notify)
-	{
-		return NotifyEvent->Notify;
-	}
-
-	if (NotifyEvent->NotifyStateClass)
-	{
-		return NotifyEvent->NotifyStateClass;
-	}
-
-	return Sequence;
+	TOptional<UObject*> Object = NodeObjectInterface->GetObjectBeingDisplayed();
+	return Object.Get(Sequence);
 }
 
 void SAnimNotifyNode::DropCancelled()
@@ -1243,14 +1533,14 @@ void SAnimNotifyNode::UpdateSizeAndPosition(const FGeometry& AllottedGeometry)
 	// Cache the geometry information, the alloted geometry is the same size as the track.
 	CachedAllotedGeometrySize = AllottedGeometry.Size;
 
-	NotifyTimePositionX = ScaleInfo.InputToLocalX(NotifyEvent->GetTime());
-	NotifyDurationSizeX = ScaleInfo.PixelsPerInput * NotifyEvent->GetDuration();
+	NotifyTimePositionX = ScaleInfo.InputToLocalX(NodeObjectInterface->GetTime());
+	NotifyDurationSizeX = ScaleInfo.PixelsPerInput * NodeObjectInterface->GetDuration();
 
 	const TSharedRef< FSlateFontMeasure > FontMeasureService = FSlateApplication::Get().GetRenderer()->GetFontMeasureService();
 	TextSize = FontMeasureService->Measure( GetNotifyText(), Font );
 	LabelWidth = TextSize.X + (TextBorderSize.X * 2.f) + (ScrubHandleSize.X / 2.f);
 
-	bool bDrawBranchingPoint = NotifyEvent->IsBranchingPoint();
+	bool bDrawBranchingPoint = NodeObjectInterface->IsBranchingPoint();
 	BranchingPointIconSize = FVector2D(TextSize.Y, TextSize.Y);
 	if (bDrawBranchingPoint)
 	{
@@ -1309,6 +1599,8 @@ int32 SAnimNotifyNode::OnPaint(const FPaintArgs& Args, const FGeometry& Allotted
 	int32 TextLayerID = ScrubHandleID + 1;
 	int32 BranchPointLayerID = TextLayerID + 1;
 
+	FAnimNotifyEvent* AnimNotifyEvent = NodeObjectInterface->GetNotifyEvent();
+
 	// Paint marker node if we have one
 	if(EndMarkerNodeOverlay.IsValid())
 	{
@@ -1335,7 +1627,7 @@ int32 SAnimNotifyNode::OnPaint(const FPaintArgs& Args, const FGeometry& Allotted
 	// Show duration of AnimNotifyState
 	if( NotifyDurationSizeX > 0.f )
 	{
-		FLinearColor BoxColor = (NotifyEvent->TrackIndex % 2) == 0 ? FLinearColor(0.0f,1.0f,0.5f,0.5f) : FLinearColor(0.0f,0.5f,1.0f,0.5f);
+		FLinearColor BoxColor = (NodeObjectInterface->GetTrackIndex() % 2) == 0 ? FLinearColor(0.0f, 1.0f, 0.5f, 0.5f) : FLinearColor(0.0f, 0.5f, 1.0f, 0.5f);
 		FVector2D DurationBoxSize = FVector2D(NotifyDurationSizeX, NotifyHeight);
 		FVector2D DurationBoxPosition = FVector2D(NotifyScrubHandleCentre, 0.f);
 		FSlateDrawElement::MakeBox( 
@@ -1350,20 +1642,20 @@ int32 SAnimNotifyNode::OnPaint(const FPaintArgs& Args, const FGeometry& Allotted
 		DrawScrubHandle(DurationBoxPosition.X + DurationBoxSize.X, OutDrawElements, ScrubHandleID, AllottedGeometry, MyClippingRect, NodeColour);
 		
 		// Render offsets if necessary
-		if(NotifyEvent->EndTriggerTimeOffset != 0.f) //Do we have an offset to render?
+		if(AnimNotifyEvent && AnimNotifyEvent->EndTriggerTimeOffset != 0.f) //Do we have an offset to render?
 		{
-			float EndTime = NotifyEvent->GetTime() + NotifyEvent->GetDuration();
+			float EndTime = AnimNotifyEvent->GetTime() + AnimNotifyEvent->GetDuration();
 			if(EndTime != Sequence->SequenceLength) //Don't render offset when we are at the end of the sequence, doesnt help the user
 			{
 				// ScrubHandle
 				float HandleCentre = NotifyDurationSizeX + ScrubHandleSize.X;
-				DrawHandleOffset(NotifyEvent->EndTriggerTimeOffset, HandleCentre, OutDrawElements, MarkerLayer, AllottedGeometry, MyClippingRect);
+				DrawHandleOffset(AnimNotifyEvent->EndTriggerTimeOffset, HandleCentre, OutDrawElements, MarkerLayer, AllottedGeometry, MyClippingRect);
 			}
 		}
 	}
 
 	// Branching point
-	bool bDrawBranchingPoint = NotifyEvent->IsBranchingPoint();
+	bool bDrawBranchingPoint = AnimNotifyEvent && AnimNotifyEvent->IsBranchingPoint();
 
 	// Background
 	FVector2D LabelSize = TextSize + TextBorderSize * 2.f;
@@ -1431,15 +1723,15 @@ int32 SAnimNotifyNode::OnPaint(const FPaintArgs& Args, const FGeometry& Allotted
 
 	DrawScrubHandle(NotifyScrubHandleCentre , OutDrawElements, ScrubHandleID, AllottedGeometry, MyClippingRect, NodeColour);
 
-	if(NotifyEvent->TriggerTimeOffset != 0.f) //Do we have an offset to render?
+	if(AnimNotifyEvent && AnimNotifyEvent->TriggerTimeOffset != 0.f) //Do we have an offset to render?
 	{
-		float NotifyTime = NotifyEvent->GetTime();
+		float NotifyTime = AnimNotifyEvent->GetTime();
 		if(NotifyTime != 0.f && NotifyTime != Sequence->SequenceLength) //Don't render offset when we are at the start/end of the sequence, doesn't help the user
 		{
 			float HandleCentre = NotifyScrubHandleCentre;
-			float &Offset = NotifyEvent->TriggerTimeOffset;
+			float &Offset = AnimNotifyEvent->TriggerTimeOffset;
 			
-			DrawHandleOffset(NotifyEvent->TriggerTimeOffset, NotifyScrubHandleCentre, OutDrawElements, MarkerLayer, AllottedGeometry, MyClippingRect);
+			DrawHandleOffset(AnimNotifyEvent->TriggerTimeOffset, NotifyScrubHandleCentre, OutDrawElements, MarkerLayer, AllottedGeometry, MyClippingRect);
 		}
 	}
 
@@ -1477,12 +1769,12 @@ FReply SAnimNotifyNode::OnMouseMove( const FGeometry& MyGeometry, const FPointer
 	if(CurrentDragHandle == ENotifyStateHandleHit::Start)
 	{
 		// Check track bounds
-		float OldDisplayTime = NotifyEvent->GetTime();
+		float OldDisplayTime = NodeObjectInterface->GetTime();
 
 		if(MouseEvent.GetScreenSpacePosition().X >= TrackScreenSpaceXPosition && MouseEvent.GetScreenSpacePosition().X <= TrackScreenSpaceXPosition + CachedAllotedGeometrySize.X)
 		{
 			float NewDisplayTime = ScaleInfo.LocalXToInput((MouseEvent.GetScreenSpacePosition() - MyGeometry.AbsolutePosition + XPositionInTrack).X);
-			float NewDuration = NotifyEvent->GetDuration() + OldDisplayTime - NewDisplayTime;
+			float NewDuration = NodeObjectInterface->GetDuration() + OldDisplayTime - NewDisplayTime;
 
 			// Check to make sure the duration is not less than the minimum allowed
 			if(NewDuration < MinimumStateDuration)
@@ -1490,10 +1782,10 @@ FReply SAnimNotifyNode::OnMouseMove( const FGeometry& MyGeometry, const FPointer
 				NewDisplayTime -= MinimumStateDuration - NewDuration;
 			}
 
-			NotifyEvent->SetTime(NewDisplayTime);
-			NotifyEvent->SetDuration(NotifyEvent->GetDuration() + OldDisplayTime - NotifyEvent->GetTime());
+			NodeObjectInterface->SetTime(NewDisplayTime);
+			NodeObjectInterface->SetDuration(NodeObjectInterface->GetDuration() + OldDisplayTime - NodeObjectInterface->GetTime());
 		}
-		else if(NotifyEvent->GetDuration() > MinimumStateDuration)
+		else if(NodeObjectInterface->GetDuration() > MinimumStateDuration)
 		{
 			float Overflow = HandleOverflowPan(MouseEvent.GetScreenSpacePosition(), TrackScreenSpaceXPosition);
 
@@ -1501,48 +1793,51 @@ FReply SAnimNotifyNode::OnMouseMove( const FGeometry& MyGeometry, const FPointer
 			ScaleInfo.ViewMinInput = ViewInputMin.Get();
 			ScaleInfo.ViewMaxInput = ViewInputMax.Get();
 
-			NotifyEvent->SetTime(ScaleInfo.LocalXToInput(Overflow < 0.0f ? 0.0f : CachedAllotedGeometrySize.X));
-			NotifyEvent->SetDuration(NotifyEvent->GetDuration() + OldDisplayTime - NotifyEvent->GetTime());
+			NodeObjectInterface->SetTime(ScaleInfo.LocalXToInput(Overflow < 0.0f ? 0.0f : CachedAllotedGeometrySize.X));
+			NodeObjectInterface->SetDuration(NodeObjectInterface->GetDuration() + OldDisplayTime - NodeObjectInterface->GetTime());
 
 			// Adjust incase we went under the minimum
-			if(NotifyEvent->GetDuration() < MinimumStateDuration)
+			if(NodeObjectInterface->GetDuration() < MinimumStateDuration)
 			{
-				float EndTimeBefore = NotifyEvent->GetTime() + NotifyEvent->GetDuration();
-				NotifyEvent->SetTime(NotifyEvent->GetTime() + NotifyEvent->GetDuration() - MinimumStateDuration);
-				NotifyEvent->SetDuration(MinimumStateDuration);
-				float EndTimeAfter = NotifyEvent->GetTime() + NotifyEvent->GetDuration();
+				float EndTimeBefore = NodeObjectInterface->GetTime() + NodeObjectInterface->GetDuration();
+				NodeObjectInterface->SetTime(NodeObjectInterface->GetTime() + NodeObjectInterface->GetDuration() - MinimumStateDuration);
+				NodeObjectInterface->SetDuration(MinimumStateDuration);
+				float EndTimeAfter = NodeObjectInterface->GetTime() + NodeObjectInterface->GetDuration();
 			}
 		}
 
 		// Now we know where the marker should be, look for possible snaps on montage marker bars
-		float NodePositionX = ScaleInfo.InputToLocalX(NotifyEvent->GetTime());
-		float MarkerSnap = GetScrubHandleSnapPosition(NodePositionX, ENotifyStateHandleHit::Start);
-		if(MarkerSnap != -1.0f)
-		{
-			// We're near to a snap bar
-			EAnimEventTriggerOffsets::Type Offset = (MarkerSnap < NodePositionX) ? EAnimEventTriggerOffsets::OffsetAfter : EAnimEventTriggerOffsets::OffsetBefore;
-			NotifyEvent->TriggerTimeOffset = GetTriggerTimeOffsetForType(Offset);
-			NodePositionX = MarkerSnap;
-
-			// Adjust our start marker
-			OldDisplayTime = NotifyEvent->GetTime();
-			NotifyEvent->SetTime(ScaleInfo.LocalXToInput(NodePositionX));
-			NotifyEvent->SetDuration(NotifyEvent->GetDuration() + OldDisplayTime - NotifyEvent->GetTime());
-		}
-		else
+		if (FAnimNotifyEvent* AnimNotifyEvent = NodeObjectInterface->GetNotifyEvent())
 		{
-			NotifyEvent->TriggerTimeOffset = GetTriggerTimeOffsetForType(EAnimEventTriggerOffsets::NoOffset);
+			float NodePositionX = ScaleInfo.InputToLocalX(AnimNotifyEvent->GetTime());
+			float MarkerSnap = GetScrubHandleSnapPosition(NodePositionX, ENotifyStateHandleHit::Start);
+			if (MarkerSnap != -1.0f)
+			{
+				// We're near to a snap bar
+				EAnimEventTriggerOffsets::Type Offset = (MarkerSnap < NodePositionX) ? EAnimEventTriggerOffsets::OffsetAfter : EAnimEventTriggerOffsets::OffsetBefore;
+				AnimNotifyEvent->TriggerTimeOffset = GetTriggerTimeOffsetForType(Offset);
+				NodePositionX = MarkerSnap;
+
+				// Adjust our start marker
+				OldDisplayTime = AnimNotifyEvent->GetTime();
+				AnimNotifyEvent->SetTime(ScaleInfo.LocalXToInput(NodePositionX));
+				AnimNotifyEvent->SetDuration(AnimNotifyEvent->GetDuration() + OldDisplayTime - AnimNotifyEvent->GetTime());
+			}
+			else
+			{
+				AnimNotifyEvent->TriggerTimeOffset = GetTriggerTimeOffsetForType(EAnimEventTriggerOffsets::NoOffset);
+			}
 		}
 	}
 	else
 	{
 		if(MouseEvent.GetScreenSpacePosition().X >= TrackScreenSpaceXPosition && MouseEvent.GetScreenSpacePosition().X <= TrackScreenSpaceXPosition + CachedAllotedGeometrySize.X)
 		{
-			float NewDuration = ScaleInfo.LocalXToInput((MouseEvent.GetScreenSpacePosition() - MyGeometry.AbsolutePosition + XPositionInTrack).X) - NotifyEvent->GetTime();
+			float NewDuration = ScaleInfo.LocalXToInput((MouseEvent.GetScreenSpacePosition() - MyGeometry.AbsolutePosition + XPositionInTrack).X) - NodeObjectInterface->GetTime();
 
-			NotifyEvent->SetDuration(FMath::Max(NewDuration, MinimumStateDuration));
+			NodeObjectInterface->SetDuration(FMath::Max(NewDuration, MinimumStateDuration));
 		}
-		else if(NotifyEvent->GetDuration() > MinimumStateDuration)
+		else if(NodeObjectInterface->GetDuration() > MinimumStateDuration)
 		{
 			float Overflow = HandleOverflowPan(MouseEvent.GetScreenSpacePosition(), TrackScreenSpaceXPosition);
 
@@ -1550,25 +1845,28 @@ FReply SAnimNotifyNode::OnMouseMove( const FGeometry& MyGeometry, const FPointer
 			ScaleInfo.ViewMinInput = ViewInputMin.Get();
 			ScaleInfo.ViewMaxInput = ViewInputMax.Get();
 
-			NotifyEvent->SetDuration(FMath::Max(ScaleInfo.LocalXToInput(Overflow > 0.0f ? CachedAllotedGeometrySize.X : 0.0f) - NotifyEvent->GetTime(), MinimumStateDuration));
+			NodeObjectInterface->SetDuration(FMath::Max(ScaleInfo.LocalXToInput(Overflow > 0.0f ? CachedAllotedGeometrySize.X : 0.0f) - NodeObjectInterface->GetTime(), MinimumStateDuration));
 		}
 
 		// Now we know where the scrub handle should be, look for possible snaps on montage marker bars
-		float NodePositionX = ScaleInfo.InputToLocalX(NotifyEvent->GetTime() + NotifyEvent->GetDuration());
-		float MarkerSnap = GetScrubHandleSnapPosition(NodePositionX, ENotifyStateHandleHit::End);
-		if(MarkerSnap != -1.0f)
+		if (FAnimNotifyEvent* AnimNotifyEvent = NodeObjectInterface->GetNotifyEvent())
 		{
-			// We're near to a snap bar
-			EAnimEventTriggerOffsets::Type Offset = (MarkerSnap < NodePositionX) ? EAnimEventTriggerOffsets::OffsetAfter : EAnimEventTriggerOffsets::OffsetBefore;
-			NotifyEvent->EndTriggerTimeOffset = GetTriggerTimeOffsetForType(Offset);
-			NodePositionX = MarkerSnap;
+			float NodePositionX = ScaleInfo.InputToLocalX(AnimNotifyEvent->GetTime() + AnimNotifyEvent->GetDuration());
+			float MarkerSnap = GetScrubHandleSnapPosition(NodePositionX, ENotifyStateHandleHit::End);
+			if (MarkerSnap != -1.0f)
+			{
+				// We're near to a snap bar
+				EAnimEventTriggerOffsets::Type Offset = (MarkerSnap < NodePositionX) ? EAnimEventTriggerOffsets::OffsetAfter : EAnimEventTriggerOffsets::OffsetBefore;
+				AnimNotifyEvent->EndTriggerTimeOffset = GetTriggerTimeOffsetForType(Offset);
+				NodePositionX = MarkerSnap;
 
-			// Adjust our end marker
-			NotifyEvent->SetDuration(ScaleInfo.LocalXToInput(NodePositionX) - NotifyEvent->GetTime());
-		}
-		else
-		{
-			NotifyEvent->EndTriggerTimeOffset = GetTriggerTimeOffsetForType(EAnimEventTriggerOffsets::NoOffset);
+				// Adjust our end marker
+				AnimNotifyEvent->SetDuration(ScaleInfo.LocalXToInput(NodePositionX) - AnimNotifyEvent->GetTime());
+			}
+			else
+			{
+				AnimNotifyEvent->EndTriggerTimeOffset = GetTriggerTimeOffsetForType(EAnimEventTriggerOffsets::NoOffset);
+			}
 		}
 	}
 
@@ -1746,6 +2044,7 @@ void SAnimNotifyTrack::Construct(const FArguments& InArgs)
 	ViewInputMax = InArgs._ViewInputMax;
 	OnSelectionChanged = InArgs._OnSelectionChanged;
 	AnimNotifies = InArgs._AnimNotifies;
+	AnimSyncMarkers = InArgs._AnimSyncMarkers;
 	OnUpdatePanel = InArgs._OnUpdatePanel;
 	OnGetNotifyBlueprintData = InArgs._OnGetNotifyBlueprintData;
 	OnGetNotifyStateBlueprintData = InArgs._OnGetNotifyStateBlueprintData;
@@ -1764,8 +2063,8 @@ void SAnimNotifyTrack::Construct(const FArguments& InArgs)
 	OnReplaceSelectedWithNotify = InArgs._OnReplaceSelectedWithNotify;
 	OnReplaceSelectedWithBlueprintNotify = InArgs._OnReplaceSelectedWithBlueprintNotify;
 	OnDeselectAllNotifies = InArgs._OnDeselectAllNotifies;
-	OnCopyNotifies = InArgs._OnCopyNotifies;
-	OnPasteNotifies = InArgs._OnPasteNotifies;
+	OnCopyNodes = InArgs._OnCopyNodes;
+	OnPasteNodes = InArgs._OnPasteNodes;
 	OnSetInputViewRange = InArgs._OnSetInputViewRange;
 	OnGetTimingNodeVisibility = InArgs._OnGetTimingNodeVisibility;
 
@@ -2146,8 +2445,17 @@ void SAnimNotifyTrack::FillNewNotifyMenu(FMenuBuilder& MenuBuilder, bool bIsRepl
 			MenuBuilder.AddMenuEntry(LOCTEXT("NewNotify", "New Notify"), LOCTEXT("NewNotifyToolTip", "Create a new animation notify"), FSlateIcon(), UIAction);
 		}
 		MenuBuilder.EndSection();
-	}
-}
+
+		MenuBuilder.BeginSection("SyncMarkerCreateNew");
+		{
+			FUIAction UIAction;
+			UIAction.ExecuteAction.BindRaw(
+				this, &SAnimNotifyTrack::OnNewSyncMarkerClicked);
+			MenuBuilder.AddMenuEntry(LOCTEXT("NewSyncMarker", "New Sync Marker"), LOCTEXT("NewSyncMarkerToolTip", "Create a new animation sync marker"), FSlateIcon(), UIAction);
+		}
+		MenuBuilder.EndSection();
+	}
+}
 
 FAnimNotifyEvent& SAnimNotifyTrack::CreateNewBlueprintNotify(FString NewNotifyName, FString BlueprintPath, float StartTime)
 {
@@ -2242,6 +2550,20 @@ void SAnimNotifyTrack::CreateNewNotifyAtCursor(FString NewNotifyName, UClass* No
 	OnUpdatePanel.ExecuteIfBound();
 }
 
+void SAnimNotifyTrack::CreateNewSyncMarkerAtCursor(FString NewSyncMarkerName, UClass* NotifyClass)
+{
+	UAnimSequence* Seq = CastChecked<UAnimSequence>(Sequence);
+
+	FScopedTransaction Transaction(LOCTEXT("AddSyncMarker", "Add Sync Marker"));
+	Seq->Modify();
+	int32 NewIndex = Seq->AuthoredSyncMarkers.Add(FAnimSyncMarker());
+	FAnimSyncMarker& SyncMarker = Seq->AuthoredSyncMarkers[NewIndex];
+	SyncMarker.MarkerName = FName(*NewSyncMarkerName);
+	SyncMarker.TrackIndex = TrackIndex;
+	SyncMarker.Time = LastClickedTime;
+	OnUpdatePanel.ExecuteIfBound();
+}
+
 void SAnimNotifyTrack::ReplaceSelectedWithBlueprintNotify(FString NewNotifyName, FString BlueprintPath)
 {
 	OnReplaceSelectedWithBlueprintNotify.ExecuteIfBound(NewNotifyName, BlueprintPath);
@@ -2295,11 +2617,11 @@ FReply SAnimNotifyTrack::OnMouseButtonUp( const FGeometry& MyGeometry, const FPo
 		{
 			if(bCtrl)
 			{
-				ToggleNotifyNodeSelectionStatus(NotifyIndex);
+				ToggleTrackObjectNodeSelectionStatus(NotifyIndex);
 			}
 			else
 			{
-				SelectNotifyNode(NotifyIndex, bShift);
+				SelectTrackObjectNode(NotifyIndex, bShift);
 			}
 		}
 
@@ -2309,9 +2631,9 @@ FReply SAnimNotifyTrack::OnMouseButtonUp( const FGeometry& MyGeometry, const FPo
 	return FReply::Unhandled();
 }
 
-void SAnimNotifyTrack::SelectNotifyNode(int32 NotifyIndex, bool Append, bool bUpdateSelection)
+void SAnimNotifyTrack::SelectTrackObjectNode(int32 TrackNodeIndex, bool Append, bool bUpdateSelection)
 {
-	if( NotifyIndex != INDEX_NONE )
+	if( TrackNodeIndex != INDEX_NONE )
 	{
 		// Deselect all other notifies if necessary.
 		if (Sequence && !Append)
@@ -2320,14 +2642,14 @@ void SAnimNotifyTrack::SelectNotifyNode(int32 NotifyIndex, bool Append, bool bUp
 		}
 
 		// Check to see if we've already selected this node
-		if (!SelectedNotifyIndices.Contains(NotifyIndex))
+		if (!SelectedNodeIndices.Contains(TrackNodeIndex))
 		{
 			// select new one
-			if (NotifyNodes.IsValidIndex(NotifyIndex))
+			if (NotifyNodes.IsValidIndex(TrackNodeIndex))
 			{
-				TSharedPtr<SAnimNotifyNode> Node = NotifyNodes[NotifyIndex];
+				TSharedPtr<SAnimNotifyNode> Node = NotifyNodes[TrackNodeIndex];
 				Node->bSelected = true;
-				SelectedNotifyIndices.Add(NotifyIndex);
+				SelectedNodeIndices.Add(TrackNodeIndex);
 
 				if(bUpdateSelection)
 				{
@@ -2338,21 +2660,21 @@ void SAnimNotifyTrack::SelectNotifyNode(int32 NotifyIndex, bool Append, bool bUp
 	}
 }
 
-void SAnimNotifyTrack::ToggleNotifyNodeSelectionStatus( int32 NotifyIndex, bool bUpdateSelection )
+void SAnimNotifyTrack::ToggleTrackObjectNodeSelectionStatus( int32 TrackNodeIndex, bool bUpdateSelection )
 {
-	check(NotifyNodes.IsValidIndex(NotifyIndex));
+	check(NotifyNodes.IsValidIndex(TrackNodeIndex));
 
-	bool bSelected = SelectedNotifyIndices.Contains(NotifyIndex);
+	bool bSelected = SelectedNodeIndices.Contains(TrackNodeIndex);
 	if(bSelected)
 	{
-		SelectedNotifyIndices.Remove(NotifyIndex);
+		SelectedNodeIndices.Remove(TrackNodeIndex);
 	}
 	else
 	{
-		SelectedNotifyIndices.Add(NotifyIndex);
+		SelectedNodeIndices.Add(TrackNodeIndex);
 	}
 
-	TSharedPtr<SAnimNotifyNode> Node = NotifyNodes[NotifyIndex];
+	TSharedPtr<SAnimNotifyNode> Node = NotifyNodes[TrackNodeIndex];
 	Node->bSelected = !Node->bSelected;
 
 	if(bUpdateSelection)
@@ -2361,13 +2683,13 @@ void SAnimNotifyTrack::ToggleNotifyNodeSelectionStatus( int32 NotifyIndex, bool
 	}
 }
 
-void SAnimNotifyTrack::DeselectNotifyNode( int32 NotifyIndex, bool bUpdateSelection )
+void SAnimNotifyTrack::DeselectTrackObjectNode( int32 TrackNodeIndex, bool bUpdateSelection )
 {
-	check(NotifyNodes.IsValidIndex(NotifyIndex));
-	TSharedPtr<SAnimNotifyNode> Node = NotifyNodes[NotifyIndex];
+	check(NotifyNodes.IsValidIndex(TrackNodeIndex));
+	TSharedPtr<SAnimNotifyNode> Node = NotifyNodes[TrackNodeIndex];
 	Node->bSelected = false;
 
-	int32 ItemsRemoved = SelectedNotifyIndices.Remove(NotifyIndex);
+	int32 ItemsRemoved = SelectedNodeIndices.Remove(TrackNodeIndex);
 	check(ItemsRemoved > 0);
 
 	if(bUpdateSelection)
@@ -2382,7 +2704,7 @@ void SAnimNotifyTrack::DeselectAllNotifyNodes(bool bUpdateSelectionSet)
 	{
 		Node->bSelected = false;
 	}
-	SelectedNotifyIndices.Empty();
+	SelectedNodeIndices.Empty();
 
 	if(bUpdateSelectionSet)
 	{
@@ -2393,65 +2715,70 @@ void SAnimNotifyTrack::DeselectAllNotifyNodes(bool bUpdateSelectionSet)
 TSharedPtr<SWidget> SAnimNotifyTrack::SummonContextMenu(const FGeometry& MyGeometry, const FPointerEvent& MouseEvent)
 {
 	FVector2D CursorPos = MouseEvent.GetScreenSpacePosition();
-	int32 NotifyIndex = GetHitNotifyNode(MyGeometry, MyGeometry.AbsoluteToLocal(CursorPos));
+	int32 NodeIndex = GetHitNotifyNode(MyGeometry, MyGeometry.AbsoluteToLocal(CursorPos));
 	LastClickedTime = CalculateTime(MyGeometry, MouseEvent.GetScreenSpacePosition());
 
 	const bool bCloseWindowAfterMenuSelection = true;
 	FMenuBuilder MenuBuilder( bCloseWindowAfterMenuSelection, AnimSequenceEditorActions );
 	FUIAction NewAction;
 
+	INodeObjectInterface* NodeObject = NodeIndex != INDEX_NONE ? NotifyNodes[NodeIndex]->NodeObjectInterface : nullptr;
+	FAnimNotifyEvent* NotifyEvent = NodeObject ? NodeObject->GetNotifyEvent(): nullptr;
+	int32 NotifyIndex = NotifyEvent ? AnimNotifies.IndexOfByKey(NotifyEvent) : INDEX_NONE;
+
 	MenuBuilder.BeginSection("AnimNotify", LOCTEXT("NotifyHeading", "Notify") );
 	{
-		if ( NotifyIndex != INDEX_NONE )
+		if (NodeObject)
 		{
-			if (!NotifyNodes[NotifyIndex]->bSelected)
+			if (!NotifyNodes[NodeIndex]->bSelected)
 			{
-				SelectNotifyNode(NotifyIndex, MouseEvent.IsControlDown());
+				SelectTrackObjectNode(NodeIndex, MouseEvent.IsControlDown());
 			}
 
-			FAnimNotifyEvent* Event = AnimNotifies[NotifyIndex];
-
 			FNumberFormattingOptions Options;
 			Options.MinimumFractionalDigits = 5;
 
 			// Add item to directly set notify time
-			const FText CurrentTime = FText::AsNumber(Event->GetTime(), &Options);
+			const FText CurrentTime = FText::AsNumber(NodeObject->GetTime(), &Options);
 			const FText TimeMenuText = FText::Format(LOCTEXT("TimeMenuText", "Set Notify Begin Time: {0}..."), CurrentTime);
 
-			NewAction.ExecuteAction.BindRaw(this, &SAnimNotifyTrack::OnSetNotifyTimeClicked, NotifyIndex);
+			NewAction.ExecuteAction.BindRaw(this, &SAnimNotifyTrack::OnSetNodeTimeClicked, NodeIndex);
 			NewAction.CanExecuteAction.BindRaw(this, &SAnimNotifyTrack::IsSingleNodeSelected);
 
 			MenuBuilder.AddMenuEntry(TimeMenuText, LOCTEXT("SetTimeToolTip", "Set the time of this notify directly"), FSlateIcon(), NewAction);
 
 			// Add item to directly set notify frame
-			const FText Frame = FText::AsNumber(Sequence->GetFrameAtTime(Event->GetTime()));
+			const FText Frame = FText::AsNumber(Sequence->GetFrameAtTime(NodeObject->GetTime()));
 			const FText FrameMenuText = FText::Format(LOCTEXT("FrameMenuText", "Set Notify Frame: {0}..."), Frame);
 
-			NewAction.ExecuteAction.BindRaw(this, &SAnimNotifyTrack::OnSetNotifyFrameClicked, NotifyIndex);
+			NewAction.ExecuteAction.BindRaw(this, &SAnimNotifyTrack::OnSetNodeFrameClicked, NodeIndex);
 			NewAction.CanExecuteAction.BindRaw(this, &SAnimNotifyTrack::IsSingleNodeSelected);
 
 			MenuBuilder.AddMenuEntry(FrameMenuText, LOCTEXT("SetFrameToolTip", "Set the frame of this notify directly"), FSlateIcon(), NewAction);
 
-			// add menu to get threshold weight for triggering this notify
-			NewAction.ExecuteAction.BindRaw(
-				this, &SAnimNotifyTrack::OnSetTriggerWeightNotifyClicked, NotifyIndex);
-			NewAction.CanExecuteAction.BindRaw(
-				this, &SAnimNotifyTrack::IsSingleNodeSelected);
-
-			const FText Threshold = FText::AsNumber( AnimNotifies[NotifyIndex]->TriggerWeightThreshold, &Options );
-			const FText MinTriggerWeightText = FText::Format( LOCTEXT("MinTriggerWeight", "Min Trigger Weight: {0}..."), Threshold );
-			MenuBuilder.AddMenuEntry(MinTriggerWeightText, LOCTEXT("MinTriggerWeightToolTip", "The minimum weight to trigger this notify"), FSlateIcon(), NewAction);
-
-			// Add menu for changing duration if this is an AnimNotifyState
-			if( AnimNotifies[NotifyIndex]->NotifyStateClass )
+			if (NotifyEvent)
 			{
+				// add menu to get threshold weight for triggering this notify
 				NewAction.ExecuteAction.BindRaw(
-					this, &SAnimNotifyTrack::OnSetDurationNotifyClicked, NotifyIndex);
+					this, &SAnimNotifyTrack::OnSetTriggerWeightNotifyClicked, NotifyIndex);
 				NewAction.CanExecuteAction.BindRaw(
 					this, &SAnimNotifyTrack::IsSingleNodeSelected);
 
-				FText SetAnimStateDurationText = FText::Format( LOCTEXT("SetAnimStateDuration", "Set AnimNotifyState duration ({0})"), FText::AsNumber( AnimNotifies[NotifyIndex]->GetDuration() ) );
-				MenuBuilder.AddMenuEntry(SetAnimStateDurationText, LOCTEXT("SetAnimStateDuration_ToolTip", "The duration of this AnimNotifyState"), FSlateIcon(), NewAction);
+				const FText Threshold = FText::AsNumber(NotifyEvent->TriggerWeightThreshold, &Options);
+				const FText MinTriggerWeightText = FText::Format(LOCTEXT("MinTriggerWeight", "Min Trigger Weight: {0}..."), Threshold);
+				MenuBuilder.AddMenuEntry(MinTriggerWeightText, LOCTEXT("MinTriggerWeightToolTip", "The minimum weight to trigger this notify"), FSlateIcon(), NewAction);
+
+				// Add menu for changing duration if this is an AnimNotifyState
+				if (NotifyEvent->NotifyStateClass)
+				{
+					NewAction.ExecuteAction.BindRaw(
+						this, &SAnimNotifyTrack::OnSetDurationNotifyClicked, NotifyIndex);
+					NewAction.CanExecuteAction.BindRaw(
+						this, &SAnimNotifyTrack::IsSingleNodeSelected);
+
+					FText SetAnimStateDurationText = FText::Format(LOCTEXT("SetAnimStateDuration", "Set AnimNotifyState duration ({0})"), FText::AsNumber(NotifyEvent->GetDuration()));
+					MenuBuilder.AddMenuEntry(SetAnimStateDurationText, LOCTEXT("SetAnimStateDuration_ToolTip", "The duration of this AnimNotifyState"), FSlateIcon(), NewAction);
+				}
 			}
 
 		}
@@ -2480,37 +2807,38 @@ TSharedPtr<SWidget> SAnimNotifyTrack::SummonContextMenu(const FGeometry& MyGeome
 
 	MenuBuilder.BeginSection("AnimEdit", LOCTEXT("NotifyEditHeading", "Edit") );
 	{
-		if ( NotifyIndex != INDEX_NONE )
+		if ( NodeObject )
 		{
 			// copy notify menu item
 			NewAction.ExecuteAction.BindRaw(
-				this, &SAnimNotifyTrack::OnCopyNotifyClicked, NotifyIndex);
+				this, &SAnimNotifyTrack::OnCopyNotifyClicked, NodeIndex);
 			MenuBuilder.AddMenuEntry(LOCTEXT("Copy", "Copy"), LOCTEXT("CopyToolTip", "Copy animation notify event"), FSlateIcon(), NewAction);
 
 			// allow it to delete
 			NewAction.ExecuteAction = OnDeleteNotify;
 			MenuBuilder.AddMenuEntry(LOCTEXT("Delete", "Delete"), LOCTEXT("DeleteToolTip", "Delete animation notify"), FSlateIcon(), NewAction);
 
-			FAnimNotifyEvent* Notify = AnimNotifies[NotifyIndex];
-
-			// For the "Replace With..." menu, make sure the current AnimNotify selection is valid for replacement
-			if (OnGetIsAnimNotifySelectionValidforReplacement.IsBound() && OnGetIsAnimNotifySelectionValidforReplacement.Execute())
+			if (NotifyEvent)
 			{
-				// If this is an AnimNotifyState (has duration) allow it to be replaced with other AnimNotifyStates
-				if (Notify && Notify->NotifyStateClass)
-				{
-					MenuBuilder.AddSubMenu(
-						NSLOCTEXT("NewNotifySubMenu", "NewNotifySubMenuReplaceWithNotifyState", "Replace with Notify State..."),
-						NSLOCTEXT("NewNotifySubMenu", "NewNotifySubMenuReplaceWithNotifyStateToolTip", "Replace with AnimNotifyState"),
-						FNewMenuDelegate::CreateRaw(this, &SAnimNotifyTrack::FillNewNotifyStateMenu, true));
-				}
-				// If this is a regular AnimNotify (no duration) allow it to be replaced with other AnimNotifies
-				else 
+				// For the "Replace With..." menu, make sure the current AnimNotify selection is valid for replacement
+				if (OnGetIsAnimNotifySelectionValidforReplacement.IsBound() && OnGetIsAnimNotifySelectionValidforReplacement.Execute())
 				{
-					MenuBuilder.AddSubMenu(
-						NSLOCTEXT("NewNotifySubMenu", "NewNotifySubMenuReplaceWithNotify", "Replace with Notify..."),
-						NSLOCTEXT("NewNotifySubMenu", "NewNotifySubMenuReplaceWithNotifyToolTip", "Replace with AnimNotifyEvent"),
-						FNewMenuDelegate::CreateRaw(this, &SAnimNotifyTrack::FillNewNotifyMenu, true));
+					// If this is an AnimNotifyState (has duration) allow it to be replaced with other AnimNotifyStates
+					if (NotifyEvent->NotifyStateClass)
+					{
+						MenuBuilder.AddSubMenu(
+							NSLOCTEXT("NewNotifySubMenu", "NewNotifySubMenuReplaceWithNotifyState", "Replace with Notify State..."),
+							NSLOCTEXT("NewNotifySubMenu", "NewNotifySubMenuReplaceWithNotifyStateToolTip", "Replace with AnimNotifyState"),
+							FNewMenuDelegate::CreateRaw(this, &SAnimNotifyTrack::FillNewNotifyStateMenu, true));
+					}
+					// If this is a regular AnimNotify (no duration) allow it to be replaced with other AnimNotifies
+					else
+					{
+						MenuBuilder.AddSubMenu(
+							NSLOCTEXT("NewNotifySubMenu", "NewNotifySubMenuReplaceWithNotify", "Replace with Notify..."),
+							NSLOCTEXT("NewNotifySubMenu", "NewNotifySubMenuReplaceWithNotifyToolTip", "Replace with AnimNotifyEvent"),
+							FNewMenuDelegate::CreateRaw(this, &SAnimNotifyTrack::FillNewNotifyMenu, true));
+					}
 				}
 			}
 		}
@@ -2520,9 +2848,10 @@ TSharedPtr<SWidget> SAnimNotifyTrack::SummonContextMenu(const FGeometry& MyGeome
 			const TCHAR* Buffer;
 			float OriginalTime;
 			float OriginalLength;
+			int32 TrackSpan;
 
 			//Check whether can we show menu item to paste anim notify event
-			if( ReadNotifyPasteHeader(PropertyString, Buffer, OriginalTime, OriginalLength) )
+			if( ReadNotifyPasteHeader(PropertyString, Buffer, OriginalTime, OriginalLength,TrackSpan) )
 			{
 				// paste notify menu item
 				if (IsSingleNodeInClipboard())
@@ -2559,11 +2888,10 @@ TSharedPtr<SWidget> SAnimNotifyTrack::SummonContextMenu(const FGeometry& MyGeome
 	}
 	MenuBuilder.EndSection(); //AnimEdit
 
-	if (NotifyIndex != INDEX_NONE)
+	if (NotifyEvent)
 	{
-		FAnimNotifyEvent* Notify = AnimNotifies[NotifyIndex];
-		UObject* NotifyObject = Notify->Notify;
-		NotifyObject = NotifyObject ? NotifyObject : Notify->NotifyStateClass;
+		UObject* NotifyObject = NotifyEvent->Notify;
+		NotifyObject = NotifyObject ? NotifyObject : NotifyEvent->NotifyStateClass;
 
 		if (NotifyObject && Cast<UBlueprintGeneratedClass>(NotifyObject->GetClass()))
 		{
@@ -2607,7 +2935,7 @@ void SAnimNotifyTrack::CreateCommands()
 
 void SAnimNotifyTrack::OnCopyNotifyClicked(int32 NotifyIndex)
 {
-	OnCopyNotifies.ExecuteIfBound();
+	OnCopyNodes.ExecuteIfBound();
 }
 
 bool SAnimNotifyTrack::CanPasteAnimNotify() const
@@ -2616,154 +2944,14 @@ bool SAnimNotifyTrack::CanPasteAnimNotify() const
 	const TCHAR* Buffer;
 	float OriginalTime;
 	float OriginalLength;
-	return ReadNotifyPasteHeader(PropertyString, Buffer, OriginalTime, OriginalLength);
-}
-
-bool SAnimNotifyTrack::ReadNotifyPasteHeader(FString& OutPropertyString, const TCHAR*& OutBuffer, float& OutOriginalTime, float& OutOriginalLength) const
-{
-	OutBuffer = NULL;
-	OutOriginalTime = -1.f;
-
-	FPlatformMisc::ClipboardPaste(OutPropertyString);
-
-	if( !OutPropertyString.IsEmpty() )
-	{
-		//Remove header text
-		const FString HeaderString( TEXT("COPY_ANIMNOTIFYEVENT") );
-		
-		//Check for string identifier in order to determine whether the text represents an FAnimNotifyEvent.
-		if(OutPropertyString.StartsWith( HeaderString ) && OutPropertyString.Len() > HeaderString.Len() )
-		{
-			int32 HeaderSize = HeaderString.Len();
-			OutBuffer = *OutPropertyString;
-			OutBuffer += HeaderSize;
-
-			FString ReadLine;
-			// Read the original time from the first notify
-			FParse::Line(&OutBuffer, ReadLine);
-			FParse::Value( *ReadLine, TEXT( "OriginalTime=" ), OutOriginalTime );
-			FParse::Value( *ReadLine, TEXT( "OriginalLength=" ), OutOriginalLength );
-			return true;
-		}
-	}
-
-	return false;
+	int32 TrackSpan;
+	return ReadNotifyPasteHeader(PropertyString, Buffer, OriginalTime, OriginalLength, TrackSpan);
 }
 
 void SAnimNotifyTrack::OnPasteNotifyClicked(ENotifyPasteMode::Type PasteMode, ENotifyPasteMultipleMode::Type MultiplePasteType)
 {
 	float ClickTime = PasteMode == ENotifyPasteMode::MousePosition ? LastClickedTime : -1.0f;
-	OnPasteNotifies.ExecuteIfBound(this, ClickTime, PasteMode, MultiplePasteType);
-}
-
-void SAnimNotifyTrack::OnPasteNotifyTimeSet(const FText& TimeText, ETextCommit::Type CommitInfo)
-{
-	if ( CommitInfo == ETextCommit::OnEnter || CommitInfo == ETextCommit::OnUserMovedFocus )
-	{
-		float PasteTime = FCString::Atof( *TimeText.ToString() );
-		PasteTime = FMath::Clamp<float>(PasteTime, 0.f, Sequence->SequenceLength-0.01f);
-		OnPasteNotify(PasteTime);
-	}
-
-	FSlateApplication::Get().DismissAllMenus();
-}
-
-void SAnimNotifyTrack::OnPasteNotify(float TimeToPasteAt, ENotifyPasteMultipleMode::Type MultiplePasteType)
-{
-	FString PropertyString;
-	const TCHAR* Buffer;
-	float OriginalTime;				// Time in the original track that the node was placed
-	float OriginalLength;			// Length of the original anim sequence
-	float Shift = 0.0f;				// The amount to shift the saved display time of subsequent nodes
-	float ScaleMultiplier = 1.0f;	// Scale value to place nodes relatively in different sequences
-
-	if( ReadNotifyPasteHeader(PropertyString, Buffer, OriginalTime, OriginalLength) )
-	{
-		// If we're using relative mode, calculate the scale between the two sequences
-		if (MultiplePasteType == ENotifyPasteMultipleMode::Relative)
-		{
-			ScaleMultiplier = Sequence->SequenceLength / OriginalLength;
-		}
-
-		//Store object to undo pasting anim notify
-		const FScopedTransaction Transaction( LOCTEXT("PasteNotifyEvent", "Paste Anim Notify") );
-
-		Sequence->Modify();
-
-		FString NotifyLine;
-		while (FParse::Line(&Buffer, NotifyLine))
-		{
-			// Add a new Anim Notify Event structure to the array
-			int32 NewNotifyIndex = Sequence->Notifies.AddZeroed();
-
-			//Get array property
-			UArrayProperty* ArrayProperty = NULL;
-			uint8* PropData = Sequence->FindNotifyPropertyData(NewNotifyIndex, ArrayProperty);
-
-			if (PropData && ArrayProperty)
-			{
-				ArrayProperty->Inner->ImportText(*NotifyLine, PropData, PPF_Copy, NULL);
-
-				FAnimNotifyEvent& NewNotify = Sequence->Notifies[NewNotifyIndex];
-
-				if (TimeToPasteAt >= 0.f)
-				{
-					Shift = TimeToPasteAt - (NewNotify.GetTime() * ScaleMultiplier);
-					NewNotify.SetTime(TimeToPasteAt);
-
-					// Don't want to paste subsequent nodes at this position
-					TimeToPasteAt = -1.0f;
-				}
-				else
-				{
-					NewNotify.SetTime((NewNotify.GetTime() * ScaleMultiplier) + Shift);
-					// Clamp into the bounds of the sequence
-					NewNotify.SetTime(FMath::Clamp(NewNotify.GetTime(), 0.0f, Sequence->SequenceLength));
-				}
-				NewNotify.TriggerTimeOffset = GetTriggerTimeOffsetForType(Sequence->CalculateOffsetForNotify(NewNotify.GetTime()));
-				NewNotify.TrackIndex = TrackIndex;
-
-				// Get anim notify ptr of the new notify event array element (This ptr is the same as the source anim notify ptr)
-				UAnimNotify* AnimNotifyPtr = NewNotify.Notify;
-				if (AnimNotifyPtr)
-				{
-					//Create a new instance of anim notify subclass.
-					UAnimNotify* NewAnimNotifyPtr = Cast<UAnimNotify>(StaticDuplicateObject(AnimNotifyPtr, Sequence, TEXT("None")));
-					if (NewAnimNotifyPtr)
-					{
-						//Reassign new anim notify ptr to the pasted anim notify event.
-						NewNotify.Notify = NewAnimNotifyPtr;
-					}
-				}
-
-				UAnimNotifyState * AnimNotifyStatePtr = NewNotify.NotifyStateClass;
-				if (AnimNotifyStatePtr)
-				{
-					//Create a new instance of anim notify subclass.
-					UAnimNotifyState * NewAnimNotifyStatePtr = Cast<UAnimNotifyState>(StaticDuplicateObject(AnimNotifyStatePtr, Sequence, TEXT("None")));
-					if (NewAnimNotifyStatePtr)
-					{
-						//Reassign new anim notify ptr to the pasted anim notify event.
-						NewNotify.NotifyStateClass = NewAnimNotifyStatePtr;
-					}
-
-					if ((NewNotify.GetTime() + NewNotify.GetDuration()) > Sequence->SequenceLength)
-					{
-						NewNotify.SetDuration(Sequence->SequenceLength - NewNotify.GetTime());
-					}
-				}
-			}
-			else
-			{
-				//Remove newly created array element if pasting is not successful
-				Sequence->Notifies.RemoveAt(NewNotifyIndex);
-			}
-		}
-
-		OnDeselectAllNotifies.ExecuteIfBound();
-		Sequence->MarkPackageDirty();
-		OnUpdatePanel.ExecuteIfBound();
-	}
+	OnPasteNodes.ExecuteIfBound(this, ClickTime, PasteMode, MultiplePasteType);
 }
 
 void SAnimNotifyTrack::OnManageNotifies()
@@ -2796,7 +2984,7 @@ void SAnimNotifyTrack::SetTriggerWeight(const FText& TriggerWeight, ETextCommit:
 
 bool SAnimNotifyTrack::IsSingleNodeSelected()
 {
-	return SelectedNotifyIndices.Num() == 1;
+	return SelectedNodeIndices.Num() == 1;
 }
 
 bool SAnimNotifyTrack::IsSingleNodeInClipboard()
@@ -2805,8 +2993,9 @@ bool SAnimNotifyTrack::IsSingleNodeInClipboard()
 	const TCHAR* Buffer;
 	float OriginalTime;
 	float OriginalLength;
+	int32 TrackSpan;
 	uint32 Count = 0;
-	if (ReadNotifyPasteHeader(PropString, Buffer, OriginalTime, OriginalLength))
+	if (ReadNotifyPasteHeader(PropString, Buffer, OriginalTime, OriginalLength, TrackSpan))
 	{
 		// If reading a single line empties the buffer then we only have one notify in there.
 		FString TempLine;
@@ -2900,6 +3089,24 @@ void SAnimNotifyTrack::OnNewNotifyClicked()
 		);
 }
 
+void SAnimNotifyTrack::OnNewSyncMarkerClicked()
+{
+	// Show dialog to enter new track name
+	TSharedRef<STextEntryPopup> TextEntry =
+		SNew(STextEntryPopup)
+		.Label(LOCTEXT("NewSyncMarkerLabel", "Sync Marker Name"))
+		.OnTextCommitted(this, &SAnimNotifyTrack::AddNewSyncMarker);
+
+	// Show dialog to enter new event name
+	FSlateApplication::Get().PushMenu(
+		AsShared(), // Menu being summoned from a menu that is closing: Parent widget should be k2 not the menu thats open or it will be closed when the menu is dismissed
+		FWidgetPath(),
+		TextEntry,
+		FSlateApplication::Get().GetCursorPos(),
+		FPopupTransitionEffect(FPopupTransitionEffect::TypeInPopup)
+		);
+}
+
 void SAnimNotifyTrack::AddNewNotify(const FText& NewNotifyName, ETextCommit::Type CommitInfo)
 {
 	USkeleton* SeqSkeleton = Sequence->GetSkeleton();
@@ -2917,6 +3124,23 @@ void SAnimNotifyTrack::AddNewNotify(const FText& NewNotifyName, ETextCommit::Typ
 	FSlateApplication::Get().DismissAllMenus();
 }
 
+void SAnimNotifyTrack::AddNewSyncMarker(const FText& NewNotifyName, ETextCommit::Type CommitInfo)
+{
+	USkeleton* SeqSkeleton = Sequence->GetSkeleton();
+	if ((CommitInfo == ETextCommit::OnEnter) && SeqSkeleton)
+	{
+		const FScopedTransaction Transaction(LOCTEXT("AddNewSyncMarker", "Add New Sync Marker"));
+		//FName NewName = FName(*NewNotifyName.ToString());
+		//SeqSkeleton->AddNewAnimationNotify(NewName);
+
+		FBlueprintActionDatabase::Get().RefreshAssetActions(SeqSkeleton);
+
+		CreateNewSyncMarkerAtCursor(NewNotifyName.ToString(), (UClass*)nullptr);
+	}
+
+	FSlateApplication::Get().DismissAllMenus();
+}
+
 void SAnimNotifyTrack::Update()
 {
 	NotifyPairs.Empty();
@@ -3025,6 +3249,34 @@ void SAnimNotifyTrack::Update()
 			NotifyPairs.Add(NotifyPair);
 		}
 	}
+
+	for (FAnimSyncMarker* SyncMarker : AnimSyncMarkers)
+	{
+		TSharedPtr<SAnimNotifyNode> AnimSyncMarkerNode = nullptr;
+		TSharedPtr<SAnimTimingNode> EndTimingNode = nullptr;
+
+		const int32 NodeIndex = NotifyNodes.Num();
+		SAssignNew(AnimSyncMarkerNode, SAnimNotifyNode)
+			.Sequence(Sequence)
+			.AnimSyncMarker(SyncMarker)
+			.OnNodeDragStarted(this, &SAnimNotifyTrack::OnNotifyNodeDragStarted, NodeIndex)
+			.OnUpdatePanel(OnUpdatePanel)
+			.PanTrackRequest(OnRequestTrackPan)
+			.ViewInputMin(ViewInputMin)
+			.ViewInputMax(ViewInputMax)
+			.MarkerBars(MarkerBars)
+			.OnDeselectAllNotifies(OnDeselectAllNotifies)
+			.StateEndTimingNode(EndTimingNode);
+
+		NodeSlots->AddSlot()
+			.Padding(TAttribute<FMargin>::Create(TAttribute<FMargin>::FGetter::CreateSP(this, &SAnimNotifyTrack::GetSyncMarkerTrackPadding, NodeIndex)))
+			.VAlign(VAlign_Center)
+			[
+				AnimSyncMarkerNode->AsShared()
+			];
+
+		NotifyNodes.Add(AnimSyncMarkerNode);
+	}
 }
 
 int32 SAnimNotifyTrack::GetHitNotifyNode(const FGeometry& MyGeometry, const FVector2D& CursorPosition)
@@ -3045,25 +3297,25 @@ FReply SAnimNotifyTrack::OnNotifyNodeDragStarted( TSharedRef<SAnimNotifyNode> No
 	// Check to see if we've already selected the triggering node
 	if (!NotifyNode->bSelected)
 	{
-		SelectNotifyNode(NotifyIndex, MouseEvent.IsShiftDown());
+		SelectTrackObjectNode(NotifyIndex, MouseEvent.IsShiftDown());
 	}
 	
 	// Sort our nodes so we're acessing them in time order
-	SelectedNotifyIndices.Sort([this](const int32& A, const int32& B)
+	SelectedNodeIndices.Sort([this](const int32& A, const int32& B)
 	{
-		FAnimNotifyEvent* First = AnimNotifies[A];
-		FAnimNotifyEvent* Second = AnimNotifies[B];
-		return First->GetTime() < Second->GetTime();
+		float TimeA = NotifyNodes[A]->NodeObjectInterface->GetTime();
+		float TimeB = NotifyNodes[B]->NodeObjectInterface->GetTime();
+		return TimeA < TimeB;
 	});
 
 	// If we're dragging one of the direction markers we don't need to call any further as we don't want the drag drop op
 	if(!bDragOnMarker)
 	{
 		TArray<TSharedPtr<SAnimNotifyNode>> NodesToDrag;
-		const float FirstNodeX = NotifyNodes[SelectedNotifyIndices[0]]->GetWidgetPosition().X;
+		const float FirstNodeX = NotifyNodes[SelectedNodeIndices[0]]->GetWidgetPosition().X;
 		
 		TSharedRef<SOverlay> DragBox = SNew(SOverlay);
-		for (auto Iter = SelectedNotifyIndices.CreateIterator(); Iter; ++Iter)
+		for (auto Iter = SelectedNodeIndices.CreateIterator(); Iter; ++Iter)
 		{
 			TSharedPtr<SAnimNotifyNode> Node = NotifyNodes[*Iter];
 			NodesToDrag.Add(Node);
@@ -3125,47 +3377,32 @@ void SAnimNotifyTrack::HandleNodeDrop(TSharedPtr<SAnimNotifyNode> Node, float Of
 	Sequence->Modify();
 
 	float LocalX = GetCachedGeometry().AbsoluteToLocal(Node->GetScreenPosition() + Offset).X;
-	float Time = GetCachedScaleInfo().LocalXToInput(LocalX);
-	FAnimNotifyEvent* DroppedEvent = Node->NotifyEvent;
-	float EventDuration = DroppedEvent->GetDuration();
-
-	if(Node->GetLastSnappedTime() != -1.0f)
-	{
-		DroppedEvent->Link(Sequence, Node->GetLastSnappedTime(), DroppedEvent->GetSlotIndex());
-	}
-	else
-	{
-		DroppedEvent->Link(Sequence, Time, DroppedEvent->GetSlotIndex());
-	}
-	DroppedEvent->RefreshTriggerOffset(Sequence->CalculateOffsetForNotify(DroppedEvent->GetTime()));
-
-	if(EventDuration > 0.0f)
-	{
-		DroppedEvent->EndLink.Link(Sequence, DroppedEvent->GetTime() + EventDuration, DroppedEvent->GetSlotIndex());
-		DroppedEvent->RefreshEndTriggerOffset(Sequence->CalculateOffsetForNotify(DroppedEvent->EndLink.GetTime()));
-	}
-	else
-	{
-		DroppedEvent->EndTriggerTimeOffset = 0.0f;
-	}
-
-	DroppedEvent->TrackIndex = TrackIndex;
+	float SnapTime = Node->GetLastSnappedTime();
+	float Time = SnapTime != -1.0f ? SnapTime : GetCachedScaleInfo().LocalXToInput(LocalX);
+	Node->NodeObjectInterface->HandleDrop(Sequence, Time, TrackIndex);
 }
 
 void SAnimNotifyTrack::DisconnectSelectedNodesForDrag(TArray<TSharedPtr<SAnimNotifyNode>>& DragNodes)
 {
-	if(SelectedNotifyIndices.Num() == 0)
+	if(SelectedNodeIndices.Num() == 0)
 	{
 		return;
 	}
 
-	const float FirstNodeX = NotifyNodes[SelectedNotifyIndices[0]]->GetWidgetPosition().X;
+	const float FirstNodeX = NotifyNodes[SelectedNodeIndices[0]]->GetWidgetPosition().X;
 
-	for(auto Iter = SelectedNotifyIndices.CreateIterator(); Iter; ++Iter)
+	for(auto Iter = SelectedNodeIndices.CreateIterator(); Iter; ++Iter)
 	{
 		TSharedPtr<SAnimNotifyNode> Node = NotifyNodes[*Iter];
-		TSharedPtr<SAnimNotifyPair> Pair = NotifyPairs[*Iter];
-		NodeSlots->RemoveSlot(Pair->AsShared());
+		if (Node->NodeObjectInterface->GetNotifyEvent())
+		{
+			TSharedPtr<SAnimNotifyPair> Pair = NotifyPairs[*Iter];
+			NodeSlots->RemoveSlot(Pair->AsShared());
+		}
+		else
+		{
+			NodeSlots->RemoveSlot(Node->AsShared());
+		}
 
 		DragNodes.Add(Node);
 	}
@@ -3174,34 +3411,27 @@ void SAnimNotifyTrack::DisconnectSelectedNodesForDrag(TArray<TSharedPtr<SAnimNot
 void SAnimNotifyTrack::AppendSelectionToSet(FGraphPanelSelectionSet& SelectionSet)
 {
 	// Add our selection to the provided set
-	for(auto Iter = SelectedNotifyIndices.CreateConstIterator(); Iter; ++Iter)
+	for(int32 Index : SelectedNodeIndices)
 	{
-		int32 Index = *Iter;
-		check(AnimNotifies.IsValidIndex(Index));
-		FAnimNotifyEvent* Event = AnimNotifies[Index];
-
-		if(Event->Notify)
-		{
-			SelectionSet.Add(Event->Notify);
-		}
-		else if(Event->NotifyStateClass)
+		if (FAnimNotifyEvent* Event = NotifyNodes[Index]->NodeObjectInterface->GetNotifyEvent())
 		{
-			SelectionSet.Add(Event->NotifyStateClass);
+			if (Event->Notify)
+			{
+				SelectionSet.Add(Event->Notify);
+			}
+			else if (Event->NotifyStateClass)
+			{
+				SelectionSet.Add(Event->NotifyStateClass);
+			}
 		}
 	}
 }
 
-void SAnimNotifyTrack::AppendSelectionToArray(TArray<FAnimNotifyEvent*>& Selection) const
+void SAnimNotifyTrack::AppendSelectionToArray(TArray<INodeObjectInterface*>& Selection) const
 {
-	for(int32 Idx : SelectedNotifyIndices)
+	for(int32 Idx : SelectedNodeIndices)
 	{
-		check(AnimNotifies.IsValidIndex(Idx));
-		FAnimNotifyEvent* Event = AnimNotifies[Idx];
-
-		if(Event)
-		{
-			Selection.Add(Event);
-		}
+		Selection.Add(NotifyNodes[Idx]->NodeObjectInterface);
 	}
 }
 
@@ -3265,6 +3495,41 @@ void SAnimNotifyTrack::PasteSingleNotify(FString& NotifyString, float PasteTime)
 	OnUpdatePanel.ExecuteIfBound();
 }
 
+void SAnimNotifyTrack::PasteSingleSyncMarker(FString& MarkerString, float PasteTime)
+{
+	if(UAnimSequence* AnimSeq = Cast<UAnimSequence>(Sequence))
+	{
+		int32 NewIdx = AnimSeq->AuthoredSyncMarkers.Add(FAnimSyncMarker());
+		UArrayProperty* ArrayProperty = NULL;
+		uint8* PropertyData = AnimSeq->FindSyncMarkerPropertyData(NewIdx, ArrayProperty);
+
+		if (PropertyData && ArrayProperty)
+		{
+			ArrayProperty->Inner->ImportText(*MarkerString, PropertyData, PPF_Copy, NULL);
+
+			FAnimSyncMarker& SyncMarker = AnimSeq->AuthoredSyncMarkers[NewIdx];
+
+			if (PasteTime != -1.0f)
+			{
+				SyncMarker.Time = PasteTime;
+			}
+
+			// Make sure the notify is within the track area
+			SyncMarker.Time = FMath::Clamp(SyncMarker.Time, 0.0f, Sequence->SequenceLength);
+			SyncMarker.TrackIndex = TrackIndex;
+		}
+		else
+		{
+			// Paste failed, remove the notify
+			AnimSeq->AuthoredSyncMarkers.RemoveAt(NewIdx);
+		}
+
+		OnDeselectAllNotifies.ExecuteIfBound();
+		Sequence->MarkPackageDirty();
+		OnUpdatePanel.ExecuteIfBound();
+	}
+}
+
 void SAnimNotifyTrack::AppendSelectedNodeWidgetsToArray(TArray<TSharedPtr<SAnimNotifyNode>>& NodeArray) const
 {
 	for(TSharedPtr<SAnimNotifyNode> Node : NotifyNodes)
@@ -3287,11 +3552,11 @@ void SAnimNotifyTrack::RefreshMarqueeSelectedNodes(FSlateRect& Rect, FNotifyMarq
 			bool bWasSelected = Marquee.OriginalSelection.Contains(Notify);
 			if(bWasSelected)
 			{
-				SelectNotifyNode(Idx, true, false);
+				SelectTrackObjectNode(Idx, true, false);
 			}
-			else if(SelectedNotifyIndices.Contains(Idx))
+			else if(SelectedNodeIndices.Contains(Idx))
 			{
-				DeselectNotifyNode(Idx, false);
+				DeselectTrackObjectNode(Idx, false);
 			}
 		}
 	}
@@ -3306,14 +3571,14 @@ void SAnimNotifyTrack::RefreshMarqueeSelectedNodes(FSlateRect& Rect, FNotifyMarq
 			// Either select or deselect the intersecting node, depending on the type of selection operation
 			if(Marquee.Operation == FNotifyMarqueeOperation::Remove)
 			{
-				if(SelectedNotifyIndices.Contains(Index))
+				if(SelectedNodeIndices.Contains(Index))
 				{
-					DeselectNotifyNode(Index, false);
+					DeselectTrackObjectNode(Index, false);
 				}
 			}
 			else
 			{
-				SelectNotifyNode(Index, true, false);
+				SelectTrackObjectNode(Index, true, false);
 			}
 		}
 	}
@@ -3345,18 +3610,19 @@ void SAnimNotifyTrack::GetNotifyMenuData(TArray<FAssetData>& NotifyAssetData, TA
 	});
 }
 
-void SAnimNotifyTrack::OnSetNotifyTimeClicked(int32 NotifyIndex)
+void SAnimNotifyTrack::OnSetNodeTimeClicked(int32 NodeIndex)
 {
-	if(AnimNotifies.IsValidIndex(NotifyIndex))
+	if (NotifyNodes.IsValidIndex(NodeIndex))
 	{
-		FString DefaultText = FString::Printf(TEXT("%0.6f"), AnimNotifies[NotifyIndex]->GetTime());
+		INodeObjectInterface* NodeObject = NotifyNodes[NodeIndex]->NodeObjectInterface;
+		FString DefaultText = FString::Printf(TEXT("%0.6f"), NodeObject->GetTime());
 
 		// Show dialog to enter time
 		TSharedRef<STextEntryPopup> TextEntry =
 			SNew(STextEntryPopup)
 			.Label(LOCTEXT("NotifyTimeClickedLabel", "Notify Time"))
 			.DefaultText(FText::FromString(DefaultText))
-			.OnTextCommitted(this, &SAnimNotifyTrack::SetNotifyTime, NotifyIndex);
+			.OnTextCommitted(this, &SAnimNotifyTrack::SetNodeTime, NodeIndex);
 
 		FSlateApplication::Get().PushMenu(
 			AsShared(), // Menu being summoned from a menu that is closing: Parent widget should be k2 not the menu thats open or it will be closed when the menu is dismissed
@@ -3368,22 +3634,25 @@ void SAnimNotifyTrack::OnSetNotifyTimeClicked(int32 NotifyIndex)
 	}
 }
 
-void SAnimNotifyTrack::SetNotifyTime(const FText& NotifyTimeText, ETextCommit::Type CommitInfo, int32 NotifyIndex)
+void SAnimNotifyTrack::SetNodeTime(const FText& NodeTimeText, ETextCommit::Type CommitInfo, int32 NodeIndex)
 {
 	if(CommitInfo == ETextCommit::OnEnter || CommitInfo == ETextCommit::OnUserMovedFocus)
 	{
-		if(AnimNotifies.IsValidIndex(NotifyIndex))
+		if (NotifyNodes.IsValidIndex(NodeIndex))
 		{
-			FAnimNotifyEvent* Event = AnimNotifies[NotifyIndex];
+			INodeObjectInterface* NodeObject = NotifyNodes[NodeIndex]->NodeObjectInterface;
 
-			float NewTime = FMath::Clamp(FCString::Atof(*NotifyTimeText.ToString()), 0.0f, Sequence->SequenceLength - Event->GetDuration());
+			float NewTime = FMath::Clamp(FCString::Atof(*NodeTimeText.ToString()), 0.0f, Sequence->SequenceLength - NodeObject->GetDuration());
 
-			Event->SetTime(NewTime);
+			NodeObject->SetTime(NewTime);
 
-			Event->RefreshTriggerOffset(Sequence->CalculateOffsetForNotify(Event->GetTime()));
-			if(Event->GetDuration() > 0.0f)
+			if (FAnimNotifyEvent* Event = NodeObject->GetNotifyEvent())
 			{
-				Event->RefreshEndTriggerOffset(Sequence->CalculateOffsetForNotify(Event->GetTime() + Event->GetDuration()));
+				Event->RefreshTriggerOffset(Sequence->CalculateOffsetForNotify(Event->GetTime()));
+				if (Event->GetDuration() > 0.0f)
+				{
+					Event->RefreshEndTriggerOffset(Sequence->CalculateOffsetForNotify(Event->GetTime() + Event->GetDuration()));
+				}
 			}
 		}
 	}
@@ -3391,15 +3660,13 @@ void SAnimNotifyTrack::SetNotifyTime(const FText& NotifyTimeText, ETextCommit::T
 	FSlateApplication::Get().DismissAllMenus();
 }
 
-void SAnimNotifyTrack::OnSetNotifyFrameClicked(int32 NotifyIndex)
+void SAnimNotifyTrack::OnSetNodeFrameClicked(int32 NodeIndex)
 {
-	if(AnimNotifies.IsValidIndex(NotifyIndex))
+	if (NotifyNodes.IsValidIndex(NodeIndex))
 	{
-		FAnimNotifyEvent* Event = AnimNotifies[NotifyIndex];
+		INodeObjectInterface* NodeObject = NotifyNodes[NodeIndex]->NodeObjectInterface;
 
-		//const float Time = Event->GetTime();
-		//const float Percentage = Time / Sequence->SequenceLength;
-		const FText Frame = FText::AsNumber(Sequence->GetFrameAtTime(Event->GetTime()));
+		const FText Frame = FText::AsNumber(Sequence->GetFrameAtTime(NodeObject->GetTime()));
 
 		FString DefaultText = FString::Printf(TEXT("%s"), *Frame.ToString());
 
@@ -3408,7 +3675,7 @@ void SAnimNotifyTrack::OnSetNotifyFrameClicked(int32 NotifyIndex)
 			SNew(STextEntryPopup)
 			.Label(LOCTEXT("NotifyFrameClickedLabel", "Notify Frame"))
 			.DefaultText(FText::FromString(DefaultText))
-			.OnTextCommitted(this, &SAnimNotifyTrack::SetNotifyFrame, NotifyIndex);
+			.OnTextCommitted(this, &SAnimNotifyTrack::SetNodeFrame, NodeIndex);
 
 		FSlateApplication::Get().PushMenu(
 			AsShared(), // Menu being summoned from a menu that is closing: Parent widget should be k2 not the menu thats open or it will be closed when the menu is dismissed
@@ -3420,23 +3687,26 @@ void SAnimNotifyTrack::OnSetNotifyFrameClicked(int32 NotifyIndex)
 	}
 }
 
-void SAnimNotifyTrack::SetNotifyFrame(const FText& NotifyFrameText, ETextCommit::Type CommitInfo, int32 NotifyIndex)
+void SAnimNotifyTrack::SetNodeFrame(const FText& NodeFrameText, ETextCommit::Type CommitInfo, int32 NodeIndex)
 {
 	if(CommitInfo == ETextCommit::OnEnter || CommitInfo == ETextCommit::OnUserMovedFocus)
 	{
-		if(AnimNotifies.IsValidIndex(NotifyIndex))
+		if (NotifyNodes.IsValidIndex(NodeIndex))
 		{
-			FAnimNotifyEvent* Event = AnimNotifies[NotifyIndex];
+			INodeObjectInterface* NodeObject = NotifyNodes[NodeIndex]->NodeObjectInterface;
 
-			int32 Frame = FCString::Atof(*NotifyFrameText.ToString());
-			float NewTime = FMath::Clamp(Sequence->GetTimeAtFrame(Frame), 0.0f, Sequence->SequenceLength - Event->Duration);
+			int32 Frame = FCString::Atof(*NodeFrameText.ToString());
+			float NewTime = FMath::Clamp(Sequence->GetTimeAtFrame(Frame), 0.0f, Sequence->SequenceLength - NodeObject->GetDuration());
 
-			Event->SetTime(NewTime);
+			NodeObject->SetTime(NewTime);
 
-			Event->RefreshTriggerOffset(Sequence->CalculateOffsetForNotify(Event->GetTime()));
-			if(Event->GetDuration() > 0.0f)
+			if (FAnimNotifyEvent* Event = NodeObject->GetNotifyEvent())
 			{
-				Event->RefreshEndTriggerOffset(Sequence->CalculateOffsetForNotify(Event->GetTime() + Event->GetDuration()));
+				Event->RefreshTriggerOffset(Sequence->CalculateOffsetForNotify(Event->GetTime()));
+				if (Event->GetDuration() > 0.0f)
+				{
+					Event->RefreshEndTriggerOffset(Sequence->CalculateOffsetForNotify(Event->GetTime() + Event->GetDuration()));
+				}
 			}
 		}
 	}
@@ -3448,10 +3718,12 @@ const EVisibility SAnimNotifyTrack::GetTimingNodeVisibility(TSharedPtr<SAnimNoti
 {
 	if(OnGetTimingNodeVisibility.IsBound())
 	{
-		FAnimNotifyEvent* Event = NotifyNode->NotifyEvent;
-		EMontageNotifyTickType::Type TickType = Event->MontageTickType;
+		if(FAnimNotifyEvent* Event = NotifyNode->NodeObjectInterface->GetNotifyEvent())
+		{
+			EMontageNotifyTickType::Type TickType = Event->MontageTickType;
 
-		return TickType == EMontageNotifyTickType::BranchingPoint ? OnGetTimingNodeVisibility.Execute(ETimingElementType::BranchPointNotify) : OnGetTimingNodeVisibility.Execute(ETimingElementType::QueuedNotify);
+			return TickType == EMontageNotifyTickType::BranchingPoint ? OnGetTimingNodeVisibility.Execute(ETimingElementType::BranchPointNotify) : OnGetTimingNodeVisibility.Execute(ETimingElementType::QueuedNotify);
+		}
 	}
 
 	// No visibility defined, not visible
@@ -3499,6 +3771,7 @@ void SNotifyEdTrack::Construct(const FArguments& InArgs)
 				.Sequence(Sequence)
 				.TrackIndex(TrackIndex)
 				.AnimNotifies(Track.Notifies)
+				.AnimSyncMarkers(Track.SyncMarkers)
 				.ViewInputMin(InArgs._ViewInputMin)
 				.ViewInputMax(InArgs._ViewInputMax)
 				.OnSelectionChanged(InArgs._OnSelectionChanged)
@@ -3519,8 +3792,8 @@ void SNotifyEdTrack::Construct(const FArguments& InArgs)
 				.OnReplaceSelectedWithNotify(PanelRef, &SAnimNotifyPanel::OnReplaceSelectedWithNotify)
 				.OnReplaceSelectedWithBlueprintNotify(PanelRef, &SAnimNotifyPanel::OnReplaceSelectedWithNotifyBlueprint)
 				.OnDeselectAllNotifies(InArgs._OnDeselectAllNotifies)
-				.OnCopyNotifies(InArgs._OnCopyNotifies)
-				.OnPasteNotifies(InArgs._OnPasteNotifies)
+				.OnCopyNodes(InArgs._OnCopyNodes)
+				.OnPasteNodes(InArgs._OnPasteNodes)
 				.OnSetInputViewRange(InArgs._OnSetInputViewRange)
 				.OnGetTimingNodeVisibility(InArgs._OnGetTimingNodeVisibility)
 			]
@@ -3692,10 +3965,19 @@ FReply SAnimNotifyPanel::InsertTrack(int32 TrackIndexToInsert)
 	for (int32 I=TrackIndexToInsert; I<Sequence->AnimNotifyTracks.Num(); ++I)
 	{
 		FAnimNotifyTrack& Track = Sequence->AnimNotifyTracks[I];
-		for (int32 J=0; J<Track.Notifies.Num(); ++J)
+
+		const int32 NewTrackIndex = I + 1;
+
+		for (FAnimNotifyEvent* Notify : Track.Notifies)
 		{
 			// fix notifies indices
-			Track.Notifies[J]->TrackIndex = I+1;
+			Notify->TrackIndex = NewTrackIndex;
+		}
+
+		for (FAnimSyncMarker* SyncMarker : Track.SyncMarkers)
+		{
+			// fix notifies indices
+			SyncMarker->TrackIndex = NewTrackIndex;
 		}
 
 		// fix track names, for now it's only index
@@ -3723,10 +4005,18 @@ FReply SAnimNotifyPanel::DeleteTrack(int32 TrackIndexToDelete)
 			for (int32 I=TrackIndexToDelete+1; I<Sequence->AnimNotifyTracks.Num(); ++I)
 			{
 				FAnimNotifyTrack& Track = Sequence->AnimNotifyTracks[I];
-				for (int32 J=0; J<Track.Notifies.Num(); ++J)
+				const int32 NewTrackIndex = I - 1;
+
+				for (FAnimNotifyEvent* Notify : Track.Notifies)
+				{
+					// fix notifies indices
+					Notify->TrackIndex = NewTrackIndex;
+				}
+
+				for (FAnimSyncMarker* SyncMarker : Track.SyncMarkers)
 				{
 					// fix notifies indices
-					Track.Notifies[J]->TrackIndex = I-1;
+					SyncMarker->TrackIndex = NewTrackIndex;
 				}
 
 				// fix track names, for now it's only index
@@ -3751,60 +4041,15 @@ bool SAnimNotifyPanel::CanDeleteTrack(int32 TrackIndexToDelete)
 	return false;
 }
 
-void SAnimNotifyPanel::DeleteNotify(FAnimNotifyEvent* Notify)
-{
-	for (int32 I=0; I<Sequence->Notifies.Num(); ++I)
-	{
-		if (Notify == &(Sequence->Notifies[I]))
-		{
-			Sequence->Notifies.RemoveAt(I);
-			Sequence->MarkPackageDirty();
-			break;
-		}
-	}
-}
-
 void SAnimNotifyPanel::Update()
 {
 	if(Sequence != NULL)
 	{
-		Sequence->UpdateAnimNotifyTrackCache();
+		Sequence->RefreshCacheData();
 	}
 	PersonaPtr.Pin()->OnAnimNotifiesChanged.Broadcast();
 }
 
-bool SAnimNotifyPanel::ReadNotifyPasteHeader(FString& OutPropertyString, const TCHAR*& OutBuffer, float& OutOriginalTime, float& OutOriginalLength, int32& OutTrackSpan) const
-{
-	OutBuffer = NULL;
-	OutOriginalTime = -1.f;
-
-	FPlatformMisc::ClipboardPaste(OutPropertyString);
-
-	if(!OutPropertyString.IsEmpty())
-	{
-		//Remove header text
-		const FString HeaderString(TEXT("COPY_ANIMNOTIFYEVENT"));
-
-		//Check for string identifier in order to determine whether the text represents an FAnimNotifyEvent.
-		if(OutPropertyString.StartsWith(HeaderString) && OutPropertyString.Len() > HeaderString.Len())
-		{
-			int32 HeaderSize = HeaderString.Len();
-			OutBuffer = *OutPropertyString;
-			OutBuffer += HeaderSize;
-
-			FString ReadLine;
-			// Read the original time from the first notify
-			FParse::Line(&OutBuffer, ReadLine);
-			FParse::Value(*ReadLine, TEXT("OriginalTime="), OutOriginalTime);
-			FParse::Value(*ReadLine, TEXT("OriginalLength="), OutOriginalLength);
-			FParse::Value(*ReadLine, TEXT("TrackSpan="), OutTrackSpan);
-			return true;
-		}
-	}
-
-	return false;
-}
-
 void SAnimNotifyPanel::RefreshNotifyTracks()
 {
 	check (Sequence);
@@ -3843,10 +4088,10 @@ void SAnimNotifyPanel::RefreshNotifyTracks()
 			.OnNodeDragStarted(this, &SAnimNotifyPanel::OnNotifyNodeDragStarted)
 			.MarkerBars(MarkerBars)
 			.OnRequestRefreshOffsets(OnRequestRefreshOffsets)
-			.OnDeleteNotify(this, &SAnimNotifyPanel::DeleteSelectedNotifies)
+			.OnDeleteNotify(this, &SAnimNotifyPanel::DeleteSelectedNodeObjects)
 			.OnDeselectAllNotifies(this, &SAnimNotifyPanel::DeselectAllNotifies)
-			.OnCopyNotifies(this, &SAnimNotifyPanel::CopySelectedNotifiesToClipboard)
-			.OnPasteNotifies(this, &SAnimNotifyPanel::OnPasteNotifies)
+			.OnCopyNodes(this, &SAnimNotifyPanel::CopySelectedNodesToClipboard)
+			.OnPasteNodes(this, &SAnimNotifyPanel::OnPasteNodes)
 			.OnSetInputViewRange(this, &SAnimNotifyPanel::InputViewRangeChanged)
 			.OnGetTimingNodeVisibility(OnGetTimingNodeVisibility)
 		];
@@ -3919,7 +4164,7 @@ void SAnimNotifyPanel::PostUndo()
 {
 	if(Sequence != NULL)
 	{
-		Sequence->UpdateAnimNotifyTrackCache();
+		Sequence->RefreshCacheData();
 	}
 }
 
@@ -3929,29 +4174,29 @@ void SAnimNotifyPanel::OnDeletePressed()
 	// so don't delete anything when the key is pressed.
 	if(HasKeyboardFocus() || HasFocusedDescendants()) 
 	{
-		DeleteSelectedNotifies();
+		DeleteSelectedNodeObjects();
 	}
 }
 
-void SAnimNotifyPanel::DeleteSelectedNotifies()
+void SAnimNotifyPanel::DeleteSelectedNodeObjects()
 {
-	TArray<FAnimNotifyEvent*> SelectedNotifies;
+	TArray<INodeObjectInterface*> SelectedNodes;
 	for(TSharedPtr<SAnimNotifyTrack> Track : NotifyAnimTracks)
 	{
-		Track->AppendSelectionToArray(SelectedNotifies);
+		Track->AppendSelectionToArray(SelectedNodes);
 	}
 
-	if(SelectedNotifies.Num() > 0)
+	if (SelectedNodes.Num() > 0)
 	{
-		FScopedTransaction Transaction(LOCTEXT("DeleteNotifies", "Delete Animation Notifies"));
+		FScopedTransaction Transaction(LOCTEXT("DeleteMarkers", "Delete Animation Markers"));
 		Sequence->Modify(true);
 
 		// Delete from highest index to lowest
-		SelectedNotifies.Sort();
-		for(int NotifyIndex = SelectedNotifies.Num() - 1; NotifyIndex >= 0 ; --NotifyIndex)
+		SelectedNodes.Sort();
+		for (int NodeIndex = SelectedNodes.Num() - 1; NodeIndex >= 0; --NodeIndex)
 		{
-			FAnimNotifyEvent* Notify = SelectedNotifies[NotifyIndex];
-			DeleteNotify(Notify);
+			INodeObjectInterface* NodeObject = SelectedNodes[NodeIndex];
+			NodeObject->Delete(Sequence);
 		}
 	}
 
@@ -3981,18 +4226,21 @@ void SAnimNotifyPanel::OnTrackSelectionChanged()
 	FGraphPanelSelectionSet SelectionSet;
 	TArray<UEditorNotifyObject*> NotifyObjects;
 
-	TArray<FAnimNotifyEvent*> Events;
 	for(int32 TrackIdx = 0 ; TrackIdx < NotifyAnimTracks.Num() ; ++TrackIdx)
 	{
 		TSharedPtr<SAnimNotifyTrack> Track = NotifyAnimTracks[TrackIdx];
 		TArray<int32> TrackIndices = Track->GetSelectedNotifyIndices();
 		for(int32 Idx : TrackIndices)
 		{
-			FString ObjName = MakeUniqueObjectName(GetTransientPackage(), UEditorNotifyObject::StaticClass()).ToString();
-			UEditorNotifyObject* NewNotifyObject = NewObject<UEditorNotifyObject>(GetTransientPackage(), FName(*ObjName), RF_Public | RF_Standalone | RF_Transient);
-			NewNotifyObject->InitFromAnim(Sequence, FOnAnimObjectChange::CreateSP(this, &SAnimNotifyPanel::OnNotifyObjectChanged));
-			NewNotifyObject->InitialiseNotify(NotifyAnimTracks.Num() - TrackIdx - 1, Idx);
-			SelectionSet.Add(NewNotifyObject);
+			INodeObjectInterface* NodeObjectInterface = Track->GetNodeObjectInterface(Idx);
+			if (FAnimNotifyEvent* NotifyEvent = NodeObjectInterface->GetNotifyEvent())
+			{
+				FString ObjName = MakeUniqueObjectName(GetTransientPackage(), UEditorNotifyObject::StaticClass()).ToString();
+				UEditorNotifyObject* NewNotifyObject = NewObject<UEditorNotifyObject>(GetTransientPackage(), FName(*ObjName), RF_Public | RF_Standalone | RF_Transient);
+				NewNotifyObject->InitFromAnim(Sequence, FOnAnimObjectChange::CreateSP(this, &SAnimNotifyPanel::OnNotifyObjectChanged));
+				NewNotifyObject->InitialiseNotify(NotifyAnimTracks.Num() - TrackIdx - 1, Idx);
+				SelectionSet.Add(NewNotifyObject);
+			}
 		}
 	}
 
@@ -4016,36 +4264,36 @@ void SAnimNotifyPanel::DeselectAllNotifies()
 	OnTrackSelectionChanged();
 }
 
-void SAnimNotifyPanel::CopySelectedNotifiesToClipboard() const
+void SAnimNotifyPanel::CopySelectedNodesToClipboard() const
 {
 	// Grab the selected events
-	TArray<FAnimNotifyEvent*> SelectedNotifies;
+	TArray<INodeObjectInterface*> SelectedNodes;
 	for(TSharedPtr<SAnimNotifyTrack> Track : NotifyAnimTracks)
 	{
-		Track->AppendSelectionToArray(SelectedNotifies);
+		Track->AppendSelectionToArray(SelectedNodes);
 	}
 
 	const FString HeaderString(TEXT("COPY_ANIMNOTIFYEVENT"));
 	
-	if(SelectedNotifies.Num() > 0)
+	if (SelectedNodes.Num() > 0)
 	{
 		FString StrValue(HeaderString);
 
 		// Sort by track
-		SelectedNotifies.Sort([](const FAnimNotifyEvent& A, const FAnimNotifyEvent& B)
+		SelectedNodes.Sort([](const INodeObjectInterface& A, const INodeObjectInterface& B)
 		{
-			return (A.TrackIndex > B.TrackIndex) || (A.TrackIndex == B.TrackIndex && A.GetTime() < B.GetTime());
+			return (A.GetTrackIndex() > B.GetTrackIndex()) || (A.GetTrackIndex() == B.GetTrackIndex() && A.GetTime() < B.GetTime());
 		});
 
 		// Need to find how many tracks this selection spans and the minimum time to use as the beginning of the selection
 		int32 MinTrack = MAX_int32;
 		int32 MaxTrack = MIN_int32;
 		float MinTime = MAX_flt;
-		for(const FAnimNotifyEvent* Event : SelectedNotifies)
+		for (const INodeObjectInterface* NodeObject : SelectedNodes)
 		{
-			MinTrack = FMath::Min(MinTrack, Event->TrackIndex);
-			MaxTrack = FMath::Max(MaxTrack, Event->TrackIndex);
-			MinTime = FMath::Min(MinTime, Event->GetTime());
+			MinTrack = FMath::Min(MinTrack, NodeObject->GetTrackIndex());
+			MaxTrack = FMath::Max(MaxTrack, NodeObject->GetTrackIndex());
+			MinTime = FMath::Min(MinTime, NodeObject->GetTime());
 		}
 
 		int32 TrackSpan = MaxTrack - MinTrack + 1;
@@ -4054,31 +4302,15 @@ void SAnimNotifyPanel::CopySelectedNotifiesToClipboard() const
 		StrValue += FString::Printf(TEXT("OriginalLength=%f,"), Sequence->SequenceLength);
 		StrValue += FString::Printf(TEXT("TrackSpan=%d"), TrackSpan);
 
-		for(const FAnimNotifyEvent* Event : SelectedNotifies)
+		for(const INodeObjectInterface* NodeObject : SelectedNodes)
 		{
 			// Locate the notify in the sequence, we need the sequence index; but also need to
 			// keep the order we're currently in.
-			int32 Index = INDEX_NONE;
-			for(int32 NotifyIdx = 0 ; NotifyIdx < Sequence->Notifies.Num() ; ++NotifyIdx)
-			{
-				if(Event == &Sequence->Notifies[NotifyIdx])
-				{
-					Index = NotifyIdx;
-					break;
-				}
-			}
 
-			check(Index != INDEX_NONE);
-			
 			StrValue += "\n";
-			UArrayProperty* ArrayProperty = NULL;
-			uint8* PropertyData = Sequence->FindNotifyPropertyData(Index, ArrayProperty);
-			StrValue += FString::Printf(TEXT("AbsTime=%f,"), Event->GetTime());
-			if(PropertyData && ArrayProperty)
-			{
-				ArrayProperty->Inner->ExportTextItem(StrValue, PropertyData, PropertyData, Sequence, PPF_Copy);
-			}
+			StrValue += FString::Printf(TEXT("AbsTime=%f,NodeObjectType=%i,"), NodeObject->GetTime(), (int32)NodeObject->GetType());
 
+			NodeObject->ExportForCopy(Sequence, StrValue);
 		}
 		FPlatformMisc::ClipboardCopy(*StrValue);
 	}
@@ -4087,16 +4319,17 @@ void SAnimNotifyPanel::CopySelectedNotifiesToClipboard() const
 bool SAnimNotifyPanel::IsNotifySelectionValidForReplacement()
 {
 	// Grab the selected events
-	TArray<FAnimNotifyEvent*> SelectedNotifies;
+	TArray<INodeObjectInterface*> SelectedNodes;
 	for (TSharedPtr<SAnimNotifyTrack> Track : NotifyAnimTracks)
 	{
-		Track->AppendSelectionToArray(SelectedNotifies);
+		Track->AppendSelectionToArray(SelectedNodes);
 	}
 
 	bool bSelectionContainsAnimNotify = false;
 	bool bSelectionContainsAnimNotifyState = false;
-	for (FAnimNotifyEvent* NotifyEvent : SelectedNotifies)
+	for (INodeObjectInterface* NodeObject : SelectedNodes)
 	{
+		FAnimNotifyEvent* NotifyEvent = NodeObject->GetNotifyEvent();
 		if (NotifyEvent)
 		{
 			if (NotifyEvent->Notify)
@@ -4125,21 +4358,21 @@ bool SAnimNotifyPanel::IsNotifySelectionValidForReplacement()
 
 void SAnimNotifyPanel::OnReplaceSelectedWithNotify(FString NewNotifyName, UClass* NewNotifyClass)
 {
-	TArray<FAnimNotifyEvent*> SelectedNotifies;
+	TArray<INodeObjectInterface*> SelectedNodes;
 	for (TSharedPtr<SAnimNotifyTrack> Track : NotifyAnimTracks)
 	{
-		Track->AppendSelectionToArray(SelectedNotifies);
+		Track->AppendSelectionToArray(SelectedNodes);
 	}
 
 	// Sort these since order is important for deletion
-	SelectedNotifies.Sort();
+	SelectedNodes.Sort();
 
 	const FScopedTransaction Transaction(LOCTEXT("AddNotifyEvent", "Replace Anim Notify"));
 	Sequence->Modify(true);
 
-	for (int NotifyIndex = SelectedNotifies.Num()-1; NotifyIndex >= 0; --NotifyIndex)
+	for (INodeObjectInterface* NodeObject : SelectedNodes)
 	{
-		FAnimNotifyEvent* OldEvent = SelectedNotifies[NotifyIndex];
+		FAnimNotifyEvent* OldEvent = NodeObject->GetNotifyEvent();
 		if (OldEvent)
 		{
 			float BeginTime = OldEvent->GetTime();
@@ -4155,7 +4388,7 @@ void SAnimNotifyPanel::OnReplaceSelectedWithNotify(FString NewNotifyName, UClass
 			EAnimLinkMethod::Type EndLinkMethod = OldEvent->EndLink.GetLinkMethod();
 
 			// Delete old one before creating new one to avoid potential array re-allocation when array temporarily increases by 1 in size
-			DeleteNotify(OldEvent);
+			NodeObject->Delete(Sequence);
 			FAnimNotifyEvent& NewEvent = NotifyAnimTracks[TargetTrackIndex]->CreateNewNotify(NewNotifyName, NewNotifyClass, BeginTime);
 
 			NewEvent.TriggerTimeOffset = TriggerTimeOffset;
@@ -4191,7 +4424,7 @@ void SAnimNotifyPanel::OnReplaceSelectedWithNotifyBlueprint(FString NewBlueprint
 	OnReplaceSelectedWithNotify(NewBlueprintNotifyName, BlueprintClass);
 }
 
-void SAnimNotifyPanel::OnPasteNotifies(SAnimNotifyTrack* RequestTrack, float ClickTime, ENotifyPasteMode::Type PasteMode, ENotifyPasteMultipleMode::Type MultiplePasteType)
+void SAnimNotifyPanel::OnPasteNodes(SAnimNotifyTrack* RequestTrack, float ClickTime, ENotifyPasteMode::Type PasteMode, ENotifyPasteMultipleMode::Type MultiplePasteType)
 {
 	int32 PasteIdx = RequestTrack->GetTrackIndex();
 	int32 NumTracks = NotifyAnimTracks.Num();
@@ -4242,11 +4475,13 @@ void SAnimNotifyPanel::OnPasteNotifies(SAnimNotifyTrack* RequestTrack, float Cli
 		{
 			int32 OriginalTrack;
 			float OrigTime;
+			int32 NodeObjectType;
 			float PasteTime = -1.0f;
-			if(FParse::Value(*CurrentLine, TEXT("TrackIndex="), OriginalTrack) && FParse::Value(*CurrentLine, TEXT("AbsTime="), OrigTime))
+			if (FParse::Value(*CurrentLine, TEXT("TrackIndex="), OriginalTrack) && FParse::Value(*CurrentLine, TEXT("AbsTime="), OrigTime) && FParse::Value(*CurrentLine, TEXT("NodeObjectType="), NodeObjectType))
 			{
-				FString NotifyExportString;
-				CurrentLine.Split(TEXT(","), NULL, &NotifyExportString);
+				const int32 FirstComma = CurrentLine.Find(TEXT(","), ESearchCase::CaseSensitive, ESearchDir::FromStart);
+				const int32 SecondComma = CurrentLine.Find(TEXT(","), ESearchCase::CaseSensitive, ESearchDir::FromStart, FirstComma + 1);
+				FString NotifyExportString = CurrentLine.RightChop(SecondComma+1);
 
 				// Store the first track so we know where to place notifies
 				if(FirstTrack < 0)
@@ -4260,8 +4495,18 @@ void SAnimNotifyPanel::OnPasteNotifies(SAnimNotifyTrack* RequestTrack, float Cli
 
 				// Have to invert the index here as tracks are stored in reverse
 				TSharedPtr<SAnimNotifyTrack> TrackToUse = NotifyAnimTracks[NotifyAnimTracks.Num() - 1 - (PasteIdx + TrackOffset)];
-
-				TrackToUse->PasteSingleNotify(NotifyExportString, TimeToPaste);
+				if (NodeObjectType == ENodeObjectTypes::NOTIFY)
+				{
+					TrackToUse->PasteSingleNotify(NotifyExportString, TimeToPaste);
+				}
+				else if (NodeObjectType == ENodeObjectTypes::SYNC_MARKER)
+				{
+					TrackToUse->PasteSingleSyncMarker(NotifyExportString, TimeToPaste);
+				}
+				else
+				{
+					check(false); //Unknown value in paste
+				}
 			}
 		}
 	}
diff --git a/Engine/Source/Editor/Persona/Private/SAnimNotifyPanel.h b/Engine/Source/Editor/Persona/Private/SAnimNotifyPanel.h
--- a/Engine/Source/Editor/Persona/Private/SAnimNotifyPanel.h
+++ b/Engine/Source/Editor/Persona/Private/SAnimNotifyPanel.h
@@ -21,7 +21,6 @@ DECLARE_DELEGATE_RetVal( bool, FOnGetIsAnimNotifySelectionValidForReplacement )
 DECLARE_DELEGATE_TwoParams( FReplaceWithNotify, FString, UClass* )
 DECLARE_DELEGATE_TwoParams( FReplaceWithBlueprintNotify, FString, FString )
 DECLARE_DELEGATE( FDeselectAllNotifies )
-DECLARE_DELEGATE( FCopyNotifies )
 DECLARE_DELEGATE_OneParam( FOnGetBlueprintNotifyData, TArray<FAssetData>& )
 DECLARE_DELEGATE_OneParam( FOnGetNativeNotifyClasses, TArray<UClass*>&)
 
@@ -182,7 +181,6 @@ public:
 	FReply InsertTrack(int32 TrackIndexToInsert);
 	FReply DeleteTrack(int32 TrackIndexToDelete);
 	bool CanDeleteTrack(int32 TrackIndexToDelete);
-	void DeleteNotify(FAnimNotifyEvent* Notify);
 	void Update();
 	TWeakPtr<FPersona> GetPersona() const { return PersonaPtr; }
 
@@ -194,8 +192,8 @@ public:
 
 	virtual float GetSequenceLength() const override {return Sequence->SequenceLength;}
 
-	void CopySelectedNotifiesToClipboard() const;
-	void OnPasteNotifies(SAnimNotifyTrack* RequestTrack, float ClickTime, ENotifyPasteMode::Type PasteMode, ENotifyPasteMultipleMode::Type MultiplePasteType);
+	void CopySelectedNodesToClipboard() const;
+	void OnPasteNodes(SAnimNotifyTrack* RequestTrack, float ClickTime, ENotifyPasteMode::Type PasteMode, ENotifyPasteMultipleMode::Type MultiplePasteType);
 
 	/** Handler for properties changing on objects */
 	FCoreUObjectDelegates::FOnObjectPropertyChanged::FDelegate OnPropertyChangedHandle;
@@ -246,9 +244,6 @@ private:
 	/** Cached list of anim tracks for notify node drag drop */
 	TArray<TSharedPtr<SAnimNotifyTrack>> NotifyAnimTracks;
 
-	// Read common info from the clipboard
-	bool ReadNotifyPasteHeader(FString& OutPropertyString, const TCHAR*& OutBuffer, float& OutOriginalTime, float& OutOriginalLength, int32& OutTrackSpan) const;
-
 	// this just refresh notify tracks - UI purpose only
 	// do not call this from here. This gets called by asset. 
 	void RefreshNotifyTracks();
@@ -258,7 +253,7 @@ private:
 	void OnDeletePressed();
 
 	/** Deletes all currently selected notifies in the panel */
-	void DeleteSelectedNotifies();
+	void DeleteSelectedNodeObjects();
 
 	/** We support keyboard focus to detect when we should process key commands like delete */
 	virtual bool SupportsKeyboardFocus() const override
diff --git a/Engine/Source/Runtime/AnimGraphRuntime/Private/AnimNodes/AnimNode_BlendSpacePlayer.cpp b/Engine/Source/Runtime/AnimGraphRuntime/Private/AnimNodes/AnimNode_BlendSpacePlayer.cpp
--- a/Engine/Source/Runtime/AnimGraphRuntime/Private/AnimNodes/AnimNode_BlendSpacePlayer.cpp
+++ b/Engine/Source/Runtime/AnimGraphRuntime/Private/AnimNodes/AnimNode_BlendSpacePlayer.cpp
@@ -57,6 +57,8 @@ void FAnimNode_BlendSpacePlayer::UpdateInternal(const FAnimationUpdateContext& C
 		// Create a tick record and fill it out
 		FAnimGroupInstance* SyncGroup;
 		FAnimTickRecord& TickRecord = Context.AnimInstance->CreateUninitializedTickRecord(GroupIndex, /*out*/ SyncGroup);
+		TickRecord.SourceNodeRef = this;
+		Context.AnimInstance->InitTickRecordFromLastFrame(GroupIndex, TickRecord);
 
 		const FVector BlendInput(X, Y, Z);
 	
diff --git a/Engine/Source/Runtime/Engine/Classes/Animation/AnimBlueprintGeneratedClass.h b/Engine/Source/Runtime/Engine/Classes/Animation/AnimBlueprintGeneratedClass.h
--- a/Engine/Source/Runtime/Engine/Classes/Animation/AnimBlueprintGeneratedClass.h
+++ b/Engine/Source/Runtime/Engine/Classes/Animation/AnimBlueprintGeneratedClass.h
@@ -183,6 +183,12 @@ class ENGINE_API UAnimBlueprintGeneratedClass : public UBlueprintGeneratedClass
 	UStructProperty* RootAnimNodeProperty;
 	TArray<UStructProperty*> AnimNodeProperties;
 
+	// Array of sync group names in the order that they are requested during compile
+	UPROPERTY()
+	TArray<FName> SyncGroupNames;
+
+	int32 GetSyncGroupIndex(FName SyncGroupName) const { return SyncGroupNames.IndexOfByKey(SyncGroupName); }
+
 #if WITH_EDITORONLY_DATA
 	FAnimBlueprintDebugData AnimBlueprintDebugData;
 
diff --git a/Engine/Source/Runtime/Engine/Classes/Animation/AnimInstance.h b/Engine/Source/Runtime/Engine/Classes/Animation/AnimInstance.h
--- a/Engine/Source/Runtime/Engine/Classes/Animation/AnimInstance.h
+++ b/Engine/Source/Runtime/Engine/Classes/Animation/AnimInstance.h
@@ -337,12 +337,10 @@ class ENGINE_API UAnimInstance : public UObject
 	USkeleton* CurrentSkeleton;
 
 	// The list of animation assets which are going to be evaluated this frame and need to be ticked (ungrouped)
-	UPROPERTY(transient)
-	TArray<FAnimTickRecord> UngroupedActivePlayers;
+	TArray<FAnimTickRecord> UngroupedActivePlayerArrays[2];
 
 	// The set of tick groups for this anim instance
-	UPROPERTY(transient)
-	TArray<FAnimGroupInstance> SyncGroups;
+	TArray<FAnimGroupInstance> SyncGroupArrays[2];
 
 	/** Array indicating active vertex anims (by reference) generated by anim instance. */
 	UPROPERTY(transient)
@@ -366,6 +364,12 @@ public:
 	// Creates an uninitialized tick record in the list for the correct group or the ungrouped array.  If the group is valid, OutSyncGroupPtr will point to the group.
 	FAnimTickRecord& CreateUninitializedTickRecord(int32 GroupIndex, FAnimGroupInstance*& OutSyncGroupPtr);
 
+	// Populate a tick record with previous frames data
+	void InitTickRecord(const TArray<FAnimTickRecord>& TickRecords, FAnimTickRecord &NewTickRecord) const;
+	
+	// Populate a tick record with previous frames data
+	void InitTickRecordFromLastFrame(int32 GroupIndex, FAnimTickRecord& NewTickRecord) const;
+
 	// slotnode interfaces
 	void GetSlotWeight(FName const & SlotNodeName, float& out_SlotNodeWeight, float& out_SourceWeight) const;
 	void SlotEvaluatePose(FName SlotNodeName, const FCompactPose& SourcePose, const FBlendedCurve& SourceCurve, FCompactPose& BlendedPose, FBlendedCurve& BlendedCurve, float SlotNodeWeight);
@@ -706,6 +710,9 @@ protected:
 	/** Gets the most relevant asset player in a specified state */
 	FAnimNode_AssetPlayerBase* GetRelevantAssetPlayerFromState(int32 MachineIndex, int32 StateIndex);
 
+	/** Returns the baked sync group index from the compile step */
+	int32 GetSyncGroupIndexFromName(FName SyncGroupName) const;
+
 	//////////////////////////////////////////////////////////////////////////
 
 public:
@@ -770,6 +777,12 @@ public:
 	void UnlockAIResources(bool bUnlockMovement, bool UnlockAILogic);
 	//--- AI communication end ---//
 
+	UFUNCTION(BlueprintCallable, Category = "Animation")
+	bool GetTimeToClosestMarker(FName SyncGroup, FName MarkerName, float& OutMarkerTime) const;
+
+	UFUNCTION(BlueprintCallable, Category = "Animation")
+	bool HasMarkerBeenHitThisFrame(FName SyncGroup, FName MarkerName) const;
+
 public:
 	// Root node of animation graph
 	struct FAnimNode_Base* RootNode;
@@ -848,6 +861,12 @@ private:
 	/** Material parameters that we had been changing and now need to clear */
 	TArray<FName> MaterialParamatersToClear;
 
+	// Current sync group buffer index
+	int32 ActiveSyncGroupsArray;
+
+	//This frames marker sync data
+	FMarkerTickContext MarkerTickContext;
+
 public: 
 	/** Update all internal curves from Blended Curve */
 	void UpdateCurves(const FBlendedCurve& InCurves);
diff --git a/Engine/Source/Runtime/Engine/Classes/Animation/AnimSequence.h b/Engine/Source/Runtime/Engine/Classes/Animation/AnimSequence.h
--- a/Engine/Source/Runtime/Engine/Classes/Animation/AnimSequence.h
+++ b/Engine/Source/Runtime/Engine/Classes/Animation/AnimSequence.h
@@ -523,8 +523,17 @@ class ENGINE_API UAnimSequence : public UAnimSequenceBase
 	UPROPERTY(transient)
 	bool bNeedsRebake;
 
+	// Track whether we have updated markers so cached data can be updated
+	int32 MarkerDataUpdateCounter;
 #endif // WITH_EDITORONLY_DATA
 
+	/** Authored Sync markers */
+	UPROPERTY()
+	TArray<FAnimSyncMarker>		AuthoredSyncMarkers;
+
+	/** List of Unique marker names in this animation sequence */
+	TArray<FName>				UniqueMarkerNames;
+
 public:
 	// Begin UObject interface
 	virtual void Serialize(FArchive& Ar) override;
@@ -542,6 +551,7 @@ public:
 
 	// Begin UAnimationAsset interface
 	virtual bool IsValidAdditive() const override;
+	virtual TArray<FName>* GetUniqueMarkerNames() { return &UniqueMarkerNames; }
 #if WITH_EDITOR
 	virtual bool GetAllAnimationSequencesReferred(TArray<UAnimSequence*>& AnimationSequences) override;
 	virtual void ReplaceReferredAnimations(const TMap<UAnimSequence*, UAnimSequence*>& ReplacementMap) override;
@@ -552,6 +562,7 @@ public:
 	// Begin UAnimSequenceBase interface
 	virtual void OnAssetPlayerTickedInternal(FAnimAssetTickContext &Context, const float PreviousTime, const float MoveDelta, const FAnimTickRecord &Instance, class UAnimInstance* InstanceOwner) const override;
 	virtual bool HasRootMotion() const override { return bEnableRootMotion; }
+	virtual void RefreshCacheData() override;
 	// End UAnimSequenceBase interface
 
 	// Extract Root Motion transform from the animation
@@ -790,8 +801,24 @@ public:
 	 * Add validation check to see if it's being ready to play or not
 	 */
 	virtual bool IsValidToPlay() const override;
+
+	// Get a pointer to the data for a given Anim Notify
+	uint8* FindSyncMarkerPropertyData(int32 SyncMarkerIndex, UArrayProperty*& ArrayProperty);
+
+	virtual int32 GetMarkerUpdateCounter() const { return MarkerDataUpdateCounter; }
 #endif
 
+	/** Sort the sync markers array by time, earliest first. */
+	void SortSyncMarkers();
+
+	// Advancing based on markers
+	float GetCurrentTimeFromMarkers(FMarkerPair& PrevMarker, FMarkerPair& NextMarker, float PositionBetweenMarkers) const;
+	virtual void AdvanceMarkerPhaseAsLeader(bool bLooping, float MoveDelta, const TArray<FName>& ValidMarkerNames, float& CurrentTime, FMarkerPair& PrevMarker, FMarkerPair& NextMarker, TArray<FPassedMarker>& MarkersPassed) const;
+	virtual void AdvanceMarkerPhaseAsFollower(const FMarkerTickContext& Context, float DeltaRemaining, bool bLooping, float& CurrentTime, FMarkerPair& PreviousMarker, FMarkerPair& NextMarker) const;
+	virtual void GetMarkerIndicesForTime(float CurrentTime, bool bLooping, const TArray<FName>& ValidMarkerNames, FMarkerPair& OutPrevMarker, FMarkerPair& OutNextMarker) const;
+	virtual FMarkerSyncAnimPosition GetMarkerSyncPositionfromMarkerIndicies(int32 PrevMarker, int32 NextMarker, float CurrentTime) const;
+	virtual void GetMarkerIndicesForPosition(const FMarkerSyncAnimPosition& SyncPosition, bool bLooping, FMarkerPair& OutPrevMarker, FMarkerPair& OutNextMarker, float& CurrentTime) const;
+
 private:
 	/**
 	* Get Bone Transform of the animation for the Time given, relative to Parent for all RequiredBones
@@ -824,6 +851,8 @@ private:
 	void ResetAnimation();
 	/** Refresh Track Map from Animation Track Names **/
 	void RefreshTrackMapFromAnimTrackNames();
+	/** Refresh sync marker data*/
+	void RefreshSyncMarkerDataFromAuthored();
 
 	/**
 	 * Utility function that helps to remove track, you can't just remove RawAnimationData
diff --git a/Engine/Source/Runtime/Engine/Classes/Animation/AnimSequenceBase.h b/Engine/Source/Runtime/Engine/Classes/Animation/AnimSequenceBase.h
--- a/Engine/Source/Runtime/Engine/Classes/Animation/AnimSequenceBase.h
+++ b/Engine/Source/Runtime/Engine/Classes/Animation/AnimSequenceBase.h
@@ -98,9 +98,6 @@ class UAnimSequenceBase : public UAnimationAsset
 
 	/** Get the time at the given frame */
 	virtual float GetTimeAtFrame(const int32 Frame) const;
-
-	// update anim notify track cache
-	ENGINE_API void UpdateAnimNotifyTrackCache();
 	
 	// @todo document
 	ENGINE_API void InitializeNotifyTrack();
@@ -116,10 +113,20 @@ class UAnimSequenceBase : public UAnimationAsset
 	// Get a pointer to the data for a given Anim Notify
 	ENGINE_API uint8* FindNotifyPropertyData(int32 NotifyIndex, UArrayProperty*& ArrayProperty);
 
+	// Get a pointer to the data for a given array property item
+	ENGINE_API uint8* FindArrayProperty(const TCHAR* PropName, UArrayProperty*& ArrayProperty, int32 ArrayIndex);
+
 #endif	//WITH_EDITORONLY_DATA
+	// update cache data (notify tracks, sync markers)
+	ENGINE_API virtual void RefreshCacheData();
 
 	// Begin UAnimationAsset interface
-	virtual void TickAssetPlayerInstance(const FAnimTickRecord& Instance, class UAnimInstance* InstanceOwner, FAnimAssetTickContext& Context) const override;
+	virtual void TickAssetPlayerInstance(FAnimTickRecord& Instance, class UAnimInstance* InstanceOwner, FAnimAssetTickContext& Context) const override;
+
+	void TickByMarkerAsFollower(FMarkerTickRecord &Instance, FMarkerTickContext &MarkerContext, float& CurrentTime, float& OutPreviousTime, const float MoveDelta, const bool bLooping) const;
+
+	void TickByMarkerAsLeader(FMarkerTickRecord& Instance, FMarkerTickContext& MarkerContext, float& CurrentTime, float& OutPreviousTime, const float MoveDelta, const bool bLooping) const;
+
 	// this is used in editor only when used for transition getter
 	// this doesn't mean max time. In Sequence, this is SequenceLength,
 	// but for BlendSpace CurrentTime is normalized [0,1], so this is 1
@@ -142,6 +149,11 @@ class UAnimSequenceBase : public UAnimationAsset
 
 	virtual void Serialize(FArchive& Ar) override;
 
+	virtual void AdvanceMarkerPhaseAsLeader(bool bLooping, float MoveDelta, const TArray<FName>& ValidMarkerNames, float& CurrentTime, FMarkerPair& PrevMarker, FMarkerPair& NextMarker, TArray<FPassedMarker>& MarkersPassed) const { check(false); /*Should never call this (either missing override or calling on unsupported asset */ }
+	virtual void AdvanceMarkerPhaseAsFollower(const FMarkerTickContext& Context, float DeltaRemaining, bool bLooping, float& CurrentTime, FMarkerPair& PreviousMarker, FMarkerPair& NextMarker) const { check(false); /*Should never call this (either missing override or calling on unsupported asset */ }
+	virtual void GetMarkerIndicesForTime(float CurrentTime, bool bLooping, const TArray<FName>& ValidMarkerNames, FMarkerPair& OutPrevMarker, FMarkerPair& OutNextMarker) const { check(false); /*Should never call this (either missing override or calling on unsupported asset */ }
+	virtual FMarkerSyncAnimPosition GetMarkerSyncPositionfromMarkerIndicies(int32 PrevMarker, int32 NextMarker, float CurrentTime) const { check(false); return FMarkerSyncAnimPosition(); /*Should never call this (either missing override or calling on unsupported asset */ }
+	virtual void GetMarkerIndicesForPosition(const FMarkerSyncAnimPosition& SyncPosition, bool bLooping, FMarkerPair& OutPrevMarker, FMarkerPair& OutNextMarker, float& CurrentTime) const { check(false); /*Should never call this (either missing override or calling on unsupported asset */ }
 #if WITH_EDITOR
 private:
 	DECLARE_MULTICAST_DELEGATE( FOnNotifyChangedMulticaster );
diff --git a/Engine/Source/Runtime/Engine/Classes/Animation/AnimationAsset.h b/Engine/Source/Runtime/Engine/Classes/Animation/AnimationAsset.h
--- a/Engine/Source/Runtime/Engine/Classes/Animation/AnimationAsset.h
+++ b/Engine/Source/Runtime/Engine/Classes/Animation/AnimationAsset.h
@@ -13,6 +13,32 @@
 #include "Engine/SkeletalMesh.h"
 #include "AnimationAsset.generated.h"
 
+namespace MarkerIndexSpecialValues
+{
+	enum Type
+	{
+		Unitialized = -2,
+		AnimationBoundary = -1,
+	};
+};
+
+struct FMarkerPair
+{
+	int32 MarkerIndex;
+	float TimeToMarker;
+
+	FMarkerPair() : MarkerIndex(MarkerIndexSpecialValues::Unitialized) {}
+};
+
+struct FMarkerTickRecord
+{
+	//Current Position in marker space, equivalent to TimeAccumulator
+	FMarkerPair PreviousMarker;
+	FMarkerPair NextMarker;
+
+	bool IsValid() const { return PreviousMarker.MarkerIndex != MarkerIndexSpecialValues::Unitialized && NextMarker.MarkerIndex != MarkerIndexSpecialValues::Unitialized; }
+};
+
 /** Transform definition */
 USTRUCT(BlueprintType)
 struct FBlendSampleData
@@ -28,6 +54,11 @@ struct FBlendSampleData
 	UPROPERTY()
 	float Time;
 
+	UPROPERTY()
+	float PreviousTime;
+
+	FMarkerTickRecord MarkerTickRecord;
+
 	// transient perbone interpolation data
 	TArray<float> PerBoneBlendData;
 
@@ -142,6 +173,29 @@ struct FAnimExtractContext
 	}
 };
 
+//Represent a current play position in an animation
+//based on sync markers
+struct FMarkerSyncAnimPosition
+{
+	/** The marker we have passed*/
+	FName PreviousMarkerName;
+
+	/** The marker we are heading towards */
+	FName NextMarkerName;
+
+	/** Value between 0 and 1 representing where we are:
+	0   we are at PreviousMarker
+	1   we are at NextMarker
+	0.5 we are half way between the two */
+	float PositionBetweenMarkers;
+};
+
+struct FPassedMarker
+{
+	FName PassedMarkerName;
+
+	float DeltaTimeWhenPassed;
+};
 
 /**
  * Information about an animation asset that needs to be ticked
@@ -161,13 +215,73 @@ struct FAnimTickRecord
 	float PlayRateMultiplier;
 	float EffectiveBlendWeight;
 	bool bLooping;
+	
+	bool bCanUseMarkerSync;
+	
+	FMarkerTickRecord MarkerTickRecord;
+
+	void* SourceNodeRef; //Used to find last frames TickRecord
 
 public:
 	FAnimTickRecord()
+		: bCanUseMarkerSync(false)
+		, SourceNodeRef(nullptr)
 	{
 	}
 };
 
+class FMarkerTickContext
+{
+public:
+	FMarkerTickContext(const TArray<FName>& ValidMarkerNames) : ValidMarkers(ValidMarkerNames) {}
+	FMarkerTickContext() {}
+
+	void SetMarkerSyncStartPosition(const FMarkerSyncAnimPosition& SyncPosition)
+	{
+		MarkerSyncStartPostion = SyncPosition;
+	}
+
+	void SetMarkerSyncEndPosition(const FMarkerSyncAnimPosition& SyncPosition)
+	{
+		MarkerSyncEndPostion = SyncPosition;
+	}
+
+	const FMarkerSyncAnimPosition& GetMarkerSyncStartPosition() const
+	{
+		return MarkerSyncStartPostion;
+	}
+
+	const FMarkerSyncAnimPosition& GetMarkerSyncEndPosition() const
+	{
+		return MarkerSyncEndPostion;
+	}
+
+	const TArray<FName>& GetValidMarkerNames() const
+	{
+		return ValidMarkers;
+	}
+
+	bool IsMarkerSyncStartValid() const
+	{
+		return !(MarkerSyncStartPostion.PreviousMarkerName == NAME_None && MarkerSyncStartPostion.NextMarkerName == NAME_None);
+	}
+
+	TArray<FPassedMarker> MarkersPassedThisTick;
+
+private:
+	// Structure representing our sync position based on markers before tick
+	// This is used to allow new animations to play from the right marker position
+	FMarkerSyncAnimPosition MarkerSyncStartPostion;
+
+	// Structure representing our sync position based on markers after tick
+	FMarkerSyncAnimPosition MarkerSyncEndPostion;
+
+
+	// Valid marker names for this sync group
+	TArray<FName> ValidMarkers;
+};
+
+
 UENUM()
 namespace EAnimGroupRole
 {
@@ -195,9 +309,19 @@ public:
 
 	// The current group leader
 	int32 GroupLeaderIndex;
+
+	// Valid marker names for this sync group
+	TArray<FName> ValidMarkers;
+
+	// Can we use sync markers for ticking this sync group
+	bool bCanUseMarkerSync;
+
+	FMarkerTickContext MarkerTickContext;
+
 public:
 	FAnimGroupInstance()
 		: GroupLeaderIndex(INDEX_NONE)
+		, bCanUseMarkerSync(false)
 	{
 	}
 
@@ -205,35 +329,15 @@ public:
 	{
 		GroupLeaderIndex = INDEX_NONE;
 		ActivePlayers.Empty(ActivePlayers.Num());
+		bCanUseMarkerSync = false;
 	}
 
 	// Checks the last tick record in the ActivePlayers array to see if it's a better leader than the current candidate.
 	// This should be called once for each record added to ActivePlayers, after the record is setup.
-	void TestTickRecordForLeadership(EAnimGroupRole::Type MembershipType)
-	{
-		int32 TestIndex = ActivePlayers.Num() - 1;
-		const FAnimTickRecord& Candidate = ActivePlayers[TestIndex];
-		
-		switch (MembershipType)
-		{
-		case EAnimGroupRole::CanBeLeader:
-			// Set it if we're better than the current leader (or if there is no leader yet)
-			if ((GroupLeaderIndex == INDEX_NONE) || (ActivePlayers[GroupLeaderIndex].EffectiveBlendWeight < Candidate.EffectiveBlendWeight))
-			{
-				// This is a better leader
-				GroupLeaderIndex = TestIndex;
-			}
-			break;
-		case EAnimGroupRole::AlwaysLeader:
-			// Always set the leader index
-			GroupLeaderIndex = TestIndex;
-			break;
-		default:
-		case EAnimGroupRole::AlwaysFollower:
-			// Never set the leader index; the actual tick code will handle the case of no leader by using the first element in the array
-			break;
-		}
-	}
+	ENGINE_API void TestTickRecordForLeadership(EAnimGroupRole::Type MembershipType);
+
+	// Called after all tick records have been added but before assets are actually ticked
+	ENGINE_API void Finalize(const TArray<FName>& PreviousValidMarkers);
 };
 
 /** Utility struct to accumulate root motion. */
@@ -346,10 +450,23 @@ struct FRootMotionMovementParams
 struct FAnimAssetTickContext
 {
 public:
+	FAnimAssetTickContext(float InDeltaTime, ERootMotionMode::Type InRootMotionMode, const TArray<FName>& ValidMarkerNames)
+		: RootMotionMode(InRootMotionMode)
+		, DeltaTime(InDeltaTime)
+		, LeaderDelta(0.f)
+		, AnimLengthRatio(0.0f)
+		, bIsMarkerPositionValid(ValidMarkerNames.Num() > 0)
+		, bIsLeader(true)
+		, MarkerTickContext(ValidMarkerNames)
+	{
+	}
+
 	FAnimAssetTickContext(float InDeltaTime, ERootMotionMode::Type InRootMotionMode)
 		: RootMotionMode(InRootMotionMode)
 		, DeltaTime(InDeltaTime)
-		, SyncPoint(0.0f)
+		, LeaderDelta(0.f)
+		, AnimLengthRatio(0.0f)
+		, bIsMarkerPositionValid(false)
 		, bIsLeader(true)
 	{
 	}
@@ -371,16 +488,31 @@ public:
 		return DeltaTime;
 	}
 
-	void SetSyncPoint(float NormalizedTime)
+	void SetLeaderDelta(float InLeaderDelta)
 	{
-		SyncPoint = NormalizedTime;
+		LeaderDelta = InLeaderDelta;
+	}
+
+	float GetLeaderDelta() const
+	{
+		return LeaderDelta;
+	}
+
+	void SetAnimationPositionRatio(float NormalizedTime)
+	{
+		AnimLengthRatio = NormalizedTime;
 	}
 
 	// Returns the synchronization point (normalized time; only legal to call if ticking a follower)
-	float GetSyncPoint() const
+	float GetAnimationPositionRatio() const
 	{
 		checkSlow(!bIsLeader);
-		return SyncPoint;
+		return AnimLengthRatio;
+	}
+
+	bool CanUseMarkerPosition() const
+	{
+		return bIsMarkerPositionValid;
 	}
 
 	void ConvertToFollower()
@@ -399,11 +531,17 @@ public:
 	// The root motion mode of the owning AnimInstance
 	ERootMotionMode::Type RootMotionMode;
 
+	FMarkerTickContext MarkerTickContext;
+
 private:
 	float DeltaTime;
 
-	// The structure used to pass synchronization state between members of a sync group
-	float SyncPoint;
+	float LeaderDelta;
+
+	// Float in 0 - 1 range representing how far through an animation we are
+	float AnimLengthRatio;
+
+	bool bIsMarkerPositionValid;
 
 	bool bIsLeader;
 };
@@ -454,7 +592,7 @@ public:
 	 * @param InstanceOwner	AnimInstance playing this asset
 	 * @param Context		The tick context (leader/follower, delta time, sync point, etc...)
 	 */
-	virtual void TickAssetPlayerInstance(const FAnimTickRecord& Instance, class UAnimInstance* InstanceOwner, FAnimAssetTickContext& Context) const {}
+	virtual void TickAssetPlayerInstance(FAnimTickRecord& Instance, class UAnimInstance* InstanceOwner, FAnimAssetTickContext& Context) const {}
 	// this is used in editor only when used for transition getter
 	// this doesn't mean max time. In Sequence, this is SequenceLength,
 	// but for BlendSpace CurrentTime is normalized [0,1], so this is 1
@@ -494,8 +632,13 @@ public:
 
 	ENGINE_API void SetPreviewMesh(USkeletalMesh* PreviewMesh);
 	ENGINE_API USkeletalMesh* GetPreviewMesh();
+
+	ENGINE_API virtual int32 GetMarkerUpdateCounter() const { return 0; }
 #endif //WITH_EDITOR
 
+	/** Return a list of unique marker names for blending compatibility */
+	ENGINE_API virtual TArray<FName>* GetUniqueMarkerNames() { return NULL; }
+
 #if WITH_EDITORONLY_DATA
 	/** Information for thumbnail rendering */
 	UPROPERTY(VisibleAnywhere, Instanced, Category = Thumbnail)
diff --git a/Engine/Source/Runtime/Engine/Classes/Animation/BlendSpaceBase.h b/Engine/Source/Runtime/Engine/Classes/Animation/BlendSpaceBase.h
--- a/Engine/Source/Runtime/Engine/Classes/Animation/BlendSpaceBase.h
+++ b/Engine/Source/Runtime/Engine/Classes/Animation/BlendSpaceBase.h
@@ -263,7 +263,7 @@ public:
 	// End UObject interface
 
 	// Begin UAnimationAsset interface
-	virtual void TickAssetPlayerInstance(const FAnimTickRecord& Instance, class UAnimInstance* InstanceOwner, FAnimAssetTickContext& Context) const override;
+	virtual void TickAssetPlayerInstance(FAnimTickRecord& Instance, class UAnimInstance* InstanceOwner, FAnimAssetTickContext& Context) const override;
 	// this is used in editor only when used for transition getter
 	// this doesn't mean max time. In Sequence, this is SequenceLength,
 	// but for BlendSpace CurrentTime is normalized [0,1], so this is 1
@@ -271,9 +271,24 @@ public:
 #if WITH_EDITOR
 	virtual bool GetAllAnimationSequencesReferred(TArray<UAnimSequence*>& AnimationSequences) override;
 	virtual void ReplaceReferredAnimations(const TMap<UAnimSequence*, UAnimSequence*>& ReplacementMap) override;
+	virtual int32 GetMarkerUpdateCounter() const;
 #endif
+	virtual TArray<FName>* GetUniqueMarkerNames() override { return (bAllSequencesHaveMatchingMarkers && SampleData.Num() > 0) ? SampleData[0].Animation->GetUniqueMarkerNames() : nullptr; }
 	// End of UAnimationAsset interface
 
+	void TickFollowerSamples(TArray<FBlendSampleData> &SampleDataList, const int32 HighestWeightIndex, FAnimAssetTickContext &Context) const
+	{
+		for (int32 SampleIndex = 0; SampleIndex < SampleDataList.Num(); ++SampleIndex)
+		{
+			FBlendSampleData& SampleDataItem = SampleDataList[SampleIndex];
+			const FBlendSample& Sample = SampleData[SampleDataItem.SampleDataIndex];
+			if (HighestWeightIndex != SampleIndex)
+			{
+				Sample.Animation->TickByMarkerAsFollower(SampleDataItem.MarkerTickRecord, Context.MarkerTickContext, SampleDataItem.Time, SampleDataItem.PreviousTime, Context.GetLeaderDelta(), true);
+			}
+		}
+	}
+
 	/**
 	 * BlendSpace Get Animation Pose function
 	 */
@@ -412,6 +427,15 @@ protected:
 	 */
 	virtual void GetRawSamplesFromBlendInput(const FVector &BlendInput, TArray<FGridBlendSample, TInlineAllocator<4> > & OutBlendSamples) const {}
 
+	/** Track whether all our sequences have the same marker set for blending (if false use original scale based on length approach) */
+	bool bAllSequencesHaveMatchingMarkers;
+
+#if WITH_EDITOR
+private:
+	// Track whether we have updated markers so cached data can be updated
+	int32 MarkerDataUpdateCounter;
+#endif
+
 public:
 	
 	/** Sample animation data **/
diff --git a/Engine/Source/Runtime/Engine/Private/Animation/AnimInstance.cpp b/Engine/Source/Runtime/Engine/Private/Animation/AnimInstance.cpp
--- a/Engine/Source/Runtime/Engine/Private/Animation/AnimInstance.cpp
+++ b/Engine/Source/Runtime/Engine/Private/Animation/AnimInstance.cpp
@@ -79,6 +79,7 @@ UAnimInstance::UAnimInstance(const FObjectInitializer& ObjectInitializer)
 	RootNode = NULL;
 	RootMotionMode = ERootMotionMode::RootMotionFromMontagesOnly;
 	SlotNodeInitializationCounter = INDEX_NONE;
+	ActiveSyncGroupsArray = 0;
 }
 
 void UAnimInstance::MakeSequenceTickRecord(FAnimTickRecord& TickRecord, class UAnimSequenceBase* Sequence, bool bLooping, float PlayRate, float FinalBlendWeight, float& CurrentTime) const
@@ -128,6 +129,7 @@ FAnimTickRecord& UAnimInstance::CreateUninitializedTickRecord(int32 GroupIndex,
 	OutSyncGroupPtr = NULL;
 	if (GroupIndex >= 0)
 	{
+		TArray<FAnimGroupInstance>& SyncGroups = SyncGroupArrays[ActiveSyncGroupsArray];
 		while (SyncGroups.Num() <= GroupIndex)
 		{
 			new (SyncGroups) FAnimGroupInstance();
@@ -136,10 +138,47 @@ FAnimTickRecord& UAnimInstance::CreateUninitializedTickRecord(int32 GroupIndex,
 	}
 
 	// Create the record
-	FAnimTickRecord* TickRecord = new ((OutSyncGroupPtr != NULL) ? OutSyncGroupPtr->ActivePlayers : UngroupedActivePlayers) FAnimTickRecord();
+	FAnimTickRecord* TickRecord = new ((OutSyncGroupPtr != NULL) ? OutSyncGroupPtr->ActivePlayers : UngroupedActivePlayerArrays[ActiveSyncGroupsArray]) FAnimTickRecord();
 	return *TickRecord;
 }
 
+void UAnimInstance::InitTickRecord(const TArray<FAnimTickRecord>& TickRecords, FAnimTickRecord &NewTickRecord) const
+{
+	for (const FAnimTickRecord& PrevTickRecord : TickRecords)
+	{
+		if (NewTickRecord.SourceNodeRef == PrevTickRecord.SourceNodeRef)
+		{
+			if (NewTickRecord.SourceAsset == PrevTickRecord.SourceAsset
+#if WITH_EDITOR					
+				&& NewTickRecord.SourceAsset->GetMarkerUpdateCounter() == PrevTickRecord.SourceAsset->GetMarkerUpdateCounter()
+#endif
+				)
+			{
+				NewTickRecord.MarkerTickRecord = PrevTickRecord.MarkerTickRecord;
+			}
+			break;
+		}
+	}
+}
+
+void UAnimInstance::InitTickRecordFromLastFrame(int32 GroupIndex, FAnimTickRecord& NewTickRecord) const
+{
+	const int32 PreviousBufferIndex = 1 - ActiveSyncGroupsArray;
+	if (GroupIndex == INDEX_NONE)
+	{
+		InitTickRecord(UngroupedActivePlayerArrays[PreviousBufferIndex], NewTickRecord);
+	}
+	else
+	{
+		const TArray<FAnimGroupInstance>& PreviousSyncGroups = SyncGroupArrays[PreviousBufferIndex];
+		if (PreviousSyncGroups.IsValidIndex(GroupIndex))
+		{
+			const TArray<FAnimTickRecord>& TickRecords = PreviousSyncGroups[GroupIndex].ActivePlayers;
+			InitTickRecord(TickRecords, NewTickRecord);
+		}
+	}
+}
+
 AActor* UAnimInstance::GetOwningActor() const
 {
 	USkeletalMeshComponent* OwnerComponent = GetSkelMeshComponent();
@@ -335,20 +374,6 @@ void UAnimInstance::UpdateAnimation(float DeltaSeconds)
 	}
 #endif
 
-	{
-		QUICK_SCOPE_CYCLE_COUNTER(STAT_UAnimInstance_UpdateAnimation_MiscClear);
-		AnimNotifies.Reset();
-
-		ClearSlotNodeWeights();
-
-	    // Reset the player tick list (but keep it presized)
-	    UngroupedActivePlayers.Empty(UngroupedActivePlayers.Num());
-	    for (int32 GroupIndex = 0; GroupIndex < SyncGroups.Num(); ++GroupIndex)
-	    {
-		    SyncGroups[GroupIndex].Reset();
-	    }
-	}
-
 	{
 		SCOPE_CYCLE_COUNTER(STAT_NativeUpdateAnimation);
 		NativeUpdateAnimation(DeltaSeconds);
@@ -358,6 +383,26 @@ void UAnimInstance::UpdateAnimation(float DeltaSeconds)
 		BlueprintUpdateAnimation(DeltaSeconds);
 	}
 
+	ActiveSyncGroupsArray = 1 - ActiveSyncGroupsArray;
+	TArray<FAnimGroupInstance>& SyncGroups = SyncGroupArrays[ActiveSyncGroupsArray];
+	TArray<FAnimTickRecord>& UngroupedActivePlayers = UngroupedActivePlayerArrays[ActiveSyncGroupsArray];
+
+	const TArray<FAnimGroupInstance>& PreviousSyncGroups = SyncGroupArrays[1 - ActiveSyncGroupsArray];
+
+	{
+		QUICK_SCOPE_CYCLE_COUNTER(STAT_UAnimInstance_UpdateAnimation_MiscClear);
+		AnimNotifies.Reset();
+
+		ClearSlotNodeWeights();
+
+		// Reset the player tick list (but keep it presized)
+		UngroupedActivePlayers.Reset();
+		for (int32 GroupIndex = 0; GroupIndex < SyncGroups.Num(); ++GroupIndex)
+		{
+			SyncGroups[GroupIndex].Reset();
+		}
+	}
+
 	// update weight before all nodes update comes in
 	Montage_UpdateWeight(DeltaSeconds);
 
@@ -372,28 +417,33 @@ void UAnimInstance::UpdateAnimation(float DeltaSeconds)
 	{
 		QUICK_SCOPE_CYCLE_COUNTER(STAT_TickAssetPlayerInstances);
 
-	// Handle all players inside sync groups
+		// Handle all players inside sync groups
 		for (int32 GroupIndex = 0; GroupIndex < SyncGroups.Num(); ++GroupIndex)
 		{
 			FAnimGroupInstance& SyncGroup = SyncGroups[GroupIndex];
     
 			if (SyncGroup.ActivePlayers.Num() > 0)
 			{
-				const int32 GroupLeaderIndex = FMath::Max(SyncGroup.GroupLeaderIndex, 0);
+				static TArray<FName> TempNames;
+				const TArray<FName>& PreviousValidMarkers = PreviousSyncGroups.IsValidIndex(GroupIndex) ? PreviousSyncGroups[GroupIndex].ValidMarkers : TempNames;
+				SyncGroup.Finalize(PreviousValidMarkers);
     
 				// Tick the group leader
-				FAnimAssetTickContext TickContext(DeltaSeconds, RootMotionMode);
-				FAnimTickRecord& GroupLeader = SyncGroup.ActivePlayers[GroupLeaderIndex];
+				FAnimAssetTickContext TickContext(DeltaSeconds, RootMotionMode, SyncGroup.ValidMarkers);
+				FAnimTickRecord& GroupLeader = SyncGroup.ActivePlayers[SyncGroup.GroupLeaderIndex];
 				{
 					QUICK_SCOPE_CYCLE_COUNTER(STAT_UAnimInstance_UpdateAnimation_TickAssetPlayerInstance);
 					FScopeCycleCounterUObject Scope(GroupLeader.SourceAsset);
 					GroupLeader.SourceAsset->TickAssetPlayerInstance(GroupLeader, this, TickContext);
 				}
+
 				if (RootMotionMode == ERootMotionMode::RootMotionFromEverything && TickContext.RootMotionMovementParams.bHasRootMotion)
 				{
 					ExtractedRootMotion.AccumulateWithBlend(TickContext.RootMotionMovementParams.RootMotionTransform, GroupLeader.EffectiveBlendWeight);
 				}
     
+				SyncGroup.MarkerTickContext = TickContext.MarkerTickContext;
+
 				// Update everything else to follow the leader
 				if (SyncGroup.ActivePlayers.Num() > 1)
 				{
@@ -401,9 +451,9 @@ void UAnimInstance::UpdateAnimation(float DeltaSeconds)
     
 					for (int32 TickIndex = 0; TickIndex < SyncGroup.ActivePlayers.Num(); ++TickIndex)
 					{
-						if (TickIndex != GroupLeaderIndex)
+						if (TickIndex != SyncGroup.GroupLeaderIndex)
 						{
-							const FAnimTickRecord& AssetPlayer = SyncGroup.ActivePlayers[TickIndex];
+							FAnimTickRecord& AssetPlayer = SyncGroup.ActivePlayers[TickIndex];
 							{
 								QUICK_SCOPE_CYCLE_COUNTER(STAT_UAnimInstance_UpdateAnimation_TickAssetPlayerInstance);
 								FScopeCycleCounterUObject Scope(AssetPlayer.SourceAsset);
@@ -422,8 +472,13 @@ void UAnimInstance::UpdateAnimation(float DeltaSeconds)
 		// Handle the remaining ungrouped animation players
 		for (int32 TickIndex = 0; TickIndex < UngroupedActivePlayers.Num(); ++TickIndex)
 		{
-			const FAnimTickRecord& AssetPlayerToTick = UngroupedActivePlayers[TickIndex];
-			FAnimAssetTickContext TickContext(DeltaSeconds, RootMotionMode);
+			FAnimTickRecord& AssetPlayerToTick = UngroupedActivePlayers[TickIndex];
+			AssetPlayerToTick.bCanUseMarkerSync = true;
+			static TArray<FName> TempNames;
+			const TArray<FName>* UniqueNames = AssetPlayerToTick.SourceAsset->GetUniqueMarkerNames();
+			const TArray<FName>& ValidMarkers = UniqueNames ? *UniqueNames : TempNames;
+
+			FAnimAssetTickContext TickContext(DeltaSeconds, RootMotionMode, ValidMarkers);
 		    {
 			    QUICK_SCOPE_CYCLE_COUNTER(STAT_UAnimInstance_UpdateAnimation_TickAssetPlayerInstance);
 			    FScopeCycleCounterUObject Scope(AssetPlayerToTick.SourceAsset);
@@ -798,6 +853,9 @@ void UAnimInstance::DisplayDebug(class UCanvas* Canvas, const FDebugDisplayInfo&
 		FIndenter AnimIndent(Indent);
 
 		//Display Sync Groups
+		const TArray<FAnimGroupInstance>& SyncGroups = SyncGroupArrays[ActiveSyncGroupsArray];
+		const TArray<FAnimTickRecord>& UngroupedActivePlayers = UngroupedActivePlayerArrays[ActiveSyncGroupsArray];
+
 		Heading = FString::Printf(TEXT("SyncGroups: %i"), SyncGroups.Num());
 		YL = Canvas->DrawText(RenderFont, Heading, Indent, YPos, 1.f, 1.f, RenderInfo);
 		YPos += YL;
@@ -805,7 +863,7 @@ void UAnimInstance::DisplayDebug(class UCanvas* Canvas, const FDebugDisplayInfo&
 		for (int32 GroupIndex = 0; GroupIndex < SyncGroups.Num(); ++GroupIndex)
 		{
 			FIndenter GroupIndent(Indent);
-			FAnimGroupInstance& SyncGroup = SyncGroups[GroupIndex];
+			const FAnimGroupInstance& SyncGroup = SyncGroups[GroupIndex];
 
 			Canvas->SetLinearDrawColor(TextYellow);
 
@@ -815,8 +873,8 @@ void UAnimInstance::DisplayDebug(class UCanvas* Canvas, const FDebugDisplayInfo&
 
 			if (SyncGroup.ActivePlayers.Num() > 0)
 			{
-				const int32 GroupLeaderIndex = FMath::Max(SyncGroup.GroupLeaderIndex, 0);
-				OutputTickRecords(SyncGroup.ActivePlayers, Canvas, Indent, GroupLeaderIndex, TextWhite, ActiveColor, InactiveColor, RenderFont, YPos, RenderInfo, YL, bFullBlendspaceDisplay);
+				check(SyncGroup.GroupLeaderIndex != -1);
+				OutputTickRecords(SyncGroup.ActivePlayers, Canvas, Indent, SyncGroup.GroupLeaderIndex, TextWhite, ActiveColor, InactiveColor, RenderFont, YPos, RenderInfo, YL, bFullBlendspaceDisplay);
 			}
 		}
 
@@ -2640,6 +2698,49 @@ void UAnimInstance::UnlockAIResources(bool bUnlockMovement, bool UnlockAILogic)
 	UE_LOG(LogAnimation, Error, TEXT("%s: UnlockAIResources is no longer supported. Please use UnlockAIResourcesWithAnimation instead."), *GetName());
 }
 
+bool UAnimInstance::GetTimeToClosestMarker(FName SyncGroup, FName MarkerName, float& OutMarkerTime) const
+{
+	const int32 SyncGroupIndex = GetSyncGroupIndexFromName(SyncGroup);
+	const TArray<FAnimGroupInstance>& SyncGroups = SyncGroupArrays[ActiveSyncGroupsArray];
+
+	if (SyncGroups.IsValidIndex(SyncGroupIndex))
+	{
+		const FAnimGroupInstance& SyncGroup = SyncGroups[SyncGroupIndex];
+		if (SyncGroup.ActivePlayers.IsValidIndex(SyncGroup.GroupLeaderIndex))
+		{
+			const FMarkerSyncAnimPosition& EndPosition = SyncGroup.MarkerTickContext.GetMarkerSyncEndPosition();
+			const FAnimTickRecord& Leader = SyncGroup.ActivePlayers[SyncGroup.GroupLeaderIndex];
+			if (EndPosition.PreviousMarkerName == MarkerName)
+			{
+				OutMarkerTime = Leader.MarkerTickRecord.PreviousMarker.TimeToMarker;
+				return true;
+			}
+			else if (EndPosition.NextMarkerName == MarkerName)
+			{
+				OutMarkerTime = Leader.MarkerTickRecord.NextMarker.TimeToMarker;
+				return true;
+			}
+		}
+	}
+	return false;
+}
+
+bool UAnimInstance::HasMarkerBeenHitThisFrame(FName SyncGroup, FName MarkerName) const
+{
+	const int32 SyncGroupIndex = GetSyncGroupIndexFromName(SyncGroup);
+	const TArray<FAnimGroupInstance>& SyncGroups = SyncGroupArrays[ActiveSyncGroupsArray];
+
+	if (SyncGroups.IsValidIndex(SyncGroupIndex))
+	{
+		const FAnimGroupInstance& SyncGroup = SyncGroups[SyncGroupIndex];
+		return SyncGroup.MarkerTickContext.MarkersPassedThisTick.ContainsByPredicate([&MarkerName](const FPassedMarker& PassedMarker) -> bool
+		{
+			return PassedMarker.PassedMarkerName == MarkerName;
+		});
+	}
+	return false;
+}
+
 void UAnimInstance::UpdateMontageEvaluationData()
 {
 	MontageEvaluationData.Empty(MontageInstances.Num());
@@ -2935,6 +3036,16 @@ FAnimNode_AssetPlayerBase* UAnimInstance::GetRelevantAssetPlayerFromState(int32
 	return ResultPlayer;
 }
 
+int32 UAnimInstance::GetSyncGroupIndexFromName(FName SyncGroupName) const
+{
+
+	if (UAnimBlueprintGeneratedClass* AnimBlueprintClass = Cast<UAnimBlueprintGeneratedClass>(GetClass()))
+	{
+		return AnimBlueprintClass->GetSyncGroupIndex(SyncGroupName);
+	}
+	return INDEX_NONE;
+}
+
 bool UAnimInstance::IsRunningParallelEvaluation() const
 {
 	USkeletalMeshComponent* Comp = GetOwningComponent();
diff --git a/Engine/Source/Runtime/Engine/Private/Animation/AnimNode_SequencePlayer.cpp b/Engine/Source/Runtime/Engine/Private/Animation/AnimNode_SequencePlayer.cpp
--- a/Engine/Source/Runtime/Engine/Private/Animation/AnimNode_SequencePlayer.cpp
+++ b/Engine/Source/Runtime/Engine/Private/Animation/AnimNode_SequencePlayer.cpp
@@ -35,6 +35,8 @@ void FAnimNode_SequencePlayer::UpdateAssetPlayer(const FAnimationUpdateContext&
 		// Create a tick record and fill it out
 		FAnimGroupInstance* SyncGroup;
 		FAnimTickRecord& TickRecord = Context.AnimInstance->CreateUninitializedTickRecord(GroupIndex, /*out*/ SyncGroup);
+		TickRecord.SourceNodeRef = this;
+		Context.AnimInstance->InitTickRecordFromLastFrame(GroupIndex, TickRecord);
 
 		if (InternalTimeAccumulator > Sequence->SequenceLength)
 		{
diff --git a/Engine/Source/Runtime/Engine/Private/Animation/AnimSequence.cpp b/Engine/Source/Runtime/Engine/Private/Animation/AnimSequence.cpp
--- a/Engine/Source/Runtime/Engine/Private/Animation/AnimSequence.cpp
+++ b/Engine/Source/Runtime/Engine/Private/Animation/AnimSequence.cpp
@@ -123,6 +123,7 @@ void UAnimSequence::PostInitProperties()
 	{
 		AssetImportData = NewObject<UAssetImportData>(this, TEXT("AssetImportData"));
 	}
+	MarkerDataUpdateCounter = 0;
 #endif
 	Super::PostInitProperties();
 }
@@ -329,6 +330,13 @@ bool UAnimSequence::IsValidToPlay() const
 	// make sure sequence length is valid and raw animation data exists, and compressed
 	return ( SequenceLength > 0.f && RawAnimationData.Num() > 0 && CompressedTrackOffsets.Num() > 0 );
 }
+
+void UAnimSequence::SortSyncMarkers()
+{
+	AuthoredSyncMarkers.Sort();
+	RefreshSyncMarkerDataFromAuthored();
+}
+
 #endif
 
 void UAnimSequence::PreSave()
@@ -3438,6 +3446,8 @@ void UAnimSequence::ResetAnimation()
 	CompressedByteStream.Empty();
 
 	Notifies.Empty();
+	AuthoredSyncMarkers.Empty();
+	UniqueMarkerNames.Empty();
 	AnimNotifyTracks.Empty();
 	RawCurveData.Empty();
 	RateScale = 1.f;
@@ -3468,6 +3478,41 @@ void UAnimSequence::RefreshTrackMapFromAnimTrackNames()
 	}
 }
 
+void UAnimSequence::RefreshSyncMarkerDataFromAuthored()
+{
+#if WITH_EDITOR
+	MarkerDataUpdateCounter++;
+#endif
+
+	if (AuthoredSyncMarkers.Num() > 0)
+	{
+		UniqueMarkerNames.Reset();
+		UniqueMarkerNames.Reserve(AuthoredSyncMarkers.Num());
+
+		const FAnimSyncMarker* PreviousMarker = nullptr;
+		for (const FAnimSyncMarker& Marker : AuthoredSyncMarkers)
+		{
+			UniqueMarkerNames.AddUnique(Marker.MarkerName);
+			PreviousMarker = &Marker;
+		}
+	}
+	else
+	{
+		UniqueMarkerNames.Empty();
+	}
+}
+
+uint8* UAnimSequence::FindSyncMarkerPropertyData(int32 SyncMarkerIndex, UArrayProperty*& ArrayProperty)
+{
+	ArrayProperty = NULL;
+
+	if (AuthoredSyncMarkers.IsValidIndex(SyncMarkerIndex))
+	{
+		return FindArrayProperty(TEXT("AuthoredSyncMarkers"), ArrayProperty, SyncMarkerIndex);
+	}
+	return NULL;
+}
+
 bool UAnimSequence::CreateAnimation(USkeletalMesh * Mesh)
 {
 	// create animation from Mesh's ref pose
@@ -3583,6 +3628,484 @@ bool UAnimSequence::CreateAnimation(UAnimSequence * Sequence)
 	return false;
 }
 
+void UAnimSequence::RefreshCacheData()
+{
+	SortSyncMarkers();
+	RefreshSyncMarkerDataFromAuthored();
+
+	for (int32 TrackIndex = 0; TrackIndex < AnimNotifyTracks.Num(); ++TrackIndex)
+	{
+		AnimNotifyTracks[TrackIndex].SyncMarkers.Empty();
+	}
+
+	for (FAnimSyncMarker& SyncMarker : AuthoredSyncMarkers)
+	{
+		const int32 TrackIndex = SyncMarker.TrackIndex;
+		if (AnimNotifyTracks.IsValidIndex(TrackIndex))
+		{
+			AnimNotifyTracks[TrackIndex].SyncMarkers.Add(&SyncMarker);
+		}
+		else
+		{
+			// This should not happen, but if it does we must find somewhere else to add it
+			ensureMsgf(0, TEXT("AnimNotifyTrack: Wrong indices found"));
+			AnimNotifyTracks[0].SyncMarkers.Add(&SyncMarker);
+			SyncMarker.TrackIndex = 0;
+		}
+	}
+
+	Super::RefreshCacheData();
+}
+
+bool IsMarkerValid(const FAnimSyncMarker* Marker, bool bLooping, const TArray<FName>& ValidMarkerNames)
+{
+	return (Marker == nullptr && !bLooping) || (Marker && ValidMarkerNames.Contains(Marker->MarkerName));
+}
+
+void UAnimSequence::AdvanceMarkerPhaseAsLeader(bool bLooping, float MoveDelta, const TArray<FName>& ValidMarkerNames, float& CurrentTime, FMarkerPair& PrevMarker, FMarkerPair& NextMarker, TArray<FPassedMarker>& MarkersPassed) const
+{
+	check(MoveDelta != 0.f);
+	const bool bPlayingForwards = MoveDelta > 0.f;
+	float CurrentMoveDelta = MoveDelta;
+
+	bool bOffsetInitialized = false;
+	float MarkerTimeOffset = 0.f;
+
+	if (bPlayingForwards)
+	{
+		while (true)
+		{
+			if (NextMarker.MarkerIndex == -1)
+			{
+				float PrevCurrentTime = CurrentTime;
+				CurrentTime = FMath::Min(CurrentTime + CurrentMoveDelta, SequenceLength);
+				NextMarker.TimeToMarker = SequenceLength - CurrentTime;
+				PrevMarker.TimeToMarker -= CurrentTime - PrevCurrentTime; //Add how far we moved to distance from previous marker
+				break;
+			}
+			const FAnimSyncMarker& NextSyncMarker = AuthoredSyncMarkers[NextMarker.MarkerIndex];
+			checkSlow(ValidMarkerNames.Contains(NextSyncMarker.MarkerName));
+			if (!bOffsetInitialized)
+			{
+				bOffsetInitialized = true;
+				if (NextSyncMarker.Time < CurrentTime)
+				{
+					MarkerTimeOffset = SequenceLength;
+				}
+			}
+			const float NextMarkerTime = NextSyncMarker.Time + MarkerTimeOffset;
+			const float TimeToMarker = NextMarkerTime - CurrentTime;
+
+			if (CurrentMoveDelta > TimeToMarker)
+			{
+				CurrentTime = NextSyncMarker.Time;
+				CurrentMoveDelta -= TimeToMarker;
+
+				PrevMarker.MarkerIndex = NextMarker.MarkerIndex;
+				PrevMarker.TimeToMarker = -CurrentMoveDelta;
+
+				int32 PassedMarker = MarkersPassed.Add(FPassedMarker());
+				MarkersPassed[PassedMarker].PassedMarkerName = NextSyncMarker.MarkerName;
+				MarkersPassed[PassedMarker].DeltaTimeWhenPassed = CurrentMoveDelta;
+
+				do
+				{
+					++NextMarker.MarkerIndex;
+					if (NextMarker.MarkerIndex >= AuthoredSyncMarkers.Num())
+					{
+						if (!bLooping)
+						{
+							NextMarker.MarkerIndex = -1;
+							break;
+						}
+						NextMarker.MarkerIndex = 0;
+						MarkerTimeOffset += SequenceLength;
+					}
+				} while (!ValidMarkerNames.Contains(AuthoredSyncMarkers[NextMarker.MarkerIndex].MarkerName));
+			}
+			else
+			{
+				CurrentTime += CurrentMoveDelta;
+				if (CurrentTime > SequenceLength)
+				{
+					CurrentTime -= SequenceLength;
+				}
+				NextMarker.TimeToMarker = TimeToMarker - CurrentMoveDelta;
+				PrevMarker.TimeToMarker -= CurrentMoveDelta;
+				break;
+			}
+		}
+	}
+	else
+	{
+		while (true)
+		{
+			if (PrevMarker.MarkerIndex == -1)
+			{
+				float PrevCurrentTime = CurrentTime;
+				CurrentTime = FMath::Max(CurrentTime + CurrentMoveDelta, 0.f);
+				PrevMarker.TimeToMarker = CurrentTime;
+				NextMarker.TimeToMarker -= CurrentTime - PrevCurrentTime; //Add how far we moved to distance from previous marker
+				break;
+			}
+			const FAnimSyncMarker& PrevSyncMarker = AuthoredSyncMarkers[PrevMarker.MarkerIndex];
+			checkSlow(ValidMarkerNames.Contains(PrevSyncMarker.MarkerName));
+			if (!bOffsetInitialized)
+			{
+				bOffsetInitialized = true;
+				if (PrevSyncMarker.Time > CurrentTime)
+				{
+					MarkerTimeOffset = -SequenceLength;
+				}
+			}
+			const float PrevMarkerTime = PrevSyncMarker.Time + MarkerTimeOffset;
+			const float TimeToMarker = PrevMarkerTime - CurrentTime;
+
+			if (CurrentMoveDelta < TimeToMarker)
+			{
+				CurrentTime = PrevSyncMarker.Time;
+				CurrentMoveDelta -= TimeToMarker;
+
+				NextMarker.MarkerIndex = PrevMarker.MarkerIndex;
+				NextMarker.TimeToMarker = -CurrentMoveDelta;
+
+				int32 PassedMarker = MarkersPassed.Add(FPassedMarker());
+				MarkersPassed[PassedMarker].PassedMarkerName = PrevSyncMarker.MarkerName;
+				MarkersPassed[PassedMarker].DeltaTimeWhenPassed = CurrentMoveDelta;
+
+				do
+				{
+					--PrevMarker.MarkerIndex;
+					if (PrevMarker.MarkerIndex < 0)
+					{
+						if (!bLooping)
+						{
+							PrevMarker.MarkerIndex = -1;
+							break;
+						}
+						PrevMarker.MarkerIndex = AuthoredSyncMarkers.Num() - 1;
+						MarkerTimeOffset -= SequenceLength;
+					}
+				} while (!ValidMarkerNames.Contains(AuthoredSyncMarkers[PrevMarker.MarkerIndex].MarkerName));
+			}
+			else
+			{
+				CurrentTime += CurrentMoveDelta;
+				if (CurrentTime < 0.f)
+				{
+					CurrentTime += SequenceLength;
+				}
+				PrevMarker.TimeToMarker = TimeToMarker - CurrentMoveDelta;
+				NextMarker.TimeToMarker -= CurrentMoveDelta;
+				break;
+			}
+		}
+	}
+}
+
+void AdvanceMarkerForwards(int32& Marker, FName MarkerToFind, bool bLooping, const TArray<FAnimSyncMarker>& AuthoredSyncMarkers)
+{
+	while (AuthoredSyncMarkers[Marker].MarkerName != MarkerToFind)
+	{
+		++Marker;
+		if (Marker == AuthoredSyncMarkers.Num() && !bLooping)
+		{
+			Marker = -1;
+			break;
+		}
+		Marker %= AuthoredSyncMarkers.Num();
+	}
+}
+
+int32 MarkerCounterSpaceTransform(int32 MaxMarker, int32 Source)
+{
+	return MaxMarker - 1 - Source;
+}
+
+void AdvanceMarkerBackwards(int32& Marker, FName MarkerToFind, bool bLooping, const TArray<FAnimSyncMarker>& AuthoredSyncMarkers)
+{
+	const int32 MarkerMax = AuthoredSyncMarkers.Num();
+	int32 Counter = MarkerCounterSpaceTransform(MarkerMax, Marker);
+	while (AuthoredSyncMarkers[Marker].MarkerName != MarkerToFind)
+	{
+		if (Marker == 0 && !bLooping)
+		{
+			Marker = -1;
+			break;
+		}
+		Counter = ++Counter % MarkerMax;
+		Marker = MarkerCounterSpaceTransform(MarkerMax, Counter);
+	}
+}
+
+void UAnimSequence::AdvanceMarkerPhaseAsFollower(const FMarkerTickContext& Context, float DeltaRemaining, bool bLooping, float& CurrentTime, FMarkerPair& PreviousMarker, FMarkerPair& NextMarker) const
+{
+	const bool bPlayingForwards = DeltaRemaining > 0.f;
+
+	if (bPlayingForwards)
+	{
+		int32 PassedMarkersIndex = 0;
+		do
+		{
+			if (NextMarker.MarkerIndex == -1)
+			{
+				check(!bLooping); // shouldnt have an end of anim marker if looping
+				CurrentTime = FMath::Min(CurrentTime + DeltaRemaining, SequenceLength);
+				break;
+			}
+			else if (PassedMarkersIndex < Context.MarkersPassedThisTick.Num())
+			{
+				PreviousMarker.MarkerIndex = NextMarker.MarkerIndex;
+
+				const FPassedMarker& PassedMarker = Context.MarkersPassedThisTick[PassedMarkersIndex];
+				AdvanceMarkerForwards(NextMarker.MarkerIndex, PassedMarker.PassedMarkerName, bLooping, AuthoredSyncMarkers);
+				if (NextMarker.MarkerIndex == -1)
+				{
+					DeltaRemaining = PassedMarker.DeltaTimeWhenPassed;
+				}
+				++PassedMarkersIndex;
+			}
+		} while (PassedMarkersIndex < Context.MarkersPassedThisTick.Num());
+
+		const FMarkerSyncAnimPosition& End = Context.GetMarkerSyncEndPosition();
+		
+		if (Context.MarkersPassedThisTick.Num() > 0)
+		{
+			AdvanceMarkerForwards(NextMarker.MarkerIndex, End.NextMarkerName, bLooping, AuthoredSyncMarkers);
+		}
+
+		//Validation
+		if (End.NextMarkerName == NAME_None)
+		{
+			check(NextMarker.MarkerIndex == -1);
+		}
+		else
+		{
+			check(AuthoredSyncMarkers[NextMarker.MarkerIndex].MarkerName == End.NextMarkerName);
+		}
+		//End Validation
+		CurrentTime = GetCurrentTimeFromMarkers(PreviousMarker, NextMarker, End.PositionBetweenMarkers);
+	}
+	else
+	{
+		int32 PassedMarkersIndex = 0;
+		do
+		{
+			if (PreviousMarker.MarkerIndex == -1)
+			{
+				check(!bLooping); // shouldnt have an end of anim marker if looping
+				CurrentTime = FMath::Max(CurrentTime + DeltaRemaining, 0.f);
+				break;
+			}
+			else if (PassedMarkersIndex < Context.MarkersPassedThisTick.Num())
+			{
+				NextMarker.MarkerIndex = PreviousMarker.MarkerIndex;
+
+				const FPassedMarker& PassedMarker = Context.MarkersPassedThisTick[PassedMarkersIndex];
+				AdvanceMarkerBackwards(PreviousMarker.MarkerIndex, PassedMarker.PassedMarkerName, bLooping, AuthoredSyncMarkers);
+				if (PreviousMarker.MarkerIndex == -1)
+				{
+					DeltaRemaining = PassedMarker.DeltaTimeWhenPassed;
+				}
+				++PassedMarkersIndex;
+			}
+		} while (PassedMarkersIndex < Context.MarkersPassedThisTick.Num());
+
+		const FMarkerSyncAnimPosition& End = Context.GetMarkerSyncEndPosition();
+
+		if (Context.MarkersPassedThisTick.Num() > 0)
+		{
+			AdvanceMarkerBackwards(PreviousMarker.MarkerIndex, End.PreviousMarkerName, bLooping, AuthoredSyncMarkers);
+		}
+
+		//Validation
+		if (End.PreviousMarkerName == NAME_None)
+		{
+			check(PreviousMarker.MarkerIndex == -1);
+		}
+		else
+		{
+			check(AuthoredSyncMarkers[PreviousMarker.MarkerIndex].MarkerName == End.PreviousMarkerName);
+		}
+		//End Validation
+		CurrentTime = GetCurrentTimeFromMarkers(PreviousMarker, NextMarker, End.PositionBetweenMarkers);
+	}
+}
+
+void UAnimSequence::GetMarkerIndicesForTime(float CurrentTime, bool bLooping, const TArray<FName>& ValidMarkerNames, FMarkerPair& OutPrevMarker, FMarkerPair& OutNextMarker) const
+{
+	const int LoopModStart = bLooping ? -1 : 0;
+	const int LoopModEnd = bLooping ? 2 : 1;
+
+	OutPrevMarker.MarkerIndex = -1;
+	OutPrevMarker.TimeToMarker = -CurrentTime;
+	OutNextMarker.MarkerIndex = -1;
+	OutNextMarker.TimeToMarker = SequenceLength - CurrentTime;
+
+	for (int32 LoopMod = LoopModStart; LoopMod < LoopModEnd; ++LoopMod)
+	{
+		const float LoopModTime = LoopMod * SequenceLength;
+		for (int Idx = 0; Idx < AuthoredSyncMarkers.Num(); ++Idx)
+		{
+			const FAnimSyncMarker& Marker = AuthoredSyncMarkers[Idx];
+			if (ValidMarkerNames.Contains(Marker.MarkerName))
+			{
+				const float MarkerTime = Marker.Time + LoopModTime;
+				if (MarkerTime < CurrentTime)
+				{
+					OutPrevMarker.MarkerIndex = Idx;
+					OutPrevMarker.TimeToMarker = MarkerTime - CurrentTime;
+				}
+				else if (MarkerTime >= CurrentTime)
+				{
+					OutNextMarker.MarkerIndex = Idx;
+					OutNextMarker.TimeToMarker = MarkerTime - CurrentTime;
+					break; // Done
+				}
+			}
+		}
+		if (OutNextMarker.MarkerIndex != -1)
+		{
+			break; // Done
+		}
+	}
+}
+
+FMarkerSyncAnimPosition UAnimSequence::GetMarkerSyncPositionfromMarkerIndicies(int32 PrevMarker, int32 NextMarker, float CurrentTime) const
+{
+	FMarkerSyncAnimPosition SyncPosition;
+	float PrevTime, NextTime;
+	
+	if (PrevMarker != -1)
+	{
+		PrevTime = AuthoredSyncMarkers[PrevMarker].Time;
+		SyncPosition.PreviousMarkerName = AuthoredSyncMarkers[PrevMarker].MarkerName;
+	}
+	else
+	{
+		PrevTime = 0.f;
+	}
+
+	if (NextMarker != -1)
+	{
+		NextTime = AuthoredSyncMarkers[NextMarker].Time;
+		SyncPosition.NextMarkerName = AuthoredSyncMarkers[NextMarker].MarkerName;
+	}
+	else
+	{
+		NextTime = SequenceLength;
+	}
+
+	// Account for looping
+	PrevTime = (PrevTime > CurrentTime) ? PrevTime - SequenceLength : PrevTime;
+	NextTime = (NextTime < CurrentTime) ? NextTime + SequenceLength : NextTime;
+
+	if (PrevTime == NextTime)
+	{
+		PrevTime -= SequenceLength;
+	}
+
+	check(NextTime > PrevTime);
+
+	SyncPosition.PositionBetweenMarkers = (CurrentTime - PrevTime) / (NextTime - PrevTime);
+	return SyncPosition;
+}
+
+float UAnimSequence::GetCurrentTimeFromMarkers(FMarkerPair& PrevMarker, FMarkerPair& NextMarker, float PositionBetweenMarkers) const
+{
+	float PrevTime = (PrevMarker.MarkerIndex != -1) ? AuthoredSyncMarkers[PrevMarker.MarkerIndex].Time : 0.f;
+	float NextTime = (NextMarker.MarkerIndex != -1) ? AuthoredSyncMarkers[NextMarker.MarkerIndex].Time : SequenceLength;
+
+	if (PrevTime >= NextTime)
+	{
+		PrevTime -= SequenceLength; //Account for looping
+	}
+	float CurrentTime = PrevTime + PositionBetweenMarkers * (NextTime - PrevTime);
+	if (CurrentTime < 0.f)
+	{
+		CurrentTime += SequenceLength;
+	}
+	PrevMarker.TimeToMarker = PrevTime - CurrentTime;
+	NextMarker.TimeToMarker = NextTime - CurrentTime;
+	return CurrentTime;
+}
+
+void UAnimSequence::GetMarkerIndicesForPosition(const FMarkerSyncAnimPosition& SyncPosition, bool bLooping, FMarkerPair& OutPrevMarker, FMarkerPair& OutNextMarker, float& OutCurrentTime) const
+{
+	if (SyncPosition.PreviousMarkerName == NAME_None)
+	{
+		OutPrevMarker.MarkerIndex = -1;
+		check(SyncPosition.NextMarkerName != NAME_None);
+
+		for (int32 Idx = 0; Idx < AuthoredSyncMarkers.Num(); ++Idx)
+		{
+			const FAnimSyncMarker& Marker = AuthoredSyncMarkers[Idx];
+			if (Marker.MarkerName == SyncPosition.NextMarkerName)
+			{
+				OutNextMarker.MarkerIndex = Idx;
+				OutCurrentTime = GetCurrentTimeFromMarkers(OutPrevMarker, OutNextMarker, SyncPosition.PositionBetweenMarkers);
+				return;
+			}
+		}
+		check(false); // Should have found a marker above!
+	}
+
+	if (SyncPosition.NextMarkerName == NAME_None)
+	{
+		OutNextMarker.MarkerIndex = -1;
+		check(SyncPosition.PreviousMarkerName != NAME_None);
+
+		for (int32 Idx = AuthoredSyncMarkers.Num() - 1; Idx >= 0; --Idx)
+		{
+			const FAnimSyncMarker& Marker = AuthoredSyncMarkers[Idx];
+			if (Marker.MarkerName == SyncPosition.PreviousMarkerName)
+			{
+				OutPrevMarker.MarkerIndex = Idx;
+				OutCurrentTime = GetCurrentTimeFromMarkers(OutPrevMarker, OutNextMarker, SyncPosition.PositionBetweenMarkers);
+				return;
+			}
+		}
+		check(false); // Should have found a marker above!
+	}
+
+	float DiffToCurrentTime = FLT_MAX;
+
+	for (int32 PrevMarkerIdx = 0; PrevMarkerIdx < AuthoredSyncMarkers.Num(); ++PrevMarkerIdx)
+	{
+		const FAnimSyncMarker& PrevMarker = AuthoredSyncMarkers[PrevMarkerIdx];
+		if (PrevMarker.MarkerName == SyncPosition.PreviousMarkerName)
+		{
+			int32 EndCount = bLooping ? AuthoredSyncMarkers.Num() : (AuthoredSyncMarkers.Num() - (PrevMarkerIdx + 1));
+			for (int32 NextMarkerCount = 0; NextMarkerCount < EndCount; ++NextMarkerCount)
+			{
+				int32 NextMarkerIdx = (PrevMarkerIdx + 1 + NextMarkerCount) % AuthoredSyncMarkers.Num();
+				if (AuthoredSyncMarkers[NextMarkerIdx].MarkerName == SyncPosition.NextMarkerName)
+				{
+					float NextMarkerTime = AuthoredSyncMarkers[NextMarkerIdx].Time;
+					if (NextMarkerTime < PrevMarker.Time)
+					{
+						NextMarkerTime += SequenceLength;
+					}
+					float ThisCurrentTime = PrevMarker.Time + SyncPosition.PositionBetweenMarkers * (NextMarkerTime - PrevMarker.Time);
+					if (ThisCurrentTime > SequenceLength)
+					{
+						ThisCurrentTime -= SequenceLength;
+					}
+					float ThisDiff = FMath::Abs(ThisCurrentTime - OutCurrentTime);
+					if (ThisDiff < DiffToCurrentTime)
+					{
+						DiffToCurrentTime = ThisDiff;
+						OutPrevMarker.MarkerIndex = PrevMarkerIdx;
+						OutNextMarker.MarkerIndex = NextMarkerIdx;
+						OutCurrentTime = GetCurrentTimeFromMarkers(OutPrevMarker, OutNextMarker, SyncPosition.PositionBetweenMarkers);
+						break;
+					}
+				}
+			}
+		}
+	}
+}
+
 #endif
 /*-----------------------------------------------------------------------------
 	AnimNotify& subclasses
diff --git a/Engine/Source/Runtime/Engine/Private/Animation/AnimSequenceBase.cpp b/Engine/Source/Runtime/Engine/Private/Animation/AnimSequenceBase.cpp
--- a/Engine/Source/Runtime/Engine/Private/Animation/AnimSequenceBase.cpp
+++ b/Engine/Source/Runtime/Engine/Private/Animation/AnimSequenceBase.cpp
@@ -158,8 +158,8 @@ void UAnimSequenceBase::PostLoad()
 
 #if WITH_EDITOR
 	InitializeNotifyTrack();
-	UpdateAnimNotifyTrackCache();
 #endif
+	RefreshCacheData();
 
 	if(USkeleton* Skeleton = GetSkeleton())
 	{
@@ -289,10 +289,10 @@ void UAnimSequenceBase::GetAnimNotifiesFromDeltaPositions(const float& PreviousP
 	}
 }
 
-void UAnimSequenceBase::TickAssetPlayerInstance(const FAnimTickRecord& Instance, class UAnimInstance* InstanceOwner, FAnimAssetTickContext& Context) const
+void UAnimSequenceBase::TickAssetPlayerInstance(FAnimTickRecord& Instance, class UAnimInstance* InstanceOwner, FAnimAssetTickContext& Context) const
 {
 	float& CurrentTime = *(Instance.TimeAccumulator);
-	const float PreviousTime = CurrentTime;
+	float PreviousTime = CurrentTime;
 	const float PlayRate = Instance.PlayRateMultiplier * this->RateScale;
 
 	float MoveDelta = 0.f;
@@ -302,18 +302,35 @@ void UAnimSequenceBase::TickAssetPlayerInstance(const FAnimTickRecord& Instance,
 		const float DeltaTime = Context.GetDeltaTime();
 		MoveDelta = PlayRate * DeltaTime;
 
-		if( MoveDelta != 0.f )
+		Context.SetLeaderDelta(MoveDelta);
+		if (MoveDelta != 0.f)
 		{
-			// Advance time
-			FAnimationRuntime::AdvanceTime(Instance.bLooping, MoveDelta, CurrentTime, SequenceLength);
+			if (Instance.bCanUseMarkerSync && Context.CanUseMarkerPosition())
+			{
+				TickByMarkerAsLeader(Instance.MarkerTickRecord, Context.MarkerTickContext, CurrentTime, PreviousTime, MoveDelta, Instance.bLooping);
+			}
+			else
+			{
+				// Advance time
+				FAnimationRuntime::AdvanceTime(Instance.bLooping, MoveDelta, CurrentTime, SequenceLength);
+			}
 		}
 
-		Context.SetSyncPoint(CurrentTime / SequenceLength);
+		Context.SetAnimationPositionRatio(CurrentTime / SequenceLength);
 	}
 	else
 	{
 		// Follow the leader
-		CurrentTime = Context.GetSyncPoint() * SequenceLength;
+		if (Instance.bCanUseMarkerSync && Context.CanUseMarkerPosition() && Context.MarkerTickContext.IsMarkerSyncStartValid())
+		{
+			TickByMarkerAsFollower(Instance.MarkerTickRecord, Context.MarkerTickContext, CurrentTime, PreviousTime, Context.GetLeaderDelta(), Instance.bLooping);
+		}
+		else
+		{
+			CurrentTime = Context.GetAnimationPositionRatio() * SequenceLength;
+		}
+
+
 		//@TODO: NOTIFIES: Calculate AdvanceType based on what the new delta time is
 
 		if( CurrentTime != PreviousTime )
@@ -329,13 +346,39 @@ void UAnimSequenceBase::TickAssetPlayerInstance(const FAnimTickRecord& Instance,
 	}
 
 	OnAssetPlayerTickedInternal(Context, PreviousTime, MoveDelta, Instance, InstanceOwner);
+}
+
+void UAnimSequenceBase::TickByMarkerAsFollower(FMarkerTickRecord &Instance, FMarkerTickContext &MarkerContext, float& CurrentTime, float& OutPreviousTime, const float MoveDelta, const bool bLooping) const
+{
+	if (!Instance.IsValid())
+	{
+		GetMarkerIndicesForPosition(MarkerContext.GetMarkerSyncStartPosition(), bLooping, Instance.PreviousMarker, Instance.NextMarker, CurrentTime);
+	}
+
+	OutPreviousTime = CurrentTime;
 
+	AdvanceMarkerPhaseAsFollower(MarkerContext, MoveDelta, bLooping, CurrentTime, Instance.PreviousMarker, Instance.NextMarker);
 }
 
-#if WITH_EDITOR
+void UAnimSequenceBase::TickByMarkerAsLeader(FMarkerTickRecord& Instance, FMarkerTickContext& MarkerContext, float& CurrentTime, float& OutPreviousTime, const float MoveDelta, const bool bLooping) const
+{
+	if (!Instance.IsValid())
+	{
+		GetMarkerIndicesForTime(CurrentTime, bLooping, MarkerContext.GetValidMarkerNames(), Instance.PreviousMarker, Instance.NextMarker);
+	}
+
+	MarkerContext.SetMarkerSyncStartPosition(GetMarkerSyncPositionfromMarkerIndicies(Instance.PreviousMarker.MarkerIndex, Instance.NextMarker.MarkerIndex, CurrentTime));
+
+	OutPreviousTime = CurrentTime;
+
+	AdvanceMarkerPhaseAsLeader(bLooping, MoveDelta, MarkerContext.GetValidMarkerNames(), CurrentTime, Instance.PreviousMarker, Instance.NextMarker, MarkerContext.MarkersPassedThisTick);
+
+	MarkerContext.SetMarkerSyncEndPosition(GetMarkerSyncPositionfromMarkerIndicies(Instance.PreviousMarker.MarkerIndex, Instance.NextMarker.MarkerIndex, CurrentTime));
+}
 
-void UAnimSequenceBase::UpdateAnimNotifyTrackCache()
+void UAnimSequenceBase::RefreshCacheData()
 {
+#if WITH_EDITOR
 	SortNotifies();
 
 	for (int32 TrackIndex=0; TrackIndex<AnimNotifyTracks.Num(); ++TrackIndex)
@@ -361,8 +404,10 @@ void UAnimSequenceBase::UpdateAnimNotifyTrackCache()
 
 	// notification broadcast
 	OnNotifyChanged.Broadcast();
+#endif //WITH_EDITOR
 }
 
+#if WITH_EDITOR
 void UAnimSequenceBase::InitializeNotifyTrack()
 {
 	if ( AnimNotifyTracks.Num() == 0 ) 
@@ -466,29 +511,34 @@ uint8* UAnimSequenceBase::FindNotifyPropertyData(int32 NotifyIndex, UArrayProper
 
 	if(Notifies.IsValidIndex(NotifyIndex))
 	{
-		// find Notifies property start point
-		UProperty* Property = FindField<UProperty>(GetClass(), TEXT("Notifies"));
+		return FindArrayProperty(TEXT("Notifies"), ArrayProperty, NotifyIndex);
+	}
+	return NULL;
+}
 
-		// found it and if it is array
-		if(Property && Property->IsA(UArrayProperty::StaticClass()))
-		{
-			// find Property Value from UObject we got
-			uint8* PropertyValue = Property->ContainerPtrToValuePtr<uint8>(this);
+uint8* UAnimSequenceBase::FindArrayProperty(const TCHAR* PropName, UArrayProperty*& ArrayProperty, int32 ArrayIndex)
+{
+	// find Notifies property start point
+	UProperty* Property = FindField<UProperty>(GetClass(), PropName);
 
-			// it is array, so now get ArrayHelper and find the raw ptr of the data
-			ArrayProperty = CastChecked<UArrayProperty>(Property);
-			FScriptArrayHelper ArrayHelper(ArrayProperty, PropertyValue);
+	// found it and if it is array
+	if (Property && Property->IsA(UArrayProperty::StaticClass()))
+	{
+		// find Property Value from UObject we got
+		uint8* PropertyValue = Property->ContainerPtrToValuePtr<uint8>(this);
 
-			if(ArrayProperty->Inner && NotifyIndex < ArrayHelper.Num())
-			{
-				//Get property data based on selected index
-				return ArrayHelper.GetRawPtr(NotifyIndex);
-			}
+		// it is array, so now get ArrayHelper and find the raw ptr of the data
+		ArrayProperty = CastChecked<UArrayProperty>(Property);
+		FScriptArrayHelper ArrayHelper(ArrayProperty, PropertyValue);
+
+		if (ArrayProperty->Inner && ArrayIndex < ArrayHelper.Num())
+		{
+			//Get property data based on selected index
+			return ArrayHelper.GetRawPtr(ArrayIndex);
 		}
 	}
 	return NULL;
 }
-
 #endif	//WITH_EDITOR
 
 
@@ -517,7 +567,7 @@ void UAnimSequenceBase::Serialize(FArchive& Ar)
 	{
 		if(USkeleton* Skeleton = GetSkeleton())
 		{
-			// we don't add track curve container unless it has been editied. 
+			// we don't add track curve container unless it has been edited. 
 			if ( RawCurveData.TransformCurves.Num() > 0 )
 			{
 				FSmartNameMapping* Mapping = GetSkeleton()->SmartNames.GetContainer(USkeleton::AnimTrackCurveMappingName);
diff --git a/Engine/Source/Runtime/Engine/Private/Animation/AnimSingleNodeInstance.cpp b/Engine/Source/Runtime/Engine/Private/Animation/AnimSingleNodeInstance.cpp
--- a/Engine/Source/Runtime/Engine/Private/Animation/AnimSingleNodeInstance.cpp
+++ b/Engine/Source/Runtime/Engine/Private/Animation/AnimSingleNodeInstance.cpp
@@ -233,6 +233,8 @@ void UAnimSingleNodeInstance::NativeUpdateAnimation(float DeltaTimeX)
 		if (UBlendSpaceBase* BlendSpace = Cast<UBlendSpaceBase>(CurrentAsset))
 		{
 			FAnimTickRecord& TickRecord = CreateUninitializedTickRecord(INDEX_NONE, /*out*/ SyncGroup);
+			TickRecord.SourceNodeRef = this;
+			InitTickRecordFromLastFrame(INDEX_NONE, TickRecord);
 			MakeBlendSpaceTickRecord(TickRecord, BlendSpace, BlendSpaceInput, BlendSampleData, BlendFilter, bLooping, NewPlayRate, 1.f, /*inout*/ CurrentTime);
 #if WITH_EDITORONLY_DATA
 			PreviewBasePose = BlendSpace->PreviewBasePose;
@@ -241,6 +243,8 @@ void UAnimSingleNodeInstance::NativeUpdateAnimation(float DeltaTimeX)
 		else if (UAnimSequence* Sequence = Cast<UAnimSequence>(CurrentAsset))
 		{
 			FAnimTickRecord& TickRecord = CreateUninitializedTickRecord(INDEX_NONE, /*out*/ SyncGroup);
+			TickRecord.SourceNodeRef = this;
+			InitTickRecordFromLastFrame(INDEX_NONE, TickRecord);
 			MakeSequenceTickRecord(TickRecord, Sequence, bLooping, NewPlayRate, 1.f, /*inout*/ CurrentTime);
 			// if it's not looping, just set play to be false when reached to end
 			if (!bLooping)
@@ -255,6 +259,8 @@ void UAnimSingleNodeInstance::NativeUpdateAnimation(float DeltaTimeX)
 		else if(UAnimComposite* Composite = Cast<UAnimComposite>(CurrentAsset))
 		{
 			FAnimTickRecord& TickRecord = CreateUninitializedTickRecord(INDEX_NONE, /*out*/ SyncGroup);
+			TickRecord.SourceNodeRef = this;
+			InitTickRecordFromLastFrame(INDEX_NONE, TickRecord);
 			MakeSequenceTickRecord(TickRecord, Composite, bLooping, NewPlayRate, 1.f, /*inout*/ CurrentTime);
 			// if it's not looping, just set play to be false when reached to end
 			if (!bLooping)
diff --git a/Engine/Source/Runtime/Engine/Private/Animation/BlendSpaceBase.cpp b/Engine/Source/Runtime/Engine/Private/Animation/BlendSpaceBase.cpp
--- a/Engine/Source/Runtime/Engine/Private/Animation/BlendSpaceBase.cpp
+++ b/Engine/Source/Runtime/Engine/Private/Animation/BlendSpaceBase.cpp
@@ -9,6 +9,64 @@
 #include "Animation/AnimInstance.h"
 #include "Animation/BlendSpaceBase.h"
 
+struct FSyncPattern
+{
+	// The markers that make up this patter
+	TArray<FName> MarkerNames;
+
+	// Returns the index of the supplied name in the array of marker names
+	// Search starts at StartIndex
+	int32 IndexOf(FName Name, int32 StartIndex = 0) const
+	{
+		for (int Ind = StartIndex; Ind < MarkerNames.Num(); ++Ind)
+		{
+			if (MarkerNames[Ind] == Name)
+			{
+				return Ind;
+			}
+		}
+		return INDEX_NONE;
+	}
+
+	// Tests the supplied pattern against ours, starting at the supplied start index
+	bool DoOneMatch(const TArray<FName>& TestMarkerNames, int32 StartIndex)
+	{
+		int32 MyMarkerIndex = StartIndex;
+		int32 TestMarkerIndex = 0;
+
+		do
+		{
+			if (MarkerNames[MyMarkerIndex] != TestMarkerNames[TestMarkerIndex])
+			{
+				return false;
+			}
+			MyMarkerIndex = (MyMarkerIndex + 1) % MarkerNames.Num();
+			TestMarkerIndex = (TestMarkerIndex + 1) % TestMarkerNames.Num();
+		} while (MyMarkerIndex != StartIndex || TestMarkerIndex != 0); // Did we get back to the start without failing
+		return true;
+	}
+
+	// Testes the supplied pattern against ourselves. Is not a straight forward array match
+	// (for example a,b,c,a would match b,c,a,a)
+	bool DoesPatternMatch(const TArray<FName>& TestMarkerNames)
+	{
+		check(TestMarkerNames.Num() > 0 && MarkerNames.Num() > 0);
+
+		FName StartMarker = TestMarkerNames[0];
+
+		int32 StartIndex = IndexOf(StartMarker);
+		while (StartIndex != INDEX_NONE)
+		{
+			if (DoOneMatch(TestMarkerNames, StartIndex))
+			{
+				return true;
+			}
+			StartIndex = IndexOf(StartMarker, StartIndex + 1);
+		}
+		return false;
+	}
+};
+
 // global variable that's used in editor when you change the property only
 // I can't make this member since I  need to change this in const function
 bool bNeedReinitializeFilter = false;
@@ -16,7 +74,7 @@ bool bNeedReinitializeFilter = false;
 UBlendSpaceBase::UBlendSpaceBase(const FObjectInitializer& ObjectInitializer)
 	: Super(ObjectInitializer)
 {
-	
+	bAllSequencesHaveMatchingMarkers = false;
 }
 
 void UBlendSpaceBase::PostLoad()
@@ -74,7 +132,23 @@ FVector UBlendSpaceBase::FilterInput(FBlendFilter * Filter, const FVector& Blend
 	return FilteredBlendInput;
 }
 
-void UBlendSpaceBase::TickAssetPlayerInstance(const FAnimTickRecord& Instance, class UAnimInstance* InstanceOwner, FAnimAssetTickContext& Context) const
+int32 GetHighestWeightSample(TArray<FBlendSampleData> &SampleDataList)
+{
+	int32 HighestWeightIndex = 0;
+	float HighestWeight = SampleDataList[HighestWeightIndex].GetWeight();
+	for (int32 I = 1; I < SampleDataList.Num(); I++)
+	{
+		if (SampleDataList[I].GetWeight() > HighestWeight)
+		{
+			HighestWeightIndex = I;
+			HighestWeight = SampleDataList[I].GetWeight();
+		}
+	}
+	return HighestWeightIndex;
+}
+
+/////////////////////////////////////////////////////
+void UBlendSpaceBase::TickAssetPlayerInstance(FAnimTickRecord& Instance, class UAnimInstance* InstanceOwner, FAnimAssetTickContext& Context) const
 {
 	const float DeltaTime = Context.GetDeltaTime();
 	float MoveDelta = Instance.PlayRateMultiplier * DeltaTime;
@@ -156,14 +230,15 @@ void UBlendSpaceBase::TickAssetPlayerInstance(const FAnimTickRecord& Instance, c
 		// get sample data from blendspace
 		if (GetSamplesFromBlendInput(BlendInput, NewSampleDataList))
 		{
-			float NewAnimLength=0;
+			float NewAnimLength=0.f;
+			float PreInterpAnimLength = 0.f;
 
 			// if target weight interpolation is set
 			if (TargetWeightInterpolationSpeedPerSec > 0.f)
 			{
 				UE_LOG(LogAnimation, Verbose, TEXT("Target Weight Interpolation: Target Samples "));
 				// recalculate AnimLength based on weight of target animations - this is used for scaling animation later (change speed)
-				float PreInterpAnimLength = GetAnimationLengthFromSampleData(NewSampleDataList);
+				PreInterpAnimLength = GetAnimationLengthFromSampleData(NewSampleDataList);
 				UE_LOG(LogAnimation, Verbose, TEXT("BlendSpace(%s) - BlendInput(%s) : PreAnimLength(%0.5f) "), *GetName(), *BlendInput.ToString(), PreInterpAnimLength);
 
 				// target weight interpolation
@@ -180,18 +255,36 @@ void UBlendSpaceBase::TickAssetPlayerInstance(const FAnimTickRecord& Instance, c
 
 				// recalculate AnimLength based on weight of animations
 				UE_LOG(LogAnimation, Verbose, TEXT("Target Weight Interpolation: Interp Samples "));
-				NewAnimLength = GetAnimationLengthFromSampleData(SampleDataList);
-				// now scale the animation
-				if (NewAnimLength > 0.f)
-				{
-					MoveDelta *= PreInterpAnimLength/NewAnimLength;
-				}
 			}
 			else
 			{
 				// when there is no target weight interpolation, just copy new to target
 				SampleDataList.Append(NewSampleDataList);
-				NewAnimLength = GetAnimationLengthFromSampleData(SampleDataList);
+			}
+
+			bool bCanDoMarkerSync = Instance.bCanUseMarkerSync && Context.CanUseMarkerPosition();
+
+			if (bCanDoMarkerSync && Instance.MarkerTickRecord.IsValid())
+			{
+				//Copy previous frame marker data to current frame
+				for (FBlendSampleData& PrevBlendSampleItem : OldSampleDataList)
+				{
+					for (FBlendSampleData& CurrentBlendSampleItem : SampleDataList)
+					{
+						if (PrevBlendSampleItem.SampleDataIndex == CurrentBlendSampleItem.SampleDataIndex)
+						{
+							CurrentBlendSampleItem.Time = PrevBlendSampleItem.Time;
+							CurrentBlendSampleItem.MarkerTickRecord = PrevBlendSampleItem.MarkerTickRecord;
+						}
+					}
+				}
+			}
+
+			NewAnimLength = GetAnimationLengthFromSampleData(SampleDataList);
+
+			if (PreInterpAnimLength > 0.f && NewAnimLength > 0.f)
+			{
+				MoveDelta *= PreInterpAnimLength / NewAnimLength;
 			}
 
 			float& NormalizedCurrentTime = *(Instance.TimeAccumulator);
@@ -201,44 +294,66 @@ void UBlendSpaceBase::TickAssetPlayerInstance(const FAnimTickRecord& Instance, c
 			{
 				// advance current time - blend spaces hold normalized time as when dealing with changing anim length it would be possible to go backwards
 				UE_LOG(LogAnimation, Verbose, TEXT("BlendSpace(%s) - BlendInput(%s) : AnimLength(%0.5f) "), *GetName(), *BlendInput.ToString(), NewAnimLength);
+				FPlatformMisc::LowLevelOutputDebugString(TEXT("\n"));
+				if (bCanDoMarkerSync)
+				{
+					const int32 HighestWeightIndex = GetHighestWeightSample(SampleDataList);
+					FBlendSampleData& SampleDataItem = SampleDataList[HighestWeightIndex];
+					const FBlendSample& Sample = SampleData[SampleDataItem.SampleDataIndex];
 
-				// Advance time using current/new anim length
-				float CurrentTime = NormalizedCurrentTime * NewAnimLength;
-				FAnimationRuntime::AdvanceTime(Instance.bLooping, MoveDelta, /*inout*/ CurrentTime, NewAnimLength);
-				NormalizedCurrentTime = NewAnimLength ? (CurrentTime / NewAnimLength) : 0.0f;
+					if (!Instance.MarkerTickRecord.IsValid())
+					{
+						SampleDataItem.Time = NormalizedCurrentTime * Sample.Animation->SequenceLength;
+					}
+					if (Context.GetDeltaTime() != 0.f)
+					{
+						Sample.Animation->TickByMarkerAsLeader(SampleDataItem.MarkerTickRecord, Context.MarkerTickContext, SampleDataItem.Time, SampleDataItem.PreviousTime, Context.GetDeltaTime(), true);
+						TickFollowerSamples(SampleDataList, HighestWeightIndex, Context);
+					}
+					NormalizedCurrentTime = SampleDataItem.Time / Sample.Animation->SequenceLength;
+					Instance.MarkerTickRecord = SampleDataItem.MarkerTickRecord;
+				}
+				else
+				{
+					// Advance time using current/new anim length
+					float CurrentTime = NormalizedCurrentTime * NewAnimLength;
+					FAnimationRuntime::AdvanceTime(Instance.bLooping, MoveDelta, /*inout*/ CurrentTime, NewAnimLength);
+					NormalizedCurrentTime = NewAnimLength ? (CurrentTime / NewAnimLength) : 0.0f;
+				}
 
-				Context.SetSyncPoint(NormalizedCurrentTime);
+				Context.SetAnimationPositionRatio(NormalizedCurrentTime);
 			}
 			else
 			{
-				NormalizedCurrentTime = Context.GetSyncPoint();
+				if(!Context.MarkerTickContext.IsMarkerSyncStartValid())
+				{
+					bCanDoMarkerSync = false;
+				}
+				if (bCanDoMarkerSync)
+				{
+					if (Context.GetDeltaTime() != 0.f)
+					{
+						TickFollowerSamples(SampleDataList, -1, Context);
+					}
+					const int32 HighestWeightIndex = GetHighestWeightSample(SampleDataList);
+					FBlendSampleData& SampleDataItem = SampleDataList[HighestWeightIndex];
+					Instance.MarkerTickRecord = SampleDataItem.MarkerTickRecord;
+				}
+				NormalizedCurrentTime = Context.GetAnimationPositionRatio();
 			}
 
 			// generate notifies and sets time
 			{
 				TArray<const FAnimNotifyEvent*> Notifies;
 
-				// now calculate time for each samples
 				const float ClampedNormalizedPreviousTime = FMath::Clamp<float>(NormalizedPreviousTime, 0.f, 1.f);
 				const float ClampedNormalizedCurrentTime = FMath::Clamp<float>(NormalizedCurrentTime, 0.f, 1.f);
 				const bool bGenerateNotifies = Context.ShouldGenerateNotifies() && (NormalizedCurrentTime != NormalizedPreviousTime) && NotifyTriggerMode != ENotifyTriggerMode::None;
-				int32 HighestWeightIndex = 0;
-
-				// Get the index of the highest weight, assuming that the first is the highest until we find otherwise
-				bool bTriggerNotifyHighestWeightedAnim = NotifyTriggerMode == ENotifyTriggerMode::HighestWeightedAnimation && SampleDataList.Num() > 0;
-				if(bGenerateNotifies && bTriggerNotifyHighestWeightedAnim)
-				{
-					float HighestWeight = SampleDataList[HighestWeightIndex].GetWeight();
-					for(int32 I = 1 ; I < SampleDataList.Num(); I++)
-					{
-						if(SampleDataList[I].GetWeight() > HighestWeight)
-						{
-							HighestWeightIndex = I;
-							HighestWeight = SampleDataList[I].GetWeight();
-						}
-					}
-				}
 				
+				// Get the index of the highest weight, assuming that the first is the highest until we find otherwise
+				const bool bTriggerNotifyHighestWeightedAnim = NotifyTriggerMode == ENotifyTriggerMode::HighestWeightedAnimation && SampleDataList.Num() > 0;
+				const int32 HighestWeightIndex = (bGenerateNotifies && bTriggerNotifyHighestWeightedAnim) ? GetHighestWeightSample(SampleDataList) : -1;
+
 				for (int32 I = 0; I < SampleDataList.Num(); ++I)
 				{
 					FBlendSampleData& SampleEntry = SampleDataList[I];
@@ -250,12 +365,20 @@ void UBlendSpaceBase::TickAssetPlayerInstance(const FAnimTickRecord& Instance, c
 						const FBlendSample& Sample = SampleData[SampleDataIndex];
 						if( Sample.Animation )
 						{
-							const float SampleNormalizedPreviousTime = Sample.Animation->RateScale >= 0.f ? ClampedNormalizedPreviousTime : 1.f - ClampedNormalizedPreviousTime;
-							const float SampleNormalizedCurrentTime = Sample.Animation->RateScale >= 0.f ? ClampedNormalizedCurrentTime : 1.f - ClampedNormalizedCurrentTime;
-
-							const float PrevSampleDataTime = SampleNormalizedPreviousTime * Sample.Animation->SequenceLength;
+							float PrevSampleDataTime;
 							float& CurrentSampleDataTime = SampleEntry.Time;
-							CurrentSampleDataTime = SampleNormalizedCurrentTime * Sample.Animation->SequenceLength;
+
+							if (!bCanDoMarkerSync) //Have already updated time if we are doing marker sync
+							{
+								const float SampleNormalizedPreviousTime = Sample.Animation->RateScale >= 0.f ? ClampedNormalizedPreviousTime : 1.f - ClampedNormalizedPreviousTime;
+								const float SampleNormalizedCurrentTime = Sample.Animation->RateScale >= 0.f ? ClampedNormalizedCurrentTime : 1.f - ClampedNormalizedCurrentTime;
+								PrevSampleDataTime = SampleNormalizedPreviousTime * Sample.Animation->SequenceLength;
+								CurrentSampleDataTime = SampleNormalizedCurrentTime * Sample.Animation->SequenceLength;
+							}
+							else
+							{
+								PrevSampleDataTime = SampleEntry.PreviousTime;
+							}
 
 							// Figure out delta time 
 							float DeltaTimePosition = CurrentSampleDataTime - PrevSampleDataTime;
@@ -676,14 +799,29 @@ bool UBlendSpaceBase::ValidateSampleInput(FBlendSample & BlendSample, int32 Orig
 	return true;
 }
 
+void PopulateMarkerNameArray(TArray<FName>& Pattern, TArray<FAnimSyncMarker>& AuthoredSyncMarkers)
+{
+	Pattern.Reserve(AuthoredSyncMarkers.Num());
+	for (FAnimSyncMarker& Marker : AuthoredSyncMarkers)
+	{
+		Pattern.Add(Marker.MarkerName);
+	}
+}
+
 void UBlendSpaceBase::ValidateSampleData()
 {
 	bool bSampleDataChanged=false;
 	AnimLength = 0.f;
 
+	bool bNoAnimationsHaveMarkers = true;
+	bool bAllMarkerPatternsMatch = true;
+	FSyncPattern BlendSpacePattern;
+
 	for (int32 I=0; I<SampleData.Num(); ++I)
 	{
-		if ( SampleData[I].Animation == 0 )
+		FBlendSample& Sample = SampleData[I];
+
+		if (Sample.Animation == nullptr)
 		{
 			SampleData.RemoveAt(I);
 			--I;
@@ -699,12 +837,12 @@ void UBlendSpaceBase::ValidateSampleData()
 		// otherwise, you have this grid area that doesn't have valid 
 		// sample points. Usually users will put it there
 		// if the value is around border, snap to border
-		SnapToBorder(SampleData[I]);
+		SnapToBorder(Sample);
 
 		// see if same data exists, by same, same values
 		for (int32 J=I+1; J<SampleData.Num(); ++J)
 		{
-			if ( IsSameSamplePoint(SampleData[I].SampleValue, SampleData[J].SampleValue) )
+			if (IsSameSamplePoint(Sample.SampleValue, SampleData[J].SampleValue))
 			{
 				SampleData.RemoveAt(J);
 				--J;
@@ -713,13 +851,33 @@ void UBlendSpaceBase::ValidateSampleData()
 			}
 		}
 
-		if (SampleData[I].Animation->SequenceLength > AnimLength)
+		if (Sample.Animation->SequenceLength > AnimLength)
 		{
 			// @todo : should apply scale? If so, we'll need to apply also when blend
-			AnimLength = SampleData[I].Animation->SequenceLength;
+			AnimLength = Sample.Animation->SequenceLength;
+		}
+
+		if (Sample.Animation->AuthoredSyncMarkers.Num() > 0)
+		{
+			bNoAnimationsHaveMarkers = false;
+			if (BlendSpacePattern.MarkerNames.Num() == 0)
+			{
+				PopulateMarkerNameArray(BlendSpacePattern.MarkerNames, Sample.Animation->AuthoredSyncMarkers);
+			}
+			else
+			{
+				TArray<FName> ThisPattern;
+				PopulateMarkerNameArray(ThisPattern, Sample.Animation->AuthoredSyncMarkers);
+				if (!BlendSpacePattern.DoesPatternMatch(ThisPattern))
+				{
+					bAllMarkerPatternsMatch = false;
+				}
+			}
 		}
 	}
 
+	bAllSequencesHaveMatchingMarkers = bAllMarkerPatternsMatch && !bNoAnimationsHaveMarkers;
+
 	if (bSampleDataChanged)
 	{
 		GridSamples.Empty();
@@ -1090,4 +1248,9 @@ void UBlendSpaceBase::PostEditChangeProperty( struct FPropertyChangedEvent& Prop
 		InitializePerBoneBlend();
 	}
 }
-#endif // WITH_EDITOR
\ No newline at end of file
+
+int32 UBlendSpaceBase::GetMarkerUpdateCounter() const
+{ 
+	return MarkerDataUpdateCounter; 
+}
+#endif // WITH_EDITOR
diff --git a/Engine/Source/Runtime/Engine/Private/AnimationAsset.cpp b/Engine/Source/Runtime/Engine/Private/AnimationAsset.cpp
--- a/Engine/Source/Runtime/Engine/Private/AnimationAsset.cpp
+++ b/Engine/Source/Runtime/Engine/Private/AnimationAsset.cpp
@@ -3,6 +3,89 @@
 #include "EnginePrivate.h"
 #include "Animation/AnimSequence.h"
 
+
+//////////////////////////////////////////////////////////////////////////
+// FAnimGroupInstance
+
+void FAnimGroupInstance::TestTickRecordForLeadership(EAnimGroupRole::Type MembershipType)
+{
+	int32 TestIndex = ActivePlayers.Num() - 1;
+	FAnimTickRecord& Candidate = ActivePlayers[TestIndex];
+
+	switch (MembershipType)
+	{
+	case EAnimGroupRole::CanBeLeader:
+		// Set it if we're better than the current leader (or if there is no leader yet)
+		if ((GroupLeaderIndex == INDEX_NONE) || (ActivePlayers[GroupLeaderIndex].EffectiveBlendWeight < Candidate.EffectiveBlendWeight))
+		{
+			// This is a better leader
+			GroupLeaderIndex = TestIndex;
+		}
+		break;
+	case EAnimGroupRole::AlwaysLeader:
+		// Always set the leader index
+		GroupLeaderIndex = TestIndex;
+		break;
+	default:
+	case EAnimGroupRole::AlwaysFollower:
+		// Never set the leader index; the actual tick code will handle the case of no leader by using the first element in the array
+		break;
+	}
+}
+
+void FAnimGroupInstance::Finalize(const TArray<FName>& PreviousValidMarkers)
+{
+	GroupLeaderIndex = FMath::Max(GroupLeaderIndex, 0);
+
+	TArray<FName>* MarkerNames = ActivePlayers[GroupLeaderIndex].SourceAsset->GetUniqueMarkerNames();
+	if (MarkerNames)
+	{
+		// Group leader has markers, off to a good start
+		ValidMarkers = *MarkerNames;
+		ActivePlayers[GroupLeaderIndex].bCanUseMarkerSync = true;
+		bCanUseMarkerSync = true;
+
+		//filter markers based on what exists in the other animations
+		for (int32 Idx = 0; Idx < ActivePlayers.Num(); ++Idx)
+		{
+			if (Idx != GroupLeaderIndex)
+			{
+				FAnimTickRecord& Candidate = ActivePlayers[Idx];
+				TArray<FName>* MarkerNames = Candidate.SourceAsset->GetUniqueMarkerNames();
+				if (MarkerNames) // Let anims with no markers set use length scaling sync
+				{
+					Candidate.bCanUseMarkerSync = true;
+					for (int Idx = ValidMarkers.Num() - 1; Idx >= 0; --Idx)
+					{
+						FName& MarkerName = ValidMarkers[Idx];
+						if (!MarkerNames->Contains(MarkerName))
+						{
+							ValidMarkers.RemoveAtSwap(Idx, 1, false);
+						}
+					}
+					if (ValidMarkers.Num() == 0) //No common markers between all anims, no marker syncing
+					{
+						bCanUseMarkerSync = false;
+						break;
+					}
+				}
+			}
+		}
+		ValidMarkers.Sort();
+		if (ValidMarkers != PreviousValidMarkers)
+		{
+			for (int32 Idx = 0; Idx < ActivePlayers.Num(); ++Idx)
+			{
+				ActivePlayers[Idx].MarkerTickRecord.PreviousMarker.MarkerIndex = MarkerIndexSpecialValues::Unitialized;
+				ActivePlayers[Idx].MarkerTickRecord.NextMarker.MarkerIndex = MarkerIndexSpecialValues::Unitialized;
+			}
+		}
+	}
+}
+
+//////////////////////////////////////////////////////////////////////////
+// UAnimationAsset
+
 UAnimationAsset::UAnimationAsset(const FObjectInitializer& ObjectInitializer)
 	: Super(ObjectInitializer)
 {
diff --git a/Engine/Source/Runtime/Engine/Public/Animation/AnimTypes.h b/Engine/Source/Runtime/Engine/Public/Animation/AnimTypes.h
--- a/Engine/Source/Runtime/Engine/Public/Animation/AnimTypes.h
+++ b/Engine/Source/Runtime/Engine/Public/Animation/AnimTypes.h
@@ -202,6 +202,29 @@ FORCEINLINE bool FAnimNotifyEvent::operator<(const FAnimNotifyEvent& Other) cons
 	}
 }
 
+USTRUCT()
+struct FAnimSyncMarker
+{
+	GENERATED_USTRUCT_BODY()
+
+	// The name of this marker
+	UPROPERTY()
+	FName MarkerName;
+
+	// Time in seconds of this marker
+	UPROPERTY()
+	float Time;
+
+#if WITH_EDITORONLY_DATA
+	// The editor track this marker sits on
+	UPROPERTY()
+	int32 TrackIndex;
+#endif
+
+	/** This can be used with the Sort() function on a TArray of FAnimSyncMarker to sort the notifies array by time, earliest first. */
+	ENGINE_API bool operator <(const FAnimSyncMarker& Other) const { return Time < Other.Time; }
+};
+
 /**
  * Keyframe position data for one track.  Pos(i) occurs at Time(i).  Pos.Num() always equals Time.Num().
  */
@@ -217,6 +240,8 @@ struct FAnimNotifyTrack
 	FLinearColor TrackColor;
 
 	TArray<FAnimNotifyEvent*> Notifies;
+	
+	TArray<FAnimSyncMarker*> SyncMarkers;
 
 	FAnimNotifyTrack()
 		: TrackName(TEXT(""))
