#ANIM: Added folder option on tooltip

Merging

//depot/UE4-Orion/Engine/Source/Editor/Persona/Private/SAnimationSequenceBrowser.cpp

to //depot/UE4/Engine/Source/Editor/Persona/Private/SAnimationSequenceBrowser.cpp

diff --git a/Engine/Source/Editor/Persona/Private/SAnimationSequenceBrowser.cpp b/Engine/Source/Editor/Persona/Private/SAnimationSequenceBrowser.cpp
--- a/Engine/Source/Editor/Persona/Private/SAnimationSequenceBrowser.cpp
+++ b/Engine/Source/Editor/Persona/Private/SAnimationSequenceBrowser.cpp
@@ -756,16 +756,10 @@ TSharedRef<SToolTip> SAnimationSequenceBrowser::CreateCustomAssetToolTip(FAssetD
 
 	// Add asset registry tags to a text list; except skeleton as that is implied in Persona
 	TSharedRef<SVerticalBox> DescriptionBox = SNew(SVerticalBox);
-	bool bDescriptionCreated = false;
 	for(TPair<FName, FString> TagPair : AssetData.TagsAndValues)
 	{
 		if(TagsToShow.Contains(TagPair.Key))
 		{
-			if(!bDescriptionCreated)
-			{
-				bDescriptionCreated = true;
-			}
-
 			DescriptionBox->AddSlot()
 			.AutoHeight()
 			.Padding(0,0,5,0)
@@ -790,6 +784,29 @@ TSharedRef<SToolTip> SAnimationSequenceBrowser::CreateCustomAssetToolTip(FAssetD
 		}
 	}
 
+	DescriptionBox->AddSlot()
+		.AutoHeight()
+		.Padding(0,0,5,0)
+		[
+			SNew(SHorizontalBox)
+			+ SHorizontalBox::Slot()
+			.AutoWidth()
+			[
+				SNew(STextBlock)
+				.Text(LOCTEXT("AssetBrowser_FolderPathLabel", "Folder :"))
+				.ColorAndOpacity(FSlateColor::UseSubduedForeground())
+			]
+
+			+ SHorizontalBox::Slot()
+			.AutoWidth()
+			[
+				SNew(STextBlock)
+				.Text(FText::FromName(AssetData.PackagePath))
+				.ColorAndOpacity(FSlateColor::UseForeground())
+				.WrapTextAt(300.f)
+			]
+		];
+
 	TSharedPtr<SHorizontalBox> ContentBox = nullptr;
 	TSharedRef<SToolTip> ToolTip = SNew(SToolTip)
 	.TextMargin(1)
@@ -849,20 +866,17 @@ TSharedRef<SToolTip> SAnimationSequenceBrowser::CreateCustomAssetToolTip(FAssetD
 		]
 	];
 
-	// If we have a description, add an extra section to the tooltip for it.
-	if(bDescriptionCreated)
-	{
-		ContentBox->AddSlot()
-		.Padding(4, 0, 0, 0)
+	// add an extra section to the tooltip for it.
+	ContentBox->AddSlot()
+	.Padding(4, 0, 0, 0)
+	[
+		SNew(SBorder)
+		.Padding(6)
+		.BorderImage(FEditorStyle::GetBrush("ContentBrowser.TileViewTooltip.ContentBorder"))
 		[
-			SNew(SBorder)
-			.Padding(6)
-			.BorderImage(FEditorStyle::GetBrush("ContentBrowser.TileViewTooltip.ContentBorder"))
-			[
-				DescriptionBox
-			]
-		];
-	}
+			DescriptionBox
+		]
+	];
 
 	return ToolTip;
 }
