[Merging using UE4-Orion-To-UE4] - Added Blend Profiles to USkeleton - Per-Bone blending scales for blend list nodes and state machine transitions.

diff --git a/Engine/Source/Editor/AnimGraph/Classes/AnimStateNodeBase.h b/Engine/Source/Editor/AnimGraph/Classes/AnimStateNodeBase.h
--- a/Engine/Source/Editor/AnimGraph/Classes/AnimStateNodeBase.h
+++ b/Engine/Source/Editor/AnimGraph/Classes/AnimStateNodeBase.h
@@ -32,6 +32,8 @@ class UAnimStateNodeBase : public UEdGraphNode
 	//
 	virtual UEdGraph* GetBoundGraph() const { return NULL; }
 
+	ANIMGRAPH_API UAnimBlueprint* GetAnimBlueprint() const;
+
 protected:
 	// Name used as a seed when pasting nodes
 	virtual FString GetDesiredNewNodeName() const { return TEXT("State"); }
diff --git a/Engine/Source/Editor/AnimGraph/Classes/AnimStateTransitionNode.h b/Engine/Source/Editor/AnimGraph/Classes/AnimStateTransitionNode.h
--- a/Engine/Source/Editor/AnimGraph/Classes/AnimStateTransitionNode.h
+++ b/Engine/Source/Editor/AnimGraph/Classes/AnimStateTransitionNode.h
@@ -40,6 +40,10 @@ class UAnimStateTransitionNode : public UAnimStateNodeBase
 	UPROPERTY(EditAnywhere, Category=Transition)
 	UCurveFloat* CustomBlendCurve;
 
+	// The blend profile to use to evaluate this transition per-bone
+	UPROPERTY(EditAnywhere, Category=Transition)
+	UBlendProfile* BlendProfile;
+
 	// Try setting the rule automatically based on the player node remaining time and the CrossfadeDuration, ignoring the internal time
 	UPROPERTY(EditAnywhere, Category=Transition)
 	bool bAutomaticRuleBasedOnSequencePlayerInState;
diff --git a/Engine/Source/Editor/AnimGraph/Private/AnimStateNodeBase.cpp b/Engine/Source/Editor/AnimGraph/Private/AnimStateNodeBase.cpp
--- a/Engine/Source/Editor/AnimGraph/Private/AnimStateNodeBase.cpp
+++ b/Engine/Source/Editor/AnimGraph/Private/AnimStateNodeBase.cpp
@@ -89,3 +89,9 @@ FString UAnimStateNodeBase::GetDocumentationLink() const
 	return TEXT("Shared/GraphNodes/AnimationStateMachine");
 }
 
+UAnimBlueprint* UAnimStateNodeBase::GetAnimBlueprint() const
+{
+	UBlueprint* Blueprint = FBlueprintEditorUtils::FindBlueprintForNode(this);
+	return CastChecked<UAnimBlueprint>(Blueprint);
+}
+
diff --git a/Engine/Source/Editor/DetailCustomizations/DetailCustomizations.Build.cs b/Engine/Source/Editor/DetailCustomizations/DetailCustomizations.Build.cs
--- a/Engine/Source/Editor/DetailCustomizations/DetailCustomizations.Build.cs
+++ b/Engine/Source/Editor/DetailCustomizations/DetailCustomizations.Build.cs
@@ -45,6 +45,7 @@ public class DetailCustomizations : ModuleRules
                 "AIModule", 
                 "MeshUtilities",
 				"ConfigEditor",
+                "Persona",
 			}
 		);
 
diff --git a/Engine/Source/Editor/DetailCustomizations/Private/AnimTransitionNodeDetails.cpp b/Engine/Source/Editor/DetailCustomizations/Private/AnimTransitionNodeDetails.cpp
--- a/Engine/Source/Editor/DetailCustomizations/Private/AnimTransitionNodeDetails.cpp
+++ b/Engine/Source/Editor/DetailCustomizations/Private/AnimTransitionNodeDetails.cpp
@@ -14,6 +14,9 @@
 #include "STextEntryPopup.h"
 #include "SExpandableArea.h"
 #include "BlueprintEditorUtils.h"
+#include "SBlendProfilePicker.h"
+#include "Animation/BlendProfile.h"
+
 
 #define LOCTEXT_NAMESPACE "FAnimStateNodeDetails"
 
@@ -45,7 +48,7 @@ void FAnimTransitionNodeDetails::CustomizeDetails( IDetailLayoutBuilder& DetailB
 		bTransitionToConduit = (NextState != NULL) && (NextState->IsA<UAnimStateConduitNode>());
 	}
 
-	////////////////////////////////////////
+	//////////////////////////////////////////////////////////////////////////
 
 	IDetailCategoryBuilder& TransitionCategory = DetailBuilder.EditCategory("Transition", LOCTEXT("TransitionCategoryTitle", "Transition") );
 
@@ -135,8 +138,7 @@ void FAnimTransitionNodeDetails::CustomizeDetails( IDetailLayoutBuilder& DetailB
 			];
 		}
 
-
-
+		//////////////////////////////////////////////////////////////////////////
 
 		IDetailCategoryBuilder& CrossfadeCategory = DetailBuilder.EditCategory("BlendSettings", LOCTEXT("BlendSettingsCategoryTitle", "BlendSettings") );
 		if (TransitionNode != NULL)
@@ -156,6 +158,30 @@ void FAnimTransitionNodeDetails::CustomizeDetails( IDetailLayoutBuilder& DetailB
 		CrossfadeCategory.AddProperty(GET_MEMBER_NAME_CHECKED(UAnimStateTransitionNode, BlendMode)).DisplayName( LOCTEXT("ModeLabel", "Mode") );
 		CrossfadeCategory.AddProperty(GET_MEMBER_NAME_CHECKED(UAnimStateTransitionNode, CustomBlendCurve)).DisplayName(LOCTEXT("CurveLabel", "Custom Blend Curve"));
 
+		USkeleton* TargetSkeleton = TransitionNode->GetAnimBlueprint()->TargetSkeleton;
+
+		if(TargetSkeleton)
+		{
+			TSharedPtr<IPropertyHandle> BlendProfileHandle = DetailBuilder.GetProperty(GET_MEMBER_NAME_CHECKED(UAnimStateTransitionNode, BlendProfile));
+			UObject* BlendProfilePropertyValue = nullptr;
+			BlendProfileHandle->GetValue(BlendProfilePropertyValue);
+			UBlendProfile* CurrentProfile = Cast<UBlendProfile>(BlendProfilePropertyValue);
+
+			CrossfadeCategory.AddProperty(BlendProfileHandle).CustomWidget(true)
+				.NameContent()
+				[
+					BlendProfileHandle->CreatePropertyNameWidget()
+				]
+				.ValueContent()
+				[
+					SNew(SBlendProfilePicker)
+					.TargetSkeleton(TargetSkeleton)
+					.AllowNew(false)
+					.OnBlendProfileSelected(this, &FAnimTransitionNodeDetails::OnBlendProfileChanged, BlendProfileHandle)
+					.InitialProfile(CurrentProfile)
+				];
+		}
+
 		// Add a button that is only visible when blend logic type is custom
 		CrossfadeCategory.AddCustomRow( LOCTEXT("EditBlendGraph", "Edit Blend Graph") )
 		[
@@ -173,9 +199,7 @@ void FAnimTransitionNodeDetails::CustomizeDetails( IDetailLayoutBuilder& DetailB
 			]
 		];
 
-
-
-
+		//////////////////////////////////////////////////////////////////////////
 
 		IDetailCategoryBuilder& NotificationCategory = DetailBuilder.EditCategory("Notifications", LOCTEXT("NotificationsCategoryTitle", "Notifications") );
 
@@ -448,4 +472,12 @@ TSharedRef<SWidget> FAnimTransitionNodeDetails::GetWidgetForInlineShareMenu(FStr
 		];
 }
 
+void FAnimTransitionNodeDetails::OnBlendProfileChanged(UBlendProfile* NewProfile, TSharedPtr<IPropertyHandle> ProfileProperty)
+{
+	if(ProfileProperty.IsValid())
+	{
+		ProfileProperty->SetValue((const UObject*&)NewProfile);
+	}
+}
+
 #undef LOCTEXT_NAMESPACE
\ No newline at end of file
diff --git a/Engine/Source/Editor/DetailCustomizations/Private/AnimTransitionNodeDetails.h b/Engine/Source/Editor/DetailCustomizations/Private/AnimTransitionNodeDetails.h
--- a/Engine/Source/Editor/DetailCustomizations/Private/AnimTransitionNodeDetails.h
+++ b/Engine/Source/Editor/DetailCustomizations/Private/AnimTransitionNodeDetails.h
@@ -32,6 +32,7 @@ protected:
 	void BecomeSharedWith(UAnimStateTransitionNode* NewNode, bool bShareRules);
 	void AssignUniqueColorsToAllSharedNodes(UEdGraph* CurrentGraph);
 
+	void OnBlendProfileChanged(UBlendProfile* NewProfile, TSharedPtr<IPropertyHandle> ProfileProperty);
 
 private:
 	TWeakObjectPtr<UAnimStateTransitionNode> TransitionNode;
diff --git a/Engine/Source/Editor/KismetCompiler/Private/AnimBlueprintCompiler.cpp b/Engine/Source/Editor/KismetCompiler/Private/AnimBlueprintCompiler.cpp
--- a/Engine/Source/Editor/KismetCompiler/Private/AnimBlueprintCompiler.cpp
+++ b/Engine/Source/Editor/KismetCompiler/Private/AnimBlueprintCompiler.cpp
@@ -813,6 +813,7 @@ void FAnimBlueprintCompiler::ProcessStateMachine(UAnimGraphNode_StateMachineBase
 			BakedTransition.InterruptNotify = FindOrAddNotify(TransitionNode->TransitionInterrupt);
 			BakedTransition.BlendMode = TransitionNode->BlendMode;
 			BakedTransition.CustomCurve = TransitionNode->CustomBlendCurve;
+			BakedTransition.BlendProfile = TransitionNode->BlendProfile;
 			BakedTransition.LogicType = TransitionNode->LogicType;
 
 			UAnimStateNodeBase* PreviousState = TransitionNode->GetPreviousState();
diff --git a/Engine/Source/Editor/Persona/Private/AnimGraphNodeDetails.cpp b/Engine/Source/Editor/Persona/Private/AnimGraphNodeDetails.cpp
--- a/Engine/Source/Editor/Persona/Private/AnimGraphNodeDetails.cpp
+++ b/Engine/Source/Editor/Persona/Private/AnimGraphNodeDetails.cpp
@@ -21,6 +21,8 @@
 #include "Persona.h"
 #include "BoneSelectionWidget.h"
 #include "SExpandableArea.h"
+#include "Animation/BlendProfile.h"
+#include "SBlendProfilePicker.h"
 
 #define LOCTEXT_NAMESPACE "KismetNodeWithOptionalPinsDetails"
 
@@ -65,7 +67,7 @@ void FAnimGraphNodeDetails::CustomizeDetails(class IDetailLayoutBuilder& DetailB
 		return;
 	}
 
-	USkeleton* TargetSkeleton = AnimGraphNode->GetAnimBlueprint()->TargetSkeleton;
+	TargetSkeleton = AnimGraphNode->GetAnimBlueprint()->TargetSkeleton;
 	TargetSkeletonName = FString::Printf(TEXT("%s'%s'"), *TargetSkeleton->GetClass()->GetName(), *TargetSkeleton->GetPathName());
 
 	// Get the node property
@@ -213,6 +215,21 @@ TSharedRef<SWidget> FAnimGraphNodeDetails::CreatePropertyWidget(UProperty* Targe
 				.AllowClear(bAllowClear)
 				.OnShouldFilterAsset(FOnShouldFilterAsset::CreateSP(this, &FAnimGraphNodeDetails::OnShouldFilterAnimAsset));
 		}
+		else if(ObjectProperty->PropertyClass->IsChildOf(UBlendProfile::StaticClass()) && TargetSkeleton)
+		{
+			TSharedPtr<IPropertyHandle> PropertyPtr(TargetPropertyHandle);
+
+			UObject* PropertyValue = nullptr;
+			TargetPropertyHandle->GetValue(PropertyValue);
+
+			UBlendProfile* CurrentProfile = Cast<UBlendProfile>(PropertyValue);
+
+			return SNew(SBlendProfilePicker)
+				.TargetSkeleton(this->TargetSkeleton)
+				.AllowNew(false)
+				.OnBlendProfileSelected(this, &FAnimGraphNodeDetails::OnBlendProfileChanged, PropertyPtr)
+				.InitialProfile(CurrentProfile);
+		}
 	}
 
 	return SNullWidget::NullWidget;
@@ -262,6 +279,14 @@ void FAnimGraphNodeDetails::AbortDisplayOfAllNodes(TArray< TWeakObjectPtr<UObjec
 	}
 }
 
+void FAnimGraphNodeDetails::OnBlendProfileChanged(UBlendProfile* NewProfile, TSharedPtr<IPropertyHandle> PropertyHandle)
+{
+	if(PropertyHandle.IsValid())
+	{
+		PropertyHandle->SetValue((const UObject*&)NewProfile);
+	}
+}
+
 
 TSharedRef<IPropertyTypeCustomization> FInputScaleBiasCustomization::MakeInstance() 
 {
diff --git a/Engine/Source/Editor/Persona/Private/AnimGraphNodeDetails.h b/Engine/Source/Editor/Persona/Private/AnimGraphNodeDetails.h
--- a/Engine/Source/Editor/Persona/Private/AnimGraphNodeDetails.h
+++ b/Engine/Source/Editor/Persona/Private/AnimGraphNodeDetails.h
@@ -29,6 +29,12 @@ protected:
 	/** Delegate to handle filtering of asset pickers */
 	bool OnShouldFilterAnimAsset( const FAssetData& AssetData ) const;
 
+	/** Called when a blend profile is selected */
+	void OnBlendProfileChanged(UBlendProfile* NewProfile, TSharedPtr<IPropertyHandle> PropertyHandle);
+
+	/** The skeleton we're operating on */
+	USkeleton* TargetSkeleton;
+
 	/** Path to the current blueprints skeleton to allow us to filter asset pickers */
 	FString TargetSkeletonName;
 };
diff --git a/Engine/Source/Editor/Persona/Private/Persona.cpp b/Engine/Source/Editor/Persona/Private/Persona.cpp
--- a/Engine/Source/Editor/Persona/Private/Persona.cpp
+++ b/Engine/Source/Editor/Persona/Private/Persona.cpp
@@ -3298,6 +3298,11 @@ FText FPersona::GetPreviewAssetTooltip() const
 	}
 }
 
+void FPersona::SetSelectedBlendProfile(UBlendProfile* InBlendProfile)
+{
+	OnBlendProfileSelected.Broadcast(InBlendProfile);
+}
+
 static class FMeshHierarchyCmd : private FSelfRegisteringExec
 {
 public:
diff --git a/Engine/Source/Editor/Persona/Private/Persona.h b/Engine/Source/Editor/Persona/Private/Persona.h
--- a/Engine/Source/Editor/Persona/Private/Persona.h
+++ b/Engine/Source/Editor/Persona/Private/Persona.h
@@ -13,6 +13,10 @@
 //////////////////////////////////////////////////////////////////////////
 // FPersona
 
+class UBlendProfile;
+class USkeletalMesh;
+class SAnimationEditorViewportTabBody;
+
 /**
  * Persona asset editor (extends Blueprint editor)
  */
@@ -170,6 +174,9 @@ public:
 	/** Called when an asset is imported into the editor */
 	void OnPostImport(UFactory* InFactory, UObject* InObject);
 
+	/** Called when the blend profile tab selects a profile */
+	void SetSelectedBlendProfile(UBlendProfile* InBlendProfile);
+
 public:
 	// IToolkit interface
 	virtual FName GetToolkitContextFName() const override;
@@ -313,13 +320,15 @@ private:
 	// called when animation asset has been changed
 	DECLARE_MULTICAST_DELEGATE_OneParam( FOnAnimChangedMulticaster, UAnimationAsset* )
 	// Called when the preview mesh has been changed
-	DECLARE_MULTICAST_DELEGATE_OneParam(FOnPreviewMeshChangedMulticaster, class USkeletalMesh*)
+	DECLARE_MULTICAST_DELEGATE_OneParam( FOnPreviewMeshChangedMulticaster, USkeletalMesh* )
 	// Called when a socket is selected
 	DECLARE_MULTICAST_DELEGATE_OneParam( FOnSelectSocket, const struct FSelectedSocketInfo& )
 	// Called when a bone is selected
 	DECLARE_MULTICAST_DELEGATE_OneParam( FOnSelectBone, const FName& )
+	// Called when a blend profile is selected in the blend profile tab
+	DECLARE_MULTICAST_DELEGATE_OneParam( FOnSelectBlendProfile, UBlendProfile*)
 	// Called when the preview viewport is created
-	DECLARE_MULTICAST_DELEGATE_OneParam(FOnCreateViewport, TWeakPtr<class SAnimationEditorViewportTabBody>)
+	DECLARE_MULTICAST_DELEGATE_OneParam( FOnCreateViewport, TWeakPtr<SAnimationEditorViewportTabBody> )
 public:
 
 	// Called when Persona refreshes
@@ -412,6 +421,18 @@ public:
 		OnSocketSelected.RemoveAll( Widget );
 	}
 
+	typedef FOnSelectBlendProfile::FDelegate FOnBlendProfileSelected;
+
+	void RegisterOnBlendProfileSelected(const FOnBlendProfileSelected& Delegate)
+	{
+		OnBlendProfileSelected.Add(Delegate);
+	}
+
+	void UnregisterOnBlendProfileSelected(SWidget* Widget)
+	{
+		OnBlendProfileSelected.RemoveAll(Widget);
+	}
+
 	// Called when selection is cleared
 	typedef FSimpleMulticastDelegate::FDelegate FOnAllDeselected;
 
@@ -574,6 +595,9 @@ protected:
 	/** Delegate for when a bone is selected by clicking its hit point */
 	FOnSelectBone OnBoneSelected;
 
+	/** Delegate for when a blend profile is selected in the blend profile tab */
+	FOnSelectBlendProfile OnBlendProfileSelected;
+
 	/** Delegate for clearing the current skeleton bone/socket selection */
 	FSimpleMulticastDelegate OnAllDeselected;
 
diff --git a/Engine/Source/Editor/Persona/Private/SBlendProfilePicker.cpp b/Engine/Source/Editor/Persona/Private/SBlendProfilePicker.cpp
new file mode 100644
--- /dev/null
+++ b/Engine/Source/Editor/Persona/Private/SBlendProfilePicker.cpp
@@ -0,0 +1,179 @@
+// Copyright 1998-2015 Epic Games, Inc. All Rights Reserved.
+
+#include "PersonaPrivatePCH.h"
+#include "SBlendProfilePicker.h"
+#include "Animation/BlendProfile.h"
+#include "STextEntryPopup.h"
+#include "ScopedTransaction.h"
+
+
+#define LOCTEXT_NAMESPACE "BlendProfilePicker"
+
+//////////////////////////////////////////////////////////////////////////
+
+BEGIN_SLATE_FUNCTION_BUILD_OPTIMIZATION
+void SBlendProfilePicker::Construct(const FArguments& InArgs)
+{
+	bShowNewOption = InArgs._AllowNew;
+	bShowClearOption = InArgs._AllowClear;
+	TargetSkeleton = InArgs._TargetSkeleton;
+
+	if(TargetSkeleton)
+	{
+		if(InArgs._InitialProfile == nullptr || TargetSkeleton->BlendProfiles.Contains(InArgs._InitialProfile))
+		{
+			SelectedProfile = InArgs._InitialProfile;
+		}
+		else
+		{
+			SelectedProfile = nullptr;
+		}
+	}
+	else
+	{
+		SelectedProfile = nullptr;
+	}
+
+	BlendProfileSelectedDelegate = InArgs._OnBlendProfileSelected;
+
+	ChildSlot
+	[
+		SNew(SComboButton)
+		.ButtonStyle(FEditorStyle::Get(), "PropertyEditor.AssetComboStyle")
+		.ForegroundColor(FEditorStyle::GetColor("PropertyEditor.AssetName.ColorAndOpacity"))
+		.ContentPadding(2.0f)
+		.OnGetMenuContent(this, &SBlendProfilePicker::GetMenuContent)
+		.ButtonContent()
+		[
+			SNew(STextBlock)
+			.TextStyle(FEditorStyle::Get(), "PropertyEditor.AssetClass")
+			.Font(FEditorStyle::GetFontStyle("PropertyWindow.NormalFont"))
+			.Text(this, &SBlendProfilePicker::GetSelectedProfileName)
+		]
+	];
+}
+END_SLATE_FUNCTION_BUILD_OPTIMIZATION
+
+FText SBlendProfilePicker::GetSelectedProfileName() const
+{
+	if(SelectedProfile)
+	{
+		return FText::Format(FText(LOCTEXT("SelectedNameEntry", "{0}")), FText::FromString(SelectedProfile->GetName()));
+	}
+	return FText(LOCTEXT("NoSelectionEntry", "None"));
+}
+
+TSharedRef<SWidget> SBlendProfilePicker::GetMenuContent()
+{
+	FMenuBuilder MenuBuilder(true, nullptr);
+
+	const bool bHasSettingsSection = bShowNewOption || bShowClearOption;
+
+	if(bHasSettingsSection)
+	{
+		MenuBuilder.BeginSection(NAME_None, LOCTEXT("MenuSettings", "Settings"));
+		{
+			if(bShowNewOption)
+			{
+				MenuBuilder.AddMenuEntry(
+					LOCTEXT("CreateNew", "Create New Blend Profile"),
+					LOCTEXT("CreateNew_ToolTip", "Creates a new blend profile inside the skeleton."),
+					FSlateIcon(),
+					FUIAction(FExecuteAction::CreateSP(this, &SBlendProfilePicker::OnCreateNewProfile)));
+			}
+
+			if(bShowClearOption)
+			{
+				MenuBuilder.AddMenuEntry(
+					LOCTEXT("Clear", "Clear"),
+					LOCTEXT("Clear_ToolTip", "Clear the selected blend profile."),
+					FSlateIcon(),
+					FUIAction(FExecuteAction::CreateSP(this, &SBlendProfilePicker::OnClearSelection)));
+			}
+		}
+		MenuBuilder.EndSection();
+	}
+
+	MenuBuilder.BeginSection(NAME_None, LOCTEXT("Profiles", "Available Blend Profiles"));
+	{
+		for(UBlendProfile* Profile : TargetSkeleton->BlendProfiles)
+		{
+			MenuBuilder.AddMenuEntry(
+				FText::Format(LOCTEXT("ProfileEntry", "{0}"), FText::FromString(Profile->GetName())),
+				LOCTEXT("ProfileEntry_ToolTip", "Select this profile for editing."),
+				FSlateIcon(),
+				FUIAction(FExecuteAction::CreateSP(this, &SBlendProfilePicker::OnProfileSelected, Profile)));
+		}
+	}
+	MenuBuilder.EndSection();
+
+	return MenuBuilder.MakeWidget();
+}
+
+void SBlendProfilePicker::OnProfileSelected(UBlendProfile* Profile)
+{
+	SelectedProfile = Profile;
+
+	BlendProfileSelectedDelegate.ExecuteIfBound(SelectedProfile);
+}
+
+void SBlendProfilePicker::OnClearSelection()
+{
+	SelectedProfile = nullptr;
+
+	BlendProfileSelectedDelegate.ExecuteIfBound(SelectedProfile);
+}
+
+void SBlendProfilePicker::OnCreateNewProfile()
+{
+	TSharedRef<STextEntryPopup> TextEntry = SNew(STextEntryPopup)
+		.Label(LOCTEXT("NewProfileName", "Profile Name"))
+		.OnTextCommitted(this, &SBlendProfilePicker::OnCreateNewProfileComitted);
+
+	FSlateApplication::Get().PushMenu(
+		AsShared(),
+		FWidgetPath(),
+		TextEntry,
+		FSlateApplication::Get().GetCursorPos(),
+		FPopupTransitionEffect(FPopupTransitionEffect::TypeInPopup));
+}
+
+void SBlendProfilePicker::OnCreateNewProfileComitted(const FText& NewName, ETextCommit::Type CommitType)
+{
+	FSlateApplication::Get().DismissAllMenus();
+
+	if(CommitType == ETextCommit::OnEnter && TargetSkeleton)
+	{
+		FScopedTransaction Transaction(LOCTEXT("Trans_NewProfile", "Create new blend profile."));
+
+		FName NameToUse = FName(*NewName.ToString());
+
+		// Only create if we don't have a matching profile
+		if(UBlendProfile* FoundProfile = TargetSkeleton->GetBlendProfile(NameToUse))
+		{
+			OnProfileSelected(FoundProfile);
+		}
+		else if(UBlendProfile* NewProfile = TargetSkeleton->CreateNewBlendProfile(NameToUse))
+		{
+			OnProfileSelected(NewProfile);
+		}
+	}
+}
+
+void SBlendProfilePicker::SetSelectedProfile(UBlendProfile* InProfile, bool bBroadcast /*= true*/)
+{
+	if(TargetSkeleton->BlendProfiles.Contains(InProfile))
+	{
+		SelectedProfile = InProfile;
+		if(bBroadcast)
+		{
+			BlendProfileSelectedDelegate.ExecuteIfBound(InProfile);
+		}
+	}
+	else if(!InProfile)
+	{
+		OnClearSelection();
+	}
+}
+
+#undef LOCTEXT_NAMESPACE
diff --git a/Engine/Source/Editor/Persona/Private/SSkeletonBlendProfiles.cpp b/Engine/Source/Editor/Persona/Private/SSkeletonBlendProfiles.cpp
new file mode 100644
--- /dev/null
+++ b/Engine/Source/Editor/Persona/Private/SSkeletonBlendProfiles.cpp
@@ -0,0 +1,427 @@
+// Copyright 1998-2015 Epic Games, Inc. All Rights Reserved.
+
+#include "PersonaPrivatePCH.h"
+#include "SSkeletonBlendProfiles.h"
+#include "Animation/BlendProfile.h"
+#include "Persona.h"
+#include "SSearchBox.h"
+#include "STextEntryPopup.h"
+#include "ScopedTransaction.h"
+
+#define LOCTEXT_NAMESPACE "SkeletonBlendProfiles"
+
+static const FName ColumnId_BlendProfileName("BlendProfileName");
+static const FName ColumnId_NumEntries("NumEntries");
+
+//////////////////////////////////////////////////////////////////////////
+
+class SBlendProfileListRow : public SMultiColumnTableRow<TSharedPtr<FDisplayedBlendProfileInfo>>
+{
+public:
+
+	SLATE_BEGIN_ARGS(SBlendProfileListRow)
+	{}
+
+	SLATE_ARGUMENT(TSharedPtr<FDisplayedBlendProfileInfo>, Item)
+	SLATE_ARGUMENT(TSharedPtr<SSkeletonBlendProfiles>, BlendProfilesView)
+	SLATE_ARGUMENT(TWeakPtr<FPersona>, WeakPersona)
+
+	SLATE_END_ARGS()
+
+	void Construct(const FArguments& InArgs, const TSharedRef<STableViewBase>& InOwnerTableView);
+
+	virtual TSharedRef<SWidget> GenerateWidgetForColumn( const FName& InColumnName ) override;
+
+private:
+
+	// Get the number of entries in the active blend profile as text
+	FText GetBlendProfileNumEntriesText() const;
+
+	// The item we are observing
+	TSharedPtr<FDisplayedBlendProfileInfo> Item;
+
+	// The view widget that owns us
+	TSharedPtr<SSkeletonBlendProfiles> BlendProfilesView;
+
+	// Hosting Persona instance
+	TWeakPtr<FPersona> WeakPersona;
+};
+
+void SBlendProfileListRow::Construct(const FArguments& InArgs, const TSharedRef<STableViewBase>& InOwnerTableView)
+{
+	Item = InArgs._Item;
+	BlendProfilesView = InArgs._BlendProfilesView;
+	WeakPersona = InArgs._WeakPersona;
+
+	SMultiColumnTableRow<TSharedPtr<FDisplayedBlendProfileInfo>>::Construct(FSuperRowType::FArguments(), InOwnerTableView);
+}
+
+TSharedRef<SWidget> SBlendProfileListRow::GenerateWidgetForColumn(const FName& InColumnName)
+{
+	if(Item.IsValid() && Item->ObservedBlendProfile.IsValid())
+	{
+		UBlendProfile* Profile = Item->ObservedBlendProfile.Get();
+		if(InColumnName == ColumnId_BlendProfileName)
+		{
+			return SNew(STextBlock).Text(FText::FromName(Profile->GetFName()));
+		}
+		else if(InColumnName == ColumnId_NumEntries)
+		{
+			return SNew(STextBlock).Text(this, &SBlendProfileListRow::GetBlendProfileNumEntriesText);
+		}
+	}
+
+	return SNullWidget::NullWidget;
+}
+
+FText SBlendProfileListRow::GetBlendProfileNumEntriesText() const
+{
+	if(Item.IsValid() && Item->ObservedBlendProfile.IsValid())
+	{
+		UBlendProfile* Profile = Item->ObservedBlendProfile.Get();
+		return FText::AsNumber(Profile->GetNumBlendEntries());
+	}
+	return FText::GetEmpty();
+}
+
+//////////////////////////////////////////////////////////////////////////
+
+FSkeletonBlendProfilesSummoner::FSkeletonBlendProfilesSummoner(TSharedPtr<FAssetEditorToolkit> InHostingApp)
+	: FWorkflowTabFactory(FPersonaTabs::BlendProfileManagerID, InHostingApp)
+{
+	TabLabel = LOCTEXT("TabLabel", "Blend Profiles");
+	
+	EnableTabPadding();
+	bIsSingleton = true;
+
+	ViewMenuDescription = LOCTEXT("MenuDescription", "Blend Profiles");
+	ViewMenuTooltip = LOCTEXT("MenuToolTip", "Shows the blend profile manager");
+}
+
+TSharedRef<SWidget> FSkeletonBlendProfilesSummoner::CreateTabBody(const FWorkflowTabSpawnInfo& Info) const
+{
+	return SNew(SSkeletonBlendProfiles)
+		.WeakPersona(StaticCastSharedPtr<FPersona>(HostingApp.Pin()));
+}
+
+//////////////////////////////////////////////////////////////////////////
+
+BEGIN_SLATE_FUNCTION_BUILD_OPTIMIZATION
+void SSkeletonBlendProfiles::Construct(const FArguments& InArgs)
+{
+	WeakPersona = InArgs._WeakPersona;
+
+	OnSuccesfulRenameCommit = FSuccessfulCommit::CreateSP(this, &SSkeletonBlendProfiles::RenameSelectedProfile);
+	OnSuccesfulNewProfileCommit = FSuccessfulCommit::CreateSP(this, &SSkeletonBlendProfiles::CreateNewBlendProfile);
+
+	TSharedPtr<FPersona> SharedPersona = WeakPersona.Pin();
+	if(SharedPersona.IsValid())
+	{
+		CurrentSkeleton = SharedPersona->GetSkeleton();
+
+		ChildSlot
+		[
+			SNew(SVerticalBox)
+
+			+ SVerticalBox::Slot()
+			.AutoHeight()
+			[
+				SAssignNew(SearchBox, SSearchBox)
+				.SelectAllTextWhenFocused(true)
+				.HintText(LOCTEXT("SearchBoxHint", "Search Blend Profiles..."))
+			]
+
+			+SVerticalBox::Slot()
+			[
+				SAssignNew(ListView, SBlendProfileListView)
+				.ListItemsSource(&BlendProfileList)
+				.OnGenerateRow(this, &SSkeletonBlendProfiles::GenerateViewRow)
+				.OnContextMenuOpening(this, &SSkeletonBlendProfiles::GetContextMenu)
+				.OnSelectionChanged(this, &SSkeletonBlendProfiles::OnListSelectionChanged)
+				.ItemHeight(22.0f)
+				.HeaderRow
+				(
+					SNew(SHeaderRow)
+					+SHeaderRow::Column(ColumnId_BlendProfileName)
+					.DefaultLabel(LOCTEXT("NameColumnLabel", "Blend Profile Name"))
+					+SHeaderRow::Column(ColumnId_NumEntries)
+					.DefaultLabel(LOCTEXT("NumEntriesColumnLabel", "Blend Entry Count"))
+				)
+			]
+		];
+
+		RebuildBlendProfileList();
+	}
+}
+END_SLATE_FUNCTION_BUILD_OPTIMIZATION
+
+TSharedRef<ITableRow> SSkeletonBlendProfiles::GenerateViewRow(TSharedPtr<FDisplayedBlendProfileInfo> InInfo, const TSharedRef<STableViewBase>& OwnerTable)
+{
+	return SNew(SBlendProfileListRow, OwnerTable)
+		.WeakPersona(WeakPersona)
+		.Item(InInfo)
+		.BlendProfilesView(SharedThis(this));
+}
+
+TSharedPtr<SWidget> SSkeletonBlendProfiles::GetContextMenu() const
+{
+	FMenuBuilder MenuBuilder(true, nullptr);
+
+	MenuBuilder.BeginSection(NAME_None, LOCTEXT("BlendProfileActions", "Blend Profile Actions"));
+	{
+		{
+			FUIAction Action;
+			Action.ExecuteAction = FExecuteAction::CreateSP(this, &SSkeletonBlendProfiles::OnRequestNewBlendProfile);
+
+			const FText Label = LOCTEXT("Menu_NewProfileLabel", "Create New...");
+			const FText ToolTip = LOCTEXT("Menu_NewProfileToolTip", "Create a new blend profile in the current skeleton.");
+
+			MenuBuilder.AddMenuEntry(Label, ToolTip, FSlateIcon(), Action);
+		}
+	}
+	MenuBuilder.EndSection();
+
+	MenuBuilder.BeginSection(NAME_None, LOCTEXT("SelectedBlendProfileActions", "Selected Blend Profile Actions"));
+	{
+		{
+			FUIAction Action;
+			Action.ExecuteAction = FExecuteAction::CreateSP(this, &SSkeletonBlendProfiles::DeleteSelectedProfiles);
+			Action.CanExecuteAction = FCanExecuteAction::CreateSP(this, &SSkeletonBlendProfiles::MenuCanDelete);
+
+			const FText Label = LOCTEXT("Menu_DeleteProfileLabel", "Delete");
+			const FText ToolTip = LOCTEXT("Menu_DeleteProfileToolTip", "Delete selected blend profiles.");
+
+			MenuBuilder.AddMenuEntry(Label, ToolTip, FSlateIcon(), Action);
+		}
+
+		{
+			FUIAction Action;
+			Action.ExecuteAction = FExecuteAction::CreateSP(this, &SSkeletonBlendProfiles::OnRequestRenameSelectedProfile);
+			Action.CanExecuteAction = FCanExecuteAction::CreateSP(this, &SSkeletonBlendProfiles::MenuCanRename);
+
+			const FText Label = LOCTEXT("Menu_RenameProfileLabel", "Rename");
+			const FText ToolTip = LOCTEXT("Menu_RenameProfileToolTip", "Rename selected profile.");
+
+			MenuBuilder.AddMenuEntry(Label, ToolTip, FSlateIcon(), Action);
+		}
+	}
+	MenuBuilder.EndSection();
+
+	return MenuBuilder.MakeWidget();
+}
+
+void SSkeletonBlendProfiles::RebuildBlendProfileList()
+{
+	BlendProfileList.Empty();
+
+	if(CurrentSkeleton)
+	{
+		for(UBlendProfile* BlendProfile : CurrentSkeleton->BlendProfiles)
+		{
+			TWeakObjectPtr<UBlendProfile> WeakBlendProfile(BlendProfile);
+			BlendProfileList.Add(FDisplayedBlendProfileInfo::Make(WeakBlendProfile));
+		}
+	}
+
+	ListView->RequestListRefresh();
+}
+
+bool SSkeletonBlendProfiles::MenuCanDelete() const
+{
+	TArray<TSharedPtr<FDisplayedBlendProfileInfo>> SelectedRows = ListView->GetSelectedItems();
+	return SelectedRows.Num() > 0;
+}
+
+bool SSkeletonBlendProfiles::MenuCanRename() const
+{
+	TArray<TSharedPtr<FDisplayedBlendProfileInfo>> SelectedRows = ListView->GetSelectedItems();
+	return SelectedRows.Num() == 1;
+}
+
+void SSkeletonBlendProfiles::OnRequestNewBlendProfile()
+{
+	SAssignNew(CurrentTextEntryPopup, STextEntryPopup)
+		.Label(LOCTEXT("RenamePopupLabel", "New Profile Name..."))
+		.OnTextChanged(this, &SSkeletonBlendProfiles::ValidateUserTextEntry)
+		.OnTextCommitted(FOnTextCommitted::CreateSP(this, &SSkeletonBlendProfiles::HandleTextEntryCommit, OnSuccesfulNewProfileCommit));
+
+	FSlateApplication::Get().PushMenu(
+		AsShared(),
+		FWidgetPath(),
+		CurrentTextEntryPopup.ToSharedRef(),
+		FSlateApplication::Get().GetCursorPos(),
+		FPopupTransitionEffect(FPopupTransitionEffect::TypeInPopup));
+}
+
+void SSkeletonBlendProfiles::CreateNewBlendProfile(const FName& InNewName)
+{
+	if(CurrentSkeleton)
+	{
+		FScopedTransaction(LOCTEXT("NewBlendProfileTransaction", "Created new Blend Profile"));
+		CurrentSkeleton->CreateNewBlendProfile(InNewName);
+
+		RefreshObservedData();
+	}
+}
+
+void SSkeletonBlendProfiles::DeleteSelectedProfiles()
+{
+	TArray<TSharedPtr<FDisplayedBlendProfileInfo>> SelectedItems = ListView->GetSelectedItems();
+
+	if(SelectedItems.Num() > 0)
+	{
+		FScopedTransaction Transaction(LOCTEXT("DeleteTransaction", "Delete Blend Profiles"));
+		for(TSharedPtr<FDisplayedBlendProfileInfo>& Item : SelectedItems)
+		{
+			if(Item->ObservedBlendProfile.IsValid())
+			{
+				UBlendProfile* Profile = Item->ObservedBlendProfile.Get();
+				CurrentSkeleton->Modify();
+				
+				CurrentSkeleton->BlendProfiles.RemoveAll([Profile](UBlendProfile* Current)
+				{
+					if(Current == Profile)
+					{
+						Current->MarkPendingKill();
+						return true;
+					}
+
+					return false;
+				});
+			}
+		}
+
+		TSharedPtr<FPersona> SharedPersona = WeakPersona.Pin();
+		if(SharedPersona.IsValid())
+		{
+			SharedPersona->SetSelectedBlendProfile(nullptr);
+		}
+
+		RefreshObservedData();
+	}
+}
+
+void SSkeletonBlendProfiles::OnRequestRenameSelectedProfile()
+{
+	SAssignNew(CurrentTextEntryPopup, STextEntryPopup)
+		.Label(LOCTEXT("RenamePopupLabel", "New Profile Name..."))
+		.OnTextChanged(this, &SSkeletonBlendProfiles::ValidateUserTextEntry)
+		.OnTextCommitted(FOnTextCommitted::CreateSP(this, &SSkeletonBlendProfiles::HandleTextEntryCommit, OnSuccesfulRenameCommit));
+
+	FSlateApplication::Get().PushMenu(
+		AsShared(),
+		FWidgetPath(),
+		CurrentTextEntryPopup.ToSharedRef(),
+		FSlateApplication::Get().GetCursorPos(),
+		FPopupTransitionEffect(FPopupTransitionEffect::TypeInPopup));
+}
+
+void SSkeletonBlendProfiles::ValidateUserTextEntry(const FText& InText)
+{
+	if(CurrentSkeleton && CurrentTextEntryPopup.IsValid())
+	{
+		bCurrentTextEntryHasError = false;
+		FString InputAsString = InText.ToString();
+
+		for(UBlendProfile* Profile : CurrentSkeleton->BlendProfiles)
+		{
+			// Exists in skeleton
+			if(InText.ToString() == Profile->GetName())
+			{
+				CurrentTextEntryPopup->SetError(FText::Format(LOCTEXT("RenameErrorString", "Error \"{0}\" already exists!"), InText));
+				bCurrentTextEntryHasError = true;
+				break;
+			}
+		}
+
+		if(!bCurrentTextEntryHasError)
+		{
+			// Check for further errors
+			FName PossibleNewName(*InputAsString);
+			FText NameErrorString;
+
+			if(!PossibleNewName.IsValidObjectName(NameErrorString))
+			{
+				CurrentTextEntryPopup->SetError(NameErrorString);
+				bCurrentTextEntryHasError = true;
+			}
+		}
+
+		if(!bCurrentTextEntryHasError)
+		{
+			// Errors, clear any message
+			CurrentTextEntryPopup->SetError(FText::GetEmpty());
+		}
+	}
+}
+
+void SSkeletonBlendProfiles::HandleTextEntryCommit(const FText& InText, ETextCommit::Type InCommitType, FSuccessfulCommit OnCommit)
+{
+	check(CurrentTextEntryPopup.IsValid());
+
+	bool bDismissMenus = true;
+
+	if(InCommitType == ETextCommit::OnEnter)
+	{
+		if(!bCurrentTextEntryHasError && OnCommit.IsBound())
+		{
+			OnCommit.Execute(FName(*InText.ToString()));
+		}
+		else
+		{
+			bDismissMenus = false;
+			CurrentTextEntryPopup->FocusDefaultWidget();
+		}
+	}
+
+	if(bDismissMenus)
+	{
+		FSlateApplication::Get().DismissAllMenus();
+	}
+}
+
+void SSkeletonBlendProfiles::RenameSelectedProfile(const FName& NewName)
+{
+	TArray<TSharedPtr<FDisplayedBlendProfileInfo>> SelectedItems = ListView->GetSelectedItems();
+	check(SelectedItems.Num() == 1);
+
+	TSharedPtr<FDisplayedBlendProfileInfo> Item = SelectedItems[0];
+
+	if(Item->ObservedBlendProfile.IsValid())
+	{
+		FScopedTransaction Transaction(LOCTEXT("RenameTransaction", "Rename Blend Profile"));
+		
+		UBlendProfile* Profile = Item->ObservedBlendProfile.Get();
+		Profile->Rename(*NewName.ToString());
+	}
+
+	RefreshObservedData();
+}
+
+void SSkeletonBlendProfiles::RefreshObservedData()
+{
+	RebuildBlendProfileList();
+	ListView->RequestListRefresh();
+}
+
+void SSkeletonBlendProfiles::OnListSelectionChanged(TSharedPtr<FDisplayedBlendProfileInfo> Item, ESelectInfo::Type InSelectionType)
+{
+	TArray<TSharedPtr<FDisplayedBlendProfileInfo>> SelectedItems = ListView->GetSelectedItems();
+
+	TSharedPtr<FPersona> SharedPersona = WeakPersona.Pin();
+	if(SharedPersona.IsValid())
+	{
+		UBlendProfile* ProfileToDisplayInTree = nullptr;
+
+		if(SelectedItems.Num() == 1)
+		{
+			TSharedPtr<FDisplayedBlendProfileInfo> Item = SelectedItems[0];
+			if(Item->ObservedBlendProfile.IsValid())
+			{
+				SharedPersona->SetSelectedBlendProfile(Item->ObservedBlendProfile.Get());
+			}
+		}
+	}
+}
+
+#undef LOCTEXT_NAMESPACE
diff --git a/Engine/Source/Editor/Persona/Private/SSkeletonBlendProfiles.h b/Engine/Source/Editor/Persona/Private/SSkeletonBlendProfiles.h
new file mode 100644
--- /dev/null
+++ b/Engine/Source/Editor/Persona/Private/SSkeletonBlendProfiles.h
@@ -0,0 +1,116 @@
+// Copyright 1998-2015 Epic Games, Inc. All Rights Reserved.
+
+#pragma once
+
+#define LOCTEXT_NAMESPACE "SkeletonBlendProfiles"
+
+class UBlendProfile;
+class FAssetEditorToolkit;
+class FPersona;
+
+//////////////////////////////////////////////////////////////////////////
+
+DECLARE_DELEGATE_OneParam(FSuccessfulCommit, const FName&);
+
+//////////////////////////////////////////////////////////////////////////
+
+class FSkeletonBlendProfilesSummoner : public FWorkflowTabFactory
+{
+public:
+	FSkeletonBlendProfilesSummoner(TSharedPtr<FAssetEditorToolkit> InHostingApp);
+
+	virtual TSharedRef<SWidget> CreateTabBody(const FWorkflowTabSpawnInfo& Info) const override;
+};
+
+//////////////////////////////////////////////////////////////////////////
+
+class FDisplayedBlendProfileInfo
+{
+public:
+	TWeakObjectPtr<UBlendProfile> ObservedBlendProfile;
+
+	static TSharedRef<FDisplayedBlendProfileInfo> Make(TWeakObjectPtr<UBlendProfile>& InBlendProfile)
+	{
+		return MakeShareable(new FDisplayedBlendProfileInfo(InBlendProfile));
+	}
+
+private:
+	// Hidden, use FDisplayedBlendProfileInfo::Make
+	FDisplayedBlendProfileInfo(TWeakObjectPtr<UBlendProfile>& InBlendProfile)
+		: ObservedBlendProfile(InBlendProfile)
+	{
+
+	}
+
+	// Disabled
+	FDisplayedBlendProfileInfo() {};
+};
+
+//////////////////////////////////////////////////////////////////////////
+
+typedef SListView<TSharedPtr<FDisplayedBlendProfileInfo>> SBlendProfileListView;
+
+// Blend profile manager widget for the blend profile tab in Persona
+class SSkeletonBlendProfiles : public SCompoundWidget//, public FGCObject
+{
+public:
+
+	SLATE_BEGIN_ARGS(SSkeletonBlendProfiles)
+		: _WeakPersona()
+	{}
+
+	// Hosting Persona Instance
+	SLATE_ARGUMENT(TWeakPtr<FPersona>, WeakPersona)
+
+	SLATE_END_ARGS()
+
+	void Construct(const FArguments& InArgs);
+
+private:
+
+	TSharedRef<ITableRow> GenerateViewRow(TSharedPtr<FDisplayedBlendProfileInfo> InInfo, const TSharedRef<STableViewBase>& OwnerTable);
+	TSharedPtr<SWidget> GetContextMenu() const;
+
+	void RebuildBlendProfileList();
+
+	void OnRequestNewBlendProfile();
+	void CreateNewBlendProfile(const FName& InNewName);
+
+	bool MenuCanDelete() const;
+	bool MenuCanRename() const;
+
+	void DeleteSelectedProfiles();
+
+	void OnRequestRenameSelectedProfile();
+	void ValidateUserTextEntry(const FText& InText);
+	void HandleTextEntryCommit(const FText& InText, ETextCommit::Type InCommitType, FSuccessfulCommit OnCommit);
+	void OnRenameTextCommitted(const FText& InText, ETextCommit::Type InCommitType);
+	void RenameSelectedProfile(const FName& NewName);
+
+	void RefreshObservedData();
+
+	void OnListSelectionChanged(TSharedPtr<FDisplayedBlendProfileInfo> Item, ESelectInfo::Type InSelectionType);
+
+	// Hosting Persona instance
+	TWeakPtr<FPersona> WeakPersona;
+
+	// Skeleton active in Persona
+	USkeleton* CurrentSkeleton;
+
+	// Widgets
+	TSharedPtr<SSearchBox> SearchBox;
+	TSharedPtr<SBlendProfileListView> ListView;
+
+	// List displayed in the list view
+	TArray<TSharedPtr<FDisplayedBlendProfileInfo>> BlendProfileList;
+
+	// Currently active text entry widget
+	TSharedPtr<STextEntryPopup> CurrentTextEntryPopup;
+	// Whether or not there's an error in the current popup content
+	bool bCurrentTextEntryHasError;
+
+	FSuccessfulCommit OnSuccesfulRenameCommit;
+	FSuccessfulCommit OnSuccesfulNewProfileCommit;
+};
+
+#undef LOCTEXT_NAMESPACE
\ No newline at end of file
diff --git a/Engine/Source/Editor/Persona/Private/SSkeletonTree.cpp b/Engine/Source/Editor/Persona/Private/SSkeletonTree.cpp
--- a/Engine/Source/Editor/Persona/Private/SSkeletonTree.cpp
+++ b/Engine/Source/Editor/Persona/Private/SSkeletonTree.cpp
@@ -33,11 +33,14 @@
 #include "SNotificationList.h"
 #include "NotificationManager.h"
 #include "GenericCommands.h"
+#include "Animation/BlendProfile.h"
+#include "SBlendProfilePicker.h"
 
 #define LOCTEXT_NAMESPACE "SSkeletonTree"
 
 static const FName	ColumnID_BoneLabel( "BoneName" );
 static const FName	ColumnID_RetargetingLabel( "TranslationRetargeting" );
+static const FName  ColumnID_BlendProfileLabel( "BlendProfile" );
 // see if mesh reduction is supported
 static bool	bMeshReductionSupported = false;
 
@@ -180,7 +183,7 @@ TSharedRef< SWidget > SSkeletonTreeRow::GenerateWidgetForColumn( const FName& Co
 	}
 	else
 	{
-		return Item->GenerateWidgetForDataColumn();
+		return Item->GenerateWidgetForDataColumn(ColumnName);
 	}
 }
 
@@ -314,21 +317,50 @@ void FDisplayedMeshBoneInfo::GenerateWidgetForNameColumn( TSharedPtr< SHorizonta
 		];
 }
 
-TSharedRef< SWidget > FDisplayedMeshBoneInfo::GenerateWidgetForDataColumn()
+TSharedRef< SWidget > FDisplayedMeshBoneInfo::GenerateWidgetForDataColumn(const FName& DataColumnName)
 {
-	return SNew( SComboButton )
-		.ContentPadding(3)
-		.OnGetMenuContent( this, &FDisplayedMeshBoneInfo::CreateBoneTranslationRetargetingModeMenu )
-		.ToolTip(IDocumentation::Get()->CreateToolTip(
-		LOCTEXT("RetargetingToolTip", "Set bone translation retargeting mode"),
-		NULL,
-		TEXT("Shared/Editors/Persona"),
-		TEXT("TranslationRetargeting")))
-		.ButtonContent()
-		[
-			SNew( STextBlock )
-			.Text( this, &FDisplayedMeshBoneInfo::GetTranslationRetargetingModeMenuTitle )
-		];
+	if(DataColumnName == ColumnID_RetargetingLabel)
+	{
+		return SNew(SComboButton)
+			.ContentPadding(3)
+			.OnGetMenuContent(this, &FDisplayedMeshBoneInfo::CreateBoneTranslationRetargetingModeMenu)
+			.ToolTip(IDocumentation::Get()->CreateToolTip(
+			LOCTEXT("RetargetingToolTip", "Set bone translation retargeting mode"),
+			NULL,
+			TEXT("Shared/Editors/Persona"),
+			TEXT("TranslationRetargeting")))
+			.ButtonContent()
+			[
+				SNew(STextBlock)
+				.Text(this, &FDisplayedMeshBoneInfo::GetTranslationRetargetingModeMenuTitle)
+			];
+	}
+	else if(DataColumnName == ColumnID_BlendProfileLabel)
+	{
+		TSharedPtr<SSkeletonTree> TreeShared = SkeletonTree.Pin();
+		
+		if(TreeShared.IsValid())
+		{
+			bool bWritable = true;
+			UBlendProfile* CurrentProfile = TreeShared->GetSelectedBlendProfile();
+
+			// We should never have this column if we don't have a profile
+			check(CurrentProfile);
+
+			return SNew(SBox)
+				.Padding(4.0f)
+				[
+					SNew(SSpinBox<float>)
+					.MinValue(0.0f)
+					.MaxValue(1000.0f)
+					.Value(CurrentProfile->GetBoneBlendScale(BoneName))
+					.OnEndSliderMovement(this, &FDisplayedMeshBoneInfo::OnBlendSliderEnd)
+					.OnValueCommitted(this, &FDisplayedMeshBoneInfo::OnBlendSliderCommitted)
+				];
+		}
+	}
+
+	return SNullWidget::NullWidget;
 }
 
 TSharedRef< SWidget > FDisplayedMeshBoneInfo::CreateBoneTranslationRetargetingModeMenu()
@@ -383,6 +415,17 @@ void FDisplayedMeshBoneInfo::SetBoneTranslationRetargetingMode(EBoneTranslationR
 	FAssetNotifications::SkeletonNeedsToBeSaved(TargetSkeleton);
 }
 
+void FDisplayedMeshBoneInfo::SetBoneBlendProfileScale(float NewScale, bool bRecurse)
+{
+	TSharedPtr<SSkeletonTree> TreeShared = SkeletonTree.Pin();
+	if(TreeShared.IsValid())
+	{
+		const int32 BoneIndex = TargetSkeleton->GetReferenceSkeleton().FindBoneIndex( BoneName );
+
+		TreeShared->SetBlendProfileBoneScale(BoneIndex, NewScale, bRecurse);
+	}
+}
+
 FSlateFontInfo FDisplayedMeshBoneInfo::GetBoneTextFont( UDebugSkelMeshComponent* PreviewComponent ) const
 {
 	if ( PreviewComponent )
@@ -483,6 +526,19 @@ FText FDisplayedMeshBoneInfo::GetBoneToolTip()
 	return ToolTip;
 }
 
+void FDisplayedMeshBoneInfo::OnBlendSliderEnd(float NewValue)
+{
+	SetBoneBlendProfileScale(NewValue, false);
+}
+
+void FDisplayedMeshBoneInfo::OnBlendSliderCommitted(float NewValue, ETextCommit::Type CommitType)
+{
+	if(CommitType == ETextCommit::OnEnter)
+	{
+		SetBoneBlendProfileScale(NewValue, false);
+	}
+}
+
 //////////////////////////////////////////////////////////////////////////
 // FDisplayedSocketInfo
 
@@ -564,7 +620,7 @@ void FDisplayedSocketInfo::GenerateWidgetForNameColumn( TSharedPtr< SHorizontalB
 	}
 }
 
-TSharedRef< SWidget > FDisplayedSocketInfo::GenerateWidgetForDataColumn()
+TSharedRef< SWidget > FDisplayedSocketInfo::GenerateWidgetForDataColumn(const FName& DataColumnName)
 {
 	return SNullWidget::NullWidget;
 }
@@ -728,24 +784,29 @@ void FDisplayedAttachedAssetInfo::GenerateWidgetForNameColumn( TSharedPtr< SHori
 		];
 }
 
-TSharedRef< SWidget > FDisplayedAttachedAssetInfo::GenerateWidgetForDataColumn()
+TSharedRef< SWidget > FDisplayedAttachedAssetInfo::GenerateWidgetForDataColumn(const FName& DataColumnName)
 {
-	return SNew( SHorizontalBox )
-		+ SHorizontalBox::Slot()
-		.AutoWidth()
-		.HAlign( HAlign_Left )
-		[
-			SNew( SCheckBox )
-			.ToolTipText( LOCTEXT( "TranslationCheckBoxToolTip", "Click to toggle visibility of this asset" ) )
-			.OnCheckStateChanged( this, &FDisplayedAttachedAssetInfo::OnToggleAssetDisplayed )
-			.IsChecked( this, &FDisplayedAttachedAssetInfo::IsAssetDisplayed )
-			.Style( FEditorStyle::Get(), "CheckboxLookToggleButtonCheckbox" )
+	if(DataColumnName == ColumnID_RetargetingLabel)
+	{
+		return SNew(SHorizontalBox)
+			+ SHorizontalBox::Slot()
+			.AutoWidth()
+			.HAlign(HAlign_Left)
 			[
-				SNew( SImage )
-				.Image( this, &FDisplayedAttachedAssetInfo::OnGetAssetDisplayedButtonImage )
-				.ColorAndOpacity( FLinearColor::Black )
-			]
-		];
+				SNew(SCheckBox)
+				.ToolTipText(LOCTEXT("TranslationCheckBoxToolTip", "Click to toggle visibility of this asset"))
+				.OnCheckStateChanged(this, &FDisplayedAttachedAssetInfo::OnToggleAssetDisplayed)
+				.IsChecked(this, &FDisplayedAttachedAssetInfo::IsAssetDisplayed)
+				.Style(FEditorStyle::Get(), "CheckboxLookToggleButtonCheckbox")
+				[
+					SNew(SImage)
+					.Image(this, &FDisplayedAttachedAssetInfo::OnGetAssetDisplayedButtonImage)
+					.ColorAndOpacity(FLinearColor::Black)
+				]
+			];
+	}
+
+	return SNullWidget::NullWidget;
 }
 
 ECheckBoxState FDisplayedAttachedAssetInfo::IsAssetDisplayed() const
@@ -800,20 +861,26 @@ void SSkeletonTree::Construct(const FArguments& InArgs)
 	BoneFilter = EBoneFilter::All;
 	SocketFilter = ESocketFilter::Active;
 	bShowingAdvancedOptions = false;
+	SelectedBlendProfile = nullptr;
 
 	PersonaPtr = InArgs._Persona;
 	IsEditable = InArgs._IsEditable;
-	TargetSkeleton = PersonaPtr.Pin()->GetSkeleton();
+
+	TSharedPtr<FPersona> SharedPersona = PersonaPtr.Pin();
+	check(SharedPersona.IsValid());
+
+	TargetSkeleton = SharedPersona->GetSkeleton();
 
 	SetPreviewComponentSocketFilter();
 
 	// Register a few delegates with Persona
-	PersonaPtr.Pin()->RegisterOnPostUndo(FPersona::FOnPostUndo::CreateSP( this, &SSkeletonTree::PostUndo ) );
-	PersonaPtr.Pin()->RegisterOnPreviewMeshChanged( FPersona::FOnPreviewMeshChanged::CreateSP( this, &SSkeletonTree::OnPreviewMeshChanged ) );
-	PersonaPtr.Pin()->RegisterOnBoneSelected(FPersona::FOnBoneSelected::CreateSP( this, &SSkeletonTree::OnExternalSelectBone ) );
-	PersonaPtr.Pin()->RegisterOnSocketSelected(FPersona::FOnSocketSelected::CreateSP( this, &SSkeletonTree::OnExternalSelectSocket ) );
-	PersonaPtr.Pin()->RegisterOnDeselectAll(FPersona::FOnAllDeselected::CreateSP( this, &SSkeletonTree::OnExternalDeselectAll ) );
-	PersonaPtr.Pin()->RegisterOnChangeSkeletonTree(FPersona::FOnSkeletonTreeChanged::CreateSP( this, &SSkeletonTree::PostUndo ) );
+	SharedPersona->RegisterOnPostUndo(FPersona::FOnPostUndo::CreateSP( this, &SSkeletonTree::PostUndo ) );
+	SharedPersona->RegisterOnPreviewMeshChanged( FPersona::FOnPreviewMeshChanged::CreateSP( this, &SSkeletonTree::OnPreviewMeshChanged ) );
+	SharedPersona->RegisterOnBoneSelected(FPersona::FOnBoneSelected::CreateSP( this, &SSkeletonTree::OnExternalSelectBone ) );
+	SharedPersona->RegisterOnSocketSelected(FPersona::FOnSocketSelected::CreateSP( this, &SSkeletonTree::OnExternalSelectSocket ) );
+	SharedPersona->RegisterOnDeselectAll(FPersona::FOnAllDeselected::CreateSP( this, &SSkeletonTree::OnExternalDeselectAll ) );
+	SharedPersona->RegisterOnChangeSkeletonTree(FPersona::FOnSkeletonTreeChanged::CreateSP( this, &SSkeletonTree::PostUndo ) );
+	SharedPersona->RegisterOnBlendProfileSelected(FPersona::FOnBlendProfileSelected::CreateSP(this, &SSkeletonTree::OnBlendProfileSelectedExternal));
 
 	// Register and bind all our menu commands
 	FSkeletonTreeCommands::Register();
@@ -834,6 +901,31 @@ void SSkeletonTree::Construct(const FArguments& InArgs)
 			.AddMetaData<FTagMetaData>(TEXT("SkelTree.Search"))
 		]
 
+		+ SVerticalBox::Slot()
+		.AutoHeight()
+		.Padding( FMargin( 0.0f, 0.0f, 0.0f, 4.0f ) )
+		[
+			SNew(SHorizontalBox)
+
+			+ SHorizontalBox::Slot()
+			.Padding(0.0f, 0.0f, 2.0f, 0.0f)
+			.VAlign(VAlign_Center)
+			.AutoWidth()
+			[
+				SNew(STextBlock)
+				.Text(LOCTEXT("BlendProfilePickerLabel", "Blend Profile: "))
+			]
+
+			+ SHorizontalBox::Slot()
+			.Padding(0.0f, 0.0f, 2.0f, 0.0f)
+			.VAlign(VAlign_Center)
+			[
+				SAssignNew(BlendProfilePicker, SBlendProfilePicker)
+				.TargetSkeleton(TargetSkeleton)
+				.OnBlendProfileSelected(this, &SSkeletonTree::OnBlendProfileSelected)
+			]
+		]
+
 		+ SVerticalBox::Slot()
 		.AutoHeight()
 		[
@@ -904,15 +996,17 @@ void SSkeletonTree::Construct(const FArguments& InArgs)
 
 SSkeletonTree::~SSkeletonTree()
 {
-	if ( PersonaPtr.IsValid() )
+	TSharedPtr<FPersona> SharedPersona = PersonaPtr.Pin();
+	if ( SharedPersona.IsValid() )
 	{
-		PersonaPtr.Pin()->UnregisterOnPostUndo( this );
-		PersonaPtr.Pin()->UnregisterOnPreviewMeshChanged( this );
-		PersonaPtr.Pin()->UnregisterOnBoneSelected( this );
-		PersonaPtr.Pin()->UnregisterOnSocketSelected( this );
-		PersonaPtr.Pin()->UnregisterOnDeselectAll( this );
-		PersonaPtr.Pin()->UnregisterOnChangeSkeletonTree( this );
-		PersonaPtr.Pin()->UnregisterOnCreateViewport( this );
+		SharedPersona->UnregisterOnPostUndo( this );
+		SharedPersona->UnregisterOnPreviewMeshChanged( this );
+		SharedPersona->UnregisterOnBoneSelected( this );
+		SharedPersona->UnregisterOnSocketSelected( this );
+		SharedPersona->UnregisterOnDeselectAll( this );
+		SharedPersona->UnregisterOnChangeSkeletonTree( this );
+		SharedPersona->UnregisterOnCreateViewport( this );
+		SharedPersona->UnregisterOnBlendProfileSelected(this);
 	}
 }
 
@@ -1068,7 +1162,7 @@ void SSkeletonTree::CreateTreeColumns()
 	TSharedRef<SHeaderRow> TreeHeaderRow = SNew(SHeaderRow)
 		+ SHeaderRow::Column(ColumnID_BoneLabel)
 		.DefaultLabel(LOCTEXT("SkeletonBoneNameLabel", "Name"))
-		.FillWidth(0.75f);
+		.FillWidth(0.5f);
 
 	if (bShowingAdvancedOptions)
 	{
@@ -1079,6 +1173,14 @@ void SSkeletonTree::CreateTreeColumns()
 			);
 	}
 
+	if(SelectedBlendProfile)
+	{
+		TreeHeaderRow->AddColumn(
+			SHeaderRow::Column(ColumnID_BlendProfileLabel)
+			.DefaultLabel(LOCTEXT("BlendProfileLabel", "Blend Profile Scale"))
+			.FillWidth(0.25f));
+	}
+
 	TreeHolder->ClearChildren();
 	TreeHolder->AddSlot()
 		[
@@ -1414,6 +1516,31 @@ TSharedPtr< SWidget > SSkeletonTree::CreateContextMenu()
 
 			MenuBuilder.EndSection();
 
+			if(SelectedBlendProfile && BoneTreeSelection.IsSingleOfTypeSelected(ESkeletonTreeRowType::Bone))
+			{
+				TSharedPtr<FDisplayedMeshBoneInfo> BoneInfo = BoneTreeSelection.SelectedBones[0];
+
+				FName BoneName = *static_cast<FName*>(BoneInfo->GetData());
+				int32 BoneIndex = TargetSkeleton->GetReferenceSkeleton().FindBoneIndex(BoneName);
+
+				float CurrentBlendScale = SelectedBlendProfile->GetBoneBlendScale(BoneIndex);
+
+				MenuBuilder.BeginSection("SkeletonTreeBlendProfileScales", LOCTEXT("BlendProfileContextOptions", "Blend Profile"));
+				{
+					FUIAction RecursiveSetScales;
+					RecursiveSetScales.ExecuteAction = FExecuteAction::CreateSP(this, &SSkeletonTree::RecursiveSetBlendProfileScales, CurrentBlendScale);
+					
+					MenuBuilder.AddMenuEntry
+						(
+						FText::Format(LOCTEXT("RecursiveSetBlendScales_Label", "Recursively Set Blend Scales To {0}"), FText::AsNumber(CurrentBlendScale)),
+						LOCTEXT("RecursiveSetBlendScales_ToolTip", "Sets all child bones to use the same blend profile scale as the selected bone"),
+						FSlateIcon(),
+						RecursiveSetScales
+						);
+				}
+				MenuBuilder.EndSection();
+			}
+
 			if(bShowingAdvancedOptions)
 			{
 				MenuBuilder.BeginSection("SkeletonTreeBoneTranslationRetargeting", LOCTEXT("BoneTranslationRetargetingHeader", "Bone Translation Retargeting"));
@@ -2567,5 +2694,57 @@ ECheckBoxState SSkeletonTree::IsShowingAdvancedOptions() const
 	return bShowingAdvancedOptions ? ECheckBoxState::Checked : ECheckBoxState::Unchecked;
 }
 
+void SSkeletonTree::OnBlendProfileSelected(UBlendProfile* NewProfile)
+{
+	// Switch out the profile and rebuild the columns, if it's compatible
+	if(NewProfile == nullptr || TargetSkeleton->BlendProfiles.Contains(NewProfile))
+	{
+		SelectedBlendProfile = NewProfile;
+		CreateTreeColumns();
+	}
+}
+
+void SSkeletonTree::SetBlendProfileBoneScale(int32 InBoneIndex, float InNewScale, bool bRecurse)
+{
+	if(SelectedBlendProfile)
+	{
+		SelectedBlendProfile->SetBoneBlendScale(InBoneIndex, InNewScale, bRecurse, true);
+	}
+}
+
+UBlendProfile* SSkeletonTree::GetSelectedBlendProfile()
+{
+	return SelectedBlendProfile;
+}
+
+void SSkeletonTree::OnBlendProfileSelectedExternal(UBlendProfile* NewProfile)
+{
+	if(BlendProfilePicker.IsValid())
+	{
+		BlendProfilePicker->SetSelectedProfile(NewProfile);
+	}
+}
+
+void SSkeletonTree::RecursiveSetBlendProfileScales(float InScaleToSet)
+{
+	if(SelectedBlendProfile)
+	{
+		const FScopedTransaction Transaction(LOCTEXT("SetBlendScalesRecursive", "Recursively Set Blend Profile Scales"));
+		TargetSkeleton->Modify();
+
+		FBoneTreeSelection TreeSelection(SkeletonTreeView->GetSelectedItems());
+
+		for(auto ItemIt = TreeSelection.SelectedBones.CreateConstIterator(); ItemIt; ++ItemIt)
+		{
+			FName BoneName = *static_cast<FName*>((*ItemIt)->GetData());
+			int32 BoneIndex = TargetSkeleton->GetReferenceSkeleton().FindBoneIndex(BoneName);
+			SelectedBlendProfile->SetBoneBlendScale(BoneIndex, InScaleToSet, true, true);
+		}
+		FAssetNotifications::SkeletonNeedsToBeSaved(TargetSkeleton);
+
+		CreateTreeColumns();
+	}
+}
+
 
 #undef LOCTEXT_NAMESPACE
diff --git a/Engine/Source/Editor/Persona/Private/SSkeletonTree.h b/Engine/Source/Editor/Persona/Private/SSkeletonTree.h
--- a/Engine/Source/Editor/Persona/Private/SSkeletonTree.h
+++ b/Engine/Source/Editor/Persona/Private/SSkeletonTree.h
@@ -57,6 +57,8 @@ namespace ESocketFilter
 }
 
 struct FDisplayedTreeRowInfo;
+class SBlendProfilePicker;
+
 typedef STreeView< TSharedPtr<FDisplayedTreeRowInfo> > SMeshSkeletonTreeRowType;
 typedef TSharedPtr< FDisplayedTreeRowInfo > FDisplayedTreeRowInfoPtr;
 
@@ -82,7 +84,7 @@ public:
 	virtual void GenerateWidgetForNameColumn( TSharedPtr< SHorizontalBox > Box, FText& FilterText, FIsSelected InIsSelected ) = 0;
 
 	/** Builds the slate widget for the data column */
-	virtual TSharedRef< SWidget > GenerateWidgetForDataColumn() = 0;
+	virtual TSharedRef< SWidget > GenerateWidgetForDataColumn(const FName& DataColumnName) = 0;
 
 	/** Get the name of the item that this row represents */
 	virtual FName GetRowItemName() const PURE_VIRTUAL(FDisplayedTreeRowInfo::GetRowItemName, return FName(););
@@ -143,7 +145,7 @@ public:
 	virtual void GenerateWidgetForNameColumn( TSharedPtr< SHorizontalBox > Box, FText& FilterText, FIsSelected InIsSelected ) override;
 
 	/** Builds the slate widget for the data column */
-	virtual TSharedRef< SWidget > GenerateWidgetForDataColumn() override;
+	virtual TSharedRef< SWidget > GenerateWidgetForDataColumn(const FName& DataColumnName) override;
 
 	/** Handle dragging a bone */
 	FReply OnDragDetected(const FGeometry& MyGeometry, const FPointerEvent& MouseEvent);
@@ -154,6 +156,9 @@ public:
 	/** Set Translation Retargeting Mode for this bone. */
 	void SetBoneTranslationRetargetingMode(EBoneTranslationRetargetingMode::Type NewRetargetingMode);
 
+	/** Set current Blend Scale for this bone */
+	void SetBoneBlendProfileScale(float NewScale, bool bRecurse);
+
 	virtual ~FDisplayedMeshBoneInfo() {}
 
 protected:
@@ -183,6 +188,12 @@ private:
 
 	/** Get Title for Bone Translation Retargeting Mode menu. */
 	FText GetTranslationRetargetingModeMenuTitle() const;
+
+	/** Callback from a slider widget when the user drops the slider */
+	void OnBlendSliderEnd(float NewValue);
+
+	/** Callback from a slider widget if the text entry is used */
+	void OnBlendSliderCommitted(float NewValue, ETextCommit::Type CommitType);
 };
 
 //////////////////////////////////////////////////////////////////////////
@@ -222,7 +233,7 @@ public:
 	virtual void GenerateWidgetForNameColumn( TSharedPtr< SHorizontalBox > Box, FText& FilterText, FIsSelected InIsSelected ) override;
 
 	/** Builds the slate widget for the data column */
-	virtual TSharedRef< SWidget > GenerateWidgetForDataColumn() override;
+	virtual TSharedRef< SWidget > GenerateWidgetForDataColumn(const FName& DataColumnName) override;
 
 	ESocketParentType::Type GetParentType() const { return ParentType; }
 
@@ -319,7 +330,7 @@ public:
 	virtual void GenerateWidgetForNameColumn( TSharedPtr< SHorizontalBox > Box, FText& FilterText, FIsSelected InIsSelected ) override;
 
 	/** Builds the slate widget for the data column */
-	virtual TSharedRef< SWidget > GenerateWidgetForDataColumn() override;
+	virtual TSharedRef< SWidget > GenerateWidgetForDataColumn(const FName& DataColumnName) override;
 
 	/** Return the name of the asset */
 	virtual FName GetRowItemName() const override {return FName( *Asset->GetName() ) ;}
@@ -425,6 +436,15 @@ public:
 	/** Add the selected bones to LOD of LODIndex when using simplygon **/
 	void AddToLOD(int32 LODIndex);
 
+	/** Sets a scale on the current blend profile */
+	void SetBlendProfileBoneScale(int32 InBoneIndex, float InNewScale, bool bRecurse);
+
+	/** Gets a scale from the currently selected blend profile */
+	float GetBlendProfileBoneScale(int32 InBoneIndex);
+
+	/** Gets the currently selected blend profile */
+	UBlendProfile* GetSelectedBlendProfile();
+
 private:
 	/** Binds the commands in FSkeletonTreeCommands to functions in this class */
 	void BindCommands();
@@ -561,6 +581,15 @@ private:
 	/** Deletes a set of attached objects from a FPreviewAssetAttachContainer and notifies Persona*/
 	void DeleteAttachedObjects( FPreviewAssetAttachContainer& AttachedAssets );
 
+	/** Called when the user selects a blend profile to edit from the blend profile panel*/
+	void OnBlendProfileSelectedExternal(UBlendProfile* NewProfile);
+
+	/** Called when the user selects a blend profile to edit */
+	void OnBlendProfileSelected(UBlendProfile* NewProfile);
+
+	/** Sets the blend scale for the selected bones and all of their children */
+	void RecursiveSetBlendProfileScales(float InScaleToSet);
+
 private:
 	/** Pointer back to the kismet 2 tool that owns us */
 	TWeakPtr<FPersona> PersonaPtr;
@@ -570,6 +599,12 @@ private:
 
 	USkeleton* TargetSkeleton;
 
+	/** Currently selected blend profile */
+	UBlendProfile* SelectedBlendProfile;
+
+	/** The blend profile picker displaying the selected profile */
+	TSharedPtr<SBlendProfilePicker> BlendProfilePicker;
+
 	/** Widget user to hold the skeleton tree */
 	TSharedPtr<SOverlay> TreeHolder;
 
diff --git a/Engine/Source/Editor/Persona/Private/Shared/PersonaMode.cpp b/Engine/Source/Editor/Persona/Private/Shared/PersonaMode.cpp
--- a/Engine/Source/Editor/Persona/Private/Shared/PersonaMode.cpp
+++ b/Engine/Source/Editor/Persona/Private/Shared/PersonaMode.cpp
@@ -8,6 +8,7 @@
 #include "SAnimBlueprintParentPlayerList.h"
 #include "SSkeletonSlotNames.h"
 #include "SSkeletonSmartNameManager.h"
+#include "SSkeletonBlendProfiles.h"
 
 #define LOCTEXT_NAMESPACE "PersonaModes"
 
@@ -40,6 +41,7 @@ const FName FPersonaTabs::SkeletonAnimNotifiesID("SkeletonAnimNotifies");
 const FName FPersonaTabs::SkeletonSlotNamesID("SkeletonSlotNames");
 const FName FPersonaTabs::SkeletonSlotGroupNamesID("SkeletonSlotGroupNames");
 const FName FPersonaTabs::CurveNameManagerID("CurveNameManager");
+const FName FPersonaTabs::BlendProfileManagerID("BlendProfileManager");
 
 /////////////////////////////////////////////////////
 // FPersonaMode
@@ -67,6 +69,7 @@ FPersonaAppMode::FPersonaAppMode(TSharedPtr<class FPersona> InPersona, FName InM
 	PersonaTabFactories.RegisterFactory(MakeShareable(new FRetargetManagerTabSummoner(InPersona)));
 	PersonaTabFactories.RegisterFactory(MakeShareable(new FSkeletonSlotNamesSummoner(InPersona)));
 	PersonaTabFactories.RegisterFactory(MakeShareable(new FSkeletonCurveNameManagerSummoner(InPersona)));
+	PersonaTabFactories.RegisterFactory(MakeShareable(new FSkeletonBlendProfilesSummoner(InPersona)));
 }
 
 void FPersonaAppMode::RegisterTabFactories(TSharedPtr<FTabManager> InTabManager)
diff --git a/Engine/Source/Editor/Persona/Private/Shared/PersonaMode.h b/Engine/Source/Editor/Persona/Private/Shared/PersonaMode.h
--- a/Engine/Source/Editor/Persona/Private/Shared/PersonaMode.h
+++ b/Engine/Source/Editor/Persona/Private/Shared/PersonaMode.h
@@ -38,6 +38,7 @@ struct FPersonaTabs
 	static const FName SkeletonSlotNamesID;
 	static const FName SkeletonSlotGroupNamesID;
 	static const FName CurveNameManagerID;
+	static const FName BlendProfileManagerID;
 
 	// Blueprint Document
 
diff --git a/Engine/Source/Editor/Persona/Public/SBlendProfilePicker.h b/Engine/Source/Editor/Persona/Public/SBlendProfilePicker.h
new file mode 100644
--- /dev/null
+++ b/Engine/Source/Editor/Persona/Public/SBlendProfilePicker.h
@@ -0,0 +1,63 @@
+// Copyright 1998-2015 Epic Games, Inc. All Rights Reserved.
+
+#pragma once
+
+class UBlendProfile;
+
+DECLARE_DELEGATE_OneParam(FOnBlendProfileSelected, UBlendProfile*);
+
+// Picker for UBlendProfile instances inside a USkeleton
+class SBlendProfilePicker : public SCompoundWidget
+{
+public:
+	SLATE_BEGIN_ARGS(SBlendProfilePicker)
+		: _TargetSkeleton(nullptr)
+		, _InitialProfile(nullptr)
+		, _OnBlendProfileSelected()
+		, _AllowNew(true)
+		, _AllowClear(true)
+	{}
+		// Skeleton to operate on
+		SLATE_ARGUMENT(USkeleton*, TargetSkeleton)
+		// Initial blend profile selected
+		SLATE_ARGUMENT(UBlendProfile*, InitialProfile)
+		// Delegate to call when the picker selection is changed
+		SLATE_EVENT(FOnBlendProfileSelected, OnBlendProfileSelected)
+		// Allow the option to create new profiles in the picker
+		SLATE_ARGUMENT(bool, AllowNew)
+		// Allow the option to clear the profile selection
+		SLATE_ARGUMENT(bool, AllowClear)
+	SLATE_END_ARGS()
+
+	PERSONA_API void Construct(const FArguments& InArgs);
+
+	/** Set the selected profile externally 
+	 *  @param InProfile New Profile to set
+	 *  @param bBroadcast Whether or not to broadcast this selection
+	 */
+	PERSONA_API void SetSelectedProfile(UBlendProfile* InProfile, bool bBroadcast = true);
+
+private:
+
+	void OnClearSelection();
+	void OnCreateNewProfile();
+	void OnCreateNewProfileComitted(const FText& NewName, ETextCommit::Type CommitType);
+
+	void OnProfileSelected(UBlendProfile* Profile);
+
+	void RefreshBlendProfileList();
+
+	FText GetSelectedProfileName() const;
+
+	TSharedRef<SWidget> GetMenuContent();
+
+	bool bShowNewOption;
+	bool bShowClearOption;
+
+	UBlendProfile* SelectedProfile;
+
+	USkeleton* TargetSkeleton;
+
+	FOnBlendProfileSelected BlendProfileSelectedDelegate;
+
+};
diff --git a/Engine/Source/Runtime/AnimGraphRuntime/Private/AnimNodes/AnimNode_BlendListBase.cpp b/Engine/Source/Runtime/AnimGraphRuntime/Private/AnimNodes/AnimNode_BlendListBase.cpp
--- a/Engine/Source/Runtime/AnimGraphRuntime/Private/AnimNodes/AnimNode_BlendListBase.cpp
+++ b/Engine/Source/Runtime/AnimGraphRuntime/Private/AnimNodes/AnimNode_BlendListBase.cpp
@@ -4,6 +4,7 @@
 #include "AnimNodes/AnimNode_BlendListBase.h"
 #include "AnimationRuntime.h"
 #include "Animation/AnimStats.h"
+#include "Animation/BlendProfile.h"
 
 /////////////////////////////////////////////////////
 // FAnimNode_BlendListBase
@@ -41,6 +42,20 @@ void FAnimNode_BlendListBase::Initialize(const FAnimationInitializeContext& Cont
 		Blend.CustomCurve = CustomBlendCurve;
 	}
 	Blends[0].SetAlpha(1.0f);
+
+	if(BlendProfile)
+	{
+		// Initialise per-bone data
+		PerBoneSampleData.Empty(NumPoses);
+		PerBoneSampleData.AddZeroed(NumPoses);
+
+		for(int32 Idx = 0 ; Idx < NumPoses ; ++Idx)
+		{
+			FBlendSampleData& SampleData = PerBoneSampleData[Idx];
+			SampleData.SampleDataIndex = Idx;
+			SampleData.PerBoneBlendData.AddZeroed(BlendProfile->GetNumBlendEntries());
+		}
+	}
 }
 
 void FAnimNode_BlendListBase::CacheBones(const FAnimationCacheBonesContext& Context) 
@@ -137,6 +152,33 @@ void FAnimNode_BlendListBase::Update(const FAnimationUpdateContext& Context)
 				PosesToEvaluate.Add(i);
 			}
 		}
+
+		// If we're using a blend profile, extract the scales and build blend sample data
+		if(BlendProfile)
+		{
+			for(int32 i = 0; i < BlendPose.Num(); ++i)
+			{
+				// Update Per-Bone Info
+				const float BlendWeight = BlendWeights[i];
+				FBlendSampleData& PoseSampleData = PerBoneSampleData[i];
+				PoseSampleData.TotalWeight = BlendWeight;
+
+				for(int32 j = 0; j < PoseSampleData.PerBoneBlendData.Num(); ++j)
+				{
+					float& BoneBlend = PoseSampleData.PerBoneBlendData[j];
+					float WeightScale = BlendProfile->GetEntryBlendScale(j);
+
+					if(ChildIndex != i)
+					{
+						WeightScale = 1.0f / WeightScale;
+					}
+
+					BoneBlend = BlendWeight * WeightScale;
+				}
+			}
+
+			FBlendSampleData::NormalizeDataWeight(PerBoneSampleData);
+		}
 	}
 }
 
@@ -154,6 +196,9 @@ void FAnimNode_BlendListBase::Evaluate(FPoseContext& Output)
 		TArray<FBlendedCurve> FilteredCurve;
 		FilteredCurve.SetNum(NumPoses);
 
+		TArray<FBlendSampleData> FilteredSampleData;
+		FilteredSampleData.SetNum(NumPoses);
+
 		TArray<float> FilteredWeights;
 		FilteredWeights.AddUninitialized(NumPoses);
 
@@ -170,10 +215,24 @@ void FAnimNode_BlendListBase::Evaluate(FPoseContext& Output)
 
 			FilteredPoses[i].MoveBonesFrom(EvaluateContext.Pose);
 			FilteredCurve[i] = EvaluateContext.Curve;
+
+			if(BlendProfile)
+			{
+				FilteredSampleData[i] = PerBoneSampleData[PoseIndex];
+			}
+
 			FilteredWeights[i] = BlendWeight;
 		}
 
-		FAnimationRuntime::BlendPosesTogether(FilteredPoses, FilteredCurve, FilteredWeights, Output.Pose, Output.Curve);
+		// Use the calculated blend sample data if we're blending per-bone
+		if(BlendProfile)
+		{
+			FAnimationRuntime::BlendPosesTogetherPerBone(FilteredPoses, FilteredCurve, BlendProfile, FilteredSampleData, Output.Pose, Output.Curve);
+		}
+		else
+		{
+			FAnimationRuntime::BlendPosesTogether(FilteredPoses, FilteredCurve, FilteredWeights, Output.Pose, Output.Curve);
+		}
 	}
 	else
 	{
diff --git a/Engine/Source/Runtime/AnimGraphRuntime/Public/AnimNodes/AnimNode_BlendListBase.h b/Engine/Source/Runtime/AnimGraphRuntime/Public/AnimNodes/AnimNode_BlendListBase.h
--- a/Engine/Source/Runtime/AnimGraphRuntime/Public/AnimNodes/AnimNode_BlendListBase.h
+++ b/Engine/Source/Runtime/AnimGraphRuntime/Public/AnimNodes/AnimNode_BlendListBase.h
@@ -24,6 +24,9 @@ public:
 	UPROPERTY(EditAnywhere, Category=BlendType)
 	UCurveFloat* CustomBlendCurve;
 
+	UPROPERTY(EditAnywhere, Category=BlendType)
+	UBlendProfile* BlendProfile;
+
 	UPROPERTY()
 	TArray<struct FAlphaBlend> Blends;
 
@@ -37,6 +40,9 @@ protected:
 	UPROPERTY()
 	int32 LastActiveChildIndex;
 
+	UPROPERTY()
+	TArray<FBlendSampleData> PerBoneSampleData;
+
 	//Store which poses we need to evaluate
 	TArray<int32> PosesToEvaluate;
 
diff --git a/Engine/Source/Runtime/Engine/Classes/Animation/AnimNode_StateMachine.h b/Engine/Source/Runtime/Engine/Classes/Animation/AnimNode_StateMachine.h
--- a/Engine/Source/Runtime/Engine/Classes/Animation/AnimNode_StateMachine.h
+++ b/Engine/Source/Runtime/Engine/Classes/Animation/AnimNode_StateMachine.h
@@ -51,10 +51,16 @@ struct FAnimationActiveTransitionEntry
 
 	TArray<FAnimNode_TransitionPoseEvaluator*> PoseEvaluators;
 
+	// Blend data used for per-bone animation evaluation
+	TArray<FBlendSampleData> StateBlendData;
+
 #if WITH_EDITORONLY_DATA
 	TArray<int32, TInlineAllocator<3>> SourceTransitionIndices;
 #endif
 
+	// Blend profile to use for this transition. Specifying this will make the transition evaluate per-bone
+	UPROPERTY()
+	UBlendProfile* BlendProfile;
 
 public:
 	FAnimationActiveTransitionEntry();
diff --git a/Engine/Source/Runtime/Engine/Classes/Animation/AnimStateMachineTypes.h b/Engine/Source/Runtime/Engine/Classes/Animation/AnimStateMachineTypes.h
--- a/Engine/Source/Runtime/Engine/Classes/Animation/AnimStateMachineTypes.h
+++ b/Engine/Source/Runtime/Engine/Classes/Animation/AnimStateMachineTypes.h
@@ -134,6 +134,9 @@ struct FAnimationTransitionBetweenStates : public FAnimationStateBase
 	UPROPERTY()
 	UCurveFloat* CustomCurve;
 
+	UPROPERTY()
+	UBlendProfile* BlendProfile;
+
 	UPROPERTY()
 	TEnumAsByte<ETransitionLogicType::Type> LogicType;
 
@@ -151,6 +154,7 @@ struct FAnimationTransitionBetweenStates : public FAnimationStateBase
 		, InterruptNotify(INDEX_NONE)
 		, BlendMode(EAlphaBlendOption::CubicInOut)
 		, CustomCurve(nullptr)
+		, BlendProfile(nullptr)
 		, LogicType(ETransitionLogicType::TLT_StandardBlend)
 #if WITH_EDITOR
 		, ReverseTransition(false)
diff --git a/Engine/Source/Runtime/Engine/Classes/Animation/AnimationAsset.h b/Engine/Source/Runtime/Engine/Classes/Animation/AnimationAsset.h
--- a/Engine/Source/Runtime/Engine/Classes/Animation/AnimationAsset.h
+++ b/Engine/Source/Runtime/Engine/Classes/Animation/AnimationAsset.h
@@ -85,6 +85,8 @@ struct FBlendSampleData
 	{
 		return FMath::Clamp<float>(TotalWeight, 0.f, 1.f);
 	}
+
+	static void ENGINE_API NormalizeDataWeight(TArray<FBlendSampleData>& SampleDataList);
 };
 
 USTRUCT()
diff --git a/Engine/Source/Runtime/Engine/Classes/Animation/BlendProfile.h b/Engine/Source/Runtime/Engine/Classes/Animation/BlendProfile.h
new file mode 100644
--- /dev/null
+++ b/Engine/Source/Runtime/Engine/Classes/Animation/BlendProfile.h
@@ -0,0 +1,104 @@
+// Copyright 1998-2015 Epic Games, Inc. All Rights Reserved.
+
+#pragma once
+
+#include "AnimationRuntime.h"
+#include "BlendProfile.generated.h"
+
+/** A single entry for a blend scale within a profile, mapping a bone to a blendscale */
+USTRUCT()
+struct FBlendProfileBoneEntry
+{
+	GENERATED_BODY()
+
+	UPROPERTY(EditAnywhere, Category=BoneSettings)
+	FBoneReference BoneReference;
+
+	UPROPERTY(EditAnywhere, Category=BoneSettings)
+	float BlendScale;
+};
+
+//////////////////////////////////////////////////////////////////////////
+
+/** A blend profile is a set of per-bone scales that can be used in transitions and blend lists
+ *  to tweak the weights of specific bones. The scales are applied to the normal weight for that bone
+ */
+UCLASS(Within=Skeleton)
+class ENGINE_API UBlendProfile : public UObject, public IInterpolationIndexProvider
+{
+public:
+
+	GENERATED_BODY()
+
+	UBlendProfile();
+
+	/** Get the number of entries in the profile (an entry is any blend scale that isn't 1.0f) */
+	int32 GetNumBlendEntries() const { return ProfileEntries.Num(); }
+
+	/** Set the blend scale for a specific bone 
+	 *  @param InBoneIdx Index of the bone to set the blend scale of
+	 *  @param InScale The scale to set the bone to
+	 *  @param bRecurse Whether or not to set the scale on all children of this bone
+	 *  @param bCreate Whether or not to create a blend profile entry if one does not exist for the specified bone
+	 */
+	void SetBoneBlendScale(int32 InBoneIdx, float InScale, bool bRecurse = false, bool bCreate = false);
+
+	/** Set the blend scale for a specific bone 
+	 *  @param InBoneName Name of the bone to set the blend scale of
+	 *  @param InScale The scale to set the bone to
+	 *  @param bRecurse Whether or not to set the scale on all children of this bone
+	 *  @param bCreate Whether or not to create a blend profile entry if one does not exist for the specified bone
+	 */
+	void SetBoneBlendScale(const FName& InBoneName, float InScale, bool bRecurse = false, bool bCreate = false);
+
+	/** Get the set blend scale for the specified bone, will return 1.0f if no entry was found (no scale)
+	 *  @param InBoneIdx Index of the bone to retrieve
+	 */
+	float GetBoneBlendScale(int32 InBoneIdx) const;
+
+	/** Get the set blend scale for the specified bone, will return 1.0f if no entry was found (no scale)
+	 *  @param InBoneName Name of the bone to retrieve
+	 */
+	float GetBoneBlendScale(const FName& InBoneName) const;
+
+	/** Get the index of the entry for the specified bone
+	 *  @param InBoneIdx Index of the bone
+	 */
+	int32 GetEntryIndex(const int32 InBoneIdx) const;
+
+	/** Get the index of the entry for the specified bone
+	 *  @param InBoneIdx Index of the bone
+	 */
+	int32 GetEntryIndex(const FName& BoneName) const;
+
+	/** Get the blend scale stored in a specific entry 
+	 *  @param InEntryIdx Index of the entry to retrieve
+	 */
+	float GetEntryBlendScale(const int32 InEntryIdx) const;
+
+	// IInterpolationIndexProvider
+	virtual int32 GetPerBoneInterpolationIndex(int32 BoneIndex, const FBoneContainer& RequiredBones) const override;
+	// End IInterpolationIndexProvider
+
+	// UObject
+	virtual bool IsSafeForRootSet() const override {return false;}
+	virtual void PostLoad() override;
+	// End UObject
+
+private:
+
+	/** Sets the skeleton this blend profile is used with */
+	void SetSkeleton(USkeleton* InSkeleton);
+
+	/** Set the blend scale for a single bone (ignore children) */
+	void SetSingleBoneBlendScale(int32 InBoneIdx, float InScale, bool bCreate = false);
+
+public:
+	// The skeleton that owns this profile
+	UPROPERTY()
+	USkeleton* OwningSkeleton;
+
+	// List of blend scale entries
+	UPROPERTY()
+	TArray<FBlendProfileBoneEntry> ProfileEntries;
+};
\ No newline at end of file
diff --git a/Engine/Source/Runtime/Engine/Classes/Animation/BlendSpaceBase.h b/Engine/Source/Runtime/Engine/Classes/Animation/BlendSpaceBase.h
--- a/Engine/Source/Runtime/Engine/Classes/Animation/BlendSpaceBase.h
+++ b/Engine/Source/Runtime/Engine/Classes/Animation/BlendSpaceBase.h
@@ -8,6 +8,7 @@
 #pragma once
 
 #include "AnimSequence.h"
+#include "AnimationRuntime.h"
 #include "BlendSpaceBase.generated.h"
 
 /** Interpolation data types. */
@@ -197,7 +198,7 @@ namespace ENotifyTriggerMode
  * Allows multiple animations to be blended between based on input parameters
  */
 UCLASS(config=Engine, hidecategories=Object, MinimalAPI, BlueprintType)
-class UBlendSpaceBase : public UAnimationAsset
+class UBlendSpaceBase : public UAnimationAsset, public IInterpolationIndexProvider
 {
 	GENERATED_UCLASS_BODY()
 
@@ -371,7 +372,7 @@ public:
 	 * Get PerBoneInterpolationIndex for the input BoneIndex
 	 * If nothing found, return INDEX_NONE
 	 */
-	int32 GetPerBoneInterpolationIndex(int32 BoneIndex, const FBoneContainer& RequiredBones) const;
+	virtual int32 GetPerBoneInterpolationIndex(int32 BoneIndex, const FBoneContainer& RequiredBones) const override;
 
 	/** return true if all sample data is additive **/
 	virtual bool IsValidAdditive() const {check(false); return false;}
diff --git a/Engine/Source/Runtime/Engine/Classes/Animation/Skeleton.h b/Engine/Source/Runtime/Engine/Classes/Animation/Skeleton.h
--- a/Engine/Source/Runtime/Engine/Classes/Animation/Skeleton.h
+++ b/Engine/Source/Runtime/Engine/Classes/Animation/Skeleton.h
@@ -12,6 +12,7 @@
 
 class UAnimSequence;
 class USkeletalMesh;
+class UBlendProfile;
 
 /** This is a mapping table between bone in a particular skeletal mesh and bone of this skeleton set. */
 USTRUCT()
@@ -277,6 +278,21 @@ public:
 	UPROPERTY()
 	FSmartNameContainer SmartNames;
 
+	//////////////////////////////////////////////////////////////////////////
+	// Blend Profiles
+
+	/** List of blend profiles available in this skeleton */
+	UPROPERTY(Instanced)
+	TArray<UBlendProfile*> BlendProfiles;
+
+	/** Get the specified blend profile by name */
+	ENGINE_API UBlendProfile* GetBlendProfile(const FName& InProfileName);
+
+	/** Create a new blend profile with the specified name */
+	ENGINE_API UBlendProfile* CreateNewBlendProfile(const FName& InProfileName);
+
+	//////////////////////////////////////////////////////////////////////////
+
 	/************************************************************************/
 	/* Slot Groups */
 	/************************************************************************/
diff --git a/Engine/Source/Runtime/Engine/Private/Animation/AnimNode_StateMachine.cpp b/Engine/Source/Runtime/Engine/Private/Animation/AnimNode_StateMachine.cpp
--- a/Engine/Source/Runtime/Engine/Private/Animation/AnimNode_StateMachine.cpp
+++ b/Engine/Source/Runtime/Engine/Private/Animation/AnimNode_StateMachine.cpp
@@ -7,6 +7,7 @@
 #include "Animation/AnimNode_SequencePlayer.h"
 #include "AnimationRuntime.h"
 #include "Animation/AnimStats.h"
+#include "Animation/BlendProfile.h"
 
 //////////////////////////////////////////////////////////////////////////
 // FAnimationActiveTransitionEntry
@@ -23,6 +24,7 @@ FAnimationActiveTransitionEntry::FAnimationActiveTransitionEntry()
 	, EndNotify(INDEX_NONE)
 	, InterruptNotify(INDEX_NONE)
 	, LogicType(ETransitionLogicType::TLT_StandardBlend)
+	, BlendProfile(nullptr)
 {
 }
 
@@ -37,6 +39,7 @@ FAnimationActiveTransitionEntry::FAnimationActiveTransitionEntry(int32 NextState
 	, EndNotify(ReferenceTransitionInfo.EndNotify)
 	, InterruptNotify(ReferenceTransitionInfo.InterruptNotify)
 	, LogicType(ReferenceTransitionInfo.LogicType)
+	, BlendProfile(ReferenceTransitionInfo.BlendProfile)
 {
 	const float Scaler = 1.0f - ExistingWeightOfNextState;
 	CrossfadeDuration = ReferenceTransitionInfo.CrossfadeDuration * CalculateInverseAlpha(BlendOption, Scaler);
@@ -83,6 +86,14 @@ void FAnimationActiveTransitionEntry::InitializeCustomGraphLinks(const FAnimatio
 			PoseEvaluators.Add(PoseEvaluator);
 		}
 	}
+
+	// Initialize blend data if necessary
+	if(BlendProfile)
+	{
+		StateBlendData.AddZeroed(2);
+		StateBlendData[0].PerBoneBlendData.AddZeroed(BlendProfile->GetNumBlendEntries());
+		StateBlendData[1].PerBoneBlendData.AddZeroed(BlendProfile->GetNumBlendEntries());
+	}
 }
 
 void FAnimationActiveTransitionEntry::Update(const FAnimationUpdateContext& Context, int32 CurrentStateIndex, bool& bOutFinished)
@@ -101,6 +112,33 @@ void FAnimationActiveTransitionEntry::Update(const FAnimationUpdateContext& Cont
 			bActive = false;
 			bOutFinished = true;
 		}
+
+		// Update state blend data (only when we're using per-bone)
+		if(BlendProfile)
+		{
+			for(int32 Idx = 0 ; Idx < 2 ; ++Idx)
+			{
+				bool bForwards = Idx == 0;
+				FBlendSampleData& CurrentData = StateBlendData[Idx];
+
+				CurrentData.TotalWeight = (bForwards) ? Alpha : 1.0f - Alpha;
+
+				for(int32 PerBoneIndex = 0 ; PerBoneIndex < CurrentData.PerBoneBlendData.Num() ; ++PerBoneIndex)
+				{
+					float& BoneBlend = CurrentData.PerBoneBlendData[PerBoneIndex];
+					float WeightScale = BlendProfile->GetEntryBlendScale(PerBoneIndex);
+
+					if(!bForwards)
+					{
+						WeightScale = 1.0f / WeightScale;
+					}
+
+					BoneBlend = CurrentData.TotalWeight * WeightScale;
+				}
+			}
+
+			FBlendSampleData::NormalizeDataWeight(StateBlendData);
+		}
 	}
 }
 
@@ -638,6 +676,22 @@ void FAnimNode_StateMachine::EvaluateTransitionStandardBlend(FPoseContext& Outpu
 	// Blend it in
 	const ScalarRegister VPreviousWeight(1.0f - Transition.Alpha);
 	const ScalarRegister VWeight(Transition.Alpha);
+
+	// If we have a blend profile we need to blend per bone
+	if(Transition.BlendProfile)
+	{
+		for(FCompactPoseBoneIndex BoneIndex : Output.Pose.ForEachBoneIndex())
+		{
+			int32 PerBoneIndex = Transition.BlendProfile->GetPerBoneInterpolationIndex(BoneIndex.GetInt(), Output.AnimInstance->RequiredBones);
+
+			// Use defined per-bone scale if the bone has a scale specified in the blend profile
+			ScalarRegister FirstWeight = PerBoneIndex != INDEX_NONE ? ScalarRegister(Transition.StateBlendData[1].PerBoneBlendData[PerBoneIndex]) : ScalarRegister(VPreviousWeight);
+			ScalarRegister SecondWeight = PerBoneIndex != INDEX_NONE ? ScalarRegister(Transition.StateBlendData[0].PerBoneBlendData[PerBoneIndex]) : ScalarRegister(VWeight);
+			Output.Pose[BoneIndex] = PreviouseStateResult.Pose[BoneIndex] * FirstWeight;
+			Output.Pose[BoneIndex].AccumulateWithShortestRotation(NextStateResult.Pose[BoneIndex], SecondWeight);
+		}
+	}
+	else
 		for(FCompactPoseBoneIndex BoneIndex : Output.Pose.ForEachBoneIndex())
 		{
 			Output.Pose[BoneIndex] = PreviouseStateResult.Pose[BoneIndex] * VPreviousWeight;
diff --git a/Engine/Source/Runtime/Engine/Private/Animation/AnimationRuntime.cpp b/Engine/Source/Runtime/Engine/Private/Animation/AnimationRuntime.cpp
--- a/Engine/Source/Runtime/Engine/Private/Animation/AnimationRuntime.cpp
+++ b/Engine/Source/Runtime/Engine/Private/Animation/AnimationRuntime.cpp
@@ -14,6 +14,7 @@
 #include "Animation/AnimCompositeBase.h"
 #include "Animation/AnimInstance.h"
 #include "BonePose.h"
+#include "Animation/BlendProfile.h"
 
 DEFINE_LOG_CATEGORY(LogAnimation);
 DEFINE_LOG_CATEGORY(LogRootMotion);
@@ -307,7 +308,7 @@ void BlendPosePerBone(const TArray<int32>& PerBoneIndices, const FBlendSampleDat
  *							This array must be in strictly increasing order.
  */
 
-void FAnimationRuntime::BlendPosesTogetherPerBone(const TArray<FCompactPose>& SourcePoses, const TArray<FBlendedCurve>& SourceCurves, const UBlendSpaceBase* BlendSpace, const TArray<FBlendSampleData>& BlendSampleDataCache, /*out*/ FCompactPose& ResultPose, /*out*/ FBlendedCurve& ResultCurve)
+void FAnimationRuntime::BlendPosesTogetherPerBone(const TArray<FCompactPose>& SourcePoses, const TArray<FBlendedCurve>& SourceCurves, const IInterpolationIndexProvider* InterpolationIndexProvider, const TArray<FBlendSampleData>& BlendSampleDataCache, /*out*/ FCompactPose& ResultPose, /*out*/ FBlendedCurve& ResultCurve)
 {
 	check(SourcePoses.Num() > 0);
 
@@ -317,7 +318,7 @@ void FAnimationRuntime::BlendPosesTogetherPerBone(const TArray<FCompactPose>& So
 	PerBoneIndices.AddUninitialized(ResultPose.GetNumBones());
 	for (int32 BoneIndex = 0; BoneIndex < PerBoneIndices.Num(); ++BoneIndex)
 	{
-		PerBoneIndices[BoneIndex] = BlendSpace->GetPerBoneInterpolationIndex(RequiredBoneIndices[BoneIndex], ResultPose.GetBoneContainer());
+		PerBoneIndices[BoneIndex] = InterpolationIndexProvider->GetPerBoneInterpolationIndex(RequiredBoneIndices[BoneIndex], ResultPose.GetBoneContainer());
 	}
 
 	BlendPosePerBone<ETransformBlendMode::Overwrite>(PerBoneIndices, BlendSampleDataCache[0], ResultPose, SourcePoses[0]);
diff --git a/Engine/Source/Runtime/Engine/Private/Animation/BlendProfile.cpp b/Engine/Source/Runtime/Engine/Private/Animation/BlendProfile.cpp
new file mode 100644
--- /dev/null
+++ b/Engine/Source/Runtime/Engine/Private/Animation/BlendProfile.cpp
@@ -0,0 +1,173 @@
+// Copyright 1998-2015 Epic Games, Inc. All Rights Reserved.
+
+#include "EnginePrivate.h"
+#include "Animation/BlendProfile.h"
+
+UBlendProfile::UBlendProfile()
+	: OwningSkeleton(nullptr)
+{
+	// Set up our owning skeleton and initialise bone references
+	if(USkeleton* OuterAsSkeleton = Cast<USkeleton>(GetOuter()))
+	{
+		SetSkeleton(OuterAsSkeleton);
+	}
+}
+
+void UBlendProfile::SetBoneBlendScale(int32 InBoneIdx, float InScale, bool bRecurse, bool bCreate)
+{
+	// Set the requested bone, then children if necessary
+	SetSingleBoneBlendScale(InBoneIdx, InScale, bCreate);
+
+	if(bRecurse)
+	{
+		const FReferenceSkeleton& RefSkeleton = OwningSkeleton->GetReferenceSkeleton();
+		const int32 NumBones = RefSkeleton.GetNum();
+		for(int32 ChildIdx = InBoneIdx + 1 ; ChildIdx < NumBones ; ++ChildIdx)
+		{
+			if(RefSkeleton.BoneIsChildOf(ChildIdx, InBoneIdx))
+			{
+				SetSingleBoneBlendScale(ChildIdx, InScale, bCreate);
+			}
+		}
+	}
+}
+
+void UBlendProfile::SetBoneBlendScale(const FName& InBoneName, float InScale, bool bRecurse, bool bCreate)
+{
+	int32 BoneIndex = OwningSkeleton->GetReferenceSkeleton().FindBoneIndex(InBoneName);
+
+	SetBoneBlendScale(BoneIndex, InScale, bRecurse, bCreate);
+}
+
+float UBlendProfile::GetBoneBlendScale(int32 InBoneIdx) const
+{
+	const FBlendProfileBoneEntry* FoundEntry = ProfileEntries.FindByPredicate([InBoneIdx](const FBlendProfileBoneEntry& Entry)
+	{
+		return Entry.BoneReference.BoneIndex == InBoneIdx;
+	});
+
+	if(FoundEntry)
+	{
+		return FoundEntry->BlendScale;
+	}
+
+	return 1.0f;
+}
+
+float UBlendProfile::GetBoneBlendScale(const FName& InBoneName) const
+{
+	const FBlendProfileBoneEntry* FoundEntry = ProfileEntries.FindByPredicate([InBoneName](const FBlendProfileBoneEntry& Entry)
+	{
+		return Entry.BoneReference.BoneName == InBoneName;
+	});
+
+	if(FoundEntry)
+	{
+		return FoundEntry->BlendScale;
+	}
+
+	return 1.0f;
+}
+
+void UBlendProfile::SetSkeleton(USkeleton* InSkeleton)
+{
+	OwningSkeleton = InSkeleton;
+
+	if(OwningSkeleton)
+	{
+		// Initialise Current profile entries
+		for(FBlendProfileBoneEntry& Entry : ProfileEntries)
+		{
+			Entry.BoneReference.Initialize(OwningSkeleton);
+		}
+	}
+}
+
+void UBlendProfile::PostLoad()
+{
+	Super::PostLoad();
+
+	if(OwningSkeleton)
+	{
+		// Initialise Current profile entries
+		for(FBlendProfileBoneEntry& Entry : ProfileEntries)
+		{
+			Entry.BoneReference.Initialize(OwningSkeleton);
+		}
+	}
+}
+
+int32 UBlendProfile::GetEntryIndex(const int32 InBoneIdx) const
+{
+	for(int32 Idx = 0 ; Idx < ProfileEntries.Num() ; ++Idx)
+	{
+		const FBlendProfileBoneEntry& Entry = ProfileEntries[Idx];
+		if(Entry.BoneReference.BoneIndex == InBoneIdx)
+		{
+			return Idx;
+		}
+	}
+	return INDEX_NONE;
+}
+
+int32 UBlendProfile::GetEntryIndex(const FName& InBoneName) const
+{
+	for(int32 Idx = 0 ; Idx < ProfileEntries.Num() ; ++Idx)
+	{
+		const FBlendProfileBoneEntry& Entry = ProfileEntries[Idx];
+		if(Entry.BoneReference.BoneName == InBoneName)
+		{
+			return Idx;
+		}
+	}
+	return INDEX_NONE;
+}
+
+float UBlendProfile::GetEntryBlendScale(const int32 InEntryIdx) const
+{
+	if(ProfileEntries.IsValidIndex(InEntryIdx))
+	{
+		return ProfileEntries[InEntryIdx].BlendScale;
+	}
+	// No overriden blend scale, return no scale
+	return 1.0f;
+}
+
+int32 UBlendProfile::GetPerBoneInterpolationIndex(int32 BoneIndex, const FBoneContainer& RequiredBones) const
+{
+	// Our internal entries are skeleton bone indices, but we have pose indices here - convert to skeleton
+	int32 ActualBoneIndex = RequiredBones.GetPoseToSkeletonBoneIndexArray()[BoneIndex];
+
+	return GetEntryIndex(ActualBoneIndex);
+}
+
+void UBlendProfile::SetSingleBoneBlendScale(int32 InBoneIdx, float InScale, bool bCreate /*= false*/)
+{
+	FBlendProfileBoneEntry* Entry = ProfileEntries.FindByPredicate([InBoneIdx](const FBlendProfileBoneEntry& Entry)
+	{
+		return Entry.BoneReference.BoneIndex == InBoneIdx;
+	});
+
+	if(!Entry && bCreate)
+	{
+		Entry = &ProfileEntries[ProfileEntries.AddZeroed()];
+		Entry->BoneReference.BoneName = OwningSkeleton->GetReferenceSkeleton().GetBoneName(InBoneIdx);
+		Entry->BoneReference.Initialize(OwningSkeleton);
+		Entry->BlendScale = InScale;
+	}
+
+	if(Entry)
+	{
+		Modify();
+		Entry->BlendScale = InScale;
+
+		// Remove any entry that gets set back to 1.0f - so we only store entries that actually contain a scale
+		if(Entry->BlendScale == 1.0f)
+		{
+			ProfileEntries.RemoveAll([InBoneIdx](const FBlendProfileBoneEntry& Current)
+			{
+				return Current.BoneReference.BoneIndex == InBoneIdx;
+			});
+		}
+	}
+}
diff --git a/Engine/Source/Runtime/Engine/Private/Animation/BlendSpaceBase.cpp b/Engine/Source/Runtime/Engine/Private/Animation/BlendSpaceBase.cpp
--- a/Engine/Source/Runtime/Engine/Private/Animation/BlendSpaceBase.cpp
+++ b/Engine/Source/Runtime/Engine/Private/Animation/BlendSpaceBase.cpp
@@ -245,7 +245,7 @@ void UBlendSpaceBase::TickAssetPlayerInstance(FAnimTickRecord& Instance, class U
 				if (InterpolateWeightOfSampleData(DeltaTime, OldSampleDataList, NewSampleDataList, SampleDataList))
 				{
 					// now I need to normalize
-					NormalizeSampleDataWeight(SampleDataList);
+					FBlendSampleData::NormalizeDataWeight(SampleDataList);
 				}
 				else
 				{
diff --git a/Engine/Source/Runtime/Engine/Private/Animation/Skeleton.cpp b/Engine/Source/Runtime/Engine/Private/Animation/Skeleton.cpp
--- a/Engine/Source/Runtime/Engine/Private/Animation/Skeleton.cpp
+++ b/Engine/Source/Runtime/Engine/Private/Animation/Skeleton.cpp
@@ -10,6 +10,7 @@
 #include "AssetRegistryModule.h"
 #include "Animation/AnimSequence.h"
 #include "Animation/Rig.h"
+#include "Animation/BlendProfile.h"
 #include "MessageLog.h"
 
 #define LOCTEXT_NAMESPACE "Skeleton"
@@ -38,7 +39,7 @@ FArchive& operator<<(FArchive& Ar, FReferencePose & P)
 	Ar << P.ReferencePose;
 #if WITH_EDITORONLY_DATA
 	//TODO: we should use strip flags but we need to rev the serialization version
-	if (!Ar.IsCooking() && !Ar.IsFilterEditorOnly())
+	if (!Ar.IsCooking())
 	{
 		Ar << P.ReferenceMesh;
 	}
@@ -269,12 +270,6 @@ bool USkeleton::IsCompatibleMesh(const USkeletalMesh* InSkelMesh) const
 		if( SkeletonBoneIndex != INDEX_NONE )
 		{
 			++NumOfBoneMatches;
-
-			// follow the parent chain to verify the chain is same
-			if(!DoesParentChainMatch(SkeletonBoneIndex, InSkelMesh))
-			{
-				return false;
-			}
 		}
 		else
 		{
@@ -387,16 +382,6 @@ int32 USkeleton::BuildLinkup(const USkeletalMesh* InSkelMesh)
 				bDismissedMessage = true;
 			}
 
-			static FName NAME_LoadErrors("LoadErrors");
-			FMessageLog LoadErrors(NAME_LoadErrors);
-
-			TSharedRef<FTokenizedMessage> Message = LoadErrors.Info();
-			Message->AddToken(FTextToken::Create(LOCTEXT("SkeletonBuildLinkupMissingBones1", "The Skeleton ")));
-			Message->AddToken(FAssetNameToken::Create(GetPathName(), FText::FromString( GetNameSafe(this) ) ));
-			Message->AddToken(FTextToken::Create(LOCTEXT("SkeletonBuildLinkupMissingBones2", " is missing bones that SkeletalMesh ")));
-			Message->AddToken(FAssetNameToken::Create(InSkelMesh->GetPathName(), FText::FromString( GetNameSafe(InSkelMesh) )));
-			Message->AddToken(FTextToken::Create(LOCTEXT("SkeletonBuildLinkupMissingBones3", "  needs. They will be added now. Please save the Skeleton!")));
-
 			// Re-add all SkelMesh bones to the Skeleton.
 			MergeAllBonesToBoneTree(InSkelMesh);
 
@@ -1296,5 +1281,28 @@ void USkeleton::GetAssetRegistryTags(TArray<FAssetRegistryTag>& OutTags) const
 	OutTags.Add(FAssetRegistryTag(USkeleton::RigTag, RigFullName, FAssetRegistryTag::TT_Hidden));
 }
 
+UBlendProfile* USkeleton::CreateNewBlendProfile(const FName& InProfileName)
+{
+	Modify();
+	UBlendProfile* NewProfile = NewObject<UBlendProfile>(this, InProfileName, RF_Public);
+	BlendProfiles.Add(NewProfile);
+
+	return NewProfile;
+}
+
+UBlendProfile* USkeleton::GetBlendProfile(const FName& InProfileName)
+{
+	UBlendProfile** FoundProfile = BlendProfiles.FindByPredicate([InProfileName](const UBlendProfile* Profile)
+	{
+		return Profile->GetName() == InProfileName.ToString();
+	});
+
+	if(FoundProfile)
+	{
+		return *FoundProfile;
+	}
+	return nullptr;
+}
+
 #endif //WITH_EDITOR
 #undef LOCTEXT_NAMESPACE 
\ No newline at end of file
diff --git a/Engine/Source/Runtime/Engine/Private/AnimationAsset.cpp b/Engine/Source/Runtime/Engine/Private/AnimationAsset.cpp
--- a/Engine/Source/Runtime/Engine/Private/AnimationAsset.cpp
+++ b/Engine/Source/Runtime/Engine/Private/AnimationAsset.cpp
@@ -408,3 +408,50 @@ void FBoneContainer::RemapFromSkeleton(USkeleton const & SourceSkeleton)
 	PoseToSkeletonBoneIndexArray = SkeletonToPoseBoneIndexArray;
 }
 
+
+void FBlendSampleData::NormalizeDataWeight(TArray<FBlendSampleData>& SampleDataList)
+{
+	float TotalSum = 0.f;
+
+	check(SampleDataList.Num() > 0);
+	int32 NumBones = SampleDataList[0].PerBoneBlendData.Num();
+
+	TArray<float> PerBoneTotalSums;
+	PerBoneTotalSums.AddZeroed(NumBones);
+
+	for(int32 I = 0; I < SampleDataList.Num(); ++I)
+	{
+		checkf(SampleDataList[I].PerBoneBlendData.Num() == NumBones, TEXT("Attempted to normalise a blend sample list, but the samples have differing numbers of bones."));
+
+		TotalSum += SampleDataList[I].GetWeight();
+
+		if(SampleDataList[I].PerBoneBlendData.Num() > 0)
+		{
+			// now interpolate the per bone weights
+			for(int32 Iter = 0; Iter<SampleDataList[I].PerBoneBlendData.Num(); ++Iter)
+			{
+				PerBoneTotalSums[Iter] += SampleDataList[I].PerBoneBlendData[Iter];
+			}
+		}
+	}
+
+	if(ensure(TotalSum > ZERO_ANIMWEIGHT_THRESH))
+	{
+		for(int32 I = 0; I < SampleDataList.Num(); ++I)
+		{
+			if(FMath::Abs<float>(TotalSum - 1.f) > ZERO_ANIMWEIGHT_THRESH)
+			{
+				SampleDataList[I].TotalWeight /= TotalSum;
+			}
+
+			// now interpolate the per bone weights
+			for(int32 Iter = 0; Iter < SampleDataList[I].PerBoneBlendData.Num(); ++Iter)
+			{
+				if(FMath::Abs<float>(PerBoneTotalSums[Iter] - 1.f) > ZERO_ANIMWEIGHT_THRESH)
+				{
+					SampleDataList[I].PerBoneBlendData[Iter] /= PerBoneTotalSums[Iter];
+				}
+			}
+		}
+	}
+}
diff --git a/Engine/Source/Runtime/Engine/Public/AnimationRuntime.h b/Engine/Source/Runtime/Engine/Public/AnimationRuntime.h
--- a/Engine/Source/Runtime/Engine/Public/AnimationRuntime.h
+++ b/Engine/Source/Runtime/Engine/Public/AnimationRuntime.h
@@ -16,9 +16,18 @@ struct FA2Pose;
 struct FPerBoneBlendWeight;
 class UBlendSpaceBase;
 typedef TArray<FTransform> FTransformArrayA2;
+
+/** Interface used to provide interpolation indices for per bone blends
+  *
+  */
+class ENGINE_API IInterpolationIndexProvider
+{
+public:
+	virtual int32 GetPerBoneInterpolationIndex(int32 BoneIndex, const FBoneContainer& RequiredBones) const = 0;
+};
+
 /** In AnimationRunTime Library, we extract animation data based on Skeleton hierarchy, not ref pose hierarchy. 
 	Ref pose will need to be re-mapped later **/
-
 class ENGINE_API FAnimationRuntime
 {
 public:
@@ -90,7 +99,7 @@ public:
 	static void BlendPosesTogetherPerBone(
 		const TArray<FCompactPose>& SourcePoses,
 		const TArray<FBlendedCurve>& SourceCurves,
-		const UBlendSpaceBase* BlendSpace,
+		const IInterpolationIndexProvider* IndexProvider,
 		const TArray<FBlendSampleData>& BlendSampleDataCache,
 		/*out*/ FCompactPose& ResultPose, 
 		/*out*/ FBlendedCurve& ResultCurve);
