#ANIM: Fix detecting modified skeleton and fix up animtrack

Merging

//depot/UE4-Orion/Engine/Source/...

to //depot/UE4/Engine/Source/...

diff --git a/Engine/Source/Editor/Persona/Private/Persona.cpp b/Engine/Source/Editor/Persona/Private/Persona.cpp
--- a/Engine/Source/Editor/Persona/Private/Persona.cpp
+++ b/Engine/Source/Editor/Persona/Private/Persona.cpp
@@ -1860,92 +1860,106 @@ void FPersona::OnActiveTabChanged( TSharedPtr<SDockTab> PreviouslyActive, TShare
 	}
 }
 
-// Sets the current preview mesh
-void FPersona::SetPreviewMesh(USkeletalMesh* NewPreviewMesh)
+void FPersona::SetPreviewMeshInternal(USkeletalMesh* NewPreviewMesh)
 {
-	if(!TargetSkeleton->IsCompatibleMesh(NewPreviewMesh))
+	ValidatePreviewAttachedAssets(NewPreviewMesh);
+	if(NewPreviewMesh != PreviewComponent->SkeletalMesh)
 	{
-		// Send a notification that the skeletal mesh cannot work with the skeleton
-		FFormatNamedArguments Args;
-		Args.Add( TEXT("PreviewMeshName"), FText::FromString( NewPreviewMesh->GetName() ) );
-		Args.Add( TEXT("TargetSkeletonName"), FText::FromString( TargetSkeleton->GetName() ) );
-		FNotificationInfo Info( FText::Format( LOCTEXT("SkeletalMeshIncompatible", "Skeletal Mesh \"{PreviewMeshName}\" incompatible with Skeleton \"{TargetSkeletonName}\"" ), Args ) );
-		Info.ExpireDuration = 3.0f;
-		Info.bUseLargeFont = false;
-		TSharedPtr<SNotificationItem> Notification = FSlateNotificationManager::Get().AddNotification(Info);
-		if ( Notification.IsValid() )
+		if(PreviewComponent->SkeletalMesh != NULL)
 		{
-			Notification->SetCompletionState( SNotificationItem::CS_Fail );
+			RemoveEditingObject(PreviewComponent->SkeletalMesh);
 		}
-	}
-	else
-	{
-		ValidatePreviewAttachedAssets(NewPreviewMesh);
-		if (NewPreviewMesh != PreviewComponent->SkeletalMesh)
-		{
-			if ( PreviewComponent->SkeletalMesh != NULL )
-			{
-				RemoveEditingObject(PreviewComponent->SkeletalMesh);
-			}
 
-			if ( NewPreviewMesh != NULL )
-			{
-				AddEditingObject(NewPreviewMesh);
-			}
+		if(NewPreviewMesh != NULL)
+		{
+			AddEditingObject(NewPreviewMesh);
+		}
 
-			// setting skeletalmesh unregister/re-register, 
-			// so I have to save the animation settings and resetting after setting mesh
-			UAnimationAsset* AnimAssetToPlay = NULL;
-			float PlayPosition = 0.f;
-			bool bPlaying = false;
-			bool bNeedsToCopyAnimationData = PreviewComponent->AnimScriptInstance && PreviewComponent->AnimScriptInstance == PreviewComponent->PreviewInstance;
-			if(bNeedsToCopyAnimationData)
-			{
-				AnimAssetToPlay = PreviewComponent->PreviewInstance->CurrentAsset;
-				PlayPosition = PreviewComponent->PreviewInstance->CurrentTime;
-				bPlaying = PreviewComponent->PreviewInstance->bPlaying;
-			}
+		// setting skeletalmesh unregister/re-register, 
+		// so I have to save the animation settings and resetting after setting mesh
+		UAnimationAsset* AnimAssetToPlay = NULL;
+		float PlayPosition = 0.f;
+		bool bPlaying = false;
+		bool bNeedsToCopyAnimationData = PreviewComponent->AnimScriptInstance && PreviewComponent->AnimScriptInstance == PreviewComponent->PreviewInstance;
+		if(bNeedsToCopyAnimationData)
+		{
+			AnimAssetToPlay = PreviewComponent->PreviewInstance->CurrentAsset;
+			PlayPosition = PreviewComponent->PreviewInstance->CurrentTime;
+			bPlaying = PreviewComponent->PreviewInstance->bPlaying;
+		}
 
-			PreviewComponent->SetSkeletalMesh(NewPreviewMesh);
+		PreviewComponent->SetSkeletalMesh(NewPreviewMesh);
 
-			if(bNeedsToCopyAnimationData)
-			{
-				SetPreviewAnimationAsset(AnimAssetToPlay);
-				PreviewComponent->PreviewInstance->SetPosition(PlayPosition);
-				PreviewComponent->PreviewInstance->bPlaying = bPlaying;
-			}
-		}
-		else
+		if(bNeedsToCopyAnimationData)
 		{
-			PreviewComponent->InitAnim(true);
+			SetPreviewAnimationAsset(AnimAssetToPlay);
+			PreviewComponent->PreviewInstance->SetPosition(PlayPosition);
+			PreviewComponent->PreviewInstance->bPlaying = bPlaying;
 		}
+	}
+	else
+	{
+		PreviewComponent->InitAnim(true);
+	}
 
-		if (NewPreviewMesh != NULL)
+	if(NewPreviewMesh != NULL)
+	{
+		PreviewScene.AddComponent(PreviewComponent, FTransform::Identity);
+		for(auto Iter = AdditionalMeshes.CreateIterator(); Iter; ++Iter)
 		{
-			PreviewScene.AddComponent(PreviewComponent, FTransform::Identity);
-			for (auto Iter = AdditionalMeshes.CreateIterator(); Iter; ++Iter)
-			{
-				PreviewScene.AddComponent((*Iter), FTransform::Identity);
-			}
+			PreviewScene.AddComponent((*Iter), FTransform::Identity);
+		}
 
-			// Set up the mesh for transactions
-			NewPreviewMesh->SetFlags(RF_Transactional);
+		// Set up the mesh for transactions
+		NewPreviewMesh->SetFlags(RF_Transactional);
 
-			AddPreviewAttachedObjects();
+		AddPreviewAttachedObjects();
 
-			if(Viewport.IsValid())
-			{
-				Viewport.Pin()->SetPreviewComponent(PreviewComponent);
-			}
+		if(Viewport.IsValid())
+		{
+			Viewport.Pin()->SetPreviewComponent(PreviewComponent);
 		}
+	}
 
-		for(auto Iter = AdditionalMeshes.CreateIterator(); Iter; ++Iter)
+	for(auto Iter = AdditionalMeshes.CreateIterator(); Iter; ++Iter)
+	{
+		(*Iter)->SetMasterPoseComponent(PreviewComponent);
+		(*Iter)->UpdateMasterBoneMap();
+	}
+
+	OnPreviewMeshChanged.Broadcast(NewPreviewMesh);
+}
+
+// Sets the current preview mesh
+void FPersona::SetPreviewMesh(USkeletalMesh* NewPreviewMesh)
+{
+	if(!TargetSkeleton->IsCompatibleMesh(NewPreviewMesh))
+	{
+		// message box, ask if they'd like to regenerate skeleton
+		if (FMessageDialog::Open(EAppMsgType::YesNo, LOCTEXT("RenerateSkeleton", "The preview mesh hierarchy doesn't match with Skeleton anymore. Would you like to regenerate skeleton?")) == EAppReturnType::Yes)
 		{
-			(*Iter)->SetMasterPoseComponent(PreviewComponent);
-			(*Iter)->UpdateMasterBoneMap();
+			TargetSkeleton->RecreateBoneTree( NewPreviewMesh );
+			SetPreviewMeshInternal(NewPreviewMesh);
 		}
-
-		OnPreviewMeshChanged.Broadcast(NewPreviewMesh);
+		else
+		{
+			// Send a notification that the skeletal mesh cannot work with the skeleton
+			FFormatNamedArguments Args;
+			Args.Add(TEXT("PreviewMeshName"), FText::FromString(NewPreviewMesh->GetName()));
+			Args.Add(TEXT("TargetSkeletonName"), FText::FromString(TargetSkeleton->GetName()));
+			FNotificationInfo Info(FText::Format(LOCTEXT("SkeletalMeshIncompatible", "Skeletal Mesh \"{PreviewMeshName}\" incompatible with Skeleton \"{TargetSkeletonName}\""), Args));
+			Info.ExpireDuration = 3.0f;
+			Info.bUseLargeFont = false;
+			TSharedPtr<SNotificationItem> Notification = FSlateNotificationManager::Get().AddNotification(Info);
+			if(Notification.IsValid())
+			{
+				Notification->SetCompletionState(SNotificationItem::CS_Fail);
+			}
+		}
+	}
+	else
+	{
+		SetPreviewMeshInternal(NewPreviewMesh);
 	}
 }
 
diff --git a/Engine/Source/Editor/Persona/Private/Persona.h b/Engine/Source/Editor/Persona/Private/Persona.h
--- a/Engine/Source/Editor/Persona/Private/Persona.h
+++ b/Engine/Source/Editor/Persona/Private/Persona.h
@@ -669,6 +669,8 @@ private:
 	void ExtendMenu();
 	/** update skeleton ref pose based on current preview mesh */
 	void UpdateSkeletonRefPose();
+	/** set preview mesh internal use only. The mesh should be verified by now.  */
+	void SetPreviewMeshInternal(USkeletalMesh* NewPreviewMesh);
 
 	/** Returns the editor objects that are applicable for our current mode (e.g mesh, animation etc) */
 	TArray<UObject*> GetEditorObjectsForMode(FName Mode) const;
diff --git a/Engine/Source/Runtime/Engine/Classes/Animation/AnimSequence.h b/Engine/Source/Runtime/Engine/Classes/Animation/AnimSequence.h
--- a/Engine/Source/Runtime/Engine/Classes/Animation/AnimSequence.h
+++ b/Engine/Source/Runtime/Engine/Classes/Animation/AnimSequence.h
@@ -142,10 +142,6 @@ struct ENGINE_API FTrackToSkeletonMap
 {
 	GENERATED_USTRUCT_BODY()
 
-	// 0 is the current Skeleton (the one above), and N is the Nth parent.
-	UPROPERTY()
-	int32 SkeletonIndex_DEPRECATED;
-
 	// Index of Skeleton.BoneTree this Track belongs to.
 	UPROPERTY()
 	int32 BoneTreeIndex;
diff --git a/Engine/Source/Runtime/Engine/Private/Animation/AnimSequence.cpp b/Engine/Source/Runtime/Engine/Private/Animation/AnimSequence.cpp
--- a/Engine/Source/Runtime/Engine/Private/Animation/AnimSequence.cpp
+++ b/Engine/Source/Runtime/Engine/Private/Animation/AnimSequence.cpp
@@ -462,11 +462,6 @@ void UAnimSequence::PostLoad()
 	// setup the Codec interfaces
 	AnimationFormat_SetInterfaceLinks(*this);
 
-#if WITH_EDITOR
-	// save track data for the case
-	VerifyTrackMap();
-#endif
-
 	if( IsRunningGame() )
 	{
 		// this probably will not show newly created animations in PIE but will show them in the game once they have been saved off
@@ -536,7 +531,7 @@ void UAnimSequence::VerifyTrackMap(USkeleton* MySkeleton)
 {
 	USkeleton* UseSkeleton = (MySkeleton)? MySkeleton: GetSkeleton();
 
-	if(AnimationTrackNames.Num() != TrackToSkeletonMapTable.Num() && UseSkeleton!=NULL)
+	if( AnimationTrackNames.Num() != TrackToSkeletonMapTable.Num() && UseSkeleton!=NULL)
 	{
 		UE_LOG(LogAnimation, Warning, TEXT("RESAVE ANIMATION NEEDED(%s): Fixing track names."), *GetName());
 
@@ -551,42 +546,52 @@ void UAnimSequence::VerifyTrackMap(USkeleton* MySkeleton)
 	}
 	else if (UseSkeleton != NULL)
 	{
-		int32 NumTracks = AnimationTrackNames.Num();
-		int32 NumSkeletonBone = UseSkeleton->GetReferenceSkeleton().GetNum();
-
-		bool bNeedsFixing = false;
-		// verify all tracks are still valid
-		for(int32 TrackIndex=0; TrackIndex<NumTracks; TrackIndex++)
-		{
-			int32 SkeletonBoneIndex = TrackToSkeletonMapTable[TrackIndex].BoneTreeIndex;
-			// invalid index found
-			if (NumSkeletonBone <= SkeletonBoneIndex)
-			{
-				// if one is invalid, fix up for all.
-				// you don't know what index got messed up
-				bNeedsFixing = true;
-				break;
-			}
-		}
-
-		if(bNeedsFixing)
+		// first check if any of them needs to be removed
 		{
-			UE_LOG(LogAnimation, Warning, TEXT("RESAVE ANIMATION NEEDED(%s): Fixing track index."), *GetName());
+			int32 NumTracks = AnimationTrackNames.Num();
+			int32 NumSkeletonBone = UseSkeleton->GetReferenceSkeleton().GetNum();
 
-			const TArray<FBoneNode>& BoneTree = UseSkeleton->GetBoneTree();
-			for(int32 I=NumTracks-1; I>=0; --I)
+			// the first fix is to make sure 
+			bool bNeedsFixing = false;
+			// verify all tracks are still valid
+			for(int32 TrackIndex=0; TrackIndex<NumTracks; TrackIndex++)
 			{
-				int32 BoneTreeIndex = UseSkeleton->GetReferenceSkeleton().FindBoneIndex(AnimationTrackNames[I]);
-				if (BoneTreeIndex == INDEX_NONE)
+				int32 SkeletonBoneIndex = TrackToSkeletonMapTable[TrackIndex].BoneTreeIndex;
+				// invalid index found
+				if(NumSkeletonBone <= SkeletonBoneIndex)
 				{
-					RemoveTrack(I);
+					// if one is invalid, fix up for all.
+					// you don't know what index got messed up
+					bNeedsFixing = true;
+					break;
 				}
-				else
+			}
+
+			if(bNeedsFixing)
+			{
+				UE_LOG(LogAnimation, Warning, TEXT("RESAVE ANIMATION NEEDED(%s): Fixing track index."), *GetName());
+
+				const TArray<FBoneNode>& BoneTree = UseSkeleton->GetBoneTree();
+				for(int32 I=NumTracks-1; I>=0; --I)
 				{
-					TrackToSkeletonMapTable[I].BoneTreeIndex = BoneTreeIndex;
+					int32 BoneTreeIndex = UseSkeleton->GetReferenceSkeleton().FindBoneIndex(AnimationTrackNames[I]);
+					if(BoneTreeIndex == INDEX_NONE)
+					{
+						RemoveTrack(I);
+					}
+					else
+					{
+						TrackToSkeletonMapTable[I].BoneTreeIndex = BoneTreeIndex;
+					}
 				}
 			}
 		}
+		
+		for(int32 I=0; I<AnimationTrackNames.Num(); ++I)
+		{
+			FTrackToSkeletonMap& TrackMap = TrackToSkeletonMapTable[I];
+			TrackMap.BoneTreeIndex = UseSkeleton->GetReferenceSkeleton().FindBoneIndex(AnimationTrackNames[I]);
+		}		
 	}
 }
 
diff --git a/Engine/Source/Runtime/Engine/Private/Animation/Skeleton.cpp b/Engine/Source/Runtime/Engine/Private/Animation/Skeleton.cpp
--- a/Engine/Source/Runtime/Engine/Private/Animation/Skeleton.cpp
+++ b/Engine/Source/Runtime/Engine/Private/Animation/Skeleton.cpp
@@ -263,13 +263,21 @@ bool USkeleton::IsCompatibleMesh(const USkeletalMesh* InSkelMesh) const
 	// first ensure the parent exists for each bone
 	for (int32 MeshBoneIndex=0; MeshBoneIndex<NumBones; MeshBoneIndex++)
 	{
+		FName MeshBoneName = MeshRefSkel.GetBoneName(MeshBoneIndex);
 		// See if Mesh bone exists in Skeleton.
-		int32 SkeletonBoneIndex = SkeletonRefSkel.FindBoneIndex( MeshRefSkel.GetBoneName(MeshBoneIndex) );
+		int32 SkeletonBoneIndex = SkeletonRefSkel.FindBoneIndex( MeshBoneName );
 
 		// if found, increase num of bone matches count
 		if( SkeletonBoneIndex != INDEX_NONE )
 		{
 			++NumOfBoneMatches;
+
+			// follow the parent chain to verify the chain is same
+			if(!DoesParentChainMatch(SkeletonBoneIndex, InSkelMesh))
+			{
+				UE_LOG(LogAnimation, Warning, TEXT("%s : Hierarchy does not match."), *MeshBoneName.ToString());
+				return false;
+			}
 		}
 		else
 		{
@@ -282,7 +290,8 @@ bool USkeleton::IsCompatibleMesh(const USkeletalMesh* InSkelMesh) const
 				if ( ParentMeshBoneIndex != INDEX_NONE )
 				{
 					// @TODO: make sure RefSkeleton's root ParentIndex < 0 if not, I'll need to fix this by checking TreeBoneIdx
-					SkeletonBoneIndex = SkeletonRefSkel.FindBoneIndex(MeshRefSkel.GetBoneName(ParentMeshBoneIndex));
+					FName ParentBoneName = MeshRefSkel.GetBoneName(ParentMeshBoneIndex);
+					SkeletonBoneIndex = SkeletonRefSkel.FindBoneIndex(ParentBoneName);
 				}
 
 				// root is reached
@@ -299,12 +308,14 @@ bool USkeleton::IsCompatibleMesh(const USkeletalMesh* InSkelMesh) const
 			// still no match, return false, no parent to look for
 			if( SkeletonBoneIndex == INDEX_NONE )
 			{
+				UE_LOG(LogAnimation, Warning, TEXT("%s : Missing joint on skeleton.  Make sure to assign to the skeleeton."), *MeshBoneName.ToString());
 				return false;
 			}
 
 			// second follow the parent chain to verify the chain is same
 			if( !DoesParentChainMatch(SkeletonBoneIndex, InSkelMesh) )
 			{
+				UE_LOG(LogAnimation, Warning, TEXT("%s : Hierarchy does not match."), *MeshBoneName.ToString());
 				return false;
 			}
 		}
diff --git a/Engine/Source/Runtime/Engine/Private/AnimationAsset.cpp b/Engine/Source/Runtime/Engine/Private/AnimationAsset.cpp
--- a/Engine/Source/Runtime/Engine/Private/AnimationAsset.cpp
+++ b/Engine/Source/Runtime/Engine/Private/AnimationAsset.cpp
@@ -153,7 +153,7 @@ bool UAnimationAsset::ReplaceSkeleton(USkeleton* NewSkeleton, bool bConvertSpace
 			for (auto Iter = AnimSeqsToReplace.CreateIterator(); Iter; ++Iter)
 			{
 				UAnimSequence* AnimSeq = *Iter;
-				if (AnimSeq && AnimSeq->Skeleton != NewSkeleton)
+				if (AnimSeq)
 				{
 					AnimSeq->RemapTracksToNewSkeleton(NewSkeleton, bConvertSpaces);
 				}
