Shadow warning fix for blend profile widget

diff --git a/Engine/Source/Editor/Persona/Private/SSkeletonBlendProfiles.cpp b/Engine/Source/Editor/Persona/Private/SSkeletonBlendProfiles.cpp
--- a/Engine/Source/Editor/Persona/Private/SSkeletonBlendProfiles.cpp
+++ b/Engine/Source/Editor/Persona/Private/SSkeletonBlendProfiles.cpp
@@ -415,10 +415,10 @@ void SSkeletonBlendProfiles::OnListSelectionChanged(TSharedPtr<FDisplayedBlendPr
 
 		if(SelectedItems.Num() == 1)
 		{
-			TSharedPtr<FDisplayedBlendProfileInfo> Item = SelectedItems[0];
-			if(Item->ObservedBlendProfile.IsValid())
+			TSharedPtr<FDisplayedBlendProfileInfo> SelectedItem = SelectedItems[0];
+			if (SelectedItem->ObservedBlendProfile.IsValid())
 			{
-				SharedPersona->SetSelectedBlendProfile(Item->ObservedBlendProfile.Get());
+				SharedPersona->SetSelectedBlendProfile(SelectedItem->ObservedBlendProfile.Get());
 			}
 		}
 	}
