Copying //UE4/Dev-Framework to Dev-Main (//UE4/Dev-Main)

diff --git a/Engine/Source/Editor/AnimGraph/Private/AnimPreviewInstance.cpp b/Engine/Source/Editor/AnimGraph/Private/AnimPreviewInstance.cpp
--- a/Engine/Source/Editor/AnimGraph/Private/AnimPreviewInstance.cpp
+++ b/Engine/Source/Editor/AnimGraph/Private/AnimPreviewInstance.cpp
@@ -129,8 +129,7 @@ bool FAnimPreviewInstanceProxy::Evaluate(FPoseContext& Output)
 	if (bEnableControllers)
 	{
 		UDebugSkelMeshComponent* Component = Cast<UDebugSkelMeshComponent>(GetSkelMeshComponent());
-
-		if(Component && GetSkeleton())
+		if(Component)
 		{
 			// update curve controllers
 			UpdateCurveController();
@@ -208,10 +207,7 @@ void FAnimPreviewInstanceProxy::RefreshCurveBoneControllers(UAnimationAsset* Ass
 		GetRequiredBones().SetUseSourceData(true);
 
 		TArray<FTransformCurve>& Curves = CurrentSequence->RawCurveData.TransformCurves;
-		USkeleton* LocalSkeleton = CurrentSequence->GetSkeleton();
-
-		const FSmartNameMapping* NameMapping = LocalSkeleton->GetSmartNameContainer(USkeleton::AnimTrackCurveMappingName);
-
+		USkeleton* MySkeleton = CurrentSequence->GetSkeleton();
 		for (auto& Curve : Curves)
 		{
 			// skip if disabled
@@ -221,21 +217,8 @@ void FAnimPreviewInstanceProxy::RefreshCurveBoneControllers(UAnimationAsset* Ass
 			}
 
 			// add bone modifier
-			FName CurveName;
-			NameMapping->GetName(Curve.CurveUid, CurveName);
-
-			// @TODO: this is going to be issue. If they don't save skeleton with it, we don't have name at all?
- 			if (CurveName == NAME_None)
- 			{
-				FSmartNameMapping::UID NewUID;
-				LocalSkeleton->AddSmartNameAndModify(USkeleton::AnimTrackCurveMappingName, Curve.LastObservedName, NewUID);
-				Curve.CurveUid = NewUID;
-
-				CurveName = Curve.LastObservedName;
- 			}
-
-			FName BoneName = CurveName;
-			if (BoneName != NAME_None && LocalSkeleton->GetReferenceSkeleton().FindBoneIndex(BoneName) != INDEX_NONE)
+			FName BoneName = Curve.Name.DisplayName;
+			if (BoneName != NAME_None && MySkeleton->GetReferenceSkeleton().FindBoneIndex(BoneName) != INDEX_NONE)
 			{
 				ModifyBone(BoneName, true);
 			}
@@ -247,11 +230,11 @@ void FAnimPreviewInstanceProxy::UpdateCurveController()
 {
 	// evaluate the curve data first
 	UAnimSequenceBase* CurrentSequence = Cast<UAnimSequenceBase>(CurrentAsset);
-
-	if (CurrentSequence && GetSkeleton())
+	USkeleton* PreviewSkeleton = (CurrentSequence) ? CurrentSequence->GetSkeleton() : nullptr;
+	if (CurrentSequence && PreviewSkeleton)
 	{
 		TMap<FName, FTransform> ActiveCurves;
-		CurrentSequence->RawCurveData.EvaluateTransformCurveData(GetSkeleton(), ActiveCurves, GetCurrentTime(), 1.f);
+		CurrentSequence->RawCurveData.EvaluateTransformCurveData(PreviewSkeleton, ActiveCurves, GetCurrentTime(), 1.f);
 
 		// make sure those curves exists in the bone controller, otherwise problem
 		if ( ActiveCurves.Num() > 0 )
@@ -314,7 +297,8 @@ void FAnimPreviewInstanceProxy::SetKeyImplementation(const FCompactPose& PreCont
 	UAnimSequence* CurrentSequence = Cast<UAnimSequence>(CurrentAsset);
 	UDebugSkelMeshComponent* Component = Cast<UDebugSkelMeshComponent> (GetSkelMeshComponent());
 
-	if(CurrentSequence && GetSkeleton() && Component && Component->SkeletalMesh)
+	USkeleton* PreviewSkeleton = (CurrentSequence) ? CurrentSequence->GetSkeleton() : nullptr;
+	if(CurrentSequence && PreviewSkeleton && Component && Component->SkeletalMesh)
 	{
 		FScopedTransaction ScopedTransaction(LOCTEXT("SetKey", "Set Key"));
 		CurrentSequence->Modify(true);
diff --git a/Engine/Source/Editor/BehaviorTreeEditor/Private/BehaviorTreeGraph.cpp b/Engine/Source/Editor/BehaviorTreeEditor/Private/BehaviorTreeGraph.cpp
--- a/Engine/Source/Editor/BehaviorTreeEditor/Private/BehaviorTreeGraph.cpp
+++ b/Engine/Source/Editor/BehaviorTreeEditor/Private/BehaviorTreeGraph.cpp
@@ -47,29 +47,29 @@ void UBehaviorTreeGraph::UpdateBlackboardChange()
 			{
 				MyNodeInstance->InitializeFromAsset(*BTAsset);
 			}
-		}
 
-		for (int32 iDecorator = 0; iDecorator < MyNode->Decorators.Num(); iDecorator++)
-		{
-			UBTNode* MyNodeInstance = MyNode->Decorators[iDecorator] ? Cast<UBTNode>(MyNode->Decorators[iDecorator]->NodeInstance) : NULL;
-			if (MyNodeInstance)
+			for (int32 iDecorator = 0; iDecorator < MyNode->Decorators.Num(); iDecorator++)
 			{
-				MyNodeInstance->InitializeFromAsset(*BTAsset);
-			}
+				UBTNode* DecoratorNodeInstance = MyNode->Decorators[iDecorator] ? Cast<UBTNode>(MyNode->Decorators[iDecorator]->NodeInstance) : NULL;
+				if (DecoratorNodeInstance)
+				{
+					DecoratorNodeInstance->InitializeFromAsset(*BTAsset);
+				}
 
-			UBehaviorTreeGraphNode_CompositeDecorator* CompDecoratorNode = Cast<UBehaviorTreeGraphNode_CompositeDecorator>(MyNode->Decorators[iDecorator]);
-			if (CompDecoratorNode)
-			{
-				CompDecoratorNode->OnBlackboardUpdate();
+				UBehaviorTreeGraphNode_CompositeDecorator* CompDecoratorNode = Cast<UBehaviorTreeGraphNode_CompositeDecorator>(MyNode->Decorators[iDecorator]);
+				if (CompDecoratorNode)
+				{
+					CompDecoratorNode->OnBlackboardUpdate();
+				}
 			}
-		}
 
-		for (int32 iService = 0; iService < MyNode->Services.Num(); iService++)
-		{
-			UBTNode* MyNodeInstance = MyNode->Services[iService] ? Cast<UBTNode>(MyNode->Services[iService]->NodeInstance) : NULL;
-			if (MyNodeInstance)
+			for (int32 iService = 0; iService < MyNode->Services.Num(); iService++)
 			{
-				MyNodeInstance->InitializeFromAsset(*BTAsset);
+				UBTNode* ServiceNodeInstance = MyNode->Services[iService] ? Cast<UBTNode>(MyNode->Services[iService]->NodeInstance) : NULL;
+				if (ServiceNodeInstance)
+				{
+					ServiceNodeInstance->InitializeFromAsset(*BTAsset);
+				}
 			}
 		}
 	}
@@ -88,6 +88,12 @@ void UBehaviorTreeGraph::UpdateAsset(int32 UpdateFlags)
 	{
 		UBehaviorTreeGraphNode* Node = Cast<UBehaviorTreeGraphNode>(Nodes[Index]);
 
+		if (Node == nullptr)
+		{
+			// ignore non-BT nodes.
+			continue;
+		}
+
 		// debugger flags
 		if (UpdateFlags & ClearDebuggerFlags)
 		{
@@ -885,18 +891,21 @@ void UBehaviorTreeGraph::CollectAllNodeInstances(TSet<UObject*>& NodeInstance)
 	for (int32 Idx = 0; Idx < Nodes.Num(); Idx++)
 	{
 		UBehaviorTreeGraphNode* MyNode = Cast<UBehaviorTreeGraphNode>(Nodes[Idx]);
-		for (int32 SubIdx = 0; SubIdx < MyNode->Decorators.Num(); SubIdx++)
+		if (MyNode)
 		{
-			UBehaviorTreeGraphNode_CompositeDecorator* SubgraphNode = Cast<UBehaviorTreeGraphNode_CompositeDecorator>(MyNode->Decorators[SubIdx]);
-			if (SubgraphNode)
+			for (int32 SubIdx = 0; SubIdx < MyNode->Decorators.Num(); SubIdx++)
 			{
-				TArray<UBTDecorator*> DecoratorInstances;
-				TArray<FBTDecoratorLogic> DummyOps;
-				SubgraphNode->CollectDecoratorData(DecoratorInstances, DummyOps);
-
-				for (int32 DecoratorIdx = 0; DecoratorIdx < DecoratorInstances.Num(); DecoratorIdx++)
+				UBehaviorTreeGraphNode_CompositeDecorator* SubgraphNode = Cast<UBehaviorTreeGraphNode_CompositeDecorator>(MyNode->Decorators[SubIdx]);
+				if (SubgraphNode)
 				{
-					NodeInstance.Add(DecoratorInstances[DecoratorIdx]);
+					TArray<UBTDecorator*> DecoratorInstances;
+					TArray<FBTDecoratorLogic> DummyOps;
+					SubgraphNode->CollectDecoratorData(DecoratorInstances, DummyOps);
+
+					for (int32 DecoratorIdx = 0; DecoratorIdx < DecoratorInstances.Num(); DecoratorIdx++)
+					{
+						NodeInstance.Add(DecoratorInstances[DecoratorIdx]);
+					}
 				}
 			}
 		}
@@ -1046,6 +1055,10 @@ void UBehaviorTreeGraph::RebuildExecutionOrder()
 	for (int32 Index = 0; Index < Nodes.Num(); ++Index)
 	{
 		UBehaviorTreeGraphNode* Node = Cast<UBehaviorTreeGraphNode>(Nodes[Index]);
+		if (Node == nullptr)
+		{
+			continue;;
+		}
 
 		// prepare node instance
 		UBTNode* NodeInstance = Cast<UBTNode>(Node->NodeInstance);
@@ -1263,6 +1276,11 @@ void UBehaviorTreeGraph::UpdateVersion_UnifiedSubNodes()
 	for (int32 NodeIdx = 0; NodeIdx < Nodes.Num(); NodeIdx++)
 	{
 		UBehaviorTreeGraphNode* MyNode = Cast<UBehaviorTreeGraphNode>(Nodes[NodeIdx]);
+		if (MyNode == nullptr)
+		{
+			continue;
+		}
+
 		MyNode->SubNodes.Reset(MyNode->Decorators.Num() + MyNode->Services.Num());
 
 		for (int32 SubIdx = 0; SubIdx < MyNode->Decorators.Num(); SubIdx++)
diff --git a/Engine/Source/Editor/ComponentVisualizers/Private/ConstraintComponentVisualizer.cpp b/Engine/Source/Editor/ComponentVisualizers/Private/ConstraintComponentVisualizer.cpp
--- a/Engine/Source/Editor/ComponentVisualizers/Private/ConstraintComponentVisualizer.cpp
+++ b/Engine/Source/Editor/ComponentVisualizers/Private/ConstraintComponentVisualizer.cpp
@@ -3,6 +3,7 @@
 #include "ComponentVisualizersPrivatePCH.h"
 
 #include "ConstraintComponentVisualizer.h"
+#include "PhysicsEngine/PhysicsConstraintComponent.h"
 
 static const FColor	JointFrame1Color(255,0,0);
 static const FColor JointFrame2Color(0,0,255);
diff --git a/Engine/Source/Editor/DetailCustomizations/Private/AnimSequenceDetails.cpp b/Engine/Source/Editor/DetailCustomizations/Private/AnimSequenceDetails.cpp
--- a/Engine/Source/Editor/DetailCustomizations/Private/AnimSequenceDetails.cpp
+++ b/Engine/Source/Editor/DetailCustomizations/Private/AnimSequenceDetails.cpp
@@ -7,6 +7,7 @@
 #include "Runtime/Engine/Public/Slate/SceneViewport.h"
 #include "Editor/KismetWidgets/Public/SScrubControlPanel.h"
 #include "AnimationEditorUtils.h"
+#include "AnimationCompressionPanel.h"
 
 #define LOCTEXT_NAMESPACE	"AnimSequenceDetails"
 
@@ -179,9 +180,9 @@ void FAnimSequenceDetails::CustomizeDetails(IDetailLayoutBuilder& DetailBuilder)
 			.HAlign(HAlign_Left)
 			[
 				SNew(SButton)
-				.Text(LOCTEXT("ApplyCompressionButton_Label", "Apply"))
-				.ToolTipText(LOCTEXT("ApplyCompressionButton_Tooltip", "Click Apply to apply the changes you've made to Compression Settings"))
-				.OnClicked(this, &FAnimSequenceDetails::OnApplyCompression)
+				.Text(LOCTEXT("EditCompressionButton_Label", "Edit Compression Settings"))
+				.ToolTipText(LOCTEXT("EditCompressionButton_Tooltip", "Click to view and edit the Compression Settings"))
+				.OnClicked(this, &FAnimSequenceDetails::OnEditCompression)
 			]
 		];
 }
@@ -331,32 +332,10 @@ TSharedPtr<FString> FAnimSequenceDetails::GetRetargetSourceString(FName Retarget
 	return RetargetSourceComboList[0];
 }
 
-FReply FAnimSequenceDetails::OnApplyCompression()
+FReply FAnimSequenceDetails::OnEditCompression()
 {
-	TArray<UAnimSequence*> AnimSequences;
-	UAnimCompress* CompressAlgorithm = NULL;
-
-	for(auto SelectionIt = SelectedAnimSequences.CreateIterator(); SelectionIt; ++SelectionIt)
-	{
-		if(UAnimSequence* TestAnimSequence = Cast<UAnimSequence>(SelectionIt->Get()))
-		{
-			AnimSequences.Add(TestAnimSequence);
-
-			if (!CompressAlgorithm)
-			{
-				// it should be all same by now
-				// better way is to get this from property handler
-				// but by this time, they should have all same value
-				CompressAlgorithm = TestAnimSequence->CompressionScheme;
-			}
-		}
-	}
-
-	if (CompressAlgorithm)
-	{
-		AnimationEditorUtils::ApplyCompressionAlgorithm(AnimSequences, CompressAlgorithm);
-	}
-
+	FDlgAnimCompression AnimCompressionDialog(SelectedAnimSequences);
+	AnimCompressionDialog.ShowModal();
 	return FReply::Handled();
 }
 ////////////////////////////////////////////////
diff --git a/Engine/Source/Editor/DetailCustomizations/Private/AnimSequenceDetails.h b/Engine/Source/Editor/DetailCustomizations/Private/AnimSequenceDetails.h
--- a/Engine/Source/Editor/DetailCustomizations/Private/AnimSequenceDetails.h
+++ b/Engine/Source/Editor/DetailCustomizations/Private/AnimSequenceDetails.h
@@ -55,7 +55,7 @@ private:
 	// cache all anim sequences that are selected
 	// but I need to know them before compress
 	TArray< TWeakObjectPtr<UAnimSequence> > SelectedAnimSequences;
-	FReply OnApplyCompression();
+	FReply OnEditCompression();
 };
 
 ///////////////////////////////////////////////////
diff --git a/Engine/Source/Editor/DetailCustomizations/Private/BodyInstanceCustomization.cpp b/Engine/Source/Editor/DetailCustomizations/Private/BodyInstanceCustomization.cpp
--- a/Engine/Source/Editor/DetailCustomizations/Private/BodyInstanceCustomization.cpp
+++ b/Engine/Source/Editor/DetailCustomizations/Private/BodyInstanceCustomization.cpp
@@ -26,10 +26,10 @@ UStaticMeshComponent* FBodyInstanceCustomization::GetDefaultCollisionProvider(co
 	UPrimitiveComponent* OwnerComp = BI->OwnerComponent.Get();
 	if(!OwnerComp)
 	{
-		if(UPrimitiveComponent* const* FoundComp = BodyInstanceToPrimComponent.Find(BI))
-		{
-			OwnerComp = *FoundComp;
-		}
+		TWeakObjectPtr<UPrimitiveComponent> FoundComp = BodyInstanceToPrimComponent.FindRef(BI);
+		OwnerComp = FoundComp.Get();
+		
+		ensureMsgf(!FoundComp.IsStale(), TEXT("BodyInstance has no owner and the map data is stale"));
 	}
 	
 	UStaticMeshComponent* SMC = Cast<UStaticMeshComponent>(OwnerComp);
diff --git a/Engine/Source/Editor/DetailCustomizations/Private/BodyInstanceCustomization.h b/Engine/Source/Editor/DetailCustomizations/Private/BodyInstanceCustomization.h
--- a/Engine/Source/Editor/DetailCustomizations/Private/BodyInstanceCustomization.h
+++ b/Engine/Source/Editor/DetailCustomizations/Private/BodyInstanceCustomization.h
@@ -100,7 +100,7 @@ private:
 
 	TArray<FBodyInstance*> BodyInstances;
 	TArray<UPrimitiveComponent*>		PrimComponents;
-	TMap<FBodyInstance*, UPrimitiveComponent*> BodyInstanceToPrimComponent;
+	TMap<FBodyInstance*, TWeakObjectPtr<UPrimitiveComponent>> BodyInstanceToPrimComponent;
 
 	TArray<FCollisionChannelInfo>	ValidCollisionChannels;
 
diff --git a/Engine/Source/Editor/DetailCustomizations/Private/DetailCustomizations.cpp b/Engine/Source/Editor/DetailCustomizations/Private/DetailCustomizations.cpp
--- a/Engine/Source/Editor/DetailCustomizations/Private/DetailCustomizations.cpp
+++ b/Engine/Source/Editor/DetailCustomizations/Private/DetailCustomizations.cpp
@@ -23,6 +23,7 @@
 #include "AnimSequenceDetails.h"
 #include "AnimStateNodeDetails.h"
 #include "AnimTransitionNodeDetails.h"
+#include "PoseAssetDetails.h"
 #include "AmbientSoundDetails.h"
 #include "ModuleManager.h"
 #include "AnimGraphDefinitions.h"
@@ -51,8 +52,6 @@
 #include "MarginCustomization.h"
 #include "SceneComponentDetails.h"
 #include "PhysicsConstraintComponentDetails.h"
-#include "Runtime/Engine/Classes/PhysicsEngine/PhysicsConstraintTemplate.h"
-#include "Runtime/Engine/Classes/PhysicsEngine/PhysicsConstraintComponent.h"
 #include "GuidStructCustomization.h"
 #include "ParticleModuleDetails.h"
 #include "CameraDetails.h"
@@ -279,6 +278,7 @@ void FDetailCustomizationsModule::RegisterObjectCustomizations()
 	RegisterCustomClassLayout("AnimStateNode", FOnGetDetailCustomizationInstance::CreateStatic(&FAnimStateNodeDetails::MakeInstance));
 	RegisterCustomClassLayout("AnimStateTransitionNode", FOnGetDetailCustomizationInstance::CreateStatic(&FAnimTransitionNodeDetails::MakeInstance));
 	RegisterCustomClassLayout("AnimGraphNode_Trail", FOnGetDetailCustomizationInstance::CreateStatic(&FAnimTrailNodeDetails::MakeInstance));
+	RegisterCustomClassLayout("PoseAsset", FOnGetDetailCustomizationInstance::CreateStatic(&FPoseAssetDetails::MakeInstance));
 
 	RegisterCustomClassLayout("SoundWave", FOnGetDetailCustomizationInstance::CreateStatic(&FSoundWaveDetails::MakeInstance));
 	RegisterCustomClassLayout("DialogueWave", FOnGetDetailCustomizationInstance::CreateStatic(&FDialogueWaveDetails::MakeInstance));
diff --git a/Engine/Source/Editor/EditorStyle/Private/SlateEditorStyle.cpp b/Engine/Source/Editor/EditorStyle/Private/SlateEditorStyle.cpp
--- a/Engine/Source/Editor/EditorStyle/Private/SlateEditorStyle.cpp
+++ b/Engine/Source/Editor/EditorStyle/Private/SlateEditorStyle.cpp
@@ -5116,6 +5116,7 @@ void FSlateEditorStyle::FStyle::SetupPersonaStyle()
 		// Persona-specific tabs
 		Set("Persona.Tabs.SkeletonTree", new IMAGE_BRUSH(TEXT("Icons/icon_Persona_Skeleton_Tree_16x"), Icon16x16));
 		Set("Persona.Tabs.MorphTargetPreviewer", new IMAGE_BRUSH(TEXT("Icons/icon_Persona_Morph_Target_Previewer_16x"), Icon16x16));
+		Set("Persona.Tabs.AnimCurvePreviewer", new IMAGE_BRUSH(TEXT("Icons/icon_Persona_AnimCurve_Previewer_16x"), Icon16x16));
 		Set("Persona.Tabs.AnimationNotifies", new IMAGE_BRUSH(TEXT("Icons/icon_Persona_Animation_Notifies_16x"), Icon16x16));
 		Set("Persona.Tabs.RetargetManager", new IMAGE_BRUSH(TEXT("Icons/icon_Persona_Retarget_Manager_16x"), Icon16x16));
 		Set("Persona.Tabs.AnimSlotManager", new IMAGE_BRUSH(TEXT("Icons/icon_Persona_Anim_Slot_Manager_16x"), Icon16x16));
@@ -5625,6 +5626,7 @@ void FSlateEditorStyle::FStyle::SetupClassIconsAndThumbnails()
 			TEXT("PlayerController"),
 			TEXT("PlayerStart"),
 			TEXT("PointLight"),
+			TEXT("PoseAsset"),
 			TEXT("PostProcessVolume"),
 			TEXT("PrecomputedVisibilityOverrideVolume"),
 			TEXT("PrecomputedVisibilityVolume"),
diff --git a/Engine/Source/Editor/EnvironmentQueryEditor/Private/EnvironmentQueryEditor.cpp b/Engine/Source/Editor/EnvironmentQueryEditor/Private/EnvironmentQueryEditor.cpp
--- a/Engine/Source/Editor/EnvironmentQueryEditor/Private/EnvironmentQueryEditor.cpp
+++ b/Engine/Source/Editor/EnvironmentQueryEditor/Private/EnvironmentQueryEditor.cpp
@@ -222,7 +222,7 @@ void FEnvironmentQueryEditor::OnSelectedNodesChanged(const TSet<class UObject*>&
 	{
 		for(TSet<class UObject*>::TConstIterator SetIt(NewSelection);SetIt;++SetIt)
 		{
-			UEnvironmentQueryGraphNode* GraphNode = CastChecked<UEnvironmentQueryGraphNode>(*SetIt);
+			UEnvironmentQueryGraphNode* GraphNode = Cast<UEnvironmentQueryGraphNode>(*SetIt);
 			if (GraphNode)
 			{
 				if (GraphNode->IsA(UEnvironmentQueryGraphNode_Root::StaticClass()))
diff --git a/Engine/Source/Editor/GraphEditor/Private/GraphEditorActions.cpp b/Engine/Source/Editor/GraphEditor/Private/GraphEditorActions.cpp
--- a/Engine/Source/Editor/GraphEditor/Private/GraphEditorActions.cpp
+++ b/Engine/Source/Editor/GraphEditor/Private/GraphEditorActions.cpp
@@ -87,6 +87,9 @@ void FGraphEditorCommandsImpl::RegisterCommands()
 	UI_COMMAND( ConvertToBSEvaluator, "Convert To Single Frame BlendSpace", "Convert to one frame BlendSpace that requires position", EUserInterfaceActionType::Button, FInputChord() )
 	UI_COMMAND( ConvertToBSPlayer, "Convert to BlendSpace Player", "Convert back to BlendSpace player without manual position set up", EUserInterfaceActionType::Button, FInputChord() )
 
+	UI_COMMAND(ConvertToPoseBlender, "Convert To Pose Blender", "Convert to pose blender that can blend by source curves", EUserInterfaceActionType::Button, FInputChord())
+	UI_COMMAND(ConvertToPoseByName, "Convert to Pose By Name", "Convert to pose node that returns by name", EUserInterfaceActionType::Button, FInputChord())
+
 	UI_COMMAND( OpenRelatedAsset, "Open Asset", "Opens the asset related to this node", EUserInterfaceActionType::Button, FInputChord() )
 
 	UI_COMMAND( CreateComment, "Create Comment", "Create a comment box", EUserInterfaceActionType::Button, FInputChord(EKeys::C))
diff --git a/Engine/Source/Editor/GraphEditor/Public/GraphEditorActions.h b/Engine/Source/Editor/GraphEditor/Public/GraphEditorActions.h
--- a/Engine/Source/Editor/GraphEditor/Public/GraphEditorActions.h
+++ b/Engine/Source/Editor/GraphEditor/Public/GraphEditorActions.h
@@ -118,6 +118,10 @@ public:
 	TSharedPtr< FUICommandInfo > ConvertToBSEvaluator;
 	TSharedPtr< FUICommandInfo > ConvertToBSPlayer;
 
+	// options for sequence/evaluator converter
+	TSharedPtr< FUICommandInfo > ConvertToPoseBlender;
+	TSharedPtr< FUICommandInfo > ConvertToPoseByName;
+
 	// option for opening the asset related to the graph node
 	TSharedPtr< FUICommandInfo > OpenRelatedAsset;
 
diff --git a/Engine/Source/Editor/Kismet/Private/BlueprintEditor.cpp b/Engine/Source/Editor/Kismet/Private/BlueprintEditor.cpp
--- a/Engine/Source/Editor/Kismet/Private/BlueprintEditor.cpp
+++ b/Engine/Source/Editor/Kismet/Private/BlueprintEditor.cpp
@@ -1158,6 +1158,14 @@ TSharedRef<SGraphEditor> FBlueprintEditor::CreateGraphEditorWidget(TSharedRef<FT
 				FExecuteAction::CreateSP( this, &FBlueprintEditor::OnConvertToBlendSpacePlayer )
 				);
 
+			GraphEditorCommands->MapAction(FGraphEditorCommands::Get().ConvertToPoseBlender,
+				FExecuteAction::CreateSP(this, &FBlueprintEditor::OnConvertToPoseBlender)
+				);
+
+			GraphEditorCommands->MapAction(FGraphEditorCommands::Get().ConvertToPoseByName,
+				FExecuteAction::CreateSP(this, &FBlueprintEditor::OnConvertToPoseByName)
+				);
+
 			GraphEditorCommands->MapAction( FGraphEditorCommands::Get().OpenRelatedAsset,
 				FExecuteAction::CreateSP( this, &FBlueprintEditor::OnOpenRelatedAsset )
 				);
diff --git a/Engine/Source/Editor/Kismet/Private/SCSEditorViewportClient.cpp b/Engine/Source/Editor/Kismet/Private/SCSEditorViewportClient.cpp
--- a/Engine/Source/Editor/Kismet/Private/SCSEditorViewportClient.cpp
+++ b/Engine/Source/Editor/Kismet/Private/SCSEditorViewportClient.cpp
@@ -24,6 +24,7 @@
 #include "Components/InstancedStaticMeshComponent.h"
 #include "Engine/StaticMesh.h"
 #include "SSCSEditorViewport.h"
+#include "PhysicsEngine/PhysicsConstraintComponent.h"
 
 DEFINE_LOG_CATEGORY_STATIC(LogSCSEditorViewport, Log, All);
 
diff --git a/Engine/Source/Editor/Kismet/Public/BlueprintEditor.h b/Engine/Source/Editor/Kismet/Public/BlueprintEditor.h
--- a/Engine/Source/Editor/Kismet/Public/BlueprintEditor.h
+++ b/Engine/Source/Editor/Kismet/Public/BlueprintEditor.h
@@ -795,6 +795,8 @@ protected:
 	virtual void OnConvertToSequencePlayer() {};
 	virtual void OnConvertToBlendSpaceEvaluator() {};
 	virtual void OnConvertToBlendSpacePlayer() {};
+	virtual void OnConvertToPoseBlender() {};
+	virtual void OnConvertToPoseByName() {};
 
 	// Opens the associated asset of the selected nodes
 	virtual void OnOpenRelatedAsset() {};
diff --git a/Engine/Source/Editor/LevelEditor/Private/LevelEditorActions.cpp b/Engine/Source/Editor/LevelEditor/Private/LevelEditorActions.cpp
--- a/Engine/Source/Editor/LevelEditor/Private/LevelEditorActions.cpp
+++ b/Engine/Source/Editor/LevelEditor/Private/LevelEditorActions.cpp
@@ -61,6 +61,7 @@
 #include "Engine/StaticMesh.h"
 #include "Engine/Light.h"
 #include "Animation/SkeletalMeshActor.h"
+#include "Animation/AnimSequence.h"
 #include "Editor/KismetWidgets/Public/CreateBlueprintFromActorDialog.h"
 #include "EditorProjectSettings.h"
 #include "HierarchicalLODUtilities.h"
diff --git a/Engine/Source/Editor/Matinee/Private/MatineeMenus.cpp b/Engine/Source/Editor/Matinee/Private/MatineeMenus.cpp
--- a/Engine/Source/Editor/Matinee/Private/MatineeMenus.cpp
+++ b/Engine/Source/Editor/Matinee/Private/MatineeMenus.cpp
@@ -59,6 +59,7 @@
 #include "Engine/Selection.h"
 #include "Sound/SoundBase.h"
 #include "Camera/CameraAnim.h"
+#include "Animation/AnimSequence.h"
 
 #include "IMenu.h"
 
diff --git a/Engine/Source/Editor/MovieSceneTools/Private/MatineeImportTools.cpp b/Engine/Source/Editor/MovieSceneTools/Private/MatineeImportTools.cpp
--- a/Engine/Source/Editor/MovieSceneTools/Private/MatineeImportTools.cpp
+++ b/Engine/Source/Editor/MovieSceneTools/Private/MatineeImportTools.cpp
@@ -42,6 +42,8 @@
 
 #include "MovieSceneSequence.h"
 
+#include "Animation/AnimSequence.h"
+
 ERichCurveInterpMode FMatineeImportTools::MatineeInterpolationToRichCurveInterpolation( EInterpCurveMode CurveMode )
 {
 	switch ( CurveMode )
diff --git a/Engine/Source/Editor/MovieSceneTools/Private/TrackEditors/TransformTrackEditor.cpp b/Engine/Source/Editor/MovieSceneTools/Private/TrackEditors/TransformTrackEditor.cpp
--- a/Engine/Source/Editor/MovieSceneTools/Private/TrackEditors/TransformTrackEditor.cpp
+++ b/Engine/Source/Editor/MovieSceneTools/Private/TrackEditors/TransformTrackEditor.cpp
@@ -17,7 +17,7 @@
 #include "Matinee/InterpTrackMove.h"
 #include "FloatCurveKeyArea.h"
 #include "ContentBrowserModule.h"
-
+#include "Animation/AnimSequence.h"
 
 #define LOCTEXT_NAMESPACE "MovieScene_TransformTrack"
 
diff --git a/Engine/Source/Editor/Persona/Private/AnimationEditorViewportClient.cpp b/Engine/Source/Editor/Persona/Private/AnimationEditorViewportClient.cpp
--- a/Engine/Source/Editor/Persona/Private/AnimationEditorViewportClient.cpp
+++ b/Engine/Source/Editor/Persona/Private/AnimationEditorViewportClient.cpp
@@ -822,8 +822,8 @@ void FAnimationViewportClient::ShowBoneNames( FCanvas* Canvas, FSceneView* View
 	{
 		const int32 BoneIndex = LODModel.RequiredBones[i];
 
-		// If previewing a specific chunk, only show the bone names that belong to it
-		if ((PreviewSkelMeshComp->ChunkIndexPreview >= 0) && !LODModel.Chunks[PreviewSkelMeshComp->ChunkIndexPreview].BoneMap.Contains(BoneIndex))
+		// If previewing a specific section, only show the bone names that belong to it
+		if ((PreviewSkelMeshComp->SectionIndexPreview >= 0) && !LODModel.Sections[PreviewSkelMeshComp->SectionIndexPreview].BoneMap.Contains(BoneIndex))
 		{
 			continue;
 		}
@@ -1000,7 +1000,6 @@ void FAnimationViewportClient::DisplayInfo(FCanvas* Canvas, FSceneView* View, bo
 
 			int32 NumBonesInUse;
 			int32 NumBonesMappedToVerts;
-			int32 NumChunksInUse;
 			int32 NumSectionsInUse;
 			FString WeightUsage;
 
@@ -1009,7 +1008,6 @@ void FAnimationViewportClient::DisplayInfo(FCanvas* Canvas, FSceneView* View, bo
 
 			NumBonesInUse = LODModel.RequiredBones.Num();
 			NumBonesMappedToVerts = LODModel.ActiveBoneIndices.Num();
-			NumChunksInUse = LODModel.Chunks.Num();
 			NumSectionsInUse = LODModel.Sections.Num();
 
 			// Calculate polys based on non clothing sections so we don't duplicate the counts.
@@ -1034,32 +1032,22 @@ void FAnimationViewportClient::DisplayInfo(FCanvas* Canvas, FSceneView* View, bo
 
 			CurYOffset += 1; // --
 
-			uint32 TotalRigidVertices = 0;
-			uint32 TotalSoftVertices = 0;
-			for (int32 ChunkIndex = 0; ChunkIndex < LODModel.Chunks.Num(); ChunkIndex++)
+			for (int32 SectionIndex = 0; SectionIndex < LODModel.Sections.Num(); SectionIndex++)
 			{
-				int32 ChunkRigidVerts = LODModel.Chunks[ChunkIndex].GetNumRigidVertices();
-				int32 ChunkSoftVerts = LODModel.Chunks[ChunkIndex].GetNumSoftVertices();
-
-				InfoString = FString::Printf(TEXT(" [Chunk %d] Verts:%d (Rigid:%d Soft:%d), Bones:%d"),
-					ChunkIndex,
-					ChunkRigidVerts + ChunkSoftVerts,
-					ChunkRigidVerts,
-					ChunkSoftVerts,
-					LODModel.Chunks[ChunkIndex].BoneMap.Num()
+				int32 SectionVerts = LODModel.Sections[SectionIndex].GetNumVertices();
+
+				InfoString = FString::Printf(TEXT(" [Section %d] Verts:%d, Bones:%d"),
+					SectionIndex,
+					SectionVerts,
+					LODModel.Sections[SectionIndex].BoneMap.Num()
 					);
 
 				CurYOffset += YL + 2;
 				Canvas->DrawShadowedString(CurXOffset, CurYOffset, *InfoString, GEngine->GetSmallFont(), TextColor*0.8f);
-
-				TotalRigidVertices += ChunkRigidVerts;
-				TotalSoftVertices += ChunkSoftVerts;
 			}
 
-			InfoString = FString::Printf(TEXT("TOTAL Verts:%d (Rigid:%d Soft:%d)"),
-				LODModel.NumVertices,
-				TotalRigidVertices,
-				TotalSoftVertices);
+			InfoString = FString::Printf(TEXT("TOTAL Verts:%d"),
+				LODModel.NumVertices);
 
 			CurYOffset += 1; // --
 
@@ -1067,8 +1055,7 @@ void FAnimationViewportClient::DisplayInfo(FCanvas* Canvas, FSceneView* View, bo
 			CurYOffset += YL + 2;
 			Canvas->DrawShadowedString(CurXOffset, CurYOffset, *InfoString, GEngine->GetSmallFont(), TextColor);
 
-			InfoString = FString::Printf(TEXT("Chunks:%d, Sections:%d %s"),
-				NumChunksInUse,
+			InfoString = FString::Printf(TEXT("Sections:%d %s"),
 				NumSectionsInUse,
 				*WeightUsage
 				);
diff --git a/Engine/Source/Editor/Persona/Private/BlueprintMode/AnimBlueprintMode.cpp b/Engine/Source/Editor/Persona/Private/BlueprintMode/AnimBlueprintMode.cpp
--- a/Engine/Source/Editor/Persona/Private/BlueprintMode/AnimBlueprintMode.cpp
+++ b/Engine/Source/Editor/Persona/Private/BlueprintMode/AnimBlueprintMode.cpp
@@ -1,6 +1,9 @@
 // Copyright 1998-2016 Epic Games, Inc. All Rights Reserved.
 
 #include "PersonaPrivatePCH.h"
+
+#include "Animation/AnimInstance.h"
+
 #include "Persona.h"
 
 #include "Editor/PropertyEditor/Public/PropertyEditorModule.h"
diff --git a/Engine/Source/Editor/Persona/Private/BoneMappingHelper.cpp b/Engine/Source/Editor/Persona/Private/BoneMappingHelper.cpp
new file mode 100644
--- /dev/null
+++ b/Engine/Source/Editor/Persona/Private/BoneMappingHelper.cpp
@@ -0,0 +1,322 @@
+// Copyright 1998-2016 Epic Games, Inc. All Rights Reserved.
+
+
+#include "PersonaPrivatePCH.h"
+#include "BoneMappingHelper.h"
+#include "AnimationRuntime.h"
+
+//////////////////////////////////////////////////////////////////////////
+// FBoneDescription
+//////////////////////////////////////////////////////////////////////////
+
+int32 FBoneDescription::GetBestIndex() const
+{
+	int32 BestIndex = INDEX_NONE;
+	float BestScore = 0.f;
+
+	for (int32 Index = 0; Index < Scores.Num(); ++Index)
+	{
+		if (Scores[Index] > BestScore)
+		{
+			BestScore = Scores[Index];
+			BestIndex = Index;
+		}
+	}
+
+	return BestIndex;
+}
+
+float FBoneDescription::CalculateNameScore(const FName& Name1, const FName& Name2) const
+{
+	FString String1 = Name1.ToString();
+	FString String2 = Name2.ToString();
+	int32 Distance = FAnimationRuntime::GetStringDistance(String1, String2);
+	float MaxLength = FMath::Max(String1.Len() , String2.Len());
+	return (MaxLength) ? (MaxLength - Distance) / (MaxLength) : 1.f;
+}
+
+float FBoneDescription::CalculateScore(const FBoneDescription& Other) const
+{
+	// if they don't have parent, it's root, so just give whole score
+	if (Other.BoneInfo.ParentIndex == INDEX_NONE && BoneInfo.ParentIndex == INDEX_NONE)
+	{
+		return 1.f;
+	}
+
+	// calculating score is tricky
+	float Score = 0.f;
+
+	// each element will exit from [0, 1], and we'll apply weight to it
+
+	// check direction of facing [-1, 1]
+	float Cosine = FVector::DotProduct(DirFromParent, Other.DirFromParent);
+	float Score_DirFromParent = FMath::Max((Cosine-0.5f)*2.f, 0.f); // scale to only care for range of [0.5-1]
+	ensure(Score_DirFromParent >= 0.f && Score_DirFromParent <= 1.f);
+
+	// check direction of facing [-1, 1]
+	Cosine = FVector::DotProduct(DirFromRoot, Other.DirFromRoot);
+	float Score_DirFromRoot = FMath::Max((Cosine - 0.5f)*2.f, 0.f); // scale to only care for range of [0.5-1]
+	ensure(Score_DirFromRoot >= 0.f && Score_DirFromRoot <= 1.f);
+
+	// check number of children
+	float Score_NumChildren = 0.f;
+	const int32 MaxNumChildren = FMath::Max(Other.NumChildren, NumChildren);
+	if (MaxNumChildren > 0)
+	{
+		Score_NumChildren = (1.f - (FMath::Abs(Other.NumChildren - NumChildren) / MaxNumChildren));
+	}
+	else
+	{
+		// no children - leaf node
+		Score_NumChildren = 1.f;
+	}
+
+	ensure(Score_NumChildren >= 0.f && Score_NumChildren <= 1.f);
+
+	// score of ratio from parent  - if you're here, you should have parent, so it shouldn't be 0.f, most likely
+	float Score_RatioFromParent = 0.f;
+	if (Other.RatioFromParent > RatioFromParent)
+	{
+		Score_RatioFromParent = (Other.RatioFromParent > 0.f) ? RatioFromParent / Other.RatioFromParent : 0.f;
+	}
+	else
+	{
+		Score_RatioFromParent = (RatioFromParent > 0.f) ? Other.RatioFromParent / RatioFromParent : 0.f;
+	}
+
+	ensure(Score_RatioFromParent >= 0.f && Score_RatioFromParent <= 1.f);
+
+	// check normalized position - since this is normalized, it should stay within 1
+	FVector DiffNormalizedPosition = (Other.NormalizedPosition - NormalizedPosition).GetAbs();
+	const float MaxNormalizedPosition = 3.f; /* since 1^2+1^2+1^2 = 3*/
+	float Score_NormalizedPosition = (MaxNormalizedPosition - DiffNormalizedPosition.SizeSquared())/ MaxNormalizedPosition;
+	ensure(Score_NormalizedPosition >= 0.f && Score_NormalizedPosition <= 1.f);
+
+	float Score_NameMatching = CalculateNameScore(BoneInfo.Name, Other.BoneInfo.Name);
+	ensure(Score_NameMatching >= 0.f && Score_NameMatching <= 1.f);
+
+	// now come up with full score
+	static float Weight_DirFromParent = 2.f;
+	static float Weight_NumChildren = 0.5f;
+	static float Weight_NormalizedPosition = 1.0f; // location can be very confusing, so give less weight on this
+	static float Weight_RatioFromParent = 1.f;
+	static float Weight_NameMatching = 2.0f;
+	static float Weight_DirFromRoot = 0.f;
+
+	float FinalScore = (Score_DirFromParent * Weight_DirFromParent + Score_NumChildren * Weight_NumChildren + Score_NormalizedPosition * Weight_NormalizedPosition + Score_RatioFromParent * Weight_RatioFromParent + Score_NameMatching * Weight_NameMatching + Score_DirFromRoot * Weight_DirFromRoot) / (Weight_DirFromParent + Weight_NumChildren + Weight_NormalizedPosition + Weight_RatioFromParent + Weight_NameMatching + Weight_DirFromRoot);
+	UE_LOG(LogAnimation, Log, TEXT("Calculate Score - [%s] - [%s] (Score_DirFromParent(%0.2f), Score_NumChildren(%0.2f), Score_NormalizedPosition(%0.2f), Score_RatioFromParent(%0.2f), Score_NameMatching(%0.2f), Score_DirFromRoot(%0.2f) )"), *BoneInfo.Name.ToString(), *Other.BoneInfo.Name.ToString(), Score_DirFromParent, Score_NumChildren, Score_NormalizedPosition, Score_RatioFromParent, Score_NameMatching, Score_DirFromRoot);
+	return FinalScore;
+}
+
+//////////////////////////////////////////////////////////////////////////
+// FBoneMappingHelper
+//////////////////////////////////////////////////////////////////////////
+
+FBoneMappingHelper::FBoneMappingHelper(const FReferenceSkeleton& InRefSkeleton1, const FReferenceSkeleton&  InRefSkeleton2)
+{
+	Initialize(0, InRefSkeleton1);
+	Initialize(1, InRefSkeleton2);
+}
+
+void FBoneMappingHelper::Initialize(int32 Index, const FReferenceSkeleton&  InRefSkeleton)
+{
+	int32 TotalNum = InRefSkeleton.GetNum();
+	TArray<FBoneDescription>& BoneDescList = BoneDescs[Index];
+
+	BoneDescList.Reset(TotalNum);
+	RefSkeleton[Index] = InRefSkeleton;
+
+	if (TotalNum > 0)
+	{
+		BoneDescList.AddZeroed(TotalNum);
+
+		// fill up spacebases
+		TArray<FTransform> SpaceBases;
+		FAnimationRuntime::FillUpSpaceBases(InRefSkeleton, InRefSkeleton.GetRefBonePose(), SpaceBases);
+
+		// get extent
+		FBox BoxExtent;
+		BoxExtent.Init();
+
+		for (int32 BoneIndex = 0; BoneIndex < TotalNum; ++BoneIndex)
+		{
+			BoxExtent += SpaceBases[BoneIndex].GetLocation();
+		}
+
+		const FVector MeshBoxSize = BoxExtent.GetSize();
+		const FVector RootPosition = SpaceBases[0].GetTranslation();
+
+		if (MeshBoxSize.GetMin() > SMALL_NUMBER)
+		{
+			// now fill up data
+			const TArray<FMeshBoneInfo>& MeshInfoList = InRefSkeleton.GetRefBoneInfo();
+			for (int32 BoneIndex = 0; BoneIndex < TotalNum; ++BoneIndex)
+			{
+				FBoneDescription& BoneDesc = BoneDescList[BoneIndex];
+				BoneDesc.BoneInfo = MeshInfoList[BoneIndex];
+				BoneDesc.NumChildren = 0;
+
+				// normalized position
+				const FVector& Position = SpaceBases[BoneIndex].GetLocation();
+				BoneDesc.NormalizedPosition = (Position - BoxExtent.Min) / MeshBoxSize;
+
+				// increase parent child
+				if (BoneDesc.BoneInfo.ParentIndex != INDEX_NONE)
+				{
+					const int32 ParentIndex = BoneDesc.BoneInfo.ParentIndex;
+					++BoneDescList[ParentIndex].NumChildren;
+
+					const FVector& ParentPosition = SpaceBases[ParentIndex].GetLocation();
+					FVector ToChild = (Position-ParentPosition);
+					BoneDesc.RatioFromParent = ToChild.Size() / MeshBoxSize.Size();
+					BoneDesc.DirFromParent = ToChild.GetSafeNormal();
+					BoneDesc.DirFromRoot = (Position-RootPosition).GetSafeNormal(); // based on whole mesh size
+				}
+			}
+		}
+	}
+}
+
+void FBoneMappingHelper::TryMatch(TMap<FName, FName>& OutBestMatches)
+{
+	TArray<FBoneDescription>& BoneDescArray0 = BoneDescs[0];
+	TArray<FBoneDescription>& BoneDescArray1 = BoneDescs[1];
+
+	for (int32 BoneIndex0 = 0; BoneIndex0 < BoneDescArray0.Num(); ++BoneIndex0)
+	{
+		BoneDescArray0[BoneIndex0].ResetScore(BoneDescArray1.Num());
+	}
+
+	for (int32 BoneIndex1 = 0; BoneIndex1 < BoneDescArray1.Num(); ++BoneIndex1)
+	{
+		BoneDescArray1[BoneIndex1].ResetScore(BoneDescArray0.Num());
+	}
+
+	for (int32 BoneIndex0 = 0; BoneIndex0 < BoneDescArray0.Num(); ++BoneIndex0)
+	{
+		FBoneDescription& BoneDesc0 = BoneDescArray0[BoneIndex0];
+		for (int32 BoneIndex1 = 0; BoneIndex1 < BoneDescArray1.Num(); ++BoneIndex1)
+		{
+			FBoneDescription& BoneDesc1 = BoneDescArray1[BoneIndex1];
+
+			// set score in both container
+			float Score = BoneDesc0.CalculateScore(BoneDesc1);
+			BoneDesc0.SetScore(BoneIndex1, Score);
+			BoneDesc1.SetScore(BoneIndex0, Score);
+		}
+	}
+
+#define MAX_CANDIDATE 10
+
+	// first find best matches up to MAX_CANDIDATE for each
+	struct FCandidate
+	{
+		FName BoneNames[MAX_CANDIDATE];
+		float Scores[MAX_CANDIDATE];
+		float StdDev;
+
+		FCandidate()
+		{
+			FMemory::Memzero(Scores);
+			StdDev = 0.f;
+		}
+
+		void Set(FName InBoneName, float InScore, int32 Index)
+		{
+			BoneNames[Index] = InBoneName;
+			Scores[Index] = InScore;
+		}
+
+		void CalculateStdDev()
+		{
+			float Avg = 0.f;
+
+			for (int32 Index = 0; Index < MAX_CANDIDATE; ++Index)
+			{
+				Avg += Scores[Index];
+			}
+
+			Avg /= MAX_CANDIDATE;
+			
+			float AccumulatedDev = 0.f;
+
+			for (int32 Index = 0; Index < MAX_CANDIDATE; ++Index)
+			{
+				AccumulatedDev += FMath::Square((Scores[Index]) - Avg);
+			}
+
+			StdDev = FGenericPlatformMath::Sqrt(AccumulatedDev/MAX_CANDIDATE);
+		}
+	};
+
+	TMap<FName, FCandidate> Candidates;
+
+	// find the best score
+	for (int32 BoneIndex0 = 0; BoneIndex0 < BoneDescArray0.Num(); ++BoneIndex0)
+	{
+		int32 BoneIndex1 = BoneDescArray0[BoneIndex0].GetBestIndex();
+
+		if (BoneIndex1 != INDEX_NONE)
+		{
+			FCandidate Candidate;
+			FName Bone0Name = BoneDescArray0[BoneIndex0].BoneInfo.Name;
+			FName Bone1Name = BoneDescArray1[BoneIndex1].BoneInfo.Name;
+			float Score = BoneDescArray0[BoneIndex0].GetScore(BoneIndex1);
+			Candidate.Set(Bone1Name, Score, 0);
+
+			UE_LOG(LogAnimation, Log, TEXT("Bone Match [%s] - [%s] (score %0.2f)"), *Bone0Name.ToString(), *Bone1Name.ToString(), Score);
+
+			// print next 3 best score, clear current one
+			for (int32 LoopCount = 0; LoopCount < MAX_CANDIDATE-1; ++LoopCount)
+			{
+				BoneDescArray0[BoneIndex0].SetScore(BoneIndex1, 0.f);
+
+				BoneIndex1 = BoneDescArray0[BoneIndex0].GetBestIndex();
+				if (BoneIndex1 != INDEX_NONE)
+				{
+					Bone1Name = BoneDescArray1[BoneIndex1].BoneInfo.Name;
+					Score = BoneDescArray0[BoneIndex0].GetScore(BoneIndex1);
+
+					Candidate.Set(Bone1Name, Score, LoopCount + 1);
+					UE_LOG(LogAnimation, Log, TEXT(" Candidate %d. - [%s] (score %0.2f)"), LoopCount + 1, *Bone1Name.ToString(), Score);
+				}
+			}
+
+			Candidate.CalculateStdDev();
+			Candidates.Add(Bone0Name, Candidate);
+		}
+		else
+		{
+			UE_LOG(LogAnimation, Log, TEXT("Bone [%s] does not have a match"), *BoneDescArray0[BoneIndex0].BoneInfo.Name.ToString());
+		}
+	}
+
+	struct FCandidateSortCallback
+	{
+		FORCEINLINE bool operator()(const FCandidate& A, const FCandidate& B) const
+		{
+			return A.StdDev > B.StdDev;
+		}
+	};
+
+	Candidates.ValueSort(FCandidateSortCallback());
+
+	TArray<FName> UsedNames;
+	for (const TPair<FName, FCandidate>& CandidatePair : Candidates)
+	{
+		const FCandidate& Candidate = CandidatePair.Value;
+		for (int32 Index = 0; Index < MAX_CANDIDATE; ++Index)
+		{
+			FName BestMatchName = Candidate.BoneNames[Index];
+			// see if it's already used by other joint
+			if (UsedNames.Find(BestMatchName) == INDEX_NONE)
+			{
+				// if that's case, just ignore and move on
+				OutBestMatches.Add(CandidatePair.Key, BestMatchName);
+				UsedNames.Add(BestMatchName);
+				break;
+			}
+		}
+	}
+}
+
diff --git a/Engine/Source/Editor/Persona/Private/Persona.cpp b/Engine/Source/Editor/Persona/Private/Persona.cpp
--- a/Engine/Source/Editor/Persona/Private/Persona.cpp
+++ b/Engine/Source/Editor/Persona/Private/Persona.cpp
@@ -21,6 +21,7 @@
 #include "SAnimationBlendSpace1D.h"
 #include "SKismetInspector.h"
 #include "SSkeletonWidget.h"
+#include "SPoseEditor.h"
 
 #include "Editor/Kismet/Public/BlueprintEditorTabs.h"
 #include "Editor/Kismet/Public/BlueprintEditorModes.h"
@@ -46,6 +47,8 @@
 #include "AnimGraphNode_LayeredBoneBlend.h"
 #include "AnimGraphNode_SequencePlayer.h"
 #include "AnimGraphNode_SequenceEvaluator.h"
+#include "AnimGraphNode_PoseByName.h"
+#include "AnimGraphNode_PoseBlendNode.h"
 #include "AnimGraphNode_Slot.h"
 #include "Customization/AnimGraphNodeSlotDetails.h"
 
@@ -72,6 +75,7 @@
 #include "Animation/AimOffsetBlendSpace.h"
 #include "Animation/AimOffsetBlendSpace1D.h"
 #include "Animation/AnimNotifies/AnimNotifyState.h"
+#include "Animation/PoseAsset.h"
 
 #include "MessageLog.h"
 
@@ -323,6 +327,14 @@ TSharedPtr<SWidget> FPersona::CreateEditorWidgetForAnimDocument(UObject* InAnimA
 
 			DocumentLink = TEXT("Engine/Animation/AnimMontage");
 		}
+		else if (UPoseAsset* PoseAsset = Cast<UPoseAsset>(InAnimAsset))
+		{
+			Result = SNew(SPoseEditor)
+				.Persona(SharedThis(this))
+				.PoseAsset(PoseAsset);
+
+			DocumentLink = TEXT("Engine/Animation/Sequences");
+		}
 		else if (UBlendSpace* BlendSpace = Cast<UBlendSpace>(InAnimAsset))
 		{
 			Result = SNew(SBlendSpaceEditor)
@@ -1663,6 +1675,113 @@ void FPersona::OnConvertToBlendSpacePlayer()
 	}
 }
 
+void FPersona::OnConvertToPoseBlender()
+{
+	FGraphPanelSelectionSet SelectedNodes = GetSelectedNodes();
+	if (SelectedNodes.Num() > 0)
+	{
+		for (auto NodeIter = SelectedNodes.CreateIterator(); NodeIter; ++NodeIter)
+		{
+			UAnimGraphNode_PoseByName* OldNode = Cast<UAnimGraphNode_PoseByName>(*NodeIter);
+
+			// see if sequence player
+			if (OldNode && OldNode->Node.PoseAsset)
+			{
+				// convert to sequence player
+				UEdGraph* TargetGraph = OldNode->GetGraph();
+				// create new player
+				FGraphNodeCreator<UAnimGraphNode_PoseBlendNode> NodeCreator(*TargetGraph);
+				UAnimGraphNode_PoseBlendNode* NewNode = NodeCreator.CreateNode();
+				NewNode->Node.PoseAsset = OldNode->Node.PoseAsset;
+				NodeCreator.Finalize();
+
+				// get default data from old node to new node
+				FEdGraphUtilities::CopyCommonState(OldNode, NewNode);
+
+				UEdGraphPin* OldPosePin = OldNode->FindPin(TEXT("Pose"));
+				UEdGraphPin* NewPosePin = NewNode->FindPin(TEXT("Pose"));
+
+				if (ensure(OldPosePin && NewPosePin))
+				{
+					NewPosePin->CopyPersistentDataFromOldPin(*OldPosePin);
+				}
+
+				// remove from selection and from graph
+				NodeIter.RemoveCurrent();
+				TargetGraph->RemoveNode(OldNode);
+
+				NewNode->Modify();
+			}
+		}
+
+		// @todo fixme: below code doesn't work
+		// because of SetAndCenterObject kicks in after new node is added
+		// will need to disable that first
+		TSharedPtr<SGraphEditor> FocusedGraphEd = FocusedGraphEdPtr.Pin();
+
+		// Update the graph so that the node will be refreshed
+		FocusedGraphEd->NotifyGraphChanged();
+		// It's possible to leave invalid objects in the selection set if they get GC'd, so clear it out
+		FocusedGraphEd->ClearSelectionSet();
+
+		FBlueprintEditorUtils::MarkBlueprintAsStructurallyModified(GetAnimBlueprint());
+	}
+}
+
+void FPersona::OnConvertToPoseByName()
+{
+	FGraphPanelSelectionSet SelectedNodes = GetSelectedNodes();
+	if (SelectedNodes.Num() > 0)
+	{
+		for (auto NodeIter = SelectedNodes.CreateIterator(); NodeIter; ++NodeIter)
+		{
+			UAnimGraphNode_PoseBlendNode* OldNode = Cast<UAnimGraphNode_PoseBlendNode>(*NodeIter);
+
+			// see if sequence player
+			if (OldNode && OldNode->Node.PoseAsset)
+			{
+				//const FScopedTransaction Transaction( LOCTEXT("ConvertToSequenceEvaluator", "Convert to Single Frame Animation") );
+				// convert to sequence player
+				UEdGraph* TargetGraph = OldNode->GetGraph();
+				// create new player
+				FGraphNodeCreator<UAnimGraphNode_PoseByName> NodeCreator(*TargetGraph);
+				UAnimGraphNode_PoseByName* NewNode = NodeCreator.CreateNode();
+				NewNode->Node.PoseAsset = OldNode->Node.PoseAsset;
+				NodeCreator.Finalize();
+
+				// get default data from old node to new node
+				FEdGraphUtilities::CopyCommonState(OldNode, NewNode);
+
+				UEdGraphPin* OldPosePin = OldNode->FindPin(TEXT("Pose"));
+				UEdGraphPin* NewPosePin = NewNode->FindPin(TEXT("Pose"));
+
+				if (ensure(OldPosePin && NewPosePin))
+				{
+					NewPosePin->CopyPersistentDataFromOldPin(*OldPosePin);
+				}
+
+				// remove from selection and from graph
+				NodeIter.RemoveCurrent();
+				TargetGraph->RemoveNode(OldNode);
+
+				NewNode->Modify();
+			}
+		}
+
+		// @todo fixme: below code doesn't work
+		// because of SetAndCenterObject kicks in after new node is added
+		// will need to disable that first
+		TSharedPtr<SGraphEditor> FocusedGraphEd = FocusedGraphEdPtr.Pin();
+
+		// Update the graph so that the node will be refreshed
+		FocusedGraphEd->NotifyGraphChanged();
+		// It's possible to leave invalid objects in the selection set if they get GC'd, so clear it out
+		FocusedGraphEd->ClearSelectionSet();
+
+		FBlueprintEditorUtils::MarkBlueprintAsStructurallyModified(GetAnimBlueprint());
+	}
+}
+
 void FPersona::OnOpenRelatedAsset()
 {
 	FGraphPanelSelectionSet SelectedNodes = GetSelectedNodes();
@@ -3405,7 +3524,7 @@ void FPersona::TestSkeletonCurveNamesForUse() const
 						TSharedPtr<FTokenizedMessage> Message;
 						for (FFloatCurve& Curve : Seq->RawCurveData.FloatCurves)
 						{
-							if (UnusedNames.Contains(Curve.LastObservedName))
+							if (UnusedNames.Contains(Curve.Name.DisplayName))
 							{
 								bFoundIssue = true;
 								if (!Message.IsValid())
@@ -3414,7 +3533,7 @@ void FPersona::TestSkeletonCurveNamesForUse() const
 									Message->AddToken(FAssetNameToken::Create(Anim.ObjectPath.ToString(), FText::FromName(Anim.AssetName)));
 									Message->AddToken(FTextToken::Create(LOCTEXT("VerifyCurves_FoundAnimationsWithUnusedReferences", "References the following curves that are not used for either morph targets or material parameters and so may be unneeded")));
 								}
-								CurveOutput.Info(FText::FromName(Curve.LastObservedName));
+								CurveOutput.Info(FText::FromName(Curve.Name.DisplayName));
 							}
 						}
 					}
diff --git a/Engine/Source/Editor/Persona/Private/Persona.h b/Engine/Source/Editor/Persona/Private/Persona.h
--- a/Engine/Source/Editor/Persona/Private/Persona.h
+++ b/Engine/Source/Editor/Persona/Private/Persona.h
@@ -254,6 +254,8 @@ protected:
 	virtual void OnConvertToSequencePlayer() override;
 	virtual void OnConvertToBlendSpaceEvaluator() override;
 	virtual void OnConvertToBlendSpacePlayer() override;
+	virtual void OnConvertToPoseBlender() override;
+	virtual void OnConvertToPoseByName() override;
 	virtual bool IsInAScriptingMode() const override;
 	virtual void OnOpenRelatedAsset() override;
 	virtual void GetCustomDebugObjects(TArray<FCustomDebugObject>& DebugList) const override;
diff --git a/Engine/Source/Editor/Persona/Private/PersonaMeshDetails.cpp b/Engine/Source/Editor/Persona/Private/PersonaMeshDetails.cpp
--- a/Engine/Source/Editor/Persona/Private/PersonaMeshDetails.cpp
+++ b/Engine/Source/Editor/Persona/Private/PersonaMeshDetails.cpp
@@ -1140,7 +1140,7 @@ TSharedRef<SWidget> FPersonaMeshDetails::OnGenerateCustomMaterialWidgetsForMater
 				SNew(STextBlock)
 				.Font(FEditorStyle::GetFontStyle("StaticMeshEditor.NormalFont"))
 				.Text(LOCTEXT("RecomputeTangent_Title", "Recompute Tangent"))
-				.ToolTipText(LOCTEXT("RecomputeTangent_Tooltip", "This feature only works if you enable skin cache (r.SkinCaching) and recompute tangent console variable(r.SkinCache.RecomputeTangents). Please note that skin cache is an experimental feature and only works if you have compute shaders."))
+				.ToolTipText(LOCTEXT("RecomputeTangent_Tooltip", "This feature only works if you enable skin cache (r.SkinCache.Mode) and recompute tangent console variable(r.SkinCache.RecomputeTangents). Please note that skin cache is an experimental feature and only works if you have compute shaders."))
 			]
 		]
 
@@ -2113,8 +2113,8 @@ void FPersonaMeshDetails::UpdateComboBoxStrings()
 			int32 ClothSection = LODModel.Sections[SecIdx].CorrespondClothSectionIndex;
 			if (ClothSection >= 0)
 			{
-				FSkelMeshChunk& Chunk = LODModel.Chunks[LODModel.Sections[ClothSection].ChunkIndex];
-				SectionAssetSubmeshIndices[SecIdx] = FClothAssetSubmeshIndex(Chunk.CorrespondClothAssetIndex, Chunk.ClothAssetSubmeshIndex);
+				FSkelMeshSection& Section = LODModel.Sections[ClothSection];
+				SectionAssetSubmeshIndices[SecIdx] = FClothAssetSubmeshIndex(Section.CorrespondClothAssetIndex, Section.ClothAssetSubmeshIndex);
 			}
 		}
 
@@ -2200,8 +2200,8 @@ void FPersonaMeshDetails::UpdateComboBoxStrings()
 				// Material index is Slot index
 				int32 SlotIdx = LODModel.Sections[SectionIdx].MaterialIndex;
 				check(SlotIdx < SectionAssetSubmeshIndices.Num());
-				FSkelMeshChunk& Chunk = LODModel.Chunks[LODModel.Sections[ClothSection].ChunkIndex];
-				SectionAssetSubmeshIndices[SlotIdx] = FClothAssetSubmeshIndex(Chunk.CorrespondClothAssetIndex, Chunk.ClothAssetSubmeshIndex);
+				FSkelMeshSection& Section = LODModel.Sections[ClothSection];
+				SectionAssetSubmeshIndices[SlotIdx] = FClothAssetSubmeshIndex(Section.CorrespondClothAssetIndex, Section.ClothAssetSubmeshIndex);
 			}
 		}
 
diff --git a/Engine/Source/Editor/Persona/Private/SAnimCompositeEditor.h b/Engine/Source/Editor/Persona/Private/SAnimCompositeEditor.h
--- a/Engine/Source/Editor/Persona/Private/SAnimCompositeEditor.h
+++ b/Engine/Source/Editor/Persona/Private/SAnimCompositeEditor.h
@@ -59,7 +59,7 @@ public:
 	
 	/** Return the animation composite being edited */
 	UAnimComposite* GetCompositeObj() const { return CompositeObj; }
-	virtual UAnimSequenceBase* GetEditorObject() const override { return GetCompositeObj(); }
+	virtual UAnimationAsset* GetEditorObject() const override { return GetCompositeObj(); }
 
 private:
 	/** Pointer to the animation composite being edited */
diff --git a/Engine/Source/Editor/Persona/Private/SAnimCurvePanel.cpp b/Engine/Source/Editor/Persona/Private/SAnimCurvePanel.cpp
--- a/Engine/Source/Editor/Persona/Private/SAnimCurvePanel.cpp
+++ b/Engine/Source/Editor/Persona/Private/SAnimCurvePanel.cpp
@@ -11,6 +11,7 @@
 #include "Kismet2NameValidators.h"
 #include "SExpandableArea.h"
 #include "STextEntryPopup.h"
+#include "Animation/AnimSequence.h"
 
 #define LOCTEXT_NAMESPACE "AnimCurvePanel"
 
@@ -28,8 +29,7 @@ private:
 		FAnimCurveBase* CurrentCurveData = RawCurveData.GetCurveData(CurveUID);
 		if (CurrentCurveData)
 		{
-			CurrentCurveData->CurveUid = RequestedNameUID;
-			CurrentCurveData->LastObservedName = RequestedName;
+			CurrentCurveData->Name.DisplayName = RequestedName;
 		}
 	}
 
@@ -94,14 +94,10 @@ public:
 	{
 		if (AnimSequenceBase.IsValid())
 		{
-			const FSmartNameMapping* NameMapping = AnimSequenceBase.Get()->GetSkeleton()->GetSmartNameContainer(USkeleton::AnimCurveMappingName);
-			if(NameMapping)
+			FSmartName CurveName;
+			if(AnimSequenceBase.Get()->GetSkeleton()->GetSmartNameByUID(USkeleton::AnimCurveMappingName, Uid, CurveName))
 			{
-				FName CurveName;
-				if(NameMapping->GetName(Uid, CurveName))
-				{
-					return FText::FromName(CurveName);
-				}
+				return FText::FromName(CurveName.DisplayName);
 			}
 		}
 		return FText::GetEmpty();
@@ -330,7 +326,7 @@ void SCurveEdTrack::Construct(const FArguments& InArgs)
 				.IsEnabled(true)
 				.Font(FEditorStyle::GetFontStyle("CurveEd.InfoFont"))
 				.SelectAllTextWhenFocused(true)
-				.Text(this, &SCurveEdTrack::GetCurveName, Curve->CurveUid)
+				.Text(this, &SCurveEdTrack::GetCurveName, Curve->Name.UID)
 				.OnTextCommitted(this, &SCurveEdTrack::NewCurveNameEntered)
 			]
 		]
@@ -467,7 +463,7 @@ void SCurveEdTrack::DeleteTrack()
 
 void SCurveEdTrack::ToggleCurveMode(ECheckBoxState NewState,EAnimCurveFlags ModeToSet)
 {
-	const int32 AllModes = (ACF_DrivesMorphTarget|ACF_DrivesMaterial);
+	const int32 AllModes = (ACF_DrivesMorphTarget|ACF_DrivesMaterial|ACF_DrivesPose);
 	check((ModeToSet&AllModes) != 0); //unexpected value for ModeToSet
 
 	FText UndoLabel;
@@ -484,6 +480,10 @@ void SCurveEdTrack::ToggleCurveMode(ECheckBoxState NewState,EAnimCurveFlags Mode
 		{
 			UndoLabel = LOCTEXT("AnimCurve_TurnOnMaterialMode", "Enable driving of materials");
 		}
+		else if (ModeToSet == ACF_DrivesPose)
+		{
+			UndoLabel = LOCTEXT("AnimCurve_TurnOnPoseMode", "Enable driving of pose");
+		}
 	}
 	else
 	{
@@ -495,6 +495,10 @@ void SCurveEdTrack::ToggleCurveMode(ECheckBoxState NewState,EAnimCurveFlags Mode
 		{
 			UndoLabel = LOCTEXT("AnimCurve_TurnOffMaterialMode", "Disable driving of materials");
 		}
+		else if (ModeToSet == ACF_DrivesPose)
+		{
+			UndoLabel = LOCTEXT("AnimCurve_TurnOffPoseMode", "Disable driving of pose");
+		}
 	}
 
 	const FScopedTransaction Transaction( UndoLabel );
@@ -566,7 +570,7 @@ public:
 		FName CurveName;
 		for(FFloatCurve& Curve : Tracks.FloatCurves)
 		{
-			if(NameMapping->GetName(Curve.CurveUid, CurveName))
+			if(NameMapping->GetName(Curve.Name.UID, CurveName))
 			{
 				Names.Add(CurveName.ToString());
 			}
@@ -712,11 +716,11 @@ void SAnimCurvePanel::CreateTrack(const FText& ComittedText, ETextCommit::Type C
 		if(Skeleton && !ComittedText.IsEmpty())
 		{
 			const FScopedTransaction Transaction(LOCTEXT("AnimCurve_AddTrack", "Add New Curve"));
-			USkeleton::AnimCurveUID CurveUid;
+			FSmartName NewTrackName;
 
-			if(Skeleton->AddSmartNameAndModify(USkeleton::AnimCurveMappingName, FName(*ComittedText.ToString()), CurveUid))
+			if(Skeleton->AddSmartNameAndModify(USkeleton::AnimCurveMappingName, FName(*ComittedText.ToString()), NewTrackName))
 			{
-				AddVariableCurve(CurveUid);
+				AddVariableCurve(NewTrackName.UID);
 			}
 		}
 
@@ -730,7 +734,6 @@ FReply SAnimCurvePanel::DuplicateTrack(USkeleton::AnimCurveUID Uid)
 	
 	const FSmartNameMapping* NameMapping = Sequence->GetSkeleton()->GetSmartNameContainer(USkeleton::AnimCurveMappingName);
 	FName CurveNameToCopy;
-	USkeleton::AnimCurveUID NewUid;
 
 	// Must have a curve that exists to duplicate
 	if(NameMapping->Exists(Uid))
@@ -741,9 +744,11 @@ FReply SAnimCurvePanel::DuplicateTrack(USkeleton::AnimCurveUID Uid)
 		// Use the validator to pick a reasonable name for the duplicated curve.
 		FString NewCurveName = CurveNameToCopy.ToString();
 		Validator->FindValidString(NewCurveName);
-		if (Sequence->GetSkeleton()->AddSmartNameAndModify(USkeleton::AnimCurveMappingName, *NewCurveName, NewUid))
+		FSmartName NewCurve, CurveToCopy;
+		//@todo : test - how to duplicate track?
+		if (NameMapping->FindSmartName(CurveNameToCopy, CurveToCopy) && Sequence->GetSkeleton()->AddSmartNameAndModify(USkeleton::AnimCurveMappingName, FName(*NewCurveName), NewCurve))
 		{
-			if(Sequence->RawCurveData.DuplicateCurveData(Uid, NewUid))
+			if(Sequence->RawCurveData.DuplicateCurveData(CurveToCopy, NewCurve))
 			{
 				Sequence->Modify();
 				Sequence->MarkRawDataAsModified();
@@ -762,10 +767,14 @@ void SAnimCurvePanel::DeleteTrack(USkeleton::AnimCurveUID Uid)
 	
 	if(Sequence->RawCurveData.GetCurveData(Uid))
 	{
-		Sequence->Modify(true);
-		Sequence->RawCurveData.DeleteCurveData(Uid);
-		Sequence->MarkRawDataAsModified();
-		UpdatePanel();
+		FSmartName TrackName;
+		if (Sequence->GetSkeleton()->GetSmartNameByUID(USkeleton::AnimCurveMappingName, Uid, TrackName))
+		{
+			Sequence->Modify(true);
+			Sequence->RawCurveData.DeleteCurveData(TrackName);
+			Sequence->MarkRawDataAsModified();
+			UpdatePanel();
+		}
 	}
 }
 
@@ -809,6 +818,18 @@ FReply SAnimCurvePanel::OnContextMenu()
 			],
 			FText()
 		);
+		MenuBuilder.AddWidget(
+			SNew(SCheckBox)
+			.IsChecked(this, &SAnimCurvePanel::AreAllCurvesOfMode, ACF_DrivesPose)
+			.OnCheckStateChanged(this, &SAnimCurvePanel::ToggleAllCurveModes, ACF_DrivesPose)
+			.ToolTipText(LOCTEXT("PoseCurveModeTooltip", "This curve drives a pose"))
+			.HAlign(HAlign_Left)
+			[
+				SNew(STextBlock)
+				.Text(LOCTEXT("PoseCurveMode", "Pose Curve"))
+			],
+			FText()
+			);
 	}
 	MenuBuilder.EndSection();
 
@@ -900,8 +921,8 @@ void SAnimCurvePanel::UpdatePanel()
 
 			FName AName;
 			FName BName;
-			MetadataNameMap->GetName(A.CurveUid, AName);
-			MetadataNameMap->GetName(B.CurveUid, BName);
+			MetadataNameMap->GetName(A.Name.UID, AName);
+			MetadataNameMap->GetName(B.Name.UID, BName);
 
 			return AName < BName;
 		});
@@ -929,7 +950,7 @@ void SAnimCurvePanel::UpdatePanel()
 			FName CurveName;
 
 			// if editable, add to the list
-			if(bEditable && NameMapping->GetName(Curve.CurveUid, CurveName))
+			if(bEditable && NameMapping->GetName(Curve.Name.UID, CurveName))
 			{
 				TSharedPtr<SCurveEdTrack> CurrentTrack;
 				PanelSlot->AddSlot()
@@ -942,7 +963,7 @@ void SAnimCurvePanel::UpdatePanel()
 					[
 						SAssignNew(CurrentTrack, SCurveEdTrack)
 						.Sequence(Sequence)
-						.CurveUid(Curve.CurveUid)
+						.CurveUid(Curve.Name.UID)
 						.AnimCurvePanel(SharedThis(this))
 						.WidgetWidth(WidgetWidth)
 						.ViewInputMin(ViewInputMin)
@@ -999,7 +1020,7 @@ TSharedRef<SWidget> SAnimCurvePanel::GenerateCurveList()
 			const FFloatCurve& Curve= *Iter;
 
 			FName CurveName;
-			NameMapping->GetName(Curve.CurveUid, CurveName);
+			NameMapping->GetName(Curve.Name.UID, CurveName);
 
 			ListBox->AddSlot()
 				.AutoHeight()
@@ -1007,8 +1028,8 @@ TSharedRef<SWidget> SAnimCurvePanel::GenerateCurveList()
 				.Padding( 2.0f, 2.0f )
 				[
 					SNew( SCheckBox )
-					.IsChecked(this, &SAnimCurvePanel::IsCurveEditable, Curve.CurveUid)
-					.OnCheckStateChanged(this, &SAnimCurvePanel::ToggleEditability, Curve.CurveUid)
+					.IsChecked(this, &SAnimCurvePanel::IsCurveEditable, Curve.Name.UID)
+					.OnCheckStateChanged(this, &SAnimCurvePanel::ToggleEditability, Curve.Name.UID)
 					.ToolTipText( LOCTEXT("Show Curves", "Show or Hide Curves") )
 					.IsEnabled( true )
 					[
@@ -1184,6 +1205,13 @@ void SAnimCurvePanel::FillMetadataEntryMenu(FMenuBuilder& Builder)
 
 void SAnimCurvePanel::FillVariableCurveMenu(FMenuBuilder& Builder)
 {
+	FText Description = LOCTEXT("NewVariableCurveCreateNew_ToolTip", "Create a new variable curve");
+	FText Label = LOCTEXT("NewVariableCurveCreateNew_Label", "Create Curve");
+	FUIAction UIAction;
+	UIAction.ExecuteAction.BindRaw(this, &SAnimCurvePanel::CreateNewCurveClicked);
+
+	Builder.AddMenuEntry(Label, Description, FSlateIcon(), UIAction);
+
 	USkeleton* CurrentSkeleton = Sequence->GetSkeleton();
 	check(CurrentSkeleton);
 
@@ -1212,10 +1240,9 @@ void SAnimCurvePanel::FillVariableCurveMenu(FMenuBuilder& Builder)
 
 			for (FSmartNameSortItem SmartNameItem : SmartNameList)
 			{
-				const FText Description = LOCTEXT("NewVariableSubMenu_ToolTip", "Add an existing variable curve");
-				const FText Label = FText::FromName(SmartNameItem.SmartName);
+				Description = LOCTEXT("NewVariableSubMenu_ToolTip", "Add an existing variable curve");
+				Label = FText::FromName(SmartNameItem.SmartName);
 
-				FUIAction UIAction;
 				UIAction.ExecuteAction.BindRaw(
 					this, &SAnimCurvePanel::AddVariableCurve,
 					SmartNameItem.ID);
@@ -1225,20 +1252,13 @@ void SAnimCurvePanel::FillVariableCurveMenu(FMenuBuilder& Builder)
 		}
 	}
 	Builder.EndSection();
-
-	Builder.AddMenuSeparator();
-
-	const FText Description = LOCTEXT("NewVariableCurveCreateNew_ToolTip", "Create a new variable curve");
-	const FText Label = LOCTEXT("NewVariableCurveCreateNew_Label", "Create Curve");
-	FUIAction UIAction;
-	UIAction.ExecuteAction.BindRaw(this, &SAnimCurvePanel::CreateNewCurveClicked);
-
-	Builder.AddMenuEntry(Label, Description, FSlateIcon(), UIAction);
 }
 
 void SAnimCurvePanel::AddMetadataEntry(USkeleton::AnimCurveUID Uid)
 {
-	if(Sequence->RawCurveData.AddCurveData(Uid))
+	FSmartName NewName;
+	ensureAlways(Sequence->GetSkeleton()->GetSmartNameByUID(USkeleton::AnimCurveMappingName, Uid, NewName));
+	if(Sequence->RawCurveData.AddCurveData(NewName))
 	{
 		Sequence->Modify(true);
 		Sequence->MarkRawDataAsModified();
@@ -1275,11 +1295,11 @@ void SAnimCurvePanel::CreateNewMetadataEntry(const FText& CommittedText, ETextCo
 		USkeleton* Skeleton = Sequence->GetSkeleton();
 		if(Skeleton && !CommittedText.IsEmpty())
 		{
-			USkeleton::AnimCurveUID CurveUid;
+			FSmartName CurveName;
 
-			if(Skeleton->AddSmartNameAndModify(USkeleton::AnimCurveMappingName, FName(*CommittedText.ToString()), CurveUid))
+			if(Skeleton->AddSmartNameAndModify(USkeleton::AnimCurveMappingName, FName(*CommittedText.ToString()), CurveName))
 			{
-				AddMetadataEntry(CurveUid);
+				AddMetadataEntry(CurveName.UID);
 			}
 		}
 	}
@@ -1331,6 +1351,18 @@ TSharedRef<SWidget> SAnimCurvePanel::CreateCurveContextMenu(FAnimCurveBaseInterf
 			],
 			FText()
 			);
+
+		MenuBuilder.AddWidget(
+			SNew(SCheckBox)
+			.IsChecked(this, &SAnimCurvePanel::GetCurveFlagAsCheckboxState, Curve, ACF_DrivesPose)
+			.OnCheckStateChanged(this, &SAnimCurvePanel::SetCurveFlagFromCheckboxState, Curve, ACF_DrivesPose)
+			.ToolTipText(LOCTEXT("PoseCurveModeTooltip", "This curve drives a pose asset"))
+			[
+				SNew(STextBlock)
+				.Text(LOCTEXT("PoseCurveMode", "Pose Curve"))
+			],
+			FText()
+			);
 	}
 	MenuBuilder.EndSection();
 
@@ -1399,7 +1431,11 @@ void SAnimCurvePanel::ToggleCurveTypeMenuCallback(FAnimCurveBaseInterface* Curve
 void SAnimCurvePanel::AddVariableCurve(USkeleton::AnimCurveUID CurveUid)
 {
 	Sequence->Modify(true);
-	Sequence->RawCurveData.AddCurveData(CurveUid);
+	
+	USkeleton* Skeleton = Sequence->GetSkeleton();
+	FSmartName NewName;
+	ensureAlways(Skeleton->GetSmartNameByUID(USkeleton::AnimCurveMappingName, CurveUid, NewName));
+	Sequence->RawCurveData.AddCurveData(NewName);
 	Sequence->MarkRawDataAsModified();
 	UpdatePanel();
 }
diff --git a/Engine/Source/Editor/Persona/Private/SAnimCurveViewer.cpp b/Engine/Source/Editor/Persona/Private/SAnimCurveViewer.cpp
new file mode 100644
--- /dev/null
+++ b/Engine/Source/Editor/Persona/Private/SAnimCurveViewer.cpp
@@ -0,0 +1,654 @@
+// Copyright 1998-2016 Epic Games, Inc. All Rights Reserved.
+
+
+#include "PersonaPrivatePCH.h"
+#include "SAnimCurveViewer.h"
+#include "ObjectTools.h"
+#include "AssetRegistryModule.h"
+#include "ScopedTransaction.h"
+#include "SSearchBox.h"
+
+#define LOCTEXT_NAMESPACE "SAnimCurveViewer"
+
+static const FName ColumnId_AnimCurveNameLabel( "Curve Name" );
+static const FName ColumnID_AnimCurveWeightLabel( "Weight" );
+static const FName ColumnID_AnimCurveEditLabel( "Edit" );
+
+const float MaxMorphWeight = 5.f;
+
+//////////////////////////////////////////////////////////////////////////
+// SAnimCurveListRow
+
+typedef TSharedPtr< FDisplayedAnimCurveInfo > FDisplayedAnimCurveInfoPtr;
+
+class SAnimCurveListRow
+	: public SMultiColumnTableRow< FDisplayedAnimCurveInfoPtr >
+{
+public:
+
+	SLATE_BEGIN_ARGS( SAnimCurveListRow ) {}
+
+	/** The item for this row **/
+	SLATE_ARGUMENT( FDisplayedAnimCurveInfoPtr, Item )
+
+	/* The SAnimCurveViewer that we push the morph target weights into */
+	SLATE_ARGUMENT( class SAnimCurveViewer*, AnimCurveViewer )
+
+	/* Widget used to display the list of animation curve */
+	SLATE_ARGUMENT( TSharedPtr<SAnimCurveListType>, AnimCurveListView )
+
+	/* Persona used to update the viewport when a weight slider is dragged */
+	SLATE_ARGUMENT( TWeakPtr<FPersona>, Persona )
+
+	SLATE_END_ARGS()
+
+	void Construct( const FArguments& InArgs, const TSharedRef<STableViewBase>& OwnerTableView );
+
+	/** Overridden from SMultiColumnTableRow.  Generates a widget for this column of the tree row. */
+	virtual TSharedRef<SWidget> GenerateWidgetForColumn( const FName& ColumnName ) override;
+
+private:
+
+	/**
+	* Called when the user changes the value of the SSpinBox
+	*
+	* @param NewWeight - The new number the SSpinBox is set to
+	*
+	*/
+	void OnAnimCurveWeightChanged( float NewWeight );
+	/**
+	* Called when the user types the value and enters
+	*
+	* @param NewWeight - The new number the SSpinBox is set to
+	*
+	*/
+	void OnAnimCurveWeightValueCommitted( float NewWeight, ETextCommit::Type CommitType);
+
+	/** Auto fill check call back functions */
+	void OnAnimCurveAutoFillChecked(ECheckBoxState InState);
+	ECheckBoxState IsAnimCurveAutoFillChangedChecked() const;
+	
+	/**
+	* Returns the weight of this morph target
+	*
+	* @return SearchText - The new number the SSpinBox is set to
+	*
+	*/
+	float GetWeight() const;
+
+	/* The SAnimCurveViewer that we push the morph target weights into */
+	SAnimCurveViewer* AnimCurveViewer;
+
+	/** Widget used to display the list of animation curve */
+	TSharedPtr<SAnimCurveListType> AnimCurveListView;
+
+	/** The name and weight of the morph target */
+	FDisplayedAnimCurveInfoPtr	Item;
+
+	/** Pointer back to the Persona that owns us */
+	TWeakPtr<FPersona> PersonaPtr;
+};
+
+void SAnimCurveListRow::Construct( const FArguments& InArgs, const TSharedRef<STableViewBase>& InOwnerTableView )
+{
+	Item = InArgs._Item;
+	AnimCurveViewer = InArgs._AnimCurveViewer;
+	AnimCurveListView = InArgs._AnimCurveListView;
+	PersonaPtr = InArgs._Persona;
+
+	check( Item.IsValid() );
+
+	SMultiColumnTableRow< FDisplayedAnimCurveInfoPtr >::Construct( FSuperRowType::FArguments(), InOwnerTableView );
+}
+
+TSharedRef< SWidget > SAnimCurveListRow::GenerateWidgetForColumn( const FName& ColumnName )
+{
+	if ( ColumnName == ColumnId_AnimCurveNameLabel )
+	{
+		return
+			SNew( SVerticalBox )
+
+			+ SVerticalBox::Slot()
+			.AutoHeight()
+			.Padding( 0.0f, 4.0f )
+			.VAlign( VAlign_Center )
+			[
+				SNew( STextBlock )
+				.Text( FText::FromName(Item->Name) )
+				.HighlightText( AnimCurveViewer->GetFilterText() )
+			];
+	}
+	else if ( ColumnName == ColumnID_AnimCurveWeightLabel )
+	{
+		// Encase the SSpinbox in an SVertical box so we can apply padding. Setting ItemHeight on the containing SListView has no effect :-(
+		return
+			SNew( SVerticalBox )
+
+			+ SVerticalBox::Slot()
+			.AutoHeight()
+			.Padding( 0.0f, 1.0f )
+			.VAlign( VAlign_Center )
+			[
+				SNew( SSpinBox<float> )
+				.MinSliderValue(-1.f)
+				.MaxSliderValue(1.f)
+				.MinValue(-MaxMorphWeight)
+				.MaxValue(MaxMorphWeight)
+				.Value( this, &SAnimCurveListRow::GetWeight )
+				.OnValueChanged( this, &SAnimCurveListRow::OnAnimCurveWeightChanged )
+				.OnValueCommitted( this, &SAnimCurveListRow::OnAnimCurveWeightValueCommitted )
+			];
+	}
+	else 
+	{
+		return
+			SNew(SVerticalBox)
+
+			+ SVerticalBox::Slot()
+			.AutoHeight()
+			.Padding(0.0f, 1.0f)
+			.VAlign(VAlign_Center)
+			.HAlign(HAlign_Center)
+			[
+				SNew(SCheckBox)
+				.OnCheckStateChanged(this, &SAnimCurveListRow::OnAnimCurveAutoFillChecked)
+				.IsChecked(this, &SAnimCurveListRow::IsAnimCurveAutoFillChangedChecked)
+			];
+	}
+}
+
+void SAnimCurveListRow::OnAnimCurveAutoFillChecked(ECheckBoxState InState)
+{
+	Item->bAutoFillData = InState == ECheckBoxState::Checked;
+
+	if (Item->bAutoFillData)
+	{
+		// clear value so that it can be filled up
+		AnimCurveViewer->RemoveAnimCurveOverride(Item->Name);
+	}
+}
+
+ECheckBoxState SAnimCurveListRow::IsAnimCurveAutoFillChangedChecked() const
+{
+	return (Item->bAutoFillData)? ECheckBoxState::Checked: ECheckBoxState::Unchecked;
+}
+
+void SAnimCurveListRow::OnAnimCurveWeightChanged( float NewWeight )
+{
+	Item->Weight = NewWeight;
+	Item->bAutoFillData = false;
+
+	AnimCurveViewer->AddAnimCurveOverride( Item->Name, Item->Weight);
+
+	if (PersonaPtr.IsValid())
+	{
+		PersonaPtr.Pin()->RefreshViewport();
+	}
+}
+
+void SAnimCurveListRow::OnAnimCurveWeightValueCommitted( float NewWeight, ETextCommit::Type CommitType)
+{
+	if (CommitType == ETextCommit::OnEnter || CommitType == ETextCommit::OnUserMovedFocus)
+	{
+		float NewValidWeight = FMath::Clamp(NewWeight, -MaxMorphWeight, MaxMorphWeight);
+		Item->Weight = NewValidWeight;
+		Item->bAutoFillData = false;
+
+		AnimCurveViewer->AddAnimCurveOverride(Item->Name, Item->Weight);
+
+		TArray< TSharedPtr< FDisplayedAnimCurveInfo > > SelectedRows = AnimCurveListView->GetSelectedItems();
+
+		// ...then any selected rows need changing by the same delta
+		for(auto ItemIt = SelectedRows.CreateIterator(); ItemIt; ++ItemIt)
+		{
+			TSharedPtr< FDisplayedAnimCurveInfo > RowItem = (*ItemIt);
+
+			if(RowItem != Item) // Don't do "this" row again if it's selected
+			{
+				RowItem->Weight = NewValidWeight;
+				RowItem->bAutoFillData = false;
+				AnimCurveViewer->AddAnimCurveOverride(RowItem->Name, RowItem->Weight);
+			}
+		}
+
+		if(PersonaPtr.IsValid())
+		{
+			PersonaPtr.Pin()->RefreshViewport();
+		}
+	}
+}
+
+float SAnimCurveListRow::GetWeight() const 
+{ 
+	if (PersonaPtr.IsValid() && Item->bAutoFillData)
+	{
+		USkeletalMeshComponent* SkelComp = PersonaPtr.Pin()->GetPreviewMeshComponent();
+		UAnimInstance* AnimInstance = (SkelComp) ? SkelComp->GetAnimInstance() : nullptr;
+		if (AnimInstance)
+		{
+			TMap<FName, float> CurveList;
+
+			// make sure if they have value that's not same as saved value
+			AnimInstance->GetAnimationCurveList(AnimCurveViewer->CurrentCurveFlag, CurveList);
+
+			if (CurveList.Num() > 0)
+			{
+				float* CurrentValue = CurveList.Find(Item->Name);
+				if (CurrentValue && *CurrentValue != 0.f)
+				{
+					return *CurrentValue;
+				}
+			}
+			else
+			{
+				return 0.f;
+			}
+		}
+	}
+
+	return Item->Weight; 
+}
+
+//////////////////////////////////////////////////////////////////////////
+// SAnimCurveTypeList
+
+class SAnimCurveTypeList : public SCompoundWidget
+{
+public:
+
+	SLATE_BEGIN_ARGS(SAnimCurveTypeList) {}
+
+	/** The item for this row **/
+	SLATE_ARGUMENT(int32, CurveFlags)
+
+	/* The SAnimCurveViewer that we push the morph target weights into */
+	SLATE_ARGUMENT(class SAnimCurveViewer*, AnimCurveViewer)
+
+	SLATE_END_ARGS()
+
+	void Construct(const FArguments& InArgs);
+
+	/** Auto fill check call back functions */
+	void OnAnimCurveTypeShowlChecked(ECheckBoxState InState);
+	ECheckBoxState IsAnimCurveTypeShowChangedChecked() const;
+
+	FText GetAnimCurveType() const;
+private:
+
+	/* The SAnimCurveViewer that we push the morph target weights into */
+	SAnimCurveViewer* AnimCurveViewer;
+
+	/** The name and weight of the morph target */
+	int32 CurveFlags;
+};
+
+void SAnimCurveTypeList::Construct(const FArguments& InArgs)
+{
+	CurveFlags = InArgs._CurveFlags;
+	AnimCurveViewer = InArgs._AnimCurveViewer;
+
+	ChildSlot
+	[
+		SNew(SHorizontalBox)
+
+		+SHorizontalBox::Slot()
+		.AutoWidth()
+		.Padding(0.0f, 1.0f)
+		.VAlign(VAlign_Center)
+		.HAlign(HAlign_Center)
+		[
+			SNew(SCheckBox)
+			.OnCheckStateChanged(this, &SAnimCurveTypeList::OnAnimCurveTypeShowlChecked)
+			.IsChecked(this, &SAnimCurveTypeList::IsAnimCurveTypeShowChangedChecked)
+		]
+
+		+SHorizontalBox::Slot()
+		.AutoWidth()
+		.Padding(3.0f, 1.0f)
+		.VAlign(VAlign_Center)
+		[
+			SNew(STextBlock)
+			.Text(this, &SAnimCurveTypeList::GetAnimCurveType)
+			.HighlightText(AnimCurveViewer->GetFilterText())
+		]
+	];
+}
+
+void SAnimCurveTypeList::OnAnimCurveTypeShowlChecked(ECheckBoxState InState)
+{
+	// clear value so that it can be filled up
+	if (InState == ECheckBoxState::Checked)
+	{
+		AnimCurveViewer->CurrentCurveFlag |= CurveFlags;
+	}
+	else
+	{
+		AnimCurveViewer->CurrentCurveFlag &= ~CurveFlags;
+	}
+
+	AnimCurveViewer->RefreshCurveList();
+}
+
+ECheckBoxState SAnimCurveTypeList::IsAnimCurveTypeShowChangedChecked() const
+{
+	return (AnimCurveViewer->CurrentCurveFlag & CurveFlags) ? ECheckBoxState::Checked : ECheckBoxState::Unchecked;
+}
+
+FText SAnimCurveTypeList::GetAnimCurveType() const
+{
+	switch (CurveFlags)
+	{
+	case ACF_DrivesMorphTarget:
+		return LOCTEXT("AnimCurveType_Morphtarget", "Morph Target");
+	case ACF_TriggerEvent:
+		return LOCTEXT("AnimCurveType_Event", "Event");
+	case ACF_DrivesMaterial:
+		return LOCTEXT("AnimCurveType_Material", "Material");
+	case ACF_Metadata:
+		return LOCTEXT("AnimCurveType_Metadata", "Meta Data");
+	case ACF_DriveTrack:
+		return LOCTEXT("AnimCurveType_Tranform", "Transform");
+	case ACF_DrivesPose:
+		return LOCTEXT("AnimCurveType_Pose", "Pose");
+	}
+
+	return LOCTEXT("AnimCurveType_Unknown", "Unknown");
+}
+//////////////////////////////////////////////////////////////////////////
+// SAnimCurveViewer
+
+void SAnimCurveViewer::Construct(const FArguments& InArgs)
+{
+	PersonaPtr = InArgs._Persona;
+	CachedPreviewInstance = nullptr;
+
+	if (PersonaPtr.IsValid())
+	{
+		PersonaPtr.Pin()->RegisterOnPreviewMeshChanged(FPersona::FOnPreviewMeshChanged::CreateSP(this, &SAnimCurveViewer::OnPreviewMeshChanged));
+		PersonaPtr.Pin()->RegisterOnAnimChanged(FPersona::FOnAnimChanged::CreateSP(this, &SAnimCurveViewer::OnPreviewAssetChanged));
+		PersonaPtr.Pin()->RegisterOnPostUndo(FPersona::FOnPostUndo::CreateSP(this, &SAnimCurveViewer::OnPostUndo));
+	}
+
+	// @todo fix this to be filtered
+	CurrentCurveFlag = ACF_DrivesMorphTarget | ACF_DrivesMaterial | ACF_DrivesPose;
+
+	RefreshCachePreviewInstance();
+	 
+	TSharedPtr<SVerticalBox> AnimTypeBoxContainer = SNew(SVerticalBox);
+
+	ChildSlot
+	[
+		SNew( SVerticalBox )
+		
+		+ SVerticalBox::Slot()
+		.AutoHeight()
+		.Padding(2)
+		[
+			SNew( STextBlock )
+			.Text( this, &SAnimCurveViewer::GetTitleText )
+		]
+
+		+ SVerticalBox::Slot()
+		.AutoHeight()
+		.Padding(0,2)
+		[
+			SNew(SHorizontalBox)
+			// Filter entry
+			+SHorizontalBox::Slot()
+			.FillWidth( 1 )
+			[
+				SAssignNew( NameFilterBox, SSearchBox )
+				.SelectAllTextWhenFocused( true )
+				.OnTextChanged( this, &SAnimCurveViewer::OnFilterTextChanged )
+				.OnTextCommitted( this, &SAnimCurveViewer::OnFilterTextCommitted )
+			]
+		]
+
+		+ SVerticalBox::Slot()
+		.AutoHeight()		// This is required to make the scrollbar work, as content overflows Slate containers by default
+		[
+			SNew(SBox)
+			.WidthOverride(150)
+			.Content()
+			[
+				AnimTypeBoxContainer.ToSharedRef()
+			]
+		]
+
+		+ SVerticalBox::Slot()
+		.FillHeight( 1.0f )		// This is required to make the scrollbar work, as content overflows Slate containers by default
+		[
+			SAssignNew( AnimCurveListView, SAnimCurveListType )
+			.ListItemsSource( &AnimCurveList )
+			.OnGenerateRow( this, &SAnimCurveViewer::GenerateAnimCurveRow )
+			.OnContextMenuOpening( this, &SAnimCurveViewer::OnGetContextMenuContent )
+			.ItemHeight( 22.0f )
+			.HeaderRow
+			(
+				SNew( SHeaderRow )
+				+ SHeaderRow::Column( ColumnId_AnimCurveNameLabel )
+				.DefaultLabel( LOCTEXT( "AnimCurveNameLabel", "Curve Name" ) )
+
+				+ SHeaderRow::Column( ColumnID_AnimCurveWeightLabel )
+				.DefaultLabel( LOCTEXT( "AnimCurveWeightLabel", "Weight" ) )
+
+				+ SHeaderRow::Column(ColumnID_AnimCurveEditLabel)
+				.DefaultLabel(LOCTEXT("AnimCurveEditLabel", "Auto"))
+			)
+		]
+	];
+
+	CreateAnimCurveTypeList(AnimTypeBoxContainer.ToSharedRef());
+	CreateAnimCurveList();
+}
+
+FText SAnimCurveViewer::GetTitleText() const
+{
+	UObject* PreviewAsset = nullptr;
+
+	if (PersonaPtr.IsValid())
+	{
+		PreviewAsset = PersonaPtr.Pin()->GetPreviewAnimationAsset();
+		if (!PreviewAsset)
+		{
+			PreviewAsset = PersonaPtr.Pin()->GetAnimBlueprint();
+		}
+	}
+
+	return (PreviewAsset? FText::FromString(PreviewAsset->GetName()) : LOCTEXT("AnimCurveMeshNameLabel", "No Asset Present"));
+}
+
+void SAnimCurveViewer::RefreshCachePreviewInstance()
+{
+	if (CachedPreviewInstance && OnAddAnimationCurveDelegate.IsBound())
+	{
+		CachedPreviewInstance->RemoveDelegate_AddCustomAnimationCurve(OnAddAnimationCurveDelegate);
+	}
+	CachedPreviewInstance = Cast<UAnimInstance>(PersonaPtr.Pin()->GetPreviewMeshComponent()->GetAnimInstance());
+
+	if (CachedPreviewInstance)
+	{
+		OnAddAnimationCurveDelegate.BindRaw(this, &SAnimCurveViewer::ApplyCustomCurveOverride);
+		CachedPreviewInstance->AddDelegate_AddCustomAnimationCurve(OnAddAnimationCurveDelegate);
+	}
+}
+
+void SAnimCurveViewer::OnPreviewMeshChanged(class USkeletalMesh* NewPreviewMesh)
+{
+	RefreshCachePreviewInstance();
+	RefreshCurveList();
+}
+
+void SAnimCurveViewer::OnFilterTextChanged( const FText& SearchText )
+{
+	FilterText = SearchText;
+
+	RefreshCurveList();
+}
+
+void SAnimCurveViewer::OnFilterTextCommitted( const FText& SearchText, ETextCommit::Type CommitInfo )
+{
+	// Just do the same as if the user typed in the box
+	OnFilterTextChanged( SearchText );
+}
+
+TSharedRef<ITableRow> SAnimCurveViewer::GenerateAnimCurveRow(TSharedPtr<FDisplayedAnimCurveInfo> InInfo, const TSharedRef<STableViewBase>& OwnerTable)
+{
+	check( InInfo.IsValid() );
+
+	return
+		SNew( SAnimCurveListRow, OwnerTable )
+		.Persona( PersonaPtr )
+		.Item( InInfo )
+		.AnimCurveViewer( this )
+		.AnimCurveListView( AnimCurveListView );
+}
+
+TSharedPtr<SWidget> SAnimCurveViewer::OnGetContextMenuContent() const
+{
+	const bool bShouldCloseWindowAfterMenuSelection = true;
+	FMenuBuilder MenuBuilder( bShouldCloseWindowAfterMenuSelection, NULL);
+
+	MenuBuilder.BeginSection("AnimCurveAction", LOCTEXT( "MorphsAction", "Selected Item Actions" ) );
+// 	{
+// 		FUIAction Action = FUIAction( FExecuteAction::CreateSP( this, &SAnimCurveViewer::OnDeleteAnimCurves ), 
+// 									  FCanExecuteAction::CreateSP( this, &SAnimCurveViewer::CanPerformDelete ) );
+// 		const FText Label = LOCTEXT("DeleteAnimCurveButtonLabel", "Delete");
+// 		const FText ToolTip = LOCTEXT("DeleteAnimCurveButtonTooltip", "Deletes the selected animation curve.");
+// 		MenuBuilder.AddMenuEntry( Label, ToolTip, FSlateIcon(), Action);
+// 	}
+	MenuBuilder.EndSection();
+
+	return MenuBuilder.MakeWidget();
+}
+
+void SAnimCurveViewer::Tick(const FGeometry& AllottedGeometry, const double InCurrentTime, const float InDeltaTime)
+{
+	if (PersonaPtr.IsValid())
+	{
+		UDebugSkelMeshComponent* MeshComponent = PersonaPtr.Pin()->GetPreviewMeshComponent();
+		UAnimInstance* AnimInstance = MeshComponent->GetAnimInstance();
+
+		if (AnimInstance)
+		{
+			TMap<FName, float> AnimCurves;
+			AnimInstance->GetAnimationCurveList(CurrentCurveFlag, AnimCurves);
+			// @fixme: this is not enough information to refresh curve, but it is better than not displaying because it didnt' tick yet
+			if (AnimCurves.Num() != AnimCurveList.Num())
+			{
+				RefreshCurveList();
+			}
+		}
+	}
+}
+void SAnimCurveViewer::CreateAnimCurveList( const FString& SearchText )
+{
+	AnimCurveList.Empty();
+
+	UDebugSkelMeshComponent* MeshComponent = PersonaPtr.Pin()->GetPreviewMeshComponent();
+	UAnimInstance* AnimInstance = MeshComponent->GetAnimInstance();
+
+	if (AnimInstance)
+	{
+		TMap<FName, float> AnimCurves;
+		AnimInstance->GetAnimationCurveList(CurrentCurveFlag, AnimCurves);
+
+		if (AnimCurves.Num() > 0)
+		{
+			bool bDoFiltering = !SearchText.IsEmpty();
+
+			for (auto Iter = AnimCurves.CreateConstIterator(); Iter; ++Iter)
+			{
+				FName CurveName = Iter.Key();
+				if (bDoFiltering && !CurveName.ToString().Contains(SearchText))
+				{
+					continue; // Skip items that don't match our filter
+				}
+
+				TSharedRef<FDisplayedAnimCurveInfo> Info = FDisplayedAnimCurveInfo::Make(CurveName);
+				AnimCurveList.Add(Info);
+			}
+		}
+		else
+		{
+			// else print somewhere no curves?
+		}
+	}
+
+	AnimCurveListView->RequestListRefresh();
+}
+
+void SAnimCurveViewer::CreateAnimCurveTypeList(TSharedRef<SVerticalBox> VerticalBox)
+{
+	TArray<int32> CurveFlagsToList;
+	CurveFlagsToList.Add(ACF_DrivesMorphTarget);
+	CurveFlagsToList.Add(ACF_TriggerEvent);
+	CurveFlagsToList.Add(ACF_DrivesMaterial);
+	CurveFlagsToList.Add(ACF_Metadata);
+	CurveFlagsToList.Add(ACF_DriveTrack);
+	CurveFlagsToList.Add(ACF_DrivesPose);
+
+	for (int32 CurveFlagIndex = 0; CurveFlagIndex < CurveFlagsToList.Num(); ++CurveFlagIndex)
+	{
+		VerticalBox->AddSlot()
+		.Padding(3, 1)
+		[
+			SNew(SAnimCurveTypeList)
+			.CurveFlags(CurveFlagsToList[CurveFlagIndex])
+			.AnimCurveViewer(this)
+		];
+	}
+}
+
+void SAnimCurveViewer::AddAnimCurveOverride( FName& Name, float Weight)
+{
+	float& Value = OverrideCurves.FindOrAdd(Name);
+	Value = Weight;
+}
+
+void SAnimCurveViewer::RemoveAnimCurveOverride(FName& Name)
+{
+	OverrideCurves.Remove(Name);
+}
+
+SAnimCurveViewer::~SAnimCurveViewer()
+{
+	if ( PersonaPtr.IsValid() )
+	{
+		// if persona isn't there, we probably don't have the preview mesh either, so no valid anim instance
+		if (CachedPreviewInstance && OnAddAnimationCurveDelegate.IsBound())
+		{
+			CachedPreviewInstance->RemoveDelegate_AddCustomAnimationCurve(OnAddAnimationCurveDelegate);
+		}
+
+		PersonaPtr.Pin()->UnregisterOnPreviewMeshChanged(this);
+		PersonaPtr.Pin()->UnregisterOnAnimChanged(this);
+		PersonaPtr.Pin()->UnregisterOnPostUndo(this);
+	}
+}
+
+void SAnimCurveViewer::OnPostUndo()
+{
+	RefreshCurveList();
+}
+
+void SAnimCurveViewer::OnPreviewAssetChanged(class UAnimationAsset* NewAsset)
+{
+	RefreshCurveList();
+}
+
+void SAnimCurveViewer::ApplyCustomCurveOverride(UAnimInstance* AnimInstance) const
+{
+	for (auto Iter = OverrideCurves.CreateConstIterator(); Iter; ++Iter)
+	{ 
+		// @todo we might want to save original curve flags? or just change curve to apply flags only
+		AnimInstance->AddCurveValue(Iter.Key(), Iter.Value(), CurrentCurveFlag);
+	}
+}
+
+void SAnimCurveViewer::RefreshCurveList()
+{
+	CreateAnimCurveList(FilterText.ToString());
+}
+#undef LOCTEXT_NAMESPACE
+
diff --git a/Engine/Source/Editor/Persona/Private/SAnimCurveViewer.h b/Engine/Source/Editor/Persona/Private/SAnimCurveViewer.h
new file mode 100644
--- /dev/null
+++ b/Engine/Source/Editor/Persona/Private/SAnimCurveViewer.h
@@ -0,0 +1,188 @@
+// Copyright 1998-2016 Epic Games, Inc. All Rights Reserved.
+
+#pragma once
+#include "Animation/AnimInstance.h"
+#include "Persona.h"
+
+//////////////////////////////////////////////////////////////////////////
+// FDisplayedAnimCurveInfo
+
+class FDisplayedAnimCurveInfo
+{
+public:
+	FName Name;
+	float Weight;
+	bool bAutoFillData;
+
+	/** Static function for creating a new item, but ensures that you can only have a TSharedRef to one */
+	static TSharedRef<FDisplayedAnimCurveInfo> Make(const FName& Source)
+	{
+		return MakeShareable(new FDisplayedAnimCurveInfo(Source));
+	}
+
+protected:
+	/** Hidden constructor, always use Make above */
+	FDisplayedAnimCurveInfo(const FName& InSource)
+		: Name( InSource )
+		, Weight( 0 )
+		, bAutoFillData(true)
+	{}
+
+	/** Hidden constructor, always use Make above */
+	FDisplayedAnimCurveInfo() {}
+};
+
+typedef SListView< TSharedPtr<FDisplayedAnimCurveInfo> > SAnimCurveListType;
+
+//////////////////////////////////////////////////////////////////////////
+// SAnimCurveViewer
+
+class SAnimCurveViewer : public SCompoundWidget
+{
+public:
+	SLATE_BEGIN_ARGS( SAnimCurveViewer )
+		: _Persona()
+	{}
+		
+		/* The Persona that owns this table */
+		SLATE_ARGUMENT( TWeakPtr< FPersona >, Persona )
+
+	SLATE_END_ARGS()
+
+	/**
+	* Slate construction function
+	*
+	* @param InArgs - Arguments passed from Slate
+	*
+	*/
+	void Construct( const FArguments& InArgs );
+
+	/**
+	* Destructor - resets the animation curve
+	*
+	*/
+	virtual ~SAnimCurveViewer();
+
+	/**
+	* Is registered with Persona to handle when its preview mesh is changed.
+	*
+	* @param NewPreviewMesh - The new preview mesh being used by Persona
+	*
+	*/
+	void OnPreviewMeshChanged(class USkeletalMesh* NewPreviewMesh);
+
+	/**
+	* Is registered with Persona to handle when its preview asset is changed.
+	*
+	* Pose Asset will have to add curve manually
+	*/
+	void OnPreviewAssetChanged(class UAnimationAsset* NewPreviewAsset);
+
+	/**
+	* Filters the SListView when the user changes the search text box (NameFilterBox)
+	*
+	* @param SearchText - The text the user has typed
+	*
+	*/
+	void OnFilterTextChanged( const FText& SearchText );
+
+	/**
+	* Filters the SListView when the user hits enter or clears the search box
+	* Simply calls OnFilterTextChanged
+	*
+	* @param SearchText - The text the user has typed
+	* @param CommitInfo - Not used
+	*
+	*/
+	void OnFilterTextCommitted( const FText& SearchText, ETextCommit::Type CommitInfo );
+
+	/**
+	* Create a widget for an entry in the tree from an info
+	*
+	* @param InInfo - Shared pointer to the morph target we're generating a row for
+	* @param OwnerTable - The table that owns this row
+	*
+	* @return A new Slate widget, containing the UI for this row
+	*/
+	TSharedRef<ITableRow> GenerateAnimCurveRow(TSharedPtr<FDisplayedAnimCurveInfo> InInfo, const TSharedRef<STableViewBase>& OwnerTable);
+
+	/**
+	* Adds a morph target override or updates the weight for an existing one
+	*
+	* @param Name - Name of the morph target we want to override
+	* @param Weight - How much of this morph target to apply (0.0 - 1.0)
+	*
+	*/
+	void AddAnimCurveOverride( FName& Name, float Weight);
+	void RemoveAnimCurveOverride(FName& Name);
+
+	/**
+	* Tells the AnimInstance to reset all of its morph target curves
+	*
+	*/
+	void ResetAnimCurves();
+	
+	/**
+	* Accessor so our rows can grab the filtertext for highlighting
+	*
+	*/
+	FText& GetFilterText() { return FilterText; }
+
+	/**
+	 * Get Title text
+	 */
+	FText GetTitleText() const;
+	/**
+	 * Refreshes the morph target list after an undo
+	 */
+	void OnPostUndo();
+
+	virtual void Tick(const FGeometry& AllottedGeometry, const double InCurrentTime, const float InDeltaTime);
+
+	void RefreshCurveList();
+
+private:
+
+	/** Handler for context menus */
+	TSharedPtr<SWidget> OnGetContextMenuContent() const;
+
+	/**
+	* Clears and rebuilds the table, according to an optional search string
+	*
+	* @param SearchText - Optional search string
+	*
+	*/
+	void CreateAnimCurveList( const FString& SearchText = FString() );
+	void CreateAnimCurveTypeList(TSharedRef<SVerticalBox> VerticalBox);
+
+	void ApplyCustomCurveOverride(UAnimInstance* AnimInstance) const;
+	void RefreshCachePreviewInstance();
+
+	/** Pointer back to the Persona that owns us */
+	TWeakPtr<FPersona> PersonaPtr;
+
+	/** Box to filter to a specific morph target name */
+	TSharedPtr<SSearchBox>	NameFilterBox;
+
+	/** Widget used to display the list of animation curve */
+	TSharedPtr<SAnimCurveListType> AnimCurveListView;
+
+	/** A list of animation curve. Used by the AnimCurveListView. */
+	TArray< TSharedPtr<FDisplayedAnimCurveInfo> > AnimCurveList;
+
+	/** The skeletal mesh that we grab the animation curve from */
+	UAnimInstance* CachedPreviewInstance;
+
+	/** Current text typed into NameFilterBox */
+	FText FilterText;
+
+	int32 CurrentCurveFlag;
+
+	TMap<FName, float> OverrideCurves;
+
+	/** Curve delegate */
+	FOnAddCustomAnimationCurves OnAddAnimationCurveDelegate;
+
+	friend class SAnimCurveListRow;
+	friend class SAnimCurveTypeList;
+};
\ No newline at end of file
diff --git a/Engine/Source/Editor/Persona/Private/SAnimEditorBase.cpp b/Engine/Source/Editor/Persona/Private/SAnimEditorBase.cpp
--- a/Engine/Source/Editor/Persona/Private/SAnimEditorBase.cpp
+++ b/Engine/Source/Editor/Persona/Private/SAnimEditorBase.cpp
@@ -28,6 +28,19 @@ void SAnimEditorBase::Construct(const FArguments& InArgs)
 
 	SetInputViewRange(0, GetSequenceLength());
 
+	UAnimSequenceBase* SeqBase = Cast<UAnimSequenceBase>(GetEditorObject());
+
+	TSharedPtr<SHorizontalBox> HorizontalBox = SNew(SHorizontalBox);
+	if (SeqBase)
+	{
+		HorizontalBox->AddSlot()
+			.FillWidth(1)
+			[
+				ConstructAnimScrubPanel()
+			];
+	}
+
+
 	this->ChildSlot
 	[
 		SNew(SVerticalBox)
@@ -167,21 +180,19 @@ void SAnimEditorBase::Construct(const FArguments& InArgs)
 		.AutoHeight() 
 		.VAlign(VAlign_Bottom)
 		[
-			SNew(SHorizontalBox)
-			+SHorizontalBox::Slot() 
-			.FillWidth(1)
-			[
-				ConstructAnimScrubPanel()
-			]
+			HorizontalBox.ToSharedRef()
 		]
 	];
 }
 
 TSharedRef<class SAnimationScrubPanel> SAnimEditorBase::ConstructAnimScrubPanel()
 {
-	return SAssignNew( AnimScrubPanel, SAnimationScrubPanel )
+	UAnimSequenceBase* AnimSeqBase = Cast<UAnimSequenceBase>(GetEditorObject());
+	check(AnimSeqBase);
+
+	return SAssignNew(AnimScrubPanel, SAnimationScrubPanel)
 		.Persona(PersonaPtr)
-		.LockedSequence(GetEditorObject())
+		.LockedSequence(AnimSeqBase)
 		.ViewInputMin(this, &SAnimEditorBase::GetViewMinInput)
 		.ViewInputMax(this, &SAnimEditorBase::GetViewMaxInput)
 		.OnSetInputViewRange(this, &SAnimEditorBase::SetInputViewRange)
@@ -305,7 +316,7 @@ float SAnimEditorBase::GetScrubValue() const
 
 void SAnimEditorBase::SetInputViewRange(float InViewMinInput, float InViewMaxInput)
 {
-	ViewMaxInput = FMath::Min<float>(InViewMaxInput, GetEditorObject()->SequenceLength);
+	ViewMaxInput = FMath::Min<float>(InViewMaxInput, GetSequenceLength());
 	ViewMinInput = FMath::Max<float>(InViewMinInput, 0.f);
 }
 
@@ -313,7 +324,7 @@ FText SAnimEditorBase::GetCurrentSequenceTime() const
 {
 	UAnimSingleNodeInstance * PreviewInstance = GetPreviewInstance();
 	float CurTime = 0.f;
-	float TotalTime = GetEditorObject()->SequenceLength;
+	float TotalTime = GetSequenceLength();
 
 	if (PreviewInstance)
 	{
@@ -332,7 +343,7 @@ float SAnimEditorBase::GetPercentageInternal() const
 	float Percentage = 0.f;
 	if (PreviewInstance)
 	{
-		float SequenceLength = GetEditorObject()->SequenceLength;
+		float SequenceLength = GetSequenceLength();
 		if (SequenceLength > 0.f)
 		{
 			Percentage = PreviewInstance->GetCurrentTime() / SequenceLength;
@@ -359,7 +370,12 @@ FText SAnimEditorBase::GetCurrentPercentage() const
 FText SAnimEditorBase::GetCurrentFrame() const
 {
 	float Percentage = GetPercentageInternal();
-	float NumFrames = GetEditorObject()->GetNumberOfFrames();
+	float NumFrames = 0;
+	
+	if (UAnimSequenceBase* AnimSeqBase = Cast<UAnimSequenceBase>(GetEditorObject()))
+	{
+		NumFrames = AnimSeqBase->GetNumberOfFrames();
+	}
 
 	static const FNumberFormattingOptions FractionNumberFormat = FNumberFormattingOptions()
 		.SetMinimumFractionalDigits(2)
@@ -367,4 +383,15 @@ FText SAnimEditorBase::GetCurrentFrame() const
 	return FText::Format(LOCTEXT("FractionKeysFmt", "{0} / {1} (key(s))"), FText::AsNumber(NumFrames * Percentage, &FractionNumberFormat), FText::AsNumber((int32)NumFrames));
 }
 
+float SAnimEditorBase::GetSequenceLength() const
+{
+	UAnimSequenceBase* AnimSeqBase = Cast<UAnimSequenceBase>(GetEditorObject());
+	
+	if (AnimSeqBase)
+	{
+		return AnimSeqBase->SequenceLength;
+	}
+	
+	return 0.f;
+}
 #undef LOCTEXT_NAMESPACE
diff --git a/Engine/Source/Editor/Persona/Private/SAnimEditorBase.h b/Engine/Source/Editor/Persona/Private/SAnimEditorBase.h
--- a/Engine/Source/Editor/Persona/Private/SAnimEditorBase.h
+++ b/Engine/Source/Editor/Persona/Private/SAnimEditorBase.h
@@ -79,10 +79,10 @@ protected:
 	virtual void InitDetailsViewEditorObject(class UEditorAnimBaseObj* EdObj) {};
 
 	/** Get the length of the current sequence */
-	float GetSequenceLength() const {return GetEditorObject()->SequenceLength;}
+	float GetSequenceLength() const;
 
 	/** Get the sequence that is currently being edited */
-	virtual UAnimSequenceBase* GetEditorObject() const PURE_VIRTUAL(SAnimEditorBase::GetEditorObject, return NULL;);
+	virtual UAnimationAsset* GetEditorObject() const PURE_VIRTUAL(SAnimEditorBase::GetEditorObject, return NULL;);
 
 	/** Get Name of Object being edited **/
 	FText GetEditorObjectName() const;
diff --git a/Engine/Source/Editor/Persona/Private/SAnimTrackCurvePanel.cpp b/Engine/Source/Editor/Persona/Private/SAnimTrackCurvePanel.cpp
--- a/Engine/Source/Editor/Persona/Private/SAnimTrackCurvePanel.cpp
+++ b/Engine/Source/Editor/Persona/Private/SAnimTrackCurvePanel.cpp
@@ -332,7 +332,7 @@ void STransformCurveEdTrack::Construct(const FArguments& InArgs)
 							.Font(FEditorStyle::GetFontStyle("CurveEd.InfoFont"))
 							.SelectAllTextWhenFocused(true)
 							.IsReadOnly(true)
-							.Text(this, &STransformCurveEdTrack::GetCurveName, Curve->CurveUid, ETransformCurve::Translation)
+							.Text(this, &STransformCurveEdTrack::GetCurveName, Curve->Name.UID, ETransformCurve::Translation)
 						]
 
 						+SHorizontalBox::Slot()
@@ -397,7 +397,7 @@ void STransformCurveEdTrack::Construct(const FArguments& InArgs)
 							.Font(FEditorStyle::GetFontStyle("CurveEd.InfoFont"))
 							.SelectAllTextWhenFocused(true)
 							.IsReadOnly(true)
-							.Text(this, &STransformCurveEdTrack::GetCurveName, Curve->CurveUid, ETransformCurve::Rotation)
+							.Text(this, &STransformCurveEdTrack::GetCurveName, Curve->Name.UID, ETransformCurve::Rotation)
 						]
 					]
 				]
@@ -445,7 +445,7 @@ void STransformCurveEdTrack::Construct(const FArguments& InArgs)
 							.Font(FEditorStyle::GetFontStyle("CurveEd.InfoFont"))
 							.SelectAllTextWhenFocused(true)
 							.IsReadOnly(true)
-							.Text(this, &STransformCurveEdTrack::GetCurveName, Curve->CurveUid, ETransformCurve::Scale)
+							.Text(this, &STransformCurveEdTrack::GetCurveName, Curve->Name.UID, ETransformCurve::Scale)
 						]
 					]
 				]
@@ -571,11 +571,15 @@ void SAnimTrackCurvePanel::DeleteTrack(USkeleton::AnimCurveUID Uid)
 	if(Sequence->RawCurveData.GetCurveData(Uid, FRawCurveTracks::TransformType))
 	{
 		const FScopedTransaction Transaction( LOCTEXT("AnimCurve_DeleteTrack", "Delete Curve") );
-		Sequence->Modify(true);
-		Sequence->bNeedsRebake = true;
-		Sequence->RawCurveData.DeleteCurveData(Uid, FRawCurveTracks::TransformType);
-		UpdatePanel();
-		WeakPersona.Pin()->RefreshPreviewInstanceTrackCurves();
+		FSmartName CurveToDelete;
+		if (Sequence->GetSkeleton()->GetSmartNameByUID(USkeleton::AnimTrackCurveMappingName, Uid, CurveToDelete))
+		{
+			Sequence->Modify(true);
+			Sequence->bNeedsRebake = true;
+			Sequence->RawCurveData.DeleteCurveData(CurveToDelete, FRawCurveTracks::TransformType);
+			UpdatePanel();
+			WeakPersona.Pin()->RefreshPreviewInstanceTrackCurves();
+		}
 	}
 }
 
@@ -598,8 +602,8 @@ void SAnimTrackCurvePanel::UpdatePanel()
 
 			FName AName;
 			FName BName;
-			MetadataNameMap->GetName(A.CurveUid, AName);
-			MetadataNameMap->GetName(B.CurveUid, BName);
+			MetadataNameMap->GetName(A.Name.UID, AName);
+			MetadataNameMap->GetName(B.Name.UID, BName);
 
 			return AName < BName;
 		});
@@ -626,7 +630,7 @@ void SAnimTrackCurvePanel::UpdatePanel()
 			FName CurveName;
 
 			// if editable, add to the list
-			if(bEditable && NameMapping->GetName(Curve.CurveUid, CurveName))
+			if(bEditable && NameMapping->GetName(Curve.Name.UID, CurveName))
 			{
 				TSharedPtr<STransformCurveEdTrack> CurrentTrack;
 				PanelSlot->AddSlot()
@@ -639,7 +643,7 @@ void SAnimTrackCurvePanel::UpdatePanel()
 					[
 						SAssignNew(CurrentTrack, STransformCurveEdTrack)
 						.Sequence(Sequence)
-						.CurveUid(Curve.CurveUid)
+						.CurveUid(Curve.Name.UID)
 						.AnimTrackCurvePanel(SharedThis(this))
 						.WidgetWidth(WidgetWidth)
 						.ViewInputMin(ViewInputMin)
@@ -687,7 +691,7 @@ TSharedRef<SWidget> SAnimTrackCurvePanel::GenerateCurveList()
 			const FTransformCurve& Curve= *Iter;
 
 			FName CurveName;
-			NameMapping->GetName(Curve.CurveUid, CurveName);
+			NameMapping->GetName(Curve.Name.UID, CurveName);
 
 			ListBox->AddSlot()
 				.AutoHeight()
@@ -695,8 +699,8 @@ TSharedRef<SWidget> SAnimTrackCurvePanel::GenerateCurveList()
 				.Padding( 2.0f, 2.0f )
 				[
 					SNew( SCheckBox )
-					.IsChecked(this, &SAnimTrackCurvePanel::IsCurveEditable, Curve.CurveUid)
-					.OnCheckStateChanged(this, &SAnimTrackCurvePanel::ToggleEditability, Curve.CurveUid)
+					.IsChecked(this, &SAnimTrackCurvePanel::IsCurveEditable, Curve.Name.UID)
+					.OnCheckStateChanged(this, &SAnimTrackCurvePanel::ToggleEditability, Curve.Name.UID)
 					.ToolTipText( LOCTEXT("Show Curves", "Show or Hide Curves") )
 					.IsEnabled( true )
 					[
diff --git a/Engine/Source/Editor/Persona/Private/SAnimTrackCurvePanel.h b/Engine/Source/Editor/Persona/Private/SAnimTrackCurvePanel.h
--- a/Engine/Source/Editor/Persona/Private/SAnimTrackCurvePanel.h
+++ b/Engine/Source/Editor/Persona/Private/SAnimTrackCurvePanel.h
@@ -9,7 +9,7 @@
 #include "SAnimTrackPanel.h"
 #include "SCurveEditor.h"
 #include "SAnimNotifyPanel.h"
-#include "Animation/AnimSequenceBase.h"
+#include "Animation/AnimSequence.h"
 
 //////////////////////////////////////////////////////////////////////////
 // SAnimTrackCurvePanel
diff --git a/Engine/Source/Editor/Persona/Private/SMontageEditor.h b/Engine/Source/Editor/Persona/Private/SMontageEditor.h
--- a/Engine/Source/Editor/Persona/Private/SMontageEditor.h
+++ b/Engine/Source/Editor/Persona/Private/SMontageEditor.h
@@ -58,7 +58,7 @@ public:
 	void SetMontageObj(UAnimMontage * NewMontage);
 	UAnimMontage * GetMontageObj() const { return MontageObj; }
 
-	virtual UAnimSequenceBase* GetEditorObject() const override { return GetMontageObj(); }
+	virtual UAnimationAsset* GetEditorObject() const override { return GetMontageObj(); }
 
 	void RestartPreview();
 	void RestartPreviewFromSection(int32 FromSectionIdx = INDEX_NONE);
diff --git a/Engine/Source/Editor/Persona/Private/SMorphTargetViewer.cpp b/Engine/Source/Editor/Persona/Private/SMorphTargetViewer.cpp
--- a/Engine/Source/Editor/Persona/Private/SMorphTargetViewer.cpp
+++ b/Engine/Source/Editor/Persona/Private/SMorphTargetViewer.cpp
@@ -8,6 +8,7 @@
 #include "ScopedTransaction.h"
 #include "SSearchBox.h"
 #include "Animation/MorphTarget.h"
+#include "Animation/AnimInstance.h"
 
 #define LOCTEXT_NAMESPACE "SMorphTargetViewer"
 
@@ -270,10 +271,13 @@ float SMorphTargetListRow::GetWeight() const
 	if (PersonaPtr.IsValid() && Item->bAutoFillData)
 	{
 		USkeletalMeshComponent* SkelComp = PersonaPtr.Pin()->GetPreviewMeshComponent();
-		if ( SkelComp && SkelComp->GetAnimInstance() )
+		UAnimInstance* AnimInstance = (SkelComp) ? SkelComp->GetAnimInstance() : nullptr;
+		if (AnimInstance)
 		{
 			// make sure if they have value that's not same as saved value
-			float* CurrentValue = SkelComp->GetAnimInstance()->GetMorphTargetCurves().Find(Item->Name);
+			TMap<FName, float> MorphCurves;
+			AnimInstance->GetAnimationCurveList(ACF_DrivesMorphTarget, MorphCurves);
+			const float* CurrentValue = MorphCurves.Find(Item->Name);
 			if (CurrentValue && *CurrentValue != 0.f)
 			{
 				return *CurrentValue;
@@ -428,7 +432,7 @@ void SMorphTargetViewer::CreateMorphTargetList( const FString& SearchText )
 			TSharedRef<FDisplayedMorphTargetInfo> Info = FDisplayedMorphTargetInfo::Make( MorphTargets[I]->GetFName(), NumberOfVerts);
 			if(MeshComponent)
 			{
-				float *CurveValPtr = MeshComponent->MorphTargetCurves.Find( MorphTargets[I]->GetFName() );
+				const float *CurveValPtr = MeshComponent->GetMorphTargetCurves().Find( MorphTargets[I]->GetFName() );
 				if(CurveValPtr)
 				{
 					Info.Get().Weight = (*CurveValPtr);
