Copying //UE4/Dev-Framework to Dev-Main (//UE4/Dev-Main)

diff --git a/Engine/Source/Editor/Persona/Private/SPoseEditor.cpp b/Engine/Source/Editor/Persona/Private/SPoseEditor.cpp
new file mode 100644
--- /dev/null
+++ b/Engine/Source/Editor/Persona/Private/SPoseEditor.cpp
@@ -0,0 +1,819 @@
+// Copyright 1998-2016 Epic Games, Inc. All Rights Reserved.
+
+
+#include "PersonaPrivatePCH.h"
+#include "SPoseEditor.h"
+#include "ObjectTools.h"
+#include "AssetRegistryModule.h"
+#include "ScopedTransaction.h"
+#include "SSearchBox.h"
+#include "Animation/AnimSingleNodeInstance.h"
+
+#define LOCTEXT_NAMESPACE "AnimPoseEditor"
+
+//////////////////////////////////////////////////////////////////////////
+// SPoseEditor
+
+void SPoseEditor::Construct(const FArguments& InArgs)
+{
+	PoseAssetObj = InArgs._PoseAsset;
+	check(PoseAssetObj);
+
+	SAnimEditorBase::Construct( SAnimEditorBase::FArguments()
+		.Persona(InArgs._Persona)
+		);
+
+ 	EditorPanels->AddSlot()
+ 	.AutoHeight()
+ 	.Padding(0, 10)
+	[
+ 		SNew( SPoseViewer)
+ 		.Persona(InArgs._Persona)
+		.PoseAsset(PoseAssetObj)
+ 	];
+}
+
+SPoseEditor::~SPoseEditor()
+{
+}
+
+////////////////////////////////////
+// Copyright 1998-2016 Epic Games, Inc. All Rights Reserved.
+static const FName ColumnId_PoseNameLabel("Pose Name");
+static const FName ColumnID_PoseWeightLabel("Weight");
+static const FName ColumnId_CurveNameLabel("Curve Name");
+
+const float MaxPoseWeight = 1.f;
+
+//////////////////////////////////////////////////////////////////////////
+// SPoseListRow
+
+typedef TSharedPtr< FDisplayedPoseInfo > FDisplayedPoseInfoPtr;
+
+class SPoseListRow
+	: public SMultiColumnTableRow< FDisplayedPoseInfoPtr >
+{
+public:
+
+	SLATE_BEGIN_ARGS(SPoseListRow) {}
+
+	/** The item for this row **/
+	SLATE_ARGUMENT(FDisplayedPoseInfoPtr, Item)
+
+	/* The SPoseViewer that we push the morph target weights into */
+	SLATE_ARGUMENT(class SPoseViewer*, PoseViewer)
+
+	/* Widget used to display the list of animation curve */
+	SLATE_ARGUMENT(TSharedPtr<SPoseListType>, PoseListView)
+
+	/* Persona used to update the viewport when a weight slider is dragged */
+	SLATE_ARGUMENT(TWeakPtr<FPersona>, Persona)
+
+	SLATE_END_ARGS()
+
+	void Construct(const FArguments& InArgs, const TSharedRef<STableViewBase>& OwnerTableView);
+
+	/** Overridden from SMultiColumnTableRow.  Generates a widget for this column of the tree row. */
+	virtual TSharedRef<SWidget> GenerateWidgetForColumn(const FName& ColumnName) override;
+
+private:
+
+	/**
+	* Called when the user changes the value of the SSpinBox
+	*
+	* @param NewWeight - The new number the SSpinBox is set to
+	*
+	*/
+	void OnPoseWeightChanged(float NewWeight);
+	/**
+	* Called when the user types the value and enters
+	*
+	* @param NewWeight - The new number the SSpinBox is set to
+	*
+	*/
+	void OnPoseWeightValueCommitted(float NewWeight, ETextCommit::Type CommitType);
+
+	/** Delegate to get labels root text from settings */
+	FText GetName() const;
+
+	/** Delegate to commit labels root text to settings */
+	void OnNameCommitted(const FText& InText, ETextCommit::Type InCommitType) const;
+
+	/**
+	* Returns the weight of this morph target
+	*
+	* @return SearchText - The new number the SSpinBox is set to
+	*
+	*/
+	float GetWeight() const;
+
+	bool CanChangeWeight() const;
+	/* The SPoseViewer that we push the morph target weights into */
+	SPoseViewer* PoseViewer;
+
+	/** Widget used to display the list of animation curve */
+	TSharedPtr<SPoseListType> PoseListView;
+
+	/** The name and weight of the morph target */
+	FDisplayedPoseInfoPtr	Item;
+
+	/** Pointer back to the Persona that owns us */
+	TWeakPtr<FPersona> PersonaPtr;
+};
+
+FText SPoseListRow::GetName() const
+{
+	return (FText::FromName(Item->Name));
+}
+
+void SPoseListRow::OnNameCommitted(const FText& InText, ETextCommit::Type InCommitType) const
+{
+	// for now only allow enter
+	// because it is important to keep the unique names per pose
+	if (InCommitType == ETextCommit::OnEnter)
+	{
+		FName NewName = FName(*InText.ToString());
+		FName OldName = Item->Name;
+
+		if (PoseViewer->ModifyName(OldName, NewName))
+		{
+			Item->Name = NewName;
+		}
+	}
+}
+
+void SPoseListRow::Construct(const FArguments& InArgs, const TSharedRef<STableViewBase>& InOwnerTableView)
+{
+	Item = InArgs._Item;
+	PoseViewer = InArgs._PoseViewer;
+	PoseListView = InArgs._PoseListView;
+	PersonaPtr = InArgs._Persona;
+
+	check(Item.IsValid());
+
+	SMultiColumnTableRow< FDisplayedPoseInfoPtr >::Construct(FSuperRowType::FArguments(), InOwnerTableView);
+}
+
+TSharedRef< SWidget > SPoseListRow::GenerateWidgetForColumn(const FName& ColumnName)
+{
+	if (ColumnName == ColumnId_PoseNameLabel)
+	{
+		return
+			SNew(SVerticalBox)
+
+			+ SVerticalBox::Slot()
+			.AutoHeight()
+			.Padding(5.0f)
+			.VAlign(VAlign_Center)
+			[
+				SNew(SEditableTextBox)
+				.Text(this, &SPoseListRow::GetName)
+				.ToolTipText(LOCTEXT("PoseName_ToolTip", "Modify Pose Name - Make sure this name is unique among all curves per skeleton."))
+				.OnTextCommitted(this, &SPoseListRow::OnNameCommitted)
+				.OnTextChanged(this, &SPoseListRow::OnNameCommitted, ETextCommit::Default)
+			];
+	}
+	else 
+	{
+		// Encase the SSpinbox in an SVertical box so we can apply padding. Setting ItemHeight on the containing SListView has no effect :-(
+		return
+			SNew(SVerticalBox)
+
+			+ SVerticalBox::Slot()
+			.AutoHeight()
+			.Padding(5.0f)
+			.VAlign(VAlign_Center)
+			[
+				SNew(SSpinBox<float>)
+				.MinSliderValue(-1.f)
+				.MaxSliderValue(1.f)
+				.MinValue(-MaxPoseWeight)
+				.MaxValue(MaxPoseWeight)
+				.Value(this, &SPoseListRow::GetWeight)
+				.OnValueChanged(this, &SPoseListRow::OnPoseWeightChanged)
+				.OnValueCommitted(this, &SPoseListRow::OnPoseWeightValueCommitted)
+				.IsEnabled(this, &SPoseListRow::CanChangeWeight)
+			];
+	}
+}
+
+void SPoseListRow::OnPoseWeightChanged(float NewWeight)
+{
+	Item->Weight = NewWeight;
+
+	PoseViewer->AddPoseOverride(Item->Name, Item->Weight);
+
+	if (PersonaPtr.IsValid())
+	{
+		PersonaPtr.Pin()->RefreshViewport();
+	}
+}
+
+void SPoseListRow::OnPoseWeightValueCommitted(float NewWeight, ETextCommit::Type CommitType)
+{
+	if (CommitType == ETextCommit::OnEnter || CommitType == ETextCommit::OnUserMovedFocus)
+	{
+		float NewValidWeight = FMath::Clamp(NewWeight, -MaxPoseWeight, MaxPoseWeight);
+		Item->Weight = NewValidWeight;
+
+		PoseViewer->AddPoseOverride(Item->Name, Item->Weight);
+
+		TArray< TSharedPtr< FDisplayedPoseInfo > > SelectedRows = PoseListView->GetSelectedItems();
+
+		// ...then any selected rows need changing by the same delta
+		for (auto ItemIt = SelectedRows.CreateIterator(); ItemIt; ++ItemIt)
+		{
+			TSharedPtr< FDisplayedPoseInfo > RowItem = (*ItemIt);
+
+			if (RowItem != Item) // Don't do "this" row again if it's selected
+			{
+				RowItem->Weight = NewValidWeight;
+				PoseViewer->AddPoseOverride(RowItem->Name, RowItem->Weight);
+			}
+		}
+
+		if (PersonaPtr.IsValid())
+		{
+			PersonaPtr.Pin()->RefreshViewport();
+		}
+	}
+}
+
+float SPoseListRow::GetWeight() const
+{
+	return Item->Weight;
+}
+
+bool SPoseListRow::CanChangeWeight() const
+{
+	return PoseViewer->IsBasePose(Item->Name) == false;
+}
+
+//////////////////////////////////////////////////////////////////////////
+// SCurveListRow
+
+typedef TSharedPtr< FDisplayedCurveInfo > FDisplayedCurveInfoPtr;
+
+class SCurveListRow
+	: public SMultiColumnTableRow< FDisplayedCurveInfoPtr >
+{
+public:
+
+	SLATE_BEGIN_ARGS(SCurveListRow) {}
+
+	/** The item for this row **/
+	SLATE_ARGUMENT(FDisplayedCurveInfoPtr, Item)
+
+	/* Widget used to display the list of animation curve */
+	SLATE_ARGUMENT(TSharedPtr<SCurveListType>, CurveListView)
+
+	SLATE_END_ARGS()
+
+	void Construct(const FArguments& InArgs, const TSharedRef<STableViewBase>& OwnerTableView);
+
+	/** Overridden from SMultiColumnTableRow.  Generates a widget for this column of the tree row. */
+	virtual TSharedRef<SWidget> GenerateWidgetForColumn(const FName& ColumnName) override;
+
+private:
+
+	/** Delegate to get labels root text from settings */
+	FText GetName() const;
+
+	/** The name and weight of the morph target */
+	FDisplayedCurveInfoPtr	Item;
+};
+
+FText SCurveListRow::GetName() const
+{
+	return (FText::FromName(Item->Name));
+}
+
+void SCurveListRow::Construct(const FArguments& InArgs, const TSharedRef<STableViewBase>& InOwnerTableView)
+{
+	Item = InArgs._Item;
+
+	check(Item.IsValid());
+
+	SMultiColumnTableRow< FDisplayedCurveInfoPtr >::Construct(FSuperRowType::FArguments(), InOwnerTableView);
+}
+
+TSharedRef< SWidget > SCurveListRow::GenerateWidgetForColumn(const FName& ColumnName)
+{
+	// for now we have one colume
+//	if (ColumnName == ColumnId_CurveNameLabel)
+	{
+		return
+			SNew(SVerticalBox)
+
+			+ SVerticalBox::Slot()
+			.AutoHeight()
+			.Padding(5.0f)
+			.VAlign(VAlign_Center)
+			[
+				SNew(STextBlock)
+				.Text(this, &SCurveListRow::GetName)
+			];
+	}
+}
+
+//////////////////////////////////////////////////////////////////////////
+// SPoseViewer
+
+void SPoseViewer::Construct(const FArguments& InArgs)
+{
+	PersonaPtr = InArgs._Persona;
+	PoseAssetPtr = InArgs._PoseAsset;
+
+	CachedPreviewInstance = nullptr;
+
+	if (PersonaPtr.IsValid())
+	{
+		PersonaPtr.Pin()->RegisterOnPreviewMeshChanged(FPersona::FOnPreviewMeshChanged::CreateSP(this, &SPoseViewer::OnPreviewMeshChanged));
+		PersonaPtr.Pin()->RegisterOnPostUndo(FPersona::FOnPostUndo::CreateSP(this, &SPoseViewer::OnPostUndo));
+	}
+
+	RefreshCachePreviewInstance();
+
+	ChildSlot
+	[
+		SNew(SVerticalBox)
+
+		+ SVerticalBox::Slot()
+		.AutoHeight()
+		.Padding(2)
+		[
+			SNew(STextBlock)
+			.Text(this, &SPoseViewer::GetTitleText)
+		]
+
+		+ SVerticalBox::Slot()
+		.AutoHeight()
+		.Padding(0, 2)
+		[
+			SNew(SHorizontalBox)
+			// Filter entry
+			+ SHorizontalBox::Slot()
+			.FillWidth(1)
+			[
+				SAssignNew(NameFilterBox, SSearchBox)
+				.SelectAllTextWhenFocused(true)
+				.OnTextChanged(this, &SPoseViewer::OnFilterTextChanged)
+				.OnTextCommitted(this, &SPoseViewer::OnFilterTextCommitted)
+			]
+		]
+
+		+ SVerticalBox::Slot()
+		.FillHeight(1.0f)		// This is required to make the scrollbar work, as content overflows Slate containers by default
+		[
+			SNew(SHorizontalBox)
+
+			+SHorizontalBox::Slot()
+			.FillWidth(1)
+			.Padding(5)
+			[
+				SAssignNew(PoseListView, SPoseListType)
+				.ListItemsSource(&PoseList)
+				.OnGenerateRow(this, &SPoseViewer::GeneratePoseRow)
+				.OnContextMenuOpening(this, &SPoseViewer::OnGetContextMenuContent)
+				.ItemHeight(22.0f)
+				.HeaderRow
+				(
+					SNew(SHeaderRow)
+					+ SHeaderRow::Column(ColumnId_PoseNameLabel)
+					.DefaultLabel(LOCTEXT("PoseNameLabel", "Pose Name"))
+
+					+ SHeaderRow::Column(ColumnID_PoseWeightLabel)
+					.DefaultLabel(LOCTEXT("PoseWeightLabel", "Weight"))
+				)
+			]
+
+			+ SHorizontalBox::Slot()
+			.FillWidth(1)
+			.Padding(5)
+			[
+				SAssignNew(CurveListView, SCurveListType)
+				.ListItemsSource(&CurveList)
+				.OnGenerateRow(this, &SPoseViewer::GenerateCurveRow)
+				.OnContextMenuOpening(this, &SPoseViewer::OnGetContextMenuContentForCurveList)
+				.ItemHeight(22.0f)
+				.HeaderRow
+				(
+					SNew(SHeaderRow)
+					+ SHeaderRow::Column(ColumnId_CurveNameLabel)
+					.DefaultLabel(LOCTEXT("CurveNameLabel", "Curve Name"))
+				)
+			]
+		]
+	];
+
+	CreatePoseList();
+	CreateCurveList();
+}
+
+FText SPoseViewer::GetTitleText() const
+{
+	UObject* PreviewAsset = nullptr;
+
+	if (PersonaPtr.IsValid())
+	{
+		PreviewAsset = PersonaPtr.Pin()->GetPreviewAnimationAsset();
+		if (!PreviewAsset)
+		{
+			PreviewAsset = PersonaPtr.Pin()->GetAnimBlueprint();
+		}
+	}
+
+	return (PreviewAsset ? FText::FromString(PreviewAsset->GetName()) : LOCTEXT("PoseMeshNameLabel", "No Asset Present"));
+}
+
+void SPoseViewer::RefreshCachePreviewInstance()
+{
+	if (CachedPreviewInstance.IsValid() && OnAddAnimationCurveDelegate.IsBound())
+	{
+		CachedPreviewInstance.Get()->RemoveDelegate_AddCustomAnimationCurve(OnAddAnimationCurveDelegate);
+	}
+	CachedPreviewInstance = Cast<UAnimSingleNodeInstance>(PersonaPtr.Pin()->GetPreviewMeshComponent()->GetAnimInstance());
+
+	if (CachedPreviewInstance.IsValid())
+	{
+		OnAddAnimationCurveDelegate.Unbind();
+		OnAddAnimationCurveDelegate.BindRaw(this, &SPoseViewer::ApplyCustomCurveOverride);
+		CachedPreviewInstance.Get()->AddDelegate_AddCustomAnimationCurve(OnAddAnimationCurveDelegate);
+	}
+}
+
+void SPoseViewer::OnPreviewMeshChanged(class USkeletalMesh* NewPreviewMesh)
+{
+	RefreshCachePreviewInstance();
+	CreatePoseList(NameFilterBox->GetText().ToString());
+	CreateCurveList(NameFilterBox->GetText().ToString());
+}
+
+void SPoseViewer::OnFilterTextChanged(const FText& SearchText)
+{
+	FilterText = SearchText;
+
+	CreatePoseList(SearchText.ToString());
+	CreateCurveList(SearchText.ToString());
+}
+
+void SPoseViewer::OnFilterTextCommitted(const FText& SearchText, ETextCommit::Type CommitInfo)
+{
+	// Just do the same as if the user typed in the box
+	OnFilterTextChanged(SearchText);
+}
+
+TSharedRef<ITableRow> SPoseViewer::GeneratePoseRow(TSharedPtr<FDisplayedPoseInfo> InInfo, const TSharedRef<STableViewBase>& OwnerTable)
+{
+	check(InInfo.IsValid());
+
+	return
+		SNew(SPoseListRow, OwnerTable)
+		.Persona(PersonaPtr)
+		.Item(InInfo)
+		.PoseViewer(this)
+		.PoseListView(PoseListView);
+}
+
+TSharedRef<ITableRow> SPoseViewer::GenerateCurveRow(TSharedPtr<FDisplayedCurveInfo> InInfo, const TSharedRef<STableViewBase>& OwnerTable)
+{
+	check(InInfo.IsValid());
+
+	return
+		SNew(SCurveListRow, OwnerTable)
+		.Item(InInfo)
+		.CurveListView(CurveListView);
+}
+
+bool SPoseViewer::IsPoseSelected() const
+{
+	// @todo: make sure not to delete base Curve
+	TArray< TSharedPtr< FDisplayedPoseInfo > > SelectedRows = PoseListView->GetSelectedItems();
+	return SelectedRows.Num() > 0;
+}
+
+bool SPoseViewer::IsCurveSelected() const
+{
+	// @todo: make sure not to delete base pose
+	TArray< TSharedPtr< FDisplayedCurveInfo > > SelectedRows = CurveListView->GetSelectedItems();
+	return SelectedRows.Num() > 0;
+}
+
+void SPoseViewer::OnDeletePoses()
+{
+	TArray< TSharedPtr< FDisplayedPoseInfo > > SelectedRows = PoseListView->GetSelectedItems();
+
+	FScopedTransaction Transaction(LOCTEXT("DeletePoses", "Delete Poses"));
+	PoseAssetPtr.Get()->Modify();
+
+	TArray<FName> PosesToDelete;
+	for (int RowIndex = 0; RowIndex < SelectedRows.Num(); ++RowIndex)
+	{
+		PosesToDelete.Add(SelectedRows[RowIndex]->Name);
+	}
+
+	PoseAssetPtr.Get()->DeletePoses(PosesToDelete);
+
+	CreatePoseList(NameFilterBox->GetText().ToString());
+}
+
+void SPoseViewer::OnDeleteCurves()
+{
+	TArray< TSharedPtr< FDisplayedCurveInfo > > SelectedRows = CurveListView->GetSelectedItems();
+
+	FScopedTransaction Transaction(LOCTEXT("DeleteCurves", "Delete Curves"));
+	PoseAssetPtr.Get()->Modify();
+
+	TArray<FName> CurvesToDelete;
+	for (int RowIndex = 0; RowIndex < SelectedRows.Num(); ++RowIndex)
+	{
+		CurvesToDelete.Add(SelectedRows[RowIndex]->Name);
+	}
+
+	PoseAssetPtr.Get()->DeleteCurves(CurvesToDelete);
+
+	CreateCurveList(NameFilterBox->GetText().ToString());
+}
+
+void SPoseViewer::OnPastePoseNamesFromClipBoard(bool bSelectedOnly)
+{
+	FString PastedString;
+
+	FPlatformMisc::ClipboardPaste(PastedString);
+
+	if (PastedString.IsEmpty() == false)
+	{
+		TArray<FString> ListOfNames;
+		PastedString.ParseIntoArrayLines(ListOfNames);
+
+		if (ListOfNames.Num() > 0)
+		{
+			TArray<FName> PosesToRename;
+			if (bSelectedOnly)
+			{
+				TArray< TSharedPtr< FDisplayedPoseInfo > > SelectedRows = PoseListView->GetSelectedItems();
+
+				for (int RowIndex = 0; RowIndex < SelectedRows.Num(); ++RowIndex)
+				{
+					PosesToRename.Add(SelectedRows[RowIndex]->Name);
+				}
+			}
+			else
+			{
+				for (auto& PoseItem : PoseList)
+				{
+					PosesToRename.Add(PoseItem->Name);
+				}
+			}
+
+			if (PosesToRename.Num() > 0)
+			{
+				FScopedTransaction Transaction(LOCTEXT("PasteNames", "Paste Curve Names"));
+				PoseAssetPtr.Get()->Modify();
+
+				int32 TotalItemCount = FMath::Min(PosesToRename.Num(), ListOfNames.Num());
+
+				for (int32 PoseIndex = 0; PoseIndex < TotalItemCount; ++PoseIndex)
+				{ 
+					ModifyName(PosesToRename[PoseIndex], FName(*ListOfNames[PoseIndex]), true);
+				}
+
+				CreatePoseList(NameFilterBox->GetText().ToString());
+			}
+		}
+	}
+
+}
+
+TSharedPtr<SWidget> SPoseViewer::OnGetContextMenuContent() const
+{
+	const bool bShouldCloseWindowAfterMenuSelection = true;
+	FMenuBuilder MenuBuilder(bShouldCloseWindowAfterMenuSelection, NULL);
+
+	FUIAction Action = FUIAction(FExecuteAction::CreateSP(this, &SPoseViewer::OnPastePoseNamesFromClipBoard, false));
+	FText MenuLabel = LOCTEXT("PasteAllPoseNamesButtonLabel", "Paste All Pose Names");
+	FText MenuToolTip = LOCTEXT("PasteAllPoseNamesButtonTooltip", "Paste All Pose Names From Clipboard");
+	MenuBuilder.AddMenuEntry(MenuLabel, MenuToolTip, FSlateIcon(), Action);
+
+	MenuBuilder.BeginSection("PoseAction", LOCTEXT("Poseaction", "Selected Item Actions"));
+	{
+	 	Action = FUIAction( FExecuteAction::CreateSP( this, &SPoseViewer::OnDeletePoses ), 
+	 									FCanExecuteAction::CreateSP( this, &SPoseViewer::IsPoseSelected ) );
+	 	MenuLabel = LOCTEXT("DeletePoseButtonLabel", "Delete");
+	 	MenuToolTip = LOCTEXT("DeletePoseButtonTooltip", "Deletes the selected animation pose.");
+	 	MenuBuilder.AddMenuEntry(MenuLabel, MenuToolTip, FSlateIcon(), Action);
+
+		Action = FUIAction(FExecuteAction::CreateSP(this, &SPoseViewer::OnPastePoseNamesFromClipBoard, true),
+			FCanExecuteAction::CreateSP(this, &SPoseViewer::IsPoseSelected));
+		MenuLabel = LOCTEXT("PastePoseNamesButtonLabel", "Paste Selected");
+		MenuToolTip = LOCTEXT("PastePoseNamesButtonTooltip", "Paste the selected pose names from ClipBoard.");
+		MenuBuilder.AddMenuEntry(MenuLabel, MenuToolTip, FSlateIcon(), Action);
+	}
+	MenuBuilder.EndSection();
+
+	return MenuBuilder.MakeWidget();
+}
+
+TSharedPtr<SWidget> SPoseViewer::OnGetContextMenuContentForCurveList() const
+{
+	const bool bShouldCloseWindowAfterMenuSelection = true;
+	FMenuBuilder MenuBuilder(bShouldCloseWindowAfterMenuSelection, NULL);
+
+	MenuBuilder.BeginSection("CurveAction", LOCTEXT("CurveActions", "Selected Item Actions"));
+	{
+		FUIAction Action = FUIAction(FExecuteAction::CreateSP(this, &SPoseViewer::OnDeleteCurves),
+			FCanExecuteAction::CreateSP(this, &SPoseViewer::IsCurveSelected));
+		const FText MenuLabel = LOCTEXT("DeleteCurveButtonLabel", "Delete");
+		const FText MenuToolTip = LOCTEXT("DeleteCurveButtonTooltip", "Deletes the selected animation curve.");
+		MenuBuilder.AddMenuEntry(MenuLabel, MenuToolTip, FSlateIcon(), Action);
+	}
+	MenuBuilder.EndSection();
+
+	return MenuBuilder.MakeWidget();
+}
+
+void SPoseViewer::CreatePoseList(const FString& SearchText)
+{
+	PoseList.Empty();
+
+	UAnimInstance* AnimInstance = CachedPreviewInstance.IsValid() ? CachedPreviewInstance.Get() : nullptr;
+	if (AnimInstance && PoseAssetPtr.IsValid())
+	{
+		UPoseAsset* PoseAsset = PoseAssetPtr.Get();
+
+		TArray<FSmartName> PoseNames = PoseAsset->GetPoseNames();
+		if (PoseNames.Num() > 0)
+		{
+			bool bDoFiltering = !SearchText.IsEmpty();
+
+			for (const auto& PoseSmartName : PoseNames)
+			{
+				FName PoseName = PoseSmartName.DisplayName;
+				if (bDoFiltering && !PoseName.ToString().Contains(SearchText))
+				{
+					continue; // Skip items that don't match our filter
+				}
+
+				TSharedRef<FDisplayedPoseInfo> Info = FDisplayedPoseInfo::Make(PoseName);
+				PoseList.Add(Info);
+			}
+		}
+		else
+		{
+			// else print somewhere no curves?
+		}
+	}
+
+	PoseListView->RequestListRefresh();
+}
+
+void SPoseViewer::CreateCurveList(const FString& SearchText)
+{
+	CurveList.Empty();
+
+	UAnimInstance* AnimInstance = CachedPreviewInstance.IsValid() ? CachedPreviewInstance.Get() : nullptr;
+	if (AnimInstance && PoseAssetPtr.IsValid())
+	{
+		UPoseAsset* PoseAsset = PoseAssetPtr.Get();
+
+		TArray<FSmartName> CurveNames = PoseAsset->GetCurveNames();
+		if (CurveNames.Num() > 0)
+		{
+			bool bDoFiltering = !SearchText.IsEmpty();
+
+			for (const auto& CurveSmartName : CurveNames)
+			{
+				FName CurveName = CurveSmartName.DisplayName;
+				if (bDoFiltering && !CurveName.ToString().Contains(SearchText))
+				{
+					continue; // Skip items that don't match our filter
+				}
+
+				TSharedRef<FDisplayedCurveInfo> Info = FDisplayedCurveInfo::Make(CurveName);
+				CurveList.Add(Info);
+			}
+		}
+		else
+		{
+			// else print somewhere no curves?
+		}
+	}
+
+	CurveListView->RequestListRefresh();
+}
+
+void SPoseViewer::AddPoseOverride(FName& Name, float Weight)
+{
+	float& Value = OverrideCurves.FindOrAdd(Name);
+	Value = Weight;
+}
+
+void SPoseViewer::RemovePoseOverride(FName& Name)
+{
+	OverrideCurves.Remove(Name);
+}
+
+SPoseViewer::~SPoseViewer()
+{
+	if (PersonaPtr.IsValid())
+	{
+		// @Todo: change this in curve editor
+		// and it won't work with this one delegate idea, so just think of a better way to do
+		// if persona isn't there, we probably don't have the preview mesh either, so no valid anim instance
+		if (CachedPreviewInstance.IsValid() && OnAddAnimationCurveDelegate.IsBound())
+		{
+			CachedPreviewInstance.Get()->RemoveDelegate_AddCustomAnimationCurve(OnAddAnimationCurveDelegate);
+		}
+
+		PersonaPtr.Pin()->UnregisterOnPreviewMeshChanged(this);
+		PersonaPtr.Pin()->UnregisterOnPostUndo(this);
+	}
+}
+
+void SPoseViewer::OnPostUndo()
+{
+	CreatePoseList(NameFilterBox->GetText().ToString());
+	CreateCurveList(NameFilterBox->GetText().ToString());
+}
+
+void SPoseViewer::ApplyCustomCurveOverride(UAnimInstance* AnimInstance) const
+{
+	for (auto Iter = OverrideCurves.CreateConstIterator(); Iter; ++Iter)
+	{
+		// @todo we might want to save original curve flags? or just change curve to apply flags only
+		AnimInstance->AddCurveValue(Iter.Key(), Iter.Value(), ACF_DrivesPose);
+	}
+}
+
+bool SPoseViewer::ModifyName(FName OldName, FName NewName, bool bSilence)
+{
+	if (PersonaPtr.IsValid())
+	{
+		USkeleton* Skeleton = PersonaPtr.Pin()->GetSkeleton();
+		if (Skeleton)
+		{
+			FScopedTransaction Transaction(LOCTEXT("RenamePoses", "Rename Pose"));
+
+			PoseAssetPtr.Get()->Modify();
+			Skeleton->Modify();
+
+			// get smart nam
+			const FSmartNameMapping::UID ExistingUID = Skeleton->GetUIDByName(USkeleton::AnimCurveMappingName, NewName);
+			// verify if this name exists in smart naming
+			if (ExistingUID != FSmartNameMapping::MaxUID)
+			{
+				// warn users
+				// if so, verify if this name is still okay
+				if (!bSilence)
+				{
+					const EAppReturnType::Type Response = FMessageDialog::Open(EAppMsgType::YesNo, LOCTEXT("UseSameNameConfirm", "The name already exists. Would you like to reuse the name? This can cause conflict of curve data."));
+
+					if (Response == EAppReturnType::No)
+					{
+						return false;
+					}
+				}
+
+				// I think this might have to be delegate of the top window
+				if (PoseAssetPtr.Get()->ModifyPoseName(OldName, NewName, &ExistingUID) == false)
+				{
+					FMessageDialog::Open(EAppMsgType::Ok, LOCTEXT("NameAlreadyUsedByTheSameAsset", "The name is used by another pose within the same asset. Please choose another name."));
+					return false;
+				}
+			}
+			else
+			{
+				// I think this might have to be delegate of the top window
+				if (PoseAssetPtr.Get()->ModifyPoseName(OldName, NewName, nullptr) == false)
+				{
+					FMessageDialog::Open(EAppMsgType::Ok, LOCTEXT("NameAlreadyUsedByTheSameAsset", "The name is used by another pose within the same asset. Please choose another name."));
+					return false;
+				}
+			}
+
+			// now refresh pose data
+			float* Value = OverrideCurves.Find(OldName);
+			if (Value)
+			{
+				AddPoseOverride(NewName, *Value);
+				RemovePoseOverride(OldName);
+			}
+
+			return true;
+		}
+	}
+
+	return false;
+}
+
+bool SPoseViewer::IsBasePose(FName PoseName) const
+{
+	if (PoseAssetPtr.IsValid() && PoseAssetPtr->IsValidAdditive())
+	{
+		int32 PoseIndex = PoseAssetPtr->GetPoseIndexByName(PoseName);
+		return (PoseIndex == PoseAssetPtr->GetBasePoseIndex());
+	}
+
+	return false;
+}
+#undef LOCTEXT_NAMESPACE
diff --git a/Engine/Source/Editor/Persona/Private/SRigWindow.cpp b/Engine/Source/Editor/Persona/Private/SRigWindow.cpp
--- a/Engine/Source/Editor/Persona/Private/SRigWindow.cpp
+++ b/Engine/Source/Editor/Persona/Private/SRigWindow.cpp
@@ -15,6 +15,9 @@
 #include "SSearchBox.h"
 #include "SInlineEditableTextBlock.h"
 #include "SRigPicker.h"
+#include "AnimationRuntime.h"
+#include "BoneMappingHelper.h"
+#include "SSkeletonWidget.h"
 
 #define LOCTEXT_NAMESPACE "SRigWindow"
 
@@ -193,7 +196,7 @@ FName SBoneMappingListRow::GetSelectedBone() const
 void SRigWindow::Construct(const FArguments& InArgs)
 {
 	PersonaPtr = InArgs._Persona;
-	Skeleton = NULL;
+	Skeleton = nullptr;
 	bDisplayAdvanced = false;
 
 	if ( PersonaPtr.IsValid() )
@@ -245,9 +248,55 @@ void SRigWindow::Construct(const FArguments& InArgs)
 					.Text(this,&SRigWindow::GetAssetName)
 				]
 			]
+		]
+
+		+ SVerticalBox::Slot()
+		.AutoHeight()
+		.HAlign(HAlign_Center)
+		.VAlign(VAlign_Center)
+		.Padding(5, 5)
+		[
+			SNew(SHorizontalBox)
+
+			+ SHorizontalBox::Slot()
+			.HAlign(HAlign_Right)
+			.Padding(5, 0)
+			[
+				SNew(SButton)
+				.OnClicked(FOnClicked::CreateSP(this, &SRigWindow::OnAutoMapping))
+				.HAlign(HAlign_Center)
+				.VAlign(VAlign_Center)
+				.Text(LOCTEXT("AutoMapping_Title", "Auto  Mapping"))
+				.ToolTipText(LOCTEXT("AutoMapping_Tooltip", "Automatically map the best matching bones"))
+			]
+
+			+SHorizontalBox::Slot()
+			.HAlign(HAlign_Right)
+			.Padding(5, 0)
+			[
+				SNew(SButton)
+				.OnClicked(FOnClicked::CreateSP(this, &SRigWindow::OnClearMapping))
+				.HAlign(HAlign_Center)
+				.VAlign(VAlign_Center)
+				.Text(LOCTEXT("ClearMapping_Title", "Clear Mapping"))
+				.ToolTipText(LOCTEXT("ClearMapping_Tooltip", "Clear currently mapping bones"))
+			]
+
+// 			+ SHorizontalBox::Slot()
+// 			.HAlign(HAlign_Right)
+// 			.Padding(5, 0)
+// 			[
+// 				SNew(SButton)
+// 				.OnClicked(FOnClicked::CreateSP(this, &SRigWindow::OnToggleView))
+// 				.HAlign(HAlign_Center)
+// 				.VAlign(VAlign_Center)
+// 				.Text(LOCTEXT("HierarchyView_Title", "Show Hierarchy"))
+// 				.ToolTipText(LOCTEXT("HierarchyView_Tooltip", "Show Hierarchy View"))
+// 			]
 
 			+SHorizontalBox::Slot()
 			.HAlign(HAlign_Right)
+			.Padding(5, 0)
 			[
 				SNew(SButton)
 				.OnClicked(FOnClicked::CreateSP(this, &SRigWindow::OnToggleAdvanced))
@@ -391,9 +440,9 @@ bool SRigWindow::ShouldFilterAsset(const class FAssetData& AssetData)
 	return (AssetData.GetAsset() == GetRigObject());
 }
 
-UObject* SRigWindow::GetRigObject() const
+URig* SRigWindow::GetRigObject() const
 {
-	return (Skeleton)? Skeleton->GetRig() : NULL;
+	return (Skeleton)? Skeleton->GetRig() : nullptr;
 }
 
 SRigWindow::~SRigWindow()
@@ -459,7 +508,7 @@ void SRigWindow::CloseComboButton()
 
 FText SRigWindow::GetAssetName() const
 {
-	UObject* Rig = GetRigObject();
+	URig* Rig = GetRigObject();
 	if (Rig)
 	{
 		return FText::FromString(Rig->GetName());
@@ -468,5 +517,118 @@ FText SRigWindow::GetAssetName() const
 	return LOCTEXT("None", "None");
 }
 
+bool SRigWindow::OnTargetSkeletonSelected(USkeleton* SelectedSkeleton, URig*  Rig) const
+{
+	if (SelectedSkeleton)
+	{
+		// make sure the skeleton contains all the rig node names
+		const FReferenceSkeleton& RefSkeleton = SelectedSkeleton->GetReferenceSkeleton();
+
+		if (RefSkeleton.GetNum() > 0)
+		{
+			const TArray<FNode> RigNodes = Rig->GetNodes();
+			int32 BoneMatched = 0;
+
+			for (const auto& RigNode : RigNodes)
+			{
+				if (RefSkeleton.FindBoneIndex(RigNode.Name) != INDEX_NONE)
+				{
+					++BoneMatched;
+				}
+			}
+
+			float BoneMatchedPercentage = BoneMatched / RefSkeleton.GetNum();
+			if (BoneMatchedPercentage > 0.5f)
+			{
+				Rig->SetSourceReferenceSkeleton(RefSkeleton);
+
+				return true;
+			}
+		}
+	}
+
+	return false;
+}
+
+bool SRigWindow::SelectSourceReferenceSkeleton(URig* Rig) const
+{
+	TSharedRef<SWindow> WidgetWindow = SNew(SWindow)
+		.Title(LOCTEXT("SelectSourceSkeletonForRig", "Select Source Skeleton for the Rig"))
+		.ClientSize(FVector2D(500, 600));
+
+	TSharedRef<SSkeletonSelectorWindow> SkeletonSelectorWindow = SNew(SSkeletonSelectorWindow).WidgetWindow(WidgetWindow);
+
+	WidgetWindow->SetContent(SkeletonSelectorWindow);
+
+	GEditor->EditorAddModalWindow(WidgetWindow);
+	USkeleton* RigSkeleton = SkeletonSelectorWindow->GetSelectedSkeleton();
+	if (RigSkeleton)
+	{
+		return OnTargetSkeletonSelected(RigSkeleton, Rig);
+	}
+
+	return false;
+}
+
+
+FReply SRigWindow::OnAutoMapping()
+{
+	URig* Rig = GetRigObject();
+	if (Rig)
+	{
+		if (!Rig->IsSourceReferenceSkeletonAvailable())
+		{
+			//ask if they want to set up source skeleton
+			EAppReturnType::Type Response = FMessageDialog::Open(EAppMsgType::YesNo, LOCTEXT("TheRigNeedsSkeleton", 
+				"In order to attempt to auto-map bones, the rig should have the source skeleton. However, the current rig is missing the source skeleton. Would you like to choose one? It's best to select the skeleton this rig is from."));
+
+			if (Response == EAppReturnType::No)
+			{
+				return FReply::Handled();
+			}
+
+			if (!SelectSourceReferenceSkeleton(Rig))
+			{
+				return FReply::Handled();
+			}
+		}
+
+		FReferenceSkeleton RigReferenceSkeleton = Rig->GetSourceReferenceSkeleton();
+		FBoneMappingHelper Helper(RigReferenceSkeleton, Skeleton->GetReferenceSkeleton());
+		TMap<FName, FName> BestMatches;
+		Helper.TryMatch(BestMatches);
+
+		for (const TPair<FName, FName>& BestCandidate : BestMatches)
+		{
+			Skeleton->SetRigBoneMapping(BestCandidate.Key, BestCandidate.Value);
+		}
+
+		BoneMappingListView->RequestListRefresh();
+	}
+
+	return FReply::Handled();
+}
+
+FReply SRigWindow::OnClearMapping()
+{
+	URig* Rig = GetRigObject();
+	if (Rig)
+	{
+		const TArray<FNode> Nodes = Rig->GetNodes();
+		for (const auto& Node : Nodes)
+		{
+			Skeleton->SetRigBoneMapping(Node.Name, NAME_None);
+		}
+
+		BoneMappingListView->RequestListRefresh();
+	}
+	return FReply::Handled();
+}
+
+FReply SRigWindow::OnToggleView()
+{
+	return FReply::Handled();
+}
+
 #undef LOCTEXT_NAMESPACE
 
diff --git a/Engine/Source/Editor/Persona/Private/SSequenceEditor.h b/Engine/Source/Editor/Persona/Private/SSequenceEditor.h
--- a/Engine/Source/Editor/Persona/Private/SSequenceEditor.h
+++ b/Engine/Source/Editor/Persona/Private/SSequenceEditor.h
@@ -36,7 +36,7 @@ public:
 	void Construct(const FArguments& InArgs);
 	virtual ~SSequenceEditor();
 
-	virtual UAnimSequenceBase* GetEditorObject() const override { return SequenceObj; }
+	virtual UAnimationAsset* GetEditorObject() const override { return SequenceObj; }
 
 private:
 	/** Pointer to the animation sequence being edited */
diff --git a/Engine/Source/Editor/Persona/Private/SSkeletonSmartNameManager.cpp b/Engine/Source/Editor/Persona/Private/SSkeletonSmartNameManager.cpp
--- a/Engine/Source/Editor/Persona/Private/SSkeletonSmartNameManager.cpp
+++ b/Engine/Source/Editor/Persona/Private/SSkeletonSmartNameManager.cpp
@@ -10,6 +10,7 @@
 #include "SSearchBox.h"
 #include "SInlineEditableTextBlock.h"
 #include "STextEntryPopup.h"
+#include "Animation/AnimSequence.h"
 
 #define LOCTEXT_NAMESPACE "SkeletonSmartNameManager"
 
@@ -352,8 +353,8 @@ void SSkeletonSmartNameManager::CreateNewNameEntry(const FText& CommittedText, E
 		if (const FSmartNameMapping* NameMapping = CurrentSkeleton->GetSmartNameContainer(ContainerName))
 		{
 			FName NewName = FName(*CommittedText.ToString());
-			FSmartNameMapping::UID NewUid;
-			if(CurrentSkeleton->AddSmartNameAndModify(ContainerName, NewName, NewUid))
+			FSmartName NewCurveName;
+			if(CurrentSkeleton->AddSmartNameAndModify(ContainerName, NewName, NewCurveName))
 			{
 				// Successfully added
 				GenerateDisplayedList(CurrentFilterText);
@@ -453,11 +454,16 @@ void SCurveNameManager::OnDeleteNameClicked()
 			for(FAssetData& Data : AnimationAssets)
 			{
 				UAnimSequenceBase* Sequence = Cast<UAnimSequenceBase>(Data.GetAsset());	
+				USkeleton* MySkeleton = Sequence->GetSkeleton();
 				Sequence->Modify(true);
 				for(FSmartNameMapping::UID Uid : SelectedUids)
 				{
-					Sequence->RawCurveData.DeleteCurveData(Uid);
-					Sequence->MarkRawDataAsModified();
+					FSmartName CurveToDelete;
+					if (MySkeleton->GetSmartNameByUID(USkeleton::AnimCurveMappingName, Uid, CurveToDelete))
+					{
+						Sequence->RawCurveData.DeleteCurveData(CurveToDelete);
+						Sequence->MarkRawDataAsModified();
+					}
 				}
 			}
 			GWarn->EndSlowTask();
diff --git a/Engine/Source/Editor/Persona/Private/Shared/PersonaMode.cpp b/Engine/Source/Editor/Persona/Private/Shared/PersonaMode.cpp
--- a/Engine/Source/Editor/Persona/Private/Shared/PersonaMode.cpp
+++ b/Engine/Source/Editor/Persona/Private/Shared/PersonaMode.cpp
@@ -17,6 +17,7 @@
 
 // Tab constants
 const FName FPersonaTabs::MorphTargetsID("MorphTargetsTab");
+const FName FPersonaTabs::AnimCurveViewID("AnimCurveViewerTab");
 const FName FPersonaTabs::SkeletonTreeViewID("SkeletonTreeView");		//@TODO: Name
 // Skeleton Pose manager
 const FName FPersonaTabs::RetargetManagerID("RetargetManager");
@@ -65,6 +66,7 @@ FPersonaAppMode::FPersonaAppMode(TSharedPtr<class FPersona> InPersona, FName InM
 	PersonaTabFactories.RegisterFactory(MakeShareable(new FAnimationAssetBrowserSummoner(InPersona)));
 	PersonaTabFactories.RegisterFactory(MakeShareable(new FPreviewViewportSummoner(InPersona)));
 	PersonaTabFactories.RegisterFactory(MakeShareable(new FMorphTargetTabSummoner(InPersona)));
+	PersonaTabFactories.RegisterFactory(MakeShareable(new FAnimCurveViewerTabSummoner(InPersona)));
 	PersonaTabFactories.RegisterFactory(MakeShareable(new FSkeletonAnimNotifiesSummoner(InPersona)));
 	PersonaTabFactories.RegisterFactory(MakeShareable(new FRetargetManagerTabSummoner(InPersona)));
 	PersonaTabFactories.RegisterFactory(MakeShareable(new FSkeletonSlotNamesSummoner(InPersona)));
@@ -155,6 +157,29 @@ TSharedRef<SWidget> FMorphTargetTabSummoner::CreateTabBody(const FWorkflowTabSpa
 	return SNew(SMorphTargetViewer)
 		.Persona(StaticCastSharedPtr<FPersona>(HostingApp.Pin()));
 }
+/////////////////////////////////////////////////////
+// FAnimCurveViewerTabSummoner
+
+#include "SAnimCurveViewer.h"
+
+FAnimCurveViewerTabSummoner::FAnimCurveViewerTabSummoner(TSharedPtr<class FAssetEditorToolkit> InHostingApp)
+	: FWorkflowTabFactory(FPersonaTabs::AnimCurveViewID, InHostingApp)
+{
+	TabLabel = LOCTEXT("AnimCurveViewTabTitle", "Anim Curve Previewer");
+	TabIcon = FSlateIcon(FEditorStyle::GetStyleSetName(), "Persona.Tabs.AnimCurvePreviewer");
+
+	EnableTabPadding();
+	bIsSingleton = true;
+
+	ViewMenuDescription = LOCTEXT("AnimCurveTabView", "Animation Curve Previewer");
+	ViewMenuTooltip = LOCTEXT("AnimCurveTabView_ToolTip", "Shows the animation curve viewer");
+}
+
+TSharedRef<SWidget> FAnimCurveViewerTabSummoner::CreateTabBody(const FWorkflowTabSpawnInfo& Info) const
+{
+	return SNew(SAnimCurveViewer)
+		.Persona(StaticCastSharedPtr<FPersona>(HostingApp.Pin()));
+}
 
 /////////////////////////////////////////////////////
 // FAnimationAssetBrowserSummoner
diff --git a/Engine/Source/Editor/Persona/Private/Shared/PersonaMode.h b/Engine/Source/Editor/Persona/Private/Shared/PersonaMode.h
--- a/Engine/Source/Editor/Persona/Private/Shared/PersonaMode.h
+++ b/Engine/Source/Editor/Persona/Private/Shared/PersonaMode.h
@@ -14,6 +14,7 @@ struct FPersonaTabs
 
 	// Selection Details
 	static const FName MorphTargetsID;
+	static const FName AnimCurveViewID;
 	static const FName SkeletonTreeViewID;
 	// Skeleton Pose manager
 	static const FName RetargetManagerID;
@@ -176,6 +177,22 @@ public:
 		return  IDocumentation::Get()->CreateToolTip(LOCTEXT("MorphTargetTooltip", "The Morph Target tab lets you preview any morph targets (aka blend shapes) available for the current mesh."), NULL, TEXT("Shared/Editors/Persona"), TEXT("MorphTarget_Window"));
 	}
 };
+/////////////////////////////////////////////////////
+// FAnimCurveViewerTabSummoner
+
+struct FAnimCurveViewerTabSummoner : public FWorkflowTabFactory
+{
+public:
+	FAnimCurveViewerTabSummoner(TSharedPtr<class FAssetEditorToolkit> InHostingApp);
+
+	virtual TSharedRef<SWidget> CreateTabBody(const FWorkflowTabSpawnInfo& Info) const override;
+
+	// Create a tooltip widget for the tab
+	virtual TSharedPtr<SToolTip> CreateTabToolTipWidget(const FWorkflowTabSpawnInfo& Info) const override
+	{
+		return  IDocumentation::Get()->CreateToolTip(LOCTEXT("AnimCurveViewTooltip", "The Anim Curve Viewer tab lets you preview any animation curves available for the current mesh from preview asset."), NULL, TEXT("Shared/Editors/Persona"), TEXT("AnimCurveView_Window"));
+	}
+};
 
 
 /////////////////////////////////////////////////////
diff --git a/Engine/Source/Editor/PropertyEditor/Private/PropertyNode.cpp b/Engine/Source/Editor/PropertyEditor/Private/PropertyNode.cpp
--- a/Engine/Source/Editor/PropertyEditor/Private/PropertyNode.cpp
+++ b/Engine/Source/Editor/PropertyEditor/Private/PropertyNode.cpp
@@ -2396,10 +2396,11 @@ void FPropertyNode::PropagatePropertyChange( UObject* ModifiedObject, const TCHA
 	UArrayProperty* ParentArrayProp = Cast<UArrayProperty>(ParentProp);
 	UProperty*      Prop            = GetProperty();
 	UMapProperty*   MapProp         = Cast<UMapProperty>(Prop);
+	USetProperty*	SetProp			= Cast<USetProperty>(Prop);
 
-	if (ParentArrayProp != NULL && ParentArrayProp->Inner != Prop)
+	if (ParentArrayProp != nullptr && ParentArrayProp->Inner != Prop)
 	{
-		ParentArrayProp = NULL;
+		ParentArrayProp = nullptr;
 	}
 
 	ObjectsToChange.Push(Object);
@@ -2429,9 +2430,9 @@ void FPropertyNode::PropagatePropertyChange( UObject* ModifiedObject, const TCHA
 			FString OrgValue;
 
 			uint8* Addr = GetValueBaseAddress( (uint8*)ActualObjToChange );
-			if (Addr != NULL)
+			if (Addr != nullptr)
 			{
-				if (MapProp != NULL)
+				if (MapProp != nullptr)
 				{
 					// Read previous value back into object
 					uint8* PreviousMap = (uint8*)FMemory::Malloc(MapProp->GetSize(), MapProp->GetMinAlignment());
@@ -2465,16 +2466,50 @@ void FPropertyNode::PropagatePropertyChange( UObject* ModifiedObject, const TCHA
 						MapProp->SerializeItem(Ar, Addr, ModifiedObjectAddr);
 					}
 				}
+				else if (SetProp != nullptr)
+				{
+					// Read previous value back into object
+					uint8* PreviousSet = (uint8*)FMemory::Malloc(SetProp->GetSize(), SetProp->GetMinAlignment());
+					ON_SCOPE_EXIT
+					{
+						FMemory::Free(PreviousSet);
+					};
+
+					SetProp->InitializeValue(PreviousSet);
+					ON_SCOPE_EXIT
+					{
+						SetProp->DestroyValue(PreviousSet);
+					};
+
+					SetProp->ImportText(*PreviousValue, PreviousSet, PPF_Localized, ModifiedObject);
+
+					uint8* ModifiedObjectAddr = GetValueBaseAddress( (uint8*)ModifiedObject );
+
+					auto ModifiedObjectAddrPtr = (TMap<int32, FString>*)ModifiedObjectAddr;
+
+					// Serialize differences from the 'default' (the old object)
+					TArray<uint8> Data;
+					{
+						FMemoryWriter Ar(Data);
+						SetProp->SerializeItem(Ar, Addr, PreviousSet);
+					}
+
+					// Deserialize differences back over the new object
+					{
+						FMemoryReader Ar(Data);
+						SetProp->SerializeItem(Ar, Addr, ModifiedObjectAddr);
+					}				
+				}
 				else
 				{
-					if (ParentArrayProp != NULL)
+					if (ParentArrayProp != nullptr)
 					{
 						uint8* ArrayAddr = ParentNode->GetValueBaseAddress( (uint8*)ActualObjToChange );
-						ParentArrayProp->ExportText_Direct(OrgValue, ArrayAddr, ArrayAddr, NULL, PPF_Localized );
+						ParentArrayProp->ExportText_Direct(OrgValue, ArrayAddr, ArrayAddr, nullptr, PPF_Localized );
 					}
 					else
 					{
-						Prop->ExportText_Direct(OrgValue, Addr, Addr, NULL, PPF_Localized );
+						Prop->ExportText_Direct(OrgValue, Addr, Addr, nullptr, PPF_Localized );
 					}
 
 					// Check if the original value was the default value and change it only then
diff --git a/Engine/Source/Editor/PropertyEditor/Private/UserInterface/PropertyEditor/SPropertyEditorText.cpp b/Engine/Source/Editor/PropertyEditor/Private/UserInterface/PropertyEditor/SPropertyEditorText.cpp
--- a/Engine/Source/Editor/PropertyEditor/Private/UserInterface/PropertyEditor/SPropertyEditorText.cpp
+++ b/Engine/Source/Editor/PropertyEditor/Private/UserInterface/PropertyEditor/SPropertyEditorText.cpp
@@ -103,6 +103,7 @@ bool SPropertyEditorText::Supports( const TSharedRef< FPropertyEditor >& InPrope
 		||	(Property->IsA(UObjectPropertyBase::StaticClass()) && !Property->HasAnyPropertyFlags(CPF_InstancedReference))
 		||	Property->IsA(UInterfaceProperty::StaticClass())
 		||	Property->IsA(UMapProperty::StaticClass())
+		||	Property->IsA(USetProperty::StaticClass())
 		) )
 	{
 		return true;
diff --git a/Engine/Source/Editor/SequenceRecorder/Private/AnimationRecorder.cpp b/Engine/Source/Editor/SequenceRecorder/Private/AnimationRecorder.cpp
--- a/Engine/Source/Editor/SequenceRecorder/Private/AnimationRecorder.cpp
+++ b/Engine/Source/Editor/SequenceRecorder/Private/AnimationRecorder.cpp
@@ -2,6 +2,7 @@
 
 #include "SequenceRecorderPrivatePCH.h"
 #include "AnimationRecorder.h"
+#include "Animation/AnimInstance.h"
 #include "Animation/AnimCompress.h"
 #include "Animation/AnimCompress_BitwiseCompressOnly.h"
 #include "SCreateAnimationDlg.h"
@@ -10,6 +11,8 @@
 #include "NotificationManager.h"
 #include "EngineLogs.h"
 #include "Animation/AnimationSettings.h"
+#include "Animation/AnimNotifies/AnimNotify.h"
+#include "Animation/AnimNotifies/AnimNotifyState.h"
 #include "Toolkits/AssetEditorManager.h"
 #include "ActorRecording.h"
 
@@ -301,6 +304,7 @@ UAnimSequence* FAnimationRecorder::StopRecord(bool bShowMessage)
 		// add to real curve data 
 		if (RecordedCurves.Num() == NumFrames && UIDList)
 		{
+			USkeleton* SkeletonObj = AnimationObject->GetSkeleton();
 			for (int32 CurveIndex = 0; CurveIndex < (*UIDList).Num(); ++CurveIndex)
 			{
 				USkeleton::AnimCurveUID UID = (*UIDList)[CurveIndex];
@@ -314,8 +318,12 @@ UAnimSequence* FAnimationRecorder::StopRecord(bool bShowMessage)
 					if (FrameIndex == 0)
 					{
 						// add one and save the cache
-						AnimationObject->RawCurveData.AddFloatCurveKey(UID, CurCurve.Flags, TimeToRecord, CurCurve.Value);
-						FloatCurveData = static_cast<FFloatCurve*>(AnimationObject->RawCurveData.GetCurveData(UID, FRawCurveTracks::FloatType));
+						FSmartName CurveName;
+						if (SkeletonObj->GetSmartNameByUID(USkeleton::AnimCurveMappingName, UID, CurveName))
+						{
+							AnimationObject->RawCurveData.AddFloatCurveKey(CurveName, CurCurve.Flags, TimeToRecord, CurCurve.Value);
+							FloatCurveData = static_cast<FFloatCurve*>(AnimationObject->RawCurveData.GetCurveData(UID, FRawCurveTracks::FloatType));
+						}
 					}
 					else if (FloatCurveData)
 					{
diff --git a/Engine/Source/Editor/UnrealEd/Classes/Factories/FbxAnimSequenceImportData.h b/Engine/Source/Editor/UnrealEd/Classes/Factories/FbxAnimSequenceImportData.h
--- a/Engine/Source/Editor/UnrealEd/Classes/Factories/FbxAnimSequenceImportData.h
+++ b/Engine/Source/Editor/UnrealEd/Classes/Factories/FbxAnimSequenceImportData.h
@@ -71,6 +71,16 @@ class UNREALED_API UFbxAnimSequenceImportData : public UFbxAssetImportData
 	UPROPERTY(EditAnywhere, AdvancedDisplay, config, Category = ImportSettings)
 	bool bDeleteExistingMorphTargetCurves;
 
+	UPROPERTY(EditAnywhere, AdvancedDisplay, config, Category = ImportSettings)
+	TArray<FString> MaterialCurveSuffixes;
+
+	UPROPERTY(EditAnywhere, AdvancedDisplay, config, Category = ImportSettings)
+	TArray<FString> PoseCurveSuffixes;
+
+	/** If enabled, this will import a curve within the animation */
+	UPROPERTY(EditAnywhere, AdvancedDisplay, config, Category = ImportSettings)
+	bool bRemoveRedundantKeys;
+
 	/** Gets or creates fbx import data for the specified anim sequence */
 	static UFbxAnimSequenceImportData* GetImportDataForAnimSequence(UAnimSequence* AnimSequence, UFbxAnimSequenceImportData* TemplateForCreation);
 
diff --git a/Engine/Source/Editor/UnrealEd/Classes/Factories/FbxImportUI.h b/Engine/Source/Editor/UnrealEd/Classes/Factories/FbxImportUI.h
--- a/Engine/Source/Editor/UnrealEd/Classes/Factories/FbxImportUI.h
+++ b/Engine/Source/Editor/UnrealEd/Classes/Factories/FbxImportUI.h
@@ -80,9 +80,9 @@ class UFbxImportUI : public UObject
 	UPROPERTY(EditAnywhere, config, Category=Animation, meta=(ImportType="SkeletalMesh|Animation"))
 	uint32 bImportAnimations:1;
 
-	/** Override for the name of the animation to import **/
+	/** Override for the name of the animation to import. By default, it will be the name of FBX **/
 	UPROPERTY(EditAnywhere, AdvancedDisplay, Category=Animation, meta=(editcondition="bImportAnimations", ImportType="SkeletalMesh")) 
-	FString AnimationName;
+	FString OverrideAnimationName;
 
 	/** Enables importing of 'rigid skeletalmesh' (unskinned, hierarchy-based animation) from this FBX file, no longer shown, used behind the scenes */
 	UPROPERTY()
diff --git a/Engine/Source/Editor/UnrealEd/Classes/Factories/PoseAssetFactory.h b/Engine/Source/Editor/UnrealEd/Classes/Factories/PoseAssetFactory.h
new file mode 100644
--- /dev/null
+++ b/Engine/Source/Editor/UnrealEd/Classes/Factories/PoseAssetFactory.h
@@ -0,0 +1,33 @@
+// Copyright 1998-2016 Epic Games, Inc. All Rights Reserved.
+
+/**
+ *
+ */
+
+#pragma once
+#include "PoseAssetFactory.generated.h"
+
+UCLASS(HideCategories=Object,MinimalAPI)
+class UPoseAssetFactory : public UFactory
+{
+	GENERATED_UCLASS_BODY()
+
+	UPROPERTY()
+	class USkeleton* TargetSkeleton;
+
+	/* Used when creating a composite from an AnimSequence, becomes the only AnimSequence contained */
+	UPROPERTY()
+	class UAnimSequence* SourceAnimation;
+
+	//~ Begin UFactory Interface
+	virtual bool ConfigureProperties() override;
+	virtual UObject* FactoryCreateNew(UClass* Class,UObject* InParent,FName Name,EObjectFlags Flags,UObject* Context,FFeedbackContext* Warn) override;
+	//~ Begin UFactory Interface	
+
+private:
+	void OnTargetSkeletonSelected(const FAssetData& SelectedAsset);
+
+private:
+	TSharedPtr<SWindow> PickerWindow;
+};
+
diff --git a/Engine/Source/Editor/UnrealEd/Private/Animation/DebugSkelMeshComponent.cpp b/Engine/Source/Editor/UnrealEd/Private/Animation/DebugSkelMeshComponent.cpp
--- a/Engine/Source/Editor/UnrealEd/Private/Animation/DebugSkelMeshComponent.cpp
+++ b/Engine/Source/Editor/UnrealEd/Private/Animation/DebugSkelMeshComponent.cpp
@@ -560,7 +560,7 @@ void UDebugSkelMeshComponent::ToggleClothSectionsVisibility(bool bShowOnlyClothS
 				if (bShowOnlyClothSections)
 				{
 					// enables only cloth sections
-					if (LODModel.Chunks[Section.ChunkIndex].HasApexClothData())
+					if (Section.HasApexClothData())
 					{
 						Section.bDisabled = false;
 					}
@@ -571,7 +571,7 @@ void UDebugSkelMeshComponent::ToggleClothSectionsVisibility(bool bShowOnlyClothS
 				}
 				else
 				{   // disables cloth sections and also corresponding original sections
-					if (LODModel.Chunks[Section.ChunkIndex].HasApexClothData())
+					if (Section.HasApexClothData())
 					{
 						Section.bDisabled = true;
 						LODModel.Sections[Section.CorrespondClothSectionIndex].bDisabled = true;
@@ -615,7 +615,7 @@ void UDebugSkelMeshComponent::RestoreClothSectionsVisibility()
 			{
 				FSkelMeshSection& Section = LODModel.Sections[SecIdx];
 
-				if(LODModel.Chunks[Section.ChunkIndex].HasApexClothData())
+				if(Section.HasApexClothData())
 				{
 					LODModel.Sections[Section.CorrespondClothSectionIndex].bDisabled = true;
 				}
@@ -662,7 +662,7 @@ int32 UDebugSkelMeshComponent::FindCurrentSectionDisplayMode()
 		{
 			FSkelMeshSection& Section = LODModel.Sections[SecIdx];
 
-			if (LODModel.Chunks[Section.ChunkIndex].HasApexClothData())
+			if (Section.HasApexClothData())
 			{
 				// Normal state if the cloth section is visible and the corresponding section is disabled
 				if (Section.bDisabled == false &&
@@ -682,7 +682,7 @@ int32 UDebugSkelMeshComponent::FindCurrentSectionDisplayMode()
 			FSkelMeshSection& Section = LODModel.Sections[SecIdx];
 
 			// not related to cloth sections
-			if (!LODModel.Chunks[Section.ChunkIndex].HasApexClothData() &&
+			if (!Section.HasApexClothData() &&
 				Section.CorrespondClothSectionIndex < 0)
 			{
 				bFoundNonClothSection = true;
diff --git a/Engine/Source/Editor/UnrealEd/Private/AnimationCompressionPanel.cpp b/Engine/Source/Editor/UnrealEd/Private/AnimationCompressionPanel.cpp
--- a/Engine/Source/Editor/UnrealEd/Private/AnimationCompressionPanel.cpp
+++ b/Engine/Source/Editor/UnrealEd/Private/AnimationCompressionPanel.cpp
@@ -4,9 +4,19 @@
 #include "Animation/AnimCompress.h"
 #include "AnimationCompressionPanel.h"
 #include "AnimationEditorUtils.h"
+#include "AnimationUtils.h"
+
+#include "Editor/PropertyEditor/Public/PropertyEditorModule.h"
+#include "Editor/PropertyEditor/Public/IDetailsView.h"
 
 #define LOCTEXT_NAMESPACE "AnimationCompression"
 
+UCompressionHolder::UCompressionHolder(const FObjectInitializer& ObjectInitializer)
+	: Super(ObjectInitializer)
+{
+	Compression = FAnimationUtils::GetDefaultAnimationCompressionAlgorithm();
+}
+
 //////////////////////////////////////////////
 //	FDlgAnimCompression
 
@@ -17,7 +27,8 @@ FDlgAnimCompression::FDlgAnimCompression( TArray<TWeakObjectPtr<UAnimSequence>>&
 		DialogWindow = SNew(SWindow)
 			.Title( LOCTEXT("AnimCompression", "Animation Compression") )
 			.SupportsMinimize(false) .SupportsMaximize(false)
-			.SizingRule( ESizingRule::Autosized );
+			.SizingRule( ESizingRule::UserSized )
+			.ClientSize(FVector2D(400, 500));
 
 		TSharedPtr<SBorder> DialogWrapper = 
 			SNew(SBorder)
@@ -43,85 +54,76 @@ void SAnimationCompressionPanel::Construct(const FArguments& InArgs)
 	ParentWindow = InArgs._ParentWindow.Get();
 
 	AnimSequences = InArgs._AnimSequences;
-	CurrentCompressionChoice = 0;
+	CompressionHolder = NewObject<UCompressionHolder>();
+	CompressionHolder->AddToRoot();
 
-	for ( TObjectIterator<UClass> It ; It ; ++It )
+	if (AnimSequences.Num() == 1)
 	{
-		UClass* Class = *It;
-		if( !Class->HasAnyClassFlags(CLASS_Abstract | CLASS_Deprecated) )
-		{
-			if ( Class->IsChildOf(UAnimCompress::StaticClass()) )
-			{
-				UAnimCompress* NewAlgorithm = NewObject<UAnimCompress>(GetTransientPackage(), Class);
-				AnimationCompressionAlgorithms.Add( NewAlgorithm );
-			}
-		}
+		UAnimSequence* Seq = AnimSequences[0].Get();
+
+		CompressionHolder->Compression = static_cast<UAnimCompress*>(StaticDuplicateObject(Seq->CompressionScheme, CompressionHolder));
 	}
 
+	FPropertyEditorModule& EditModule = FModuleManager::Get().GetModuleChecked<FPropertyEditorModule>("PropertyEditor");
+
+	const bool bAllowSearch = true;
+	FDetailsViewArgs DetailsViewArgs(/*bUpdateFromSelection=*/ false, /*bLockable=*/ false, bAllowSearch, FDetailsViewArgs::HideNameArea, /*bHideSelectionTip=*/ true);
+
+	TSharedPtr<class IDetailsView> PropertyView = EditModule.CreateDetailView(DetailsViewArgs);
+
+	TArray<UObject*> SelectedObjects;
+	SelectedObjects.Add(CompressionHolder);
+	PropertyView->SetObjects(SelectedObjects);
+
 	TSharedRef<SVerticalBox> Box =
 		SNew(SVerticalBox)
-		+SVerticalBox::Slot()
-		.AutoHeight()
+		+ SVerticalBox::Slot()
+		.FillHeight(1.f)
 		.Padding(8.0f, 4.0f, 8.0f, 4.0f)
 		[
-			SNew( STextBlock )
-			.Text( LOCTEXT("AnimCompressionLabel", "Compression Algorithms") )
+			PropertyView.ToSharedRef()
 		]
 		+SVerticalBox::Slot()
 		.AutoHeight()
-		.Padding(8.0f, 4.0f, 8.0f, 4.0f)
+		.Padding(8.0f, 4.0f, 4.0f, 8.0f)
 		[
 			SNew(SSeparator)
-		];
-
-	for( int32 i = 0 ; i < AnimationCompressionAlgorithms.Num() ; ++i )
-	{
-		UAnimCompress* Alg = AnimationCompressionAlgorithms[i];
-		(*Box).AddSlot()
-		.Padding(8.0f, 4.0f, 8.0f, 4.0f)
+		]
+		+SVerticalBox::Slot()
+		.Padding(4.0f)
+		.HAlign(HAlign_Right)
+		.VAlign(VAlign_Bottom)
+		.AutoHeight()
 		[
-			CreateRadioButton( FText::FromString( Alg->Description ), i)
-		];
-	}
-
-	(*Box)	.AddSlot()
-			.AutoHeight()
-			.Padding(8.0f, 4.0f, 4.0f, 8.0f)
-			[
-				SNew(SSeparator)
-			];
-	(*Box)	.AddSlot()
-			.Padding(4.0f)
-			.HAlign(HAlign_Right)
-			.VAlign(VAlign_Bottom)
-			.AutoHeight()
+			SNew(SUniformGridPanel)
+			.SlotPadding(FEditorStyle::GetMargin("StandardDialog.SlotPadding"))
+			.MinDesiredSlotWidth(FEditorStyle::GetFloat("StandardDialog.MinDesiredSlotWidth"))
+			.MinDesiredSlotHeight(FEditorStyle::GetFloat("StandardDialog.MinDesiredSlotHeight"))
+			+ SUniformGridPanel::Slot(0, 0)
 			[
-				SNew(SUniformGridPanel)
-				.SlotPadding(FEditorStyle::GetMargin("StandardDialog.SlotPadding"))
-				.MinDesiredSlotWidth(FEditorStyle::GetFloat("StandardDialog.MinDesiredSlotWidth"))
-				.MinDesiredSlotHeight(FEditorStyle::GetFloat("StandardDialog.MinDesiredSlotHeight"))
-				+SUniformGridPanel::Slot(0,0)
-				[
-					SNew(SButton)
-					.Text(LOCTEXT("AnimCompressionApply", "Apply"))
-					.HAlign(HAlign_Center)
-					.ContentPadding(FEditorStyle::GetMargin("StandardDialog.ContentPadding"))
-					.OnClicked(this, &SAnimationCompressionPanel::ApplyClicked)
-				]
-			];
+				SNew(SButton)
+				.Text(LOCTEXT("AnimCompressionApply", "Apply"))
+			.HAlign(HAlign_Center)
+			.ContentPadding(FEditorStyle::GetMargin("StandardDialog.ContentPadding"))
+			.OnClicked(this, &SAnimationCompressionPanel::ApplyClicked)
+			]
+		];
 	
 	this->ChildSlot[Box];
 }
 
 SAnimationCompressionPanel::~SAnimationCompressionPanel()
 {
-
+	if (CompressionHolder)
+	{
+		CompressionHolder->Compression = nullptr;
+		CompressionHolder->RemoveFromRoot();
+	}
 }
 
 FReply SAnimationCompressionPanel::ApplyClicked()
 {
-	UAnimCompress* CompressionAlgorithm = AnimationCompressionAlgorithms[CurrentCompressionChoice];
-	ApplyAlgorithm( CompressionAlgorithm );
+	ApplyAlgorithm( CompressionHolder->Compression );
 	return FReply::Handled();
 }
 
@@ -139,7 +141,11 @@ void SAnimationCompressionPanel::ApplyAlgorithm(class UAnimCompress* Algorithm)
 
 		if (AnimationEditorUtils::ApplyCompressionAlgorithm(AnimSequencePtrs, Algorithm))
 		{
-			ParentWindow->RequestDestroyWindow();
+			TSharedPtr<SWindow> Win = ParentWindow.Pin();
+			if (Win.IsValid())
+			{
+				Win->RequestDestroyWindow();
+			}
 		}
 	}
 }
diff --git a/Engine/Source/Editor/UnrealEd/Private/ApexClothingUtils.cpp b/Engine/Source/Editor/UnrealEd/Private/ApexClothingUtils.cpp
--- a/Engine/Source/Editor/UnrealEd/Private/ApexClothingUtils.cpp
+++ b/Engine/Source/Editor/UnrealEd/Private/ApexClothingUtils.cpp
@@ -772,18 +772,15 @@ bool AssociateClothingAssetWithSkeletalMesh(USkeletalMesh* SkelMesh, int32 LODIn
 
 	//TempChunk makes it easy to go back to previous status when error occurred
 	//to check errors about bone-mapping first
-	FSkelMeshChunk		TempChunk; 
+	FSkelMeshSection	TempSection; 
 
 	uint32 NumImportedVertices = ImportData.Positions.Num();
 
-	TempChunk.NumRigidVertices = 0;
-	TempChunk.RigidVertices.Empty();
-	TempChunk.NumSoftVertices = NumImportedVertices;
-	TempChunk.SoftVertices.Empty(NumImportedVertices);
+	TempSection.SoftVertices.Empty(NumImportedVertices);
 
 	for( uint32 VertexIndex=0;VertexIndex<NumImportedVertices;VertexIndex++ )
 	{
-		FSoftSkinVertex& OutVertex = TempChunk.SoftVertices[TempChunk.SoftVertices.AddZeroed()];
+		FSoftSkinVertex& OutVertex = TempSection.SoftVertices[TempSection.SoftVertices.AddZeroed()];
 
 		OutVertex.Position	= ImportData.Positions[VertexIndex];
 		OutVertex.TangentX	= ImportData.Tangents[VertexIndex];
@@ -803,19 +800,19 @@ bool AssociateClothingAssetWithSkeletalMesh(USkeletalMesh* SkelMesh, int32 LODIn
 
 	// Set the bone map
 	int32 NumUsedBones = UsedBones.Num();
-	TempChunk.BoneMap.Empty(NumUsedBones);
+	TempSection.BoneMap.Empty(NumUsedBones);
 
 	for( int32 UsedBoneIndex=0; UsedBoneIndex<NumUsedBones; UsedBoneIndex++ )
 	{
 		const int32 BoneIdx = SkelMesh->RefSkeleton.FindBoneIndex(UsedBones[UsedBoneIndex]);
 		if(BoneIdx != INDEX_NONE)
 		{
-			TempChunk.BoneMap.Add(BoneIdx);
+			TempSection.BoneMap.Add(BoneIdx);
 		}
 	}
 
 	//if bone count is 0, then crash will occur while drawing
-	if(TempChunk.BoneMap.Num() == 0)
+	if(TempSection.BoneMap.Num() == 0)
 	{
 		FString FailedBones = "\nMapping failed bones : \n";
 
@@ -847,13 +844,12 @@ bool AssociateClothingAssetWithSkeletalMesh(USkeletalMesh* SkelMesh, int32 LODIn
 	// Check the influences on the original chunk to make sure the influences match
 	{
 		FSkelMeshSection& CheckOriginMeshSection = LODModel.Sections[OriginSectionIndex];
-		FSkelMeshChunk& CheckOriginChunk = LODModel.Chunks[CheckOriginMeshSection.ChunkIndex];
-		if(CheckOriginChunk.MaxBoneInfluences > MAX_INFLUENCES_PER_STREAM)
+		if(CheckOriginMeshSection.MaxBoneInfluences > MAX_INFLUENCES_PER_STREAM)
 		{
-			const FText Text = FText::Format(LOCTEXT("Error_TooManyInfluences", "Chunk {0} in Skeletal Mesh {1} has up to {2} influences on it's vertices. The maximum when using cloth is 4, reduce the number of influences and reimport the mesh to allow cloth on this mesh."),
-				FText::AsNumber(CheckOriginMeshSection.ChunkIndex),
+			const FText Text = FText::Format(LOCTEXT("Error_TooManyInfluences", "Section {0} in Skeletal Mesh {1} has up to {2} influences on it's vertices. The maximum when using cloth is 4, reduce the number of influences and reimport the mesh to allow cloth on this mesh."),
+				FText::AsNumber(OriginSectionIndex),
 				FText::FromString(SkelMesh->GetName()),
-				FText::AsNumber(CheckOriginChunk.MaxBoneInfluences));
+				FText::AsNumber(CheckOriginMeshSection.MaxBoneInfluences));
 
 			FMessageDialog::Open(EAppMsgType::Ok, Text);
 
@@ -870,35 +866,38 @@ bool AssociateClothingAssetWithSkeletalMesh(USkeletalMesh* SkelMesh, int32 LODIn
 
 	SkelMesh->PreEditChange(NULL);
 
-	//add new section & chunk for clothing
-	int32 NewSectionIndex = LODModel.Sections.AddZeroed();
-	int16 NewChunkIndex = LODModel.Chunks.Add(TempChunk);
-
+	// Get original section
 	FSkelMeshSection& OriginMeshSection = LODModel.Sections[OriginSectionIndex];
-	FSkelMeshChunk& OriginChunk = LODModel.Chunks[OriginMeshSection.ChunkIndex];
 
+	// Add new section for clothing
+	int32 NewSectionIndex = LODModel.Sections.AddZeroed();
 	FSkelMeshSection& ClothSection = LODModel.Sections[NewSectionIndex];
 
-	// copy default info from original section
-	ClothSection = OriginMeshSection;
-
-	FSkelMeshChunk& ClothChunk = LODModel.Chunks[NewChunkIndex];
-	
-	// copy data from original chunk
-	ClothChunk.BaseVertexIndex = OriginChunk.BaseVertexIndex;
-	ClothChunk.MaxBoneInfluences = OriginChunk.MaxBoneInfluences;
+	// Copy 'chunk' properties from TempSection
+	// int16 NewChunkIndex = LODModel.Chunks.Add(TempChunk);
+	ClothSection.SoftVertices = TempSection.SoftVertices;
+	ClothSection.BoneMap = TempSection.BoneMap;
+	ClothSection.MaxBoneInfluences = TempSection.MaxBoneInfluences;
+	ClothSection.CorrespondClothAssetIndex = TempSection.CorrespondClothAssetIndex;
+	ClothSection.ClothAssetSubmeshIndex = TempSection.ClothAssetSubmeshIndex;
+
+	// copy 'material' properties from original section
+	// ClothSection = OriginMeshSection;
+	ClothSection.MaterialIndex = OriginMeshSection.MaterialIndex;
+	ClothSection.TriangleSorting = OriginMeshSection.TriangleSorting;
+	ClothSection.bSelected = OriginMeshSection.bSelected;
+	ClothSection.bRecomputeTangent = OriginMeshSection.bRecomputeTangent;
+	ClothSection.bEnableClothLOD_DEPRECATED = OriginMeshSection.bEnableClothLOD_DEPRECATED;
 
 	OriginMeshSection.bDisabled = true;
-	// assign newly added chunk
-	ClothSection.ChunkIndex = NewChunkIndex;
 	ClothSection.bDisabled = false;
 
 	ClothSection.CorrespondClothSectionIndex = OriginSectionIndex;
 	OriginMeshSection.CorrespondClothSectionIndex = NewSectionIndex;
 
-	ClothChunk.ApexClothMappingData = RenderToPhysicalMapping;
-	ClothChunk.PhysicalMeshVertices = PhysicalMeshVertices;
-	ClothChunk.PhysicalMeshNormals = PhysicalMeshNormals;
+	ClothSection.ApexClothMappingData = RenderToPhysicalMapping;
+	ClothSection.PhysicalMeshVertices = PhysicalMeshVertices;
+	ClothSection.PhysicalMeshNormals = PhysicalMeshNormals;
 
 	//adjust index buffer
 	TArray<uint32> OutIndexBuffer;
@@ -910,25 +909,26 @@ bool AssociateClothingAssetWithSkeletalMesh(USkeletalMesh* SkelMesh, int32 LODIn
 	uint32 NumVertices = 0;
 	for(int32 SectionIdx = 0; SectionIdx < NewSectionIndex; SectionIdx++)
 	{
-		uint16 ChunkIdx = LODModel.Sections[SectionIdx].ChunkIndex;
-		FSkelMeshChunk& Chunk = LODModel.Chunks[ChunkIdx]; 	
-		NumVertices += Chunk.GetNumVertices();
+		FSkelMeshSection& Section = LODModel.Sections[SectionIdx];
+		NumVertices += Section.GetNumVertices();
 	}
 
-	ClothChunk.BaseVertexIndex = NumVertices;
+	ClothSection.BaseVertexIndex = NumVertices;
 
 	// if Section Index is greater than 0, Vertex Index has to be adjusted
 	if(NewSectionIndex > 0)
 	{
-		for(int32 Index=0; Index<IndexBuffer.Num(); Index++)
-			IndexBuffer[Index] += ClothChunk.BaseVertexIndex;
+		for (int32 Index = 0; Index < IndexBuffer.Num(); Index++)
+		{
+			IndexBuffer[Index] += ClothSection.BaseVertexIndex;
+		}
 	}
 
 	OutIndexBuffer += IndexBuffer;
 
 	LODModel.MultiSizeIndexContainer.CopyIndexBuffer(OutIndexBuffer);
 
-	LODModel.NumVertices += ClothChunk.GetNumVertices();
+	LODModel.NumVertices += ClothSection.GetNumVertices();
 
 	// build adjacency information for this clothing section
 	{
@@ -946,12 +946,12 @@ bool AssociateClothingAssetWithSkeletalMesh(USkeletalMesh* SkelMesh, int32 LODIn
 
 
 	// set asset submesh info to new chunk 
-	ClothChunk.SetClothSubmeshIndex(AssetIndex, AssetSubmeshIndex);
+	ClothSection.SetClothSubmeshIndex(AssetIndex, AssetSubmeshIndex);
 
 	// It's possible that the required/active bones for the new cloth chunk differ from
 	// the original chunk - here we inspect the bone map and activate any new bones.
 	bool bRequiredBonesChanged = false;
-	for(FBoneIndexType BoneIdx : ClothChunk.BoneMap)
+	for(FBoneIndexType BoneIdx : ClothSection.BoneMap)
 	{
 		if(!LODModel.RequiredBones.Contains(BoneIdx))
 		{
@@ -1268,13 +1268,13 @@ void RemoveAssetFromSkeletalMesh(USkeletalMesh* SkelMesh, uint32 AssetIndex, boo
 
 		FStaticLODModel& LODModel = ImportedResource->LODModels[LODIdx];
 
-		int32 NumChunks = LODModel.Chunks.Num();
+		int32 NumSections = LODModel.Sections.Num();
 
 		//decrease asset index because one asset is removed
-		for(int32 ChunkIdx=0; ChunkIdx < NumChunks; ChunkIdx++)
+		for(int32 SectionIdx=0; SectionIdx < NumSections; SectionIdx++)
 		{
-			if(LODModel.Chunks[ChunkIdx].CorrespondClothAssetIndex > (int16)AssetIndex)
-				LODModel.Chunks[ChunkIdx].CorrespondClothAssetIndex--;
+			if(LODModel.Sections[SectionIdx].CorrespondClothAssetIndex > (int16)AssetIndex)
+				LODModel.Sections[SectionIdx].CorrespondClothAssetIndex--;
 		}
 	}
 
@@ -1302,13 +1302,11 @@ void RestoreOriginalClothingSection(USkeletalMesh* SkelMesh, uint32 LODIndex, ui
 
 	FSkelMeshSection& Section = LODModel.Sections[SectionIndex];
 
-	uint16 ChunkIndex = Section.ChunkIndex;
-
 	int32 OriginSectionIndex = -1, ClothSectionIndex = -1;
 
 	// if this section is original mesh section ( without clothing data )
 	if(Section.CorrespondClothSectionIndex >= 0
-		&& !LODModel.Chunks[ChunkIndex].HasApexClothData())
+		&& !Section.HasApexClothData())
 	{
 		ClothSectionIndex = Section.CorrespondClothSectionIndex;
 		OriginSectionIndex = SectionIndex;
@@ -1316,7 +1314,7 @@ void RestoreOriginalClothingSection(USkeletalMesh* SkelMesh, uint32 LODIndex, ui
 
 	// if this section is clothing section
 	if(Section.CorrespondClothSectionIndex >= 0
-		&&LODModel.Chunks[ChunkIndex].HasApexClothData())
+		&& Section.HasApexClothData())
 	{
 		ClothSectionIndex = SectionIndex;
 		OriginSectionIndex = Section.CorrespondClothSectionIndex;
@@ -1334,17 +1332,13 @@ void RestoreOriginalClothingSection(USkeletalMesh* SkelMesh, uint32 LODIndex, ui
 		TArray<uint32> OutIndexBuffer;
 		LODModel.MultiSizeIndexContainer.GetIndexBuffer(OutIndexBuffer);
 
-		int16 RemovedChunkIndex = ClothSection.ChunkIndex;
 		uint32 RemovedBaseIndex = ClothSection.BaseIndex;
 		uint32 RemovedNumIndices = ClothSection.NumTriangles * 3;
-
-		FSkelMeshChunk& ClothChunk = LODModel.Chunks[ClothSection.ChunkIndex];
-
-		uint32 RemovedBaseVertexIndex = ClothChunk.BaseVertexIndex;
+		uint32 RemovedBaseVertexIndex = ClothSection.BaseVertexIndex;
 
 		int32 NumIndexBuffer = OutIndexBuffer.Num();
 
-		int32 NumRemovedVertices = ClothChunk.GetNumVertices();
+		int32 NumRemovedVertices = ClothSection.GetNumVertices();
 
 		// Remove old indices and update any indices for other sections
 		OutIndexBuffer.RemoveAt(RemovedBaseIndex, RemovedNumIndices);
@@ -1353,13 +1347,12 @@ void RestoreOriginalClothingSection(USkeletalMesh* SkelMesh, uint32 LODIndex, ui
 
 		for(int32 i=0; i < NumIndexBuffer; i++)
 		{
-			if(OutIndexBuffer[i] >= ClothChunk.BaseVertexIndex)
+			if(OutIndexBuffer[i] >= ClothSection.BaseVertexIndex)
 				OutIndexBuffer[i] -= NumRemovedVertices;
 		}
 
 		LODModel.MultiSizeIndexContainer.CopyIndexBuffer(OutIndexBuffer);
 
-		LODModel.Chunks.RemoveAt(ClothSection.ChunkIndex);
 		LODModel.Sections.RemoveAt(OriginSection.CorrespondClothSectionIndex);
 
 		LODModel.NumVertices -= NumRemovedVertices;
@@ -1371,17 +1364,11 @@ void RestoreOriginalClothingSection(USkeletalMesh* SkelMesh, uint32 LODIndex, ui
 			if(LODModel.Sections[i].CorrespondClothSectionIndex > OriginSection.CorrespondClothSectionIndex)
 				LODModel.Sections[i].CorrespondClothSectionIndex -= 1;
 
-			if(LODModel.Sections[i].ChunkIndex > RemovedChunkIndex)
-				LODModel.Sections[i].ChunkIndex -= 1;
-
 			if(LODModel.Sections[i].BaseIndex > RemovedBaseIndex)
 				LODModel.Sections[i].BaseIndex -= RemovedNumIndices;
 
-			int16 NewChunkIndex = LODModel.Sections[i].ChunkIndex;
-
-			if(LODModel.Chunks[NewChunkIndex].BaseVertexIndex > RemovedBaseVertexIndex)
-				LODModel.Chunks[NewChunkIndex].BaseVertexIndex -= NumRemovedVertices;
-
+			if(LODModel.Sections[i].BaseVertexIndex > RemovedBaseVertexIndex)
+				LODModel.Sections[i].BaseVertexIndex -= NumRemovedVertices;
 		}
 
 		OriginSection.bDisabled = false;
@@ -1532,9 +1519,8 @@ void ReImportClothingSectionFromClothingAsset(USkeletalMesh* SkelMesh, int32 LOD
 		return;
 	}
 
-	int16 ClothChunkIdx = Model.Sections[ClothSecIdx].ChunkIndex;
-	int16 AssetIndex = Model.Chunks[ClothChunkIdx].CorrespondClothAssetIndex;
-	int16 SubmeshIdx = Model.Chunks[ClothChunkIdx].ClothAssetSubmeshIndex;
+	int16 AssetIndex = Model.Sections[ClothSecIdx].CorrespondClothAssetIndex;
+	int16 SubmeshIdx = Model.Sections[ClothSecIdx].ClothAssetSubmeshIndex;
 
 	ImportClothingSectionFromClothingAsset(SkelMesh, LODIndex, SectionIndex, AssetIndex, SubmeshIdx);
 }
@@ -1629,8 +1615,7 @@ EClothUtilRetType ImportApexAssetFromApexFile(FString& ApexFile, USkeletalMesh*
 				FSkelMeshSection& Section = LODModel.Sections[SectionIndices[i]];
 
 				int16 ClothSecIdx = Section.CorrespondClothSectionIndex;
-				int16 ClothChunkIdx = LODModel.Sections[ClothSecIdx].ChunkIndex;
-				int16 SubmeshIdx = LODModel.Chunks[ClothChunkIdx].ClothAssetSubmeshIndex;
+				int16 SubmeshIdx = LODModel.Sections[ClothSecIdx].ClothAssetSubmeshIndex;
 
 				int32 SecIdx = SectionIndices[i];
 
@@ -1820,7 +1805,6 @@ void BackupClothingDataFromSkeletalMesh(USkeletalMesh* SkelMesh, FClothingBackup
 	FStaticLODModel& LODModel = ImportedResource->LODModels[0];
 
 	int32 NumSections = LODModel.Sections.Num();
-	int32 NumChunks = LODModel.Chunks.Num();
 
 	int32 NumClothingSections = 0;
 	int32 NumIndexBuffer = 0;
@@ -1833,14 +1817,13 @@ void BackupClothingDataFromSkeletalMesh(USkeletalMesh* SkelMesh, FClothingBackup
 			FSkelMeshSection& Section = LODModel.Sections[i];
 			NumClothingSections++;
 			NumIndexBuffer += (Section.NumTriangles * 3);
-			NumBoneIndices += LODModel.Chunks[Section.ChunkIndex].BoneMap.Num();
+			NumBoneIndices += Section.BoneMap.Num();
 		}
 	}
 
 	if(NumClothingSections > 0)
 	{
 		ClothingBackup.Sections.Empty(NumClothingSections);
-		ClothingBackup.Chunks.Empty(NumClothingSections);
 		ClothingBackup.IndexBuffer.Empty(NumIndexBuffer);
 
 		TArray<uint32> OutIndexBuffer;
@@ -1852,19 +1835,15 @@ void BackupClothingDataFromSkeletalMesh(USkeletalMesh* SkelMesh, FClothingBackup
 			{
 				FSkelMeshSection& Section = LODModel.Sections[SecIdx];
 				int32 BackupSectionIndex = ClothingBackup.Sections.Add(Section);
-				FSkelMeshChunk& Chunk = LODModel.Chunks[Section.ChunkIndex];
-				int32 BackupChunkIndex = ClothingBackup.Chunks.Add(Chunk);
 				// back up index buffer
 				FSkelMeshSection& NewSection = ClothingBackup.Sections[BackupSectionIndex];
-				FSkelMeshChunk& NewChunk = ClothingBackup.Chunks[BackupChunkIndex];
 
-				NewSection.ChunkIndex = BackupChunkIndex;
 				int32 BaseIndex = NewSection.BaseIndex;
 				NewSection.BaseIndex = ClothingBackup.IndexBuffer.Num();
 				int32 NumIndices = NewSection.NumTriangles * 3;
 				int32 EndIndex = BaseIndex + NumIndices;
 
-				int32 BaseVertexIndex = NewChunk.BaseVertexIndex;
+				int32 BaseVertexIndex = NewSection.BaseVertexIndex;
 
 				for(int32 Idx=BaseIndex; Idx < EndIndex; Idx++)
 				{
@@ -1908,16 +1887,12 @@ void ReapplyClothingDataToSkeletalMesh( USkeletalMesh* SkelMesh, FClothingBackup
 			continue;
 		}
 
-		int16 ChunkIndex = Section.ChunkIndex;
-
-		FSkelMeshChunk& Chunk = ClothingBackup.Chunks[ChunkIndex];
-
 		ImportClothingSectionFromClothingAsset(
 			SkelMesh, 
 			0,
 			Section.CorrespondClothSectionIndex,
-			Chunk.CorrespondClothAssetIndex,
-			Chunk.ClothAssetSubmeshIndex
+			Section.CorrespondClothAssetIndex,
+			Section.ClothAssetSubmeshIndex
 		);
 
 	}
diff --git a/Engine/Source/Editor/UnrealEd/Private/Commandlets/CookCommandlet.cpp b/Engine/Source/Editor/UnrealEd/Private/Commandlets/CookCommandlet.cpp
--- a/Engine/Source/Editor/UnrealEd/Private/Commandlets/CookCommandlet.cpp
+++ b/Engine/Source/Editor/UnrealEd/Private/Commandlets/CookCommandlet.cpp
@@ -89,7 +89,7 @@ namespace DetailedCookStats
 			// This value is set by an INI private to Epic.
 			if (GConfig->GetString(TEXT("CookAnalytics"), TEXT("APIServer"), APIServerET, GEngineIni))
 			{
-				TSharedPtr<IAnalyticsProviderET> CookAnalytics = FAnalyticsET::Get().CreateAnalyticsProvider(FAnalyticsET::Config(TEXT("Cook"), APIServerET));
+				TSharedPtr<IAnalyticsProviderET> CookAnalytics = FAnalyticsET::Get().CreateAnalyticsProvider(FAnalyticsET::Config(TEXT("Cook"), APIServerET, FString(), true));
 				if (CookAnalytics.IsValid())
 				{
 					// start the session
diff --git a/Engine/Source/Editor/UnrealEd/Private/EditorAnimUtils.cpp b/Engine/Source/Editor/UnrealEd/Private/EditorAnimUtils.cpp
--- a/Engine/Source/Editor/UnrealEd/Private/EditorAnimUtils.cpp
+++ b/Engine/Source/Editor/UnrealEd/Private/EditorAnimUtils.cpp
@@ -530,9 +530,7 @@ namespace EditorAnimUtils
 	{
 		// Copy curve data from source asset, preserving data in the target if present.
 		const FSmartNameMapping* OldNameMapping = OldSkeleton->GetSmartNameContainer(ContainerName);
-		const FSmartNameMapping* NewNameMapping = NewSkeleton->GetOrAddSmartNameContainer(ContainerName);
-
-		SequenceBase->RawCurveData.UpdateLastObservedNames(OldNameMapping, CurveType);
+		SequenceBase->RawCurveData.RefreshName(OldNameMapping, CurveType);
 
 		switch (CurveType)
 		{
@@ -540,7 +538,7 @@ namespace EditorAnimUtils
 			{
 				for(FFloatCurve& Curve : SequenceBase->RawCurveData.FloatCurves)
 				{
-					NewSkeleton->AddSmartNameAndModify(ContainerName, Curve.LastObservedName, Curve.CurveUid);
+					NewSkeleton->AddSmartNameAndModify(ContainerName, Curve.Name.DisplayName, Curve.Name);
 				}
 				break;
 			}
@@ -548,7 +546,7 @@ namespace EditorAnimUtils
 			{
 				for(FVectorCurve& Curve : SequenceBase->RawCurveData.VectorCurves)
 				{
-					NewSkeleton->AddSmartNameAndModify(ContainerName, Curve.LastObservedName, Curve.CurveUid);
+					NewSkeleton->AddSmartNameAndModify(ContainerName, Curve.Name.DisplayName, Curve.Name);
 				}
 				break;
 			}
@@ -556,7 +554,7 @@ namespace EditorAnimUtils
 			{
 				for(FTransformCurve& Curve : SequenceBase->RawCurveData.TransformCurves)
 				{
-					NewSkeleton->AddSmartNameAndModify(ContainerName, Curve.LastObservedName, Curve.CurveUid);
+					NewSkeleton->AddSmartNameAndModify(ContainerName, Curve.Name.DisplayName, Curve.Name);
 				}
 				break;
 			}
diff --git a/Engine/Source/Editor/UnrealEd/Private/Factories/EditorFactories.cpp b/Engine/Source/Editor/UnrealEd/Private/Factories/EditorFactories.cpp
--- a/Engine/Source/Editor/UnrealEd/Private/Factories/EditorFactories.cpp
+++ b/Engine/Source/Editor/UnrealEd/Private/Factories/EditorFactories.cpp
@@ -95,6 +95,7 @@
 #include "Animation/BlendSpace.h"
 #include "Animation/AimOffsetBlendSpace.h"
 #include "Animation/AimOffsetBlendSpace1D.h"
+#include "Animation/AnimInstance.h"
 #include "Engine/UserDefinedEnum.h"
 #include "Engine/UserDefinedStruct.h"
 #include "GameFramework/ForceFeedbackEffect.h"
@@ -6031,7 +6032,7 @@ EReimportResult::Type UReimportFbxSkeletalMeshFactory::Reimport( UObject* Obj )
 	ReimportUI->bCreatePhysicsAsset = false;
 	ReimportUI->PhysicsAsset = SkeletalMesh->PhysicsAsset;
 	ReimportUI->bImportAnimations = false;
-	ReimportUI->AnimationName = TEXT("");
+	ReimportUI->OverrideAnimationName = TEXT("");
 	ReimportUI->bImportRigidMesh = false;
 
 	bool bSuccess = false;
diff --git a/Engine/Source/Editor/UnrealEd/Private/Factories/PoseAssetFactory.cpp b/Engine/Source/Editor/UnrealEd/Private/Factories/PoseAssetFactory.cpp
new file mode 100644
--- /dev/null
+++ b/Engine/Source/Editor/UnrealEd/Private/Factories/PoseAssetFactory.cpp
@@ -0,0 +1,87 @@
+// Copyright 1998-2016 Epic Games, Inc. All Rights Reserved.
+
+/*=============================================================================
+	PoseAssetFactory.cpp: Factory for PoseAsset
+=============================================================================*/
+
+#include "UnrealEd.h"
+
+#include "AssetData.h"
+#include "ContentBrowserModule.h"
+
+#include "Animation/PoseAsset.h"
+
+#define LOCTEXT_NAMESPACE "PoseAssetFactory"
+
+UPoseAssetFactory::UPoseAssetFactory(const FObjectInitializer& ObjectInitializer)
+	: Super(ObjectInitializer)
+{
+	bCreateNew = true;
+	SupportedClass = UPoseAsset::StaticClass();
+}
+
+bool UPoseAssetFactory::ConfigureProperties()
+{
+	// Null the skeleton so we can check for selection later
+	TargetSkeleton = NULL;
+	SourceAnimation = NULL;
+
+	// Load the content browser module to display an asset picker
+	FContentBrowserModule& ContentBrowserModule = FModuleManager::LoadModuleChecked<FContentBrowserModule>("ContentBrowser");
+
+	FAssetPickerConfig AssetPickerConfig;
+	
+	/** The asset picker will only show skeletons */
+	AssetPickerConfig.Filter.ClassNames.Add(USkeleton::StaticClass()->GetFName());
+	AssetPickerConfig.Filter.bRecursiveClasses = true;
+
+	/** The delegate that fires when an asset was selected */
+	AssetPickerConfig.OnAssetSelected = FOnAssetSelected::CreateUObject(this, &UPoseAssetFactory::OnTargetSkeletonSelected);
+
+	/** The default view mode should be a list view */
+	AssetPickerConfig.InitialAssetViewType = EAssetViewType::List;
+
+	PickerWindow = SNew(SWindow)
+	.Title(LOCTEXT("CreatePoseAssetOptions", "Pick Skeleton"))
+	.ClientSize(FVector2D(500, 600))
+	.SupportsMinimize(false) .SupportsMaximize(false)
+	[
+		SNew(SBorder)
+		.BorderImage( FEditorStyle::GetBrush("Menu.Background") )
+		[
+			ContentBrowserModule.Get().CreateAssetPicker(AssetPickerConfig)
+		]
+	];
+
+	GEditor->EditorAddModalWindow(PickerWindow.ToSharedRef());
+	PickerWindow.Reset();
+
+	return TargetSkeleton != NULL;
+}
+
+UObject* UPoseAssetFactory::FactoryCreateNew(UClass* Class, UObject* InParent, FName Name, EObjectFlags Flags, UObject* Context, FFeedbackContext* Warn)
+{
+	if (TargetSkeleton || SourceAnimation)
+	{
+		UPoseAsset* PoseAsset = NewObject<UPoseAsset>(InParent, Class, Name, Flags);
+
+		if(SourceAnimation)
+		{
+			PoseAsset->CreatePoseFromAnimation( SourceAnimation );
+		}
+		
+		PoseAsset->SetSkeleton( TargetSkeleton );
+
+		return PoseAsset;
+	}
+
+	return NULL;
+}
+
+void UPoseAssetFactory::OnTargetSkeletonSelected(const FAssetData& SelectedAsset)
+{
+	TargetSkeleton = Cast<USkeleton>(SelectedAsset.GetAsset());
+	PickerWindow->RequestDestroyWindow();
+}
+
+#undef LOCTEXT_NAMESPACE
diff --git a/Engine/Source/Editor/UnrealEd/Private/Factories/SkeletalMeshImport.cpp b/Engine/Source/Editor/UnrealEd/Private/Factories/SkeletalMeshImport.cpp
--- a/Engine/Source/Editor/UnrealEd/Private/Factories/SkeletalMeshImport.cpp
+++ b/Engine/Source/Editor/UnrealEd/Private/Factories/SkeletalMeshImport.cpp
@@ -943,22 +943,22 @@ void RestoreExistingSkelMeshData(ExistingSkelMeshData* MeshData, USkeletalMesh*
 					LODModel.RequiredBones.Sort();
 					LODModel.ActiveBoneIndices.Sort();
 
-					// Fix the chunks' BoneMaps.
-					for (int32 ChunkIndex = 0; ChunkIndex < LODModel.Chunks.Num(); ChunkIndex++)
+					// Fix the sections' BoneMaps.
+					for (int32 SectionIndex = 0; SectionIndex < LODModel.Sections.Num(); SectionIndex++)
 					{
-						FSkelMeshChunk& Chunk = LODModel.Chunks[ChunkIndex];
-						for (int32 BoneIndex = 0; BoneIndex < Chunk.BoneMap.Num(); BoneIndex++)
+						FSkelMeshSection& Section = LODModel.Sections[SectionIndex];
+						for (int32 BoneIndex = 0; BoneIndex < Section.BoneMap.Num(); BoneIndex++)
 						{
-							int32 NewBoneIndex = OldToNewMap[Chunk.BoneMap[BoneIndex]];
+							int32 NewBoneIndex = OldToNewMap[Section.BoneMap[BoneIndex]];
 							if (NewBoneIndex == INDEX_NONE)
 							{
 								bMissingBone = true;
-								MissingBoneName = MeshData->ExistingRefSkeleton.GetBoneName(Chunk.BoneMap[BoneIndex]);
+								MissingBoneName = MeshData->ExistingRefSkeleton.GetBoneName(Section.BoneMap[BoneIndex]);
 								break;
 							}
 							else
 							{
-								Chunk.BoneMap[BoneIndex] = NewBoneIndex;
+								Section.BoneMap[BoneIndex] = NewBoneIndex;
 							}
 						}
 						if (bMissingBone)
diff --git a/Engine/Source/Editor/UnrealEd/Private/Fbx/FbxAnimSequenceImportData.cpp b/Engine/Source/Editor/UnrealEd/Private/Fbx/FbxAnimSequenceImportData.cpp
--- a/Engine/Source/Editor/UnrealEd/Private/Fbx/FbxAnimSequenceImportData.cpp
+++ b/Engine/Source/Editor/UnrealEd/Private/Fbx/FbxAnimSequenceImportData.cpp
@@ -1,13 +1,18 @@
 // Copyright 1998-2016 Epic Games, Inc. All Rights Reserved.
 
 #include "UnrealEd.h"
+#include "Animation/AnimSequence.h"
 
 UFbxAnimSequenceImportData::UFbxAnimSequenceImportData(const FObjectInitializer& ObjectInitializer)
 	: Super(ObjectInitializer)
 	, bImportCustomAttribute(true)
+	, bRemoveRedundantKeys(true)
 {
 	FrameImportRange.Min = 0;
 	FrameImportRange.Max = 0;
+
+	MaterialCurveSuffixes.Add(TEXT("_mat"));
+	PoseCurveSuffixes.Add(TEXT("_pose"));
 }
 
 UFbxAnimSequenceImportData* UFbxAnimSequenceImportData::GetImportDataForAnimSequence(UAnimSequence* AnimSequence, UFbxAnimSequenceImportData* TemplateForCreation)
diff --git a/Engine/Source/Editor/UnrealEd/Private/Fbx/FbxMainImport.cpp b/Engine/Source/Editor/UnrealEd/Private/Fbx/FbxMainImport.cpp
--- a/Engine/Source/Editor/UnrealEd/Private/Fbx/FbxMainImport.cpp
+++ b/Engine/Source/Editor/UnrealEd/Private/Fbx/FbxMainImport.cpp
@@ -234,13 +234,16 @@ void ApplyImportUIToImportOptions(UFbxImportUI* ImportUI, FBXImportOptions& InOu
 	InOutImportOptions.AnimationLengthImportType = ImportUI->AnimSequenceImportData->AnimationLength;
 	InOutImportOptions.AnimationRange.X = ImportUI->AnimSequenceImportData->FrameImportRange.Min;
 	InOutImportOptions.AnimationRange.Y = ImportUI->AnimSequenceImportData->FrameImportRange.Max;
-	InOutImportOptions.AnimationName = ImportUI->AnimationName;
+	InOutImportOptions.AnimationName = ImportUI->OverrideAnimationName;
 	// only re-sample if they don't want to use default sample rate
 	InOutImportOptions.bResample = !ImportUI->AnimSequenceImportData->bUseDefaultSampleRate;
 	InOutImportOptions.bPreserveLocalTransform = ImportUI->AnimSequenceImportData->bPreserveLocalTransform;
 	InOutImportOptions.bDeleteExistingMorphTargetCurves = ImportUI->AnimSequenceImportData->bDeleteExistingMorphTargetCurves;
+	InOutImportOptions.bRemoveRedundantKeys = ImportUI->AnimSequenceImportData->bRemoveRedundantKeys;
 	InOutImportOptions.bImportCustomAttribute = ImportUI->AnimSequenceImportData->bImportCustomAttribute;
 	InOutImportOptions.bSetMaterialDriveParameterOnCustomAttribute = ImportUI->AnimSequenceImportData->bSetMaterialDriveParameterOnCustomAttribute;
+	InOutImportOptions.MaterialCurveSuffixes = ImportUI->AnimSequenceImportData->MaterialCurveSuffixes;
+	InOutImportOptions.PoseCurveSuffixes = ImportUI->AnimSequenceImportData->PoseCurveSuffixes;
 }
 
 void FImportedMaterialData::AddImportedMaterial( FbxSurfaceMaterial& FbxMaterial, UMaterialInterface& UnrealMaterial )
diff --git a/Engine/Source/Editor/UnrealEd/Private/Fbx/FbxSkeletalMeshExport.cpp b/Engine/Source/Editor/UnrealEd/Private/Fbx/FbxSkeletalMeshExport.cpp
--- a/Engine/Source/Editor/UnrealEd/Private/Fbx/FbxSkeletalMeshExport.cpp
+++ b/Engine/Source/Editor/UnrealEd/Private/Fbx/FbxSkeletalMeshExport.cpp
@@ -7,6 +7,7 @@
 #include "UnrealEd.h"
 
 #include "FbxExporter.h"
+#include "Animation/AnimSequence.h"
 
 DEFINE_LOG_CATEGORY_STATIC(LogFbxSkeletalMeshExport, Log, All);
 
@@ -304,33 +305,18 @@ void FFbxExporter::BindMeshToSkeleton(const USkeletalMesh* SkelMesh, FbxNode* Me
 		// list of verts.  The convert the chunk bone index to the mesh bone index, the chunk's
 		// boneMap is needed
 		int32 VertIndex = 0;
-		const int32 ChunkCount = SourceModel.Chunks.Num();
-		for(int32 ChunkIndex = 0; ChunkIndex < ChunkCount; ++ChunkIndex)
+		const int32 SectionCount = SourceModel.Sections.Num();
+		for(int32 SectionIndex = 0; SectionIndex < SectionCount; ++SectionIndex)
 		{
-			const FSkelMeshChunk& Chunk = SourceModel.Chunks[ChunkIndex];
+			const FSkelMeshSection& Section = SourceModel.Sections[SectionIndex];
 
-			for(int32 RigidIndex = 0; RigidIndex < Chunk.NumRigidVertices; ++RigidIndex)
+			for(int32 SoftIndex = 0; SoftIndex < Section.SoftVertices.Num(); ++SoftIndex)
 			{
-				const FRigidSkinVertex& Vert = Chunk.RigidVertices[RigidIndex];
-
-				int32 InfluenceBone		= Chunk.BoneMap[ Vert.Bone ];
-				float InfluenceWeight	= 1.f;
-
-				if(InfluenceBone == BoneIndex && InfluenceWeight > 0.f)
-				{
-					CurrentCluster->AddControlPointIndex(VertIndex, InfluenceWeight);
-				}
-
-				++VertIndex;
-			}
-
-			for(int32 SoftIndex = 0; SoftIndex < Chunk.NumSoftVertices; ++SoftIndex)
-			{
-				const FSoftSkinVertex& Vert = Chunk.SoftVertices[SoftIndex];
+				const FSoftSkinVertex& Vert = Section.SoftVertices[SoftIndex];
 
 				for(int32 InfluenceIndex = 0; InfluenceIndex < MAX_TOTAL_INFLUENCES; ++InfluenceIndex)
 				{
-					int32 InfluenceBone		= Chunk.BoneMap[ Vert.InfluenceBones[InfluenceIndex] ];
+					int32 InfluenceBone		= Section.BoneMap[ Vert.InfluenceBones[InfluenceIndex] ];
 					float InfluenceWeight	= Vert.InfluenceWeights[InfluenceIndex] / 255.f;
 
 					if(InfluenceBone == BoneIndex && InfluenceWeight > 0.f)
diff --git a/Engine/Source/Editor/UnrealEd/Private/Fbx/FbxSkeletalMeshImport.cpp b/Engine/Source/Editor/UnrealEd/Private/Fbx/FbxSkeletalMeshImport.cpp
--- a/Engine/Source/Editor/UnrealEd/Private/Fbx/FbxSkeletalMeshImport.cpp
+++ b/Engine/Source/Editor/UnrealEd/Private/Fbx/FbxSkeletalMeshImport.cpp
@@ -2721,13 +2721,13 @@ bool UnFbx::FFbxImporter::ImportSkeletalMeshLOD(USkeletalMesh* InSkeletalMesh, U
 
 	// Enforce LODs having only single-influence vertices.
 	bool bCheckSingleInfluence;
-	GConfig->GetBool(TEXT("AnimSetViewer"), TEXT("CheckSingleInfluenceLOD"), bCheckSingleInfluence, GEditorIni);
+	GConfig->GetBool(TEXT("ImportSetting"), TEXT("CheckSingleInfluenceLOD"), bCheckSingleInfluence, GEditorIni);
 	if (bCheckSingleInfluence &&
 		DesiredLOD > 0)
 	{
-		for (int32 ChunkIndex = 0; ChunkIndex < NewLODModel.Chunks.Num(); ChunkIndex++)
+		for (int32 SectionIndex = 0; SectionIndex < NewLODModel.Sections.Num(); SectionIndex++)
 		{
-			if (NewLODModel.Chunks[ChunkIndex].SoftVertices.Num() > 0)
+			if (NewLODModel.Sections[SectionIndex].SoftVertices.Num() > 0)
 			{
 				AddTokenizedErrorMessage(FTokenizedMessage::Create(EMessageSeverity::Warning, FText(LOCTEXT("LODHasSoftVertices", "Warning: The mesh LOD you are importing has some vertices with more than one influence."))), FFbxErrors::SkeletalMesh_LOD_HasSoftVerts);
 			}
@@ -2767,15 +2767,15 @@ bool UnFbx::FFbxImporter::ImportSkeletalMeshLOD(USkeletalMesh* InSkeletalMesh, U
 	}
 
 	// Fix up the chunk BoneMaps.
-	for (int32 ChunkIndex = 0; ChunkIndex < NewLODModel.Chunks.Num(); ChunkIndex++)
+	for (int32 SectionIndex = 0; SectionIndex < NewLODModel.Sections.Num(); SectionIndex++)
 	{
-		FSkelMeshChunk& Chunk = NewLODModel.Chunks[ChunkIndex];
-		for (int32 i = 0; i < Chunk.BoneMap.Num(); i++)
+		FSkelMeshSection& Section = NewLODModel.Sections[SectionIndex];
+		for (int32 i = 0; i < Section.BoneMap.Num(); i++)
 		{
-			int32 LODBoneIndex = Chunk.BoneMap[i];
+			int32 LODBoneIndex = Section.BoneMap[i];
 			FName LODBoneName = InSkeletalMesh->RefSkeleton.GetBoneName(LODBoneIndex);
 			int32 BaseBoneIndex = BaseSkeletalMesh->RefSkeleton.FindBoneIndex(LODBoneName);
-			Chunk.BoneMap[i] = BaseBoneIndex;
+			Section.BoneMap[i] = BaseBoneIndex;
 		}
 	}
 
@@ -2801,25 +2801,17 @@ bool UnFbx::FFbxImporter::ImportSkeletalMeshLOD(USkeletalMesh* InSkeletalMesh, U
 	// To be extra-nice, we apply the difference between the root transform of the meshes to the verts.
 	FMatrix LODToBaseTransform = InSkeletalMesh->GetRefPoseMatrix(0).InverseFast() * BaseSkeletalMesh->GetRefPoseMatrix(0);
 
-	for (int32 ChunkIndex = 0; ChunkIndex < NewLODModel.Chunks.Num(); ChunkIndex++)
+	for (int32 SectionIndex = 0; SectionIndex < NewLODModel.Sections.Num(); SectionIndex++)
 	{
-		FSkelMeshChunk& Chunk = NewLODModel.Chunks[ChunkIndex];
-		// Fix up rigid verts.
-		for (int32 i = 0; i < Chunk.RigidVertices.Num(); i++)
-		{
-			Chunk.RigidVertices[i].Position = LODToBaseTransform.TransformPosition(Chunk.RigidVertices[i].Position);
-			Chunk.RigidVertices[i].TangentX = LODToBaseTransform.TransformVector(Chunk.RigidVertices[i].TangentX);
-			Chunk.RigidVertices[i].TangentY = LODToBaseTransform.TransformVector(Chunk.RigidVertices[i].TangentY);
-			Chunk.RigidVertices[i].TangentZ = LODToBaseTransform.TransformVector(Chunk.RigidVertices[i].TangentZ);
-		}
+		FSkelMeshSection& Section = NewLODModel.Sections[SectionIndex];
 
 		// Fix up soft verts.
-		for (int32 i = 0; i < Chunk.SoftVertices.Num(); i++)
+		for (int32 i = 0; i < Section.SoftVertices.Num(); i++)
 		{
-			Chunk.SoftVertices[i].Position = LODToBaseTransform.TransformPosition(Chunk.SoftVertices[i].Position);
-			Chunk.SoftVertices[i].TangentX = LODToBaseTransform.TransformVector(Chunk.SoftVertices[i].TangentX);
-			Chunk.SoftVertices[i].TangentY = LODToBaseTransform.TransformVector(Chunk.SoftVertices[i].TangentY);
-			Chunk.SoftVertices[i].TangentZ = LODToBaseTransform.TransformVector(Chunk.SoftVertices[i].TangentZ);
+			Section.SoftVertices[i].Position = LODToBaseTransform.TransformPosition(Section.SoftVertices[i].Position);
+			Section.SoftVertices[i].TangentX = LODToBaseTransform.TransformVector(Section.SoftVertices[i].TangentX);
+			Section.SoftVertices[i].TangentY = LODToBaseTransform.TransformVector(Section.SoftVertices[i].TangentY);
+			Section.SoftVertices[i].TangentZ = LODToBaseTransform.TransformVector(Section.SoftVertices[i].TangentZ);
 		}
 	}
 
diff --git a/Engine/Source/Editor/UnrealEd/Private/FbxImporter.h b/Engine/Source/Editor/UnrealEd/Private/FbxImporter.h
--- a/Engine/Source/Editor/UnrealEd/Private/FbxImporter.h
+++ b/Engine/Source/Editor/UnrealEd/Private/FbxImporter.h
@@ -128,6 +128,9 @@ struct FBXImportOptions
 	bool	bDeleteExistingMorphTargetCurves;
 	bool	bImportCustomAttribute;
 	bool	bSetMaterialDriveParameterOnCustomAttribute;
+	bool	bRemoveRedundantKeys;
+	TArray<FString> MaterialCurveSuffixes;
+	TArray<FString>	PoseCurveSuffixes;
 
 	/** This allow to add a prefix to the material name when unreal material get created.	
 	*   This prefix can just modify the name of the asset for materials (i.e. TEXT("Mat"))
diff --git a/Engine/Source/Editor/UnrealEd/Private/PhysicsAssetUtils.cpp b/Engine/Source/Editor/UnrealEd/Private/PhysicsAssetUtils.cpp
--- a/Engine/Source/Editor/UnrealEd/Private/PhysicsAssetUtils.cpp
+++ b/Engine/Source/Editor/UnrealEd/Private/PhysicsAssetUtils.cpp
@@ -464,9 +464,8 @@ bool CreateCollisionFromBone( UBodySetup* bs, USkeletalMesh* skelMesh, int32 Bon
 	else if (Params.GeomType == EFG_MultiConvexHull)
 	{
 		// Just feed in all of the verts which are affected by this bone
-		int32 ChunkIndex;
+		int32 SectionIndex;
 		int32 VertIndex;
-		bool bSoftVertex;
 		bool bHasExtraInfluences;
 
 		// Storage for the hull generation
@@ -476,16 +475,21 @@ bool CreateCollisionFromBone( UBodySetup* bs, USkeletalMesh* skelMesh, int32 Bon
 
 		// Get the static LOD from the skeletal mesh and loop through the chunks		
 		FStaticLODModel& LODModel = skelMesh->GetSourceModel();
-		TArray<uint32> indexBufferInOrder;
-		LODModel.MultiSizeIndexContainer.GetIndexBuffer( indexBufferInOrder );
-		uint32 indexBufferSize = indexBufferInOrder.Num();
-		uint32 currentIndex = 0;
+		TArray<uint32> IndexBufferInOrder;
+		LODModel.MultiSizeIndexContainer.GetIndexBuffer(IndexBufferInOrder);
+		uint32 IndexBufferSize = IndexBufferInOrder.Num();
+		uint32 CurrentIndex = 0;
 
 		// Add all of the verts and indices to a list I can loop over
-		for ( uint32 index = 0; index < indexBufferSize; index++ )
+		for ( uint32 Index = 0; Index < IndexBufferSize; Index++ )
 		{
-			LODModel.GetChunkAndSkinType(indexBufferInOrder[index], ChunkIndex, VertIndex, bSoftVertex, bHasExtraInfluences);
-			const FSkelMeshChunk& Chunk = LODModel.Chunks[ChunkIndex];
+			LODModel.GetSectionFromVertexIndex(IndexBufferInOrder[Index], SectionIndex, VertIndex, bHasExtraInfluences);
+			const FSkelMeshSection& Section = LODModel.Sections[SectionIndex];
+			const FSoftSkinVertex& SoftVert = Section.SoftVertices[VertIndex];
+
+			uint8 VertBone = 0;
+			bool bSoftVertex = !SoftVert.GetRigidWeightBone(VertBone);
+
 			if ( bSoftVertex )
 			{
 				// We dont want to support soft verts, only rigid
@@ -496,9 +500,8 @@ bool CreateCollisionFromBone( UBodySetup* bs, USkeletalMesh* skelMesh, int32 Bon
 			}
 
 			// Using the same code in GetSkinnedVertexPosition
-			const FRigidSkinVertex& RigidVert = Chunk.RigidVertices[VertIndex];
-			const int LocalBoneIndex = Chunk.BoneMap[RigidVert.Bone];
-			const FVector& VertPosition = skelMesh->RefBasesInvMatrix[LocalBoneIndex].TransformPosition(RigidVert.Position);
+			const int LocalBoneIndex = Section.BoneMap[VertBone];
+			const FVector& VertPosition = skelMesh->RefBasesInvMatrix[LocalBoneIndex].TransformPosition(SoftVert.Position);
 
 			if ( LocalBoneIndex == BoneIndex )
 			{
@@ -508,8 +511,8 @@ bool CreateCollisionFromBone( UBodySetup* bs, USkeletalMesh* skelMesh, int32 Bon
 				}
 				else
 				{
-					Indices.Add( currentIndex );
-					IndexMap.Add( VertIndex, currentIndex++ );
+					Indices.Add( CurrentIndex );
+					IndexMap.Add( VertIndex, CurrentIndex++ );
 					Verts.Add( VertPosition );					
 				}
 			}
diff --git a/Engine/Source/Editor/UnrealEd/Private/SSkeletonWidget.cpp b/Engine/Source/Editor/UnrealEd/Private/SSkeletonWidget.cpp
--- a/Engine/Source/Editor/UnrealEd/Private/SSkeletonWidget.cpp
+++ b/Engine/Source/Editor/UnrealEd/Private/SSkeletonWidget.cpp
@@ -1763,7 +1763,7 @@ const FAssetData* SAnimationRemapAssets::FindBestGuessMatch(const FAssetData& As
 
 		if(Data.AssetClass == AssetData.AssetClass)
 		{
-			int32 Distance = GetStringDistance(AssetData.AssetName.ToString(), Data.AssetName.ToString());
+			int32 Distance = FAnimationRuntime::GetStringDistance(AssetData.AssetName.ToString(), Data.AssetName.ToString());
 
 			if(Distance < LowestScore)
 			{
@@ -1781,62 +1781,6 @@ const FAssetData* SAnimationRemapAssets::FindBestGuessMatch(const FAssetData& As
 	return nullptr;
 }
 
-int32 SAnimationRemapAssets::GetStringDistance(const FString& First, const FString& Second) const
-{
-	// Finds the distance between strings, where the distance is the number of operations we would need
-	// to perform on First to match Second.
-	// Operations are: Adding a character, Removing a character, changing a character.
-
-	const int32 FirstLength = First.Len();
-	const int32 SecondLength = Second.Len();
-
-	// Already matching
-	if(First == Second)
-	{
-		return 0;
-	}
-
-	// No first string, so we need to add SecondLength characters to match
-	if(FirstLength == 0)
-	{
-		return SecondLength;
-	}
-
-	// No Second string, so we need to add FirstLength characters to match
-	if(SecondLength == 0)
-	{
-		return FirstLength;
-	}
-
-	TArray<int32> PrevRow;
-	TArray<int32> NextRow;
-	PrevRow.AddZeroed(SecondLength + 1);
-	NextRow.AddZeroed(SecondLength + 1);
-
-	// Initialise prev row to num characters we need to remove from Second
-	for(int32 I = 0 ; I < PrevRow.Num() ; ++I)
-	{
-		PrevRow[I] = I;
-	}
-
-	for(int32 I = 0 ; I < FirstLength ; ++I)
-	{
-		// Calculate current row
-		NextRow[0] = I + 1;
-
-		for(int32 J = 0 ; J < SecondLength ; ++J)
-		{
-			int32 Indicator = (First[I] == Second[J]) ? 0 : 1;
-			NextRow[J + 1] = FMath::Min3(NextRow[J] + 1, PrevRow[J + 1] + 1, PrevRow[J] + Indicator);
-		}
-
-		// Copy back
-		PrevRow = NextRow;
-	}
-
-	return NextRow[SecondLength];
-}
-
 void SAssetEntryRow::Construct(const FArguments& InArgs, const TSharedRef<STableViewBase>& InOwnerTableView)
 {
 	check(InArgs._DisplayedInfo.IsValid());
diff --git a/Engine/Source/Editor/UnrealEd/Private/SkeletalMeshEdit.cpp b/Engine/Source/Editor/UnrealEd/Private/SkeletalMeshEdit.cpp
--- a/Engine/Source/Editor/UnrealEd/Private/SkeletalMeshEdit.cpp
+++ b/Engine/Source/Editor/UnrealEd/Private/SkeletalMeshEdit.cpp
@@ -940,22 +940,22 @@ bool UnFbx::FFbxImporter::ImportCurveToAnimSequence(class UAnimSequence * Target
 		const FSmartNameMapping* NameMapping = Skeleton->GetSmartNameContainer(USkeleton::AnimCurveMappingName);
 
 		// Add or retrieve curve
-		USkeleton::AnimCurveUID Uid;
 		if (!NameMapping->Exists(Name))
 		{
 			// mark skeleton dirty
 			Skeleton->Modify();
 		}
 
-		Skeleton->AddSmartNameAndModify(USkeleton::AnimCurveMappingName, Name, Uid);
+		FSmartName NewName;
+		Skeleton->AddSmartNameAndModify(USkeleton::AnimCurveMappingName, Name, NewName);
 
-		FFloatCurve * CurveToImport = static_cast<FFloatCurve *>(TargetSequence->RawCurveData.GetCurveData(Uid, FRawCurveTracks::FloatType));
+		FFloatCurve * CurveToImport = static_cast<FFloatCurve *>(TargetSequence->RawCurveData.GetCurveData(NewName.UID, FRawCurveTracks::FloatType));
 		if(CurveToImport==NULL)
 		{
-			if (TargetSequence->RawCurveData.AddCurveData(Uid, ACF_DefaultCurve | CurveFlags))
+			if (TargetSequence->RawCurveData.AddCurveData(NewName, ACF_DefaultCurve | CurveFlags))
 			{
-				CurveToImport = static_cast<FFloatCurve *> (TargetSequence->RawCurveData.GetCurveData(Uid, FRawCurveTracks::FloatType));
-				CurveToImport->LastObservedName = Name;
+				CurveToImport = static_cast<FFloatCurve *> (TargetSequence->RawCurveData.GetCurveData(NewName.UID, FRawCurveTracks::FloatType));
+				CurveToImport->Name = NewName;
 			}
 			else
 			{
@@ -972,10 +972,17 @@ bool UnFbx::FFbxImporter::ImportCurveToAnimSequence(class UAnimSequence * Target
 
 		// update last observed name. If not, sometimes it adds new UID while fixing up that will confuse Compressed Raw Data
 		const FSmartNameMapping* Mapping = Skeleton->GetSmartNameContainer(USkeleton::AnimCurveMappingName);
-		TargetSequence->RawCurveData.UpdateLastObservedNames(Mapping);
+		TargetSequence->RawCurveData.RefreshName(Mapping);
 
 		TargetSequence->MarkRawDataAsModified();
-		return ImportCurve(FbxCurve, CurveToImport, AnimTimeSpan, ValueScale);
+		if (ImportCurve(FbxCurve, CurveToImport, AnimTimeSpan, ValueScale))
+		{
+			if (ImportOptions->bRemoveRedundantKeys)
+			{
+				CurveToImport->FloatCurve.RemoveRedundantKeys(SMALL_NUMBER);
+			}
+			return true;
+		}
 	}
 
 	return false;
@@ -1113,11 +1120,36 @@ bool UnFbx::FFbxImporter::ImportAnimation(USkeleton* Skeleton, UAnimSequence * D
 							const FText StatusUpate = FText::Format(LOCTEXT("ImportingCustomAttributeCurvesDetail", "Importing Custom Attribute [{CurveName}]"), Args);
 							GWarn->StatusUpdate(CurLinkIndex + 1, TotalLinks, StatusUpate);
 
-							int32 CurveFlags = 0;
+							int32 CurveFlags = ACF_DefaultCurve;
+							// first let them override material curve if required
 							if (ImportOptions->bSetMaterialDriveParameterOnCustomAttribute)
 							{
 								CurveFlags |= ACF_DrivesMaterial;
 							}
+							else
+							{
+								// if not material set by default, apply naming convention for material
+								for (const auto& Suffix : ImportOptions->MaterialCurveSuffixes)
+								{
+									int32 TotalSuffix = Suffix.Len();
+									if (CurveName.Right(TotalSuffix) == Suffix)
+									{
+										CurveFlags |= ACF_DrivesMaterial;
+										break;
+									}
+								}
+							}
+
+							// apply naming convention for pose curve
+							for (const auto& Suffix : ImportOptions->PoseCurveSuffixes)
+							{
+								int32 TotalSuffix = Suffix.Len();
+								if (ChannelName.Right(TotalSuffix) == Suffix)
+								{
+									CurveFlags |= ACF_DrivesPose;
+									break;
+								}
+							}
 
 							ImportCurveToAnimSequence(DestSeq, FinalCurveName, AnimCurve, CurveFlags, AnimTimeSpan);
 						}
diff --git a/Engine/Source/Editor/UnrealEd/Private/Tests/EditorBuildPromotionTests.cpp b/Engine/Source/Editor/UnrealEd/Private/Tests/EditorBuildPromotionTests.cpp
--- a/Engine/Source/Editor/UnrealEd/Private/Tests/EditorBuildPromotionTests.cpp
+++ b/Engine/Source/Editor/UnrealEd/Private/Tests/EditorBuildPromotionTests.cpp
@@ -53,6 +53,10 @@
 #include "LightingBuildOptions.h"
 #include "../Private/StaticLightingSystem/StaticLightingPrivate.h"
 
+//Animation
+#include "Animation/AnimSequence.h"
+
+
 #include "ScopedTransaction.h"
 #include "LevelEditor.h"
 #include "ObjectTools.h"
diff --git a/Engine/Source/Editor/UnrealEd/Private/ThumbnailHelpers.cpp b/Engine/Source/Editor/UnrealEd/Private/ThumbnailHelpers.cpp
--- a/Engine/Source/Editor/UnrealEd/Private/ThumbnailHelpers.cpp
+++ b/Engine/Source/Editor/UnrealEd/Private/ThumbnailHelpers.cpp
@@ -737,7 +737,7 @@ FBlendSpaceThumbnailScene::FBlendSpaceThumbnailScene()
 	SpawnInfo.SpawnCollisionHandlingOverride = ESpawnActorCollisionHandlingMethod::AlwaysSpawn;
 	SpawnInfo.bNoFail = true;
 	SpawnInfo.ObjectFlags = RF_Transient;
-	PreviewActor = GetWorld()->SpawnActor<ASkeletalMeshActor>(SpawnInfo);
+	PreviewActor = GetWorld()->SpawnActor<AAnimationThumbnailSkeletalMeshActor>(SpawnInfo);
 
 	PreviewActor->SetActorEnableCollision(false);
 }
@@ -762,11 +762,14 @@ bool FBlendSpaceThumbnailScene::SetBlendSpace(class UBlendSpaceBase* InBlendSpac
 			{
 				bSetSucessfully = true;
 
+				UDebugSkelMeshComponent* MeshComponent = CastChecked<UDebugSkelMeshComponent>(PreviewActor->GetSkeletalMeshComponent());
+
 				// Handle posing the mesh at the middle of the animation
-				PreviewActor->GetSkeletalMeshComponent()->PlayAnimation(InBlendSpace, false);
-				PreviewActor->GetSkeletalMeshComponent()->Stop();
+				MeshComponent->EnablePreview(true, InBlendSpace);
+				MeshComponent->Play(false);
+				MeshComponent->Stop();
 
-				UAnimSingleNodeInstance* AnimInstance = PreviewActor->GetSkeletalMeshComponent()->GetSingleNodeInstance();
+				UAnimSingleNodeInstance* AnimInstance = MeshComponent->GetSingleNodeInstance();
 				if (AnimInstance)
 				{
 					FVector BlendInput(0.f);
@@ -777,7 +780,9 @@ bool FBlendSpaceThumbnailScene::SetBlendSpace(class UBlendSpaceBase* InBlendSpac
 					}
 					AnimInstance->UpdateBlendspaceSamples(BlendInput);
 				}
-				PreviewActor->GetSkeletalMeshComponent()->RefreshBoneTransforms(nullptr);
+
+				MeshComponent->TickAnimation(0.f, false);
+				MeshComponent->RefreshBoneTransforms(nullptr);
 
 				FTransform MeshTransform = FTransform::Identity;
 
@@ -972,6 +977,11 @@ void FClassActorThumbnailScene::SpawnPreviewActor(UClass* InClass)
 {
 	if (PreviewActor)
 	{
+		if (PreviewActor->GetClass() == InClass)
+		{
+			return;
+		}
+
 		PreviewActor->Destroy();
 		PreviewActor = nullptr;
 	}
