Fix for non-unity builds

diff --git a/Engine/Source/Editor/Persona/Private/AnimGraphNodeDetails.cpp b/Engine/Source/Editor/Persona/Private/AnimGraphNodeDetails.cpp
--- a/Engine/Source/Editor/Persona/Private/AnimGraphNodeDetails.cpp
+++ b/Engine/Source/Editor/Persona/Private/AnimGraphNodeDetails.cpp
@@ -24,6 +24,7 @@
 #include "Animation/BlendProfile.h"
 #include "SBlendProfilePicker.h"
 #include "AnimGraphNode_AssetPlayerBase.h"
+#include "Animation/AnimInstance.h"
 
 #define LOCTEXT_NAMESPACE "KismetNodeWithOptionalPinsDetails"
 
