Copying //UE4/Dev-Framework to //UE4/Dev-Main (Source: //UE4/Dev-Framework @ 3038004)

diff --git a/Engine/Source/Editor/KismetCompiler/Private/AnimBlueprintCompiler.cpp b/Engine/Source/Editor/KismetCompiler/Private/AnimBlueprintCompiler.cpp
--- a/Engine/Source/Editor/KismetCompiler/Private/AnimBlueprintCompiler.cpp
+++ b/Engine/Source/Editor/KismetCompiler/Private/AnimBlueprintCompiler.cpp
@@ -31,6 +31,11 @@
 #include "K2Node_TransitionRuleGetter.h"
 #include "K2Node_AnimGetter.h"
 #include "AnimGraphNode_StateMachine.h"
+#include "AnimGraphNode_SubInput.h"
+#include "Animation/AnimNode_SubInstance.h"
+#include "AnimGraphNode_SubInstance.h"
+#include "AnimGraphNode_Slot.h"
+#include "EdGraph/EdGraphPin.h"
 
 #define LOCTEXT_NAMESPACE "AnimBlueprintCompiler"
 
@@ -44,8 +49,18 @@ class UAnimStateNodeBase;
 
 bool FAnimBlueprintCompiler::FEffectiveConstantRecord::Apply(UObject* Object)
 {
-	uint8* StructPtr = NodeVariableProperty->ContainerPtrToValuePtr<uint8>(Object);
-	uint8* PropertyPtr = ConstantProperty->ContainerPtrToValuePtr<uint8>(StructPtr);
+	uint8* StructPtr = nullptr;
+	uint8* PropertyPtr = nullptr;
+	
+	if(NodeVariableProperty->Struct == FAnimNode_SubInstance::StaticStruct())
+	{
+		PropertyPtr = ConstantProperty->ContainerPtrToValuePtr<uint8>(Object);
+	}
+	else
+	{
+		StructPtr = NodeVariableProperty->ContainerPtrToValuePtr<uint8>(Object);
+		PropertyPtr = ConstantProperty->ContainerPtrToValuePtr<uint8>(StructPtr);
+	}
 
 	if (ArrayIndex != INDEX_NONE)
 	{
@@ -112,7 +127,7 @@ UK2Node_CallFunction* FAnimBlueprintCompiler::SpawnCallAnimInstanceFunction(UEdG
 	return FunctionCall;
 }
 
-void FAnimBlueprintCompiler::CreateEvaluationHandler(UAnimGraphNode_Base* VisualAnimNode, FEvaluationHandlerRecord& Record)
+void FAnimBlueprintCompiler::CreateEvaluationHandlerStruct(UAnimGraphNode_Base* VisualAnimNode, FEvaluationHandlerRecord& Record)
 {
 	// Shouldn't create a handler if there is nothing to work with
 	check(Record.ServicedProperties.Num() > 0);
@@ -241,6 +256,79 @@ void FAnimBlueprintCompiler::CreateEvaluationHandler(UAnimGraphNode_Base* Visual
 	AssignmentNode->ReconstructNode();
 }
 
+void FAnimBlueprintCompiler::CreateEvaluationHandlerInstance(UAnimGraphNode_Base* VisualAnimNode, FEvaluationHandlerRecord& Record)
+{
+	// Shouldn't create a handler if there is nothing to work with
+	check(Record.ServicedProperties.Num() > 0);
+	check(Record.NodeVariableProperty != nullptr);
+	check(Record.InstanceProperty != nullptr)
+
+	// Use the node GUID for a stable name across compiles
+	FString FunctionName = FString::Printf(TEXT("%s_%s_%s_%s"), *Record.EvaluationHandlerProperty->GetName(), *VisualAnimNode->GetOuter()->GetName(), *VisualAnimNode->GetClass()->GetName(), *VisualAnimNode->NodeGuid.ToString());
+	Record.HandlerFunctionName = FName(*FunctionName);
+
+	// check function name isnt already used (data exists that can contain duplicate GUIDs) and apply a numeric extension until it is unique
+	int32 ExtensionIndex = 0;
+	FName* ExistingName = HandlerFunctionNames.Find(Record.HandlerFunctionName);
+	while(ExistingName != nullptr)
+	{
+		FunctionName = FString::Printf(TEXT("%s_%s_%s_%s_%d"), *Record.EvaluationHandlerProperty->GetName(), *VisualAnimNode->GetOuter()->GetName(), *VisualAnimNode->GetClass()->GetName(), *VisualAnimNode->NodeGuid.ToString(), ExtensionIndex);
+		Record.HandlerFunctionName = FName(*FunctionName);
+		ExistingName = HandlerFunctionNames.Find(Record.HandlerFunctionName);
+		ExtensionIndex++;
+	}
+
+	HandlerFunctionNames.Add(Record.HandlerFunctionName);
+
+	// Add a custom event in the graph
+	UK2Node_CustomEvent* EntryNode = SpawnIntermediateNode<UK2Node_CustomEvent>(VisualAnimNode, ConsolidatedEventGraph);
+	EntryNode->bInternalEvent = true;
+	EntryNode->CustomFunctionName = Record.HandlerFunctionName;
+	EntryNode->AllocateDefaultPins();
+
+	// The ExecChain is the current exec output pin in the linear chain
+	UEdGraphPin* ExecChain = Schema->FindExecutionPin(*EntryNode, EGPD_Output);
+
+	// Create the set node for the instance variable
+	UK2Node_VariableSet* VarAssignNode = SpawnIntermediateNode<UK2Node_VariableSet>(VisualAnimNode, ConsolidatedEventGraph);
+	VarAssignNode->VariableReference.SetSelfMember(Record.InstanceProperty->GetFName());
+	VarAssignNode->AllocateDefaultPins();
+
+	// Wire up the variable node execution wires
+	UEdGraphPin* ExecVariablesIn = Schema->FindExecutionPin(*VarAssignNode, EGPD_Input);
+	ExecChain->MakeLinkTo(ExecVariablesIn);
+	ExecChain = Schema->FindExecutionPin(*VarAssignNode, EGPD_Output);
+
+	// Run through each property
+	for(auto TargetPinIt = VarAssignNode->Pins.CreateIterator(); TargetPinIt; ++TargetPinIt)
+	{
+		UEdGraphPin* TargetPin = *TargetPinIt;
+		FString PropertyNameStr = TargetPin->PinName;
+		FName PropertyName(*PropertyNameStr);
+
+		// Does it get serviced by this handler?
+		if(FAnimNodeSinglePropertyHandler* SourceInfo = Record.ServicedProperties.Find(PropertyName))
+		{
+			// Check for arrays, we can't copy from array properties yet so we need to wire it up properly
+			if(SourceInfo->SimpleCopyPropertyName != NAME_None && !TargetPin->PinType.bIsArray)
+			{
+				continue;
+			}
+
+			if(SourceInfo->SinglePin != NULL)
+			{
+				UEdGraphPin* SourcePin = SourceInfo->SinglePin;
+
+				TargetPin->CopyPersistentDataFromOldPin(*SourcePin);
+				MessageLog.NotifyIntermediatePinCreation(TargetPin, SourcePin);
+				SourcePin->BreakAllPinLinks();
+			}
+		}
+	}
+
+	VarAssignNode->ReconstructNode();
+}
+
 void FAnimBlueprintCompiler::ProcessAnimationNode(UAnimGraphNode_Base* VisualAnimNode)
 {
 	// Early out if this node has already been processed
@@ -310,9 +398,13 @@ void FAnimBlueprintCompiler::ProcessAnimationNode(UAnimGraphNode_Base* VisualAni
 		// Handle a save/use cached pose linkage
 		ProcessUseCachedPose(UseCachedPose);
 	}
+	else if(UAnimGraphNode_SubInstance* SubInstanceNode = Cast<UAnimGraphNode_SubInstance>(VisualAnimNode))
+	{
+		ProcessSubInstance(SubInstanceNode);
+	}
 
 	// Record pose pins for later patchup and gather pins that have an associated evaluation handler
-	TMap<FName, FEvaluationHandlerRecord> EvaluationHandlers;
+	TMap<FName, FEvaluationHandlerRecord> StructEvalHandlers;
 
 	for (auto SourcePinIt = VisualAnimNode->Pins.CreateIterator(); SourcePinIt; ++SourcePinIt)
 	{
@@ -332,10 +424,25 @@ void FAnimBlueprintCompiler::ProcessAnimationNode(UAnimGraphNode_Base* VisualAni
 		}
 		else
 		{
+			// The property source for our data, either the struct property for an anim node, or the
+			// owning anim instance if using a sub instance node.
+			UProperty* SourcePinProperty = nullptr;
+			int32 SourceArrayIndex = INDEX_NONE;
+
+			// We have special handling below if we're targeting a subinstance instead of our own instance properties
+			UAnimGraphNode_SubInstance* SubInstanceNode = Cast<UAnimGraphNode_SubInstance>(VisualAnimNode);
+
 			// Does this pin have an associated evaluation handler?
-			UProperty* SourcePinProperty;
-			int32 SourceArrayIndex;
-			VisualAnimNode->GetPinAssociatedProperty(NodeType, SourcePin, /*out*/ SourcePinProperty, /*out*/ SourceArrayIndex);
+			if(SubInstanceNode)
+			{
+				// Subinstance nodes use instance properties not node properties as they aren't UObjects
+				// and we can't store non-native properties there
+				SubInstanceNode->GetInstancePinProperty(NewAnimBlueprintClass, SourcePin, SourcePinProperty);
+			}
+			else
+			{
+				VisualAnimNode->GetPinAssociatedProperty(NodeType, SourcePin, /*out*/ SourcePinProperty, /*out*/ SourceArrayIndex);
+			}
 			
 			if (SourcePinProperty != NULL)
 			{
@@ -355,7 +462,21 @@ void FAnimBlueprintCompiler::ProcessAnimationNode(UAnimGraphNode_Base* VisualAni
 						EvaluationHandlerName = AnimGraphDefaultSchema->DefaultEvaluationHandlerName;
 					}
 
-					EvaluationHandlers.FindOrAdd(EvaluationHandlerName).RegisterPin(SourcePin, SourcePinProperty, SourceArrayIndex);
+					FEvaluationHandlerRecord& EvalHandler = StructEvalHandlers.FindOrAdd(EvaluationHandlerName);
+
+					EvalHandler.RegisterPin(SourcePin, SourcePinProperty, SourceArrayIndex);
+
+					if(SubInstanceNode)
+					{
+						EvalHandler.InstanceProperty = SourcePinProperty;
+
+						FAnimNodeSinglePropertyHandler* SinglePropHandler = EvalHandler.ServicedProperties.Find(SourcePinProperty->GetFName());
+						check(SinglePropHandler); // Should have been added in RegisterPin
+
+						// Flag that the target property is actually on the instance class and not the node
+						SinglePropHandler->bInstanceIsTarget = true;
+					}
+
 					bConsumed = true;
 				}
 
@@ -384,24 +505,36 @@ void FAnimBlueprintCompiler::ProcessAnimationNode(UAnimGraphNode_Base* VisualAni
 				// Register this property to the list of pins that need to be updated
 				// (it's OK if there isn't an entry for this handler, it means that the values are static and don't need to be calculated every frame)
 				FName EvaluationHandlerName(StructProp->GetFName());
-				if (FEvaluationHandlerRecord* pRecord = EvaluationHandlers.Find(EvaluationHandlerName))
+
+				FEvaluationHandlerRecord* pRecord = StructEvalHandlers.Find(EvaluationHandlerName);
+
+				if (pRecord)
 				{
 					pRecord->NodeVariableProperty = NewProperty;
 					pRecord->EvaluationHandlerProperty = StructProp;
 				}
+				
 			}
 		}
 	}
 
 	// Generate a new event to update the value of these properties
-	for (auto HandlerIt = EvaluationHandlers.CreateIterator(); HandlerIt; ++HandlerIt)
+	for (auto HandlerIt = StructEvalHandlers.CreateIterator(); HandlerIt; ++HandlerIt)
 	{
 		FName EvaluationHandlerName = HandlerIt.Key();
 		FEvaluationHandlerRecord& Record = HandlerIt.Value();
 
 		if (Record.IsValid())
 		{
-			CreateEvaluationHandler(VisualAnimNode, Record);
+			if(Record.InstanceProperty)
+			{
+				CreateEvaluationHandlerInstance(VisualAnimNode, Record);
+			}
+			else
+			{
+				CreateEvaluationHandlerStruct(VisualAnimNode, Record);
+			}
+
 			ValidEvaluationHandlerList.Add(Record);
 		}
 		else
@@ -449,6 +582,119 @@ void FAnimBlueprintCompiler::ProcessUseCachedPose(UAnimGraphNode_UseCachedPose*
 	}
 }
 
+void FAnimBlueprintCompiler::ProcessSubInstance(UAnimGraphNode_SubInstance* SubInstance)
+{
+	if(!SubInstance)
+	{
+		return;
+	}
+
+	for(UEdGraphPin* Pin : SubInstance->Pins)
+	{
+		if(Pin->PinName == TEXT("InPose") || Pin->PinName == TEXT("Pose"))
+		{
+			continue;
+		}
+
+		// Add prefix to avoid collisions
+		FString PrefixedName = SubInstance->GetPinTargetVariableName(Pin);
+
+		// Create a property on the new class to hold the pin data
+		UProperty* NewProperty = FKismetCompilerUtilities::CreatePropertyOnScope(NewAnimBlueprintClass, FName(*PrefixedName), Pin->PinType, NewAnimBlueprintClass, 0, GetSchema(), MessageLog);
+		if(NewProperty)
+		{
+			NewProperty->SetMetaData(TEXT("Category"), TEXT("SubInstance"));
+			NewProperty->SetPropertyFlags(CPF_Edit | CPF_BlueprintVisible);
+			FKismetCompilerUtilities::LinkAddedProperty(NewAnimBlueprintClass, NewProperty);
+
+			// Add mappings to the node
+			FAnimNode_SubInstance& RuntimeNode = SubInstance->Node;
+
+			RuntimeNode.SourcePropertyNames.Add(NewProperty->GetFName());
+			
+			// Find the property on the internal instance
+			if(UProperty* FoundProperty = FindField<UProperty>(*RuntimeNode.InstanceClass, *Pin->PinName))
+			{
+				RuntimeNode.DestPropertyNames.Add(FoundProperty->GetFName());
+			}
+		}
+	}
+
+	// Check for duplicated slot and state machine names to warn the user about how these 
+	// are boxed
+	NameToCountMap SlotNameToCountMap;
+	NameToCountMap StateMachineNameToCountMap;
+
+	GetDuplicatedSlotAndStateNames(SubInstance, StateMachineNameToCountMap, SlotNameToCountMap);
+
+
+	for(TPair<FName, int32>& Pair : SlotNameToCountMap)
+	{
+		if(Pair.Value > 1)
+		{
+			// Duplicated slot node
+			FString CompilerMessage = FString::Printf(TEXT("Slot name \"%s\" found across multiple instances. Slots are not visible outside of instances so duplicates or subinstances may not perform as expected."), *Pair.Key.ToString());
+			MessageLog.Warning(*CompilerMessage);
+		}
+	}
+
+	for(TPair<FName, int32>& Pair : StateMachineNameToCountMap)
+	{
+		if(Pair.Value > 1)
+		{
+			// Duplicated slot node
+			FString CompilerMessage = FString::Printf(TEXT("State machine \"%s\" found across multiple instances. States are not visible outside of instances so duplicates or subinstances may not perform as expected."), *Pair.Key.ToString());
+			MessageLog.Warning(*CompilerMessage);
+		}
+	}
+}
+
+void FAnimBlueprintCompiler::GetDuplicatedSlotAndStateNames(UAnimGraphNode_SubInstance* InSubInstance, NameToCountMap& OutStateMachineNameToCountMap, NameToCountMap& OutSlotNameToCountMap)
+{
+	if(!InSubInstance)
+	{
+		// Nothing to inspect
+		return;
+	}
+
+	if(UClass* InstanceClass = *InSubInstance->Node.InstanceClass)
+	{
+		UBlueprint* ClassBP = UBlueprint::GetBlueprintFromClass(InstanceClass);
+
+		TArray<UEdGraph*> AllGraphs;
+		ClassBP->GetAllGraphs(AllGraphs);
+
+		for(UEdGraph* Graph : AllGraphs)
+		{
+			TArray<UAnimGraphNode_StateMachine*> StateMachineNodes;
+			TArray<UAnimGraphNode_Slot*> SlotNodes;
+			TArray<UAnimGraphNode_SubInstance*> SubInstanceNodes;
+
+			Graph->GetNodesOfClass(StateMachineNodes);
+			Graph->GetNodesOfClass(SlotNodes);
+			Graph->GetNodesOfClass(SubInstanceNodes);
+
+			for(UAnimGraphNode_StateMachine* StateMachineNode : StateMachineNodes)
+			{
+				int32& Count = OutStateMachineNameToCountMap.FindOrAdd(FName(*StateMachineNode->GetStateMachineName()));
+				// Add one to count as we've encountered this name
+				Count++;
+			}
+
+			for(UAnimGraphNode_Slot* SlotNode : SlotNodes)
+			{
+				int32& Count = OutSlotNameToCountMap.FindOrAdd(SlotNode->Node.SlotName);
+				Count++;
+			}
+
+			for(UAnimGraphNode_SubInstance* SubInstanceNode : SubInstanceNodes)
+			{
+				GetDuplicatedSlotAndStateNames(SubInstanceNode, OutStateMachineNameToCountMap, OutSlotNameToCountMap);
+			}
+		}
+	}
+}
+
 int32 FAnimBlueprintCompiler::GetAllocationIndexOfNode(UAnimGraphNode_Base* VisualAnimNode)
 {
 	ProcessAnimationNode(VisualAnimNode);
@@ -2088,42 +2334,45 @@ void FAnimBlueprintCompiler::FEvaluationHandlerRecord::PatchFunctionNameAndCopyR
 	HandlerPtr->CopyRecords.Empty();
 
 	bool bOnlyUsesCopyRecords = true;
-	for(TMap<FName, FAnimNodeSinglePropertyHandler>::TConstIterator PropertyIt(ServicedProperties); PropertyIt; ++PropertyIt)
+	for(TPair<FName, FAnimNodeSinglePropertyHandler> ServicedPropPair : ServicedProperties)
 	{
-		const FAnimNodeSinglePropertyHandler& AnimNodeSinglePropertyHandler = PropertyIt.Value();
-		if(AnimNodeSinglePropertyHandler.SimpleCopyPropertyName != NAME_None)
+		const FName& PropertyName = ServicedPropPair.Key;
+		const FAnimNodeSinglePropertyHandler& PropertyHandler = ServicedPropPair.Value;
+
+		if(PropertyHandler.SimpleCopyPropertyName != NAME_None)
 		{
-			UProperty* SimpleCopyPropertySource = TargetObject->GetClass()->FindPropertyByName(AnimNodeSinglePropertyHandler.SimpleCopyPropertyName);
+			UProperty* SimpleCopyPropertySource = TargetObject->GetClass()->FindPropertyByName(PropertyHandler.SimpleCopyPropertyName);
 
 			// we dont support copying *from* array properties at the moment
 			if (SimpleCopyPropertySource && !SimpleCopyPropertySource->IsA<UArrayProperty>())
 			{
 				// are we copying to an array element (or more than one)?
-				if (AnimNodeSinglePropertyHandler.ArrayPins.Num() > 0)
+				if (PropertyHandler.ArrayPins.Num() > 0)
 				{
-					UArrayProperty* SimpleCopyPropertyDest = CastChecked<UArrayProperty>(NodeVariableProperty->Struct->FindPropertyByName(PropertyIt.Key()));
+					UArrayProperty* SimpleCopyPropertyDest = CastChecked<UArrayProperty>(NodeVariableProperty->Struct->FindPropertyByName(PropertyName));
 
-					if (AnimNodeSinglePropertyHandler.SubStructPropertyName != NAME_None)
+					if (PropertyHandler.SubStructPropertyName != NAME_None)
 					{
 						UStructProperty* SourceStructProperty = CastChecked<UStructProperty>(SimpleCopyPropertySource);
-						UProperty* SourceSubProperty = SourceStructProperty->Struct->FindPropertyByName(AnimNodeSinglePropertyHandler.SubStructPropertyName);
+						UProperty* SourceSubProperty = SourceStructProperty->Struct->FindPropertyByName(PropertyHandler.SubStructPropertyName);
 						if (SourceSubProperty == nullptr || SourceSubProperty->GetSize() != SimpleCopyPropertyDest->Inner->GetSize())
 						{
 							bOnlyUsesCopyRecords = false;
 							break;
 						}
 
-						for (TMap<int32, UEdGraphPin*>::TConstIterator ArrayIt(AnimNodeSinglePropertyHandler.ArrayPins); ArrayIt; ++ArrayIt)
+						for (TMap<int32, UEdGraphPin*>::TConstIterator ArrayIt(PropertyHandler.ArrayPins); ArrayIt; ++ArrayIt)
 						{
 							// Single sub-struct variable get->array set
 							FExposedValueCopyRecord CopyRecord;
 							CopyRecord.DestProperty = SimpleCopyPropertyDest;
 							CopyRecord.DestArrayIndex = ArrayIt.Key();
-							CopyRecord.SourcePropertyName = AnimNodeSinglePropertyHandler.SimpleCopyPropertyName;
-							CopyRecord.SourceSubPropertyName = AnimNodeSinglePropertyHandler.SubStructPropertyName;
+							CopyRecord.SourcePropertyName = PropertyHandler.SimpleCopyPropertyName;
+							CopyRecord.SourceSubPropertyName = PropertyHandler.SubStructPropertyName;
 							CopyRecord.SourceArrayIndex = 0;
 							CopyRecord.Size = SimpleCopyPropertyDest->Inner->GetSize();
-							CopyRecord.PostCopyOperation = AnimNodeSinglePropertyHandler.Operation;
+							CopyRecord.PostCopyOperation = PropertyHandler.Operation;
+							CopyRecord.bInstanceIsTarget = PropertyHandler.bInstanceIsTarget;
 							HandlerPtr->CopyRecords.Add(CopyRecord);
 						}
 					}
@@ -2135,34 +2384,41 @@ void FAnimBlueprintCompiler::FEvaluationHandlerRecord::PatchFunctionNameAndCopyR
 							break;
 						}
 
-						for (TMap<int32, UEdGraphPin*>::TConstIterator ArrayIt(AnimNodeSinglePropertyHandler.ArrayPins); ArrayIt; ++ArrayIt)
+						for (TMap<int32, UEdGraphPin*>::TConstIterator ArrayIt(PropertyHandler.ArrayPins); ArrayIt; ++ArrayIt)
 						{
 							// Single variable get->array set
 							FExposedValueCopyRecord CopyRecord;
 							CopyRecord.DestProperty = SimpleCopyPropertyDest;
 							CopyRecord.DestArrayIndex = ArrayIt.Key();
-							CopyRecord.SourcePropertyName = AnimNodeSinglePropertyHandler.SimpleCopyPropertyName;
+							CopyRecord.SourcePropertyName = PropertyHandler.SimpleCopyPropertyName;
 							CopyRecord.SourceSubPropertyName = NAME_None;
 							CopyRecord.SourceArrayIndex = 0;
 							CopyRecord.Size = SimpleCopyPropertyDest->Inner->GetSize();
-							CopyRecord.PostCopyOperation = AnimNodeSinglePropertyHandler.Operation;
+							CopyRecord.PostCopyOperation = PropertyHandler.Operation;
+							CopyRecord.bInstanceIsTarget = PropertyHandler.bInstanceIsTarget;
 							HandlerPtr->CopyRecords.Add(CopyRecord);
 						}
 					}
 				}
 				else
 				{
-					UProperty* SimpleCopyPropertyDest = NodeVariableProperty->Struct->FindPropertyByName(PropertyIt.Key());
+					UProperty* SimpleCopyPropertyDest = NodeVariableProperty->Struct->FindPropertyByName(PropertyName);
+
+					if(!SimpleCopyPropertyDest)
+					{
+						SimpleCopyPropertyDest = FindField<UProperty>(TargetObject->GetClass(), PropertyName);
+					}
+
 					if(SimpleCopyPropertyDest == nullptr)
 					{
 						bOnlyUsesCopyRecords = false;
 						break;
 					}
 
-					if(AnimNodeSinglePropertyHandler.SubStructPropertyName != NAME_None)
+					if(PropertyHandler.SubStructPropertyName != NAME_None)
 					{
 						UStructProperty* SourceStructProperty = CastChecked<UStructProperty>(SimpleCopyPropertySource);
-						UProperty* SourceSubProperty = SourceStructProperty->Struct->FindPropertyByName(AnimNodeSinglePropertyHandler.SubStructPropertyName);
+						UProperty* SourceSubProperty = SourceStructProperty->Struct->FindPropertyByName(PropertyHandler.SubStructPropertyName);
 						if(SourceSubProperty == nullptr || SourceSubProperty->GetSize() != SimpleCopyPropertyDest->GetSize())
 						{
 							bOnlyUsesCopyRecords = false;
@@ -2173,11 +2429,12 @@ void FAnimBlueprintCompiler::FEvaluationHandlerRecord::PatchFunctionNameAndCopyR
 						FExposedValueCopyRecord CopyRecord;
 						CopyRecord.DestProperty = SimpleCopyPropertyDest;
 						CopyRecord.DestArrayIndex = 0;
-						CopyRecord.SourcePropertyName = AnimNodeSinglePropertyHandler.SimpleCopyPropertyName;
-						CopyRecord.SourceSubPropertyName = AnimNodeSinglePropertyHandler.SubStructPropertyName;
+						CopyRecord.SourcePropertyName = PropertyHandler.SimpleCopyPropertyName;
+						CopyRecord.SourceSubPropertyName = PropertyHandler.SubStructPropertyName;
 						CopyRecord.SourceArrayIndex = 0;
 						CopyRecord.Size = SimpleCopyPropertyDest->GetSize();
-						CopyRecord.PostCopyOperation = AnimNodeSinglePropertyHandler.Operation;
+						CopyRecord.PostCopyOperation = PropertyHandler.Operation;
+						CopyRecord.bInstanceIsTarget = PropertyHandler.bInstanceIsTarget;
 						HandlerPtr->CopyRecords.Add(CopyRecord);
 					}
 					else
@@ -2192,11 +2449,12 @@ void FAnimBlueprintCompiler::FEvaluationHandlerRecord::PatchFunctionNameAndCopyR
 						FExposedValueCopyRecord CopyRecord;
 						CopyRecord.DestProperty = SimpleCopyPropertyDest;
 						CopyRecord.DestArrayIndex = 0;
-						CopyRecord.SourcePropertyName = AnimNodeSinglePropertyHandler.SimpleCopyPropertyName;
+						CopyRecord.SourcePropertyName = PropertyHandler.SimpleCopyPropertyName;
 						CopyRecord.SourceSubPropertyName = NAME_None;
 						CopyRecord.SourceArrayIndex = 0;
 						CopyRecord.Size = SimpleCopyPropertyDest->GetSize();
-						CopyRecord.PostCopyOperation = AnimNodeSinglePropertyHandler.Operation;
+						CopyRecord.PostCopyOperation = PropertyHandler.Operation;
+						CopyRecord.bInstanceIsTarget = PropertyHandler.bInstanceIsTarget;
 						HandlerPtr->CopyRecords.Add(CopyRecord);
 					}
 				}
diff --git a/Engine/Source/Editor/KismetCompiler/Private/AnimBlueprintCompiler.h b/Engine/Source/Editor/KismetCompiler/Private/AnimBlueprintCompiler.h
--- a/Engine/Source/Editor/KismetCompiler/Private/AnimBlueprintCompiler.h
+++ b/Engine/Source/Editor/KismetCompiler/Private/AnimBlueprintCompiler.h
@@ -7,18 +7,23 @@
 #include "Animation/AnimNodeBase.h"
 #include "AnimationGraphSchema.h"
 #include "K2Node_TransitionRuleGetter.h"
+#include "AnimGraphNode_Base.h"
 //
 // Forward declarations.
 //
 class UAnimGraphNode_SaveCachedPose;
 class UAnimGraphNode_UseCachedPose;
+class UAnimGraphNode_SubInput;
+class UAnimGraphNode_SubInstance;
+
+class UStructProperty;
 class UBlueprintGeneratedClass;
 struct FPoseLinkMappingRecord;
 
 //////////////////////////////////////////////////////////////////////////
 // FAnimBlueprintCompiler
 
-class FAnimBlueprintCompiler : public FKismetCompilerContext
+class KISMETCOMPILER_API FAnimBlueprintCompiler : public FKismetCompilerContext
 {
 protected:
 	typedef FKismetCompilerContext Super;
@@ -65,12 +70,16 @@ protected:
 		// Whether this node chain only ever accesses members
 		bool bHasOnlyMemberAccess;
 
+		// If the anim instance is the container target instead of the node.
+		bool bInstanceIsTarget;
+
 		FAnimNodeSinglePropertyHandler()
 			: SinglePin(NULL)
 			, SimpleCopyPropertyName(NAME_None)
 			, SubStructPropertyName(NAME_None)
 			, Operation(EPostCopyOperation::None)
 			, bHasOnlyMemberAccess(false)
+			, bInstanceIsTarget(false)
 		{
 		}
 	};
@@ -110,28 +119,33 @@ protected:
 		bool Apply(UObject* Object);
 	};
 
-	// Struct used to record information about one evaluation handler
 	struct FEvaluationHandlerRecord
 	{
 	public:
+
+		// The instance property to use if we aren't accessing data through the node variable
+		UProperty* InstanceProperty;
+
 		// The node variable that the handler is in
-		class UStructProperty* NodeVariableProperty;
+		UStructProperty* NodeVariableProperty;
 
 		// The specific evaluation handler inside the specified node
-		class UStructProperty* EvaluationHandlerProperty;
+		UStructProperty* EvaluationHandlerProperty;
 
 		// Set of properties serviced by this handler (Map from property name to the record for that property)
 		TMap<FName, FAnimNodeSinglePropertyHandler> ServicedProperties;
 
 		// The resulting function name
 		FName HandlerFunctionName;
+
 	public:
+
 		FEvaluationHandlerRecord()
-			: NodeVariableProperty(NULL)
-			, EvaluationHandlerProperty(NULL)
+			: InstanceProperty(nullptr)
+			, NodeVariableProperty(nullptr)
+			, EvaluationHandlerProperty(nullptr)
 			, HandlerFunctionName(NAME_None)
-		{
-		}
+		{}
 
 		bool OnlyUsesCopyRecords() const
 		{
@@ -149,14 +163,17 @@ protected:
 
 		bool IsValid() const
 		{
-			return (NodeVariableProperty != NULL) && (EvaluationHandlerProperty != NULL);
+			return NodeVariableProperty != nullptr && EvaluationHandlerProperty != nullptr;
 		}
 
 		void PatchFunctionNameAndCopyRecordsInto(UObject* TargetObject) const;
 
 		void RegisterPin(UEdGraphPin* SourcePin, UProperty* AssociatedProperty, int32 AssociatedPropertyArrayIndex);
 
+		UStructProperty* GetHandlerNodeProperty() const { return NodeVariableProperty; }
+
 	private:
+
 		bool CheckForVariableGet(FAnimNodeSinglePropertyHandler& Handler, UEdGraphPin* SourcePin);
 
 		bool CheckForLogicalNot(FAnimNodeSinglePropertyHandler& Handler, UEdGraphPin* SourcePin);
@@ -229,7 +246,8 @@ private:
 	UK2Node_CallFunction* SpawnCallAnimInstanceFunction(UEdGraphNode* SourceNode, FName FunctionName);
 
 	// Creates an evaluation handler for an FExposedValue property in an animation node
-	void CreateEvaluationHandler(UAnimGraphNode_Base* VisualAnimNode, FEvaluationHandlerRecord& Record);
+	void CreateEvaluationHandlerStruct(UAnimGraphNode_Base* VisualAnimNode, FEvaluationHandlerRecord& Record);
+	void CreateEvaluationHandlerInstance(UAnimGraphNode_Base* VisualAnimNode, FEvaluationHandlerRecord& Record);
 
 	// Prunes any nodes that aren't reachable via a pose link
 	void PruneIsolatedAnimationNodes(const TArray<UAnimGraphNode_Base*>& RootSet, TArray<UAnimGraphNode_Base*>& GraphNodes);
@@ -243,6 +261,13 @@ private:
 	// Compiles one use cached pose instance
 	void ProcessUseCachedPose(UAnimGraphNode_UseCachedPose* UseCachedPose);
 
+	// Compiles one sub instance node
+	void ProcessSubInstance(UAnimGraphNode_SubInstance* SubInstance);
+
+	// Traverses subinstance links looking for slot names and state machine names, returning their count in a name map
+	typedef TMap<FName, int32> NameToCountMap;
+	void GetDuplicatedSlotAndStateNames(UAnimGraphNode_SubInstance* InSubInstance, NameToCountMap& OutStateMachineNameToCountMap, NameToCountMap& OutSlotNameToCountMap);
+
 	// Compiles an entire animation graph
 	void ProcessAllAnimationNodes();
 
diff --git a/Engine/Source/Editor/MovieSceneCaptureDialog/Private/MovieSceneCaptureDialogModule.cpp b/Engine/Source/Editor/MovieSceneCaptureDialog/Private/MovieSceneCaptureDialogModule.cpp
--- a/Engine/Source/Editor/MovieSceneCaptureDialog/Private/MovieSceneCaptureDialogModule.cpp
+++ b/Engine/Source/Editor/MovieSceneCaptureDialog/Private/MovieSceneCaptureDialogModule.cpp
@@ -335,8 +335,8 @@ struct FInEditorCapture
 		FAudioDevice* AudioDevice = GWorld->GetAudioDevice();
 		if (AudioDevice != nullptr)
 		{
-			TransientMasterVolume = AudioDevice->TransientMasterVolume;
-			AudioDevice->TransientMasterVolume = 0.0f;
+			TransientMasterVolume = AudioDevice->GetTransientMasterVolume();
+			AudioDevice->SetTransientMasterVolume(0.0f);
 		}
 
 		GEditor->RequestPlaySession(true, nullptr, false);
@@ -454,7 +454,7 @@ struct FInEditorCapture
 		FAudioDevice* AudioDevice = GWorld->GetAudioDevice();
 		if (AudioDevice != nullptr)
 		{
-			AudioDevice->TransientMasterVolume = TransientMasterVolume;
+			AudioDevice->SetTransientMasterVolume(TransientMasterVolume);
 		}
 
 		CaptureObject->Close();
diff --git a/Engine/Source/Editor/Persona/Private/AnimGraphNodeDetails.cpp b/Engine/Source/Editor/Persona/Private/AnimGraphNodeDetails.cpp
--- a/Engine/Source/Editor/Persona/Private/AnimGraphNodeDetails.cpp
+++ b/Engine/Source/Editor/Persona/Private/AnimGraphNodeDetails.cpp
@@ -240,10 +240,11 @@ TSharedRef<SWidget> FAnimGraphNodeDetails::CreatePropertyWidget(UProperty* Targe
 bool FAnimGraphNodeDetails::OnShouldFilterAnimAsset( const FAssetData& AssetData, const UClass* NodeToFilterFor ) const
 {
 	const FString* SkeletonName = AssetData.TagsAndValues.Find(TEXT("Skeleton"));
-	if (*SkeletonName == TargetSkeletonName)
+	if ((SkeletonName != nullptr) && (*SkeletonName == TargetSkeletonName))
 	{
 		const UClass* AssetClass = AssetData.GetClass();
-		if (NodeToFilterFor == GetNodeClassForAsset(AssetClass))
+		// If node is an 'asset player', only let you select the right kind of asset for it
+		if (!NodeToFilterFor->IsChildOf(UAnimGraphNode_AssetPlayerBase::StaticClass()) || SupportNodeClassForAsset(AssetClass, NodeToFilterFor))
 		{
 			return false;
 		}
@@ -516,8 +517,7 @@ void FBoneReferenceCustomization::CustomizeHeader( TSharedRef<IPropertyHandle> S
 
 		HeaderRow.ValueContent()
 		[
-			SNew(SBoneSelectionWidget)
-			.Skeleton(TargetSkeleton)
+			SNew(SBoneSelectionWidget, TargetSkeleton)
 			.Tooltip(StructPropertyHandle->GetToolTipText())
 			.OnBoneSelectionChanged(this, &FBoneReferenceCustomization::OnBoneSelectionChanged)
 			.OnGetSelectedBone(this, &FBoneReferenceCustomization::GetSelectedBone)
diff --git a/Engine/Source/Editor/Persona/Private/Persona.cpp b/Engine/Source/Editor/Persona/Private/Persona.cpp
--- a/Engine/Source/Editor/Persona/Private/Persona.cpp
+++ b/Engine/Source/Editor/Persona/Private/Persona.cpp
@@ -79,6 +79,7 @@
 
 #include "MessageLog.h"
 #include "SAdvancedPreviewDetailsTab.h"
+#include "PhysicsEngine/PhysicsAsset.h"
 
 #define LOCTEXT_NAMESPACE "FPersona"
 
@@ -865,7 +866,52 @@ void FPersona::CreateAnimation(const TArray<UObject*> NewAssets, int32 Option)
 		}
 		else
 		{
-			// give warning
+			FMessageDialog::Open(EAppMsgType::Ok, LOCTEXT("FailedToCreateAsset", "Failed to create asset"));
+		}
+	}
+}
+
+void FPersona::CreatePoseAsset(const TArray<UObject*> NewAssets, int32 Option)
+{
+	bool bResult = false;
+	if (NewAssets.Num() > 0)
+	{
+		USkeletalMeshComponent* MeshComponent = GetPreviewMeshComponent();
+		UAnimSequence* Sequence = Cast<UAnimSequence>(GetPreviewAnimationAsset());
+
+		for (auto NewAsset : NewAssets)
+		{
+			UPoseAsset* NewPoseAsset = Cast<UPoseAsset>(NewAsset);
+			if (NewPoseAsset)
+			{
+				switch (Option)
+				{
+				case 0:
+					NewPoseAsset->AddOrUpdatePoseWithUniqueName(MeshComponent);
+					break;
+				case 1:
+					NewPoseAsset->CreatePoseFromAnimation(Sequence);
+					break;
+				}
+				
+				bResult = true;
+			}
+		}
+
+		// if it contains error, warn them
+		if (bResult)
+		{
+			OnAssetCreated(NewAssets);
+			
+			// if it created based on current mesh component, 
+			if (Option == 0)
+			{
+				PreviewComponent->PreviewInstance->ResetModifiedBone();
+			}
+		}
+		else
+		{
+			FMessageDialog::Open(EAppMsgType::Ok, LOCTEXT("FailedToCreateAsset", "Failed to create asset"));
 		}
 	}
 }
@@ -912,6 +958,38 @@ void FPersona::FillCreateAnimationMenu(FMenuBuilder& MenuBuilder) const
 	MenuBuilder.EndSection();
 }
 
+void FPersona::FillCreatePoseAssetMenu(FMenuBuilder& MenuBuilder) const
+{
+	TArray<TWeakObjectPtr<USkeleton>> Skeletons;
+
+	Skeletons.Add(TargetSkeleton);
+
+	// create rig
+	MenuBuilder.BeginSection("CreatePoseAssetSubMenu", LOCTEXT("CreatePoseAssetSubMenuHeading", "Create PoseAsset"));
+	{
+		MenuBuilder.AddMenuEntry(
+			LOCTEXT("CreatePoseAsset_CurrentPose", "From Current Pose"),
+			LOCTEXT("CreatePoseAsset_CurrentPose_Tooltip", "Create PoseAsset from current pose."),
+			FSlateIcon(),
+			FUIAction(
+				FExecuteAction::CreateStatic(&AnimationEditorUtils::ExecuteNewAnimAsset<UPoseAssetFactory, UPoseAsset>, Skeletons, FString("_PoseAsset"), FAnimAssetCreated::CreateSP(this, &FPersona::CreatePoseAsset, 0), false),
+				FCanExecuteAction()
+				)
+			);
+
+		MenuBuilder.AddMenuEntry(
+			LOCTEXT("CreatePoseAsset_CurrentAnimation", "From Current Animation"),
+			LOCTEXT("CreatePoseAsset_CurrentAnimation_Tooltip", "Create Animation from current animation."),
+			FSlateIcon(),
+			FUIAction(
+				FExecuteAction::CreateStatic(&AnimationEditorUtils::ExecuteNewAnimAsset<UPoseAssetFactory, UPoseAsset>, Skeletons, FString("_PoseAsset"), FAnimAssetCreated::CreateSP(this, &FPersona::CreatePoseAsset, 1), false),
+				FCanExecuteAction::CreateSP(this, &FPersona::HasValidAnimationSequencePlaying)
+				)
+			);
+	}
+	MenuBuilder.EndSection();
+}
+
 TSharedRef< SWidget > FPersona::GenerateCreateAssetMenu( USkeleton* Skeleton ) const
 {
 	const bool bShouldCloseWindowAfterMenuSelection = true;
@@ -926,8 +1004,16 @@ TSharedRef< SWidget > FPersona::GenerateCreateAssetMenu( USkeleton* Skeleton ) c
 				LOCTEXT("CreateAnimationSubmenu_ToolTip", "Create Animation for this skeleton"),
 				FNewMenuDelegate::CreateSP(this, &FPersona::FillCreateAnimationMenu),
 				false,
-				FSlateIcon(FEditorStyle::GetStyleSetName(), "Persona.AssetActions.CreateAnimAsset")
+				FSlateIcon(FEditorStyle::GetStyleSetName(), "ClassIcon.AnimSequence")
 				);
+
+		MenuBuilder.AddSubMenu(
+			LOCTEXT("CreatePoseAssetSubmenu", "Create PoseAsset"),
+			LOCTEXT("CreatePoseAsssetSubmenu_ToolTip", "Create PoseAsset for this skeleton"),
+			FNewMenuDelegate::CreateSP(this, &FPersona::FillCreatePoseAssetMenu),
+			false,
+			FSlateIcon(FEditorStyle::GetStyleSetName(), "ClassIcon.PoseAsset")
+			);
 	}
 	MenuBuilder.EndSection();
 
@@ -1114,11 +1200,18 @@ void FPersona::SetDetailObject(UObject* Obj)
 	UpdateSelectionDetails(Obj, ForcedTitle);
 }
 
-UDebugSkelMeshComponent* FPersona::GetPreviewMeshComponent()
+UDebugSkelMeshComponent* FPersona::GetPreviewMeshComponent() const
 {
 	return PreviewComponent;
 }
 
+UAnimInstance* FPersona::GetPreviewAnimInstance() const
+{
+	UDebugSkelMeshComponent* DebugSkelComp = GetPreviewMeshComponent();
+	return (DebugSkelComp != nullptr) ? DebugSkelComp->GetAnimInstance() : nullptr;
+}
+
+
 void FPersona::OnPostReimport(UObject* InObject, bool bSuccess)
 {
 	if (bSuccess)
diff --git a/Engine/Source/Editor/Persona/Private/Persona.h b/Engine/Source/Editor/Persona/Private/Persona.h
--- a/Engine/Source/Editor/Persona/Private/Persona.h
+++ b/Engine/Source/Editor/Persona/Private/Persona.h
@@ -63,8 +63,10 @@ public:
 
 	virtual void OnActiveTabChanged(TSharedPtr<SDockTab> PreviouslyActive, TSharedPtr<SDockTab> NewlyActivated) override;
 
-	/** handle after compile **/
-	UDebugSkelMeshComponent* GetPreviewMeshComponent();
+	/** Get the preview DebugSkelMeshComponent */
+	UDebugSkelMeshComponent* GetPreviewMeshComponent() const;
+	/** Get the AnimInstance associated with the preview DEbugSkelMeshComponent */
+	UAnimInstance* GetPreviewAnimInstance() const;
 
 	// Gets the skeleton being edited/viewed by this Persona instance
 	USkeleton* GetSkeleton() const;
@@ -696,7 +698,9 @@ private:
 	void OnAssetCreated(const TArray<UObject*> NewAssets);
 	TSharedRef< SWidget > GenerateCreateAssetMenu( USkeleton* Skeleton ) const;
 	void FillCreateAnimationMenu(FMenuBuilder& MenuBuilder) const;
+	void FillCreatePoseAssetMenu(FMenuBuilder& MenuBuilder) const;
 	void CreateAnimation(const TArray<UObject*> NewAssets, int32 Option);
+	void CreatePoseAsset(const TArray<UObject*> NewAssets, int32 Option);
 
 	/** Extend menu and toolbar */
 	void ExtendMenu();
diff --git a/Engine/Source/Editor/Persona/Private/SAnimCurvePanel.cpp b/Engine/Source/Editor/Persona/Private/SAnimCurvePanel.cpp
--- a/Engine/Source/Editor/Persona/Private/SAnimCurvePanel.cpp
+++ b/Engine/Source/Editor/Persona/Private/SAnimCurvePanel.cpp
@@ -463,7 +463,7 @@ void SCurveEdTrack::DeleteTrack()
 
 void SCurveEdTrack::ToggleCurveMode(ECheckBoxState NewState,EAnimCurveFlags ModeToSet)
 {
-	const int32 AllModes = (ACF_DrivesMorphTarget|ACF_DrivesMaterial|ACF_DrivesPose);
+	const int32 AllModes = (ACF_DriveMorphTarget|ACF_DriveMaterial);
 	check((ModeToSet&AllModes) != 0); //unexpected value for ModeToSet
 
 	FText UndoLabel;
@@ -472,33 +472,25 @@ void SCurveEdTrack::ToggleCurveMode(ECheckBoxState NewState,EAnimCurveFlags Mode
 	
 	if(bIsSwitchingFlagOn)
 	{
-		if(ModeToSet == ACF_DrivesMorphTarget)
+		if(ModeToSet == ACF_DriveMorphTarget)
 		{
 			UndoLabel = LOCTEXT("AnimCurve_TurnOnMorphMode", "Enable driving of morph targets");
 		}
-		else if(ModeToSet == ACF_DrivesMaterial)
+		else if(ModeToSet == ACF_DriveMaterial)
 		{
 			UndoLabel = LOCTEXT("AnimCurve_TurnOnMaterialMode", "Enable driving of materials");
 		}
-		else if (ModeToSet == ACF_DrivesPose)
-		{
-			UndoLabel = LOCTEXT("AnimCurve_TurnOnPoseMode", "Enable driving of pose");
-		}
 	}
 	else
 	{
-		if(ModeToSet == ACF_DrivesMorphTarget)
+		if(ModeToSet == ACF_DriveMorphTarget)
 		{
 			UndoLabel = LOCTEXT("AnimCurve_TurnOffMorphMode", "Disable driving of morph targets");
 		}
-		else if(ModeToSet == ACF_DrivesMaterial)
+		else if(ModeToSet == ACF_DriveMaterial)
 		{
 			UndoLabel = LOCTEXT("AnimCurve_TurnOffMaterialMode", "Disable driving of materials");
 		}
-		else if (ModeToSet == ACF_DrivesPose)
-		{
-			UndoLabel = LOCTEXT("AnimCurve_TurnOffPoseMode", "Disable driving of pose");
-		}
 	}
 
 	const FScopedTransaction Transaction( UndoLabel );
@@ -796,8 +788,8 @@ FReply SAnimCurvePanel::OnContextMenu()
 	{
 		MenuBuilder.AddWidget(
 			SNew(SCheckBox)
-			.IsChecked( this, &SAnimCurvePanel::AreAllCurvesOfMode, ACF_DrivesMorphTarget )
-			.OnCheckStateChanged( this, &SAnimCurvePanel::ToggleAllCurveModes, ACF_DrivesMorphTarget )
+			.IsChecked( this, &SAnimCurvePanel::AreAllCurvesOfMode, ACF_DriveMorphTarget )
+			.OnCheckStateChanged( this, &SAnimCurvePanel::ToggleAllCurveModes, ACF_DriveMorphTarget )
 			.ToolTipText(LOCTEXT("MorphCurveModeTooltip", "This curve drives a morph target"))
 			[
 				SNew(STextBlock)
@@ -808,8 +800,8 @@ FReply SAnimCurvePanel::OnContextMenu()
 
 		MenuBuilder.AddWidget(
 			SNew(SCheckBox)
-			.IsChecked( this, &SAnimCurvePanel::AreAllCurvesOfMode, ACF_DrivesMaterial )
-			.OnCheckStateChanged( this, &SAnimCurvePanel::ToggleAllCurveModes, ACF_DrivesMaterial )
+			.IsChecked( this, &SAnimCurvePanel::AreAllCurvesOfMode, ACF_DriveMaterial )
+			.OnCheckStateChanged( this, &SAnimCurvePanel::ToggleAllCurveModes, ACF_DriveMaterial )
 			.ToolTipText(LOCTEXT("MaterialCurveModeTooltip", "This curve drives a material"))
 			.HAlign(HAlign_Left)
 			[
@@ -818,18 +810,6 @@ FReply SAnimCurvePanel::OnContextMenu()
 			],
 			FText()
 		);
-		MenuBuilder.AddWidget(
-			SNew(SCheckBox)
-			.IsChecked(this, &SAnimCurvePanel::AreAllCurvesOfMode, ACF_DrivesPose)
-			.OnCheckStateChanged(this, &SAnimCurvePanel::ToggleAllCurveModes, ACF_DrivesPose)
-			.ToolTipText(LOCTEXT("PoseCurveModeTooltip", "This curve drives a pose"))
-			.HAlign(HAlign_Left)
-			[
-				SNew(STextBlock)
-				.Text(LOCTEXT("PoseCurveMode", "Pose Curve"))
-			],
-			FText()
-			);
 	}
 	MenuBuilder.EndSection();
 
@@ -1330,8 +1310,8 @@ TSharedRef<SWidget> SAnimCurvePanel::CreateCurveContextMenu(FAnimCurveBaseInterf
 	{
 		MenuBuilder.AddWidget(
 			SNew(SCheckBox)
-			.IsChecked(this, &SAnimCurvePanel::GetCurveFlagAsCheckboxState, Curve, ACF_DrivesMorphTarget)
-			.OnCheckStateChanged(this, &SAnimCurvePanel::SetCurveFlagFromCheckboxState, Curve, ACF_DrivesMorphTarget)
+			.IsChecked(this, &SAnimCurvePanel::GetCurveFlagAsCheckboxState, Curve, ACF_DriveMorphTarget)
+			.OnCheckStateChanged(this, &SAnimCurvePanel::SetCurveFlagFromCheckboxState, Curve, ACF_DriveMorphTarget)
 			.ToolTipText(LOCTEXT("MorphCurveModeTooltip", "This curve drives a morph target"))
 			[
 				SNew(STextBlock)
@@ -1342,8 +1322,8 @@ TSharedRef<SWidget> SAnimCurvePanel::CreateCurveContextMenu(FAnimCurveBaseInterf
 
 		MenuBuilder.AddWidget(
 			SNew(SCheckBox)
-			.IsChecked(this, &SAnimCurvePanel::GetCurveFlagAsCheckboxState, Curve, ACF_DrivesMaterial)
-			.OnCheckStateChanged(this, &SAnimCurvePanel::SetCurveFlagFromCheckboxState, Curve, ACF_DrivesMaterial)
+			.IsChecked(this, &SAnimCurvePanel::GetCurveFlagAsCheckboxState, Curve, ACF_DriveMaterial)
+			.OnCheckStateChanged(this, &SAnimCurvePanel::SetCurveFlagFromCheckboxState, Curve, ACF_DriveMaterial)
 			.ToolTipText(LOCTEXT("MaterialCurveModeTooltip", "This curve drives a material"))
 			[
 				SNew(STextBlock)
@@ -1351,18 +1331,6 @@ TSharedRef<SWidget> SAnimCurvePanel::CreateCurveContextMenu(FAnimCurveBaseInterf
 			],
 			FText()
 			);
-
-		MenuBuilder.AddWidget(
-			SNew(SCheckBox)
-			.IsChecked(this, &SAnimCurvePanel::GetCurveFlagAsCheckboxState, Curve, ACF_DrivesPose)
-			.OnCheckStateChanged(this, &SAnimCurvePanel::SetCurveFlagFromCheckboxState, Curve, ACF_DrivesPose)
-			.ToolTipText(LOCTEXT("PoseCurveModeTooltip", "This curve drives a pose asset"))
-			[
-				SNew(STextBlock)
-				.Text(LOCTEXT("PoseCurveMode", "Pose Curve"))
-			],
-			FText()
-			);
 	}
 	MenuBuilder.EndSection();
 
diff --git a/Engine/Source/Editor/Persona/Private/SAnimCurveViewer.cpp b/Engine/Source/Editor/Persona/Private/SAnimCurveViewer.cpp
--- a/Engine/Source/Editor/Persona/Private/SAnimCurveViewer.cpp
+++ b/Engine/Source/Editor/Persona/Private/SAnimCurveViewer.cpp
@@ -7,6 +7,9 @@
 #include "AssetRegistryModule.h"
 #include "ScopedTransaction.h"
 #include "SSearchBox.h"
+#include "SInlineEditableTextBlock.h"
+#include "STextEntryPopup.h"
+#include "Animation/AnimSingleNodeInstance.h"
 
 #define LOCTEXT_NAMESPACE "SAnimCurveViewer"
 
@@ -19,104 +22,45 @@ const float MaxMorphWeight = 5.f;
 //////////////////////////////////////////////////////////////////////////
 // SAnimCurveListRow
 
-typedef TSharedPtr< FDisplayedAnimCurveInfo > FDisplayedAnimCurveInfoPtr;
 
-class SAnimCurveListRow
-	: public SMultiColumnTableRow< FDisplayedAnimCurveInfoPtr >
-{
-public:
-
-	SLATE_BEGIN_ARGS( SAnimCurveListRow ) {}
-
-	/** The item for this row **/
-	SLATE_ARGUMENT( FDisplayedAnimCurveInfoPtr, Item )
-
-	/* The SAnimCurveViewer that we push the morph target weights into */
-	SLATE_ARGUMENT( class SAnimCurveViewer*, AnimCurveViewer )
-
-	/* Widget used to display the list of animation curve */
-	SLATE_ARGUMENT( TSharedPtr<SAnimCurveListType>, AnimCurveListView )
-
-	/* Persona used to update the viewport when a weight slider is dragged */
-	SLATE_ARGUMENT( TWeakPtr<FPersona>, Persona )
-
-	SLATE_END_ARGS()
-
-	void Construct( const FArguments& InArgs, const TSharedRef<STableViewBase>& OwnerTableView );
-
-	/** Overridden from SMultiColumnTableRow.  Generates a widget for this column of the tree row. */
-	virtual TSharedRef<SWidget> GenerateWidgetForColumn( const FName& ColumnName ) override;
-
-private:
-
-	/**
-	* Called when the user changes the value of the SSpinBox
-	*
-	* @param NewWeight - The new number the SSpinBox is set to
-	*
-	*/
-	void OnAnimCurveWeightChanged( float NewWeight );
-	/**
-	* Called when the user types the value and enters
-	*
-	* @param NewWeight - The new number the SSpinBox is set to
-	*
-	*/
-	void OnAnimCurveWeightValueCommitted( float NewWeight, ETextCommit::Type CommitType);
-
-	/** Auto fill check call back functions */
-	void OnAnimCurveAutoFillChecked(ECheckBoxState InState);
-	ECheckBoxState IsAnimCurveAutoFillChangedChecked() const;
-	
-	/**
-	* Returns the weight of this morph target
-	*
-	* @return SearchText - The new number the SSpinBox is set to
-	*
-	*/
-	float GetWeight() const;
-
-	/* The SAnimCurveViewer that we push the morph target weights into */
-	SAnimCurveViewer* AnimCurveViewer;
-
-	/** Widget used to display the list of animation curve */
-	TSharedPtr<SAnimCurveListType> AnimCurveListView;
-
-	/** The name and weight of the morph target */
-	FDisplayedAnimCurveInfoPtr	Item;
-
-	/** Pointer back to the Persona that owns us */
-	TWeakPtr<FPersona> PersonaPtr;
-};
 
 void SAnimCurveListRow::Construct( const FArguments& InArgs, const TSharedRef<STableViewBase>& InOwnerTableView )
 {
 	Item = InArgs._Item;
-	AnimCurveViewer = InArgs._AnimCurveViewer;
-	AnimCurveListView = InArgs._AnimCurveListView;
-	PersonaPtr = InArgs._Persona;
+	AnimCurveViewerPtr = InArgs._AnimCurveViewerPtr;
 
 	check( Item.IsValid() );
 
-	SMultiColumnTableRow< FDisplayedAnimCurveInfoPtr >::Construct( FSuperRowType::FArguments(), InOwnerTableView );
+	SMultiColumnTableRow< TSharedPtr<FDisplayedAnimCurveInfo> >::Construct( FSuperRowType::FArguments(), InOwnerTableView );
 }
 
 TSharedRef< SWidget > SAnimCurveListRow::GenerateWidgetForColumn( const FName& ColumnName )
 {
 	if ( ColumnName == ColumnId_AnimCurveNameLabel )
 	{
-		return
-			SNew( SVerticalBox )
-
-			+ SVerticalBox::Slot()
-			.AutoHeight()
-			.Padding( 0.0f, 4.0f )
-			.VAlign( VAlign_Center )
-			[
-				SNew( STextBlock )
-				.Text( FText::FromName(Item->Name) )
-				.HighlightText( AnimCurveViewer->GetFilterText() )
-			];
+		TSharedPtr<SAnimCurveViewer> AnimCurveViewer = AnimCurveViewerPtr.Pin();
+		if (AnimCurveViewer.IsValid())
+		{
+			return
+				SNew(SVerticalBox)
+
+				+ SVerticalBox::Slot()
+				.AutoHeight()
+				.Padding(4)
+				.VAlign(VAlign_Center)
+				[
+					SAssignNew(Item->EditableText, SInlineEditableTextBlock)
+					.OnTextCommitted(AnimCurveViewer.Get(), &SAnimCurveViewer::OnNameCommitted, Item)
+					.ColorAndOpacity(this, &SAnimCurveListRow::GetItemTextColor)
+					.IsSelected(this, &SAnimCurveListRow::IsSelected)
+					.Text(this, &SAnimCurveListRow::GetItemName)
+					.HighlightText(AnimCurveViewer->GetFilterText())
+				];
+		}
+		else
+		{
+			return SNullWidget::NullWidget;
+		}
 	}
 	else if ( ColumnName == ColumnID_AnimCurveWeightLabel )
 	{
@@ -161,10 +105,18 @@ void SAnimCurveListRow::OnAnimCurveAutoFillChecked(ECheckBoxState InState)
 {
 	Item->bAutoFillData = InState == ECheckBoxState::Checked;
 
-	if (Item->bAutoFillData)
+	TSharedPtr<SAnimCurveViewer> AnimCurveViewer = AnimCurveViewerPtr.Pin();
+	if (AnimCurveViewer.IsValid())
 	{
-		// clear value so that it can be filled up
-		AnimCurveViewer->RemoveAnimCurveOverride(Item->Name);
+		if (Item->bAutoFillData)
+		{
+			// clear value so that it can be filled up
+			AnimCurveViewer->RemoveAnimCurveOverride(Item->SmartName.DisplayName);
+		}
+		else
+		{
+			AnimCurveViewer->AddAnimCurveOverride(Item->SmartName.DisplayName, Item->Weight);
+		}
 	}
 }
 
@@ -175,78 +127,132 @@ ECheckBoxState SAnimCurveListRow::IsAnimCurveAutoFillChangedChecked() const
 
 void SAnimCurveListRow::OnAnimCurveWeightChanged( float NewWeight )
 {
-	Item->Weight = NewWeight;
+	float NewValidWeight = FMath::Clamp(NewWeight, -MaxMorphWeight, MaxMorphWeight);
+	Item->Weight = NewValidWeight;
 	Item->bAutoFillData = false;
 
-	AnimCurveViewer->AddAnimCurveOverride( Item->Name, Item->Weight);
-
-	if (PersonaPtr.IsValid())
+	TSharedPtr<SAnimCurveViewer> AnimCurveViewer = AnimCurveViewerPtr.Pin();
+	if (AnimCurveViewer.IsValid())
 	{
-		PersonaPtr.Pin()->RefreshViewport();
-	}
-}
-
-void SAnimCurveListRow::OnAnimCurveWeightValueCommitted( float NewWeight, ETextCommit::Type CommitType)
-{
-	if (CommitType == ETextCommit::OnEnter || CommitType == ETextCommit::OnUserMovedFocus)
-	{
-		float NewValidWeight = FMath::Clamp(NewWeight, -MaxMorphWeight, MaxMorphWeight);
-		Item->Weight = NewValidWeight;
-		Item->bAutoFillData = false;
-
-		AnimCurveViewer->AddAnimCurveOverride(Item->Name, Item->Weight);
+		// If adjusting an element that is not selected, make it the only selection
+		if (!AnimCurveViewer->AnimCurveListView->IsItemSelected(Item))
+		{
+			AnimCurveViewer->AnimCurveListView->SetSelection(Item);
+		}
 
-		TArray< TSharedPtr< FDisplayedAnimCurveInfo > > SelectedRows = AnimCurveListView->GetSelectedItems();
+		// Add override
+		AnimCurveViewer->AddAnimCurveOverride(Item->SmartName.DisplayName, Item->Weight);
 
 		// ...then any selected rows need changing by the same delta
-		for(auto ItemIt = SelectedRows.CreateIterator(); ItemIt; ++ItemIt)
+		TArray< TSharedPtr< FDisplayedAnimCurveInfo > > SelectedRows = AnimCurveViewer->AnimCurveListView->GetSelectedItems();
+		for (auto ItemIt = SelectedRows.CreateIterator(); ItemIt; ++ItemIt)
 		{
 			TSharedPtr< FDisplayedAnimCurveInfo > RowItem = (*ItemIt);
 
-			if(RowItem != Item) // Don't do "this" row again if it's selected
+			if (RowItem != Item) // Don't do "this" row again if it's selected
 			{
 				RowItem->Weight = NewValidWeight;
 				RowItem->bAutoFillData = false;
-				AnimCurveViewer->AddAnimCurveOverride(RowItem->Name, RowItem->Weight);
+				AnimCurveViewer->AddAnimCurveOverride(RowItem->SmartName.DisplayName, RowItem->Weight);
 			}
 		}
 
-		if(PersonaPtr.IsValid())
+		// Refresh viewport
+		TSharedPtr<FPersona> Persona = AnimCurveViewer->PersonaPtr.Pin();
+		if (Persona.IsValid())
 		{
-			PersonaPtr.Pin()->RefreshViewport();
+			Persona->RefreshViewport();
 		}
 	}
 }
 
-float SAnimCurveListRow::GetWeight() const 
-{ 
-	if (PersonaPtr.IsValid() && Item->bAutoFillData)
+void SAnimCurveListRow::OnAnimCurveWeightValueCommitted( float NewWeight, ETextCommit::Type CommitType)
+{
+	if (CommitType == ETextCommit::OnEnter || CommitType == ETextCommit::OnUserMovedFocus)
 	{
-		USkeletalMeshComponent* SkelComp = PersonaPtr.Pin()->GetPreviewMeshComponent();
-		UAnimInstance* AnimInstance = (SkelComp) ? SkelComp->GetAnimInstance() : nullptr;
-		if (AnimInstance)
-		{
-			TMap<FName, float> CurveList;
+		OnAnimCurveWeightChanged(NewWeight);
+	}
+}
+
+FText SAnimCurveListRow::GetItemName() const
+{
+	FName ItemName;
+
+	TSharedPtr<SAnimCurveViewer> AnimCurveViewer = AnimCurveViewerPtr.Pin();
+	if (AnimCurveViewer.IsValid())
+	{
+		const FSmartNameMapping* Mapping = Item->Skeleton->GetSmartNameContainer(AnimCurveViewer->ContainerName);
+		check(Mapping);
+		Mapping->GetName(Item->SmartName.UID, ItemName);
+	}
+
+	return FText::FromName(ItemName);
+}
 
-			// make sure if they have value that's not same as saved value
-			AnimInstance->GetAnimationCurveList(AnimCurveViewer->CurrentCurveFlag, CurveList);
+bool SAnimCurveListRow::GetActiveWeight(float& OutWeight) const
+{
+	bool bFoundActive = false;
 
-			if (CurveList.Num() > 0)
+	// If anim viewer
+	TSharedPtr<SAnimCurveViewer> AnimCurveViewer = AnimCurveViewerPtr.Pin();
+	if (AnimCurveViewer.IsValid())
+	{
+		// If persona
+		TSharedPtr<FPersona> Persona = AnimCurveViewer->PersonaPtr.Pin();
+		if (Persona.IsValid())
+		{
+			// If anim instance
+			UAnimInstance* AnimInstance = Persona->GetPreviewAnimInstance();
+			if (AnimInstance)
 			{
-				float* CurrentValue = CurveList.Find(Item->Name);
-				if (CurrentValue && *CurrentValue != 0.f)
+				// See if curve is in active set
+				TMap<FName, float> CurveList;
+				AnimInstance->GetAnimationCurveList(ACF_EditorPreviewCurves, CurveList);
+
+				float* CurrentValue = CurveList.Find(Item->SmartName.DisplayName);
+				if (CurrentValue)
 				{
-					return *CurrentValue;
+					OutWeight = *CurrentValue;
+					// Remember we found it
+					bFoundActive = true;
 				}
 			}
-			else
-			{
-				return 0.f;
-			}
 		}
 	}
 
-	return Item->Weight; 
+	return bFoundActive;
+}
+
+
+FSlateColor SAnimCurveListRow::GetItemTextColor() const
+{
+	// If row is selected, show text as black to make it easier to read
+	if (IsSelected())
+	{
+		return FLinearColor(0, 0, 0);
+	}
+
+	// If not selected, show bright if active
+	bool bItemActive = true;
+	if (Item->bAutoFillData)
+	{
+		float Weight = 0.f;
+		bItemActive = GetActiveWeight(Weight);
+	}
+
+	return bItemActive ? FLinearColor(1, 1, 1) : FLinearColor(0.5, 0.5, 0.5);
+}
+
+float SAnimCurveListRow::GetWeight() const 
+{ 
+	float Weight = Item->Weight;
+
+	if (Item->bAutoFillData)
+	{
+		GetActiveWeight(Weight);
+	}
+
+	return Weight;
 }
 
 //////////////////////////////////////////////////////////////////////////
@@ -258,25 +264,25 @@ public:
 
 	SLATE_BEGIN_ARGS(SAnimCurveTypeList) {}
 
-	/** The item for this row **/
-	SLATE_ARGUMENT(int32, CurveFlags)
+		/** The item for this row **/
+		SLATE_ARGUMENT(int32, CurveFlags)
 
-	/* The SAnimCurveViewer that we push the morph target weights into */
-	SLATE_ARGUMENT(class SAnimCurveViewer*, AnimCurveViewer)
+		/* The SAnimCurveViewer that we push the morph target weights into */
+		SLATE_ARGUMENT(TWeakPtr<SAnimCurveViewer>, AnimCurveViewerPtr)
 
 	SLATE_END_ARGS()
 
 	void Construct(const FArguments& InArgs);
 
 	/** Auto fill check call back functions */
-	void OnAnimCurveTypeShowlChecked(ECheckBoxState InState);
+	void OnAnimCurveTypeShowChecked(ECheckBoxState InState);
 	ECheckBoxState IsAnimCurveTypeShowChangedChecked() const;
 
 	FText GetAnimCurveType() const;
 private:
 
 	/* The SAnimCurveViewer that we push the morph target weights into */
-	SAnimCurveViewer* AnimCurveViewer;
+	TWeakPtr<SAnimCurveViewer> AnimCurveViewerPtr;
 
 	/** The name and weight of the morph target */
 	int32 CurveFlags;
@@ -285,71 +291,85 @@ private:
 void SAnimCurveTypeList::Construct(const FArguments& InArgs)
 {
 	CurveFlags = InArgs._CurveFlags;
-	AnimCurveViewer = InArgs._AnimCurveViewer;
-
-	ChildSlot
-	[
-		SNew(SHorizontalBox)
+	AnimCurveViewerPtr = InArgs._AnimCurveViewerPtr;
 
-		+SHorizontalBox::Slot()
-		.AutoWidth()
-		.Padding(0.0f, 1.0f)
-		.VAlign(VAlign_Center)
-		.HAlign(HAlign_Center)
+	TSharedPtr<SAnimCurveViewer> AnimCurveViewer = AnimCurveViewerPtr.Pin();
+	if (AnimCurveViewer.IsValid())
+	{
+		ChildSlot
 		[
-			SNew(SCheckBox)
-			.OnCheckStateChanged(this, &SAnimCurveTypeList::OnAnimCurveTypeShowlChecked)
-			.IsChecked(this, &SAnimCurveTypeList::IsAnimCurveTypeShowChangedChecked)
-		]
+			SNew(SHorizontalBox)
+
+			+ SHorizontalBox::Slot()
+			.AutoWidth()
+			.Padding(0.0f, 1.0f)
+			.VAlign(VAlign_Center)
+			.HAlign(HAlign_Center)
+			[
+				SNew(SCheckBox)
+				.OnCheckStateChanged(this, &SAnimCurveTypeList::OnAnimCurveTypeShowChecked)
+				.IsChecked(this, &SAnimCurveTypeList::IsAnimCurveTypeShowChangedChecked)
+			]
+
+			+ SHorizontalBox::Slot()
+			.AutoWidth()
+			.Padding(3.0f, 1.0f)
+			.VAlign(VAlign_Center)
+			[
+				SNew(STextBlock)
+				.Text(this, &SAnimCurveTypeList::GetAnimCurveType)
+				.HighlightText(AnimCurveViewer->GetFilterText())
+			]
+		];
+	}
 
-		+SHorizontalBox::Slot()
-		.AutoWidth()
-		.Padding(3.0f, 1.0f)
-		.VAlign(VAlign_Center)
-		[
-			SNew(STextBlock)
-			.Text(this, &SAnimCurveTypeList::GetAnimCurveType)
-			.HighlightText(AnimCurveViewer->GetFilterText())
-		]
-	];
 }
 
-void SAnimCurveTypeList::OnAnimCurveTypeShowlChecked(ECheckBoxState InState)
+void SAnimCurveTypeList::OnAnimCurveTypeShowChecked(ECheckBoxState InState)
 {
-	// clear value so that it can be filled up
-	if (InState == ECheckBoxState::Checked)
-	{
-		AnimCurveViewer->CurrentCurveFlag |= CurveFlags;
-	}
-	else
+	TSharedPtr<SAnimCurveViewer> AnimCurveViewer = AnimCurveViewerPtr.Pin();
+	if (AnimCurveViewer.IsValid())
 	{
-		AnimCurveViewer->CurrentCurveFlag &= ~CurveFlags;
-	}
+		// clear value so that it can be filled up
+		if (InState == ECheckBoxState::Checked)
+		{
+			AnimCurveViewer->CurrentCurveFlag |= CurveFlags;
+		}
+		else
+		{
+			AnimCurveViewer->CurrentCurveFlag &= ~CurveFlags;
+		}
 
-	AnimCurveViewer->RefreshCurveList();
+		AnimCurveViewer->RefreshCurveList();
+	}
 }
 
 ECheckBoxState SAnimCurveTypeList::IsAnimCurveTypeShowChangedChecked() const
 {
-	return (AnimCurveViewer->CurrentCurveFlag & CurveFlags) ? ECheckBoxState::Checked : ECheckBoxState::Unchecked;
+	TSharedPtr<SAnimCurveViewer> AnimCurveViewer = AnimCurveViewerPtr.Pin();
+	if (AnimCurveViewer.IsValid())
+	{
+		return (AnimCurveViewer->CurrentCurveFlag & CurveFlags) ? ECheckBoxState::Checked : ECheckBoxState::Unchecked;
+	}
+	else
+	{
+		return ECheckBoxState::Unchecked;
+	}
 }
 
 FText SAnimCurveTypeList::GetAnimCurveType() const
 {
 	switch (CurveFlags)
 	{
-	case ACF_DrivesMorphTarget:
+	case ACF_DriveMorphTarget:
 		return LOCTEXT("AnimCurveType_Morphtarget", "Morph Target");
-	case ACF_TriggerEvent:
-		return LOCTEXT("AnimCurveType_Event", "Event");
-	case ACF_DrivesMaterial:
+	case ACF_DriveAttribute:
+		return LOCTEXT("AnimCurveType_Attribute", "Attribute");
+	case ACF_DriveMaterial:
 		return LOCTEXT("AnimCurveType_Material", "Material");
+		// @fixme: when we move curve types to skeleton fix this
 	case ACF_Metadata:
 		return LOCTEXT("AnimCurveType_Metadata", "Meta Data");
-	case ACF_DriveTrack:
-		return LOCTEXT("AnimCurveType_Tranform", "Transform");
-	case ACF_DrivesPose:
-		return LOCTEXT("AnimCurveType_Pose", "Pose");
 	}
 
 	return LOCTEXT("AnimCurveType_Unknown", "Unknown");
@@ -360,34 +380,35 @@ FText SAnimCurveTypeList::GetAnimCurveType() const
 void SAnimCurveViewer::Construct(const FArguments& InArgs)
 {
 	PersonaPtr = InArgs._Persona;
+
+	ContainerName = USkeleton::AnimCurveMappingName;
+
 	CachedPreviewInstance = nullptr;
 
-	if (PersonaPtr.IsValid())
+	bShowAllCurves = true;
+
+	TSharedPtr<FPersona> Persona = PersonaPtr.Pin();
+	if (Persona.IsValid())
 	{
-		PersonaPtr.Pin()->RegisterOnPreviewMeshChanged(FPersona::FOnPreviewMeshChanged::CreateSP(this, &SAnimCurveViewer::OnPreviewMeshChanged));
-		PersonaPtr.Pin()->RegisterOnAnimChanged(FPersona::FOnAnimChanged::CreateSP(this, &SAnimCurveViewer::OnPreviewAssetChanged));
-		PersonaPtr.Pin()->RegisterOnPostUndo(FPersona::FOnPostUndo::CreateSP(this, &SAnimCurveViewer::OnPostUndo));
+		CurrentSkeleton = Persona->GetSkeleton();
+
+		Persona->RegisterOnPreviewMeshChanged(FPersona::FOnPreviewMeshChanged::CreateSP(this, &SAnimCurveViewer::OnPreviewMeshChanged));
+		Persona->RegisterOnAnimChanged(FPersona::FOnAnimChanged::CreateSP(this, &SAnimCurveViewer::OnPreviewAssetChanged));
+		Persona->RegisterOnPostUndo(FPersona::FOnPostUndo::CreateSP(this, &SAnimCurveViewer::OnPostUndo));
+		Persona->RegisterOnChangeCurves(FPersona::FOnCurvesChanged::CreateSP(this, &SAnimCurveViewer::OnCurvesChanged));
 	}
 
 	// @todo fix this to be filtered
-	CurrentCurveFlag = ACF_DrivesMorphTarget | ACF_DrivesMaterial | ACF_DrivesPose;
+	CurrentCurveFlag = ACF_DriveMorphTarget | ACF_DriveMaterial | ACF_DriveAttribute;
 
 	RefreshCachePreviewInstance();
 	 
-	TSharedPtr<SVerticalBox> AnimTypeBoxContainer = SNew(SVerticalBox);
+	TSharedPtr<SHorizontalBox> AnimTypeBoxContainer = SNew(SHorizontalBox);
 
 	ChildSlot
 	[
 		SNew( SVerticalBox )
 		
-		+ SVerticalBox::Slot()
-		.AutoHeight()
-		.Padding(2)
-		[
-			SNew( STextBlock )
-			.Text( this, &SAnimCurveViewer::GetTitleText )
-		]
-
 		+ SVerticalBox::Slot()
 		.AutoHeight()
 		.Padding(0,2)
@@ -423,16 +444,20 @@ void SAnimCurveViewer::Construct(const FArguments& InArgs)
 			.OnGenerateRow( this, &SAnimCurveViewer::GenerateAnimCurveRow )
 			.OnContextMenuOpening( this, &SAnimCurveViewer::OnGetContextMenuContent )
 			.ItemHeight( 22.0f )
+			.SelectionMode(ESelectionMode::Multi)
 			.HeaderRow
 			(
 				SNew( SHeaderRow )
 				+ SHeaderRow::Column( ColumnId_AnimCurveNameLabel )
+				.FillWidth(1.f)
 				.DefaultLabel( LOCTEXT( "AnimCurveNameLabel", "Curve Name" ) )
 
 				+ SHeaderRow::Column( ColumnID_AnimCurveWeightLabel )
+				.FillWidth(1.f)
 				.DefaultLabel( LOCTEXT( "AnimCurveWeightLabel", "Weight" ) )
 
 				+ SHeaderRow::Column(ColumnID_AnimCurveEditLabel)
+				.FillWidth(0.25f)
 				.DefaultLabel(LOCTEXT("AnimCurveEditLabel", "Auto"))
 			)
 		]
@@ -442,20 +467,39 @@ void SAnimCurveViewer::Construct(const FArguments& InArgs)
 	CreateAnimCurveList();
 }
 
-FText SAnimCurveViewer::GetTitleText() const
+SAnimCurveViewer::~SAnimCurveViewer()
 {
-	UObject* PreviewAsset = nullptr;
-
-	if (PersonaPtr.IsValid())
+	TSharedPtr<FPersona> Persona = PersonaPtr.Pin();
+	if (Persona.IsValid())
 	{
-		PreviewAsset = PersonaPtr.Pin()->GetPreviewAnimationAsset();
-		if (!PreviewAsset)
+		// if persona isn't there, we probably don't have the preview mesh either, so no valid anim instance
+		if (CachedPreviewInstance && OnAddAnimationCurveDelegate.IsBound())
 		{
-			PreviewAsset = PersonaPtr.Pin()->GetAnimBlueprint();
+			CachedPreviewInstance->RemoveDelegate_AddCustomAnimationCurve(OnAddAnimationCurveDelegate);
 		}
+
+		Persona->UnregisterOnPreviewMeshChanged(this);
+		Persona->UnregisterOnAnimChanged(this);
+		Persona->UnregisterOnPostUndo(this);
+		Persona->UnregisterOnChangeCurves(this);
 	}
+}
 
-	return (PreviewAsset? FText::FromString(PreviewAsset->GetName()) : LOCTEXT("AnimCurveMeshNameLabel", "No Asset Present"));
+bool SAnimCurveViewer::IsCurveFilterEnabled() const
+{
+	return !bShowAllCurves;
+}
+
+ECheckBoxState SAnimCurveViewer::IsShowingAllCurves() const
+{
+	return bShowAllCurves ? ECheckBoxState::Checked : ECheckBoxState::Unchecked;
+}
+
+void SAnimCurveViewer::OnToggleShowingAllCurves(ECheckBoxState NewState)
+{
+	bShowAllCurves = (NewState == ECheckBoxState::Checked);
+
+	RefreshCurveList();
 }
 
 void SAnimCurveViewer::RefreshCachePreviewInstance()
@@ -464,7 +508,8 @@ void SAnimCurveViewer::RefreshCachePreviewInstance()
 	{
 		CachedPreviewInstance->RemoveDelegate_AddCustomAnimationCurve(OnAddAnimationCurveDelegate);
 	}
-	CachedPreviewInstance = Cast<UAnimInstance>(PersonaPtr.Pin()->GetPreviewMeshComponent()->GetAnimInstance());
+
+	CachedPreviewInstance = Cast<UAnimInstance>(PersonaPtr.Pin()->GetPreviewAnimInstance());
 
 	if (CachedPreviewInstance)
 	{
@@ -486,6 +531,11 @@ void SAnimCurveViewer::OnFilterTextChanged( const FText& SearchText )
 	RefreshCurveList();
 }
 
+void SAnimCurveViewer::OnCurvesChanged()
+{
+	RefreshCurveList();
+}
+
 void SAnimCurveViewer::OnFilterTextCommitted( const FText& SearchText, ETextCommit::Type CommitInfo )
 {
 	// Just do the same as if the user typed in the box
@@ -498,10 +548,8 @@ TSharedRef<ITableRow> SAnimCurveViewer::GenerateAnimCurveRow(TSharedPtr<FDisplay
 
 	return
 		SNew( SAnimCurveListRow, OwnerTable )
-		.Persona( PersonaPtr )
 		.Item( InInfo )
-		.AnimCurveViewer( this )
-		.AnimCurveListView( AnimCurveListView );
+		.AnimCurveViewerPtr( SharedThis(this) );
 }
 
 TSharedPtr<SWidget> SAnimCurveViewer::OnGetContextMenuContent() const
@@ -510,92 +558,243 @@ TSharedPtr<SWidget> SAnimCurveViewer::OnGetContextMenuContent() const
 	FMenuBuilder MenuBuilder( bShouldCloseWindowAfterMenuSelection, NULL);
 
 	MenuBuilder.BeginSection("AnimCurveAction", LOCTEXT( "MorphsAction", "Selected Item Actions" ) );
-// 	{
-// 		FUIAction Action = FUIAction( FExecuteAction::CreateSP( this, &SAnimCurveViewer::OnDeleteAnimCurves ), 
-// 									  FCanExecuteAction::CreateSP( this, &SAnimCurveViewer::CanPerformDelete ) );
-// 		const FText Label = LOCTEXT("DeleteAnimCurveButtonLabel", "Delete");
-// 		const FText ToolTip = LOCTEXT("DeleteAnimCurveButtonTooltip", "Deletes the selected animation curve.");
-// 		MenuBuilder.AddMenuEntry( Label, ToolTip, FSlateIcon(), Action);
-// 	}
+
+	{
+		FUIAction Action;
+		// Rename selection
+		{
+			Action.ExecuteAction = FExecuteAction::CreateSP(this, &SAnimCurveViewer::OnRenameClicked);
+			Action.CanExecuteAction = FCanExecuteAction::CreateSP(this, &SAnimCurveViewer::CanRename);
+			const FText Label = LOCTEXT("RenameSmartNameLabel", "Rename");
+			const FText ToolTipText = LOCTEXT("RenameSmartNameToolTip", "Rename the selected item");
+			MenuBuilder.AddMenuEntry(Label, ToolTipText, FSlateIcon(), Action);
+		}
+		// Delete a name
+		{
+			Action.ExecuteAction = FExecuteAction::CreateSP(this, &SAnimCurveViewer::OnDeleteNameClicked);
+			Action.CanExecuteAction = FCanExecuteAction::CreateSP(this, &SAnimCurveViewer::CanDelete);
+			const FText Label = LOCTEXT("DeleteSmartNameLabel", "Delete");
+			const FText ToolTipText = LOCTEXT("DeleteSmartNameToolTip", "Delete the selected item");
+			MenuBuilder.AddMenuEntry(Label, ToolTipText, FSlateIcon(), Action);
+		}
+		// Add a name
+		MenuBuilder.AddMenuSeparator();
+		{
+			Action.ExecuteAction = FExecuteAction::CreateSP(this, &SAnimCurveViewer::OnAddClicked);
+			Action.CanExecuteAction = nullptr;
+			const FText Label = LOCTEXT("AddSmartNameLabel", "Add...");
+			const FText ToolTipText = LOCTEXT("AddSmartNameToolTip", "Add an entry to the skeleton");
+			MenuBuilder.AddMenuEntry(Label, ToolTipText, FSlateIcon(), Action);
+		}
+	}
+
 	MenuBuilder.EndSection();
 
 	return MenuBuilder.MakeWidget();
 }
 
+void SAnimCurveViewer::OnRenameClicked()
+{
+	TArray< TSharedPtr<FDisplayedAnimCurveInfo> > SelectedItems = AnimCurveListView->GetSelectedItems();
+
+	SelectedItems[0]->EditableText->EnterEditingMode();
+}
+
+bool SAnimCurveViewer::CanRename()
+{
+	return AnimCurveListView->GetNumItemsSelected() == 1;
+}
+
+void SAnimCurveViewer::OnAddClicked()
+{
+	TSharedRef<STextEntryPopup> TextEntry =
+		SNew(STextEntryPopup)
+		.Label(LOCTEXT("NewSmartnameLabel", "New Name"))
+		.OnTextCommitted(this, &SAnimCurveViewer::CreateNewNameEntry);
+
+	FSlateApplication& SlateApp = FSlateApplication::Get();
+	SlateApp.PushMenu(
+		AsShared(),
+		FWidgetPath(),
+		TextEntry,
+		SlateApp.GetCursorPos(),
+		FPopupTransitionEffect::TypeInPopup
+		);
+}
+
+
+
 void SAnimCurveViewer::Tick(const FGeometry& AllottedGeometry, const double InCurrentTime, const float InDeltaTime)
 {
 	if (PersonaPtr.IsValid())
 	{
-		UDebugSkelMeshComponent* MeshComponent = PersonaPtr.Pin()->GetPreviewMeshComponent();
-		UAnimInstance* AnimInstance = MeshComponent->GetAnimInstance();
-
+		UAnimInstance* AnimInstance = PersonaPtr.Pin()->GetPreviewAnimInstance();
 		if (AnimInstance)
 		{
-			TMap<FName, float> AnimCurves;
-			AnimInstance->GetAnimationCurveList(CurrentCurveFlag, AnimCurves);
-			// @fixme: this is not enough information to refresh curve, but it is better than not displaying because it didnt' tick yet
-			if (AnimCurves.Num() != AnimCurveList.Num())
+			RefreshCurveList();
+		}
+	}
+}
+
+void SAnimCurveViewer::CreateNewNameEntry(const FText& CommittedText, ETextCommit::Type CommitType)
+{
+	FSlateApplication::Get().DismissAllMenus();
+	if (!CommittedText.IsEmpty() && CommitType == ETextCommit::OnEnter)
+	{
+		if (const FSmartNameMapping* NameMapping = GetAnimCurveMapping())
+		{
+			FName NewName = FName(*CommittedText.ToString());
+			FSmartName NewCurveName;
+			if (CurrentSkeleton->AddSmartNameAndModify(ContainerName, NewName, NewCurveName))
 			{
+				// Successfully added
 				RefreshCurveList();
 			}
 		}
 	}
 }
-void SAnimCurveViewer::CreateAnimCurveList( const FString& SearchText )
-{
-	AnimCurveList.Empty();
 
-	UDebugSkelMeshComponent* MeshComponent = PersonaPtr.Pin()->GetPreviewMeshComponent();
-	UAnimInstance* AnimInstance = MeshComponent->GetAnimInstance();
+const FSmartNameMapping* SAnimCurveViewer::GetAnimCurveMapping()
+{
+	const FSmartNameMapping* Mapping = nullptr;
+	if (CurrentSkeleton)
+	{
+		Mapping = CurrentSkeleton->GetSmartNameContainer(ContainerName);
+	}
+	return Mapping;
+}
 
-	if (AnimInstance)
+int32 FindIndexOfAnimCurveInfo(TArray<TSharedPtr<FDisplayedAnimCurveInfo>>& AnimCurveInfos, const FName& CurveName)
+{
+	for (int32 CurveIdx = 0; CurveIdx < AnimCurveInfos.Num(); CurveIdx++)
 	{
-		TMap<FName, float> AnimCurves;
-		AnimInstance->GetAnimationCurveList(CurrentCurveFlag, AnimCurves);
+		if (AnimCurveInfos[CurveIdx]->SmartName.DisplayName == CurveName)
+		{
+			return CurveIdx;
+		}
+	}
+
+	return INDEX_NONE;
+}
 
-		if (AnimCurves.Num() > 0)
+void SAnimCurveViewer::CreateAnimCurveList( const FString& SearchText )
+{
+	TSharedPtr<FPersona> Persona = PersonaPtr.Pin();
+	if (Persona.IsValid() && CurrentSkeleton)
+	{
+		const FSmartNameMapping* Mapping = GetAnimCurveMapping();
+		if (Mapping)
 		{
-			bool bDoFiltering = !SearchText.IsEmpty();
+			TArray<FSmartNameMapping::UID> UidList;
+			Mapping->FillUidArray(UidList);
+
+			// Get set of active curves
+			TMap<FName, float> ActiveCurves;
+			UAnimInstance* AnimInstance = Persona->GetPreviewAnimInstance();
+			if (!bShowAllCurves && AnimInstance)
+			{
+				AnimInstance->GetAnimationCurveList(CurrentCurveFlag, ActiveCurves);
+			}
 
-			for (auto Iter = AnimCurves.CreateConstIterator(); Iter; ++Iter)
+			// Iterate through all curves..
+			for (FSmartNameMapping::UID Uid : UidList)
 			{
-				FName CurveName = Iter.Key();
-				if (bDoFiltering && !CurveName.ToString().Contains(SearchText))
+				bool bAddToList = true;
+
+				FSmartName SmartName;
+				Mapping->FindSmartNameByUID(Uid, SmartName);
+
+				// See if we pass the search filter
+				if (!FilterText.IsEmpty())
+				{
+					if (!SmartName.DisplayName.ToString().Contains(*FilterText.ToString()))
+					{
+						bAddToList = false;
+					}
+				}
+
+				// If we passed that, see if we are filtering to only active
+				if (bAddToList && !bShowAllCurves)
 				{
-					continue; // Skip items that don't match our filter
+					bAddToList = ActiveCurves.Contains(SmartName.DisplayName);
 				}
 
-				TSharedRef<FDisplayedAnimCurveInfo> Info = FDisplayedAnimCurveInfo::Make(CurveName);
-				AnimCurveList.Add(Info);
+				// If we still want to add
+				if (bAddToList)
+				{
+					// If not already in list, add it
+					if (FindIndexOfAnimCurveInfo(AnimCurveList, SmartName.DisplayName) == INDEX_NONE)
+					{
+						AnimCurveList.Add(FDisplayedAnimCurveInfo::Make(CurrentSkeleton, ContainerName, SmartName));
+					}
+				}
+				// Don't want in list
+				else
+				{
+					// If already in list, remove it
+					int32 CurrentIndex = FindIndexOfAnimCurveInfo(AnimCurveList, SmartName.DisplayName);
+					if (CurrentIndex != INDEX_NONE)
+					{
+						AnimCurveList.RemoveAt(CurrentIndex);
+					}
+				}
 			}
-		}
-		else
-		{
-			// else print somewhere no curves?
+
+			// Sort final list
+			struct FSortSmartNamesAlphabetically
+			{
+				bool operator()(const TSharedPtr<FDisplayedAnimCurveInfo>& A, const TSharedPtr<FDisplayedAnimCurveInfo>& B) const
+				{
+					return (A.Get()->SmartName.DisplayName.Compare(B.Get()->SmartName.DisplayName) < 0);
+				}
+			};
+
+			AnimCurveList.Sort(FSortSmartNamesAlphabetically());
 		}
 	}
 
 	AnimCurveListView->RequestListRefresh();
 }
 
-void SAnimCurveViewer::CreateAnimCurveTypeList(TSharedRef<SVerticalBox> VerticalBox)
+void SAnimCurveViewer::CreateAnimCurveTypeList(TSharedRef<SHorizontalBox> HorizontalBox)
 {
+	// Add toggle button for 'all curves'
+	HorizontalBox->AddSlot()
+	.AutoWidth()
+	.Padding(3, 1)
+	[
+		SNew(SCheckBox)
+		.Style(FEditorStyle::Get(), "ToggleButtonCheckbox")
+		.ToolTipText(LOCTEXT("ShowAllCurvesTooltip", "Show all curves, or only active curves."))
+		.Type(ESlateCheckBoxType::ToggleButton)
+		.IsChecked(this, &SAnimCurveViewer::IsShowingAllCurves)
+		.OnCheckStateChanged(this, &SAnimCurveViewer::OnToggleShowingAllCurves)
+		.Padding(4)
+		.Content()
+		[
+			SNew(STextBlock)
+			.Text(LOCTEXT("ShowAllCurves", "All Curves"))
+		]
+	];
+
+	// Add check box for each curve type flag
 	TArray<int32> CurveFlagsToList;
-	CurveFlagsToList.Add(ACF_DrivesMorphTarget);
-	CurveFlagsToList.Add(ACF_TriggerEvent);
-	CurveFlagsToList.Add(ACF_DrivesMaterial);
-	CurveFlagsToList.Add(ACF_Metadata);
-	CurveFlagsToList.Add(ACF_DriveTrack);
-	CurveFlagsToList.Add(ACF_DrivesPose);
+	CurveFlagsToList.Add(ACF_DriveMorphTarget);
+	CurveFlagsToList.Add(ACF_DriveAttribute);
+	CurveFlagsToList.Add(ACF_DriveMaterial);
+	// @fixme: when we move curve types to skeleton fix this
+//	CurveFlagsToList.Add(ACF_Metadata);
 
 	for (int32 CurveFlagIndex = 0; CurveFlagIndex < CurveFlagsToList.Num(); ++CurveFlagIndex)
 	{
-		VerticalBox->AddSlot()
+		HorizontalBox->AddSlot()
+		.AutoWidth()
 		.Padding(3, 1)
 		[
 			SNew(SAnimCurveTypeList)
+			.IsEnabled(this, &SAnimCurveViewer::IsCurveFilterEnabled)
 			.CurveFlags(CurveFlagsToList[CurveFlagIndex])
-			.AnimCurveViewer(this)
+			.AnimCurveViewerPtr(SharedThis(this))
 		];
 	}
 }
@@ -604,29 +803,26 @@ void SAnimCurveViewer::AddAnimCurveOverride( FName& Name, float Weight)
 {
 	float& Value = OverrideCurves.FindOrAdd(Name);
 	Value = Weight;
+
+	UAnimSingleNodeInstance* SingleNodeInstance = Cast<UAnimSingleNodeInstance>(CachedPreviewInstance);
+	if (SingleNodeInstance)
+	{
+		SingleNodeInstance->SetPreviewCurveOverride(Name, Value, false);
+	}
 }
 
 void SAnimCurveViewer::RemoveAnimCurveOverride(FName& Name)
 {
 	OverrideCurves.Remove(Name);
-}
 
-SAnimCurveViewer::~SAnimCurveViewer()
-{
-	if ( PersonaPtr.IsValid() )
+	UAnimSingleNodeInstance* SingleNodeInstance = Cast<UAnimSingleNodeInstance>(CachedPreviewInstance);
+	if (SingleNodeInstance)
 	{
-		// if persona isn't there, we probably don't have the preview mesh either, so no valid anim instance
-		if (CachedPreviewInstance && OnAddAnimationCurveDelegate.IsBound())
-		{
-			CachedPreviewInstance->RemoveDelegate_AddCustomAnimationCurve(OnAddAnimationCurveDelegate);
-		}
-
-		PersonaPtr.Pin()->UnregisterOnPreviewMeshChanged(this);
-		PersonaPtr.Pin()->UnregisterOnAnimChanged(this);
-		PersonaPtr.Pin()->UnregisterOnPostUndo(this);
+		SingleNodeInstance->SetPreviewCurveOverride(Name, 0.f, true);
 	}
 }
 
+
 void SAnimCurveViewer::OnPostUndo()
 {
 	RefreshCurveList();
@@ -634,6 +830,7 @@ void SAnimCurveViewer::OnPostUndo()
 
 void SAnimCurveViewer::OnPreviewAssetChanged(class UAnimationAsset* NewAsset)
 {
+	OverrideCurves.Empty();
 	RefreshCurveList();
 }
 
@@ -650,5 +847,161 @@ void SAnimCurveViewer::RefreshCurveList()
 {
 	CreateAnimCurveList(FilterText.ToString());
 }
+
+void SAnimCurveViewer::OnNameCommitted(const FText& InNewName, ETextCommit::Type, TSharedPtr<FDisplayedAnimCurveInfo> Item)
+{
+	const FSmartNameMapping* Mapping = GetAnimCurveMapping();
+	if (Mapping)
+	{
+		FName NewName(*InNewName.ToString());
+		if (!Mapping->Exists(NewName))
+		{
+			FScopedTransaction Transaction(LOCTEXT("TransactionRename", "Rename Element"));
+			CurrentSkeleton->RenameSmartnameAndModify(ContainerName, Item->SmartName.UID, NewName);
+		}
+	}
+}
+
+void SAnimCurveViewer::OnDeleteNameClicked()
+{
+	TArray< TSharedPtr<FDisplayedAnimCurveInfo> > SelectedItems = AnimCurveListView->GetSelectedItems();
+	TArray<FSmartNameMapping::UID> SelectedUids;
+
+	for (TSharedPtr<FDisplayedAnimCurveInfo> Item : SelectedItems)
+	{
+		SelectedUids.Add(Item->SmartName.UID);
+	}
+
+	FAssetRegistryModule& AssetModule = FModuleManager::LoadModuleChecked<FAssetRegistryModule>(TEXT("AssetRegistry"));
+	TArray<FAssetData> AnimationAssets;
+	AssetModule.Get().GetAssetsByClass(UAnimSequenceBase::StaticClass()->GetFName(), AnimationAssets, true);
+	FAssetData SkeletonData(CurrentSkeleton);
+	const FString CurrentSkeletonName = SkeletonData.GetExportTextName();
+
+	GWarn->BeginSlowTask(LOCTEXT("CollectAnimationsTaskDesc", "Collecting assets..."), true);
+
+	for (int32 Idx = AnimationAssets.Num() - 1; Idx >= 0; --Idx)
+	{
+		FAssetData& Data = AnimationAssets[Idx];
+		bool bRemove = true;
+
+		const FString SkeletonDataTag = Data.GetTagValueRef<FString>("Skeleton");
+		if (!SkeletonDataTag.IsEmpty() && SkeletonDataTag == CurrentSkeletonName)
+		{
+			FString CurveData;
+			if (!Data.GetTagValue(USkeleton::CurveTag, CurveData))
+			{
+				// This asset is old; it hasn't been loaded since before smartnames were added for curves.
+				// unfortunately the only way to delete curves safely is to load old assets to see if they're
+				// using the selected name. We only load what we have to here.
+				UObject* Asset = Data.GetAsset();
+				check(Asset);
+				TArray<UObject::FAssetRegistryTag> Tags;
+				Asset->GetAssetRegistryTags(Tags);
+
+				UObject::FAssetRegistryTag* CurveTag = Tags.FindByPredicate([](const UObject::FAssetRegistryTag& InTag)
+				{
+					return InTag.Name == USkeleton::CurveTag;
+				});
+				CurveData = CurveTag->Value;
+			}
+
+			TArray<FString> ParsedStringUids;
+			CurveData.ParseIntoArray(ParsedStringUids, *USkeleton::CurveTagDelimiter, true);
+
+			for (const FString& UidString : ParsedStringUids)
+			{
+				FSmartNameMapping::UID Uid = (FSmartNameMapping::UID)TCString<TCHAR>::Strtoui64(*UidString, NULL, 10);
+				if (SelectedUids.Contains(Uid))
+				{
+					bRemove = false;
+					break;
+				}
+			}
+		}
+
+		if (bRemove)
+		{
+			AnimationAssets.RemoveAt(Idx);
+		}
+	}
+
+	GWarn->EndSlowTask();
+
+	// AnimationAssets now only contains assets that are using the selected curve(s)
+	if (AnimationAssets.Num() > 0)
+	{
+		FString ConfirmMessage = LOCTEXT("DeleteCurveMessage", "The following assets use the selected curve(s), deleting will remove the curves from these assets. Continue?\n\n").ToString();
+
+		for (FAssetData& Data : AnimationAssets)
+		{
+			ConfirmMessage += Data.AssetName.ToString() + " (" + Data.AssetClass.ToString() + ")\n";
+		}
+
+		FText Title = LOCTEXT("DeleteCurveDialogTitle", "Confirm Deletion");
+		FText Message = FText::FromString(ConfirmMessage);
+		if (FMessageDialog::Open(EAppMsgType::YesNo, Message, &Title) == EAppReturnType::Yes)
+		{
+			// Proceed to delete the curves
+			GWarn->BeginSlowTask(FText::Format(LOCTEXT("DeleteCurvesTaskDesc", "Deleting curve from skeleton {0}"), FText::FromString(CurrentSkeleton->GetName())), true);
+			FScopedTransaction Transaction(LOCTEXT("DeleteCurvesTransactionName", "Delete skeleton curve"));
+
+			// Remove curves from animation assets
+			for (FAssetData& Data : AnimationAssets)
+			{
+				UAnimSequenceBase* Sequence = Cast<UAnimSequenceBase>(Data.GetAsset());
+				USkeleton* MySkeleton = Sequence->GetSkeleton();
+				Sequence->Modify(true);
+				for (FSmartNameMapping::UID Uid : SelectedUids)
+				{
+					FSmartName CurveToDelete;
+					if (MySkeleton->GetSmartNameByUID(ContainerName, Uid, CurveToDelete))
+					{
+						Sequence->RawCurveData.DeleteCurveData(CurveToDelete);
+						Sequence->MarkRawDataAsModified();
+					}
+				}
+			}
+			GWarn->EndSlowTask();
+
+			GWarn->BeginSlowTask(LOCTEXT("RebuildingAnimations", "Rebaking/compressing modified animations"), true);
+
+			// Rebake/compress the animations
+			for (TObjectIterator<UAnimSequence> It; It; ++It)
+			{
+				UAnimSequence* Seq = *It;
+
+				GWarn->StatusUpdate(1, 2, FText::Format(LOCTEXT("RebuildingAnimationsStatus", "Rebuilding {0}"), FText::FromString(Seq->GetName())));
+				Seq->RequestSyncAnimRecompression();
+			}
+			GWarn->EndSlowTask();
+		}
+	}
+
+	if (SelectedUids.Num() > 0)
+	{
+		// Remove names from skeleton
+		CurrentSkeleton->RemoveSmartnamesAndModify(ContainerName, SelectedUids);
+	}
+
+	RefreshCurveList();
+
+	TSharedPtr<FPersona> Persona = PersonaPtr.Pin();
+	if (Persona.IsValid())
+	{
+		if (Persona->GetCurrentMode() == FPersonaModes::AnimationEditMode)
+		{
+			// If we're in animation mode then we need to reinitialise the mode to
+			// refresh the sequencer which will now have a widget it shouldn't have
+			Persona->ReinitMode();
+		}
+	}
+}
+
+bool SAnimCurveViewer::CanDelete()
+{
+	return AnimCurveListView->GetNumItemsSelected() > 0;
+}
+
 #undef LOCTEXT_NAMESPACE
 
diff --git a/Engine/Source/Editor/Persona/Private/SAnimCurveViewer.h b/Engine/Source/Editor/Persona/Private/SAnimCurveViewer.h
--- a/Engine/Source/Editor/Persona/Private/SAnimCurveViewer.h
+++ b/Engine/Source/Editor/Persona/Private/SAnimCurveViewer.h
@@ -4,28 +4,35 @@
 #include "Animation/AnimInstance.h"
 #include "Persona.h"
 
+class SAnimCurveViewer;
+
 //////////////////////////////////////////////////////////////////////////
 // FDisplayedAnimCurveInfo
 
 class FDisplayedAnimCurveInfo
 {
 public:
-	FName Name;
+	FSmartName SmartName;
 	float Weight;
 	bool bAutoFillData;
+	const class USkeleton* Skeleton;		// The skeleton we're associated with
+	TSharedPtr<SInlineEditableTextBlock> EditableText;	// The editable text box in the list, used to focus from the context menu
+	FName ContainerName;	// The container in the skeleton this name resides in
 
 	/** Static function for creating a new item, but ensures that you can only have a TSharedRef to one */
-	static TSharedRef<FDisplayedAnimCurveInfo> Make(const FName& Source)
+	static TSharedRef<FDisplayedAnimCurveInfo> Make(const class USkeleton* InSkeleton, const FName& InContainerName, const FSmartName& InSmartName)
 	{
-		return MakeShareable(new FDisplayedAnimCurveInfo(Source));
+		return MakeShareable(new FDisplayedAnimCurveInfo(InSkeleton, InContainerName, InSmartName));
 	}
 
 protected:
 	/** Hidden constructor, always use Make above */
-	FDisplayedAnimCurveInfo(const FName& InSource)
-		: Name( InSource )
+	FDisplayedAnimCurveInfo(const class USkeleton* InSkeleton, const FName& InContainerName, const FSmartName& InSmartName)
+		: SmartName(InSmartName)
 		, Weight( 0 )
 		, bAutoFillData(true)
+		, Skeleton(InSkeleton)
+		, ContainerName(InContainerName)
 	{}
 
 	/** Hidden constructor, always use Make above */
@@ -34,6 +41,66 @@ protected:
 
 typedef SListView< TSharedPtr<FDisplayedAnimCurveInfo> > SAnimCurveListType;
 
+//////////////////////////////////////////////////////////////////////////
+// SAnimCurveListRow
+
+class SAnimCurveListRow : public SMultiColumnTableRow< TSharedPtr<FDisplayedAnimCurveInfo> >
+{
+public:
+
+	SLATE_BEGIN_ARGS(SAnimCurveListRow) {}
+
+		/** The item for this row **/
+		SLATE_ARGUMENT(TSharedPtr<FDisplayedAnimCurveInfo>, Item)
+
+		/* The SAnimCurveViewer that we push the morph target weights into */
+		SLATE_ARGUMENT(class TWeakPtr<SAnimCurveViewer>, AnimCurveViewerPtr)
+
+	SLATE_END_ARGS()
+
+	void Construct(const FArguments& InArgs, const TSharedRef<STableViewBase>& OwnerTableView);
+
+	/** Overridden from SMultiColumnTableRow.  Generates a widget for this column of the tree row. */
+	virtual TSharedRef<SWidget> GenerateWidgetForColumn(const FName& ColumnName) override;
+
+private:
+
+	/**
+	* Called when the user changes the value of the SSpinBox
+	*
+	* @param NewWeight - The new number the SSpinBox is set to
+	*
+	*/
+	void OnAnimCurveWeightChanged(float NewWeight);
+	/**
+	* Called when the user types the value and enters
+	*
+	* @param NewWeight - The new number the SSpinBox is set to
+	*
+	*/
+	void OnAnimCurveWeightValueCommitted(float NewWeight, ETextCommit::Type CommitType);
+
+	/** Auto fill check call back functions */
+	void OnAnimCurveAutoFillChecked(ECheckBoxState InState);
+	ECheckBoxState IsAnimCurveAutoFillChangedChecked() const;
+
+	/** Returns the weight of this curve */
+	float GetWeight() const;
+	/** Returns name of this curve */
+	FText GetItemName() const;
+	/** Return color for text of item */
+	FSlateColor GetItemTextColor() const;
+
+	/** Get current active weight. Returns false if not currently active */
+	bool GetActiveWeight(float& OutWeight) const;
+
+	/* The SAnimCurveViewer that we push the morph target weights into */
+	TWeakPtr<SAnimCurveViewer> AnimCurveViewerPtr;
+
+	/** The name and weight of the morph target */
+	TSharedPtr<FDisplayedAnimCurveInfo>	Item;
+};
+
 //////////////////////////////////////////////////////////////////////////
 // SAnimCurveViewer
 
@@ -44,7 +111,7 @@ public:
 		: _Persona()
 	{}
 		
-		/* The Persona that owns this table */
+		/** The Persona that owns this table */
 		SLATE_ARGUMENT( TWeakPtr< FPersona >, Persona )
 
 	SLATE_END_ARGS()
@@ -78,6 +145,9 @@ public:
 	*/
 	void OnPreviewAssetChanged(class UAnimationAsset* NewPreviewAsset);
 
+	/** Is registered with Persona to handle when curves change. */
+	void OnCurvesChanged();
+
 	/**
 	* Filters the SListView when the user changes the search text box (NameFilterBox)
 	*
@@ -86,6 +156,9 @@ public:
 	*/
 	void OnFilterTextChanged( const FText& SearchText );
 
+
+
+
 	/**
 	* Filters the SListView when the user hits enter or clears the search box
 	* Simply calls OnFilterTextChanged
@@ -128,10 +201,6 @@ public:
 	*/
 	FText& GetFilterText() { return FilterText; }
 
-	/**
-	 * Get Title text
-	 */
-	FText GetTitleText() const;
 	/**
 	 * Refreshes the morph target list after an undo
 	 */
@@ -141,6 +210,9 @@ public:
 
 	void RefreshCurveList();
 
+	// When a name is committed after being edited in the list
+	virtual void OnNameCommitted(const FText& NewName, ETextCommit::Type CommitType, TSharedPtr<FDisplayedAnimCurveInfo> Item);
+
 private:
 
 	/** Handler for context menus */
@@ -153,29 +225,58 @@ private:
 	*
 	*/
 	void CreateAnimCurveList( const FString& SearchText = FString() );
-	void CreateAnimCurveTypeList(TSharedRef<SVerticalBox> VerticalBox);
+	void CreateAnimCurveTypeList(TSharedRef<SHorizontalBox> HorizontalBox);
 
 	void ApplyCustomCurveOverride(UAnimInstance* AnimInstance) const;
 	void RefreshCachePreviewInstance();
 
+	void OnDeleteNameClicked();
+	bool CanDelete();
+
+	void OnRenameClicked();
+	bool CanRename();
+
+	void OnAddClicked();
+
+	ECheckBoxState IsShowingAllCurves() const;
+	void OnToggleShowingAllCurves(ECheckBoxState NewState);
+
+	bool IsCurveFilterEnabled() const;
+
+	// Adds a new smartname entry to the skeleton in the container we are managing
+	void CreateNewNameEntry(const FText& CommittedText, ETextCommit::Type CommitType);
+
+	/** Get the SmartNameMapping for anim curves */
+	const struct FSmartNameMapping* GetAnimCurveMapping();
+
+
 	/** Pointer back to the Persona that owns us */
 	TWeakPtr<FPersona> PersonaPtr;
 
 	/** Box to filter to a specific morph target name */
 	TSharedPtr<SSearchBox>	NameFilterBox;
 
-	/** Widget used to display the list of animation curve */
-	TSharedPtr<SAnimCurveListType> AnimCurveListView;
-
 	/** A list of animation curve. Used by the AnimCurveListView. */
 	TArray< TSharedPtr<FDisplayedAnimCurveInfo> > AnimCurveList;
 
 	/** The skeletal mesh that we grab the animation curve from */
 	UAnimInstance* CachedPreviewInstance;
 
+	/** The current skeleton being edited by out Persona instance */
+	class USkeleton* CurrentSkeleton;									
+
+	/** Widget used to display the list of animation curve */
+	TSharedPtr<SAnimCurveListType> AnimCurveListView;
+
+	/** Name of the skeleton smart name container to display in the list */
+	FName ContainerName;
+
 	/** Current text typed into NameFilterBox */
 	FText FilterText;
 
+	/** Are we showing all curves, or only active ones. */
+	bool bShowAllCurves;
+
 	int32 CurrentCurveFlag;
 
 	TMap<FName, float> OverrideCurves;
diff --git a/Engine/Source/Editor/Persona/Private/SAnimEditorBase.cpp b/Engine/Source/Editor/Persona/Private/SAnimEditorBase.cpp
--- a/Engine/Source/Editor/Persona/Private/SAnimEditorBase.cpp
+++ b/Engine/Source/Editor/Persona/Private/SAnimEditorBase.cpp
@@ -28,22 +28,11 @@ void SAnimEditorBase::Construct(const FArguments& InArgs)
 
 	SetInputViewRange(0, GetSequenceLength());
 
-	UAnimSequenceBase* SeqBase = Cast<UAnimSequenceBase>(GetEditorObject());
-
-	TSharedPtr<SHorizontalBox> HorizontalBox = SNew(SHorizontalBox);
-	if (SeqBase)
-	{
-		HorizontalBox->AddSlot()
-			.FillWidth(1)
-			[
-				ConstructAnimScrubPanel()
-			];
-	}
-
+	TSharedPtr<SVerticalBox> AnimVerticalBox;
 
 	this->ChildSlot
 	[
-		SNew(SVerticalBox)
+		SAssignNew(AnimVerticalBox, SVerticalBox)
 		+SVerticalBox::Slot()
 		.AutoHeight()
 		[
@@ -84,8 +73,12 @@ void SAnimEditorBase::Construct(const FArguments& InArgs)
 				]
 			]
 		]
+	];
 
-		+SVerticalBox::Slot()
+	// If we want to create anim info bar, display that now
+	if (InArgs._DisplayAnimInfoBar)
+	{
+		AnimVerticalBox->AddSlot()
 		.AutoHeight()
 		.VAlign(VAlign_Center)
 		[
@@ -174,15 +167,27 @@ void SAnimEditorBase::Construct(const FArguments& InArgs)
 					]
 				]
 			]
-		]
-	
-		+SVerticalBox::Slot()
-		.AutoHeight() 
+		];
+
+	}
+
+	// If we are an anim sequence, add scrub panel as well
+	UAnimSequenceBase* SeqBase = Cast<UAnimSequenceBase>(GetEditorObject());
+	if (SeqBase)
+	{
+		AnimVerticalBox->AddSlot()
+		.AutoHeight()
 		.VAlign(VAlign_Bottom)
 		[
-			HorizontalBox.ToSharedRef()
-		]
-	];
+			SNew(SHorizontalBox)
+			+SHorizontalBox::Slot()
+			.FillWidth(1)
+			[
+				ConstructAnimScrubPanel()
+			]
+		];
+	}
+
 }
 
 TSharedRef<class SAnimationScrubPanel> SAnimEditorBase::ConstructAnimScrubPanel()
diff --git a/Engine/Source/Editor/Persona/Private/SAnimEditorBase.h b/Engine/Source/Editor/Persona/Private/SAnimEditorBase.h
--- a/Engine/Source/Editor/Persona/Private/SAnimEditorBase.h
+++ b/Engine/Source/Editor/Persona/Private/SAnimEditorBase.h
@@ -29,9 +29,12 @@ class SAnimEditorBase : public SCompoundWidget, public FGCObject
 public:
 	SLATE_BEGIN_ARGS( SAnimEditorBase )
 		: _Persona()
+		, _DisplayAnimInfoBar(true)
 		{}
 
 		SLATE_ARGUMENT( TSharedPtr<FPersona>, Persona )
+
+		SLATE_ARGUMENT( bool, DisplayAnimInfoBar )
 	SLATE_END_ARGS()
 
 	void Construct(const FArguments& InArgs);
diff --git a/Engine/Source/Editor/Persona/Private/SMorphTargetViewer.cpp b/Engine/Source/Editor/Persona/Private/SMorphTargetViewer.cpp
--- a/Engine/Source/Editor/Persona/Private/SMorphTargetViewer.cpp
+++ b/Engine/Source/Editor/Persona/Private/SMorphTargetViewer.cpp
@@ -276,7 +276,7 @@ float SMorphTargetListRow::GetWeight() const
 		{
 			// make sure if they have value that's not same as saved value
 			TMap<FName, float> MorphCurves;
-			AnimInstance->GetAnimationCurveList(ACF_DrivesMorphTarget, MorphCurves);
+			AnimInstance->GetAnimationCurveList(ACF_DriveMorphTarget, MorphCurves);
 			const float* CurrentValue = MorphCurves.Find(Item->Name);
 			if (CurrentValue && *CurrentValue != 0.f)
 			{
