Copying //UE4/Dev-Framework to //UE4/Dev-Main (Source: //UE4/Dev-Framework @ 3038004)

diff --git a/Engine/Source/Editor/Persona/Private/SPoseEditor.cpp b/Engine/Source/Editor/Persona/Private/SPoseEditor.cpp
--- a/Engine/Source/Editor/Persona/Private/SPoseEditor.cpp
+++ b/Engine/Source/Editor/Persona/Private/SPoseEditor.cpp
@@ -9,8 +9,18 @@
 #include "SSearchBox.h"
 #include "Animation/AnimSingleNodeInstance.h"
 
+#include "SInlineEditableTextBlock.h"
+
+
 #define LOCTEXT_NAMESPACE "AnimPoseEditor"
 
+static const FName ColumnId_PoseNameLabel("Pose Name");
+static const FName ColumnID_PoseWeightLabel("Weight");
+static const FName ColumnId_CurveNameLabel("Curve Name");
+static const FName ColumnID_CurveValueLabel("Curve Value");
+
+const float MaxPoseWeight = 1.f;
+
 //////////////////////////////////////////////////////////////////////////
 // SPoseEditor
 
@@ -21,6 +31,7 @@ void SPoseEditor::Construct(const FArguments& InArgs)
 
 	SAnimEditorBase::Construct( SAnimEditorBase::FArguments()
 		.Persona(InArgs._Persona)
+		.DisplayAnimInfoBar(false)
 		);
 
  	EditorPanels->AddSlot()
@@ -37,90 +48,10 @@ SPoseEditor::~SPoseEditor()
 {
 }
 
-////////////////////////////////////
-// Copyright 1998-2016 Epic Games, Inc. All Rights Reserved.
-static const FName ColumnId_PoseNameLabel("Pose Name");
-static const FName ColumnID_PoseWeightLabel("Weight");
-static const FName ColumnId_CurveNameLabel("Curve Name");
-
-const float MaxPoseWeight = 1.f;
 
 //////////////////////////////////////////////////////////////////////////
 // SPoseListRow
 
-typedef TSharedPtr< FDisplayedPoseInfo > FDisplayedPoseInfoPtr;
-
-class SPoseListRow
-	: public SMultiColumnTableRow< FDisplayedPoseInfoPtr >
-{
-public:
-
-	SLATE_BEGIN_ARGS(SPoseListRow) {}
-
-	/** The item for this row **/
-	SLATE_ARGUMENT(FDisplayedPoseInfoPtr, Item)
-
-	/* The SPoseViewer that we push the morph target weights into */
-	SLATE_ARGUMENT(class SPoseViewer*, PoseViewer)
-
-	/* Widget used to display the list of animation curve */
-	SLATE_ARGUMENT(TSharedPtr<SPoseListType>, PoseListView)
-
-	/* Persona used to update the viewport when a weight slider is dragged */
-	SLATE_ARGUMENT(TWeakPtr<FPersona>, Persona)
-
-	SLATE_END_ARGS()
-
-	void Construct(const FArguments& InArgs, const TSharedRef<STableViewBase>& OwnerTableView);
-
-	/** Overridden from SMultiColumnTableRow.  Generates a widget for this column of the tree row. */
-	virtual TSharedRef<SWidget> GenerateWidgetForColumn(const FName& ColumnName) override;
-
-private:
-
-	/**
-	* Called when the user changes the value of the SSpinBox
-	*
-	* @param NewWeight - The new number the SSpinBox is set to
-	*
-	*/
-	void OnPoseWeightChanged(float NewWeight);
-	/**
-	* Called when the user types the value and enters
-	*
-	* @param NewWeight - The new number the SSpinBox is set to
-	*
-	*/
-	void OnPoseWeightValueCommitted(float NewWeight, ETextCommit::Type CommitType);
-
-	/** Delegate to get labels root text from settings */
-	FText GetName() const;
-
-	/** Delegate to commit labels root text to settings */
-	void OnNameCommitted(const FText& InText, ETextCommit::Type InCommitType) const;
-
-	/**
-	* Returns the weight of this morph target
-	*
-	* @return SearchText - The new number the SSpinBox is set to
-	*
-	*/
-	float GetWeight() const;
-
-	bool CanChangeWeight() const;
-	/* The SPoseViewer that we push the morph target weights into */
-	SPoseViewer* PoseViewer;
-
-	/** Widget used to display the list of animation curve */
-	TSharedPtr<SPoseListType> PoseListView;
-
-	/** The name and weight of the morph target */
-	FDisplayedPoseInfoPtr	Item;
-
-	/** Pointer back to the Persona that owns us */
-	TWeakPtr<FPersona> PersonaPtr;
-};
-
 FText SPoseListRow::GetName() const
 {
 	return (FText::FromName(Item->Name));
@@ -135,43 +66,80 @@ void SPoseListRow::OnNameCommitted(const FText& InText, ETextCommit::Type InComm
 		FName NewName = FName(*InText.ToString());
 		FName OldName = Item->Name;
 
-		if (PoseViewer->ModifyName(OldName, NewName))
+		if (PoseViewerPtr.IsValid() && PoseViewerPtr.Pin()->ModifyName(OldName, NewName))
 		{
 			Item->Name = NewName;
 		}
 	}
 }
 
+bool SPoseListRow::OnVerifyNameChanged(const FText& InText, FText& OutErrorMessage)
+{
+	bool bVerifyName = false;
+
+	FName NewName = FName(*InText.ToString());
+
+	if (NewName == NAME_None)
+	{
+		OutErrorMessage = LOCTEXT("EmptyPoseName", "Poses must have a name!");
+	}
+
+	if (PoseViewerPtr.IsValid())
+	{
+		UPoseAsset* PoseAsset = PoseViewerPtr.Pin()->PoseAssetPtr.Get();
+		if (PoseAsset != nullptr)
+		{
+			if (PoseAsset->ContainsPose(NewName))
+			{
+				OutErrorMessage = LOCTEXT("NameAlreadyUsedByTheSameAsset", "The name is used by another pose within the same asset. Please choose another name.");
+			}
+			else
+			{
+				bVerifyName = true;
+			}
+		}
+	}
+
+	return bVerifyName;
+}
+
+
 void SPoseListRow::Construct(const FArguments& InArgs, const TSharedRef<STableViewBase>& InOwnerTableView)
 {
 	Item = InArgs._Item;
-	PoseViewer = InArgs._PoseViewer;
-	PoseListView = InArgs._PoseListView;
-	PersonaPtr = InArgs._Persona;
+	PoseViewerPtr = InArgs._PoseViewer;
+	FilterText = InArgs._FilterText;
 
 	check(Item.IsValid());
 
-	SMultiColumnTableRow< FDisplayedPoseInfoPtr >::Construct(FSuperRowType::FArguments(), InOwnerTableView);
+	SMultiColumnTableRow< TSharedPtr<FDisplayedPoseInfo> >::Construct(FSuperRowType::FArguments(), InOwnerTableView);
 }
 
 TSharedRef< SWidget > SPoseListRow::GenerateWidgetForColumn(const FName& ColumnName)
 {
 	if (ColumnName == ColumnId_PoseNameLabel)
 	{
-		return
-			SNew(SVerticalBox)
+		TSharedPtr< SInlineEditableTextBlock > InlineWidget;
 
-			+ SVerticalBox::Slot()
-			.AutoHeight()
+		TSharedRef<SWidget> NameWidget =
+			SNew(SHorizontalBox)
+
+			+ SHorizontalBox::Slot()
+			.AutoWidth()
 			.Padding(5.0f)
 			.VAlign(VAlign_Center)
 			[
-				SNew(SEditableTextBox)
+				SAssignNew(InlineWidget, SInlineEditableTextBlock)
 				.Text(this, &SPoseListRow::GetName)
+				.HighlightText(FilterText)
 				.ToolTipText(LOCTEXT("PoseName_ToolTip", "Modify Pose Name - Make sure this name is unique among all curves per skeleton."))
+				.OnVerifyTextChanged(this, &SPoseListRow::OnVerifyNameChanged)
 				.OnTextCommitted(this, &SPoseListRow::OnNameCommitted)
-				.OnTextChanged(this, &SPoseListRow::OnNameCommitted, ETextCommit::Default)
 			];
+
+		Item->OnRenameRequested.BindSP(InlineWidget.Get(), &SInlineEditableTextBlock::EnterEditingMode);
+
+		return NameWidget;
 	}
 	else 
 	{
@@ -201,11 +169,15 @@ void SPoseListRow::OnPoseWeightChanged(float NewWeight)
 {
 	Item->Weight = NewWeight;
 
-	PoseViewer->AddPoseOverride(Item->Name, Item->Weight);
-
-	if (PersonaPtr.IsValid())
+	if (PoseViewerPtr.IsValid())
 	{
-		PersonaPtr.Pin()->RefreshViewport();
+		TSharedPtr<SPoseViewer> PoseViewer = PoseViewerPtr.Pin();
+		PoseViewer->AddCurveOverride(Item->Name, Item->Weight);
+
+		if (PoseViewer->PersonaPtr.IsValid())
+		{
+			PoseViewer->PersonaPtr.Pin()->RefreshViewport();
+		}
 	}
 }
 
@@ -214,28 +186,7 @@ void SPoseListRow::OnPoseWeightValueCommitted(float NewWeight, ETextCommit::Type
 	if (CommitType == ETextCommit::OnEnter || CommitType == ETextCommit::OnUserMovedFocus)
 	{
 		float NewValidWeight = FMath::Clamp(NewWeight, -MaxPoseWeight, MaxPoseWeight);
-		Item->Weight = NewValidWeight;
-
-		PoseViewer->AddPoseOverride(Item->Name, Item->Weight);
-
-		TArray< TSharedPtr< FDisplayedPoseInfo > > SelectedRows = PoseListView->GetSelectedItems();
-
-		// ...then any selected rows need changing by the same delta
-		for (auto ItemIt = SelectedRows.CreateIterator(); ItemIt; ++ItemIt)
-		{
-			TSharedPtr< FDisplayedPoseInfo > RowItem = (*ItemIt);
-
-			if (RowItem != Item) // Don't do "this" row again if it's selected
-			{
-				RowItem->Weight = NewValidWeight;
-				PoseViewer->AddPoseOverride(RowItem->Name, RowItem->Weight);
-			}
-		}
-
-		if (PersonaPtr.IsValid())
-		{
-			PersonaPtr.Pin()->RefreshViewport();
-		}
+		OnPoseWeightChanged(NewValidWeight);
 	}
 }
 
@@ -246,61 +197,73 @@ float SPoseListRow::GetWeight() const
 
 bool SPoseListRow::CanChangeWeight() const
 {
-	return PoseViewer->IsBasePose(Item->Name) == false;
+	if (PoseViewerPtr.IsValid())
+	{
+		return PoseViewerPtr.Pin()->IsBasePose(Item->Name) == false;
+	}
+	else
+	{
+		return false;
+	}
 }
 
 //////////////////////////////////////////////////////////////////////////
 // SCurveListRow
 
-typedef TSharedPtr< FDisplayedCurveInfo > FDisplayedCurveInfoPtr;
-
-class SCurveListRow
-	: public SMultiColumnTableRow< FDisplayedCurveInfoPtr >
+FText SCurveListRow::GetName() const
 {
-public:
-
-	SLATE_BEGIN_ARGS(SCurveListRow) {}
-
-	/** The item for this row **/
-	SLATE_ARGUMENT(FDisplayedCurveInfoPtr, Item)
-
-	/* Widget used to display the list of animation curve */
-	SLATE_ARGUMENT(TSharedPtr<SCurveListType>, CurveListView)
-
-	SLATE_END_ARGS()
-
-	void Construct(const FArguments& InArgs, const TSharedRef<STableViewBase>& OwnerTableView);
+	return (FText::FromName(Item->Name));
+}
 
-	/** Overridden from SMultiColumnTableRow.  Generates a widget for this column of the tree row. */
-	virtual TSharedRef<SWidget> GenerateWidgetForColumn(const FName& ColumnName) override;
+FText SCurveListRow::GetValue() const
+{
+	FText ValueText;
+	if (PoseViewerPtr.IsValid())
+	{
+		TSharedPtr<SPoseViewer> PoseViewer = PoseViewerPtr.Pin();
 
-private:
+		// Get pose asset
+		UPoseAsset* PoseAsset = PoseViewer->PoseAssetPtr.Get();
+		if (PoseAsset != nullptr)
+		{
+			// Get selected row (only show values if only one selected)
+			TArray< TSharedPtr<FDisplayedPoseInfo> > SelectedRows = PoseViewer->PoseListView->GetSelectedItems();
+			if (SelectedRows.Num() == 1)
+			{
+				TSharedPtr<FDisplayedPoseInfo> PoseInfo = SelectedRows[0];
 
-	/** Delegate to get labels root text from settings */
-	FText GetName() const;
+				// Get pose index that we have selected
+				int32 PoseIndex = PoseAsset->GetPoseIndexByName(PoseInfo->Name);
+				int32 CurveIndex = PoseAsset->GetCurveIndexByName(Item->Name);
 
-	/** The name and weight of the morph target */
-	FDisplayedCurveInfoPtr	Item;
-};
+				float CurveValue = 0.f;
+				bool bSuccess = PoseAsset->GetCurveValue(PoseIndex, CurveIndex, CurveValue);
 
-FText SCurveListRow::GetName() const
-{
-	return (FText::FromName(Item->Name));
+				if (bSuccess)
+				{
+					ValueText = FText::FromString(FString::Printf(TEXT("%f"), CurveValue));
+				}
+			}
+		}
+	}
+	return ValueText;
 }
 
+
 void SCurveListRow::Construct(const FArguments& InArgs, const TSharedRef<STableViewBase>& InOwnerTableView)
 {
 	Item = InArgs._Item;
+	PoseViewerPtr = InArgs._PoseViewer;
 
 	check(Item.IsValid());
 
-	SMultiColumnTableRow< FDisplayedCurveInfoPtr >::Construct(FSuperRowType::FArguments(), InOwnerTableView);
+	SMultiColumnTableRow< TSharedPtr<FDisplayedCurveInfo> >::Construct(FSuperRowType::FArguments(), InOwnerTableView);
 }
 
 TSharedRef< SWidget > SCurveListRow::GenerateWidgetForColumn(const FName& ColumnName)
 {
 	// for now we have one colume
-//	if (ColumnName == ColumnId_CurveNameLabel)
+	if (ColumnName == ColumnId_CurveNameLabel)
 	{
 		return
 			SNew(SVerticalBox)
@@ -314,6 +277,20 @@ TSharedRef< SWidget > SCurveListRow::GenerateWidgetForColumn(const FName& Column
 				.Text(this, &SCurveListRow::GetName)
 			];
 	}
+	else
+	{
+		return
+			SNew(SVerticalBox)
+
+			+ SVerticalBox::Slot()
+			.AutoHeight()
+			.Padding(5.0f)
+			.VAlign(VAlign_Center)
+			[
+				SNew(STextBlock)
+				.Text(this, &SCurveListRow::GetValue)
+			];
+	}
 }
 
 //////////////////////////////////////////////////////////////////////////
@@ -329,52 +306,52 @@ void SPoseViewer::Construct(const FArguments& InArgs)
 	if (PersonaPtr.IsValid())
 	{
 		PersonaPtr.Pin()->RegisterOnPreviewMeshChanged(FPersona::FOnPreviewMeshChanged::CreateSP(this, &SPoseViewer::OnPreviewMeshChanged));
-		PersonaPtr.Pin()->RegisterOnPostUndo(FPersona::FOnPostUndo::CreateSP(this, &SPoseViewer::OnPostUndo));
+		PersonaPtr.Pin()->RegisterOnPostUndo(FPersona::FOnPostUndo::CreateSP(this, &SPoseViewer::RefreshList));
+	}
+
+	if (PoseAssetPtr.IsValid())
+	{
+		PoseAssetPtr->RegisterOnAssetModifiedNotifier(UPoseAsset::FAssetModifiedNotifier::CreateSP(this, &SPoseViewer::RefreshList));
 	}
 
 	RefreshCachePreviewInstance();
 
 	ChildSlot
 	[
-		SNew(SVerticalBox)
+		SNew(SHorizontalBox)
 
-		+ SVerticalBox::Slot()
-		.AutoHeight()
-		.Padding(2)
+		// Pose List
+		+SHorizontalBox::Slot()
+		.FillWidth(1)
+		.Padding(3)
 		[
-			SNew(STextBlock)
-			.Text(this, &SPoseViewer::GetTitleText)
-		]
+			SNew(SVerticalBox)
 
-		+ SVerticalBox::Slot()
-		.AutoHeight()
-		.Padding(0, 2)
-		[
-			SNew(SHorizontalBox)
-			// Filter entry
-			+ SHorizontalBox::Slot()
-			.FillWidth(1)
+			+ SVerticalBox::Slot()
+			.AutoHeight()
+			.Padding(0, 2)
 			[
-				SAssignNew(NameFilterBox, SSearchBox)
-				.SelectAllTextWhenFocused(true)
-				.OnTextChanged(this, &SPoseViewer::OnFilterTextChanged)
-				.OnTextCommitted(this, &SPoseViewer::OnFilterTextCommitted)
+				SNew(SHorizontalBox)
+				// Filter entry
+				+ SHorizontalBox::Slot()
+				.FillWidth(1)
+				[
+					SAssignNew(NameFilterBox, SSearchBox)
+					.SelectAllTextWhenFocused(true)
+					.OnTextChanged(this, &SPoseViewer::OnFilterTextChanged)
+					.OnTextCommitted(this, &SPoseViewer::OnFilterTextCommitted)
+				]
 			]
-		]
 
-		+ SVerticalBox::Slot()
-		.FillHeight(1.0f)		// This is required to make the scrollbar work, as content overflows Slate containers by default
-		[
-			SNew(SHorizontalBox)
-
-			+SHorizontalBox::Slot()
-			.FillWidth(1)
-			.Padding(5)
+			+ SVerticalBox::Slot()
+			.FillHeight(1)
+			.Padding(0, 2)
 			[
 				SAssignNew(PoseListView, SPoseListType)
 				.ListItemsSource(&PoseList)
 				.OnGenerateRow(this, &SPoseViewer::GeneratePoseRow)
 				.OnContextMenuOpening(this, &SPoseViewer::OnGetContextMenuContent)
+				.OnMouseButtonDoubleClick(this, &SPoseViewer::OnListDoubleClick)
 				.ItemHeight(22.0f)
 				.HeaderRow
 				(
@@ -384,12 +361,18 @@ void SPoseViewer::Construct(const FArguments& InArgs)
 
 					+ SHeaderRow::Column(ColumnID_PoseWeightLabel)
 					.DefaultLabel(LOCTEXT("PoseWeightLabel", "Weight"))
-				)
+					)
 			]
+		]
 
-			+ SHorizontalBox::Slot()
-			.FillWidth(1)
-			.Padding(5)
+		// Curve List
+		+ SHorizontalBox::Slot()
+		.FillWidth(1)
+		.Padding(5)
+		[
+			SNew(SBorder)
+			.Padding(3)
+			.BorderImage(FEditorStyle::GetBrush("ToolPanel.DarkGroupBorder"))
 			[
 				SAssignNew(CurveListView, SCurveListType)
 				.ListItemsSource(&CurveList)
@@ -401,8 +384,12 @@ void SPoseViewer::Construct(const FArguments& InArgs)
 					SNew(SHeaderRow)
 					+ SHeaderRow::Column(ColumnId_CurveNameLabel)
 					.DefaultLabel(LOCTEXT("CurveNameLabel", "Curve Name"))
+
+					+ SHeaderRow::Column(ColumnID_CurveValueLabel)
+					.DefaultLabel(LOCTEXT("CurveValueLabel", "Value"))
 				)
 			]
+
 		]
 	];
 
@@ -410,22 +397,6 @@ void SPoseViewer::Construct(const FArguments& InArgs)
 	CreateCurveList();
 }
 
-FText SPoseViewer::GetTitleText() const
-{
-	UObject* PreviewAsset = nullptr;
-
-	if (PersonaPtr.IsValid())
-	{
-		PreviewAsset = PersonaPtr.Pin()->GetPreviewAnimationAsset();
-		if (!PreviewAsset)
-		{
-			PreviewAsset = PersonaPtr.Pin()->GetAnimBlueprint();
-		}
-	}
-
-	return (PreviewAsset ? FText::FromString(PreviewAsset->GetName()) : LOCTEXT("PoseMeshNameLabel", "No Asset Present"));
-}
-
 void SPoseViewer::RefreshCachePreviewInstance()
 {
 	if (CachedPreviewInstance.IsValid() && OnAddAnimationCurveDelegate.IsBound())
@@ -469,10 +440,9 @@ TSharedRef<ITableRow> SPoseViewer::GeneratePoseRow(TSharedPtr<FDisplayedPoseInfo
 
 	return
 		SNew(SPoseListRow, OwnerTable)
-		.Persona(PersonaPtr)
 		.Item(InInfo)
-		.PoseViewer(this)
-		.PoseListView(PoseListView);
+		.PoseViewer(SharedThis(this))
+		.FilterText(GetFilterText());
 }
 
 TSharedRef<ITableRow> SPoseViewer::GenerateCurveRow(TSharedPtr<FDisplayedCurveInfo> InInfo, const TSharedRef<STableViewBase>& OwnerTable)
@@ -482,7 +452,7 @@ TSharedRef<ITableRow> SPoseViewer::GenerateCurveRow(TSharedPtr<FDisplayedCurveIn
 	return
 		SNew(SCurveListRow, OwnerTable)
 		.Item(InInfo)
-		.CurveListView(CurveListView);
+		.PoseViewer(SharedThis(this));
 }
 
 bool SPoseViewer::IsPoseSelected() const
@@ -492,6 +462,13 @@ bool SPoseViewer::IsPoseSelected() const
 	return SelectedRows.Num() > 0;
 }
 
+bool SPoseViewer::IsSinglePoseSelected() const
+{
+	// @todo: make sure not to delete base Curve
+	TArray< TSharedPtr< FDisplayedPoseInfo > > SelectedRows = PoseListView->GetSelectedItems();
+	return SelectedRows.Num() == 1;
+}
+
 bool SPoseViewer::IsCurveSelected() const
 {
 	// @todo: make sure not to delete base pose
@@ -517,6 +494,19 @@ void SPoseViewer::OnDeletePoses()
 	CreatePoseList(NameFilterBox->GetText().ToString());
 }
 
+void SPoseViewer::OnRenamePose()
+{
+	TArray< TSharedPtr< FDisplayedPoseInfo > > SelectedRows = PoseListView->GetSelectedItems();
+	if (SelectedRows.Num() > 0)
+	{
+		TSharedPtr<FDisplayedPoseInfo> SelectedRow = SelectedRows[0];
+		if (SelectedRow.IsValid())
+		{
+			SelectedRow->OnRenameRequested.ExecuteIfBound();
+		}
+	}
+}
+
 void SPoseViewer::OnDeleteCurves()
 {
 	TArray< TSharedPtr< FDisplayedCurveInfo > > SelectedRows = CurveListView->GetSelectedItems();
@@ -597,12 +587,21 @@ TSharedPtr<SWidget> SPoseViewer::OnGetContextMenuContent() const
 
 	MenuBuilder.BeginSection("PoseAction", LOCTEXT("Poseaction", "Selected Item Actions"));
 	{
+		// Delete
 	 	Action = FUIAction( FExecuteAction::CreateSP( this, &SPoseViewer::OnDeletePoses ), 
 	 									FCanExecuteAction::CreateSP( this, &SPoseViewer::IsPoseSelected ) );
 	 	MenuLabel = LOCTEXT("DeletePoseButtonLabel", "Delete");
 	 	MenuToolTip = LOCTEXT("DeletePoseButtonTooltip", "Deletes the selected animation pose.");
 	 	MenuBuilder.AddMenuEntry(MenuLabel, MenuToolTip, FSlateIcon(), Action);
 
+		// Rename
+		Action = FUIAction(FExecuteAction::CreateSP(this, &SPoseViewer::OnRenamePose),
+			FCanExecuteAction::CreateSP(this, &SPoseViewer::IsSinglePoseSelected));
+		MenuLabel = LOCTEXT("RenamePoseButtonLabel", "Rename");
+		MenuToolTip = LOCTEXT("RenamePoseButtonTooltip", "Renames the selected animation pose.");
+		MenuBuilder.AddMenuEntry(MenuLabel, MenuToolTip, FSlateIcon(), Action);
+
+		// Paste
 		Action = FUIAction(FExecuteAction::CreateSP(this, &SPoseViewer::OnPastePoseNamesFromClipBoard, true),
 			FCanExecuteAction::CreateSP(this, &SPoseViewer::IsPoseSelected));
 		MenuLabel = LOCTEXT("PastePoseNamesButtonLabel", "Paste Selected");
@@ -632,12 +631,41 @@ TSharedPtr<SWidget> SPoseViewer::OnGetContextMenuContentForCurveList() const
 	return MenuBuilder.MakeWidget();
 }
 
+void SPoseViewer::OnListDoubleClick(TSharedPtr<FDisplayedPoseInfo> InItem)
+{
+	if(InItem.IsValid())
+	{
+		const float CurrentWeight = InItem->Weight;
+
+		// Clear all preview poses
+		for (TSharedPtr<FDisplayedPoseInfo>& Pose : PoseList)
+		{
+			Pose->Weight = 0.f;
+			AddCurveOverride(Pose->Name, 0.f);
+		}
+
+		// If current weight was already at 1.0, do nothing (we are setting it to zero)
+		if (!FMath::IsNearlyEqual(CurrentWeight, 1.f))
+		{
+			// Otherwise set to 1.0
+			InItem->Weight = 1.f;
+			AddCurveOverride(InItem->Name, 1.f);
+		}
+
+		// Force update viewport
+		if (PersonaPtr.IsValid())
+		{
+			PersonaPtr.Pin()->RefreshViewport();
+		}
+	}
+}
+
+
 void SPoseViewer::CreatePoseList(const FString& SearchText)
 {
 	PoseList.Empty();
 
-	UAnimInstance* AnimInstance = CachedPreviewInstance.IsValid() ? CachedPreviewInstance.Get() : nullptr;
-	if (AnimInstance && PoseAssetPtr.IsValid())
+	if (PoseAssetPtr.IsValid())
 	{
 		UPoseAsset* PoseAsset = PoseAssetPtr.Get();
 
@@ -658,10 +686,6 @@ void SPoseViewer::CreatePoseList(const FString& SearchText)
 				PoseList.Add(Info);
 			}
 		}
-		else
-		{
-			// else print somewhere no curves?
-		}
 	}
 
 	PoseListView->RequestListRefresh();
@@ -671,46 +695,45 @@ void SPoseViewer::CreateCurveList(const FString& SearchText)
 {
 	CurveList.Empty();
 
-	UAnimInstance* AnimInstance = CachedPreviewInstance.IsValid() ? CachedPreviewInstance.Get() : nullptr;
-	if (AnimInstance && PoseAssetPtr.IsValid())
+	if (PoseAssetPtr.IsValid())
 	{
 		UPoseAsset* PoseAsset = PoseAssetPtr.Get();
 
 		TArray<FSmartName> CurveNames = PoseAsset->GetCurveNames();
 		if (CurveNames.Num() > 0)
 		{
-			bool bDoFiltering = !SearchText.IsEmpty();
-
 			for (const auto& CurveSmartName : CurveNames)
 			{
 				FName CurveName = CurveSmartName.DisplayName;
-				if (bDoFiltering && !CurveName.ToString().Contains(SearchText))
-				{
-					continue; // Skip items that don't match our filter
-				}
 
 				TSharedRef<FDisplayedCurveInfo> Info = FDisplayedCurveInfo::Make(CurveName);
 				CurveList.Add(Info);
 			}
 		}
-		else
-		{
-			// else print somewhere no curves?
-		}
 	}
 
 	CurveListView->RequestListRefresh();
 }
 
-void SPoseViewer::AddPoseOverride(FName& Name, float Weight)
+void SPoseViewer::AddCurveOverride(FName& Name, float Weight)
 {
 	float& Value = OverrideCurves.FindOrAdd(Name);
 	Value = Weight;
+
+	if (CachedPreviewInstance.IsValid())
+	{
+		CachedPreviewInstance->SetPreviewCurveOverride(Name, Value, false);
+	}
 }
 
-void SPoseViewer::RemovePoseOverride(FName& Name)
+void SPoseViewer::RemoveCurveOverride(FName& Name)
 {
 	OverrideCurves.Remove(Name);
+
+	if (CachedPreviewInstance.IsValid())
+	{
+		CachedPreviewInstance->SetPreviewCurveOverride(Name, 0.f, true);
+	}
 }
 
 SPoseViewer::~SPoseViewer()
@@ -728,9 +751,14 @@ SPoseViewer::~SPoseViewer()
 		PersonaPtr.Pin()->UnregisterOnPreviewMeshChanged(this);
 		PersonaPtr.Pin()->UnregisterOnPostUndo(this);
 	}
+
+	if (PoseAssetPtr.IsValid())
+	{
+		PoseAssetPtr->UnregisterOnAssetModifiedNotifier(this);
+	}
 }
 
-void SPoseViewer::OnPostUndo()
+void SPoseViewer::RefreshList()
 {
 	CreatePoseList(NameFilterBox->GetText().ToString());
 	CreateCurveList(NameFilterBox->GetText().ToString());
@@ -741,7 +769,7 @@ void SPoseViewer::ApplyCustomCurveOverride(UAnimInstance* AnimInstance) const
 	for (auto Iter = OverrideCurves.CreateConstIterator(); Iter; ++Iter)
 	{
 		// @todo we might want to save original curve flags? or just change curve to apply flags only
-		AnimInstance->AddCurveValue(Iter.Key(), Iter.Value(), ACF_DrivesPose);
+		AnimInstance->AddCurveValue(Iter.Key(), Iter.Value(), ACF_DriveAttribute);
 	}
 }
 
@@ -777,7 +805,6 @@ bool SPoseViewer::ModifyName(FName OldName, FName NewName, bool bSilence)
 				// I think this might have to be delegate of the top window
 				if (PoseAssetPtr.Get()->ModifyPoseName(OldName, NewName, &ExistingUID) == false)
 				{
-					FMessageDialog::Open(EAppMsgType::Ok, LOCTEXT("NameAlreadyUsedByTheSameAsset", "The name is used by another pose within the same asset. Please choose another name."));
 					return false;
 				}
 			}
@@ -786,7 +813,6 @@ bool SPoseViewer::ModifyName(FName OldName, FName NewName, bool bSilence)
 				// I think this might have to be delegate of the top window
 				if (PoseAssetPtr.Get()->ModifyPoseName(OldName, NewName, nullptr) == false)
 				{
-					FMessageDialog::Open(EAppMsgType::Ok, LOCTEXT("NameAlreadyUsedByTheSameAsset", "The name is used by another pose within the same asset. Please choose another name."));
 					return false;
 				}
 			}
@@ -795,8 +821,8 @@ bool SPoseViewer::ModifyName(FName OldName, FName NewName, bool bSilence)
 			float* Value = OverrideCurves.Find(OldName);
 			if (Value)
 			{
-				AddPoseOverride(NewName, *Value);
-				RemovePoseOverride(OldName);
+				AddCurveOverride(NewName, *Value);
+				RemoveCurveOverride(OldName);
 			}
 
 			return true;
diff --git a/Engine/Source/Editor/Persona/Private/SPoseEditor.h b/Engine/Source/Editor/Persona/Private/SPoseEditor.h
--- a/Engine/Source/Editor/Persona/Private/SPoseEditor.h
+++ b/Engine/Source/Editor/Persona/Private/SPoseEditor.h
@@ -12,6 +12,8 @@
 #include "Animation/PoseAsset.h"
 
 
+class SPoseViewer;
+
 //////////////////////////////////////////////////////////////////////////
 // FDisplayedPoseInfo
 
@@ -27,6 +29,10 @@ public:
 		return MakeShareable(new FDisplayedPoseInfo(Source));
 	}
 
+	/** Delegate for when the context menu requests a rename */
+	DECLARE_DELEGATE(FOnRenameRequested);
+	FOnRenameRequested OnRenameRequested;
+
 protected:
 	/** Hidden constructor, always use Make above */
 	FDisplayedPoseInfo(const FName& InSource)
@@ -40,6 +46,77 @@ protected:
 
 typedef SListView< TSharedPtr<FDisplayedPoseInfo> > SPoseListType;
 
+//////////////////////////////////////////////////////////////////////////
+// SPoseListRow
+
+class SPoseListRow : public SMultiColumnTableRow< TSharedPtr<FDisplayedPoseInfo> >
+{
+public:
+
+	SLATE_BEGIN_ARGS(SPoseListRow) {}
+
+		/** The item for this row **/
+		SLATE_ARGUMENT(TSharedPtr<FDisplayedPoseInfo>, Item)
+
+		/* The SPoseViewer that we push the morph target weights into */
+		SLATE_ARGUMENT(TWeakPtr<SPoseViewer>, PoseViewer)
+
+		/** Filter text typed by the user into the parent tree's search widget */
+		SLATE_ARGUMENT(FText, FilterText);
+
+	SLATE_END_ARGS()
+
+	void Construct(const FArguments& InArgs, const TSharedRef<STableViewBase>& OwnerTableView);
+
+	/** Overridden from SMultiColumnTableRow.  Generates a widget for this column of the tree row. */
+	virtual TSharedRef<SWidget> GenerateWidgetForColumn(const FName& ColumnName) override;
+
+private:
+
+	/**
+	* Called when the user changes the value of the SSpinBox
+	*
+	* @param NewWeight - The new number the SSpinBox is set to
+	*
+	*/
+	void OnPoseWeightChanged(float NewWeight);
+	/**
+	* Called when the user types the value and enters
+	*
+	* @param NewWeight - The new number the SSpinBox is set to
+	*
+	*/
+	void OnPoseWeightValueCommitted(float NewWeight, ETextCommit::Type CommitType);
+
+	/** Delegate to get labels root text from settings */
+	FText GetName() const;
+
+	/** Delegate to commit labels root text to settings */
+	void OnNameCommitted(const FText& InText, ETextCommit::Type InCommitType) const;
+
+	bool OnVerifyNameChanged(const FText& InText, FText& OutErrorMessage);
+
+
+	/**
+	* Returns the weight of this morph target
+	*
+	* @return SearchText - The new number the SSpinBox is set to
+	*
+	*/
+	float GetWeight() const;
+
+	bool CanChangeWeight() const;
+
+	/* The SPoseViewer that we push the pose weights into */
+	TWeakPtr<SPoseViewer> PoseViewerPtr;
+
+	/** The name and weight of the morph target */
+	TSharedPtr<FDisplayedPoseInfo>	Item;
+
+	/** Text the user typed into the search box - used for text highlighting */
+	FText FilterText;
+};
+
 //////////////////////////////////////////////////////////////////////////
 // FDisplayedCurveInfo
 
@@ -64,7 +141,47 @@ protected:
 	FDisplayedCurveInfo() {}
 };
 
+//////////////////////////////////////////////////////////////////////////
+// SCurveListRow
+
 typedef SListView< TSharedPtr<FDisplayedCurveInfo> > SCurveListType;
+
+
+class SCurveListRow : public SMultiColumnTableRow< TSharedPtr<FDisplayedCurveInfo> >
+{
+public:
+
+	SLATE_BEGIN_ARGS(SCurveListRow) {}
+
+		/** The item for this row **/
+		SLATE_ARGUMENT(TSharedPtr<FDisplayedCurveInfo>, Item)
+
+		/* The SPoseViewer that we push the morph target weights into */
+		SLATE_ARGUMENT(TWeakPtr<SPoseViewer>, PoseViewer)
+
+	SLATE_END_ARGS()
+
+	void Construct(const FArguments& InArgs, const TSharedRef<STableViewBase>& OwnerTableView);
+
+	/** Overridden from SMultiColumnTableRow.  Generates a widget for this column of the tree row. */
+	virtual TSharedRef<SWidget> GenerateWidgetForColumn(const FName& ColumnName) override;
+
+private:
+
+	/** Delegate to get labels root text from settings */
+	FText GetName() const;
+
+	/** Delegate to get weight of curve in selected pose */
+	FText GetValue() const;
+
+	/** The name and weight of the morph target */
+	TSharedPtr<FDisplayedCurveInfo>	Item;
+
+	/* The SPoseViewer that we push the morph target weights into */
+	TWeakPtr<SPoseViewer> PoseViewerPtr;
+};
+
+
 //////////////////////////////////////////////////////////////////////////
 // SPoseViewer
 
@@ -81,13 +198,13 @@ public:
 
 	SLATE_END_ARGS()
 
-		/**
-		* Slate construction function
-		*
-		* @param InArgs - Arguments passed from Slate
-		*
-		*/
-		void Construct(const FArguments& InArgs);
+	/**
+	* Slate construction function
+	*
+	* @param InArgs - Arguments passed from Slate
+	*
+	*/
+	void Construct(const FArguments& InArgs);
 
 	/**
 	* Destructor - resets the animation curve
@@ -132,19 +249,20 @@ public:
 	TSharedRef<ITableRow> GeneratePoseRow(TSharedPtr<FDisplayedPoseInfo> InInfo, const TSharedRef<STableViewBase>& OwnerTable);
 	TSharedRef<ITableRow> GenerateCurveRow(TSharedPtr<FDisplayedCurveInfo> InInfo, const TSharedRef<STableViewBase>& OwnerTable);
 
-	/**
-	* Provides state to the IsEnabled property of the delete morph targets button
-	*
-	*/
 	bool IsPoseSelected() const;
+	bool IsSinglePoseSelected() const;
 	bool IsCurveSelected() const;
 
-	/**
-	* Handler for the delete poses button
-	*/
+	/** Handler for the delete poses option */
 	void OnDeletePoses();
+
+	/** Handler for rename pose option */
+	void OnRenamePose();
+
+	/** Handler for delete curves option */
 	void OnDeleteCurves();
 
+	/** Handler for pasting names from clipboard */
 	void OnPastePoseNamesFromClipBoard(bool bSelectedOnly);
 
 	/**
@@ -152,16 +270,11 @@ public:
 	*
 	* @param Name - Name of the morph target we want to override
 	* @param Weight - How much of this morph target to apply (0.0 - 1.0)
-	*
 	*/
-	void AddPoseOverride(FName& Name, float Weight);
-	void RemovePoseOverride(FName& Name);
+	void AddCurveOverride(FName& Name, float Weight);
 
-	/**
-	* Tells the AnimInstance to reset all of its morph target curves
-	*
-	*/
-	void ResetPoses();
+	/** Remove a named curve override */ 
+	void RemoveCurveOverride(FName& Name);
 
 	/**
 	* Accessor so our rows can grab the filtertext for highlighting
@@ -176,7 +289,7 @@ public:
 	/**
 	* Refreshes the morph target list after an undo
 	*/
-	void OnPostUndo();
+	void RefreshList();
 
 	bool ModifyName(FName OldName, FName NewName, bool bSilence = false);
 	bool IsBasePose(FName PoseName) const;
@@ -187,6 +300,9 @@ private:
 	TSharedPtr<SWidget> OnGetContextMenuContent() const;
 	/** Handler for curve list context menu*/
 	TSharedPtr<SWidget> OnGetContextMenuContentForCurveList() const;
+	/** Called when list double-clicked */
+	void OnListDoubleClick(TSharedPtr<FDisplayedPoseInfo> InItem);
+
 
 	/**
 	* Clears and rebuilds the table, according to an optional search string
diff --git a/Engine/Source/Editor/Persona/Private/SRigWindow.cpp b/Engine/Source/Editor/Persona/Private/SRigWindow.cpp
--- a/Engine/Source/Editor/Persona/Private/SRigWindow.cpp
+++ b/Engine/Source/Editor/Persona/Private/SRigWindow.cpp
@@ -141,8 +141,7 @@ TSharedRef< SWidget > SBoneMappingListRow::GenerateWidgetForColumn( const FName&
 
 				+SHorizontalBox::Slot()
 				[
-					SNew(SBoneSelectionWidget)
-					.Skeleton(Item->Skeleton)
+					SNew(SBoneSelectionWidget, Item->Skeleton)
 					.Tooltip(FText::Format(LOCTEXT("BoneSelectinWidget", "Select Bone for node {0}"), FText::FromString(Item->GetDisplayName())))
 					.OnBoneSelectionChanged(this, &SBoneMappingListRow::OnBoneSelectionChanged)
 					.OnGetSelectedBone(this, &SBoneMappingListRow::GetSelectedBone)
diff --git a/Engine/Source/Editor/Persona/Private/Shared/BoneSelectionWidget.cpp b/Engine/Source/Editor/Persona/Private/Shared/BoneSelectionWidget.cpp
--- a/Engine/Source/Editor/Persona/Private/Shared/BoneSelectionWidget.cpp
+++ b/Engine/Source/Editor/Persona/Private/Shared/BoneSelectionWidget.cpp
@@ -9,49 +9,25 @@
 #define LOCTEXT_NAMESPACE "SBoneSelectionWidget"
 
 /////////////////////////////////////////////////////
-
-void SBoneSelectionWidget::Construct(const FArguments& InArgs)
+void SBoneTreeMenu::Construct(const FArguments& InArgs, TWeakObjectPtr<const USkeleton> Skeleton)
 {
-	if(InArgs._Skeleton.IsValid())
-	{
-		TargetSkeleton = InArgs._Skeleton.Get();
-	}
-
-	check(TargetSkeleton);
-
-	OnBoneSelectionChanged = InArgs._OnBoneSelectionChanged;
-	OnGetSelectedBone = InArgs._OnGetSelectedBone;
+	TargetSkeleton = Skeleton;
+	OnSelectionChangedDelegate = InArgs._OnBoneSelectionChanged;
 
-	FText FinalTooltip = FText::Format(LOCTEXT("BoneClickToolTip", "{0}\nClick to choose a different bone"), InArgs._Tooltip);
+	FText TitleToUse = !InArgs._Title.IsEmpty() ? InArgs._Title  : LOCTEXT("BonePickerTitle", "Pick Bone...");
 
 	SAssignNew(TreeView, STreeView<TSharedPtr<FBoneNameInfo>>)
 		.TreeItemsSource(&SkeletonTreeInfo)
-		.OnGenerateRow(this, &SBoneSelectionWidget::MakeTreeRowWidget)
-		.OnGetChildren(this, &SBoneSelectionWidget::GetChildrenForInfo)
-		.OnSelectionChanged(this, &SBoneSelectionWidget::OnSelectionChanged)
+		.OnGenerateRow(this, &SBoneTreeMenu::MakeTreeRowWidget)
+		.OnGetChildren(this, &SBoneTreeMenu::GetChildrenForInfo)
+		.OnSelectionChanged(this, &SBoneTreeMenu::OnSelectionChanged)
 		.SelectionMode(ESelectionMode::Single);
 
-	this->ChildSlot
-	[
-		SAssignNew(BonePickerButton, SComboButton)
-		.OnGetMenuContent(FOnGetContent::CreateSP(this, &SBoneSelectionWidget::CreateSkeletonWidgetMenu))
-		.ContentPadding(FMargin(4.0f, 2.0f, 4.0f, 2.0f))
-		.ButtonContent()
-		[
-			SNew(STextBlock)
-			.Text(this, &SBoneSelectionWidget::GetCurrentBoneName)
-			.Font(IDetailLayoutBuilder::GetDetailFont())
-			.ToolTipText(FinalTooltip)
-		]
-	];
-}
-
-TSharedRef<SWidget> SBoneSelectionWidget::CreateSkeletonWidgetMenu()
-{
 	RebuildBoneList();
 
-	TSharedPtr<SSearchBox> SearchWidgetToFocus = NULL;
-	TSharedRef<SBorder> MenuWidget = SNew(SBorder)
+	ChildSlot
+	[
+		SNew(SBorder)
 		.Padding(6)
 		.BorderImage(FEditorStyle::GetBrush("NoBorder"))
 		.Content()
@@ -62,42 +38,38 @@ TSharedRef<SWidget> SBoneSelectionWidget::CreateSkeletonWidgetMenu()
 			.Content()
 			[
 				SNew(SVerticalBox)
-				+SVerticalBox::Slot()
+				+ SVerticalBox::Slot()
 				.AutoHeight()
 				[
 					SNew(STextBlock)
 					.Font(FEditorStyle::GetFontStyle("BoldFont"))
-						.Text(LOCTEXT("BonePickerTitle", "Pick Bone..."))
+					.Text(TitleToUse)
 				]
-				+SVerticalBox::Slot()
+				+ SVerticalBox::Slot()
 				.AutoHeight()
 				[
 					SNew(SSeparator)
 					.SeparatorImage(FEditorStyle::GetBrush("Menu.Separator"))
-						.Orientation(Orient_Horizontal)
+					.Orientation(Orient_Horizontal)
 				]
-				+SVerticalBox::Slot()
+				+ SVerticalBox::Slot()
 				.AutoHeight()
 				[
-					SAssignNew(SearchWidgetToFocus, SSearchBox)
+					SAssignNew(FilterTextWidget, SSearchBox)
 					.SelectAllTextWhenFocused(true)
-					.OnTextChanged(this, &SBoneSelectionWidget::OnFilterTextChanged)
-					.OnTextCommitted(this, &SBoneSelectionWidget::OnFilterTextCommitted)
+					.OnTextChanged(this, &SBoneTreeMenu::OnFilterTextChanged)
 					.HintText(NSLOCTEXT("BonePicker", "Search", "Search..."))
 				]
-				+SVerticalBox::Slot()
+				+ SVerticalBox::Slot()
 				[
 					TreeView->AsShared()
 				]
 			]
-		];
-
-	BonePickerButton->SetMenuContentWidgetToFocus(SearchWidgetToFocus);
-
-	return MenuWidget;
+		]
+	];
 }
 
-TSharedRef<ITableRow> SBoneSelectionWidget::MakeTreeRowWidget(TSharedPtr<FBoneNameInfo> InInfo, const TSharedRef<STableViewBase>& OwnerTable)
+TSharedRef<ITableRow> SBoneTreeMenu::MakeTreeRowWidget(TSharedPtr<FBoneNameInfo> InInfo, const TSharedRef<STableViewBase>& OwnerTable)
 {
 	return SNew(STableRow<TSharedPtr<FBoneNameInfo>>, OwnerTable)
 		.Content()
@@ -108,27 +80,36 @@ TSharedRef<ITableRow> SBoneSelectionWidget::MakeTreeRowWidget(TSharedPtr<FBoneNa
 		];
 }
 
-void SBoneSelectionWidget::GetChildrenForInfo(TSharedPtr<FBoneNameInfo> InInfo, TArray< TSharedPtr<FBoneNameInfo> >& OutChildren)
+void SBoneTreeMenu::GetChildrenForInfo(TSharedPtr<FBoneNameInfo> InInfo, TArray< TSharedPtr<FBoneNameInfo> >& OutChildren)
 {
 	OutChildren = InInfo->Children;
 }
 
-void SBoneSelectionWidget::OnFilterTextChanged(const FText& InFilterText)
+void SBoneTreeMenu::OnFilterTextChanged(const FText& InFilterText)
 {
 	FilterText = InFilterText;
 
 	RebuildBoneList();
 }
 
-void SBoneSelectionWidget::OnFilterTextCommitted(const FText& SearchText, ETextCommit::Type CommitInfo)
+void SBoneTreeMenu::OnSelectionChanged(TSharedPtr<SBoneTreeMenu::FBoneNameInfo> BoneInfo, ESelectInfo::Type SelectInfo)
 {
-	// Already committed as the text was typed
+	//Because we recreate all our items on tree refresh we will get a spurious null selection event initially.
+	if (BoneInfo.IsValid())
+	{
+		OnSelectionChangedDelegate.ExecuteIfBound(BoneInfo->BoneName);
+	}
 }
 
-void SBoneSelectionWidget::RebuildBoneList()
+void SBoneTreeMenu::RebuildBoneList()
 {
 	SkeletonTreeInfo.Empty();
 	SkeletonTreeInfoFlat.Empty();
+	if (!TargetSkeleton.IsValid())
+	{
+		return;
+	}
+
 	const FReferenceSkeleton& RefSkeleton = TargetSkeleton->GetReferenceSkeleton();
 	for(int32 BoneIdx = 0; BoneIdx < RefSkeleton.GetNum(); ++BoneIdx)
 	{
@@ -180,20 +161,53 @@ void SBoneSelectionWidget::RebuildBoneList()
 	TreeView->RequestTreeRefresh();
 }
 
-void SBoneSelectionWidget::OnSelectionChanged(TSharedPtr<FBoneNameInfo> BoneInfo, ESelectInfo::Type SelectInfo)
+/////////////////////////////////////////////////////
+
+void SBoneSelectionWidget::Construct(const FArguments& InArgs, TWeakObjectPtr<const USkeleton> Skeleton)
 {
-	FilterText = FText::FromString("");
+	TargetSkeleton = Skeleton;
+
+	check(TargetSkeleton.IsValid());
+
+	OnBoneSelectionChanged = InArgs._OnBoneSelectionChanged;
+	OnGetSelectedBone = InArgs._OnGetSelectedBone;
 
+	FText FinalTooltip = FText::Format(LOCTEXT("BoneClickToolTip", "{0}\nClick to choose a different bone"), InArgs._Tooltip);
+
+	this->ChildSlot
+	[
+		SAssignNew(BonePickerButton, SComboButton)
+		.OnGetMenuContent(FOnGetContent::CreateSP(this, &SBoneSelectionWidget::CreateSkeletonWidgetMenu))
+		.ContentPadding(FMargin(4.0f, 2.0f, 4.0f, 2.0f))
+		.ButtonContent()
+		[
+			SNew(STextBlock)
+			.Text(this, &SBoneSelectionWidget::GetCurrentBoneName)
+			.Font(IDetailLayoutBuilder::GetDetailFont())
+			.ToolTipText(FinalTooltip)
+		]
+	];
+}
+
+TSharedRef<SWidget> SBoneSelectionWidget::CreateSkeletonWidgetMenu()
+{
+	TSharedRef<SBoneTreeMenu> MenuWidget = SNew(SBoneTreeMenu, TargetSkeleton)
+									.OnBoneSelectionChanged(this, &SBoneSelectionWidget::OnSelectionChanged);
+
+	BonePickerButton->SetMenuContentWidgetToFocus(MenuWidget->FilterTextWidget);
+
+	return MenuWidget;
+}
+
+void SBoneSelectionWidget::OnSelectionChanged(FName BoneName)
+{
 	//Because we recreate all our items on tree refresh we will get a spurious null selection event initially.
-	if (BoneInfo.IsValid())
+	if (OnBoneSelectionChanged.IsBound())
 	{
-		if (OnBoneSelectionChanged.IsBound())
-		{
-			OnBoneSelectionChanged.Execute(BoneInfo->BoneName);
-		}
-
-		BonePickerButton->SetIsOpen(false);
+		OnBoneSelectionChanged.Execute(BoneName);
 	}
+
+	BonePickerButton->SetIsOpen(false);
 }
 
 FText SBoneSelectionWidget::GetCurrentBoneName() const
diff --git a/Engine/Source/Editor/Persona/Private/Shared/PersonaMode.cpp b/Engine/Source/Editor/Persona/Private/Shared/PersonaMode.cpp
--- a/Engine/Source/Editor/Persona/Private/Shared/PersonaMode.cpp
+++ b/Engine/Source/Editor/Persona/Private/Shared/PersonaMode.cpp
@@ -7,7 +7,6 @@
 #include "IDocumentation.h"
 #include "SAnimBlueprintParentPlayerList.h"
 #include "SSkeletonSlotNames.h"
-#include "SSkeletonSmartNameManager.h"
 #include "SSkeletonBlendProfiles.h"
 #include "SDockTab.h"
 #include "SAdvancedPreviewDetailsTab.h"
@@ -43,7 +42,6 @@ const FName FPersonaTabs::PreviewManagerID("AnimPreviewSetup");		//@TODO: Name
 const FName FPersonaTabs::SkeletonAnimNotifiesID("SkeletonAnimNotifies");
 const FName FPersonaTabs::SkeletonSlotNamesID("SkeletonSlotNames");
 const FName FPersonaTabs::SkeletonSlotGroupNamesID("SkeletonSlotGroupNames");
-const FName FPersonaTabs::CurveNameManagerID("CurveNameManager");
 const FName FPersonaTabs::BlendProfileManagerID("BlendProfileManager");
 
 const FName FPersonaTabs::AdvancedPreviewSceneSettingsID("AdvancedPreviewTab");
@@ -74,7 +72,6 @@ FPersonaAppMode::FPersonaAppMode(TSharedPtr<class FPersona> InPersona, FName InM
 	PersonaTabFactories.RegisterFactory(MakeShareable(new FSkeletonAnimNotifiesSummoner(InPersona)));
 	PersonaTabFactories.RegisterFactory(MakeShareable(new FRetargetManagerTabSummoner(InPersona)));
 	PersonaTabFactories.RegisterFactory(MakeShareable(new FSkeletonSlotNamesSummoner(InPersona)));
-	PersonaTabFactories.RegisterFactory(MakeShareable(new FSkeletonCurveNameManagerSummoner(InPersona)));
 	PersonaTabFactories.RegisterFactory(MakeShareable(new FSkeletonBlendProfilesSummoner(InPersona)));
 	PersonaTabFactories.RegisterFactory(MakeShareable(new FAdvancedPreviewSceneTabSummoner(InPersona)));
 }
@@ -170,13 +167,13 @@ TSharedRef<SWidget> FMorphTargetTabSummoner::CreateTabBody(const FWorkflowTabSpa
 FAnimCurveViewerTabSummoner::FAnimCurveViewerTabSummoner(TSharedPtr<class FAssetEditorToolkit> InHostingApp)
 	: FWorkflowTabFactory(FPersonaTabs::AnimCurveViewID, InHostingApp)
 {
-	TabLabel = LOCTEXT("AnimCurveViewTabTitle", "Anim Curve Previewer");
+	TabLabel = LOCTEXT("AnimCurveViewTabTitle", "Anim Curves");
 	TabIcon = FSlateIcon(FEditorStyle::GetStyleSetName(), "Persona.Tabs.AnimCurvePreviewer");
 
 	EnableTabPadding();
 	bIsSingleton = true;
 
-	ViewMenuDescription = LOCTEXT("AnimCurveTabView", "Animation Curve Previewer");
+	ViewMenuDescription = LOCTEXT("AnimCurveTabView", "Animation Curves");
 	ViewMenuTooltip = LOCTEXT("AnimCurveTabView_ToolTip", "Shows the animation curve viewer");
 }
 
diff --git a/Engine/Source/Editor/Persona/Public/BoneSelectionWidget.h b/Engine/Source/Editor/Persona/Public/BoneSelectionWidget.h
--- a/Engine/Source/Editor/Persona/Public/BoneSelectionWidget.h
+++ b/Engine/Source/Editor/Persona/Public/BoneSelectionWidget.h
@@ -7,19 +7,78 @@
 DECLARE_DELEGATE_OneParam(FOnBoneSelectionChanged, FName);
 DECLARE_DELEGATE_RetVal(FName, FGetSelectedBone);
 
+class SBoneTreeMenu : public SCompoundWidget
+{
+public:
+	// Storage object for bone hierarchy
+	struct FBoneNameInfo
+	{
+		FBoneNameInfo(FName Name) : BoneName(Name) {}
+
+		FName BoneName;
+		TArray<TSharedPtr<FBoneNameInfo>> Children;
+	};
+
+	SLATE_BEGIN_ARGS(SBoneTreeMenu)
+		{}
+
+		SLATE_ARGUMENT(FText, Title)
+		SLATE_EVENT(FOnBoneSelectionChanged, OnBoneSelectionChanged)
+
+	SLATE_END_ARGS();
+
+	/**
+	* Construct this widget
+	*
+	* @param	InArgs	The declaration data for this widget
+	*/
+	void Construct(const FArguments& InArgs, TWeakObjectPtr<const USkeleton> Skeleton);
+
+	//Filter text widget
+	TSharedPtr<SSearchBox> FilterTextWidget;
+
+private:
+
+	// Using the current filter, repopulate the tree view
+	void RebuildBoneList();
+
+	// Make a single tree row widget
+	TSharedRef<ITableRow> MakeTreeRowWidget(TSharedPtr<FBoneNameInfo> InInfo, const TSharedRef<STableViewBase>& OwnerTable);
+
+	// Get the children for the provided bone info
+	void GetChildrenForInfo(TSharedPtr<FBoneNameInfo> InInfo, TArray< TSharedPtr<FBoneNameInfo> >& OutChildren);
+
+	// Called when the user changes the search filter
+	void OnFilterTextChanged(const FText& InFilterText);
+
+	void OnSelectionChanged(TSharedPtr<SBoneTreeMenu::FBoneNameInfo> BoneInfo, ESelectInfo::Type SelectInfo);
+
+	// Tree info entries for bone picker
+	TArray<TSharedPtr<FBoneNameInfo>> SkeletonTreeInfo;
+	// Mirror of SkeletonTreeInfo but flattened for searching
+	TArray<TSharedPtr<FBoneNameInfo>> SkeletonTreeInfoFlat;
+
+	// Text to filter bone tree with
+	FText FilterText;
+
+	// Skeleton to search
+	TWeakObjectPtr<const USkeleton> TargetSkeleton;
+
+	// Tree view used in the button menu
+	TSharedPtr<STreeView<TSharedPtr<FBoneNameInfo>>> TreeView;
+
+	FOnBoneSelectionChanged OnSelectionChangedDelegate;
+};
+
 class SBoneSelectionWidget : public SCompoundWidget
 {
 public: 
 
 	SLATE_BEGIN_ARGS( SBoneSelectionWidget )
-		:_Skeleton()
-		,_Tooltip()
+		:_Tooltip()
 		,_OnBoneSelectionChanged()
 		,_OnGetSelectedBone()
 	{}
-
-		/** Vertical alignment */
-		SLATE_ARGUMENT(TWeakObjectPtr<USkeleton>, Skeleton);
 		/** Set tooltip attribute */
 		SLATE_ARGUMENT(FText, Tooltip);
 
@@ -36,51 +95,22 @@
 	 *
 	 * @param	InArgs	The declaration data for this widget
 	 */
-	void Construct( const FArguments& InArgs );
+	void Construct( const FArguments& InArgs, TWeakObjectPtr<const USkeleton> Skeleton);
 
 private: 
 
-	// Storage object for bone hierarchy
-	struct FBoneNameInfo
-	{
-		FBoneNameInfo(FName Name): BoneName(Name) {}
-
-		FName BoneName;
-		TArray<TSharedPtr<FBoneNameInfo>> Children;
-	};
-
 	// Creates the combo button menu when clicked
 	TSharedRef<SWidget> CreateSkeletonWidgetMenu();
-	// Using the current filter, repopulate the tree view
-	void RebuildBoneList();
-	// Make a single tree row widget
-	TSharedRef<ITableRow> MakeTreeRowWidget(TSharedPtr<FBoneNameInfo> InInfo, const TSharedRef<STableViewBase>& OwnerTable);
-	// Get the children for the provided bone info
-	void GetChildrenForInfo(TSharedPtr<FBoneNameInfo> InInfo, TArray< TSharedPtr<FBoneNameInfo> >& OutChildren);
-	// Called when the user changes the search filter
-	void OnFilterTextChanged(const FText& InFilterText);
-	void OnFilterTextCommitted(const FText& SearchText, ETextCommit::Type CommitInfo);
 	// Called when the user selects a bone name
-	void OnSelectionChanged(TSharedPtr<FBoneNameInfo>, ESelectInfo::Type SelectInfo);
+	void OnSelectionChanged(FName BoneName);
 	// Gets the current bone name, used to get the right name for the combo button
 	FText GetCurrentBoneName() const;
 
-	// Tree view used in the button menu
-	TSharedPtr<STreeView<TSharedPtr<FBoneNameInfo>>> TreeView;
-
-	// Tree info entries for bone picker
-	TArray<TSharedPtr<FBoneNameInfo>> SkeletonTreeInfo;
-	// Mirror of SkeletonTreeInfo but flattened for searching
-	TArray<TSharedPtr<FBoneNameInfo>> SkeletonTreeInfoFlat;
-
 	// Base combo button 
 	TSharedPtr<SComboButton> BonePickerButton;
 
 	// Skeleton to search
-	USkeleton* TargetSkeleton;
+	TWeakObjectPtr<const USkeleton> TargetSkeleton;
 
-	// Text to filter bone tree with
-	FText FilterText;
-
 	// delegates
 	FOnBoneSelectionChanged OnBoneSelectionChanged;
diff --git a/Engine/Source/Editor/PhAT/Private/PhAT.cpp b/Engine/Source/Editor/PhAT/Private/PhAT.cpp
--- a/Engine/Source/Editor/PhAT/Private/PhAT.cpp
+++ b/Engine/Source/Editor/PhAT/Private/PhAT.cpp
@@ -30,6 +30,7 @@
 #include "PhysicsEngine/BodySetup.h"
 #include "PhysicsEngine/PhysicsConstraintTemplate.h"
 #include "PhysicsEngine/ConstraintUtils.h"
+#include "PhysicsEngine/PhysicsAsset.h"
 #include "Engine/Selection.h"
 #include "Engine/StaticMesh.h"
 #include "EngineLogs.h"
@@ -37,6 +38,8 @@
 #include "PersonaModule.h"
 #include "Animation/AnimSequence.h"
 
+#include "STextComboBox.h"
+
 const FName PhATAppIdentifier = FName(TEXT("PhATApp"));
 
 DEFINE_LOG_CATEGORY(LogPhAT);
@@ -81,6 +84,7 @@ public:
 
 static const FName PhATPreviewViewportName("PhAT_PreviewViewport");
 static const FName PhATPropertiesName("PhAT_Properties");
+static const FName PhATPhysicsAssetPropertiesName("PhAT_PhysicsAssetProperties");
 static const FName PhATHierarchyName("PhAT_Hierarchy");
 
 
@@ -102,6 +106,11 @@ void FPhAT::RegisterTabSpawners(const TSharedRef<class FTabManager>& InTabManage
 		.SetGroup( WorkspaceMenuCategoryRef )
 		.SetIcon( FSlateIcon( FEditorStyle::GetStyleSetName(), "LevelEditor.Tabs.Details" ) );
 
+	InTabManager->RegisterTabSpawner(PhATPhysicsAssetPropertiesName, FOnSpawnTab::CreateSP(this, &FPhAT::SpawnTab, PhATPhysicsAssetPropertiesName))
+		.SetDisplayName(LOCTEXT("PhysicsAssetPropertiesTab", "Physics Asset"))
+		.SetGroup(WorkspaceMenuCategoryRef)
+		.SetIcon(FSlateIcon(FEditorStyle::GetStyleSetName(), "LevelEditor.Tabs.Details"));
+
 	InTabManager->RegisterTabSpawner( PhATHierarchyName, FOnSpawnTab::CreateSP( this, &FPhAT::SpawnTab, PhATHierarchyName ) )
 		.SetDisplayName( LOCTEXT( "HierarchyTab", "Hierarchy" ) )
 		.SetGroup( WorkspaceMenuCategoryRef )
@@ -140,6 +149,15 @@ TSharedRef<SDockTab> FPhAT::SpawnTab( const FSpawnTabArgs& TabSpawnArgs, FName T
 				Properties.ToSharedRef()
 			];
 	}
+	else if(TabIdentifier == PhATPhysicsAssetPropertiesName)
+	{
+		return SNew(SDockTab)
+			.Icon(FEditorStyle::GetBrush("PhAT.Tabs.Properties"))
+			.Label(LOCTEXT("PhATPhysicsAssetPropertiesTitle", "Physics Asset"))
+			[
+				PhysAssetProperties.ToSharedRef()
+			];
+	}
 	else if (TabIdentifier == PhATHierarchyName)
 	{
 		const TSharedRef<SDockTab> NewTab = SNew(SDockTab)
@@ -338,7 +356,7 @@ void FPhAT::SetPropertiesSelection(UObject* Obj, FPhATSharedData::FSelection * B
 
 			if (!bFound && (SharedData->EditingMode == FPhATSharedData::PEM_BodyEdit))
 			{
-				int32 BoneIndex = SharedData->EditorSkelComp->GetBoneIndex(SharedData->PhysicsAsset->BodySetup[Body->Index]->BoneName);
+				int32 BoneIndex = SharedData->EditorSkelComp->GetBoneIndex(SharedData->PhysicsAsset->SkeletalBodySetups[Body->Index]->BoneName);
 				for (int32 ItemIdx = 0; ItemIdx < TreeElements.Num(); ++ItemIdx)
 				{
 					FPhATTreeInfo& Info = *TreeElements[ItemIdx];
@@ -382,7 +400,7 @@ bool TreeElemSelected(FTreeElemPtr TreeElem, TSharedPtr<FPhATSharedData> SharedD
 		{
 			for(int32 i=0; i<SharedData->SelectedBodies.Num(); ++i)
 			{
-				if (SharedData->PhysicsAsset->BodySetup[SharedData->SelectedBodies[i].Index]->BoneName == (*TreeElem).Name)
+				if (SharedData->PhysicsAsset->SkeletalBodySetups[SharedData->SelectedBodies[i].Index]->BoneName == (*TreeElem).Name)
 				{
 					return true;
 				}
@@ -448,16 +466,16 @@ void FPhAT::RefreshHierachyTree()
 	// if next event is selecting a bone to create a new body, Fill up tree with bone names
 	if (SharedData->EditingMode == FPhATSharedData::PEM_BodyEdit)
 	{
-		for (int32 i = 0; i < SharedData->PhysicsAsset->BodySetup.Num(); ++i)
+		for (int32 i = 0; i < SharedData->PhysicsAsset->SkeletalBodySetups.Num(); ++i)
 		{
-			const int32 BoneIndex = SharedData->EditorSkelComp->GetBoneIndex(SharedData->PhysicsAsset->BodySetup[i]->BoneName);
+			const int32 BoneIndex = SharedData->EditorSkelComp->GetBoneIndex(SharedData->PhysicsAsset->SkeletalBodySetups[i]->BoneName);
 			if (BoneIndex != INDEX_NONE)
 			{
-				const FKAggregateGeom& AggGeom = SharedData->PhysicsAsset->BodySetup[i]->AggGeom;
+				const FKAggregateGeom& AggGeom = SharedData->PhysicsAsset->SkeletalBodySetups[i]->AggGeom;
 
 				if (AggGeom.SphereElems.Num() + AggGeom.BoxElems.Num() + AggGeom.SphylElems.Num() + AggGeom.ConvexElems.Num() > 0)
 				{
-					TreeElements.Add(FTreeElemPtr(new FPhATTreeInfo(SharedData->PhysicsAsset->BodySetup[i]->BoneName, true, INDEX_NONE, BoneIndex)));
+					TreeElements.Add(FTreeElemPtr(new FPhATTreeInfo(SharedData->PhysicsAsset->SkeletalBodySetups[i]->BoneName, true, INDEX_NONE, BoneIndex)));
 				}
 			}
 		}
@@ -600,9 +618,9 @@ void FPhAT::PostUndo(bool bSuccess)
 
 void FPhAT::PostRedo( bool bSuccess )
 {
-	for (int32 BodyIdx=0; BodyIdx < SharedData->PhysicsAsset->BodySetup.Num(); ++BodyIdx)
+	for (int32 BodyIdx=0; BodyIdx < SharedData->PhysicsAsset->SkeletalBodySetups.Num(); ++BodyIdx)
 	{
-		UBodySetup* Body = SharedData->PhysicsAsset->BodySetup[BodyIdx];
+		UBodySetup* Body = SharedData->PhysicsAsset->SkeletalBodySetups[BodyIdx];
 		
 		bool bRecreate = false;
 		for (int32 ElemIdx=0; ElemIdx < Body->AggGeom.ConvexElems.Num(); ++ElemIdx)
@@ -642,6 +660,9 @@ void FPhAT::CreateInternalWidgets()
 	Properties->SetObject(SharedData->EditorSimOptions);
 	Properties->OnFinishedChangingProperties().AddSP(this, &FPhAT::OnFinishedChangingProperties);
 
+	PhysAssetProperties = PropertyModule.CreateDetailView(Args);
+	PhysAssetProperties->SetObject(SharedData->PhysicsAsset);
+
 	HierarchyControl = 
 	SNew(SBorder)
 	.Padding(8)
@@ -751,6 +772,98 @@ void FPhAT::ExtendToolbar()
 			return MenuBuilder.MakeWidget();
 		}
 
+		static TSharedRef< SWidget > FillPhysicalAnimationProfileOptions(TSharedRef<FUICommandList> InCommandList, TSharedPtr<FPhATSharedData> SharedData)
+		{
+			const bool bShouldCloseWindowAfterMenuSelection = true;
+			FMenuBuilder MenuBuilder(bShouldCloseWindowAfterMenuSelection, InCommandList);
+
+			const FPhATCommands& Commands = FPhATCommands::Get();
+
+			if(SharedData->PhysicsAsset)
+			{
+				TArray<FName> ProfileNames;
+				ProfileNames.Add(NAME_None);
+				ProfileNames.Append(SharedData->PhysicsAsset->GetPhysicalAnimationProfileNames(ProfileNames));
+				
+				for(FName ProfileName : ProfileNames)
+				{
+					FUIAction Action;
+					Action.ExecuteAction = FExecuteAction::CreateLambda( [SharedData, ProfileName]()
+					{
+						SharedData->PhysicsAsset->CurrentPhysicalAnimationProfileName = ProfileName;
+						for(USkeletalBodySetup* BS : SharedData->PhysicsAsset->SkeletalBodySetups)
+						{
+							if(FPhysicalAnimationProfile* Profile = BS->FindPhysicalAnimationProfile(ProfileName))
+							{
+								BS->CurrentPhysicalAnimationProfile = *Profile;
+							}
+						}
+					});
+
+					Action.GetActionCheckState = FGetActionCheckState::CreateLambda([SharedData, ProfileName]()
+					{
+						return SharedData->PhysicsAsset->CurrentPhysicalAnimationProfileName == ProfileName ? ECheckBoxState::Checked : ECheckBoxState::Unchecked;
+					});
+
+					auto SearchClickedLambda = [ProfileName, SharedData]()
+					{
+						SharedData->SetSelectedBody(nullptr);	//clear selection
+						for(int32 BSIndex = 0; BSIndex < SharedData->PhysicsAsset->SkeletalBodySetups.Num(); ++BSIndex)
+						{
+							const USkeletalBodySetup* BS = SharedData->PhysicsAsset->SkeletalBodySetups[BSIndex];
+							if(BS->FindPhysicalAnimationProfile(ProfileName))
+							{
+								SharedData->SetSelectedBodyAnyPrim(BSIndex, true);
+							}
+						}
+
+
+						return FReply::Handled();
+					};
+					
+					TSharedRef<SWidget> PhysAnimProfileButton = SNew(SHorizontalBox)
+						+ SHorizontalBox::Slot()
+						.FillWidth(1.f)
+						.VAlign(VAlign_Center)
+						[
+							SNew(STextBlock)
+							.Text(FText::FromString(ProfileName.ToString()))
+						]
+						+ SHorizontalBox::Slot()
+						.AutoWidth()
+						.Padding(2.f, 0.f, 0.f, 0.f)
+						.VAlign(VAlign_Center)
+						[
+							SNew(SButton)
+							.ButtonStyle(FEditorStyle::Get(), "HoverHintOnly")
+							.OnClicked_Lambda(SearchClickedLambda)
+							[
+								SNew(SBox)
+								.WidthOverride(MultiBoxConstants::MenuIconSize)
+								.HeightOverride(MultiBoxConstants::MenuIconSize)
+								.Visibility_Lambda([ProfileName](){ return ProfileName == NAME_None ? EVisibility::Collapsed : EVisibility::Visible; })
+								[
+									SNew(SImage)
+									.Image(FSlateIcon(FEditorStyle::GetStyleSetName(), "Symbols.SearchGlass").GetIcon())
+								]
+								
+							]
+						];
+
+
+					//MenuBuilder.AddMenuEntry( FText::FromString(ProfileName.ToString()), FText::GetEmpty(), FSlateIcon(), Action, NAME_None, EUserInterfaceActionType::Check);
+					MenuBuilder.AddMenuEntry(Action, PhysAnimProfileButton, NAME_None, TAttribute<FText>(), EUserInterfaceActionType::Check);
+				}
+				
+				MenuBuilder.AddMenuSeparator();
+				//MenuBuilder.AddMenuEntry(Commands.EditPhysicalAnimations);
+
+			}
+
+			
+			return MenuBuilder.MakeWidget();
+		}
+
 		static TSharedRef< SWidget > FillEditMode(TSharedRef<FUICommandList> InCommandList)
 		{
 			const bool bShouldCloseWindowAfterMenuSelection = true;
@@ -764,8 +877,9 @@ void FPhAT::ExtendToolbar()
 			return MenuBuilder.MakeWidget();
 		}
 
-		static void FillToolbar(FToolBarBuilder& ToolbarBuilder, TSharedRef<SWidget> PhATAnimation, FPhATSharedData::EPhATEditingMode InPhATEditingMode, FPhAT * Phat )
+		static void FillToolbar(FToolBarBuilder& ToolbarBuilder, TSharedRef<SWidget> PhATAnimation, TSharedPtr<FPhATSharedData> SharedData, FPhAT * Phat )
 		{
+			const FPhATSharedData::EPhATEditingMode InPhATEditingMode = SharedData->EditingMode;
 			const FPhATCommands& Commands = FPhATCommands::Get();
 			TSharedRef<FUICommandList> InCommandList = Phat->GetToolkitCommands();
 
@@ -793,6 +907,7 @@ void FPhAT::ExtendToolbar()
 					true
 					);
 			}
+
 			ToolbarBuilder.EndSection();
 
 			//selected simulation
@@ -840,6 +955,23 @@ void FPhAT::ExtendToolbar()
 				ToolbarBuilder.EndSection();
 			}
 
+
+			FUIAction PhysicalAnimationProfileAction;
+			PhysicalAnimationProfileAction.GetActionCheckState = FGetActionCheckState::CreateLambda([SharedData]() { return SharedData->PhysicsAsset->CurrentPhysicalAnimationProfileName == NAME_None ? ECheckBoxState::Unchecked : ECheckBoxState::Checked; });
+
+			ToolbarBuilder.BeginSection("PhATPhysicalAnimation");
+			{
+				ToolbarBuilder.AddComboButton(
+					PhysicalAnimationProfileAction,
+					FOnGetContent::CreateStatic(&FillPhysicalAnimationProfileOptions, InCommandList, SharedData),
+					LOCTEXT("PhysicalAnimationProfile_Label", "Profile "),
+					LOCTEXT("PhysicalAnimationProfileToolTip", "Changes the Physical Animation profile"),
+					FSlateIcon(FEditorStyle::GetStyleSetName(), "Persona.ImportAnimation"),
+					false
+					);
+			}
+			ToolbarBuilder.EndSection();
+
 			ToolbarBuilder.BeginSection("PhATPlayAnimation");
 			{
 				ToolbarBuilder.AddToolBarButton(Commands.PlayAnimation);
@@ -873,37 +1005,33 @@ void FPhAT::ExtendToolbar()
 	}
 
 	ToolbarExtender = MakeShareable(new FExtender);
-
-	TSharedRef<SWidget> PhATAnimation = SNew(SBox)		
-		.WidthOverride(250)
+	
+	TSharedRef<SWidget> PhATAnimation = SNew(SHorizontalBox)
+		+ SHorizontalBox::Slot()
+		.AutoWidth()
+		.VAlign(VAlign_Center)
 		[
-			SNew(SHorizontalBox)
-			+SHorizontalBox::Slot()
-			.AutoWidth()
-			.VAlign(VAlign_Center)
-			[
-				SNew(STextBlock)
-				.Text( LOCTEXT( "PhATToolbarAnimation", "Animation: " ) )
-			]
-			+SHorizontalBox::Slot()
-			.FillWidth(1.0f)
-			[
-				SNew(SContentReference)
-				.WidthOverride(80.0f)
-				.AllowSelectingNewAsset(true)
-				.AssetReference(this, &FPhAT::GetSelectedAnimation)
-				.AllowedClass(UAnimSequence::StaticClass())
-				.OnShouldFilterAsset(this, &FPhAT::ShouldFilterAssetBasedOnSkeleton)
-				.OnSetReference(this, &FPhAT::AnimationSelectionChanged)
-				.IsEnabled(this, &FPhAT::IsToggleSimulation)
-			]
+			SNew(STextBlock)
+			.Text(LOCTEXT("PhATToolbarAnimation", "Animation: "))
+		]
+		+ SHorizontalBox::Slot()
+		.FillWidth(1.0f)
+		[
+			SNew(SContentReference)
+			.WidthOverride(80.0f)
+			.AllowSelectingNewAsset(true)
+			.AssetReference(this, &FPhAT::GetSelectedAnimation)
+			.AllowedClass(UAnimSequence::StaticClass())
+			.OnShouldFilterAsset(this, &FPhAT::ShouldFilterAssetBasedOnSkeleton)
+			.OnSetReference(this, &FPhAT::AnimationSelectionChanged)
+			.IsEnabled(this, &FPhAT::IsToggleSimulation)
 		];
 
 	ToolbarExtender->AddToolBarExtension(
 		"Asset",
 		EExtensionHook::After,
 		GetToolkitCommands(),
-		FToolBarExtensionDelegate::CreateStatic( &Local::FillToolbar, PhATAnimation, SharedData->EditingMode, this)
+		FToolBarExtensionDelegate::CreateStatic( &Local::FillToolbar, PhATAnimation, SharedData, this)
 		);
 
 	AddToolbarExtender(ToolbarExtender);
@@ -1342,6 +1470,12 @@ void FPhAT::BindCommands()
 		FCanExecuteAction::CreateSP(this, &FPhAT::IsNotSimulation)
 		);
 
+	/*ToolkitCommands->MapAction(
+		Commands.EditPhysicalAnimations,
+		FExecuteAction::CreateSP(this, &FPhAT::EditPhysicalAnimations),
+		FCanExecuteAction::CreateSP(this, &FPhAT::IsNotSimulation)
+		);*/
+
 	// record animation
 	ToolkitCommands->MapAction(
 		Commands.RecordAnimation,
@@ -1357,6 +1491,11 @@ void FPhAT::Mirror()
 	SharedData->Mirror();
 }
 
+void FPhAT::EditPhysicalAnimations()
+{
+	
+}
+
 TSharedRef<ITableRow> FPhAT::OnGenerateRowForTree(FTreeElemPtr Item, const TSharedRef<STableViewBase>& OwnerTable)
 {
 	class SHoverDetectTableRow : public STableRow<FTreeElemPtr>
@@ -1415,12 +1554,12 @@ void FPhAT::OnGetChildrenForTree(FTreeElemPtr Parent, TArray<FTreeElemPtr>& OutC
 {
 	if (SharedData->EditingMode == FPhATSharedData::PEM_BodyEdit)
 	{
-		for (int32 i = 0; i < SharedData->PhysicsAsset->BodySetup.Num(); ++i)
+		for (int32 i = 0; i < SharedData->PhysicsAsset->SkeletalBodySetups.Num(); ++i)
 		{
-			if (SharedData->PhysicsAsset->BodySetup[i]->BoneName == (*Parent).Name)
+			if (SharedData->PhysicsAsset->SkeletalBodySetups[i]->BoneName == (*Parent).Name)
 			{
-				int32 BoneIndex = SharedData->EditorSkelComp->GetBoneIndex(SharedData->PhysicsAsset->BodySetup[i]->BoneName);
-				FKAggregateGeom& AggGeom = SharedData->PhysicsAsset->BodySetup[i]->AggGeom;
+				int32 BoneIndex = SharedData->EditorSkelComp->GetBoneIndex(SharedData->PhysicsAsset->SkeletalBodySetups[i]->BoneName);
+				FKAggregateGeom& AggGeom = SharedData->PhysicsAsset->SkeletalBodySetups[i]->AggGeom;
 
 				if (AggGeom.SphereElems.Num() + AggGeom.BoxElems.Num() + AggGeom.SphylElems.Num() + AggGeom.ConvexElems.Num() > 1)
 				{
@@ -1522,11 +1661,11 @@ void FPhAT::OnTreeSelectionChanged(FTreeElemPtr TreeElem, ESelectInfo::Type Sele
 		{
 			if (ObjIndex != INDEX_NONE)
 			{
-				for (int32 i = 0; i < SharedData->PhysicsAsset->BodySetup.Num(); ++i)
+				for (int32 i = 0; i < SharedData->PhysicsAsset->SkeletalBodySetups.Num(); ++i)
 				{
-					if (SharedData->PhysicsAsset->BodySetup[i]->BoneName == (*SelectedElem).Name)
+					if (SharedData->PhysicsAsset->SkeletalBodySetups[i]->BoneName == (*SelectedElem).Name)
 					{
-						FKAggregateGeom& AggGeom = SharedData->PhysicsAsset->BodySetup[i]->AggGeom;
+						FKAggregateGeom& AggGeom = SharedData->PhysicsAsset->SkeletalBodySetups[i]->AggGeom;
 
 						//select all primitives
 						for(int32 j=0; j<AggGeom.BoxElems.Num(); ++j)
@@ -1853,7 +1992,7 @@ void FPhAT::AddNewPrimitive(EKCollisionPrimitiveType InPrimitiveType, bool bCopy
 
 		for(int32 i=0; i<NewSelection.Num(); ++i)
 		{
-			UBodySetup* BodySetup = SharedData->PhysicsAsset->BodySetup[NewSelection[i].Index];
+			UBodySetup* BodySetup = SharedData->PhysicsAsset->SkeletalBodySetups[NewSelection[i].Index];
 			EKCollisionPrimitiveType PrimitiveType;
 			if (bCopySelected)
 			{
@@ -1994,14 +2133,14 @@ void FPhAT::SetBodiesBelowPhysicsType( EPhysicsType InPhysicsType, const TArray<
 	for(int32 i=0; i<Indices.Num(); ++i)
 	{
 		// Get the index of this body
-		UBodySetup* BaseSetup = SharedData->PhysicsAsset->BodySetup[Indices[i]];
+		UBodySetup* BaseSetup = SharedData->PhysicsAsset->SkeletalBodySetups[Indices[i]];
 		SharedData->PhysicsAsset->GetBodyIndicesBelow(BelowBodies, BaseSetup->BoneName, SharedData->EditorSkelMesh);
 	}
 
 	for (int32 i = 0; i < BelowBodies.Num(); ++i)
 	{
 		int32 BodyIndex = BelowBodies[i];
-		UBodySetup* BodySetup = SharedData->PhysicsAsset->BodySetup[BodyIndex];
+		UBodySetup* BodySetup = SharedData->PhysicsAsset->SkeletalBodySetups[BodyIndex];
 		if (bMarkAsDirty)
 		{
 			BodySetup->Modify();
@@ -2061,9 +2200,9 @@ void FPhAT::OnChangeDefaultMesh()
 	{
 		// See if any bones are missing from the skeletal mesh we are trying to use
 		// @todo Could do more here - check for bone lengths etc. Maybe modify asset?
-		for (int32 i = 0; i <SharedData->PhysicsAsset->BodySetup.Num(); ++i)
+		for (int32 i = 0; i <SharedData->PhysicsAsset->SkeletalBodySetups.Num(); ++i)
 		{
-			FName BodyName = SharedData->PhysicsAsset->BodySetup[i]->BoneName;
+			FName BodyName = SharedData->PhysicsAsset->SkeletalBodySetups[i]->BoneName;
 			int32 BoneIndex = NewSkelMesh->RefSkeleton.FindBoneIndex(BodyName);
 			if (BoneIndex == INDEX_NONE)
 			{
@@ -2113,7 +2252,7 @@ void FPhAT::OnResetEntireAsset()
 		SharedData->SetSelectedConstraint(INDEX_NONE);	
 
 		// Empty current asset data.
-		SharedData->PhysicsAsset->BodySetup.Empty();
+		SharedData->PhysicsAsset->SkeletalBodySetups.Empty();
 		SharedData->PhysicsAsset->BodySetupIndexMap.Empty();
 		SharedData->PhysicsAsset->ConstraintSetup.Empty();
 
@@ -2150,7 +2289,7 @@ void FPhAT::OnResetBoneCollision()
 		const FScopedTransaction Transaction( NSLOCTEXT("PhAT", "ResetBoneCollision", "Reset Bone Collision") );
 		for(int32 i=0; i<SharedData->SelectedBodies.Num(); ++i)
 		{
-			UBodySetup* BodySetup = SharedData->PhysicsAsset->BodySetup[SharedData->SelectedBodies[i].Index];
+			UBodySetup* BodySetup = SharedData->PhysicsAsset->SkeletalBodySetups[SharedData->SelectedBodies[i].Index];
 			check(BodySetup);
 			BodySetup->Modify();
 
@@ -2186,9 +2325,9 @@ void FPhAT::OnApplyPhysicalMaterial()
 
 	if (SelectedPhysMaterial)
 	{
-		for (int32 BodyIdx=0; BodyIdx<SharedData->PhysicsAsset->BodySetup.Num(); BodyIdx++)
+		for (int32 BodyIdx=0; BodyIdx<SharedData->PhysicsAsset->SkeletalBodySetups.Num(); BodyIdx++)
 		{
-			UBodySetup* BodySetup = SharedData->PhysicsAsset->BodySetup[BodyIdx];
+			UBodySetup* BodySetup = SharedData->PhysicsAsset->SkeletalBodySetups[BodyIdx];
 			BodySetup->Modify();
 			BodySetup->PhysMaterial = SelectedPhysMaterial;
 		}
@@ -2291,12 +2430,12 @@ bool FPhAT::IsInstanceProperties() const
 void FPhAT::FixPhysicsState()
 {
 	UPhysicsAsset * PhysicsAsset = SharedData->PhysicsAsset;
-	TArray<UBodySetup*> & BodySetup = PhysicsAsset->BodySetup;
+	TArray<USkeletalBodySetup*>& BodySetup = PhysicsAsset->SkeletalBodySetups;
 
 	if(!SharedData->bRunningSimulation)
 	{
 		PhysicsTypeState.Reset();
-		for(int32 i=0; i<SharedData->PhysicsAsset->BodySetup.Num(); ++i)
+		for(int32 i=0; i<SharedData->PhysicsAsset->SkeletalBodySetups.Num(); ++i)
 		{
 			PhysicsTypeState.Add(BodySetup[i]->PhysicsType);
 		}
@@ -2381,10 +2520,10 @@ void FPhAT::OnSelectedSimulation()
 	{
 
 		UPhysicsAsset * PhysicsAsset = SharedData->PhysicsAsset;
-		TArray<UBodySetup*> & BodySetup = PhysicsAsset->BodySetup;
+		TArray<USkeletalBodySetup*>& BodySetup = PhysicsAsset->SkeletalBodySetups;
 
 		//first we fix all the bodies
-		for(int32 i=0; i<SharedData->PhysicsAsset->BodySetup.Num(); ++i)
+		for(int32 i=0; i<SharedData->PhysicsAsset->SkeletalBodySetups.Num(); ++i)
 		{
 			BodySetup[i]->PhysicsType = PhysType_Kinematic;
 		}
@@ -2746,7 +2885,7 @@ void FPhAT::OnSetBodyPhysicsType( EPhysicsType InPhysicsType )
 		
 		for(int32 i=0; i<SharedData->SelectedBodies.Num(); ++i)
 		{
-			UBodySetup* BodySetup = SharedData->PhysicsAsset->BodySetup[SharedData->SelectedBodies[i].Index];
+			UBodySetup* BodySetup = SharedData->PhysicsAsset->SkeletalBodySetups[SharedData->SelectedBodies[i].Index];
 			BodySetup->Modify();
 			BodySetup->PhysicsType = InPhysicsType;
 		}
@@ -2758,7 +2897,7 @@ bool FPhAT::IsBodyPhysicsType( EPhysicsType InPhysicsType )
 {
 	for(int32 i=0; i<SharedData->SelectedBodies.Num(); ++i)
 	{
-		UBodySetup* BodySetup = SharedData->PhysicsAsset->BodySetup[SharedData->SelectedBodies[i].Index];
+		UBodySetup* BodySetup = SharedData->PhysicsAsset->SkeletalBodySetups[SharedData->SelectedBodies[i].Index];
 		if(BodySetup->PhysicsType == InPhysicsType)
 		{
 			return true;
@@ -2779,7 +2918,7 @@ void FPhAT::OnDeleteBody()
 
 		for(int32 i=0; i<SharedData->SelectedBodies.Num(); ++i)
 		{
-			BodySetups.Add( SharedData->PhysicsAsset->BodySetup[SharedData->SelectedBodies[i].Index] );
+			BodySetups.Add( SharedData->PhysicsAsset->SkeletalBodySetups[SharedData->SelectedBodies[i].Index] );
 		}
 
 		const FScopedTransaction Transaction( LOCTEXT( "DeleteBodies", "Delete Bodies" ) );
@@ -2804,7 +2943,7 @@ void FPhAT::OnDeleteAllBodiesBelow()
 
 	for (FPhATSharedData::FSelection SelectedBody : SharedData->SelectedBodies)
 	{
-		UBodySetup* BaseSetup = SharedData->PhysicsAsset->BodySetup[SelectedBody.Index];
+		UBodySetup* BaseSetup = SharedData->PhysicsAsset->SkeletalBodySetups[SelectedBody.Index];
 		
 		// Build a list of BodySetups below this one
 		TArray<int32> BelowBodies;
@@ -2812,7 +2951,7 @@ void FPhAT::OnDeleteAllBodiesBelow()
 
 		for (const int32 BodyIndex : BelowBodies)
 		{
-			UBodySetup* BodySetup = SharedData->PhysicsAsset->BodySetup[BodyIndex];
+			UBodySetup* BodySetup = SharedData->PhysicsAsset->SkeletalBodySetups[BodyIndex];
 			BodySetups.Add(BodySetup);
 		}
 	}
@@ -2910,7 +3049,7 @@ void FPhAT::OnFocusSelection()
 		{
 			if(SharedData->GetSelectedBody())
 			{
-				int32 BoneIdx = SharedData->EditorSkelComp->GetBoneIndex(SharedData->PhysicsAsset->BodySetup[SharedData->GetSelectedBody()->Index]->BoneName);
+				int32 BoneIdx = SharedData->EditorSkelComp->GetBoneIndex(SharedData->PhysicsAsset->SkeletalBodySetups[SharedData->GetSelectedBody()->Index]->BoneName);
 				FMatrix BoneTransform = SharedData->EditorSkelComp->GetBoneMatrix(BoneIdx);
 				FBoxSphereBounds Bounds(BoneTransform.GetOrigin(), FVector(20), 20);
 
@@ -3004,7 +3143,7 @@ void FPhAT::OnAssetSelectedFromStaticMeshAssetPicker( const FAssetData& AssetDat
 		SharedData->PhysicsAsset->Modify();
 
 		// Build a list of BodySetups below this one
-		UBodySetup* BaseSetup = SharedData->PhysicsAsset->BodySetup[SharedData->GetSelectedBody()->Index];
+		UBodySetup* BaseSetup = SharedData->PhysicsAsset->SkeletalBodySetups[SharedData->GetSelectedBody()->Index];
 		BaseSetup->Modify();
 	
 		UStaticMesh* SM = Cast<UStaticMesh>(AssetData.GetAsset());
@@ -3063,14 +3202,14 @@ void FPhAT::OnSelectAll()
 		SharedData->SetSelectedBody(NULL);
 	
 		//go through every body and add every geom
-		for (int32 i = 0; i <PhysicsAsset->BodySetup.Num(); ++i)
+		for (int32 i = 0; i <PhysicsAsset->SkeletalBodySetups.Num(); ++i)
 		{
-			int32 BoneIndex = SharedData->EditorSkelComp->GetBoneIndex(PhysicsAsset->BodySetup[i]->BoneName);
+			int32 BoneIndex = SharedData->EditorSkelComp->GetBoneIndex(PhysicsAsset->SkeletalBodySetups[i]->BoneName);
 
 			// If we found a bone for it, add all geom
 			if (BoneIndex != INDEX_NONE)
 			{
-				FKAggregateGeom* AggGeom = &PhysicsAsset->BodySetup[i]->AggGeom;
+				FKAggregateGeom* AggGeom = &PhysicsAsset->SkeletalBodySetups[i]->AggGeom;
 
 				for (int32 j = 0; j <AggGeom->SphereElems.Num(); ++j)
 				{
diff --git a/Engine/Source/Editor/PhAT/Private/PhAT.h b/Engine/Source/Editor/PhAT/Private/PhAT.h
--- a/Engine/Source/Editor/PhAT/Private/PhAT.h
+++ b/Engine/Source/Editor/PhAT/Private/PhAT.h
@@ -236,6 +236,8 @@ private:
 	
 	void Mirror();
 
+	void EditPhysicalAnimations();
+
 	//menu commands
 	void OnSelectAll();
 	void SetHierarchyFilter(EPhatHierarchyFilterMode Mode);
@@ -261,6 +263,9 @@ private:
 	/** Properties Tab */
 	TSharedPtr<class IDetailsView> Properties;
 
+	/** Physics asset properties tab */
+	TSharedPtr<class IDetailsView> PhysAssetProperties;
+
 	/** Bone Hierarchy Tree View */
 	TSharedPtr< STreeView<FTreeElemPtr> > Hierarchy;
 	TSharedPtr<SBorder> HierarchyControl;
@@ -274,6 +279,8 @@ private:
 	/** Data and methods shared across multiple classes */
 	TSharedPtr<FPhATSharedData> SharedData;
 
+	TSharedPtr<STextComboBox> PhysicalAnimationComboBox;
+
 	/** Toolbar extender - used repeatedly as the body/constraints mode will remove/add this when changed */
 	TSharedPtr<FExtender> ToolbarExtender;
 
@@ -289,6 +296,9 @@ private:
 	/** True if we want to only simulate from selected body/constraint down*/
 	bool SelectedSimulation;
 
+	/** Holds the array of strings the UI needs to show physical animation profiles for the editing physics asset*/
+	TArray<TSharedPtr<FString>> PhysicalAnimationProfiles;
+
 	/** Determines which simulation mode wer're currently in */
 	EPhATSimulationMode SimulationMode;
 
diff --git a/Engine/Source/Editor/PhAT/Private/PhATActions.cpp b/Engine/Source/Editor/PhAT/Private/PhATActions.cpp
--- a/Engine/Source/Editor/PhAT/Private/PhATActions.cpp
+++ b/Engine/Source/Editor/PhAT/Private/PhATActions.cpp
@@ -12,15 +12,15 @@ void FPhATCommands::RegisterCommands()
 	UI_COMMAND(ResetEntireAsset, "Reset Asset", "Opens Body Creation Settings And Replaces Asset Using The New Settings", EUserInterfaceActionType::Button, FInputChord());
 	UI_COMMAND(ResetBoneCollision, "Reset", "Reset Bone Collision", EUserInterfaceActionType::Button, FInputChord());
 	UI_COMMAND(ApplyPhysicalMaterial, "Apply PhysMat", "Applies Currently Selected Physical Material In Content Browser To All Bodies", EUserInterfaceActionType::Button, FInputChord());
-	UI_COMMAND(EditingMode_Body, "Body Mode", "Body Editing Mode", EUserInterfaceActionType::RadioButton, FInputChord());
-	UI_COMMAND(EditingMode_Constraint, "Constraint Mode", "Constraint Editing Mode", EUserInterfaceActionType::RadioButton, FInputChord());
+	UI_COMMAND(EditingMode_Body, "Bodies", "Body Editing Mode", EUserInterfaceActionType::RadioButton, FInputChord());
+	UI_COMMAND(EditingMode_Constraint, "Constraints", "Constraint Editing Mode", EUserInterfaceActionType::RadioButton, FInputChord());
 	UI_COMMAND(CopyProperties, "Copy Properties", "Copy Properties: Copy Properties Of Currently Selected Object To Next Selected Object", EUserInterfaceActionType::Button, FInputChord(EModifierKey::Control, EKeys::C));
 	UI_COMMAND(PasteProperties, "Paste Properties", "Paste Properties: Copy Properties Of Currently Selected Object To Next Selected Object", EUserInterfaceActionType::Button, FInputChord(EModifierKey::Control, EKeys::V));
 	UI_COMMAND(InstanceProperties, "Instance Properties", "Instance Properties: Displays Instance Properties When In Body Editing Mode", EUserInterfaceActionType::ToggleButton, FInputChord(EKeys::I));
 	UI_COMMAND(RepeatLastSimulation, "Simulate", "Previews Physics Simulation", EUserInterfaceActionType::RadioButton, FInputChord(EKeys::Enter));
 	UI_COMMAND(SimulationNormal, "Real Simulation", "Previews Normal Physics Simulation", EUserInterfaceActionType::RadioButton, FInputChord());
 	UI_COMMAND(SimulationNoGravity, "No Gravity Simulation", "Run Physics Simulation without gravity. Use this to debug issues with your ragdoll. If the setup is correct, the asset should not move!", EUserInterfaceActionType::RadioButton, FInputChord());
-	UI_COMMAND(ToggleSelectedSimulation, "Selected Simulation", "Run Physics Simulation on selected objects. Use this to tune  specific parts of your ragdoll.", EUserInterfaceActionType::ToggleButton, FInputChord());
+	UI_COMMAND(ToggleSelectedSimulation, "Selected", "Run Physics Simulation on selected objects. Use this to tune  specific parts of your ragdoll.", EUserInterfaceActionType::ToggleButton, FInputChord());
 	UI_COMMAND(MeshRenderingMode_Solid, "Solid", "Solid Mesh Rendering Mode", EUserInterfaceActionType::RadioButton, FInputChord());
 	UI_COMMAND(MeshRenderingMode_Wireframe, "Wireframe", "Wireframe Mesh Rendering Mode", EUserInterfaceActionType::RadioButton, FInputChord());
 	UI_COMMAND(MeshRenderingMode_None, "None", "No Mesh Rendering Mode", EUserInterfaceActionType::RadioButton, FInputChord());
@@ -51,7 +51,7 @@ void FPhATCommands::RegisterCommands()
 	UI_COMMAND(ConvertToPrismatic, "To Prismatic", "Convert Selected Constraint To Prismatic", EUserInterfaceActionType::Button, FInputChord());
 	UI_COMMAND(ConvertToSkeletal, "To Skeletal", "Convert Selected Constraint To Skeletal", EUserInterfaceActionType::Button, FInputChord());
 	UI_COMMAND(DeleteConstraint, "Delete", "Delete Selected Constraint", EUserInterfaceActionType::Button, FInputChord());
-	UI_COMMAND(PlayAnimation, "PlayAnimation", "Play Animation", EUserInterfaceActionType::ToggleButton, FInputChord());
+	UI_COMMAND(PlayAnimation, "Play", "Play", EUserInterfaceActionType::ToggleButton, FInputChord());
 	UI_COMMAND(RecordAnimation, "RecordAnimation", "Record Animation", EUserInterfaceActionType::ToggleButton, FInputChord());
 	UI_COMMAND(ShowSkeleton, "Skeleton", "Show Skeleton", EUserInterfaceActionType::ToggleButton, FInputChord());
 	UI_COMMAND(MakeBodyKinematic, "Kinematic", "Make Body Kinematic", EUserInterfaceActionType::RadioButton, FInputChord());
@@ -77,6 +77,8 @@ void FPhATCommands::RegisterCommands()
 	UI_COMMAND(BackView, "Back", "Orthographic view from back", EUserInterfaceActionType::Button, FInputChord());
 	UI_COMMAND(Mirror, "Mirror", "Finds the body on the other side and duplicates constraint and body", EUserInterfaceActionType::Button, FInputChord(EKeys::M));
 
+	//UI_COMMAND(EditPhysicalAnimations, "Edit...", "Edits the physical animation profiles", EUserInterfaceActionType::Button, FInputChord());
+
 
 	UI_COMMAND(SelectionLock, "Lock Selection", "", EUserInterfaceActionType::Button, FInputChord(EKeys::X));
 
diff --git a/Engine/Source/Editor/PhAT/Private/PhATPreviewViewportClient.cpp b/Engine/Source/Editor/PhAT/Private/PhATPreviewViewportClient.cpp
--- a/Engine/Source/Editor/PhAT/Private/PhATPreviewViewportClient.cpp
+++ b/Engine/Source/Editor/PhAT/Private/PhATPreviewViewportClient.cpp
@@ -15,6 +15,7 @@
 #include "PhysicsEngine/BodySetup.h"
 #include "Engine/Font.h"
 #include "PhysicsEngine/PhysicsConstraintTemplate.h"
+#include "PhysicsEngine/PhysicsAsset.h"
 #include "PhysicsEngine/PhysicsSettings.h"
 #include "PhysicsEngine/PhysicsHandleComponent.h"
 #include "DrawDebugHelpers.h"
@@ -122,9 +123,9 @@ void FPhATEdPreviewViewportClient::DrawCanvas( FViewport& InViewport, FSceneView
 	// Write body/constraint count at top.
 	FString StatusString = FText::Format(
 		NSLOCTEXT("UnrealEd", "BodiesConstraints_F", "{0} Bodies  {1} Considered for bounds  {2} Ratio  {3} Constraints"),
-		FText::AsNumber(SharedData->PhysicsAsset->BodySetup.Num()),
+		FText::AsNumber(SharedData->PhysicsAsset->SkeletalBodySetups.Num()),
 		FText::AsNumber(SharedData->PhysicsAsset->BoundsBodies.Num()),
-		FText::AsNumber(static_cast<float>(SharedData->PhysicsAsset->BoundsBodies.Num())/static_cast<float>(SharedData->PhysicsAsset->BodySetup.Num())),
+		FText::AsNumber(static_cast<float>(SharedData->PhysicsAsset->BoundsBodies.Num())/static_cast<float>(SharedData->PhysicsAsset->SkeletalBodySetups.Num())),
 		FText::AsNumber(SharedData->PhysicsAsset->ConstraintSetup.Num()) ).ToString();
 
 	TextItem.Text = FText::FromString( StatusString );
@@ -188,7 +189,7 @@ void FPhATEdPreviewViewportClient::DrawCanvas( FViewport& InViewport, FSceneView
 				for(int32 j=0; j< SharedData->SelectedBodies.Num(); ++j)
 				{
 					int32 SelectedBodyIndex = SharedData->SelectedBodies[j].Index;
-					FName SelectedBoneName = SharedData->PhysicsAsset->BodySetup[SelectedBodyIndex]->BoneName;
+					FName SelectedBoneName = SharedData->PhysicsAsset->SkeletalBodySetups[SelectedBodyIndex]->BoneName;
 					if(SelectedBoneName == BoneName)
 					{
 						BoneNameColor = FColor::Green;
@@ -473,7 +474,7 @@ bool FPhATEdPreviewViewportClient::InputWidgetDelta( FViewport* InViewport, EAxi
 			float BoneScale = 1.f;
 			if (SharedData->EditingMode == FPhATSharedData::PEM_BodyEdit) /// BODY EDITING ///
 			{
-				int32 BoneIndex = SharedData->EditorSkelComp->GetBoneIndex(SharedData->PhysicsAsset->BodySetup[SelectedObject.Index]->BoneName);
+				int32 BoneIndex = SharedData->EditorSkelComp->GetBoneIndex(SharedData->PhysicsAsset->SkeletalBodySetups[SelectedObject.Index]->BoneName);
 
 				FTransform BoneTM = SharedData->EditorSkelComp->GetBoneTransform(BoneIndex);
 				BoneScale = BoneTM.GetScale3D().GetAbsMax();
@@ -597,7 +598,7 @@ FVector FPhATEdPreviewViewportClient::GetWidgetLocation() const
 			return FVector::ZeroVector;
 		}
 
-		int32 BoneIndex = SharedData->EditorSkelComp->GetBoneIndex(SharedData->PhysicsAsset->BodySetup[SharedData->GetSelectedBody()->Index]->BoneName);
+		int32 BoneIndex = SharedData->EditorSkelComp->GetBoneIndex(SharedData->PhysicsAsset->SkeletalBodySetups[SharedData->GetSelectedBody()->Index]->BoneName);
 
 		FTransform BoneTM = SharedData->EditorSkelComp->GetBoneTransform(BoneIndex);
 		const float Scale = BoneTM.GetScale3D().GetAbsMax();
@@ -628,7 +629,7 @@ FMatrix FPhATEdPreviewViewportClient::GetWidgetCoordSystem() const
 				return FMatrix::Identity;
 			}
 
-			int32 BoneIndex = SharedData->EditorSkelComp->GetBoneIndex(SharedData->PhysicsAsset->BodySetup[SharedData->GetSelectedBody()->Index]->BoneName);
+			int32 BoneIndex = SharedData->EditorSkelComp->GetBoneIndex(SharedData->PhysicsAsset->SkeletalBodySetups[SharedData->GetSelectedBody()->Index]->BoneName);
 
 			FTransform BoneTM = SharedData->EditorSkelComp->GetBoneTransform(BoneIndex);
 			BoneTM.RemoveScaling();
@@ -681,6 +682,11 @@ void FPhATEdPreviewViewportClient::Tick(float DeltaSeconds)
 		SharedData->EditorSkelComp->AnimationSpaceBases = SharedData->EditorSkelComp->GetSpaceBases();
 		SharedData->EditorSkelComp->SetPhysicsBlendWeight(SharedData->EditorSimOptions->PhysicsBlend);
 		SharedData->EditorSkelComp->bUpdateJointsFromAnimation = SharedData->EditorSimOptions->bUpdateJointsFromAnimation;
+
+		static FPhysicalAnimationData EmptyProfile;
+
+		SharedData->PhysicalAnimationComponent->ApplyPhysicalAnimationProfileBelow(NAME_None, SharedData->PhysicsAsset->CurrentPhysicalAnimationProfileName, /*Include Self=*/true, /*Clear Not Found=*/true);
+		
 	}
 
 	World->Tick(LEVELTICK_All, DeltaSeconds * SharedData->EditorSimOptions->TimeDilation);
@@ -742,7 +748,7 @@ void FPhATEdPreviewViewportClient::StartManipulating(EAxisList::Type Axis, const
 		GEditor->BeginTransaction( NSLOCTEXT("UnrealEd", "MoveElement", "Move Element") );
 		for(int32 i=0; i<SharedData->SelectedBodies.Num(); ++i)
 		{
-			SharedData->PhysicsAsset->BodySetup[SharedData->SelectedBodies[i].Index]->Modify();
+			SharedData->PhysicsAsset->SkeletalBodySetups[SharedData->SelectedBodies[i].Index]->Modify();
 			SharedData->SelectedBodies[i].ManipulateTM = FTransform::Identity;
 		}
 
@@ -783,7 +789,7 @@ void FPhATEdPreviewViewportClient::EndManipulating()
 				for(int32 i=0; i<SharedData->SelectedBodies.Num(); ++i)
 				{
 					FPhATSharedData::FSelection & SelectedObject = SharedData->SelectedBodies[i];
-					UBodySetup* BodySetup = SharedData->PhysicsAsset->BodySetup[SelectedObject.Index];
+					UBodySetup* BodySetup = SharedData->PhysicsAsset->SkeletalBodySetups[SelectedObject.Index];
 
 					FKAggregateGeom* AggGeom = &BodySetup->AggGeom;
 
@@ -833,7 +839,7 @@ void FPhATEdPreviewViewportClient::SimMousePress(FViewport* InViewport, bool bCo
 	if (bHit)
 	{
 		check(Result.Item != INDEX_NONE);
-		FName BoneName = SharedData->PhysicsAsset->BodySetup[Result.Item]->BoneName;
+		FName BoneName = SharedData->PhysicsAsset->SkeletalBodySetups[Result.Item]->BoneName;
 
 		//UE_LOG(LogPhysics, Warning, TEXT("Hit Bone Name (%s)"), *BoneName.ToString());
 
@@ -951,7 +957,7 @@ void FPhATEdPreviewViewportClient::ModifyPrimitiveSize(int32 BodyIndex, EKCollis
 {
 	check(SharedData->GetSelectedBody());
 
-	FKAggregateGeom* AggGeom = &SharedData->PhysicsAsset->BodySetup[BodyIndex]->AggGeom;
+	FKAggregateGeom* AggGeom = &SharedData->PhysicsAsset->SkeletalBodySetups[BodyIndex]->AggGeom;
 
 	if (PrimType == KPT_Sphere)
 	{
diff --git a/Engine/Source/Editor/PhAT/Private/PhATRender.cpp b/Engine/Source/Editor/PhAT/Private/PhATRender.cpp
--- a/Engine/Source/Editor/PhAT/Private/PhATRender.cpp
+++ b/Engine/Source/Editor/PhAT/Private/PhATRender.cpp
@@ -6,6 +6,7 @@
 #include "PhATEdSkeletalMeshComponent.h"
 #include "PhysicsEngine/BodySetup.h"
 #include "PhysicsEngine/PhysicsConstraintTemplate.h"
+#include "PhysicsEngine/PhysicsAsset.h"
 
 UPhATEdSkeletalMeshComponent::UPhATEdSkeletalMeshComponent(const FObjectInitializer& ObjectInitializer)
 	: Super(ObjectInitializer)
@@ -71,9 +72,9 @@ void UPhATEdSkeletalMeshComponent::RenderAssetTools(const FSceneView* View, clas
 	PDI->DrawPoint(SharedData->LastClickOrigin, FLinearColor(1, 0, 0), 5, SDPG_Foreground);
 #endif
 	// Draw bodies
-	for (int32 i = 0; i <PhysicsAsset->BodySetup.Num(); ++i)
+	for (int32 i = 0; i <PhysicsAsset->SkeletalBodySetups.Num(); ++i)
 	{
-		int32 BoneIndex = GetBoneIndex(PhysicsAsset->BodySetup[i]->BoneName);
+		int32 BoneIndex = GetBoneIndex(PhysicsAsset->SkeletalBodySetups[i]->BoneName);
 
 		// If we found a bone for it, draw the collision.
 		// The logic is as follows; always render in the ViewMode requested when not in hit mode - but if we are in hit mode and the right editing mode, render as solid
@@ -84,7 +85,7 @@ void UPhATEdSkeletalMeshComponent::RenderAssetTools(const FSceneView* View, clas
 			FVector VectorScale(Scale);
 			BoneTM.RemoveScaling();
 
-			FKAggregateGeom* AggGeom = &PhysicsAsset->BodySetup[i]->AggGeom;
+			FKAggregateGeom* AggGeom = &PhysicsAsset->SkeletalBodySetups[i]->AggGeom;
 
 			for (int32 j = 0; j <AggGeom->SphereElems.Num(); ++j)
 			{
@@ -371,7 +372,7 @@ void UPhATEdSkeletalMeshComponent::DrawCurrentInfluences(FPrimitiveDrawInterface
 
 FTransform UPhATEdSkeletalMeshComponent::GetPrimitiveTransform(FTransform& BoneTM, int32 BodyIndex, EKCollisionPrimitiveType PrimType, int32 PrimIndex, float Scale)
 {
-	UBodySetup* SharedBodySetup = SharedData->PhysicsAsset->BodySetup[BodyIndex];
+	UBodySetup* SharedBodySetup = SharedData->PhysicsAsset->SkeletalBodySetups[BodyIndex];
 	FVector Scale3D(Scale);
 
 	FTransform ManTM = FTransform::Identity;
@@ -423,7 +424,7 @@ FTransform UPhATEdSkeletalMeshComponent::GetPrimitiveTransform(FTransform& BoneT
 
 FColor UPhATEdSkeletalMeshComponent::GetPrimitiveColor(int32 BodyIndex, EKCollisionPrimitiveType PrimitiveType, int32 PrimitiveIndex)
 {
-	UBodySetup* SharedBodySetup = SharedData->PhysicsAsset->BodySetup[ BodyIndex ];
+	UBodySetup* SharedBodySetup = SharedData->PhysicsAsset->SkeletalBodySetups[ BodyIndex ];
 
 	if (!SharedData->bRunningSimulation && SharedData->EditingMode == FPhATSharedData::PEM_ConstraintEdit && SharedData->GetSelectedConstraint())
 	{
