Copying //UE4/Dev-Framework to //UE4/Dev-Main (Source: //UE4/Dev-Framework @ 3136612)

diff --git a/Engine/Source/Editor/Persona/Private/AnimationEditorViewportClient.cpp b/Engine/Source/Editor/Persona/Private/AnimationEditorViewportClient.cpp
--- a/Engine/Source/Editor/Persona/Private/AnimationEditorViewportClient.cpp
+++ b/Engine/Source/Editor/Persona/Private/AnimationEditorViewportClient.cpp
@@ -25,11 +25,14 @@
 #include "EngineUtils.h"
 #include "GameFramework/WorldSettings.h"
 #include "PhysicsEngine/PhysicsSettings.h"
-#include "Components/WindDirectionalSourceComponent.h"
 #include "Engine/StaticMesh.h"
 #include "SAnimationEditorViewport.h"
-
+#include "ISkeletonTree.h"
+#include "IEditableSkeleton.h"
+#include "PersonaPreviewSceneDescription.h"
+#include "Engine/PreviewMeshCollection.h"
 #include "AssetViewerSettings.h"
+#include "IPersonaEditorModeManager.h"
 
 namespace {
 	// Value from UE3
@@ -45,34 +48,13 @@ namespace {
 	const float FOVMax = 170.f;
 }
 
-/**
- * A hit proxy class for sockets in the Persona viewport.
- */
-struct HPersonaSocketProxy : public HHitProxy
+namespace EAnimationPlaybackSpeeds
 {
-	DECLARE_HIT_PROXY();
-
-	FSelectedSocketInfo SocketInfo;
+	// Speed scales for animation playback, must match EAnimationPlaybackSpeeds::Type
+	float Values[EAnimationPlaybackSpeeds::NumPlaybackSpeeds] = { 0.1f, 0.25f, 0.5f, 1.0f, 2.0f, 5.0f, 10.0f };
+}
 
-	explicit HPersonaSocketProxy( FSelectedSocketInfo InSocketInfo )
-		:	SocketInfo( InSocketInfo )
-	{}
-};
 IMPLEMENT_HIT_PROXY( HPersonaSocketProxy, HHitProxy );
-
-/**
- * A hit proxy class for sockets in the Persona viewport.
- */
-struct HPersonaBoneProxy : public HHitProxy
-{
-	DECLARE_HIT_PROXY();
-
-	FName BoneName;
-
-	explicit HPersonaBoneProxy(const FName& InBoneName)
-		:	BoneName( InBoneName )
-	{}
-};
 IMPLEMENT_HIT_PROXY( HPersonaBoneProxy, HHitProxy );
 
 #define LOCTEXT_NAMESPACE "FAnimationViewportClient"
@@ -80,18 +62,25 @@ IMPLEMENT_HIT_PROXY( HPersonaBoneProxy, HHitProxy );
 /////////////////////////////////////////////////////////////////////////
 // FAnimationViewportClient
 
-FAnimationViewportClient::FAnimationViewportClient(FAdvancedPreviewScene& InPreviewScene, TWeakPtr<FPersona> InPersonaPtr, const TSharedRef<SAnimationEditorViewport>& InAnimationEditorViewport)
-	: FEditorViewportClient(nullptr, &InPreviewScene, StaticCastSharedRef<SEditorViewport>(InAnimationEditorViewport))
-	, PersonaPtr( InPersonaPtr )
-	, bManipulating(false)
-	, bInTransaction(false)
-	, GravityScaleSliderValue(0.25f)
-	, PrevWindStrength(0.0f)
-	, SelectedWindActor(NULL)
+FAnimationViewportClient::FAnimationViewportClient(const TSharedRef<ISkeletonTree>& InSkeletonTree, const TSharedRef<IPersonaPreviewScene>& InPreviewScene, const TSharedRef<SAnimationEditorViewport>& InAnimationEditorViewport, const TSharedRef<FAssetEditorToolkit>& InAssetEditorToolkit, bool bInShowStats)
+	: FEditorViewportClient(FModuleManager::LoadModuleChecked<FPersonaModule>("Persona").CreatePersonaEditorModeManager(), &InPreviewScene.Get(), StaticCastSharedRef<SEditorViewport>(InAnimationEditorViewport))
+	, SkeletonTreePtr(InSkeletonTree)
+	, PreviewScenePtr(InPreviewScene)
+	, AssetEditorToolkitPtr(InAssetEditorToolkit)
+	, AnimationPlaybackSpeedMode(EAnimationPlaybackSpeeds::Normal)
 	, bFocusOnDraw(false)
-	, BodyTraceDistance(100000.0f)
-	, bShouldUpdateDefaultValues(false)
+	, bInstantFocusOnDraw(true)
+	, bShowMeshStats(bInShowStats)
 {
+	// we actually own the mode tools here, we just override its type in the FEditorViewportClient constructor above
+	bOwnsModeTools = true;
+
+	// Let the asset editor toolkit know about the mode manager so it can be used outside of thew viewport
+	InAssetEditorToolkit->SetAssetEditorModeManager((FAssetEditorModeManager*)ModeTools);
+
+	Widget->SetUsesEditorModeTools(ModeTools);
+	((FAssetEditorModeManager*)ModeTools)->SetPreviewScene(&InPreviewScene.Get());
+
 	// load config
 	ConfigOption = UPersonaOptions::StaticClass()->GetDefaultObject<UPersonaOptions>();
 	check (ConfigOption);
@@ -101,10 +90,8 @@ FAnimationViewportClient::FAnimationViewportClient(FAdvancedPreviewScene& InPrev
 	DrawHelper.AxesLineThickness = ConfigOption->bHighlightOrigin ? 1.0f : 0.0f;
 	DrawHelper.bDrawGrid = ConfigOption->bShowGrid;
 
-	LocalAxesMode = static_cast<ELocalAxesMode::Type>(ConfigOption->DefaultLocalAxesSelection);
-	BoneDrawMode = static_cast<EBoneDrawMode::Type>(ConfigOption->DefaultBoneDrawSelection);
-
 	WidgetMode = FWidget::WM_Rotate;
+	ModeTools->SetWidgetMode(WidgetMode);
 
 	SetSelectedBackgroundColor(ConfigOption->ViewportBackgroundColor, false);
 
@@ -118,7 +105,7 @@ FAnimationViewportClient::FAnimationViewportClient(FAdvancedPreviewScene& InPrev
 	{
 		SetRealtime(false,true); // We are PIE, don't start in realtime mode
 	}
-	
+
 	ViewFOV = FMath::Clamp<float>(ConfigOption->ViewFOV, FOVMin, FOVMax);
 
 	EngineShowFlags.SetSeparateTranslucency(true);
@@ -139,27 +126,10 @@ FAnimationViewportClient::FAnimationViewportClient(FAdvancedPreviewScene& InPrev
 		World->bAllowAudioPlayback = !ConfigOption->bMuteAudio;
 	}
 
-	// Grab a wind actor if it exists (could have been placed in the scene before a mode transition)
-	TActorIterator<AWindDirectionalSource> WindIter(World);
-	if(WindIter)
-	{
-		AWindDirectionalSource* WindActor = *WindIter;
-		WindSourceActor = WindActor;
-
-		PrevWindLocation = WindActor->GetActorLocation();
-		PrevWindRotation = WindActor->GetActorRotation();
-		PrevWindStrength = WindActor->GetComponent()->Strength;
-	}
-	else
-	{
-		//wind actor's initial position, rotation and strength
-		PrevWindLocation = FVector(100, 100, 100);
-		PrevWindRotation = FRotator(0, 0, 0); // roll, yaw, pitch
-		PrevWindStrength = 0.2f;
-	}
+	InPreviewScene->RegisterOnPreviewMeshChanged(FOnPreviewMeshChanged::CreateRaw(this, &FAnimationViewportClient::HandleSkeletalMeshChanged));
+	HandleSkeletalMeshChanged(InPreviewScene->GetPreviewMeshComponent()->SkeletalMesh);
+	InPreviewScene->RegisterOnInvalidateViews(FSimpleDelegate::CreateRaw(this, &FAnimationViewportClient::HandleInvalidateViews));
 
-	// Store direct pointer to advanced preview scene
-	AdvancedPreviewScene = static_cast<FAdvancedPreviewScene*>(PreviewScene);
 	// Register delegate to update the show flags when the post processing is turned on or off
 	UAssetViewerSettings::Get()->OnAssetViewerSettingsChanged().AddRaw(this, &FAnimationViewportClient::OnAssetViewerSettingsChanged);
 	// Set correct flags according to current profile settings
@@ -168,6 +138,14 @@ FAnimationViewportClient::FAnimationViewportClient(FAdvancedPreviewScene& InPrev
 
 FAnimationViewportClient::~FAnimationViewportClient()
 {
+	if (PreviewScenePtr.IsValid())
+	{
+		GetPreviewScene()->UnregisterOnPreviewMeshChanged(this);
+		GetPreviewScene()->UnregisterOnInvalidateViews(this);
+	}
+
+	((FAssetEditorModeManager*)ModeTools)->SetPreviewScene(nullptr);
+
 	UAssetViewerSettings::Get()->OnAssetViewerSettingsChanged().RemoveAll(this);
 }
 
@@ -275,9 +253,10 @@ void FAnimationViewportClient::SetCameraFollow()
 
 		EnableCameraLock(false);
 
-		if (PreviewSkelMeshComp.IsValid())
+		UDebugSkelMeshComponent* PreviewMeshComponent = GetAnimPreviewScene()->GetPreviewMeshComponent();
+		if (PreviewMeshComponent != nullptr)
 		{
-			FBoxSphereBounds Bound = PreviewSkelMeshComp.Get()->CalcBounds(FTransform::Identity);
+			FBoxSphereBounds Bound = PreviewMeshComponent->CalcBounds(FTransform::Identity);
 			SetViewLocationForOrbiting(Bound.Origin);
 		}
 	}
@@ -293,188 +272,150 @@ bool FAnimationViewportClient::IsSetCameraFollowChecked() const
 	return bCameraFollow;
 }
 
-void FAnimationViewportClient::SetPreviewMeshComponent(UDebugSkelMeshComponent* InPreviewSkelMeshComp)
+void FAnimationViewportClient::HandleSkeletalMeshChanged(USkeletalMesh* InSkeletalMesh)
 {
-	PreviewSkelMeshComp = InPreviewSkelMeshComp;
+	FocusViewportOnPreviewMesh();
 
-	PreviewSkelMeshComp->BonesOfInterest.Empty();
+	UDebugSkelMeshComponent* PreviewMeshComponent = GetAnimPreviewScene()->GetPreviewMeshComponent();
+
+	PreviewMeshComponent->BonesOfInterest.Empty();
 
 	UpdateCameraSetup();
 
 	// Setup physics data from physics assets if available, clearing any physics setup on the component
-	UPhysicsAsset* PhysAsset = PreviewSkelMeshComp->GetPhysicsAsset();
+	UPhysicsAsset* PhysAsset = PreviewMeshComponent->GetPhysicsAsset();
 	if(PhysAsset)
 	{
 		PhysAsset->InvalidateAllPhysicsMeshes();
-		PreviewSkelMeshComp->TermArticulated();
-		PreviewSkelMeshComp->InitArticulated(GetWorld()->GetPhysicsScene());
+		PreviewMeshComponent->TermArticulated();
+		PreviewMeshComponent->InitArticulated(GetWorld()->GetPhysicsScene());
 
 		// Set to block all to enable tracing.
-		PreviewSkelMeshComp->SetCollisionProfileName(UCollisionProfile::BlockAll_ProfileName);
+		PreviewMeshComponent->SetCollisionProfileName(UCollisionProfile::BlockAll_ProfileName);
 	}
+
+	Invalidate();
 }
 
 void FAnimationViewportClient::Draw(const FSceneView* View, FPrimitiveDrawInterface* PDI)
 {
 	FEditorViewportClient::Draw(View, PDI);
 
-	if (PreviewSkelMeshComp.IsValid() && PreviewSkelMeshComp->SkeletalMesh)
+	UDebugSkelMeshComponent* PreviewMeshComponent = GetAnimPreviewScene()->GetPreviewMeshComponent();
+	if (PreviewMeshComponent->SkeletalMesh)
 	{
 		// Can't have both bones of interest and sockets of interest set
-		check( !(PreviewSkelMeshComp->BonesOfInterest.Num() && PreviewSkelMeshComp->SocketsOfInterest.Num() ) )
+		check( !(GetAnimPreviewScene()->GetSelectedBoneIndex() != INDEX_NONE && GetAnimPreviewScene()->GetSelectedSocket().IsValid() ) )
 
 		// if we have BonesOfInterest, draw sub set of the bones only
-		if ( PreviewSkelMeshComp->BonesOfInterest.Num() > 0 )
+		if (GetAnimPreviewScene()->GetSelectedBoneIndex() != INDEX_NONE)
 		{
-			DrawMeshSubsetBones(PreviewSkelMeshComp.Get(), PreviewSkelMeshComp->BonesOfInterest, PDI);
+			DrawMeshSubsetBones(PreviewMeshComponent, PreviewMeshComponent->BonesOfInterest, PDI);
 		}
 		// otherwise, if we display bones, display
-		if ( BoneDrawMode != EBoneDrawMode::None )
+		if ( GetBoneDrawMode() != EBoneDrawMode::None )
 		{
-			DrawMeshBones(PreviewSkelMeshComp.Get(), PDI);
+			DrawMeshBones(PreviewMeshComponent, PDI);
 		}
-		if ( PreviewSkelMeshComp->bDisplayRawAnimation )
+		if (PreviewMeshComponent->bDisplayRawAnimation )
 		{
-			DrawMeshBonesUncompressedAnimation(PreviewSkelMeshComp.Get(), PDI);
+			DrawMeshBonesUncompressedAnimation(PreviewMeshComponent, PDI);
 		}
-		if ( PreviewSkelMeshComp->NonRetargetedSpaceBases.Num() > 0 )
+		if (PreviewMeshComponent->NonRetargetedSpaceBases.Num() > 0 )
 		{
-			DrawMeshBonesNonRetargetedAnimation(PreviewSkelMeshComp.Get(), PDI);
+			DrawMeshBonesNonRetargetedAnimation(PreviewMeshComponent, PDI);
 		}
-		if( PreviewSkelMeshComp->bDisplayAdditiveBasePose )
+		if(PreviewMeshComponent->bDisplayAdditiveBasePose )
 		{
-			DrawMeshBonesAdditiveBasePose(PreviewSkelMeshComp.Get(), PDI);
+			DrawMeshBonesAdditiveBasePose(PreviewMeshComponent, PDI);
 		}
-		if(PreviewSkelMeshComp->bDisplayBakedAnimation)
+		if(PreviewMeshComponent->bDisplayBakedAnimation)
 		{
-			DrawMeshBonesBakedAnimation(PreviewSkelMeshComp.Get(), PDI);
+			DrawMeshBonesBakedAnimation(PreviewMeshComponent, PDI);
 		}
-		if(PreviewSkelMeshComp->bDisplaySourceAnimation)
+		if(PreviewMeshComponent->bDisplaySourceAnimation)
 		{
-			DrawMeshBonesSourceRawAnimation(PreviewSkelMeshComp.Get(), PDI);
+			DrawMeshBonesSourceRawAnimation(PreviewMeshComponent, PDI);
 		}
 		
-		DrawWatchedPoses(PreviewSkelMeshComp.Get(), PDI);
+		DrawWatchedPoses(PreviewMeshComponent, PDI);
 
 		// Display normal vectors of each simulation vertex
-		if ( PreviewSkelMeshComp->bDisplayClothingNormals )
+		if (PreviewMeshComponent->bDisplayClothingNormals )
 		{
-			PreviewSkelMeshComp->DrawClothingNormals(PDI);
+			PreviewMeshComponent->DrawClothingNormals(PDI);
 		}
 
 		// Display tangent spaces of each graphical vertex
-		if ( PreviewSkelMeshComp->bDisplayClothingTangents )
+		if (PreviewMeshComponent->bDisplayClothingTangents )
 		{
-			PreviewSkelMeshComp->DrawClothingTangents(PDI);
+			PreviewMeshComponent->DrawClothingTangents(PDI);
 		}
 
 		// Display collision volumes of current selected cloth
-		if ( PreviewSkelMeshComp->bDisplayClothingCollisionVolumes )
+		if (PreviewMeshComponent->bDisplayClothingCollisionVolumes )
 		{
-			PreviewSkelMeshComp->DrawClothingCollisionVolumes(PDI);
+			PreviewMeshComponent->DrawClothingCollisionVolumes(PDI);
 		}
 
 		// Display collision volumes of current selected cloth
-		if ( PreviewSkelMeshComp->bDisplayClothPhysicalMeshWire )
+		if (PreviewMeshComponent->bDisplayClothPhysicalMeshWire )
 		{
-			PreviewSkelMeshComp->DrawClothingPhysicalMeshWire(PDI);
+			PreviewMeshComponent->DrawClothingPhysicalMeshWire(PDI);
 		}
 
 		// Display collision volumes of current selected cloth
-		if ( PreviewSkelMeshComp->bDisplayClothMaxDistances )
+		if (PreviewMeshComponent->bDisplayClothMaxDistances )
 		{
-			PreviewSkelMeshComp->DrawClothingMaxDistances(PDI);
+			PreviewMeshComponent->DrawClothingMaxDistances(PDI);
 		}
 
 		// Display collision volumes of current selected cloth
-		if ( PreviewSkelMeshComp->bDisplayClothBackstops )
+		if (PreviewMeshComponent->bDisplayClothBackstops )
 		{
-			PreviewSkelMeshComp->DrawClothingBackstops(PDI);
+			PreviewMeshComponent->DrawClothingBackstops(PDI);
 		}
 
-		if( PreviewSkelMeshComp->bDisplayClothFixedVertices )
+		if(PreviewMeshComponent->bDisplayClothFixedVertices )
 		{
-			PreviewSkelMeshComp->DrawClothingFixedVertices(PDI);
+			PreviewMeshComponent->DrawClothingFixedVertices(PDI);
 		}
 		
 		// Display socket hit points
-		if ( PreviewSkelMeshComp->bDrawSockets )
+		if (PreviewMeshComponent->bDrawSockets )
 		{
-			if ( PreviewSkelMeshComp->bSkeletonSocketsVisible && PreviewSkelMeshComp->SkeletalMesh->Skeleton )
+			if (PreviewMeshComponent->bSkeletonSocketsVisible && PreviewMeshComponent->SkeletalMesh->Skeleton )
 			{
-				DrawSockets( PreviewSkelMeshComp.Get()->SkeletalMesh->Skeleton->Sockets, PDI, true );
+				DrawSockets(PreviewMeshComponent, PreviewMeshComponent->SkeletalMesh->Skeleton->Sockets, FSelectedSocketInfo(), PDI, true);
 			}
 
-			if ( PreviewSkelMeshComp->bMeshSocketsVisible )
+			if ( PreviewMeshComponent->bMeshSocketsVisible )
 			{
-				DrawSockets( PreviewSkelMeshComp.Get()->SkeletalMesh->GetMeshOnlySocketList(), PDI, false );
+				DrawSockets(PreviewMeshComponent, PreviewMeshComponent->SkeletalMesh->GetMeshOnlySocketList(), FSelectedSocketInfo(), PDI, false);
 			}
 		}
-
-		// If we have a socket of interest, draw the widget
-		if ( PreviewSkelMeshComp->SocketsOfInterest.Num() == 1 )
-		{
-			USkeletalMeshSocket* Socket = PreviewSkelMeshComp->SocketsOfInterest[0].Socket;
-
-			bool bSocketIsOnSkeleton = PreviewSkelMeshComp->SocketsOfInterest[0].bSocketIsOnSkeleton;
-
-			TArray<USkeletalMeshSocket*> SocketAsArray;
-			SocketAsArray.Add( Socket );
-			DrawSockets( SocketAsArray, PDI, false );
-		}
 	}
 
-	if ( PersonaPtr.IsValid() && PreviewSkelMeshComp.IsValid() )
+	if (bFocusOnDraw)
 	{
-		// Allow selected nodes to draw debug rendering if they support it
-		const FGraphPanelSelectionSet SelectedNodes = PersonaPtr.Pin()->GetSelectedNodes();
-		for (FGraphPanelSelectionSet::TConstIterator NodeIt(SelectedNodes); NodeIt; ++NodeIt)
-		{
-			UAnimGraphNode_SkeletalControlBase* Node = Cast<UAnimGraphNode_SkeletalControlBase>(*NodeIt);
-
-			if (Node)
-			{
-				Node->Draw(PDI, PreviewSkelMeshComp.Get());
-			}
-		}
-
-		if(bFocusOnDraw)
-		{
-			bFocusOnDraw = false;
-			FocusViewportOnPreviewMesh();
-		}
+		bFocusOnDraw = false;
+		FocusViewportOnPreviewMesh(bInstantFocusOnDraw);
 	}
 }
 
 void FAnimationViewportClient::DrawCanvas( FViewport& InViewport, FSceneView& View, FCanvas& Canvas )
 {
-	if (PreviewSkelMeshComp.IsValid())
-	{
+	FEditorViewportClient::DrawCanvas(InViewport, View, Canvas);
 
+	UDebugSkelMeshComponent* PreviewMeshComponent = GetAnimPreviewScene()->GetPreviewMeshComponent();
+	if (PreviewMeshComponent != nullptr)
+	{
 		// Display bone names
-		if (PreviewSkelMeshComp->bShowBoneNames)
+		if (PreviewMeshComponent->bShowBoneNames)
 		{
 			ShowBoneNames(&Canvas, &View);
 		}
 
-		// Allow nodes to draw with the canvas, and collect on screen strings to draw later
-		TArray<FText> NodeDebugLines;
-		if(PersonaPtr.IsValid())
-		{
-			// Allow selected nodes to draw debug rendering if they support it
-			const FGraphPanelSelectionSet SelectedNodes = PersonaPtr.Pin()->GetSelectedNodes();
-			for(FGraphPanelSelectionSet::TConstIterator NodeIt(SelectedNodes); NodeIt; ++NodeIt)
-			{
-				UAnimGraphNode_SkeletalControlBase* Node = Cast<UAnimGraphNode_SkeletalControlBase>(*NodeIt);
-
-				if(Node)
-				{
-					Node->DrawCanvas(InViewport, View, Canvas, PreviewSkelMeshComp.Get());
-					Node->GetOnScreenDebugInfo(NodeDebugLines, PreviewSkelMeshComp.Get());
-				}
-			}
-		}
-
 		// Display info
 		if (IsShowingMeshStats())
 		{
@@ -482,55 +423,12 @@ void FAnimationViewportClient::DrawCanvas( FViewport& InViewport, FSceneView& Vi
 		}
 		else if(IsShowingSelectedNodeStats())
 		{
-			// Draw Node info instead of mesh info if we have entries
-			DrawNodeDebugLines(NodeDebugLines, &Canvas, &View);
-		}
-
-		// Draw name of selected bone
-		if (PreviewSkelMeshComp->BonesOfInterest.Num() == 1)
-		{
-			const FIntPoint ViewPortSize = Viewport->GetSizeXY();
-			const int32 HalfX = ViewPortSize.X/2;
-			const int32 HalfY = ViewPortSize.Y/2;
-
-			int32 BoneIndex = PreviewSkelMeshComp->BonesOfInterest[0];
-			const FName BoneName = PreviewSkelMeshComp->SkeletalMesh->RefSkeleton.GetBoneName(BoneIndex);
-
-			FMatrix BoneMatrix = PreviewSkelMeshComp->GetBoneMatrix(BoneIndex);
-			const FPlane proj = View.Project(BoneMatrix.GetOrigin());
-			if (proj.W > 0.f)
-			{
-				const int32 XPos = HalfX + ( HalfX * proj.X );
-				const int32 YPos = HalfY + ( HalfY * (proj.Y * -1) );
-
-				FQuat BoneQuat = PreviewSkelMeshComp->GetBoneQuaternion(BoneName);
-				FVector Loc = PreviewSkelMeshComp->GetBoneLocation(BoneName);
-				FCanvasTextItem TextItem( FVector2D( XPos, YPos), FText::FromString( BoneName.ToString() ), GEngine->GetSmallFont(), FLinearColor::White );
-				Canvas.DrawItem( TextItem );				
-			}
-		}
-
-		// Draw name of selected socket
-		if (PreviewSkelMeshComp->SocketsOfInterest.Num() == 1)
-		{
-			USkeletalMeshSocket* Socket = PreviewSkelMeshComp->SocketsOfInterest[0].Socket;
-
-			FMatrix SocketMatrix;
-			Socket->GetSocketMatrix( SocketMatrix, PreviewSkelMeshComp.Get() );
-			const FVector SocketPos	= SocketMatrix.GetOrigin();
-
-			const FPlane proj = View.Project( SocketPos );
-			if(proj.W > 0.f)
-			{
-				const FIntPoint ViewPortSize = Viewport->GetSizeXY();
-				const int32 HalfX = ViewPortSize.X/2;
-				const int32 HalfY = ViewPortSize.Y/2;
+			// Allow edit modes (inc. skeletal control modes) to draw with the canvas, and collect on screen strings to draw later
+			TArray<FText> EditModeDebugText;
+			GetPersonaModeManager().GetOnScreenDebugInfo(EditModeDebugText);
 
-				const int32 XPos = HalfX + ( HalfX * proj.X );
-				const int32 YPos = HalfY + ( HalfY * (proj.Y * -1) );
-				FCanvasTextItem TextItem( FVector2D( XPos, YPos), FText::FromString( Socket->SocketName.ToString() ), GEngine->GetSmallFont(), FLinearColor::White );
-				Canvas.DrawItem( TextItem );				
-			}
+			// Draw Node info instead of mesh info if we have entries
+			DrawNodeDebugLines(EditModeDebugText, &Canvas, &View);
 		}
 
 		if (bDrawUVs)
@@ -542,158 +440,14 @@ void FAnimationViewportClient::DrawCanvas( FViewport& InViewport, FSceneView& Vi
 
 void FAnimationViewportClient::DrawUVsForMesh(FViewport* InViewport, FCanvas* InCanvas, int32 InTextYPos)
 {
+	UDebugSkelMeshComponent* PreviewMeshComponent = GetAnimPreviewScene()->GetPreviewMeshComponent();
+
 	//use the overridden LOD level
-	const uint32 LODLevel = FMath::Clamp(PreviewSkelMeshComp->ForcedLodModel - 1, 0, PreviewSkelMeshComp->SkeletalMesh->LODInfo.Num() - 1);
+	const uint32 LODLevel = FMath::Clamp(PreviewMeshComponent->ForcedLodModel - 1, 0, PreviewMeshComponent->SkeletalMesh->LODInfo.Num() - 1);
 
 	TArray<FVector2D> SelectedEdgeTexCoords; //No functionality in Persona for this (yet?)
 
-	DrawUVs(InViewport, InCanvas, InTextYPos, LODLevel, UVChannelToDraw, SelectedEdgeTexCoords, NULL, &PreviewSkelMeshComp->GetSkeletalMeshResource()->LODModels[LODLevel] );
-}
-
-FAnimNode_SkeletalControlBase* FAnimationViewportClient::FindSkeletalControlAnimNode(TWeakObjectPtr<UAnimGraphNode_SkeletalControlBase> AnimGraphNode) const
-{
-	FAnimNode_SkeletalControlBase* AnimNode = NULL;
-
-	if (PreviewSkelMeshComp.IsValid() && PreviewSkelMeshComp->GetAnimInstance() && AnimGraphNode.IsValid())
-	{
-		AnimNode = AnimGraphNode.Get()->FindDebugAnimNode(PreviewSkelMeshComp.Get());
-	}
-
-	return AnimNode;
-}
-
-void FAnimationViewportClient::FindSelectedAnimGraphNode()
-{
-	bool bSelected = false;
-	// finding a selected node
-	if (PersonaPtr.IsValid() && PreviewSkelMeshComp.IsValid())
-	{
-		USkeletalMeshComponent* PreviewComponent = PreviewSkelMeshComp.Get();
-		check(PreviewComponent);
-
-		const FGraphPanelSelectionSet SelectedNodes = PersonaPtr.Pin()->GetSelectedNodes();
-
-		// don't support multi-selection
-		if (SelectedNodes.Num() == 1)
-		{
-			FGraphPanelSelectionSet::TConstIterator NodeIt(SelectedNodes);
-			// first element
-			UAnimGraphNode_SkeletalControlBase* AnimGraphNode = Cast<UAnimGraphNode_SkeletalControlBase>(*NodeIt);
-
-			if (AnimGraphNode)
-			{
-				FAnimNode_SkeletalControlBase* AnimNode = FindSkeletalControlAnimNode(AnimGraphNode);
-
-				if (AnimNode)
-				{
-					bSelected = true;
-
-					// when selected first after AnimGraph is opened, assign previous data to current node
-					if (SelectedSkelControlAnimGraph != AnimGraphNode)
-					{
-
-						if (SelectedSkelControlAnimGraph.IsValid())
-						{
-							SelectedSkelControlAnimGraph->DeselectActor(PreviewComponent);
-						}
-
-						// make same values to ensure data consistency
-						AnimGraphNode->CopyNodeDataTo(AnimNode);
-
-						WidgetMode = (FWidget::EWidgetMode)AnimGraphNode->GetWidgetMode(PreviewComponent);
-						ECoordSystem DesiredCoordSystem = (ECoordSystem)AnimGraphNode->GetWidgetCoordinateSystem(PreviewComponent);
-						SetWidgetCoordSystemSpace(DesiredCoordSystem);
-					}
-					else
-					{
-						if (bShouldUpdateDefaultValues)
-						{
-							// copy updated values into internal node to ensure data consistency
-							AnimGraphNode->CopyNodeDataFrom(AnimNode);
-							bShouldUpdateDefaultValues = false;
-						}
-					}
-
-					AnimGraphNode->MoveSelectActorLocation(PreviewComponent, AnimNode);
-					SelectedSkelControlAnimGraph = AnimGraphNode;
-				}
-			}
-		}
-
-		if (!bSelected)
-		{
-			if (SelectedSkelControlAnimGraph.IsValid())
-			{
-				SelectedSkelControlAnimGraph->DeselectActor(PreviewComponent);
-			}
-			SelectedSkelControlAnimGraph.Reset();
-		}
-	}
-}
-
-void FAnimationViewportClient::ClearSelectedAnimGraphNode()
-{
-	if (SelectedSkelControlAnimGraph.IsValid())
-	{
-		SelectedSkelControlAnimGraph->DeselectActor(PreviewSkelMeshComp.Get());
-	}
-}
-
-void FAnimationViewportClient::PostUndo()
-{
-	// undo for skeletal controllers
-	// search all UAnimGraphNode_SkeletalControlBase nodes and apply data
-	UEdGraph* FocusedGraph = PersonaPtr.Pin()->GetFocusedGraph();
-
-	// @fixme : fix this to better way than looping all nodes
-	if (FocusedGraph)
-	{
-		// find UAnimGraphNode_SkeletalControlBase
-		for (UEdGraphNode* Node : FocusedGraph->Nodes)
-		{
-			UAnimGraphNode_SkeletalControlBase* AnimGraphNode = Cast<UAnimGraphNode_SkeletalControlBase>(Node);
-
-			if (AnimGraphNode)
-			{
-				FAnimNode_SkeletalControlBase* AnimNode = FindSkeletalControlAnimNode(AnimGraphNode);
-
-				if (AnimNode)
-				{
-					// copy undo data
-					AnimGraphNode->CopyNodeDataTo(AnimNode);
-				}
-			}
-		}
-	}
-}
-
-void FAnimationViewportClient::PostCompile()
-{
-	// reset the selected anim graph to copy 
-	SelectedSkelControlAnimGraph.Reset();
-
-	// if the user manipulated Pin values directly from the node, then should copy updated values to the internal node to retain data consistency
-	UEdGraph* FocusedGraph = PersonaPtr.Pin()->GetFocusedGraph();
-
-	// @fixme : fix this to better way than looping all nodes
-	if (FocusedGraph)
-	{
-		// find UAnimGraphNode_SkeletalControlBase
-		for (UEdGraphNode* Node : FocusedGraph->Nodes)
-		{
-			UAnimGraphNode_SkeletalControlBase* AnimGraphNode = Cast<UAnimGraphNode_SkeletalControlBase>(Node);
-
-			if (AnimGraphNode)
-			{
-				FAnimNode_SkeletalControlBase* AnimNode = FindSkeletalControlAnimNode(AnimGraphNode);
-
-				if (AnimNode)
-				{
-					AnimGraphNode->CopyNodeDataFrom(AnimNode);
-				}
-			}
-		}
-	}
+	DrawUVs(InViewport, InCanvas, InTextYPos, LODLevel, UVChannelToDraw, SelectedEdgeTexCoords, NULL, &PreviewMeshComponent->GetSkeletalMeshResource()->LODModels[LODLevel] );
 }
 
 void FAnimationViewportClient::Tick(float DeltaSeconds) 
@@ -701,7 +455,8 @@ void FAnimationViewportClient::Tick(float DeltaSeconds)
 	FEditorViewportClient::Tick(DeltaSeconds);
 
 	// @todo fixme
-	if (bCameraFollow && PreviewSkelMeshComp.IsValid())
+	UDebugSkelMeshComponent* PreviewMeshComponent = GetAnimPreviewScene()->GetPreviewMeshComponent();
+	if (bCameraFollow && PreviewMeshComponent != nullptr)
 	{
 		// if camera isn't lock, make the mesh bounds to be center
 		FSphere BoundSphere = GetCameraTarget();
@@ -715,17 +470,13 @@ void FAnimationViewportClient::Tick(float DeltaSeconds)
 		PreviewScene->GetWorld()->Tick(LEVELTICK_All, DeltaSeconds);
 	}
 
-	UDebugSkelMeshComponent* PreviewComp = PreviewSkelMeshComp.Get();
+	UDebugSkelMeshComponent* PreviewComp = PreviewMeshComponent;
 	if (PreviewComp)
 	{
-		// Handle updating the preview component to represent the effects of root motion
-		const UStaticMeshComponent* FloorMeshComponent = AdvancedPreviewScene->GetFloorMeshComponent();
-		FBoxSphereBounds Bounds = FloorMeshComponent->CalcBounds(FloorMeshComponent->GetRelativeTransform());
+		// Handle updating the preview component to represent the effects of root motion	
+		const FBoxSphereBounds& Bounds = GetAnimPreviewScene()->GetFloorBounds();
 		PreviewComp->ConsumeRootMotion(Bounds.GetBox().Min, Bounds.GetBox().Max);
-	}	
-
-	FindSelectedAnimGraphNode();
-
+	}
 }
 
 void FAnimationViewportClient::SetCameraTargetLocation(const FSphere &BoundSphere, float DeltaSeconds)
@@ -746,15 +497,16 @@ void FAnimationViewportClient::SetCameraTargetLocation(const FSphere &BoundSpher
 
 void FAnimationViewportClient::ShowBoneNames( FCanvas* Canvas, FSceneView* View )
 {
-	if (!PreviewSkelMeshComp.IsValid() || !PreviewSkelMeshComp->MeshObject)
+	UDebugSkelMeshComponent* PreviewMeshComponent = GetAnimPreviewScene()->GetPreviewMeshComponent();
+	if (PreviewMeshComponent == nullptr || PreviewMeshComponent->MeshObject == nullptr)
 	{
 		return;
 	}
 
 	//Most of the code taken from FASVViewportClient::Draw() in AnimSetViewerMain.cpp
-	FSkeletalMeshResource* SkelMeshResource = PreviewSkelMeshComp->GetSkeletalMeshResource();
+	FSkeletalMeshResource* SkelMeshResource = PreviewMeshComponent->GetSkeletalMeshResource();
 	check(SkelMeshResource);
-	const int32 LODIndex = FMath::Clamp(PreviewSkelMeshComp->PredictedLODLevel, 0, SkelMeshResource->LODModels.Num()-1);
+	const int32 LODIndex = FMath::Clamp(PreviewMeshComponent->PredictedLODLevel, 0, SkelMeshResource->LODModels.Num()-1);
 	FStaticLODModel& LODModel = SkelMeshResource->LODModels[ LODIndex ];
 
 	const int32 HalfX = Viewport->GetSizeXY().X/2;
@@ -765,7 +517,7 @@ void FAnimationViewportClient::ShowBoneNames( FCanvas* Canvas, FSceneView* View
 		const int32 BoneIndex = LODModel.RequiredBones[i];
 
 		// If previewing a specific section, only show the bone names that belong to it
-		if ((PreviewSkelMeshComp->SectionIndexPreview >= 0) && !LODModel.Sections[PreviewSkelMeshComp->SectionIndexPreview].BoneMap.Contains(BoneIndex))
+		if ((PreviewMeshComponent->SectionIndexPreview >= 0) && !LODModel.Sections[PreviewMeshComponent->SectionIndexPreview].BoneMap.Contains(BoneIndex))
 		{
 			continue;
 		}
@@ -773,7 +525,7 @@ void FAnimationViewportClient::ShowBoneNames( FCanvas* Canvas, FSceneView* View
 		const FColor BoneColor = FColor::White;
 		if (BoneColor.A != 0)
 		{
-			const FVector BonePos = PreviewSkelMeshComp->ComponentToWorld.TransformPosition(PreviewSkelMeshComp->GetComponentSpaceTransforms()[BoneIndex].GetLocation());
+			const FVector BonePos = PreviewMeshComponent->ComponentToWorld.TransformPosition(PreviewMeshComponent->GetComponentSpaceTransforms()[BoneIndex].GetLocation());
 
 			const FPlane proj = View->Project(BonePos);
 			if (proj.W > 0.f)
@@ -781,7 +533,7 @@ void FAnimationViewportClient::ShowBoneNames( FCanvas* Canvas, FSceneView* View
 				const int32 XPos = HalfX + ( HalfX * proj.X );
 				const int32 YPos = HalfY + ( HalfY * (proj.Y * -1) );
 
-				const FName BoneName = PreviewSkelMeshComp->SkeletalMesh->RefSkeleton.GetBoneName(BoneIndex);
+				const FName BoneName = PreviewMeshComponent->SkeletalMesh->RefSkeleton.GetBoneName(BoneIndex);
 				const FString BoneString = FString::Printf( TEXT("%d: %s"), BoneIndex, *BoneName.ToString() );
 				FCanvasTextItem TextItem( FVector2D( XPos, YPos), FText::FromString( BoneString ), GEngine->GetSmallFont(), BoneColor );
 				TextItem.EnableShadow(FLinearColor::Black);
@@ -791,6 +543,26 @@ void FAnimationViewportClient::ShowBoneNames( FCanvas* Canvas, FSceneView* View
 	}
 }
 
+bool FAnimationViewportClient::ShouldDisplayAdditiveScaleErrorMessage()
+{
+	UAnimSequence* AnimSequence = Cast<UAnimSequence>(GetAnimPreviewScene()->GetPreviewAnimationAsset());
+	if (AnimSequence)
+	{
+		if (AnimSequence->IsValidAdditive() && AnimSequence->RefPoseSeq)
+		{
+			if (RefPoseGuid != AnimSequence->RefPoseSeq->RawDataGuid)
+			{
+				RefPoseGuid = AnimSequence->RefPoseSeq->RawDataGuid;
+				bDoesAdditiveRefPoseHaveZeroScale = AnimSequence->DoesSequenceContainZeroScale();
+			}
+			return bDoesAdditiveRefPoseHaveZeroScale;
+		}
+	}
+
+	RefPoseGuid.Invalidate();
+	return false;
+}
+
 void FAnimationViewportClient::DisplayInfo(FCanvas* Canvas, FSceneView* View, bool bDisplayAllInfo)
 {
 	int32 CurXOffset = 5;
@@ -810,21 +582,20 @@ void FAnimationViewportClient::DisplayInfo(FCanvas* Canvas, FSceneView* View, bo
 	const FColor SubHeadlineColour(202, 66, 0);
 
 	// if not valid skeletalmesh
-	if (!PreviewSkelMeshComp.IsValid() || !PreviewSkelMeshComp->SkeletalMesh)
+	UDebugSkelMeshComponent* PreviewMeshComponent = GetAnimPreviewScene()->GetPreviewMeshComponent();
+	if (PreviewMeshComponent == nullptr || PreviewMeshComponent->SkeletalMesh == nullptr)
 	{
 		return;
 	}
 
-	TSharedPtr<FPersona> SharedPersona = PersonaPtr.Pin();
-
-	if (SharedPersona.IsValid() && SharedPersona->ShouldDisplayAdditiveScaleErrorMessage())
+	if (ShouldDisplayAdditiveScaleErrorMessage())
 	{
 		InfoString = TEXT("Additve ref pose contains scales of 0.0, this can cause additive animations to not give the desired results");
 		Canvas->DrawShadowedString(CurXOffset, CurYOffset, *InfoString, GEngine->GetSmallFont(), SubHeadlineColour);
 		CurYOffset += YL + 2;
 	}
 
-	if (PreviewSkelMeshComp->SkeletalMesh->MorphTargets.Num() > 0)
+	if (PreviewMeshComponent->SkeletalMesh->MorphTargets.Num() > 0)
 	{
 		int32 SubHeadingIndent = CurXOffset + 10;
 
@@ -832,9 +603,9 @@ void FAnimationViewportClient::DisplayInfo(FCanvas* Canvas, FSceneView* View, bo
 		TArray<UMaterial*> MaterialsThatNeedMorphFlagOn;
 		TArray<UMaterial*> MaterialsThatNeedSaving;
 
-		for (int i = 0; i < PreviewSkelMeshComp->GetNumMaterials(); ++i)
+		for (int i = 0; i < PreviewMeshComponent->GetNumMaterials(); ++i)
 		{
-			if (UMaterialInterface* MaterialInterface = PreviewSkelMeshComp->GetMaterial(i))
+			if (UMaterialInterface* MaterialInterface = PreviewMeshComponent->GetMaterial(i))
 			{
 				UMaterial* Material = MaterialInterface->GetMaterial();
 				if ((Material != nullptr) && !ProcessedMaterials.Contains(Material))
@@ -887,7 +658,7 @@ void FAnimationViewportClient::DisplayInfo(FCanvas* Canvas, FSceneView* View, bo
 		}
 	}
 
-	UAnimPreviewInstance* PreviewInstance = PreviewSkelMeshComp->PreviewInstance;
+	UAnimPreviewInstance* PreviewInstance = PreviewMeshComponent->PreviewInstance;
 	if( PreviewInstance )
 	{
 		// see if you have anim sequence that has transform curves
@@ -910,11 +681,11 @@ void FAnimationViewportClient::DisplayInfo(FCanvas* Canvas, FSceneView* View, bo
 		}
 	}
 
-	if (PreviewSkelMeshComp->IsUsingInGameBounds())
+	if (PreviewMeshComponent->IsUsingInGameBounds())
 	{
-		if (!PreviewSkelMeshComp->CheckIfBoundsAreCorrrect())
+		if (!PreviewMeshComponent->CheckIfBoundsAreCorrrect())
 		{
-			if( PreviewSkelMeshComp->GetPhysicsAsset() == NULL )
+			if( PreviewMeshComponent->GetPhysicsAsset() == NULL )
 			{
 				InfoString = FString::Printf( *LOCTEXT("NeedToSetupPhysicsAssetForAccurateBounds", "You may need to setup Physics Asset to use more accurate bounds").ToString() );
 			}
@@ -927,15 +698,15 @@ void FAnimationViewportClient::DisplayInfo(FCanvas* Canvas, FSceneView* View, bo
 		}
 	}
 
-	if (PreviewSkelMeshComp != NULL && PreviewSkelMeshComp->MeshObject != NULL)
+	if (PreviewMeshComponent != NULL && PreviewMeshComponent->MeshObject != NULL)
 	{
 		if (bDisplayAllInfo)
 		{
-			FSkeletalMeshResource* SkelMeshResource = PreviewSkelMeshComp->GetSkeletalMeshResource();
+			FSkeletalMeshResource* SkelMeshResource = PreviewMeshComponent->GetSkeletalMeshResource();
 			check(SkelMeshResource);
 
 			// Draw stats about the mesh
-			const FBoxSphereBounds& SkelBounds = PreviewSkelMeshComp->Bounds;
+			const FBoxSphereBounds& SkelBounds = PreviewMeshComponent->Bounds;
 			const FPlane ScreenPosition = View->Project(SkelBounds.Origin);
 
 			const int32 HalfX = Viewport->GetSizeXY().X / 2;
@@ -949,7 +720,7 @@ void FAnimationViewportClient::DisplayInfo(FCanvas* Canvas, FSceneView* View, bo
 			int32 NumSectionsInUse;
 			FString WeightUsage;
 
-			const int32 LODIndex = FMath::Clamp(PreviewSkelMeshComp->PredictedLODLevel, 0, SkelMeshResource->LODModels.Num() - 1);
+			const int32 LODIndex = FMath::Clamp(PreviewMeshComponent->PredictedLODLevel, 0, SkelMeshResource->LODModels.Num() - 1);
 			FStaticLODModel& LODModel = SkelMeshResource->LODModels[LODIndex];
 
 			NumBonesInUse = LODModel.RequiredBones.Num();
@@ -1009,12 +780,12 @@ void FAnimationViewportClient::DisplayInfo(FCanvas* Canvas, FSceneView* View, bo
 			CurYOffset += YL + 2;
 			Canvas->DrawShadowedString(CurXOffset, CurYOffset, *InfoString, GEngine->GetSmallFont(), TextColor);
 
-			if (PreviewSkelMeshComp->BonesOfInterest.Num() > 0)
+			if (PreviewMeshComponent->BonesOfInterest.Num() > 0)
 			{
-				int32 BoneIndex = PreviewSkelMeshComp->BonesOfInterest[0];
-				FTransform ReferenceTransform = PreviewSkelMeshComp->SkeletalMesh->RefSkeleton.GetRefBonePose()[BoneIndex];
-				FTransform LocalTransform = PreviewSkelMeshComp->BoneSpaceTransforms[BoneIndex];
-				FTransform ComponentTransform = PreviewSkelMeshComp->GetComponentSpaceTransforms()[BoneIndex];
+				int32 BoneIndex = PreviewMeshComponent->BonesOfInterest[0];
+				FTransform ReferenceTransform = PreviewMeshComponent->SkeletalMesh->RefSkeleton.GetRefBonePose()[BoneIndex];
+				FTransform LocalTransform = PreviewMeshComponent->BoneSpaceTransforms[BoneIndex];
+				FTransform ComponentTransform = PreviewMeshComponent->GetComponentSpaceTransforms()[BoneIndex];
 
 				CurYOffset += YL + 2;
 				InfoString = FString::Printf(TEXT("Local :%s"), *LocalTransform.ToString());
@@ -1031,28 +802,28 @@ void FAnimationViewportClient::DisplayInfo(FCanvas* Canvas, FSceneView* View, bo
 
 			CurYOffset += YL + 2;
 			InfoString = FString::Printf(TEXT("Approximate Size: %ix%ix%i"),
-				FMath::RoundToInt(PreviewSkelMeshComp->Bounds.BoxExtent.X * 2.0f),
-				FMath::RoundToInt(PreviewSkelMeshComp->Bounds.BoxExtent.Y * 2.0f),
-				FMath::RoundToInt(PreviewSkelMeshComp->Bounds.BoxExtent.Z * 2.0f));
+				FMath::RoundToInt(PreviewMeshComponent->Bounds.BoxExtent.X * 2.0f),
+				FMath::RoundToInt(PreviewMeshComponent->Bounds.BoxExtent.Y * 2.0f),
+				FMath::RoundToInt(PreviewMeshComponent->Bounds.BoxExtent.Z * 2.0f));
 			Canvas->DrawShadowedString(CurXOffset, CurYOffset, *InfoString, GEngine->GetSmallFont(), TextColor);
 
-			uint32 NumNotiesWithErrors = PreviewSkelMeshComp->AnimNotifyErrors.Num();
+			uint32 NumNotiesWithErrors = PreviewMeshComponent->AnimNotifyErrors.Num();
 			for (uint32 i = 0; i < NumNotiesWithErrors; ++i)
 			{
-				uint32 NumErrors = PreviewSkelMeshComp->AnimNotifyErrors[i].Errors.Num();
+				uint32 NumErrors = PreviewMeshComponent->AnimNotifyErrors[i].Errors.Num();
 				for (uint32 ErrorIdx = 0; ErrorIdx < NumErrors; ++ErrorIdx)
 				{
 					CurYOffset += YL + 2;
-					Canvas->DrawShadowedString(CurXOffset, CurYOffset, *PreviewSkelMeshComp->AnimNotifyErrors[i].Errors[ErrorIdx], GEngine->GetSmallFont(), FLinearColor(1.0f, 0.0f, 0.0f, 1.0f));
+					Canvas->DrawShadowedString(CurXOffset, CurYOffset, *PreviewMeshComponent->AnimNotifyErrors[i].Errors[ErrorIdx], GEngine->GetSmallFont(), FLinearColor(1.0f, 0.0f, 0.0f, 1.0f));
 				}
 			}
 		}
 		else // simplified default display info to be same as static mesh editor
 		{
-			FSkeletalMeshResource* SkelMeshResource = PreviewSkelMeshComp->GetSkeletalMeshResource();
+			FSkeletalMeshResource* SkelMeshResource = PreviewMeshComponent->GetSkeletalMeshResource();
 			check(SkelMeshResource);
 
-			const int32 LODIndex = FMath::Clamp(PreviewSkelMeshComp->PredictedLODLevel, 0, SkelMeshResource->LODModels.Num() - 1);
+			const int32 LODIndex = FMath::Clamp(PreviewMeshComponent->PredictedLODLevel, 0, SkelMeshResource->LODModels.Num() - 1);
 			FStaticLODModel& LODModel = SkelMeshResource->LODModels[LODIndex];
 
 			// Current LOD 
@@ -1060,7 +831,7 @@ void FAnimationViewportClient::DisplayInfo(FCanvas* Canvas, FSceneView* View, bo
 			Canvas->DrawShadowedString(CurXOffset, CurYOffset, *InfoString, GEngine->GetSmallFont(), TextColor);
 
 			// current screen size
-			const FBoxSphereBounds& SkelBounds = PreviewSkelMeshComp->Bounds;
+			const FBoxSphereBounds& SkelBounds = PreviewMeshComponent->Bounds;
 			const FPlane ScreenPosition = View->Project(SkelBounds.Origin);
 
 			const int32 HalfX = Viewport->GetSizeXY().X / 2;
@@ -1098,14 +869,14 @@ void FAnimationViewportClient::DisplayInfo(FCanvas* Canvas, FSceneView* View, bo
 			// Approx Size 
 			CurYOffset += YL + 2;
 			InfoString = FString::Printf(TEXT("Approx Size: %ix%ix%i"),
-				FMath::RoundToInt(PreviewSkelMeshComp->Bounds.BoxExtent.X * 2.0f),
-				FMath::RoundToInt(PreviewSkelMeshComp->Bounds.BoxExtent.Y * 2.0f),
-				FMath::RoundToInt(PreviewSkelMeshComp->Bounds.BoxExtent.Z * 2.0f));
+				FMath::RoundToInt(PreviewMeshComponent->Bounds.BoxExtent.X * 2.0f),
+				FMath::RoundToInt(PreviewMeshComponent->Bounds.BoxExtent.Y * 2.0f),
+				FMath::RoundToInt(PreviewMeshComponent->Bounds.BoxExtent.Z * 2.0f));
 			Canvas->DrawShadowedString(CurXOffset, CurYOffset, *InfoString, GEngine->GetSmallFont(), TextColor);
 		}
 	}
 
-	if (PreviewSkelMeshComp->SectionIndexPreview != INDEX_NONE)
+	if (PreviewMeshComponent->SectionIndexPreview != INDEX_NONE)
 	{
 		// Notify the user if they are isolating a mesh section.
 		CurYOffset += YL + 2;
@@ -1140,409 +911,21 @@ void FAnimationViewportClient::DrawNodeDebugLines(TArray<FText>& Lines, FCanvas*
 	}
 }
 
-int32 FAnimationViewportClient::FindSelectedBone() const
-{
-	int32 SelectedBone = INDEX_NONE;
-	if (SelectedSkelControlAnimGraph.IsValid())
-	{		
-		FName BoneName = SelectedSkelControlAnimGraph->FindSelectedBone();
-		SelectedBone = PreviewSkelMeshComp->GetBoneIndex(BoneName);
-	}
-
-	if (SelectedBone == INDEX_NONE && PreviewSkelMeshComp.IsValid() && (PreviewSkelMeshComp->BonesOfInterest.Num() == 1) )
-	{
-		SelectedBone = PreviewSkelMeshComp->BonesOfInterest[0];
-	}
-	return SelectedBone;
-}
-
-USkeletalMeshSocket* FAnimationViewportClient::FindSelectedSocket() const
-{
-	USkeletalMeshSocket* SelectedSocket = NULL;
-
-	if (PreviewSkelMeshComp.IsValid() && (PreviewSkelMeshComp->SocketsOfInterest.Num() == 1) )
-	{
-		SelectedSocket = PreviewSkelMeshComp->SocketsOfInterest[0].Socket;
-	}
-
-	return SelectedSocket;
-}
-
-TWeakObjectPtr<AWindDirectionalSource> FAnimationViewportClient::FindSelectedWindActor() const
-{
-	return SelectedWindActor;
-}
-
-void FAnimationViewportClient::ProcessClick(class FSceneView& View, class HHitProxy* HitProxy, FKey Key, EInputEvent Event, uint32 HitX, uint32 HitY)
-{
-	TSharedPtr<FPersona> SharedPersona = PersonaPtr.Pin();
-
-	if(!SharedPersona.IsValid())
-	{
-		// No reason to continue, cannot select elements
-		return;
-	}
-
-	if ( HitProxy )
-	{
-		if ( HitProxy->IsA( HPersonaSocketProxy::StaticGetType() ) )
-		{
-			HPersonaSocketProxy* SocketProxy = ( HPersonaSocketProxy* )HitProxy;
-
-			// Tell Persona that the socket has been selected - this will sort out the skeleton tree, etc.
-			SharedPersona->SetSelectedSocket( SocketProxy->SocketInfo );
-		}
-		else if ( HitProxy->IsA( HPersonaBoneProxy::StaticGetType() ) )
-		{
-			HPersonaBoneProxy* BoneProxy = ( HPersonaBoneProxy* )HitProxy;
-			// Tell Persona that the bone has been selected - this will sort out the skeleton tree, etc.
-			SharedPersona->SetSelectedBone( PreviewSkelMeshComp.Get()->SkeletalMesh->Skeleton, BoneProxy->BoneName );
-		}
-		else if ( HitProxy->IsA( HActor::StaticGetType() ) )
-		{
-			HActor* ActorHitProxy = static_cast<HActor*>(HitProxy);
-			AWindDirectionalSource* WindActor = Cast<AWindDirectionalSource>(ActorHitProxy->Actor);
-
-			if( WindActor )
-			{
-				//clear previously selected things
-				SharedPersona->DeselectAll();
-				SelectedWindActor = WindActor;
-			}
-			else if (SelectedSkelControlAnimGraph.IsValid() && SelectedSkelControlAnimGraph->IsActorClicked(ActorHitProxy))
-			{
-				// do some actions when hit
-				SelectedSkelControlAnimGraph->ProcessActorClick(ActorHitProxy);
-			}
-		}
-	}
-	else
-	{
-		// Cast for phys bodies if we didn't get any hit proxies
-		FHitResult Result(1.0f);
-		const FViewportClick Click(&View, this, EKeys::Invalid, IE_Released, Viewport->GetMouseX(), Viewport->GetMouseY());
-		bool bHit = PreviewSkelMeshComp.Get()->LineTraceComponent(Result, Click.GetOrigin(), Click.GetOrigin() + Click.GetDirection() * BodyTraceDistance, FCollisionQueryParams(NAME_None,true));
-		
-		if(bHit)
-		{
-			SharedPersona->SetSelectedBone(PreviewSkelMeshComp->SkeletalMesh->Skeleton, Result.BoneName);
-		}
-		else
-		{
-			// We didn't hit a proxy or a physics object, so deselect all objects
-			SharedPersona->DeselectAll();
-		}
-}
-}
-
-bool FAnimationViewportClient::InputWidgetDelta( FViewport* InViewport, EAxisList::Type CurrentAxis, FVector& Drag, FRotator& Rot, FVector& Scale )
-{
-	// Get some useful info about buttons being held down
-	const bool bCtrlDown = InViewport->KeyState(EKeys::LeftControl) || InViewport->KeyState(EKeys::RightControl);
-	const bool bShiftDown = InViewport->KeyState(EKeys::LeftShift) || InViewport->KeyState(EKeys::RightShift);
-	const bool bMouseButtonDown = InViewport->KeyState( EKeys::LeftMouseButton ) || InViewport->KeyState( EKeys::MiddleMouseButton ) || InViewport->KeyState( EKeys::RightMouseButton );
-
-	bool bHandled = false;
-
-	if ( bManipulating && CurrentAxis != EAxisList::None )
-	{
-		bHandled = true;
-
-		int32 BoneIndex = FindSelectedBone();
-		USkeletalMeshSocket* SelectedSocket = FindSelectedSocket();
-		TWeakObjectPtr<AWindDirectionalSource> WindActor = FindSelectedWindActor();
-
-		FAnimNode_ModifyBone* SkelControl = NULL;
-
-		if ( BoneIndex >= 0 )
-		{
-			//Get the skeleton control manipulating this bone
-			const FName BoneName = PreviewSkelMeshComp->SkeletalMesh->RefSkeleton.GetBoneName(BoneIndex);
-			SkelControl = &(PreviewSkelMeshComp->PreviewInstance->ModifyBone(BoneName));
-		}
-
-		if ( SkelControl || SelectedSocket )
-		{
-			FTransform CurrentSkelControlTM(
-				SelectedSocket ? SelectedSocket->RelativeRotation : SkelControl->Rotation,
-				SelectedSocket ? SelectedSocket->RelativeLocation : SkelControl->Translation,
-				SelectedSocket ? SelectedSocket->RelativeScale : SkelControl->Scale);
-
-			FTransform BaseTM;
-
-			if ( SelectedSocket )
-			{
-				BaseTM = SelectedSocket->GetSocketTransform( PreviewSkelMeshComp.Get() );
-			}
-			else
-			{
-				BaseTM = PreviewSkelMeshComp->GetBoneTransform( BoneIndex );
-			}
-
-			// Remove SkelControl's orientation from BoneMatrix, as we need to translate/rotate in the non-SkelControlled space
-			BaseTM = BaseTM.GetRelativeTransformReverse(CurrentSkelControlTM);
-
-			const bool bDoRotation    = WidgetMode == FWidget::WM_Rotate    || WidgetMode == FWidget::WM_TranslateRotateZ;
-			const bool bDoTranslation = WidgetMode == FWidget::WM_Translate || WidgetMode == FWidget::WM_TranslateRotateZ;
-			const bool bDoScale = WidgetMode == FWidget::WM_Scale;
-
-			if (bDoRotation)
-			{
-				FVector RotAxis;
-				float RotAngle;
-				Rot.Quaternion().ToAxisAndAngle( RotAxis, RotAngle );
-
-				FVector4 BoneSpaceAxis = BaseTM.TransformVector( RotAxis );
-
-				//Calculate the new delta rotation
-				FQuat DeltaQuat( BoneSpaceAxis, RotAngle );
-				DeltaQuat.Normalize();
-
-				FRotator NewRotation = ( CurrentSkelControlTM * FTransform( DeltaQuat )).Rotator();
-
-				if (SelectedSkelControlAnimGraph.IsValid())
-				{
-					FAnimNode_SkeletalControlBase* AnimNode = FindSkeletalControlAnimNode(SelectedSkelControlAnimGraph);
-					if (AnimNode)
-					{
-						SelectedSkelControlAnimGraph->DoRotation(PreviewSkelMeshComp.Get(), Rot, AnimNode);
-						bShouldUpdateDefaultValues = true;
-					}
-				}
-				else if ( SelectedSocket )
-				{
-					SelectedSocket->RelativeRotation = NewRotation;
-				}
-				else
-				{
-					SkelControl->Rotation = NewRotation;
-				}
-			}
-
-			if (bDoTranslation)
-			{
-				FVector4 BoneSpaceOffset = BaseTM.TransformVector(Drag);
-				if (SelectedSkelControlAnimGraph.IsValid())
-				{
-					FAnimNode_SkeletalControlBase* AnimNode = FindSkeletalControlAnimNode(SelectedSkelControlAnimGraph);
-					if (AnimNode)
-					{
-						SelectedSkelControlAnimGraph->DoTranslation(PreviewSkelMeshComp.Get(), Drag, AnimNode);
-						bShouldUpdateDefaultValues = true;
-					}
-				}
-				else if (SelectedSocket)
-				{
-					SelectedSocket->RelativeLocation += BoneSpaceOffset;
-				}
-				else
-				{
-					SkelControl->Translation += BoneSpaceOffset;
-				}
-			}
-			if(bDoScale)
-			{
-				FVector4 BoneSpaceScaleOffset;
-
-				if (ModeTools->GetCoordSystem() == COORD_World)
-				{
-					BoneSpaceScaleOffset = BaseTM.TransformVector(Scale);
-				}
-				else
-				{
-					BoneSpaceScaleOffset = Scale;
-				}
-
-				if (SelectedSkelControlAnimGraph.IsValid())
-				{
-					FAnimNode_SkeletalControlBase* AnimNode = FindSkeletalControlAnimNode(SelectedSkelControlAnimGraph);
-					if (AnimNode)
-					{
-						SelectedSkelControlAnimGraph->DoScale(PreviewSkelMeshComp.Get(), Scale, AnimNode);
-						bShouldUpdateDefaultValues = true;
-					}
-				}
-				else if(SelectedSocket)
-				{
-					SelectedSocket->RelativeScale += BoneSpaceScaleOffset;
-				}
-				else
-				{
-					SkelControl->Scale += BoneSpaceScaleOffset;
-				}
-			}
-
-			}
-		else if( WindActor.IsValid() )
-		{
-			if (WidgetMode == FWidget::WM_Rotate)
-			{
-				FTransform WindTransform = WindActor->GetTransform();
-
-				FRotator NewRotation = ( WindTransform * FTransform( Rot ) ).Rotator();
-
-				WindActor->SetActorRotation( NewRotation );
-			}
-			else
-			{
-				FVector Location = WindActor->GetActorLocation();
-				Location += Drag;
-				WindActor->SetActorLocation(Location);
-			}
-		}
-
-		InViewport->Invalidate();
-	}
-
-	return bHandled;
-}
-
-
 void FAnimationViewportClient::TrackingStarted( const struct FInputEventState& InInputState, bool bIsDraggingWidget, bool bNudge )
 {
-	int32 BoneIndex = FindSelectedBone();
-	USkeletalMeshSocket* SelectedSocket = FindSelectedSocket();
-	TWeakObjectPtr<AWindDirectionalSource> WindActor = FindSelectedWindActor();
-
-	if( (BoneIndex >= 0 || SelectedSocket || WindActor.IsValid()) && bIsDraggingWidget )
-	{
-		bool bValidAxis = false;
-		FVector WorldAxisDir;
-
-		if(InInputState.IsLeftMouseButtonPressed() && (Widget->GetCurrentAxis() & EAxisList::XYZ) != 0)
-		{
-			if ( PreviewSkelMeshComp->SocketsOfInterest.Num() == 1 )
-		{
-			const bool bAltDown = InInputState.IsAltButtonPressed();
-
-			if ( bAltDown && PersonaPtr.IsValid() )
-			{
-				// Rather than moving/rotating the selected socket, copy it and move the copy instead
-				PersonaPtr.Pin()->DuplicateAndSelectSocket( PreviewSkelMeshComp->SocketsOfInterest[0] );
-			}
-
-			// Socket movement is transactional - we want undo/redo and saving of it
-			USkeletalMeshSocket* Socket = PreviewSkelMeshComp->SocketsOfInterest[0].Socket;
-
-			if ( Socket && bInTransaction == false )
-			{
-				if (WidgetMode == FWidget::WM_Rotate )
-				{
-					GEditor->BeginTransaction( LOCTEXT("AnimationEditorViewport_RotateSocket", "Rotate Socket" ) );
-				}
-				else
-				{
-					GEditor->BeginTransaction( LOCTEXT("AnimationEditorViewport_TranslateSocket", "Translate Socket" ) );
-				}
-
-				Socket->SetFlags( RF_Transactional );	// Undo doesn't work without this!
-				Socket->Modify();
-				bInTransaction = true;
-			}
-		}
-			else if( BoneIndex >= 0 )
-			{
-				if ( bInTransaction == false )
-				{
-					// we also allow undo/redo of bone manipulations
-					if (WidgetMode == FWidget::WM_Rotate )
-					{
-						GEditor->BeginTransaction( LOCTEXT("AnimationEditorViewport_RotateBone", "Rotate Bone" ) );
-					}
-					else
-					{
-						GEditor->BeginTransaction( LOCTEXT("AnimationEditorViewport_TranslateBone", "Translate Bone" ) );
-					}
-
-					PreviewSkelMeshComp->PreviewInstance->SetFlags( RF_Transactional );	// Undo doesn't work without this!
-					PreviewSkelMeshComp->PreviewInstance->Modify();
-					bInTransaction = true;
-
-					// now modify the bone array
-					const FName BoneName = PreviewSkelMeshComp->SkeletalMesh->RefSkeleton.GetBoneName(BoneIndex);
-					PreviewSkelMeshComp->PreviewInstance->ModifyBone(BoneName);
-				}
-			}
-		}
-
-
-		bManipulating = true;
-	}
+	ModeTools->StartTracking(this, Viewport);
 }
 
 void FAnimationViewportClient::TrackingStopped() 
 {
-	if (bManipulating)
-	{
-		// Socket movement is transactional - we want undo/redo and saving of it
-		if ( bInTransaction )
-		{
-			GEditor->EndTransaction();
-			bInTransaction = false;
-		}
+	ModeTools->EndTracking(this, Viewport);
 
-		bManipulating = false;
-	}
 	Invalidate();
 }
 
-FWidget::EWidgetMode FAnimationViewportClient::GetWidgetMode() const
-{
-	FWidget::EWidgetMode Mode = FWidget::WM_None;
-	if (!PreviewSkelMeshComp->IsAnimBlueprintInstanced())
-	{
-	if ((PreviewSkelMeshComp != NULL) && ((PreviewSkelMeshComp->BonesOfInterest.Num() == 1) || (PreviewSkelMeshComp->SocketsOfInterest.Num() == 1)))
-	{
-			Mode = WidgetMode;
-	}
-	else if (SelectedWindActor.IsValid())
-	{
-			Mode = WidgetMode;
-	}
-	}
-
-	if (SelectedSkelControlAnimGraph.IsValid())
-	{
-		Mode = (FWidget::EWidgetMode)SelectedSkelControlAnimGraph->GetWidgetMode(PreviewSkelMeshComp.Get());
-	}
-
-	return Mode;
-}
-
 FVector FAnimationViewportClient::GetWidgetLocation() const
 {
-	if (SelectedSkelControlAnimGraph.IsValid())
-	{
-		FAnimNode_SkeletalControlBase* AnimNode = FindSkeletalControlAnimNode(SelectedSkelControlAnimGraph.Get());
-		if (AnimNode)
-		{
-			return SelectedSkelControlAnimGraph->GetWidgetLocation(PreviewSkelMeshComp.Get(), AnimNode);
-		}
-	}
-	else if( PreviewSkelMeshComp->BonesOfInterest.Num() > 0 )
-	{
-		int32 BoneIndex = PreviewSkelMeshComp->BonesOfInterest.Last();
-		const FName BoneName = PreviewSkelMeshComp->SkeletalMesh->RefSkeleton.GetBoneName(BoneIndex);
-
-		FMatrix BoneMatrix = PreviewSkelMeshComp->GetBoneMatrix(BoneIndex);
-
-		return BoneMatrix.GetOrigin();
-	}
-	else if( PreviewSkelMeshComp->SocketsOfInterest.Num() > 0 )
-	{
-		USkeletalMeshSocket* Socket = PreviewSkelMeshComp->SocketsOfInterest.Last().Socket;
-
- 		FMatrix SocketMatrix;
-		Socket->GetSocketMatrix( SocketMatrix, PreviewSkelMeshComp.Get() );
-
-		return SocketMatrix.GetOrigin();
-	}
-	else if( SelectedWindActor.IsValid() )
-	{
-		return SelectedWindActor.Get()->GetActorLocation();
-	}
-
-	return FVector::ZeroVector;
+	return ModeTools->GetWidgetLocation();
 }
 
 FMatrix FAnimationViewportClient::GetWidgetCoordSystem() const
@@ -1551,36 +934,7 @@ FMatrix FAnimationViewportClient::GetWidgetCoordSystem() const
 
 	if( bIsLocal )
 	{
-		if (SelectedSkelControlAnimGraph.IsValid())
-		{
-			FName BoneName = SelectedSkelControlAnimGraph->FindSelectedBone();
-			int32 BoneIndex = PreviewSkelMeshComp->GetBoneIndex(BoneName);
-			if (BoneIndex != INDEX_NONE)
-			{
-				FTransform BoneMatrix = PreviewSkelMeshComp->GetBoneTransform(BoneIndex);
-				return BoneMatrix.ToMatrixNoScale().RemoveTranslation();
-			}
-		}
-		else if ( PreviewSkelMeshComp->BonesOfInterest.Num() > 0 )
-		{
-			int32 BoneIndex = PreviewSkelMeshComp->BonesOfInterest.Last();
-
-			FTransform BoneMatrix = PreviewSkelMeshComp->GetBoneTransform(BoneIndex);
-
-			return BoneMatrix.ToMatrixNoScale().RemoveTranslation();
-		}
-		else if( PreviewSkelMeshComp->SocketsOfInterest.Num() > 0 )
-		{
-			USkeletalMeshSocket* Socket = PreviewSkelMeshComp->SocketsOfInterest.Last().Socket;
-
-			FTransform SocketMatrix = Socket->GetSocketTransform( PreviewSkelMeshComp.Get() );
-			
-			return SocketMatrix.ToMatrixNoScale().RemoveTranslation();
-		}
-		else if ( SelectedWindActor.IsValid() )
-		{
-			return SelectedWindActor->GetTransform().ToMatrixNoScale().RemoveTranslation();
-		}
+		return ModeTools->GetCustomInputCoordinateSystem();
 	}
 
 	return FMatrix::Identity;
@@ -1618,54 +972,9 @@ void FAnimationViewportClient::RotateViewportType()
 
 bool FAnimationViewportClient::InputKey( FViewport* InViewport, int32 ControllerId, FKey Key, EInputEvent Event, float AmountDepressed, bool bGamepad )
 {
-	const int32 HitX = InViewport->GetMouseX();
-	const int32 HitY = InViewport->GetMouseY();
-
 	bool bHandled = false;
 
-	
-	// Handle switching modes - only allowed when not already manipulating
-	if ( (Event == IE_Pressed) && (Key == EKeys::SpaceBar) && !bManipulating )
-	{
-		int32 SelectedBone = FindSelectedBone();
-		USkeletalMeshSocket* SelectedSocket = FindSelectedSocket();
-		TWeakObjectPtr<AWindDirectionalSource> SelectedWind = FindSelectedWindActor();
-
-		if (SelectedBone >= 0 || SelectedSocket || SelectedWind.IsValid())
-		{
-			if (SelectedSkelControlAnimGraph.IsValid())
-			{
-				WidgetMode = (FWidget::EWidgetMode)SelectedSkelControlAnimGraph->ChangeToNextWidgetMode(PreviewSkelMeshComp.Get(), WidgetMode);
-				if (WidgetMode == FWidget::WM_Scale)
-				{
-					SetWidgetCoordSystemSpace(COORD_Local);
-				}
-				else
-				{
-					SetWidgetCoordSystemSpace(COORD_World);
-				}
-			}
-			else
-			if (WidgetMode == FWidget::WM_Rotate)
-			{
-				WidgetMode = FWidget::WM_Scale;
-			}
-			else if (WidgetMode == FWidget::WM_Scale)
-			{
-				WidgetMode = FWidget::WM_Translate;
-			}
-			else
-			{
-				WidgetMode = FWidget::WM_Rotate;
-			}
-
-			//@TODO: ANIMATION: Add scaling support
-		}
-		bHandled = true;
-		Invalidate();
-	}
-
-	if( (Event == IE_Pressed) && (Key == EKeys::F) )
+	if(Event == IE_Pressed && Key == EKeys::F)
 	{
 		bHandled = true;
 		FocusViewportOnPreviewMesh();
@@ -1677,7 +986,7 @@ bool FAnimationViewportClient::InputKey( FViewport* InViewport, int32 Controller
 	// Pass keys to standard controls, if we didn't consume input
 	return (bHandled)
 		? true
-		: FEditorViewportClient::InputKey(InViewport,  ControllerId,  Key,  Event,  AmountDepressed,  bGamepad);
+		: FEditorViewportClient::InputKey(InViewport, ControllerId, Key, Event, AmountDepressed, bGamepad);
 }
 
 bool FAnimationViewportClient::InputAxis(FViewport* InViewport, int32 ControllerId, FKey Key, float Delta, float DeltaTime, int32 NumSamples /*= 1*/, bool bGamepad /*= false*/)
@@ -1701,59 +1010,34 @@ bool FAnimationViewportClient::InputAxis(FViewport* InViewport, int32 Controller
 	return bResult;
 }
 
-void FAnimationViewportClient::SetWidgetMode(FWidget::EWidgetMode InMode)
-{
-	bool bAnimBPMode = PersonaPtr.Pin()->IsModeCurrent(FPersonaModes::AnimBlueprintEditMode);
-
-	// support WER keys to change gizmo mode for Bone Controller preivew in anim blueprint
-	if (bAnimBPMode)
-	{
-		if (!bManipulating)
-		{
-			int32 SelectedBone = FindSelectedBone();
-			if (SelectedBone >= 0 && SelectedSkelControlAnimGraph.IsValid())
-			{
-				if (SelectedSkelControlAnimGraph->SetWidgetMode(PreviewSkelMeshComp.Get(), InMode))
-				{
-					WidgetMode = InMode;
-				}
-			}
-		}
-	}
-	else
-	{
-		WidgetMode = InMode;
-	}
-
-	// Force viewport to redraw
-	Viewport->Invalidate();
-}
-
-bool FAnimationViewportClient::CanSetWidgetMode(FWidget::EWidgetMode NewMode) const
-{
-	return true;
-}
-
 void FAnimationViewportClient::SetLocalAxesMode(ELocalAxesMode::Type AxesMode)
 {
-	LocalAxesMode = AxesMode;
 	ConfigOption->SetDefaultLocalAxesSelection( AxesMode );
 }
 
 bool FAnimationViewportClient::IsLocalAxesModeSet( ELocalAxesMode::Type AxesMode ) const
 {
-	return LocalAxesMode == AxesMode;
+	return (ELocalAxesMode::Type)ConfigOption->DefaultLocalAxesSelection == AxesMode;
+}
+
+ELocalAxesMode::Type FAnimationViewportClient::GetLocalAxesMode() const
+{ 
+	return (ELocalAxesMode::Type)ConfigOption->DefaultLocalAxesSelection;
 }
 
 void FAnimationViewportClient::SetBoneDrawMode(EBoneDrawMode::Type AxesMode)
 {
-	BoneDrawMode = AxesMode;
 	ConfigOption->SetDefaultBoneDrawSelection(AxesMode);
 }
 
 bool FAnimationViewportClient::IsBoneDrawModeSet(EBoneDrawMode::Type AxesMode) const
 {
-	return BoneDrawMode == AxesMode;
+	return (EBoneDrawMode::Type)ConfigOption->DefaultBoneDrawSelection == AxesMode;
+}
+
+EBoneDrawMode::Type FAnimationViewportClient::GetBoneDrawMode() const 
+{ 
+	return (EBoneDrawMode::Type)ConfigOption->DefaultBoneDrawSelection;
 }
 
 void FAnimationViewportClient::DrawBonesFromTransforms(TArray<FTransform>& Transforms, UDebugSkelMeshComponent * MeshComponent, FPrimitiveDrawInterface* PDI, FLinearColor BoneColour, FLinearColor RootBoneColour) const
@@ -1925,8 +1209,8 @@ void FAnimationViewportClient::DrawBones(const USkeletalMeshComponent* MeshCompo
 			const int32 BoneIndex = RequiredBones[Index];
 
 			if (bForceDraw ||
-				(BoneDrawMode == EBoneDrawMode::All) ||
-				((BoneDrawMode == EBoneDrawMode::Selected) && SelectedBones.Contains(BoneIndex) )
+				(GetBoneDrawMode() == EBoneDrawMode::All) ||
+				((GetBoneDrawMode() == EBoneDrawMode::Selected) && SelectedBones.Contains(BoneIndex) )
 				)
 			{
 				const int32 ParentIndex = MeshComponent->SkeletalMesh->RefSkeleton.GetParentIndex(BoneIndex);
@@ -1959,8 +1243,8 @@ void FAnimationViewportClient::DrawBones(const USkeletalMeshComponent* MeshCompo
 				PDI->SetHitProxy(NULL);
 
 				// draw gizmo
-				if ((LocalAxesMode == ELocalAxesMode::All) ||
-					((LocalAxesMode == ELocalAxesMode::Selected) && SelectedBones.Contains(BoneIndex))
+				if ((GetLocalAxesMode() == ELocalAxesMode::All) ||
+					((GetLocalAxesMode() == ELocalAxesMode::Selected) && SelectedBones.Contains(BoneIndex))
 					)
 				{
 					RenderGizmo(WorldTransforms[BoneIndex], PDI);
@@ -1970,7 +1254,7 @@ void FAnimationViewportClient::DrawBones(const USkeletalMeshComponent* MeshCompo
 	}
 }
 
-void FAnimationViewportClient::RenderGizmo(const FTransform& Transform, FPrimitiveDrawInterface* PDI) const
+void FAnimationViewportClient::RenderGizmo(const FTransform& Transform, FPrimitiveDrawInterface* PDI)
 {
 	// Display colored coordinate system axes for this joint.
 	const float AxisLength = 3.75f;
@@ -2057,27 +1341,26 @@ void FAnimationViewportClient::DrawMeshSubsetBones(const USkeletalMeshComponent*
 	}
 }
 
-void FAnimationViewportClient::DrawSockets( TArray<USkeletalMeshSocket*>& Sockets, FPrimitiveDrawInterface* PDI, bool bUseSkeletonSocketColor ) const
+void FAnimationViewportClient::DrawSockets(const UDebugSkelMeshComponent* InPreviewMeshComponent, TArray<USkeletalMeshSocket*>& InSockets, FSelectedSocketInfo InSelectedSocket, FPrimitiveDrawInterface* PDI, bool bUseSkeletonSocketColor)
 {
-	if ( const UDebugSkelMeshComponent* DebugMeshComponent = PreviewSkelMeshComp.Get() )
+	if (InPreviewMeshComponent)
 	{
-		TArray<FSelectedSocketInfo> SelectedSockets;
-		SelectedSockets = DebugMeshComponent->SocketsOfInterest;
+		ELocalAxesMode::Type LocalAxesMode = (ELocalAxesMode::Type)GetDefault<UPersonaOptions>()->DefaultLocalAxesSelection;
 
-		for ( auto SocketIt = Sockets.CreateConstIterator(); SocketIt; ++SocketIt )
+		for ( auto SocketIt = InSockets.CreateConstIterator(); SocketIt; ++SocketIt )
 		{
 			USkeletalMeshSocket* Socket = *(SocketIt);
-			FReferenceSkeleton& RefSkeleton = DebugMeshComponent->SkeletalMesh->RefSkeleton;
+			FReferenceSkeleton& RefSkeleton = InPreviewMeshComponent->SkeletalMesh->RefSkeleton;
 
 			const int32 ParentIndex = RefSkeleton.FindBoneIndex(Socket->BoneName);
 
-			FTransform WorldTransformSocket = Socket->GetSocketTransform(DebugMeshComponent);
+			FTransform WorldTransformSocket = Socket->GetSocketTransform(InPreviewMeshComponent);
 
 			FLinearColor SocketColor;
 			FVector Start, End;
 			if (ParentIndex >=0)
 			{
-				FTransform WorldTransformParent = DebugMeshComponent->GetComponentSpaceTransforms()[ParentIndex]*DebugMeshComponent->ComponentToWorld;
+				FTransform WorldTransformParent = InPreviewMeshComponent->GetComponentSpaceTransforms()[ParentIndex] * InPreviewMeshComponent->ComponentToWorld;
 				Start = WorldTransformParent.GetLocation();
 				End = WorldTransformSocket.GetLocation();
 			}
@@ -2087,7 +1370,7 @@ void FAnimationViewportClient::DrawSockets( TArray<USkeletalMeshSocket*>& Socket
 				End = WorldTransformSocket.GetLocation();
 			}
 
-			const bool bSelectedSocket = SelectedSockets.ContainsByPredicate([Socket](const FSelectedSocketInfo& Info){ return Info.Socket == Socket; });
+			const bool bSelectedSocket = (InSelectedSocket.Socket == Socket);
 
 			if( bSelectedSocket )
 			{
@@ -2115,7 +1398,7 @@ void FAnimationViewportClient::DrawSockets( TArray<USkeletalMeshSocket*>& Socket
 			if( (LocalAxesMode == ELocalAxesMode::All) || bSelectedSocket )
 			{
 				FMatrix SocketMatrix;
-				Socket->GetSocketMatrix( SocketMatrix, PreviewSkelMeshComp.Get() );
+				Socket->GetSocketMatrix( SocketMatrix, InPreviewMeshComponent);
 
 				PDI->SetHitProxy( new HPersonaSocketProxy( FSelectedSocketInfo( Socket, bUseSkeletonSocketColor ) ) );
 				DrawWireDiamond( PDI, SocketMatrix, 2.f, SocketColor, SDPG_Foreground );
@@ -2131,38 +1414,23 @@ FSphere FAnimationViewportClient::GetCameraTarget()
 {
 	bool bFoundTarget = false;
 	FSphere Sphere(FVector(0,0,0), 100.0f); // default
-	if( PreviewSkelMeshComp.IsValid() )
+	UDebugSkelMeshComponent* PreviewMeshComponent = GetAnimPreviewScene()->GetPreviewMeshComponent();
+	if( PreviewMeshComponent != nullptr )
 	{
-		FTransform ComponentToWorld = PreviewSkelMeshComp->ComponentToWorld;
-		PreviewSkelMeshComp.Get()->CalcBounds(ComponentToWorld);
-
-		if( PreviewSkelMeshComp->BonesOfInterest.Num() > 0 )
-		{
-			const int32 FocusBoneIndex = PreviewSkelMeshComp->BonesOfInterest[0];
-			if( FocusBoneIndex != INDEX_NONE )
-			{
-				const FName BoneName = PreviewSkelMeshComp->SkeletalMesh->RefSkeleton.GetBoneName(FocusBoneIndex);
-				Sphere.Center = PreviewSkelMeshComp->GetBoneLocation(BoneName);
-				Sphere.W = 30.0f;
-				bFoundTarget = true;
-			}
-		}
+		FTransform ComponentToWorld = PreviewMeshComponent->ComponentToWorld;
+		PreviewMeshComponent->CalcBounds(ComponentToWorld);
 
-
-		if( !bFoundTarget && (PreviewSkelMeshComp->SocketsOfInterest.Num() > 0) )
+		// give the editor mode a chance to give us a camera target
+		FSphere Target;
+		if (GetPersonaModeManager().GetCameraTarget(Target))
 		{
-			USkeletalMeshSocket * Socket = PreviewSkelMeshComp->SocketsOfInterest[0].Socket;
-			if( Socket )
-			{
-				Sphere.Center = Socket->GetSocketLocation( PreviewSkelMeshComp.Get() );
-				Sphere.W = 30.0f;
-				bFoundTarget = true;
-			}
+			Sphere = Target;
+			bFoundTarget = true;
 		}
 
 		if( !bFoundTarget )
 		{
-			FBoxSphereBounds Bounds = PreviewSkelMeshComp.Get()->CalcBounds(FTransform::Identity);
+			FBoxSphereBounds Bounds = PreviewMeshComponent->CalcBounds(FTransform::Identity);
 			Sphere = Bounds.GetSphere();
 		}
 	}
@@ -2172,248 +1440,170 @@ FSphere FAnimationViewportClient::GetCameraTarget()
 void FAnimationViewportClient::UpdateCameraSetup()
 {
 	static FRotator CustomOrbitRotation(-33.75, -135, 0);
-	if (PreviewSkelMeshComp.IsValid() && PreviewSkelMeshComp->SkeletalMesh)
+	UDebugSkelMeshComponent* PreviewMeshComponent = GetAnimPreviewScene()->GetPreviewMeshComponent();
+	if (PreviewMeshComponent != nullptr && PreviewMeshComponent->SkeletalMesh)
 	{
 		FSphere BoundSphere = GetCameraTarget();
 		FVector CustomOrbitZoom(0, BoundSphere.W / (75.0f * (float)PI / 360.0f), 0);
 		FVector CustomOrbitLookAt = BoundSphere.Center;
 
 		SetCameraSetup(CustomOrbitLookAt, CustomOrbitRotation, CustomOrbitZoom, CustomOrbitLookAt, GetViewLocation(), GetViewRotation() );
-		
+
 		// Move the floor to the bottom of the bounding box of the mesh, rather than on the origin
-		FVector Bottom = PreviewSkelMeshComp->Bounds.GetBoxExtrema(0);
-		const float FloorOffset = GetFloorOffset() + (bAutoAlignFloor ? -Bottom.Z : 0.0f);
-		AdvancedPreviewScene->SetFloorOffset(FloorOffset);
+		FVector Bottom = PreviewMeshComponent->Bounds.GetBoxExtrema(0);
+
+		FVector FloorPos(0.f, 0.f, 0.f);
+		if (bAutoAlignFloor)
+		{
+			FloorPos.Z = GetFloorOffset() + Bottom.Z;
+		}
+		GetAnimPreviewScene()->SetFloorLocation(FloorPos);
 	}
 }
 
-void FAnimationViewportClient::FocusViewportOnSphere( FSphere& Sphere )
+void FAnimationViewportClient::FocusViewportOnSphere( FSphere& Sphere, bool bInstant /*= true*/ )
 {
 	FBox Box( Sphere.Center - FVector(Sphere.W, 0.0f, 0.0f), Sphere.Center + FVector(Sphere.W, 0.0f, 0.0f) );
 
-	bool bInstant = false;
 	FocusViewportOnBox( Box, bInstant );
 
 	Invalidate();
 }
 
-void FAnimationViewportClient::FocusViewportOnPreviewMesh()
+void FAnimationViewportClient::FocusViewportOnPreviewMesh(bool bInstant /*= true*/)
 {
-	FIntPoint ViewportSize = Viewport->GetSizeXY();
+	FIntPoint ViewportSize(FIntPoint::ZeroValue);
+	if (Viewport != nullptr)
+	{
+		ViewportSize = Viewport->GetSizeXY();
+	}
 
 	if(!(ViewportSize.SizeSquared() > 0))
 	{
 		// We cannot focus fully right now as the viewport does not know its size
 		// and we must have the aspect to correctly focus on the component,
 		bFocusOnDraw = true;
+		bInstantFocusOnDraw = bInstant;
 	}
-	FSphere Sphere = GetCameraTarget();
-	FocusViewportOnSphere(Sphere);
-}
-
-float FAnimationViewportClient::GetFloorOffset() const
-{
-	if ( PersonaPtr.IsValid() )
+	else
 	{
-		USkeletalMesh* Mesh = PersonaPtr.Pin()->GetMesh();
-
-		if ( Mesh )
+		// dont auto-focus if there is nothing to focus on
+		UPersonaPreviewSceneDescription* PreviewSceneDescription = GetAnimPreviewScene()->GetPreviewSceneDescription();
+		if (!(!PreviewSceneDescription->PreviewMesh.IsValid() && (!PreviewSceneDescription->AdditionalMeshes.IsValid() || PreviewSceneDescription->AdditionalMeshes.LoadSynchronous()->SkeletalMeshes.Num() == 0)))
 		{
-			return Mesh->FloorOffset;
+			FSphere Sphere = GetCameraTarget();
+			FocusViewportOnSphere(Sphere);
 		}
 	}
-
-	return 0.0f;
 }
 
-void FAnimationViewportClient::SetFloorOffset( float NewValue )
+float FAnimationViewportClient::GetFloorOffset() const
 {
-	if ( PersonaPtr.IsValid() )
+	USkeletalMesh* Mesh = GetPreviewScene()->GetPreviewMeshComponent()->SkeletalMesh;
+	if ( Mesh )
 	{
-		USkeletalMesh* Mesh = PersonaPtr.Pin()->GetMesh();
-
-		if ( Mesh )
-		{
-			// This value is saved in a UPROPERTY for the mesh, so changes are transactional
-			FScopedTransaction Transaction( LOCTEXT( "SetFloorOffset", "Set Floor Offset" ) );
-			Mesh->Modify();
-
-			Mesh->FloorOffset = NewValue;
-			UpdateCameraSetup(); // This does the actual moving of the floor mesh
-			Invalidate();
-		}
+		return Mesh->FloorOffset;
 	}
-}
-
-TWeakObjectPtr<AWindDirectionalSource> FAnimationViewportClient::CreateWindActor(UWorld* World)
-{
-	TWeakObjectPtr<AWindDirectionalSource> Wind = World->SpawnActor<AWindDirectionalSource>(PrevWindLocation, PrevWindRotation);
 
-	check(Wind.IsValid());
-	//initial wind strength value 
-	Wind->GetComponent()->Strength = PrevWindStrength;
-	return Wind;
+	return 0.0f;
 }
 
-void FAnimationViewportClient::EnableWindActor(bool bEnableWind)
+void FAnimationViewportClient::SetFloorOffset( float NewValue )
 {
-	UWorld* World = PersonaPtr.Pin()->PreviewComponent->GetWorld();
-	check(World);
-
-	if(bEnableWind)
-	{
-		if(!WindSourceActor.IsValid())
-		{
-			WindSourceActor = CreateWindActor(World);
-		}
-	}
-	else if(WindSourceActor.IsValid())
-	{
-		PrevWindLocation = WindSourceActor->GetActorLocation();
-		PrevWindRotation = WindSourceActor->GetActorRotation();
-		PrevWindStrength = WindSourceActor->GetComponent()->Strength;
+	USkeletalMesh* Mesh = GetPreviewScene()->GetPreviewMeshComponent()->SkeletalMesh;
 
-		if(World->DestroyActor(WindSourceActor.Get()))
-		{
-			//clear the gizmo (translation or rotation)
-			PersonaPtr.Pin()->DeselectAll();
-		}
-	}
-}
-
-void FAnimationViewportClient::SetWindStrength( float SliderPos )
-{
-	if(WindSourceActor.IsValid())
+	if ( Mesh )
 	{
-		//Clamp grid size slider value between 0 - 1
-		WindSourceActor->GetComponent()->Strength = FMath::Clamp<float>(SliderPos, 0.0f, 1.0f);
-		//to apply this new wind strength
-		WindSourceActor->UpdateComponentTransforms();
-	}
-}
+		// This value is saved in a UPROPERTY for the mesh, so changes are transactional
+		FScopedTransaction Transaction( LOCTEXT( "SetFloorOffset", "Set Floor Offset" ) );
+		Mesh->Modify();
 
-float FAnimationViewportClient::GetWindStrengthSliderValue() const
-{
-	if(WindSourceActor.IsValid())
-	{
-		return WindSourceActor->GetComponent()->Strength;
+		Mesh->FloorOffset = NewValue;
+		UpdateCameraSetup(); // This does the actual moving of the floor mesh
+		Invalidate();
 	}
-
-	return 0;
-}
-
-FText FAnimationViewportClient::GetWindStrengthLabel() const
-{
-	//Clamp slide value so that minimum value displayed is 0.00 and maximum is 1.0
-	float SliderValue = FMath::Clamp<float>(GetWindStrengthSliderValue(), 0.0f, 1.0f);
-
-	static const FNumberFormattingOptions FormatOptions = FNumberFormattingOptions()
-		.SetMinimumFractionalDigits(2)
-		.SetMaximumFractionalDigits(2);
-	return FText::AsNumber(SliderValue, &FormatOptions);
-}
-
-void FAnimationViewportClient::SetGravityScale( float SliderPos )
-{
-	GravityScaleSliderValue = SliderPos;
-	float RealGravityScale = SliderPos * 4;
-
-	check(PersonaPtr.Pin()->PreviewComponent);
-
-	UWorld* World = PersonaPtr.Pin()->PreviewComponent->GetWorld();
-	AWorldSettings* Setting = World->GetWorldSettings();
-	Setting->WorldGravityZ = UPhysicsSettings::Get()->DefaultGravityZ*RealGravityScale;
-	Setting->bWorldGravitySet = true;
-}
-
-float FAnimationViewportClient::GetGravityScaleSliderValue() const
-{
-	return GravityScaleSliderValue;
-}
-
-FText FAnimationViewportClient::GetGravityScaleLabel() const
-{
-	//Clamp slide value so that minimum value displayed is 0.00 and maximum is 4.0
-	float SliderValue = FMath::Clamp<float>(GetGravityScaleSliderValue()*4, 0.0f, 4.0f);
-
-	static const FNumberFormattingOptions FormatOptions = FNumberFormattingOptions()
-		.SetMinimumFractionalDigits(2)
-		.SetMaximumFractionalDigits(2);
-	return FText::AsNumber(SliderValue, &FormatOptions);
 }
 
 void FAnimationViewportClient::ToggleCPUSkinning()
 {
-	if (PreviewSkelMeshComp.IsValid())
+	UDebugSkelMeshComponent* PreviewMeshComponent = GetAnimPreviewScene()->GetPreviewMeshComponent();
+	if (PreviewMeshComponent != nullptr)
 	{
-		PreviewSkelMeshComp->bCPUSkinning = !PreviewSkelMeshComp->bCPUSkinning;
-		PreviewSkelMeshComp->MarkRenderStateDirty();
+		PreviewMeshComponent->bCPUSkinning = !PreviewMeshComponent->bCPUSkinning;
+		PreviewMeshComponent->MarkRenderStateDirty();
 		Invalidate();
 	}
 }
 
 bool FAnimationViewportClient::IsSetCPUSkinningChecked() const
 {
-	if (PreviewSkelMeshComp.IsValid())
+	UDebugSkelMeshComponent* PreviewMeshComponent = GetAnimPreviewScene()->GetPreviewMeshComponent();
+	if (PreviewMeshComponent != nullptr)
 	{
-		return PreviewSkelMeshComp->bCPUSkinning;
+		return PreviewMeshComponent->bCPUSkinning;
 	}
 	return false;
 }
 
 void FAnimationViewportClient::ToggleShowNormals()
 {
-	if(PreviewSkelMeshComp.IsValid())
+	UDebugSkelMeshComponent* PreviewMeshComponent = GetAnimPreviewScene()->GetPreviewMeshComponent();
+	if(PreviewMeshComponent != nullptr)
 	{
-		PreviewSkelMeshComp->bDrawNormals = !PreviewSkelMeshComp->bDrawNormals;
-		PreviewSkelMeshComp->MarkRenderStateDirty();
+		PreviewMeshComponent->bDrawNormals = !PreviewMeshComponent->bDrawNormals;
+		PreviewMeshComponent->MarkRenderStateDirty();
 		Invalidate();
 	}
 }
 
 bool FAnimationViewportClient::IsSetShowNormalsChecked() const
 {
-	if (PreviewSkelMeshComp.IsValid())
+	UDebugSkelMeshComponent* PreviewMeshComponent = GetAnimPreviewScene()->GetPreviewMeshComponent();
+	if (PreviewMeshComponent != nullptr)
 	{
-		return PreviewSkelMeshComp->bDrawNormals;
+		return PreviewMeshComponent->bDrawNormals;
 	}
 	return false;
 }
 
 void FAnimationViewportClient::ToggleShowTangents()
 {
-	if(PreviewSkelMeshComp.IsValid())
+	UDebugSkelMeshComponent* PreviewMeshComponent = GetAnimPreviewScene()->GetPreviewMeshComponent();
+	if(PreviewMeshComponent != nullptr)
 	{
-		PreviewSkelMeshComp->bDrawTangents = !PreviewSkelMeshComp->bDrawTangents;
-		PreviewSkelMeshComp->MarkRenderStateDirty();
+		PreviewMeshComponent->bDrawTangents = !PreviewMeshComponent->bDrawTangents;
+		PreviewMeshComponent->MarkRenderStateDirty();
 		Invalidate();
 	}
 }
 
 bool FAnimationViewportClient::IsSetShowTangentsChecked() const
 {
-	if (PreviewSkelMeshComp.IsValid())
+	UDebugSkelMeshComponent* PreviewMeshComponent = GetAnimPreviewScene()->GetPreviewMeshComponent();
+	if (PreviewMeshComponent != nullptr)
 	{
-		return PreviewSkelMeshComp->bDrawTangents;
+		return PreviewMeshComponent->bDrawTangents;
 	}
 	return false;
 }
 
 void FAnimationViewportClient::ToggleShowBinormals()
 {
-	if(PreviewSkelMeshComp.IsValid())
+	UDebugSkelMeshComponent* PreviewMeshComponent = GetAnimPreviewScene()->GetPreviewMeshComponent();
+	if(PreviewMeshComponent != nullptr)
 	{
-		PreviewSkelMeshComp->bDrawBinormals = !PreviewSkelMeshComp->bDrawBinormals;
-		PreviewSkelMeshComp->MarkRenderStateDirty();
+		PreviewMeshComponent->bDrawBinormals = !PreviewMeshComponent->bDrawBinormals;
+		PreviewMeshComponent->MarkRenderStateDirty();
 		Invalidate();
 	}
 }
 
 bool FAnimationViewportClient::IsSetShowBinormalsChecked() const
 {
-	if (PreviewSkelMeshComp.IsValid())
+	UDebugSkelMeshComponent* PreviewMeshComponent = GetAnimPreviewScene()->GetPreviewMeshComponent();
+	if (PreviewMeshComponent != nullptr)
 	{
-		return PreviewSkelMeshComp->bDrawBinormals;
+		return PreviewMeshComponent->bDrawBinormals;
 	}
 	return false;
 }
@@ -2437,14 +1627,13 @@ void FAnimationViewportClient::OnSetShowMeshStats(int32 ShowMode)
 bool FAnimationViewportClient::IsShowingMeshStats() const
 {
 	const bool bShouldBeEnabled = ConfigOption->ShowMeshStats != EDisplayInfoMode::None;
-	const bool bCanBeEnabled = !PersonaPtr.Pin()->IsModeCurrent(FPersonaModes::AnimBlueprintEditMode);
 
-	return bShouldBeEnabled && bCanBeEnabled;
+	return bShouldBeEnabled && bShowMeshStats;
 }
 
 bool FAnimationViewportClient::IsShowingSelectedNodeStats() const
 {
-	return PersonaPtr.Pin()->IsModeCurrent(FPersonaModes::AnimBlueprintEditMode) && ConfigOption->ShowMeshStats == EDisplayInfoMode::SkeletalControls;
+	return ConfigOption->ShowMeshStats == EDisplayInfoMode::SkeletalControls;
 }
 
 bool FAnimationViewportClient::IsDetailedMeshStats() const
@@ -2467,7 +1656,7 @@ void FAnimationViewportClient::OnAssetViewerSettingsChanged(const FName& InPrope
 
 void FAnimationViewportClient::SetAdvancedShowFlagsForScene()
 {
-	const bool bAdvancedShowFlags = UAssetViewerSettings::Get()->Profiles[AdvancedPreviewScene->GetCurrentProfileIndex()].bPostProcessingEnabled;
+	const bool bAdvancedShowFlags = UAssetViewerSettings::Get()->Profiles[GetPreviewScene()->GetCurrentProfileIndex()].bPostProcessingEnabled;
 	if (bAdvancedShowFlags)
 	{
 		EngineShowFlags.EnableAdvancedFeatures();
@@ -2478,5 +1667,35 @@ void FAnimationViewportClient::SetAdvancedShowFlagsForScene()
 	}
 }
 
+void FAnimationViewportClient::SetPlaybackSpeedMode(EAnimationPlaybackSpeeds::Type InMode)
+{
+	AnimationPlaybackSpeedMode = InMode;
+
+	if (GetWorld())
+	{
+		GetWorld()->GetWorldSettings()->TimeDilation = EAnimationPlaybackSpeeds::Values[AnimationPlaybackSpeedMode];
+	}
+}
+
+EAnimationPlaybackSpeeds::Type FAnimationViewportClient::GetPlaybackSpeedMode() const
+{
+	return AnimationPlaybackSpeedMode;
+}
+
+TSharedRef<FAnimationEditorPreviewScene> FAnimationViewportClient::GetAnimPreviewScene() const
+{
+	return StaticCastSharedRef<FAnimationEditorPreviewScene>(GetPreviewScene());
+}
+
+IPersonaEditorModeManager& FAnimationViewportClient::GetPersonaModeManager() const
+{
+	return *static_cast<IPersonaEditorModeManager*>(ModeTools);
+}
+
+void FAnimationViewportClient::HandleInvalidateViews()
+{
+	Invalidate();
+}
+
 #undef LOCTEXT_NAMESPACE
 
diff --git a/Engine/Source/Editor/Persona/Private/AnimationEditorViewportClient.h b/Engine/Source/Editor/Persona/Private/AnimationEditorViewportClient.h
--- a/Engine/Source/Editor/Persona/Private/AnimationEditorViewportClient.h
+++ b/Engine/Source/Editor/Persona/Private/AnimationEditorViewportClient.h
@@ -3,7 +3,6 @@
 #pragma once
 
 #include "EditorViewportClient.h"
-#include "Engine/WindDirectionalSource.h"
 
 class FAdvancedPreviewScene;
 struct FCompactHeapPose;
@@ -51,6 +50,51 @@ namespace EDisplayInfoMode
 	};
 };
 
+namespace EAnimationPlaybackSpeeds
+{
+	enum Type
+	{
+		OneTenth = 0,
+		Quarter,
+		Half,
+		Normal,
+		Double,
+		FiveTimes,
+		TenTimes,
+		NumPlaybackSpeeds
+	};
+
+	extern float Values[NumPlaybackSpeeds];
+};
+
+/**
+* A hit proxy class for sockets in the Persona viewport.
+*/
+struct HPersonaSocketProxy : public HHitProxy
+{
+	DECLARE_HIT_PROXY();
+
+	FSelectedSocketInfo SocketInfo;
+
+	explicit HPersonaSocketProxy(FSelectedSocketInfo InSocketInfo)
+		: SocketInfo(InSocketInfo)
+	{}
+};
+
+/**
+* A hit proxy class for sockets in the Persona viewport.
+*/
+struct HPersonaBoneProxy : public HHitProxy
+{
+	DECLARE_HIT_PROXY();
+
+	FName BoneName;
+
+	explicit HPersonaBoneProxy(const FName& InBoneName)
+		: BoneName(InBoneName)
+	{}
+};
+
 /////////////////////////////////////////////////////////////////////////
 // FAnimationViewportClient
 
@@ -58,9 +102,6 @@ class FAnimationViewportClient : public FEditorViewportClient
 {
 protected:
 
-	/** Skeletal Mesh Component used for preview */
-	TWeakObjectPtr<UDebugSkelMeshComponent> PreviewSkelMeshComp;
-
 	/** Function to display bone names*/
 	void ShowBoneNames( FCanvas* Canvas, FSceneView* View );
 
@@ -71,7 +112,7 @@ protected:
 	void DrawNodeDebugLines(TArray<FText>& Lines, FCanvas* Canvas, FSceneView* View);
 
 public:
-	FAnimationViewportClient(FAdvancedPreviewScene& InPreviewScene, TWeakPtr<FPersona> InPersonaPtr, const TSharedRef<SAnimationEditorViewport>& InAnimationEditorViewport);
+	FAnimationViewportClient(const TSharedRef<class ISkeletonTree>& InSkeletonTree, const TSharedRef<class IPersonaPreviewScene>& InPreviewScene, const TSharedRef<class SAnimationEditorViewport>& InAnimationEditorViewport, const TSharedRef<class FAssetEditorToolkit>& InAssetEditorToolkit, bool bInShowStats);
 	virtual ~FAnimationViewportClient();
 
 	// FEditorViewportClient interface
@@ -81,13 +122,13 @@ public:
 	virtual void DrawCanvas( FViewport& InViewport, FSceneView& View, FCanvas& Canvas ) override;
 	virtual bool InputKey(FViewport* Viewport, int32 ControllerId, FKey Key, EInputEvent Event, float AmountDepressed = 1.f, bool bGamepad=false) override;
 	virtual bool InputAxis(FViewport* InViewport, int32 ControllerId, FKey Key, float Delta, float DeltaTime, int32 NumSamples = 1, bool bGamepad = false) override;
-	virtual void ProcessClick(class FSceneView& View, class HHitProxy* HitProxy, FKey Key, EInputEvent Event, uint32 HitX, uint32 HitY) override;
-	virtual bool InputWidgetDelta( FViewport* Viewport, EAxisList::Type CurrentAxis, FVector& Drag, FRotator& Rot, FVector& Scale ) override;
+//	virtual void ProcessClick(class FSceneView& View, class HHitProxy* HitProxy, FKey Key, EInputEvent Event, uint32 HitX, uint32 HitY) override;
+//	virtual bool InputWidgetDelta( FViewport* Viewport, EAxisList::Type CurrentAxis, FVector& Drag, FRotator& Rot, FVector& Scale ) override;
 	virtual void TrackingStarted( const struct FInputEventState& InInputState, bool bIsDragging, bool bNudge ) override;
 	virtual void TrackingStopped() override;
-	virtual FWidget::EWidgetMode GetWidgetMode() const override;
-	virtual void SetWidgetMode(FWidget::EWidgetMode InWidgetMode) override;
-	virtual bool CanSetWidgetMode(FWidget::EWidgetMode NewMode) const override;
+//	virtual FWidget::EWidgetMode GetWidgetMode() const override;
+//	virtual void SetWidgetMode(FWidget::EWidgetMode InWidgetMode) override;
+//	virtual bool CanSetWidgetMode(FWidget::EWidgetMode NewMode) const override;
 	virtual FVector GetWidgetLocation() const override;
 	virtual FMatrix GetWidgetCoordSystem() const override;
 	virtual ECoordSystem GetWidgetCoordSystemSpace() const override;
@@ -106,8 +147,8 @@ public:
 	/** Callback for checking the camera lock flag. */
 	bool IsSetCameraFollowChecked() const;
 
-	/** Function to set the mesh component used for preview */
-	void SetPreviewMeshComponent(UDebugSkelMeshComponent* InPreviewSkelMeshComp);
+	/** Handle the skeletal mesh mesh component being used for preview changing */
+	void HandleSkeletalMeshChanged(class USkeletalMesh* InSkeletalMesh);
 
 	/** Function to display bone names*/
 	void ShowBoneNames(FViewport* Viewport, FCanvas* Canvas);
@@ -146,7 +187,7 @@ public:
 	bool IsLocalAxesModeSet(ELocalAxesMode::Type AxesMode) const;
 
 	/** Get the Bone local axis mode */
-	ELocalAxesMode::Type GetLocalAxesMode() const {return LocalAxesMode;}
+	ELocalAxesMode::Type GetLocalAxesMode() const;
 
 	/** Function to set Bone Draw  mode for the EBoneDrawType */
 	void SetBoneDrawMode(EBoneDrawMode::Type AxesMode);
@@ -155,7 +196,7 @@ public:
 	bool IsBoneDrawModeSet(EBoneDrawMode::Type AxesMode) const;
 
 	/** Get the Bone local axis mode */
-	EBoneDrawMode::Type GetBoneDrawMode() const { return BoneDrawMode; }
+	EBoneDrawMode::Type GetBoneDrawMode() const;
 	
 	/** Returns the desired target of the camera */
 	FSphere GetCameraTarget();
@@ -164,10 +205,10 @@ public:
 	void UpdateCameraSetup();
 
 	/* Places the viewport camera at a good location to view the supplied sphere */
-	void FocusViewportOnSphere( FSphere& Sphere );
+	void FocusViewportOnSphere( FSphere& Sphere, bool bInstant = true );
 
 	/* Places the viewport camera at a good location to view the preview target */
-	void FocusViewportOnPreviewMesh();
+	void FocusViewportOnPreviewMesh(bool bInstant = true);
 
 	/** Callback for toggling the normals show flag. */
 	void ToggleCPUSkinning();
@@ -206,40 +247,12 @@ public:
 	/** Sets the UV Channel that will be drawn when Draw UV Overlay is turned on */
 	void SetUVChannelToDraw(int32 UVChannel) { UVChannelToDraw = UVChannel; }
 
-	void ClearSelectedWindActor()
-	{
-		SelectedWindActor = NULL;
-	}
-
-	void ClearSelectedAnimGraphNode();
-	void PostUndo();
-	void PostCompile();
-
 	/* Returns the floor height offset */	
 	float GetFloorOffset() const;
 
 	/* Sets the floor height offset, saves it to config and invalidates the viewport so it shows up immediately */
 	void SetFloorOffset(float NewValue);
 
-	/* create AWindDirectionalSource actor to apply wind to clothes for preview */
-	TWeakObjectPtr<AWindDirectionalSource> CreateWindActor(UWorld* World);
-	/** if set to true, shows the wind actor to control wind direction 
-	*   and enables the wind strength slider
-	*/
-	void EnableWindActor(bool bEnableWind);
-	/** Function to set wind strength (0.0 - 1.0) */
-	void SetWindStrength(float SliderPos);
-	/** Function to get slide value used to represent wind strength */
-	float GetWindStrengthSliderValue() const;
-
-	/** Function to get wind strength label */
-	FText GetWindStrengthLabel() const;
-	/** Function to set gravity scale (0.0 - 4.0) */
-	void SetGravityScale(float SliderPos);
-	/** Function to get slide value used to represent gravity scale */
-	float GetGravityScaleSliderValue() const;
-	/** Function to get gravity scale label */
-	FText GetGravityScaleLabel() const;
 	/** Function to set mesh stat drawing state */
 	void OnSetShowMeshStats(int32 ShowMode);
 	/** Whether or not mesh stats are being displayed */
@@ -251,47 +264,59 @@ public:
 
 	int32 GetShowMeshStats() const;
 
+	/** Set the playback speed mode */
+	void SetPlaybackSpeedMode(EAnimationPlaybackSpeeds::Type InMode);
+
+	/** Get the playback speed mode */
+	EAnimationPlaybackSpeeds::Type GetPlaybackSpeedMode() const;
+
+	/** Get the skeleton tree we are bound to */
+	TSharedRef<class ISkeletonTree> GetSkeletonTree() const  { return SkeletonTreePtr.Pin().ToSharedRef(); }
+
+	/** Get the preview scene we are viewing */
+	TSharedRef<class IPersonaPreviewScene> GetPreviewScene() const { return PreviewScenePtr.Pin().ToSharedRef(); }
+
+	/** Get the asset editor we are embedded in */
+	TSharedRef<class FAssetEditorToolkit> GetAssetEditorToolkit() const { return AssetEditorToolkitPtr.Pin().ToSharedRef(); }
+
+	/* Handle error checking for additive base pose */
+	bool ShouldDisplayAdditiveScaleErrorMessage();
+
+	/** Draws Mesh Sockets in foreground - bUseSkeletonSocketColor = true for grey (skeleton), false for red (mesh) **/
+	static void DrawSockets(const UDebugSkelMeshComponent* InPreviewMeshComponent, TArray<USkeletalMeshSocket*>& InSockets, FSelectedSocketInfo InSelectedSocket, FPrimitiveDrawInterface* PDI, bool bUseSkeletonSocketColor);
+
+	/** Draws Gizmo for the Transform in foreground **/
+	static void RenderGizmo(const FTransform& Transform, FPrimitiveDrawInterface* PDI);
+
 public:
 
 	/** persona config options **/
 	UPersonaOptions* ConfigOption;
 
 private:
-	/** Weak pointer back to the FPersona that owns us */
-	TWeakPtr<FPersona> PersonaPtr;
+	/** Weak pointer back to the skeleton tree we are bound to */
+	TWeakPtr<class ISkeletonTree> SkeletonTreePtr;
+
+	/** Weak pointer back to the preview scene we are viewing */
+	TWeakPtr<class IPersonaPreviewScene> PreviewScenePtr;
+
+	/** Weak pointer back to asset editor we are embedded in */
+	TWeakPtr<class FAssetEditorToolkit> AssetEditorToolkitPtr;
 
 	// Current widget mode
 	FWidget::EWidgetMode WidgetMode;
 
-	/** True when the user is manipulating a bone widget. */
-	bool bManipulating;
-
 	/** add follow option @todo change to enum later - we share editorviewportclient, which is only problem*/
 	bool bCameraFollow;
 
 	/** Should we auto align floor to mesh bounds */
 	bool bAutoAlignFloor;
 
-	/** True when we're in an editor transaction (moving/rotating sockets) */
-	bool bInTransaction;
-
-	/** Control where we display local axes for bones/sockets */
-	ELocalAxesMode::Type LocalAxesMode;
-
-	/** Control where we display local axes for bones/sockets */
-	EBoneDrawMode::Type BoneDrawMode;
-
 	/** User selected color using color picker */
 	FLinearColor SelectedHSVColor;
 
-	/** User specified gravity scale value (0.0 - 1.0) */
-	float GravityScaleSliderValue;
-
-	/** Previous information of a wind actor */
-	FVector PrevWindLocation;
-	FRotator PrevWindRotation;
-	float PrevWindStrength;
-	TWeakObjectPtr<AWindDirectionalSource> WindSourceActor;
+	/** Selected playback speed mode, used for deciding scale */
+	EAnimationPlaybackSpeeds::Type AnimationPlaybackSpeedMode;
 
 	/** Flag for displaying the UV data in the viewport */
 	bool bDrawUVs;
@@ -306,19 +331,18 @@ private:
 		MaxGridSize	= 50,
 	};
 
-	/** Wind actor used for preview */
-	TWeakObjectPtr<AWindDirectionalSource> SelectedWindActor;
-
 	/** Focus on the preview component the next time we draw the viewport */
 	bool bFocusOnDraw;
+	bool bInstantFocusOnDraw;
+
+	/** Handle additive anim scale validation */
+	bool bDoesAdditiveRefPoseHaveZeroScale;
+	FGuid RefPoseGuid;
 
-	/** Distance to trace for physics bodies */
-	const float BodyTraceDistance;
 private:
-	int32 FindSelectedBone() const;
-	class USkeletalMeshSocket* FindSelectedSocket() const;
+
 	void SetSelectedBackgroundColor(const FLinearColor& RGBColor, bool bSave = true);
-	void SaveGridSize(float NewGridSize, bool bSave = true);
+
 	void SetCameraTargetLocation(const FSphere &BoundSphere, float DeltaSeconds);
 
 	/** Draws Mesh Bones in foreground **/
@@ -341,23 +365,18 @@ private:
 	void DrawBones(const USkeletalMeshComponent* MeshComponent, const TArray<FBoneIndexType> & RequiredBones, const TArray<FTransform> & WorldTransforms, FPrimitiveDrawInterface* PDI, const TArray<FLinearColor>& BoneColours, float LineThickness = 0.f, bool bForceDraw = false) const;
 	/** Draw Sub set of Bones **/
 	void DrawMeshSubsetBones(const USkeletalMeshComponent* MeshComponent, const TArray<int32>& BonesOfInterest, FPrimitiveDrawInterface* PDI) const;
-	/** Draws Gizmo for the Transform in foreground **/
-	void RenderGizmo(const FTransform& Transform, FPrimitiveDrawInterface* PDI) const;
-	/** Draws Mesh Sockets in foreground - bUseSkeletonSocketColor = true for grey (skeleton), false for red (mesh) **/
-	void DrawSockets( TArray<class USkeletalMeshSocket*>& Sockets, FPrimitiveDrawInterface* PDI, bool bUseSkeletonSocketColor ) const;
+
 	/** Draws bones from watched poses*/
 	void DrawWatchedPoses(UDebugSkelMeshComponent * MeshComponent, FPrimitiveDrawInterface* PDI);
 
-	TWeakObjectPtr<AWindDirectionalSource> FindSelectedWindActor() const;
+	/** Get the typed anim preview scene */
+	TSharedRef<class FAnimationEditorPreviewScene> GetAnimPreviewScene() const;
 
-	struct FAnimNode_SkeletalControlBase* FindSkeletalControlAnimNode(TWeakObjectPtr<class UAnimGraphNode_SkeletalControlBase> AnimGraphNode) const;
+	/** Get the persona mode manager */
+	class IPersonaEditorModeManager& GetPersonaModeManager() const;
 
-	void FindSelectedAnimGraphNode();
-
-	// selected skeletal control anim graph node 
-	TWeakObjectPtr<class UAnimGraphNode_SkeletalControlBase> SelectedSkelControlAnimGraph;
-	// to check whether we should update literal values in selected AnimGraphNode
-	bool bShouldUpdateDefaultValues;
+	/** Invalidate this view in response to a preview scene change */
+	void HandleInvalidateViews();
 
 	/** Delegate for preview profile is changed (used for updating show flags) */
 	void OnAssetViewerSettingsChanged(const FName& InPropertyName);
@@ -365,6 +384,7 @@ private:
 	/** Sets up the ShowFlag according to the current preview scene profile */
 	void SetAdvancedShowFlagsForScene();
 
-	/** Direct pointer to preview scene */
-	FAdvancedPreviewScene* AdvancedPreviewScene;
+private:
+	/** Allow mesh stats to be disabled for specific viewport instances */
+	bool bShowMeshStats;
 };
diff --git a/Engine/Source/Editor/Persona/Private/AnimationMode/AnimationMode.cpp b/Engine/Source/Editor/Persona/Private/AnimationMode/AnimationMode.cpp
--- a/Engine/Source/Editor/Persona/Private/AnimationMode/AnimationMode.cpp
+++ b/Engine/Source/Editor/Persona/Private/AnimationMode/AnimationMode.cpp
@@ -5,7 +5,6 @@
 #include "Persona.h"
 #include "Editor/Kismet/Public/BlueprintEditorTabs.h"
 #include "Editor/KismetWidgets/Public/SSingleObjectDetailsPanel.h"
-#include "Editor/PropertyEditor/Public/IDetailsView.h"
 
 #include "AnimationMode.h"
 #include "IDocumentation.h"
@@ -15,21 +14,28 @@
 /////////////////////////////////////////////////////
 // SAnimAssetPropertiesTabBody
 
-class SAnimAssetPropertiesTabBody : public SSingleObjectDetailsPanel
+class SAssetPropertiesTabBody : public SSingleObjectDetailsPanel
 {
 public:
-	SLATE_BEGIN_ARGS(SAnimAssetPropertiesTabBody) {}
+	SLATE_BEGIN_ARGS(SAssetPropertiesTabBody) {}
+
+	SLATE_ARGUMENT(FOnGetAsset, OnGetAsset)
+
+	SLATE_ARGUMENT(FOnDetailsCreated, OnDetailsCreated)
+
 	SLATE_END_ARGS()
 
 private:
-	// Pointer back to owning Persona instance (the keeper of state)
-	TWeakPtr<class FPersona> PersonaPtr;
+	FOnGetAsset OnGetAsset;
+
 public:
-	void Construct(const FArguments& InArgs, TSharedPtr<FPersona> InPersona)
+	void Construct(const FArguments& InArgs)
 	{
-		PersonaPtr = InPersona;
+		OnGetAsset = InArgs._OnGetAsset;
+
+		SSingleObjectDetailsPanel::Construct(SSingleObjectDetailsPanel::FArguments(), true, true);
 
-		SSingleObjectDetailsPanel::Construct(SSingleObjectDetailsPanel::FArguments().HostCommandList(InPersona->GetToolkitCommands()), true, true);
+		InArgs._OnDetailsCreated.ExecuteIfBound(PropertyView.ToSharedRef());
 	}
 
 	virtual EVisibility GetAssetDisplayNameVisibility() const
@@ -52,92 +58,39 @@ public:
 	// SSingleObjectDetailsPanel interface
 	virtual UObject* GetObjectToObserve() const override
 	{
-		return PersonaPtr.Pin()->GetAnimationAssetBeingEdited();
-	}
-
-	bool IsReadOnlyProperty()
-	{
-		class UAnimationAsset* AnimAsset = Cast<UAnimationAsset> (PersonaPtr.Pin()->GetAnimationAssetBeingEdited());
-		if (AnimAsset)
+		if (OnGetAsset.IsBound())
 		{
-			// if child anim montage, you can't edit property for it. 
-			return !AnimAsset->HasParentAsset();
-		}
-
-		return false;
+			return OnGetAsset.Execute();
 	}
-
-	virtual TSharedRef<SWidget> PopulateSlot(TSharedRef<SWidget> PropertyEditorWidget) override
-	{
-		PropertyView->SetIsPropertyEditingEnabledDelegate(FIsPropertyEditingEnabled::CreateSP(this, &SAnimAssetPropertiesTabBody::IsReadOnlyProperty));
-
-		return SNew(SVerticalBox)
-			+SVerticalBox::Slot()
-			.AutoHeight()
-			[
-				// Header, shows name of the blend space we are editing
-				SNew(SBorder)
-				.BorderImage( FEditorStyle::GetBrush( TEXT("Graph.TitleBackground") ) )
-				.HAlign(HAlign_Center)
-				.Visibility(this, &SAnimAssetPropertiesTabBody::GetAssetDisplayNameVisibility)
-				[
-					SNew(SHorizontalBox)
-					+SHorizontalBox::Slot()
-					.AutoWidth()
-					.VAlign(VAlign_Center)
-					[
-						SNew(STextBlock)
-						.Font( FSlateFontInfo(FPaths::EngineContentDir() / TEXT("Slate/Fonts/Roboto-Regular.ttf"), 14 ) )
-						.ColorAndOpacity( FLinearColor(1,1,1,0.5) )
-						.Text( this, &SAnimAssetPropertiesTabBody::GetAssetDisplayName)
-					]
-				]
-			]
-
-			+SVerticalBox::Slot()
-			.FillHeight(1)
-			[
-				PropertyEditorWidget
-			];
+		return nullptr;
 	}
 	// End of SSingleObjectDetailsPanel interface
 };
 
-/////////////////////////////////////////////////////
-// FAnimAssetPropertiesSummoner
-
-struct FAnimAssetPropertiesSummoner : public FWorkflowTabFactory
-{
-	FAnimAssetPropertiesSummoner(TSharedPtr<class FAssetEditorToolkit> InHostingApp);
-
-	// FWorkflowTabFactory interface
-	virtual TSharedRef<SWidget> CreateTabBody(const FWorkflowTabSpawnInfo& Info) const override;
-
-	// Create a tooltip widget for the tab
-	virtual TSharedPtr<SToolTip> CreateTabToolTipWidget(const FWorkflowTabSpawnInfo& Info) const override
-	{
-		return  IDocumentation::Get()->CreateToolTip(LOCTEXT("AnimAssetPropertiesTooltip", "The Anim Asset Details tab lets you edit properties of the selection animation asset (animation, blend space etc)."), NULL, TEXT("Shared/Editors/Persona"), TEXT("AnimationAssetDetail_Window"));
-	}
-	// FWorkflowTabFactory interface
-};
-
-FAnimAssetPropertiesSummoner::FAnimAssetPropertiesSummoner(TSharedPtr<class FAssetEditorToolkit> InHostingApp)
+FAssetPropertiesSummoner::FAssetPropertiesSummoner(TSharedPtr<class FAssetEditorToolkit> InHostingApp, FOnGetAsset InOnGetAsset, FOnDetailsCreated InOnDetailsCreated)
 	: FWorkflowTabFactory(FPersonaTabs::AnimAssetPropertiesID, InHostingApp)
+	, OnGetAsset(InOnGetAsset)
+	, OnDetailsCreated(InOnDetailsCreated)
 {
-	TabLabel = LOCTEXT("AnimAssetProperties_TabTitle", "Anim Asset Details");
+	TabLabel = LOCTEXT("AssetProperties_TabTitle", "Asset Details");
 	TabIcon = FSlateIcon(FEditorStyle::GetStyleSetName(), "Persona.Tabs.AnimAssetDetails");
 
 	bIsSingleton = true;
 
-	ViewMenuDescription = LOCTEXT("AnimAssetProperties_MenuTitle", "Anim Asset Details");
-	ViewMenuTooltip = LOCTEXT("AnimAssetProperties_MenuToolTip", "Shows the animation asset properties");
+	ViewMenuDescription = LOCTEXT("AssetProperties_MenuTitle", "Asset Details");
+	ViewMenuTooltip = LOCTEXT("AssetProperties_MenuToolTip", "Shows the asset properties");
 }
 
-TSharedRef<SWidget> FAnimAssetPropertiesSummoner::CreateTabBody(const FWorkflowTabSpawnInfo& Info) const
+TSharedPtr<SToolTip> FAssetPropertiesSummoner::CreateTabToolTipWidget(const FWorkflowTabSpawnInfo& Info) const
 {
-	TSharedPtr<FPersona> PersonaApp = StaticCastSharedPtr<FPersona>(HostingApp.Pin());
+	return  IDocumentation::Get()->CreateToolTip(LOCTEXT("AssetPropertiesTooltip", "The Asset Details tab lets you edit properties of the current asset (animation, blend space etc)."), NULL, TEXT("Shared/Editors/Persona"), TEXT("AnimationAssetDetail_Window"));
+}
 
-	return SNew(SAnimAssetPropertiesTabBody, PersonaApp);
+TSharedRef<SWidget> FAssetPropertiesSummoner::CreateTabBody(const FWorkflowTabSpawnInfo& Info) const
+{
+	return SNew(SAssetPropertiesTabBody)
+		.OnGetAsset(OnGetAsset)
+		.OnDetailsCreated(OnDetailsCreated);
 }
 
 /////////////////////////////////////////////////////
@@ -147,7 +100,7 @@ FAnimEditAppMode::FAnimEditAppMode(TSharedPtr<FPersona> InPersona)
 	: FPersonaAppMode(InPersona, FPersonaModes::AnimationEditMode)
 {
 	PersonaTabFactories.RegisterFactory(MakeShareable(new FSelectionDetailsSummoner(InPersona)));
-	PersonaTabFactories.RegisterFactory(MakeShareable(new FAnimAssetPropertiesSummoner(InPersona)));
+	PersonaTabFactories.RegisterFactory(MakeShareable(new FAssetPropertiesSummoner(InPersona, FOnGetAsset::CreateSP(InPersona.Get(), &FPersona::HandleGetObject), FOnDetailsCreated())));
 
 	TabLayout = FTabManager::NewLayout("Persona_AnimEditMode_Layout_v7")
 		->AddArea
diff --git a/Engine/Source/Editor/Persona/Private/AnimationMode/AnimationMode.h b/Engine/Source/Editor/Persona/Private/AnimationMode/AnimationMode.h
--- a/Engine/Source/Editor/Persona/Private/AnimationMode/AnimationMode.h
+++ b/Engine/Source/Editor/Persona/Private/AnimationMode/AnimationMode.h
@@ -2,6 +2,9 @@
 
 #pragma once
 
+#include "PersonaDelegates.h"
+#include "PersonaModule.h"
+
 /////////////////////////////////////////////////////
 // FAnimEditAppMode
 
@@ -10,3 +13,20 @@ class FAnimEditAppMode : public FPersonaAppMode
 public:
 	FAnimEditAppMode(TSharedPtr<class FPersona> InPersona);
 };
+
+/////////////////////////////////////////////////////
+// FAnimAssetPropertiesSummoner
+
+struct FAssetPropertiesSummoner : public FWorkflowTabFactory
+{
+	FAssetPropertiesSummoner(TSharedPtr<class FAssetEditorToolkit> InHostingApp, FOnGetAsset InOnGetAsset, FOnDetailsCreated InOnDetailsCreated);
+
+	// FWorkflowTabFactory interface
+	virtual TSharedRef<SWidget> CreateTabBody(const FWorkflowTabSpawnInfo& Info) const override;
+	virtual TSharedPtr<SToolTip> CreateTabToolTipWidget(const FWorkflowTabSpawnInfo& Info) const override;
+	// FWorkflowTabFactory interface
+
+private:
+	FOnGetAsset OnGetAsset;
+	FOnDetailsCreated OnDetailsCreated;
+};
\ No newline at end of file
