Copying //UE4/Dev-Framework to //UE4/Dev-Main (Source: //UE4/Dev-Framework @ 3136612)

diff --git a/Engine/Source/Editor/Persona/Private/BlueprintMode/AnimBlueprintMode.cpp b/Engine/Source/Editor/Persona/Private/BlueprintMode/AnimBlueprintMode.cpp
--- a/Engine/Source/Editor/Persona/Private/BlueprintMode/AnimBlueprintMode.cpp
+++ b/Engine/Source/Editor/Persona/Private/BlueprintMode/AnimBlueprintMode.cpp
@@ -14,6 +14,8 @@
 #include "SSkeletonAnimNotifies.h"
 
 #include "AnimBlueprintMode.h"
+#include "IPersonaToolkit.h"
+#include "ISkeletonTree.h"
 
 /////////////////////////////////////////////////////
 // FAnimBlueprintEditAppMode
@@ -110,13 +112,13 @@ FAnimBlueprintEditAppMode::FAnimBlueprintEditAppMode(TSharedPtr<FPersona> InPers
 		);
 
 	PersonaTabFactories.RegisterFactory(MakeShareable(new FSkeletonTreeSummoner(InPersona)));
-	PersonaTabFactories.RegisterFactory(MakeShareable(new FAnimationAssetBrowserSummoner(InPersona)));
-	PersonaTabFactories.RegisterFactory(MakeShareable(new FPreviewViewportSummoner(InPersona)));
-	PersonaTabFactories.RegisterFactory(MakeShareable(new FSkeletonAnimNotifiesSummoner(InPersona)));
-	PersonaTabFactories.RegisterFactory(MakeShareable(new FAnimBlueprintPreviewEditorSummoner(InPersona)));
-	PersonaTabFactories.RegisterFactory(MakeShareable(new FAnimBlueprintParentPlayerEditorSummoner(InPersona)));
-	PersonaTabFactories.RegisterFactory(MakeShareable(new FSkeletonSlotNamesSummoner(InPersona)));
-	PersonaTabFactories.RegisterFactory(MakeShareable(new FAdvancedPreviewSceneTabSummoner(InPersona)));
+	PersonaTabFactories.RegisterFactory(MakeShareable(new FAnimationAssetBrowserSummoner(InPersona, InPersona->GetPersonaToolkit(), FOnOpenNewAsset::CreateSP(InPersona.Get(), &FPersona::HandleOpenNewAsset), FOnAnimationSequenceBrowserCreated(), false)));
+	PersonaTabFactories.RegisterFactory(MakeShareable(new FPreviewViewportSummoner(InPersona, InPersona->GetSkeletonTree(), InPersona->GetPreviewScene(), InPersona->OnPostUndo, InPersona, FOnViewportCreated::CreateSP(InPersona.Get(), &FPersona::HandleViewportCreated), false, false)));
+	PersonaTabFactories.RegisterFactory(MakeShareable(new FSkeletonAnimNotifiesSummoner(InPersona, InPersona->GetSkeletonTree()->GetEditableSkeleton(), InPersona->OnAnimNotifiesChanged, InPersona->OnPostUndo, FOnObjectsSelected::CreateSP(InPersona.Get(), &FPersona::HandleObjectsSelected))));
+	PersonaTabFactories.RegisterFactory(MakeShareable(new FAnimBlueprintPreviewEditorSummoner(InPersona, InPersona->GetPreviewScene())));
+	PersonaTabFactories.RegisterFactory(MakeShareable(new FAnimBlueprintParentPlayerEditorSummoner(InPersona, InPersona->OnPostUndo)));
+	PersonaTabFactories.RegisterFactory(MakeShareable(new FSkeletonSlotNamesSummoner(InPersona, InPersona->GetSkeletonTree()->GetEditableSkeleton(), InPersona->OnPostUndo, FOnObjectSelected::CreateSP(InPersona.Get(), &FPersona::HandleObjectSelected))));
+	PersonaTabFactories.RegisterFactory(MakeShareable(new FAdvancedPreviewSceneTabSummoner(InPersona, InPersona->GetPreviewScene())));
 
 	// setup toolbar - clear existing toolbar extender from the BP mode
 	//@TODO: Keep this in sync with BlueprintEditorModes.cpp
@@ -146,10 +148,10 @@ void FAnimBlueprintEditAppMode::PostActivateMode()
 	if (UAnimBlueprint* AnimBlueprint = Persona->GetAnimBlueprint())
 	{
 		// Switch off any active preview when going to graph editing mode
-		Persona->SetPreviewAnimationAsset(NULL, false);
+		Persona->GetPersonaToolkit()->GetPreviewScene()->SetPreviewAnimationAsset(NULL, false);
 
 		// When switching to anim blueprint mode, make sure the object being debugged is either a valid world object or the preview instance
-		UDebugSkelMeshComponent* PreviewComponent = Persona->GetPreviewMeshComponent();
+		UDebugSkelMeshComponent* PreviewComponent = Persona->GetPersonaToolkit()->GetPreviewScene()->GetPreviewMeshComponent();
 		if ((AnimBlueprint->GetObjectBeingDebugged() == NULL) && (PreviewComponent->IsAnimBlueprintInstanced()))
 		{
 			AnimBlueprint->SetObjectBeingDebugged(PreviewComponent->GetAnimInstance());
diff --git a/Engine/Source/Editor/Persona/Private/Customization/AnimGraphNodeSlotDetails.cpp b/Engine/Source/Editor/Persona/Private/Customization/AnimGraphNodeSlotDetails.cpp
--- a/Engine/Source/Editor/Persona/Private/Customization/AnimGraphNodeSlotDetails.cpp
+++ b/Engine/Source/Editor/Persona/Private/Customization/AnimGraphNodeSlotDetails.cpp
@@ -17,8 +17,8 @@
 #define LOCTEXT_NAMESPACE "AnimNodeSlotDetails"
 
 ////////////////////////////////////////////////////////////////
-FAnimGraphNodeSlotDetails::FAnimGraphNodeSlotDetails(TWeakPtr<class FPersona> InPersonalEditor)
-:PersonaEditor(InPersonalEditor)
+FAnimGraphNodeSlotDetails::FAnimGraphNodeSlotDetails(FOnInvokeTab InOnInvokeTab)
+	: OnInvokeTab(InOnInvokeTab)
 {
 }
 
@@ -192,10 +192,7 @@ void FAnimGraphNodeSlotDetails::OnSlotListOpening()
 
 FReply FAnimGraphNodeSlotDetails::OnOpenAnimSlotManager()
 {
-	if(PersonaEditor.IsValid())
-	{
-		PersonaEditor.Pin()->GetTabManager()->InvokeTab(FPersonaTabs::SkeletonSlotNamesID);
-	}
+	OnInvokeTab.ExecuteIfBound(FPersonaTabs::SkeletonSlotNamesID);
 	return FReply::Handled();
 }
 
diff --git a/Engine/Source/Editor/Persona/Private/Customization/AnimGraphNodeSlotDetails.h b/Engine/Source/Editor/Persona/Private/Customization/AnimGraphNodeSlotDetails.h
--- a/Engine/Source/Editor/Persona/Private/Customization/AnimGraphNodeSlotDetails.h
+++ b/Engine/Source/Editor/Persona/Private/Customization/AnimGraphNodeSlotDetails.h
@@ -9,19 +9,19 @@
 class FAnimGraphNodeSlotDetails : public IDetailCustomization
 {
 public:
-	FAnimGraphNodeSlotDetails(TWeakPtr<class FPersona> PersonalEditor);
+	FAnimGraphNodeSlotDetails(FOnInvokeTab InOnInvokeTab);
 
-	static TSharedRef<IDetailCustomization> MakeInstance(TWeakPtr<class FPersona> PersonalEditor) 
+	static TSharedRef<IDetailCustomization> MakeInstance(FOnInvokeTab InOnInvokeTab) 
 	{
-		return MakeShareable( new FAnimGraphNodeSlotDetails(PersonalEditor) );
+		return MakeShareable( new FAnimGraphNodeSlotDetails(InOnInvokeTab) );
 	}
 
 	// IDetailCustomization interface
 	virtual void CustomizeDetails(class IDetailLayoutBuilder& DetailBuilder) override;
 	
 private:
-	// persona editor 
-	TWeakPtr<class FPersona> PersonaEditor;
+	// delegate used to invoke a tab in the containing editor
+	FOnInvokeTab OnInvokeTab;
 
 	// property of the two 
 	TSharedPtr<IPropertyHandle> SlotNodeNamePropertyHandle;
diff --git a/Engine/Source/Editor/Persona/Private/EditModes/SkeletonSelectionEditMode.cpp b/Engine/Source/Editor/Persona/Private/EditModes/SkeletonSelectionEditMode.cpp
new file mode 100644
--- /dev/null
+++ b/Engine/Source/Editor/Persona/Private/EditModes/SkeletonSelectionEditMode.cpp
@@ -0,0 +1,541 @@
+// Copyright 1998-2016 Epic Games, Inc. All Rights Reserved.
+
+#include "PersonaPrivatePCH.h"
+#include "SkeletonSelectionEditMode.h"
+#include "AnimationEditorViewportClient.h"
+#include "Animation/DebugSkelMeshComponent.h"
+#include "AnimPreviewInstance.h"
+#include "AnimNode_ModifyBone.h"
+#include "ISkeletonTree.h"
+#include "IPersonaPreviewScene.h"
+#include "AssetEditorModeManager.h"
+
+#define LOCTEXT_NAMESPACE "SkeletonSelectionEditMode"
+
+namespace SkeletonSelectionModeConstants
+{
+	/** Distance to trace for physics bodies */
+	static const float BodyTraceDistance = 10000.0f;
+}
+
+FSkeletonSelectionEditMode::FSkeletonSelectionEditMode()
+	: bManipulating(false)
+	, bInTransaction(false)
+{
+	// Disable grid drawing for this mode as the viewport handles this
+	bDrawGrid = false;
+}
+
+bool FSkeletonSelectionEditMode::GetCameraTarget(FSphere& OutTarget) const
+{
+	bool bHandled = false;
+
+	UDebugSkelMeshComponent* PreviewMeshComponent = GetAnimPreviewScene().GetPreviewMeshComponent();
+
+	if (GetAnimPreviewScene().GetSelectedBoneIndex() != INDEX_NONE)
+	{
+		const int32 FocusBoneIndex = GetAnimPreviewScene().GetSelectedBoneIndex();
+		if (FocusBoneIndex != INDEX_NONE)
+		{
+			const FName BoneName = PreviewMeshComponent->SkeletalMesh->RefSkeleton.GetBoneName(FocusBoneIndex);
+			OutTarget.Center = PreviewMeshComponent->GetBoneLocation(BoneName);
+			OutTarget.W = 30.0f;
+			bHandled = true;
+		}
+	}
+
+	if (!bHandled && GetAnimPreviewScene().GetSelectedSocket().IsValid())
+	{
+		USkeletalMeshSocket * Socket = GetAnimPreviewScene().GetSelectedSocket().Socket;
+		if (Socket)
+		{
+			OutTarget.Center = Socket->GetSocketLocation(PreviewMeshComponent);
+			OutTarget.W = 30.0f;
+			bHandled = true;
+		}
+	}
+
+	return bHandled;
+}
+
+IPersonaPreviewScene& FSkeletonSelectionEditMode::GetAnimPreviewScene() const
+{
+	return *static_cast<IPersonaPreviewScene*>(static_cast<FAssetEditorModeManager*>(Owner)->GetPreviewScene());
+}
+
+void FSkeletonSelectionEditMode::GetOnScreenDebugInfo(TArray<FText>& OutDebugInfo) const
+{
+
+}
+
+bool FSkeletonSelectionEditMode::StartTracking(FEditorViewportClient* InViewportClient, FViewport* InViewport)
+{
+	EAxisList::Type CurrentAxis = InViewportClient->GetCurrentWidgetAxis();
+	FWidget::EWidgetMode WidgetMode = InViewportClient->GetWidgetMode();
+
+	int32 BoneIndex = GetAnimPreviewScene().GetSelectedBoneIndex();
+	USkeletalMeshSocket* SelectedSocket = GetAnimPreviewScene().GetSelectedSocket().Socket;
+	AActor* SelectedActor = GetAnimPreviewScene().GetSelectedActor();
+
+	if (BoneIndex >= 0 || SelectedSocket != nullptr || SelectedActor != nullptr)
+	{
+		bool bValidAxis = false;
+		FVector WorldAxisDir;
+
+		if ((CurrentAxis & EAxisList::XYZ) != 0)
+		{
+			UDebugSkelMeshComponent* PreviewMeshComponent = GetAnimPreviewScene().GetPreviewMeshComponent();
+			FSelectedSocketInfo SelectedSocketInfo = GetAnimPreviewScene().GetSelectedSocket();
+			if (SelectedSocketInfo.IsValid())
+			{
+				const bool bAltDown = InViewportClient->IsAltPressed();
+
+				if (bAltDown)
+				{
+					// Rather than moving/rotating the selected socket, copy it and move the copy instead
+					static_cast<FAnimationViewportClient*>(InViewportClient)->GetSkeletonTree()->DuplicateAndSelectSocket(SelectedSocketInfo);
+				}
+
+				// Socket movement is transactional - we want undo/redo and saving of it
+				USkeletalMeshSocket* Socket = SelectedSocketInfo.Socket;
+
+				if (Socket && bInTransaction == false)
+				{
+					if (WidgetMode == FWidget::WM_Rotate)
+					{
+						GEditor->BeginTransaction(LOCTEXT("AnimationEditorViewport_RotateSocket", "Rotate Socket"));
+					}
+					else
+					{
+						GEditor->BeginTransaction(LOCTEXT("AnimationEditorViewport_TranslateSocket", "Translate Socket"));
+					}
+
+					Socket->SetFlags(RF_Transactional);	// Undo doesn't work without this!
+					Socket->Modify();
+					bInTransaction = true;
+				}
+			}
+			else if (BoneIndex >= 0)
+			{
+				if (bInTransaction == false)
+				{
+					// we also allow undo/redo of bone manipulations
+					if (WidgetMode == FWidget::WM_Rotate)
+					{
+						GEditor->BeginTransaction(LOCTEXT("AnimationEditorViewport_RotateBone", "Rotate Bone"));
+					}
+					else
+					{
+						GEditor->BeginTransaction(LOCTEXT("AnimationEditorViewport_TranslateBone", "Translate Bone"));
+					}
+
+					PreviewMeshComponent->PreviewInstance->SetFlags(RF_Transactional);	// Undo doesn't work without this!
+					PreviewMeshComponent->PreviewInstance->Modify();
+					bInTransaction = true;
+
+					// now modify the bone array
+					const FName BoneName = PreviewMeshComponent->SkeletalMesh->RefSkeleton.GetBoneName(BoneIndex);
+					PreviewMeshComponent->PreviewInstance->ModifyBone(BoneName);
+				}
+			}
+		}
+
+		bManipulating = true;
+		return true;
+	}
+
+	return false;
+}
+
+bool FSkeletonSelectionEditMode::EndTracking(FEditorViewportClient* InViewportClient, FViewport* InViewport)
+{
+	if (bManipulating)
+	{
+		// Socket movement is transactional - we want undo/redo and saving of it
+		if (bInTransaction)
+		{
+			GEditor->EndTransaction();
+			bInTransaction = false;
+		}
+
+		bManipulating = false;
+		return true;
+	}
+
+	return false;
+}
+
+bool FSkeletonSelectionEditMode::InputKey(FEditorViewportClient* InViewportClient, FViewport* InViewport, FKey InKey, EInputEvent InEvent)
+{
+	const FWidget::EWidgetMode WidgetMode = InViewportClient->GetWidgetMode();
+	const int32 HitX = InViewport->GetMouseX();
+	const int32 HitY = InViewport->GetMouseY();
+
+	bool bHandled = false;
+
+	// Handle switching modes - only allowed when not already manipulating
+	if ((InEvent == IE_Pressed) && (InKey == EKeys::SpaceBar) && !bManipulating)
+	{
+		int32 SelectedBoneIndex = GetAnimPreviewScene().GetSelectedBoneIndex();
+		USkeletalMeshSocket* SelectedSocket = GetAnimPreviewScene().GetSelectedSocket().Socket;
+		AActor* SelectedActor = GetAnimPreviewScene().GetSelectedActor();
+
+		if (SelectedBoneIndex >= 0 || SelectedSocket || SelectedActor != nullptr)
+		{
+			if (WidgetMode == FWidget::WM_Rotate)
+			{
+				InViewportClient->SetWidgetMode(FWidget::WM_Scale);
+			}
+			else if (WidgetMode == FWidget::WM_Scale)
+			{
+				InViewportClient->SetWidgetMode(FWidget::WM_Translate);
+			}
+			else
+			{
+				InViewportClient->SetWidgetMode(FWidget::WM_Rotate);
+			}
+		}
+		bHandled = true;
+		InViewportClient->Invalidate();
+	}
+
+	return bHandled;
+}
+
+bool FSkeletonSelectionEditMode::InputDelta(FEditorViewportClient* InViewportClient, FViewport* InViewport, FVector& InDrag, FRotator& InRot, FVector& InScale)
+{
+	const EAxisList::Type CurrentAxis = InViewportClient->GetCurrentWidgetAxis();
+	const FWidget::EWidgetMode WidgetMode = InViewportClient->GetWidgetMode();
+	const ECoordSystem CoordSystem = InViewportClient->GetWidgetCoordSystemSpace();
+
+	// Get some useful info about buttons being held down
+	const bool bCtrlDown = InViewport->KeyState(EKeys::LeftControl) || InViewport->KeyState(EKeys::RightControl);
+	const bool bShiftDown = InViewport->KeyState(EKeys::LeftShift) || InViewport->KeyState(EKeys::RightShift);
+	const bool bMouseButtonDown = InViewport->KeyState( EKeys::LeftMouseButton ) || InViewport->KeyState( EKeys::MiddleMouseButton ) || InViewport->KeyState( EKeys::RightMouseButton );
+
+	bool bHandled = false;
+
+	UDebugSkelMeshComponent* PreviewMeshComponent = GetAnimPreviewScene().GetPreviewMeshComponent();
+
+	if ( bManipulating && CurrentAxis != EAxisList::None )
+	{
+		bHandled = true;
+
+		int32 BoneIndex = GetAnimPreviewScene().GetSelectedBoneIndex();
+		USkeletalMeshSocket* SelectedSocket = GetAnimPreviewScene().GetSelectedSocket().Socket;
+		AActor* SelectedActor = GetAnimPreviewScene().GetSelectedActor();
+		FAnimNode_ModifyBone* SkelControl = nullptr;
+
+		if ( BoneIndex >= 0 )
+		{
+			//Get the skeleton control manipulating this bone
+			const FName BoneName = PreviewMeshComponent->SkeletalMesh->RefSkeleton.GetBoneName(BoneIndex);
+			SkelControl = &(PreviewMeshComponent->PreviewInstance->ModifyBone(BoneName));
+		}
+
+		if ( SkelControl || SelectedSocket )
+		{
+			FTransform CurrentSkelControlTM(
+				SelectedSocket ? SelectedSocket->RelativeRotation : SkelControl->Rotation,
+				SelectedSocket ? SelectedSocket->RelativeLocation : SkelControl->Translation,
+				SelectedSocket ? SelectedSocket->RelativeScale : SkelControl->Scale);
+
+			FTransform BaseTM;
+
+			if ( SelectedSocket )
+			{
+				BaseTM = SelectedSocket->GetSocketTransform( PreviewMeshComponent );
+			}
+			else
+			{
+				BaseTM = PreviewMeshComponent->GetBoneTransform( BoneIndex );
+			}
+
+			// Remove SkelControl's orientation from BoneMatrix, as we need to translate/rotate in the non-SkelControlled space
+			BaseTM = BaseTM.GetRelativeTransformReverse(CurrentSkelControlTM);
+
+			const bool bDoRotation    = WidgetMode == FWidget::WM_Rotate    || WidgetMode == FWidget::WM_TranslateRotateZ;
+			const bool bDoTranslation = WidgetMode == FWidget::WM_Translate || WidgetMode == FWidget::WM_TranslateRotateZ;
+			const bool bDoScale = WidgetMode == FWidget::WM_Scale;
+
+			if (bDoRotation)
+			{
+				FVector RotAxis;
+				float RotAngle;
+				InRot.Quaternion().ToAxisAndAngle( RotAxis, RotAngle );
+
+				FVector4 BoneSpaceAxis = BaseTM.TransformVector( RotAxis );
+
+				//Calculate the new delta rotation
+				FQuat DeltaQuat( BoneSpaceAxis, RotAngle );
+				DeltaQuat.Normalize();
+
+				FRotator NewRotation = ( CurrentSkelControlTM * FTransform( DeltaQuat )).Rotator();
+
+				if ( SelectedSocket )
+				{
+					SelectedSocket->RelativeRotation = NewRotation;
+				}
+				else
+				{
+					SkelControl->Rotation = NewRotation;
+				}
+			}
+
+			if (bDoTranslation)
+			{
+				FVector4 BoneSpaceOffset = BaseTM.TransformVector(InDrag);
+				if (SelectedSocket)
+				{
+					SelectedSocket->RelativeLocation += BoneSpaceOffset;
+				}
+				else
+				{
+					SkelControl->Translation += BoneSpaceOffset;
+				}
+			}
+			if(bDoScale)
+			{
+				FVector4 BoneSpaceScaleOffset;
+
+				if (CoordSystem == COORD_World)
+				{
+					BoneSpaceScaleOffset = BaseTM.TransformVector(InScale);
+				}
+				else
+				{
+					BoneSpaceScaleOffset = InScale;
+				}
+
+				if(SelectedSocket)
+				{
+					SelectedSocket->RelativeScale += BoneSpaceScaleOffset;
+				}
+				else
+				{
+					SkelControl->Scale += BoneSpaceScaleOffset;
+				}
+			}
+
+		}
+		else if( SelectedActor != nullptr )
+		{
+			if (WidgetMode == FWidget::WM_Rotate)
+			{
+				FTransform Transform = SelectedActor->GetTransform();
+				FRotator NewRotation = (Transform * FTransform( InRot ) ).Rotator();
+
+				SelectedActor->SetActorRotation( NewRotation );
+			}
+			else
+			{
+				FVector Location = SelectedActor->GetActorLocation();
+				Location += InDrag;
+				SelectedActor->SetActorLocation(Location);
+			}
+		}
+
+		InViewport->Invalidate();
+	}
+
+	return bHandled;
+}
+
+void FSkeletonSelectionEditMode::Render(const FSceneView* View, FViewport* Viewport, FPrimitiveDrawInterface* PDI)
+{
+	// If we have a socket of interest, draw the widget
+	if (GetAnimPreviewScene().GetSelectedSocket().IsValid())
+	{
+		TArray<USkeletalMeshSocket*> SocketAsArray;
+		SocketAsArray.Add(GetAnimPreviewScene().GetSelectedSocket().Socket);
+		FAnimationViewportClient::DrawSockets(GetAnimPreviewScene().GetPreviewMeshComponent(), SocketAsArray, GetAnimPreviewScene().GetSelectedSocket(), PDI, false);
+	}
+}
+
+void FSkeletonSelectionEditMode::DrawHUD(FEditorViewportClient* ViewportClient, FViewport* Viewport, const FSceneView* View, FCanvas* Canvas)
+{
+	UDebugSkelMeshComponent* PreviewMeshComponent = GetAnimPreviewScene().GetPreviewMeshComponent();
+
+	// Draw name of selected bone
+	if (GetAnimPreviewScene().GetSelectedBoneIndex() != INDEX_NONE)
+	{
+		const FIntPoint ViewPortSize = Viewport->GetSizeXY();
+		const int32 HalfX = ViewPortSize.X / 2;
+		const int32 HalfY = ViewPortSize.Y / 2;
+
+		int32 BoneIndex = GetAnimPreviewScene().GetSelectedBoneIndex();
+		const FName BoneName = PreviewMeshComponent->SkeletalMesh->RefSkeleton.GetBoneName(BoneIndex);
+
+		FMatrix BoneMatrix = PreviewMeshComponent->GetBoneMatrix(BoneIndex);
+		const FPlane Proj = View->Project(BoneMatrix.GetOrigin());
+		if (Proj.W > 0.f)
+		{
+			const int32 XPos = HalfX + (HalfX * Proj.X);
+			const int32 YPos = HalfY + (HalfY * (Proj.Y * -1));
+
+			FQuat BoneQuat = PreviewMeshComponent->GetBoneQuaternion(BoneName);
+			FVector Loc = PreviewMeshComponent->GetBoneLocation(BoneName);
+			FCanvasTextItem TextItem(FVector2D(XPos, YPos), FText::FromString(BoneName.ToString()), GEngine->GetSmallFont(), FLinearColor::White);
+			Canvas->DrawItem(TextItem);
+		}
+	}
+
+	// Draw name of selected socket
+	if (GetAnimPreviewScene().GetSelectedSocket().IsValid())
+	{
+		USkeletalMeshSocket* Socket = GetAnimPreviewScene().GetSelectedSocket().Socket;
+
+		FMatrix SocketMatrix;
+		Socket->GetSocketMatrix(SocketMatrix, PreviewMeshComponent);
+		const FVector SocketPos = SocketMatrix.GetOrigin();
+
+		const FPlane Proj = View->Project(SocketPos);
+		if (Proj.W > 0.f)
+		{
+			const FIntPoint ViewPortSize = Viewport->GetSizeXY();
+			const int32 HalfX = ViewPortSize.X / 2;
+			const int32 HalfY = ViewPortSize.Y / 2;
+
+			const int32 XPos = HalfX + (HalfX * Proj.X);
+			const int32 YPos = HalfY + (HalfY * (Proj.Y * -1));
+			FCanvasTextItem TextItem(FVector2D(XPos, YPos), FText::FromString(Socket->SocketName.ToString()), GEngine->GetSmallFont(), FLinearColor::White);
+			Canvas->DrawItem(TextItem);
+		}
+	}
+}
+
+bool FSkeletonSelectionEditMode::AllowWidgetMove()
+{
+	UDebugSkelMeshComponent* PreviewMeshComponent = GetAnimPreviewScene().GetPreviewMeshComponent();
+	return !PreviewMeshComponent->IsAnimBlueprintInstanced();
+}
+
+bool FSkeletonSelectionEditMode::ShouldDrawWidget() const
+{
+	UDebugSkelMeshComponent* PreviewMeshComponent = GetAnimPreviewScene().GetPreviewMeshComponent();
+	if (!PreviewMeshComponent->IsAnimBlueprintInstanced())
+	{
+		return (GetAnimPreviewScene().GetSelectedBoneIndex() != INDEX_NONE) || GetAnimPreviewScene().GetSelectedSocket().IsValid() || GetAnimPreviewScene().GetSelectedActor() != nullptr;
+	}
+
+	return false;
+}
+
+bool FSkeletonSelectionEditMode::UsesTransformWidget() const
+{
+	return true;
+}
+
+bool FSkeletonSelectionEditMode::GetCustomDrawingCoordinateSystem(FMatrix& InMatrix, void* InData)
+{
+	UDebugSkelMeshComponent* PreviewMeshComponent = GetAnimPreviewScene().GetPreviewMeshComponent();
+
+	if (GetAnimPreviewScene().GetSelectedBoneIndex() != INDEX_NONE)
+	{
+		int32 BoneIndex = PreviewMeshComponent->BonesOfInterest.Last();
+
+		FTransform BoneMatrix = PreviewMeshComponent->GetBoneTransform(BoneIndex);
+
+		InMatrix = BoneMatrix.ToMatrixNoScale().RemoveTranslation();
+		return true;
+	}
+	else if (GetAnimPreviewScene().GetSelectedSocket().IsValid())
+	{
+		USkeletalMeshSocket* Socket = GetAnimPreviewScene().GetSelectedSocket().Socket;
+
+		FTransform SocketMatrix = Socket->GetSocketTransform(PreviewMeshComponent);
+
+		InMatrix = SocketMatrix.ToMatrixNoScale().RemoveTranslation();
+		return true;
+	}
+	else if (AActor* SelectedActor = GetAnimPreviewScene().GetSelectedActor())
+	{
+		InMatrix = SelectedActor->GetTransform().ToMatrixNoScale().RemoveTranslation();
+		return true;
+	}
+
+	return false;
+}
+
+bool FSkeletonSelectionEditMode::GetCustomInputCoordinateSystem(FMatrix& InMatrix, void* InData)
+{
+	return GetCustomDrawingCoordinateSystem(InMatrix, InData);
+}
+
+FVector FSkeletonSelectionEditMode::GetWidgetLocation() const
+{
+	UDebugSkelMeshComponent* PreviewMeshComponent = GetAnimPreviewScene().GetPreviewMeshComponent();
+
+	if (GetAnimPreviewScene().GetSelectedBoneIndex() != INDEX_NONE)
+	{
+		int32 BoneIndex = PreviewMeshComponent->BonesOfInterest.Last();
+		const FName BoneName = PreviewMeshComponent->SkeletalMesh->RefSkeleton.GetBoneName(BoneIndex);
+
+		FMatrix BoneMatrix = PreviewMeshComponent->GetBoneMatrix(BoneIndex);
+
+		return BoneMatrix.GetOrigin();
+	}
+	else if (GetAnimPreviewScene().GetSelectedSocket().IsValid())
+	{
+		USkeletalMeshSocket* Socket = GetAnimPreviewScene().GetSelectedSocket().Socket;
+
+		FMatrix SocketMatrix;
+		Socket->GetSocketMatrix(SocketMatrix, PreviewMeshComponent);
+
+		return SocketMatrix.GetOrigin();
+	}
+	else if (AActor* SelectedActor = GetAnimPreviewScene().GetSelectedActor())
+	{
+		return SelectedActor->GetActorLocation();
+	}
+
+	return FVector::ZeroVector;
+}
+
+bool FSkeletonSelectionEditMode::HandleClick(FEditorViewportClient* InViewportClient, HHitProxy *HitProxy, const FViewportClick &Click)
+{
+	bool bHandled = false;
+
+	if ( HitProxy )
+	{
+		if ( HitProxy->IsA( HPersonaSocketProxy::StaticGetType() ) )
+		{
+			// Tell the skeleton tree that the socket has been selected - this will sort out the skeleton tree, etc.
+			static_cast<FAnimationViewportClient*>(InViewportClient)->GetSkeletonTree()->SetSelectedSocket(static_cast<HPersonaSocketProxy*>(HitProxy)->SocketInfo);
+			bHandled = true;
+		}
+		else if ( HitProxy->IsA( HPersonaBoneProxy::StaticGetType() ) )
+		{			
+			// Tell the skeleton tree that the bone has been selected - this will sort out the skeleton tree, etc.
+			static_cast<FAnimationViewportClient*>(InViewportClient)->GetSkeletonTree()->SetSelectedBone(static_cast<HPersonaBoneProxy*>(HitProxy)->BoneName);
+			bHandled = true;
+		}
+		else if ( HitProxy->IsA( HActor::StaticGetType() ) )
+		{
+			GetAnimPreviewScene().SetSelectedActor(static_cast<HActor*>(HitProxy)->Actor);
+			bHandled = true;
+		}
+	}
+	else
+	{
+		// Cast for phys bodies if we didn't get any hit proxies
+		FHitResult Result(1.0f);
+		UDebugSkelMeshComponent* PreviewMeshComponent = GetAnimPreviewScene().GetPreviewMeshComponent();
+		bool bHit = PreviewMeshComponent->LineTraceComponent(Result, Click.GetOrigin(), Click.GetOrigin() + Click.GetDirection() * SkeletonSelectionModeConstants::BodyTraceDistance, FCollisionQueryParams(NAME_None,true));
+		
+		if(bHit)
+		{
+			static_cast<FAnimationViewportClient*>(InViewportClient)->GetSkeletonTree()->SetSelectedBone(Result.BoneName);
+			bHandled = true;
+		}
+		else
+		{
+			// We didn't hit a proxy or a physics object, so deselect all objects
+			static_cast<FAnimationViewportClient*>(InViewportClient)->GetSkeletonTree()->DeselectAll();
+		}
+	}
+
+	return bHandled;
+}
+
+#undef LOCTEXT_NAMESPACE
\ No newline at end of file
diff --git a/Engine/Source/Editor/Persona/Private/MeshMode/MeshMode.cpp b/Engine/Source/Editor/Persona/Private/MeshMode/MeshMode.cpp
--- a/Engine/Source/Editor/Persona/Private/MeshMode/MeshMode.cpp
+++ b/Engine/Source/Editor/Persona/Private/MeshMode/MeshMode.cpp
@@ -11,16 +11,17 @@
 #include "PersonaMeshDetails.h"
 
 #include "MeshMode.h"
+#include "IPersonaToolkit.h"
 
 #define LOCTEXT_NAMESPACE "PersonaMeshMode"
 
 /////////////////////////////////////////////////////
 // SMeshPropertiesTabBody
 
-class SMeshPropertiesTabBody : public SSingleObjectDetailsPanel
+class SPersonaMeshPropertiesTabBody : public SSingleObjectDetailsPanel
 {
 public:
-	SLATE_BEGIN_ARGS(SMeshPropertiesTabBody) {}
+	SLATE_BEGIN_ARGS(SPersonaMeshPropertiesTabBody) {}
 	SLATE_END_ARGS()
 
 private:
@@ -33,13 +34,14 @@ public:
 
 		SSingleObjectDetailsPanel::Construct(SSingleObjectDetailsPanel::FArguments().HostCommandList(InPersona->GetToolkitCommands()), true, true);
 
-		PropertyView->SetGenericLayoutDetailsDelegate( FOnGetDetailCustomizationInstance::CreateStatic( &FPersonaMeshDetails::MakeInstance, InPersona ) );
+		FPersonaModule& PersonaModule = FModuleManager::GetModuleChecked<FPersonaModule>("Persona");
+		PersonaModule.CustomizeMeshDetails(PropertyView.ToSharedRef(), InPersona->GetPersonaToolkit());
 	}
 
 	// SSingleObjectDetailsPanel interface
 	virtual UObject* GetObjectToObserve() const override
 	{
-		return PersonaPtr.Pin()->GetMesh();
+		return PersonaPtr.Pin()->GetPersonaToolkit()->GetMesh();
 	}
 	// End of SSingleObjectDetailsPanel interface
 };
@@ -47,7 +49,7 @@ public:
 /////////////////////////////////////////////////////
 // FMeshPropertiesSummoner
 
-FMeshPropertiesSummoner::FMeshPropertiesSummoner(TSharedPtr<class FAssetEditorToolkit> InHostingApp)
+FPersonaMeshPropertiesSummoner::FPersonaMeshPropertiesSummoner(TSharedPtr<class FAssetEditorToolkit> InHostingApp)
 	: FWorkflowTabFactory(FPersonaTabs::MeshAssetPropertiesID, InHostingApp)
 {
 	TabLabel = LOCTEXT("MeshProperties_TabTitle", "Mesh Details");
@@ -59,10 +61,10 @@ FMeshPropertiesSummoner::FMeshPropertiesSummoner(TSharedPtr<class FAssetEditorTo
 	ViewMenuTooltip = LOCTEXT("MeshProperties_MenuToolTip", "Shows the skeletal mesh properties");
 }
 
-TSharedRef<SWidget> FMeshPropertiesSummoner::CreateTabBody(const FWorkflowTabSpawnInfo& Info) const
+TSharedRef<SWidget> FPersonaMeshPropertiesSummoner::CreateTabBody(const FWorkflowTabSpawnInfo& Info) const
 {
 	TSharedPtr<FPersona> PersonaApp = StaticCastSharedPtr<FPersona>(HostingApp.Pin());
-	return SNew(SMeshPropertiesTabBody, PersonaApp);
+	return SNew(SPersonaMeshPropertiesTabBody, PersonaApp);
 }
 
 /////////////////////////////////////////////////////
@@ -72,7 +74,7 @@ FMeshEditAppMode::FMeshEditAppMode(TSharedPtr<FPersona> InPersona)
 	: FPersonaAppMode(InPersona, FPersonaModes::MeshEditMode)
 {
 	PersonaTabFactories.RegisterFactory(MakeShareable(new FSelectionDetailsSummoner(InPersona)));
-	PersonaTabFactories.RegisterFactory(MakeShareable(new FMeshPropertiesSummoner(InPersona)));
+	PersonaTabFactories.RegisterFactory(MakeShareable(new FPersonaMeshPropertiesSummoner(InPersona)));
 
 	TabLayout = FTabManager::NewLayout( "Persona_MeshEditMode_Layout_v7" )
 		->AddArea
diff --git a/Engine/Source/Editor/Persona/Private/MeshMode/MeshMode.h b/Engine/Source/Editor/Persona/Private/MeshMode/MeshMode.h
--- a/Engine/Source/Editor/Persona/Private/MeshMode/MeshMode.h
+++ b/Engine/Source/Editor/Persona/Private/MeshMode/MeshMode.h
@@ -9,9 +9,9 @@
 /////////////////////////////////////////////////////
 // FMeshPropertiesSummoner
 
-struct FMeshPropertiesSummoner : public FWorkflowTabFactory
+struct FPersonaMeshPropertiesSummoner : public FWorkflowTabFactory
 {
-	FMeshPropertiesSummoner(TSharedPtr<class FAssetEditorToolkit> InHostingApp);
+	FPersonaMeshPropertiesSummoner(TSharedPtr<class FAssetEditorToolkit> InHostingApp);
 
 	// FWorkflowTabFactory interface
 	virtual TSharedRef<SWidget> CreateTabBody(const FWorkflowTabSpawnInfo& Info) const override;
