Copying //UE4/Dev-Framework to //UE4/Dev-Main (Source: //UE4/Dev-Framework @ 3136612)

diff --git a/Engine/Source/Editor/Persona/Private/Persona.cpp b/Engine/Source/Editor/Persona/Private/Persona.cpp
--- a/Engine/Source/Editor/Persona/Private/Persona.cpp
+++ b/Engine/Source/Editor/Persona/Private/Persona.cpp
@@ -58,13 +58,8 @@
 
 #include "Particles/ParticleSystemComponent.h"
 
-#include "DesktopPlatformModule.h"
-#include "MainFrame.h"
-#include "AnimationCompressionPanel.h"
-#include "FbxAnimUtils.h"
-
 #include "SAnimationDlgs.h"
-#include "Developer/AssetTools/Public/AssetToolsModule.h"
+#include "AssetToolsModule.h"
 #include "FbxMeshUtils.h"
 #include "AnimationEditorUtils.h"
 #include "SAnimationSequenceBrowser.h"
@@ -73,7 +68,6 @@
 #include "SNotificationList.h"
 #include "NotificationManager.h"
 
-#include "Editor/KismetWidgets/Public/SSingleObjectDetailsPanel.h"
 #include "Animation/AimOffsetBlendSpace.h"
 #include "Animation/AimOffsetBlendSpace1D.h"
 #include "Animation/AnimNotifies/AnimNotifyState.h"
@@ -84,6 +78,15 @@
 #include "PhysicsEngine/PhysicsAsset.h"
 
 #include "Editor/AnimGraph/Public/AnimGraphCommands.h"
+#include "PersonaUtils.h"
+#include "IPersonaToolkit.h"
+#include "ISkeletonEditorModule.h"
+#include "SMorphTargetViewer.h"
+#include "SSkeletonAnimNotifies.h"
+#include "IEditableSkeleton.h"
+#include "SBlueprintEditorToolbar.h"
+#include "IPersonaEditorModeManager.h"
+#include "AnimationGraph.h"
 
 #define LOCTEXT_NAMESPACE "FPersona"
 
@@ -100,68 +103,6 @@ struct FLocalCharEditorCallbacks
 	}
 };
 
-/////////////////////////////////////////////////////
-// SPersonaPreviewPropertyEditor
-
-class SPersonaPreviewPropertyEditor : public SSingleObjectDetailsPanel
-{
-public:
-	SLATE_BEGIN_ARGS(SPersonaPreviewPropertyEditor) {}
-	SLATE_END_ARGS()
-
-private:
-	// Pointer back to owning Persona editor instance (the keeper of state)
-	TWeakPtr<FPersona> PersonaPtr;
-public:
-	void Construct(const FArguments& InArgs, TSharedPtr<FPersona> InPersona)
-	{
-		PersonaPtr = InPersona;
-
-		SSingleObjectDetailsPanel::Construct(SSingleObjectDetailsPanel::FArguments().HostCommandList(InPersona->GetToolkitCommands()), /*bAutomaticallyObserveViaGetObjectToObserve*/ true, /*bAllowSearch*/ true);
-
-		PropertyView->SetIsPropertyEditingEnabledDelegate(FIsPropertyEditingEnabled::CreateStatic([] { return !GIntraFrameDebuggingGameThread; }));
-	}
-
-	// SSingleObjectDetailsPanel interface
-	virtual UObject* GetObjectToObserve() const override
-	{
-		if (UDebugSkelMeshComponent* PreviewComponent = PersonaPtr.Pin()->GetPreviewMeshComponent())
-		{
-			if (PreviewComponent->GetAnimInstance() != nullptr)
-			{
-				return PreviewComponent->GetAnimInstance();
-			}
-		}
-
-		return nullptr;
-	}
-
-	virtual TSharedRef<SWidget> PopulateSlot(TSharedRef<SWidget> PropertyEditorWidget) override
-	{
-		return SNew(SVerticalBox)
-			+SVerticalBox::Slot()
-			.AutoHeight()
-			.Padding(0.f, 8.f, 0.f, 0.f)
-			[
-				SNew(SBorder)
-				.BorderImage(FEditorStyle::GetBrush("Persona.PreviewPropertiesWarning"))
-				[
-					SNew(STextBlock)
-					.Text(LOCTEXT("AnimBlueprintEditPreviewText", "Changes to preview options are not saved in the asset."))
-					.Font(FEditorStyle::GetFontStyle("PropertyWindow.NormalFont"))
-					.ShadowColorAndOpacity(FLinearColor::Black.CopyWithNewOpacity(0.3f))
-					.ShadowOffset(FVector2D::UnitVector)
-				]
-			]
-			+SVerticalBox::Slot()
-			.FillHeight(1)
-			[
-				PropertyEditorWidget
-			];
-	}
-	// End of SSingleObjectDetailsPanel interface
-};
-
 //////////////////////////////////////////////////////
 /////////////////////////////////////////////////////
 // FSkeletonEditAppMode
@@ -262,21 +203,13 @@ TArray<UObject*> GetEditorObjectsOfClass( const TArray< UObject* >& Objects, con
 // FPersona
 
 FPersona::FPersona()
-	: TargetSkeleton(NULL)
-	, PreviewComponent(NULL)
-	, PersonaMeshDetailLayout(NULL)
-	, PreviewScene(FPreviewScene::ConstructionValues().AllowAudioPlayback(true).ShouldSimulatePhysics(true))
-	, LastCachedLODForPreviewComponent(0)
+	: PersonaMeshDetailLayout(NULL)
 {
 	// Register to be notified when properties are edited
 	OnPropertyChangedHandle = FCoreUObjectDelegates::FOnObjectPropertyChanged::FDelegate::CreateRaw(this, &FPersona::OnPropertyChanged);
 	OnPropertyChangedHandleDelegateHandle = FCoreUObjectDelegates::OnObjectPropertyChanged.Add(OnPropertyChangedHandle);
 
 	GEditor->OnBlueprintPreCompile().AddRaw(this, &FPersona::OnBlueprintPreCompile);
-
-	//Temporary fix for missing attached assets - MDW
-	PreviewScene.GetWorld()->GetWorldSettings()->SetIsTemporarilyHiddenInEditor(false);
-
 }
 
 FPersona::~FPersona()
@@ -286,105 +219,11 @@ FPersona::~FPersona()
 	FEditorDelegates::OnAssetPostImport.RemoveAll(this);
 	FReimportManager::Instance()->OnPostReimport().RemoveAll(this);
 
-
-	if(TargetSkeleton)
-	{
-		TargetSkeleton->UnregisterOnSkeletonHierarchyChanged(this);
-	}
-
-	if(Viewport.IsValid())
-	{
-		Viewport.Pin()->CleanupPersonaReferences();
-	}
-
 	FCoreUObjectDelegates::OnObjectPropertyChanged.Remove(OnPropertyChangedHandleDelegateHandle);
 
-	if(PreviewComponent)
-	{
-		PreviewComponent->RemoveFromRoot();
-	}
-	
 	// NOTE: Any tabs that we still have hanging out when destroyed will be cleaned up by FBaseToolkit's destructor
 }
 
-TSharedPtr<SWidget> FPersona::CreateEditorWidgetForAnimDocument(UObject* InAnimAsset, FString& DocumentLink)
-{
-	TSharedPtr<SWidget> Result;
-	if (InAnimAsset)
-	{
-		if (UAnimSequence* Sequence = Cast<UAnimSequence>(InAnimAsset))
-		{
-			Result = SNew(SSequenceEditor)
-				.Persona(SharedThis(this))
-				.Sequence(Sequence);
-
-			DocumentLink = TEXT("Engine/Animation/Sequences");
-		}
-		else if (UAnimComposite* Composite = Cast<UAnimComposite>(InAnimAsset))
-		{
-			Result = SNew(SAnimCompositeEditor)
-				.Persona(SharedThis(this))
-				.Composite(Composite);
-
-			DocumentLink = TEXT("Engine/Animation/AnimationComposite");
-		}
-		else if (UAnimMontage* Montage = Cast<UAnimMontage>(InAnimAsset))
-		{
-			Result = SNew(SMontageEditor)
-				.Persona(SharedThis(this))
-				.Montage(Montage);
-
-			DocumentLink = TEXT("Engine/Animation/AnimMontage");
-		}
-		else if (UPoseAsset* PoseAsset = Cast<UPoseAsset>(InAnimAsset))
-		{
-			Result = SNew(SPoseEditor)
-				.Persona(SharedThis(this))
-				.PoseAsset(PoseAsset);
-
-			DocumentLink = TEXT("Engine/Animation/Sequences");
-		}
-		else if (UBlendSpace* BlendSpace = Cast<UBlendSpace>(InAnimAsset))
-		{
-			Result = SNew(SBlendSpaceEditor)
-				.Persona(SharedThis(this))
-				.BlendSpace(BlendSpace);
-
-			if (Cast<UAimOffsetBlendSpace>(InAnimAsset))
-			{
-				DocumentLink = TEXT("Engine/Animation/AimOffset");
-			}
-			else
-			{
-				DocumentLink = TEXT("Engine/Animation/Blendspaces");
-			}
-		}
-		else if (UBlendSpace1D* BlendSpace1D = Cast<UBlendSpace1D>(InAnimAsset))
-		{
-			Result = SNew(SBlendSpaceEditor1D)
-				.Persona(SharedThis(this))
-				.BlendSpace1D(BlendSpace1D);
-
-
-			if (Cast<UAimOffsetBlendSpace1D>(InAnimAsset))
-			{
-				DocumentLink = TEXT("Engine/Animation/AimOffset");
-			}
-			else
-			{
-				DocumentLink = TEXT("Engine/Animation/Blendspaces");
-			}
-		}
-	}
-
-	if (Result.IsValid())
-	{
-		InAnimAsset->SetFlags(RF_Transactional);
-	}
-
-	return Result;
-}
-
 void FPersona::ReinitMode()
 {
 	FName CurrentMode = GetCurrentMode();
@@ -392,23 +231,32 @@ void FPersona::ReinitMode()
 	if (CurrentMode == FPersonaModes::AnimationEditMode)
 	{
 		// if opening animation edit mode, open the animation editor
-		OpenNewAnimationDocumentTab(GetPreviewAnimationAsset());
+		OpenNewAnimationDocumentTab(GetPreviewScene()->GetPreviewAnimationAsset());
 	}
 }
 
-TSharedPtr<SDockTab> FPersona::OpenNewAnimationDocumentTab(UObject* InAnimAsset)
+TSharedPtr<SDockTab> FPersona::OpenNewAnimationDocumentTab(UAnimationAsset* InAnimAsset)
 {
 	TSharedPtr<SDockTab> OpenedTab;
 
 	if (InAnimAsset != NULL)
 	{
 		FString	DocumentLink;
-		TSharedPtr<SWidget> TabContents = CreateEditorWidgetForAnimDocument(InAnimAsset, DocumentLink);
+
+		FAnimDocumentArgs Args(GetPreviewScene(), GetPersonaToolkit(), GetSkeletonTree()->GetEditableSkeleton(), OnPostUndo, OnCurvesChanged, OnAnimNotifiesChanged, OnSectionsChanged);
+		Args.OnDespatchObjectsSelected = FOnObjectsSelected::CreateSP(this, &FPersona::HandleObjectsSelected);
+		Args.OnDespatchAnimNotifiesChanged = FSimpleDelegate::CreateSP(this, &FPersona::HandleAnimNotifiesChanged);
+		Args.OnDespatchInvokeTab = FOnInvokeTab::CreateSP(this, &FAssetEditorToolkit::InvokeTab);
+		Args.OnDespatchCurvesChanged = FSimpleDelegate::CreateSP(this, &FPersona::HandleCurvesChanged);
+		Args.OnDespatchSectionsChanged = FSimpleDelegate::CreateSP(this, &FPersona::HandleSectionsChanged);
+
+		FPersonaModule& PersonaModule = FModuleManager::GetModuleChecked<FPersonaModule>("Persona");
+		TSharedPtr<SWidget> TabContents = PersonaModule.CreateEditorWidgetForAnimDocument(SharedThis(this), InAnimAsset, Args, DocumentLink);
 		if (TabContents.IsValid())
 		{
-			if ( SharedAnimAssetBeingEdited.IsValid() )
+			if ( GetPersonaToolkit()->GetAnimationAsset() )
 			{
-				RemoveEditingObject(SharedAnimAssetBeingEdited.Get());
+				RemoveEditingObject(GetPersonaToolkit()->GetAnimationAsset());
 			}
 
 			if ( InAnimAsset != NULL )
@@ -416,7 +264,7 @@ TSharedPtr<SDockTab> FPersona::OpenNewAnimationDocumentTab(UObject* InAnimAsset)
 				AddEditingObject(InAnimAsset);
 			}
 
-			SharedAnimAssetBeingEdited = InAnimAsset;
+			GetPersonaToolkit()->SetAnimationAsset(InAnimAsset);
 			TAttribute<FText> NameAttribute = TAttribute<FText>::Create(TAttribute<FText>::FGetter::CreateStatic(&FLocalCharEditorCallbacks::GetObjectName, (UObject*)InAnimAsset));
 
 			if (SharedAnimDocumentTab.IsValid())
@@ -448,8 +296,7 @@ TSharedPtr<SDockTab> FPersona::OpenNewAnimationDocumentTab(UObject* InAnimAsset)
 
 		if(SequenceBrowser.IsValid())
 		{
-			UAnimationAsset* NewAsset = CastChecked<UAnimationAsset>(InAnimAsset);
-			SequenceBrowser.Pin()->SelectAsset(NewAsset);
+			SequenceBrowser.Pin()->SelectAsset(InAnimAsset);
 		}
 	}
 
@@ -478,7 +325,7 @@ TSharedPtr<SDockTab> FPersona::OpenNewDocumentTab(class UAnimationAsset* InAnimA
 			NewTab = OpenNewAnimationDocumentTab(InAnimAsset);
 		}
 
-		SetPreviewAnimationAsset(InAnimAsset);
+		GetPreviewScene()->SetPreviewAnimationAsset(InAnimAsset);
 	}
 
 	return NewTab;
@@ -487,25 +334,9 @@ TSharedPtr<SDockTab> FPersona::OpenNewDocumentTab(class UAnimationAsset* InAnimA
 void FPersona::SetViewport(TWeakPtr<class SAnimationEditorViewportTabBody> NewViewport)
 {
 	Viewport = NewViewport;
-	if(Viewport.IsValid())
-	{
-		Viewport.Pin()->SetPreviewComponent(PreviewComponent);
-	}
 	OnViewportCreated.Broadcast(Viewport);
 }
 
-void FPersona::SetSequenceBrowser(class SAnimationSequenceBrowser* InSequenceBrowser)
-{
-	if (InSequenceBrowser)
-	{
-		SequenceBrowser = SharedThis(InSequenceBrowser);
-	}
-	else
-	{
-		SequenceBrowser = NULL;
-	}
-}
-
 void FPersona::RefreshViewport()
 {
 	if (Viewport.IsValid())
@@ -514,46 +345,11 @@ void FPersona::RefreshViewport()
 	}
 }
 
-USkeleton* FPersona::GetSkeleton() const
-{
-	return TargetSkeleton;
-}
-
-USkeletalMesh* FPersona::GetMesh() const
-{
-	return PreviewComponent->SkeletalMesh;
-}
-
-UPhysicsAsset* FPersona::GetPhysicsAsset() const
-{
-	return NULL;
-}
-
 UAnimBlueprint* FPersona::GetAnimBlueprint() const
 {
 	return Cast<UAnimBlueprint>(GetBlueprintObj());
 }
 
-UObject* FPersona::GetMeshAsObject() const
-{
-	return GetMesh();
-}
-
-UObject* FPersona::GetSkeletonAsObject() const
-{
-	return GetSkeleton();
-}
-
-UObject* FPersona::GetPhysicsAssetAsObject() const
-{
-	return GetPhysicsAsset();
-}
-
-UObject* FPersona::GetAnimBlueprintAsObject() const
-{
-	return GetAnimBlueprint();
-}
-
 void FPersona::ExtendMenu()
 {
 	// Add additional persona editor menus
@@ -594,14 +390,6 @@ void FPersona::ExtendMenu()
 				MenuBuilder.AddMenuEntry(FPersonaCommands::Get().AddLoopingInterpolation);
 			}
 			MenuBuilder.EndSection();
-
-			MenuBuilder.BeginSection("Persona", LOCTEXT("PersonaAssetMenuMenu_Record", "Record"));
-			{
-				MenuBuilder.AddMenuEntry(FPersonaCommands::Get().RecordAnimation, 
-					NAME_None, 
-					TAttribute<FText>(PersonaPtr, &FPersona::GetRecordMenuLabel));
-			}
-			MenuBuilder.EndSection();
 		}
 	};
 
@@ -664,45 +452,55 @@ void FPersona::InitPersona(const EToolkitMode::Type Mode, const TSharedPtr< clas
 
 	GetToolkitCommands()->Append(FPlayWorldCommands::GlobalPlayWorldActions.ToSharedRef());
 
-	// Create a list of available modes
-	{
-		TSharedPtr<FPersona> ThisPtr(SharedThis(this));
-
-		TArray< TSharedRef<FApplicationMode> > TempModeList;
-		TempModeList.Add(MakeShareable(new FSkeletonEditAppMode(ThisPtr)));
-		TempModeList.Add(MakeShareable(new FMeshEditAppMode(ThisPtr)));
-		//TempModeList.Add(MakeShareable(new FPhysicsEditAppMode(ThisPtr)));
-		TempModeList.Add(MakeShareable(new FAnimEditAppMode(ThisPtr)));
-		TempModeList.Add(MakeShareable(new FAnimBlueprintEditAppMode(ThisPtr)));
-
-		for (auto ModeListIt = TempModeList.CreateIterator(); ModeListIt; ++ModeListIt)
-		{
-			TSharedRef<FApplicationMode> ApplicationMode = *ModeListIt;
-			AddApplicationMode(ApplicationMode->GetModeName(), ApplicationMode);
-		}
-	}
+	FPersonaModule& PersonaModule = FModuleManager::GetModuleChecked<FPersonaModule>("Persona");
 
 	// Determine the initial mode the app will open up in
 	FName InitialMode = NAME_None;
 	if (InitAnimBlueprint != NULL)
 	{
 		InitialMode = FPersonaModes::AnimBlueprintEditMode;
-		TargetSkeleton = InitAnimBlueprint->TargetSkeleton;
+		PersonaToolkit = PersonaModule.CreatePersonaToolkit(InitAnimBlueprint);
 	}
 	else if (InitAnimationAsset != NULL)
 	{
 		InitialMode = FPersonaModes::AnimationEditMode;
-		TargetSkeleton = InitSkeleton;
+		PersonaToolkit = PersonaModule.CreatePersonaToolkit(InitAnimationAsset);
 	}
 	else if (InitMesh != NULL)
 	{
 		InitialMode = FPersonaModes::MeshEditMode;
-		TargetSkeleton = InitSkeleton;
+		PersonaToolkit = PersonaModule.CreatePersonaToolkit(InitMesh);
 	}
 	else
 	{
 		InitialMode = FPersonaModes::SkeletonDisplayMode;
-		TargetSkeleton = InitSkeleton;
+		PersonaToolkit = PersonaModule.CreatePersonaToolkit(InitSkeleton);
+	}
+
+	// create the skeleton tree
+	FSkeletonTreeArgs SkeletonTreeArgs(OnPostUndo);
+	SkeletonTreeArgs.OnObjectSelected = FOnObjectSelected::CreateSP(this, &FPersona::HandleObjectSelected);
+	SkeletonTreeArgs.PreviewScene = GetPreviewScene();
+
+	ISkeletonEditorModule& SkeletonEditorModule = FModuleManager::LoadModuleChecked<ISkeletonEditorModule>("SkeletonEditor");
+	SkeletonTree = SkeletonEditorModule.CreateSkeletonTree(PersonaToolkit->GetSkeleton(), PersonaToolkit->GetMesh(), SkeletonTreeArgs);
+
+	// Create a list of available modes
+	{
+		TSharedPtr<FPersona> ThisPtr(SharedThis(this));
+
+		TArray< TSharedRef<FApplicationMode> > TempModeList;
+		TempModeList.Add(MakeShareable(new FSkeletonEditAppMode(ThisPtr)));
+		TempModeList.Add(MakeShareable(new FMeshEditAppMode(ThisPtr)));
+		//TempModeList.Add(MakeShareable(new FPhysicsEditAppMode(ThisPtr)));
+		TempModeList.Add(MakeShareable(new FAnimEditAppMode(ThisPtr)));
+		TempModeList.Add(MakeShareable(new FAnimBlueprintEditAppMode(ThisPtr)));
+
+		for (auto ModeListIt = TempModeList.CreateIterator(); ModeListIt; ++ModeListIt)
+		{
+			TSharedRef<FApplicationMode> ApplicationMode = *ModeListIt;
+			AddApplicationMode(ApplicationMode->GetModeName(), ApplicationMode);
+		}
 	}
 
 	// Build up a list of objects being edited in this asset editor
@@ -725,15 +523,14 @@ void FPersona::InitPersona(const EToolkitMode::Type Mode, const TSharedPtr< clas
 	//	ObjectsBeingEdited.Add(InitMesh);
 	//}
 
-	ObjectsBeingEdited.Add(TargetSkeleton);
-	check(TargetSkeleton);
-
-	TargetSkeleton->CollectAnimationNotifies();
+	USkeleton* Skeleton = PersonaToolkit->GetSkeleton();
+	check(Skeleton);
+	ObjectsBeingEdited.Add(Skeleton);
 
-	TargetSkeleton->RegisterOnSkeletonHierarchyChanged( USkeleton::FOnSkeletonHierarchyChanged::CreateSP( this, &FPersona::OnSkeletonHierarchyChanged ) );
+	Skeleton->CollectAnimationNotifies();
 
 	// We could modify the skeleton within Persona (add/remove sockets), so we need to enable undo/redo on it
-	TargetSkeleton->SetFlags( RF_Transactional );
+	Skeleton->SetFlags( RF_Transactional );
 
 	// Initialize the asset editor and spawn tabs
 	const TSharedRef<FTabManager::FLayout> DummyLayout = FTabManager::NewLayout("NullLayout")->AddArea(FTabManager::NewPrimaryArea());
@@ -748,15 +545,6 @@ void FPersona::InitPersona(const EToolkitMode::Type Mode, const TSharedPtr< clas
 	}
 	CommonInitialization(AnimBlueprints);
 
-	check(PreviewComponent == NULL);
-
-	// Create the preview component
-	PreviewComponent = NewObject<UDebugSkelMeshComponent>();
-
-	// note: we add to root here (rather than using RF_Standalone) as all standalone objects will be cleaned up when we switch
-	// preview worlds (in UWorld::CleanupWorld), but we want this to stick around while Persona exists
-	PreviewComponent->AddToRoot();
-
 	bool bSetMesh = false;
 
 	// Set the mesh
@@ -765,9 +553,9 @@ void FPersona::InitPersona(const EToolkitMode::Type Mode, const TSharedPtr< clas
 		SetPreviewMesh(InitMesh);
 		bSetMesh = true;
 
-		if (TargetSkeleton && !TargetSkeleton->GetPreviewMesh())
+		if (!Skeleton->GetPreviewMesh())
 		{
-			TargetSkeleton->SetPreviewMesh(InitMesh, false);
+			Skeleton->SetPreviewMesh(InitMesh, false);
 		}
 	}
 	else if (InitAnimationAsset != NULL)
@@ -780,10 +568,10 @@ void FPersona::InitPersona(const EToolkitMode::Type Mode, const TSharedPtr< clas
 		}
 	}
 
-	if (!bSetMesh && TargetSkeleton)
+	if (!bSetMesh && Skeleton)
 	{
 		//If no preview mesh set, just find the first mesh that uses this skeleton
-		USkeletalMesh* PreviewMesh = TargetSkeleton->GetPreviewMesh(true);
+		USkeletalMesh* PreviewMesh = Skeleton->GetPreviewMesh(true);
 		if ( PreviewMesh )
 		{
 			SetPreviewMesh( PreviewMesh );
@@ -791,20 +579,21 @@ void FPersona::InitPersona(const EToolkitMode::Type Mode, const TSharedPtr< clas
 	}
 
 	// Force validation of preview attached assets (catch case of never doing it if we dont have a valid preview mesh)
-	ValidatePreviewAttachedAssets(NULL);
+	GetPreviewScene()->ValidatePreviewAttachedAssets(nullptr);
 
-	UAnimBlueprint* AnimBlueprint = GetAnimBlueprint();
-	PreviewComponent->SetAnimInstanceClass(AnimBlueprint ? AnimBlueprint->GeneratedClass : NULL);
+	UDebugSkelMeshComponent* PreviewMeshComponent = PersonaToolkit->GetPreviewMeshComponent();
+	UAnimBlueprint* AnimBlueprint = PersonaToolkit->GetAnimBlueprint();
+	PreviewMeshComponent->SetAnimInstanceClass(AnimBlueprint ? AnimBlueprint->GeneratedClass : NULL);
 
 	// We always want a preview instance unless we are using blueprints so that bone manipulation works
 	if (AnimBlueprint == NULL)
 	{
-		PreviewComponent->EnablePreview(true, nullptr);
+		PreviewMeshComponent->EnablePreview(true, nullptr);
 	}
 	else
 	{
 		// Make sure the object being debugged is the preview instance
-		AnimBlueprint->SetObjectBeingDebugged(PreviewComponent->GetAnimInstance());
+		AnimBlueprint->SetObjectBeingDebugged(PreviewMeshComponent->GetAnimInstance());
 	}
 
 	ExtendMenu();
@@ -815,6 +604,10 @@ void FPersona::InitPersona(const EToolkitMode::Type Mode, const TSharedPtr< clas
 	SetCurrentMode(InitialMode);
 	check(CurrentAppModePtr.IsValid());
 
+	// set up our editor mode
+	check(AssetEditorModeManager);
+	AssetEditorModeManager->SetDefaultMode(FPersonaEditModes::SkeletonSelection);
+
 	// Post-layout initialization
 	PostLayoutBlueprintEditorInitialization();
 
@@ -826,9 +619,7 @@ void FPersona::InitPersona(const EToolkitMode::Type Mode, const TSharedPtr< clas
 
 	// register customization of Slot node for this Persona
 	// this is so that you can open the manage window per Persona
-	TWeakPtr<FPersona> PersonaPtr = SharedThis(this);
-	Inspector->GetPropertyView()->RegisterInstancedCustomPropertyLayout(UAnimGraphNode_Slot::StaticClass(),
-		FOnGetDetailCustomizationInstance::CreateStatic(&FAnimGraphNodeSlotDetails::MakeInstance, PersonaPtr));
+	PersonaModule.CustomizeSlotNodeDetails(Inspector->GetPropertyView().ToSharedRef(), FOnInvokeTab::CreateSP(this, &FAssetEditorToolkit::InvokeTab));
 
 	// Register post import callback to catch animation imports when we have the asset open (we need to reinit)
 	FEditorDelegates::OnAssetPostImport.AddRaw(this, &FPersona::OnPostImport);
@@ -841,8 +632,8 @@ void FPersona::CreateAnimation(const TArray<UObject*> NewAssets, int32 Option)
 	bool bResult = true;
 	if (NewAssets.Num() > 0)
 	{
-		USkeletalMeshComponent* MeshComponent = GetPreviewMeshComponent();
-		UAnimSequence* Sequence = Cast<UAnimSequence> (GetPreviewAnimationAsset());
+		USkeletalMeshComponent* MeshComponent = PersonaToolkit->GetPreviewMeshComponent();
+		UAnimSequence* Sequence = Cast<UAnimSequence> (GetPreviewScene()->GetPreviewAnimationAsset());
 
 		for (auto NewAsset : NewAssets)
 		{
@@ -872,7 +663,11 @@ void FPersona::CreateAnimation(const TArray<UObject*> NewAssets, int32 Option)
 			// if it created based on current mesh component, 
 			if (Option == 1)
 			{
-				PreviewComponent->PreviewInstance->ResetModifiedBone();
+				UDebugSkelMeshComponent* PreviewMeshComponent = PersonaToolkit->GetPreviewMeshComponent();
+				if (PreviewMeshComponent && PreviewMeshComponent->PreviewInstance)
+				{
+					PreviewMeshComponent->PreviewInstance->ResetModifiedBone();
+				}
 			}
 		}
 		else
@@ -887,8 +682,8 @@ void FPersona::CreatePoseAsset(const TArray<UObject*> NewAssets, int32 Option)
 	bool bResult = false;
 	if (NewAssets.Num() > 0)
 	{
-		USkeletalMeshComponent* MeshComponent = GetPreviewMeshComponent();
-		UAnimSequence* Sequence = Cast<UAnimSequence>(GetPreviewAnimationAsset());
+		UDebugSkelMeshComponent* PreviewComponent = GetPreviewScene()->GetPreviewMeshComponent();
+		UAnimSequence* Sequence = Cast<UAnimSequence>(GetPreviewScene()->GetPreviewAnimationAsset());
 
 		for (auto NewAsset : NewAssets)
 		{
@@ -898,7 +693,7 @@ void FPersona::CreatePoseAsset(const TArray<UObject*> NewAssets, int32 Option)
 				switch (Option)
 				{
 				case 0:
-					NewPoseAsset->AddOrUpdatePoseWithUniqueName(MeshComponent);
+					NewPoseAsset->AddOrUpdatePoseWithUniqueName(PreviewComponent);
 					break;
 				case 1:
 					NewPoseAsset->CreatePoseFromAnimation(Sequence);
@@ -931,7 +726,7 @@ void FPersona::FillCreateAnimationMenu(FMenuBuilder& MenuBuilder) const
 {
 	TArray<TWeakObjectPtr<USkeleton>> Skeletons;
 
-	Skeletons.Add(TargetSkeleton);
+	Skeletons.Add(PersonaToolkit->GetSkeleton());
 
 	// create rig
 	MenuBuilder.BeginSection("CreateAnimationSubMenu", LOCTEXT("CreateAnimationSubMenuHeading", "Create Animation"));
@@ -975,7 +770,7 @@ void FPersona::FillCreatePoseAssetMenu(FMenuBuilder& MenuBuilder) const
 
 	Skeletons.Add(TargetSkeleton);
 
-	// create rig
+	// create pose asset
 	MenuBuilder.BeginSection("CreatePoseAssetSubMenu", LOCTEXT("CreatePoseAssetSubMenuHeading", "Create PoseAsset"));
 	{
 		MenuBuilder.AddMenuEntry(
@@ -999,6 +794,95 @@ void FPersona::FillCreatePoseAssetMenu(FMenuBuilder& MenuBuilder) const
 			);
 	}
 	MenuBuilder.EndSection();
+
+	// insert pose to pose asset
+	MenuBuilder.BeginSection("InsertPoseSubMenuSection", LOCTEXT("InsertPoseSubMenuSubMenuHeading", "Insert Pose"));
+	{
+		MenuBuilder.AddSubMenu(
+			LOCTEXT("InsertPoseSubmenu", "Insert Pose"),
+			LOCTEXT("InsertPoseSubmenu_ToolTip", "Insert current pose to selected PoseAsset"),
+			FNewMenuDelegate::CreateSP(this, &FPersona::FillInsertPoseMenu),
+			false,
+			FSlateIcon(FEditorStyle::GetStyleSetName(), "ClassIcon.PoseAsset")
+		);
+	}
+	MenuBuilder.EndSection();
+}
+
+void FPersona::FillInsertPoseMenu(FMenuBuilder& MenuBuilder) const
+{
+	FAssetPickerConfig AssetPickerConfig;
+	
+	USkeleton* Skeleton = GetPersonaToolkit()->GetSkeleton();
+
+	/** The asset picker will only show skeletons */
+	AssetPickerConfig.Filter.ClassNames.Add(*UPoseAsset::StaticClass()->GetName());
+	AssetPickerConfig.Filter.bRecursiveClasses = false;
+	AssetPickerConfig.bAllowNullSelection = false;
+	AssetPickerConfig.Filter.TagsAndValues.Add(TEXT("Skeleton"), FAssetData(Skeleton).GetExportTextName());
+
+	/** The delegate that fires when an asset was selected */
+	AssetPickerConfig.OnAssetSelected = FOnAssetSelected::CreateRaw(this, &FPersona::InsertCurrentPoseToAsset);
+
+	/** The default view mode should be a list view */
+	AssetPickerConfig.InitialAssetViewType = EAssetViewType::List;
+
+	FContentBrowserModule& ContentBrowserModule = FModuleManager::Get().LoadModuleChecked<FContentBrowserModule>(TEXT("ContentBrowser"));
+
+	MenuBuilder.AddWidget(
+		ContentBrowserModule.Get().CreateAssetPicker(AssetPickerConfig),
+		LOCTEXT("Select_Label", "")
+	);
+}
+
+void FPersona::InsertCurrentPoseToAsset(const FAssetData& NewPoseAssetData)
+{
+	UPoseAsset* PoseAsset = Cast<UPoseAsset>(NewPoseAssetData.GetAsset());
+	FScopedTransaction ScopedTransaction(LOCTEXT("InsertPose", "Insert Pose"));
+
+	if (PoseAsset)
+	{
+		PoseAsset->Modify();
+
+		UDebugSkelMeshComponent* PreviewMeshComponent = PersonaToolkit->GetPreviewMeshComponent();
+		if (PreviewMeshComponent)
+		{
+			FSmartName NewPoseName;
+
+			bool bSuccess =  PoseAsset->AddOrUpdatePoseWithUniqueName(PreviewMeshComponent, &NewPoseName);
+
+			if (bSuccess)
+			{
+				FFormatNamedArguments Args;
+				Args.Add(TEXT("PoseAsset"), FText::FromString(PoseAsset->GetName()));
+				Args.Add(TEXT("PoseName"), FText::FromName(NewPoseName.DisplayName));
+				FNotificationInfo Info(FText::Format(LOCTEXT("InsertPoseSucceeded", "The current pose has inserted to {PoseAsset} with {PoseName}"), Args));
+				Info.ExpireDuration = 7.0f;
+				Info.bUseLargeFont = false;
+				TSharedPtr<SNotificationItem> Notification = FSlateNotificationManager::Get().AddNotification(Info);
+				if (Notification.IsValid())
+				{
+					Notification->SetCompletionState(SNotificationItem::CS_Success);
+				}
+			}
+			else
+			{
+				FFormatNamedArguments Args;
+				Args.Add(TEXT("PoseAsset"), FText::FromString(PoseAsset->GetName()));
+				FNotificationInfo Info(FText::Format(LOCTEXT("InsertPoseFailed", "Inserting pose to asset {PoseAsset} has failed"), Args));
+				Info.ExpireDuration = 7.0f;
+				Info.bUseLargeFont = false;
+				TSharedPtr<SNotificationItem> Notification = FSlateNotificationManager::Get().AddNotification(Info);
+				if (Notification.IsValid())
+				{
+					Notification->SetCompletionState(SNotificationItem::CS_Fail);
+				}
+			}
+		}
+	}
+
+	// it doesn't work well if I leave the window open. The delegate goes weired or it stop showing the popups. 
+	FSlateApplication::Get().DismissAllMenus();
 }
 
 TSharedRef< SWidget > FPersona::GenerateCreateAssetMenu( USkeleton* Skeleton ) const
@@ -1078,13 +962,6 @@ void FPersona::ExtendDefaultPersonaToolbar()
 				ToolbarBuilder.AddToolBarButton(FPersonaCommands::Get().ImportAnimation, NAME_None, LOCTEXT("Toolbar_ImportAnimation", "Import"));
 				ToolbarBuilder.AddToolBarButton(FPersonaCommands::Get().ReimportAnimation, NAME_None, LOCTEXT("Toolbar_ReimportAnimation", "Reimport"));
 				ToolbarBuilder.AddToolBarButton(FPersonaCommands::Get().ExportToFBX, NAME_None, LOCTEXT("Toolbar_ExportToFBX", "Export"));
-
-				ToolbarBuilder.AddToolBarButton(FPersonaCommands::Get().RecordAnimation, 
-											NAME_None,
-											 TAttribute<FText>(PersonaPtr, &FPersona::GetRecordStatusLabel),
-											 TAttribute<FText>(PersonaPtr, &FPersona::GetRecordStatusTooltip),
-											 TAttribute<FSlateIcon>(PersonaPtr, &FPersona::GetRecordStatusImage),
-											 NAME_None);
 			}
 			ToolbarBuilder.EndSection();
 
@@ -1124,7 +1001,7 @@ void FPersona::ExtendDefaultPersonaToolbar()
 		"Asset",
 		EExtensionHook::After,
 		GetToolkitCommands(),
-		FToolBarExtensionDelegate::CreateStatic(&Local::FillToolbar, TargetSkeleton, this)
+		FToolBarExtensionDelegate::CreateStatic(&Local::FillToolbar, PersonaToolkit->GetSkeleton(), this)
 		);
 
 	AddToolbarExtender(ToolbarExtender);
@@ -1145,84 +1022,31 @@ UBlueprint* FPersona::GetBlueprintObj() const
 
 UObject* FPersona::GetPreviewAnimationAsset() const
 {
-	if (Viewport.IsValid())
-	{
-		if (PreviewComponent)
-		{
-			// if same, do not overwrite. It will reset time and everything
-			if (PreviewComponent->PreviewInstance != NULL)
-			{
-				return PreviewComponent->PreviewInstance->GetCurrentAsset();
-			}
-		}
-	}
-
-	return NULL;
+	return GetPreviewScene()->GetPreviewAnimationAsset();
 }
 
 UObject* FPersona::GetAnimationAssetBeingEdited() const
 {
-	if (!SharedAnimDocumentTab.IsValid())
-	{
-		SharedAnimAssetBeingEdited = NULL;
-	}
-
-	return SharedAnimAssetBeingEdited.Get();
+	return GetPersonaToolkit()->GetAnimationAsset();
 }
 
 void FPersona::SetPreviewAnimationAsset(UAnimationAsset* AnimAsset, bool bEnablePreview)
 {
-	if (Viewport.IsValid() && !Viewport.Pin()->bPreviewLockModeOn)
-	{
-		if (PreviewComponent)
-		{
-			RemoveAttachedComponent(false);
-
-			if (AnimAsset != NULL)
-			{
-				// Early out if the new preview asset is the same as the current one, to avoid replaying from the beginning, etc...
-				if (AnimAsset == GetPreviewAnimationAsset() && PreviewComponent->IsPreviewOn())
-				{
-					return;
-				}
-
-				// Treat it as invalid if it's got a bogus skeleton pointer
-				if (AnimAsset->GetSkeleton() != TargetSkeleton)
-				{
-					return;
-				}
+	return GetPreviewScene()->SetPreviewAnimationAsset(AnimAsset, bEnablePreview);
 			}
 
-			PreviewComponent->EnablePreview(bEnablePreview, AnimAsset);
-		}
-
-		OnAnimChanged.Broadcast(AnimAsset);
-	}
-}
-
-void FPersona::UpdateSelectionDetails(UObject* Object, const FText& ForcedTitle)
+void FPersona::SetDetailObjects(const TArray<UObject*>& InObjects)
 {
-	Inspector->ShowDetailsForSingleObject(Object, SKismetInspector::FShowDetailsOptions(ForcedTitle));
+	Inspector->ShowDetailsForObjects(InObjects);
 }
 
 void FPersona::SetDetailObject(UObject* Obj)
 {
-	FText ForcedTitle = (Obj != NULL) ? FText::FromString(Obj->GetName()) : FText::GetEmpty();
-	UpdateSelectionDetails(Obj, ForcedTitle);
+	TArray<UObject*> Objects;
+	Objects.Add(Obj);
+	SetDetailObjects(Objects);
 }
 
-UDebugSkelMeshComponent* FPersona::GetPreviewMeshComponent() const
-{
-	return PreviewComponent;
-}
-
-UAnimInstance* FPersona::GetPreviewAnimInstance() const
-{
-	UDebugSkelMeshComponent* DebugSkelComp = GetPreviewMeshComponent();
-	return (DebugSkelComp != nullptr) ? DebugSkelComp->GetAnimInstance() : nullptr;
-}
-
-
 void FPersona::OnPostReimport(UObject* InObject, bool bSuccess)
 {
 	if (bSuccess)
@@ -1239,8 +1063,9 @@ void FPersona::OnPostImport(UFactory* InFactory, UObject* InObject)
 void FPersona::ConditionalRefreshEditor(UObject* InObject)
 {
 	bool bInterestingAsset = true;
+	USkeleton* Skeleton = GetPersonaToolkit()->GetSkeleton();
 	// Ignore if this is regarding a different object
-	if(InObject != TargetSkeleton && InObject != TargetSkeleton->GetPreviewMesh() && InObject != GetAnimationAssetBeingEdited())
+	if(InObject != Skeleton && InObject != Skeleton->GetPreviewMesh() && InObject != GetAnimationAssetBeingEdited())
 	{
 		bInterestingAsset = false;
 	}
@@ -1270,8 +1095,6 @@ void FPersona::ConditionalRefreshEditor(UObject* InObject)
 	{
 		RefreshViewport();
 		ReinitMode();
-		
-		OnPersonaRefresh.Broadcast();
 	}
 }
 
@@ -1280,6 +1103,24 @@ void FPersona::OnGraphEditorFocused(const TSharedRef<class SGraphEditor>& InGrap
 {
 	// in the future, depending on which graph editor is this will act different
 	FBlueprintEditor::OnGraphEditorFocused(InGraphEditor);
+
+	// install callback to allow us to propagate pin default changes live to the preview
+	UAnimationGraph* AnimationGraph = Cast<UAnimationGraph>(InGraphEditor->GetCurrentGraph());
+	if (AnimationGraph)
+	{
+		OnPinDefaultValueChangedHandle = AnimationGraph->OnPinDefaultValueChanged.Add(FOnPinDefaultValueChanged::FDelegate::CreateSP(this, &FPersona::HandlePinDefaultValueChanged));
+	}
+}
+
+void FPersona::OnGraphEditorBackgrounded(const TSharedRef<SGraphEditor>& InGraphEditor)
+{
+	FBlueprintEditor::OnGraphEditorBackgrounded(InGraphEditor);
+
+	UAnimationGraph* AnimationGraph = Cast<UAnimationGraph>(InGraphEditor->GetCurrentGraph());
+	if (AnimationGraph)
+	{
+		AnimationGraph->OnPinDefaultValueChanged.Remove(OnPinDefaultValueChangedHandle);
+	}
 }
 
 /** Create Default Tabs **/
@@ -1306,14 +1147,6 @@ void FPersona::CreateDefaultCommands()
 		FCanExecuteAction::CreateSP(this, &FPersona::CanSaveAnimationAssets)
 		);
 
-	// record animation
-	ToolkitCommands->MapAction( FPersonaCommands::Get().RecordAnimation,
-		FExecuteAction::CreateSP( this, &FPersona::RecordAnimation ),
-		FCanExecuteAction(),
-		FIsActionChecked(),
-		FIsActionButtonVisible::CreateSP( this, &FPersona::IsRecordAvailable )
-		);
-
 	ToolkitCommands->MapAction(FPersonaCommands::Get().ApplyCompression,
 		FExecuteAction::CreateSP(this, &FPersona::OnApplyCompression),
 		FCanExecuteAction::CreateSP(this, &FPersona::HasValidAnimationSequencePlaying),
@@ -1356,13 +1189,13 @@ void FPersona::CreateDefaultCommands()
 
 	ToolkitCommands->MapAction(FPersonaCommands::Get().ToggleReferencePose,
 	FExecuteAction::CreateSP(this, &FPersona::ShowReferencePose, true),
-	FCanExecuteAction::CreateSP(this, &FPersona::CanShowReferencePose),
+	FCanExecuteAction(),
 	FIsActionChecked::CreateSP(this, &FPersona::IsShowReferencePoseEnabled)
 	);
 
 	ToolkitCommands->MapAction(FPersonaCommands::Get().TogglePreviewAsset,
 	FExecuteAction::CreateSP(this, &FPersona::ShowReferencePose, false),
-	FCanExecuteAction::CreateSP(this, &FPersona::CanPreviewAsset),
+	FCanExecuteAction(),
 	FIsActionChecked::CreateSP(this, &FPersona::IsPreviewAssetEnabled)
 	);
 
@@ -2090,36 +1923,33 @@ void FPersona::Compile()
 {
 	// Note if we were debugging the preview
 	UObject* CurrentDebugObject = GetBlueprintObj()->GetObjectBeingDebugged();
-	const bool bIsDebuggingPreview = (PreviewComponent != NULL) && PreviewComponent->IsAnimBlueprintInstanced() && (PreviewComponent->GetAnimInstance() == CurrentDebugObject);
+	UDebugSkelMeshComponent* PreviewMeshComponent = PersonaToolkit->GetPreviewMeshComponent();
+	const bool bIsDebuggingPreview = (PreviewMeshComponent != NULL) && PreviewMeshComponent->IsAnimBlueprintInstanced() && (PreviewMeshComponent->GetAnimInstance() == CurrentDebugObject);
 
-	if (PreviewComponent != NULL)
+	if (PreviewMeshComponent != NULL)
 	{
 		// Force close any asset editors that are using the AnimScriptInstance (such as the Property Matrix), the class will be garbage collected
-		FAssetEditorManager::Get().CloseOtherEditors(PreviewComponent->GetAnimInstance(), nullptr);
+		FAssetEditorManager::Get().CloseOtherEditors(PreviewMeshComponent->GetAnimInstance(), nullptr);
 	}
 
 	// Compile the blueprint
 	FBlueprintEditor::Compile();
 
-	if (PreviewComponent != NULL)
+	if (PreviewMeshComponent != NULL)
 	{
-		if (PreviewComponent->GetAnimInstance() == NULL)
+		if (PreviewMeshComponent->GetAnimInstance() == NULL)
 		{
 			// try reinitialize animation if it doesn't exist
-			PreviewComponent->InitAnim(true);
+			PreviewMeshComponent->InitAnim(true);
 		}
 
 		if (bIsDebuggingPreview)
 		{
-			GetBlueprintObj()->SetObjectBeingDebugged(PreviewComponent->GetAnimInstance());
+			GetBlueprintObj()->SetObjectBeingDebugged(PreviewMeshComponent->GetAnimInstance());
 		}
 	}
 
-	// calls PostCompile to copy proper values between anim nodes
-	if (Viewport.IsValid())
-	{
-		Viewport.Pin()->GetAnimationViewportClient()->PostCompile();
-	}
+	OnPostCompile();
 }
 
 FName FPersona::GetToolkitContextFName() const
@@ -2151,11 +1981,10 @@ FText FPersona::GetToolkitName() const
 	}
 	else
 	{
-		check(TargetSkeleton != NULL);
-
-		const bool bDirtyState = TargetSkeleton->GetOutermost()->IsDirty();
+		USkeleton* Skeleton = GetPersonaToolkit()->GetSkeleton();
+		const bool bDirtyState = Skeleton->GetOutermost()->IsDirty();
 
-		Args.Add( TEXT("SkeletonName"), FText::FromString( TargetSkeleton->GetName() ) );
+		Args.Add( TEXT("SkeletonName"), FText::FromString(Skeleton->GetName() ) );
 		Args.Add( TEXT("DirtyState"), bDirtyState ? FText::FromString( TEXT( "*" ) ) : FText::GetEmpty() );
 		return FText::Format( LOCTEXT("KismetToolkitName_SingleSkeleton", "{SkeletonName}{DirtyState}"), Args );
 	}
@@ -2169,8 +1998,8 @@ FText FPersona::GetToolkitToolTipText() const
 	}
 	else
 	{
-		check(TargetSkeleton != NULL);
-		return FAssetEditorToolkit::GetToolTipTextForObject(TargetSkeleton);
+		USkeleton* Skeleton = GetPersonaToolkit()->GetSkeleton();
+		return FAssetEditorToolkit::GetToolTipTextForObject(Skeleton);
 	}
 }
 
@@ -2213,9 +2042,9 @@ void FPersona::OnActiveTabChanged( TSharedPtr<SDockTab> PreviouslyActive, TShare
 	}
 	else if (SharedAnimDocumentTab.IsValid() && NewlyActivated.Get() == SharedAnimDocumentTab.Pin().Get())
 	{
-		if (UAnimationAsset* SharedAsset = Cast<UAnimationAsset>(SharedAnimAssetBeingEdited.Get()))
+		if (GetPersonaToolkit()->GetAnimationAsset())
 		{
-			SetPreviewAnimationAsset(SharedAsset);
+			GetPersonaToolkit()->GetPreviewScene()->SetPreviewAnimationAsset(GetPersonaToolkit()->GetAnimationAsset());
 		}
 	}
 	else if (SharedAnimDocumentTab.IsValid() && PreviouslyActive.Get() == SharedAnimDocumentTab.Pin().Get())
@@ -2229,108 +2058,10 @@ void FPersona::OnActiveTabChanged( TSharedPtr<SDockTab> PreviouslyActive, TShare
 	}
 }
 
-void FPersona::SetPreviewMeshInternal(USkeletalMesh* NewPreviewMesh)
+void FPersona::SetPreviewMesh(USkeletalMesh* NewPreviewMesh)
 {
-	ValidatePreviewAttachedAssets(NewPreviewMesh);
-	if(NewPreviewMesh != PreviewComponent->SkeletalMesh)
-	{
-		if(PreviewComponent->SkeletalMesh != NULL)
-		{
-			RemoveEditingObject(PreviewComponent->SkeletalMesh);
-		}
-
-		if(NewPreviewMesh != NULL)
-		{
-			AddEditingObject(NewPreviewMesh);
-		}
-
-		// setting skeletalmesh unregister/re-register, 
-		// so I have to save the animation settings and resetting after setting mesh
-		UAnimationAsset* AnimAssetToPlay = NULL;
-		float PlayPosition = 0.f;
-		bool bPlaying = false;
-		bool bNeedsToCopyAnimationData = PreviewComponent->GetAnimInstance() && PreviewComponent->GetAnimInstance() == PreviewComponent->PreviewInstance;
-		if(bNeedsToCopyAnimationData)
-		{
-			AnimAssetToPlay = PreviewComponent->PreviewInstance->GetCurrentAsset();
-			PlayPosition = PreviewComponent->PreviewInstance->GetCurrentTime();
-			bPlaying = PreviewComponent->PreviewInstance->IsPlaying();
-		}
-
-		PreviewComponent->SetSkeletalMesh(NewPreviewMesh);
-
-		if(bNeedsToCopyAnimationData)
-		{
-			SetPreviewAnimationAsset(AnimAssetToPlay);
-			PreviewComponent->PreviewInstance->SetPosition(PlayPosition);
-			PreviewComponent->PreviewInstance->SetPlaying(bPlaying);
-		}
-	}
-	else
-	{
-		PreviewComponent->InitAnim(true);
-	}
-
-	if(NewPreviewMesh != NULL)
-	{
-		PreviewScene.AddComponent(PreviewComponent, FTransform::Identity);
-		for(auto Iter = AdditionalMeshes.CreateIterator(); Iter; ++Iter)
-		{
-			PreviewScene.AddComponent((*Iter), FTransform::Identity);
-		}
-
-		// Set up the mesh for transactions
-		NewPreviewMesh->SetFlags(RF_Transactional);
-
-		AddPreviewAttachedObjects();
-
-		if(Viewport.IsValid())
-		{
-			Viewport.Pin()->SetPreviewComponent(PreviewComponent);
-		}
-	}
-
-	for(auto Iter = AdditionalMeshes.CreateIterator(); Iter; ++Iter)
-	{
-		(*Iter)->SetMasterPoseComponent(PreviewComponent);
-		(*Iter)->UpdateMasterBoneMap();
-	}
-
-	OnPreviewMeshChanged.Broadcast(NewPreviewMesh);
-}
-
-// Sets the current preview mesh
-void FPersona::SetPreviewMesh(USkeletalMesh* NewPreviewMesh)
-{
-	if(!TargetSkeleton->IsCompatibleMesh(NewPreviewMesh))
-	{
-		// message box, ask if they'd like to regenerate skeleton
-		if (FMessageDialog::Open(EAppMsgType::YesNo, LOCTEXT("RenerateSkeleton", "The preview mesh hierarchy doesn't match with Skeleton anymore. Would you like to regenerate skeleton?")) == EAppReturnType::Yes)
-		{
-			TargetSkeleton->RecreateBoneTree( NewPreviewMesh );
-			SetPreviewMeshInternal(NewPreviewMesh);
-		}
-		else
-		{
-			// Send a notification that the skeletal mesh cannot work with the skeleton
-			FFormatNamedArguments Args;
-			Args.Add(TEXT("PreviewMeshName"), FText::FromString(NewPreviewMesh->GetName()));
-			Args.Add(TEXT("TargetSkeletonName"), FText::FromString(TargetSkeleton->GetName()));
-			FNotificationInfo Info(FText::Format(LOCTEXT("SkeletalMeshIncompatible", "Skeletal Mesh \"{PreviewMeshName}\" incompatible with Skeleton \"{TargetSkeletonName}\""), Args));
-			Info.ExpireDuration = 3.0f;
-			Info.bUseLargeFont = false;
-			TSharedPtr<SNotificationItem> Notification = FSlateNotificationManager::Get().AddNotification(Info);
-			if(Notification.IsValid())
-			{
-				Notification->SetCompletionState(SNotificationItem::CS_Fail);
-			}
-		}
-	}
-	else
-	{
-		SetPreviewMeshInternal(NewPreviewMesh);
-	}
-}
+	GetSkeletonTree()->SetSkeletalMesh(NewPreviewMesh);
+}
 
 void FPersona::OnPropertyChanged(UObject* ObjectBeingModified, FPropertyChangedEvent& PropertyChangedEvent)
 {
@@ -2341,9 +2072,10 @@ void FPersona::OnPropertyChanged(UObject* ObjectBeingModified, FPropertyChangedE
 void FPersona::RefreshPreviewInstanceTrackCurves()
 {
 	// need to refresh the preview mesh
-	if(PreviewComponent->PreviewInstance)
+	UDebugSkelMeshComponent* PreviewMeshComponent = PersonaToolkit->GetPreviewMeshComponent();
+	if(PreviewMeshComponent->PreviewInstance)
 	{
-		PreviewComponent->PreviewInstance->RefreshCurveBoneControllers();
+		PreviewMeshComponent->PreviewInstance->RefreshCurveBoneControllers();
 	}
 }
 
@@ -2361,9 +2093,10 @@ void FPersona::PostUndo(bool bSuccess)
 	OnPostUndo.Broadcast();	
 
 	// Re-init preview instance
-	if (PreviewComponent && PreviewComponent->AnimScriptInstance)
+	UDebugSkelMeshComponent* PreviewMeshComponent = PersonaToolkit->GetPreviewMeshComponent();
+	if (PreviewMeshComponent && PreviewMeshComponent->AnimScriptInstance)
 	{
-		PreviewComponent->AnimScriptInstance->InitializeAnimation();
+		PreviewMeshComponent->AnimScriptInstance->InitializeAnimation();
 	}
 
 	RefreshPreviewInstanceTrackCurves();
@@ -2373,11 +2106,13 @@ void FPersona::PostUndo(bool bSuccess)
 	// if those are undoed or redoed, they have to be 
 	// cleared up, otherwise, they might have invalid data
 	ClearupPreviewMeshAnimNotifyStates();
+
+	OnPostCompile();
 }
 
 void FPersona::ClearupPreviewMeshAnimNotifyStates()
 {
-	USkeletalMeshComponent* PreviewMeshComponent = GetPreviewMeshComponent();
+	UDebugSkelMeshComponent* PreviewMeshComponent = PersonaToolkit->GetPreviewMeshComponent();
 	if ( PreviewMeshComponent )
 	{
 		UAnimInstance* AnimInstantace = PreviewMeshComponent->GetAnimInstance();
@@ -2392,11 +2127,6 @@ void FPersona::ClearupPreviewMeshAnimNotifyStates()
 	}
 }
 
-void FPersona::OnSkeletonHierarchyChanged()
-{
-	OnSkeletonTreeChanged.Broadcast();
-}
-
 bool FPersona::IsInAScriptingMode() const
 {
 	return IsModeCurrent(FPersonaModes::AnimBlueprintEditMode);
@@ -2404,9 +2134,10 @@ bool FPersona::IsInAScriptingMode() const
 
 void FPersona::GetCustomDebugObjects(TArray<FCustomDebugObject>& DebugList) const
 {
-	if (PreviewComponent->IsAnimBlueprintInstanced())
+	UDebugSkelMeshComponent* PreviewMeshComponent = PersonaToolkit->GetPreviewMeshComponent();
+	if (PreviewMeshComponent->IsAnimBlueprintInstanced())
 	{
-		new (DebugList) FCustomDebugObject(PreviewComponent->GetAnimInstance(), LOCTEXT("PreviewObjectLabel", "Preview Instance").ToString());
+		new (DebugList) FCustomDebugObject(PreviewMeshComponent->GetAnimInstance(), LOCTEXT("PreviewObjectLabel", "Preview Instance").ToString());
 	}
 }
 
@@ -2414,7 +2145,7 @@ void FPersona::CreateDefaultTabContents(const TArray<UBlueprint*>& InBlueprints)
 {
 	FBlueprintEditor::CreateDefaultTabContents(InBlueprints);
 
-	PreviewEditor = SNew(SPersonaPreviewPropertyEditor, SharedThis(this));
+	PreviewEditor = SNew(SPersonaPreviewPropertyEditor, GetPreviewScene());
 }
 
 FGraphAppearanceInfo FPersona::GetGraphAppearance(UEdGraph* InGraph) const
@@ -2456,480 +2187,30 @@ void FPersona::FocusWindow(UObject* ObjectToFocusOn)
 	}
 }
 
-void FPersona::ClearSelectedBones()
-{
-	if (UDebugSkelMeshComponent* Preview = GetPreviewMeshComponent())
-	{
-		Preview->BonesOfInterest.Empty();
-	}
-}
-
-void FPersona::SetSelectedBone(USkeleton* InTargetSkeleton, const FName& BoneName, bool bRebroadcast /* = true */ )
-{
-	const int32 BoneIndex = InTargetSkeleton->GetReferenceSkeleton().FindBoneIndex(BoneName);
-	if (BoneIndex != INDEX_NONE)
-	{
-		if (UDebugSkelMeshComponent* Preview = GetPreviewMeshComponent())
-		{
-			Preview->BonesOfInterest.Empty();
-			ClearSelectedSocket();
-
-			// Add in bone of interest only if we have a preview instance set-up
-			if (Preview->PreviewInstance != NULL)
-			{
-				// need to get mesh bone base since BonesOfInterest is saved in SkeletalMeshComponent
-				// and it is used by renderer. It is not Skeleton base
-				const int32 MeshBoneIndex = Preview->GetBoneIndex(BoneName);
-
-				if (MeshBoneIndex != INDEX_NONE)
-				{
-					Preview->BonesOfInterest.Add(MeshBoneIndex);
-				}
-
-				if ( bRebroadcast )
-				{
-					// Broadcast that a bone has been selected
-					OnBoneSelected.Broadcast( BoneName );
-				}
-
-				if( Viewport.IsValid() )
-				{
-					Viewport.Pin()->GetLevelViewportClient().Invalidate();
-				}
-			}
-		}
-	}
-}
-
-void FPersona::SetSelectedMorphTargets(USkeletalMesh *SkeletalMesh, const TArray<FName>& SelectedMorphTargetNames)
-{
-	UDebugSkelMeshComponent* Preview = GetPreviewMeshComponent();
-	if (Preview)
-	{
-		Preview->MorphTargetOfInterests.Reset();
-
-		if (SelectedMorphTargetNames.Num() > 0)
-		{
-			for (const FName& MorphTargetName : SelectedMorphTargetNames)
-			{
-				int32 MorphtargetIdx;
-				UMorphTarget* MorphTarget = SkeletalMesh->FindMorphTargetAndIndex(MorphTargetName, MorphtargetIdx);
-				if (MorphTarget != nullptr)
-				{
-					Preview->MorphTargetOfInterests.AddUnique(MorphTarget);
-				}
-			}
-
-			if (Viewport.IsValid())
-			{
-				Viewport.Pin()->GetLevelViewportClient().Invalidate();
-			}
-			Preview->PostInitMeshObject(PreviewComponent->MeshObject);
-		}
-	}
-}
-
-void FPersona::SetSelectedSocket( const FSelectedSocketInfo& SocketInfo, bool bRebroadcast /* = true */ )
-{
-	if ( UDebugSkelMeshComponent* Preview = GetPreviewMeshComponent() )
-	{
-		Preview->SocketsOfInterest.Empty();
-		Preview->BonesOfInterest.Empty();
-
-		Preview->SocketsOfInterest.Add( SocketInfo );
-
-		if ( bRebroadcast )
-		{
-			OnSocketSelected.Broadcast( SocketInfo );
-		}
-
-		// This populates the details panel with the information for the socket
-		SetDetailObject( SocketInfo.Socket );
-
-		if( Viewport.IsValid() )
-		{
-			Viewport.Pin()->GetLevelViewportClient().Invalidate();
-		}
-	}
-}
-
-void FPersona::ClearSelectedSocket()
-{
-	if (UDebugSkelMeshComponent* Preview = GetPreviewMeshComponent())
-	{
-		Preview->SocketsOfInterest.Empty();
-		Preview->BonesOfInterest.Empty();
-
-		SetDetailObject( NULL );
-	}
-}
-
 void FPersona::ClearSelectedWindActor()
 {
-	if (Viewport.IsValid())
-	{
-		Viewport.Pin()->GetAnimationViewportClient()->ClearSelectedWindActor();
-	}
+	GetPreviewScene()->ClearSelectedActor();
 }
 
 void FPersona::ClearSelectedAnimGraphNode()
 {
-	if(Viewport.IsValid())
-	{
-		Viewport.Pin()->GetAnimationViewportClient()->ClearSelectedAnimGraphNode();
-	}
-}
-
-void FPersona::RenameSocket( USkeletalMeshSocket* Socket, const FName& NewSocketName )
-{
-	Socket->SocketName = NewSocketName;
-
-	// Broadcast that the skeleton tree has changed due to the new name
-	OnSkeletonTreeChanged.Broadcast();
-
-}
-
-void FPersona::ChangeSocketParent( USkeletalMeshSocket* Socket, const FName& NewSocketParent )
-{
-	// Update the bone name.
-	Socket->BoneName = NewSocketParent;
-
-	// Broadcast that the skeleton tree has changed due to the new parent
-	OnSkeletonTreeChanged.Broadcast();
-}
-
-void FPersona::DuplicateAndSelectSocket( const FSelectedSocketInfo& SocketInfoToDuplicate, const FName& NewParentBoneName /*= FName()*/ )
-{
-	check( SocketInfoToDuplicate.Socket );
-
-	USkeletalMesh* Mesh = GetMesh();
-
-	const FScopedTransaction Transaction( LOCTEXT( "CopySocket", "Copy Socket" ) );
-
-	USkeletalMeshSocket* NewSocket;
-
-	bool bModifiedSkeleton = false;
-
-	if ( SocketInfoToDuplicate.bSocketIsOnSkeleton )
-	{
-		TargetSkeleton->Modify();
-		bModifiedSkeleton = true;
-
-		NewSocket = NewObject<USkeletalMeshSocket>(TargetSkeleton);
-		check(NewSocket);
-	}
-	else if ( Mesh )
-	{
-		Mesh->Modify();
-
-		NewSocket = NewObject<USkeletalMeshSocket>(Mesh);
-		check(NewSocket);
-	}
-	else 
-	{
-		// Original socket was on the mesh, but we have no mesh. Huh?
-		check( 0 );
-		return;
-	}
-
-	NewSocket->SocketName = GenerateUniqueSocketName( SocketInfoToDuplicate.Socket->SocketName );
-	NewSocket->BoneName = NewParentBoneName != NAME_None ? NewParentBoneName : SocketInfoToDuplicate.Socket->BoneName;
-	NewSocket->RelativeLocation = SocketInfoToDuplicate.Socket->RelativeLocation;
-	NewSocket->RelativeRotation = SocketInfoToDuplicate.Socket->RelativeRotation;
-	NewSocket->RelativeScale = SocketInfoToDuplicate.Socket->RelativeScale;
-
-	if ( SocketInfoToDuplicate.bSocketIsOnSkeleton )
-	{
-		TargetSkeleton->Sockets.Add( NewSocket );
-	}
-	else if ( Mesh )
-	{
-		Mesh->GetMeshOnlySocketList().Add( NewSocket );
-	}
-
-	// Duplicated attached assets
-	int32 NumExistingAttachedObjects = TargetSkeleton->PreviewAttachedAssetContainer.Num();
-	for(int AttachedObjectIndex = 0; AttachedObjectIndex < NumExistingAttachedObjects; ++AttachedObjectIndex)
-	{
-		FPreviewAttachedObjectPair& Pair = TargetSkeleton->PreviewAttachedAssetContainer[AttachedObjectIndex];
-		if(Pair.AttachedTo == SocketInfoToDuplicate.Socket->SocketName)
-		{
-			if(!bModifiedSkeleton)
-			{
-				bModifiedSkeleton = true;
-				TargetSkeleton->Modify();
-			}
-			FPreviewAttachedObjectPair NewPair = Pair;
-			NewPair.AttachedTo = NewSocket->SocketName;
-			AttachObjectToPreviewComponent( NewPair.GetAttachedObject(), NewPair.AttachedTo, &TargetSkeleton->PreviewAttachedAssetContainer );
-		}
-	}
-
-	// Broadcast that the skeleton tree has changed due to the new socket
-	OnSkeletonTreeChanged.Broadcast();
-
-	SetSelectedSocket( FSelectedSocketInfo( NewSocket, SocketInfoToDuplicate.bSocketIsOnSkeleton ) );
-}
-
-FName FPersona::GenerateUniqueSocketName( FName InName )
-{
-	USkeletalMesh* Mesh = GetMesh();
-
-	int32 NewNumber = InName.GetNumber();
-
-	// Increment NewNumber until we have a unique name (potential infinite loop if *all* int32 values are used!)
-	while ( DoesSocketAlreadyExist( NULL, FText::FromName( FName( InName, NewNumber ) ), TargetSkeleton->Sockets ) ||
-		( Mesh ? DoesSocketAlreadyExist( NULL, FText::FromName( FName( InName, NewNumber ) ), Mesh->GetMeshOnlySocketList() ) : false ) )
-	{
-		++NewNumber;
-	}
-
-	return FName( InName, NewNumber );
-}
-
-bool FPersona::DoesSocketAlreadyExist( const USkeletalMeshSocket* InSocket, const FText& InSocketName, const TArray< USkeletalMeshSocket* >& InSocketArray ) const
-{
-	for ( auto SocketIt = InSocketArray.CreateConstIterator(); SocketIt; ++SocketIt )
-	{
-		USkeletalMeshSocket* Socket = *( SocketIt );
-
-		if ( InSocket != Socket && Socket->SocketName.ToString() == InSocketName.ToString() )
-		{
-			return true;
-		}
-	}
-
-	return false;
-}
-
-USkeletalMeshComponent*	FPersona::CreateNewSkeletalMeshComponent()
-{
-	USkeletalMeshComponent* NewComp = NewObject<USkeletalMeshComponent>();
-	AdditionalMeshes.Add(NewComp);
-	NewComp->SetMasterPoseComponent(GetPreviewMeshComponent());
-	return NewComp;
-}
-
-void FPersona::RemoveAdditionalMesh(USkeletalMeshComponent* MeshToRemove)
-{
-	PreviewScene.RemoveComponent(MeshToRemove);
-	AdditionalMeshes.Remove(MeshToRemove);
-}
-
-void FPersona::ClearAllAdditionalMeshes()
-{
-	for(auto Iter = AdditionalMeshes.CreateIterator(); Iter; ++Iter)
-	{
-		PreviewScene.RemoveComponent( *Iter );
-	}
-
-	AdditionalMeshes.Empty();
-}
-
-void FPersona::AddPreviewAttachedObjects()
-{
-	// Load up mesh attachments...
-	USkeletalMesh* Mesh = GetMesh();
-
-	if ( Mesh )
-	{
-		for(int32 i = 0; i < Mesh->PreviewAttachedAssetContainer.Num(); i++)
-		{
-			FPreviewAttachedObjectPair& PreviewAttachedObject = Mesh->PreviewAttachedAssetContainer[i];
-
-			AttachObjectToPreviewComponent(PreviewAttachedObject.GetAttachedObject(), PreviewAttachedObject.AttachedTo);
-		}
-	}
-
-	// ...and then skeleton attachments
-	for(int32 i = 0; i < TargetSkeleton->PreviewAttachedAssetContainer.Num(); i++)
-	{
-		FPreviewAttachedObjectPair& PreviewAttachedObject = TargetSkeleton->PreviewAttachedAssetContainer[i];
-
-		AttachObjectToPreviewComponent(PreviewAttachedObject.GetAttachedObject(), PreviewAttachedObject.AttachedTo);
-	}
-}
-
-bool FPersona::AttachObjectToPreviewComponent( UObject* Object, FName AttachTo, FPreviewAssetAttachContainer* PreviewAssetContainer /* = NULL */ )
-{
-	if(GetComponentForAttachedObject(Object, AttachTo))
-	{
-		return false; // Already have this asset attached at this location
-	}
-
-	TSubclassOf<UActorComponent> ComponentClass = FComponentAssetBrokerage::GetPrimaryComponentForAsset(Object->GetClass());
-	if(PreviewComponent && *ComponentClass && ComponentClass->IsChildOf(USceneComponent::StaticClass()))
-	{
-		if ( PreviewAssetContainer )
-		{
-			// Set up the attached object for transactions
-			Object->SetFlags(RF_Transactional);
-			PreviewAssetContainer->AddAttachedObject( Object, AttachTo );
-		}
-
-		//set up world info for undo
-		AWorldSettings* WorldSettings = PreviewScene.GetWorld()->GetWorldSettings();
-		WorldSettings->SetFlags(RF_Transactional);
-		WorldSettings->Modify();
-
-		USceneComponent* SceneComponent = NewObject<USceneComponent>(WorldSettings, ComponentClass, NAME_None, RF_Transactional);
-
-		FComponentAssetBrokerage::AssignAssetToComponent(SceneComponent, Object);
-
-		if(UParticleSystemComponent* NewPSysComp = Cast<UParticleSystemComponent>(SceneComponent))
-		{
-			//maybe this should happen in FComponentAssetBrokerage::AssignAssetToComponent?
-			NewPSysComp->SetTickGroup(TG_PostUpdateWork);
-		}
-
-		//set up preview component for undo
-		PreviewComponent->SetFlags(RF_Transactional);
-		PreviewComponent->Modify();
-
-		// Attach component to the preview component
-		SceneComponent->SetupAttachment(PreviewComponent, AttachTo);
-		SceneComponent->RegisterComponent();
-		return true;
-	}
-	return false;
-}
-
-void FPersona::RemoveAttachedObjectFromPreviewComponent(UObject* Object, FName AttachedTo)
-{
-	// clean up components	
-	if (PreviewComponent)
-	{
-		AWorldSettings* WorldSettings = PreviewScene.GetWorld()->GetWorldSettings();
-		WorldSettings->SetFlags(RF_Transactional);
-		WorldSettings->Modify();
-
-		//set up preview component for undo
-		PreviewComponent->SetFlags(RF_Transactional);
-		PreviewComponent->Modify();
-
-		for (int32 I=PreviewComponent->GetAttachChildren().Num()-1; I >= 0; --I) // Iterate backwards because CleanupComponent will remove from AttachChildren
-		{
-			USceneComponent* ChildComponent = PreviewComponent->GetAttachChildren()[I];
-			UObject* Asset = FComponentAssetBrokerage::GetAssetFromComponent(ChildComponent);
-
-			if( Asset == Object && ChildComponent->GetAttachSocketName() == AttachedTo)
-			{
-				// PreviewComponet will be cleaned up by PreviewScene, 
-				// but if anything is attached, it won't be cleaned up, 
-				// so we'll need to clean them up manually
-				CleanupComponent(PreviewComponent->GetAttachChildren()[I]);
-				break;
-			}
-		}
-	}
-}
-
-USceneComponent* FPersona::GetComponentForAttachedObject(UObject* Object, FName AttachedTo)
-{
-	if (PreviewComponent)
-	{
-		for (USceneComponent* ChildComponent : PreviewComponent->GetAttachChildren())
-		{
-			UObject* Asset = FComponentAssetBrokerage::GetAssetFromComponent(ChildComponent);
-
-			if( Asset == Object && ChildComponent->GetAttachSocketName() == AttachedTo)
-			{
-				return ChildComponent;
-			}
-		}
-	}
-	return NULL;
-}
-
-void FPersona::RemoveAttachedComponent( bool bRemovePreviewAttached /* = true */ )
-{
-	TMap<UObject*, TArray<FName>> PreviewAttachedObjects;
-
-	if( !bRemovePreviewAttached )
-	{
-		for(auto Iter = TargetSkeleton->PreviewAttachedAssetContainer.CreateConstIterator(); Iter; ++Iter)
-		{
-			PreviewAttachedObjects.FindOrAdd(Iter->GetAttachedObject()).Add(Iter->AttachedTo);
-		}
-
-		if ( USkeletalMesh* PreviewMesh = GetMesh() )
-		{
-			for(auto Iter = PreviewMesh->PreviewAttachedAssetContainer.CreateConstIterator(); Iter; ++Iter)
-			{
-				PreviewAttachedObjects.FindOrAdd(Iter->GetAttachedObject()).Add(Iter->AttachedTo);
-			}
-		}
-	}
-
-	// clean up components	
-	if (PreviewComponent)
-	{
-		for (int32 I=PreviewComponent->GetAttachChildren().Num()-1; I >= 0; --I) // Iterate backwards because CleanupComponent will remove from AttachChildren
-		{
-			USceneComponent* ChildComponent = PreviewComponent->GetAttachChildren()[I];
-			UObject* Asset = FComponentAssetBrokerage::GetAssetFromComponent(ChildComponent);
-
-			bool bRemove = true;
-
-			//are we supposed to leave assets that came from the skeleton
-			if(	!bRemovePreviewAttached )
-			{
-				//could this asset have come from the skeleton
-				if(PreviewAttachedObjects.Contains(Asset))
-				{
-					if(PreviewAttachedObjects.Find(Asset)->Contains(ChildComponent->GetAttachSocketName()))
-					{
-						bRemove = false;
-					}
-				}
-			}
-
-			if(bRemove)
-			{
-				// PreviewComponet will be cleaned up by PreviewScene, 
-				// but if anything is attached, it won't be cleaned up, 
-				// so we'll need to clean them up manually
-				CleanupComponent(PreviewComponent->GetAttachChildren()[I]);
-			}
-		}
-
-		if( bRemovePreviewAttached )
-		{
-			check(PreviewComponent->GetAttachChildren().Num() == 0);
-		}
-	}
-}
-
-void FPersona::CleanupComponent(USceneComponent* Component)
-{
-	if (Component)
-	{
-		for (int32 I = Component->GetAttachChildren().Num() - 1; I >= 0; --I) // Iterate backwards because CleanupComponent will remove from AttachChildren
-		{
-			CleanupComponent(Component->GetAttachChildren()[I]);
-		}
-
-		check(Component->GetAttachChildren().Num() == 0);
-		Component->DestroyComponent();
-	}
+	SelectedAnimGraphNode.Reset();
 }
 
 void FPersona::DeselectAll()
 {
-	ClearSelectedBones();
-	ClearSelectedSocket();
+	GetSkeletonTree()->DeselectAll();
 	ClearSelectedWindActor();
 	ClearSelectedAnimGraphNode();
-
-	OnAllDeselected.Broadcast();
 }
 
 FReply FPersona::OnClickEditMesh()
 {
-	USkeletalMesh* PreviewMesh = TargetSkeleton->GetPreviewMesh();
+	USkeleton* Skeleton = PersonaToolkit->GetSkeleton();
+	USkeletalMesh* PreviewMesh = Skeleton->GetPreviewMesh();
 	if(PreviewMesh)
 	{
-		UpdateSelectionDetails(PreviewMesh, FText::FromString(PreviewMesh->GetName()));
+		SetDetailObject(PreviewMesh);
 	}
 	return FReply::Handled();
 }
@@ -2944,9 +2225,10 @@ void FPersona::PostRedo(bool bSuccess)
 	// PostUndo broadcast, OnPostRedo
 	OnPostUndo.Broadcast();
 
-	if (PreviewComponent && PreviewComponent->AnimScriptInstance)
+	UDebugSkelMeshComponent* PreviewMeshComponent = PersonaToolkit->GetPreviewMeshComponent();
+	if (PreviewMeshComponent && PreviewMeshComponent->AnimScriptInstance)
 	{
-		PreviewComponent->AnimScriptInstance->InitializeAnimation();
+		PreviewMeshComponent->AnimScriptInstance->InitializeAnimation();
 	}
 
 	// clear up preview anim notify states
@@ -3015,17 +2297,17 @@ void FPersona::FindInContentBrowser_Execute()
 void FPersona::SaveAsset_Execute()
 {
 	//Clean the unused Material entry before saving
-	if (GetMesh())
+	if (USkeletalMesh* Mesh = PersonaToolkit->GetMesh())
 	{
-		for (int32 MaterialIndex = 0; MaterialIndex < GetMesh()->Materials.Num(); ++MaterialIndex)
+		for (int32 MaterialIndex = 0; MaterialIndex < Mesh->Materials.Num(); ++MaterialIndex)
 		{
 			bool MaterialIsUsed = false;
-			FSkeletalMeshResource* ImportedResource = GetMesh()->GetImportedResource();
+			FSkeletalMeshResource* ImportedResource = Mesh->GetImportedResource();
 			check(ImportedResource);
 			for (int32 LODIndex = 0; !MaterialIsUsed && LODIndex < ImportedResource->LODModels.Num(); ++LODIndex)
 			{
-				FSkeletalMeshLODInfo& Info = GetMesh()->LODInfo[LODIndex];
-				if (LODIndex == 0 || GetMesh()->LODInfo[LODIndex].LODMaterialMap.Num() == 0)
+				FSkeletalMeshLODInfo& Info = Mesh->LODInfo[LODIndex];
+				if (LODIndex == 0 || Mesh->LODInfo[LODIndex].LODMaterialMap.Num() == 0)
 				{
 					for (int32 SectionIndex = 0; SectionIndex < ImportedResource->LODModels[LODIndex].Sections.Num(); ++SectionIndex)
 					{
@@ -3038,7 +2320,7 @@ void FPersona::SaveAsset_Execute()
 				}
 				else
 				{
-					for (int32 SectionIndex = 0; SectionIndex < GetMesh()->LODInfo[LODIndex].LODMaterialMap.Num(); ++SectionIndex)
+					for (int32 SectionIndex = 0; SectionIndex < Mesh->LODInfo[LODIndex].LODMaterialMap.Num(); ++SectionIndex)
 					{
 						if (Info.LODMaterialMap[SectionIndex] == MaterialIndex)
 						{
@@ -3050,99 +2332,21 @@ void FPersona::SaveAsset_Execute()
 			}
 			if (!MaterialIsUsed)
 			{
-				GetMesh()->Materials[MaterialIndex].MaterialInterface = nullptr;
+				Mesh->Materials[MaterialIndex].MaterialInterface = nullptr;
 			}
 		}
 	}
 	FAssetEditorToolkit::SaveAsset_Execute();
 }
 
-void FPersona::OnCommandGenericDelete()
-{
-	OnGenericDelete.Broadcast();
-}
-
 void FPersona::UndoAction()
 {
 	GEditor->UndoTransaction();
 }
-
-void FPersona::RedoAction()
-{
-	GEditor->RedoTransaction();
-}
-
-FSlateIcon FPersona::GetRecordStatusImage() const
-{
-	if (IsRecording())
-	{
-		return FSlateIcon(FEditorStyle::GetStyleSetName(), "Persona.StopRecordAnimation");
-	}
-
-	return FSlateIcon(FEditorStyle::GetStyleSetName(), "Persona.StartRecordAnimation");
-}
-
-FText FPersona::GetRecordMenuLabel() const
-{
-	if(IsRecording())
-	{
-		return LOCTEXT("Persona_StopRecordAnimationMenuLabel", "Stop Record Animation");
-	}
-
-	return LOCTEXT("Persona_StartRecordAnimationMenuLabel", "Start Record Animation");
-}
-
-FText FPersona::GetRecordStatusLabel() const
-{
-	FPersonaModule& PersonaModule = FModuleManager::GetModuleChecked<FPersonaModule>(TEXT("Persona"));
-
-	bool bInRecording = false;
-	PersonaModule.OnIsRecordingActive().ExecuteIfBound(PreviewComponent, bInRecording);
-
-	if (bInRecording)
-	{
-		return LOCTEXT("Persona_StopRecordAnimationLabel", "Stop");
-	}
-
-	return LOCTEXT("Persona_StartRecordAnimationLabel", "Record");
-}
-
-FText FPersona::GetRecordStatusTooltip() const
-{
-	FPersonaModule& PersonaModule = FModuleManager::GetModuleChecked<FPersonaModule>(TEXT("Persona"));
-
-	bool bInRecording = false;
-	PersonaModule.OnIsRecordingActive().ExecuteIfBound(PreviewComponent, bInRecording);
-
-	if (bInRecording)
-	{
-		return LOCTEXT("Persona_StopRecordAnimation", "Stop Record Animation");
-	}
-
-	return LOCTEXT("Persona_StartRecordAnimation", "Start Record Animation");
-}
-
-void FPersona::RecordAnimation()
-{
-	if (!PreviewComponent || !PreviewComponent->SkeletalMesh)
-	{
-		// error
-		return;
-	}
-
-	FPersonaModule& PersonaModule = FModuleManager::GetModuleChecked<FPersonaModule>(TEXT("Persona"));
-
-	bool bInRecording = false;
-	PersonaModule.OnIsRecordingActive().ExecuteIfBound(PreviewComponent, bInRecording);
-
-	if (bInRecording)
-	{
-		PersonaModule.OnStopRecording().ExecuteIfBound(PreviewComponent);
-	}
-	else
-	{
-		PersonaModule.OnRecord().ExecuteIfBound(PreviewComponent);
-	}
+
+void FPersona::RedoAction()
+{
+	GEditor->RedoTransaction();
 }
 
 void FPersona::OnSetKeyCompleted()
@@ -3152,7 +2356,8 @@ void FPersona::OnSetKeyCompleted()
 
 bool FPersona::CanSetKey() const
 {
-	return ( HasValidAnimationSequencePlaying() && PreviewComponent->BonesOfInterest.Num() > 0);
+	UDebugSkelMeshComponent* PreviewMeshComponent = PersonaToolkit->GetPreviewMeshComponent();
+	return ( HasValidAnimationSequencePlaying() && PreviewMeshComponent->BonesOfInterest.Num() > 0);
 }
 
 void FPersona::OnSetKey()
@@ -3160,7 +2365,7 @@ void FPersona::OnSetKey()
 	UAnimSequence* AnimSequence = Cast<UAnimSequence> (GetAnimationAssetBeingEdited());
 	if (AnimSequence)
 	{
-		UDebugSkelMeshComponent* Component = GetPreviewMeshComponent();
+		UDebugSkelMeshComponent* Component = PersonaToolkit->GetPreviewMeshComponent();
 		Component->PreviewInstance->SetKey(FSimpleDelegate::CreateSP(this, &FPersona::OnSetKeyCompleted));
 	}
 }
@@ -3177,7 +2382,7 @@ void FPersona::OnApplyRawAnimChanges()
 	UAnimSequence* AnimSequence = Cast<UAnimSequence>(GetAnimationAssetBeingEdited());
 	if(AnimSequence)
 	{
-		UDebugSkelMeshComponent* Component = GetPreviewMeshComponent();
+		UDebugSkelMeshComponent* Component = PersonaToolkit->GetPreviewMeshComponent();
 		// now bake
 		if (AnimSequence->DoesNeedRebake())
 		{
@@ -3197,183 +2402,40 @@ void FPersona::OnApplyRawAnimChanges()
 
 void FPersona::OnApplyCompression()
 {
-	UAnimSequence* AnimSequence = Cast<UAnimSequence> (GetPreviewAnimationAsset());
+	UAnimSequence* AnimSequence = Cast<UAnimSequence> (GetPreviewScene()->GetPreviewAnimationAsset());
 
 	if (AnimSequence)
 	{
 		TArray<TWeakObjectPtr<UAnimSequence>> AnimSequences;
 		AnimSequences.Add(AnimSequence);
-		ApplyCompression(AnimSequences);
+		FPersonaModule& PersonaModule = FModuleManager::GetModuleChecked<FPersonaModule>("Persona");
+		PersonaModule.ApplyCompression(AnimSequences);
 	}
 }
 
 void FPersona::OnExportToFBX()
 {
-	UAnimSequence* AnimSequence = Cast<UAnimSequence>(GetPreviewAnimationAsset());
+	UAnimSequence* AnimSequence = Cast<UAnimSequence>(GetPreviewScene()->GetPreviewAnimationAsset());
 
 	if(AnimSequence)
 	{
 		TArray<TWeakObjectPtr<UAnimSequence>> AnimSequences;
 		AnimSequences.Add(AnimSequence);
-		ExportToFBX(AnimSequences);
+		FPersonaModule& PersonaModule = FModuleManager::GetModuleChecked<FPersonaModule>("Persona");
+		PersonaModule.ExportToFBX(AnimSequences, GetPreviewScene()->GetPreviewMeshComponent()->SkeletalMesh);
 	}
 }
 
 void FPersona::OnAddLoopingInterpolation()
 {
-	UAnimSequence* AnimSequence = Cast<UAnimSequence>(GetPreviewAnimationAsset());
+	UAnimSequence* AnimSequence = Cast<UAnimSequence>(GetPreviewScene()->GetPreviewAnimationAsset());
 
 	if(AnimSequence)
 	{
 		TArray<TWeakObjectPtr<UAnimSequence>> AnimSequences;
 		AnimSequences.Add(AnimSequence);
-		AddLoopingInterpolation(AnimSequences);
-	}
-}
-
-void FPersona::ApplyCompression(TArray<TWeakObjectPtr<UAnimSequence>>& AnimSequences)
-{
-	FDlgAnimCompression AnimCompressionDialog(AnimSequences);
-	AnimCompressionDialog.ShowModal();
-}
-
-void FPersona::ExportToFBX(TArray<TWeakObjectPtr<UAnimSequence>>& AnimSequences)
-{
-	IDesktopPlatform* DesktopPlatform = FDesktopPlatformModule::Get();
-
-	if(DesktopPlatform)
-	{
-		USkeletalMesh* PreviewMesh = GetPreviewMeshComponent()->SkeletalMesh;
-		if(PreviewMesh == NULL)
-		{
-			FMessageDialog::Open(EAppMsgType::Ok, LOCTEXT("ExportToFBXExportMissingPreviewMesh", "ERROR: Missing preview mesh"));
-			return;
-		}
-
-		if(AnimSequences.Num() > 0)
-		{
-			const void* ParentWindowWindowHandle = FSlateApplication::Get().FindBestParentWindowHandleForDialogs(nullptr);
-
-			//Cache anim file names
-			TArray<FString> AnimFileNames;
-			for(auto Iter = AnimSequences.CreateIterator(); Iter; ++Iter)
-			{
-				AnimFileNames.Add(Iter->Get()->GetName() + TEXT(".fbx"));
-			}
-
-			IPlatformFile& PlatformFile = FPlatformFileManager::Get().GetPlatformFile();
-			FString DestinationFolder;
-
-			const FString Title = LOCTEXT("ExportFBXsToFolderTitle", "Choose a destination folder for the FBX file(s)").ToString();
-
-			if(AnimSequences.Num() > 1)
-			{
-				bool bFolderValid = false;
-				// More than one file, just ask for directory
-				while(!bFolderValid)
-				{
-					const bool bFolderSelected = DesktopPlatform->OpenDirectoryDialog(
-						ParentWindowWindowHandle,
-						Title,
-						FEditorDirectories::Get().GetLastDirectory(ELastDirectory::GENERIC_EXPORT),
-						DestinationFolder
-						);
-
-					if(!bFolderSelected)
-					{
-						// User canceled, return
-						return;
-					}
-
-					FEditorDirectories::Get().SetLastDirectory(ELastDirectory::GENERIC_EXPORT, DestinationFolder);
-					FPaths::NormalizeFilename(DestinationFolder);
-
-					//Check whether there are any fbx filename conflicts in this folder
-					for(auto Iter = AnimFileNames.CreateIterator(); Iter; ++Iter)
-					{
-						FString& AnimFileName = *Iter;
-						FString FullPath = DestinationFolder + "/" + AnimFileName;
-
-						bFolderValid = true;
-						if(PlatformFile.FileExists(*FullPath))
-						{
-							FFormatNamedArguments Args;
-							Args.Add(TEXT("DestinationFolder"), FText::FromString(DestinationFolder));
-							const FText DialogMessage = FText::Format(LOCTEXT("ExportToFBXFileOverwriteMessage", "Exporting to '{DestinationFolder}' will cause one or more existing FBX files to be overwritten. Would you like to continue?"), Args);
-							EAppReturnType::Type DialogReturn = FMessageDialog::Open(EAppMsgType::YesNo, DialogMessage);
-							bFolderValid = (EAppReturnType::Yes == DialogReturn);
-							break;
-						}
-					}
-				}
-			}
-			else
-			{
-				// One file only, ask for full filename.
-				// Can set bFolderValid from the SaveFileDialog call as the window will handle 
-				// duplicate files for us.
-				TArray<FString> TempDestinationNames;
-				bool bSave = DesktopPlatform->SaveFileDialog(
-					ParentWindowWindowHandle,
-					Title,
-					FEditorDirectories::Get().GetLastDirectory(ELastDirectory::GENERIC_EXPORT),
-					AnimSequences[0]->GetName(),
-					"FBX  |*.fbx",
-					EFileDialogFlags::None,
-					TempDestinationNames
-					);
-
-				if(!bSave)
-				{
-					// Canceled
-					return;
-				}
-				check(TempDestinationNames.Num() == 1);
-				check(AnimFileNames.Num() == 1);
-
-				DestinationFolder = FPaths::GetPath(TempDestinationNames[0]);
-				AnimFileNames[0] = FPaths::GetCleanFilename(TempDestinationNames[0]);
-
-				FEditorDirectories::Get().SetLastDirectory(ELastDirectory::GENERIC_EXPORT, DestinationFolder);
-			}
-
-			EAppReturnType::Type DialogReturn = FMessageDialog::Open(EAppMsgType::YesNo, LOCTEXT("ExportToFBXExportPreviewMeshToo", "Would you like to export the current skeletal mesh with the animation(s)?"));
-			bool bSaveSkeletalMesh = EAppReturnType::Yes == DialogReturn;
-
-			const bool bShowCancel = false;
-			const bool bShowProgressDialog = true;
-			GWarn->BeginSlowTask(LOCTEXT("ExportToFBXProgress", "Exporting Animation(s) to FBX"), bShowProgressDialog, bShowCancel);
-
-			// make sure to use PreviewMesh, when export inside of Persona
-			const int32 NumberOfAnimations = AnimSequences.Num();
-			for(int32 i = 0; i < NumberOfAnimations; ++i)
-			{
-				GWarn->UpdateProgress(i, NumberOfAnimations);
-
-				UAnimSequence* AnimSequence = AnimSequences[i].Get();
-
-				FString FileName = FString::Printf(TEXT("%s/%s"), *DestinationFolder, *AnimFileNames[i]);
-
-				FbxAnimUtils::ExportAnimFbx(*FileName, AnimSequence, PreviewMesh, bSaveSkeletalMesh);
-			}
-
-			GWarn->EndSlowTask();
-		}
-	}
-}
-
-void FPersona::AddLoopingInterpolation(TArray<TWeakObjectPtr<UAnimSequence>>& AnimSequences)
-{
-	FText WarningMessage = LOCTEXT("AddLoopiingInterpolation", "This will add an extra first frame at the end of the animation to create a better looping interpolation. This action cannot be undone. Would you like to proceed?");
-
-	if(FMessageDialog::Open(EAppMsgType::YesNo, WarningMessage) == EAppReturnType::Yes)
-	{
-		for(auto Animation : AnimSequences)
-		{
-			// get first frame and add to the last frame and go through track
-			// now calculating old animated space bases
-			Animation->AddLoopingInterpolation();
-		}
+		FPersonaModule& PersonaModule = FModuleManager::GetModuleChecked<FPersonaModule>("Persona");
+		PersonaModule.AddLoopingInterpolation(AnimSequences);
 	}
 }
 
@@ -3391,15 +2453,17 @@ bool FPersona::IsInPersonaMode(const FName InPersonaMode) const
 
 void FPersona::ChangeSkeletonPreviewMesh()
 {
+	UDebugSkelMeshComponent* PreviewMeshComponent = PersonaToolkit->GetPreviewMeshComponent();
+	const USkeleton& Skeleton = GetSkeletonTree()->GetEditableSkeleton()->GetSkeleton();
+
 	// Menu option cannot be called unless this is valid
-	if(TargetSkeleton->GetPreviewMesh() != PreviewComponent->SkeletalMesh)
+	if(Skeleton.GetPreviewMesh() != PreviewMeshComponent->SkeletalMesh)
 	{
-		const FScopedTransaction Transaction( LOCTEXT("ChangeSkeletonPreviewMesh", "Change Skeleton Preview Mesh") );
-		TargetSkeleton->SetPreviewMesh(PreviewComponent->SkeletalMesh);
+		GetSkeletonTree()->GetEditableSkeleton()->SetPreviewMesh(PreviewMeshComponent->SkeletalMesh);
 
 		FFormatNamedArguments Args;
-		Args.Add( TEXT("PreviewMeshName"), FText::FromString( PreviewComponent->SkeletalMesh->GetName() ) );
-		Args.Add( TEXT("TargetSkeletonName"), FText::FromString( TargetSkeleton->GetName() ) );
+		Args.Add( TEXT("PreviewMeshName"), FText::FromString( PreviewMeshComponent->SkeletalMesh->GetName() ) );
+		Args.Add( TEXT("TargetSkeletonName"), FText::FromString( Skeleton.GetName() ) );
 		FNotificationInfo Info( FText::Format( LOCTEXT("SaveSkeletonWarning", "Please save Skeleton {TargetSkeletonName} if you'd like to keep {PreviewMeshName} as default preview mesh" ), Args ) );
 		Info.ExpireDuration = 5.0f;
 		Info.bUseLargeFont = false;
@@ -3413,100 +2477,12 @@ void FPersona::ChangeSkeletonPreviewMesh()
 
 void FPersona::UpdateSkeletonRefPose()
 {
-	// update ref pose with current preview mesh
-	if (TargetSkeleton)
-	{
-		TargetSkeleton->UpdateReferencePoseFromMesh(GetPreviewMeshComponent()->SkeletalMesh);
-	}
+	GetSkeletonTree()->GetEditableSkeleton()->UpdateSkeletonReferencePose(PersonaToolkit->GetPreviewMeshComponent()->SkeletalMesh);
 }
 
 void FPersona::RemoveUnusedBones()
 {
-	// Menu option cannot be called unless this is valid
-	if(TargetSkeleton)
-	{
-		TArray<FName> SkeletonBones;
-		const FReferenceSkeleton& RefSkeleton = TargetSkeleton->GetReferenceSkeleton();
-
-		for(int32 BoneIndex = 0; BoneIndex < RefSkeleton.GetNum(); ++BoneIndex)
-		{
-			SkeletonBones.Add(RefSkeleton.GetBoneName(BoneIndex));
-		}
-
-		FARFilter Filter;
-		Filter.ClassNames.Add(USkeletalMesh::StaticClass()->GetFName());
-
-		FString SkeletonString = FAssetData(TargetSkeleton).GetExportTextName();
-		Filter.TagsAndValues.Add(GET_MEMBER_NAME_CHECKED(USkeletalMesh, Skeleton), SkeletonString);
-
-		TArray<FAssetData> SkeletalMeshes;
-		FAssetRegistryModule& AssetRegistryModule = FModuleManager::LoadModuleChecked<FAssetRegistryModule>(TEXT("AssetRegistry"));
-		AssetRegistryModule.Get().GetAssets(Filter, SkeletalMeshes);
-
-		FText TimeTakenMessage = FText::Format( LOCTEXT("RemoveUnusedBones_TimeTakenWarning", "In order to verify bone use all Skeletal Meshes that use this skeleton will be loaded, this may take some time.\n\nProceed?\n\nNumber of Meshes: {0}"), FText::AsNumber(SkeletalMeshes.Num()) );
-		
-		if(FMessageDialog::Open( EAppMsgType::YesNo, TimeTakenMessage ) == EAppReturnType::Yes)
-		{
-			const FText StatusUpdate = FText::Format(LOCTEXT("RemoveUnusedBones_ProcessingAssetsFor", "Processing Skeletal Meshes for {0}"), FText::FromString(TargetSkeleton->GetName()) );
-			GWarn->BeginSlowTask(StatusUpdate, true );
-
-			// Loop through all SkeletalMeshes and remove the bones they use from our list
-			for ( int32 MeshIdx = 0; MeshIdx < SkeletalMeshes.Num(); ++MeshIdx )
-			{
-				GWarn->StatusUpdate( MeshIdx, SkeletalMeshes.Num(), StatusUpdate );
-
-				USkeletalMesh* Mesh = Cast<USkeletalMesh>(SkeletalMeshes[MeshIdx].GetAsset());
-				const FReferenceSkeleton& MeshRefSkeleton = Mesh->RefSkeleton;
-
-				for(int32 BoneIndex = 0; BoneIndex < MeshRefSkeleton.GetNum(); ++BoneIndex)
-				{
-					SkeletonBones.Remove(MeshRefSkeleton.GetBoneName(BoneIndex));
-					if(SkeletonBones.Num() == 0)
-					{
-						break;
-					}
-				}
-				if(SkeletonBones.Num() == 0)
-				{
-					break;
-				}
-			}
-
-			GWarn->EndSlowTask();
-
-			//Remove bones that are a parent to bones we aren't removing
-			for(int32 BoneIndex = RefSkeleton.GetNum() -1; BoneIndex >= 0; --BoneIndex)
-			{
-				FName CurrBoneName = RefSkeleton.GetBoneName(BoneIndex);
-				if(!SkeletonBones.Contains(CurrBoneName))
-				{
-					//We aren't removing this bone, so remove parent from list of bones to remove too
-					int32 ParentIndex = RefSkeleton.GetParentIndex(BoneIndex);
-					if(ParentIndex != INDEX_NONE)
-					{
-						SkeletonBones.Remove(RefSkeleton.GetBoneName(ParentIndex));
-					}
-				}
-			}
-
-			// If we have any bones left they are unused
-			if(SkeletonBones.Num() > 0)
-			{
-				const FText RemoveBoneMessage = FText::Format(LOCTEXT("RemoveBoneWarning", "Continuing will remove the following bones from the skeleton '{0}'. These bones are not being used by any of the SkeletalMeshes assigned to this skeleton\n\nOnce the bones have been removed all loaded animations for this skeleton will be recompressed (any that aren't loaded will be recompressed the next time they are loaded)."), FText::FromString(TargetSkeleton->GetName()) );
-
-				// Ask User whether they would like to remove the bones from the skeleton
-				if (SSkeletonBoneRemoval::ShowModal(SkeletonBones, RemoveBoneMessage))
-				{
-					//Remove these bones from the skeleton
-					TargetSkeleton->RemoveBonesFromSkeleton(SkeletonBones, true);
-				}
-			}
-			else
-			{
-				FMessageDialog::Open(EAppMsgType::Ok, LOCTEXT("NoBonesToRemove", "No unused bones were found."));
-			}
-		}
-	}
+	GetSkeletonTree()->GetEditableSkeleton()->RemoveUnusedBones();
 }
 
 void FPersona::Tick(float DeltaTime)
@@ -3517,28 +2493,18 @@ void FPersona::Tick(float DeltaTime)
 	{
 		Viewport.Pin()->RefreshViewport();
 	}
-
-	if (PreviewComponent && LastCachedLODForPreviewComponent != PreviewComponent->PredictedLODLevel)
-	{
-		OnLODChanged.Broadcast();
-		LastCachedLODForPreviewComponent = PreviewComponent->PredictedLODLevel;
-	}
 }
 
 bool FPersona::CanChangeSkeletonPreviewMesh() const
 {
-	return PreviewComponent && PreviewComponent->SkeletalMesh;
+	UDebugSkelMeshComponent* PreviewMeshComponent = PersonaToolkit->GetPreviewMeshComponent();
+	return PreviewMeshComponent && PreviewMeshComponent->SkeletalMesh;
 }
 
 bool FPersona::CanRemoveBones() const
 {
-	return PreviewComponent && PreviewComponent->SkeletalMesh;
-}
-
-bool FPersona::IsRecordAvailable() const
-{
-	// make sure mesh exists
-	return (PreviewComponent && PreviewComponent->SkeletalMesh);
+	UDebugSkelMeshComponent* PreviewMeshComponent = PersonaToolkit->GetPreviewMeshComponent();
+	return PreviewMeshComponent && PreviewMeshComponent->SkeletalMesh;
 }
 
 bool FPersona::IsEditable(UEdGraph* InGraph) const
@@ -3558,32 +2524,7 @@ FText FPersona::GetGraphDecorationString(UEdGraph* InGraph) const
 	return FText::GetEmpty();
 }
 
-void FPersona::ValidatePreviewAttachedAssets(USkeletalMesh* PreviewSkeletalMesh)
-{
-	// Validate the skeleton/meshes attached objects and display a notification to the user if any were broken
-	int32 NumBrokenAssets = TargetSkeleton->ValidatePreviewAttachedObjects();
-	if (PreviewSkeletalMesh)
-	{
-		NumBrokenAssets += PreviewSkeletalMesh->ValidatePreviewAttachedObjects();
-	}
-
-	if (NumBrokenAssets > 0)
-	{
-		// Tell the user that there were assets that could not be loaded
-		FFormatNamedArguments Args;
-		Args.Add(TEXT("NumBrokenAssets"), NumBrokenAssets);
-		FNotificationInfo Info(FText::Format(LOCTEXT("MissingPreviewAttachedAssets", "{NumBrokenAssets} attached assets could not be found on loading and were removed"), Args));
-
-		Info.bUseLargeFont = false;
-		Info.ExpireDuration = 5.0f;
 
-		TSharedPtr<SNotificationItem> Notification = FSlateNotificationManager::Get().AddNotification(Info);
-		if (Notification.IsValid())
-		{
-			Notification->SetCompletionState(SNotificationItem::CS_Fail);
-		}
-	}
-}
 
 void FPersona::OnAnimNotifyWindow()
 {
@@ -3597,42 +2538,14 @@ void FPersona::OnRetargetManager()
 
 void FPersona::OnImportAsset(enum EFBXImportType DefaultImportType)
 {
-	// open dialog
-	// get path
-	FString AssetPath;
-
-	TSharedRef<SImportPathDialog> NewAnimDlg =
-	SNew(SImportPathDialog);
-
-	if(NewAnimDlg->ShowModal() != EAppReturnType::Cancel)
-	{
-		AssetPath = NewAnimDlg->GetAssetPath();
-	
-		UFbxImportUI* ImportUI = NewObject<UFbxImportUI>();
-		ImportUI->Skeleton = TargetSkeleton;
-		ImportUI->MeshTypeToImport = DefaultImportType;
-
-		FbxMeshUtils::SetImportOption(ImportUI);
-
-		// now I have to set skeleton on it. 
-		FAssetToolsModule& AssetToolsModule = FModuleManager::GetModuleChecked<FAssetToolsModule>("AssetTools");
-		TArray<UObject*> NewAssets = AssetToolsModule.Get().ImportAssets(AssetPath);
-
-		if ( NewAssets.Num() > 0 )
-		{
-
-		}
-		else
-		{
-			// error
-		}
-	}
+	FPersonaModule& PersonaModule = FModuleManager::GetModuleChecked<FPersonaModule>("Persona");
+	PersonaModule.ImportNewAsset(PersonaToolkit->GetSkeleton(), DefaultImportType);
 }
 
 void FPersona::OnReimportMesh()
 {
 	// Reimport the asset
-	UDebugSkelMeshComponent* PreviewMeshComponent = GetPreviewMeshComponent();
+	UDebugSkelMeshComponent* PreviewMeshComponent = PersonaToolkit->GetPreviewMeshComponent();
 	if(PreviewMeshComponent && PreviewMeshComponent->SkeletalMesh)
 	{
 		FReimportManager::Instance()->Reimport(PreviewMeshComponent->SkeletalMesh, true);
@@ -3656,12 +2569,13 @@ TStatId FPersona::GetStatId() const
 
 void FPersona::OnBlueprintPreCompile(UBlueprint* BlueprintToCompile)
 {
-	if(PreviewComponent && PreviewComponent->PreviewInstance)
+	UDebugSkelMeshComponent* PreviewMeshComponent = PersonaToolkit->GetPreviewMeshComponent();
+	if(PreviewMeshComponent && PreviewMeshComponent->PreviewInstance)
 	{
 		// If we are compiling an anim notify state the class will soon be sanitized and 
 		// if an anim instance is running a state when that happens it will likely
 		// crash, so we end any states that are about to compile.
-		UAnimPreviewInstance* Instance = PreviewComponent->PreviewInstance;
+		UAnimPreviewInstance* Instance = PreviewMeshComponent->PreviewInstance;
 		USkeletalMeshComponent* SkelMeshComp = Instance->GetSkelMeshComponent();
 
 		for(int32 Idx = Instance->ActiveAnimNotifyState.Num() - 1 ; Idx >= 0 ; --Idx)
@@ -3681,356 +2595,214 @@ void FPersona::OnBlueprintChangedImpl(UBlueprint* InBlueprint, bool bIsJustBeing
 	FBlueprintEditor::OnBlueprintChangedImpl(InBlueprint, bIsJustBeingCompiled);
 
 	UObject* CurrentDebugObject = GetBlueprintObj()->GetObjectBeingDebugged();
-	const bool bIsDebuggingPreview = (PreviewComponent != NULL) && PreviewComponent->IsAnimBlueprintInstanced() && (PreviewComponent->GetAnimInstance() == CurrentDebugObject);
+	UDebugSkelMeshComponent* PreviewMeshComponent = PersonaToolkit->GetPreviewMeshComponent();
 
-	if(PreviewComponent != NULL)
+	if(PreviewMeshComponent != NULL)
 	{
 		// Reinitialize the animation, anything we reference could have changed triggering
 		// the blueprint change
-		PreviewComponent->InitAnim(true);
+		PreviewMeshComponent->InitAnim(true);
 
+		const bool bIsDebuggingPreview = PreviewMeshComponent->IsAnimBlueprintInstanced() && (PreviewMeshComponent->GetAnimInstance() == CurrentDebugObject);
 		if(bIsDebuggingPreview)
 		{
-			GetBlueprintObj()->SetObjectBeingDebugged(PreviewComponent->GetAnimInstance());
+			GetBlueprintObj()->SetObjectBeingDebugged(PreviewMeshComponent->GetAnimInstance());
 		}
 	}
 
 	// calls PostCompile to copy proper values between anim nodes
-	if(Viewport.IsValid())
-	{
-		Viewport.Pin()->GetAnimationViewportClient()->PostCompile();
-	}
+	OnPostCompile();
 }
 
 void FPersona::ShowReferencePose(bool bReferencePose)
 {
-	if(PreviewComponent)
+	GetPreviewScene()->ShowReferencePose(bReferencePose);
+}
+
+void FPersona::TestSkeletonCurveNamesForUse() const
 	{
-		if(bReferencePose == false)
-		{
-			if(IsInPersonaMode(FPersonaModes::AnimBlueprintEditMode))
-			{
-				PreviewComponent->EnablePreview(false, nullptr);
+	FPersonaModule& PersonaModule = FModuleManager::GetModuleChecked<FPersonaModule>("Persona");
+	PersonaModule.TestSkeletonCurveNamesForUse(GetSkeletonTree()->GetEditableSkeleton());
+}
 
-				UAnimBlueprint* AnimBP = GetAnimBlueprint();
-				if(AnimBP)
+bool FPersona::IsShowReferencePoseEnabled() const
 				{
-					PreviewComponent->SetAnimInstanceClass(AnimBP->GeneratedClass);
+	return GetPreviewScene()->IsShowReferencePoseEnabled() == true;
 				}
-			}
-			else
+
+bool FPersona::IsPreviewAssetEnabled() const
 			{
-				UObject* PreviewAsset = CachedPreviewAsset.IsValid()? CachedPreviewAsset.Get() : (GetAnimationAssetBeingEdited());
-				PreviewComponent->EnablePreview(true, Cast<UAnimationAsset>(PreviewAsset));
+	return GetPreviewScene()->IsShowReferencePoseEnabled() == false;
 			}
-		}
-		else
-		{
-			if (PreviewComponent->PreviewInstance && PreviewComponent->PreviewInstance->GetCurrentAsset())
+
+FText FPersona::GetPreviewAssetTooltip() const
 			{
-				CachedPreviewAsset = PreviewComponent->PreviewInstance->GetCurrentAsset();
+	return GetPreviewScene()->GetPreviewAssetTooltip(IsInPersonaMode(FPersonaModes::AnimBlueprintEditMode));
 			}
 			
-			PreviewComponent->EnablePreview(true, nullptr);
-		}
-	}
+void FPersona::SetSelectedBlendProfile(UBlendProfile* InBlendProfile)
+{
+	OnBlendProfileSelected.Broadcast(InBlendProfile);
 }
 
-bool FPersona::ShouldDisplayAdditiveScaleErrorMessage()
-{
-	UAnimSequence* AnimSequence = Cast<UAnimSequence>(GetAnimationAssetBeingEdited());
-	if (AnimSequence)
+TSharedRef<FAnimationEditorPreviewScene> FPersona::GetPreviewScene() const
 	{
-		if (AnimSequence->IsValidAdditive() && AnimSequence->RefPoseSeq)
-		{
-			if (RefPoseGuid != AnimSequence->RefPoseSeq->RawDataGuid)
-			{
-				RefPoseGuid = AnimSequence->RefPoseSeq->RawDataGuid;
-				bDoesAdditiveRefPoseHaveZeroScale = AnimSequence->DoesSequenceContainZeroScale();
-			}
-			return bDoesAdditiveRefPoseHaveZeroScale;
-		}
+	check(PersonaToolkit.IsValid()); 
+	return StaticCastSharedRef<FAnimationEditorPreviewScene>(PersonaToolkit->GetPreviewScene()); 
 	}
 
-	RefPoseGuid.Invalidate();
-	return false;
+void FPersona::HandleCurvesChanged()
+{
+	OnCurvesChanged.Broadcast();
 }
 
-void PopulateWithAssets(FName ClassName, FName SkeletonMemberName, const FString& SkeletonString, TArray<FAssetData>& OutAssets)
+void FPersona::HandleSectionsChanged()
 {
-	FAssetRegistryModule& AssetRegistryModule = FModuleManager::LoadModuleChecked<FAssetRegistryModule>(TEXT("AssetRegistry"));
-
-	FARFilter Filter;
-	Filter.ClassNames.Add(ClassName);
-	Filter.TagsAndValues.Add(SkeletonMemberName, SkeletonString);
-
-	AssetRegistryModule.Get().GetAssets(Filter, OutAssets);
+	OnSectionsChanged.Broadcast();
 }
 
-void FPersona::TestSkeletonCurveNamesForUse() const
+void FPersona::HandleObjectsSelected(const TArray<UObject*>& InObjects)
 {
-	if (TargetSkeleton)
-	{
-		if (const FSmartNameMapping* Mapping = TargetSkeleton->GetSmartNameContainer(USkeleton::AnimCurveMappingName))
-		{
-			const FString SkeletonString = FAssetData(TargetSkeleton).GetExportTextName();
+	SetDetailObjects(InObjects);
+}
 			
-			TArray<FAssetData> SkeletalMeshes;
-			PopulateWithAssets(USkeletalMesh::StaticClass()->GetFName(), GET_MEMBER_NAME_CHECKED(USkeletalMesh, Skeleton), SkeletonString, SkeletalMeshes);
-			TArray<FAssetData> Animations;
-			PopulateWithAssets(UAnimSequence::StaticClass()->GetFName(), FName("Skeleton"), SkeletonString, Animations);
-
-			FText TimeTakenMessage = FText::Format(LOCTEXT("VerifyCurves_TimeTakenWarning", "In order to verify curve usage all Skeletal Meshes and Animations that use this skeleton will be loaded, this may take some time.\n\nProceed?\n\nNumber of Meshes: {0}\nNumber of Animations: {1}"), FText::AsNumber(SkeletalMeshes.Num()), FText::AsNumber(Animations.Num()));
-
-			if (FMessageDialog::Open(EAppMsgType::YesNo, TimeTakenMessage) == EAppReturnType::Yes)
+void FPersona::HandleObjectSelected(UObject* InObject)
 			{
-				const FText LoadingStatusUpdate = FText::Format(LOCTEXT("VerifyCurves_LoadingAllAnimations", "Loading all animations for skeleton '{0}'"), FText::FromString(TargetSkeleton->GetName()));
-				{
-					FScopedSlowTask LoadingAnimSlowTask(Animations.Num(), LoadingStatusUpdate);
-					LoadingAnimSlowTask.MakeDialog();
+	SetDetailObject(InObject);
+}
 
-					// Loop through all animations to load then, this makes sure smart names are all up to date
-					for (const FAssetData& Anim : Animations)
+void FPersona::HandleAnimNotifiesChanged()
 					{
-						LoadingAnimSlowTask.EnterProgressFrame();
-						UAnimSequence* Seq = Cast<UAnimSequence>(Anim.GetAsset());
-					}
+	OnAnimNotifiesChanged.Broadcast();
 				}
 
-				// Grab all curve names for this skeleton
-				TArray<FName> UnusedNames;
-				Mapping->FillNameArray(UnusedNames);
-
-				const FText ProcessingStatusUpdate = FText::Format(LOCTEXT("VerifyCurves_ProcessingSkeletalMeshes", "Looking at curve useage for each skeletal mesh of skeleton '{0}'"), FText::FromString(TargetSkeleton->GetName()));
+void FPersona::HandleViewportCreated(const TSharedRef<IPersonaViewport>& InViewport)
 				{
-					FScopedSlowTask LoadingSkelMeshSlowTask(SkeletalMeshes.Num(), ProcessingStatusUpdate);
-					LoadingSkelMeshSlowTask.MakeDialog();
-
-					for (int32 MeshIdx = 0; MeshIdx < SkeletalMeshes.Num(); ++MeshIdx)
+	// mode switch data sharing: We save data from previous viewports and restore after they are switched
+	bool bRestoreData = Viewport.IsValid();
+	if (bRestoreData)
 					{
-						LoadingSkelMeshSlowTask.EnterProgressFrame();
+		ViewportState = InViewport->SaveState();
+	}
 
-						const USkeletalMesh* Mesh = Cast<USkeletalMesh>(SkeletalMeshes[MeshIdx].GetAsset());
+	SetViewport(StaticCastSharedRef<SAnimationEditorViewportTabBody>(InViewport));
 
-						// Filter morph targets from curves
-						const TArray<UMorphTarget*>& MorphTargets = Mesh->MorphTargets;
-						for (int32 I = 0; I < MorphTargets.Num(); ++I)
+	if (bRestoreData)
 						{
-							const int32 CurveIndex = UnusedNames.RemoveSingleSwap(MorphTargets[I]->GetFName(), false);
+		InViewport->RestoreState(ViewportState.ToSharedRef());
 						}
-
-						// Filter material params from curves
-						for (const FSkeletalMaterial& Mat : Mesh->Materials)
-						{
-							if (UnusedNames.Num() == 0)
-							{
-								break; // Done
 							}
 
-							UMaterial* Material = (Mat.MaterialInterface != nullptr) ? Mat.MaterialInterface->GetMaterial() : nullptr;
-							if (Material)
+void FPersona::HandleOpenNewAsset(UObject* InNewAsset)
 							{
-								TArray<FName> OutParameterNames;
-								TArray<FGuid> OutParameterIds;
-
-								// Retrieve all scalar parameter names from the material
-								Material->GetAllScalarParameterNames(OutParameterNames, OutParameterIds);
-
-								for (FName SPName : OutParameterNames)
+	if (UAnimationAsset* NewAnimationAsset = Cast<UAnimationAsset>(InNewAsset))
 								{
-									UnusedNames.RemoveSingleSwap(SPName);
-								}
+		OpenNewDocumentTab(NewAnimationAsset);
+		SetPreviewAnimationAsset(NewAnimationAsset);
 							}
 						}
-					}
-				}
-
-				FMessageLog CurveOutput("Persona");
-				CurveOutput.NewPage(LOCTEXT("PersonaMessageLogName", "Persona"));
 
-				bool bFoundIssue = false;
-
-				const FText ProcessingAnimStatusUpdate = FText::Format(LOCTEXT("VerifyCurves_FindAnimationsWithUnusedCurves", "Finding animations that reference unused curves on skeleton '{0}'"), FText::FromString(TargetSkeleton->GetName()));
+UObject* FPersona::HandleGetObject()
 				{
-					FScopedSlowTask ProcessingAnimationsSlowTask(Animations.Num(), ProcessingAnimStatusUpdate);
-					ProcessingAnimationsSlowTask.MakeDialog();
+	return GetEditingObject();
+}
 
-					for (const FAssetData& Anim : Animations)
+void FPersona::OnPostCompile()
 					{
-						ProcessingAnimationsSlowTask.EnterProgressFrame();
-						UAnimSequence* Seq = Cast<UAnimSequence>(Anim.GetAsset());
+	// act as if we have re-selected, so internal pointers are updated
+	OnSelectedNodesChangedImpl(GetSelectedNodes());
 
-						TSharedPtr<FTokenizedMessage> Message;
-						for (FFloatCurve& Curve : Seq->RawCurveData.FloatCurves)
+	// if the user manipulated Pin values directly from the node, then should copy updated values to the internal node to retain data consistency
+	UEdGraph* FocusedGraph = GetFocusedGraph();
+	if (FocusedGraph)
 						{
-							if (UnusedNames.Contains(Curve.Name.DisplayName))
+		// find UAnimGraphNode_Base
+		for (UEdGraphNode* Node : FocusedGraph->Nodes)
 							{
-								bFoundIssue = true;
-								if (!Message.IsValid())
+			UAnimGraphNode_Base* AnimGraphNode = Cast<UAnimGraphNode_Base>(Node);
+			if (AnimGraphNode)
 								{
-									Message = CurveOutput.Warning();
-									Message->AddToken(FAssetNameToken::Create(Anim.ObjectPath.ToString(), FText::FromName(Anim.AssetName)));
-									Message->AddToken(FTextToken::Create(LOCTEXT("VerifyCurves_FoundAnimationsWithUnusedReferences", "References the following curves that are not used for either morph targets or material parameters and so may be unneeded")));
-								}
-								CurveOutput.Info(FText::FromName(Curve.Name.DisplayName));
-							}
-						}
-					}
-				}
-
-				if (bFoundIssue)
+				FAnimNode_Base* AnimNode = FindAnimNode(AnimGraphNode);
+				if (AnimNode)
 				{
-					CurveOutput.Notify();
+					AnimGraphNode->CopyNodeDataToPreviewNode(AnimNode);
 				}
 			}
 		}
 	}
 }
 
-bool FPersona::CanShowReferencePose() const
+FAnimNode_Base* FPersona::FindAnimNode(UAnimGraphNode_Base* AnimGraphNode) const
 {
-	return PreviewComponent != NULL;
-}
-
-bool FPersona::IsShowReferencePoseEnabled() const
+	FAnimNode_Base* AnimNode = nullptr;
+	if (AnimGraphNode)
 {
-	if(PreviewComponent)
+		UDebugSkelMeshComponent* PreviewMeshComponent = GetPreviewScene()->GetPreviewMeshComponent();
+		if (PreviewMeshComponent != nullptr && PreviewMeshComponent->GetAnimInstance() != nullptr)
 	{
-		return PreviewComponent->IsPreviewOn() && PreviewComponent->PreviewInstance->GetCurrentAsset() == NULL;
+			AnimNode = AnimGraphNode->FindDebugAnimNode(PreviewMeshComponent);
 	}
-	return false;
 }
 
-bool FPersona::CanPreviewAsset() const
-{
-	return CanShowReferencePose();
+	return AnimNode;
 }
 
-bool FPersona::IsPreviewAssetEnabled() const
+void FPersona::OnSelectedNodesChangedImpl(const TSet<class UObject*>& NewSelection)
 {
-	return IsShowReferencePoseEnabled() == false;
-}
+	FBlueprintEditor::OnSelectedNodesChangedImpl(NewSelection);
 
-FText FPersona::GetPreviewAssetTooltip() const
+	if (SelectedAnimGraphNode.IsValid())
 {
-	// if already looking at ref pose
-	if(IsShowReferencePoseEnabled())
-	{
-		FString AssetName = TEXT("None Available. Please select asset to preview.");
+		// copy values from the preview node to ensure data consistency
+		FAnimNode_Base* PreviewNode = FindAnimNode(SelectedAnimGraphNode.Get());
 
-		if(IsInPersonaMode(FPersonaModes::AnimBlueprintEditMode))
-		{
-			UAnimBlueprint* AnimBP = GetAnimBlueprint();
-			if(AnimBP)
-			{
-				AssetName = AnimBP->GetName();
+		SelectedAnimGraphNode->OnNodeSelected(false, *GetAssetEditorModeManager(), PreviewNode);
+		SelectedAnimGraphNode.Reset();
 			}
-		}
-		else
+
+	// if we only have one node selected, let it know
+	UAnimGraphNode_Base* NewSelectedAnimGraphNode = nullptr;
+	if (NewSelection.Num() == 1)
 		{
-			UObject* PreviewAsset = CachedPreviewAsset.IsValid()? CachedPreviewAsset.Get() : (GetAnimationAssetBeingEdited());
-			if (PreviewAsset)
+		NewSelectedAnimGraphNode = Cast<UAnimGraphNode_SkeletalControlBase>(*NewSelection.CreateConstIterator());
+		if (NewSelectedAnimGraphNode != nullptr)
 			{
-				AssetName = PreviewAsset->GetName();
+			SelectedAnimGraphNode = NewSelectedAnimGraphNode;
+
+			// copy values to the preview node to ensure data consistency
+			FAnimNode_Base* PreviewNode = FindAnimNode(SelectedAnimGraphNode.Get());
+
+			SelectedAnimGraphNode->OnNodeSelected(true, *GetAssetEditorModeManager(), PreviewNode);
 			}
 		}
-		return FText::FromString(FString::Printf(TEXT("Preview %s"), *AssetName));
-	}
-	else
-	{
-		return FText::FromString(FString::Printf(TEXT("Currently previewing %s"), *PreviewComponent->GetPreviewText()));
-	}
 }
 
-void FPersona::SetSelectedBlendProfile(UBlendProfile* InBlendProfile)
+void FPersona::HandlePinDefaultValueChanged(UEdGraphPin* InPinThatChanged)
 {
-	OnBlendProfileSelected.Broadcast(InBlendProfile);
-}
-
-bool FPersona::IsRecording() const
+	UAnimGraphNode_Base* AnimGraphNode = Cast<UAnimGraphNode_Base>(InPinThatChanged->GetOwningNode());
+	if (AnimGraphNode)
 {
-	FPersonaModule& PersonaModule = FModuleManager::GetModuleChecked<FPersonaModule>("Persona");
-
-	bool bInRecording = false;
-	PersonaModule.OnIsRecordingActive().ExecuteIfBound(PreviewComponent, bInRecording);
-
-	return bInRecording;
-}
-
-void FPersona::StopRecording()
+		FAnimNode_Base* AnimNode = FindAnimNode(AnimGraphNode);
+		if (AnimNode)
 {
-	FPersonaModule& PersonaModule = FModuleManager::GetModuleChecked<FPersonaModule>("Persona");
-	PersonaModule.OnStopRecording().ExecuteIfBound(PreviewComponent);
+			AnimGraphNode->CopyNodeDataToPreviewNode(AnimNode);
 }
-
-UAnimSequence* FPersona::GetCurrentRecording() const
-{
-	FPersonaModule& PersonaModule = FModuleManager::GetModuleChecked<FPersonaModule>("Persona");
-	UAnimSequence* Recording = nullptr;
-	PersonaModule.OnGetCurrentRecording().ExecuteIfBound(PreviewComponent, Recording);
-	return Recording;
 }
-
-float FPersona::GetCurrentRecordingTime() const
-{
-	FPersonaModule& PersonaModule = FModuleManager::GetModuleChecked<FPersonaModule>("Persona");
-	float RecordingTime = 0.0f;
-	PersonaModule.OnGetCurrentRecordingTime().ExecuteIfBound(PreviewComponent, RecordingTime);
-	return RecordingTime;
 }
 
 void FPersona::TogglePlayback()
 {
+	UDebugSkelMeshComponent* PreviewComponent = GetPreviewScene()->GetPreviewMeshComponent();
 	if (PreviewComponent && PreviewComponent->PreviewInstance)
 	{
 		PreviewComponent->PreviewInstance->SetPlaying(!PreviewComponent->PreviewInstance->IsPlaying());
 	}
 }
 
-static class FMeshHierarchyCmd : private FSelfRegisteringExec
-{
-public:
-	/** Console commands, see embeded usage statement **/
-	virtual bool Exec( UWorld*, const TCHAR* Cmd, FOutputDevice& Ar ) override
-	{
-		bool bResult = false;
-		if(FParse::Command(&Cmd,TEXT("TMH")))
+void FPersona::HandleAnimationSequenceBrowserCreated(const TSharedRef<IAnimationSequenceBrowser>& InSequenceBrowser)
 		{
-			Ar.Log(TEXT("Starting Mesh Test"));
-			FARFilter Filter;
-			Filter.ClassNames.Add(USkeletalMesh::StaticClass()->GetFName());
-
-			TArray<FAssetData> SkeletalMeshes;
-			FAssetRegistryModule& AssetRegistryModule = FModuleManager::LoadModuleChecked<FAssetRegistryModule>(TEXT("AssetRegistry"));
-			AssetRegistryModule.Get().GetAssets(Filter, SkeletalMeshes);
-
-			const FText StatusUpdate = LOCTEXT("RemoveUnusedBones_ProcessingAssets", "Processing Skeletal Meshes");
-			GWarn->BeginSlowTask(StatusUpdate, true );
-
-			// go through all assets try load
-			for ( int32 MeshIdx = 0; MeshIdx < SkeletalMeshes.Num(); ++MeshIdx )
-			{
-				GWarn->StatusUpdate( MeshIdx, SkeletalMeshes.Num(), StatusUpdate );
-
-				USkeletalMesh* Mesh = Cast<USkeletalMesh>(SkeletalMeshes[MeshIdx].GetAsset());
-				if(Mesh->Skeleton)
-				{
-					FName MeshRoot = Mesh->RefSkeleton.GetBoneName(0);
-					FName SkelRoot = Mesh->Skeleton->GetReferenceSkeleton().GetBoneName(0);
-
-					if(MeshRoot != SkelRoot)
-					{
-						Ar.Logf( TEXT("Mesh Found '%s' %s->%s"), *Mesh->GetName(), *MeshRoot.ToString(), *SkelRoot.ToString());
-					}
-				}
-			}
-
-			GWarn->EndSlowTask();
-			Ar.Log(TEXT("Mesh Test Finished"));
-		}
-		return bResult;
+	SequenceBrowser = InSequenceBrowser;
 	}
-} MeshHierarchyCmdExec;
 
 #undef LOCTEXT_NAMESPACE
 
