Copying //UE4/Dev-Framework to //UE4/Dev-Main (Source: //UE4/Dev-Framework @ 3136612)

diff --git a/Engine/Source/Editor/Persona/Private/Persona.h b/Engine/Source/Editor/Persona/Private/Persona.h
--- a/Engine/Source/Editor/Persona/Private/Persona.h
+++ b/Engine/Source/Editor/Persona/Private/Persona.h
@@ -4,10 +4,10 @@
 
 #include "Editor/Kismet/Public/BlueprintEditor.h"
 
-#include "AdvancedPreviewScene.h"
-
 #include "PersonaModule.h"
 #include "PersonaCommands.h"
+#include "AnimationEditorPreviewScene.h"
+#include "PersonaDelegates.h"
 
 //////////////////////////////////////////////////////////////////////////
 // FPersona
@@ -45,8 +45,6 @@ public:
 
 	/** Set the current preview viewport for Persona */
 	void SetViewport(TWeakPtr<class SAnimationEditorViewportTabBody> NewViewport);
-	/** Set Sequence browser */
-	void SetSequenceBrowser(class SAnimationSequenceBrowser* SequenceBrowser);
 
 	/** Refresh viewport */
 	void RefreshViewport();
@@ -57,26 +55,12 @@ public:
 	UObject* GetAnimationAssetBeingEdited() const;
 
 	/** Update the inspector that displays information about the current selection*/
-	void UpdateSelectionDetails(UObject* Object, const FText& ForcedTitle);
-
+	void SetDetailObjects(const TArray<UObject*>& InObjects);
 	void SetDetailObject(UObject* Obj);
 
+	/** FBlueprintEditor interface */
 	virtual void OnActiveTabChanged(TSharedPtr<SDockTab> PreviouslyActive, TSharedPtr<SDockTab> NewlyActivated) override;
-
-	/** Get the preview DebugSkelMeshComponent */
-	UDebugSkelMeshComponent* GetPreviewMeshComponent() const;
-	/** Get the AnimInstance associated with the preview DEbugSkelMeshComponent */
-	UAnimInstance* GetPreviewAnimInstance() const;
-
-	// Gets the skeleton being edited/viewed by this Persona instance
-	USkeleton* GetSkeleton() const;
-	UObject* GetSkeletonAsObject() const;
-
-	// Gets the skeletal mesh being edited/viewed by this Persona instance
-	USkeletalMesh* GetMesh() const;
-
-	// Gets the physics asset being edited/viewed by this Persona instance
-	UPhysicsAsset* GetPhysicsAsset() const;
+	virtual void OnSelectedNodesChangedImpl(const TSet<class UObject*>& NewSelection) override;
 
 	// Gets the Anim Blueprint being edited/viewed by this Persona instance
 	UAnimBlueprint* GetAnimBlueprint() const;
@@ -84,38 +68,6 @@ public:
 	// Sets the current preview mesh
 	void SetPreviewMesh(USkeletalMesh* NewPreviewMesh);
 
-	// Validate preview attached assets on skeleton and supplied skeletal mesh, notifying user if any are removed
-	void ValidatePreviewAttachedAssets(USkeletalMesh* SkeletalMeshToUse);
-
-	// Thunks for SContentReference
-	UObject* GetMeshAsObject() const;
-	UObject* GetPhysicsAssetAsObject() const;
-	UObject* GetAnimBlueprintAsObject() const;
-
-	// Sets the selected bone to preview
-	void SetSelectedBone(class USkeleton* Skeleton, const FName& BoneName, bool bRebroadcast = true );
-
-	// Sets the selected morph target
-	void SetSelectedMorphTargets(class USkeletalMesh *SkeletonMesh, const TArray<FName>& SelectedMorphTargetNames);
-
-	/** Helper method to clear out any selected bones */
-	void ClearSelectedBones();
-
-	/** Sets the selected socket to preview */
-	void SetSelectedSocket( const FSelectedSocketInfo& SocketInfo, bool bRebroadcast = true );
-
-	/** Duplicates and selects the socket specified (used when Alt-Dragging in the viewport ) */
-	void DuplicateAndSelectSocket( const FSelectedSocketInfo& SocketInfoToDuplicate, const FName& NewParentBoneName = FName() );
-
-	/** Clears the selected socket (removes both the rotate/transform widget and clears the details panel) */
-	void ClearSelectedSocket();
-
-	/** Rename socket to new name */
-	void RenameSocket( USkeletalMeshSocket* Socket, const FName& NewSocketName );
-
-	/** Reparent socket to new parent */
-	void ChangeSocketParent( USkeletalMeshSocket* Socket, const FName& NewSocketParent );
-
 	/** Clears the selected wind actor */
 	void ClearSelectedWindActor();
 
@@ -128,46 +80,8 @@ public:
 	/* Open details panel for the skeletal preview mesh */
 	FReply OnClickEditMesh();
 
-	/** Generates a unique socket name from the input name, by changing the FName's number */
-	FName GenerateUniqueSocketName( FName InName );
-
-	/** Function to tell you if a socket name already exists in a given array of sockets */
-	bool DoesSocketAlreadyExist( const class USkeletalMeshSocket* InSocket, const FText& InSocketName, const TArray< USkeletalMeshSocket* >& InSocketArray ) const;
-
-	/** Creates a new skeletal mesh component with the right master pose component set */
-	USkeletalMeshComponent*	CreateNewSkeletalMeshComponent();
-
-	/** Handler for a user removing a mesh in the loaded meshes panel */
-	void RemoveAdditionalMesh(USkeletalMeshComponent* MeshToRemove);
-
-	/** Cleans up any additional meshes that have been loaded */
-	void ClearAllAdditionalMeshes();
-
-	/** Adds to the viewport all the attached preview objects that the current skeleton and mesh contains */
-	void AddPreviewAttachedObjects();
-
-	/** Attaches an object to the preview component using the supplied attach name, returning whether it was successfully attached or not */
-	bool AttachObjectToPreviewComponent( UObject* Object, FName AttachTo, FPreviewAssetAttachContainer* PreviewAssetContainer = NULL );
-
-	/**Removes a currently attached object from the preview component */
-	void RemoveAttachedObjectFromPreviewComponent(UObject* Object, FName AttachedTo);
-
-	/** Get the viewport scene component for the specified attached asset */
-	USceneComponent* GetComponentForAttachedObject(UObject* Asset, FName AttachedTo);
-
-	/** Removes attached components from the preview component. (WARNING: There is a possibility that this function will
-	 *  remove the wrong component if 2 of the same type (same UObject) are attached at the same location!
-	 *
-	 * @param bRemovePreviewAttached	Specifies whether all attached components are removed or just the ones
-	 *									that haven't come from the target skeleton.
-	 */
-	void RemoveAttachedComponent(bool bRemovePreviewAttached = true);
-
-	/** Destroy the supplied component (and its children) */
-	void CleanupComponent(USceneComponent* Component);
-
 	/** Returns the editors preview scene */
-	FAdvancedPreviewScene& GetPreviewScene() { return PreviewScene; }
+	TSharedRef<FAnimationEditorPreviewScene> GetPreviewScene() const;
 
 	/** Gets called when it's clicked via mode tab - Reinitialize the mode **/
 	void ReinitMode();
@@ -178,21 +92,34 @@ public:
 	/** Called when the blend profile tab selects a profile */
 	void SetSelectedBlendProfile(UBlendProfile* InBlendProfile);
 
-	/** Check whether this Persona instance is recording */
-	bool IsRecording() const;
+	/** Handle sections changing - broadcasts event */
+	void HandleSectionsChanged();
 
-	/** Stop recording in this Persona instance */
-	void StopRecording();
+	/** Handle curves changing - broadcasts event */
+	void HandleCurvesChanged();
 
-	/** Get the currently recording animation */
-	UAnimSequence* GetCurrentRecording() const;
+	/** Handle anim notifes changing - broadcast delegate */
+	void HandleAnimNotifiesChanged();
 
-	/** Get the currently recording animation time */
-	float GetCurrentRecordingTime() const;
+	/** Handle general object selection */
+	void HandleObjectsSelected(const TArray<UObject*>& InObjects);
+	void HandleObjectSelected(UObject* InObject);
+
+	/** Handle when a viewport widget is created */
+	void HandleViewportCreated(const TSharedRef<IPersonaViewport>& InViewport);
+
+	/** Handling opening a new asset from within the editor */
+	void HandleOpenNewAsset(UObject* InNewAsset);
+
+	/** Get the object to be displayed in the asset properties */
+	UObject* HandleGetObject();
 
 	/** Toggle playback of the current animation */
 	void TogglePlayback();
 
+	/** Store a ptr to the anim browser when it is created */
+	void HandleAnimationSequenceBrowserCreated(const TSharedRef<class IAnimationSequenceBrowser>& InSequenceBrowser);
+
 public:
 	//~ Begin IToolkit Interface
 	virtual FName GetToolkitContextFName() const override;
@@ -236,12 +163,15 @@ public:
 	bool CanShowReferencePose() const;
 	void ShowReferencePose(bool bReferencePose);
 
-	/* Handle error checking for additive base pose */
-	bool ShouldDisplayAdditiveScaleErrorMessage();
-
 	/** Validates curve use */
 	void TestSkeletonCurveNamesForUse() const;
 
+	/** Get the skeleton tree this Persona editor is hosting */
+	TSharedRef<class ISkeletonTree> GetSkeletonTree() const { return SkeletonTree.ToSharedRef(); }
+
+	/** Get the Persona toolkit embedded in this editor */
+	TSharedRef<IPersonaToolkit> GetPersonaToolkit() const { return PersonaToolkit.ToSharedRef(); }
+
 protected:
 	bool IsPreviewAssetEnabled() const;
 	bool CanPreviewAsset() const;
@@ -259,6 +189,7 @@ protected:
 	virtual bool CanRemovePosePin() const override;
 	virtual void Compile() override;
 	virtual void OnGraphEditorFocused(const TSharedRef<class SGraphEditor>& InGraphEditor) override;
+	virtual void OnGraphEditorBackgrounded(const TSharedRef<SGraphEditor>& InGraphEditor) override;
 	virtual void OnConvertToSequenceEvaluator() override;
 	virtual void OnConvertToSequencePlayer() override;
 	virtual void OnConvertToBlendSpaceEvaluator() override;
@@ -291,14 +222,11 @@ protected:
 	virtual void SaveAsset_Execute() override;
 	// End of FAssetEditorToolkit
 
-	// Generic Command handlers
-	void OnCommandGenericDelete();
-
 	// Pose watch handler
 	void OnTogglePoseWatch();
 protected:
 	// 
-	TSharedPtr<SDockTab> OpenNewAnimationDocumentTab(UObject* InAnimAsset);
+	TSharedPtr<SDockTab> OpenNewAnimationDocumentTab(UAnimationAsset* InAnimAsset);
 
 	// Creates an editor widget for a specified animation document and returns the document link 
 	TSharedPtr<SWidget> CreateEditorWidgetForAnimDocument(UObject* InAnimAsset, FString& OutDocumentLink);
@@ -312,25 +240,15 @@ protected:
 	/** Clear up Preview Mesh's AnimNotifyStates. Called when undo or redo**/
 	void ClearupPreviewMeshAnimNotifyStates();
 
-	/** Handler connected to open skeleton so we get notified when it changes */
-	void OnSkeletonHierarchyChanged();
-
 protected:
 	//~ IToolkit interface
 	virtual void GetSaveableObjects(TArray<UObject*>& OutObjects) const override;
 
 protected:
 	USkeleton* TargetSkeleton;
-	TWeakObjectPtr<UObject> CachedPreviewAsset;
 
 public:
-	class UDebugSkelMeshComponent* PreviewComponent;
-
-	// Array of loaded additional meshes
-	TArray<USkeletalMeshComponent*> AdditionalMeshes;
-
 	// The animation document currently being edited
-	mutable TWeakObjectPtr<UObject> SharedAnimAssetBeingEdited;
 	TWeakPtr<SDockTab> SharedAnimDocumentTab;
 
 public:
@@ -342,53 +260,12 @@ public:
 	void OnPropertyChanged(UObject* ObjectBeingModified, FPropertyChangedEvent& PropertyChangedEvent);
 
 	/** Shared data between modes - for now only used for viewport **/
-	FPersonaModeSharedData ModeSharedData;
+	TSharedPtr<IPersonaViewportState> ViewportState;
 
 	/** holding this pointer to refresh persona mesh detials tab when LOD is changed **/
 	class IDetailLayoutBuilder* PersonaMeshDetailLayout;
-private:
-	// called when animation asset has been changed
-	DECLARE_MULTICAST_DELEGATE_OneParam( FOnAnimChangedMulticaster, UAnimationAsset* )
-	// Called when the preview mesh has been changed
-	DECLARE_MULTICAST_DELEGATE_OneParam( FOnPreviewMeshChangedMulticaster, USkeletalMesh* )
-	// Called when a socket is selected
-	DECLARE_MULTICAST_DELEGATE_OneParam( FOnSelectSocket, const struct FSelectedSocketInfo& )
-	// Called when a bone is selected
-	DECLARE_MULTICAST_DELEGATE_OneParam( FOnSelectBone, const FName& )
-	// Called when a blend profile is selected in the blend profile tab
-	DECLARE_MULTICAST_DELEGATE_OneParam( FOnSelectBlendProfile, UBlendProfile*)
-	// Called when the preview viewport is created
-	DECLARE_MULTICAST_DELEGATE_OneParam( FOnCreateViewport, TWeakPtr<SAnimationEditorViewportTabBody> )
-public:
-	// Called when Persona refreshes
-	typedef FSimpleMulticastDelegate::FDelegate FOnPersonaRefresh;
-
-	/** Registers a delegate to be called when Persona is Refreshed */
-	void RegisterOnPersonaRefresh(const FOnPersonaRefresh& Delegate)
-	{
-		OnPersonaRefresh.Add(Delegate);
-	}
-
-	/** Unregisters refresh delegate */
-	void UnregisterOnPersonaRefresh(SWidget* Widget)
-	{
-		OnPersonaRefresh.RemoveAll(Widget);
-	}
-
-	// anim changed 
-	typedef FOnAnimChangedMulticaster::FDelegate FOnAnimChanged;
 
-	/** Registers a delegate to be called after the preview animation has been changed */
-	void RegisterOnAnimChanged(const FOnAnimChanged& Delegate)
-	{
-		OnAnimChanged.Add(Delegate);
-	}
-
-	/** Unregisters a delegate to be called after the preview animation has been changed */
-	void UnregisterOnAnimChanged(SWidget* Widget)
-	{
-		OnAnimChanged.RemoveAll(Widget);
-	}
+public:
 
 	// Called after an undo is performed to give child widgets a chance to refresh
 	typedef FSimpleMulticastDelegate::FDelegate FOnPostUndo;
@@ -405,50 +282,8 @@ public:
 		OnPostUndo.RemoveAll(Widget);
 	}
 
-	// preview mesh changed 
-	typedef FOnPreviewMeshChangedMulticaster::FDelegate FOnPreviewMeshChanged;
-
-	/** Registers a delegate to be called when the preview mesh is changed */
-	void RegisterOnPreviewMeshChanged(const FOnPreviewMeshChanged& Delegate)
-	{
-		OnPreviewMeshChanged.Add(Delegate);
-	}
-
-	/** Unregisters a delegate to be called when the preview mesh is changed */
-	void UnregisterOnPreviewMeshChanged(SWidget* Widget)
-	{
-		OnPreviewMeshChanged.RemoveAll(Widget);
-	}
-
-	// Bone selected
-	typedef FOnSelectBone::FDelegate FOnBoneSelected;
-
-	/** Registers a delegate to be called when the selected bone is changed */
-	void RegisterOnBoneSelected( const FOnBoneSelected& Delegate )
-	{
-		OnBoneSelected.Add( Delegate );
-	}
-
-	/** Unregisters a delegate to be called when the selected bone is changed */
-	void UnregisterOnBoneSelected( SWidget* Widget )
-	{
-		OnBoneSelected.RemoveAll( Widget );
-	}
-
-	// Socket selected
-	typedef FOnSelectSocket::FDelegate FOnSocketSelected;
-
-	/** Registers a delegate to be called when the selected socket is changed */
-	void RegisterOnSocketSelected( const FOnSocketSelected& Delegate )
-	{
-		OnSocketSelected.Add( Delegate );
-	}
-
-	/** Unregisters a delegate to be called when the selected socket is changed */
-	void UnregisterOnSocketSelected(SWidget* Widget)
-	{
-		OnSocketSelected.RemoveAll( Widget );
-	}
+	/** Delegate called after an undo operation for child widgets to refresh */
+	FSimpleMulticastDelegate OnPostUndo;
 
 	typedef FOnSelectBlendProfile::FDelegate FOnBlendProfileSelected;
 
@@ -462,36 +297,6 @@ public:
 		OnBlendProfileSelected.RemoveAll(Widget);
 	}
 
-	// Called when selection is cleared
-	typedef FSimpleMulticastDelegate::FDelegate FOnAllDeselected;
-
-	/** Registers a delegate to be called when all bones/sockets are deselected */
-	void RegisterOnDeselectAll(const FOnAllDeselected& Delegate)
-	{
-		OnAllDeselected.Add( Delegate );
-	}
-
-	/** Unregisters a delegate to be called when all bones/sockets are deselected */
-	void UnregisterOnDeselectAll(SWidget* Widget)
-	{
-		OnAllDeselected.RemoveAll( Widget );
-	}
-
-	// Called when the skeleton tree has been changed (socket added/deleted/etc)
-	typedef FSimpleMulticastDelegate::FDelegate FOnSkeletonTreeChanged;
-
-	/** Registers a delegate to be called when the skeleton tree has changed */
-	void RegisterOnChangeSkeletonTree(const FOnSkeletonTreeChanged& Delegate)
-	{
-		OnSkeletonTreeChanged.Add( Delegate );
-	}
-
-	/** Unregisters a delegate to be called when all bones/sockets are deselected */
-	void UnregisterOnChangeSkeletonTree(SWidget* Widget)
-	{
-		OnSkeletonTreeChanged.RemoveAll( Widget );
-	}
-
 	// Called when the notifies of the current animation are changed
 	typedef FSimpleMulticastDelegate::FDelegate FOnAnimNotifiesChanged;
 
@@ -558,21 +363,6 @@ public:
 		OnViewportCreated.RemoveAll( Widget );
 	}
 
-	// Called when generic delete happens
-	typedef FSimpleMulticastDelegate::FDelegate FOnDeleteGeneric;
-
-	/** Registers a delegate to be called when Persona receives a generic delete command */
-	void RegisterOnGenericDelete(const FOnDeleteGeneric& Delegate)
-	{
-		OnGenericDelete.Add(Delegate);
-	}
-
-	/** Unregisters a delegate to be called when Persona receives a generic delete command */
-	void UnregisterOnGenericDelete(SWidget* Widget)
-	{
-		OnGenericDelete.RemoveAll(Widget);
-	}
-
 	/** Broadcasts section changes */
 	FSimpleMulticastDelegate OnSectionsChanged;
 
@@ -591,31 +381,6 @@ public:
 		OnSectionsChanged.RemoveAll(Widget);
 	}
 
-	// Called when the notifies of the current animation are changed
-	typedef FSimpleMulticastDelegate::FDelegate FOnLODChanged;
-
-	/** Registers a delegate to be called when the skeleton anim notifies have been changed */
-	void RegisterOnLODChanged(const FOnLODChanged& Delegate)
-	{
-		OnLODChanged.Add(Delegate);
-	}
-
-	/** Unregisters a delegate to be called when the skeleton anim notifies have been changed */
-	void UnregisterOnLODChanged(SWidget* Widget)
-	{
-		OnLODChanged.RemoveAll(Widget);
-	}
-
-	/** Delegate for when the skeletons animation notifies have been changed */
-	FSimpleMulticastDelegate OnLODChanged;
-
-	/** Apply Compression to list of animations */
-	void ApplyCompression(TArray<TWeakObjectPtr<UAnimSequence>>& AnimSequences);
-	/** Export to FBX files of the list of animations */
-	void ExportToFBX(TArray<TWeakObjectPtr<UAnimSequence>>& AnimSequences);
-	/** Add looping interpolation to the list of animations */
-	void AddLoopingInterpolation(TArray<TWeakObjectPtr<UAnimSequence>>& AnimSequences);
-
 protected:
 	/** Undo Action**/
 	void UndoAction();
@@ -624,60 +389,20 @@ protected:
 
 protected:
 
-	/** Called when persona is refreshed through an external action (reimport etc) */
-	FSimpleMulticastDelegate OnPersonaRefresh;
-
-	/** Delegate called after an undo operation for child widgets to refresh */
-	FSimpleMulticastDelegate OnPostUndo;	
-
-	/**	Broadcasts whenever the animation changes */
-	FOnAnimChangedMulticaster OnAnimChanged;
-
-	/** Broadcasts whenever the preview mesh changes */
-	FOnPreviewMeshChangedMulticaster OnPreviewMeshChanged;
-	
-	/** Delegate for when a socket is selected by clicking its hit point */
-	FOnSelectSocket OnSocketSelected;
-
-	/** Delegate for when a bone is selected by clicking its hit point */
-	FOnSelectBone OnBoneSelected;
-
 	/** Delegate for when a blend profile is selected in the blend profile tab */
 	FOnSelectBlendProfile OnBlendProfileSelected;
 
-	/** Delegate for clearing the current skeleton bone/socket selection */
-	FSimpleMulticastDelegate OnAllDeselected;
-
-	/** Delegate for when the skeleton tree has changed (e.g. a socket has been duplicated in the viewport) */
-	FSimpleMulticastDelegate OnSkeletonTreeChanged;
-
 	/** Delegate for when the preview viewport is created */
 	FOnCreateViewport OnViewportCreated;
 
 	/** Delegate for changing / updating of track curves panel */
 	FSimpleMulticastDelegate OnTrackCurvesChanged;
 
-	/**
-	 * Delegate for handling generic delete command
-	 * Register to this if you want the generic delete command (delete key) from the editor. Also
-	 * remember that the user may be using another part of the editor when you get the broadcast so
-	 * make sure not to delete unless the user intends it.
-	 */
-	FSimpleMulticastDelegate OnGenericDelete;
-
 	/** Persona toolbar builder class */
 	TSharedPtr<class FPersonaToolbar> PersonaToolbar;
 
 private:
 
-	/** Recording animation functions **/
-	void RecordAnimation();
-	bool IsRecordAvailable() const;
-	FSlateIcon GetRecordStatusImage() const;
-	FText GetRecordStatusTooltip() const;
-	FText GetRecordStatusLabel() const;
-	FText GetRecordMenuLabel() const;
-
 	/** Animation menu functions **/
 	void OnApplyCompression();
 	void OnExportToFBX();
@@ -711,6 +436,8 @@ private:
 	TSharedRef< SWidget > GenerateCreateAssetMenu( USkeleton* Skeleton ) const;
 	void FillCreateAnimationMenu(FMenuBuilder& MenuBuilder) const;
 	void FillCreatePoseAssetMenu(FMenuBuilder& MenuBuilder) const;
+	void FillInsertPoseMenu(FMenuBuilder& MenuBuilder) const;
+	void InsertCurrentPoseToAsset(const FAssetData& NewPoseAssetData);
 	void CreateAnimation(const TArray<UObject*> NewAssets, int32 Option);
 	void CreatePoseAsset(const TArray<UObject*> NewAssets, int32 Option);
 
@@ -718,8 +445,6 @@ private:
 	void ExtendMenu();
 	/** update skeleton ref pose based on current preview mesh */
 	void UpdateSkeletonRefPose();
-	/** set preview mesh internal use only. The mesh should be verified by now.  */
-	void SetPreviewMeshInternal(USkeletalMesh* NewPreviewMesh);
 
 	/** Returns the editor objects that are applicable for our current mode (e.g mesh, animation etc) */
 	TArray<UObject*> GetEditorObjectsForMode(FName Mode) const;
@@ -727,15 +452,21 @@ private:
 	/** Called immediately prior to a blueprint compilation */
 	void OnBlueprintPreCompile(UBlueprint* BlueprintToCompile);
 
+	/** Post compile we copy values between preview & defaults */
+	void OnPostCompile();
+
+	/** Helper function used to keep nodes in preview & instance in sync */
+	struct FAnimNode_Base* FindAnimNode(class UAnimGraphNode_Base* AnimGraphNode) const;
+
+	/** Handle a pin's default value changing be propagating it to the preview */
+	void HandlePinDefaultValueChanged(UEdGraphPin* InPinThatChanged);
+
 	/** The extender to pass to the level editor to extend it's window menu */
 	TSharedPtr<FExtender> MenuExtender;
 
 	/** Toolbar extender */
 	TSharedPtr<FExtender> ToolbarExtender;
 
-	/** Preview scene for the editor */
-	FAdvancedPreviewScene PreviewScene;
-
 	/** Brush to use as a dirty marker for assets */
 	const FSlateBrush* AssetDirtyBrush;
 
@@ -743,17 +474,22 @@ private:
 	TSharedPtr<class SWidget> PreviewEditor;
 
 	/** Sequence Browser **/
-	TWeakPtr<class SAnimationSequenceBrowser> SequenceBrowser;
+	TWeakPtr<class IAnimationSequenceBrowser> SequenceBrowser;
 
 	/** Handle to the registered OnPropertyChangedHandle delegate */
 	FDelegateHandle OnPropertyChangedHandleDelegateHandle;
 
-	/** Last Cached LOD value of Preview Mesh Component */
-	int32 LastCachedLODForPreviewComponent;
+	/** Persona toolkit */
+	TSharedPtr<class IPersonaToolkit> PersonaToolkit;
 
-	/** Handle additive anim scale validation */
-	bool bDoesAdditiveRefPoseHaveZeroScale;
-	FGuid RefPoseGuid;
+	/** Skeleton tree */
+	TSharedPtr<class ISkeletonTree> SkeletonTree;
+
+	// selected skeletal control anim graph node 
+	TWeakObjectPtr<class UAnimGraphNode_Base> SelectedAnimGraphNode;
 
 	static const FName PreviewSceneSettingsTabId;
+
+	/** Delegate handle registered for when pin default values change */
+	FDelegateHandle OnPinDefaultValueChangedHandle;
 };
diff --git a/Engine/Source/Editor/Persona/Private/PersonaAssetFamily.cpp b/Engine/Source/Editor/Persona/Private/PersonaAssetFamily.cpp
new file mode 100644
--- /dev/null
+++ b/Engine/Source/Editor/Persona/Private/PersonaAssetFamily.cpp
@@ -0,0 +1,264 @@
+// Copyright 1998-2016 Epic Games, Inc. All Rights Reserved.
+
+#include "PersonaPrivatePCH.h"
+#include "PersonaAssetFamily.h"
+#include "AssetRegistryModule.h"
+
+#define LOCTEXT_NAMESPACE "PersonaAssetFamily"
+
+FPersonaAssetFamily::FPersonaAssetFamily(const UObject* InFromObject)
+	: Skeleton(nullptr)
+	, Mesh(nullptr)
+	, AnimBlueprint(nullptr)
+	, AnimationAsset(nullptr)
+{
+	if (InFromObject)
+	{
+		if (InFromObject->IsA<USkeleton>())
+		{
+			Skeleton = CastChecked<USkeleton>(InFromObject);
+			Mesh = Skeleton->GetPreviewMesh();
+		}
+		else if (InFromObject->IsA<UAnimationAsset>())
+		{
+			AnimationAsset = CastChecked<UAnimationAsset>(InFromObject);
+			Skeleton = AnimationAsset->GetSkeleton();
+			Mesh = AnimationAsset->GetPreviewMesh();
+			if (Mesh == nullptr)
+			{
+				Mesh = Skeleton->GetPreviewMesh();
+			}
+		}
+		else if (InFromObject->IsA<USkeletalMesh>())
+		{
+			Mesh = CastChecked<USkeletalMesh>(InFromObject);
+			Skeleton = Mesh->Skeleton;
+		}
+		else if (InFromObject->IsA<UAnimBlueprint>())
+		{
+			AnimBlueprint = CastChecked<UAnimBlueprint>(InFromObject);
+			Skeleton = AnimBlueprint->TargetSkeleton;
+			check(AnimBlueprint->TargetSkeleton);
+			Mesh = AnimBlueprint->TargetSkeleton->GetPreviewMesh();
+		}
+	}
+}
+
+void FPersonaAssetFamily::GetAssetTypes(TArray<UClass*>& OutAssetTypes) const
+{
+	OutAssetTypes.Reset();
+	OutAssetTypes.Add(USkeleton::StaticClass());
+	OutAssetTypes.Add(USkeletalMesh::StaticClass());
+	OutAssetTypes.Add(UAnimationAsset::StaticClass());
+	OutAssetTypes.Add(UAnimBlueprint::StaticClass());
+}
+
+template<typename AssetType>
+static void FindAssets(const USkeleton* InSkeleton, TArray<FAssetData>& OutAssetData, FName SkeletonTag)
+{
+	FAssetRegistryModule& AssetRegistryModule = FModuleManager::LoadModuleChecked<FAssetRegistryModule>("AssetRegistry");
+	FARFilter Filter;
+	Filter.bRecursiveClasses = true;
+	Filter.ClassNames.Add(AssetType::StaticClass()->GetFName());
+	Filter.TagsAndValues.Add(SkeletonTag, FAssetData(InSkeleton).GetExportTextName());
+
+	AssetRegistryModule.Get().GetAssets(Filter, OutAssetData);
+}
+
+FAssetData FPersonaAssetFamily::FindAssetOfType(UClass* InAssetClass) const
+{
+	if (InAssetClass)
+	{
+		if (InAssetClass->IsChildOf<USkeleton>())
+		{
+			// we should always have a skeleton here, this asset family is based on it
+			return FAssetData(Skeleton.Get());
+		}
+		else if (InAssetClass->IsChildOf<UAnimationAsset>())
+		{
+			if (AnimationAsset.IsValid())
+			{
+				return FAssetData(AnimationAsset.Get());
+			}
+			else
+			{
+				TArray<FAssetData> Assets;
+				FindAssets<UAnimationAsset>(Skeleton.Get(), Assets, "Skeleton");
+				if (Assets.Num() > 0)
+				{
+					return Assets[0];
+				}
+			}
+		}
+		else if (InAssetClass->IsChildOf<USkeletalMesh>())
+		{
+			if (Mesh.IsValid())
+			{
+				return FAssetData(Mesh.Get());
+			}
+			else
+			{
+				TArray<FAssetData> Assets;
+				FindAssets<USkeletalMesh>(Skeleton.Get(), Assets, "Skeleton");
+				if (Assets.Num() > 0)
+				{
+					return Assets[0];
+				}
+			}
+		}
+		else if (InAssetClass->IsChildOf<UAnimBlueprint>())
+		{
+			if (AnimBlueprint.IsValid())
+			{
+				return FAssetData(AnimBlueprint.Get());
+			}
+			else
+			{
+				TArray<FAssetData> Assets;
+				FindAssets<UAnimBlueprint>(Skeleton.Get(), Assets, "TargetSkeleton");
+				if (Assets.Num() > 0)
+				{
+					return Assets[0];
+				}
+			}
+		}
+	}
+
+	return FAssetData();
+}
+
+void FPersonaAssetFamily::FindAssetsOfType(UClass* InAssetClass, TArray<FAssetData>& OutAssets) const
+{
+	if (InAssetClass)
+	{
+		if (InAssetClass->IsChildOf<USkeleton>())
+		{
+			// we should always have a skeleton here, this asset family is based on it
+			OutAssets.Add(FAssetData(Skeleton.Get()));
+		}
+		else if (InAssetClass->IsChildOf<UAnimationAsset>())
+		{
+			FindAssets<UAnimationAsset>(Skeleton.Get(), OutAssets, "Skeleton");
+		}
+		else if (InAssetClass->IsChildOf<USkeletalMesh>())
+		{
+			FindAssets<USkeletalMesh>(Skeleton.Get(), OutAssets, "Skeleton");
+		}
+		else if (InAssetClass->IsChildOf<UAnimBlueprint>())
+		{
+			FindAssets<UAnimBlueprint>(Skeleton.Get(), OutAssets, "TargetSkeleton");
+		}
+	}
+}
+
+FText FPersonaAssetFamily::GetAssetTypeDisplayName(UClass* InAssetClass) const
+{
+	if (InAssetClass)
+	{
+		if (InAssetClass->IsChildOf<USkeleton>())
+		{
+			return LOCTEXT("SkeletonAssetDisplayName", "Skeleton");
+		}
+		else if (InAssetClass->IsChildOf<UAnimationAsset>())
+		{
+			return LOCTEXT("AnimationAssetDisplayName", "Animation");
+		}
+		else if (InAssetClass->IsChildOf<USkeletalMesh>())
+		{
+			return LOCTEXT("SkeletalMeshAssetDisplayName", "Mesh");
+		}
+		else if (InAssetClass->IsChildOf<UAnimBlueprint>())
+		{
+			return LOCTEXT("AnimBlueprintAssetDisplayName", "Blueprint");
+		}
+	}
+
+	return FText();
+}
+
+bool FPersonaAssetFamily::IsAssetCompatible(const FAssetData& InAssetData) const
+{
+	UClass* Class = InAssetData.GetClass();
+	if (Class)
+	{
+		if (Class->IsChildOf<USkeleton>())
+		{
+			return FAssetData(Skeleton.Get()) == InAssetData;
+		}
+		else if (Class->IsChildOf<UAnimationAsset>() || Class->IsChildOf<USkeletalMesh>())
+		{
+			const FString* TargetSkeleton = InAssetData.TagsAndValues.Find("Skeleton");
+			if (TargetSkeleton)
+			{
+				return *TargetSkeleton == FAssetData(Skeleton.Get()).GetExportTextName();
+			}
+		}
+		else if (Class->IsChildOf<UAnimBlueprint>())
+		{
+			const FString* TargetSkeleton = InAssetData.TagsAndValues.Find("TargetSkeleton");
+			if (TargetSkeleton)
+			{
+				return *TargetSkeleton == FAssetData(Skeleton.Get()).GetExportTextName();
+			}
+		}
+	}
+
+	return false;
+}
+
+UClass* FPersonaAssetFamily::GetAssetFamilyClass(UClass* InClass) const
+{
+	if (InClass)
+	{
+		if (InClass->IsChildOf<USkeleton>())
+		{
+			return USkeleton::StaticClass();
+		}
+		else if (InClass->IsChildOf<UAnimationAsset>())
+		{
+			return UAnimationAsset::StaticClass();
+		}
+		else if (InClass->IsChildOf<USkeletalMesh>())
+		{
+			return USkeletalMesh::StaticClass();
+		}
+		else if (InClass->IsChildOf<UAnimBlueprint>())
+		{
+			return UAnimBlueprint::StaticClass();
+		}
+	}
+
+	return nullptr;
+}
+
+void FPersonaAssetFamily::RecordAssetOpened(const FAssetData& InAssetData)
+{
+	if (IsAssetCompatible(InAssetData))
+	{
+		UClass* Class = InAssetData.GetClass();
+		if (Class)
+		{
+			if (Class->IsChildOf<USkeleton>())
+			{
+				Skeleton = Cast<USkeleton>(InAssetData.GetAsset());
+			}
+			else if (Class->IsChildOf<UAnimationAsset>())
+			{
+				AnimationAsset = Cast<UAnimationAsset>(InAssetData.GetAsset());
+			}
+			else if (Class->IsChildOf<USkeletalMesh>())
+			{
+				Mesh = Cast<USkeletalMesh>(InAssetData.GetAsset());
+			}
+			else if (Class->IsChildOf<UAnimBlueprint>())
+			{
+				AnimBlueprint = Cast<UAnimBlueprint>(InAssetData.GetAsset());
+			}
+		}
+
+		OnAssetOpened.Broadcast(InAssetData.GetAsset());
+	}
+}
+
+#undef LOCTEXT_NAMESPACE
+ 
\ No newline at end of file
diff --git a/Engine/Source/Editor/Persona/Private/PersonaAssetFamily.h b/Engine/Source/Editor/Persona/Private/PersonaAssetFamily.h
new file mode 100644
--- /dev/null
+++ b/Engine/Source/Editor/Persona/Private/PersonaAssetFamily.h
@@ -0,0 +1,38 @@
+// Copyright 1998-2016 Epic Games, Inc. All Rights Reserved.
+
+#pragma once
+
+#include "IAssetFamily.h"
+
+class FPersonaAssetFamily : public IAssetFamily
+{
+public:
+	FPersonaAssetFamily(const UObject* InFromObject);
+
+	/** IAssetFamily interface */
+	virtual void GetAssetTypes(TArray<UClass*>& OutAssetTypes) const override;
+	virtual FAssetData FindAssetOfType(UClass* InAssetClass) const override;
+	virtual void FindAssetsOfType(UClass* InAssetClass, TArray<FAssetData>& OutAssets) const override;
+	virtual FText GetAssetTypeDisplayName(UClass* InAssetClass) const override;
+	virtual bool IsAssetCompatible(const FAssetData& InAssetData) const override;
+	virtual UClass* GetAssetFamilyClass(UClass* InClass) const override;
+	virtual void RecordAssetOpened(const FAssetData& InAssetData) override;
+	DECLARE_DERIVED_EVENT(FPersonaAssetFamily, IAssetFamily::FOnAssetOpened, FOnAssetOpened)
+	virtual FOnAssetOpened& GetOnAssetOpened() override { return OnAssetOpened;  }
+
+private:
+	/** The skeleton that links all assets */
+	TWeakObjectPtr<const USkeleton> Skeleton;
+
+	/** The last mesh that was encountered */
+	TWeakObjectPtr<const USkeletalMesh> Mesh;
+
+	/** The last anim blueprint that was encountered */
+	TWeakObjectPtr<const UAnimBlueprint> AnimBlueprint;
+
+	/** The last animation asset that was encountered */
+	TWeakObjectPtr<const UAnimationAsset> AnimationAsset;
+
+	/** Event fired when an asset is opened */
+	FOnAssetOpened OnAssetOpened;
+};
\ No newline at end of file
diff --git a/Engine/Source/Editor/Persona/Private/PersonaAssetFamilyManager.cpp b/Engine/Source/Editor/Persona/Private/PersonaAssetFamilyManager.cpp
new file mode 100644
--- /dev/null
+++ b/Engine/Source/Editor/Persona/Private/PersonaAssetFamilyManager.cpp
@@ -0,0 +1,32 @@
+// Copyright 1998-2016 Epic Games, Inc. All Rights Reserved.
+
+#include "PersonaPrivatePCH.h"
+#include "PersonaAssetFamilyManager.h"
+#include "PersonaAssetFamily.h"
+
+FPersonaAssetFamilyManager& FPersonaAssetFamilyManager::Get()
+{
+	static FPersonaAssetFamilyManager TheManager;
+	return TheManager;
+}
+
+TSharedRef<IAssetFamily> FPersonaAssetFamilyManager::CreatePersonaAssetFamily(const UObject* InAsset)
+{
+	// compact any invalid entries
+	AssetFamilies.RemoveAll([](const TWeakPtr<class IAssetFamily>& InAssetFamily) { return !InAssetFamily.IsValid(); });
+
+	// look for an existing matching asset family
+	FAssetData AssetData(InAsset);
+	for (TWeakPtr<class IAssetFamily>& AssetFamily : AssetFamilies)
+	{
+		if (AssetFamily.Pin()->IsAssetCompatible(AssetData))
+		{
+			return AssetFamily.Pin().ToSharedRef();
+		}
+	}
+
+	// not found - make a new one
+	TSharedRef<IAssetFamily> NewAssetFamily = MakeShareable(new FPersonaAssetFamily(InAsset));
+	AssetFamilies.Add(NewAssetFamily);
+	return NewAssetFamily;
+}
\ No newline at end of file
diff --git a/Engine/Source/Editor/Persona/Private/PersonaAssetFamilyManager.h b/Engine/Source/Editor/Persona/Private/PersonaAssetFamilyManager.h
new file mode 100644
--- /dev/null
+++ b/Engine/Source/Editor/Persona/Private/PersonaAssetFamilyManager.h
@@ -0,0 +1,22 @@
+// Copyright 1998-2016 Epic Games, Inc. All Rights Reserved.
+
+#pragma once
+
+/** Central registry of persona asset families */
+class FPersonaAssetFamilyManager
+{
+public:
+	/** Singleton access */
+	static FPersonaAssetFamilyManager& Get();
+
+	/** Create a skeleton tree for the requested skeleton */
+	TSharedRef<class IAssetFamily> CreatePersonaAssetFamily(const UObject* InAsset);
+
+private:
+	/** Hidden constructor */
+	FPersonaAssetFamilyManager() {}
+
+private:
+	/** All current asset families */
+	TArray<TWeakPtr<class IAssetFamily>> AssetFamilies;
+};
diff --git a/Engine/Source/Editor/Persona/Private/PersonaCommands.cpp b/Engine/Source/Editor/Persona/Private/PersonaCommands.cpp
--- a/Engine/Source/Editor/Persona/Private/PersonaCommands.cpp
+++ b/Engine/Source/Editor/Persona/Private/PersonaCommands.cpp
@@ -29,7 +29,6 @@ void FPersonaCommands::RegisterCommands()
 	UI_COMMAND( ImportAnimation, "Import Animation", "Import new animation for the skeleton.", EUserInterfaceActionType::Button, FInputChord());
 	UI_COMMAND( ReimportAnimation, "Reimport Animation", "Reimport current animation.", EUserInterfaceActionType::Button, FInputChord());
 	UI_COMMAND( ApplyCompression, "Apply Compression", "Apply compression to current animation", EUserInterfaceActionType::Button, FInputChord() );
-	UI_COMMAND( RecordAnimation, "Record to new Animation", "Create new animation from currently playing", EUserInterfaceActionType::Button, FInputChord() );
 	UI_COMMAND( ExportToFBX, "Export to FBX", "Export current animation to FBX", EUserInterfaceActionType::Button, FInputChord() );
 	UI_COMMAND( AddLoopingInterpolation, "Add Looping Interpolation", "Add an extra first frame at the end of the animation to create interpolation when looping", EUserInterfaceActionType::Button, FInputChord() );
 	UI_COMMAND( SetKey, "Set Key", "Add Bone Transform to Additive Layer Tracks", EUserInterfaceActionType::Button, FInputChord(EKeys::S) );
diff --git a/Engine/Source/Editor/Persona/Private/PersonaCommands.h b/Engine/Source/Editor/Persona/Private/PersonaCommands.h
--- a/Engine/Source/Editor/Persona/Private/PersonaCommands.h
+++ b/Engine/Source/Editor/Persona/Private/PersonaCommands.h
@@ -46,8 +46,6 @@ public:
 	TSharedPtr<FUICommandInfo> ImportAnimation;
 	// reimport animation
 	TSharedPtr<FUICommandInfo> ReimportAnimation;
-	// record animation 
- 	TSharedPtr<FUICommandInfo> RecordAnimation;
 	// apply compression
 	TSharedPtr<FUICommandInfo> ApplyCompression;
 	// export to FBX
diff --git a/Engine/Source/Editor/Persona/Private/PersonaEditorModeManager.cpp b/Engine/Source/Editor/Persona/Private/PersonaEditorModeManager.cpp
new file mode 100644
--- /dev/null
+++ b/Engine/Source/Editor/Persona/Private/PersonaEditorModeManager.cpp
@@ -0,0 +1,33 @@
+// Copyright 1998-2016 Epic Games, Inc. All Rights Reserved.
+
+#include "PersonaPrivatePCH.h"
+#include "PersonaEditorModeManager.h"
+#include "IPersonaEditMode.h"
+
+bool FPersonaEditorModeManager::GetCameraTarget(FSphere& OutTarget) const
+{
+	// Note: assumes all of our modes are IPersonaEditMode!
+	for(int32 ModeIndex = 0; ModeIndex < Modes.Num(); ++ModeIndex)
+	{
+		TSharedPtr<IPersonaEditMode> EditMode = StaticCastSharedPtr<IPersonaEditMode>(Modes[ModeIndex]);
+
+		FSphere Target;
+		if (EditMode->GetCameraTarget(Target))
+		{
+			OutTarget = Target;
+			return true;
+		}
+	}
+
+	return false;
+}
+
+void FPersonaEditorModeManager::GetOnScreenDebugInfo(TArray<FText>& OutDebugText) const
+{
+	// Note: assumes all of our modes are IPersonaEditMode!
+	for (int32 ModeIndex = 0; ModeIndex < Modes.Num(); ++ModeIndex)
+	{
+		TSharedPtr<IPersonaEditMode> EditMode = StaticCastSharedPtr<IPersonaEditMode>(Modes[ModeIndex]);
+		EditMode->GetOnScreenDebugInfo(OutDebugText);
+	}
+}
\ No newline at end of file
diff --git a/Engine/Source/Editor/Persona/Private/PersonaMeshDetails.cpp b/Engine/Source/Editor/Persona/Private/PersonaMeshDetails.cpp
--- a/Engine/Source/Editor/Persona/Private/PersonaMeshDetails.cpp
+++ b/Engine/Source/Editor/Persona/Private/PersonaMeshDetails.cpp
@@ -14,7 +14,6 @@
 	#include "ApexClothingOptionWindow.h"
 #endif // #if WITH_APEX_CLOTHING
 
-#include "MeshMode/SAdditionalMeshesEditor.h"
 #include "LODUtilities.h"
 #include "Developer/MeshUtilities/Public/MeshUtilities.h"
 #include "FbxMeshUtils.h"
@@ -23,6 +22,8 @@
 #include "STextComboBox.h"
 
 #include "Engine/SkeletalMeshReductionSettings.h"
+#include "Animation/AnimInstance.h"
+#include "IPersonaToolkit.h"
 
 #define LOCTEXT_NAMESPACE "PersonaMeshDetails"
 
@@ -455,14 +456,14 @@ void FSkelMeshReductionSettingsLayout::OnSkinningImportanceChanged(TSharedPtr<FS
 /**
 * FPersonaMeshDetails
 */
-TSharedRef<IDetailCustomization> FPersonaMeshDetails::MakeInstance(TSharedPtr<FPersona> InPersona)
+TSharedRef<IDetailCustomization> FPersonaMeshDetails::MakeInstance(TSharedRef<class IPersonaToolkit> InPersonaToolkit)
 {
-	return MakeShareable( new FPersonaMeshDetails(InPersona) );
+	return MakeShareable( new FPersonaMeshDetails(InPersonaToolkit) );
 }
 
 void FPersonaMeshDetails::AddLODLevelCategories(IDetailLayoutBuilder& DetailLayout)
 {
-	USkeletalMesh* SkelMesh = PersonaPtr->GetMesh();
+	USkeletalMesh* SkelMesh = GetPersonaToolkit()->GetMesh();
 
 	if (SkelMesh)
 	{
@@ -649,7 +650,7 @@ void FPersonaMeshDetails::AddLODLevelCategories(IDetailLayoutBuilder& DetailLayo
 
 void FPersonaMeshDetails::CustomizeLODSettingsCategories(IDetailLayoutBuilder& DetailLayout)
 {
-	USkeletalMesh* SkelMesh = PersonaPtr->GetMesh();
+	USkeletalMesh* SkelMesh = GetPersonaToolkit()->GetMesh();
 	LODCount = SkelMesh->LODInfo.Num();
 
 	UpdateLODNames();
@@ -714,7 +715,7 @@ void FPersonaMeshDetails::OnImportLOD(TSharedPtr<FString> NewValue, ESelectInfo:
 	int32 LODIndex = 0;
 	if (LODNames.Find(NewValue, LODIndex) && LODIndex > 0)
 	{
-		USkeletalMesh* SkelMesh = PersonaPtr->GetMesh();
+		USkeletalMesh* SkelMesh = GetPersonaToolkit()->GetMesh();
 		check(SkelMesh);
 
 		FbxMeshUtils::ImportMeshLODDialog(SkelMesh, LODIndex);
@@ -748,30 +749,30 @@ FReply FPersonaMeshDetails::OnApplyChanges()
 
 FReply FPersonaMeshDetails::RemoveOneLOD(int32 LODIndex)
 {
-	USkeletalMesh* SkelMesh = PersonaPtr->GetMesh();
+	USkeletalMesh* SkelMesh = GetPersonaToolkit()->GetMesh();
 	check(SkelMesh);
 	check(SkelMesh->LODInfo.IsValidIndex(LODIndex));
 
 	FSkeletalMeshUpdateContext UpdateContext;
 	UpdateContext.SkeletalMesh = SkelMesh;
-	UpdateContext.AssociatedComponents.Push(PersonaPtr->PreviewComponent);
+	UpdateContext.AssociatedComponents.Push(GetPersonaToolkit()->GetPreviewMeshComponent());
 
 	FLODUtilities::RemoveLOD(UpdateContext, LODIndex);
 
-	PersonaPtr->PersonaMeshDetailLayout->ForceRefreshDetails();
+	MeshDetailLayout->ForceRefreshDetails();
 	return FReply::Handled();
 }
 
 FReply FPersonaMeshDetails::RemoveBones(int32 LODIndex)
 {
-	USkeletalMesh* SkelMesh = PersonaPtr->GetMesh();
+	USkeletalMesh* SkelMesh = GetPersonaToolkit()->GetMesh();
 	check(SkelMesh);
 	check(SkelMesh->LODInfo.IsValidIndex(LODIndex));
 
 	IMeshUtilities& MeshUtilities = FModuleManager::Get().LoadModuleChecked<IMeshUtilities>("MeshUtilities");
 	MeshUtilities.RemoveBonesFromMesh(SkelMesh, LODIndex, NULL);
 
-	PersonaPtr->PersonaMeshDetailLayout->ForceRefreshDetails();
+	MeshDetailLayout->ForceRefreshDetails();
 	return FReply::Handled();
 }
 
@@ -791,12 +792,12 @@ FText FPersonaMeshDetails::GetApplyButtonText() const
 
 void FPersonaMeshDetails::ApplyChanges(int32 DesiredLOD, const FSkeletalMeshOptimizationSettings& ReductionSettings)
 {
-	USkeletalMesh* SkelMesh = PersonaPtr->GetMesh();
+	USkeletalMesh* SkelMesh = GetPersonaToolkit()->GetMesh();
 	check(SkelMesh);
 
 	FSkeletalMeshUpdateContext UpdateContext;
 	UpdateContext.SkeletalMesh = SkelMesh;
-	UpdateContext.AssociatedComponents.Push(PersonaPtr->PreviewComponent);
+	UpdateContext.AssociatedComponents.Push(GetPersonaToolkit()->GetPreviewMeshComponent());
 
 	if (SkelMesh->LODInfo.IsValidIndex(DesiredLOD))
 	{
@@ -820,12 +821,12 @@ void FPersonaMeshDetails::ApplyChanges(int32 DesiredLOD, const FSkeletalMeshOpti
 
 void FPersonaMeshDetails::ApplyChanges()
 {
-	USkeletalMesh* SkelMesh = PersonaPtr->GetMesh();
+	USkeletalMesh* SkelMesh = GetPersonaToolkit()->GetMesh();
 	check(SkelMesh);
 
 	FSkeletalMeshUpdateContext UpdateContext;
 	UpdateContext.SkeletalMesh = SkelMesh;
-	UpdateContext.AssociatedComponents.Push(PersonaPtr->PreviewComponent);
+	UpdateContext.AssociatedComponents.Push(GetPersonaToolkit()->GetPreviewMeshComponent());
 
 	// remove LODs
 	int32 CurrentNumLODs = SkelMesh->LODInfo.Num();
@@ -929,7 +930,7 @@ void FPersonaMeshDetails::ApplyChanges()
 		}
 	}
 
-	PersonaPtr->PersonaMeshDetailLayout->ForceRefreshDetails();
+	MeshDetailLayout->ForceRefreshDetails();
 }
 
 void FPersonaMeshDetails::UpdateLODNames()
@@ -955,7 +956,7 @@ bool FPersonaMeshDetails::IsGenerateAvailable() const
 }
 bool FPersonaMeshDetails::IsApplyNeeded() const
 {
-	USkeletalMesh* SkelMesh = PersonaPtr->GetMesh();
+	USkeletalMesh* SkelMesh = GetPersonaToolkit()->GetMesh();
 	check(SkelMesh);
 
 	if (SkelMesh->LODInfo.Num() != LODCount)
@@ -990,7 +991,7 @@ FText FPersonaMeshDetails::GetLODCountTooltip() const
 
 FText FPersonaMeshDetails::GetLODImportedText(int32 LODIndex) const
 {
-	USkeletalMesh* Mesh = GetMesh();
+	USkeletalMesh* Mesh = GetPersonaToolkit()->GetMesh();
 	if (Mesh && Mesh->LODInfo.IsValidIndex(LODIndex))
 	{
 		if (Mesh->LODInfo[LODIndex].bHasBeenSimplified)
@@ -1004,7 +1005,7 @@ FText FPersonaMeshDetails::GetLODImportedText(int32 LODIndex) const
 
 FText FPersonaMeshDetails::GetMaterialSlotNameText(int32 MaterialIndex) const
 {
-	USkeletalMesh* Mesh = GetMesh();
+	USkeletalMesh* Mesh = GetPersonaToolkit()->GetMesh();
 	if (Mesh && Mesh->Materials.IsValidIndex(MaterialIndex))
 	{
 		return FText::FromName(Mesh->Materials[MaterialIndex].MaterialSlotName);
@@ -1021,7 +1022,7 @@ void FPersonaMeshDetails::CustomizeDetails( IDetailLayoutBuilder& DetailLayout )
 	SkeletalMeshPtr = SelectedObjects.Num() > 0 ? Cast<USkeletalMesh>(SelectedObjects[0].Get()) : nullptr;
 
 	// copy temporarily to refresh Mesh details tab from the LOD settings window
-	PersonaPtr->PersonaMeshDetailLayout = &DetailLayout;	
+	MeshDetailLayout = &DetailLayout;
 
 	// add multiple LOD levels to LOD category
 	AddLODLevelCategories(DetailLayout);
@@ -1033,10 +1034,26 @@ void FPersonaMeshDetails::CustomizeDetails( IDetailLayoutBuilder& DetailLayout )
 	CustomizeClothingProperties(DetailLayout,ClothingCategory);
 #endif// #if WITH_APEX_CLOTHING
 
-	IDetailCategoryBuilder& AdditionalMeshCategory = DetailLayout.EditCategory("AdditionalBodyPart", LOCTEXT("AdditionalMeshesCollapsable", "Additional Body Part"), ECategoryPriority::TypeSpecific);
-	AdditionalMeshCategory.AddCustomRow(FText::GetEmpty())
+	// Post process selector
+	IDetailCategoryBuilder& SkelMeshCategory = DetailLayout.EditCategory("SkeletalMesh");
+	TSharedRef<IPropertyHandle> PostProcessHandle = DetailLayout.GetProperty(GET_MEMBER_NAME_CHECKED(USkeletalMesh, PostProcessAnimBlueprint), USkeletalMesh::StaticClass());
+	PostProcessHandle->SetOnPropertyValueChanged(FSimpleDelegate::CreateSP(this, &FPersonaMeshDetails::OnPostProcessBlueprintChanged, &DetailLayout));
+	PostProcessHandle->MarkHiddenByCustomization();
+
+	FDetailWidgetRow& PostProcessRow = SkelMeshCategory.AddCustomRow(LOCTEXT("PostProcessFilterString", "Post Process Blueprint"));
+	PostProcessRow.NameContent()
+		[
+			PostProcessHandle->CreatePropertyNameWidget()
+		];
+
+	PostProcessRow.ValueContent()
 	[
-		SNew(SAdditionalMeshesEditor, PersonaPtr)
+			SNew(SObjectPropertyEntryBox)
+			.ObjectPath(this, &FPersonaMeshDetails::GetCurrentPostProcessBlueprintPath)
+			.AllowedClass(UAnimBlueprint::StaticClass())
+			.NewAssetFactories(TArray<UFactory*>())
+			.OnShouldFilterAsset(FOnShouldFilterAsset::CreateSP(this, &FPersonaMeshDetails::OnShouldFilterPostProcessBlueprint))
+			.OnObjectChanged(FOnSetObject::CreateSP(this, &FPersonaMeshDetails::OnSetPostProcessBlueprint, PostProcessHandle))
 	];
 
 	HideUnnecessaryProperties(DetailLayout);
@@ -1077,12 +1094,48 @@ void FPersonaMeshDetails::HideUnnecessaryProperties(IDetailLayoutBuilder& Detail
 	{
 		DetailLayout.HideProperty(MirroringProperties[MirrorPropertyIdx]);
 	}
+}
+
+void FPersonaMeshDetails::OnPostProcessBlueprintChanged(IDetailLayoutBuilder* DetailBuilder)
+{
+	DetailBuilder->ForceRefreshDetails();
+}
+
+FString FPersonaMeshDetails::GetCurrentPostProcessBlueprintPath() const
+{
+	USkeletalMesh* SkelMesh = GetPersonaToolkit()->GetMesh();
+	if(UClass* PostProcessClass = *SkelMesh->PostProcessAnimBlueprint)
+	{
+		return PostProcessClass->GetPathName();
+	}
+
+	return FString();
+}
+
+bool FPersonaMeshDetails::OnShouldFilterPostProcessBlueprint(const FAssetData& AssetData) const
+{
+	if(USkeletalMesh* SkelMesh = GetPersonaToolkit()->GetMesh())
+	{
+		const FString CurrentMeshSkeletonName = FString::Printf(TEXT("%s'%s'"), *SkelMesh->Skeleton->GetClass()->GetName(), *SkelMesh->Skeleton->GetPathName());
+		const FString SkeletonName = AssetData.GetTagValueRef<FString>("TargetSkeleton");
+
+		return SkeletonName != CurrentMeshSkeletonName;
+	}
 
+	return true;
+}
+
+void FPersonaMeshDetails::OnSetPostProcessBlueprint(const FAssetData& AssetData, TSharedRef<IPropertyHandle> BlueprintProperty)
+{
+	if(UAnimBlueprint* SelectedBlueprint = Cast<UAnimBlueprint>(AssetData.GetAsset()))
+	{
+		BlueprintProperty->SetValue(SelectedBlueprint->GetAnimBlueprintGeneratedClass());
+	}
 }
 
 void FPersonaMeshDetails::OnGetMaterialsForArray(class IMaterialListBuilder& OutMaterials, int32 LODIndex)
 {
-	USkeletalMesh* SkelMesh = PersonaPtr->GetMesh();
+	USkeletalMesh* SkelMesh = SkeletalMeshPtr.Get();
 
 	if (!SkelMesh)
 		return;
@@ -1161,7 +1214,7 @@ FText FPersonaMeshDetails::GetMaterialArrayText() const
 
 void FPersonaMeshDetails::OnGetSectionsForView(ISectionListBuilder& OutSections, int32 LODIndex)
 {
-	USkeletalMesh* SkelMesh = PersonaPtr->GetMesh();
+	USkeletalMesh* SkelMesh = GetPersonaToolkit()->GetMesh();
 
 	FSkeletalMeshResource* ImportedResource = SkelMesh->GetImportedResource();
 
@@ -1241,9 +1294,9 @@ FText FPersonaMeshDetails::GetOriginalImportMaterialNameText(int32 MaterialIndex
 		SkeletalMeshPtr->Materials[MaterialIndex].ImportedMaterialSlotName.ToString(OriginalImportMaterialName);
 		OriginalImportMaterialName = TEXT("Original Imported Material Name: ") + OriginalImportMaterialName;
 		return FText::FromString(OriginalImportMaterialName);
-	}
+			}
 	return FText::FromName(NAME_None);
-}
+		}
 
 void FPersonaMeshDetails::OnMaterialNameCommitted(const FText& InValue, ETextCommit::Type CommitType, int32 MaterialIndex)
 {
@@ -1275,7 +1328,7 @@ TSharedRef<SWidget> FPersonaMeshDetails::OnGenerateCustomNameWidgetsForMaterialA
 			.VAlign(VAlign_Center)
 			.Padding(0.0f, 3.0f, 0.0f, 0.0f)
 			.FillWidth(0.4f)
-			[
+		[
 				SNew(STextBlock)
 				.Font(IDetailLayoutBuilder::GetDetailFont())
 				.Text(LOCTEXT("MaterialArrayNameLabelStringKey", "Name"))
@@ -1304,12 +1357,12 @@ TSharedRef<SWidget> FPersonaMeshDetails::OnGenerateCustomNameWidgetsForMaterialA
 				.Padding(0.0f, 3.0f, 0.0f, 0.0f)
 				.VAlign(VAlign_Center)
 				.FillWidth(0.4f)
-				[
-					SNew(STextBlock)
-					.Font(IDetailLayoutBuilder::GetDetailFont())
+			[
+				SNew(STextBlock)
+				.Font(IDetailLayoutBuilder::GetDetailFont())
 					.Text(LOCTEXT("CustomNameMaterialNotUsed", "Unused"))
 					.ToolTipText(LOCTEXT("CustomNameMaterialNotUsedTooltip", "This material is not use by any LODs section! It will be set to None when the skeletal mesh will be saved."))
-				]
+			]
 				+ SHorizontalBox::Slot()
 				.Padding(5.0f, 3.0f, 5.0f, 0.0f)
 				.VAlign(VAlign_Center)
@@ -1320,15 +1373,15 @@ TSharedRef<SWidget> FPersonaMeshDetails::OnGenerateCustomNameWidgetsForMaterialA
 					.ToolTipText(LOCTEXT("CustomNameMaterialNotUsedDeleteTooltip", "Delete this material slot."))
 					.IsEnabled(this, &FPersonaMeshDetails::CanDeleteMaterialSlot, MaterialIndex)
 					.OnClicked(this, &FPersonaMeshDetails::OnDeleteMaterialSlot, MaterialIndex)
-				]
+		]
 			];
 	}
 	else
 	{
 		CustomNameWidget->AddSlot()
-			.AutoHeight()
-			.Padding(0, 2, 0, 0)
-			[
+		.AutoHeight()
+		.Padding(0, 2, 0, 0)
+		[
 				SNew(SHorizontalBox)
 				+ SHorizontalBox::Slot()
 				.VAlign(VAlign_Center)
@@ -1351,14 +1404,14 @@ TSharedRef<SWidget> FPersonaMeshDetails::OnGenerateCustomNameWidgetsForMaterialA
 					.VAlign(VAlign_Center)
 					.ContentPadding(2)
 					.ButtonContent()
-					[
-						SNew(STextBlock)
-						.Font(IDetailLayoutBuilder::GetDetailFont())
+			[
+				SNew(STextBlock)
+				.Font(IDetailLayoutBuilder::GetDetailFont())
 						.Text(this, &FPersonaMeshDetails::GetFirstMaterialSlotUsedBySection, MaterialIndex)
 					]
-				]
-			];
-	}
+			]
+		];
+}
 	return CustomNameWidget;
 }
 
@@ -1382,7 +1435,7 @@ TSharedRef<SWidget> FPersonaMeshDetails::OnGetMaterialSlotUsedByMenuContent(int3
 
 	TArray<FSectionLocalizer> *SectionLocalizers;
 	if (SkeletalMeshPtr.IsValid() && MaterialUsedMap.Contains(MaterialIndex))
-	{
+{
 		SectionLocalizers = MaterialUsedMap.Find(MaterialIndex);
 		FUIAction Action;
 		FText EmptyTooltip;
@@ -1399,12 +1452,12 @@ TSharedRef<SWidget> FPersonaMeshDetails::OnGetMaterialSlotUsedByMenuContent(int3
 }
 
 bool FPersonaMeshDetails::CanDeleteMaterialSlot(int32 MaterialIndex) const
-{
+	{
 	if (!SkeletalMeshPtr.IsValid())
 		return false;
 	return (MaterialIndex + 1) == SkeletalMeshPtr->Materials.Num();
-}
-
+	}
+	
 FReply FPersonaMeshDetails::OnDeleteMaterialSlot(int32 MaterialIndex)
 {
 	if (!SkeletalMeshPtr.IsValid() || !CanDeleteMaterialSlot(MaterialIndex))
@@ -1488,7 +1541,7 @@ bool FPersonaMeshDetails::OnMaterialListDirty()
 				}
 			}
 			if (ForceMaterialListRefresh)
-			{
+	{
 				break;
 			}
 		}
@@ -1501,10 +1554,10 @@ bool FPersonaMeshDetails::OnMaterialListDirty()
 TSharedRef<SWidget> FPersonaMeshDetails::OnGenerateCustomMaterialWidgetsForMaterialArray(UMaterialInterface* Material, int32 MaterialIndex, int32 LODIndex)
 {
 	TSharedRef<SWidget> MaterialWidget
-		= SNew(SVerticalBox)
+	= SNew(SVerticalBox)
 
-		+ SVerticalBox::Slot()
-		.Padding(0, 2, 0, 0)
+		+SVerticalBox::Slot()
+		.Padding(0,2,0,0)
 		[
 			SNew(SCheckBox)
 			.IsChecked(this, &FPersonaMeshDetails::IsShadowCastingEnabled, MaterialIndex)
@@ -1516,8 +1569,8 @@ TSharedRef<SWidget> FPersonaMeshDetails::OnGenerateCustomMaterialWidgetsForMater
 			]
 		]
 
-		+ SVerticalBox::Slot()
-		.Padding(0, 2, 0, 0)
+		+SVerticalBox::Slot()
+		.Padding(0,2,0,0)
 		[
 			SNew(SCheckBox)
 			.IsChecked(this, &FPersonaMeshDetails::IsRecomputeTangentEnabled, MaterialIndex)
@@ -1537,7 +1590,7 @@ TSharedRef<SWidget> FPersonaMeshDetails::OnGenerateCustomMaterialWidgetsForMater
 TSharedRef<SWidget> FPersonaMeshDetails::OnGenerateCustomNameWidgetsForSection(int32 LodIndex, int32 SectionIndex)
 {
 	return SNew(SVerticalBox)
-		+ SVerticalBox::Slot()
+		+SVerticalBox::Slot()
 		.AutoHeight()
 		[
 			SNew(SCheckBox)
@@ -1554,17 +1607,17 @@ TSharedRef<SWidget> FPersonaMeshDetails::OnGenerateCustomNameWidgetsForSection(i
 		+ SVerticalBox::Slot()
 		.AutoHeight()
 		.Padding(0, 2, 0, 0)
-		[
+			[
 			SNew(SCheckBox)
 			.IsChecked(this, &FPersonaMeshDetails::IsIsolateSectionEnabled, SectionIndex)
 			.OnCheckStateChanged(this, &FPersonaMeshDetails::OnSectionIsolatedChanged, SectionIndex)
 			.ToolTipText(LOCTEXT("Isolate_ToolTip", "Isolates this section in the viewport"))
-			[
+				[
 				SNew(STextBlock)
 				.Font(IDetailLayoutBuilder::GetDetailFont())
 				.ColorAndOpacity(FLinearColor(0.4f, 0.4f, 0.4f, 1.0f))
 				.Text(LOCTEXT("Isolate", "Isolate"))
-			]
+				]
 		];
 }
 
@@ -1676,7 +1729,7 @@ void FPersonaMeshDetails::OnSectionSelectedChanged(ECheckBoxState NewState, int3
 	USkeletalMesh* Mesh = SkeletalMeshPtr.Get();
 
 	// Currently assumes that we only ever have one preview mesh in Persona.
-	UDebugSkelMeshComponent* MeshComponent = PersonaPtr->GetPreviewMeshComponent();
+	UDebugSkelMeshComponent* MeshComponent = GetPersonaToolkit()->GetPreviewScene()->GetPreviewMeshComponent();
 
 	if (Mesh && MeshComponent)
 	{
@@ -1693,14 +1746,14 @@ void FPersonaMeshDetails::OnSectionSelectedChanged(ECheckBoxState NewState, int3
 		{
 			Mesh->SelectedEditorSection = INDEX_NONE;
 		}
-		PersonaPtr->RefreshViewport();
+		GetPersonaToolkit()->GetPreviewScene()->InvalidateViews();
 	}
 }
 
 ECheckBoxState FPersonaMeshDetails::IsIsolateSectionEnabled(int32 SectionIndex) const
 {
 	ECheckBoxState State = ECheckBoxState::Unchecked;
-	const UDebugSkelMeshComponent* MeshComponent = PersonaPtr->GetPreviewMeshComponent();
+	const UDebugSkelMeshComponent* MeshComponent = GetPersonaToolkit()->GetPreviewScene()->GetPreviewMeshComponent();
 	if (MeshComponent)
 	{
 		State = MeshComponent->SectionIndexPreview == SectionIndex ? ECheckBoxState::Checked : ECheckBoxState::Unchecked;
@@ -1711,7 +1764,7 @@ ECheckBoxState FPersonaMeshDetails::IsIsolateSectionEnabled(int32 SectionIndex)
 void FPersonaMeshDetails::OnSectionIsolatedChanged(ECheckBoxState NewState, int32 SectionIndex)
 {
 	USkeletalMesh* Mesh = SkeletalMeshPtr.Get();
-	UDebugSkelMeshComponent * MeshComponent = PersonaPtr->GetPreviewMeshComponent();
+	UDebugSkelMeshComponent * MeshComponent = GetPersonaToolkit()->GetPreviewScene()->GetPreviewMeshComponent();
 	if (Mesh && MeshComponent)
 	{
 		if (NewState == ECheckBoxState::Checked)
@@ -1726,7 +1779,7 @@ void FPersonaMeshDetails::OnSectionIsolatedChanged(ECheckBoxState NewState, int3
 		{
 			MeshComponent->SetSectionPreview(INDEX_NONE);
 		}
-		PersonaPtr->RefreshViewport();
+		GetPersonaToolkit()->GetPreviewScene()->InvalidateViews();
 	}
 }
 
@@ -1817,7 +1870,7 @@ void FPersonaMeshDetails::OnShadowCastingChanged(ECheckBoxState NewState, int32
 				MeshComponent->MarkRenderStateDirty();
 			}
 		}
-		PersonaPtr->RefreshViewport();
+		GetPersonaToolkit()->GetPreviewScene()->InvalidateViews();
 	}
 }
 
@@ -1896,7 +1949,7 @@ void FPersonaMeshDetails::OnRecomputeTangentChanged(ECheckBoxState NewState, int
 				MeshComponent->MarkRenderStateDirty();
 			}
 		}
-		PersonaPtr->RefreshViewport();
+		GetPersonaToolkit()->GetPreviewScene()->InvalidateViews();
 	}
 }
 
@@ -1947,7 +2000,7 @@ void FPersonaMeshDetails::OnSectionShadowCastingChanged(ECheckBoxState NewState,
 		const FScopedTransaction Transaction(LOCTEXT("SetSectionShadowCastingFlag", "Set Shadow Casting For Section"));
 		Mesh->Modify();
 		Section.bCastShadow = true;
-	}
+		}
 	else if (NewState == ECheckBoxState::Unchecked)
 	{
 		const FScopedTransaction Transaction(LOCTEXT("ClearSectionShadowCastingFlag", "Clear Shadow Casting For Section"));
@@ -1955,18 +2008,18 @@ void FPersonaMeshDetails::OnSectionShadowCastingChanged(ECheckBoxState NewState,
 		Section.bCastShadow = false;
 	}
 	
-	for (TObjectIterator<USkinnedMeshComponent> It; It; ++It)
-	{
-		USkinnedMeshComponent* MeshComponent = *It;
-		if (MeshComponent &&
-			!MeshComponent->IsTemplate() &&
-			MeshComponent->SkeletalMesh == Mesh)
+		for (TObjectIterator<USkinnedMeshComponent> It; It; ++It)
 		{
-			MeshComponent->MarkRenderStateDirty();
+			USkinnedMeshComponent* MeshComponent = *It;
+			if (MeshComponent &&
+				!MeshComponent->IsTemplate() &&
+				MeshComponent->SkeletalMesh == Mesh)
+			{
+				MeshComponent->MarkRenderStateDirty();
+			}
 		}
+		GetPersonaToolkit()->GetPreviewScene()->InvalidateViews();
 	}
-	PersonaPtr->RefreshViewport();
-}
 
 ECheckBoxState FPersonaMeshDetails::IsSectionRecomputeTangentEnabled(int32 LODIndex, int32 SectionIndex) const
 {
@@ -2035,12 +2088,12 @@ void FPersonaMeshDetails::OnSectionRecomputeTangentChanged(ECheckBoxState NewSta
 			MeshComponent->MarkRenderStateDirty();
 		}
 	}
-	PersonaPtr->RefreshViewport();
+	GetPersonaToolkit()->GetPreviewScene()->InvalidateViews();
 }
 
 int32 FPersonaMeshDetails::GetMaterialIndex(int32 LODIndex, int32 SectionIndex) const
 {
-	USkeletalMesh* SkelMesh = PersonaPtr->GetMesh();
+	USkeletalMesh* SkelMesh = GetPersonaToolkit()->GetMesh();
 
 	check(LODIndex < SkelMesh->LODInfo.Num());
 
@@ -2059,7 +2112,7 @@ int32 FPersonaMeshDetails::GetMaterialIndex(int32 LODIndex, int32 SectionIndex)
 
 bool FPersonaMeshDetails::IsDuplicatedMaterialIndex(int32 LODIndex, int32 MaterialIndex)
 {
-	USkeletalMesh* SkelMesh = PersonaPtr->GetMesh();
+	USkeletalMesh* SkelMesh = GetPersonaToolkit()->GetMesh();
 
 	// finding whether this material index is being used in parent LODs
 	for (int32 LODInfoIdx = 0; LODInfoIdx < LODIndex; LODInfoIdx++)
@@ -2100,7 +2153,7 @@ bool FPersonaMeshDetails::IsDuplicatedMaterialIndex(int32 LODIndex, int32 Materi
 void FPersonaMeshDetails::OnSectionChanged(int32 LODIndex, int32 SectionIndex, int32 NewMaterialSlotIndex, FName NewMaterialSlotName)
 {
 	USkeletalMesh* Mesh = SkeletalMeshPtr.Get();
-	if (Mesh)
+	if(Mesh)
 	{
 		FSkeletalMeshResource* ImportedResource = Mesh->GetImportedResource();
 		check(ImportedResource && ImportedResource->LODModels.IsValidIndex(LODIndex));
@@ -2116,7 +2169,7 @@ void FPersonaMeshDetails::OnSectionChanged(int32 LODIndex, int32 SectionIndex, i
 				NewSkeletalMaterialIndex = SkeletalMaterialIndex;
 				break;
 			}
-		}
+			}
 
 		check(NewSkeletalMaterialIndex != INDEX_NONE);
 
@@ -2127,21 +2180,21 @@ void FPersonaMeshDetails::OnSectionChanged(int32 LODIndex, int32 SectionIndex, i
 
 		FSkeletalMeshLODInfo& Info = Mesh->LODInfo[LODIndex];
 		if (LODIndex == 0 || Info.LODMaterialMap.Num() == 0)
-		{
+				{
 			ImportedResource->LODModels[LODIndex].Sections[SectionIndex].MaterialIndex = NewSkeletalMaterialIndex;
-		}
-		else
-		{
+				}
+				else
+				{
 			check(SectionIndex < Info.LODMaterialMap.Num());
 			Info.LODMaterialMap[SectionIndex] = NewSkeletalMaterialIndex;
-		}
+				}
 		Mesh->PostEditChange();
-		// End the transation if we created one
-		GEditor->EndTransaction();
-		// Redraw viewports to reflect the material changes 
-		GUnrealEd->RedrawLevelEditingViewports();
+			// End the transation if we created one
+			GEditor->EndTransaction();
+			// Redraw viewports to reflect the material changes 
+			GUnrealEd->RedrawLevelEditingViewports();
+		}
 	}
-}
 
 #if WITH_APEX_CLOTHING
 
@@ -2259,7 +2312,7 @@ void FPersonaMeshDetails::OnGenerateElementForClothingAsset( TSharedRef<IPropert
 		MakeApexDetailsWidget(ElementIndex)
 	];	
 	
-	USkeletalMesh* SkelMesh = PersonaPtr->GetMesh();
+	USkeletalMesh* SkelMesh = GetPersonaToolkit()->GetMesh();
 	FClothingAssetData& AssetData = SkelMesh->ClothingAssets[ElementIndex];
 
 	ApexClothingUtils::GetPhysicsPropertiesFromApexAsset(AssetData.ApexClothingAsset, AssetData.PhysicsProperties);
@@ -2281,7 +2334,7 @@ TSharedRef<SUniformGridPanel> FPersonaMeshDetails::MakeApexDetailsWidget(int32 A
 {
 	const FSlateFontInfo DetailFontInfo = IDetailLayoutBuilder::GetDetailFont();
 
-	USkeletalMesh* SkelMesh = PersonaPtr->GetMesh();
+	USkeletalMesh* SkelMesh = GetPersonaToolkit()->GetMesh();
 
 	FClothingAssetData& Asset = SkelMesh->ClothingAssets[AssetIndex];
 
@@ -2471,7 +2524,7 @@ TSharedRef<SUniformGridPanel> FPersonaMeshDetails::MakeApexDetailsWidget(int32 A
 
 FReply FPersonaMeshDetails::OnReimportApexFileClicked(int32 AssetIndex, IDetailLayoutBuilder* DetailLayout)
 {
-	USkeletalMesh* SkelMesh = PersonaPtr->GetMesh();
+	USkeletalMesh* SkelMesh = GetPersonaToolkit()->GetMesh();
 
 	check(SkelMesh);
 
@@ -2528,7 +2581,7 @@ FReply FPersonaMeshDetails::OnReimportApexFileClicked(int32 AssetIndex, IDetailL
 
 FReply FPersonaMeshDetails::OnRemoveApexFileClicked(int32 AssetIndex, IDetailLayoutBuilder* DetailLayout)
 {
-	USkeletalMesh* SkelMesh = PersonaPtr->GetMesh();
+	USkeletalMesh* SkelMesh = GetPersonaToolkit()->GetMesh();
 
 	// SkeMesh->ClothingAssets.RemoveAt(AssetIndex) is conducted inside RemoveAssetFromSkeletalMesh
 	ApexClothingUtils::RemoveAssetFromSkeletalMesh(SkelMesh, AssetIndex, true);
@@ -2544,7 +2597,7 @@ FReply FPersonaMeshDetails::OnRemoveApexFileClicked(int32 AssetIndex, IDetailLay
 
 FReply FPersonaMeshDetails::OnOpenClothingFileClicked(IDetailLayoutBuilder* DetailLayout)
 {
-	USkeletalMesh* SkelMesh = PersonaPtr->GetMesh();
+	USkeletalMesh* SkelMesh = GetPersonaToolkit()->GetMesh();
 
 	if( ensure(SkelMesh) )
 	{
@@ -2663,7 +2716,7 @@ FReply FPersonaMeshDetails::OnOpenClothingFileClicked(IDetailLayoutBuilder* Deta
 
 void FPersonaMeshDetails::UpdateComboBoxStrings()
 {
-	USkeletalMesh* SkelMesh = PersonaPtr->GetMesh();
+	USkeletalMesh* SkelMesh = GetPersonaToolkit()->GetMesh();
 
 	if (!SkelMesh)
 	{
@@ -2862,7 +2915,7 @@ void FPersonaMeshDetails::UpdateComboBoxStrings()
 EVisibility FPersonaMeshDetails::IsClothingComboBoxVisible(int32 LODIndex, int32 SectionIndex) const
 {
 	// if this material doesn't have correspondence with any mesh section, hide this combo box 
-	USkeletalMesh* SkelMesh = PersonaPtr->GetMesh();
+	USkeletalMesh* SkelMesh = GetPersonaToolkit()->GetMesh();
 
 	// ClothingComboStrings.Num() should be greater than 1 because it is including "None"
 	bool bExistClothingLOD = ClothingComboLODInfos.IsValidIndex(LODIndex) && ClothingComboLODInfos[LODIndex].ClothingComboStrings.Num() > 1;
@@ -2883,7 +2936,7 @@ void FPersonaMeshDetails::HandleSectionsComboBoxSelectionChanged( TSharedPtr<FSt
 		return;
 	}
 
-	USkeletalMesh* SkelMesh = PersonaPtr->GetMesh();
+	USkeletalMesh* SkelMesh = GetPersonaToolkit()->GetMesh();
 
 	FClothAssetSubmeshIndex* ComboItem;
 
@@ -2911,9 +2964,9 @@ void FPersonaMeshDetails::HandleSectionsComboBoxSelectionChanged( TSharedPtr<FSt
 		else
 		{
 			// when clothing is mapped, need to re-prepare cloth morph targets
-			if (PersonaPtr->PreviewComponent)
+			if (GetPersonaToolkit()->GetPreviewMeshComponent())
 			{
-				PersonaPtr->PreviewComponent->bPreparedClothMorphTargets = false;
+				GetPersonaToolkit()->GetPreviewMeshComponent()->bPreparedClothMorphTargets = false;
 			}
 		}
 	}
@@ -2925,7 +2978,7 @@ void FPersonaMeshDetails::HandleSectionsComboBoxSelectionChanged( TSharedPtr<FSt
 
 void FPersonaMeshDetails::UpdateClothPhysicsProperties(int32 AssetIndex)
 {
-	USkeletalMesh* SkelMesh = PersonaPtr->GetMesh();
+	USkeletalMesh* SkelMesh = GetPersonaToolkit()->GetMesh();
 	FClothingAssetData& Asset = SkelMesh->ClothingAssets[AssetIndex];
 
 	ApexClothingUtils::SetPhysicsPropertiesToApexAsset(Asset.ApexClothingAsset, Asset.PhysicsProperties);
@@ -2946,7 +2999,7 @@ FReply FPersonaMeshDetails::OnDeleteButtonClicked(int32 LODIndex, int32 SectionI
 
 	int32 MaterialIndex = GetMaterialIndex(LODIndex, SectionIndex);
 
-	USkeletalMesh* SkelMesh = PersonaPtr->GetMesh();
+	USkeletalMesh* SkelMesh = GetPersonaToolkit()->GetMesh();
 
 	// Move any mappings pointing to the requested material to point to the first
 	// and decrement any above it
diff --git a/Engine/Source/Editor/Persona/Private/PersonaMeshDetails.h b/Engine/Source/Editor/Persona/Private/PersonaMeshDetails.h
--- a/Engine/Source/Editor/Persona/Private/PersonaMeshDetails.h
+++ b/Engine/Source/Editor/Persona/Private/PersonaMeshDetails.h
@@ -119,10 +119,10 @@ private:
 class FPersonaMeshDetails : public IDetailCustomization
 {
 public:
-	FPersonaMeshDetails(TSharedPtr<FPersona> InPersona) : PersonaPtr(InPersona) {}
+	FPersonaMeshDetails(TSharedRef<class IPersonaToolkit> InPersonaToolkit) : PersonaToolkitPtr(InPersonaToolkit) {}
 
 	/** Makes a new instance of this detail layout class for a specific detail view requesting it */
-	static TSharedRef<IDetailCustomization> MakeInstance(TSharedPtr<FPersona> InPersona);
+	static TSharedRef<IDetailCustomization> MakeInstance(TSharedRef<class IPersonaToolkit> InPersonaToolkit);
 
 	/** IDetailCustomization interface */
 	virtual void CustomizeDetails( IDetailLayoutBuilder& DetailLayout ) override;
@@ -349,6 +349,15 @@ private:
 	/** hide properties which don't need to be showed to end users */
 	void HideUnnecessaryProperties(IDetailLayoutBuilder& DetailLayout);
 
+	// Handling functions for post process blueprint selection combo box
+	void OnPostProcessBlueprintChanged(IDetailLayoutBuilder* DetailBuilder);
+	FString GetCurrentPostProcessBlueprintPath() const;
+	bool OnShouldFilterPostProcessBlueprint(const FAssetData& AssetData) const;
+	void OnSetPostProcessBlueprint(const FAssetData& AssetData, TSharedRef<IPropertyHandle> BlueprintProperty);
+
+	/** Access the persona toolkit ptr. It should always be valid in the lifetime of this customization */
+	TSharedRef<IPersonaToolkit> GetPersonaToolkit() const { check(PersonaToolkitPtr.IsValid()); return PersonaToolkitPtr.Pin().ToSharedRef(); }
+
 public:
 
 	bool IsApplyNeeded() const;
@@ -357,23 +366,12 @@ public:
 	void ApplyChanges();
 	FText GetApplyButtonText() const;
 
-	USkeletalMesh* GetMesh() const
-	{ 
-		if (PersonaPtr.IsValid())
-		{
-			return PersonaPtr->GetMesh();
-		}
-		else
-		{
-			return NULL;
-		}
-	}
 private:
 	// Container for the objects to display
 	TWeakObjectPtr<USkeletalMesh> SkeletalMeshPtr;
 
-	// Pointer back to Persona
-	TSharedPtr<FPersona> PersonaPtr;
+	// Reference the persona toolkit
+	TWeakPtr<class IPersonaToolkit> PersonaToolkitPtr;
 
 	IDetailLayoutBuilder* MeshDetailLayout;
 
diff --git a/Engine/Source/Editor/Persona/Private/PersonaModule.cpp b/Engine/Source/Editor/Persona/Private/PersonaModule.cpp
--- a/Engine/Source/Editor/Persona/Private/PersonaModule.cpp
+++ b/Engine/Source/Editor/Persona/Private/PersonaModule.cpp
@@ -15,11 +15,45 @@
 #include "AnimInstanceDetails.h"
 #include "Editor/UnrealEd/Public/Kismet2/KismetEditorUtilities.h"
 #include "Animation/AnimInstance.h"
+#include "PersonaToolkit.h"
+#include "Shared/PersonaMode.h"
+#include "SSkeletonAnimNotifies.h"
+#include "PersonaAssetFamily.h"
+#include "SAssetFamilyShortcutBar.h"
+#include "SMontageEditor.h"
+#include "SSequenceEditor.h"
+#include "SAnimCompositeEditor.h"
+#include "SPoseEditor.h"
+#include "SAnimationBlendSpace.h"
+#include "SAnimationBlendSpace1D.h"
+#include "Animation/AimOffsetBlendSpace.h"
+#include "Animation/AimOffsetBlendSpace1D.h"
+#include "Animation/BlendSpace1D.h"
+#include "SAnimationDlgs.h"
+#include "FbxMeshUtils.h"
+#include "AssetToolsModule.h"
+#include "IEditableSkeleton.h"
+#include "MessageLog.h"
+#include "AnimationCompressionPanel.h"
+#include "DesktopPlatformModule.h"
+#include "FbxAnimUtils.h"
+#include "AnimationMode/AnimationMode.h"
+#include "PersonaAssetFamilyManager.h"
+#include "AnimGraphNode_Slot.h"
+#include "Customization/AnimGraphNodeSlotDetails.h"
+#include "EditModes/SkeletonSelectionEditMode.h"
+#include "PersonaEditorModeManager.h"
+#include "PreviewSceneCustomizations.h"
+#include "SSkeletonSlotNames.h"
+#include "PersonaMeshDetails.h"
 
 IMPLEMENT_MODULE( FPersonaModule, Persona );
 
 const FName PersonaAppName = FName(TEXT("PersonaApp"));
 
+const FEditorModeID FPersonaEditModes::SkeletonSelection(TEXT("PersonaSkeletonSelection"));
+
+#define LOCTEXT_NAMESPACE "PersonaModule"
 
 void FPersonaModule::StartupModule()
 {
@@ -59,14 +93,17 @@ void FPersonaModule::StartupModule()
 
 		PropertyModule.RegisterCustomPropertyTypeLayout( "InputScaleBias", FOnGetPropertyTypeCustomizationInstance::CreateStatic( &FInputScaleBiasCustomization::MakeInstance ) );
 		PropertyModule.RegisterCustomPropertyTypeLayout( "BoneReference", FOnGetPropertyTypeCustomizationInstance::CreateStatic( &FBoneReferenceCustomization::MakeInstance ) );
+		PropertyModule.RegisterCustomPropertyTypeLayout( "PreviewMeshCollectionEntry", FOnGetPropertyTypeCustomizationInstance::CreateStatic( &FPreviewMeshCollectionEntryCustomization::MakeInstance ) );
 	}
 
-	FKismetEditorUtilities::RegisterOnBlueprintCreatedCallback(this, UAnimInstance::StaticClass(), FKismetEditorUtilities::FOnBlueprintCreated::CreateRaw(this, &FPersonaModule::OnNewBlueprintCreated));
+	// Register the editor modes
+	FEditorModeRegistry::Get().RegisterMode<FSkeletonSelectionEditMode>(FPersonaEditModes::SkeletonSelection, LOCTEXT("SkeletonSelectionEditMode", "Skeleton Selection"), FSlateIcon(), false);
 }
 
 void FPersonaModule::ShutdownModule()
 {
-	FKismetEditorUtilities::UnregisterAutoBlueprintNodeCreation(this);
+	// Unregister the editor modes
+	FEditorModeRegistry::Get().UnregisterMode(FPersonaEditModes::SkeletonSelection);
 
 	MenuExtensibilityManager.Reset();
 	ToolBarExtensibilityManager.Reset();
@@ -92,34 +129,530 @@ TSharedRef<IBlueprintEditor> FPersonaModule::CreatePersona( const EToolkitMode::
 	return NewPersona;
 }
 
-void FPersonaModule::OnNewBlueprintCreated(UBlueprint* InBlueprint)
+static void SetupPersonaToolkit(const TSharedRef<FPersonaToolkit>& Toolkit, const FPersonaToolkitArgs& PersonaToolkitArgs)
 {
-	if (ensure(InBlueprint->UbergraphPages.Num() > 0))
+	if (PersonaToolkitArgs.bCreatePreviewScene)
 	{
-		UEdGraph* EventGraph = InBlueprint->UbergraphPages[0];
+		Toolkit->CreatePreviewScene();
+	}
+}
+
+TSharedRef<IPersonaToolkit> FPersonaModule::CreatePersonaToolkit(USkeleton* InSkeleton, const FPersonaToolkitArgs& PersonaToolkitArgs) const
+{
+	TSharedRef<FPersonaToolkit> NewPersonaToolkit(new FPersonaToolkit());
+
+	NewPersonaToolkit->Initialize(InSkeleton);
+
+	SetupPersonaToolkit(NewPersonaToolkit, PersonaToolkitArgs);
+
+	return NewPersonaToolkit;
+}
+
+TSharedRef<IPersonaToolkit> FPersonaModule::CreatePersonaToolkit(UAnimationAsset* InAnimationAsset, const FPersonaToolkitArgs& PersonaToolkitArgs) const
+{
+	TSharedRef<FPersonaToolkit> NewPersonaToolkit(new FPersonaToolkit());
+
+	NewPersonaToolkit->Initialize(InAnimationAsset);
+
+	SetupPersonaToolkit(NewPersonaToolkit, PersonaToolkitArgs);
+
+	return NewPersonaToolkit;
+}
+
+TSharedRef<IPersonaToolkit> FPersonaModule::CreatePersonaToolkit(USkeletalMesh* InSkeletalMesh, const FPersonaToolkitArgs& PersonaToolkitArgs) const
+{
+	TSharedRef<FPersonaToolkit> NewPersonaToolkit(new FPersonaToolkit());
+
+	NewPersonaToolkit->Initialize(InSkeletalMesh);
+
+	SetupPersonaToolkit(NewPersonaToolkit, PersonaToolkitArgs);
+
+	return NewPersonaToolkit;
+}
+
+TSharedRef<IPersonaToolkit> FPersonaModule::CreatePersonaToolkit(UAnimBlueprint* InAnimBlueprint, const FPersonaToolkitArgs& PersonaToolkitArgs) const
+{
+	TSharedRef<FPersonaToolkit> NewPersonaToolkit(new FPersonaToolkit());
+
+	NewPersonaToolkit->Initialize(InAnimBlueprint);
+
+	SetupPersonaToolkit(NewPersonaToolkit, PersonaToolkitArgs);
+
+	return NewPersonaToolkit;
+}
 
-		int32 SafeXPosition = 0;
-		int32 SafeYPosition = 0;
+TSharedRef<class IAssetFamily> FPersonaModule::CreatePersonaAssetFamily(const UObject* InAsset) const
+{
+	return FPersonaAssetFamilyManager::Get().CreatePersonaAssetFamily(InAsset);
+}
+
+TSharedRef<SWidget> FPersonaModule::CreateAssetFamilyShortcutWidget(const TSharedRef<class FWorkflowCentricApplication>& InHostingApp, const TSharedRef<class IAssetFamily>& InAssetFamily) const
+{
+	return SNew(SAssetFamilyShortcutBar, InHostingApp, InAssetFamily);
+}
+
+TSharedRef<class FWorkflowTabFactory> FPersonaModule::CreateDetailsTabFactory(const TSharedRef<class FWorkflowCentricApplication>& InHostingApp, FOnDetailsCreated InOnDetailsCreated) const
+{
+	return MakeShareable(new FDetailsTabSummoner(InHostingApp, InOnDetailsCreated));
+}
+
+TSharedRef<class FWorkflowTabFactory> FPersonaModule::CreatePersonaViewportTabFactory(const TSharedRef<class FWorkflowCentricApplication>& InHostingApp, const TSharedRef<ISkeletonTree>& InSkeletonTree, const TSharedRef<IPersonaPreviewScene>& InPreviewScene, FSimpleMulticastDelegate& OnPostUndo, const TSharedPtr<FBlueprintEditor>& InBlueprintEditor, FOnViewportCreated InOnViewportCreated, bool bInShowTimeline, bool bInShowStats) const
+{
+	return MakeShareable(new FPreviewViewportSummoner(InHostingApp, InSkeletonTree, InPreviewScene, OnPostUndo, InBlueprintEditor, InOnViewportCreated, bInShowTimeline, bInShowStats));
+}
+
+TSharedRef<class FWorkflowTabFactory> FPersonaModule::CreateAnimNotifiesTabFactory(const TSharedRef<class FWorkflowCentricApplication>& InHostingApp, const TSharedRef<class IEditableSkeleton>& InEditableSkeleton, FSimpleMulticastDelegate& InOnChangeAnimNotifies, FSimpleMulticastDelegate& InOnPostUndo, FOnObjectsSelected InOnObjectsSelected) const
+{
+	return MakeShareable(new FSkeletonAnimNotifiesSummoner(InHostingApp, InEditableSkeleton, InOnChangeAnimNotifies, InOnPostUndo, InOnObjectsSelected));
+}
+
+TSharedRef<class FWorkflowTabFactory> FPersonaModule::CreateCurveViewerTabFactory(const TSharedRef<class FWorkflowCentricApplication>& InHostingApp, const TSharedRef<class IEditableSkeleton>& InEditableSkeleton, const TSharedRef<IPersonaPreviewScene>& InPreviewScene, FSimpleMulticastDelegate& InOnCurvesChanged, FSimpleMulticastDelegate& InOnPostUndo) const
+{
+	return MakeShareable(new FAnimCurveViewerTabSummoner(InHostingApp, InEditableSkeleton, InPreviewScene, InOnCurvesChanged, InOnPostUndo));
+}
+
+TSharedRef<class FWorkflowTabFactory> FPersonaModule::CreateRetargetManagerTabFactory(const TSharedRef<class FWorkflowCentricApplication>& InHostingApp, const TSharedRef<class IEditableSkeleton>& InEditableSkeleton, const TSharedRef<IPersonaPreviewScene>& InPreviewScene, FSimpleMulticastDelegate& InOnPostUndo) const
+{
+	return MakeShareable(new FRetargetManagerTabSummoner(InHostingApp, InEditableSkeleton, InPreviewScene, InOnPostUndo));
+}
+
+TSharedRef<class FWorkflowTabFactory> FPersonaModule::CreateAdvancedPreviewSceneTabFactory(const TSharedRef<class FWorkflowCentricApplication>& InHostingApp, const TSharedRef<IPersonaPreviewScene>& InPreviewScene) const
+{
+	return MakeShareable(new FAdvancedPreviewSceneTabSummoner(InHostingApp, InPreviewScene));
+}
+
+TSharedRef<class FWorkflowTabFactory> FPersonaModule::CreateAnimationAssetBrowserTabFactory(const TSharedRef<class FWorkflowCentricApplication>& InHostingApp, const TSharedRef<IPersonaToolkit>& InPersonaToolkit, FOnOpenNewAsset InOnOpenNewAsset, FOnAnimationSequenceBrowserCreated InOnAnimationSequenceBrowserCreated, bool bInShowHistory) const
+{
+	return MakeShareable(new FAnimationAssetBrowserSummoner(InHostingApp, InPersonaToolkit, InOnOpenNewAsset, InOnAnimationSequenceBrowserCreated, bInShowHistory));
+}
 
-		if (EventGraph->Nodes.Num() != 0)
+TSharedRef<class FWorkflowTabFactory> FPersonaModule::CreateAssetDetailsTabFactory(const TSharedRef<class FWorkflowCentricApplication>& InHostingApp, FOnGetAsset InOnGetAsset, FOnDetailsCreated InOnDetailsCreated) const
+{
+	return MakeShareable(new FAssetPropertiesSummoner(InHostingApp, InOnGetAsset, InOnDetailsCreated));
+}
+
+TSharedRef<class FWorkflowTabFactory> FPersonaModule::CreateMorphTargetTabFactory(const TSharedRef<class FWorkflowCentricApplication>& InHostingApp, const TSharedRef<IPersonaPreviewScene>& InPreviewScene, FSimpleMulticastDelegate& OnPostUndo) const
+{
+	return MakeShareable(new FMorphTargetTabSummoner(InHostingApp, InPreviewScene, OnPostUndo));
+}
+
+TSharedRef<class FWorkflowTabFactory> FPersonaModule::CreateAnimBlueprintPreviewTabFactory(const TSharedRef<class FBlueprintEditor>& InBlueprintEditor, const TSharedRef<IPersonaPreviewScene>& InPreviewScene) const
+{
+	return MakeShareable(new FAnimBlueprintPreviewEditorSummoner(InBlueprintEditor, InPreviewScene));
+}
+
+TSharedRef<class FWorkflowTabFactory> FPersonaModule::CreateAnimBlueprintAssetOverridesTabFactory(const TSharedRef<class FBlueprintEditor>& InBlueprintEditor, UAnimBlueprint* InAnimBlueprint, FSimpleMulticastDelegate& InOnPostUndo) const
+{
+	return MakeShareable(new FAnimBlueprintParentPlayerEditorSummoner(InBlueprintEditor, InOnPostUndo));
+}
+
+TSharedRef<class FWorkflowTabFactory> FPersonaModule::CreateSkeletonSlotNamesTabFactory(const TSharedRef<class FWorkflowCentricApplication>& InHostingApp, const TSharedRef<class IEditableSkeleton>& InEditableSkeleton, FSimpleMulticastDelegate& InOnPostUndo, FOnObjectSelected InOnObjectSelected) const
+{
+	return MakeShareable(new FSkeletonSlotNamesSummoner(InHostingApp, InEditableSkeleton, InOnPostUndo, InOnObjectSelected));
+}
+
+TSharedRef<SWidget> FPersonaModule::CreateEditorWidgetForAnimDocument(const TSharedRef<class FWorkflowCentricApplication>& InHostingApp, UObject* InAnimAsset, const FAnimDocumentArgs& InArgs, FString& OutDocumentLink)
+{
+	TSharedPtr<SWidget> Result = SNullWidget::NullWidget;
+	if (InAnimAsset)
+	{
+		if (UAnimSequence* Sequence = Cast<UAnimSequence>(InAnimAsset))
+		{
+			Result = SNew(SSequenceEditor, InArgs.PreviewScene.Pin().ToSharedRef(), InArgs.EditableSkeleton.Pin().ToSharedRef(), InArgs.OnPostUndo, InArgs.OnCurvesChanged)
+				.Sequence(Sequence)
+				.OnObjectsSelected(InArgs.OnDespatchObjectsSelected)
+				.OnAnimNotifiesChanged(InArgs.OnDespatchAnimNotifiesChanged)
+				.OnInvokeTab(InArgs.OnDespatchInvokeTab);
+
+			OutDocumentLink = TEXT("Engine/Animation/Sequences");
+		}
+		else if (UAnimComposite* Composite = Cast<UAnimComposite>(InAnimAsset))
+		{
+			Result = SNew(SAnimCompositeEditor, InArgs.PreviewScene.Pin().ToSharedRef(), InArgs.EditableSkeleton.Pin().ToSharedRef(), InArgs.OnPostUndo)
+				.Composite(Composite)
+				.OnObjectsSelected(InArgs.OnDespatchObjectsSelected)
+				.OnAnimNotifiesChanged(InArgs.OnDespatchAnimNotifiesChanged)
+				.OnInvokeTab(InArgs.OnDespatchInvokeTab);
+
+			OutDocumentLink = TEXT("Engine/Animation/AnimationComposite");
+		}
+		else if (UAnimMontage* Montage = Cast<UAnimMontage>(InAnimAsset))
+		{
+			FMontageEditorRequiredArgs RequiredArgs(InArgs.PreviewScene.Pin().ToSharedRef(), InArgs.EditableSkeleton.Pin().ToSharedRef(), InArgs.OnPostUndo, InArgs.OnAnimNotifiesChanged, InArgs.OnSectionsChanged);
+
+			Result = SNew(SMontageEditor, RequiredArgs)
+				.Montage(Montage)
+				.OnCurvesChanged(InArgs.OnDespatchCurvesChanged)
+				.OnSectionsChanged(InArgs.OnDespatchSectionsChanged)
+				.OnInvokeTab(InArgs.OnDespatchInvokeTab)
+				.OnObjectsSelected(InArgs.OnDespatchObjectsSelected)
+				.OnAnimNotifiesChanged(InArgs.OnDespatchAnimNotifiesChanged);
+
+			OutDocumentLink = TEXT("Engine/Animation/AnimMontage");
+		}
+		else if (UPoseAsset* PoseAsset = Cast<UPoseAsset>(InAnimAsset))
+		{
+			Result = SNew(SPoseEditor, InArgs.PersonaToolkit.Pin().ToSharedRef(), InArgs.EditableSkeleton.Pin().ToSharedRef(), InArgs.PreviewScene.Pin().ToSharedRef())
+				.PoseAsset(PoseAsset);
+
+			OutDocumentLink = TEXT("Engine/Animation/Sequences");
+		}
+		else if (UBlendSpace* BlendSpace = Cast<UBlendSpace>(InAnimAsset))
 		{
-			SafeXPosition = EventGraph->Nodes[0]->NodePosX;
-			SafeYPosition = EventGraph->Nodes[EventGraph->Nodes.Num() - 1]->NodePosY + EventGraph->Nodes[EventGraph->Nodes.Num() - 1]->NodeHeight + 100;
+			Result = SNew(SBlendSpaceEditor, InArgs.PreviewScene.Pin().ToSharedRef(), InArgs.OnPostUndo)
+				.BlendSpace(BlendSpace);
+
+			if (Cast<UAimOffsetBlendSpace>(InAnimAsset))
+			{
+				OutDocumentLink = TEXT("Engine/Animation/AimOffset");
+			}
+			else
+			{
+				OutDocumentLink = TEXT("Engine/Animation/Blendspaces");
+			}
+		}
+		else if (UBlendSpace1D* BlendSpace1D = Cast<UBlendSpace1D>(InAnimAsset))
+		{
+			Result = SNew(SBlendSpaceEditor1D, InArgs.PreviewScene.Pin().ToSharedRef(), InArgs.OnPostUndo)
+				.BlendSpace1D(BlendSpace1D);
+
+			if (Cast<UAimOffsetBlendSpace1D>(InAnimAsset))
+			{
+				OutDocumentLink = TEXT("Engine/Animation/AimOffset");
+			}
+			else
+			{
+				OutDocumentLink = TEXT("Engine/Animation/Blendspaces");
+			}
 		}
+	}
 
-		// add try get owner node
-		UK2Node_CallFunction* GetOwnerNode = NewObject<UK2Node_CallFunction>(EventGraph);
-		UFunction* MakeNodeFunction = UAnimInstance::StaticClass()->FindFunctionByName(GET_FUNCTION_NAME_CHECKED(UAnimInstance, TryGetPawnOwner));
-		GetOwnerNode->CreateNewGuid();
-		GetOwnerNode->PostPlacedNewNode();
-		GetOwnerNode->SetFromFunction(MakeNodeFunction);
-		GetOwnerNode->SetFlags(RF_Transactional);
-		GetOwnerNode->AllocateDefaultPins();
-		GetOwnerNode->NodePosX = SafeXPosition;
-		GetOwnerNode->NodePosY = SafeYPosition;
-		UEdGraphSchema_K2::SetNodeMetaData(GetOwnerNode, FNodeMetadata::DefaultGraphNode);
-		GetOwnerNode->DisableNode();
-
-		EventGraph->AddNode(GetOwnerNode);
+	if (Result.IsValid())
+	{
+		InAnimAsset->SetFlags(RF_Transactional);
 	}
+
+	return Result.ToSharedRef();
 }
+
+void FPersonaModule::CustomizeMeshDetails(const TSharedRef<IDetailsView>& InDetailsView, const TSharedRef<IPersonaToolkit>& InPersonaToolkit)
+{
+	InDetailsView->SetGenericLayoutDetailsDelegate(FOnGetDetailCustomizationInstance::CreateStatic(&FPersonaMeshDetails::MakeInstance, InPersonaToolkit));
+}
+
+void FPersonaModule::ImportNewAsset(USkeleton* InSkeleton, EFBXImportType DefaultImportType)
+{
+	TSharedRef<SImportPathDialog> NewAnimDlg =
+		SNew(SImportPathDialog);
+
+	if (NewAnimDlg->ShowModal() != EAppReturnType::Cancel)
+	{
+		FString AssetPath = NewAnimDlg->GetAssetPath();
+
+		UFbxImportUI* ImportUI = NewObject<UFbxImportUI>();
+		ImportUI->Skeleton = InSkeleton;
+		ImportUI->MeshTypeToImport = DefaultImportType;
+
+		FbxMeshUtils::SetImportOption(ImportUI);
+
+		// now I have to set skeleton on it. 
+		FAssetToolsModule& AssetToolsModule = FModuleManager::GetModuleChecked<FAssetToolsModule>("AssetTools");
+		AssetToolsModule.Get().ImportAssets(AssetPath);
+	}
+}
+
+void PopulateWithAssets(FName ClassName, FName SkeletonMemberName, const FString& SkeletonString, TArray<FAssetData>& OutAssets)
+{
+	FAssetRegistryModule& AssetRegistryModule = FModuleManager::LoadModuleChecked<FAssetRegistryModule>(TEXT("AssetRegistry"));
+
+	FARFilter Filter;
+	Filter.ClassNames.Add(ClassName);
+	Filter.TagsAndValues.Add(SkeletonMemberName, SkeletonString);
+
+	AssetRegistryModule.Get().GetAssets(Filter, OutAssets);
+}
+
+void FPersonaModule::TestSkeletonCurveNamesForUse(const TSharedRef<IEditableSkeleton>& InEditableSkeleton) const
+{
+	const USkeleton& Skeleton = InEditableSkeleton->GetSkeleton();
+
+	if (const FSmartNameMapping* Mapping = Skeleton.GetSmartNameContainer(USkeleton::AnimCurveMappingName))
+	{
+		const FString SkeletonString = FAssetData(&Skeleton).GetExportTextName();
+
+		TArray<FAssetData> SkeletalMeshes;
+		PopulateWithAssets(USkeletalMesh::StaticClass()->GetFName(), GET_MEMBER_NAME_CHECKED(USkeletalMesh, Skeleton), SkeletonString, SkeletalMeshes);
+		TArray<FAssetData> Animations;
+		PopulateWithAssets(UAnimSequence::StaticClass()->GetFName(), FName("Skeleton"), SkeletonString, Animations);
+
+		FText TimeTakenMessage = FText::Format(LOCTEXT("TimeTakenWarning", "In order to verify curve usage all Skeletal Meshes and Animations that use this skeleton will be loaded, this may take some time.\n\nProceed?\n\nNumber of Meshes: {0}\nNumber of Animations: {1}"), FText::AsNumber(SkeletalMeshes.Num()), FText::AsNumber(Animations.Num()));
+
+		if (FMessageDialog::Open(EAppMsgType::YesNo, TimeTakenMessage) == EAppReturnType::Yes)
+		{
+			const FText LoadingStatusUpdate = FText::Format(LOCTEXT("VerifyCurves_LoadingAllAnimations", "Loading all animations for skeleton '{0}'"), FText::FromString(Skeleton.GetName()));
+			{
+				FScopedSlowTask LoadingAnimSlowTask(Animations.Num(), LoadingStatusUpdate);
+				LoadingAnimSlowTask.MakeDialog();
+
+				// Loop through all animations to load then, this makes sure smart names are all up to date
+				for (const FAssetData& Anim : Animations)
+				{
+					LoadingAnimSlowTask.EnterProgressFrame();
+					UAnimSequence* Seq = Cast<UAnimSequence>(Anim.GetAsset());
+				}
+			}
+
+			// Grab all curve names for this skeleton
+			TArray<FName> UnusedNames;
+			Mapping->FillNameArray(UnusedNames);
+
+			const FText ProcessingStatusUpdate = FText::Format(LOCTEXT("VerifyCurves_ProcessingSkeletalMeshes", "Looking at curve useage for each skeletal mesh of skeleton '{0}'"), FText::FromString(Skeleton.GetName()));
+			{
+				FScopedSlowTask LoadingSkelMeshSlowTask(SkeletalMeshes.Num(), ProcessingStatusUpdate);
+				LoadingSkelMeshSlowTask.MakeDialog();
+
+				for (int32 MeshIdx = 0; MeshIdx < SkeletalMeshes.Num(); ++MeshIdx)
+				{
+					LoadingSkelMeshSlowTask.EnterProgressFrame();
+
+					const USkeletalMesh* Mesh = Cast<USkeletalMesh>(SkeletalMeshes[MeshIdx].GetAsset());
+
+					// Filter morph targets from curves
+					const TArray<UMorphTarget*>& MorphTargets = Mesh->MorphTargets;
+					for (int32 I = 0; I < MorphTargets.Num(); ++I)
+					{
+						const int32 CurveIndex = UnusedNames.RemoveSingleSwap(MorphTargets[I]->GetFName(), false);
+					}
+
+					// Filter material params from curves
+					for (const FSkeletalMaterial& Mat : Mesh->Materials)
+					{
+						if (UnusedNames.Num() == 0)
+						{
+							break; // Done
+						}
+
+						UMaterial* Material = (Mat.MaterialInterface != nullptr) ? Mat.MaterialInterface->GetMaterial() : nullptr;
+						if (Material)
+						{
+							TArray<FName> OutParameterNames;
+							TArray<FGuid> OutParameterIds;
+
+							// Retrieve all scalar parameter names from the material
+							Material->GetAllScalarParameterNames(OutParameterNames, OutParameterIds);
+
+							for (FName SPName : OutParameterNames)
+							{
+								UnusedNames.RemoveSingleSwap(SPName);
+							}
+						}
+					}
+				}
+			}
+
+			FMessageLog CurveOutput("Persona");
+			CurveOutput.NewPage(LOCTEXT("PersonaMessageLogName", "Persona"));
+
+			bool bFoundIssue = false;
+
+			const FText ProcessingAnimStatusUpdate = FText::Format(LOCTEXT("VerifyCurves_ProcessingSkeletalMeshes", "Finding animations that reference unused curves on skeleton '{0}'"), FText::FromString(Skeleton.GetName()));
+			{
+				FScopedSlowTask ProcessingAnimationsSlowTask(Animations.Num(), ProcessingAnimStatusUpdate);
+				ProcessingAnimationsSlowTask.MakeDialog();
+
+				for (const FAssetData& Anim : Animations)
+				{
+					ProcessingAnimationsSlowTask.EnterProgressFrame();
+					UAnimSequence* Seq = Cast<UAnimSequence>(Anim.GetAsset());
+
+					TSharedPtr<FTokenizedMessage> Message;
+					for (FFloatCurve& Curve : Seq->RawCurveData.FloatCurves)
+					{
+						if (UnusedNames.Contains(Curve.Name.DisplayName))
+						{
+							bFoundIssue = true;
+							if (!Message.IsValid())
+							{
+								Message = CurveOutput.Warning();
+								Message->AddToken(FAssetNameToken::Create(Anim.ObjectPath.ToString(), FText::FromName(Anim.AssetName)));
+								Message->AddToken(FTextToken::Create(LOCTEXT("VerifyCurves_FoundAnimationsWithUnusedReferences", "References the following curves that are not used for either morph targets or material parameters and so may be unneeded")));
+							}
+							CurveOutput.Info(FText::FromName(Curve.Name.DisplayName));
+						}
+					}
+				}
+			}
+
+			if (bFoundIssue)
+			{
+				CurveOutput.Notify();
+			}
+		}
+	}
+}
+
+void FPersonaModule::ApplyCompression(TArray<TWeakObjectPtr<UAnimSequence>>& AnimSequences)
+{
+	FDlgAnimCompression AnimCompressionDialog(AnimSequences);
+	AnimCompressionDialog.ShowModal();
+}
+
+void FPersonaModule::ExportToFBX(TArray<TWeakObjectPtr<UAnimSequence>>& AnimSequences, USkeletalMesh* SkeletalMesh)
+{
+	IDesktopPlatform* DesktopPlatform = FDesktopPlatformModule::Get();
+
+	if (DesktopPlatform)
+	{
+		if (SkeletalMesh == NULL)
+		{
+			FMessageDialog::Open(EAppMsgType::Ok, LOCTEXT("ExportToFBXExportMissingSkeletalMesh", "ERROR: Missing skeletal mesh"));
+			return;
+		}
+
+		if (AnimSequences.Num() > 0)
+		{
+			//Get parent window for dialogs
+			TSharedPtr<SWindow> RootWindow = FGlobalTabmanager::Get()->GetRootWindow();
+
+			void* ParentWindowWindowHandle = NULL;
+
+			if (RootWindow.IsValid() && RootWindow->GetNativeWindow().IsValid())
+			{
+				ParentWindowWindowHandle = RootWindow->GetNativeWindow()->GetOSWindowHandle();
+			}
+
+			//Cache anim file names
+			TArray<FString> AnimFileNames;
+			for (auto Iter = AnimSequences.CreateIterator(); Iter; ++Iter)
+			{
+				AnimFileNames.Add(Iter->Get()->GetName() + TEXT(".fbx"));
+			}
+
+			IPlatformFile& PlatformFile = FPlatformFileManager::Get().GetPlatformFile();
+			FString DestinationFolder;
+
+			const FString Title = LOCTEXT("ExportFBXsToFolderTitle", "Choose a destination folder for the FBX file(s)").ToString();
+
+			if (AnimSequences.Num() > 1)
+			{
+				bool bFolderValid = false;
+				// More than one file, just ask for directory
+				while (!bFolderValid)
+				{
+					const bool bFolderSelected = DesktopPlatform->OpenDirectoryDialog(
+						ParentWindowWindowHandle,
+						Title,
+						FEditorDirectories::Get().GetLastDirectory(ELastDirectory::GENERIC_EXPORT),
+						DestinationFolder
+						);
+
+					if (!bFolderSelected)
+					{
+						// User canceled, return
+						return;
+					}
+
+					FEditorDirectories::Get().SetLastDirectory(ELastDirectory::GENERIC_EXPORT, DestinationFolder);
+					FPaths::NormalizeFilename(DestinationFolder);
+
+					//Check whether there are any fbx filename conflicts in this folder
+					for (auto Iter = AnimFileNames.CreateIterator(); Iter; ++Iter)
+					{
+						FString& AnimFileName = *Iter;
+						FString FullPath = DestinationFolder + "/" + AnimFileName;
+
+						bFolderValid = true;
+						if (PlatformFile.FileExists(*FullPath))
+						{
+							FFormatNamedArguments Args;
+							Args.Add(TEXT("DestinationFolder"), FText::FromString(DestinationFolder));
+							const FText DialogMessage = FText::Format(LOCTEXT("ExportToFBXFileOverwriteMessage", "Exporting to '{DestinationFolder}' will cause one or more existing FBX files to be overwritten. Would you like to continue?"), Args);
+							EAppReturnType::Type DialogReturn = FMessageDialog::Open(EAppMsgType::YesNo, DialogMessage);
+							bFolderValid = (EAppReturnType::Yes == DialogReturn);
+							break;
+						}
+					}
+				}
+			}
+			else
+			{
+				// One file only, ask for full filename.
+				// Can set bFolderValid from the SaveFileDialog call as the window will handle 
+				// duplicate files for us.
+				TArray<FString> TempDestinationNames;
+				bool bSave = DesktopPlatform->SaveFileDialog(
+					ParentWindowWindowHandle,
+					Title,
+					FEditorDirectories::Get().GetLastDirectory(ELastDirectory::GENERIC_EXPORT),
+					AnimSequences[0]->GetName(),
+					"FBX  |*.fbx",
+					EFileDialogFlags::None,
+					TempDestinationNames
+					);
+
+				if (!bSave)
+				{
+					// Canceled
+					return;
+				}
+				check(TempDestinationNames.Num() == 1);
+				check(AnimFileNames.Num() == 1);
+
+				DestinationFolder = FPaths::GetPath(TempDestinationNames[0]);
+				AnimFileNames[0] = FPaths::GetCleanFilename(TempDestinationNames[0]);
+
+				FEditorDirectories::Get().SetLastDirectory(ELastDirectory::GENERIC_EXPORT, DestinationFolder);
+			}
+
+			EAppReturnType::Type DialogReturn = FMessageDialog::Open(EAppMsgType::YesNo, LOCTEXT("ExportToFBXExportSkeletalMeshToo", "Would you like to export the current skeletal mesh with the animation(s)?"));
+			bool bSaveSkeletalMesh = EAppReturnType::Yes == DialogReturn;
+
+			const bool bShowCancel = false;
+			const bool bShowProgressDialog = true;
+			GWarn->BeginSlowTask(LOCTEXT("ExportToFBXProgress", "Exporting Animation(s) to FBX"), bShowProgressDialog, bShowCancel);
+
+			// make sure to use SkeletalMesh, when export inside of Persona
+			const int32 NumberOfAnimations = AnimSequences.Num();
+			for (int32 i = 0; i < NumberOfAnimations; ++i)
+			{
+				GWarn->UpdateProgress(i, NumberOfAnimations);
+
+				UAnimSequence* AnimSequence = AnimSequences[i].Get();
+
+				FString FileName = FString::Printf(TEXT("%s/%s"), *DestinationFolder, *AnimFileNames[i]);
+
+				FbxAnimUtils::ExportAnimFbx(*FileName, AnimSequence, SkeletalMesh, bSaveSkeletalMesh);
+			}
+
+			GWarn->EndSlowTask();
+		}
+	}
+}
+
+void FPersonaModule::AddLoopingInterpolation(TArray<TWeakObjectPtr<UAnimSequence>>& AnimSequences)
+{
+	FText WarningMessage = LOCTEXT("AddLoopiingInterpolation", "This will add an extra first frame at the end of the animation to create a better looping interpolation. This action cannot be undone. Would you like to proceed?");
+
+	if (FMessageDialog::Open(EAppMsgType::YesNo, WarningMessage) == EAppReturnType::Yes)
+	{
+		for (auto Animation : AnimSequences)
+		{
+			// get first frame and add to the last frame and go through track
+			// now calculating old animated space bases
+			Animation->AddLoopingInterpolation();
+		}
+	}
+}
+
+void FPersonaModule::CustomizeSlotNodeDetails(const TSharedRef<class IDetailsView>& InDetailsView, FOnInvokeTab InOnInvokeTab)
+{
+	InDetailsView->RegisterInstancedCustomPropertyLayout(UAnimGraphNode_Slot::StaticClass(),
+		FOnGetDetailCustomizationInstance::CreateStatic(&FAnimGraphNodeSlotDetails::MakeInstance, InOnInvokeTab));
+}
+
+IPersonaEditorModeManager* FPersonaModule::CreatePersonaEditorModeManager()
+{
+	return new FPersonaEditorModeManager();
+}
+
+#undef LOCTEXT_NAMESPACE
\ No newline at end of file
diff --git a/Engine/Source/Editor/Persona/Private/PersonaPreviewSceneDescription.cpp b/Engine/Source/Editor/Persona/Private/PersonaPreviewSceneDescription.cpp
new file mode 100644
--- /dev/null
+++ b/Engine/Source/Editor/Persona/Private/PersonaPreviewSceneDescription.cpp
@@ -0,0 +1,5 @@
+// Copyright 1998-2016 Epic Games, Inc. All Rights Reserved.
+
+#include "PersonaPrivatePCH.h"
+#include "PersonaPreviewSceneDescription.h"
+
diff --git a/Engine/Source/Editor/Persona/Private/PersonaPreviewSceneDescription.h b/Engine/Source/Editor/Persona/Private/PersonaPreviewSceneDescription.h
new file mode 100644
--- /dev/null
+++ b/Engine/Source/Editor/Persona/Private/PersonaPreviewSceneDescription.h
@@ -0,0 +1,40 @@
+// Copyright 1998-2016 Epic Games, Inc. All Rights Reserved.
+
+#pragma once
+
+#include "PersonaPreviewSceneDescription.generated.h"
+
+UENUM()
+enum class EPreviewAnimationMode : uint8
+{
+	/** Animate the preview based on the current editor (e.g. Skeleton = Reference Pose, Animation = Animation, Animation Blueprint = Animation Blueprint) */
+	Default,
+
+	/** Use the reference pose from the skeleton */
+	ReferencePose,
+
+	/** Use a specified animation */
+	UseSpecificAnimation,
+};
+
+UCLASS()
+class UPersonaPreviewSceneDescription : public UObject
+{
+public:
+	GENERATED_BODY()
+
+	/** The method by which the preview is animated */
+	UPROPERTY(EditAnywhere, Category = "Animation")
+	EPreviewAnimationMode AnimationMode;
+
+	/** The preview animation to use */
+	UPROPERTY(EditAnywhere, Category = "Animation", meta=(DisplayThumbnail=true))
+	TAssetPtr<UAnimationAsset> Animation;
+
+	/** The preview mesh to use */
+	UPROPERTY(EditAnywhere, Category = "Mesh", meta=(DisplayThumbnail=true))
+	TAssetPtr<USkeletalMesh> PreviewMesh;
+
+	UPROPERTY(EditAnywhere, Category = "Additional Meshes")
+	TAssetPtr<UPreviewMeshCollection> AdditionalMeshes;
+};
\ No newline at end of file
diff --git a/Engine/Source/Editor/Persona/Private/PersonaPrivatePCH.h b/Engine/Source/Editor/Persona/Private/PersonaPrivatePCH.h
--- a/Engine/Source/Editor/Persona/Private/PersonaPrivatePCH.h
+++ b/Engine/Source/Editor/Persona/Private/PersonaPrivatePCH.h
@@ -5,11 +5,15 @@
 
 #include "UnrealEd.h"
 #include "SlateBasics.h"
+#include "NotificationManager.h"
+#include "SNotificationList.h"
 
 #include "Editor/Kismet/Public/WorkflowOrientedApp/ApplicationMode.h"
 #include "Editor/Kismet/Public/WorkflowOrientedApp/WorkflowTabFactory.h"
 #include "Editor/Kismet/Public/WorkflowOrientedApp/WorkflowTabManager.h"
 
 #include "Shared/PersonaMode.h"
+#include "PropertyEditing.h"
+#include "PropertyCustomizationHelpers.h"
 
 #endif // __PersonaPrivatePCH_h__
