Copying //UE4/Dev-Framework to //UE4/Dev-Main (Source: //UE4/Dev-Framework @ 3136612)

diff --git a/Engine/Source/Editor/Persona/Private/PersonaToolkit.cpp b/Engine/Source/Editor/Persona/Private/PersonaToolkit.cpp
new file mode 100644
--- /dev/null
+++ b/Engine/Source/Editor/Persona/Private/PersonaToolkit.cpp
@@ -0,0 +1,148 @@
+// Copyright 1998-2016 Epic Games, Inc. All Rights Reserved.
+
+#include "PersonaPrivatePCH.h"
+#include "PersonaToolkit.h"
+#include "ISkeletonEditorModule.h"
+
+FPersonaToolkit::FPersonaToolkit()
+	: Skeleton(nullptr)
+	, Mesh(nullptr)
+	, AnimBlueprint(nullptr)
+	, AnimationAsset(nullptr)
+{
+}
+
+void FPersonaToolkit::Initialize(USkeleton* InSkeleton)
+{
+	check(InSkeleton);
+	Skeleton = InSkeleton;
+	Mesh = InSkeleton->GetPreviewMesh(true);
+}
+
+void FPersonaToolkit::Initialize(UAnimationAsset* InAnimationAsset)
+{
+	check(InAnimationAsset);
+	AnimationAsset = InAnimationAsset;
+	Skeleton = InAnimationAsset->GetSkeleton();
+	Mesh = InAnimationAsset->GetPreviewMesh();
+	if (Mesh == nullptr)
+	{
+		Mesh = Skeleton->GetPreviewMesh(true);
+	}
+}
+
+void FPersonaToolkit::Initialize(USkeletalMesh* InSkeletalMesh)
+{
+	check(InSkeletalMesh);
+	Skeleton = InSkeletalMesh->Skeleton;
+	Mesh = InSkeletalMesh;
+}
+
+void FPersonaToolkit::Initialize(UAnimBlueprint* InAnimBlueprint)
+{
+	check(InAnimBlueprint);
+	AnimBlueprint = InAnimBlueprint;
+	Skeleton = InAnimBlueprint->TargetSkeleton;
+	check(InAnimBlueprint->TargetSkeleton);
+	Mesh = InAnimBlueprint->TargetSkeleton->GetPreviewMesh(true);
+}
+
+void FPersonaToolkit::CreatePreviewScene()
+{
+	if (!PreviewScene.IsValid())
+	{
+		if (!EditableSkeleton.IsValid())
+		{
+			ISkeletonEditorModule& SkeletonEditorModule = FModuleManager::LoadModuleChecked<ISkeletonEditorModule>("SkeletonEditor");
+			EditableSkeleton = SkeletonEditorModule.CreateEditableSkeleton(Skeleton);
+		}
+
+		PreviewScene = MakeShareable(new FAnimationEditorPreviewScene(FPreviewScene::ConstructionValues().AllowAudioPlayback(true).ShouldSimulatePhysics(true), EditableSkeleton.ToSharedRef(), AsShared()));
+
+		//Temporary fix for missing attached assets - MDW
+		PreviewScene->GetWorld()->GetWorldSettings()->SetIsTemporarilyHiddenInEditor(false);
+
+		bool bSetMesh = false;
+
+		// Set the mesh
+		if (Mesh != nullptr)
+		{
+			PreviewScene->SetPreviewMesh(Mesh);
+			bSetMesh = true;
+
+			if (!Skeleton->GetPreviewMesh())
+			{
+				Skeleton->SetPreviewMesh(Mesh, false);
+			}
+		}
+		else if (AnimationAsset != nullptr)
+		{
+			USkeletalMesh* AssetMesh = AnimationAsset->GetPreviewMesh();
+			if (AssetMesh)
+			{
+				PreviewScene->SetPreviewMesh(AssetMesh);
+				bSetMesh = true;
+			}
+		}
+
+		if (!bSetMesh && Skeleton)
+		{
+			//If no preview mesh set, just find the first mesh that uses this skeleton
+			USkeletalMesh* PreviewMesh = Skeleton->GetPreviewMesh(true);
+			if (PreviewMesh)
+			{
+				PreviewScene->SetPreviewMesh(PreviewMesh);
+			}
+		}
+	}
+}
+
+USkeleton* FPersonaToolkit::GetSkeleton() const
+{
+	return Skeleton;
+}
+
+UDebugSkelMeshComponent* FPersonaToolkit::GetPreviewMeshComponent() const
+{
+	return PreviewScene.IsValid() ? PreviewScene->GetPreviewMeshComponent() : nullptr;
+}
+
+USkeletalMesh* FPersonaToolkit::GetMesh() const
+{
+	return Mesh;
+}
+
+void FPersonaToolkit::SetMesh(class USkeletalMesh* InSkeletalMesh)
+{
+	if (InSkeletalMesh != nullptr)
+	{
+		check(InSkeletalMesh->Skeleton == Skeleton);
+	}
+
+	Mesh = InSkeletalMesh;
+}
+
+UAnimBlueprint* FPersonaToolkit::GetAnimBlueprint() const
+{
+	return AnimBlueprint;
+}
+
+UAnimationAsset* FPersonaToolkit::GetAnimationAsset() const
+{
+	return AnimationAsset;
+}
+
+void FPersonaToolkit::SetAnimationAsset(class UAnimationAsset* InAnimationAsset)
+{
+	if (InAnimationAsset != nullptr)
+	{
+		check(InAnimationAsset->GetSkeleton() == Skeleton);
+	}
+
+	AnimationAsset = InAnimationAsset;
+}
+
+TSharedRef<IPersonaPreviewScene> FPersonaToolkit::GetPreviewScene() const
+{
+	return PreviewScene.ToSharedRef();
+}
\ No newline at end of file
diff --git a/Engine/Source/Editor/Persona/Private/PreviewSceneCustomizations.cpp b/Engine/Source/Editor/Persona/Private/PreviewSceneCustomizations.cpp
new file mode 100644
--- /dev/null
+++ b/Engine/Source/Editor/Persona/Private/PreviewSceneCustomizations.cpp
@@ -0,0 +1,235 @@
+// Copyright 1998-2016 Epic Games, Inc. All Rights Reserved.
+
+#include "PersonaPrivatePCH.h"
+#include "PreviewSceneCustomizations.h"
+#include "PersonaPreviewSceneDescription.h"
+#include "Engine/PreviewMeshCollection.h"
+#include "Factories/PreviewMeshCollectionFactory.h"
+#include "IEditableSkeleton.h"
+#include "IPersonaPreviewScene.h"
+#include "IPropertyUtilities.h"
+
+void FPreviewSceneDescriptionCustomization::CustomizeDetails(IDetailLayoutBuilder& DetailBuilder)
+{
+	FPropertyEditorModule& PropertyEditorModule = FModuleManager::LoadModuleChecked<FPropertyEditorModule>("PropertyEditor");
+
+	TSharedRef<IPropertyHandle> AnimationModeProperty = DetailBuilder.GetProperty(GET_MEMBER_NAME_CHECKED(UPersonaPreviewSceneDescription, AnimationMode));
+	AnimationModeProperty->SetOnPropertyValueChanged(FSimpleDelegate::CreateSP(this, &FPreviewSceneDescriptionCustomization::HandleAnimationModeChanged));
+	TSharedRef<IPropertyHandle> AnimationProperty = DetailBuilder.GetProperty(GET_MEMBER_NAME_CHECKED(UPersonaPreviewSceneDescription, Animation));
+	TSharedRef<IPropertyHandle> SkeletalMeshProperty = DetailBuilder.GetProperty(GET_MEMBER_NAME_CHECKED(UPersonaPreviewSceneDescription, PreviewMesh));
+	TSharedRef<IPropertyHandle> AdditionalMeshesProperty = DetailBuilder.GetProperty(GET_MEMBER_NAME_CHECKED(UPersonaPreviewSceneDescription, AdditionalMeshes));
+
+	DetailBuilder.EditCategory("Animation")
+	.AddProperty(AnimationModeProperty);
+
+	DetailBuilder.EditCategory("Animation")
+	.AddProperty(AnimationProperty)
+	.CustomWidget()
+	.NameContent()
+	[
+		AnimationProperty->CreatePropertyNameWidget()
+	]
+	.ValueContent()
+	.MaxDesiredWidth(250.0f)
+	.MinDesiredWidth(250.0f)
+	[
+		SNew(SObjectPropertyEntryBox)
+		.AllowedClass(UAnimationAsset::StaticClass())
+		.PropertyHandle(AnimationProperty)
+		.OnShouldFilterAsset(this, &FPreviewSceneDescriptionCustomization::HandleShouldFilterAsset, false)
+		.OnObjectChanged(this, &FPreviewSceneDescriptionCustomization::HandleAnimationChanged)
+		.ThumbnailPool(DetailBuilder.GetThumbnailPool())
+	];
+
+	DetailBuilder.EditCategory("Mesh")
+	.AddProperty(SkeletalMeshProperty)
+	.CustomWidget()
+	.NameContent()
+	[
+		SkeletalMeshProperty->CreatePropertyNameWidget()
+	]
+	.ValueContent()
+	.MaxDesiredWidth(250.0f)
+	.MinDesiredWidth(250.0f)
+	[
+		SNew(SObjectPropertyEntryBox)
+		.AllowedClass(USkeletalMesh::StaticClass())
+		.PropertyHandle(SkeletalMeshProperty)
+		.OnShouldFilterAsset(this, &FPreviewSceneDescriptionCustomization::HandleShouldFilterAsset, false)
+		.OnObjectChanged(this, &FPreviewSceneDescriptionCustomization::HandleMeshChanged)
+		.ThumbnailPool(DetailBuilder.GetThumbnailPool())
+	];
+
+	// use a factory with a custom fixed skeleton as we shouldn't be picking one here
+	UPreviewMeshCollectionFactory* FactoryToUse = NewObject<UPreviewMeshCollectionFactory>();
+	FactoryToUse->CurrentSkeleton = &EditableSkeleton.Pin()->GetSkeleton();
+
+	TArray<UFactory*> FactoriesToUse({ FactoryToUse });
+
+	DetailBuilder.EditCategory("Additional Meshes")
+	.AddProperty(AdditionalMeshesProperty)
+	.CustomWidget()
+	.NameContent()
+	[
+		AdditionalMeshesProperty->CreatePropertyNameWidget()
+	]
+	.ValueContent()
+	.MaxDesiredWidth(250.0f)
+	.MinDesiredWidth(250.0f)
+	[
+		SNew(SObjectPropertyEntryBox)
+		.AllowedClass(UPreviewMeshCollection::StaticClass())
+		.PropertyHandle(AdditionalMeshesProperty)
+		.OnShouldFilterAsset(this, &FPreviewSceneDescriptionCustomization::HandleShouldFilterAsset, true)
+		.OnObjectChanged(this, &FPreviewSceneDescriptionCustomization::HandleAdditionalMeshesChanged, &DetailBuilder)
+		.ThumbnailPool(DetailBuilder.GetThumbnailPool())
+		.NewAssetFactories(FactoriesToUse)
+	];
+
+	FAssetData AdditionalMeshesAsset;
+	AdditionalMeshesProperty->GetValue(AdditionalMeshesAsset);
+
+	if (AdditionalMeshesAsset.IsValid())
+	{
+		TArray<UObject*> Objects;
+		Objects.Add(AdditionalMeshesAsset.GetAsset());
+
+		IDetailPropertyRow* PropertyRow = DetailBuilder.EditCategory("Additional Meshes")
+		.AddExternalProperty(Objects, "SkeletalMeshes");
+
+		if (PropertyRow)
+		{
+			PropertyRow->ShouldAutoExpand(true);
+		}
+	}
+}
+
+bool FPreviewSceneDescriptionCustomization::HandleShouldFilterAsset(const FAssetData& InAssetData, bool bCanUseDifferentSkeleton)
+{
+	if (bCanUseDifferentSkeleton && GetDefault<UPersonaOptions>()->bAllowPreviewMeshCollectionsToSelectFromDifferentSkeletons)
+	{
+		return true;
+	}
+
+	FString SkeletonTag = InAssetData.GetTagValueRef<FString>("Skeleton");
+	if (SkeletonTag == SkeletonName)
+	{
+		return false;
+	}
+
+	return true;
+}
+
+void FPreviewSceneDescriptionCustomization::HandleAnimationModeChanged()
+{
+	UPersonaPreviewSceneDescription* PersonaPreviewSceneDescription = PreviewScene.Pin()->GetPreviewSceneDescription();
+	switch (PersonaPreviewSceneDescription->AnimationMode)
+	{
+	case EPreviewAnimationMode::Default:
+		PreviewScene.Pin()->ShowDefaultMode();
+		break;
+	case EPreviewAnimationMode::ReferencePose:
+		PreviewScene.Pin()->ShowReferencePose(true);
+		break;
+	case EPreviewAnimationMode::UseSpecificAnimation:
+		PreviewScene.Pin()->SetPreviewAnimationAsset(Cast<UAnimationAsset>(PersonaPreviewSceneDescription->Animation.LoadSynchronous()));
+		break;
+	}
+}
+
+void FPreviewSceneDescriptionCustomization::HandleAnimationChanged(const FAssetData& InAssetData)
+{
+	UAnimationAsset* AnimationAsset = Cast<UAnimationAsset>(InAssetData.GetAsset());
+	PreviewScene.Pin()->SetPreviewAnimationAsset(AnimationAsset);
+
+	if (AnimationAsset != nullptr)
+	{
+		PreviewScene.Pin()->GetPreviewSceneDescription()->AnimationMode = EPreviewAnimationMode::UseSpecificAnimation;
+	}
+}
+
+void FPreviewSceneDescriptionCustomization::HandleMeshChanged(const FAssetData& InAssetData)
+{
+	USkeletalMesh* NewPreviewMesh = Cast<USkeletalMesh>(InAssetData.GetAsset());
+
+	EditableSkeleton.Pin()->SetPreviewMesh(NewPreviewMesh);
+	PreviewScene.Pin()->SetPreviewMesh(NewPreviewMesh);
+}
+
+void FPreviewSceneDescriptionCustomization::HandleAdditionalMeshesChanged(const FAssetData& InAssetData, IDetailLayoutBuilder* DetailLayoutBuilder)
+{
+	PreviewScene.Pin()->SetAdditionalMeshes(Cast<UPreviewMeshCollection>(InAssetData.GetAsset()));
+	DetailLayoutBuilder->ForceRefreshDetails();
+}
+
+void FPreviewMeshCollectionEntryCustomization::CustomizeHeader(TSharedRef<IPropertyHandle> PropertyHandle, FDetailWidgetRow& HeaderRow, IPropertyTypeCustomizationUtils& CustomizationUtils)
+{
+	// get the enclosing preview mesh collection to determine the skeleton we want
+	TArray<UObject*> OuterObjects;
+	PropertyHandle->GetOuterObjects(OuterObjects);
+
+	check(OuterObjects.Num() > 0);
+		
+	if (OuterObjects[0] != nullptr)
+	{
+		FString SkeletonName = FAssetData(CastChecked<UPreviewMeshCollection>(OuterObjects[0])->Skeleton).GetExportTextName();
+
+		PropertyHandle->GetParentHandle()->SetOnPropertyValueChanged(FSimpleDelegate::CreateSP(this, &FPreviewMeshCollectionEntryCustomization::HandleMeshesArrayChanged, CustomizationUtils.GetPropertyUtilities()));
+
+		TSharedPtr<IPropertyHandle> SkeletalMeshProperty = PropertyHandle->GetChildHandle(GET_MEMBER_NAME_CHECKED(FPreviewMeshCollectionEntry, SkeletalMesh));
+		if (SkeletalMeshProperty.IsValid())
+		{
+			HeaderRow.NameContent()
+			[
+				SkeletalMeshProperty->CreatePropertyNameWidget()
+			]
+			.ValueContent()
+			.MaxDesiredWidth(250.0f)
+			.MinDesiredWidth(250.0f)
+			[
+				SNew(SObjectPropertyEntryBox)
+				.AllowedClass(USkeletalMesh::StaticClass())
+				.PropertyHandle(SkeletalMeshProperty)
+				.OnShouldFilterAsset(this, &FPreviewMeshCollectionEntryCustomization::HandleShouldFilterAsset, SkeletonName)
+				.OnObjectChanged(this, &FPreviewMeshCollectionEntryCustomization::HandleMeshChanged)
+				.ThumbnailPool(CustomizationUtils.GetThumbnailPool())
+			];
+		}
+	}
+}
+
+bool FPreviewMeshCollectionEntryCustomization::HandleShouldFilterAsset(const FAssetData& InAssetData, FString SkeletonName)
+{
+	if (GetDefault<UPersonaOptions>()->bAllowPreviewMeshCollectionsToSelectFromDifferentSkeletons)
+	{
+		return true;
+	}
+
+	FString SkeletonTag = InAssetData.GetTagValueRef<FString>("Skeleton");
+	if (SkeletonTag == SkeletonName)
+	{
+		return false;
+	}
+
+	return true;
+}
+
+void FPreviewMeshCollectionEntryCustomization::HandleMeshChanged(const FAssetData& InAssetData)
+{
+	if (PreviewScene.IsValid())
+	{
+		PreviewScene.Pin()->RefreshAdditionalMeshes();
+	}
+}
+
+void FPreviewMeshCollectionEntryCustomization::HandleMeshesArrayChanged(TSharedPtr<IPropertyUtilities> PropertyUtilities)
+{
+	if (PreviewScene.IsValid())
+	{
+		PreviewScene.Pin()->RefreshAdditionalMeshes();
+		if (PropertyUtilities.IsValid())
+		{
+			PropertyUtilities->ForceRefresh();
+		}
+	}
+}
diff --git a/Engine/Source/Editor/Persona/Private/PreviewSceneCustomizations.h b/Engine/Source/Editor/Persona/Private/PreviewSceneCustomizations.h
new file mode 100644
--- /dev/null
+++ b/Engine/Source/Editor/Persona/Private/PreviewSceneCustomizations.h
@@ -0,0 +1,65 @@
+// Copyright 1998-2016 Epic Games, Inc. All Rights Reserved.
+
+#pragma once
+
+#include "AnimationEditorPreviewScene.h"
+
+class FPreviewSceneDescriptionCustomization : public IDetailCustomization
+{
+public:
+	FPreviewSceneDescriptionCustomization(const FString& InSkeletonName, const TSharedRef<class IPersonaPreviewScene>& InPreviewScene, const TSharedRef<class IEditableSkeleton>& InEditableSkeleton)
+		: SkeletonName(InSkeletonName)
+		, PreviewScene(StaticCastSharedRef<FAnimationEditorPreviewScene>(InPreviewScene))
+		, EditableSkeleton(InEditableSkeleton)
+	{}
+
+	virtual void CustomizeDetails(IDetailLayoutBuilder& DetailBuilder) override;
+
+private:
+	bool HandleShouldFilterAsset(const FAssetData& InAssetData, bool bCanUseDifferentSkeleton);
+
+	void HandleAnimationModeChanged();
+
+	void HandleAnimationChanged(const FAssetData& InAssetData);
+
+	void HandleMeshChanged(const FAssetData& InAssetData);
+
+	void HandleAdditionalMeshesChanged(const FAssetData& InAssetData, IDetailLayoutBuilder* DetailLayoutBuilder);
+
+private:
+	/** Cached skeleton name to check for asset registry tags */
+	FString SkeletonName;
+
+	/** Preview scene we will be editing */
+	TWeakPtr<class FAnimationEditorPreviewScene> PreviewScene;
+
+	/** Editable Skeleton scene we will be editing */
+	TWeakPtr<class IEditableSkeleton> EditableSkeleton;
+};
+
+class FPreviewMeshCollectionEntryCustomization : public IPropertyTypeCustomization
+{
+public:
+	static TSharedRef<IPropertyTypeCustomization> MakeInstance()
+	{
+		return MakeShareable(new FPreviewMeshCollectionEntryCustomization(nullptr));
+	}
+
+	FPreviewMeshCollectionEntryCustomization(const TSharedPtr<IPersonaPreviewScene>& InPreviewScene)
+		: PreviewScene(InPreviewScene)
+	{}
+
+	virtual void CustomizeHeader(TSharedRef<IPropertyHandle> PropertyHandle, FDetailWidgetRow& HeaderRow, IPropertyTypeCustomizationUtils& CustomizationUtils) override;
+	virtual void CustomizeChildren(TSharedRef<IPropertyHandle> PropertyHandle, IDetailChildrenBuilder& ChildBuilder, IPropertyTypeCustomizationUtils& CustomizationUtils) override {}
+
+private:
+	bool HandleShouldFilterAsset(const FAssetData& InAssetData, FString SkeletonName);
+
+	void HandleMeshChanged(const FAssetData& InAssetData);
+
+	void HandleMeshesArrayChanged(TSharedPtr<IPropertyUtilities> PropertyUtilities);
+
+private:
+	/** Preview scene we will be editing */
+	TWeakPtr<class IPersonaPreviewScene> PreviewScene;
+};
\ No newline at end of file
diff --git a/Engine/Source/Editor/Persona/Private/SAnimBlueprintParentPlayerList.cpp b/Engine/Source/Editor/Persona/Private/SAnimBlueprintParentPlayerList.cpp
--- a/Engine/Source/Editor/Persona/Private/SAnimBlueprintParentPlayerList.cpp
+++ b/Engine/Source/Editor/Persona/Private/SAnimBlueprintParentPlayerList.cpp
@@ -37,27 +37,15 @@ SAnimBlueprintParentPlayerList::~SAnimBlueprintParentPlayerList()
 		CurrentBlueprint->OnChanged().RemoveAll(this);
 		CurrentBlueprint->OnCompiled().RemoveAll(this);
 	}
-
-	TSharedPtr<FPersona> PersonaShared = PersonaPtr.Pin();
-	if(PersonaPtr.IsValid())
-	{
-		PersonaPtr.Pin()->UnregisterOnPostUndo(this);
-	}
 }
 
-void SAnimBlueprintParentPlayerList::Construct(const FArguments& InArgs)
+void SAnimBlueprintParentPlayerList::Construct(const FArguments& InArgs, const TSharedRef<FBlueprintEditor>& InBlueprintEditor, FSimpleMulticastDelegate& InOnPostUndo)
 {
-	PersonaPtr = InArgs._Persona;
-
 	// Register a refresh on post undo to grab the blueprint assets again
-	TSharedPtr<FPersona> PersonaShared = PersonaPtr.Pin();
-	if(PersonaPtr.IsValid())
-	{
-		PersonaPtr.Pin()->RegisterOnPostUndo(FPersona::FOnPostUndo::CreateSP(this, &SAnimBlueprintParentPlayerList::RefreshDetailView));
-	}
+	InOnPostUndo.Add(FSimpleDelegate::CreateSP(this, &SAnimBlueprintParentPlayerList::RefreshDetailView));
 
 	UEditorParentPlayerListObj* EditorObject = Cast<UEditorParentPlayerListObj>(ObjectTracker.GetEditorObjectForClass(UEditorParentPlayerListObj::StaticClass()));
-	CurrentBlueprint = PersonaPtr.Pin()->GetAnimBlueprint();
+	CurrentBlueprint = CastChecked<UAnimBlueprint>(InBlueprintEditor->GetBlueprintObj());
 	check(CurrentBlueprint);
 	CurrentParentClass = CurrentBlueprint->ParentClass;
 
@@ -72,12 +60,12 @@ void SAnimBlueprintParentPlayerList::Construct(const FArguments& InArgs)
 
 		DetailView = PropertyModule.CreateDetailView(ViewArgs);
 
-		DetailView->RegisterInstancedCustomPropertyLayout(UEditorParentPlayerListObj::StaticClass(), FOnGetDetailCustomizationInstance::CreateStatic(&FAnimGraphParentPlayerDetails::MakeInstance, PersonaPtr));
+		DetailView->RegisterInstancedCustomPropertyLayout(UEditorParentPlayerListObj::StaticClass(), FOnGetDetailCustomizationInstance::CreateStatic(&FAnimGraphParentPlayerDetails::MakeInstance, InBlueprintEditor));
 
 		this->ChildSlot
-			[
-				DetailView->AsShared()
-			];
+		[
+			DetailView->AsShared()
+		];
 
 		DetailView->SetObject(EditorObject, true);
 
diff --git a/Engine/Source/Editor/Persona/Private/SAnimBlueprintParentPlayerList.h b/Engine/Source/Editor/Persona/Private/SAnimBlueprintParentPlayerList.h
--- a/Engine/Source/Editor/Persona/Private/SAnimBlueprintParentPlayerList.h
+++ b/Engine/Source/Editor/Persona/Private/SAnimBlueprintParentPlayerList.h
@@ -10,15 +10,14 @@ class SAnimBlueprintParentPlayerList : public SCompoundWidget
 {
 public:
 	SLATE_BEGIN_ARGS(SAnimBlueprintParentPlayerList)
-		: _Persona()
-		{}
-	SLATE_ARGUMENT(TWeakPtr<FPersona>, Persona)
-		SLATE_END_ARGS()
+	{}
 
-		SAnimBlueprintParentPlayerList();
+	SLATE_END_ARGS()
+
+	SAnimBlueprintParentPlayerList();
 	~SAnimBlueprintParentPlayerList();
 
-	void Construct(const FArguments& InArgs);
+	void Construct(const FArguments& InArgs, const TSharedRef<FBlueprintEditor>& InBlueprintEditor, FSimpleMulticastDelegate& InOnPostUndo);
 
 private:
 
@@ -38,9 +37,6 @@ private:
 
 	void RefreshDetailView();
 
-	// Persona instance we're currently in
-	TWeakPtr<FPersona> PersonaPtr;
-
 	// Object tracker to maintain single instance of editor objects
 	FEditorObjectTracker ObjectTracker;
 
diff --git a/Engine/Source/Editor/Persona/Private/SAnimCompositeEditor.cpp b/Engine/Source/Editor/Persona/Private/SAnimCompositeEditor.cpp
--- a/Engine/Source/Editor/Persona/Private/SAnimCompositeEditor.cpp
+++ b/Engine/Source/Editor/Persona/Private/SAnimCompositeEditor.cpp
@@ -20,25 +20,23 @@ TSharedRef<SWidget> SAnimCompositeEditor::CreateDocumentAnchor()
 	return IDocumentation::Get()->CreateAnchor(TEXT("Engine/Animation/AnimationComposite"));
 }
 
-void SAnimCompositeEditor::Construct(const FArguments& InArgs)
+void SAnimCompositeEditor::Construct(const FArguments& InArgs, const TSharedRef<class IPersonaPreviewScene>& InPreviewScene, const TSharedRef<class IEditableSkeleton>& InEditableSkeleton, FSimpleMulticastDelegate& OnPostUndo)
 {
 	bIsActiveTimerRegistered = false;
-	PersonaPtr = InArgs._Persona;
 	CompositeObj = InArgs._Composite;
 	check(CompositeObj);
 
 	SAnimEditorBase::Construct( SAnimEditorBase::FArguments()
-		.Persona(InArgs._Persona)
-		);
+		.OnObjectsSelected(InArgs._OnObjectsSelected), 
+		InPreviewScene );
 
-	PersonaPtr.Pin()->RegisterOnPostUndo(FPersona::FOnPostUndo::CreateSP( this, &SAnimCompositeEditor::PostUndo ) );	
+	OnPostUndo.Add(FSimpleDelegate::CreateSP( this, &SAnimCompositeEditor::PostUndo ) );
 
 	EditorPanels->AddSlot()
 		.AutoHeight()
 		.Padding(0, 10)
 		[
 			SAssignNew( AnimCompositePanel, SAnimCompositePanel )
-			.Persona(PersonaPtr)
 			.Composite(CompositeObj)
 			.CompositeEditor(SharedThis(this))
 			.WidgetWidth(S2ColumnWidget::DEFAULT_RIGHT_COLUMN_WIDTH)
@@ -51,8 +49,7 @@ void SAnimCompositeEditor::Construct(const FArguments& InArgs)
 		.AutoHeight()
 		.Padding(0, 10)
 		[
-			SAssignNew( AnimNotifyPanel, SAnimNotifyPanel )
-			.Persona(InArgs._Persona)
+			SAssignNew( AnimNotifyPanel, SAnimNotifyPanel, OnPostUndo)
 			.Sequence(CompositeObj)
 			.WidgetWidth(S2ColumnWidget::DEFAULT_RIGHT_COLUMN_WIDTH)
 			.InputMin(this, &SAnimEditorBase::GetMinInput)
@@ -62,14 +59,15 @@ void SAnimCompositeEditor::Construct(const FArguments& InArgs)
 			.OnSetInputViewRange(this, &SAnimEditorBase::SetInputViewRange)
 			.OnGetScrubValue(this, &SAnimEditorBase::GetScrubValue)
 			.OnSelectionChanged(this, &SAnimEditorBase::OnSelectionChanged)
+			.OnAnimNotifiesChanged(InArgs._OnAnimNotifiesChanged)
+			.OnInvokeTab(InArgs._OnInvokeTab)
 		];
 
 	EditorPanels->AddSlot()
 		.AutoHeight()
 		.Padding(0, 10)
 		[
-			SAssignNew( AnimCurvePanel, SAnimCurvePanel )
-			.Persona(PersonaPtr)
+			SAssignNew( AnimCurvePanel, SAnimCurvePanel, InEditableSkeleton)
 			.Sequence(CompositeObj)
 			.WidgetWidth(S2ColumnWidget::DEFAULT_RIGHT_COLUMN_WIDTH)
 			.ViewInputMin(this, &SAnimEditorBase::GetViewMinInput)
@@ -83,14 +81,6 @@ void SAnimCompositeEditor::Construct(const FArguments& InArgs)
 	CollapseComposite();
 }
 
-SAnimCompositeEditor::~SAnimCompositeEditor()
-{
-	if (PersonaPtr.IsValid())
-	{
-		PersonaPtr.Pin()->UnregisterOnPostUndo(this);
-	}
-}
-
 void SAnimCompositeEditor::PreAnimUpdate()
 {
 	CompositeObj->Modify();
diff --git a/Engine/Source/Editor/Persona/Private/SAnimCompositeEditor.h b/Engine/Source/Editor/Persona/Private/SAnimCompositeEditor.h
--- a/Engine/Source/Editor/Persona/Private/SAnimCompositeEditor.h
+++ b/Engine/Source/Editor/Persona/Private/SAnimCompositeEditor.h
@@ -22,20 +22,17 @@ class SAnimCompositeEditor : public SAnimEditorBase
 {
 public:
 	SLATE_BEGIN_ARGS( SAnimCompositeEditor )
-		: _Persona()
-		, _Composite(NULL)
+		: _Composite(NULL)
 		{}
 
-		SLATE_ARGUMENT( TSharedPtr<FPersona>,	Persona )
-		SLATE_ARGUMENT( UAnimComposite*,		Composite)
-	SLATE_END_ARGS()
+		SLATE_ARGUMENT( UAnimComposite*, Composite)
+		SLATE_EVENT(FOnObjectsSelected, OnObjectsSelected)
+		SLATE_EVENT(FSimpleDelegate, OnAnimNotifiesChanged)
+		SLATE_EVENT(FOnInvokeTab, OnInvokeTab)
 
-	~SAnimCompositeEditor();
+	SLATE_END_ARGS()
 
 private:
-	/** Persona reference **/
-	TWeakPtr<FPersona> PersonaPtr;
-
 	/** Slate editor panels */
 	TSharedPtr<class SAnimCompositePanel> AnimCompositePanel;
 	TSharedPtr<class SAnimNotifyPanel> AnimNotifyPanel;
@@ -55,7 +52,7 @@ protected:
 	virtual void InitDetailsViewEditorObject(class UEditorAnimBaseObj* EdObj) override;
 
 public:
-	void Construct(const FArguments& InArgs);
+	void Construct(const FArguments& InArgs, const TSharedRef<class IPersonaPreviewScene>& InPreviewScene, const TSharedRef<class IEditableSkeleton>& InEditableSkeleton, FSimpleMulticastDelegate& OnPostUndo);
 	
 	/** Return the animation composite being edited */
 	UAnimComposite* GetCompositeObj() const { return CompositeObj; }
diff --git a/Engine/Source/Editor/Persona/Private/SAnimCompositePanel.cpp b/Engine/Source/Editor/Persona/Private/SAnimCompositePanel.cpp
--- a/Engine/Source/Editor/Persona/Private/SAnimCompositePanel.cpp
+++ b/Engine/Source/Editor/Persona/Private/SAnimCompositePanel.cpp
@@ -23,7 +23,6 @@ void SAnimCompositePanel::Construct(const FArguments& InArgs)
 		.InputMax(InArgs._InputMax)
 		.OnSetInputViewRange(InArgs._OnSetInputViewRange));
 
-	Persona = InArgs._Persona;
 	Composite = InArgs._Composite;
 	CompositeEditor = InArgs._CompositeEditor;
 
diff --git a/Engine/Source/Editor/Persona/Private/SAnimCompositePanel.h b/Engine/Source/Editor/Persona/Private/SAnimCompositePanel.h
--- a/Engine/Source/Editor/Persona/Private/SAnimCompositePanel.h
+++ b/Engine/Source/Editor/Persona/Private/SAnimCompositePanel.h
@@ -22,7 +22,6 @@ class SAnimCompositePanel : public SAnimTrackPanel
 public:
 	SLATE_BEGIN_ARGS( SAnimCompositePanel )
 		: _Composite()
-		, _Persona()
 		, _CompositeEditor()
 		, _ViewInputMin()
 		, _ViewInputMax()
@@ -32,7 +31,6 @@ public:
 	{}
 
 	SLATE_ARGUMENT( class UAnimComposite*, Composite)
-	SLATE_ARGUMENT( TWeakPtr<FPersona>, Persona)
 	SLATE_ARGUMENT( TWeakPtr<class SAnimCompositeEditor>, CompositeEditor)
 	SLATE_ARGUMENT( float, WidgetWidth )
 	SLATE_ATTRIBUTE( float, ViewInputMin )
@@ -54,9 +52,6 @@ public:
 	void ShowSegmentInDetailsView(int32 SegmentIndex);
 
 private:
-	/** Persona reference */
-	TWeakPtr<FPersona>			Persona;
-
 	/** Reference to our editor object */
 	TWeakPtr<SAnimCompositeEditor>	CompositeEditor;
 
diff --git a/Engine/Source/Editor/Persona/Private/SAnimCurveEd.cpp b/Engine/Source/Editor/Persona/Private/SAnimCurveEd.cpp
--- a/Engine/Source/Editor/Persona/Private/SAnimCurveEd.cpp
+++ b/Engine/Source/Editor/Persona/Private/SAnimCurveEd.cpp
@@ -17,18 +17,8 @@ float SAnimCurveEd::GetTimeStep(FTrackScaleInfo &ScaleInfo)const
 	if(NumberOfKeys.Get())
 	{
 		int32 Divider = SScrubWidget::GetDivider(ViewMinInput.Get(), ViewMaxInput.Get(), ScaleInfo.WidgetSize, TimelineLength.Get(), NumberOfKeys.Get());
-
-		float TimePerKey;
-
-		if(NumberOfKeys.Get() != 0.f)
-		{
-			TimePerKey = TimelineLength.Get()/(float)NumberOfKeys.Get();
-		}
-		else
-		{
-			TimePerKey = 1.f;
-		}
-
+		const FAnimKeyHelper Helper(TimelineLength.Get(), NumberOfKeys.Get());
+		float TimePerKey = Helper.TimePerKey();
 		return TimePerKey * Divider;
 	}
 
diff --git a/Engine/Source/Editor/Persona/Private/SAnimCurveEd.h b/Engine/Source/Editor/Persona/Private/SAnimCurveEd.h
--- a/Engine/Source/Editor/Persona/Private/SAnimCurveEd.h
+++ b/Engine/Source/Editor/Persona/Private/SAnimCurveEd.h
@@ -6,7 +6,7 @@
 #include "Persona.h"
 #include "SCurveEditor.h"
 
-DECLARE_DELEGATE_OneParam(FOnSelectionChanged, const FGraphPanelSelectionSet&)
+DECLARE_DELEGATE_OneParam(FOnSelectionChanged, const TArray<UObject*>&)
 DECLARE_DELEGATE(FOnUpdatePanel)
 DECLARE_DELEGATE_RetVal(float, FOnGetScrubValue)
 
diff --git a/Engine/Source/Editor/Persona/Private/SAnimCurvePanel.cpp b/Engine/Source/Editor/Persona/Private/SAnimCurvePanel.cpp
--- a/Engine/Source/Editor/Persona/Private/SAnimCurvePanel.cpp
+++ b/Engine/Source/Editor/Persona/Private/SAnimCurvePanel.cpp
@@ -12,6 +12,7 @@
 #include "SExpandableArea.h"
 #include "STextEntryPopup.h"
 #include "Animation/AnimSequence.h"
+#include "IEditableSkeleton.h"
 
 #define LOCTEXT_NAMESPACE "AnimCurvePanel"
 
@@ -24,7 +25,7 @@ class FAnimCurveBaseInterface : public FCurveOwnerInterface
 private:
 	FAnimCurveBase*	CurveData;
 
-	void UpdateNameInternal(FRawCurveTracks& RawCurveData, const FSmartNameMapping::UID& RequestedNameUID, FName RequestedName)
+	void UpdateNameInternal(FRawCurveTracks& RawCurveData, const SmartName::UID_Type& RequestedNameUID, FName RequestedName)
 	{
 		FAnimCurveBase* CurrentCurveData = RawCurveData.GetCurveData(CurveUID);
 		if (CurrentCurveData)
@@ -132,7 +133,7 @@ public:
 		CurrentCurveData->FloatCurve.AddKey(0.0f, 1.0f);
 	}
 
-	void UpdateName(const FSmartNameMapping::UID& RequestedNameUID, FName RequestedName)
+	void UpdateName(const SmartName::UID_Type& RequestedNameUID, FName RequestedName)
 	{
 		CurveUID = RequestedNameUID;
 		UpdateNameInternal(AnimSequenceBase.Get()->RawCurveData, RequestedNameUID, RequestedName);
@@ -145,7 +146,7 @@ public:
 	/**
 	* Set InFlag to bValue
 	*/
-	void SetCurveTypeFlag(EAnimCurveFlags InFlag, bool bValue)
+	void SetCurveTypeFlag(EAnimAssetCurveFlags InFlag, bool bValue)
 	{
 		AnimSequenceBase.Get()->RawCurveData.GetCurveData(CurveUID)->SetCurveTypeFlag(InFlag, bValue);
 		if (UAnimSequence* Seq = AnimSequence.Get())
@@ -160,7 +161,7 @@ public:
 	/**
 	* Toggle the value of the specified flag
 	*/
-	void ToggleCurveTypeFlag(EAnimCurveFlags InFlag)
+	void ToggleCurveTypeFlag(EAnimAssetCurveFlags InFlag)
 	{
 		AnimSequenceBase.Get()->RawCurveData.GetCurveData(CurveUID)->ToggleCurveTypeFlag(InFlag);
 		if (UAnimSequence* Seq = AnimSequence.Get())
@@ -175,7 +176,7 @@ public:
 	/**
 	* Return true if InFlag is set, false otherwise
 	*/
-	bool GetCurveTypeFlag(EAnimCurveFlags InFlag) const
+	bool GetCurveTypeFlag(EAnimAssetCurveFlags InFlag) const
 	{
 		return AnimSequenceBase.Get()->RawCurveData.GetCurveData(CurveUID)->GetCurveTypeFlag(InFlag);
 	}
@@ -234,11 +235,6 @@ public:
 	// Delete current track
 	void DeleteTrack();
 
-	// Sets the current mode for this curve
-	void ToggleCurveMode(ECheckBoxState NewState, EAnimCurveFlags ModeToSet);
-
-	// Returns whether this curve is of the specificed mode type
-	ECheckBoxState IsCurveOfMode(EAnimCurveFlags ModeToTest) const;
 
 	/**
 	 * Build and display curve track context menu.
@@ -405,7 +401,7 @@ void SCurveEdTrack::NewCurveNameEntered( const FText& NewText, ETextCommit::Type
 			const FSmartNameMapping* NameMapping = Skeleton->GetSmartNameContainer(USkeleton::AnimCurveMappingName);
 
 			// If requested name exists, make sure it's not currently in use in this sequence.
-			const FSmartNameMapping::UID* RequestedNameUID = NameMapping->FindUID(RequestedName);
+			const SmartName::UID_Type* RequestedNameUID = NameMapping->FindUID(RequestedName);
 			if (RequestedNameUID != nullptr)
 			{
 				// Already in use in this sequence, skip
@@ -465,49 +461,6 @@ void SCurveEdTrack::DeleteTrack()
 	}
 }
 
-void SCurveEdTrack::ToggleCurveMode(ECheckBoxState NewState,EAnimCurveFlags ModeToSet)
-{
-	const int32 AllModes = (ACF_DriveMorphTarget|ACF_DriveMaterial);
-	check((ModeToSet&AllModes) != 0); //unexpected value for ModeToSet
-
-	FText UndoLabel;
-	bool bIsSwitchingFlagOn = !CurveInterface->GetCurveTypeFlag(ModeToSet);
-	check(bIsSwitchingFlagOn == (NewState==ECheckBoxState::Checked));
-	
-	if(bIsSwitchingFlagOn)
-	{
-		if(ModeToSet == ACF_DriveMorphTarget)
-		{
-			UndoLabel = LOCTEXT("AnimCurve_TurnOnMorphMode", "Enable driving of morph targets");
-		}
-		else if(ModeToSet == ACF_DriveMaterial)
-		{
-			UndoLabel = LOCTEXT("AnimCurve_TurnOnMaterialMode", "Enable driving of materials");
-		}
-	}
-	else
-	{
-		if(ModeToSet == ACF_DriveMorphTarget)
-		{
-			UndoLabel = LOCTEXT("AnimCurve_TurnOffMorphMode", "Disable driving of morph targets");
-		}
-		else if(ModeToSet == ACF_DriveMaterial)
-		{
-			UndoLabel = LOCTEXT("AnimCurve_TurnOffMaterialMode", "Disable driving of materials");
-		}
-	}
-
-	const FScopedTransaction Transaction( UndoLabel );
-	CurveInterface->MakeTransactional();
-	CurveInterface->ModifyOwner();
-
-	CurveInterface->SetCurveTypeFlag(ModeToSet, bIsSwitchingFlagOn);
-}
-
-ECheckBoxState SCurveEdTrack::IsCurveOfMode(EAnimCurveFlags ModeToTest) const
-{
-	return CurveInterface->GetCurveTypeFlag(ModeToTest) ? ECheckBoxState::Checked : ECheckBoxState::Unchecked;
-}
 
 FReply SCurveEdTrack::OnContextMenu()
 {
@@ -574,7 +527,7 @@ public:
 	}
 };
 
-void SAnimCurvePanel::Construct(const FArguments& InArgs)
+void SAnimCurvePanel::Construct(const FArguments& InArgs, const TSharedRef<class IEditableSkeleton>& InEditableSkeleton)
 {
 	SAnimTrackPanel::Construct( SAnimTrackPanel::FArguments()
 		.WidgetWidth(InArgs._WidgetWidth)
@@ -584,10 +537,12 @@ void SAnimCurvePanel::Construct(const FArguments& InArgs)
 		.InputMax(InArgs._InputMax)
 		.OnSetInputViewRange(InArgs._OnSetInputViewRange));
 
-	WeakPersona = InArgs._Persona;
 	Sequence = InArgs._Sequence;
 	WidgetWidth = InArgs._WidgetWidth;
 	OnGetScrubValue = InArgs._OnGetScrubValue;
+	OnCurvesChanged = InArgs._OnCurvesChanged;
+
+	InEditableSkeleton->RegisterOnSmartNameRemoved(FOnSmartNameRemoved::FDelegate::CreateSP(this, &SAnimCurvePanel::HandleSmartNameRemoved));
 
 	this->ChildSlot
 	[
@@ -791,34 +746,6 @@ FReply SAnimCurvePanel::OnContextMenu()
 {
 	FMenuBuilder MenuBuilder(true, NULL);
 
-	MenuBuilder.BeginSection("AnimCurvePanelCurveTypes", LOCTEXT("AllCurveTypesHeading", "All Curve Types"));
-	{
-		MenuBuilder.AddWidget(
-			SNew(SCheckBox)
-			.IsChecked( this, &SAnimCurvePanel::AreAllCurvesOfMode, ACF_DriveMorphTarget )
-			.OnCheckStateChanged( this, &SAnimCurvePanel::ToggleAllCurveModes, ACF_DriveMorphTarget )
-			.ToolTipText(LOCTEXT("MorphCurveModeTooltip", "This curve drives a morph target"))
-			[
-				SNew(STextBlock)
-				.Text(LOCTEXT("MorphCurveMode", "Morph Curve"))
-			],
-			FText()
-		);
-
-		MenuBuilder.AddWidget(
-			SNew(SCheckBox)
-			.IsChecked( this, &SAnimCurvePanel::AreAllCurvesOfMode, ACF_DriveMaterial )
-			.OnCheckStateChanged( this, &SAnimCurvePanel::ToggleAllCurveModes, ACF_DriveMaterial )
-			.ToolTipText(LOCTEXT("MaterialCurveModeTooltip", "This curve drives a material"))
-			.HAlign(HAlign_Left)
-			[
-				SNew(STextBlock)
-				.Text(LOCTEXT("MaterialCurveMode", "Material Curve"))
-			],
-			FText()
-		);
-	}
-	MenuBuilder.EndSection();
 
 	MenuBuilder.BeginSection("AnimCurvePanelOptions", LOCTEXT("OptionsHeading", "Options"));
 	{
@@ -847,47 +774,6 @@ EVisibility SAnimCurvePanel::IsSetAllTracksButtonVisible() const
 	return (Tracks.Num() > 1) ? EVisibility::Visible : EVisibility::Hidden;
 }
 
-void SAnimCurvePanel::ToggleAllCurveModes(ECheckBoxState NewState, EAnimCurveFlags ModeToSet)
-{
-	const ECheckBoxState CurrentAllState = AreAllCurvesOfMode(ModeToSet);
-	for(TWeakPtr<SCurveEdTrack> TrackWeak : Tracks)
-	{
-		TSharedPtr<SCurveEdTrack> TrackWidget = TrackWeak.Pin();
-		if( TrackWidget.IsValid() )
-		{
-			const ECheckBoxState CurrentTrackState = TrackWidget->IsCurveOfMode(ModeToSet);
-			if( (CurrentAllState == CurrentTrackState) || ((CurrentAllState == ECheckBoxState::Undetermined) && (CurrentTrackState == ECheckBoxState::Unchecked)) )
-			{
-				TrackWidget->ToggleCurveMode( NewState, ModeToSet );
-			}
-		}
-	}
-}
-
-ECheckBoxState SAnimCurvePanel::AreAllCurvesOfMode(EAnimCurveFlags ModeToSet) const
-{
-	int32 NumChecked = 0;
-	for(const TWeakPtr<SCurveEdTrack> TrackWeak : Tracks)
-	{
-		const TSharedPtr<SCurveEdTrack> TrackWidget = TrackWeak.Pin();
-		if( TrackWidget.IsValid() )
-		{
-			if ( TrackWidget->IsCurveOfMode(ModeToSet) == ECheckBoxState::Checked )
-			{
-				NumChecked++;
-			}
-		}
-	}
-	if( NumChecked == Tracks.Num() )
-	{
-		return ECheckBoxState::Checked;
-	}
-	else if( NumChecked == 0 )
-	{
-		return ECheckBoxState::Unchecked;
-	}
-	return ECheckBoxState::Undetermined;
-}
 
 void SAnimCurvePanel::UpdatePanel()
 {
@@ -963,11 +849,7 @@ void SAnimCurvePanel::UpdatePanel()
 			}
 		}
 
-		TSharedPtr<FPersona> SharedPersona = WeakPersona.Pin();
-		if(SharedPersona.IsValid())
-		{
-			SharedPersona->OnCurvesChanged.Broadcast();
-		}
+		OnCurvesChanged.ExecuteIfBound();
 	}
 }
 
@@ -1314,33 +1196,6 @@ TSharedRef<SWidget> SAnimCurvePanel::CreateCurveContextMenu(FAnimCurveBaseInterf
 {
 	FMenuBuilder MenuBuilder(true, NULL);
 
-	MenuBuilder.BeginSection("AnimCurvePanelCurveTypes", LOCTEXT("CurveTypesHeading", "Curve Types"));
-	{
-		MenuBuilder.AddWidget(
-			SNew(SCheckBox)
-			.IsChecked(this, &SAnimCurvePanel::GetCurveFlagAsCheckboxState, Curve, ACF_DriveMorphTarget)
-			.OnCheckStateChanged(this, &SAnimCurvePanel::SetCurveFlagFromCheckboxState, Curve, ACF_DriveMorphTarget)
-			.ToolTipText(LOCTEXT("MorphCurveModeTooltip", "This curve drives a morph target"))
-			[
-				SNew(STextBlock)
-				.Text(LOCTEXT("MorphCurveMode", "Morph Curve"))
-			],
-			FText()
-			);
-
-		MenuBuilder.AddWidget(
-			SNew(SCheckBox)
-			.IsChecked(this, &SAnimCurvePanel::GetCurveFlagAsCheckboxState, Curve, ACF_DriveMaterial)
-			.OnCheckStateChanged(this, &SAnimCurvePanel::SetCurveFlagFromCheckboxState, Curve, ACF_DriveMaterial)
-			.ToolTipText(LOCTEXT("MaterialCurveModeTooltip", "This curve drives a material"))
-			[
-				SNew(STextBlock)
-				.Text(LOCTEXT("MaterialCurveMode", "Material Curve"))
-			],
-			FText()
-			);
-	}
-	MenuBuilder.EndSection();
 
 	MenuBuilder.BeginSection("AnimCurvePanelTrackOptions", LOCTEXT("TrackOptionsHeading", "Track Options"));
 	{
@@ -1375,17 +1230,6 @@ TSharedRef<SWidget> SAnimCurvePanel::CreateCurveContextMenu(FAnimCurveBaseInterf
 	return MenuBuilder.MakeWidget();
 }
 
-ECheckBoxState SAnimCurvePanel::GetCurveFlagAsCheckboxState(FAnimCurveBaseInterface* Curve, EAnimCurveFlags InFlag) const
-{
-	return Curve->GetCurveTypeFlag(InFlag) ? ECheckBoxState::Checked : ECheckBoxState::Unchecked;
-}
-
-void SAnimCurvePanel::SetCurveFlagFromCheckboxState(ECheckBoxState CheckState, FAnimCurveBaseInterface* Curve, EAnimCurveFlags InFlag)
-{
-	bool Enabled = CheckState == ECheckBoxState::Checked;
-	Curve->SetCurveTypeFlag(InFlag, Enabled);
-}
-
 void SAnimCurvePanel::ToggleCurveTypeMenuCallback(FAnimCurveBaseInterface* Curve)
 {
 	check(Curve);
@@ -1417,4 +1261,9 @@ void SAnimCurvePanel::AddVariableCurve(USkeleton::AnimCurveUID CurveUid)
 	UpdatePanel();
 }
 
+void SAnimCurvePanel::HandleSmartNameRemoved(const FName& InContainerName, const TArray<SmartName::UID_Type>& InNameUids)
+{
+	UpdatePanel();
+}
+
 #undef LOCTEXT_NAMESPACE
diff --git a/Engine/Source/Editor/Persona/Private/SAnimCurvePanel.h b/Engine/Source/Editor/Persona/Private/SAnimCurvePanel.h
--- a/Engine/Source/Editor/Persona/Private/SAnimCurvePanel.h
+++ b/Engine/Source/Editor/Persona/Private/SAnimCurvePanel.h
@@ -10,7 +10,7 @@
 #include "SCurveEditor.h"
 #include "Animation/AnimSequenceBase.h"
 
-DECLARE_DELEGATE_OneParam( FOnSelectionChanged, const FGraphPanelSelectionSet& )
+DECLARE_DELEGATE_OneParam( FOnSelectionChanged, const TArray<UObject*>& )
 DECLARE_DELEGATE( FOnUpdatePanel )
 DECLARE_DELEGATE_RetVal( float, FOnGetScrubValue )
 
@@ -49,12 +49,9 @@ public:
 		, _InputMax()
 		, _OnSetInputViewRange()
 		, _OnGetScrubValue()
+		, _OnCurvesChanged()
 	{}
 
-	/**
-	 * Weak ptr to our Persona instance
-	 */
-	SLATE_ARGUMENT(TWeakPtr<FPersona>, Persona)
 	/**
 	 * AnimSequenceBase to be used for this panel
 	 */
@@ -75,9 +72,11 @@ public:
 	 * Get current value
 	 */
 	SLATE_EVENT( FOnGetScrubValue, OnGetScrubValue )
+
+	SLATE_EVENT( FSimpleDelegate, OnCurvesChanged )
 	SLATE_END_ARGS()
 
-	void Construct(const FArguments& InArgs);
+	void Construct(const FArguments& InArgs, const TSharedRef<class IEditableSkeleton>& InEditableSkeleton);
 	/**
 	 * Set New Sequence
 	 */
@@ -125,12 +124,12 @@ public:
 	/**
 	 * Toggle all the curve modes of all the tracks
 	 */
-	void ToggleAllCurveModes(ECheckBoxState NewState, EAnimCurveFlags ModeToSet);
+	void ToggleAllCurveModes(ECheckBoxState NewState, EAnimAssetCurveFlags ModeToSet);
 
 	/**
 	 * Are all the curve modes of all the tracks set to the mode
 	 */
-	ECheckBoxState AreAllCurvesOfMode(EAnimCurveFlags ModeToTest) const;
+	ECheckBoxState AreAllCurvesOfMode(EAnimAssetCurveFlags ModeToTest) const;
 
 	/**
 	 * Visibility of the set-all-tracks button
@@ -145,7 +144,7 @@ public:
 	/**
 	 * Sets the specified flag value to State for the provided curve
 	 */
-	void SetCurveFlag(FFloatCurve* CurveInterface, bool State, EAnimCurveFlags FlagToSet);
+	void SetCurveFlag(FFloatCurve* CurveInterface, bool State, EAnimAssetCurveFlags FlagToSet);
 
 	/**
 	* Update Panel
@@ -160,13 +159,14 @@ public:
 	float GetLength() const { return Sequence->SequenceLength; }
 
 private:
-	TWeakPtr<FPersona> WeakPersona;
 	TSharedPtr<SSplitter> PanelSlot;
 	class UAnimSequenceBase* Sequence;
 	TAttribute<float> CurrentPosition;
 	FOnGetScrubValue OnGetScrubValue;
 	TArray<TWeakPtr<class SCurveEdTrack>> Tracks;
 
+	FSimpleDelegate OnCurvesChanged;
+
 	/**
 	 * This is to control visibility of the curves, so you can edit or not
 	 * Get Widget that shows all curve list and edit
@@ -221,18 +221,12 @@ private:
 	*/
 	void CreateNewCurveClicked();
 
-	/**
-	 * Convert the requested flag bool value into a checkbox state
-	 */
-	ECheckBoxState GetCurveFlagAsCheckboxState(FAnimCurveBaseInterface* Curve, EAnimCurveFlags InFlag) const;
-
-	/**
-	 * Convert a given checkbox state into a flag value in the provided curve
-	 */
-	void SetCurveFlagFromCheckboxState(ECheckBoxState CheckState, FAnimCurveBaseInterface* Curve, EAnimCurveFlags InFlag);
 
 	/**
 	 * Handler for converting a curve from metadata to variable curve or vice versa
 	 */
 	void ToggleCurveTypeMenuCallback(FAnimCurveBaseInterface* Curve);
+
+	/** Handle smart names getting removed (and regenerate widgets) */
+	void HandleSmartNameRemoved(const FName& InContainerName, const TArray<SmartName::UID_Type>& InNameUids);
 };
\ No newline at end of file
diff --git a/Engine/Source/Editor/Persona/Private/SAnimCurveViewer.cpp b/Engine/Source/Editor/Persona/Private/SAnimCurveViewer.cpp
--- a/Engine/Source/Editor/Persona/Private/SAnimCurveViewer.cpp
+++ b/Engine/Source/Editor/Persona/Private/SAnimCurveViewer.cpp
@@ -12,10 +12,16 @@
 #include "Animation/AnimSingleNodeInstance.h"
 #include "SNotificationList.h"
 #include "NotificationManager.h"
+#include "IPersonaPreviewScene.h"
+#include "IEditableSkeleton.h"
+
+#include "GenericCommands.h"
+#include "CurveViewerCommands.h"
 
 #define LOCTEXT_NAMESPACE "SAnimCurveViewer"
 
 static const FName ColumnId_AnimCurveNameLabel( "Curve Name" );
+static const FName ColumnID_AnimCurveTypeLabel("Type");
 static const FName ColumnID_AnimCurveWeightLabel( "Weight" );
 static const FName ColumnID_AnimCurveEditLabel( "Edit" );
 
@@ -24,12 +30,23 @@ const float MaxMorphWeight = 5.f;
 //////////////////////////////////////////////////////////////////////////
 // SAnimCurveListRow
 
+typedef TSharedPtr< FDisplayedAnimCurveInfo > FDisplayedAnimCurveInfoPtr;
+//  This is a flag that is used to filter UI part
+enum EAnimCurveEditorFlags
+{
+	// Used as morph target curve
+	ACEF_DriveMorphTarget		= 0x00000001, 
+	// Used as triggering event
+	ACEF_DriveAttribute			= 0x00000002, 
+	// Used as a material curve
+	ACEF_DriveMaterial			= 0x00000004, 
+};
 
-
-void SAnimCurveListRow::Construct( const FArguments& InArgs, const TSharedRef<STableViewBase>& InOwnerTableView )
+void SAnimCurveListRow::Construct( const FArguments& InArgs, const TSharedRef<STableViewBase>& InOwnerTableView, const TSharedRef<IPersonaPreviewScene>& InPreviewScene)
 {
 	Item = InArgs._Item;
 	AnimCurveViewerPtr = InArgs._AnimCurveViewerPtr;
+	PreviewScenePtr = InPreviewScene;
 
 	check( Item.IsValid() );
 
@@ -64,6 +81,27 @@ TSharedRef< SWidget > SAnimCurveListRow::GenerateWidgetForColumn( const FName& C
 			return SNullWidget::NullWidget;
 		}
 	}
+	else if (ColumnName == ColumnID_AnimCurveTypeLabel)
+	{
+		TSharedPtr<SAnimCurveViewer> AnimCurveViewer = AnimCurveViewerPtr.Pin();
+		if (AnimCurveViewer.IsValid())
+		{
+			return
+				SNew(SVerticalBox)
+
+				+ SVerticalBox::Slot()
+				.AutoHeight()
+				.Padding(4)
+				.VAlign(VAlign_Center)
+				[
+					GetCurveTypeWidget()
+				];
+		}
+		else
+		{
+			return SNullWidget::NullWidget;
+		}
+	}
 	else if ( ColumnName == ColumnID_AnimCurveWeightLabel )
 	{
 		// Encase the SSpinbox in an SVertical box so we can apply padding. Setting ItemHeight on the containing SListView has no effect :-(
@@ -103,6 +141,79 @@ TSharedRef< SWidget > SAnimCurveListRow::GenerateWidgetForColumn( const FName& C
 	}
 }
 
+TSharedRef< SWidget > SAnimCurveListRow::GetCurveTypeWidget()
+{
+	return SNew(SHorizontalBox)
+
+		+SHorizontalBox::Slot()
+		.AutoWidth()
+		.Padding(0.f, 1.f, 1.f, 1.f)
+		.VAlign(VAlign_Center)
+		.HAlign(HAlign_Center)
+		[
+			SNew(SCheckBox)
+			.OnCheckStateChanged(this, &SAnimCurveListRow::OnAnimCurveTypeBoxChecked, true)
+			.IsChecked(this, &SAnimCurveListRow::IsAnimCurveTypeBoxChangedChecked, true)
+			.IsEnabled(false)
+			.CheckedImage(FEditorStyle::GetBrush("AnimCurveViewer.MorphTargetOn"))
+			.CheckedPressedImage(FEditorStyle::GetBrush("AnimCurveViewer.MorphTargetOn"))
+			.UncheckedImage(FEditorStyle::GetBrush("AnimCurveViewer.MorphTargetOff"))
+			.CheckedHoveredImage(FEditorStyle::GetBrush("AnimCurveViewer.MorphTargetOn"))
+			.UncheckedHoveredImage(FEditorStyle::GetBrush("AnimCurveViewer.MorphTargetOff"))
+			.ToolTipText(LOCTEXT("CurveTypeMorphTarget_Tooltip", "MorphTarget"))
+			.ForegroundColor(FEditorStyle::GetSlateColor("DefaultForeground"))
+		]
+
+		+SHorizontalBox::Slot()
+		.AutoWidth()
+		.Padding(0.f, 1.f, 1.f, 1.f)
+		.VAlign(VAlign_Center)
+		.HAlign(HAlign_Center)
+		[
+			SNew(SCheckBox)
+			.OnCheckStateChanged(this, &SAnimCurveListRow::OnAnimCurveTypeBoxChecked, false)
+			.IsChecked(this, &SAnimCurveListRow::IsAnimCurveTypeBoxChangedChecked, false)
+			.CheckedImage(FEditorStyle::GetBrush("AnimCurveViewer.MaterialOn"))
+			.CheckedPressedImage(FEditorStyle::GetBrush("AnimCurveViewer.MaterialOn"))
+			.UncheckedImage(FEditorStyle::GetBrush("AnimCurveViewer.MaterialOff"))
+			.CheckedHoveredImage(FEditorStyle::GetBrush("AnimCurveViewer.MaterialOn"))
+			.UncheckedHoveredImage(FEditorStyle::GetBrush("AnimCurveViewer.MaterialOff"))
+			.ToolTipText(LOCTEXT("CurveTypeMaterial_Tooltip", "Material"))
+			.ForegroundColor(FEditorStyle::GetSlateColor("DefaultForeground"))
+		];
+
+}
+
+void SAnimCurveListRow::OnAnimCurveTypeBoxChecked(ECheckBoxState InState, bool bMorphTarget)
+{
+	// currently only material curve is set
+	bool bNewData = (InState == ECheckBoxState::Checked);
+	if (!bMorphTarget)
+	{
+		Item->EditableSkeleton.Pin()->SetCurveMetaDataMaterial(Item->SmartName, bNewData);
+	}
+}
+
+ECheckBoxState SAnimCurveListRow::IsAnimCurveTypeBoxChangedChecked(bool bMorphTarget) const
+{
+	const FCurveMetaData* CurveMetaData = Item->EditableSkeleton.Pin()->GetSkeleton().GetCurveMetaData(Item->SmartName);
+
+	bool bData = false;
+	if (CurveMetaData)
+	{
+		if (bMorphTarget)
+		{
+			bData = CurveMetaData->Type.bMorphtarget != 0;
+		}
+		else
+		{
+			bData = CurveMetaData->Type.bMaterial != 0;
+		}
+	}
+
+	return (bData)? ECheckBoxState::Checked : ECheckBoxState::Unchecked;
+}
+
 void SAnimCurveListRow::OnAnimCurveAutoFillChecked(ECheckBoxState InState)
 {
 	Item->bAutoFillData = InState == ECheckBoxState::Checked;
@@ -136,11 +247,6 @@ void SAnimCurveListRow::OnAnimCurveWeightChanged( float NewWeight )
 	TSharedPtr<SAnimCurveViewer> AnimCurveViewer = AnimCurveViewerPtr.Pin();
 	if (AnimCurveViewer.IsValid())
 	{
-		// If adjusting an element that is not selected, make it the only selection
-		if (!AnimCurveViewer->AnimCurveListView->IsItemSelected(Item))
-		{
-			AnimCurveViewer->AnimCurveListView->SetSelection(Item);
-		}
 
 		// Add override
 		AnimCurveViewer->AddAnimCurveOverride(Item->SmartName.DisplayName, Item->Weight);
@@ -159,11 +265,9 @@ void SAnimCurveListRow::OnAnimCurveWeightChanged( float NewWeight )
 			}
 		}
 
-		// Refresh viewport
-		TSharedPtr<FPersona> Persona = AnimCurveViewer->PersonaPtr.Pin();
-		if (Persona.IsValid())
+		if(PreviewScenePtr.IsValid())
 		{
-			Persona->RefreshViewport();
+			PreviewScenePtr.Pin()->InvalidateViews();
 		}
 	}
 }
@@ -183,7 +287,7 @@ FText SAnimCurveListRow::GetItemName() const
 	TSharedPtr<SAnimCurveViewer> AnimCurveViewer = AnimCurveViewerPtr.Pin();
 	if (AnimCurveViewer.IsValid())
 	{
-		const FSmartNameMapping* Mapping = Item->Skeleton->GetSmartNameContainer(AnimCurveViewer->ContainerName);
+		const FSmartNameMapping* Mapping = Item->EditableSkeleton.Pin()->GetSkeleton().GetSmartNameContainer(AnimCurveViewer->ContainerName);
 		check(Mapping);
 		Mapping->GetName(Item->SmartName.UID, ItemName);
 	}
@@ -213,25 +317,20 @@ bool SAnimCurveListRow::GetActiveWeight(float& OutWeight) const
 	TSharedPtr<SAnimCurveViewer> AnimCurveViewer = AnimCurveViewerPtr.Pin();
 	if (AnimCurveViewer.IsValid())
 	{
-		// If persona
-		TSharedPtr<FPersona> Persona = AnimCurveViewer->PersonaPtr.Pin();
-		if (Persona.IsValid())
+		// If anim instance
+		UAnimInstance* AnimInstance = PreviewScenePtr.Pin()->GetPreviewMeshComponent()->GetAnimInstance();
+		if (AnimInstance)
 		{
-			// If anim instance
-			UAnimInstance* AnimInstance = Persona->GetPreviewAnimInstance();
-			if (AnimInstance)
-			{
-				// See if curve is in active set
-				TMap<FName, float> CurveList;
-				AnimInstance->GetAnimationCurveList(ACF_EditorPreviewCurves, CurveList);
+			// See if curve is in active set, attribute curve should have everything
+			TMap<FName, float> CurveList;
+			AnimInstance->GetAnimationCurveList(EAnimCurveType::AttributeCurve, CurveList);
 
-				float* CurrentValue = CurveList.Find(Item->SmartName.DisplayName);
-				if (CurrentValue)
-				{
-					OutWeight = *CurrentValue;
-					// Remember we found it
-					bFoundActive = true;
-				}
+			float* CurrentValue = CurveList.Find(Item->SmartName.DisplayName);
+			if (CurrentValue)
+			{
+				OutWeight = *CurrentValue;
+				// Remember we found it
+				bFoundActive = true;
 			}
 		}
 	}
@@ -253,7 +352,9 @@ FSlateColor SAnimCurveListRow::GetItemTextColor() const
 	if (Item->bAutoFillData)
 	{
 		float Weight = 0.f;
-		bItemActive = GetActiveWeight(Weight);
+		GetActiveWeight(Weight);
+		// change so that print white if it has weight on it. 
+		bItemActive = (Weight != 0.f);
 	}
 
 	return bItemActive ? FLinearColor(1, 1, 1) : FLinearColor(0.5, 0.5, 0.5);
@@ -377,15 +478,12 @@ FText SAnimCurveTypeList::GetAnimCurveType() const
 {
 	switch (CurveFlags)
 	{
-	case ACF_DriveMorphTarget:
+	case ACEF_DriveMorphTarget:
 		return LOCTEXT("AnimCurveType_Morphtarget", "Morph Target");
-	case ACF_DriveAttribute:
+	case ACEF_DriveAttribute:
 		return LOCTEXT("AnimCurveType_Attribute", "Attribute");
-	case ACF_DriveMaterial:
+	case ACEF_DriveMaterial:
 		return LOCTEXT("AnimCurveType_Material", "Material");
-		// @fixme: when we move curve types to skeleton fix this
-	case ACF_Metadata:
-		return LOCTEXT("AnimCurveType_Metadata", "Meta Data");
 	}
 
 	return LOCTEXT("AnimCurveType_Unknown", "Unknown");
@@ -393,29 +491,26 @@ FText SAnimCurveTypeList::GetAnimCurveType() const
 //////////////////////////////////////////////////////////////////////////
 // SAnimCurveViewer
 
-void SAnimCurveViewer::Construct(const FArguments& InArgs)
+void SAnimCurveViewer::Construct(const FArguments& InArgs, const TSharedRef<class IEditableSkeleton>& InEditableSkeleton, const TSharedRef<IPersonaPreviewScene>& InPreviewScene, FSimpleMulticastDelegate& InOnCurvesChanged, FSimpleMulticastDelegate& InOnPostUndo)
 {
-	PersonaPtr = InArgs._Persona;
-
 	ContainerName = USkeleton::AnimCurveMappingName;
 
 	CachedPreviewInstance = nullptr;
 
-	bShowAllCurves = true;
+	PreviewScenePtr = InPreviewScene;
+	EditableSkeletonPtr = InEditableSkeleton;
 
-	TSharedPtr<FPersona> Persona = PersonaPtr.Pin();
-	if (Persona.IsValid())
-	{
-		CurrentSkeleton = Persona->GetSkeleton();
+	InPreviewScene->RegisterOnPreviewMeshChanged(FOnPreviewMeshChanged::CreateSP(this, &SAnimCurveViewer::OnPreviewMeshChanged));
+	InPreviewScene->RegisterOnAnimChanged(FOnAnimChanged::CreateSP(this, &SAnimCurveViewer::OnPreviewAssetChanged));
+	InOnPostUndo.Add(FSimpleDelegate::CreateSP(this, &SAnimCurveViewer::OnPostUndo));
+	InOnCurvesChanged.Add(FSimpleDelegate::CreateSP(this, &SAnimCurveViewer::OnCurvesChanged));
 
-		Persona->RegisterOnPreviewMeshChanged(FPersona::FOnPreviewMeshChanged::CreateSP(this, &SAnimCurveViewer::OnPreviewMeshChanged));
-		Persona->RegisterOnAnimChanged(FPersona::FOnAnimChanged::CreateSP(this, &SAnimCurveViewer::OnPreviewAssetChanged));
-		Persona->RegisterOnPostUndo(FPersona::FOnPostUndo::CreateSP(this, &SAnimCurveViewer::OnPostUndo));
-		Persona->RegisterOnChangeCurves(FPersona::FOnCurvesChanged::CreateSP(this, &SAnimCurveViewer::OnCurvesChanged));
-	}
+	// Register and bind all our menu commands
+	FCurveViewerCommands::Register();
+	BindCommands();
 
 	// @todo fix this to be filtered
-	CurrentCurveFlag = ACF_DriveMorphTarget | ACF_DriveMaterial | ACF_DriveAttribute;
+	CurrentCurveFlag = ACEF_DriveMorphTarget | ACEF_DriveMaterial | ACEF_DriveAttribute;
 
 	RefreshCachePreviewInstance();
 	 
@@ -468,6 +563,10 @@ void SAnimCurveViewer::Construct(const FArguments& InArgs)
 				.FillWidth(1.f)
 				.DefaultLabel( LOCTEXT( "AnimCurveNameLabel", "Curve Name" ) )
 
+				+ SHeaderRow::Column(ColumnID_AnimCurveTypeLabel)
+				.FillWidth(1.f)
+				.DefaultLabel(LOCTEXT("AnimCurveTypeLabel", "Type"))
+
 				+ SHeaderRow::Column( ColumnID_AnimCurveWeightLabel )
 				.FillWidth(1.f)
 				.DefaultLabel( LOCTEXT( "AnimCurveWeightLabel", "Weight" ) )
@@ -485,37 +584,61 @@ void SAnimCurveViewer::Construct(const FArguments& InArgs)
 
 SAnimCurveViewer::~SAnimCurveViewer()
 {
-	TSharedPtr<FPersona> Persona = PersonaPtr.Pin();
-	if (Persona.IsValid())
+	if (PreviewScenePtr.IsValid() )
 	{
-		// if persona isn't there, we probably don't have the preview mesh either, so no valid anim instance
 		if (CachedPreviewInstance && OnAddAnimationCurveDelegate.IsBound())
 		{
 			CachedPreviewInstance->RemoveDelegate_AddCustomAnimationCurve(OnAddAnimationCurveDelegate);
 		}
 
-		Persona->UnregisterOnPreviewMeshChanged(this);
-		Persona->UnregisterOnAnimChanged(this);
-		Persona->UnregisterOnPostUndo(this);
-		Persona->UnregisterOnChangeCurves(this);
+		PreviewScenePtr.Pin()->UnregisterOnPreviewMeshChanged(this);
+		PreviewScenePtr.Pin()->UnregisterOnAnimChanged(this);
 	}
 }
 
-bool SAnimCurveViewer::IsCurveFilterEnabled() const
+FReply SAnimCurveViewer::OnKeyDown(const FGeometry& MyGeometry, const FKeyEvent& InKeyEvent)
 {
-	return !bShowAllCurves;
+	if (UICommandList.IsValid() && UICommandList->ProcessCommandBindings(InKeyEvent))
+	{
+		return FReply::Handled();
+	}
+	return FReply::Unhandled();
 }
 
-ECheckBoxState SAnimCurveViewer::IsShowingAllCurves() const
+void SAnimCurveViewer::BindCommands()
 {
-	return bShowAllCurves ? ECheckBoxState::Checked : ECheckBoxState::Unchecked;
+	// This should not be called twice on the same instance
+	check(!UICommandList.IsValid());
+
+	UICommandList = MakeShareable(new FUICommandList);
+
+	FUICommandList& CommandList = *UICommandList;
+
+	// Grab the list of menu commands to bind...
+	const FCurveViewerCommands& MenuActions = FCurveViewerCommands::Get();
+
+	// ...and bind them all
+
+	CommandList.MapAction(
+		FGenericCommands::Get().Rename,
+		FExecuteAction::CreateSP(this, &SAnimCurveViewer::OnRenameClicked),
+		FCanExecuteAction::CreateSP(this, &SAnimCurveViewer::CanRename));
+
+	CommandList.MapAction(
+		FGenericCommands::Get().Delete,
+		FExecuteAction::CreateSP(this, &SAnimCurveViewer::OnDeleteNameClicked),
+		FCanExecuteAction::CreateSP(this, &SAnimCurveViewer::CanDelete));
+
+	CommandList.MapAction(
+		MenuActions.AddCurve,
+		FExecuteAction::CreateSP(this, &SAnimCurveViewer::OnAddClicked),
+		FCanExecuteAction());
 }
 
-void SAnimCurveViewer::OnToggleShowingAllCurves(ECheckBoxState NewState)
-{
-	bShowAllCurves = (NewState == ECheckBoxState::Checked);
 
-	RefreshCurveList();
+bool SAnimCurveViewer::IsCurveFilterEnabled() const
+{
+	return true;
 }
 
 void SAnimCurveViewer::RefreshCachePreviewInstance()
@@ -524,8 +647,7 @@ void SAnimCurveViewer::RefreshCachePreviewInstance()
 	{
 		CachedPreviewInstance->RemoveDelegate_AddCustomAnimationCurve(OnAddAnimationCurveDelegate);
 	}
-
-	CachedPreviewInstance = Cast<UAnimInstance>(PersonaPtr.Pin()->GetPreviewAnimInstance());
+	CachedPreviewInstance = Cast<UAnimInstance>(PreviewScenePtr.Pin()->GetPreviewMeshComponent()->GetAnimInstance());
 
 	if (CachedPreviewInstance)
 	{
@@ -563,7 +685,7 @@ TSharedRef<ITableRow> SAnimCurveViewer::GenerateAnimCurveRow(TSharedPtr<FDisplay
 	check( InInfo.IsValid() );
 
 	return
-		SNew( SAnimCurveListRow, OwnerTable )
+		SNew( SAnimCurveListRow, OwnerTable, PreviewScenePtr.Pin().ToSharedRef() )
 		.Item( InInfo )
 		.AnimCurveViewerPtr( SharedThis(this) );
 }
@@ -571,38 +693,15 @@ TSharedRef<ITableRow> SAnimCurveViewer::GenerateAnimCurveRow(TSharedPtr<FDisplay
 TSharedPtr<SWidget> SAnimCurveViewer::OnGetContextMenuContent() const
 {
 	const bool bShouldCloseWindowAfterMenuSelection = true;
-	FMenuBuilder MenuBuilder( bShouldCloseWindowAfterMenuSelection, NULL);
+	FMenuBuilder MenuBuilder( bShouldCloseWindowAfterMenuSelection, UICommandList);
 
-	MenuBuilder.BeginSection("AnimCurveAction", LOCTEXT( "MorphsAction", "Selected Item Actions" ) );
+	const FCurveViewerCommands& Actions = FCurveViewerCommands::Get();
 
-	{
-		FUIAction Action;
-		// Rename selection
-		{
-			Action.ExecuteAction = FExecuteAction::CreateSP(this, &SAnimCurveViewer::OnRenameClicked);
-			Action.CanExecuteAction = FCanExecuteAction::CreateSP(this, &SAnimCurveViewer::CanRename);
-			const FText Label = LOCTEXT("RenameSmartNameLabel", "Rename");
-			const FText ToolTipText = LOCTEXT("RenameSmartNameToolTip", "Rename the selected item");
-			MenuBuilder.AddMenuEntry(Label, ToolTipText, FSlateIcon(), Action);
-		}
-		// Delete a name
-		{
-			Action.ExecuteAction = FExecuteAction::CreateSP(this, &SAnimCurveViewer::OnDeleteNameClicked);
-			Action.CanExecuteAction = FCanExecuteAction::CreateSP(this, &SAnimCurveViewer::CanDelete);
-			const FText Label = LOCTEXT("DeleteSmartNameLabel", "Delete");
-			const FText ToolTipText = LOCTEXT("DeleteSmartNameToolTip", "Delete the selected item");
-			MenuBuilder.AddMenuEntry(Label, ToolTipText, FSlateIcon(), Action);
-		}
-		// Add a name
-		MenuBuilder.AddMenuSeparator();
-		{
-			Action.ExecuteAction = FExecuteAction::CreateSP(this, &SAnimCurveViewer::OnAddClicked);
-			Action.CanExecuteAction = nullptr;
-			const FText Label = LOCTEXT("AddSmartNameLabel", "Add...");
-			const FText ToolTipText = LOCTEXT("AddSmartNameToolTip", "Add an entry to the skeleton");
-			MenuBuilder.AddMenuEntry(Label, ToolTipText, FSlateIcon(), Action);
-		}
-	}
+	MenuBuilder.BeginSection("AnimCurveAction", LOCTEXT( "CurveAction", "Curve Actions" ) );
+
+	MenuBuilder.AddMenuEntry(FGenericCommands::Get().Rename, NAME_None, LOCTEXT("RenameSmartNameLabel", "Rename Curve"), LOCTEXT("RenameSmartNameToolTip", "Rename the selected curve"));
+	MenuBuilder.AddMenuEntry(FGenericCommands::Get().Delete, NAME_None, LOCTEXT("DeleteSmartNameLabel", "Delete Curve"), LOCTEXT("DeleteSmartNameToolTip", "Delete the selected curve"));
+	MenuBuilder.AddMenuEntry(Actions.AddCurve);
 
 	MenuBuilder.EndSection();
 
@@ -642,13 +741,12 @@ void SAnimCurveViewer::OnAddClicked()
 
 void SAnimCurveViewer::Tick(const FGeometry& AllottedGeometry, const double InCurrentTime, const float InDeltaTime)
 {
-	if (PersonaPtr.IsValid())
+	UDebugSkelMeshComponent* MeshComponent = PreviewScenePtr.Pin()->GetPreviewMeshComponent();
+	UAnimInstance* AnimInstance = MeshComponent->GetAnimInstance();
+
+	if (AnimInstance)
 	{
-		UAnimInstance* AnimInstance = PersonaPtr.Pin()->GetPreviewAnimInstance();
-		if (AnimInstance)
-		{
-			RefreshCurveList();
-		}
+		RefreshCurveList();
 	}
 }
 
@@ -661,7 +759,7 @@ void SAnimCurveViewer::CreateNewNameEntry(const FText& CommittedText, ETextCommi
 		{
 			FName NewName = FName(*CommittedText.ToString());
 			FSmartName NewCurveName;
-			if (CurrentSkeleton->AddSmartNameAndModify(ContainerName, NewName, NewCurveName))
+			if (EditableSkeletonPtr.Pin()->AddSmartname(ContainerName, NewName, NewCurveName))
 			{
 				// Successfully added
 				RefreshCurveList();
@@ -670,14 +768,10 @@ void SAnimCurveViewer::CreateNewNameEntry(const FText& CommittedText, ETextCommi
 	}
 }
 
+
 const FSmartNameMapping* SAnimCurveViewer::GetAnimCurveMapping()
 {
-	const FSmartNameMapping* Mapping = nullptr;
-	if (CurrentSkeleton)
-	{
-		Mapping = CurrentSkeleton->GetSmartNameContainer(ContainerName);
-	}
-	return Mapping;
+	return EditableSkeletonPtr.Pin()->GetSkeleton().GetSmartNameContainer(ContainerName);
 }
 
 int32 FindIndexOfAnimCurveInfo(TArray<TSharedPtr<FDisplayedAnimCurveInfo>>& AnimCurveInfos, const FName& CurveName)
@@ -695,78 +789,91 @@ int32 FindIndexOfAnimCurveInfo(TArray<TSharedPtr<FDisplayedAnimCurveInfo>>& Anim
 
 void SAnimCurveViewer::CreateAnimCurveList( const FString& SearchText )
 {
-	TSharedPtr<FPersona> Persona = PersonaPtr.Pin();
-	if (Persona.IsValid() && CurrentSkeleton)
+	const FSmartNameMapping* Mapping = GetAnimCurveMapping();
+	if (Mapping)
 	{
-		const FSmartNameMapping* Mapping = GetAnimCurveMapping();
-		if (Mapping)
-		{
-			TArray<FSmartNameMapping::UID> UidList;
-			Mapping->FillUidArray(UidList);
+		TArray<SmartName::UID_Type> UidList;
+		Mapping->FillUidArray(UidList);
 
-			// Get set of active curves
-			TMap<FName, float> ActiveCurves;
-			UAnimInstance* AnimInstance = Persona->GetPreviewAnimInstance();
-			if (!bShowAllCurves && AnimInstance)
+		// Get set of active curves
+		TMap<FName, float> ActiveCurves;
+		UDebugSkelMeshComponent* MeshComponent = PreviewScenePtr.Pin()->GetPreviewMeshComponent();
+		UAnimInstance* AnimInstance = MeshComponent->GetAnimInstance();
+		if (AnimInstance)
+		{
+			// attribute curve should contain everything
+			// so only search other container if attribute is off
+			if (CurrentCurveFlag & ACEF_DriveAttribute)
 			{
-				AnimInstance->GetAnimationCurveList(CurrentCurveFlag, ActiveCurves);
+				AnimInstance->GetAnimationCurveList(EAnimCurveType::AttributeCurve, ActiveCurves);
 			}
-
-			// Iterate through all curves..
-			for (FSmartNameMapping::UID Uid : UidList)
+			else
 			{
-				bool bAddToList = true;
-
-				FSmartName SmartName;
-				Mapping->FindSmartNameByUID(Uid, SmartName);
-
-				// See if we pass the search filter
-				if (!FilterText.IsEmpty())
+				if (CurrentCurveFlag & ACEF_DriveMorphTarget)
 				{
-					if (!SmartName.DisplayName.ToString().Contains(*FilterText.ToString()))
-					{
-						bAddToList = false;
-					}
+					AnimInstance->GetAnimationCurveList(EAnimCurveType::MorphTargetCurve, ActiveCurves);
 				}
-
-				// If we passed that, see if we are filtering to only active
-				if (bAddToList && !bShowAllCurves)
+				if (CurrentCurveFlag & ACEF_DriveMaterial)
 				{
-					bAddToList = ActiveCurves.Contains(SmartName.DisplayName);
+					AnimInstance->GetAnimationCurveList(EAnimCurveType::MaterialCurve, ActiveCurves);
 				}
+			}
+		}
+
+		// Iterate through all curves..
+		for (SmartName::UID_Type Uid : UidList)
+		{
+			bool bAddToList = true;
+
+			FSmartName SmartName;
+			Mapping->FindSmartNameByUID(Uid, SmartName);
 
-				// If we still want to add
-				if (bAddToList)
+			// See if we pass the search filter
+			if (!FilterText.IsEmpty())
+			{
+				if (!SmartName.DisplayName.ToString().Contains(*FilterText.ToString()))
 				{
-					// If not already in list, add it
-					if (FindIndexOfAnimCurveInfo(AnimCurveList, SmartName.DisplayName) == INDEX_NONE)
-					{
-						AnimCurveList.Add(FDisplayedAnimCurveInfo::Make(CurrentSkeleton, ContainerName, SmartName));
-					}
+					bAddToList = false;
 				}
-				// Don't want in list
-				else
+			}
+
+			// If we passed that, see if we are filtering to only active
+			if (bAddToList)
+			{
+				bAddToList = ActiveCurves.Contains(SmartName.DisplayName);
+			}
+
+			// If we still want to add
+			if (bAddToList)
+			{
+				// If not already in list, add it
+				if (FindIndexOfAnimCurveInfo(AnimCurveList, SmartName.DisplayName) == INDEX_NONE)
 				{
-					// If already in list, remove it
-					int32 CurrentIndex = FindIndexOfAnimCurveInfo(AnimCurveList, SmartName.DisplayName);
-					if (CurrentIndex != INDEX_NONE)
-					{
-						AnimCurveList.RemoveAt(CurrentIndex);
-					}
+					AnimCurveList.Add(FDisplayedAnimCurveInfo::Make(EditableSkeletonPtr, ContainerName, SmartName));
 				}
 			}
-
-			// Sort final list
-			struct FSortSmartNamesAlphabetically
+			// Don't want in list
+			else
 			{
-				bool operator()(const TSharedPtr<FDisplayedAnimCurveInfo>& A, const TSharedPtr<FDisplayedAnimCurveInfo>& B) const
+				// If already in list, remove it
+				int32 CurrentIndex = FindIndexOfAnimCurveInfo(AnimCurveList, SmartName.DisplayName);
+				if (CurrentIndex != INDEX_NONE)
 				{
-					return (A.Get()->SmartName.DisplayName.Compare(B.Get()->SmartName.DisplayName) < 0);
+					AnimCurveList.RemoveAt(CurrentIndex);
 				}
-			};
-
-			AnimCurveList.Sort(FSortSmartNamesAlphabetically());
+			}
 		}
+
+		// Sort final list
+		struct FSortSmartNamesAlphabetically
+		{
+			bool operator()(const TSharedPtr<FDisplayedAnimCurveInfo>& A, const TSharedPtr<FDisplayedAnimCurveInfo>& B) const
+			{
+				return (A.Get()->SmartName.DisplayName.Compare(B.Get()->SmartName.DisplayName) < 0);
+			}
+		};
+
+		AnimCurveList.Sort(FSortSmartNamesAlphabetically());
 	}
 
 	AnimCurveListView->RequestListRefresh();
@@ -774,32 +881,11 @@ void SAnimCurveViewer::CreateAnimCurveList( const FString& SearchText )
 
 void SAnimCurveViewer::CreateAnimCurveTypeList(TSharedRef<SHorizontalBox> HorizontalBox)
 {
-	// Add toggle button for 'all curves'
-	HorizontalBox->AddSlot()
-	.AutoWidth()
-	.Padding(3, 1)
-	[
-		SNew(SCheckBox)
-		.Style(FEditorStyle::Get(), "ToggleButtonCheckbox")
-		.ToolTipText(LOCTEXT("ShowAllCurvesTooltip", "Show all curves, or only active curves."))
-		.Type(ESlateCheckBoxType::ToggleButton)
-		.IsChecked(this, &SAnimCurveViewer::IsShowingAllCurves)
-		.OnCheckStateChanged(this, &SAnimCurveViewer::OnToggleShowingAllCurves)
-		.Padding(4)
-		.Content()
-		[
-			SNew(STextBlock)
-			.Text(LOCTEXT("ShowAllCurves", "All Curves"))
-		]
-	];
-
 	// Add check box for each curve type flag
 	TArray<int32> CurveFlagsToList;
-	CurveFlagsToList.Add(ACF_DriveMorphTarget);
-	CurveFlagsToList.Add(ACF_DriveAttribute);
-	CurveFlagsToList.Add(ACF_DriveMaterial);
-	// @fixme: when we move curve types to skeleton fix this
-//	CurveFlagsToList.Add(ACF_Metadata);
+	CurveFlagsToList.Add(ACEF_DriveMorphTarget);
+	CurveFlagsToList.Add(ACEF_DriveAttribute);
+	CurveFlagsToList.Add(ACEF_DriveMaterial);
 
 	for (int32 CurveFlagIndex = 0; CurveFlagIndex < CurveFlagsToList.Num(); ++CurveFlagIndex)
 	{
@@ -855,7 +941,7 @@ void SAnimCurveViewer::ApplyCustomCurveOverride(UAnimInstance* AnimInstance) con
 	for (auto Iter = OverrideCurves.CreateConstIterator(); Iter; ++Iter)
 	{ 
 		// @todo we might want to save original curve flags? or just change curve to apply flags only
-		AnimInstance->AddCurveValue(Iter.Key(), Iter.Value(), CurrentCurveFlag);
+		AnimInstance->AddCurveValue(Iter.Key(), Iter.Value());
 	}
 }
 
@@ -870,10 +956,13 @@ void SAnimCurveViewer::OnNameCommitted(const FText& InNewName, ETextCommit::Type
 	if (Mapping)
 	{
 		FName NewName(*InNewName.ToString());
-		if (!Mapping->Exists(NewName))
+		if (NewName == Item->SmartName.DisplayName)
 		{
-			FScopedTransaction Transaction(LOCTEXT("TransactionRename", "Rename Element"));
-			CurrentSkeleton->RenameSmartnameAndModify(ContainerName, Item->SmartName.UID, NewName);
+			// Do nothing if trying to rename to existing name...
+		}
+		else if (!Mapping->Exists(NewName))
+		{
+			EditableSkeletonPtr.Pin()->RenameSmartname(ContainerName, Item->SmartName.UID, NewName);
 			// remove it, so that it can readd it. 
 			AnimCurveList.Remove(Item);
 		}
@@ -898,145 +987,18 @@ void SAnimCurveViewer::OnNameCommitted(const FText& InNewName, ETextCommit::Type
 void SAnimCurveViewer::OnDeleteNameClicked()
 {
 	TArray< TSharedPtr<FDisplayedAnimCurveInfo> > SelectedItems = AnimCurveListView->GetSelectedItems();
-	TArray<FSmartNameMapping::UID> SelectedUids;
+	TArray<SmartName::UID_Type> SelectedUids;
 
 	for (TSharedPtr<FDisplayedAnimCurveInfo> Item : SelectedItems)
 	{
 		SelectedUids.Add(Item->SmartName.UID);
 	}
 
-	FAssetRegistryModule& AssetModule = FModuleManager::LoadModuleChecked<FAssetRegistryModule>(TEXT("AssetRegistry"));
-	TArray<FAssetData> AnimationAssets;
-	AssetModule.Get().GetAssetsByClass(UAnimSequenceBase::StaticClass()->GetFName(), AnimationAssets, true);
-	FAssetData SkeletonData(CurrentSkeleton);
-	const FString CurrentSkeletonName = SkeletonData.GetExportTextName();
+	EditableSkeletonPtr.Pin()->RemoveSmartnamesAndFixupAnimations(ContainerName, SelectedUids);
 
-	GWarn->BeginSlowTask(LOCTEXT("CollectAnimationsTaskDesc", "Collecting assets..."), true);
-
-	for (int32 Idx = AnimationAssets.Num() - 1; Idx >= 0; --Idx)
-	{
-		FAssetData& Data = AnimationAssets[Idx];
-		bool bRemove = true;
-
-		const FString SkeletonDataTag = Data.GetTagValueRef<FString>("Skeleton");
-		if (!SkeletonDataTag.IsEmpty() && SkeletonDataTag == CurrentSkeletonName)
-		{
-			FString CurveData;
-			if (!Data.GetTagValue(USkeleton::CurveTag, CurveData))
-			{
-				// This asset is old; it hasn't been loaded since before smartnames were added for curves.
-				// unfortunately the only way to delete curves safely is to load old assets to see if they're
-				// using the selected name. We only load what we have to here.
-				UObject* Asset = Data.GetAsset();
-				check(Asset);
-				TArray<UObject::FAssetRegistryTag> Tags;
-				Asset->GetAssetRegistryTags(Tags);
-
-				UObject::FAssetRegistryTag* CurveTag = Tags.FindByPredicate([](const UObject::FAssetRegistryTag& InTag)
-				{
-					return InTag.Name == USkeleton::CurveTag;
-				});
-				CurveData = CurveTag->Value;
-			}
-
-			TArray<FString> ParsedStringUids;
-			CurveData.ParseIntoArray(ParsedStringUids, *USkeleton::CurveTagDelimiter, true);
-
-			for (const FString& UidString : ParsedStringUids)
-			{
-				FSmartNameMapping::UID Uid = (FSmartNameMapping::UID)TCString<TCHAR>::Strtoui64(*UidString, NULL, 10);
-				if (SelectedUids.Contains(Uid))
-				{
-					bRemove = false;
-					break;
-				}
-			}
-		}
-
-		if (bRemove)
-		{
-			AnimationAssets.RemoveAt(Idx);
-		}
-	}
-
-	GWarn->EndSlowTask();
-
-	// AnimationAssets now only contains assets that are using the selected curve(s)
-	if (AnimationAssets.Num() > 0)
-	{
-		FString ConfirmMessage = LOCTEXT("DeleteCurveMessage", "The following assets use the selected curve(s), deleting will remove the curves from these assets. Continue?\n\n").ToString();
-
-		for (FAssetData& Data : AnimationAssets)
-		{
-			ConfirmMessage += Data.AssetName.ToString() + " (" + Data.AssetClass.ToString() + ")\n";
-		}
-
-		FText Title = LOCTEXT("DeleteCurveDialogTitle", "Confirm Deletion");
-		FText Message = FText::FromString(ConfirmMessage);
-		if (FMessageDialog::Open(EAppMsgType::YesNo, Message, &Title) == EAppReturnType::No)
-		{
-			return;
-		}
-		else
-		{
-			// Proceed to delete the curves
-			GWarn->BeginSlowTask(FText::Format(LOCTEXT("DeleteCurvesTaskDesc", "Deleting curve from skeleton {0}"), FText::FromString(CurrentSkeleton->GetName())), true);
-			FScopedTransaction Transaction(LOCTEXT("DeleteCurvesTransactionName", "Delete skeleton curve"));
-
-			// Remove curves from animation assets
-			for (FAssetData& Data : AnimationAssets)
-			{
-				UAnimSequenceBase* Sequence = Cast<UAnimSequenceBase>(Data.GetAsset());
-				USkeleton* MySkeleton = Sequence->GetSkeleton();
-				Sequence->Modify(true);
-				for (FSmartNameMapping::UID Uid : SelectedUids)
-				{
-					FSmartName CurveToDelete;
-					if (MySkeleton->GetSmartNameByUID(ContainerName, Uid, CurveToDelete))
-					{
-						Sequence->RawCurveData.DeleteCurveData(CurveToDelete);
-						Sequence->MarkRawDataAsModified();
-					}
-				}
-			}
-			GWarn->EndSlowTask();
-
-			GWarn->BeginSlowTask(LOCTEXT("RebuildingAnimations", "Rebaking/compressing modified animations"), true);
-
-			// Rebake/compress the animations
-			for (TObjectIterator<UAnimSequence> It; It; ++It)
-			{
-				UAnimSequence* Seq = *It;
-
-				GWarn->StatusUpdate(1, 2, FText::Format(LOCTEXT("RebuildingAnimationsStatus", "Rebuilding {0}"), FText::FromString(Seq->GetName())));
-				Seq->RequestSyncAnimRecompression();
-			}
-			GWarn->EndSlowTask();
-		}
-	}
-
-	if (SelectedUids.Num() > 0)
-	{
-		// Remove names from skeleton
-		CurrentSkeleton->RemoveSmartnamesAndModify(ContainerName, SelectedUids);
-		for (int32 ListIndex = SelectedItems.Num() - 1; ListIndex >= 0 ; --ListIndex)
-		{
-			AnimCurveList.Remove(SelectedItems[ListIndex]);
-		}
-	}
+	AnimCurveList.Empty();
 
 	RefreshCurveList();
-
-	TSharedPtr<FPersona> Persona = PersonaPtr.Pin();
-	if (Persona.IsValid())
-	{
-		if (Persona->GetCurrentMode() == FPersonaModes::AnimationEditMode)
-		{
-			// If we're in animation mode then we need to reinitialise the mode to
-			// refresh the sequencer which will now have a widget it shouldn't have
-			Persona->ReinitMode();
-		}
-	}
 }
 
 bool SAnimCurveViewer::CanDelete()
diff --git a/Engine/Source/Editor/Persona/Private/SAnimCurveViewer.h b/Engine/Source/Editor/Persona/Private/SAnimCurveViewer.h
--- a/Engine/Source/Editor/Persona/Private/SAnimCurveViewer.h
+++ b/Engine/Source/Editor/Persona/Private/SAnimCurveViewer.h
@@ -1,8 +1,8 @@
 // Copyright 1998-2016 Epic Games, Inc. All Rights Reserved.
 
 #pragma once
+
 #include "Animation/AnimInstance.h"
-#include "Persona.h"
 
 class SAnimCurveViewer;
 
@@ -15,28 +15,25 @@ public:
 	FSmartName SmartName;
 	float Weight;
 	bool bAutoFillData;
-	const class USkeleton* Skeleton;		// The skeleton we're associated with
+	TWeakPtr<class IEditableSkeleton> EditableSkeleton;		// The skeleton we're associated with
 	TSharedPtr<SInlineEditableTextBlock> EditableText;	// The editable text box in the list, used to focus from the context menu
 	FName ContainerName;	// The container in the skeleton this name resides in
 
 	/** Static function for creating a new item, but ensures that you can only have a TSharedRef to one */
-	static TSharedRef<FDisplayedAnimCurveInfo> Make(const class USkeleton* InSkeleton, const FName& InContainerName, const FSmartName& InSmartName)
+	static TSharedRef<FDisplayedAnimCurveInfo> Make(TWeakPtr<class IEditableSkeleton> InEditableSkeleton, const FName& InContainerName, const FSmartName& InSmartName)
 	{
-		return MakeShareable(new FDisplayedAnimCurveInfo(InSkeleton, InContainerName, InSmartName));
+		return MakeShareable(new FDisplayedAnimCurveInfo(InEditableSkeleton, InContainerName, InSmartName));
 	}
 
 protected:
 	/** Hidden constructor, always use Make above */
-	FDisplayedAnimCurveInfo(const class USkeleton* InSkeleton, const FName& InContainerName, const FSmartName& InSmartName)
+	FDisplayedAnimCurveInfo(TWeakPtr<class IEditableSkeleton> InEditableSkeleton, const FName& InContainerName, const FSmartName& InSmartName)
 		: SmartName(InSmartName)
 		, Weight( 0 )
 		, bAutoFillData(true)
-		, Skeleton(InSkeleton)
+		, EditableSkeleton(InEditableSkeleton)
 		, ContainerName(InContainerName)
 	{}
-
-	/** Hidden constructor, always use Make above */
-	FDisplayedAnimCurveInfo() {}
 };
 
 typedef SListView< TSharedPtr<FDisplayedAnimCurveInfo> > SAnimCurveListType;
@@ -58,7 +55,7 @@ public:
 
 	SLATE_END_ARGS()
 
-	void Construct(const FArguments& InArgs, const TSharedRef<STableViewBase>& OwnerTableView);
+	void Construct(const FArguments& InArgs, const TSharedRef<STableViewBase>& OwnerTableView, const TSharedRef<IPersonaPreviewScene>& InPreviewScene);
 
 	/** Overridden from SMultiColumnTableRow.  Generates a widget for this column of the tree row. */
 	virtual TSharedRef<SWidget> GenerateWidgetForColumn(const FName& ColumnName) override;
@@ -84,6 +81,10 @@ private:
 	void OnAnimCurveAutoFillChecked(ECheckBoxState InState);
 	ECheckBoxState IsAnimCurveAutoFillChangedChecked() const;
 
+	/* Curve Flag checks for morphtarget or material */
+	void OnAnimCurveTypeBoxChecked(ECheckBoxState InState, bool bMorphTarget);
+	ECheckBoxState IsAnimCurveTypeBoxChangedChecked(bool bMorphTarget) const;
+
 	/** Returns the weight of this curve */
 	float GetWeight() const;
 	/** Returns name of this curve */
@@ -101,6 +102,11 @@ private:
 
 	/** The name and weight of the morph target */
 	TSharedPtr<FDisplayedAnimCurveInfo>	Item;
+
+	/** Preview scene used to update on scrub */
+	TWeakPtr<IPersonaPreviewScene> PreviewScenePtr;
+	/** Returns curve type widget constructed */
+	TSharedRef<SWidget> GetCurveTypeWidget();
 };
 
 //////////////////////////////////////////////////////////////////////////
@@ -110,12 +116,8 @@ class SAnimCurveViewer : public SCompoundWidget
 {
 public:
 	SLATE_BEGIN_ARGS( SAnimCurveViewer )
-		: _Persona()
 	{}
 		
-		/** The Persona that owns this table */
-		SLATE_ARGUMENT( TWeakPtr< FPersona >, Persona )
-
 	SLATE_END_ARGS()
 
 	/**
@@ -124,7 +126,7 @@ public:
 	* @param InArgs - Arguments passed from Slate
 	*
 	*/
-	void Construct( const FArguments& InArgs );
+	void Construct( const FArguments& InArgs, const TSharedRef<class IEditableSkeleton>& InEditableSkeleton, const TSharedRef<class IPersonaPreviewScene>& InPreviewScene, FSimpleMulticastDelegate& InOnCurvesChanged, FSimpleMulticastDelegate& InOnPostUndo );
 
 	/**
 	* Destructor - resets the animation curve
@@ -132,6 +134,9 @@ public:
 	*/
 	virtual ~SAnimCurveViewer();
 
+	/** SWidget interface */
+	virtual FReply OnKeyDown(const FGeometry& MyGeometry, const FKeyEvent& InKeyEvent);
+
 	/**
 	* Is registered with Persona to handle when its preview mesh is changed.
 	*
@@ -217,6 +222,8 @@ public:
 
 private:
 
+	void BindCommands();
+
 	/** Handler for context menus */
 	TSharedPtr<SWidget> OnGetContextMenuContent() const;
 
@@ -245,15 +252,18 @@ private:
 
 	bool IsCurveFilterEnabled() const;
 
+
 	// Adds a new smartname entry to the skeleton in the container we are managing
 	void CreateNewNameEntry(const FText& CommittedText, ETextCommit::Type CommitType);
 
 	/** Get the SmartNameMapping for anim curves */
 	const struct FSmartNameMapping* GetAnimCurveMapping();
 
+	/** Pointer to the preview scene we are bound to */
+	TWeakPtr<class IPersonaPreviewScene> PreviewScenePtr;
 
-	/** Pointer back to the Persona that owns us */
-	TWeakPtr<FPersona> PersonaPtr;
+	/** Pointer to the editable skeleton */
+	TWeakPtr<class IEditableSkeleton> EditableSkeletonPtr;
 
 	/** Box to filter to a specific morph target name */
 	TSharedPtr<SSearchBox>	NameFilterBox;
@@ -262,10 +272,7 @@ private:
 	TArray< TSharedPtr<FDisplayedAnimCurveInfo> > AnimCurveList;
 
 	/** The skeletal mesh that we grab the animation curve from */
-	UAnimInstance* CachedPreviewInstance;
-
-	/** The current skeleton being edited by out Persona instance */
-	class USkeleton* CurrentSkeleton;									
+	UAnimInstance* CachedPreviewInstance;						
 
 	/** Widget used to display the list of animation curve */
 	TSharedPtr<SAnimCurveListType> AnimCurveListView;
@@ -276,13 +283,13 @@ private:
 	/** Current text typed into NameFilterBox */
 	FText FilterText;
 
-	/** Are we showing all curves, or only active ones. */
-	bool bShowAllCurves;
-
 	int32 CurrentCurveFlag;
 
 	TMap<FName, float> OverrideCurves;
 
+	/** Commands that are bound to delegates*/
+	TSharedPtr<FUICommandList> UICommandList;
+
 	/** Curve delegate */
 	FOnAddCustomAnimationCurves OnAddAnimationCurveDelegate;
 
diff --git a/Engine/Source/Editor/Persona/Private/SAnimEditorBase.cpp b/Engine/Source/Editor/Persona/Private/SAnimEditorBase.cpp
--- a/Engine/Source/Editor/Persona/Private/SAnimEditorBase.cpp
+++ b/Engine/Source/Editor/Persona/Private/SAnimEditorBase.cpp
@@ -11,6 +11,8 @@
 #include "SAnimNotifyPanel.h"
 #include "SAnimCurvePanel.h"
 #include "AnimPreviewInstance.h"
+#include "Animation/BlendSpaceBase.h"
+
 
 #define LOCTEXT_NAMESPACE "AnimEditorBase"
 
@@ -22,9 +24,10 @@ TSharedRef<SWidget> SAnimEditorBase::CreateDocumentAnchor()
 	return IDocumentation::Get()->CreateAnchor(TEXT("Engine/Animation/Sequences"));
 }
 
-void SAnimEditorBase::Construct(const FArguments& InArgs)
+void SAnimEditorBase::Construct(const FArguments& InArgs, const TSharedRef<class IPersonaPreviewScene>& InPreviewScene)
 {
-	PersonaPtr = InArgs._Persona;
+	PreviewScenePtr = InPreviewScene;
+	OnObjectsSelected = InArgs._OnObjectsSelected;
 
 	SetInputViewRange(0, GetSequenceLength());
 
@@ -34,42 +37,24 @@ void SAnimEditorBase::Construct(const FArguments& InArgs)
 	[
 		SAssignNew(AnimVerticalBox, SVerticalBox)
 		+SVerticalBox::Slot()
-		.AutoHeight()
+		.FillHeight(1)
 		[
-			// Header, shows name of timeline we are editing
 			SNew(SBorder)
-			. BorderImage( FEditorStyle::GetBrush( TEXT("Graph.TitleBackground") ) )
-			. HAlign(HAlign_Center)
+			.BorderImage( FEditorStyle::GetBrush("ToolPanel.GroupBorder") )
 			[
-				SNew(SHorizontalBox)
-				+ SHorizontalBox::Slot()
-				.AutoWidth()
-				. VAlign(VAlign_Center)
+				SNew(SOverlay)
+				+SOverlay::Slot()
 				[
-					SNew(STextBlock)
-					.Font( FSlateFontInfo(FPaths::EngineContentDir() / TEXT("Slate/Fonts/Roboto-Regular.ttf"), 14 ) )
-					.ColorAndOpacity( FLinearColor(1,1,1,0.5) )
-					.Text( this, &SAnimEditorBase::GetEditorObjectName )
+					SAssignNew(NonScrollEditorPanels, SVerticalBox)
 				]
-				+ SHorizontalBox::Slot()
-					.HAlign(HAlign_Left)
-					.VAlign(VAlign_Center)
+				+SOverlay::Slot()
+				[
+					SNew( SScrollBox )
+					.Visibility(EVisibility::SelfHitTestInvisible)
+					+SScrollBox::Slot() 
 					[
-						CreateDocumentAnchor()
+						SAssignNew (EditorPanels, SVerticalBox)
 					]
-			]
-		]
-
-		+SVerticalBox::Slot()
-		.FillHeight(1)
-		[
-			SNew(SBorder)
-			.BorderImage( FEditorStyle::GetBrush("ToolPanel.GroupBorder") )
-			[
-				SNew( SScrollBox )
-				+SScrollBox::Slot() 
-				[
-					SAssignNew (EditorPanels, SVerticalBox)
 				]
 			]
 		]
@@ -172,32 +157,23 @@ void SAnimEditorBase::Construct(const FArguments& InArgs)
 	}
 
 	// If we are an anim sequence, add scrub panel as well
-	UAnimSequenceBase* SeqBase = Cast<UAnimSequenceBase>(GetEditorObject());
-	if (SeqBase)
-	{
-		AnimVerticalBox->AddSlot()
-		.AutoHeight()
-		.VAlign(VAlign_Bottom)
+	AnimVerticalBox->AddSlot()
+	.AutoHeight()
+	.VAlign(VAlign_Bottom)
+	[
+		SNew(SHorizontalBox)
+		+SHorizontalBox::Slot()
+		.FillWidth(1)
 		[
-			SNew(SHorizontalBox)
-			+SHorizontalBox::Slot()
-			.FillWidth(1)
-			[
-				ConstructAnimScrubPanel()
-			]
-		];
-	}
-
+			ConstructAnimScrubPanel()
+		]
+	];
 }
 
 TSharedRef<class SAnimationScrubPanel> SAnimEditorBase::ConstructAnimScrubPanel()
 {
-	UAnimSequenceBase* AnimSeqBase = Cast<UAnimSequenceBase>(GetEditorObject());
-	check(AnimSeqBase);
-
-	return SAssignNew(AnimScrubPanel, SAnimationScrubPanel)
-		.Persona(PersonaPtr)
-		.LockedSequence(AnimSeqBase)
+	return SAssignNew( AnimScrubPanel, SAnimationScrubPanel, PreviewScenePtr.Pin().ToSharedRef() )
+		.LockedSequence(Cast<UAnimSequenceBase>(GetEditorObject()))
 		.ViewInputMin(this, &SAnimEditorBase::GetViewMinInput)
 		.ViewInputMax(this, &SAnimEditorBase::GetViewMaxInput)
 		.OnSetInputViewRange(this, &SAnimEditorBase::SetInputViewRange)
@@ -221,7 +197,10 @@ UObject* SAnimEditorBase::ShowInDetailsView( UClass* EdClass )
 		{
 			UEditorAnimBaseObj *EdObj = Cast<UEditorAnimBaseObj>(Obj);
 			InitDetailsViewEditorObject(EdObj);
-			PersonaPtr.Pin()->SetDetailObject(EdObj);
+
+			TArray<UObject*> Objects;
+			Objects.Add(EdObj);
+			OnObjectsSelected.ExecuteIfBound(Objects);
 		}
 	}
 	return Obj;
@@ -229,7 +208,8 @@ UObject* SAnimEditorBase::ShowInDetailsView( UClass* EdClass )
 
 void SAnimEditorBase::ClearDetailsView()
 {
-	PersonaPtr.Pin()->SetDetailObject(NULL);
+	TArray<UObject*> Objects;
+	OnObjectsSelected.ExecuteIfBound(Objects);
 }
 
 FText SAnimEditorBase::GetEditorObjectName() const
@@ -285,24 +265,14 @@ bool SAnimEditorBase::ClampToEndTime(float NewEndTime)
 	return (SequenceLength > 0.f && NewEndTime < SequenceLength);
 }
 
-void SAnimEditorBase::OnSelectionChanged(const FGraphPanelSelectionSet& SelectedItems)
+void SAnimEditorBase::OnSelectionChanged(const TArray<UObject*>& SelectedItems)
 {
-	if (SelectedItems.Num() == 0)
-	{
-		// If nothing is selected, unselect it instead of selecting current editor object. 
-		// That means, it shows same information as Anim Detail Panel
-		PersonaPtr.Pin()->UpdateSelectionDetails(NULL, LOCTEXT("Edit Sequence", "Edit Sequence"));
-	}
-	else
-	{
-		// Edit selected notifications
-		PersonaPtr.Pin()->FocusInspectorOnGraphSelection(SelectedItems);
-	}
+	OnObjectsSelected.ExecuteIfBound(SelectedItems);
 }
 
 class UAnimSingleNodeInstance* SAnimEditorBase::GetPreviewInstance() const
 {
-	return (PersonaPtr.Pin()->GetPreviewMeshComponent())? PersonaPtr.Pin()->GetPreviewMeshComponent()->PreviewInstance:NULL;
+	return (GetPreviewScene()->GetPreviewMeshComponent()) ? GetPreviewScene()->GetPreviewMeshComponent()->PreviewInstance : nullptr;
 }
 
 float SAnimEditorBase::GetScrubValue() const
@@ -375,27 +345,30 @@ FText SAnimEditorBase::GetCurrentPercentage() const
 FText SAnimEditorBase::GetCurrentFrame() const
 {
 	float Percentage = GetPercentageInternal();
-	float NumFrames = 0;
+	float LastFrame = 0;
 	
 	if (UAnimSequenceBase* AnimSeqBase = Cast<UAnimSequenceBase>(GetEditorObject()))
 	{
-		NumFrames = AnimSeqBase->GetNumberOfFrames();
+		LastFrame = AnimSeqBase->GetNumberOfFrames() - 1;
 	}
 
 	static const FNumberFormattingOptions FractionNumberFormat = FNumberFormattingOptions()
 		.SetMinimumFractionalDigits(2)
 		.SetMaximumFractionalDigits(2);
-	return FText::Format(LOCTEXT("FractionKeysFmt", "{0} / {1} (key(s))"), FText::AsNumber(NumFrames * Percentage, &FractionNumberFormat), FText::AsNumber((int32)NumFrames));
+	return FText::Format(LOCTEXT("FractionKeysFmt", "{0} / {1} Frame"), FText::AsNumber(LastFrame * Percentage, &FractionNumberFormat), FText::AsNumber((int32)LastFrame));
 }
 
 float SAnimEditorBase::GetSequenceLength() const
 {
-	UAnimSequenceBase* AnimSeqBase = Cast<UAnimSequenceBase>(GetEditorObject());
-	
-	if (AnimSeqBase)
+	if (UAnimSequenceBase* AnimSeqBase = Cast<UAnimSequenceBase>(GetEditorObject()))
 	{
 		return AnimSeqBase->SequenceLength;
 	}
+	else if (UBlendSpaceBase* BlendSpaceBase = Cast<UBlendSpaceBase>(GetEditorObject()))
+	{
+		// Blendspaces use normalized time, so we just return 1 here
+		return 1.0f;
+	}
 	
 	return 0.f;
 }
diff --git a/Engine/Source/Editor/Persona/Private/SAnimEditorBase.h b/Engine/Source/Editor/Persona/Private/SAnimEditorBase.h
--- a/Engine/Source/Editor/Persona/Private/SAnimEditorBase.h
+++ b/Engine/Source/Editor/Persona/Private/SAnimEditorBase.h
@@ -27,17 +27,17 @@ public:
 class SAnimEditorBase : public SCompoundWidget, public FGCObject
 {
 public:
-	SLATE_BEGIN_ARGS( SAnimEditorBase )
-		: _Persona()
-		, _DisplayAnimInfoBar(true)
-		{}
+	SLATE_BEGIN_ARGS(SAnimEditorBase)
+		: _DisplayAnimInfoBar(true)
+	{}
 
-		SLATE_ARGUMENT( TSharedPtr<FPersona>, Persona )
+	SLATE_EVENT(FOnObjectsSelected, OnObjectsSelected)
+
+	SLATE_ARGUMENT( bool, DisplayAnimInfoBar )
 
-		SLATE_ARGUMENT( bool, DisplayAnimInfoBar )
 	SLATE_END_ARGS()
 
-	void Construct(const FArguments& InArgs);
+	void Construct(const FArguments& InArgs, const TSharedRef<class IPersonaPreviewScene>& InPreviewScene);
 
 	/** Accessors to the current viewed Min/Max input range of the editor */
 	float GetViewMinInput() const { return ViewMinInput; }
@@ -53,7 +53,7 @@ public:
 	void ScrollInputViewRange(float ScreenDelta, FVector2D ScreenViewSize);
 
 	/** Delegate when selection changed **/
-	void OnSelectionChanged(const FGraphPanelSelectionSet& SelectedItems);
+	void OnSelectionChanged(const TArray<UObject*>& SelectedItems);
 
 	/** Get Current Scrub Value **/
 	float GetScrubValue() const;
@@ -72,8 +72,10 @@ public:
 	virtual TSharedRef<SWidget> CreateDocumentAnchor();
 
 protected:
-	/** Persona reference **/
-	TWeakPtr<FPersona> PersonaPtr;
+	/** Preview scene reference **/
+	TWeakPtr<IPersonaPreviewScene> PreviewScenePtr;
+
+	TSharedRef<IPersonaPreviewScene> GetPreviewScene() const { return PreviewScenePtr.Pin().ToSharedRef(); }
 
 	/** Allows derived classes to create different animation scrub panel */
 	virtual TSharedRef<class SAnimationScrubPanel> ConstructAnimScrubPanel();
@@ -85,7 +87,7 @@ protected:
 	float GetSequenceLength() const;
 
 	/** Get the sequence that is currently being edited */
-	virtual UAnimationAsset* GetEditorObject() const PURE_VIRTUAL(SAnimEditorBase::GetEditorObject, return NULL;);
+	virtual UAnimationAsset* GetEditorObject() const = 0;
 
 	/** Get Name of Object being edited **/
 	FText GetEditorObjectName() const;
@@ -120,12 +122,18 @@ protected:
 	/** The slate container that the editor panels are placed in */
 	TSharedPtr<SVerticalBox>	 EditorPanels;
 
+	/** Slate container used to add controls that are not embedded in a scroll box */
+	TSharedPtr<SVerticalBox> NonScrollEditorPanels;
+
 	/** The editors Animation Scrub Panel */
 	TSharedPtr<class SAnimationScrubPanel> AnimScrubPanel;
 
 	/** Tracks objects created for the details panel */
 	FEditorObjectTracker EditorObjectTracker;
 
+	/** Delegate called to select objects */
+	FOnObjectsSelected OnObjectsSelected;
+
 private:
 	float GetPercentageInternal() const;
 };
diff --git a/Engine/Source/Editor/Persona/Private/SAnimMontagePanel.cpp b/Engine/Source/Editor/Persona/Private/SAnimMontagePanel.cpp
--- a/Engine/Source/Editor/Persona/Private/SAnimMontagePanel.cpp
+++ b/Engine/Source/Editor/Persona/Private/SAnimMontagePanel.cpp
@@ -4,7 +4,6 @@
 #include "SAnimMontagePanel.h"
 #include "ScopedTransaction.h"
 #include "SCurveEditor.h"
-#include "SAnimationSequenceBrowser.h"
 #include "SAnimSegmentsPanel.h"
 #include "SMontageEditor.h"
 #include "Editor/UnrealEd/Public/DragAndDrop/AssetDragDropOp.h"
@@ -18,7 +17,7 @@
 //////////////////////////////////////////////////////////////////////////
 // SAnimMontagePanel
 
-void SAnimMontagePanel::Construct(const FArguments& InArgs)
+void SAnimMontagePanel::Construct(const FArguments& InArgs, FSimpleMulticastDelegate& OnAnimNotifiesChanged, FSimpleMulticastDelegate& OnSectionsChanged)
 {
 	SAnimTrackPanel::Construct( SAnimTrackPanel::FArguments()
 		.WidgetWidth(InArgs._WidgetWidth)
@@ -28,8 +27,8 @@ void SAnimMontagePanel::Construct(const FArguments& InArgs)
 		.InputMax(InArgs._InputMax)
 		.OnSetInputViewRange(InArgs._OnSetInputViewRange));
 
-	Persona = InArgs._Persona;
 	Montage = InArgs._Montage;
+	OnInvokeTab = InArgs._OnInvokeTab;
 	MontageEditor = InArgs._MontageEditor;
 	SectionTimingNodeVisibility = InArgs._SectionTimingNodeVisibility;
 	bChildAnimMontage = InArgs._bChildAnimMontage;
@@ -55,26 +54,12 @@ void SAnimMontagePanel::Construct(const FArguments& InArgs)
 		]
 	];
 
-	TSharedPtr<FPersona> SharedPersona = Persona.Pin();
-	if(SharedPersona.IsValid())
-	{
-		SharedPersona->RegisterOnChangeAnimNotifies(FPersona::FOnAnimNotifiesChanged::CreateSP(this, &SAnimMontagePanel::Update));
-		SharedPersona->RegisterOnSectionsChanged(FPersona::FOnAnimNotifiesChanged::CreateSP(this, &SAnimMontagePanel::Update));
-	}
+	OnAnimNotifiesChanged.Add(FSimpleDelegate::CreateSP(this, &SAnimMontagePanel::Update));
+	OnSectionsChanged.Add(FSimpleDelegate::CreateSP(this, &SAnimMontagePanel::Update));
 
 	Update();
 }
 
-SAnimMontagePanel::~SAnimMontagePanel()
-{
-	TSharedPtr<FPersona> SharedPersona = Persona.Pin();
-	if(SharedPersona.IsValid())
-	{
-		SharedPersona->UnregisterOnChangeAnimNotifies(this);
-		SharedPersona->UnregisterOnSectionsChanged(this);
-	}
-}
-
 /** This is the main function that creates the UI widgets for the montage tool.*/
 void SAnimMontagePanel::Update()
 {
@@ -498,10 +483,8 @@ void SAnimMontagePanel::OnSlotListOpening(int32 AnimSlotIndex)
 
 FReply SAnimMontagePanel::OnOpenAnimSlotManager()
 {
-	if (Persona.IsValid())
-	{
-		Persona.Pin()->GetTabManager()->InvokeTab(FPersonaTabs::SkeletonSlotNamesID);
-	}
+	OnInvokeTab.ExecuteIfBound(FPersonaTabs::SkeletonSlotNamesID);
+
 	return FReply::Handled();
 }
 
diff --git a/Engine/Source/Editor/Persona/Private/SAnimMontagePanel.h b/Engine/Source/Editor/Persona/Private/SAnimMontagePanel.h
--- a/Engine/Source/Editor/Persona/Private/SAnimMontagePanel.h
+++ b/Engine/Source/Editor/Persona/Private/SAnimMontagePanel.h
@@ -29,7 +29,6 @@ class SAnimMontagePanel : public SAnimTrackPanel
 public:
 	SLATE_BEGIN_ARGS( SAnimMontagePanel )
 		: _Montage()
-		, _Persona()
 		, _MontageEditor()
 		, _CurrentPosition()
 		, _ViewInputMin()
@@ -39,11 +38,11 @@ public:
 		, _OnSetInputViewRange()
 		, _OnGetScrubValue()
 		, _OnMontageChange()
+		, _OnInvokeTab()
 		, _bChildAnimMontage(false)
 	{}
 
 	SLATE_ARGUMENT( class UAnimMontage*, Montage)
-	SLATE_ARGUMENT( TWeakPtr<FPersona>, Persona)
 	SLATE_ARGUMENT( TWeakPtr<class SMontageEditor>, MontageEditor)
 	SLATE_ARGUMENT( float, WidgetWidth )
 	SLATE_ATTRIBUTE( float, CurrentPosition )
@@ -55,12 +54,11 @@ public:
 	SLATE_EVENT( FOnSetInputViewRange, OnSetInputViewRange )
 	SLATE_EVENT( FOnGetScrubValue, OnGetScrubValue )
 	SLATE_EVENT( FOnAnimObjectChange, OnMontageChange )
+	SLATE_EVENT( FOnInvokeTab, OnInvokeTab )
 	SLATE_ARGUMENT(bool, bChildAnimMontage)
 	SLATE_END_ARGS()
 
-	~SAnimMontagePanel();
-
-	void Construct(const FArguments& InArgs);
+	void Construct(const FArguments& InArgs, FSimpleMulticastDelegate& OnAnimNotifiesChanged, FSimpleMulticastDelegate& OnSectionsChanged);
 	void SetMontage(class UAnimMontage * InMontage);
 
 	// SWidget interface
@@ -105,7 +103,6 @@ public:
 	void RefreshTimingNodes();
 
 private:
-	TWeakPtr<FPersona> Persona;
 	TWeakPtr<SMontageEditor>	MontageEditor;
 	TSharedPtr<SBorder> PanelArea;
 	class UAnimMontage* Montage;
@@ -158,4 +155,7 @@ private:
 	/** Gets the length of the montage we are currently editing
 	 */
 	virtual float GetSequenceLength() const override;
+
+	/** Delegate used to invoke a tab in the containing editor */
+	FOnInvokeTab OnInvokeTab;
 };
\ No newline at end of file
diff --git a/Engine/Source/Editor/Persona/Private/SAnimMontageScrubPanel.cpp b/Engine/Source/Editor/Persona/Private/SAnimMontageScrubPanel.cpp
--- a/Engine/Source/Editor/Persona/Private/SAnimMontageScrubPanel.cpp
+++ b/Engine/Source/Editor/Persona/Private/SAnimMontageScrubPanel.cpp
@@ -7,18 +7,18 @@
 // keep localization same
 #define LOCTEXT_NAMESPACE "AnimationScrubPanel"
 
-void SAnimMontageScrubPanel::Construct( const SAnimMontageScrubPanel::FArguments& InArgs )
+void SAnimMontageScrubPanel::Construct( const SAnimMontageScrubPanel::FArguments& InArgs, const TSharedRef<IPersonaPreviewScene>& InPreviewScene)
 {
 	MontageEditor = InArgs._MontageEditor;
 	SAnimationScrubPanel::Construct( SAnimationScrubPanel::FArguments()
-		.Persona(InArgs._Persona)
 		.LockedSequence(InArgs._LockedSequence)
 		.ViewInputMin(InArgs._ViewInputMin)
 		.ViewInputMax(InArgs._ViewInputMax)
 		.OnSetInputViewRange(InArgs._OnSetInputViewRange)
 		.OnCropAnimSequence(InArgs._OnCropAnimSequence)
 		.OnReZeroAnimSequence(InArgs._OnReZeroAnimSequence)
-		.bAllowZoom(InArgs._bAllowZoom)
+		.bAllowZoom(InArgs._bAllowZoom),
+		InPreviewScene
 		);
 }
 
diff --git a/Engine/Source/Editor/Persona/Private/SAnimMontageScrubPanel.h b/Engine/Source/Editor/Persona/Private/SAnimMontageScrubPanel.h
--- a/Engine/Source/Editor/Persona/Private/SAnimMontageScrubPanel.h
+++ b/Engine/Source/Editor/Persona/Private/SAnimMontageScrubPanel.h
@@ -10,11 +10,9 @@ public:
 
 	SLATE_BEGIN_ARGS(SAnimMontageScrubPanel)
 		: _MontageEditor()
-		, _Persona()
 		, _LockedSequence()
 	{}
 		SLATE_ARGUMENT( TWeakPtr<class SMontageEditor>, MontageEditor)
-		SLATE_ARGUMENT(TWeakPtr<FPersona>, Persona )
 		/** If you'd like to lock to one asset for this scrub control, give this**/
 		SLATE_ARGUMENT(UAnimSequenceBase*, LockedSequence)
 		/** View Input range **/
@@ -33,7 +31,7 @@ public:
 	 * 
 	 * @param InArgs   A declaration from which to construct the widget
 	 */
-	void Construct( const FArguments& InArgs );
+	void Construct( const FArguments& InArgs, const TSharedRef<IPersonaPreviewScene>& InPreviewScene );
 protected:
 	// notifiers 
 	virtual FReply OnClick_Backward_End() override;
diff --git a/Engine/Source/Editor/Persona/Private/SAnimMontageSectionsPanel.cpp b/Engine/Source/Editor/Persona/Private/SAnimMontageSectionsPanel.cpp
--- a/Engine/Source/Editor/Persona/Private/SAnimMontageSectionsPanel.cpp
+++ b/Engine/Source/Editor/Persona/Private/SAnimMontageSectionsPanel.cpp
@@ -6,7 +6,6 @@
 #include "SAnimMontageSectionsPanel.h"
 #include "ScopedTransaction.h"
 #include "SCurveEditor.h"
-#include "SAnimationSequenceBrowser.h"
 #include "SMontageEditor.h"
 #include "SExpandableArea.h"
 
