Copying //UE4/Dev-Framework to //UE4/Dev-Main (Source: //UE4/Dev-Framework @ 3136612)

diff --git a/Engine/Source/Editor/Persona/Private/SAnimNotifyPanel.cpp b/Engine/Source/Editor/Persona/Private/SAnimNotifyPanel.cpp
--- a/Engine/Source/Editor/Persona/Private/SAnimNotifyPanel.cpp
+++ b/Engine/Source/Editor/Persona/Private/SAnimNotifyPanel.cpp
@@ -581,8 +581,7 @@ class SAnimNotifyTrack : public SCompoundWidget
 {
 public:
 	SLATE_BEGIN_ARGS( SAnimNotifyTrack )
-		: _Persona()
-		, _Sequence(NULL)
+		: _Sequence(NULL)
 		, _ViewInputMin()
 		, _ViewInputMax()
 		, _TrackIndex()
@@ -608,7 +607,6 @@ public:
 		, _OnSetInputViewRange()
 		{}
 
-		SLATE_ARGUMENT( TSharedPtr<FPersona>,	Persona )
 		SLATE_ARGUMENT( class UAnimSequenceBase*, Sequence )
 		SLATE_ARGUMENT( TArray<FAnimNotifyEvent *>, AnimNotifies )
 		SLATE_ARGUMENT( TArray<FAnimSyncMarker *>, AnimSyncMarkers)
@@ -639,6 +637,7 @@ public:
 		SLATE_EVENT(FPasteNodes, OnPasteNodes)
 		SLATE_EVENT( FOnSetInputViewRange, OnSetInputViewRange )
 		SLATE_EVENT( FOnGetTimingNodeVisibility, OnGetTimingNodeVisibility )
+		SLATE_EVENT(FOnInvokeTab, OnInvokeTab)
 		SLATE_END_ARGS()
 public:
 
@@ -896,8 +895,9 @@ protected:
 	FReplaceWithNotify						OnReplaceSelectedWithNotify;
 	FReplaceWithBlueprintNotify				OnReplaceSelectedWithBlueprintNotify;
 
+	FOnInvokeTab							OnInvokeTab;
+
 	TSharedPtr<SBorder>						TrackArea;
-	TWeakPtr<FPersona>						PersonaPtr;
 
 	/** Cache the SOverlay used to store all this tracks nodes */
 	TSharedPtr<SOverlay> NodeSlots;
@@ -967,6 +967,7 @@ public:
 	SLATE_EVENT( FPasteNodes, OnPasteNodes )
 	SLATE_EVENT( FOnSetInputViewRange, OnSetInputViewRange )
 	SLATE_EVENT( FOnGetTimingNodeVisibility, OnGetTimingNodeVisibility )
+	SLATE_EVENT(FOnInvokeTab, OnInvokeTab)
 	SLATE_END_ARGS()
 
 	void Construct(const FArguments& InArgs);
@@ -2053,7 +2054,6 @@ void SAnimNotifyTrack::Construct(const FArguments& InArgs)
 	FAnimSequenceEditorCommands::Register();
 	CreateCommands();
 
-	PersonaPtr = InArgs._Persona;
 	Sequence = InArgs._Sequence;
 	ViewInputMin = InArgs._ViewInputMin;
 	ViewInputMax = InArgs._ViewInputMax;
@@ -2082,6 +2082,7 @@ void SAnimNotifyTrack::Construct(const FArguments& InArgs)
 	OnPasteNodes = InArgs._OnPasteNodes;
 	OnSetInputViewRange = InArgs._OnSetInputViewRange;
 	OnGetTimingNodeVisibility = InArgs._OnGetTimingNodeVisibility;
+	OnInvokeTab = InArgs._OnInvokeTab;
 
 	this->ChildSlot
 	[
@@ -2121,12 +2122,14 @@ int32 SAnimNotifyTrack::OnPaint(const FPaintArgs& Args, const FGeometry& Allotte
 
 	// draw line for every 1/4 length
 	FTrackScaleInfo ScaleInfo(ViewInputMin.Get(), ViewInputMax.Get(), 0.f, 0.f, AllottedGeometry.Size);
-	if (Sequence->GetNumberOfFrames() > 0 )
+	const int32 NumOfFrames = Sequence->GetNumberOfFrames();
+	if (NumOfFrames > 0)
 	{
 		int32 Divider = SScrubWidget::GetDivider( ViewInputMin.Get(), ViewInputMax.Get(), AllottedGeometry.Size, Sequence->SequenceLength, Sequence->GetNumberOfFrames());
+		const FAnimKeyHelper Helper(Sequence->SequenceLength, NumOfFrames);
 
-		float TimePerKey = Sequence->SequenceLength/(float)Sequence->GetNumberOfFrames();
-		for (int32 I=1; I<Sequence->GetNumberOfFrames(); ++I)
+		float TimePerKey = Helper.TimePerKey();
+		for (int32 I=1; I<Helper.GetNumKeys(); ++I)
 		{
 			if ( I % Divider == 0 )
 			{
@@ -3024,11 +3027,7 @@ void SAnimNotifyTrack::OnPasteNotifyClicked(ENotifyPasteMode::Type PasteMode, EN
 
 void SAnimNotifyTrack::OnManageNotifies()
 {
-	TSharedPtr< FPersona > PersonalPin = PersonaPtr.Pin();
-	if( PersonalPin.IsValid() )
-	{
-		PersonalPin->GetTabManager()->InvokeTab( FPersonaTabs::SkeletonAnimNotifiesID );
-	}
+	OnInvokeTab.ExecuteIfBound(FPersonaTabs::SkeletonAnimNotifiesID);
 }
 
 void SAnimNotifyTrack::OnOpenNotifySource(UBlueprint* InSourceBlueprint) const
@@ -3853,7 +3852,6 @@ void SNotifyEdTrack::Construct(const FArguments& InArgs)
 			[
 				// Notification editor panel
 				SAssignNew(NotifyTrack, SAnimNotifyTrack)
-				.Persona(PanelRef->GetPersona().Pin())
 				.Sequence(Sequence)
 				.TrackIndex(TrackIndex)
 				.AnimNotifies(Track.Notifies)
@@ -3882,6 +3880,7 @@ void SNotifyEdTrack::Construct(const FArguments& InArgs)
 				.OnPasteNodes(InArgs._OnPasteNodes)
 				.OnSetInputViewRange(InArgs._OnSetInputViewRange)
 				.OnGetTimingNodeVisibility(InArgs._OnGetTimingNodeVisibility)
+				.OnInvokeTab(InArgs._OnInvokeTab)
 			]
 
 			+SHorizontalBox::Slot()
@@ -3943,7 +3942,7 @@ void FAnimNotifyPanelCommands::RegisterCommands()
 //////////////////////////////////////////////////////////////////////////
 // SAnimNotifyPanel
 
-void SAnimNotifyPanel::Construct(const FArguments& InArgs)
+void SAnimNotifyPanel::Construct(const FArguments& InArgs, FSimpleMulticastDelegate& OnPostUndo)
 {
 	SAnimTrackPanel::Construct( SAnimTrackPanel::FArguments()
 		.WidgetWidth(InArgs._WidgetWidth)
@@ -3953,9 +3952,10 @@ void SAnimNotifyPanel::Construct(const FArguments& InArgs)
 		.InputMax(InArgs._InputMax)
 		.OnSetInputViewRange(InArgs._OnSetInputViewRange));
 
-	PersonaPtr = InArgs._Persona;
 	Sequence = InArgs._Sequence;
 	MarkerBars = InArgs._MarkerBars;
+	OnAnimNotifiesChanged = InArgs._OnAnimNotifiesChanged;
+	OnInvokeTab = InArgs._OnInvokeTab;
 
 	FAnimNotifyPanelCommands::Register();
 	BindCommands();
@@ -3965,8 +3965,8 @@ void SAnimNotifyPanel::Construct(const FArguments& InArgs)
 	// the function name in that case will need to change
 	Sequence->InitializeNotifyTrack();
 	Sequence->RegisterOnNotifyChanged(UAnimSequenceBase::FOnNotifyChanged::CreateSP( this, &SAnimNotifyPanel::RefreshNotifyTracks )  );
-	PersonaPtr.Pin()->RegisterOnPostUndo(FPersona::FOnPostUndo::CreateSP( this, &SAnimNotifyPanel::PostUndo ) );
-	PersonaPtr.Pin()->RegisterOnGenericDelete(FPersona::FOnDeleteGeneric::CreateSP(this, &SAnimNotifyPanel::OnDeletePressed));
+
+	OnPostUndo.Add(FSimpleDelegate::CreateSP( this, &SAnimNotifyPanel::PostUndo ) );
 
 	CurrentPosition = InArgs._CurrentPosition;
 	OnSelectionChanged = InArgs._OnSelectionChanged;
@@ -4036,11 +4036,6 @@ void SAnimNotifyPanel::Construct(const FArguments& InArgs)
 SAnimNotifyPanel::~SAnimNotifyPanel()
 {
 	Sequence->UnregisterOnNotifyChanged(this);
-	if (PersonaPtr.IsValid())
-	{
-		PersonaPtr.Pin()->UnregisterOnPostUndo(this);
-		PersonaPtr.Pin()->UnregisterOnGenericDelete(this);
-	}
 
 	FCoreUObjectDelegates::OnObjectPropertyChanged.Remove(OnPropertyChangedHandleDelegateHandle);
 }
@@ -4133,7 +4128,8 @@ void SAnimNotifyPanel::Update()
 	{
 		Sequence->RefreshCacheData();
 	}
-	PersonaPtr.Pin()->OnAnimNotifiesChanged.Broadcast();
+
+	OnAnimNotifiesChanged.ExecuteIfBound();
 }
 
 void SAnimNotifyPanel::RefreshNotifyTracks()
@@ -4287,8 +4283,8 @@ void SAnimNotifyPanel::DeleteSelectedNodeObjects()
 	}
 
 	// clear selection and update the panel
-	FGraphPanelSelectionSet ObjectSet;
-	OnSelectionChanged.ExecuteIfBound(ObjectSet);
+	TArray<UObject*> Objects;
+	OnSelectionChanged.ExecuteIfBound(Objects);
 
 	Update();
 }
@@ -4309,8 +4305,7 @@ void SAnimNotifyPanel::SetSequence(class UAnimSequenceBase *	InSequence)
 void SAnimNotifyPanel::OnTrackSelectionChanged()
 {
 	// Need to collect selection info from all tracks
-	FGraphPanelSelectionSet SelectionSet;
-	TArray<UEditorNotifyObject*> NotifyObjects;
+	TArray<UObject*> NotifyObjects;
 
 	for(int32 TrackIdx = 0 ; TrackIdx < NotifyAnimTracks.Num() ; ++TrackIdx)
 	{
@@ -4325,12 +4320,12 @@ void SAnimNotifyPanel::OnTrackSelectionChanged()
 				UEditorNotifyObject* NewNotifyObject = NewObject<UEditorNotifyObject>(GetTransientPackage(), FName(*ObjName), RF_Public | RF_Standalone | RF_Transient);
 				NewNotifyObject->InitFromAnim(Sequence, FOnAnimObjectChange::CreateSP(this, &SAnimNotifyPanel::OnNotifyObjectChanged));
 				NewNotifyObject->InitialiseNotify(NotifyAnimTracks.Num() - TrackIdx - 1, Idx);
-				SelectionSet.Add(NewNotifyObject);
+				NotifyObjects.AddUnique(NewNotifyObject);
 			}
 		}
 	}
 
-	OnSelectionChanged.ExecuteIfBound(SelectionSet);
+	OnSelectionChanged.ExecuteIfBound(NotifyObjects);
 }
 
 void SAnimNotifyPanel::DeselectAllNotifies()
@@ -4341,11 +4336,7 @@ void SAnimNotifyPanel::DeselectAllNotifies()
 	}
 
 	// Broadcast the change so the editor can update
-	TSharedPtr<FPersona> SharedPersona = PersonaPtr.Pin();
-	if(SharedPersona.IsValid())
-	{
-		SharedPersona->OnAnimNotifiesChanged.Broadcast();
-	}
+	OnAnimNotifiesChanged.ExecuteIfBound();
 
 	OnTrackSelectionChanged();
 }
@@ -4497,8 +4488,8 @@ void SAnimNotifyPanel::OnReplaceSelectedWithNotify(FString NewNotifyName, UClass
 	}
 
 	// clear selection  
-	FGraphPanelSelectionSet ObjectSet;
-	OnSelectionChanged.ExecuteIfBound(ObjectSet);
+	TArray<UObject*> Objects;
+	OnSelectionChanged.ExecuteIfBound(Objects);
 	// TODO: set selection to new notifies?
 	// update the panel
 	Update();
@@ -4621,11 +4612,7 @@ void SAnimNotifyPanel::OnPropertyChanged(UObject* ChangedObject, FPropertyChange
 		}
 
 		// Broadcast the change so the editor can update
-		TSharedPtr<FPersona> SharedPersona = PersonaPtr.Pin();
-		if(SharedPersona.IsValid())
-		{
-			SharedPersona->OnAnimNotifiesChanged.Broadcast();
-		}
+		OnAnimNotifiesChanged.ExecuteIfBound();
 	}
 }
 
@@ -4880,11 +4867,7 @@ void SAnimNotifyPanel::OnNotifyObjectChanged(UObject* EditorBaseObj, bool bRebui
 		}
 
 		// Broadcast the change so the editor can update
-		TSharedPtr<FPersona> SharedPersona = PersonaPtr.Pin();
-		if(SharedPersona.IsValid())
-		{
-			SharedPersona->OnAnimNotifiesChanged.Broadcast();
-		}
+		OnAnimNotifiesChanged.ExecuteIfBound();
 	}
 }
 
diff --git a/Engine/Source/Editor/Persona/Private/SAnimNotifyPanel.h b/Engine/Source/Editor/Persona/Private/SAnimNotifyPanel.h
--- a/Engine/Source/Editor/Persona/Private/SAnimNotifyPanel.h
+++ b/Engine/Source/Editor/Persona/Private/SAnimNotifyPanel.h
@@ -11,7 +11,7 @@
 #include "SAnimEditorBase.h"
 #include "SAnimTimingPanel.h"
 
-DECLARE_DELEGATE_OneParam( FOnSelectionChanged, const FGraphPanelSelectionSet& )
+DECLARE_DELEGATE_OneParam( FOnSelectionChanged, const TArray<UObject*>& )
 DECLARE_DELEGATE( FOnTrackSelectionChanged )
 DECLARE_DELEGATE( FOnUpdatePanel )
 DECLARE_DELEGATE_RetVal( float, FOnGetScrubValue )
@@ -143,8 +143,7 @@ class SAnimNotifyPanel: public SAnimTrackPanel
 {
 public:
 	SLATE_BEGIN_ARGS( SAnimNotifyPanel )
-		: _Persona()
-		, _Sequence()
+		: _Sequence()
 		, _CurrentPosition()
 		, _ViewInputMin()
 		, _ViewInputMax()
@@ -156,7 +155,6 @@ public:
 		, _OnRequestRefreshOffsets()
 	{}
 
-	SLATE_ARGUMENT( TSharedPtr<FPersona>,	Persona )
 	SLATE_ARGUMENT( class UAnimSequenceBase*, Sequence)
 	SLATE_ARGUMENT( float, WidgetWidth )
 	SLATE_ATTRIBUTE( float, CurrentPosition )
@@ -170,10 +168,12 @@ public:
 	SLATE_EVENT( FOnGetScrubValue, OnGetScrubValue )
 	SLATE_EVENT( FRefreshOffsetsRequest, OnRequestRefreshOffsets )
 	SLATE_EVENT( FOnGetTimingNodeVisibility, OnGetTimingNodeVisibility )
+	SLATE_EVENT( FSimpleDelegate, OnAnimNotifiesChanged )
+	SLATE_EVENT( FOnInvokeTab, OnInvokeTab )
 
 	SLATE_END_ARGS()
 
-	void Construct(const FArguments& InArgs);
+	void Construct(const FArguments& InArgs, FSimpleMulticastDelegate& OnPostUndo);
 	virtual ~SAnimNotifyPanel();
 
 	void SetSequence(class UAnimSequenceBase *	InSequence);
@@ -182,7 +182,6 @@ public:
 	FReply DeleteTrack(int32 TrackIndexToDelete);
 	bool CanDeleteTrack(int32 TrackIndexToDelete);
 	void Update();
-	TWeakPtr<FPersona> GetPersona() const { return PersonaPtr; }
 
 	/** Returns the position of the notify node currently being dragged. Returns -1 if no node is being dragged */
 	float CalculateDraggedNodePos() const;
@@ -294,9 +293,6 @@ private:
 
 	virtual void InputViewRangeChanged(float ViewMin, float ViewMax) override;
 
-	/** Persona reference **/
-	TWeakPtr<FPersona> PersonaPtr;
-
 	/** Attribute for accessing any section/branching point positions we have to draw */
 	TAttribute<TArray<FTrackMarkerBar>>	MarkerBars;
 
@@ -311,4 +307,10 @@ private:
 
 	/** Handle to the registered OnPropertyChangedHandle delegate */
 	FDelegateHandle OnPropertyChangedHandleDelegateHandle;
+
+	/** Delegate fired when anim notifies are changed */
+	FSimpleDelegate OnAnimNotifiesChanged;
+
+	/** Delegate used to invoke a tab */
+	FOnInvokeTab OnInvokeTab;
 };
diff --git a/Engine/Source/Editor/Persona/Private/SAnimSegmentsPanel.cpp b/Engine/Source/Editor/Persona/Private/SAnimSegmentsPanel.cpp
--- a/Engine/Source/Editor/Persona/Private/SAnimSegmentsPanel.cpp
+++ b/Engine/Source/Editor/Persona/Private/SAnimSegmentsPanel.cpp
@@ -285,6 +285,11 @@ void SAnimSegmentsPanel::SummonSegmentNodeContextMenu(FMenuBuilder& MenuBuilder,
 			UIAction.ExecuteAction.BindRaw(this, &SAnimSegmentsPanel::RevertToParent, AnimSegmentIndex);
 			MenuBuilder.AddMenuEntry(LOCTEXT("RevertToParentSegment", "Revert To Parent"), LOCTEXT("RevertToParentSegment_ToolTip", "Revert to Parent Animation"), FSlateIcon(), UIAction);
 			MenuBuilder.AddSubMenu(LOCTEXT("PickAnimationForTheSegment", "Replace animation with..."), LOCTEXT("PickAnimationForTheSegment_TooTip", "Replace the current animation with another animation."), FNewMenuDelegate::CreateSP(this, &SAnimSegmentsPanel::FillSubMenu, AnimSegmentIndex));
+
+			MenuBuilder.AddMenuSeparator();
+			// open asset option
+			UIAction.ExecuteAction.BindRaw(this, &SAnimSegmentsPanel::OpenAsset, AnimSegmentIndex);
+			MenuBuilder.AddMenuEntry(LOCTEXT("OpenAssetOfSegment", "Open Asset"), LOCTEXT("OpenAssetOfSegment_ToolTip", "Open Asset"), FSlateIcon(), UIAction);
 		}
 
 		MenuBuilder.EndSection();
@@ -295,6 +300,9 @@ void SAnimSegmentsPanel::SummonSegmentNodeContextMenu(FMenuBuilder& MenuBuilder,
 		{
 			UIAction.ExecuteAction.BindRaw(this, &SAnimSegmentsPanel::RemoveAnimSegment, AnimSegmentIndex);
 			MenuBuilder.AddMenuEntry(LOCTEXT("DeleteSegment", "Delete Segment"), LOCTEXT("DeleteSegmentHint", "Delete Segment"), FSlateIcon(), UIAction);
+			// open asset option
+			UIAction.ExecuteAction.BindRaw(this, &SAnimSegmentsPanel::OpenAsset, AnimSegmentIndex);
+			MenuBuilder.AddMenuEntry(LOCTEXT("OpenAssetOfSegment", "Open Asset"), LOCTEXT("OpenAssetOfSegment_ToolTip", "Open Asset"), FSlateIcon(), UIAction);
 		}
 		MenuBuilder.EndSection();
 	}
@@ -330,6 +338,9 @@ void SAnimSegmentsPanel::ReplaceAnimSegment(int32 AnimSegmentIndex, UAnimSequenc
 			OnPostAnimUpdateDelegate.Execute();
 		}
 	}
+
+	// it doesn't work well if I leave the window open. The delegate goes weired or it stop showing the popups. 
+	FSlateApplication::Get().DismissAllMenus();
 }
 
 void SAnimSegmentsPanel::ReplaceAnimSegment(const FAssetData& NewSequenceData, int32 AnimSegmentIndex)
@@ -376,6 +387,18 @@ void SAnimSegmentsPanel::RevertToParent(int32 AnimSegmentIndex)
 	ReplaceAnimSegment(AnimSegmentIndex, nullptr);
 }
 
+void SAnimSegmentsPanel::OpenAsset(int32 AnimSegmentIndex)
+{
+	if (ValidIndex(AnimSegmentIndex))
+	{
+		UAnimSequenceBase* Asset = AnimTrack->AnimSegments[AnimSegmentIndex].AnimReference;
+		if (Asset)
+		{
+			FAssetEditorManager::Get().OpenEditorForAsset(Asset);
+		}
+	}
+}
+
 void SAnimSegmentsPanel::FillSubMenu(FMenuBuilder& MenuBuilder, int32 AnimSegmentIndex)
 {
 	if (ValidIndex(AnimSegmentIndex))
diff --git a/Engine/Source/Editor/Persona/Private/SAnimSegmentsPanel.h b/Engine/Source/Editor/Persona/Private/SAnimSegmentsPanel.h
--- a/Engine/Source/Editor/Persona/Private/SAnimSegmentsPanel.h
+++ b/Engine/Source/Editor/Persona/Private/SAnimSegmentsPanel.h
@@ -129,6 +129,7 @@ private:
 	// Remove the specified anim segment
 	void RemoveAnimSegment(int32 AnimSegmentIndex);
 	void RevertToParent(int32 AnimSegmentIndex);
+	void OpenAsset(int32 AnimSegmentIndex);
 	void FillSubMenu(FMenuBuilder& MenuBuilder, int32 AnimSegmentIndex);
 
 	/** Bind UI commands for this widget */
diff --git a/Engine/Source/Editor/Persona/Private/SAnimTimingPanel.cpp b/Engine/Source/Editor/Persona/Private/SAnimTimingPanel.cpp
--- a/Engine/Source/Editor/Persona/Private/SAnimTimingPanel.cpp
+++ b/Engine/Source/Editor/Persona/Private/SAnimTimingPanel.cpp
@@ -174,7 +174,7 @@ void SAnimTimingTrackNode::Construct(const FArguments& InArgs)
 		);
 }
 
-void SAnimTimingPanel::Construct(const FArguments& InArgs)
+void SAnimTimingPanel::Construct(const FArguments& InArgs, FSimpleMulticastDelegate& OnAnimNotifiesChanged, FSimpleMulticastDelegate& OnSectionsChanged)
 {
 	SAnimTrackPanel::Construct(SAnimTrackPanel::FArguments()
 		.WidgetWidth(InArgs._WidgetWidth)
@@ -184,10 +184,8 @@ void SAnimTimingPanel::Construct(const FArguments& InArgs)
 		.InputMax(InArgs._InputMax)
 		.OnSetInputViewRange(InArgs._OnSetInputViewRange));
 
-	WeakPersona = InArgs._InWeakPersona;
 	AnimSequence = InArgs._InSequence;
 
-	check(WeakPersona.IsValid());
 	check(AnimSequence);
 
 	this->ChildSlot
@@ -211,28 +209,13 @@ void SAnimTimingPanel::Construct(const FArguments& InArgs)
 	Update();
 
 	// Register to some delegates to update the interface
-	TSharedPtr<FPersona> SharedPersona = WeakPersona.Pin();
-	if(SharedPersona.IsValid())
-	{
-		SharedPersona->RegisterOnChangeAnimNotifies(FPersona::FOnAnimNotifiesChanged::CreateSP(this, &SAnimTimingPanel::RefreshTrackNodes));
-		SharedPersona->RegisterOnSectionsChanged(FPersona::FOnSectionsChanged::CreateSP(this, &SAnimTimingPanel::RefreshTrackNodes));
-	}
+	OnAnimNotifiesChanged.Add(FSimpleDelegate::CreateSP(this, &SAnimTimingPanel::RefreshTrackNodes));
+	OnSectionsChanged.Add(FSimpleDelegate::CreateSP(this, &SAnimTimingPanel::RefreshTrackNodes));
 
 	// Clear display flags
 	FMemory::Memset(bElementNodeDisplayFlags, false, sizeof(bElementNodeDisplayFlags));
 }
 
-SAnimTimingPanel::~SAnimTimingPanel()
-{
-	// Clean up our registered delegates
-	TSharedPtr<FPersona> SharedPersona = WeakPersona.Pin();
-	if(SharedPersona.IsValid())
-	{
-		SharedPersona->UnregisterOnChangeAnimNotifies(this);
-		SharedPersona->UnregisterOnSectionsChanged(this);
-	}
-}
-
 void SAnimTimingPanel::Update()
 {
 	check(PanelArea.IsValid());
diff --git a/Engine/Source/Editor/Persona/Private/SAnimTimingPanel.h b/Engine/Source/Editor/Persona/Private/SAnimTimingPanel.h
--- a/Engine/Source/Editor/Persona/Private/SAnimTimingPanel.h
+++ b/Engine/Source/Editor/Persona/Private/SAnimTimingPanel.h
@@ -166,8 +166,7 @@ class SAnimTimingPanel : public SAnimTrackPanel
 public:
 
 	SLATE_BEGIN_ARGS(SAnimTimingPanel)
-		: _InWeakPersona()
-		, _InSequence(nullptr)
+		: _InSequence(nullptr)
 		, _WidgetWidth(0)
 		, _CurrentPosition(0)
 		, _ViewInputMin(0)
@@ -177,7 +176,6 @@ public:
 		, _OnSetInputViewRange()
 	{}
 
-	SLATE_ARGUMENT(TWeakPtr<FPersona>, InWeakPersona)
 	SLATE_ARGUMENT(UAnimSequenceBase*, InSequence)
 	SLATE_ARGUMENT(float, WidgetWidth)
 	SLATE_ATTRIBUTE(float, CurrentPosition)
@@ -189,11 +187,9 @@ public:
 
 	SLATE_END_ARGS()
 
-	virtual ~SAnimTimingPanel();
-
 	// Construct the panel
 	// @param InArgs - Slate arguments
-	void Construct(const FArguments& InArgs);
+	void Construct(const FArguments& InArgs, FSimpleMulticastDelegate& OnAnimNotifiesChanged, FSimpleMulticastDelegate& OnSectionsChanged);
 
 	// Updates panel widgets
 	void Update();
@@ -227,8 +223,6 @@ protected:
 
 	// Observed timing elements
 	TArray<TSharedPtr<FTimingRelevantElementBase>> Elements;
-	// Toolkit reference
-	TWeakPtr<FPersona> WeakPersona;
 	// Anim sequence that contains the timing elements we are observing
 	UAnimSequenceBase* AnimSequence;
 	// Main panel widget
diff --git a/Engine/Source/Editor/Persona/Private/SAnimTrackCurvePanel.cpp b/Engine/Source/Editor/Persona/Private/SAnimTrackCurvePanel.cpp
--- a/Engine/Source/Editor/Persona/Private/SAnimTrackCurvePanel.cpp
+++ b/Engine/Source/Editor/Persona/Private/SAnimTrackCurvePanel.cpp
@@ -822,13 +822,13 @@ FReply		SAnimTrackCurvePanel::ShowAll(bool bShow)
 	return FReply::Handled();
 }
 
-ECheckBoxState SAnimTrackCurvePanel::GetCurveFlagAsCheckboxState(USkeleton::AnimCurveUID CurveUid, EAnimCurveFlags InFlag) const
+ECheckBoxState SAnimTrackCurvePanel::GetCurveFlagAsCheckboxState(USkeleton::AnimCurveUID CurveUid, EAnimAssetCurveFlags InFlag) const
 {
 	FAnimCurveBase* Curve = Sequence->RawCurveData.GetCurveData(CurveUid, FRawCurveTracks::TransformType);
 	return Curve && Curve->GetCurveTypeFlag(InFlag) ? ECheckBoxState::Checked : ECheckBoxState::Unchecked;
 }
 
-void SAnimTrackCurvePanel::SetCurveFlagFromCheckboxState(ECheckBoxState CheckState, USkeleton::AnimCurveUID CurveUid, EAnimCurveFlags InFlag)
+void SAnimTrackCurvePanel::SetCurveFlagFromCheckboxState(ECheckBoxState CheckState, USkeleton::AnimCurveUID CurveUid, EAnimAssetCurveFlags InFlag)
 {
 	bool Enabled = CheckState == ECheckBoxState::Checked;
 	FAnimCurveBase* Curve = Sequence->RawCurveData.GetCurveData(CurveUid, FRawCurveTracks::TransformType);
diff --git a/Engine/Source/Editor/Persona/Private/SAnimTrackCurvePanel.h b/Engine/Source/Editor/Persona/Private/SAnimTrackCurvePanel.h
--- a/Engine/Source/Editor/Persona/Private/SAnimTrackCurvePanel.h
+++ b/Engine/Source/Editor/Persona/Private/SAnimTrackCurvePanel.h
@@ -63,7 +63,7 @@ public:
 	/**
 	 * Sets the specified flag value to State for the provided curve
 	 */
-	void SetCurveFlag(FAnimCurveBase* CurveInterface, bool State, EAnimCurveFlags FlagToSet);
+	void SetCurveFlag(FAnimCurveBase* CurveInterface, bool State, EAnimAssetCurveFlags FlagToSet);
 
 	/**
 	* Update Panel
@@ -111,10 +111,10 @@ private:
 	/**
 	 * Convert the requested flag bool value into a checkbox state
 	 */
-	ECheckBoxState GetCurveFlagAsCheckboxState(USkeleton::AnimCurveUID CurveUid, EAnimCurveFlags InFlag) const;
+	ECheckBoxState GetCurveFlagAsCheckboxState(USkeleton::AnimCurveUID CurveUid, EAnimAssetCurveFlags InFlag) const;
 
 	/**
 	 * Convert a given checkbox state into a flag value in the provided curve
 	 */
-	void SetCurveFlagFromCheckboxState(ECheckBoxState CheckState, USkeleton::AnimCurveUID CurveUid, EAnimCurveFlags InFlag);
+	void SetCurveFlagFromCheckboxState(ECheckBoxState CheckState, USkeleton::AnimCurveUID CurveUid, EAnimAssetCurveFlags InFlag);
 };
\ No newline at end of file
diff --git a/Engine/Source/Editor/Persona/Private/SAnimViewportToolBar.cpp b/Engine/Source/Editor/Persona/Private/SAnimViewportToolBar.cpp
--- a/Engine/Source/Editor/Persona/Private/SAnimViewportToolBar.cpp
+++ b/Engine/Source/Editor/Persona/Private/SAnimViewportToolBar.cpp
@@ -17,9 +17,13 @@
 #include "SColorPicker.h"
 #include "SNumericEntryBox.h"
 #include "STextComboBox.h"
-
 #include "AssetViewerSettings.h"
-
+#include "IPersonaToolkit.h"
+#include "PersonaPreviewSceneDescription.h"
+#include "ISkeletonTree.h"
+#include "IEditableSkeleton.h"
+#include "Engine/PreviewMeshCollection.h"
+#include "PreviewSceneCustomizations.h"
 
 #define LOCTEXT_NAMESPACE "AnimViewportToolBar"
 
@@ -274,7 +278,6 @@ void SAnimViewportToolBar::Construct(const FArguments& InArgs, TSharedPtr<class
 				SNew( SEditorViewportToolbarMenu )
 				.ParentToolBar( SharedThis( this ) )
 				.Image("EditorViewportToolBar.MenuDropdown")
-//					.Label( this, &SAnimViewportToolBar::GetViewMenuLabel )
 				.OnGetMenuContent( this, &SAnimViewportToolBar::GenerateViewMenu ) 
 			]
 
@@ -332,18 +335,6 @@ void SAnimViewportToolBar::Construct(const FArguments& InArgs, TSharedPtr<class
 				.LabelIcon(FEditorStyle::GetBrush("AnimViewportMenu.PlayBackSpeed"))
 				.OnGetMenuContent( this, &SAnimViewportToolBar::GeneratePlaybackMenu ) 
 			]
-
-			// Turn table menu
-			+ SHorizontalBox::Slot()
-			.AutoWidth()
-			.Padding(2.0f, 2.0f)
-			[
-				SNew(SEditorViewportToolbarMenu)
-				.ParentToolBar(SharedThis(this))
-				.Label(this, &SAnimViewportToolBar::GetTurnTableMenuLabel)
-				.LabelIcon(FEditorStyle::GetBrush("AnimViewportMenu.TurnTableSpeed"))
-				.OnGetMenuContent(this, &SAnimViewportToolBar::GenerateTurnTableMenu)
-			]
 				
  			+SHorizontalBox::Slot()
 			.Padding(3.0f, 1.0f)
@@ -395,21 +386,6 @@ EVisibility SAnimViewportToolBar::GetTransformToolbarVisibility() const
 	return Viewport.Pin()->CanUseGizmos() ? EVisibility::Visible : EVisibility::Hidden;
 }
 
-FText SAnimViewportToolBar::GetViewMenuLabel() const
-{
-	FText Label = LOCTEXT("ViewMenu_AutoLabel", "Menu");
-	if (Viewport.IsValid())
-	{
-		// lock mode on
-		if (Viewport.Pin()->IsPreviewModeOn(1))
-		{
-			Label = LOCTEXT("ViewMenu_LockLabel", "Lock");
-		}
-	}
-
-	return Label;
-}
-
 TSharedRef<SWidget> SAnimViewportToolBar::GenerateViewMenu() const
 {
 	const FAnimViewportMenuCommands& Actions = FAnimViewportMenuCommands::Get();
@@ -417,27 +393,40 @@ TSharedRef<SWidget> SAnimViewportToolBar::GenerateViewMenu() const
 	const bool bInShouldCloseWindowAfterMenuSelection = true;
 	FMenuBuilder ViewMenuBuilder( bInShouldCloseWindowAfterMenuSelection, Viewport.Pin()->GetCommandList() );
 	{
-		// View modes
+		ViewMenuBuilder.BeginSection("AnimViewportSceneSetup", LOCTEXT("ViewMenu_SceneSetupLabel", "Scene Setup"));
 		{
-			ViewMenuBuilder.BeginSection("AnimViewportPreviewMode", LOCTEXT("ViewMenu_PreviewModeLabel", "Preview Mode") );
+			ViewMenuBuilder.AddMenuEntry(FAnimViewportMenuCommands::Get().PreviewSceneSettings);
+
+			// We cant allow animation configuration via the viewport menu in 'old regular' Persona
+			if (GetDefault<UPersonaOptions>()->bUseStandaloneAnimationEditors)
 			{
-				ViewMenuBuilder.AddMenuEntry( Actions.Auto );
-				ViewMenuBuilder.AddMenuEntry( Actions.Lock );
+				ViewMenuBuilder.AddSubMenu(
+					LOCTEXT("SceneSetupLabel", "Scene Setup"),
+					LOCTEXT("SceneSetupTooltip", "Set up preview meshes, animations etc."),
+					FNewMenuDelegate::CreateRaw(this, &SAnimViewportToolBar::GenerateSceneSetupMenu),
+					false,
+					FSlateIcon(FEditorStyle::GetStyleSetName(), "AnimViewportMenu.SceneSetup")
+					);
 			}
-			ViewMenuBuilder.EndSection();
 
-			ViewMenuBuilder.AddMenuEntry(FAnimViewportMenuCommands::Get().CameraFollow);
+			ViewMenuBuilder.AddSubMenu(
+				LOCTEXT("TurnTableLabel", "Turn Table"),
+				LOCTEXT("TurnTableTooltip", "Set up auto-rotation of preview."),
+				FNewMenuDelegate::CreateRaw(this, &SAnimViewportToolBar::GenerateTurnTableMenu),
+				false,
+				FSlateIcon(FEditorStyle::GetStyleSetName(), "AnimViewportMenu.TurnTableSpeed")
+				);
+		}
+		ViewMenuBuilder.EndSection();
 
-			ViewMenuBuilder.BeginSection("AnimViewportPreview", LOCTEXT("ViewMenu_PreviewLabel", "Preview") );
-			{
-				ViewMenuBuilder.AddMenuEntry( Actions.UseInGameBound );
-			}
-			ViewMenuBuilder.EndSection();
+		ViewMenuBuilder.BeginSection("AnimViewportCamera", LOCTEXT("ViewMenu_CameraLabel", "Camera"));
+		{
+			ViewMenuBuilder.AddMenuEntry(FAnimViewportMenuCommands::Get().CameraFollow);
 		}
+		ViewMenuBuilder.EndSection();
 	}
 
 	return ViewMenuBuilder.MakeWidget();
-
 }
 
 TSharedRef<SWidget> SAnimViewportToolBar::GenerateShowMenu() const
@@ -478,9 +467,9 @@ TSharedRef<SWidget> SAnimViewportToolBar::GenerateShowMenu() const
 
 			ShowMenuBuilder.BeginSection("AnimViewportMesh", LOCTEXT("ShowMenu_Actions_Mesh", "Mesh"));
 			{
-				ShowMenuBuilder.AddMenuEntry( Actions.ShowReferencePose );
 				ShowMenuBuilder.AddMenuEntry( Actions.ShowRetargetBasePose );
 				ShowMenuBuilder.AddMenuEntry( Actions.ShowBound );
+				ShowMenuBuilder.AddMenuEntry( Actions.UseInGameBound );
 				ShowMenuBuilder.AddMenuEntry( Actions.ShowPreviewMesh );
 				ShowMenuBuilder.AddMenuEntry( Actions.ShowMorphTargets );
 			}
@@ -522,7 +511,7 @@ TSharedRef<SWidget> SAnimViewportToolBar::GenerateShowMenu() const
 				FNewMenuDelegate::CreateRaw(this, &SAnimViewportToolBar::FillShowDisplayInfoMenu));
 
 #if WITH_APEX_CLOTHING
-			UDebugSkelMeshComponent* PreviewComp = Viewport.Pin()->GetPersona().Pin()->PreviewComponent;
+			UDebugSkelMeshComponent* PreviewComp = Viewport.Pin()->GetPreviewScene()->GetPreviewMeshComponent();
 
 			if( PreviewComp && PreviewComp->HasValidClothingActors() )
 			{
@@ -808,53 +797,54 @@ TSharedRef<SWidget> SAnimViewportToolBar::GeneratePlaybackMenu() const
 
 }
 
-TSharedRef<SWidget> SAnimViewportToolBar::GenerateTurnTableMenu() const
+void SAnimViewportToolBar::GenerateTurnTableMenu(FMenuBuilder& MenuBuilder) const
 {
 	const FAnimViewportPlaybackCommands& Actions = FAnimViewportPlaybackCommands::Get();
 
 	const bool bInShouldCloseWindowAfterMenuSelection = true;
 
-	FMenuBuilder TurnTableMenuBuilder(bInShouldCloseWindowAfterMenuSelection, Viewport.Pin()->GetCommandList());
+	MenuBuilder.BeginSection("AnimViewportTurnTableMode", LOCTEXT("TurnTableMenu_ModeLabel", "Turn Table Mode"));
 	{
-		TurnTableMenuBuilder.BeginSection("AnimViewportTurnTableMode", LOCTEXT("TurnTableMenu_ModeLabel", "Turn Table Mode"));
-		{
-			TurnTableMenuBuilder.AddMenuEntry(Actions.PersonaTurnTablePlay);
-			TurnTableMenuBuilder.AddMenuEntry(Actions.PersonaTurnTablePause);
-			TurnTableMenuBuilder.AddMenuEntry(Actions.PersonaTurnTableStop);
-		}
-		TurnTableMenuBuilder.EndSection();
+		MenuBuilder.AddMenuEntry(Actions.PersonaTurnTablePlay);
+		MenuBuilder.AddMenuEntry(Actions.PersonaTurnTablePause);
+		MenuBuilder.AddMenuEntry(Actions.PersonaTurnTableStop);
+	}
+	MenuBuilder.EndSection();
 
-		TurnTableMenuBuilder.BeginSection("AnimViewportTurnTableSpeed", LOCTEXT("TurnTableMenu_SpeedLabel", "Turn Table Speed"));
+	MenuBuilder.BeginSection("AnimViewportTurnTableSpeed", LOCTEXT("TurnTableMenu_SpeedLabel", "Turn Table Speed"));
+	{
+		for (int i = 0; i < EAnimationPlaybackSpeeds::NumPlaybackSpeeds; ++i)
 		{
-			for (int i = 0; i < EAnimationPlaybackSpeeds::NumPlaybackSpeeds; ++i)
-			{
-				TurnTableMenuBuilder.AddMenuEntry(Actions.TurnTableSpeeds[i]);
-			}
+			MenuBuilder.AddMenuEntry(Actions.TurnTableSpeeds[i]);
 		}
-		TurnTableMenuBuilder.EndSection();
 	}
+	MenuBuilder.EndSection();
+}
+
+void SAnimViewportToolBar::GenerateSceneSetupMenu(FMenuBuilder& MenuBuilder)
+{
+	FPropertyEditorModule& PropertyEditorModule = FModuleManager::LoadModuleChecked<FPropertyEditorModule>("PropertyEditor");
+	FDetailsViewArgs Args(false, false, false, FDetailsViewArgs::HideNameArea, true, nullptr, false, "PersonaPreviewSceneDescription");
+	Args.bShowScrollBar = false;
+
+	TSharedRef<IDetailsView> DetailsView = PropertyEditorModule.CreateDetailView(Args);
 
-	return TurnTableMenuBuilder.MakeWidget();
+	DetailsView->RegisterInstancedCustomPropertyLayout(UPersonaPreviewSceneDescription::StaticClass(), FOnGetDetailCustomizationInstance::CreateSP(this, &SAnimViewportToolBar::CustomizePreviewSceneDescription));
+	PropertyEditorModule.RegisterCustomPropertyTypeLayout(FPreviewMeshCollectionEntry::StaticStruct()->GetFName(), FOnGetPropertyTypeCustomizationInstance::CreateSP(this, &SAnimViewportToolBar::CustomizePreviewMeshCollectionEntry), nullptr, DetailsView);
+
+	DetailsView->SetObject(StaticCastSharedRef<FAnimationEditorPreviewScene>(Viewport.Pin()->GetPreviewScene())->GetPreviewSceneDescription());
+
+	MenuBuilder.AddWidget(DetailsView, FText(), true);
 }
 
-FText SAnimViewportToolBar::GetTurnTableMenuLabel() const
+TSharedRef<class IDetailCustomization> SAnimViewportToolBar::CustomizePreviewSceneDescription()
 {
-	FText Label = LOCTEXT("TurnTableError", "Error");
-	if (Viewport.IsValid())
-	{
-		for (int i = 0; i < EAnimationPlaybackSpeeds::NumPlaybackSpeeds; ++i)
-		{
-			if (Viewport.Pin()->IsTurnTableSpeedSelected(i))
-			{
-				Label = FText::FromString(FString::Printf(
-					(i == EAnimationPlaybackSpeeds::Quarter) ? TEXT("x%.2f") : TEXT("x%.1f"),
-					EAnimationPlaybackSpeeds::Values[i]
-					));
-				break;
-			}
-		}
-	}
-	return Label;
+	return MakeShareable(new FPreviewSceneDescriptionCustomization(FAssetData(&Viewport.Pin()->GetSkeletonTree()->GetEditableSkeleton()->GetSkeleton()).GetExportTextName(), Viewport.Pin()->GetPreviewScene(), Viewport.Pin()->GetSkeletonTree()->GetEditableSkeleton()));
+}
+
+TSharedRef<class IPropertyTypeCustomization> SAnimViewportToolBar::CustomizePreviewMeshCollectionEntry()
+{
+	return MakeShareable(new FPreviewMeshCollectionEntryCustomization(Viewport.Pin()->GetPreviewScene()));
 }
 
 FSlateColor SAnimViewportToolBar::GetFontColor() const
diff --git a/Engine/Source/Editor/Persona/Private/SAnimViewportToolBar.h b/Engine/Source/Editor/Persona/Private/SAnimViewportToolBar.h
--- a/Engine/Source/Editor/Persona/Private/SAnimViewportToolBar.h
+++ b/Engine/Source/Editor/Persona/Private/SAnimViewportToolBar.h
@@ -16,13 +16,6 @@ public:
 	void Construct(const FArguments& InArgs, TSharedPtr<class SAnimationEditorViewportTabBody> InViewport, TSharedPtr<class SEditorViewport> InRealViewport);
 
 private:
-	/**
-	 * Returns the label for the "View" tool bar menu, which changes depending on the viewport type
-	 *
-	 * @return	Label to use for this menu label
-	 */
-	FText GetViewMenuLabel() const;
-
 	/**
 	 * Generates the toolbar view menu content 
 	 *
@@ -95,8 +88,17 @@ private:
 	 */
 	TSharedRef<SWidget> GeneratePlaybackMenu() const;
 
-	TSharedRef<SWidget> GenerateTurnTableMenu() const;
-	FText GetTurnTableMenuLabel() const;
+	/** Generate the turntable menu entries */
+	void GenerateTurnTableMenu(FMenuBuilder& MenuBuilder) const;
+
+	/** Generate the scene setup menu */
+	void GenerateSceneSetupMenu(FMenuBuilder& MenuBuilder);
+
+	/** Customize the details of the scene setup object */
+	TSharedRef<class IDetailCustomization> CustomizePreviewSceneDescription();
+
+	/** Customize a preview mesh collection entry */
+	TSharedRef<class IPropertyTypeCustomization> CustomizePreviewMeshCollectionEntry();
 
 	/**
 	* Generate color of the text on the top
diff --git a/Engine/Source/Editor/Persona/Private/SAnimationBlendSpace.cpp b/Engine/Source/Editor/Persona/Private/SAnimationBlendSpace.cpp
--- a/Engine/Source/Editor/Persona/Private/SAnimationBlendSpace.cpp
+++ b/Engine/Source/Editor/Persona/Private/SAnimationBlendSpace.cpp
@@ -3,7 +3,6 @@
 
 #include "PersonaPrivatePCH.h"
 #include "SAnimationBlendSpace.h"
-#include "SAnimationSequenceBrowser.h"
 #include "Persona.h"
 #include "AssetData.h"
 #include "Editor/ContentBrowser/Public/ContentBrowserModule.h"
@@ -1285,22 +1284,17 @@ void SBlendSpaceEditor::OnBlendSpaceParamtersChanged()
 }
 
 BEGIN_SLATE_FUNCTION_BUILD_OPTIMIZATION
-void SBlendSpaceEditor::Construct(const FArguments& InArgs)
+void SBlendSpaceEditor::Construct(const FArguments& InArgs, const TSharedRef<class IPersonaPreviewScene>& InPreviewScene, FSimpleMulticastDelegate& OnPostUndo)
 {
 	SBlendSpaceEditorBase::Construct(SBlendSpaceEditorBase::FArguments()
-									 .Persona(InArgs._Persona)
-									 .BlendSpace(InArgs._BlendSpace)
+									 .BlendSpace(InArgs._BlendSpace),
+									InPreviewScene,
+									OnPostUndo
 									);
-	this->ChildSlot
+	NonScrollEditorPanels->AddSlot()
 	[
 		SNew(SVerticalBox)
 
-		+SVerticalBox::Slot()
-		.AutoHeight()
-		[
-			MakeEditorHeader()
-		]
-
 		+SVerticalBox::Slot()
 		.FillHeight(1.0f)
 		[
diff --git a/Engine/Source/Editor/Persona/Private/SAnimationBlendSpace.h b/Engine/Source/Editor/Persona/Private/SAnimationBlendSpace.h
--- a/Engine/Source/Editor/Persona/Private/SAnimationBlendSpace.h
+++ b/Engine/Source/Editor/Persona/Private/SAnimationBlendSpace.h
@@ -604,14 +604,12 @@ class SBlendSpaceEditor : public SBlendSpaceEditorBase
 public:
 	SLATE_BEGIN_ARGS(SBlendSpaceEditor)
 		: _BlendSpace(NULL)			
-		, _Persona()
 		{}
 
 		SLATE_ARGUMENT(UBlendSpace*, BlendSpace)
-		SLATE_ARGUMENT(TSharedPtr<class FPersona>, Persona)
 	SLATE_END_ARGS()
 
-	void Construct(const FArguments& InArgs);
+	void Construct(const FArguments& InArgs, const TSharedRef<class IPersonaPreviewScene>& InPreviewScene, FSimpleMulticastDelegate& OnPostUndo);
 	
 	void SetBlendSpace(UBlendSpace* NewBlendSpace);
 	UBlendSpace* GetBlendSpace() const { return Cast<UBlendSpace>(BlendSpace); }
diff --git a/Engine/Source/Editor/Persona/Private/SAnimationBlendSpace1D.cpp b/Engine/Source/Editor/Persona/Private/SAnimationBlendSpace1D.cpp
--- a/Engine/Source/Editor/Persona/Private/SAnimationBlendSpace1D.cpp
+++ b/Engine/Source/Editor/Persona/Private/SAnimationBlendSpace1D.cpp
@@ -3,7 +3,6 @@
 
 #include "PersonaPrivatePCH.h"
 #include "SAnimationBlendSpace1D.h"
-#include "SAnimationSequenceBrowser.h"
 #include "Persona.h"
 #include "AssetData.h"
 #include "Editor/ContentBrowser/Public/ContentBrowserModule.h"
@@ -518,11 +517,12 @@ SBlendSpaceEditor1D::~SBlendSpaceEditor1D()
 	FCoreUObjectDelegates::OnObjectPropertyChanged.Remove(OnPropertyChangedHandleDelegateHandle);
 }
 
-void SBlendSpaceEditor1D::Construct(const FArguments& InArgs)
+void SBlendSpaceEditor1D::Construct(const FArguments& InArgs, const TSharedRef<IPersonaPreviewScene>& InPreviewScene, FSimpleMulticastDelegate& OnPostUndo)
 {
 	SBlendSpaceEditorBase::Construct(SBlendSpaceEditorBase::FArguments()
-		.Persona(InArgs._Persona)
-		.BlendSpace(InArgs._BlendSpace1D)
+		.BlendSpace(InArgs._BlendSpace1D),
+		InPreviewScene,
+		OnPostUndo
 		);
 
 	// Register to be notified when properties are edited
@@ -531,16 +531,10 @@ void SBlendSpaceEditor1D::Construct(const FArguments& InArgs)
 
 	bCachedDisplayVerticalValue = GetBlendSpace()->bDisplayEditorVertically;
 
-	this->ChildSlot
+	NonScrollEditorPanels->AddSlot()
 	[
 		SNew(SVerticalBox)
 
-		+SVerticalBox::Slot()
-		.AutoHeight()
-		[
-			MakeEditorHeader()
-		]
-
 		+SVerticalBox::Slot()
 		.FillHeight(1.0f)
 		[
diff --git a/Engine/Source/Editor/Persona/Private/SAnimationBlendSpace1D.h b/Engine/Source/Editor/Persona/Private/SAnimationBlendSpace1D.h
--- a/Engine/Source/Editor/Persona/Private/SAnimationBlendSpace1D.h
+++ b/Engine/Source/Editor/Persona/Private/SAnimationBlendSpace1D.h
@@ -157,17 +157,15 @@ class SBlendSpaceEditor1D : public SBlendSpaceEditorBase
 {
 public:
 	SLATE_BEGIN_ARGS(SBlendSpaceEditor1D)
-		: _BlendSpace1D(NULL)			
-		, _Persona()
+		: _BlendSpace1D(NULL)
 		{}
 
 		SLATE_ARGUMENT(UBlendSpace1D*, BlendSpace1D)
-		SLATE_ARGUMENT(TSharedPtr<class FPersona>, Persona)
 	SLATE_END_ARGS()
 
 	~SBlendSpaceEditor1D();
 
-	void Construct(const FArguments& InArgs);
+	void Construct(const FArguments& InArgs, const TSharedRef<IPersonaPreviewScene>& InPreviewScene, FSimpleMulticastDelegate& OnPostUndo);
 
 protected:
 
diff --git a/Engine/Source/Editor/Persona/Private/SAnimationBlendSpaceBase.cpp b/Engine/Source/Editor/Persona/Private/SAnimationBlendSpaceBase.cpp
--- a/Engine/Source/Editor/Persona/Private/SAnimationBlendSpaceBase.cpp
+++ b/Engine/Source/Editor/Persona/Private/SAnimationBlendSpaceBase.cpp
@@ -3,7 +3,6 @@
 
 #include "PersonaPrivatePCH.h"
 #include "SAnimationBlendSpaceBase.h"
-#include "SAnimationSequenceBrowser.h"
 #include "Persona.h"
 #include "AssetData.h"
 #include "Editor/ContentBrowser/Public/ContentBrowserModule.h"
@@ -15,6 +14,7 @@
 #include "Animation/BlendSpaceBase.h"
 #include "Animation/AimOffsetBlendSpace.h"
 #include "Animation/AimOffsetBlendSpace1D.h"
+#include "IPersonaPreviewScene.h"
 
 DEFINE_LOG_CATEGORY_STATIC(LogBlendSpaceBase, Log, All);
 
@@ -1174,23 +1174,18 @@ FReply SBlendSpaceWidget::OpenAsset( int32 SampleIndex )
 
 ////////////////////////////////////////////////////////////////////////////////////////////////////////////////////////////////////////////////////////////////////
 
-SBlendSpaceEditorBase::~SBlendSpaceEditorBase()
-{
-	if (PersonaPtr.IsValid())
-	{
-		PersonaPtr.Pin()->UnregisterOnPostUndo(this);
-	}
-}
-
-void SBlendSpaceEditorBase::Construct(const FArguments& InArgs)
+void SBlendSpaceEditorBase::Construct(const FArguments& InArgs, const TSharedRef<IPersonaPreviewScene>& InPreviewScene, FSimpleMulticastDelegate& OnPostUndo)
 {
 	bIsActiveTimerRegistered = false;
 	BlendSpace = InArgs._BlendSpace;
-	PersonaPtr = InArgs._Persona;
-	PersonaPtr.Pin()->RegisterOnPostUndo(FPersona::FOnPostUndo::CreateSP( this, &SBlendSpaceEditorBase::PostUndo ) );
+	PreviewScenePtr = InPreviewScene;
+
+	OnPostUndo.Add(FSimpleDelegate::CreateSP( this, &SBlendSpaceEditorBase::PostUndo ) );
 
 	bIsActiveTimerRegistered = true;
 	RegisterActiveTimer(0.f, FWidgetActiveTimerDelegate::CreateSP(this, &SBlendSpaceEditorBase::UpdatePreview));
+
+	SAnimEditorBase::Construct(SAnimEditorBase::FArguments().DisplayAnimInfoBar(false), InPreviewScene);
 }
 
 EActiveTimerReturnType SBlendSpaceEditorBase::UpdatePreview( double InCurrentTime, float InDeltaTime )
@@ -1212,42 +1207,6 @@ void SBlendSpaceEditorBase::PostUndo()
 	RefreshSampleDataPanel();
 }
 
-TSharedRef<SWidget> SBlendSpaceEditorBase::MakeEditorHeader() const
-{
-	FString DocumentLink;
-
-	if (Cast<UAimOffsetBlendSpace>(BlendSpace) || Cast<UAimOffsetBlendSpace1D>(BlendSpace))
-	{
-		DocumentLink = TEXT("Engine/Animation/AimOffset");
-	}
-	else
-	{
-		DocumentLink = TEXT("Engine/Animation/Blendspaces");
-	}
-	
-	return SNew(SBorder)
-	. BorderImage( FEditorStyle::GetBrush( TEXT("Graph.TitleBackground") ) )
-	. HAlign(HAlign_Center)
-	[
-		SNew(SHorizontalBox)
-		+ SHorizontalBox::Slot()
-		.AutoWidth()
-		. VAlign(VAlign_Center)
-		[
-			SNew(STextBlock)
-			.Font( FSlateFontInfo(FPaths::EngineContentDir() / TEXT("Slate/Fonts/Roboto-Regular.ttf"), 14 ) )
-			.ColorAndOpacity( FLinearColor(1,1,1,0.5) )
-			.Text( this, &SBlendSpaceEditorBase::GetBlendSpaceDisplayName )
-		]
-		+ SHorizontalBox::Slot()
-		.HAlign(HAlign_Left)
-		.VAlign(VAlign_Center)
-		[
-			IDocumentation::Get()->CreateAnchor(DocumentLink)
-		]
-	];
-}
-
 TSharedRef<SWidget> SBlendSpaceEditorBase::MakeDisplayOptionsBox() const
 {
 	TSharedRef<SBlendSpaceDisplayOptionsWidget> DisplayOptionsPanel = SNew(SBlendSpaceDisplayOptionsWidget);
@@ -1278,7 +1237,7 @@ TSharedRef<SWidget> SBlendSpaceEditorBase::MakeDisplayOptionsBox() const
 
 ECheckBoxState SBlendSpaceEditorBase::IsPreviewOn() const
 {
-	class UDebugSkelMeshComponent* Component = PersonaPtr.Pin()->GetPreviewMeshComponent();
+	class UDebugSkelMeshComponent* Component = GetPreviewScene()->GetPreviewMeshComponent();
 
 	if (Component != NULL && Component->IsPreviewOn())
 	{
@@ -1293,7 +1252,7 @@ ECheckBoxState SBlendSpaceEditorBase::IsPreviewOn() const
 void SBlendSpaceEditorBase::ShowPreview_OnIsCheckedChanged( ECheckBoxState NewValue )
 {
 	bool bPreviewEnabled = (NewValue != ECheckBoxState::Unchecked);
-	class UDebugSkelMeshComponent* Component = PersonaPtr.Pin()->GetPreviewMeshComponent();
+	class UDebugSkelMeshComponent* Component = GetPreviewScene()->GetPreviewMeshComponent();
 
 	if ( Component != NULL )
 	{
@@ -1319,7 +1278,7 @@ void SBlendSpaceEditorBase::ShowToolTip_OnIsCheckedChanged( ECheckBoxState NewVa
 
 void SBlendSpaceEditorBase::UpdatePreviewParameter() const
 {
-	class UDebugSkelMeshComponent* Component = PersonaPtr.Pin()->GetPreviewMeshComponent();
+	class UDebugSkelMeshComponent* Component = GetPreviewScene()->GetPreviewMeshComponent();
 
 	if (Component != nullptr && Component->IsPreviewOn())
 	{
@@ -1328,10 +1287,7 @@ void SBlendSpaceEditorBase::UpdatePreviewParameter() const
 			FVector BlendInput = BlendSpaceWidget->LastValidMouseEditorPoint;
 			Component->PreviewInstance->SetBlendSpaceInput(BlendInput);
 
-			if (PersonaPtr.IsValid())
-			{
-				PersonaPtr.Pin()->RefreshViewport();
-			}
+			GetPreviewScene()->InvalidateViews();
 		}
 	}
 }
@@ -1348,18 +1304,15 @@ void SBlendSpaceEditorBase::OnBlendSpaceSamplesChanged()
 
 FVector SBlendSpaceEditorBase::GetPreviewBlendInput() const
 {
-	if (PersonaPtr.IsValid())
+	UDebugSkelMeshComponent * Mesh = GetPreviewScene()->GetPreviewMeshComponent();
+	if ( Mesh )
 	{
-		UDebugSkelMeshComponent * Mesh = PersonaPtr.Pin()->GetPreviewMeshComponent();
-		if ( Mesh )
+		UAnimSingleNodeInstance * SingleNodeInstance = Mesh->GetSingleNodeInstance();
+		if ( SingleNodeInstance )
 		{
-			UAnimSingleNodeInstance * SingleNodeInstance = Mesh->GetSingleNodeInstance();
-			if ( SingleNodeInstance )
-			{
-				return SingleNodeInstance->GetFilterLastOutput();
-			}
+			return SingleNodeInstance->GetFilterLastOutput();
 		}
-	}	
+	}
 
 	return FVector::ZeroVector;
 }
diff --git a/Engine/Source/Editor/Persona/Private/SAnimationBlendSpaceBase.h b/Engine/Source/Editor/Persona/Private/SAnimationBlendSpaceBase.h
--- a/Engine/Source/Editor/Persona/Private/SAnimationBlendSpaceBase.h
+++ b/Engine/Source/Editor/Persona/Private/SAnimationBlendSpaceBase.h
@@ -5,6 +5,7 @@
 #include "SlateBasics.h"
 #include "Sorting.h"
 #include "Animation/BlendSpaceBase.h"
+#include "SAnimEditorBase.h"
 
 extern FText GetErrorMessageForSequence(UBlendSpaceBase* BlendSpace, UAnimSequence* AnimSequence, FVector SampleValue, int32 OriginalIndex);
 
@@ -341,31 +342,26 @@ public:
 //////////////////////////////////////////////////////////////////////////
 // SBlendSpaceEditorBase
 
-class SBlendSpaceEditorBase : public SCompoundWidget
+class SBlendSpaceEditorBase : public SAnimEditorBase
 {
 public:
 	SLATE_BEGIN_ARGS(SBlendSpaceEditorBase)
 		: _BlendSpace(NULL)			
-		, _Persona()
 		{}
 
 		SLATE_ARGUMENT(UBlendSpaceBase*, BlendSpace)
-		SLATE_ARGUMENT(TSharedPtr<class FPersona>, Persona)
 	SLATE_END_ARGS()
 
-	virtual ~SBlendSpaceEditorBase();
-
-	void Construct(const FArguments& InArgs);
+	void Construct(const FArguments& InArgs, const TSharedRef<class IPersonaPreviewScene>& InPreviewScene, FSimpleMulticastDelegate& OnPostUndo);
 
 	void PostUndo();
 
 protected:
+	virtual UAnimationAsset* GetEditorObject() const override { return BlendSpace; }
+
 	/** Updates the preview */
 	EActiveTimerReturnType UpdatePreview( double InCurrentTime, float InDeltaTime );
 
-	/** Creates the editor heading UI */
-	TSharedRef<SWidget> MakeEditorHeader() const;
-
 	/** Creates the options panel for the editor */
 	virtual TSharedRef<SWidget> MakeDisplayOptionsBox() const;
 
@@ -394,8 +390,7 @@ protected:
 		return FText::FromString(BlendSpace->GetName());
 	}
 
-	/** Pointer back to the Persona that owns us */
-	TWeakPtr<class FPersona> PersonaPtr;
+	TSharedRef<class IPersonaPreviewScene> GetPreviewScene() const { return PreviewScenePtr.Pin().ToSharedRef();  }
 
 	/** The blend space being edited */
 	UBlendSpaceBase* BlendSpace;
@@ -412,6 +407,9 @@ protected:
 	/** This is used to tell the user when they do something wrong */
 	TSharedPtr<SNotificationList> NotificationListPtr;
 
+	/** The preview scene we are viewing */
+	TWeakPtr<class IPersonaPreviewScene> PreviewScenePtr;
+
 	/** Whether the active timer to update the preview is registered */
 	bool bIsActiveTimerRegistered;
 };
