Copying //UE4/Dev-Framework to //UE4/Dev-Main (Source: //UE4/Dev-Framework @ 3136612)

diff --git a/Engine/Source/Editor/Persona/Private/SAnimationEditorViewport.cpp b/Engine/Source/Editor/Persona/Private/SAnimationEditorViewport.cpp
--- a/Engine/Source/Editor/Persona/Private/SAnimationEditorViewport.cpp
+++ b/Engine/Source/Editor/Persona/Private/SAnimationEditorViewport.cpp
@@ -19,32 +19,24 @@
 #include "Editor/UnrealEd/Public/LODUtilities.h"
 #include "DetailLayoutBuilder.h"
 #include "STextComboBox.h"
+#include "ISkeletonTree.h"
+#include "IEditableSkeleton.h"
+#include "IPersonaPreviewScene.h"
 
 #define LOCTEXT_NAMESPACE "PersonaViewportToolbar"
 
-namespace EAnimationPlaybackSpeeds
-{
-	// Speed scales for animation playback, must match EAnimationPlaybackSpeeds::Type
-	float Values[EAnimationPlaybackSpeeds::NumPlaybackSpeeds] = {0.1f, 0.25f, 0.5f, 1.0f, 2.0f, 5.0f, 10.0f};
-};
-
 //////////////////////////////////////////////////////////////////////////
 // SAnimationEditorViewport
 
-SAnimationEditorViewport::~SAnimationEditorViewport()
+void SAnimationEditorViewport::Construct(const FArguments& InArgs, const FAnimationEditorViewportRequiredArgs& InRequiredArgs)
 {
-	if(PersonaPtr.IsValid())
-	{		
-		PersonaPtr.Pin()->UnregisterOnPostUndo(this);
-	}
-}
+	SkeletonTreePtr = InRequiredArgs.SkeletonTree;
+	PreviewScenePtr = InRequiredArgs.PreviewScene;
+	TabBodyPtr = InRequiredArgs.TabBody;
+	AssetEditorToolkitPtr = InRequiredArgs.AssetEditorToolkit;
+	bShowStats = InArgs._ShowStats;
 
-void SAnimationEditorViewport::Construct(const FArguments& InArgs, TSharedPtr<class FPersona> InPersona, TSharedPtr<class SAnimationEditorViewportTabBody> InTabBody)
-{
-	PersonaPtr = InPersona;
-	TabBodyPtr = InTabBody;
-
-	PersonaPtr.Pin()->RegisterOnPostUndo(FPersona::FOnPostUndo::CreateSP(this, &SAnimationEditorViewport::OnUndoRedo));
+	InRequiredArgs.OnPostUndo.Add(FSimpleDelegate::CreateSP(this, &SAnimationEditorViewport::OnUndoRedo));
 
 	SEditorViewport::Construct(
 		SEditorViewport::FArguments()
@@ -58,7 +50,7 @@ void SAnimationEditorViewport::Construct(const FArguments& InArgs, TSharedPtr<cl
 TSharedRef<FEditorViewportClient> SAnimationEditorViewport::MakeEditorViewportClient()
 {
 	// Create an animation viewport client
-	LevelViewportClient = MakeShareable(new FAnimationViewportClient(PersonaPtr.Pin()->GetPreviewScene(), PersonaPtr, SharedThis(this)));
+	LevelViewportClient = MakeShareable(new FAnimationViewportClient(SkeletonTreePtr.Pin().ToSharedRef(), PreviewScenePtr.Pin().ToSharedRef(), SharedThis(this), AssetEditorToolkitPtr.Pin().ToSharedRef(), bShowStats));
 
 	LevelViewportClient->ViewportType = LVT_Perspective;
 	LevelViewportClient->bSetListenerPosition = false;
@@ -77,7 +69,6 @@ TSharedPtr<SWidget> SAnimationEditorViewport::MakeViewportToolbar()
 void SAnimationEditorViewport::OnUndoRedo()
 {
 	LevelViewportClient->Invalidate();
-	LevelViewportClient->PostUndo();
 }
 
 //////////////////////////////////////////////////////////////////////////
@@ -185,14 +176,11 @@ public:
 SAnimationEditorViewportTabBody::SAnimationEditorViewportTabBody()
 	: SelectedTurnTableSpeed(EAnimationPlaybackSpeeds::Normal)
 	, SelectedTurnTableMode(EPersonaTurnTableMode::Stopped)
-	, AnimationPlaybackSpeedMode(EAnimationPlaybackSpeeds::Normal)
 {
 }
 
 SAnimationEditorViewportTabBody::~SAnimationEditorViewportTabBody()
 {
-	CleanupPersonaReferences();
-
 	// Close viewport
 	if (LevelViewportClient.IsValid())
 	{
@@ -203,18 +191,9 @@ SAnimationEditorViewportTabBody::~SAnimationEditorViewportTabBody()
 	LevelViewportClient.Reset();
 }
 
-void SAnimationEditorViewportTabBody::CleanupPersonaReferences()
-{
-	if (PersonaPtr.IsValid())
-	{
-		PersonaPtr.Pin()->UnregisterOnAnimChanged(this);
-		PersonaPtr.Reset();
-	}
-}
-
 bool SAnimationEditorViewportTabBody::CanUseGizmos() const
 {
-	class UDebugSkelMeshComponent* Component = PersonaPtr.Pin()->PreviewComponent;
+	class UDebugSkelMeshComponent* Component = GetPreviewScene()->GetPreviewMeshComponent();
 
 	if (Component != NULL)
 	{
@@ -233,8 +212,8 @@ bool SAnimationEditorViewportTabBody::CanUseGizmos() const
 
 FText SAnimationEditorViewportTabBody::GetDisplayString() const
 {
-	class UDebugSkelMeshComponent* Component = PersonaPtr.Pin()->PreviewComponent;
-	FString TargetSkeletonName = TargetSkeleton ? TargetSkeleton->GetName() : FName(NAME_None).ToString();
+	class UDebugSkelMeshComponent* Component = GetPreviewScene()->GetPreviewMeshComponent();
+	FName TargetSkeletonName = GetSkeletonTree()->GetEditableSkeleton()->GetSkeleton().GetFName();
 
 	if (Component != NULL)
 	{
@@ -248,7 +227,7 @@ FText SAnimationEditorViewportTabBody::GetDisplayString() const
 		}
 		else if (Component->AnimClass != NULL)
 		{
-			const bool bWarnAboutBoneManip = !PersonaPtr.Pin()->IsModeCurrent(FPersonaModes::AnimBlueprintEditMode);
+			const bool bWarnAboutBoneManip = BlueprintEditorPtr.Pin()->IsModeCurrent(FPersonaModes::AnimBlueprintEditMode);
 			if (bWarnAboutBoneManip)
 			{
 				return FText::Format(LOCTEXT("PreviewingAnimBP_WarnDisabled", "Previewing {0}. \nBone manipulation is disabled in this mode. "), FText::FromString(Component->AnimClass->GetName()));
@@ -260,7 +239,7 @@ FText SAnimationEditorViewportTabBody::GetDisplayString() const
 		}
 		else if (Component->SkeletalMesh == NULL)
 		{
-			return FText::Format(LOCTEXT("NoMeshFound", "No skeletal mesh found for skeleton '{0}'"), FText::FromString(TargetSkeletonName));
+			return FText::Format(LOCTEXT("NoMeshFound", "No skeletal mesh found for skeleton '{0}'"), FText::FromName(TargetSkeletonName));
 		}
 		else
 		{
@@ -269,10 +248,23 @@ FText SAnimationEditorViewportTabBody::GetDisplayString() const
 	}
 	else
 	{
-		return FText::Format(LOCTEXT("NoMeshFound", "No skeletal mesh found for skeleton '{0}'"), FText::FromString(TargetSkeletonName));
+		return FText::Format(LOCTEXT("NoMeshFound", "No skeletal mesh found for skeleton '{0}'"), FText::FromName(TargetSkeletonName));
 	}
 }
 
+TSharedRef<IPersonaViewportState> SAnimationEditorViewportTabBody::SaveState() const
+{
+	TSharedRef<FPersonaModeSharedData> State = MakeShareable(new(FPersonaModeSharedData));
+	State->Save(StaticCastSharedRef<FAnimationViewportClient>(LevelViewportClient.ToSharedRef()));
+	return State;
+}
+
+void SAnimationEditorViewportTabBody::RestoreState(TSharedRef<IPersonaViewportState> InState)
+{
+	TSharedRef<FPersonaModeSharedData> State = StaticCastSharedRef<FPersonaModeSharedData>(InState);
+	State->Restore(StaticCastSharedRef<FAnimationViewportClient>(LevelViewportClient.ToSharedRef()));
+}
+
 void SAnimationEditorViewportTabBody::RefreshViewport()
 {
 	LevelViewportClient->Invalidate();
@@ -284,21 +276,21 @@ bool SAnimationEditorViewportTabBody::IsVisible() const
 }
 
 
-void SAnimationEditorViewportTabBody::Construct(const FArguments& InArgs)
+void SAnimationEditorViewportTabBody::Construct(const FArguments& InArgs, const TSharedRef<class ISkeletonTree>& InSkeletonTree, const TSharedRef<class IPersonaPreviewScene>& InPreviewScene, const TSharedRef<class FAssetEditorToolkit>& InAssetEditorToolkit, FSimpleMulticastDelegate& InOnUndoRedo)
 {
 	UICommandList = MakeShareable(new FUICommandList);
 
-	PersonaPtr = InArgs._Persona;
-	IsEditable = InArgs._IsEditable;
-	TargetSkeleton = InArgs._Skeleton;
-	bPreviewLockModeOn = 0;
+	SkeletonTreePtr = InSkeletonTree;
+	PreviewScenePtr = StaticCastSharedRef<FAnimationEditorPreviewScene>(InPreviewScene);
+	AssetEditorToolkitPtr = InAssetEditorToolkit;
+	BlueprintEditorPtr = InArgs._BlueprintEditor;
+	bShowTimeline = InArgs._ShowTimeline;
+	OnInvokeTab = InArgs._OnInvokeTab;
 	LODSelection = 0;
 
-	TSharedPtr<FPersona> SharedPersona = PersonaPtr.Pin();
-	check(SharedPersona.IsValid());
-
-	// register delegate for anim change notification
-	SharedPersona->RegisterOnAnimChanged(FPersona::FOnAnimChanged::CreateSP(this, &SAnimationEditorViewportTabBody::AnimChanged));
+	// register delegates for change notifications
+	InPreviewScene->RegisterOnAnimChanged(FOnAnimChanged::CreateSP(this, &SAnimationEditorViewportTabBody::AnimChanged));
+	InPreviewScene->RegisterOnPreviewMeshChanged(FOnPreviewMeshChanged::CreateSP(this, &SAnimationEditorViewportTabBody::HandlePreviewMeshChanged));
 
 	const FSlateFontInfo SmallLayoutFont( FPaths::EngineContentDir() / TEXT("Slate/Fonts/Roboto-Regular.ttf"), 10 );
 
@@ -312,7 +304,10 @@ void SAnimationEditorViewportTabBody::Construct(const FArguments& InArgs)
 		.OptionsSource(&UVChannels)
 		.OnSelectionChanged(this, &SAnimationEditorViewportTabBody::ComboBoxSelectionChanged);
 
-	ViewportWidget = SNew(SAnimationEditorViewport, InArgs._Persona, SharedThis(this));
+	FAnimationEditorViewportRequiredArgs ViewportArgs(InSkeletonTree, InPreviewScene, SharedThis(this), InAssetEditorToolkit, InOnUndoRedo);
+
+	ViewportWidget = SNew(SAnimationEditorViewport, ViewportArgs)
+		.ShowStats(InArgs._ShowStats);
 
 	TSharedPtr<SVerticalBox> ViewportContainer = nullptr;
 	this->ChildSlot
@@ -343,29 +338,16 @@ void SAnimationEditorViewportTabBody::Construct(const FArguments& InArgs)
 				.OnClicked(this, &SAnimationEditorViewportTabBody::ClickedOnViewportCornerText)
 				.Content()
 				[
-					SNew(SHorizontalBox)
-					+SHorizontalBox::Slot()
-					.AutoWidth()
-					[
-						SNew(SImage)
-						.Visibility(this, &SAnimationEditorViewportTabBody::GetViewportCornerImageVisibility)
-						.Image(this, &SAnimationEditorViewportTabBody::GetViewportCornerImage)
-					]
-						
-					+SHorizontalBox::Slot()
-					.FillWidth(1)
-					[
-						SNew(STextBlock)
-						.TextStyle(FEditorStyle::Get(), "Persona.Viewport.BlueprintDirtyText")
-						.Text(this, &SAnimationEditorViewportTabBody::GetViewportCornerText)
-						.ToolTipText(this, &SAnimationEditorViewportTabBody::GetViewportCornerTooltip)
-					]
+					SNew(STextBlock)
+					.TextStyle(FEditorStyle::Get(), "Persona.Viewport.BlueprintDirtyText")
+					.Text(this, &SAnimationEditorViewportTabBody::GetViewportCornerText)
+					.ToolTipText(this, &SAnimationEditorViewportTabBody::GetViewportCornerTooltip)
 				]
 			]
 		]
 	];
 
-	if(!SharedPersona->IsModeCurrent(FPersonaModes::AnimationEditMode) && ViewportContainer.IsValid())
+	if(bShowTimeline && ViewportContainer.IsValid())
 	{
 		ViewportContainer->AddSlot()
 		.AutoHeight()
@@ -374,15 +356,14 @@ void SAnimationEditorViewportTabBody::Construct(const FArguments& InArgs)
 			+SVerticalBox::Slot()
 			.AutoHeight()
 			[
-				SNew(SAnimationScrubPanel)
-				.Persona(PersonaPtr)
+				SNew(SAnimationScrubPanel, GetPreviewScene())
 				.ViewInputMin(this, &SAnimationEditorViewportTabBody::GetViewMinInput)
 				.ViewInputMax(this, &SAnimationEditorViewportTabBody::GetViewMaxInput)
 				.bAllowZoom(true)
 			]
 		];
 
-		UpdateScrubPanel(Cast<UAnimationAsset>(SharedPersona->GetPreviewAnimationAsset()));
+		UpdateScrubPanel(InPreviewScene->GetPreviewAnimationAsset());
 	}
 
 	LevelViewportClient = ViewportWidget->GetViewportClient();
@@ -398,6 +379,8 @@ void SAnimationEditorViewportTabBody::Construct(const FArguments& InArgs)
 	OnSetTurnTableSpeed(SelectedTurnTableSpeed);
 
 	BindCommands();
+
+	PopulateNumUVChannels();
 }
 
 void SAnimationEditorViewportTabBody::BindCommands()
@@ -407,18 +390,6 @@ void SAnimationEditorViewportTabBody::BindCommands()
 	//Bind menu commands
 	const FAnimViewportMenuCommands& MenuActions = FAnimViewportMenuCommands::Get();
 
-	CommandList.MapAction( 
-		MenuActions.Lock,
-		FExecuteAction::CreateSP(this, &SAnimationEditorViewportTabBody::SetPreviewMode, 1),
-		FCanExecuteAction(),
-		FIsActionChecked::CreateSP(this, &SAnimationEditorViewportTabBody::IsPreviewModeOn, 1));
-
-	CommandList.MapAction( 
-		MenuActions.Auto,
-		FExecuteAction::CreateSP(this, &SAnimationEditorViewportTabBody::SetPreviewMode, 0),
-		FCanExecuteAction(),
-		FIsActionChecked::CreateSP(this, &SAnimationEditorViewportTabBody::IsPreviewModeOn, 0));
-
 	CommandList.MapAction(
 		MenuActions.CameraFollow,
 		FExecuteAction::CreateSP(this, &SAnimationEditorViewportTabBody::ToggleCameraFollow),
@@ -426,10 +397,8 @@ void SAnimationEditorViewportTabBody::BindCommands()
 		FIsActionChecked::CreateSP(this, &SAnimationEditorViewportTabBody::IsCameraFollowEnabled));
 
 	CommandList.MapAction(
-		MenuActions.UseInGameBound,
-		FExecuteAction::CreateSP(this, &SAnimationEditorViewportTabBody::UseInGameBound),
-		FCanExecuteAction::CreateSP(this, &SAnimationEditorViewportTabBody::CanUseInGameBound),
-		FIsActionChecked::CreateSP(this, &SAnimationEditorViewportTabBody::IsUsingInGameBound));
+		MenuActions.PreviewSceneSettings,
+		FExecuteAction::CreateSP(this, &SAnimationEditorViewportTabBody::OpenPreviewSceneSettings));
 
 	TSharedRef<FAnimationViewportClient> EditorViewportClientRef = GetAnimationViewportClient();
 
@@ -465,12 +434,6 @@ void SAnimationEditorViewportTabBody::BindCommands()
 
 	//Bind Show commands
 	const FAnimViewportShowCommands& ViewportShowMenuCommands = FAnimViewportShowCommands::Get();
-	
-	CommandList.MapAction(
-		ViewportShowMenuCommands.ShowReferencePose,
-		FExecuteAction::CreateSP(this, &SAnimationEditorViewportTabBody::ShowReferencePose),
-		FCanExecuteAction::CreateSP(this, &SAnimationEditorViewportTabBody::CanShowReferencePose),
-		FIsActionChecked::CreateSP(this, &SAnimationEditorViewportTabBody::IsShowReferencePoseEnabled));
 
 	CommandList.MapAction(
 		ViewportShowMenuCommands.ShowRetargetBasePose,
@@ -484,6 +447,12 @@ void SAnimationEditorViewportTabBody::BindCommands()
 		FCanExecuteAction::CreateSP(this, &SAnimationEditorViewportTabBody::CanShowBound),
 		FIsActionChecked::CreateSP(this, &SAnimationEditorViewportTabBody::IsShowBoundEnabled));
 
+	CommandList.MapAction(
+		ViewportShowMenuCommands.UseInGameBound,
+		FExecuteAction::CreateSP(this, &SAnimationEditorViewportTabBody::UseInGameBound),
+		FCanExecuteAction::CreateSP(this, &SAnimationEditorViewportTabBody::CanUseInGameBound),
+		FIsActionChecked::CreateSP(this, &SAnimationEditorViewportTabBody::IsUsingInGameBound));
+
 	CommandList.MapAction(
 		ViewportShowMenuCommands.ShowPreviewMesh,
 		FExecuteAction::CreateSP(this, &SAnimationEditorViewportTabBody::ToggleShowPreviewMesh),
@@ -794,7 +763,7 @@ void SAnimationEditorViewportTabBody::OnSetTurnTableSpeed(int32 SpeedIndex)
 {
 	SelectedTurnTableSpeed = (EAnimationPlaybackSpeeds::Type)SpeedIndex;
 
-	UDebugSkelMeshComponent* PreviewComponent = PersonaPtr.Pin()->PreviewComponent;
+	UDebugSkelMeshComponent* PreviewComponent = GetPreviewScene()->GetPreviewMeshComponent();
 	if (PreviewComponent)
 	{
 		PreviewComponent->TurnTableSpeedScaling = EAnimationPlaybackSpeeds::Values[SelectedTurnTableSpeed];
@@ -810,7 +779,7 @@ void SAnimationEditorViewportTabBody::OnSetTurnTableMode(int32 ModeIndex)
 {
 	SelectedTurnTableMode = (EPersonaTurnTableMode::Type)ModeIndex;
 
-	UDebugSkelMeshComponent* PreviewComponent = PersonaPtr.Pin()->PreviewComponent;
+	UDebugSkelMeshComponent* PreviewComponent = GetPreviewScene()->GetPreviewMeshComponent();
 	if (PreviewComponent)
 	{
 		PreviewComponent->TurnTableMode = SelectedTurnTableMode;
@@ -827,35 +796,9 @@ bool SAnimationEditorViewportTabBody::IsTurnTableModeSelected(int32 ModeIndex) c
 	return (SelectedTurnTableMode == ModeIndex);
 }
 
-bool SAnimationEditorViewportTabBody::IsPreviewModeOn(int32 PreviewMode) const
-{
-	if (bPreviewLockModeOn && PreviewMode)
-	{
-		return true;
-	}
-	else if (!bPreviewLockModeOn && !PreviewMode)
-	{
-		return true;
-	}
-
-	return false;
-
-}
-void SAnimationEditorViewportTabBody::SetPreviewMode(int32 PreviewMode)
-{
-	if (PreviewMode)
-	{
-		bPreviewLockModeOn = true; 
-	}
-	else
-	{
-		bPreviewLockModeOn = false; 
-	}
-}
-
 int32 SAnimationEditorViewportTabBody::GetLODModelCount() const
 {
-	UDebugSkelMeshComponent* PreviewComponent = PersonaPtr.Pin()->PreviewComponent;
+	UDebugSkelMeshComponent* PreviewComponent = GetPreviewScene()->GetPreviewMeshComponent();
 	if( PreviewComponent && PreviewComponent->SkeletalMesh )
 	{
 		return PreviewComponent->SkeletalMesh->GetImportedResource()->LODModels.Num();
@@ -865,7 +808,7 @@ int32 SAnimationEditorViewportTabBody::GetLODModelCount() const
 
 void SAnimationEditorViewportTabBody::OnShowMorphTargets()
 {
-	if (UDebugSkelMeshComponent* PreviewComponent = PersonaPtr.Pin()->PreviewComponent)
+	if (UDebugSkelMeshComponent* PreviewComponent = GetPreviewScene()->GetPreviewMeshComponent())
 	{
 		PreviewComponent->bDisableMorphTarget = !PreviewComponent->bDisableMorphTarget;
 		PreviewComponent->MarkRenderStateDirty();
@@ -875,7 +818,7 @@ void SAnimationEditorViewportTabBody::OnShowMorphTargets()
 
 void SAnimationEditorViewportTabBody::OnShowBoneNames()
 {
-	if (UDebugSkelMeshComponent* PreviewComponent = PersonaPtr.Pin()->PreviewComponent)
+	if (UDebugSkelMeshComponent* PreviewComponent = GetPreviewScene()->GetPreviewMeshComponent())
 	{
 		PreviewComponent->bShowBoneNames = !PreviewComponent->bShowBoneNames;
 		PreviewComponent->MarkRenderStateDirty();
@@ -885,7 +828,7 @@ void SAnimationEditorViewportTabBody::OnShowBoneNames()
 
 void SAnimationEditorViewportTabBody::OnShowRawAnimation()
 {
-	if (UDebugSkelMeshComponent* PreviewComponent = PersonaPtr.Pin()->PreviewComponent)
+	if (UDebugSkelMeshComponent* PreviewComponent = GetPreviewScene()->GetPreviewMeshComponent())
 	{
 		PreviewComponent->bDisplayRawAnimation = !PreviewComponent->bDisplayRawAnimation;
 		PreviewComponent->MarkRenderStateDirty();
@@ -894,7 +837,7 @@ void SAnimationEditorViewportTabBody::OnShowRawAnimation()
 
 void SAnimationEditorViewportTabBody::OnShowNonRetargetedAnimation()
 {
-	if (UDebugSkelMeshComponent* PreviewComponent = PersonaPtr.Pin()->PreviewComponent)
+	if (UDebugSkelMeshComponent* PreviewComponent = GetPreviewScene()->GetPreviewMeshComponent())
 	{
 		PreviewComponent->bDisplayNonRetargetedPose = !PreviewComponent->bDisplayNonRetargetedPose;
 		PreviewComponent->MarkRenderStateDirty();
@@ -903,7 +846,7 @@ void SAnimationEditorViewportTabBody::OnShowNonRetargetedAnimation()
 
 void SAnimationEditorViewportTabBody::OnShowSourceRawAnimation()
 {
-	if (UDebugSkelMeshComponent* PreviewComponent = PersonaPtr.Pin()->PreviewComponent)
+	if (UDebugSkelMeshComponent* PreviewComponent = GetPreviewScene()->GetPreviewMeshComponent())
 	{
 		PreviewComponent->bDisplaySourceAnimation = !PreviewComponent->bDisplaySourceAnimation;
 		PreviewComponent->MarkRenderStateDirty();
@@ -912,7 +855,7 @@ void SAnimationEditorViewportTabBody::OnShowSourceRawAnimation()
 
 void SAnimationEditorViewportTabBody::OnShowBakedAnimation()
 {
-	if (UDebugSkelMeshComponent* PreviewComponent = PersonaPtr.Pin()->PreviewComponent)
+	if (UDebugSkelMeshComponent* PreviewComponent = GetPreviewScene()->GetPreviewMeshComponent())
 	{
 		PreviewComponent->bDisplayBakedAnimation = !PreviewComponent->bDisplayBakedAnimation;
 		PreviewComponent->MarkRenderStateDirty();
@@ -921,7 +864,7 @@ void SAnimationEditorViewportTabBody::OnShowBakedAnimation()
 
 void SAnimationEditorViewportTabBody::OnShowAdditiveBase()
 {
-	if (UDebugSkelMeshComponent* PreviewComponent = PersonaPtr.Pin()->PreviewComponent)
+	if (UDebugSkelMeshComponent* PreviewComponent = GetPreviewScene()->GetPreviewMeshComponent())
 	{
 		PreviewComponent->bDisplayAdditiveBasePose = !PreviewComponent->bDisplayAdditiveBasePose;
 		PreviewComponent->MarkRenderStateDirty();
@@ -930,49 +873,49 @@ void SAnimationEditorViewportTabBody::OnShowAdditiveBase()
 
 bool SAnimationEditorViewportTabBody::IsPreviewingAnimation() const
 {
-	UDebugSkelMeshComponent* PreviewComponent = PersonaPtr.Pin()->PreviewComponent;
+	UDebugSkelMeshComponent* PreviewComponent = GetPreviewScene()->GetPreviewMeshComponent();
 	return (PreviewComponent && PreviewComponent->PreviewInstance && (PreviewComponent->PreviewInstance == PreviewComponent->GetAnimInstance()));
 }
 
 bool SAnimationEditorViewportTabBody::IsShowingMorphTargets() const
 {
-	UDebugSkelMeshComponent* PreviewComponent = PersonaPtr.Pin()->PreviewComponent;
+	UDebugSkelMeshComponent* PreviewComponent = GetPreviewScene()->GetPreviewMeshComponent();
 	return PreviewComponent != NULL && PreviewComponent->bDisableMorphTarget == false;
 }
 
 bool SAnimationEditorViewportTabBody::IsShowingBoneNames() const
 {
-	UDebugSkelMeshComponent* PreviewComponent = PersonaPtr.Pin()->PreviewComponent;
+	UDebugSkelMeshComponent* PreviewComponent = GetPreviewScene()->GetPreviewMeshComponent();
 	return PreviewComponent != NULL && PreviewComponent->bShowBoneNames;
 }
 
 bool SAnimationEditorViewportTabBody::IsShowingRawAnimation() const
 {
-	UDebugSkelMeshComponent* PreviewComponent = PersonaPtr.Pin()->PreviewComponent;
+	UDebugSkelMeshComponent* PreviewComponent = GetPreviewScene()->GetPreviewMeshComponent();
 	return PreviewComponent != NULL && PreviewComponent->bDisplayRawAnimation;
 }
 
 bool SAnimationEditorViewportTabBody::IsShowingNonRetargetedPose() const
 {
-	UDebugSkelMeshComponent* PreviewComponent = PersonaPtr.Pin()->PreviewComponent;
+	UDebugSkelMeshComponent* PreviewComponent = GetPreviewScene()->GetPreviewMeshComponent();
 	return PreviewComponent != NULL && PreviewComponent->bDisplayNonRetargetedPose;
 }
 
 bool SAnimationEditorViewportTabBody::IsShowingAdditiveBase() const
 {
-	UDebugSkelMeshComponent* PreviewComponent = PersonaPtr.Pin()->PreviewComponent;
+	UDebugSkelMeshComponent* PreviewComponent = GetPreviewScene()->GetPreviewMeshComponent();
 	return PreviewComponent != NULL && PreviewComponent->bDisplayAdditiveBasePose;
 }
 
 bool SAnimationEditorViewportTabBody::IsShowingSourceRawAnimation() const
 {
-	UDebugSkelMeshComponent* PreviewComponent = PersonaPtr.Pin()->PreviewComponent;
+	UDebugSkelMeshComponent* PreviewComponent = GetPreviewScene()->GetPreviewMeshComponent();
 	return PreviewComponent != NULL && PreviewComponent->bDisplaySourceAnimation;
 }
 
 bool SAnimationEditorViewportTabBody::IsShowingBakedAnimation() const
 {
-	UDebugSkelMeshComponent* PreviewComponent = PersonaPtr.Pin()->PreviewComponent;
+	UDebugSkelMeshComponent* PreviewComponent = GetPreviewScene()->GetPreviewMeshComponent();
 	return PreviewComponent != NULL && PreviewComponent->bDisplayBakedAnimation;
 }
 
@@ -988,7 +931,7 @@ bool SAnimationEditorViewportTabBody::IsShowingMeshInfo(int32 DisplayInfoMode) c
 
 void SAnimationEditorViewportTabBody::OnShowOverlayNone()
 {
-	UDebugSkelMeshComponent* PreviewComponent = PersonaPtr.Pin()->PreviewComponent;
+	UDebugSkelMeshComponent* PreviewComponent = GetPreviewScene()->GetPreviewMeshComponent();
 	if (PreviewComponent)
 	{
 		PreviewComponent->SetShowBoneWeight(false);
@@ -1001,13 +944,13 @@ void SAnimationEditorViewportTabBody::OnShowOverlayNone()
 
 bool SAnimationEditorViewportTabBody::IsShowingOverlayNone() const
 {
-	UDebugSkelMeshComponent* PreviewComponent = PersonaPtr.Pin()->PreviewComponent;
+	UDebugSkelMeshComponent* PreviewComponent = GetPreviewScene()->GetPreviewMeshComponent();
 	return PreviewComponent != NULL && !PreviewComponent->bDrawBoneInfluences && !PreviewComponent->bDrawMorphTargetVerts;
 }
 
 void SAnimationEditorViewportTabBody::OnShowOverlayBoneWeight()
 {
-	UDebugSkelMeshComponent* PreviewComponent = PersonaPtr.Pin()->PreviewComponent;
+	UDebugSkelMeshComponent* PreviewComponent = GetPreviewScene()->GetPreviewMeshComponent();
 	if( PreviewComponent )
 	{
 		PreviewComponent->SetShowBoneWeight( !PreviewComponent->bDrawBoneInfluences );
@@ -1019,13 +962,13 @@ void SAnimationEditorViewportTabBody::OnShowOverlayBoneWeight()
 
 bool SAnimationEditorViewportTabBody::IsShowingOverlayBoneWeight() const
 {
-	UDebugSkelMeshComponent* PreviewComponent = PersonaPtr.Pin()->PreviewComponent;
+	UDebugSkelMeshComponent* PreviewComponent = GetPreviewScene()->GetPreviewMeshComponent();
 	return PreviewComponent != NULL && PreviewComponent->bDrawBoneInfluences;
 }
 
 void SAnimationEditorViewportTabBody::OnShowOverlayMorphTargetVert()
 {
-	UDebugSkelMeshComponent* PreviewComponent = PersonaPtr.Pin()->PreviewComponent;
+	UDebugSkelMeshComponent* PreviewComponent = GetPreviewScene()->GetPreviewMeshComponent();
 	if (PreviewComponent)
 	{
 		PreviewComponent->SetShowMorphTargetVerts(!PreviewComponent->bDrawMorphTargetVerts);
@@ -1037,7 +980,7 @@ void SAnimationEditorViewportTabBody::OnShowOverlayMorphTargetVert()
 
 bool SAnimationEditorViewportTabBody::IsShowingOverlayMorphTargetVerts() const
 {
-	UDebugSkelMeshComponent* PreviewComponent = PersonaPtr.Pin()->PreviewComponent;
+	UDebugSkelMeshComponent* PreviewComponent = GetPreviewScene()->GetPreviewMeshComponent();
 	return PreviewComponent != NULL && PreviewComponent->bDrawMorphTargetVerts;
 }
 
@@ -1067,7 +1010,7 @@ bool SAnimationEditorViewportTabBody::IsLocalAxesModeSet(int32 LocalAxesMode) co
 
 void SAnimationEditorViewportTabBody::OnShowSockets()
 {
-	UDebugSkelMeshComponent* PreviewComponent = PersonaPtr.Pin()->PreviewComponent;
+	UDebugSkelMeshComponent* PreviewComponent = GetPreviewScene()->GetPreviewMeshComponent();
 	if( PreviewComponent )
 	{
 		PreviewComponent->bDrawSockets = !PreviewComponent->bDrawSockets;
@@ -1078,7 +1021,7 @@ void SAnimationEditorViewportTabBody::OnShowSockets()
 
 bool SAnimationEditorViewportTabBody::IsShowingSockets() const
 {
-	UDebugSkelMeshComponent* PreviewComponent = PersonaPtr.Pin()->PreviewComponent;
+	UDebugSkelMeshComponent* PreviewComponent = GetPreviewScene()->GetPreviewMeshComponent();
 	return PreviewComponent != NULL && PreviewComponent->bDrawSockets;
 }
 
@@ -1110,38 +1053,19 @@ bool SAnimationEditorViewportTabBody::IsAutoAlignFloor() const
 /** Function to set the current playback speed*/
 void SAnimationEditorViewportTabBody::OnSetPlaybackSpeed(int32 PlaybackSpeedMode)
 {
-	AnimationPlaybackSpeedMode = (EAnimationPlaybackSpeeds::Type)PlaybackSpeedMode;
-	
-	float PlaySpeed = EAnimationPlaybackSpeeds::Values[AnimationPlaybackSpeedMode];
-	if(LevelViewportClient->GetWorld())
-	{
-		LevelViewportClient->GetWorld()->GetWorldSettings()->TimeDilation = PlaySpeed;
-	}
+	TSharedRef<FAnimationViewportClient> AnimViewportClient = StaticCastSharedRef<FAnimationViewportClient>(LevelViewportClient.ToSharedRef());
+	AnimViewportClient->SetPlaybackSpeedMode((EAnimationPlaybackSpeeds::Type)PlaybackSpeedMode);
 }
 
 bool SAnimationEditorViewportTabBody::IsPlaybackSpeedSelected(int32 PlaybackSpeedMode)
 {
-	return PlaybackSpeedMode == AnimationPlaybackSpeedMode;
-}
-
-void SAnimationEditorViewportTabBody::ShowReferencePose()
-{
-	PersonaPtr.Pin()->ShowReferencePose(IsShowReferencePoseEnabled() == false);
-}
-
-bool SAnimationEditorViewportTabBody::CanShowReferencePose() const
-{
-	return PersonaPtr.Pin()->CanShowReferencePose();
-}
-
-bool SAnimationEditorViewportTabBody::IsShowReferencePoseEnabled() const
-{
-	return PersonaPtr.Pin()->IsShowReferencePoseEnabled();
+	TSharedRef<FAnimationViewportClient> AnimViewportClient = StaticCastSharedRef<FAnimationViewportClient>(LevelViewportClient.ToSharedRef());
+	return PlaybackSpeedMode == AnimViewportClient->GetPlaybackSpeedMode();
 }
 
 void SAnimationEditorViewportTabBody::ShowRetargetBasePose()
 {
-	UDebugSkelMeshComponent* PreviewComponent = PersonaPtr.Pin()->PreviewComponent;
+	UDebugSkelMeshComponent* PreviewComponent = GetPreviewScene()->GetPreviewMeshComponent();
 	if(PreviewComponent && PreviewComponent->PreviewInstance)
 	{
 		PreviewComponent->PreviewInstance->SetForceRetargetBasePose(!PreviewComponent->PreviewInstance->GetForceRetargetBasePose());
@@ -1150,13 +1074,13 @@ void SAnimationEditorViewportTabBody::ShowRetargetBasePose()
 
 bool SAnimationEditorViewportTabBody::CanShowRetargetBasePose() const
 {
-	UDebugSkelMeshComponent* PreviewComponent = PersonaPtr.Pin()->PreviewComponent;
+	UDebugSkelMeshComponent* PreviewComponent = GetPreviewScene()->GetPreviewMeshComponent();
 	return PreviewComponent != NULL && PreviewComponent->PreviewInstance;
 }
 
 bool SAnimationEditorViewportTabBody::IsShowRetargetBasePoseEnabled() const
 {
-	UDebugSkelMeshComponent* PreviewComponent = PersonaPtr.Pin()->PreviewComponent;
+	UDebugSkelMeshComponent* PreviewComponent = GetPreviewScene()->GetPreviewMeshComponent();
 	if(PreviewComponent && PreviewComponent->PreviewInstance)
 	{
 		return PreviewComponent->PreviewInstance->GetForceRetargetBasePose();
@@ -1169,7 +1093,7 @@ void SAnimationEditorViewportTabBody::ShowBound()
 	TSharedRef<FAnimationViewportClient> AnimViewportClient = StaticCastSharedRef<FAnimationViewportClient>(LevelViewportClient.ToSharedRef());	
 	AnimViewportClient->ToggleShowBounds();
 
-	UDebugSkelMeshComponent* PreviewComponent = PersonaPtr.Pin()->PreviewComponent;
+	UDebugSkelMeshComponent* PreviewComponent = GetPreviewScene()->GetPreviewMeshComponent();
 	if(PreviewComponent)
 	{
 		PreviewComponent->bDisplayBound = AnimViewportClient->EngineShowFlags.Bounds;
@@ -1179,7 +1103,7 @@ void SAnimationEditorViewportTabBody::ShowBound()
 
 bool SAnimationEditorViewportTabBody::CanShowBound() const
 {
-	UDebugSkelMeshComponent* PreviewComponent = PersonaPtr.Pin()->PreviewComponent;
+	UDebugSkelMeshComponent* PreviewComponent = GetPreviewScene()->GetPreviewMeshComponent();
 	return PreviewComponent != NULL;
 }
 
@@ -1191,7 +1115,7 @@ bool SAnimationEditorViewportTabBody::IsShowBoundEnabled() const
 
 void SAnimationEditorViewportTabBody::ToggleShowPreviewMesh()
 {
-	if (UDebugSkelMeshComponent* PreviewComponent = PersonaPtr.Pin()->PreviewComponent)
+	if (UDebugSkelMeshComponent* PreviewComponent = GetPreviewScene()->GetPreviewMeshComponent())
 	{
 		bool bCurrentlyVisible = IsShowPreviewMeshEnabled();
 		PreviewComponent->SetVisibility(!bCurrentlyVisible);
@@ -1200,19 +1124,19 @@ void SAnimationEditorViewportTabBody::ToggleShowPreviewMesh()
 
 bool SAnimationEditorViewportTabBody::CanShowPreviewMesh() const
 {
-	UDebugSkelMeshComponent* PreviewComponent = PersonaPtr.Pin()->PreviewComponent;
+	UDebugSkelMeshComponent* PreviewComponent = GetPreviewScene()->GetPreviewMeshComponent();
 	return PreviewComponent != NULL;
 }
 
 bool SAnimationEditorViewportTabBody::IsShowPreviewMeshEnabled() const
 {
-	UDebugSkelMeshComponent* PreviewComponent = PersonaPtr.Pin()->PreviewComponent;
+	UDebugSkelMeshComponent* PreviewComponent = GetPreviewScene()->GetPreviewMeshComponent();
 	return (PreviewComponent != NULL) && PreviewComponent->IsVisible();
 }
 
 void SAnimationEditorViewportTabBody::UseInGameBound()
 {
-	UDebugSkelMeshComponent* PreviewComponent = PersonaPtr.Pin()->PreviewComponent;
+	UDebugSkelMeshComponent* PreviewComponent = GetPreviewScene()->GetPreviewMeshComponent();
 	if (PreviewComponent != NULL)
 	{
 		PreviewComponent->UseInGameBounds(! PreviewComponent->IsUsingInGameBounds());
@@ -1221,35 +1145,18 @@ void SAnimationEditorViewportTabBody::UseInGameBound()
 
 bool SAnimationEditorViewportTabBody::CanUseInGameBound() const
 {
-	UDebugSkelMeshComponent* PreviewComponent = PersonaPtr.Pin()->PreviewComponent;
+	UDebugSkelMeshComponent* PreviewComponent = GetPreviewScene()->GetPreviewMeshComponent();
 	return PreviewComponent != NULL;
 }
 
 bool SAnimationEditorViewportTabBody::IsUsingInGameBound() const
 {
-	UDebugSkelMeshComponent* PreviewComponent = PersonaPtr.Pin()->PreviewComponent;
+	UDebugSkelMeshComponent* PreviewComponent = GetPreviewScene()->GetPreviewMeshComponent();
 	return PreviewComponent != NULL && PreviewComponent->IsUsingInGameBounds();
 }
 
-void SAnimationEditorViewportTabBody::SetPreviewComponent(class UDebugSkelMeshComponent* PreviewComponent)
+void SAnimationEditorViewportTabBody::HandlePreviewMeshChanged(class USkeletalMesh* SkeletalMesh)
 {
-	PreviewComponent->MeshComponentUpdateFlag = EMeshComponentUpdateFlag::AlwaysTickPoseAndRefreshBones;
-	PreviewComponent->bCanHighlightSelectedSections = true;
-
-	//Set preview skeleton mesh component
-	TSharedRef<FAnimationViewportClient> AnimViewportClient = StaticCastSharedRef<FAnimationViewportClient>(LevelViewportClient.ToSharedRef());
-	AnimViewportClient->SetPreviewMeshComponent(PreviewComponent);
-
-	AnimViewportClient->FocusViewportOnPreviewMesh();
-
-	// Adding the component to the PreviewScene can change AnimScriptInstance so we must re register it
-	// with the AnimBlueprint
-	UAnimBlueprint* SourceBlueprint = Cast<UAnimBlueprint>(PersonaPtr.Pin()->GetBlueprintObj());
-	if ( SourceBlueprint && PreviewComponent->IsAnimBlueprintInstanced() )
-	{
-		SourceBlueprint->SetObjectBeingDebugged(PreviewComponent->GetAnimInstance());
-	}
-
 	PopulateNumUVChannels();
 }
 
@@ -1272,7 +1179,7 @@ void SAnimationEditorViewportTabBody::PopulateNumUVChannels()
 {
 	NumUVChannels.Empty();
 
-	if (UDebugSkelMeshComponent* PreviewComponent = PersonaPtr.Pin()->PreviewComponent)
+	if (UDebugSkelMeshComponent* PreviewComponent = GetPreviewScene()->GetPreviewMeshComponent())
 	{
 		if (FSkeletalMeshResource* MeshResource = PreviewComponent->GetSkeletalMeshResource())
 		{
@@ -1293,7 +1200,7 @@ void SAnimationEditorViewportTabBody::PopulateUVChoices()
 	// Fill out the UV channels combo.
 	UVChannels.Empty();
 	
-	if (UDebugSkelMeshComponent* PreviewComponent = PersonaPtr.Pin()->PreviewComponent)
+	if (UDebugSkelMeshComponent* PreviewComponent = GetPreviewScene()->GetPreviewMeshComponent())
 	{
 		int32 CurrentLOD = FMath::Clamp(PreviewComponent->ForcedLodModel - 1, 0, NumUVChannels.Num() - 1);
 
@@ -1333,8 +1240,7 @@ void SAnimationEditorViewportTabBody::UpdateScrubPanel(UAnimationAsset* AnimAsse
 			ScrubPanelContainer->AddSlot()
 				.AutoHeight()
 				[
-					SNew(SAnimMontageScrubPanel)
-					.Persona(PersonaPtr)
+					SNew(SAnimMontageScrubPanel, GetPreviewScene())
 					.ViewInputMin(this, &SAnimationEditorViewportTabBody::GetViewMinInput)
 					.ViewInputMax(this, &SAnimationEditorViewportTabBody::GetViewMaxInput)
 					.bAllowZoom(true)
@@ -1346,8 +1252,7 @@ void SAnimationEditorViewportTabBody::UpdateScrubPanel(UAnimationAsset* AnimAsse
 			ScrubPanelContainer->AddSlot()
 				.AutoHeight()
 				[
-					SNew(SAnimationScrubPanel)
-					.Persona(PersonaPtr)
+					SNew(SAnimationScrubPanel, GetPreviewScene())
 					.ViewInputMin(this, &SAnimationEditorViewportTabBody::GetViewMinInput)
 					.ViewInputMax(this, &SAnimationEditorViewportTabBody::GetViewMaxInput)
 					.bAllowZoom(true)
@@ -1358,9 +1263,9 @@ void SAnimationEditorViewportTabBody::UpdateScrubPanel(UAnimationAsset* AnimAsse
 
 float SAnimationEditorViewportTabBody::GetViewMinInput() const
 {
-	if (UDebugSkelMeshComponent* PreviewComponent = PersonaPtr.Pin()->PreviewComponent)
+	if (UDebugSkelMeshComponent* PreviewComponent = GetPreviewScene()->GetPreviewMeshComponent())
 	{
-		UObject* PreviewAsset = PersonaPtr.Pin()->GetPreviewAnimationAsset();
+		UObject* PreviewAsset = GetPreviewScene()->GetPreviewAnimationAsset();
 		if (PreviewAsset != NULL)
 		{
 			return 0.0f;
@@ -1376,10 +1281,10 @@ float SAnimationEditorViewportTabBody::GetViewMinInput() const
 
 float SAnimationEditorViewportTabBody::GetViewMaxInput() const
 { 
-	UDebugSkelMeshComponent* PreviewComponent = PersonaPtr.Pin()->PreviewComponent;
+	UDebugSkelMeshComponent* PreviewComponent = GetPreviewScene()->GetPreviewMeshComponent();
 	if (PreviewComponent != NULL)
 	{
-		UObject* PreviewAsset = PersonaPtr.Pin()->GetPreviewAnimationAsset();
+		UObject* PreviewAsset = GetPreviewScene()->GetPreviewAnimationAsset();
 		if ((PreviewAsset != NULL) && (PreviewComponent->PreviewInstance != NULL))
 		{
 			return PreviewComponent->PreviewInstance->GetLength();
@@ -1396,7 +1301,7 @@ float SAnimationEditorViewportTabBody::GetViewMaxInput() const
 void SAnimationEditorViewportTabBody::UpdateShowFlagForMeshEdges()
 {
 	bool bUseOverlayMaterial = false;
-	if (UDebugSkelMeshComponent* PreviewComponent = PersonaPtr.Pin()->PreviewComponent)
+	if (UDebugSkelMeshComponent* PreviewComponent = GetPreviewScene()->GetPreviewMeshComponent())
 	{
 		bUseOverlayMaterial = PreviewComponent->bDrawBoneInfluences || PreviewComponent->bDrawMorphTargetVerts;
 	}
@@ -1419,7 +1324,7 @@ void SAnimationEditorViewportTabBody::OnSetLODModel(int32 LODSelectionType)
 {
 	LODSelection = LODSelectionType;
 
-	UDebugSkelMeshComponent* PreviewComponent = PersonaPtr.Pin()->PreviewComponent;
+	UDebugSkelMeshComponent* PreviewComponent = GetPreviewScene()->GetPreviewMeshComponent();
 	if( PreviewComponent )
 	{
 		PreviewComponent->ForcedLodModel = LODSelection;
@@ -1457,6 +1362,11 @@ TSharedRef<FAnimationViewportClient> SAnimationEditorViewportTabBody::GetAnimati
 	return StaticCastSharedRef<FAnimationViewportClient>(LevelViewportClient.ToSharedRef());
 }
 
+void SAnimationEditorViewportTabBody::OpenPreviewSceneSettings()
+{
+	OnInvokeTab.ExecuteIfBound(FPersonaTabs::AdvancedPreviewSceneSettingsID);
+}
+
 void SAnimationEditorViewportTabBody::ToggleCameraFollow()
 {
 	// Switch to rotation mode
@@ -1478,78 +1388,6 @@ bool SAnimationEditorViewportTabBody::CanChangeCameraMode() const
 	return !LevelViewportClient->IsOrtho();
 }
 
-void SAnimationEditorViewportTabBody::SaveData(class SAnimationEditorViewportTabBody* OldViewport)
-{
-	if ( PersonaPtr.IsValid() && OldViewport )
-	{
-		FPersonaModeSharedData& SharedData = PersonaPtr.Pin()->ModeSharedData;
-
-		TSharedRef<FAnimationViewportClient> OldAnimViewportClient = StaticCastSharedRef<FAnimationViewportClient>(OldViewport->LevelViewportClient.ToSharedRef());
-		// set camera set up
-		SharedData.ViewLocation = OldAnimViewportClient.Get().GetViewLocation();
-		SharedData.ViewRotation = OldAnimViewportClient.Get().GetViewRotation();
-		SharedData.LookAtLocation = OldAnimViewportClient.Get().GetLookAtLocation();
-		SharedData.OrthoZoom = OldAnimViewportClient.Get().GetOrthoZoom();
-		SharedData.bCameraLock = OldAnimViewportClient.Get().IsCameraLocked();
-		SharedData.bCameraFollow = OldAnimViewportClient.Get().IsSetCameraFollowChecked();
-		SharedData.bShowBound = OldAnimViewportClient.Get().IsSetShowBoundsChecked();
-		SharedData.LocalAxesMode = OldAnimViewportClient.Get().GetLocalAxesMode();
-		SharedData.ViewportType = OldAnimViewportClient.Get().ViewportType;
-		SharedData.PlaybackSpeedMode = OldViewport->AnimationPlaybackSpeedMode;
-	}
-}
-
-void SAnimationEditorViewportTabBody::RestoreData()
-{
-	if ( PersonaPtr.IsValid() )
-	{
-		FPersonaModeSharedData& SharedData = PersonaPtr.Pin()->ModeSharedData;
-		TSharedRef<FAnimationViewportClient> AnimViewportClient = StaticCastSharedRef<FAnimationViewportClient>(LevelViewportClient.ToSharedRef());
-
-
-		AnimViewportClient.Get().SetViewportType((ELevelViewportType)SharedData.ViewportType);
-		AnimViewportClient.Get().SetViewLocation( SharedData.ViewLocation );
-		AnimViewportClient.Get().SetViewRotation( SharedData.ViewRotation );
-		AnimViewportClient.Get().SetShowBounds(SharedData.bShowBound);
-		AnimViewportClient.Get().SetLocalAxesMode((ELocalAxesMode::Type)SharedData.LocalAxesMode);
-		AnimViewportClient.Get().SetOrthoZoom(SharedData.OrthoZoom);
-
-		OnSetPlaybackSpeed(SharedData.PlaybackSpeedMode);
-
-		if (SharedData.bCameraLock)
-		{
-			AnimViewportClient.Get().SetLookAtLocation( SharedData.LookAtLocation );
-		}
-		else if(SharedData.bCameraFollow)
-		{
-			AnimViewportClient.Get().SetCameraFollow();
-		}
-	}
-}
-
-void SAnimationEditorViewportTabBody::OnLODChanged()
-{
-	int32 LodCount = 0;
-
-	UDebugSkelMeshComponent* PreviewComponent = PersonaPtr.Pin()->PreviewComponent;
-
-	if( PreviewComponent && PreviewComponent->SkeletalMesh )
-	{
-		LodCount = PreviewComponent->SkeletalMesh->GetImportedResource()->LODModels.Num();
-	}
-
-	// If the currently selected LoD is invalid, revert to LOD_Auto
-	if ( LodCount < LODSelection )
-	{
-		OnSetLODModel( 0 );
-	}
-
-	if (PersonaPtr.Pin()->PersonaMeshDetailLayout)
-	{
-		PersonaPtr.Pin()->PersonaMeshDetailLayout->ForceRefreshDetails();
-	}
-}
-
 void SAnimationEditorViewportTabBody::OnMuteAudio()
 {
 	GetAnimationViewportClient()->OnToggleMuteAudio();
@@ -1562,7 +1400,7 @@ bool SAnimationEditorViewportTabBody::IsAudioMuted()
 
 void SAnimationEditorViewportTabBody::OnTogglePreviewRootMotion()
 {
-	UDebugSkelMeshComponent* PreviewComponent = PersonaPtr.Pin()->PreviewComponent;
+	UDebugSkelMeshComponent* PreviewComponent = GetPreviewScene()->GetPreviewMeshComponent();
 
 	if (PreviewComponent)
 	{
@@ -1572,7 +1410,7 @@ void SAnimationEditorViewportTabBody::OnTogglePreviewRootMotion()
 
 bool SAnimationEditorViewportTabBody::IsPreviewingRootMotion() const
 {
-	UDebugSkelMeshComponent* PreviewComponent = PersonaPtr.Pin()->PreviewComponent;
+	UDebugSkelMeshComponent* PreviewComponent = GetPreviewScene()->GetPreviewMeshComponent();
 
 	if (PreviewComponent)
 	{
@@ -1584,7 +1422,7 @@ bool SAnimationEditorViewportTabBody::IsPreviewingRootMotion() const
 #if WITH_APEX_CLOTHING
 bool SAnimationEditorViewportTabBody::IsDisablingClothSimulation() const
 {
-	UDebugSkelMeshComponent* PreviewComponent = PersonaPtr.Pin()->PreviewComponent;
+	UDebugSkelMeshComponent* PreviewComponent = GetPreviewScene()->GetPreviewMeshComponent();
 
 	if( PreviewComponent )
 	{
@@ -1596,7 +1434,7 @@ bool SAnimationEditorViewportTabBody::IsDisablingClothSimulation() const
 
 void SAnimationEditorViewportTabBody::OnDisableClothSimulation()
 {
-	UDebugSkelMeshComponent* PreviewComponent = PersonaPtr.Pin()->PreviewComponent;
+	UDebugSkelMeshComponent* PreviewComponent = GetPreviewScene()->GetPreviewMeshComponent();
 
 	if( PreviewComponent )
 	{
@@ -1620,47 +1458,40 @@ void SAnimationEditorViewportTabBody::OnDisableClothSimulation()
 
 bool SAnimationEditorViewportTabBody::IsApplyingClothWind() const
 {
-	UDebugSkelMeshComponent* PreviewComponent = PersonaPtr.Pin()->PreviewComponent;
-
-	if( PreviewComponent )
-	{
-		return PreviewComponent->IsWindEnabled();
-	}
-
-	return false;
+	return GetPreviewScene()->IsWindEnabled();
 }
 
 void SAnimationEditorViewportTabBody::OnApplyClothWind()
 {
-	UDebugSkelMeshComponent* PreviewComponent = PersonaPtr.Pin()->PreviewComponent;
-
-	if( PreviewComponent )
-	{
-		PreviewComponent->bEnableWind = !PreviewComponent->IsWindEnabled();
-		GetAnimationViewportClient()->EnableWindActor(PreviewComponent->IsWindEnabled());
-		RefreshViewport();
-	}
+	GetPreviewScene()->EnableWind(!GetPreviewScene()->IsWindEnabled());
+	RefreshViewport();
 }
 
 void SAnimationEditorViewportTabBody::SetWindStrength(float SliderPos)
 {
-	GetAnimationViewportClient()->SetWindStrength( SliderPos );
+	GetPreviewScene()->SetWindStrength(SliderPos);
 	RefreshViewport();
 }
 
 float SAnimationEditorViewportTabBody::GetWindStrengthSliderValue() const
 {
-	return GetAnimationViewportClient()->GetWindStrengthSliderValue();
+	return GetPreviewScene()->GetWindStrength();
 }
 
 FText SAnimationEditorViewportTabBody::GetWindStrengthLabel() const
 {
-	return GetAnimationViewportClient()->GetWindStrengthLabel();
+	//Clamp slide value so that minimum value displayed is 0.00 and maximum is 1.0
+	float SliderValue = FMath::Clamp<float>(GetWindStrengthSliderValue(), 0.0f, 1.0f);
+
+	static const FNumberFormattingOptions FormatOptions = FNumberFormattingOptions()
+		.SetMinimumFractionalDigits(2)
+		.SetMaximumFractionalDigits(2);
+	return FText::AsNumber(SliderValue, &FormatOptions);
 }
 
 void SAnimationEditorViewportTabBody::OnShowClothSimulationNormals()
 {
-	UDebugSkelMeshComponent* PreviewComponent = PersonaPtr.Pin()->PreviewComponent;
+	UDebugSkelMeshComponent* PreviewComponent = GetPreviewScene()->GetPreviewMeshComponent();
 
 	if( PreviewComponent )
 	{
@@ -1671,7 +1502,7 @@ void SAnimationEditorViewportTabBody::OnShowClothSimulationNormals()
 
 bool SAnimationEditorViewportTabBody::IsShowingClothSimulationNormals() const
 {
-	UDebugSkelMeshComponent* PreviewComponent = PersonaPtr.Pin()->PreviewComponent;
+	UDebugSkelMeshComponent* PreviewComponent = GetPreviewScene()->GetPreviewMeshComponent();
 
 	if( PreviewComponent )
 	{
@@ -1683,7 +1514,7 @@ bool SAnimationEditorViewportTabBody::IsShowingClothSimulationNormals() const
 
 void SAnimationEditorViewportTabBody::OnShowClothGraphicalTangents()
 {
-	UDebugSkelMeshComponent* PreviewComponent = PersonaPtr.Pin()->PreviewComponent;
+	UDebugSkelMeshComponent* PreviewComponent = GetPreviewScene()->GetPreviewMeshComponent();
 
 	if( PreviewComponent )
 	{
@@ -1694,7 +1525,7 @@ void SAnimationEditorViewportTabBody::OnShowClothGraphicalTangents()
 
 bool SAnimationEditorViewportTabBody::IsShowingClothGraphicalTangents() const
 {
-	UDebugSkelMeshComponent* PreviewComponent = PersonaPtr.Pin()->PreviewComponent;
+	UDebugSkelMeshComponent* PreviewComponent = GetPreviewScene()->GetPreviewMeshComponent();
 
 	if( PreviewComponent )
 	{
@@ -1706,7 +1537,7 @@ bool SAnimationEditorViewportTabBody::IsShowingClothGraphicalTangents() const
 
 void SAnimationEditorViewportTabBody::OnShowClothCollisionVolumes()
 {
-	UDebugSkelMeshComponent* PreviewComponent = PersonaPtr.Pin()->PreviewComponent;
+	UDebugSkelMeshComponent* PreviewComponent = GetPreviewScene()->GetPreviewMeshComponent();
 
 	if( PreviewComponent )
 	{
@@ -1717,7 +1548,7 @@ void SAnimationEditorViewportTabBody::OnShowClothCollisionVolumes()
 
 bool SAnimationEditorViewportTabBody::IsShowingClothCollisionVolumes() const
 {
-	UDebugSkelMeshComponent* PreviewComponent = PersonaPtr.Pin()->PreviewComponent;
+	UDebugSkelMeshComponent* PreviewComponent = GetPreviewScene()->GetPreviewMeshComponent();
 
 	if( PreviewComponent )
 	{
@@ -1729,23 +1560,29 @@ bool SAnimationEditorViewportTabBody::IsShowingClothCollisionVolumes() const
 
 void SAnimationEditorViewportTabBody::SetGravityScale(float SliderPos)
 {
-	GetAnimationViewportClient()->SetGravityScale(SliderPos);
+	GetPreviewScene()->SetGravityScale(SliderPos);
 	RefreshViewport();
 }
 
 float SAnimationEditorViewportTabBody::GetGravityScaleSliderValue() const
 {
-	return GetAnimationViewportClient()->GetGravityScaleSliderValue();
+	return GetPreviewScene()->GetGravityScale();
 }
 
 FText SAnimationEditorViewportTabBody::GetGravityScaleLabel() const
 {
-	return GetAnimationViewportClient()->GetGravityScaleLabel();
+	//Clamp slide value so that minimum value displayed is 0.00 and maximum is 4.0
+	float SliderValue = FMath::Clamp<float>(GetGravityScaleSliderValue() * 4, 0.0f, 4.0f);
+
+	static const FNumberFormattingOptions FormatOptions = FNumberFormattingOptions()
+		.SetMinimumFractionalDigits(2)
+		.SetMaximumFractionalDigits(2);
+	return FText::AsNumber(SliderValue, &FormatOptions);
 }
 
 void SAnimationEditorViewportTabBody::OnEnableCollisionWithAttachedClothChildren()
 {
-	UDebugSkelMeshComponent* PreviewComponent = PersonaPtr.Pin()->PreviewComponent;
+	UDebugSkelMeshComponent* PreviewComponent = GetPreviewScene()->GetPreviewMeshComponent();
 
 	if( PreviewComponent )
 	{
@@ -1756,7 +1593,7 @@ void SAnimationEditorViewportTabBody::OnEnableCollisionWithAttachedClothChildren
 
 bool SAnimationEditorViewportTabBody::IsEnablingCollisionWithAttachedClothChildren() const
 {
-	UDebugSkelMeshComponent* PreviewComponent = PersonaPtr.Pin()->PreviewComponent;
+	UDebugSkelMeshComponent* PreviewComponent = GetPreviewScene()->GetPreviewMeshComponent();
 
 	if( PreviewComponent )
 	{
@@ -1768,7 +1605,7 @@ bool SAnimationEditorViewportTabBody::IsEnablingCollisionWithAttachedClothChildr
 
 void SAnimationEditorViewportTabBody::OnShowClothPhysicalMeshWire()
 {
-	UDebugSkelMeshComponent* PreviewComponent = PersonaPtr.Pin()->PreviewComponent;
+	UDebugSkelMeshComponent* PreviewComponent = GetPreviewScene()->GetPreviewMeshComponent();
 
 	if( PreviewComponent )
 	{
@@ -1779,7 +1616,7 @@ void SAnimationEditorViewportTabBody::OnShowClothPhysicalMeshWire()
 
 bool SAnimationEditorViewportTabBody::IsShowingClothPhysicalMeshWire() const
 {
-	UDebugSkelMeshComponent* PreviewComponent = PersonaPtr.Pin()->PreviewComponent;
+	UDebugSkelMeshComponent* PreviewComponent = GetPreviewScene()->GetPreviewMeshComponent();
 
 	if( PreviewComponent )
 	{
@@ -1791,7 +1628,7 @@ bool SAnimationEditorViewportTabBody::IsShowingClothPhysicalMeshWire() const
 
 void SAnimationEditorViewportTabBody::OnShowClothMaxDistances()
 {
-	UDebugSkelMeshComponent* PreviewComponent = PersonaPtr.Pin()->PreviewComponent;
+	UDebugSkelMeshComponent* PreviewComponent = GetPreviewScene()->GetPreviewMeshComponent();
 
 	if( PreviewComponent )
 	{
@@ -1816,7 +1653,7 @@ void SAnimationEditorViewportTabBody::OnShowClothMaxDistances()
 
 bool SAnimationEditorViewportTabBody::IsShowingClothMaxDistances() const
 {
-	UDebugSkelMeshComponent* PreviewComponent = PersonaPtr.Pin()->PreviewComponent;
+	UDebugSkelMeshComponent* PreviewComponent = GetPreviewScene()->GetPreviewMeshComponent();
 
 	if( PreviewComponent )
 	{
@@ -1828,7 +1665,7 @@ bool SAnimationEditorViewportTabBody::IsShowingClothMaxDistances() const
 
 void SAnimationEditorViewportTabBody::OnShowClothBackstops()
 {
-	UDebugSkelMeshComponent* PreviewComponent = PersonaPtr.Pin()->PreviewComponent;
+	UDebugSkelMeshComponent* PreviewComponent = GetPreviewScene()->GetPreviewMeshComponent();
 
 	if( PreviewComponent )
 	{
@@ -1851,7 +1688,7 @@ void SAnimationEditorViewportTabBody::OnShowClothBackstops()
 
 bool SAnimationEditorViewportTabBody::IsShowingClothBackstops() const
 {
-	if (UDebugSkelMeshComponent* PreviewComponent = PersonaPtr.Pin()->PreviewComponent)
+	if (UDebugSkelMeshComponent* PreviewComponent = GetPreviewScene()->GetPreviewMeshComponent())
 	{
 		return PreviewComponent->bDisplayClothBackstops;
 	}
@@ -1861,7 +1698,7 @@ bool SAnimationEditorViewportTabBody::IsShowingClothBackstops() const
 
 void SAnimationEditorViewportTabBody::OnShowClothFixedVertices()
 {
-	if (UDebugSkelMeshComponent* PreviewComponent = PersonaPtr.Pin()->PreviewComponent)
+	if (UDebugSkelMeshComponent* PreviewComponent = GetPreviewScene()->GetPreviewMeshComponent())
 	{
 		PreviewComponent->bDisplayClothFixedVertices = !PreviewComponent->bDisplayClothFixedVertices;
 		RefreshViewport();
@@ -1870,7 +1707,7 @@ void SAnimationEditorViewportTabBody::OnShowClothFixedVertices()
 
 bool SAnimationEditorViewportTabBody::IsShowingClothFixedVertices() const
 {
-	if (UDebugSkelMeshComponent* PreviewComponent = PersonaPtr.Pin()->PreviewComponent)
+	if (UDebugSkelMeshComponent* PreviewComponent = GetPreviewScene()->GetPreviewMeshComponent())
 	{
 		return PreviewComponent->bDisplayClothFixedVertices;
 	}
@@ -1880,7 +1717,7 @@ bool SAnimationEditorViewportTabBody::IsShowingClothFixedVertices() const
 
 void SAnimationEditorViewportTabBody::OnSetSectionsDisplayMode(int32 DisplayMode)
 {
-	UDebugSkelMeshComponent* PreviewComponent = PersonaPtr.Pin()->PreviewComponent;
+	UDebugSkelMeshComponent* PreviewComponent = GetPreviewScene()->GetPreviewMeshComponent();
 
 	if (!PreviewComponent)
 	{
@@ -1910,7 +1747,7 @@ void SAnimationEditorViewportTabBody::OnSetSectionsDisplayMode(int32 DisplayMode
 
 bool SAnimationEditorViewportTabBody::IsSectionsDisplayMode(int32 DisplayMode) const
 {
-	UDebugSkelMeshComponent* PreviewComponent = PersonaPtr.Pin()->PreviewComponent;
+	UDebugSkelMeshComponent* PreviewComponent = GetPreviewScene()->GetPreviewMeshComponent();
 
 	if (ensure(PreviewComponent))
 	{
@@ -1921,47 +1758,15 @@ bool SAnimationEditorViewportTabBody::IsSectionsDisplayMode(int32 DisplayMode) c
 }
 #endif // #if WITH_APEX_CLOTHING
 
-EVisibility SAnimationEditorViewportTabBody::GetViewportCornerImageVisibility() const
-{
-	TSharedPtr<FPersona> Persona = PersonaPtr.Pin();	
-	return Persona->IsRecording() ? EVisibility::Visible : EVisibility::Collapsed;
-}
-
-const FSlateBrush* SAnimationEditorViewportTabBody::GetViewportCornerImage() const
-{
-	static int32 Count=0;
-
-	TSharedPtr<FPersona> Persona = PersonaPtr.Pin();
-//	if(Persona->Recorder.InRecording())
-	{
-		if(Count++ < 5)
-		{
-			return FEditorStyle::GetBrush("Persona.StopRecordAnimation");
-		}
-		else
-		{
-			if(Count == 10)
-			{
-				Count = 0;
-			}
-
-			return FEditorStyle::GetBrush("Persona.StopRecordAnimation_Alt");
-		}
-	}
-
-//	return NULL;
-}
-
 EVisibility SAnimationEditorViewportTabBody::GetViewportCornerTextVisibility() const
 {
-	TSharedPtr<FPersona> Persona = PersonaPtr.Pin();
-	if (Persona->IsRecording())
+	if (GetPreviewScene()->IsRecording())
 	{
 		return EVisibility::Visible;
 	}
-	else if (Persona->IsModeCurrent(FPersonaModes::AnimBlueprintEditMode))
+	else if (BlueprintEditorPtr.IsValid() && BlueprintEditorPtr.Pin()->IsModeCurrent(FPersonaModes::AnimBlueprintEditMode))
 	{
-		if (UBlueprint* Blueprint = Persona->GetBlueprintObj())
+		if (UBlueprint* Blueprint = BlueprintEditorPtr.Pin()->GetBlueprintObj())
 		{
 			const bool bUpToDate = (Blueprint->Status == BS_UpToDate) || (Blueprint->Status == BS_UpToDateWithWarnings);
 			return bUpToDate ? EVisibility::Collapsed : EVisibility::Visible;
@@ -1973,21 +1778,21 @@ EVisibility SAnimationEditorViewportTabBody::GetViewportCornerTextVisibility() c
 
 FText SAnimationEditorViewportTabBody::GetViewportCornerText() const
 {
-	TSharedPtr<FPersona> Persona = PersonaPtr.Pin();
-	if(Persona->IsRecording())
+	if(GetPreviewScene()->IsRecording())
 	{
-		UAnimSequence* Recording = Persona->GetCurrentRecording();
+		UAnimSequence* Recording = GetPreviewScene()->GetCurrentRecording();
 		const FString& Name = Recording ? Recording->GetName() : TEXT("None");
-		float TimeRecorded = Persona->GetCurrentRecordingTime();
+		float TimeRecorded = GetPreviewScene()->GetCurrentRecordingTime();
 		FNumberFormattingOptions NumberOption;
 		NumberOption.MaximumFractionalDigits = 2;
 		NumberOption.MinimumFractionalDigits = 2;
-		return FText::Format(LOCTEXT("AnimRecorder", "Recording '{0}' - Time {1} sec(s)]\nTo stop, click here. "),
+		return FText::Format(LOCTEXT("AnimRecorder", "Recording '{0}' [{1} sec(s)]"),
 			FText::FromString(Name), FText::AsNumber(TimeRecorded, &NumberOption));
 	}
-	if (Persona->IsModeCurrent(FPersonaModes::AnimBlueprintEditMode))
+
+	if (BlueprintEditorPtr.IsValid() && BlueprintEditorPtr.Pin()->IsModeCurrent(FPersonaModes::AnimBlueprintEditMode))
 	{
-		if (UBlueprint* Blueprint = Persona->GetBlueprintObj())
+		if (UBlueprint* Blueprint = BlueprintEditorPtr.Pin()->GetBlueprintObj())
 		{
 			switch (Blueprint->Status)
 			{
@@ -2010,31 +1815,30 @@ FText SAnimationEditorViewportTabBody::GetViewportCornerText() const
 
 FText SAnimationEditorViewportTabBody::GetViewportCornerTooltip() const
 {
-	TSharedPtr<FPersona> Persona = PersonaPtr.Pin();
-	if(Persona->IsRecording())
+	if(GetPreviewScene()->IsRecording())
 	{
 		return LOCTEXT("RecordingStatusTooltip", "Shows the status of animation recording.\nClick to stop the recording.");
 	}
-	if(Persona->IsModeCurrent(FPersonaModes::AnimBlueprintEditMode))
+
+	if(BlueprintEditorPtr.IsValid() && BlueprintEditorPtr.Pin()->IsModeCurrent(FPersonaModes::AnimBlueprintEditMode))
 	{
 		return LOCTEXT("BlueprintStatusTooltip", "Shows the status of the animation blueprint.\nClick to recompile a dirty blueprint");
 	}
 
-return FText::GetEmpty();
+	return FText::GetEmpty();
 }
 
 FReply SAnimationEditorViewportTabBody::ClickedOnViewportCornerText()
 {
-	TSharedPtr<FPersona> Persona = PersonaPtr.Pin();
-	// if it's recording, it won't be able to see the message
-	// so disable it
-	if (Persona->IsRecording())
-	{
-		Persona->StopRecording();
-	}
-	else 
+	if(BlueprintEditorPtr.IsValid())
 	{
-		Persona->RecompileAnimBlueprintIfDirty();
+		if (UBlueprint* Blueprint = BlueprintEditorPtr.Pin()->GetBlueprintObj())
+		{
+			if (!Blueprint->IsUpToDate())
+			{
+				BlueprintEditorPtr.Pin()->Compile();
+			}
+		}
 	}
 
 	return FReply::Handled();
diff --git a/Engine/Source/Editor/Persona/Private/SAnimationEditorViewport.h b/Engine/Source/Editor/Persona/Private/SAnimationEditorViewport.h
--- a/Engine/Source/Editor/Persona/Private/SAnimationEditorViewport.h
+++ b/Engine/Source/Editor/Persona/Private/SAnimationEditorViewport.h
@@ -5,25 +5,27 @@
 
 #include "Persona.h"
 #include "SEditorViewport.h"
+#include "IPersonaViewport.h"
 
-//////////////////////////////////////////////////////////////////////////
-// EAnimationPlaybackSpeeds
-
-namespace EAnimationPlaybackSpeeds
+struct FAnimationEditorViewportRequiredArgs
 {
-	enum Type
-	{
-		OneTenth=0,
-		Quarter,
-		Half,
-		Normal,
-		Double,
-		FiveTimes,
-		TenTimes,
-		NumPlaybackSpeeds
-	};
-
-	extern float Values[NumPlaybackSpeeds];
+	FAnimationEditorViewportRequiredArgs(const TSharedRef<class ISkeletonTree>& InSkeletonTree, const TSharedRef<class IPersonaPreviewScene>& InPreviewScene, TSharedRef<class SAnimationEditorViewportTabBody> InTabBody, TSharedRef<class FAssetEditorToolkit> InAssetEditorToolkit, FSimpleMulticastDelegate& InOnPostUndo)
+		: SkeletonTree(InSkeletonTree)
+		, PreviewScene(InPreviewScene)
+		, TabBody(InTabBody)
+		, AssetEditorToolkit(InAssetEditorToolkit)
+		, OnPostUndo(InOnPostUndo)
+	{}
+
+	TSharedRef<class ISkeletonTree> SkeletonTree;
+
+	TSharedRef<class IPersonaPreviewScene> PreviewScene;
+
+	TSharedRef<class SAnimationEditorViewportTabBody> TabBody;
+
+	TSharedRef<class FAssetEditorToolkit> AssetEditorToolkit;
+
+	FSimpleMulticastDelegate& OnPostUndo;
 };
 
 //////////////////////////////////////////////////////////////////////////
@@ -32,12 +34,15 @@ namespace EAnimationPlaybackSpeeds
 class SAnimationEditorViewport : public SEditorViewport
 {
 public:
-	SLATE_BEGIN_ARGS(SAnimationEditorViewport) {}
-	SLATE_END_ARGS()
+	SLATE_BEGIN_ARGS(SAnimationEditorViewport) 
+		: _ShowStats(true)
+	{}
 
-	~SAnimationEditorViewport();
+	SLATE_ARGUMENT(bool, ShowStats)
+
+	SLATE_END_ARGS()
 
-	void Construct(const FArguments& InArgs, TSharedPtr<class FPersona> InPersona, TSharedPtr<class SAnimationEditorViewportTabBody> InTabBody);
+	void Construct(const FArguments& InArgs, const FAnimationEditorViewportRequiredArgs& InRequiredArgs);
 
 protected:
 	// SEditorViewport interface
@@ -55,32 +60,48 @@ protected:
 	// Pointer to the compound widget that owns this viewport widget
 	TWeakPtr<class SAnimationEditorViewportTabBody> TabBodyPtr;
 
-	// Pointer back to the Persona tool that owns us
-	TWeakPtr<class FPersona> PersonaPtr;
+	// The preview scene that we are viewing
+	TWeakPtr<class IPersonaPreviewScene> PreviewScenePtr;
+
+	// The skeleton tree we are editing
+	TWeakPtr<class ISkeletonTree> SkeletonTreePtr;
+
+	// The asset editor we are embedded in
+	TWeakPtr<class FAssetEditorToolkit> AssetEditorToolkitPtr;
+
+	/** Whether we should show stats for this viewport */
+	bool bShowStats;
 };
 
 //////////////////////////////////////////////////////////////////////////
 // SAnimationEditorViewportTabBody
 
-class SAnimationEditorViewportTabBody : public SCompoundWidget
+class SAnimationEditorViewportTabBody : public IPersonaViewport
 {
 public:
 	SLATE_BEGIN_ARGS( SAnimationEditorViewportTabBody )
-		: _Persona()
-		, _Skeleton()
-		, _IsEditable(true)
+		: _BlueprintEditor()
+		, _ShowTimeline(true)
 		{}
 
-		SLATE_ARGUMENT( TSharedPtr<FPersona>, Persona )
-		SLATE_ARGUMENT( USkeleton*, Skeleton )
-		SLATE_ATTRIBUTE( bool, IsEditable )
+		SLATE_ARGUMENT(TWeakPtr<FBlueprintEditor>, BlueprintEditor)
+
+		SLATE_ARGUMENT(FOnInvokeTab, OnInvokeTab)
+
+		SLATE_ARGUMENT(bool, ShowTimeline)
+
+		SLATE_ARGUMENT(bool, ShowStats)
 	SLATE_END_ARGS()
 public:
 
-	void Construct(const FArguments& InArgs);
+	void Construct(const FArguments& InArgs, const TSharedRef<class ISkeletonTree>& InSkeletonTree, const TSharedRef<class IPersonaPreviewScene>& InPreviewScene, const TSharedRef<class FAssetEditorToolkit>& InAssetEditorToolkit, FSimpleMulticastDelegate& InOnUndoRedo);
 	SAnimationEditorViewportTabBody();
 	virtual ~SAnimationEditorViewportTabBody();
 
+	/** IPersonaViewport interface */
+	virtual TSharedRef<IPersonaViewportState> SaveState() const override;
+	virtual void RestoreState(TSharedRef<IPersonaViewportState> InState) override;
+	
 	void RefreshViewport();
 
 	/**
@@ -88,11 +109,8 @@ public:
 	 */
 	const TSharedPtr<FUICommandList>& GetCommandList() const { return UICommandList; }
 
-	/** Preview mode related function **/
-	bool IsPreviewModeOn( int32 PreviewMode ) const;
-
-	/** Sets up the viewport with Persona's preview component */
-	void SetPreviewComponent( class UDebugSkelMeshComponent* PreviewComponent);
+	/** Handle the skeletal mesh changing */
+	void HandlePreviewMeshChanged(class USkeletalMesh* SkeletalMesh);
 
 	/** Function to get the number of LOD models associated with the preview skeletal mesh*/
 	int32 GetLODModelCount() const;
@@ -134,10 +152,6 @@ public:
 	/** Gets the animation viewport client */
 	TSharedRef<class FAnimationViewportClient> GetAnimationViewportClient() const;
 
-	/** Save data from OldViewport, and Restore **/
-	void SaveData(class SAnimationEditorViewportTabBody* OldViewport);
-	void RestoreData();
-
 	/** Returns Detail description of what's going with viewport **/
 	FText GetDisplayString() const;
 
@@ -146,15 +160,6 @@ public:
 
 	/** Function to check whether floor is auto aligned or not */
 	bool IsAutoAlignFloor() const;
-	
-	/** Clears our reference to Persona, also cleaning up anything that depends on Persona first */
-	void CleanupPersonaReferences();
-
-	/** Gets Persona to access to PreviewComponent for checking whether it shows clothing options or not */
-	TWeakPtr<FPersona> GetPersona() const 
-	{ 
-		return PersonaPtr; 
-	}
 
 	void SetWindStrength( float SliderPos );
 
@@ -174,6 +179,15 @@ public:
 	/** Function to set LOD model selection*/
 	void OnSetLODModel(int32 LODSelectionType);
 
+	/** Get the skeleton tree we are bound to */
+	TSharedRef<class ISkeletonTree> GetSkeletonTree() const { return SkeletonTreePtr.Pin().ToSharedRef(); }
+
+	/** Get the preview scene we are viewing */
+	TSharedRef<class FAnimationEditorPreviewScene> GetPreviewScene() const { return PreviewScenePtr.Pin().ToSharedRef(); }
+
+	/** Get the skeleton tree we are bound to */
+	TSharedPtr<class FAssetEditorToolkit> GetAssetEditorToolkit() const { return AssetEditorToolkitPtr.Pin(); }
+
 protected:
 
 
@@ -185,9 +199,6 @@ private:
 	 */ 
 	void BindCommands();
 
-	/** Preview mode related function **/
-	void SetPreviewMode( int32 PreviewMode );
-
 	/** Show Morphtarget of SkeletalMesh **/
 	void OnShowMorphTargets();
 
@@ -270,19 +281,11 @@ private:
 	bool IsShowingMeshInfo(int32 DisplayInfoMode) const;
 
 	/** Function to show/hide grid in the viewport */
-	void OnShowGrid();
-	
+	void OnShowGrid();	
+
 	/** Toggles floor alignment in the preview scene */
 	void OnToggleAutoAlignFloor();
 
-	/** update reference pose with current preview mesh */
-	void UpdateReferencePose();
-
-	/** Called to toggle showing of reference pose on current preview mesh */
-	void ShowReferencePose();
-	bool CanShowReferencePose() const;
-	bool IsShowReferencePoseEnabled() const;
-
 	/** Called to toggle showing of reference pose on current preview mesh */
 	void ShowRetargetBasePose();
 	bool CanShowRetargetBasePose() const;
@@ -314,6 +317,9 @@ private:
 
 	void AnimChanged(UAnimationAsset* AnimAsset);
 
+	/** Open the preview scene settings */
+	void OpenPreviewSceneSettings();
+
 	/** Called to toggle camera lock for naviagating **/
 	void ToggleCameraFollow();
 	bool IsCameraFollowEnabled() const;
@@ -324,9 +330,6 @@ private:
 	/** Tests to see if bone move mode buttons should be visibile */
 	EVisibility GetBoneMoveModeButtonVisibility() const;
 
-	/** Changes the currently selected LoD if the current one becomes invalid */
-	void OnLODChanged();
-
 	/** Function to mute/unmute viewport audio */
 	void OnMuteAudio();
 
@@ -402,14 +405,20 @@ private:
 #endif // #if WITH_APEX_CLOTHING
 
 private:
-	/** Pointer back to the Persona tool that owns us */
-	TWeakPtr<FPersona> PersonaPtr;
+	/** Weak pointer back to the skeleton tree we are bound to */
+	TWeakPtr<class ISkeletonTree> SkeletonTreePtr;
+
+	/** Weak pointer back to the preview scene we are viewing */
+	TWeakPtr<class FAnimationEditorPreviewScene> PreviewScenePtr;
+
+	/** Weak pointer back to asset editor we are embedded in */
+	TWeakPtr<class FAssetEditorToolkit> AssetEditorToolkitPtr;
 
-	/** Skeleton */
-	USkeleton* TargetSkeleton;
+	/** Weak pointer to the blueprint editor we are optionally embedded in */
+	TWeakPtr<class FBlueprintEditor> BlueprintEditorPtr;
 
-	/** Is this view editable */
-	TAttribute<bool> IsEditable;
+	/** Whether to show the timeline */
+	bool bShowTimeline;
 
 	/** Level viewport client */
 	TSharedPtr<FEditorViewportClient> LevelViewportClient;
@@ -423,6 +432,9 @@ private:
 	/** Commands that are bound to delegates*/
 	TSharedPtr<FUICommandList> UICommandList;
 
+	/** Delegate used to invoke tabs in the containing asset editor */
+	FOnInvokeTab OnInvokeTab;
+
 public:
 	/** UV Channel Selector */
 	TSharedPtr< class STextComboBox > UVChannelCombo;
@@ -437,14 +449,9 @@ private:
 	/** Box that contains scrub panel */
 	TSharedPtr<SVerticalBox> ScrubPanelContainer;
 
-	bool bPreviewLockModeOn;
-
 	/** Current LOD selection*/
 	int32 LODSelection;
 
-	/** Selected playback speed mode, used for deciding scale */
-	EAnimationPlaybackSpeeds::Type AnimationPlaybackSpeedMode;
-
 	/** Get Min/Max Input of value **/
 	float GetViewMinInput() const;
 	float GetViewMaxInput() const;
@@ -457,9 +464,6 @@ private:
 private:
 	friend class FPersona;
 
-	EVisibility GetViewportCornerImageVisibility() const;
-	const FSlateBrush* GetViewportCornerImage() const;
-
 	EVisibility GetViewportCornerTextVisibility() const;
 	FText GetViewportCornerText() const;
 	FText GetViewportCornerTooltip() const;
diff --git a/Engine/Source/Editor/Persona/Private/SAnimationScrubPanel.cpp b/Engine/Source/Editor/Persona/Private/SAnimationScrubPanel.cpp
--- a/Engine/Source/Editor/Persona/Private/SAnimationScrubPanel.cpp
+++ b/Engine/Source/Editor/Persona/Private/SAnimationScrubPanel.cpp
@@ -8,19 +8,17 @@
 #include "Editor/KismetWidgets/Public/SScrubControlPanel.h"
 #include "AnimationUtils.h"
 #include "ScopedTransaction.h"
+#include "Animation/BlendSpaceBase.h"
 
 #define LOCTEXT_NAMESPACE "AnimationScrubPanel"
 
-void SAnimationScrubPanel::Construct( const SAnimationScrubPanel::FArguments& InArgs )
+void SAnimationScrubPanel::Construct( const SAnimationScrubPanel::FArguments& InArgs, const TSharedRef<IPersonaPreviewScene>& InPreviewScene)
 {
 	bSliderBeingDragged = false;
-	PersonaPtr = InArgs._Persona;
 	LockedSequence = InArgs._LockedSequence;
 	OnSetInputViewRange = InArgs._OnSetInputViewRange;
 
-	// register delegate for anim change notification
-	PersonaPtr.Pin()->RegisterOnAnimChanged( FPersona::FOnAnimChanged::CreateSP(this, &SAnimationScrubPanel::AnimChanged) );
-	check(PersonaPtr.IsValid());
+	PreviewScenePtr = InPreviewScene;
 
 	this->ChildSlot
 	[
@@ -30,13 +28,14 @@ void SAnimationScrubPanel::Construct( const SAnimationScrubPanel::FArguments& In
 		.HAlign(HAlign_Fill) 
 		.VAlign(VAlign_Center)
 		.FillWidth(1)
-		.Padding( FMargin( 3.0f, 0.0f) )
+		.Padding(0.0f)
 		[
 			SAssignNew(ScrubControlPanel, SScrubControlPanel)
 			.IsEnabled(true)//this, &SAnimationScrubPanel::DoesSyncViewport)
 			.Value(this, &SAnimationScrubPanel::GetScrubValue)
 			.NumOfKeys(this, &SAnimationScrubPanel::GetNumOfFrames)
 			.SequenceLength(this, &SAnimationScrubPanel::GetSequenceLength)
+			.DisplayDrag(this, &SAnimationScrubPanel::GetDisplayDrag)
 			.OnValueChanged(this, &SAnimationScrubPanel::OnValueChanged)
 			.OnBeginSliderMovement(this, &SAnimationScrubPanel::OnBeginSliderMovement)
 			.OnEndSliderMovement(this, &SAnimationScrubPanel::OnEndSliderMovement)
@@ -47,8 +46,10 @@ void SAnimationScrubPanel::Construct( const SAnimationScrubPanel::FArguments& In
 			.OnClickedBackwardStep(this, &SAnimationScrubPanel::OnClick_Backward_Step)
 			.OnClickedBackwardEnd(this, &SAnimationScrubPanel::OnClick_Backward_End)
 			.OnClickedToggleLoop(this, &SAnimationScrubPanel::OnClick_ToggleLoop)
+			.OnClickedRecord(this, &SAnimationScrubPanel::OnClick_Record)
 			.OnGetLooping(this, &SAnimationScrubPanel::IsLoopStatusOn)
 			.OnGetPlaybackMode(this, &SAnimationScrubPanel::GetPlaybackMode)
+			.OnGetRecording(this, &SAnimationScrubPanel::IsRecording)
 			.ViewInputMin(InArgs._ViewInputMin)
 			.ViewInputMax(InArgs._ViewInputMax)
 			.OnSetInputViewRange(InArgs._OnSetInputViewRange)
@@ -60,15 +61,6 @@ void SAnimationScrubPanel::Construct( const SAnimationScrubPanel::FArguments& In
 			.IsRealtimeStreamingMode(this, &SAnimationScrubPanel::IsRealtimeStreamingMode)
 		]
 	];
-
-}
-
-SAnimationScrubPanel::~SAnimationScrubPanel()
-{
-	if (PersonaPtr.IsValid())
-	{
-		PersonaPtr.Pin()->UnregisterOnAnimChanged(this);
-	}
 }
 
 FReply SAnimationScrubPanel::OnClick_Forward_Step()
@@ -78,7 +70,7 @@ FReply SAnimationScrubPanel::OnClick_Forward_Step()
 		PreviewInstance->SetPlaying(false);
 		PreviewInstance->StepForward();
 	}
-	else if (UDebugSkelMeshComponent* SMC = PersonaPtr.Pin()->GetPreviewMeshComponent())
+	else if (UDebugSkelMeshComponent* SMC = GetPreviewScene()->GetPreviewMeshComponent())
 	{
 		//@TODO: Should we hardcode 30 Hz here?
 		{
@@ -158,7 +150,7 @@ FReply SAnimationScrubPanel::OnClick_Forward()
 			PreviewInstance->SetPlaying(true);
 		}
 	}
-	else if (UDebugSkelMeshComponent* SMC = PersonaPtr.Pin()->GetPreviewMeshComponent())
+	else if (UDebugSkelMeshComponent* SMC = GetPreviewScene()->GetPreviewMeshComponent())
 	{
 
 		SMC->GlobalAnimRateScale = (SMC->GlobalAnimRateScale > 0.0f) ? 0.0f : 1.0f;
@@ -209,6 +201,13 @@ FReply SAnimationScrubPanel::OnClick_ToggleLoop()
 	return FReply::Handled();
 }
 
+FReply SAnimationScrubPanel::OnClick_Record()
+{
+	StaticCastSharedRef<FAnimationEditorPreviewScene>(GetPreviewScene())->RecordAnimation();
+
+	return FReply::Handled();
+}
+
 bool SAnimationScrubPanel::IsLoopStatusOn() const
 {
 	UAnimSingleNodeInstance* PreviewInstance = GetPreviewInstance();
@@ -225,7 +224,7 @@ EPlaybackMode::Type SAnimationScrubPanel::GetPlaybackMode() const
 		}
 		return EPlaybackMode::Stopped;
 	}
-	else if (UDebugSkelMeshComponent* SMC = PersonaPtr.Pin()->GetPreviewMeshComponent())
+	else if (UDebugSkelMeshComponent* SMC = GetPreviewScene()->GetPreviewMeshComponent())
 	{
 		return (SMC->GlobalAnimRateScale > 0.0f) ? EPlaybackMode::PlayingForward : EPlaybackMode::Stopped;
 	}
@@ -233,6 +232,11 @@ EPlaybackMode::Type SAnimationScrubPanel::GetPlaybackMode() const
 	return EPlaybackMode::Stopped;
 }
 
+bool SAnimationScrubPanel::IsRecording() const
+{
+	return StaticCastSharedRef<FAnimationEditorPreviewScene>(GetPreviewScene())->IsRecording();
+}
+
 bool SAnimationScrubPanel::IsRealtimeStreamingMode() const
 {
 	UAnimSingleNodeInstance* PreviewInstance = GetPreviewInstance();
@@ -272,10 +276,6 @@ void SAnimationScrubPanel::OnEndSliderMovement(float NewValue)
 	bSliderBeingDragged = false;
 }
 
-void SAnimationScrubPanel::AnimChanged(UAnimationAsset* AnimAsset)
-{
-}
-
 uint32 SAnimationScrubPanel::GetNumOfFrames() const
 {
 	if (DoesSyncViewport())
@@ -284,9 +284,17 @@ uint32 SAnimationScrubPanel::GetNumOfFrames() const
 		float Length = PreviewInstance->GetLength();
 		// if anim sequence, use correct num frames
 		int32 NumFrames = (int32) (Length/0.0333f); 
-		if (PreviewInstance->GetCurrentAsset() && PreviewInstance->GetCurrentAsset()->IsA(UAnimSequenceBase::StaticClass()))
+		if (PreviewInstance->GetCurrentAsset())
 		{
-			NumFrames = CastChecked<UAnimSequenceBase>(PreviewInstance->GetCurrentAsset())->GetNumberOfFrames();
+			if (PreviewInstance->GetCurrentAsset()->IsA(UAnimSequenceBase::StaticClass()))
+			{
+				NumFrames = CastChecked<UAnimSequenceBase>(PreviewInstance->GetCurrentAsset())->GetNumberOfFrames();
+			}
+			else if(PreviewInstance->GetCurrentAsset()->IsA(UBlendSpaceBase::StaticClass()))
+			{
+				// Blendspaces dont display frame notches, so just return 0 here
+				NumFrames = 0;
+			}
 		}
 		return NumFrames;
 	}
@@ -342,21 +350,13 @@ void SAnimationScrubPanel::Tick( const FGeometry& AllottedGeometry, const double
 {
 	if (bSliderBeingDragged)
 	{
-		if (PersonaPtr.IsValid())
-		{
-			PersonaPtr.Pin()->RefreshViewport();
-		}
-		else
-		{
-			// if character editor isn't valid, we should just ignore update
-			bSliderBeingDragged = false;
-		}
+		GetPreviewScene()->InvalidateViews();
 	}
 }
 
 class UAnimSingleNodeInstance* SAnimationScrubPanel::GetPreviewInstance() const
 {
-	UDebugSkelMeshComponent* PreviewMeshComponent = PersonaPtr.Pin()->GetPreviewMeshComponent();
+	UDebugSkelMeshComponent* PreviewMeshComponent = GetPreviewScene()->GetPreviewMeshComponent();
 	return PreviewMeshComponent && PreviewMeshComponent->IsPreviewOn()? PreviewMeshComponent->PreviewInstance : NULL;
 }
 
@@ -390,7 +390,7 @@ void SAnimationScrubPanel::ReplaceLockedSequence(class UAnimSequenceBase* NewLoc
 
 UAnimInstance* SAnimationScrubPanel::GetAnimInstanceWithBlueprint() const
 {
-	if (UDebugSkelMeshComponent* DebugComponent = PersonaPtr.Pin()->GetPreviewMeshComponent())
+	if (UDebugSkelMeshComponent* DebugComponent = GetPreviewScene()->GetPreviewMeshComponent())
 	{
 		UAnimInstance* Instance = DebugComponent->GetAnimInstance();
 
@@ -515,7 +515,7 @@ void SAnimationScrubPanel::OnReZeroAnimSequence( )
 	UAnimSingleNodeInstance* PreviewInstance = GetPreviewInstance();
 	if(PreviewInstance)
 	{
-		UDebugSkelMeshComponent* PreviewSkelComp = PersonaPtr.Pin()->GetPreviewMeshComponent();
+		UDebugSkelMeshComponent* PreviewSkelComp = GetPreviewScene()->GetPreviewMeshComponent();
 
 		if (PreviewInstance->GetCurrentAsset() && PreviewSkelComp )
 		{
@@ -552,4 +552,15 @@ void SAnimationScrubPanel::OnReZeroAnimSequence( )
 	}
 }
 
+bool SAnimationScrubPanel::GetDisplayDrag() const
+{
+	UAnimSingleNodeInstance* PreviewInstance = GetPreviewInstance();
+	if (PreviewInstance && PreviewInstance->GetCurrentAsset())
+	{
+		return true;
+	}
+
+	return false;
+}
+
 #undef LOCTEXT_NAMESPACE
diff --git a/Engine/Source/Editor/Persona/Private/SAnimationScrubPanel.h b/Engine/Source/Editor/Persona/Private/SAnimationScrubPanel.h
--- a/Engine/Source/Editor/Persona/Private/SAnimationScrubPanel.h
+++ b/Engine/Source/Editor/Persona/Private/SAnimationScrubPanel.h
@@ -13,10 +13,8 @@ class SAnimationScrubPanel : public SCompoundWidget
 public:
 
 	SLATE_BEGIN_ARGS(SAnimationScrubPanel)
-		: _Persona()
-		, _LockedSequence()
+		: _LockedSequence()
 	{}
-		SLATE_ARGUMENT(TWeakPtr<FPersona>, Persona )
 		/** If you'd like to lock to one asset for this scrub control, give this**/
 		SLATE_ARGUMENT(UAnimSequenceBase*, LockedSequence)
 		/** View Input range **/
@@ -35,17 +33,15 @@ public:
 	 * 
 	 * @param InArgs   A declaration from which to construct the widget
 	 */
-	void Construct( const FArguments& InArgs );
+	void Construct( const FArguments& InArgs, const TSharedRef<IPersonaPreviewScene>& InPreviewScene );
 
 	// SWidget interface
 	virtual void Tick( const FGeometry& AllottedGeometry, const double InCurrentTime, const float InDeltaTime ) override;
 	// End of SWidget interface
 
-	virtual ~SAnimationScrubPanel();
-
 	void ReplaceLockedSequence(class UAnimSequenceBase * NewLockedSequence);
 protected:
-	TWeakPtr<FPersona> PersonaPtr;
+	TWeakPtr<IPersonaPreviewScene> PreviewScenePtr;
 	FOnSetInputViewRange OnSetInputViewRange;
 
 	bool bSliderBeingDragged;
@@ -61,6 +57,8 @@ protected:
 
 	virtual FReply OnClick_ToggleLoop();
 
+	virtual FReply OnClick_Record();
+
 	void AnimChanged(UAnimationAsset * AnimAsset);
 
 	virtual void OnValueChanged(float NewValue);
@@ -78,6 +76,7 @@ protected:
 	void OnEndSliderMovement(float NewValue);
 
 	EPlaybackMode::Type GetPlaybackMode() const;
+	bool IsRecording() const;
 	bool IsLoopStatusOn() const;
 	bool IsRealtimeStreamingMode() const;
 
@@ -97,5 +96,9 @@ protected:
 
 	// Returns the debug data if the current preview is of an anim blueprint that is the selected debug object, or NULL
 	bool GetAnimBlueprintDebugData(UAnimInstance*& Instance, FAnimBlueprintDebugData*& DebugInfo) const;
+
+	TSharedRef<IPersonaPreviewScene> GetPreviewScene() const { return PreviewScenePtr.Pin().ToSharedRef(); }
+
+	bool GetDisplayDrag() const;
 };
 #endif // __SAnimationScrubPanel_h__
diff --git a/Engine/Source/Editor/Persona/Private/SAnimationSequenceBrowser.cpp b/Engine/Source/Editor/Persona/Private/SAnimationSequenceBrowser.cpp
--- a/Engine/Source/Editor/Persona/Private/SAnimationSequenceBrowser.cpp
+++ b/Engine/Source/Editor/Persona/Private/SAnimationSequenceBrowser.cpp
@@ -16,9 +16,13 @@
 #include "SceneViewport.h"
 #include "AnimPreviewInstance.h"
 #include "ObjectEditorUtils.h"
+#include "IPersonaToolkit.h"
+#include "IAnimationEditorModule.h"
 
 #define LOCTEXT_NAMESPACE "SequenceBrowser"
 
+const FString SAnimationSequenceBrowser::SettingsIniSection = TEXT("SequenceBrowser");
+
 /** A filter that displays animations that are additive */
 class FFrontendFilter_AdditiveAnimAssets : public FFrontendFilter
 {
@@ -61,25 +65,20 @@ SAnimationSequenceBrowser::~SAnimationSequenceBrowser()
 	{
 		ViewportClient->Viewport = NULL;
 	}
-
-	if (PersonaPtr.IsValid())
-	{
-		PersonaPtr.Pin()->SetSequenceBrowser(NULL);
-	}
 }
 
 void SAnimationSequenceBrowser::OnRequestOpenAsset(const FAssetData& AssetData, bool bFromHistory)
 {
-	TSharedPtr<FPersona> Persona = PersonaPtr.Pin();
-	if (Persona.IsValid())
-	{
 		if (UObject* RawAsset = AssetData.GetAsset())
 		{
 			if (UAnimationAsset* Asset = Cast<UAnimationAsset>(RawAsset))
 			{
-				Persona->OpenNewDocumentTab(Asset, !bFromHistory);
-				Persona->SetPreviewAnimationAsset(Asset);
+			if(!bFromHistory)
+			{
+				AddAssetToHistory(AssetData);
 			}
+
+			OnOpenNewAsset.ExecuteIfBound(Asset);
 		}
 	}
 }
@@ -251,7 +250,8 @@ void SAnimationSequenceBrowser::OnApplyCompression(TArray<FAssetData> SelectedAs
 			}
 		}
 
-		PersonaPtr.Pin()->ApplyCompression(AnimSequences);
+		FPersonaModule& PersonaModule = FModuleManager::GetModuleChecked<FPersonaModule>("Persona");
+		PersonaModule.ApplyCompression(AnimSequences);
 	}
 }
 
@@ -269,7 +269,8 @@ void SAnimationSequenceBrowser::OnExportToFBX(TArray<FAssetData> SelectedAssets)
 			}
 		}
 
-		PersonaPtr.Pin()->ExportToFBX(AnimSequences);
+		FPersonaModule& PersonaModule = FModuleManager::GetModuleChecked<FPersonaModule>("Persona");
+		PersonaModule.ExportToFBX(AnimSequences, PersonaToolkitPtr.Pin()->GetPreviewScene()->GetPreviewMeshComponent()->SkeletalMesh);
 	}
 }
 
@@ -277,8 +278,7 @@ void SAnimationSequenceBrowser::OnSetCurrentPreviewMesh(TArray<FAssetData> Selec
 {
 	if(SelectedAssets.Num() > 0)
 	{
-		USkeletalMesh * PreviewMesh = PersonaPtr.Pin()->GetMesh();
-
+		USkeletalMesh* PreviewMesh = PersonaToolkitPtr.Pin()->GetPreviewScene()->GetPreviewMeshComponent()->SkeletalMesh;
 		if (PreviewMesh)
 		{
 			TArray<TWeakObjectPtr<UAnimSequence>> AnimSequences;
@@ -308,7 +308,8 @@ void SAnimationSequenceBrowser::OnAddLoopingInterpolation(TArray<FAssetData> Sel
 			}
 		}
 
-		PersonaPtr.Pin()->AddLoopingInterpolation(AnimSequences);
+		FPersonaModule& PersonaModule = FModuleManager::GetModuleChecked<FPersonaModule>("Persona");
+		PersonaModule.AddLoopingInterpolation(AnimSequences);
 	}
 }
 
@@ -335,19 +336,36 @@ void SAnimationSequenceBrowser::RetargetAnimationHandler(USkeleton* OldSkeleton,
 	if(UAnimationAsset* AnimAsset = Cast<UAnimationAsset>(AssetToOpen))
 	{
 		FAssetRegistryModule::AssetCreated(AssetToOpen);
+
+		if (GetDefault<UPersonaOptions>()->bUseStandaloneAnimationEditors)
+		{
+			const bool bBringToFrontIfOpen = true;
+			if (IAssetEditorInstance* EditorInstance = FAssetEditorManager::Get().FindEditorForAsset(AnimAsset, bBringToFrontIfOpen))
+			{
+				EditorInstance->FocusWindow(AnimAsset);
+			}
+			else
+			{
+				IAnimationEditorModule& AnimationEditorModule = FModuleManager::LoadModuleChecked<IAnimationEditorModule>("AnimationEditor");
+				AnimationEditorModule.CreateAnimationEditor(EToolkitMode::Standalone, nullptr, AnimAsset);
+			}
+		}
+		else
+		{
 		// once all success, attempt to open new persona module with new skeleton
 		EToolkitMode::Type Mode = EToolkitMode::Standalone;
 		FPersonaModule& PersonaModule = FModuleManager::LoadModuleChecked<FPersonaModule>("Persona");
 		PersonaModule.CreatePersona(Mode, TSharedPtr<IToolkitHost>(), NewSkeleton, NULL, AnimAsset, NULL);
 	}
 }
+}
 
 void SAnimationSequenceBrowser::OnCreateCopy(TArray<FAssetData> Selected)
 {
 	if ( Selected.Num() > 0 )
 	{
 		// ask which skeleton users would like to choose
-		USkeleton* OldSkeleton = PersonaPtr.Pin()->GetSkeleton();
+		USkeleton* OldSkeleton = PersonaToolkitPtr.Pin()->GetSkeleton();
 		USkeleton* NewSkeleton = NULL;
 		bool		bDuplicateAssets = true;
 
@@ -377,14 +395,11 @@ bool SAnimationSequenceBrowser::CanShowColumnForAssetRegistryTag(FName AssetType
 	return !AssetRegistryTagsToIgnore.Contains(TagName);
 }
 
-void SAnimationSequenceBrowser::Construct(const FArguments& InArgs)
+void SAnimationSequenceBrowser::Construct(const FArguments& InArgs, const TSharedRef<IPersonaToolkit>& InPersonaToolkit)
 {
-	PersonaPtr = InArgs._Persona;
-
-	if (PersonaPtr.IsValid())
-	{
-		PersonaPtr.Pin()->SetSequenceBrowser(this);
-	}
+	PersonaToolkitPtr = InPersonaToolkit;
+	OnOpenNewAsset = InArgs._OnOpenNewAsset;
+	bShowHistory = InArgs._ShowHistory;
 
 	Commands = MakeShareable(new FUICommandList());
 	Commands->MapAction(FGlobalEditorCommonCommands::Get().FindInContentBrowser, FUIAction(
@@ -412,16 +427,12 @@ void SAnimationSequenceBrowser::Construct(const FArguments& InArgs)
 	Config.bShowPathInColumnView = true;
 	Config.bSortByPathInColumnView = true;
 
-	TSharedPtr<FPersona> Persona = PersonaPtr.Pin();
-	if (Persona.IsValid())
-	{
-		USkeleton* DesiredSkeleton = Persona->GetSkeleton();
+	USkeleton* DesiredSkeleton = InPersonaToolkit->GetSkeleton();
 		if(DesiredSkeleton)
 		{
 			FString SkeletonString = FAssetData(DesiredSkeleton).GetExportTextName();
 			Config.Filter.TagsAndValues.Add(TEXT("Skeleton"), SkeletonString);
 		}
-	}
 
 	// Configure response to click and double-click
 	Config.OnAssetDoubleClicked = FOnAssetDoubleClicked::CreateSP(this, &SAnimationSequenceBrowser::OnRequestOpenAsset, false);
@@ -432,6 +443,8 @@ void SAnimationSequenceBrowser::Construct(const FArguments& InArgs)
 	Config.bFocusSearchBoxWhenOpened = false;
 	Config.DefaultFilterMenuExpansion = EAssetTypeCategories::Animation;
 
+	Config.SaveSettingsName = SettingsIniSection;
+
 	TSharedPtr<FFrontendFilterCategory> AnimCategory = MakeShareable( new FFrontendFilterCategory(LOCTEXT("ExtraAnimationFilters", "Anim Filters"), LOCTEXT("ExtraAnimationFiltersTooltip", "Filter assets by all filters in this category.")) );
 	Config.ExtraFrontendFilters.Add( MakeShareable(new FFrontendFilter_AdditiveAnimAssets(AnimCategory)) );
 	
@@ -439,18 +452,24 @@ void SAnimationSequenceBrowser::Construct(const FArguments& InArgs)
 	Config.OnVisualizeAssetToolTip = FOnVisualizeAssetToolTip::CreateSP(this, &SAnimationSequenceBrowser::OnVisualizeAssetToolTip);
 	Config.OnAssetToolTipClosing = FOnAssetToolTipClosing::CreateSP( this, &SAnimationSequenceBrowser::OnAssetToolTipClosing );
 
+	static const FName DefaultForegroundName("DefaultForeground");
+
 	TSharedRef< SMenuAnchor > BackMenuAnchorPtr = SNew(SMenuAnchor)
 		.Placement(MenuPlacement_BelowAnchor)
 		.OnGetMenuContent(this, &SAnimationSequenceBrowser::CreateHistoryMenu, true)
 		[
 			SNew(SButton)
 			.OnClicked(this, &SAnimationSequenceBrowser::OnGoBackInHistory)
-			.ButtonStyle(FEditorStyle::Get(), "GraphBreadcrumbButton")
+			.ForegroundColor(FEditorStyle::GetSlateColor(DefaultForegroundName))
+			.ButtonStyle(FEditorStyle::Get(), "FlatButton")
+			.ContentPadding(FMargin(1, 0))
 			.IsEnabled(this, &SAnimationSequenceBrowser::CanStepBackwardInHistory)
 			.ToolTipText(LOCTEXT("Backward_Tooltip", "Step backward in the asset history. Right click to see full history."))
 			[
-				SNew(SImage)
-				.Image(FEditorStyle::GetBrush("GraphBreadcrumb.BrowseBack"))
+				SNew(STextBlock)
+				.TextStyle(FEditorStyle::Get(), "ContentBrowser.TopBar.Font")
+				.Font(FEditorStyle::Get().GetFontStyle("FontAwesome.11"))
+				.Text(FText::FromString(FString(TEXT("\xf060"))) /*fa-arrow-left*/)
 			]
 		];
 
@@ -460,12 +479,16 @@ void SAnimationSequenceBrowser::Construct(const FArguments& InArgs)
 		[
 			SNew(SButton)
 			.OnClicked(this, &SAnimationSequenceBrowser::OnGoForwardInHistory)
-			.ButtonStyle(FEditorStyle::Get(), "GraphBreadcrumbButton")
+			.ForegroundColor(FEditorStyle::GetSlateColor(DefaultForegroundName))
+			.ButtonStyle(FEditorStyle::Get(), "FlatButton")
+			.ContentPadding(FMargin(1, 0))
 			.IsEnabled(this, &SAnimationSequenceBrowser::CanStepForwardInHistory)
 			.ToolTipText(LOCTEXT("Forward_Tooltip", "Step forward in the asset history. Right click to see full history."))
 			[
-				SNew(SImage)
-				.Image(FEditorStyle::GetBrush("GraphBreadcrumb.BrowseForward"))
+				SNew(STextBlock)
+				.TextStyle(FEditorStyle::Get(), "ContentBrowser.TopBar.Font")
+				.Font(FEditorStyle::Get().GetFontStyle("FontAwesome.11"))
+				.Text(FText::FromString(FString(TEXT("\xf061"))) /*fa-arrow-right*/)
 			]
 		];
 
@@ -473,27 +496,23 @@ void SAnimationSequenceBrowser::Construct(const FArguments& InArgs)
 	[
 		SNew(SVerticalBox)
 		+SVerticalBox::Slot()
-		.FillHeight(1.f)
-		[
-			ContentBrowserModule.Get().CreateAssetPicker(Config)
-		]
-		+SVerticalBox::Slot()
 		.AutoHeight()
 		[
-			SNew(SSeparator)
-			.Visibility(this, &SAnimationSequenceBrowser::GetNonBlueprintModeVisibility)
-		]
-		+SVerticalBox::Slot()
-		.HAlign(HAlign_Right)
-		.AutoHeight()
+			SNew(SBorder)
+			.Visibility(this, &SAnimationSequenceBrowser::GetHistoryVisibility)
+			.Padding(FMargin(3))
+			.BorderImage( FEditorStyle::GetBrush("ToolPanel.GroupBorder") )
+		[
+				SNew(SHorizontalBox)
+				+SHorizontalBox::Slot()
+				.HAlign(HAlign_Left)
 		[
 			SNew(SHorizontalBox)
-			.Visibility(this, &SAnimationSequenceBrowser::GetNonBlueprintModeVisibility)
 			+ SHorizontalBox::Slot()
 			.AutoWidth()
 			[
 				SNew(SBorder)
-				.OnMouseButtonDown(this, &SAnimationSequenceBrowser::OnMouseDownHisory, TWeakPtr<SMenuAnchor>(BackMenuAnchorPtr))
+						.OnMouseButtonDown(this, &SAnimationSequenceBrowser::OnMouseDownHistory, TWeakPtr<SMenuAnchor>(BackMenuAnchorPtr))
 				.BorderImage( FEditorStyle::GetBrush("NoBorder") )
 				[
 					BackMenuAnchorPtr
@@ -504,13 +523,25 @@ void SAnimationSequenceBrowser::Construct(const FArguments& InArgs)
 			.AutoWidth()
 			[
 				SNew(SBorder)
-				.OnMouseButtonDown(this, &SAnimationSequenceBrowser::OnMouseDownHisory, TWeakPtr<SMenuAnchor>(FwdMenuAnchorPtr))
+						.OnMouseButtonDown(this, &SAnimationSequenceBrowser::OnMouseDownHistory, TWeakPtr<SMenuAnchor>(FwdMenuAnchorPtr))
 				.BorderImage( FEditorStyle::GetBrush("NoBorder") )
 				[
 					FwdMenuAnchorPtr
 				]
 			]
 		]
+			]
+		]
+		+SVerticalBox::Slot()
+		.FillHeight(1.f)
+		[
+			SNew(SBorder)
+			.Padding(FMargin(3))
+			.BorderImage(FEditorStyle::GetBrush("ToolPanel.GroupBorder"))
+			[
+				ContentBrowserModule.Get().CreateAssetPicker(Config)
+			]
+		]
 	];
 
 	// Create the ignore set for asset registry tags
@@ -554,7 +585,7 @@ void SAnimationSequenceBrowser::AddAssetToHistory(const FAssetData& AssetData)
 	CurrentAssetHistoryIndex = AssetHistory.Num() - 1;
 }
 
-FReply SAnimationSequenceBrowser::OnMouseDownHisory( const FGeometry& MyGeometry, const FPointerEvent& MouseEvent, TWeakPtr< SMenuAnchor > InMenuAnchor )
+FReply SAnimationSequenceBrowser::OnMouseDownHistory( const FGeometry& MyGeometry, const FPointerEvent& MouseEvent, TWeakPtr< SMenuAnchor > InMenuAnchor )
 {
 	if(MouseEvent.GetEffectingButton() == EKeys::RightMouseButton)
 	{
@@ -693,11 +724,11 @@ void SAnimationSequenceBrowser::CacheOriginalAnimAssetHistory()
 	{
 		bTriedToCacheOrginalAsset = true;
 
-		if(AssetHistory.Num() == 0 && PersonaPtr.IsValid())
+		if(AssetHistory.Num() == 0)
 		{
-			USkeleton* DesiredSkeleton = PersonaPtr.Pin()->GetSkeleton();
+			USkeleton* DesiredSkeleton = PersonaToolkitPtr.Pin()->GetSkeleton();
 
-			if(UObject* PreviewAsset = PersonaPtr.Pin()->GetPreviewAnimationAsset())
+			if(UObject* PreviewAsset = PersonaToolkitPtr.Pin()->GetPreviewScene()->GetPreviewAnimationAsset())
 			{
 				FAssetRegistryModule& AssetRegistryModule = FModuleManager::LoadModuleChecked<FAssetRegistryModule>("AssetRegistry");
 				FAssetData AssetData = AssetRegistryModule.Get().GetAssetByObjectPath(FName(*PreviewAsset->GetPathName()));
@@ -758,6 +789,17 @@ TSharedRef<SToolTip> SAnimationSequenceBrowser::CreateCustomAssetToolTip(FAssetD
 	{
 		if(TagsToShow.Contains(TagPair.Key))
 		{
+			// Check for DisplayName metadata
+			FText DisplayName;
+			if (UProperty* Field = FindField<UProperty>(AssetClass, TagPair.Key))
+			{
+				DisplayName = Field->GetDisplayNameText();
+			}
+			else
+			{
+				DisplayName = FText::FromName(TagPair.Key);
+			}
+
 			DescriptionBox->AddSlot()
 			.AutoHeight()
 			.Padding(0,0,5,0)
@@ -767,7 +809,7 @@ TSharedRef<SToolTip> SAnimationSequenceBrowser::CreateCustomAssetToolTip(FAssetD
 				.AutoWidth()
 				[
 					SNew(STextBlock)
-					.Text(FText::Format(LOCTEXT("AssetTagKey", "{0} :"), FText::FromName(TagPair.Key)))
+					.Text(FText::Format(LOCTEXT("AssetTagKey", "{0} :"), DisplayName))
 					.ColorAndOpacity(FSlateColor::UseSubduedForeground())
 				]
 
@@ -1011,26 +1053,17 @@ EActiveTimerReturnType SAnimationSequenceBrowser::UpdateTootipPreview( double In
 bool SAnimationSequenceBrowser::IsToolTipPreviewVisible()
 {
 	bool bVisible = false;
-	// during persona recording, disable this
-	if( PersonaPtr.IsValid() && PersonaPtr.Pin()->IsRecording() == false 
-		&& ViewportWidget.IsValid())
+
+	if(ViewportWidget.IsValid())
 	{
 		bVisible = ViewportWidget->GetVisibility() == EVisibility::Visible;
 	}
 	return bVisible;
 }
 
-EVisibility SAnimationSequenceBrowser::GetNonBlueprintModeVisibility() const
-{
-	if (PersonaPtr.IsValid())
+EVisibility SAnimationSequenceBrowser::GetHistoryVisibility() const
 	{
-		if (PersonaPtr.Pin()->GetCurrentMode() == FPersonaModes::AnimBlueprintEditMode)
-		{
-			return EVisibility::Collapsed;
-		}
-	}
-
-	return EVisibility::Visible;
+	return bShowHistory ? EVisibility::Visible : EVisibility::Collapsed;
 }
 
 
diff --git a/Engine/Source/Editor/Persona/Private/SAnimationSequenceBrowser.h b/Engine/Source/Editor/Persona/Private/SAnimationSequenceBrowser.h
--- a/Engine/Source/Editor/Persona/Private/SAnimationSequenceBrowser.h
+++ b/Engine/Source/Editor/Persona/Private/SAnimationSequenceBrowser.h
@@ -7,6 +7,7 @@
 #include "Editor/ContentBrowser/Public/ContentBrowserModule.h"
 #include "PreviewScene.h"
 #include "EditorAnimUtils.h"
+#include "IAnimationSequenceBrowser.h"
 
 //////////////////////////////////////////////////////////////////////////
 // FAnimationAssetViewportClient
@@ -24,17 +25,20 @@ public:
 //////////////////////////////////////////////////////////////////////////
 // SAnimationSequenceBrowser
 
-class SAnimationSequenceBrowser : public SCompoundWidget
+class SAnimationSequenceBrowser : public IAnimationSequenceBrowser
 {
 public:
 	SLATE_BEGIN_ARGS(SAnimationSequenceBrowser)
-		: _Persona()
-		{}
+		: _ShowHistory(false)
+	{}
+
+	SLATE_ARGUMENT(bool, ShowHistory)
+
+	SLATE_ARGUMENT(FOnOpenNewAsset, OnOpenNewAsset)
 
-		SLATE_ARGUMENT(TSharedPtr<class FPersona>, Persona)
 	SLATE_END_ARGS()
 public:
-	void Construct(const FArguments& InArgs);
+	void Construct(const FArguments& InArgs, const TSharedRef<IPersonaToolkit>& InPersonaToolkit);
 
 	void OnRequestOpenAsset(const FAssetData& AssetData, bool bFromHistory);
 
@@ -42,6 +46,9 @@ public:
 
 	virtual ~SAnimationSequenceBrowser();
 
+	/** IAnimationSequenceBrowser interface */
+	virtual void SelectAsset(UAnimationAsset * AnimAsset) override;
+
 	/** Delegate that handles creation of context menu */
 	TSharedPtr<SWidget> OnGetAssetContextMenu(const TArray<FAssetData>& SelectedAssets);
 
@@ -78,11 +85,8 @@ public:
 	 */
 	void OnCreateCopy(TArray<FAssetData> Selected);
 
-	/** Refresh list */
-	void SelectAsset(UAnimationAsset * AnimAsset);
-
 	/** public reference to add to history */
-	void AddToHistory(UAnimationAsset * AnimAsset);
+	virtual void AddToHistory(UAnimationAsset * AnimAsset) override;
 
 protected:
 	bool CanShowColumnForAssetRegistryTag(FName AssetType, FName TagName) const;
@@ -115,7 +119,7 @@ protected:
 	 *
 	 * @param InMenuAnchor		This is the anchor the menu will use for positioning
 	 */
-	FReply OnMouseDownHisory( const FGeometry& MyGeometry, const FPointerEvent& MouseEvent, TWeakPtr< SMenuAnchor > InMenuAnchor );
+	FReply OnMouseDownHistory( const FGeometry& MyGeometry, const FPointerEvent& MouseEvent, TWeakPtr< SMenuAnchor > InMenuAnchor );
 
 	/** 
 	 * Callback to create the history menu.
@@ -157,7 +161,7 @@ protected:
 	bool IsToolTipPreviewVisible();
 
 	/** Returns visible when not in a Blueprint mode (anim mode, etc...) */
-	EVisibility GetNonBlueprintModeVisibility() const;
+	EVisibility GetHistoryVisibility() const;
 protected:
 	/**
 	 * The actual viewport widget
@@ -189,8 +193,8 @@ protected:
 	 */
 	TSharedPtr<FUICommandList> Commands;
 
-	// Pointer back to persona tool that owns us
-	TWeakPtr<class FPersona> PersonaPtr;
+	/** The persona toolkit we are using */
+	TWeakPtr<class IPersonaToolkit> PersonaToolkitPtr;
 
 	// Set of tags to prevent creating details view columns for (infrequently used)
 	TSet<FName> AssetRegistryTagsToIgnore;
@@ -207,6 +211,12 @@ protected:
 	// Track if we have tried to cache the first asset we were playing
 	bool bTriedToCacheOrginalAsset;
 
+	/** Whether to show the history widgets */
+	bool bShowHistory;
+
+	/** Delegate called to open a new asset for editing */
+	FOnOpenNewAsset OnOpenNewAsset;
+
 	// delegate to sync the asset picker to selected assets
 	FSyncToAssetsDelegate SyncToAssetsDelegate;
 	FGetCurrentSelectionDelegate GetCurrentSelectionDelegate;
@@ -225,4 +235,8 @@ private:
 
 	/** Whether the active timer should stop */
 	uint8 bToolTipVisualizedThisFrame : 1;
+
+public:
+	/** The section of EditorPerProjectUserSettings in which to save settings */
+	static const FString SettingsIniSection;
 };
\ No newline at end of file
