Copying //UE4/Dev-Framework to //UE4/Dev-Main (Source: //UE4/Dev-Framework @ 3136612)

diff --git a/Engine/Source/Editor/Persona/Private/SAssetFamilyShortcutBar.cpp b/Engine/Source/Editor/Persona/Private/SAssetFamilyShortcutBar.cpp
new file mode 100644
--- /dev/null
+++ b/Engine/Source/Editor/Persona/Private/SAssetFamilyShortcutBar.cpp
@@ -0,0 +1,454 @@
+// Copyright 1998-2016 Epic Games, Inc. All Rights Reserved.
+
+#include "PersonaPrivatePCH.h"
+#include "SAssetFamilyShortcutBar.h"
+#include "IAssetFamily.h"
+#include "AssetThumbnail.h"
+#include "SHyperlink.h"
+#include "BreakIterator.h"
+#include "ContentBrowserModule.h"
+#include "AssetRegistryModule.h"
+#include "WorkflowCentricApplication.h"
+
+#define LOCTEXT_NAMESPACE "SAssetFamilyShortcutBar"
+
+namespace AssetShortcutConstants
+{
+	const int32 ThumbnailSize = 40;
+	const int32 ThumbnailSizeSmall = 16;
+}
+
+class SAssetShortcut : public SCompoundWidget
+{
+public:
+	SLATE_BEGIN_ARGS(SAssetShortcut)
+	{}
+
+	SLATE_END_ARGS()
+
+	void Construct(const FArguments& InArgs, const TSharedRef<class FWorkflowCentricApplication>& InHostingApp, const TSharedRef<class IAssetFamily>& InAssetFamily, const FAssetData& InAssetData, const TSharedRef<FAssetThumbnailPool>& InThumbnailPool)
+	{
+		AssetData = InAssetData;
+		AssetFamily = InAssetFamily;
+		HostingApp = InHostingApp;
+		ThumbnailPoolPtr = InThumbnailPool;
+		bPackageDirty = false;
+
+		FAssetRegistryModule& AssetRegistryModule = FModuleManager::LoadModuleChecked<FAssetRegistryModule>(TEXT("AssetRegistry"));
+		AssetRegistryModule.Get().OnFilesLoaded().AddSP(this, &SAssetShortcut::HandleFilesLoaded);
+		AssetRegistryModule.Get().OnAssetAdded().AddSP(this, &SAssetShortcut::HandleAssetAdded);
+		AssetRegistryModule.Get().OnAssetRemoved().AddSP(this, &SAssetShortcut::HandleAssetRemoved);
+
+		FAssetEditorManager::Get().OnAssetEditorRequestedOpen().AddSP(this, &SAssetShortcut::HandleAssetOpened);
+		AssetFamily->GetOnAssetOpened().AddSP(this, &SAssetShortcut::HandleAssetOpened);
+
+		AssetThumbnail = MakeShareable(new FAssetThumbnail(InAssetData, AssetShortcutConstants::ThumbnailSize, AssetShortcutConstants::ThumbnailSize, InThumbnailPool));
+		AssetThumbnailSmall = MakeShareable(new FAssetThumbnail(InAssetData, AssetShortcutConstants::ThumbnailSizeSmall, AssetShortcutConstants::ThumbnailSizeSmall, InThumbnailPool));
+
+		TArray<FAssetData> Assets;
+		InAssetFamily->FindAssetsOfType(InAssetData.GetClass(), Assets);
+		bMultipleAssetsExist = Assets.Num() > 1;
+		AssetDirtyBrush = FEditorStyle::GetBrush("ContentBrowser.ContentDirty");
+
+		ChildSlot
+		[
+			SNew(SHorizontalBox)
+			+SHorizontalBox::Slot()
+			[
+				SAssignNew(CheckBox, SCheckBox)
+				.Style(FEditorStyle::Get(), "ToolBar.ToggleButton")
+				.ForegroundColor(FSlateColor::UseForeground())
+				.Padding(0.0f)
+				.OnCheckStateChanged(this, &SAssetShortcut::HandleOpenAssetShortcut)
+				.IsChecked(this, &SAssetShortcut::GetCheckState)
+				.Visibility(this, &SAssetShortcut::GetButtonVisibility)
+				.ToolTipText(FText::FromString(InAssetData.GetExportTextName()))
+				[
+					SNew(SHorizontalBox)
+					+ SHorizontalBox::Slot()
+					.VAlign(VAlign_Center)
+					.AutoWidth()
+					[
+						SNew(SBorder)
+						.Padding(4.0f)
+						.BorderImage(FEditorStyle::GetBrush("PropertyEditor.AssetThumbnailShadow"))
+						[
+							SNew(SHorizontalBox)
+							+SHorizontalBox::Slot()
+							[
+								SAssignNew(ThumbnailBox, SBox)
+								.WidthOverride(AssetShortcutConstants::ThumbnailSize)
+								.HeightOverride(AssetShortcutConstants::ThumbnailSize)
+								.Visibility(this, &SAssetShortcut::GetThumbnailVisibility)
+								[
+									SNew(SOverlay)
+									+SOverlay::Slot()
+									[
+										AssetThumbnail->MakeThumbnailWidget()
+									]
+									+SOverlay::Slot()
+									.HAlign(HAlign_Left)
+									.VAlign(VAlign_Bottom)
+									[
+										SNew(SImage)
+										.Image(this, &SAssetShortcut::GetDirtyImage)
+									]
+								]
+							]
+							+SHorizontalBox::Slot()
+							[
+								SAssignNew(ThumbnailSmallBox, SBox)
+								.WidthOverride(AssetShortcutConstants::ThumbnailSizeSmall)
+								.HeightOverride(AssetShortcutConstants::ThumbnailSizeSmall)
+								.Visibility(this, &SAssetShortcut::GetSmallThumbnailVisibility)
+								[
+									SNew(SOverlay)
+									+SOverlay::Slot()
+									[
+										AssetThumbnailSmall->MakeThumbnailWidget()
+									]
+									+SOverlay::Slot()
+									.HAlign(HAlign_Left)
+									.VAlign(VAlign_Bottom)
+									[
+										SNew(SImage)
+										.Image(this, &SAssetShortcut::GetDirtyImage)
+									]
+								]
+							]
+						]
+					]
+					+SHorizontalBox::Slot()
+					.VAlign(VAlign_Center)
+					.AutoWidth()
+					[
+						SNew(SBox)
+						.WidthOverride(AssetShortcutConstants::ThumbnailSize + 16.0f)
+						.HeightOverride(AssetShortcutConstants::ThumbnailSize + 4.0f)
+						.VAlign(VAlign_Center)
+						.Visibility(this, &SAssetShortcut::GetThumbnailVisibility)
+						[
+							SNew(STextBlock)
+							.Font(FEditorStyle::GetFontStyle("ContentBrowser.AssetTileViewNameFontSmall"))
+							.Text(this, &SAssetShortcut::GetAssetText)
+							.ColorAndOpacity(this, &SAssetShortcut::GetAssetTextColor)
+							.Justification(ETextJustify::Center)
+							.LineBreakPolicy(FBreakIterator::CreateCamelCaseBreakIterator())
+							.AutoWrapText(true)
+						]
+					]
+				]
+			]
+			+SHorizontalBox::Slot()
+			.VAlign(VAlign_Fill)
+			.AutoWidth()
+			[
+				SNew(SComboButton)
+				.Visibility(this, &SAssetShortcut::GetComboVisibility)
+				.ContentPadding(0)
+				.ForegroundColor(FSlateColor::UseForeground())
+				.ButtonStyle(FEditorStyle::Get(), "Toolbar.Button")
+				.OnGetMenuContent(this, &SAssetShortcut::HandleGetMenuContent)
+			]
+		];
+
+		EnableToolTipForceField(true);
+
+		DirtyStateTimerHandle = RegisterActiveTimer(1.0f / 10.0f, FWidgetActiveTimerDelegate::CreateSP(this, &SAssetShortcut::HandleRefreshDirtyState));
+	}
+
+	~SAssetShortcut()
+	{
+		if (FModuleManager::Get().IsModuleLoaded(TEXT("AssetRegistry")))
+		{
+			FAssetRegistryModule& AssetRegistryModule = FModuleManager::GetModuleChecked<FAssetRegistryModule>(TEXT("AssetRegistry"));
+			AssetRegistryModule.Get().OnFilesLoaded().RemoveAll(this);
+			AssetRegistryModule.Get().OnAssetAdded().RemoveAll(this);
+			AssetRegistryModule.Get().OnAssetRemoved().RemoveAll(this);
+		}
+
+		AssetFamily->GetOnAssetOpened().RemoveAll(this);
+		FAssetEditorManager::Get().OnAssetEditorRequestedOpen().RemoveAll(this);
+		UnRegisterActiveTimer(DirtyStateTimerHandle.ToSharedRef());
+	}
+
+	void HandleOpenAssetShortcut(ECheckBoxState InState)
+	{
+		TArray<UObject*> Assets;
+		Assets.Add(AssetData.GetAsset());
+		FAssetEditorManager::Get().OpenEditorForAssets(Assets);
+	}
+
+	FText GetAssetText() const
+	{
+		return FText::FromName(AssetData.AssetName);
+	}
+
+	ECheckBoxState GetCheckState() const
+	{
+		const TArray<UObject*>* Objects = HostingApp.Pin()->GetObjectsCurrentlyBeingEdited();
+		if (Objects != nullptr)
+		{
+			for (UObject* Object : *Objects)
+			{
+				if (FAssetData(Object) == AssetData)
+				{
+					return ECheckBoxState::Checked;
+				}
+			}
+		}
+		return ECheckBoxState::Unchecked;
+	}
+
+	FSlateColor GetAssetTextColor() const
+	{
+		static const FName InvertedForeground("InvertedForeground");
+		return GetCheckState() == ECheckBoxState::Checked || CheckBox->IsHovered() ? FEditorStyle::GetSlateColor(InvertedForeground) : FSlateColor::UseForeground();
+	}
+
+	TSharedRef<SWidget> HandleGetMenuContent()
+	{
+		FContentBrowserModule& ContentBrowserModule = FModuleManager::Get().LoadModuleChecked<FContentBrowserModule>(TEXT("ContentBrowser"));
+
+		const bool bInShouldCloseWindowAfterMenuSelection = true;
+		FMenuBuilder MenuBuilder(bInShouldCloseWindowAfterMenuSelection, nullptr);
+
+		MenuBuilder.BeginSection("AssetActions", LOCTEXT("AssetActionsSection", "Asset Actions"));
+		{
+			MenuBuilder.AddMenuEntry(
+				LOCTEXT("ShowInContentBrowser", "Show In Content Browser"),
+				LOCTEXT("ShowInContentBrowser_ToolTip", "Show this asset in the content browser."),
+				FSlateIcon(FEditorStyle::GetStyleSetName(), "PropertyWindow.Button_Browse"),
+				FUIAction(FExecuteAction::CreateSP(this, &SAssetShortcut::HandleShowInContentBrowser)));
+		}
+		MenuBuilder.EndSection();
+
+		if (bMultipleAssetsExist)
+		{
+			MenuBuilder.BeginSection("AssetSelection", LOCTEXT("AssetSelectionSection", "Select Asset"));
+			{
+				FAssetPickerConfig AssetPickerConfig;
+
+				UClass* FilterClass = AssetFamily->GetAssetFamilyClass(AssetData.GetClass());
+				if (FilterClass != nullptr)
+				{
+					AssetPickerConfig.Filter.ClassNames.Add(FilterClass->GetFName());
+					AssetPickerConfig.Filter.bRecursiveClasses = true;
+				}
+
+				AssetPickerConfig.OnAssetSelected = FOnAssetSelected::CreateSP(this, &SAssetShortcut::HandleAssetSelectedFromPicker);
+				AssetPickerConfig.OnShouldFilterAsset = FOnShouldFilterAsset::CreateSP(this, &SAssetShortcut::HandleFilterAsset);
+				AssetPickerConfig.bAllowNullSelection = false;
+				AssetPickerConfig.ThumbnailLabel = EThumbnailLabel::ClassName;
+				AssetPickerConfig.InitialAssetViewType = EAssetViewType::List;
+
+				MenuBuilder.AddWidget(
+					SNew(SBox)
+					.WidthOverride(300)
+					.HeightOverride(600)
+					[
+						ContentBrowserModule.Get().CreateAssetPicker(AssetPickerConfig)
+					],
+					FText(), true);
+			}
+			MenuBuilder.EndSection();
+		}
+		
+		return MenuBuilder.MakeWidget();
+	}
+
+	void HandleAssetSelectedFromPicker(const class FAssetData& InAssetData)
+	{
+		TArray<UObject*> Assets;
+		Assets.Add(InAssetData.GetAsset());
+		FAssetEditorManager::Get().OpenEditorForAssets(Assets);
+	}
+
+	bool HandleFilterAsset(const class FAssetData& InAssetData)
+	{
+		return !AssetFamily->IsAssetCompatible(InAssetData);
+	}
+
+	EVisibility GetButtonVisibility() const
+	{
+		return AssetData.IsValid() || bMultipleAssetsExist ? EVisibility::Visible : EVisibility::Collapsed;
+	}
+
+	EVisibility GetComboVisibility() const
+	{
+		return bMultipleAssetsExist && AssetData.IsValid() ? EVisibility::Visible : EVisibility::Collapsed;
+	}
+
+	void HandleFilesLoaded()
+	{
+		TArray<FAssetData> Assets;
+		AssetFamily->FindAssetsOfType(AssetData.GetClass(), Assets);
+		bMultipleAssetsExist = Assets.Num() > 1;
+	}
+
+	void HandleAssetRemoved(const FAssetData& InAssetData)
+	{
+		if (AssetFamily->IsAssetCompatible(InAssetData))
+		{
+			TArray<FAssetData> Assets;
+			AssetFamily->FindAssetsOfType(AssetData.GetClass(), Assets);
+			bMultipleAssetsExist = Assets.Num() > 1;
+		}
+	}
+
+	void HandleAssetAdded(const FAssetData& InAssetData)
+	{
+		if (AssetFamily->IsAssetCompatible(InAssetData))
+		{
+			TArray<FAssetData> Assets;
+			AssetFamily->FindAssetsOfType(AssetData.GetClass(), Assets);
+			bMultipleAssetsExist = Assets.Num() > 1;
+		}
+	}
+
+	void HandleShowInContentBrowser()
+	{
+		FContentBrowserModule& ContentBrowserModule = FModuleManager::Get().LoadModuleChecked<FContentBrowserModule>(TEXT("ContentBrowser"));
+		TArray<FAssetData> Assets;
+		Assets.Add(AssetData);
+		ContentBrowserModule.Get().SyncBrowserToAssets(Assets);
+	}
+
+	void HandleAssetOpened(UObject* InAsset)
+	{
+		RefreshAsset();
+	}
+
+	EVisibility GetThumbnailVisibility() const
+	{
+		return FMultiBoxSettings::UseSmallToolBarIcons.Get() ? EVisibility::Collapsed : EVisibility::Visible;
+	}
+
+	EVisibility GetSmallThumbnailVisibility() const
+	{
+		return FMultiBoxSettings::UseSmallToolBarIcons.Get() ? EVisibility::Visible : EVisibility::Collapsed;
+	}
+
+	const FSlateBrush* GetDirtyImage() const
+	{
+		return bPackageDirty ? AssetDirtyBrush : nullptr;
+	}
+
+	void RefreshAsset()
+	{
+		// if this is the asset being edited by our hosting asset editor, don't switch it
+		bool bAssetBeingEdited = false;
+		const TArray<UObject*>* Objects = HostingApp.Pin()->GetObjectsCurrentlyBeingEdited();
+		if (Objects != nullptr)
+		{
+			for (UObject* Object : *Objects)
+			{
+				if (FAssetData(Object) == AssetData)
+				{
+					bAssetBeingEdited = true;
+					break;
+				}
+			}
+		}
+
+		// switch to new asset if needed
+		FAssetData NewAssetData = AssetFamily->FindAssetOfType(AssetData.GetClass());
+		if (!bAssetBeingEdited && NewAssetData.IsValid() && NewAssetData != AssetData)
+		{
+			AssetData = NewAssetData;
+
+			// regenerate thumbnail
+			AssetThumbnail = MakeShareable(new FAssetThumbnail(AssetData, AssetShortcutConstants::ThumbnailSize, AssetShortcutConstants::ThumbnailSize, ThumbnailPoolPtr.Pin()));
+			AssetThumbnailSmall = MakeShareable(new FAssetThumbnail(AssetData, AssetShortcutConstants::ThumbnailSizeSmall, AssetShortcutConstants::ThumbnailSizeSmall, ThumbnailPoolPtr.Pin()));
+
+			ThumbnailBox->SetContent(AssetThumbnail->MakeThumbnailWidget());
+			ThumbnailSmallBox->SetContent(AssetThumbnailSmall->MakeThumbnailWidget());
+		}
+	}
+
+	EActiveTimerReturnType HandleRefreshDirtyState(double InCurrentTime, float InDeltaTime)
+	{
+		if (AssetData.IsAssetLoaded())
+		{
+			if (!AssetPackage.IsValid())
+			{
+				AssetPackage = AssetData.GetPackage();
+			}
+
+			if (AssetPackage.IsValid())
+			{
+				bPackageDirty = AssetPackage->IsDirty();
+			}
+		}
+
+		return EActiveTimerReturnType::Continue;
+	}
+
+private:
+	/** The current asset data for this widget */
+	FAssetData AssetData;
+
+	/** Cache the package of the object for checking dirty state */
+	TWeakObjectPtr<UPackage> AssetPackage;
+
+	/** Timer handle used to updating dirty state */
+	TSharedPtr<class FActiveTimerHandle> DirtyStateTimerHandle;
+
+	/** The asset family we are working with */
+	TSharedPtr<class IAssetFamily> AssetFamily;
+
+	/** Our asset thumbnails */
+	TSharedPtr<FAssetThumbnail> AssetThumbnail;
+	TSharedPtr<FAssetThumbnail> AssetThumbnailSmall;
+
+	/** Thumbnail widget containers */
+	TSharedPtr<SBox> ThumbnailBox;
+	TSharedPtr<SBox> ThumbnailSmallBox;
+
+	/** The asset editor we are embedded in */
+	TWeakPtr<class FWorkflowCentricApplication> HostingApp;
+
+	/** Thumbnail pool */
+	TWeakPtr<FAssetThumbnailPool> ThumbnailPoolPtr;
+
+	/** Check box */
+	TSharedPtr<SCheckBox> CheckBox;
+
+	/** Cached dirty brush */
+	const FSlateBrush* AssetDirtyBrush;
+
+	/** Whether there are multiple (>1) of this asset type in existence */
+	bool bMultipleAssetsExist;
+
+	/** Cache the package's dirty state */
+	bool bPackageDirty;
+};
+
+void SAssetFamilyShortcutBar::Construct(const FArguments& InArgs, const TSharedRef<class FWorkflowCentricApplication>& InHostingApp, const TSharedRef<class IAssetFamily>& InAssetFamily)
+{
+	ThumbnailPool = MakeShareable(new FAssetThumbnailPool(16, false));
+
+	TSharedRef<SUniformGridPanel> Panel = SNew(SUniformGridPanel);
+
+	int32 ColumnIndex = 0;
+	TArray<UClass*> AssetTypes;
+	InAssetFamily->GetAssetTypes(AssetTypes);
+
+	for (UClass* Class : AssetTypes)
+	{
+		FAssetData AssetData = InAssetFamily->FindAssetOfType(Class);
+		Panel->AddSlot(ColumnIndex, 0)
+		[
+			SNew(SAssetShortcut, InHostingApp, InAssetFamily, AssetData, ThumbnailPool.ToSharedRef())
+		];
+
+		ColumnIndex++;
+	}
+
+	ChildSlot
+	[
+		Panel
+	];
+}
+
+#undef LOCTEXT_NAMESPACE
\ No newline at end of file
diff --git a/Engine/Source/Editor/Persona/Private/SMontageEditor.cpp b/Engine/Source/Editor/Persona/Private/SMontageEditor.cpp
--- a/Engine/Source/Editor/Persona/Private/SMontageEditor.cpp
+++ b/Engine/Source/Editor/Persona/Private/SMontageEditor.cpp
@@ -26,15 +26,13 @@ TSharedRef<SWidget> SMontageEditor::CreateDocumentAnchor()
 	return IDocumentation::Get()->CreateAnchor(TEXT("Engine/Animation/AnimMontage"));
 }
 
-void SMontageEditor::Construct(const FArguments& InArgs)
+void SMontageEditor::Construct(const FArguments& InArgs, const FMontageEditorRequiredArgs& InRequiredArgs)
 {
 	MontageObj = InArgs._Montage;
 	check(MontageObj);
-
+	OnSectionsChanged = InArgs._OnSectionsChanged;
 	MontageObj->RegisterOnMontageChanged(UAnimMontage::FOnMontageChanged::CreateSP(this, &SMontageEditor::RebuildMontagePanel, false));
 		
-	WeakPersona = InArgs._Persona;
-	
 	// set child montage if montage has parent
 	bChildAnimMontage = MontageObj->HasParentAsset();
 
@@ -42,19 +40,13 @@ void SMontageEditor::Construct(const FArguments& InArgs)
 	bIsActiveTimerRegistered = false;
 
 	SAnimEditorBase::Construct( SAnimEditorBase::FArguments()
-		.Persona(InArgs._Persona)
-		);
+		.OnObjectsSelected(InArgs._OnObjectsSelected), 
+		InRequiredArgs.PreviewScene );
 
-	TSharedPtr<FPersona> SharedPersona = PersonaPtr.Pin();
-	if(SharedPersona.IsValid())
-	{
-		SharedPersona->RegisterOnPostUndo(FPersona::FOnPostUndo::CreateSP( this, &SMontageEditor::PostUndo ) );
-		SharedPersona->RegisterOnPersonaRefresh(FPersona::FOnPersonaRefresh::CreateSP(this, &SMontageEditor::RebuildMontagePanel, true));
-	}
+	InRequiredArgs.OnPostUndo.Add(FSimpleDelegate::CreateSP(this, &SMontageEditor::PostUndo));
 
-	SAssignNew(AnimTimingPanel, SAnimTimingPanel)
+	SAssignNew(AnimTimingPanel, SAnimTimingPanel, InRequiredArgs.OnAnimNotifiesChanged, InRequiredArgs.OnSectionsChanged)
 		.IsEnabled(!bChildAnimMontage)
-		.InWeakPersona(PersonaPtr)
 		.InSequence(MontageObj)
 		.WidgetWidth(S2ColumnWidget::DEFAULT_RIGHT_COLUMN_WIDTH)
 		.ViewInputMin(this, &SAnimEditorBase::GetViewMinInput)
@@ -75,76 +67,75 @@ void SMontageEditor::Construct(const FArguments& InArgs)
 		// add child montage warning section - and link to parent
 		EditorPanels->AddSlot()
 		.AutoHeight()
-		.Padding(0, 10)
+		.Padding(0, 0)
 		[
-			SNew(SVerticalBox)
-			
-			+SVerticalBox::Slot()
-			.Padding(3.f, 1.f)
-			.AutoHeight()
+			SNew(SBorder)
+			.Padding(FMargin(10.f, 4.f))
+			.BorderImage(FEditorStyle::GetBrush(TEXT("Graph.InstructionBackground")))
+			.BorderBackgroundColor(FLinearColor(0.1f, 0.1f, 0.1f, 0.5f))
+			.HAlign(HAlign_Center)
+			.ColorAndOpacity(FLinearColor(1.f, 1.f, 1.f, 0.8f))
 			[
-				SNew(SHorizontalBox)
-
-				+SHorizontalBox::Slot()
-				.AutoWidth()
-				[
-					SNew(STextBlock)
-					.ShadowOffset(FVector2D::UnitVector)
-					.Text(LOCTEXT("ParentAnimMontageLink", " This is a child anim montage. To edit the lay out, please go to parent montage "))
-					.Font(FEditorStyle::GetFontStyle("Persona.RetargetManager.BoldFont"))
-					.ColorAndOpacity(FLinearColor(FLinearColor::Yellow))
-				]
-				+ SHorizontalBox::Slot()
-				.AutoWidth()
+				SNew(SVerticalBox)
+			
+				+SVerticalBox::Slot()
+				.AutoHeight()
 				[
-					SNew(SButton)
-					.VAlign(VAlign_Center)
-					.ButtonStyle(FEditorStyle::Get(), "HoverHintOnly")
-					.OnClicked(this, &SMontageEditor::OnFindParentClassInContentBrowserClicked)
-					.ToolTipText(LOCTEXT("FindParentInCBToolTip", "Find parent in Content Browser"))
-					.ContentPadding(4.0f)
-					.ForegroundColor(FSlateColor::UseForeground())
+					SNew(SHorizontalBox)
+
+					+SHorizontalBox::Slot()
+					.AutoWidth()
 					[
-						SNew(SImage)
-						.Image(FEditorStyle::GetBrush("PropertyWindow.Button_Browse"))
+						SNew(STextBlock)
+						.Text(LOCTEXT("ParentAnimMontageLink", " This is a child anim montage. To edit the lay out, please go to the parent montage "))
+						.TextStyle(FEditorStyle::Get(), "Persona.MontageEditor.ChildMontageInstruction")
 					]
-				]
-				+ SHorizontalBox::Slot()
-				.AutoWidth()
-				[
-					SNew(SButton)
-					.VAlign(VAlign_Center)
-					.ButtonStyle(FEditorStyle::Get(), "HoverHintOnly")
-					.OnClicked(this, &SMontageEditor::OnEditParentClassClicked)
-					.ToolTipText(LOCTEXT("EditParentClassToolTip", "Open parent in editor"))
-					.ContentPadding(4.0f)
-					.ForegroundColor(FSlateColor::UseForeground())
+					+ SHorizontalBox::Slot()
+					.AutoWidth()
+					[
+						SNew(SButton)
+						.VAlign(VAlign_Center)
+						.ButtonStyle(FEditorStyle::Get(), "HoverHintOnly")
+						.OnClicked(this, &SMontageEditor::OnFindParentClassInContentBrowserClicked)
+						.ToolTipText(LOCTEXT("FindParentInCBToolTip", "Find parent in Content Browser"))
+						.ForegroundColor(FSlateColor::UseForeground())
+						[
+							SNew(SImage)
+							.Image(FEditorStyle::GetBrush("PropertyWindow.Button_Browse"))
+						]
+					]
+					+ SHorizontalBox::Slot()
+					.AutoWidth()
 					[
-						SNew(SImage)
-						.Image(FEditorStyle::GetBrush("PropertyWindow.Button_Edit"))
+						SNew(SButton)
+						.VAlign(VAlign_Center)
+						.ButtonStyle(FEditorStyle::Get(), "HoverHintOnly")
+						.OnClicked(this, &SMontageEditor::OnEditParentClassClicked)
+						.ToolTipText(LOCTEXT("EditParentClassToolTip", "Open parent in editor"))
+						.ForegroundColor(FSlateColor::UseForeground())
+						[
+							SNew(SImage)
+							.Image(FEditorStyle::GetBrush("PropertyWindow.Button_Edit"))
+						]
 					]
 				]
-			]
 
-			+ SVerticalBox::Slot()
-			.Padding(3.f, 1.f)
-			.AutoHeight()
-			[
-				SNew(STextBlock)
-				.ShadowOffset(FVector2D::UnitVector)
-				.Text(LOCTEXT("RemapHelpText", " To remap asset to different asset, use context menu or drag and drop the animation on the segment."))
-				.Font(FEditorStyle::GetFontStyle("Persona.RetargetManager.BoldFont"))
-				.ColorAndOpacity(FLinearColor(FLinearColor::Yellow))
+				+ SVerticalBox::Slot()
+				.AutoHeight()
+				[
+					SNew(STextBlock)
+					.Text(LOCTEXT("RemapHelpText", " To remap an asset to a different asset, use the context menu or drag and drop the animation on to the segment."))
+					.TextStyle(FEditorStyle::Get(), "Persona.MontageEditor.ChildMontageInstruction")
+				]
 			]
 		];
 	}
 
 	EditorPanels->AddSlot()
 	.AutoHeight()
-	.Padding(0, 10)
+	.Padding(0, 5)
 	[
-		SAssignNew( AnimMontagePanel, SAnimMontagePanel )
-		.Persona(PersonaPtr)
+		SAssignNew( AnimMontagePanel, SAnimMontagePanel, InRequiredArgs.OnAnimNotifiesChanged, InRequiredArgs.OnSectionsChanged)
 		.Montage(MontageObj)
 		.bChildAnimMontage(bChildAnimMontage)
 		.MontageEditor(SharedThis(this))
@@ -155,6 +146,7 @@ void SMontageEditor::Construct(const FArguments& InArgs)
 		.InputMax(this, &SAnimEditorBase::GetMaxInput)
 		.OnSetInputViewRange(this, &SAnimEditorBase::SetInputViewRange)
 		.SectionTimingNodeVisibility(SectionVisibility)
+		.OnInvokeTab(InArgs._OnInvokeTab)
 	];
 
 	EditorPanels->AddSlot()
@@ -178,10 +170,9 @@ void SMontageEditor::Construct(const FArguments& InArgs)
 	.AutoHeight()
 	.Padding(0, 10)
 	[
-		SAssignNew( AnimNotifyPanel, SAnimNotifyPanel )
-		.IsEnabled(!bChildAnimMontage)
-		.Persona(InArgs._Persona)
+		SAssignNew( AnimNotifyPanel, SAnimNotifyPanel, InRequiredArgs.OnPostUndo )
 		.Sequence(MontageObj)
+		.IsEnabled(!bChildAnimMontage)
 		.WidgetWidth(S2ColumnWidget::DEFAULT_RIGHT_COLUMN_WIDTH)
 		.InputMin(this, &SAnimEditorBase::GetMinInput)
 		.InputMax(this, &SAnimEditorBase::GetMaxInput)
@@ -193,15 +184,16 @@ void SMontageEditor::Construct(const FArguments& InArgs)
 		.MarkerBars(this, &SMontageEditor::GetMarkerBarInformation)
 		.OnRequestRefreshOffsets(this, &SMontageEditor::RefreshNotifyTriggerOffsets)
 		.OnGetTimingNodeVisibility(TimingNodeVisibilityDelegate)
+		.OnAnimNotifiesChanged(InArgs._OnAnimNotifiesChanged)
+		.OnInvokeTab(InArgs._OnInvokeTab)
 	];
 
 	EditorPanels->AddSlot()
 	.AutoHeight()
 	.Padding(0, 10)
 	[
-		SAssignNew( AnimCurvePanel, SAnimCurvePanel )
+		SAssignNew( AnimCurvePanel, SAnimCurvePanel, InRequiredArgs.EditableSkeleton )
 		.IsEnabled(!bChildAnimMontage)
-		.Persona(InArgs._Persona)
 		.Sequence(MontageObj)
 		.WidgetWidth(S2ColumnWidget::DEFAULT_RIGHT_COLUMN_WIDTH)
 		.ViewInputMin(this, &SAnimEditorBase::GetViewMinInput)
@@ -210,6 +202,7 @@ void SMontageEditor::Construct(const FArguments& InArgs)
 		.InputMax(this, &SAnimEditorBase::GetMaxInput)
 		.OnSetInputViewRange(this, &SAnimEditorBase::SetInputViewRange)
 		.OnGetScrubValue(this, &SAnimEditorBase::GetScrubValue)
+		.OnCurvesChanged(InArgs._OnCurvesChanged)
 	];
 
 	if (MontageObj)
@@ -257,19 +250,11 @@ SMontageEditor::~SMontageEditor()
 	{
 		MontageObj->UnregisterOnMontageChanged(this);
 	}
-
-	TSharedPtr<FPersona> SharedPersona = PersonaPtr.Pin();
-	if (SharedPersona.IsValid())
-	{
-		SharedPersona->UnregisterOnPostUndo(this);
-		SharedPersona->UnregisterOnPersonaRefresh(this);
-	}
 }
 
 TSharedRef<class SAnimationScrubPanel> SMontageEditor::ConstructAnimScrubPanel()
 {
-	return SAssignNew(AnimMontageScrubPanel, SAnimMontageScrubPanel)
-		.Persona(PersonaPtr)
+	return SAssignNew(AnimMontageScrubPanel, SAnimMontageScrubPanel, GetPreviewScene())
 		.LockedSequence(MontageObj)
 		.ViewInputMin(this, &SMontageEditor::GetViewMinInput)
 		.ViewInputMax(this, &SMontageEditor::GetViewMaxInput)
@@ -426,11 +411,7 @@ void SMontageEditor::OnEditSectionTimeFinish( int32 SectionIndex )
 		AnimMontageSectionsPanel->Update();
 	}
 
-	TSharedPtr<FPersona> SharedPersona = WeakPersona.Pin();
-	if(SharedPersona.IsValid())
-	{
-		SharedPersona->OnSectionsChanged.Broadcast();
-	}
+	OnSectionsChanged.ExecuteIfBound();
 }
 
 void SMontageEditor::PreAnimUpdate()
@@ -831,7 +812,7 @@ void SMontageEditor::ShowSectionInDetailsView(int32 SectionIndex)
 
 void SMontageEditor::RestartPreview()
 {
-	UAnimPreviewInstance * Preview = Cast<UAnimPreviewInstance>(PersonaPtr.Pin()->GetPreviewMeshComponent() ? PersonaPtr.Pin()->GetPreviewMeshComponent()->PreviewInstance:NULL);
+	UAnimPreviewInstance * Preview = Cast<UAnimPreviewInstance>(GetPreviewScene()->GetPreviewMeshComponent() ? GetPreviewScene()->GetPreviewMeshComponent()->PreviewInstance:NULL);
 	if (Preview)
 	{
 		Preview->MontagePreview_PreviewNormal(INDEX_NONE, Preview->IsPlaying());
@@ -840,7 +821,7 @@ void SMontageEditor::RestartPreview()
 
 void SMontageEditor::RestartPreviewFromSection(int32 FromSectionIdx)
 {
-	UAnimPreviewInstance * Preview = Cast<UAnimPreviewInstance>(PersonaPtr.Pin()->GetPreviewMeshComponent() ? PersonaPtr.Pin()->GetPreviewMeshComponent()->PreviewInstance:NULL);
+	UAnimPreviewInstance * Preview = Cast<UAnimPreviewInstance>(GetPreviewScene()->GetPreviewMeshComponent() ? GetPreviewScene()->GetPreviewMeshComponent()->PreviewInstance:NULL);
 	if(Preview)
 	{
 		Preview->MontagePreview_PreviewNormal(FromSectionIdx, Preview->IsPlaying());
@@ -849,7 +830,7 @@ void SMontageEditor::RestartPreviewFromSection(int32 FromSectionIdx)
 
 void SMontageEditor::RestartPreviewPlayAllSections()
 {
-	UAnimPreviewInstance * Preview = Cast<UAnimPreviewInstance>(PersonaPtr.Pin()->GetPreviewMeshComponent() ? PersonaPtr.Pin()->GetPreviewMeshComponent()->PreviewInstance:NULL);
+	UAnimPreviewInstance * Preview = Cast<UAnimPreviewInstance>(GetPreviewScene()->GetPreviewMeshComponent() ? GetPreviewScene()->GetPreviewMeshComponent()->PreviewInstance:NULL);
 	if(Preview)
 	{
 		Preview->MontagePreview_PreviewAllSections(Preview->IsPlaying());
diff --git a/Engine/Source/Editor/Persona/Private/SMontageEditor.h b/Engine/Source/Editor/Persona/Private/SMontageEditor.h
--- a/Engine/Source/Editor/Persona/Private/SMontageEditor.h
+++ b/Engine/Source/Editor/Persona/Private/SMontageEditor.h
@@ -15,6 +15,23 @@ class SAnimMontageSectionsPanel;
 class SAnimMontageScrubPanel;
 class SAnimTimingPanel;
 
+struct FMontageEditorRequiredArgs
+{
+	FMontageEditorRequiredArgs(const TSharedRef<class IPersonaPreviewScene>& InPreviewScene, const TSharedRef<class IEditableSkeleton>& InEditableSkeleton, FSimpleMulticastDelegate& InOnPostUndo, FSimpleMulticastDelegate& InOnAnimNotifiesChanged, FSimpleMulticastDelegate& InOnSectionsChanged)
+		: PreviewScene(InPreviewScene)
+		, EditableSkeleton(InEditableSkeleton)
+		, OnPostUndo(InOnPostUndo)
+		, OnAnimNotifiesChanged(InOnAnimNotifiesChanged)
+		, OnSectionsChanged(InOnSectionsChanged)
+	{}
+
+	TSharedRef<class IPersonaPreviewScene> PreviewScene;
+	TSharedRef<class IEditableSkeleton> EditableSkeleton;
+	FSimpleMulticastDelegate& OnPostUndo;
+	FSimpleMulticastDelegate& OnAnimNotifiesChanged;
+	FSimpleMulticastDelegate& OnSectionsChanged;
+};
+
 //////////////////////////////////////////////////////////////////////////
 // SMontageEditor
 
@@ -28,18 +45,20 @@ class SMontageEditor : public SAnimEditorBase
 {
 public:
 	SLATE_BEGIN_ARGS( SMontageEditor )
-		: _Persona()
-		, _Montage(NULL)
+		: _Montage(NULL)
 		{}
 
-		SLATE_ARGUMENT( TSharedPtr<FPersona>, Persona )
+		SLATE_EVENT(FOnInvokeTab, OnInvokeTab)
+		SLATE_EVENT(FSimpleDelegate, OnCurvesChanged)
+		SLATE_EVENT(FSimpleDelegate, OnSectionsChanged)
 		SLATE_ARGUMENT( UAnimMontage*, Montage )
+		SLATE_EVENT(FOnObjectsSelected, OnObjectsSelected)
+		SLATE_EVENT(FSimpleDelegate, OnAnimNotifiesChanged)
 	SLATE_END_ARGS()
 
 	~SMontageEditor();
 
 private:
-	/** Persona reference **/
 	TSharedPtr<SAnimMontagePanel> AnimMontagePanel;
 	TSharedPtr<SAnimNotifyPanel> AnimNotifyPanel;
 	TSharedPtr<SAnimCurvePanel>	AnimCurvePanel;
@@ -47,22 +66,14 @@ private:
 	TSharedPtr<SAnimMontageScrubPanel> AnimMontageScrubPanel;
 	TSharedPtr<SAnimTimingPanel> AnimTimingPanel;
 
-	TWeakPtr<FPersona> WeakPersona;
-
- 	/* 
-	 * Child Anim Montage: Child Anim Montage only can replace name of animations, and no other meaningful edits 
-	 * as it will derive every data from Parent. There might be some other data that will allow to be replaced, but for now, it is
-	 * not. 
-	 */	
-	bool bChildAnimMontage;
-
 protected:
 	//~ Begin SAnimEditorBase Interface
 	virtual TSharedRef<class SAnimationScrubPanel> ConstructAnimScrubPanel() override;
 	//~ End SAnimEditorBase Interface
 
 public:
-	void Construct(const FArguments& InArgs);
+	void Construct(const FArguments& InArgs, const FMontageEditorRequiredArgs& InRequiredArgs);
+
 	void SetMontageObj(UAnimMontage * NewMontage);
 	UAnimMontage * GetMontageObj() const { return MontageObj; }
 
@@ -74,7 +85,14 @@ public:
 
 private:
 	/** Pointer to the animation sequence being edited */
+	// @todo fix this
 	UAnimMontage* MontageObj;
+	/*
+	* Child Anim Montage: Child Anim Montage only can replace name of animations, and no other meaningful edits
+	* as it will derive every data from Parent. There might be some other data that will allow to be replaced, but for now, it is
+	* not.
+	*/
+	bool bChildAnimMontage;
 
 	/** If previewing section, it is section used to restart previewing when play button is pushed */
 	int32 PreviewingStartSectionIdx;
@@ -157,6 +175,9 @@ public:
 	virtual TSharedRef<SWidget> CreateDocumentAnchor() override;
 	//~ End SAnimEditorBase Interface
 
+	/** Delegate fired when montage sections have changed */
+	FSimpleDelegate			OnSectionsChanged;
+
 	FReply	OnFindParentClassInContentBrowserClicked();
 	FReply	OnEditParentClassClicked();
 };
diff --git a/Engine/Source/Editor/Persona/Private/SMorphTargetViewer.cpp b/Engine/Source/Editor/Persona/Private/SMorphTargetViewer.cpp
--- a/Engine/Source/Editor/Persona/Private/SMorphTargetViewer.cpp
+++ b/Engine/Source/Editor/Persona/Private/SMorphTargetViewer.cpp
@@ -40,12 +40,9 @@ public:
 		/* Widget used to display the list of morph targets */
 		SLATE_ARGUMENT( TSharedPtr<SMorphTargetListType>, MorphTargetListView )
 
-		/* Persona used to update the viewport when a weight slider is dragged */
-		SLATE_ARGUMENT( TWeakPtr<FPersona>, Persona )
-
 		SLATE_END_ARGS()
 
-		void Construct( const FArguments& InArgs, const TSharedRef<STableViewBase>& OwnerTableView );
+		void Construct( const FArguments& InArgs, const TSharedRef<IPersonaPreviewScene>& InPreviewScene, const TSharedRef<STableViewBase>& OwnerTableView );
 
 	/** Overridden from SMultiColumnTableRow.  Generates a widget for this column of the tree row. */
 	virtual TSharedRef<SWidget> GenerateWidgetForColumn( const FName& ColumnName ) override;
@@ -88,16 +85,16 @@ private:
 	/** The name and weight of the morph target */
 	FDisplayedMorphTargetInfoPtr	Item;
 
-	/** Pointer back to the Persona that owns us */
-	TWeakPtr<FPersona> PersonaPtr;
+	/** Preview scene - we invalidate this etc. */
+	TWeakPtr<IPersonaPreviewScene> PreviewScenePtr;
 };
 
-void SMorphTargetListRow::Construct( const FArguments& InArgs, const TSharedRef<STableViewBase>& InOwnerTableView )
+void SMorphTargetListRow::Construct( const FArguments& InArgs, const TSharedRef<IPersonaPreviewScene>& InPreviewScene, const TSharedRef<STableViewBase>& InOwnerTableView )
 {
 	Item = InArgs._Item;
 	MorphTargetViewer = InArgs._MorphTargetViewer;
 	MorphTargetListView = InArgs._MorphTargetListView;
-	PersonaPtr = InArgs._Persona;
+	PreviewScenePtr = InPreviewScene;
 
 	check( Item.IsValid() );
 
@@ -211,10 +208,7 @@ void SMorphTargetListRow::OnMorphTargetWeightChanged( float NewWeight )
 
 	MorphTargetViewer->AddMorphTargetOverride( Item->Name, Item->Weight, false );
 
-	if (PersonaPtr.IsValid())
-	{
-		PersonaPtr.Pin()->RefreshViewport();
-	}
+	PreviewScenePtr.Pin()->InvalidateViews();
 
 #if 0 
 	TArray< TSharedPtr< FDisplayedMorphTargetInfo > > SelectedRows = MorphTargetListView->GetSelectedItems();
@@ -259,24 +253,21 @@ void SMorphTargetListRow::OnMorphTargetWeightValueCommitted( float NewWeight, ET
 			}
 		}
 
-		if(PersonaPtr.IsValid())
-		{
-			PersonaPtr.Pin()->RefreshViewport();
-		}
+		PreviewScenePtr.Pin()->InvalidateViews();
 	}
 }
 
 float SMorphTargetListRow::GetWeight() const 
 { 
-	if (PersonaPtr.IsValid() && Item->bAutoFillData)
+	if (Item->bAutoFillData)
 	{
-		USkeletalMeshComponent* SkelComp = PersonaPtr.Pin()->GetPreviewMeshComponent();
+		USkeletalMeshComponent* SkelComp = PreviewScenePtr.Pin()->GetPreviewMeshComponent();
 		UAnimInstance* AnimInstance = (SkelComp) ? SkelComp->GetAnimInstance() : nullptr;
 		if (AnimInstance)
 		{
 			// make sure if they have value that's not same as saved value
 			TMap<FName, float> MorphCurves;
-			AnimInstance->GetAnimationCurveList(ACF_DriveMorphTarget, MorphCurves);
+			AnimInstance->GetAnimationCurveList(EAnimCurveType::MorphTargetCurve, MorphCurves);
 			const float* CurrentValue = MorphCurves.Find(Item->Name);
 			if (CurrentValue && *CurrentValue != 0.f)
 			{
@@ -290,17 +281,13 @@ float SMorphTargetListRow::GetWeight() const
 //////////////////////////////////////////////////////////////////////////
 // SMorphTargetViewer
 
-void SMorphTargetViewer::Construct(const FArguments& InArgs)
+void SMorphTargetViewer::Construct(const FArguments& InArgs, const TSharedRef<IPersonaPreviewScene>& InPreviewScene, FSimpleMulticastDelegate& OnPostUndo)
 {
-	PersonaPtr = InArgs._Persona;
-	SkeletalMesh = NULL;
+	PreviewScenePtr = InPreviewScene;
 
-	if ( PersonaPtr.IsValid() )
-	{
-		SkeletalMesh = PersonaPtr.Pin()->GetMesh();
-		PersonaPtr.Pin()->RegisterOnPreviewMeshChanged( FPersona::FOnPreviewMeshChanged::CreateSP( this, &SMorphTargetViewer::OnPreviewMeshChanged ) );
-		PersonaPtr.Pin()->RegisterOnPostUndo(FPersona::FOnPostUndo::CreateSP(this, &SMorphTargetViewer::OnPostUndo));
-	}
+	SkeletalMesh = InPreviewScene->GetPreviewMeshComponent()->SkeletalMesh;
+	InPreviewScene->RegisterOnPreviewMeshChanged( FOnPreviewMeshChanged::CreateSP( this, &SMorphTargetViewer::OnPreviewMeshChanged ) );
+	OnPostUndo.Add(FSimpleDelegate::CreateSP(this, &SMorphTargetViewer::OnPostUndo));
 
 	const FText SkeletalMeshName = SkeletalMesh ? FText::FromString( SkeletalMesh->GetName() ) : LOCTEXT( "MorphTargetMeshNameLabel", "No Skeletal Mesh Present" );
 
@@ -385,8 +372,7 @@ TSharedRef<ITableRow> SMorphTargetViewer::GenerateMorphTargetRow(TSharedPtr<FDis
 	check( InInfo.IsValid() );
 
 	return
-		SNew( SMorphTargetListRow, OwnerTable )
-		.Persona( PersonaPtr )
+		SNew( SMorphTargetListRow, PreviewScenePtr.Pin().ToSharedRef(), OwnerTable )
 		.Item( InInfo )
 		.MorphTargetViewer( this )
 		.MorphTargetListView( MorphTargetListView );
@@ -399,11 +385,24 @@ TSharedPtr<SWidget> SMorphTargetViewer::OnGetContextMenuContent() const
 
 	MenuBuilder.BeginSection("MorphTargetAction", LOCTEXT( "MorphsAction", "Selected Item Actions" ) );
 	{
-		FUIAction Action = FUIAction( FExecuteAction::CreateSP( this, &SMorphTargetViewer::OnDeleteMorphTargets ), 
-									  FCanExecuteAction::CreateSP( this, &SMorphTargetViewer::CanPerformDelete ) );
-		const FText Label = LOCTEXT("DeleteMorphTargetButtonLabel", "Delete");
-		const FText ToolTipText = LOCTEXT("DeleteMorphTargetButtonTooltip", "Deletes the selected morph targets.");
-		MenuBuilder.AddMenuEntry( Label, ToolTipText, FSlateIcon(), Action);
+		FUIAction Action;
+
+		{
+			Action.ExecuteAction = FExecuteAction::CreateSP(this, &SMorphTargetViewer::OnDeleteMorphTargets);
+			Action.CanExecuteAction = FCanExecuteAction::CreateSP(this, &SMorphTargetViewer::CanPerformDelete);
+			const FText Label = LOCTEXT("DeleteMorphTargetButtonLabel", "Delete");
+			const FText ToolTipText = LOCTEXT("DeleteMorphTargetButtonTooltip", "Deletes the selected morph targets.");
+			MenuBuilder.AddMenuEntry(Label, ToolTipText, FSlateIcon(), Action);
+		}
+
+		{
+			Action.ExecuteAction = FExecuteAction::CreateSP(this, &SMorphTargetViewer::OnCopyMorphTargetNames);
+			Action.CanExecuteAction = nullptr;
+			const FText Label = LOCTEXT("CopyMorphTargetNamesButtonLabel", "Copy Names");
+			const FText ToolTipText = LOCTEXT("CopyMorphTargetNamesButtonTooltip", "Copy the names of selected morph targets to clipboard");
+			MenuBuilder.AddMenuEntry(Label, ToolTipText, FSlateIcon(), Action);
+		}
+
 	}
 	MenuBuilder.EndSection();
 
@@ -416,7 +415,7 @@ void SMorphTargetViewer::CreateMorphTargetList( const FString& SearchText )
 
 	if ( SkeletalMesh )
 	{
-		UDebugSkelMeshComponent* MeshComponent = PersonaPtr.Pin()->GetPreviewMeshComponent();
+		UDebugSkelMeshComponent* MeshComponent = PreviewScenePtr.Pin()->GetPreviewMeshComponent();
 		TArray<UMorphTarget*>& MorphTargets = SkeletalMesh->MorphTargets;
 
 		bool bDoFiltering = !SearchText.IsEmpty();
@@ -450,14 +449,11 @@ void SMorphTargetViewer::CreateMorphTargetList( const FString& SearchText )
 
 void SMorphTargetViewer::AddMorphTargetOverride( FName& Name, float Weight, bool bRemoveZeroWeight )
 {
-	if ( PersonaPtr.IsValid() )
-	{
-		UDebugSkelMeshComponent* Mesh = PersonaPtr.Pin()->GetPreviewMeshComponent();
+	UDebugSkelMeshComponent* Mesh = PreviewScenePtr.Pin()->GetPreviewMeshComponent();
 
-		if ( Mesh )
-		{
-			Mesh->SetMorphTarget( Name, Weight, bRemoveZeroWeight );
-		}
+	if ( Mesh )
+	{
+		Mesh->SetMorphTarget( Name, Weight, bRemoveZeroWeight );
 	}
 }
 
@@ -493,16 +489,33 @@ void SMorphTargetViewer::OnDeleteMorphTargets()
 	CreateMorphTargetList( NameFilterBox->GetText().ToString() );
 }
 
-SMorphTargetViewer::~SMorphTargetViewer()
+void SMorphTargetViewer::OnCopyMorphTargetNames()
 {
-	if ( PersonaPtr.IsValid() )
+	FString CopyText;
+
+	TArray< TSharedPtr< FDisplayedMorphTargetInfo > > SelectedRows = MorphTargetListView->GetSelectedItems();
+	for (int RowIndex = 0; RowIndex < SelectedRows.Num(); ++RowIndex)
+	{
+		UMorphTarget* MorphTarget = SkeletalMesh->FindMorphTarget(SelectedRows[RowIndex]->Name);
+		if (MorphTarget)
+		{
+			CopyText += FString::Printf(TEXT("%s\r\n"), *MorphTarget->GetName());
+		}
+	}
+
+	if(!CopyText.IsEmpty())
 	{
-		PersonaPtr.Pin()->UnregisterOnPreviewMeshChanged(this);
-		PersonaPtr.Pin()->UnregisterOnPostUndo(this);
+		FPlatformMisc::ClipboardCopy(*CopyText);
+	}
+}
 
-		UDebugSkelMeshComponent* Mesh = PersonaPtr.Pin()->GetPreviewMeshComponent();
+SMorphTargetViewer::~SMorphTargetViewer()
+{
+	if (PreviewScenePtr.IsValid())
+	{
+		UDebugSkelMeshComponent* Mesh = PreviewScenePtr.Pin()->GetPreviewMeshComponent();
 
-		if ( Mesh )
+		if (Mesh)
 		{
 			Mesh->ClearMorphTargets();
 		}
@@ -527,7 +540,7 @@ void SMorphTargetViewer::NotifySelectionChange() const
 	}
 
 	// stil have to call this even if empty, otherwise it won't clear it
-	PersonaPtr.Pin()->SetSelectedMorphTargets(SkeletalMesh, SelectedMorphtargetNames);
+	SetSelectedMorphTargets(SelectedMorphtargetNames);
 }
 
 void SMorphTargetViewer::OnRowsSelectedChanged(TSharedPtr<FDisplayedMorphTargetInfo> Item, ESelectInfo::Type SelectInfo)
@@ -535,5 +548,31 @@ void SMorphTargetViewer::OnRowsSelectedChanged(TSharedPtr<FDisplayedMorphTargetI
 	NotifySelectionChange();
 }
 
+void SMorphTargetViewer::SetSelectedMorphTargets(const TArray<FName>& SelectedMorphTargetNames) const
+{
+	UDebugSkelMeshComponent* PreviewComponent = PreviewScenePtr.Pin()->GetPreviewMeshComponent();
+	if (PreviewComponent)
+	{
+		PreviewComponent->MorphTargetOfInterests.Reset();
+
+		if (SelectedMorphTargetNames.Num() > 0)
+		{
+			for (const FName& MorphTargetName : SelectedMorphTargetNames)
+			{
+				int32 MorphtargetIdx;
+				UMorphTarget* MorphTarget = SkeletalMesh->FindMorphTargetAndIndex(MorphTargetName, MorphtargetIdx);
+				if (MorphTarget != nullptr)
+				{
+					PreviewComponent->MorphTargetOfInterests.AddUnique(MorphTarget);
+				}
+			}
+
+			PreviewScenePtr.Pin()->InvalidateViews();
+			PreviewComponent->PostInitMeshObject(PreviewComponent->MeshObject);
+		}
+	}
+}
+
+
 #undef LOCTEXT_NAMESPACE
 
diff --git a/Engine/Source/Editor/Persona/Private/SMorphTargetViewer.h b/Engine/Source/Editor/Persona/Private/SMorphTargetViewer.h
--- a/Engine/Source/Editor/Persona/Private/SMorphTargetViewer.h
+++ b/Engine/Source/Editor/Persona/Private/SMorphTargetViewer.h
@@ -43,12 +43,8 @@ class SMorphTargetViewer : public SCompoundWidget
 {
 public:
 	SLATE_BEGIN_ARGS( SMorphTargetViewer )
-		: _Persona()
 	{}
-		
-		/* The Persona that owns this table */
-		SLATE_ARGUMENT( TWeakPtr< FPersona >, Persona )
-
+	
 	SLATE_END_ARGS()
 
 	/**
@@ -57,7 +53,7 @@ public:
 	* @param InArgs - Arguments passed from Slate
 	*
 	*/
-	void Construct( const FArguments& InArgs );
+	void Construct( const FArguments& InArgs, const TSharedRef<IPersonaPreviewScene>& InPreviewScene, FSimpleMulticastDelegate& OnPostUndo );
 
 	/**
 	* Destructor - resets the morph targets
@@ -127,6 +123,9 @@ public:
 	*/
 	void OnDeleteMorphTargets();
 
+	/** Handler for copying morph taret names */
+	void OnCopyMorphTargetNames();
+
 	/**
 	* Accessor so our rows can grab the filtertext for highlighting
 	*
@@ -154,8 +153,11 @@ private:
 	*/
 	void CreateMorphTargetList( const FString& SearchText = FString() );
 
+	// Sets the selected morph target
+	void SetSelectedMorphTargets(const TArray<FName>& SelectedMorphTargetNames) const;
+
 	/** Pointer back to the Persona that owns us */
-	TWeakPtr<FPersona> PersonaPtr;
+	TWeakPtr<class IPersonaPreviewScene> PreviewScenePtr;
 
 	/** Box to filter to a specific morph target name */
 	TSharedPtr<SSearchBox>	NameFilterBox;
diff --git a/Engine/Source/Editor/Persona/Private/SPersonaDetails.cpp b/Engine/Source/Editor/Persona/Private/SPersonaDetails.cpp
new file mode 100644
--- /dev/null
+++ b/Engine/Source/Editor/Persona/Private/SPersonaDetails.cpp
@@ -0,0 +1,17 @@
+// Copyright 1998-2016 Epic Games, Inc. All Rights Reserved.
+
+#include "PersonaPrivatePCH.h"
+#include "SPersonaDetails.h"
+
+void SPersonaDetails::Construct(const FArguments& InArgs)
+{
+	FDetailsViewArgs DetailsViewArgs(false, false, true, FDetailsViewArgs::HideNameArea, true);
+
+	FPropertyEditorModule& PropertyEditorModule = FModuleManager::LoadModuleChecked<FPropertyEditorModule>("PropertyEditor");
+	DetailsView = PropertyEditorModule.CreateDetailView(DetailsViewArgs);
+
+	ChildSlot
+	[
+		DetailsView.ToSharedRef()
+	];
+}
\ No newline at end of file
diff --git a/Engine/Source/Editor/Persona/Private/SPersonaToolbar.cpp b/Engine/Source/Editor/Persona/Private/SPersonaToolbar.cpp
--- a/Engine/Source/Editor/Persona/Private/SPersonaToolbar.cpp
+++ b/Engine/Source/Editor/Persona/Private/SPersonaToolbar.cpp
@@ -11,6 +11,7 @@
 #include "AnimGraphDefinitions.h"
 #include "Editor/MainFrame/Public/MainFrame.h"
 #include "IDocumentation.h"
+#include "IPersonaToolkit.h"
 
 #define LOCTEXT_NAMESPACE "PersonaToolbar"
 
@@ -69,6 +70,16 @@ void FPersonaToolbar::FillPersonaToolbar(FToolBarBuilder& ParentToolbarBuilder)
 {
 }
 
+UObject* FPersonaToolbar::GetSkeletonAsUObject() const
+{
+	return Persona.Pin()->GetPersonaToolkit()->GetSkeleton();
+}
+
+UObject* FPersonaToolbar::GetMeshAsUObject() const
+{
+	return Persona.Pin()->GetPersonaToolkit()->GetMesh();
+}
+
 void FPersonaToolbar::FillPersonaModeToolbar(FToolBarBuilder& ParentToolbarBuilder)
 {
 	TSharedPtr<FPersona> PersonaPtr = Persona.Pin();
@@ -99,7 +110,7 @@ void FPersonaToolbar::FillPersonaModeToolbar(FToolBarBuilder& ParentToolbarBuild
 			.ShortContents()
 			[
 				SNew(SContentReference)
-				.AssetReference(PersonaPtr.ToSharedRef(), &FPersona::GetSkeletonAsObject)
+				.AssetReference(this, &FPersonaToolbar::GetSkeletonAsUObject)
 				.AllowSelectingNewAsset(false)
 				.AllowClearingReference(false)
 				.WidthOverride(ContentRefWidth)
@@ -125,7 +136,7 @@ void FPersonaToolbar::FillPersonaModeToolbar(FToolBarBuilder& ParentToolbarBuild
 				SNew(SContentReference)
 				.WidthOverride(ContentRefWidth)
 				.AllowSelectingNewAsset(true)
-				.AssetReference(PersonaPtr.ToSharedRef(), &FPersona::GetMeshAsObject)
+				.AssetReference(this, &FPersonaToolbar::GetMeshAsUObject)
 				.AllowedClass(USkeletalMesh::StaticClass())
 				.OnShouldFilterAsset(this, &FPersonaToolbar::ShouldFilterAssetBasedOnSkeleton)
 				.OnSetReference(this, &FPersonaToolbar::OnSetSkeletalMeshReference)
@@ -223,7 +234,7 @@ bool FPersonaToolbar::ShouldFilterAssetBasedOnSkeleton(const FAssetData& AssetDa
 
 	if ( !SkeletonName.IsEmpty() )
 	{
-		USkeleton* Skeleton = Persona.Pin()->GetSkeleton();
+		USkeleton* Skeleton = Persona.Pin()->GetPersonaToolkit()->GetSkeleton();
 		if ( SkeletonName == FString::Printf(TEXT("%s'%s'"), *Skeleton->GetClass()->GetName(), *Skeleton->GetPathName()) )
 		{
 			return false;
diff --git a/Engine/Source/Editor/Persona/Private/SPersonaToolbar.h b/Engine/Source/Editor/Persona/Private/SPersonaToolbar.h
--- a/Engine/Source/Editor/Persona/Private/SPersonaToolbar.h
+++ b/Engine/Source/Editor/Persona/Private/SPersonaToolbar.h
@@ -1,17 +1,11 @@
 // Copyright 1998-2016 Epic Games, Inc. All Rights Reserved.
 
-
-#ifndef __SPersonaToolbar_h__
-#define __SPersonaToolbar_h__
-
-#include "Editor/Kismet/Public/SBlueprintEditorToolbar.h"
-
-
+#pragma once
 
 class FPersonaToolbar : public TSharedFromThis<FPersonaToolbar>
 {
 public:
-	void SetupPersonaToolbar(TSharedPtr<FExtender> Extender, TSharedPtr<FPersona> InPersona);
+	void SetupPersonaToolbar(TSharedPtr<FExtender> Extender, TSharedPtr<class FPersona> InPersona);
 	
 private:
 	void FillPersonaModeToolbar(FToolBarBuilder& ToolbarBuilder);
@@ -27,11 +21,13 @@ protected:
 	/* Set new reference for the current animation */
 	void OnSetAnimationAssetReference(UObject* Object);
 
+	/** Thunk for SContentReference */
+	UObject* GetSkeletonAsUObject() const;
+
+	/** Thunk for SContentReference */
+	UObject* GetMeshAsUObject() const;
+
 private:
 	/** Pointer back to the blueprint editor tool that owns us */
-	TWeakPtr<FPersona> Persona;
+	TWeakPtr<class FPersona> Persona;
 };
-
-
-
-#endif		// __SPersonaToolbar_h__
diff --git a/Engine/Source/Editor/Persona/Private/SPoseEditor.cpp b/Engine/Source/Editor/Persona/Private/SPoseEditor.cpp
--- a/Engine/Source/Editor/Persona/Private/SPoseEditor.cpp
+++ b/Engine/Source/Editor/Persona/Private/SPoseEditor.cpp
@@ -8,9 +8,14 @@
 #include "ScopedTransaction.h"
 #include "SSearchBox.h"
 #include "Animation/AnimSingleNodeInstance.h"
+#include "IPersonaPreviewScene.h"
+#include "IPersonaToolkit.h"
+#include "IEditableSkeleton.h"
 
 #include "SInlineEditableTextBlock.h"
 
+#include "GenericCommands.h"
+#include "PoseEditorCommands.h"
 
 #define LOCTEXT_NAMESPACE "AnimPoseEditor"
 
@@ -24,30 +29,24 @@ const float MaxPoseWeight = 1.f;
 //////////////////////////////////////////////////////////////////////////
 // SPoseEditor
 
-void SPoseEditor::Construct(const FArguments& InArgs)
+void SPoseEditor::Construct(const FArguments& InArgs, const TSharedRef<class IPersonaToolkit>& InPersonaToolkit, const TSharedRef<IEditableSkeleton>& InEditableSkeleton, const TSharedRef<class IPersonaPreviewScene>& InPreviewScene)
 {
 	PoseAssetObj = InArgs._PoseAsset;
 	check(PoseAssetObj);
 
-	SAnimEditorBase::Construct( SAnimEditorBase::FArguments()
-		.Persona(InArgs._Persona)
-		.DisplayAnimInfoBar(false)
-		);
+	SAnimEditorBase::Construct(SAnimEditorBase::FArguments()
+		.DisplayAnimInfoBar(false),
+		InPreviewScene);
 
  	EditorPanels->AddSlot()
  	.AutoHeight()
  	.Padding(0, 10)
 	[
- 		SNew( SPoseViewer)
- 		.Persona(InArgs._Persona)
+ 		SNew( SPoseViewer, InPersonaToolkit, InEditableSkeleton, InPreviewScene)
 		.PoseAsset(PoseAssetObj)
  	];
 }
 
-SPoseEditor::~SPoseEditor()
-{
-}
-
 
 //////////////////////////////////////////////////////////////////////////
 // SPoseListRow
@@ -104,11 +103,12 @@ bool SPoseListRow::OnVerifyNameChanged(const FText& InText, FText& OutErrorMessa
 }
 
 
-void SPoseListRow::Construct(const FArguments& InArgs, const TSharedRef<STableViewBase>& InOwnerTableView)
+void SPoseListRow::Construct(const FArguments& InArgs, const TSharedRef<STableViewBase>& InOwnerTableView, const TSharedRef<IPersonaPreviewScene>& InPreviewScene)
 {
 	Item = InArgs._Item;
 	PoseViewerPtr = InArgs._PoseViewer;
 	FilterText = InArgs._FilterText;
+	PreviewScenePtr = InPreviewScene;
 
 	check(Item.IsValid());
 
@@ -174,10 +174,7 @@ void SPoseListRow::OnPoseWeightChanged(float NewWeight)
 		TSharedPtr<SPoseViewer> PoseViewer = PoseViewerPtr.Pin();
 		PoseViewer->AddCurveOverride(Item->Name, Item->Weight);
 
-		if (PoseViewer->PersonaPtr.IsValid())
-		{
-			PoseViewer->PersonaPtr.Pin()->RefreshViewport();
-		}
+		PreviewScenePtr.Pin()->InvalidateViews();
 	}
 }
 
@@ -296,24 +293,23 @@ TSharedRef< SWidget > SCurveListRow::GenerateWidgetForColumn(const FName& Column
 //////////////////////////////////////////////////////////////////////////
 // SPoseViewer
 
-void SPoseViewer::Construct(const FArguments& InArgs)
+void SPoseViewer::Construct(const FArguments& InArgs, const TSharedRef<IPersonaToolkit>& InPersonaToolkit, const TSharedRef<IEditableSkeleton>& InEditableSkeleton, const TSharedRef<IPersonaPreviewScene>& InPreviewScene)
 {
-	PersonaPtr = InArgs._Persona;
+	PreviewScenePtr = InPreviewScene;
+	PersonaToolkitPtr = InPersonaToolkit;
+	EditableSkeletonPtr = InEditableSkeleton;
 	PoseAssetPtr = InArgs._PoseAsset;
 
 	CachedPreviewInstance = nullptr;
 
-	if (PersonaPtr.IsValid())
-	{
-		PersonaPtr.Pin()->RegisterOnPreviewMeshChanged(FPersona::FOnPreviewMeshChanged::CreateSP(this, &SPoseViewer::OnPreviewMeshChanged));
-		PersonaPtr.Pin()->RegisterOnPostUndo(FPersona::FOnPostUndo::CreateSP(this, &SPoseViewer::RefreshList));
-		PersonaPtr.Pin()->RegisterOnAnimChanged(FPersona::FOnAnimChanged::CreateSP(this, &SPoseViewer::OnAssestChanged));
-	}
+	InPreviewScene->RegisterOnPreviewMeshChanged(FOnPreviewMeshChanged::CreateSP(this, &SPoseViewer::OnPreviewMeshChanged));
+	InPreviewScene->RegisterOnAnimChanged(FOnAnimChanged::CreateSP(this, &SPoseViewer::OnAssetChanged));
 
-	if (PoseAssetPtr.IsValid())
-	{
-		PoseAssetPtr->RegisterOnAssetModifiedNotifier(UPoseAsset::FAssetModifiedNotifier::CreateSP(this, &SPoseViewer::RefreshList));
-	}
+	OnDelegatePoseListChangedDelegateHandle = PoseAssetPtr->RegisterOnPoseListChanged(UPoseAsset::FOnPoseListChanged::CreateSP(this, &SPoseViewer::RefreshList));
+
+	// Register and bind all our menu commands
+	FPoseEditorCommands::Register();
+	BindCommands();
 
 	RefreshCachePreviewInstance();
 
@@ -398,7 +394,7 @@ void SPoseViewer::Construct(const FArguments& InArgs)
 	CreateCurveList();
 }
 
-void SPoseViewer::OnAssestChanged(UAnimationAsset* NewAsset) 
+void SPoseViewer::OnAssetChanged(UAnimationAsset* NewAsset) 
 {
 	// this is odd that we're adding delegate
 	// this is because we're caching anim instance here
@@ -413,7 +409,7 @@ void SPoseViewer::RefreshCachePreviewInstance()
 	{
 		CachedPreviewInstance.Get()->RemoveDelegate_AddCustomAnimationCurve(OnAddAnimationCurveDelegate);
 	}
-	CachedPreviewInstance = Cast<UAnimSingleNodeInstance>(PersonaPtr.Pin()->GetPreviewMeshComponent()->GetAnimInstance());
+	CachedPreviewInstance = Cast<UAnimSingleNodeInstance>(PreviewScenePtr.Pin()->GetPreviewMeshComponent()->GetAnimInstance());
 
 	if (CachedPreviewInstance.IsValid())
 	{
@@ -449,7 +445,7 @@ TSharedRef<ITableRow> SPoseViewer::GeneratePoseRow(TSharedPtr<FDisplayedPoseInfo
 	check(InInfo.IsValid());
 
 	return
-		SNew(SPoseListRow, OwnerTable)
+		SNew(SPoseListRow, OwnerTable, PreviewScenePtr.Pin().ToSharedRef())
 		.Item(InInfo)
 		.PoseViewer(SharedThis(this))
 		.FilterText(GetFilterText());
@@ -585,39 +581,62 @@ void SPoseViewer::OnPastePoseNamesFromClipBoard(bool bSelectedOnly)
 
 }
 
+FReply SPoseViewer::OnKeyDown(const FGeometry& MyGeometry, const FKeyEvent& InKeyEvent)
+{
+	if (UICommandList.IsValid() && UICommandList->ProcessCommandBindings(InKeyEvent))
+	{
+		return FReply::Handled();
+	}
+	return FReply::Unhandled();
+}
+
+void SPoseViewer::BindCommands()
+{
+	// This should not be called twice on the same instance
+	check(!UICommandList.IsValid());
+	UICommandList = MakeShareable(new FUICommandList);
+	FUICommandList& CommandList = *UICommandList;
+
+	// Grab the list of menu commands to bind...
+	const FPoseEditorCommands& PoseEditorCommands = FPoseEditorCommands::Get();
+
+	// ...and bind them all
+
+	CommandList.MapAction(
+		FGenericCommands::Get().Rename,
+		FExecuteAction::CreateSP(this, &SPoseViewer::OnRenamePose),
+		FCanExecuteAction::CreateSP(this, &SPoseViewer::IsSinglePoseSelected));
+
+	CommandList.MapAction(
+		FGenericCommands::Get().Delete,
+		FExecuteAction::CreateSP(this, &SPoseViewer::OnDeletePoses),
+		FCanExecuteAction::CreateSP(this, &SPoseViewer::IsPoseSelected));
+
+	CommandList.MapAction(
+		FGenericCommands::Get().Paste,
+		FExecuteAction::CreateSP(this, &SPoseViewer::OnPastePoseNamesFromClipBoard, true),
+		FCanExecuteAction::CreateSP(this, &SPoseViewer::IsPoseSelected));
+
+	CommandList.MapAction(
+		PoseEditorCommands.PasteAllNames,
+		FExecuteAction::CreateSP(this, &SPoseViewer::OnPastePoseNamesFromClipBoard, false),
+		FCanExecuteAction());
+}
+
+
 TSharedPtr<SWidget> SPoseViewer::OnGetContextMenuContent() const
 {
 	const bool bShouldCloseWindowAfterMenuSelection = true;
-	FMenuBuilder MenuBuilder(bShouldCloseWindowAfterMenuSelection, NULL);
+	FMenuBuilder MenuBuilder(bShouldCloseWindowAfterMenuSelection, UICommandList);
 
-	FUIAction Action = FUIAction(FExecuteAction::CreateSP(this, &SPoseViewer::OnPastePoseNamesFromClipBoard, false));
-	FText MenuLabel = LOCTEXT("PasteAllPoseNamesButtonLabel", "Paste All Pose Names");
-	FText MenuToolTip = LOCTEXT("PasteAllPoseNamesButtonTooltip", "Paste All Pose Names From Clipboard");
-	MenuBuilder.AddMenuEntry(MenuLabel, MenuToolTip, FSlateIcon(), Action);
-
-	MenuBuilder.BeginSection("PoseAction", LOCTEXT("Poseaction", "Selected Item Actions"));
-	{
-		// Delete
-	 	Action = FUIAction( FExecuteAction::CreateSP( this, &SPoseViewer::OnDeletePoses ), 
-	 									FCanExecuteAction::CreateSP( this, &SPoseViewer::IsPoseSelected ) );
-	 	MenuLabel = LOCTEXT("DeletePoseButtonLabel", "Delete");
-	 	MenuToolTip = LOCTEXT("DeletePoseButtonTooltip", "Deletes the selected animation pose.");
-	 	MenuBuilder.AddMenuEntry(MenuLabel, MenuToolTip, FSlateIcon(), Action);
-
-		// Rename
-		Action = FUIAction(FExecuteAction::CreateSP(this, &SPoseViewer::OnRenamePose),
-			FCanExecuteAction::CreateSP(this, &SPoseViewer::IsSinglePoseSelected));
-		MenuLabel = LOCTEXT("RenamePoseButtonLabel", "Rename");
-		MenuToolTip = LOCTEXT("RenamePoseButtonTooltip", "Renames the selected animation pose.");
-		MenuBuilder.AddMenuEntry(MenuLabel, MenuToolTip, FSlateIcon(), Action);
+	const FPoseEditorCommands& PoseEditorCommands = FPoseEditorCommands::Get();
 
-		// Paste
-		Action = FUIAction(FExecuteAction::CreateSP(this, &SPoseViewer::OnPastePoseNamesFromClipBoard, true),
-			FCanExecuteAction::CreateSP(this, &SPoseViewer::IsPoseSelected));
-		MenuLabel = LOCTEXT("PastePoseNamesButtonLabel", "Paste Selected");
-		MenuToolTip = LOCTEXT("PastePoseNamesButtonTooltip", "Paste the selected pose names from ClipBoard.");
-		MenuBuilder.AddMenuEntry(MenuLabel, MenuToolTip, FSlateIcon(), Action);
-	}
+	MenuBuilder.AddMenuEntry(PoseEditorCommands.PasteAllNames);
+
+	MenuBuilder.BeginSection("PoseAction", LOCTEXT("SelectedItems", "Selected Item Actions"));
+	MenuBuilder.AddMenuEntry(FGenericCommands::Get().Delete, NAME_None, LOCTEXT("DeletePoseButtonLabel", "Delete"), LOCTEXT("DeletePoseButtonTooltip", "Delete the selected pose(s)"));
+	MenuBuilder.AddMenuEntry(FGenericCommands::Get().Rename, NAME_None, LOCTEXT("RenamePoseButtonLabel", "Rename"), LOCTEXT("RenamePoseButtonTooltip", "Renames the selected pose"));
+	MenuBuilder.AddMenuEntry(FGenericCommands::Get().Paste, NAME_None, LOCTEXT("PastePoseNamesButtonLabel", "Paste Selected"), LOCTEXT("PastePoseNamesButtonTooltip", "Paste the selected pose names from clipBoard"));
 	MenuBuilder.EndSection();
 
 	return MenuBuilder.MakeWidget();
@@ -663,10 +682,7 @@ void SPoseViewer::OnListDoubleClick(TSharedPtr<FDisplayedPoseInfo> InItem)
 		}
 
 		// Force update viewport
-		if (PersonaPtr.IsValid())
-		{
-			PersonaPtr.Pin()->RefreshViewport();
-		}
+		PreviewScenePtr.Pin()->InvalidateViews();
 	}
 }
 
@@ -754,7 +770,7 @@ void SPoseViewer::RemoveCurveOverride(FName& Name)
 
 SPoseViewer::~SPoseViewer()
 {
-	if (PersonaPtr.IsValid())
+	if (PreviewScenePtr.IsValid())
 	{
 		// @Todo: change this in curve editor
 		// and it won't work with this one delegate idea, so just think of a better way to do
@@ -764,14 +780,13 @@ SPoseViewer::~SPoseViewer()
 			CachedPreviewInstance.Get()->RemoveDelegate_AddCustomAnimationCurve(OnAddAnimationCurveDelegate);
 		}
 
-		PersonaPtr.Pin()->UnregisterOnPreviewMeshChanged(this);
-		PersonaPtr.Pin()->UnregisterOnPostUndo(this);
-		PersonaPtr.Pin()->UnregisterOnAnimChanged(this);
+		PreviewScenePtr.Pin()->UnregisterOnPreviewMeshChanged(this);
+		PreviewScenePtr.Pin()->UnregisterOnAnimChanged(this);
 	}
 
 	if (PoseAssetPtr.IsValid())
 	{
-		PoseAssetPtr->UnregisterOnAssetModifiedNotifier(this);
+		PoseAssetPtr->UnregisterOnPoseListChanged(OnDelegatePoseListChangedDelegateHandle);
 	}
 }
 
@@ -786,67 +801,57 @@ void SPoseViewer::ApplyCustomCurveOverride(UAnimInstance* AnimInstance) const
 	for (auto Iter = OverrideCurves.CreateConstIterator(); Iter; ++Iter)
 	{
 		// @todo we might want to save original curve flags? or just change curve to apply flags only
-		AnimInstance->AddCurveValue(Iter.Key(), Iter.Value(), ACF_DriveAttribute);
+		AnimInstance->AddCurveValue(Iter.Key(), Iter.Value());
 	}
 }
 
 bool SPoseViewer::ModifyName(FName OldName, FName NewName, bool bSilence)
 {
-	if (PersonaPtr.IsValid())
+	FScopedTransaction Transaction(LOCTEXT("RenamePoses", "Rename Pose"));
+	PoseAssetPtr.Get()->Modify();
+
+	// get smart name
+	const USkeleton& Skeleton = EditableSkeletonPtr.Pin()->GetSkeleton();
+	const SmartName::UID_Type ExistingUID = Skeleton.GetUIDByName(USkeleton::AnimCurveMappingName, NewName);
+	// verify if this name exists in smart naming
+	if (ExistingUID != SmartName::MaxUID)
 	{
-		USkeleton* Skeleton = PersonaPtr.Pin()->GetSkeleton();
-		if (Skeleton)
+		// warn users
+		// if so, verify if this name is still okay
+		if (!bSilence)
 		{
-			FScopedTransaction Transaction(LOCTEXT("RenamePoses", "Rename Pose"));
-
-			PoseAssetPtr.Get()->Modify();
-			Skeleton->Modify();
-
-			// get smart nam
-			const FSmartNameMapping::UID ExistingUID = Skeleton->GetUIDByName(USkeleton::AnimCurveMappingName, NewName);
-			// verify if this name exists in smart naming
-			if (ExistingUID != FSmartNameMapping::MaxUID)
-			{
-				// warn users
-				// if so, verify if this name is still okay
-				if (!bSilence)
-				{
-					const EAppReturnType::Type Response = FMessageDialog::Open(EAppMsgType::YesNo, LOCTEXT("UseSameNameConfirm", "The name already exists. Would you like to reuse the name? This can cause conflict of curve data."));
-
-					if (Response == EAppReturnType::No)
-					{
-						return false;
-					}
-				}
-
-				// I think this might have to be delegate of the top window
-				if (PoseAssetPtr.Get()->ModifyPoseName(OldName, NewName, &ExistingUID) == false)
-				{
-					return false;
-				}
-			}
-			else
-			{
-				// I think this might have to be delegate of the top window
-				if (PoseAssetPtr.Get()->ModifyPoseName(OldName, NewName, nullptr) == false)
-				{
-					return false;
-				}
-			}
+			const EAppReturnType::Type Response = FMessageDialog::Open(EAppMsgType::YesNo, LOCTEXT("UseSameNameConfirm", "The name already exists. Would you like to reuse the name? This can cause conflict of curve data."));
 
-			// now refresh pose data
-			float* Value = OverrideCurves.Find(OldName);
-			if (Value)
+			if (Response == EAppReturnType::No)
 			{
-				AddCurveOverride(NewName, *Value);
-				RemoveCurveOverride(OldName);
+				return false;
 			}
+		}
 
-			return true;
+		// I think this might have to be delegate of the top window
+		if (PoseAssetPtr.Get()->ModifyPoseName(OldName, NewName, &ExistingUID) == false)
+		{
+			return false;
+		}
+	}
+	else
+	{
+		// I think this might have to be delegate of the top window
+		if (PoseAssetPtr.Get()->ModifyPoseName(OldName, NewName, nullptr) == false)
+		{
+			return false;
 		}
 	}
 
-	return false;
+	// now refresh pose data
+	float* Value = OverrideCurves.Find(OldName);
+	if (Value)
+	{
+		AddCurveOverride(NewName, *Value);
+		RemoveCurveOverride(OldName);
+	}
+
+	return true;
 }
 
 bool SPoseViewer::IsBasePose(FName PoseName) const
diff --git a/Engine/Source/Editor/Persona/Private/SPoseEditor.h b/Engine/Source/Editor/Persona/Private/SPoseEditor.h
--- a/Engine/Source/Editor/Persona/Private/SPoseEditor.h
+++ b/Engine/Source/Editor/Persona/Private/SPoseEditor.h
@@ -66,7 +66,7 @@ public:
 
 	SLATE_END_ARGS()
 
-	void Construct(const FArguments& InArgs, const TSharedRef<STableViewBase>& OwnerTableView);
+	void Construct(const FArguments& InArgs, const TSharedRef<STableViewBase>& OwnerTableView, const TSharedRef<IPersonaPreviewScene>& InPreviewScene);
 
 	/** Overridden from SMultiColumnTableRow.  Generates a widget for this column of the tree row. */
 	virtual TSharedRef<SWidget> GenerateWidgetForColumn(const FName& ColumnName) override;
@@ -115,6 +115,9 @@ private:
 
 	/** Text the user typed into the search box - used for text highlighting */
 	FText FilterText;
+
+	/** The preview scene we are viewing */
+	TWeakPtr<class IPersonaPreviewScene> PreviewScenePtr;
 };
 
 //////////////////////////////////////////////////////////////////////////
@@ -189,11 +192,8 @@ class SPoseViewer : public SCompoundWidget
 {
 public:
 	SLATE_BEGIN_ARGS(SPoseViewer)
-		: _Persona()
 	{}
 
-	/* The Persona that owns this table */
-	SLATE_ARGUMENT(TWeakPtr< FPersona >, Persona)
 	SLATE_ARGUMENT(TWeakObjectPtr< UPoseAsset >, PoseAsset)
 
 	SLATE_END_ARGS()
@@ -204,7 +204,7 @@ public:
 	* @param InArgs - Arguments passed from Slate
 	*
 	*/
-	void Construct(const FArguments& InArgs);
+	void Construct(const FArguments& InArgs, const TSharedRef<class IPersonaToolkit>& InPersonaToolkit, const TSharedRef<IEditableSkeleton>& InEditableSkeleton, const TSharedRef<class IPersonaPreviewScene>& InPreviewScene);
 
 	/**
 	* Destructor - resets the animation curve
@@ -212,6 +212,9 @@ public:
 	*/
 	virtual ~SPoseViewer();
 
+	/** SWidget interface */
+	virtual FReply OnKeyDown(const FGeometry& MyGeometry, const FKeyEvent& InKeyEvent);
+
 	/**
 	* Is registered with Persona to handle when its preview mesh is changed.
 	*
@@ -282,10 +285,6 @@ public:
 	*/
 	FText& GetFilterText() { return FilterText; }
 
-	/**
-	* Get Title text
-	*/
-	FText GetTitleText() const;
 	/**
 	* Refreshes the morph target list after an undo
 	*/
@@ -296,6 +295,8 @@ public:
 
 private:
 
+	void BindCommands();
+
 	/** Handler for context menus */
 	TSharedPtr<SWidget> OnGetContextMenuContent() const;
 	/** Handler for curve list context menu*/
@@ -315,10 +316,16 @@ private:
 
 	void ApplyCustomCurveOverride(UAnimInstance* AnimInstance) const;
 	void RefreshCachePreviewInstance();
-	void OnAssestChanged(UAnimationAsset* NewAsset);
+	void OnAssetChanged(UAnimationAsset* NewAsset);
+
+	/** Pointer to the preview scene we are viewing */
+	TWeakPtr<class IPersonaPreviewScene> PreviewScenePtr;
 
-	/** Pointer back to the Persona that owns us */
-	TWeakPtr<FPersona> PersonaPtr;
+	/** Pointer to the persona toolkit we are embedded in */
+	TWeakPtr<class IPersonaToolkit> PersonaToolkitPtr;
+
+	/** Pointer to the editable skeleton we will need to modify */
+	TWeakPtr<class IEditableSkeleton> EditableSkeletonPtr;
 
 	/** Pointer to the pose asset */
 	TWeakObjectPtr<UPoseAsset> PoseAssetPtr;
@@ -344,10 +351,14 @@ private:
 	/** Current text typed into NameFilterBox */
 	FText FilterText;
 
+	/** Commands that are bound to delegates*/
+	TSharedPtr<FUICommandList> UICommandList;
+
 	TMap<FName, float> OverrideCurves;
 
 	/** Add curve delegate */
 	FOnAddCustomAnimationCurves OnAddAnimationCurveDelegate;
+	FDelegateHandle OnDelegatePoseListChangedDelegateHandle;
 
 	friend class SPoseListRow;
 	friend class SCurveListRow;
@@ -360,17 +371,14 @@ class SPoseEditor : public SAnimEditorBase
 {
 public:
 	SLATE_BEGIN_ARGS( SPoseEditor )
-		: _Persona()
-		, _PoseAsset(NULL)
+		: _PoseAsset(NULL)
 		{}
 
-		SLATE_ARGUMENT( TSharedPtr<FPersona>, Persona )
 		SLATE_ARGUMENT( UPoseAsset*, PoseAsset )
 	SLATE_END_ARGS()
 
 public:
-	void Construct(const FArguments& InArgs);
-	virtual ~SPoseEditor();
+	void Construct(const FArguments& InArgs, const TSharedRef<class IPersonaToolkit>& InPersonaToolkit, const TSharedRef<IEditableSkeleton>& InEditableSkeleton, const TSharedRef<class IPersonaPreviewScene>& InPreviewScene);
 
 	virtual UAnimationAsset* GetEditorObject() const override { return PoseAssetObj; }
 
diff --git a/Engine/Source/Editor/Persona/Private/SRetargetManager.cpp b/Engine/Source/Editor/Persona/Private/SRetargetManager.cpp
--- a/Engine/Source/Editor/Persona/Private/SRetargetManager.cpp
+++ b/Engine/Source/Editor/Persona/Private/SRetargetManager.cpp
@@ -14,6 +14,7 @@
 #include "SRetargetSourceWindow.h"
 #include "SRigWindow.h"
 #include "Editor/AnimGraph/Classes/AnimPreviewInstance.h"
+#include "IEditableSkeleton.h"
 
 #define LOCTEXT_NAMESPACE "SRetargetManager"
 
@@ -21,16 +22,11 @@
 //////////////////////////////////////////////////////////////////////////
 // SRetargetManager
 
-void SRetargetManager::Construct(const FArguments& InArgs)
+void SRetargetManager::Construct(const FArguments& InArgs, const TSharedRef<IEditableSkeleton>& InEditableSkeleton, const TSharedRef<IPersonaPreviewScene>& InPreviewScene, FSimpleMulticastDelegate& InOnPostUndo)
 {
-	PersonaPtr = InArgs._Persona;
-	Skeleton = NULL;
-
-	if ( PersonaPtr.IsValid() )
-	{
-		Skeleton = PersonaPtr.Pin()->GetSkeleton();
-		PersonaPtr.Pin()->RegisterOnPostUndo(FPersona::FOnPostUndo::CreateSP( this, &SRetargetManager::PostUndo ) );
-	}
+	EditableSkeletonPtr = InEditableSkeleton;
+	PreviewScenePtr = InPreviewScene;
+	InOnPostUndo.Add(FSimpleDelegate::CreateSP(this, &SRetargetManager::PostUndo));
 
 	const FString DocLink = TEXT("Shared/Editors/Persona");
 	ChildSlot
@@ -41,11 +37,11 @@ void SRetargetManager::Construct(const FArguments& InArgs)
 		.Padding(5, 5)
 		.AutoHeight()
 		[
-			// explainint this is retarget source window
+			// explain this is retarget source window
 			// and what it is
 			SNew(STextBlock)
 			.TextStyle( FEditorStyle::Get(), "Persona.RetargetManager.ImportantText" )
-			.Text(LOCTEXT("RetargetSource_Title", "Manager Retarget Source"))
+			.Text(LOCTEXT("RetargetSource_Title", "Manage Retarget Source"))
 		]
 
 		+ SVerticalBox::Slot()
@@ -61,9 +57,9 @@ void SRetargetManager::Construct(const FArguments& InArgs)
 																			DocLink,
 																			TEXT("RetargetSource")))
 			.Font(FEditorStyle::GetFontStyle(TEXT("Persona.RetargetManager.FilterFont")))
-			.Text(LOCTEXT("RetargetSource_Description", "You can add/rename/delete Retarget Sources. When you have different proportional meshes per skeleton, you can use this settings to indicate if this animation is from different source. \
-														For example, if your default skeleton is from small guy, and if you have animation for big guy, you can create Retarget Source from the big guy and set it for the animation. \
-														Retargeting system will use this information when extracting animation. "))
+			.Text(LOCTEXT("RetargetSource_Description", "You can add/rename/delete Retarget Sources. When you have different proportional meshes per skeleton, you can use this setting to indicate if this animation is from a different source." \
+														"For example, if your default skeleton is from a small guy, and if you have an animation for a big guy, you can create a Retarget Source from the big guy and set it for the animation." \
+														"The Retargeting system will use this information when extracting animation. "))
 		]
 
 		+ SVerticalBox::Slot()
@@ -71,8 +67,7 @@ void SRetargetManager::Construct(const FArguments& InArgs)
 		.FillHeight(0.5)
 		[
 			// construct retarget source window
-			SNew(SRetargetSourceWindow)
-			.Persona(PersonaPtr)
+			SNew(SRetargetSourceWindow, InEditableSkeleton, InOnPostUndo)
 		]
 
 		+SVerticalBox::Slot()
@@ -107,7 +102,7 @@ void SRetargetManager::Construct(const FArguments& InArgs)
 																			DocLink,
 																			TEXT("RigSetup")))
 			.Font(FEditorStyle::GetFontStyle(TEXT("Persona.RetargetManager.FilterFont")))
-			.Text(LOCTEXT("RigTemplate_Description", "You can set up Rig for this skeleton, then when you retarget animation to different skeleton with the same Rig, it will use the information to convert data. "))
+			.Text(LOCTEXT("RigTemplate_Description", "You can set up a Rig for this skeleton, then when you retarget the animation to a different skeleton with the same Rig, it will use the information to convert data. "))
 		]
 
 		+ SVerticalBox::Slot()
@@ -115,8 +110,7 @@ void SRetargetManager::Construct(const FArguments& InArgs)
 		.Padding(2, 5)
 		[
 			// construct rig manager window
-			SNew(SRigWindow)
-			.Persona(PersonaPtr)
+			SNew(SRigWindow, InEditableSkeleton, InOnPostUndo)
 		]
 
 		+SVerticalBox::Slot()
@@ -151,8 +145,8 @@ void SRetargetManager::Construct(const FArguments& InArgs)
 																			DocLink,
 																			TEXT("SetupBasePose")))
 			.Font(FEditorStyle::GetFontStyle(TEXT("Persona.RetargetManager.FilterFont")))
-			.Text(LOCTEXT("BasePose_Description", "This information is used when retargeting assets to different skeleton. You need to make sure the ref pose of both mesh is same when retargeting, so you can see the pose and \
-											edit using bone transform widget, and click Save button below. "))
+			.Text(LOCTEXT("BasePose_Description", "This information is used when retargeting assets to a different skeleton. You need to make sure the ref pose of both meshes is the same when retargeting, so you can see the pose and" \
+											" edit using the bone transform widget, and click the Save button below. "))
 		]
 
 		+ SVerticalBox::Slot()
@@ -203,17 +197,9 @@ void SRetargetManager::Construct(const FArguments& InArgs)
 	];
 }
 
-SRetargetManager::~SRetargetManager()
-{
-	if (PersonaPtr.IsValid())
-	{
-		PersonaPtr.Pin()->UnregisterOnPostUndo(this);
-	}
-}
-
 FReply SRetargetManager::OnViewRetargetBasePose()
 {
-	UDebugSkelMeshComponent * PreviewMeshComp = PersonaPtr.Pin()->GetPreviewMeshComponent();
+	UDebugSkelMeshComponent * PreviewMeshComp = PreviewScenePtr.Pin()->GetPreviewMeshComponent();
 	if (PreviewMeshComp && PreviewMeshComp->PreviewInstance)
 	{
 		const FScopedTransaction Transaction(LOCTEXT("ViewRetargetBasePose_Action", "Edit Retarget Base Pose"));
@@ -233,12 +219,12 @@ FReply SRetargetManager::OnViewRetargetBasePose()
 
 FReply SRetargetManager::OnSaveRetargetBasePose()
 {
-	UDebugSkelMeshComponent * PreviewMeshComp = PersonaPtr.Pin()->GetPreviewMeshComponent();
+	UDebugSkelMeshComponent * PreviewMeshComp = PreviewScenePtr.Pin()->GetPreviewMeshComponent();
 	if (PreviewMeshComp && PreviewMeshComp->SkeletalMesh)
 	{
 		USkeletalMesh * PreviewMesh = PreviewMeshComp->SkeletalMesh;
 
-		check(PreviewMesh && Skeleton == PreviewMesh->Skeleton);
+		check(PreviewMesh && &EditableSkeletonPtr.Pin()->GetSkeleton() == PreviewMesh->Skeleton);
 
 		if (PreviewMesh)
 		{
@@ -288,12 +274,12 @@ FReply SRetargetManager::OnResetRetargetBasePose()
 	EAppReturnType::Type Response = FMessageDialog::Open(EAppMsgType::OkCancel, Message);
 	if(Response == EAppReturnType::Ok)
 	{
-		UDebugSkelMeshComponent * PreviewMeshComp = PersonaPtr.Pin()->GetPreviewMeshComponent();
+		UDebugSkelMeshComponent * PreviewMeshComp = PreviewScenePtr.Pin()->GetPreviewMeshComponent();
 		if(PreviewMeshComp && PreviewMeshComp->SkeletalMesh)
 		{
 			USkeletalMesh * PreviewMesh = PreviewMeshComp->SkeletalMesh;
 
-			check(PreviewMesh && Skeleton == PreviewMesh->Skeleton);
+			check(PreviewMesh && &EditableSkeletonPtr.Pin()->GetSkeleton() == PreviewMesh->Skeleton);
 
 			if(PreviewMesh)
 			{
@@ -316,7 +302,7 @@ void SRetargetManager::PostUndo()
 
 FText SRetargetManager::GetToggleRetargetBasePose() const
 {
-	UDebugSkelMeshComponent * PreviewMeshComp = PersonaPtr.Pin()->GetPreviewMeshComponent();
+	UDebugSkelMeshComponent * PreviewMeshComp = PreviewScenePtr.Pin()->GetPreviewMeshComponent();
 	if(PreviewMeshComp && PreviewMeshComp->PreviewInstance)
 	{
 		if (PreviewMeshComp->PreviewInstance->GetForceRetargetBasePose())
diff --git a/Engine/Source/Editor/Persona/Private/SRetargetManager.h b/Engine/Source/Editor/Persona/Private/SRetargetManager.h
--- a/Engine/Source/Editor/Persona/Private/SRetargetManager.h
+++ b/Engine/Source/Editor/Persona/Private/SRetargetManager.h
@@ -2,8 +2,6 @@
 
 #pragma once
 
-#include "Persona.h"
-
 //////////////////////////////////////////////////////////////////////////
 // SRetargetManager
 
@@ -11,11 +9,7 @@ class SRetargetManager : public SCompoundWidget
 {
 public:
 	SLATE_BEGIN_ARGS( SRetargetManager )
-		: _Persona()
 	{}
-		
-		/* The Persona that owns this table */
-		SLATE_ARGUMENT( TWeakPtr< FPersona >, Persona )
 
 	SLATE_END_ARGS()
 
@@ -25,23 +19,18 @@ public:
 	* @param InArgs - Arguments passed from Slate
 	*
 	*/
-	void Construct( const FArguments& InArgs );
-
-	/**
-	* Destructor - resets the morph targets
-	*
-	*/
-	virtual ~SRetargetManager();
+	void Construct(const FArguments& InArgs, const TSharedRef<IEditableSkeleton>& InEditableSkeleton, const TSharedRef<IPersonaPreviewScene>& InPreviewScene, FSimpleMulticastDelegate& InOnPostUndo);
 
 private:
 	FReply OnViewRetargetBasePose();
 	FReply OnResetRetargetBasePose();
 	FReply OnSaveRetargetBasePose();
 
-	/** Pointer back to the Persona that owns us */
-	TWeakPtr<FPersona> PersonaPtr;
-	/** The skeletal mesh that we grab the morph targets from */
-	USkeleton* Skeleton;
+	/** The editable skeleton */
+	TWeakPtr<class IEditableSkeleton> EditableSkeletonPtr;
+
+	/** The preview scene  */
+	TWeakPtr<class IPersonaPreviewScene> PreviewScenePtr;
 
 	/** Delegate for undo/redo transaction **/
 	void PostUndo();
diff --git a/Engine/Source/Editor/Persona/Private/SRetargetSourceWindow.cpp b/Engine/Source/Editor/Persona/Private/SRetargetSourceWindow.cpp
--- a/Engine/Source/Editor/Persona/Private/SRetargetSourceWindow.cpp
+++ b/Engine/Source/Editor/Persona/Private/SRetargetSourceWindow.cpp
@@ -14,6 +14,7 @@
 #include "SNotificationList.h"
 #include "NotificationManager.h"
 #include "Animation/AnimSequence.h"
+#include "IEditableSkeleton.h"
 
 #define LOCTEXT_NAMESPACE "SRetargetSourceWindow"
 
@@ -44,9 +45,6 @@ public:
 	/* Widget used to display the list of retarget sources*/
 	SLATE_ARGUMENT( TSharedPtr<SRetargetSourceListType>, RetargetSourceListView )
 
-	/* Persona used to update the viewport when a weight slider is dragged */
-	SLATE_ARGUMENT( TWeakPtr<FPersona>, Persona )
-
 	/** Delegate for when an asset name has been entered for an item that is in a rename state */
 	SLATE_EVENT( FOnRenameCommit, OnRenameCommit )
 
@@ -74,9 +72,6 @@ private:
 	/** The name and weight of the retarget source*/
 	FDisplayedRetargetSourceInfoPtr	Item;
 
-	/** Pointer back to the Persona that owns us */
-	TWeakPtr<FPersona> PersonaPtr;
-
 	/** Delegate for when an asset name has been entered for an item that is in a rename state */
 	FOnRenameCommit OnRenameCommit;
 
@@ -98,8 +93,6 @@ void SRetargetSourceListRow::Construct( const FArguments& InArgs, const TSharedR
 	OnRenameCommit = InArgs._OnRenameCommit;
 	OnVerifyRenameCommit = InArgs._OnVerifyRenameCommit;
 
-	PersonaPtr = InArgs._Persona;
-
 	check( Item.IsValid() );
 
 	SMultiColumnTableRow< FDisplayedRetargetSourceInfoPtr >::Construct( FSuperRowType::FArguments(), InOwnerTableView );
@@ -163,18 +156,13 @@ bool SRetargetSourceListRow::HandleVerifyNameChanged( const FText& NewText, FTex
 //////////////////////////////////////////////////////////////////////////
 // SRetargetSourceWindow
 
-void SRetargetSourceWindow::Construct(const FArguments& InArgs)
+void SRetargetSourceWindow::Construct(const FArguments& InArgs, const TSharedRef<IEditableSkeleton>& InEditableSkeleton, FSimpleMulticastDelegate& InOnPostUndo)
 {
-	PersonaPtr = InArgs._Persona;
-	Skeleton = NULL;
+	EditableSkeletonPtr = InEditableSkeleton;
 
-	if ( PersonaPtr.IsValid() )
-	{
-		Skeleton = PersonaPtr.Pin()->GetSkeleton();
-		PersonaPtr.Pin()->RegisterOnPostUndo(FPersona::FOnPostUndo::CreateSP( this, &SRetargetSourceWindow::PostUndo ) );
-	}
+	InOnPostUndo.Add(FSimpleDelegate::CreateSP( this, &SRetargetSourceWindow::PostUndo ) );
 
-	FText SkeletonName = Skeleton ? FText::FromString(Skeleton->GetName()) : LOCTEXT( "RetargetSourceMeshNameLabel", "No Skeleton Present" );
+	FText SkeletonName = FText::FromString(InEditableSkeleton->GetSkeleton().GetName());
 
 	ChildSlot
 	[
@@ -268,7 +256,6 @@ TSharedRef<ITableRow> SRetargetSourceWindow::GenerateRetargetSourceRow(TSharedPt
 
 	return
 		SNew( SRetargetSourceListRow, OwnerTable )
-		.Persona( PersonaPtr )
 		.Item( InInfo )
 		.RetargetSourceWindow( this )
 		.RetargetSourceListView( RetargetSourceListView )
@@ -277,77 +264,16 @@ TSharedRef<ITableRow> SRetargetSourceWindow::GenerateRetargetSourceRow(TSharedPt
 }
 
 
-void SRetargetSourceWindow::OnRenameCommit( const FName& OldName,  const FString& NewName )
+void SRetargetSourceWindow::OnRenameCommit( const FName& InOldName, const FString& InNewName )
 {
-	FString NewNameString = NewName;
-	if (OldName != FName(*NewNameString.Trim().TrimTrailing()))
+	FString NewNameString = InNewName;
+	if (InOldName != FName(*NewNameString.Trim().TrimTrailing()))
 	{
-		const FScopedTransaction Transaction( LOCTEXT("RetargetSourceWindow_Rename", "Rename Retarget Source") );
-
-		FReferencePose * Pose = Skeleton->AnimRetargetSources.Find(OldName);
-		if (Pose)
-		{
-			FReferencePose NewPose = *Pose;
-			FName NewPoseName = FName(*NewName);
-			NewPose.PoseName = NewPoseName;
-
-			Skeleton->Modify();
-
-			Skeleton->AnimRetargetSources.Remove(OldName);
-			Skeleton->AnimRetargetSources.Add(NewPoseName, NewPose);		
-			// push renamed info to the stack
-			RenamedPoseNameStack.Push(FPoseNameForUndo(OldName, NewPoseName));
-
-			// need to verify if these pose is used by anybody else
-			FAssetRegistryModule& AssetRegistryModule = FModuleManager::LoadModuleChecked<FAssetRegistryModule>(TEXT("AssetRegistry"));
-
-			TArray<FAssetData> AssetList;
-			TMultiMap<FName, FString> TagsAndValues;
-			TagsAndValues.Add(GET_MEMBER_NAME_CHECKED(UAnimSequence, RetargetSource), OldName.ToString());
-			AssetRegistryModule.Get().GetAssetsByTagValues(TagsAndValues, AssetList);
+		FName NewName = *InNewName;
+		EditableSkeletonPtr.Pin()->RenameRetargetSource(InOldName, NewName);
 
-			// ask users if they'd like to continue and/or fix up
-			if (AssetList.Num() > 0 )
-			{
-				FString ListOfAssets;
-				// if users is sure to delete, delete
-				for (auto Iter=AssetList.CreateConstIterator(); Iter; ++Iter)
-				{
-					ListOfAssets += Iter->AssetName.ToString();
-					ListOfAssets += TEXT("\n");
-				}
-
-
-				// if so, ask if they'd like us to fix the animations as well
-				FString Message = NSLOCTEXT("RetargetSourceEditor", "RetargetSourceRename_FixUpAnimation_Message", "Would you like to fix up the following animation(s)? You'll have to save all the assets in the list.").ToString();
-				Message += TEXT("\n");
-				Message += ListOfAssets;
-
-				EAppReturnType::Type Result = FMessageDialog::Open( EAppMsgType::YesNo, FText::FromString(Message) );
-
-				if (Result == EAppReturnType::Yes)
-				{
-					// now fix up all assets
-					TArray<UObject*> ObjectsToUpdate;
-					for (auto Iter=AssetList.CreateConstIterator(); Iter; ++Iter)
-					{
-						UAnimSequence * AnimSequence = Cast<UAnimSequence>(Iter->GetAsset());
-						if (AnimSequence)
-						{
-							ObjectsToUpdate.Add(AnimSequence);
-
-							AnimSequence->Modify();
-							// clear name
-							AnimSequence->RetargetSource = NewPoseName;
-						}
-					}
-				}
-			}
-
-			Skeleton->CallbackRetargetSourceChanged();
-			FAssetNotifications::SkeletonNeedsToBeSaved(Skeleton);
-			CreateRetargetSourceList( NameFilterBox->GetText().ToString() );
-		}
+		FAssetNotifications::SkeletonNeedsToBeSaved(&EditableSkeletonPtr.Pin()->GetSkeleton());
+		CreateRetargetSourceList( NameFilterBox->GetText().ToString() );
 	}
 }
 
@@ -362,15 +288,15 @@ bool SRetargetSourceWindow::OnVerifyRenameCommit( const FName& OldName, const FS
 		return false;
 	}
 
-	FReferencePose * Pose = Skeleton->AnimRetargetSources.Find(OldName);
+	const USkeleton& Skeleton = EditableSkeletonPtr.Pin()->GetSkeleton();
+	const FReferencePose* Pose = Skeleton.AnimRetargetSources.Find(OldName);
 	if (!Pose)
 	{
 		OutErrorMessage = FText::Format (LOCTEXT("RetargetSourceWindowNameNotFound", "{0} Not found"), FText::FromString(OldName.ToString()) );
 		return false;
 	}
 
-	Pose = Skeleton->AnimRetargetSources.Find(FName(*NewName));
-
+	Pose = Skeleton.AnimRetargetSources.Find(FName(*NewName));
 	if (Pose)
 	{
 		OutErrorMessage = FText::Format (LOCTEXT("RetargetSourceWindowNameDuplicated", "{0} already exists"), FText::FromString(NewName) );
@@ -382,50 +308,43 @@ bool SRetargetSourceWindow::OnVerifyRenameCommit( const FName& OldName, const FS
 
 TSharedPtr<SWidget> SRetargetSourceWindow::OnGetContextMenuContent() const
 {
-	if ( Skeleton )
-	{
-		const bool bShouldCloseWindowAfterMenuSelection = true;
-		FMenuBuilder MenuBuilder( bShouldCloseWindowAfterMenuSelection, NULL);
-
-		MenuBuilder.BeginSection("RetargetSourceAction", LOCTEXT( "New", "New" ) );
-		{
-			FUIAction Action = FUIAction( FExecuteAction::CreateSP( this, &SRetargetSourceWindow::OnAddRetargetSource ) );
-			const FText Label = LOCTEXT("AddRetargetSourceActionLabel", "Add...");
-			const FText ToolTipText = LOCTEXT("AddRetargetSourceActionTooltip", "Add new retarget source.");
-			MenuBuilder.AddMenuEntry( Label, ToolTipText, FSlateIcon(), Action);
-		}
-		MenuBuilder.EndSection();
-
-		MenuBuilder.BeginSection("RetargetSourceAction", LOCTEXT( "Selected", "Selected Item Actions" ) );
-		{
-			FUIAction Action = FUIAction( FExecuteAction::CreateSP( this, &SRetargetSourceWindow::OnRenameRetargetSource ), 
-				FCanExecuteAction::CreateSP( this, &SRetargetSourceWindow::CanPerformRename ) );
-			const FText Label = LOCTEXT("RenameRetargetSourceActionLabel", "Rename");
-			const FText ToolTipText = LOCTEXT("RenameRetargetSourceActionTooltip", "Rename the selected retarget source.");
-			MenuBuilder.AddMenuEntry( Label, ToolTipText, FSlateIcon(), Action);
-		}
-		{
-			FUIAction Action = FUIAction( FExecuteAction::CreateSP( this, &SRetargetSourceWindow::OnDeleteRetargetSource ), 
-				FCanExecuteAction::CreateSP( this, &SRetargetSourceWindow::CanPerformDelete ) );
-			const FText Label = LOCTEXT("DeleteRetargetSourceActionLabel", "Delete");
-			const FText ToolTipText = LOCTEXT("DeleteRetargetSourceActionTooltip", "Deletes the selected retarget sources.");
-			MenuBuilder.AddMenuEntry( Label, ToolTipText, FSlateIcon(), Action);
-		}
-		{
-			FUIAction Action = FUIAction( FExecuteAction::CreateSP( this, &SRetargetSourceWindow::OnRefreshRetargetSource ), 
-				FCanExecuteAction::CreateSP( this, &SRetargetSourceWindow::CanPerformRefresh ) );
-			const FText Label = LOCTEXT("RefreshRetargetSourceActionLabel", "Refresh");
-			const FText ToolTipText = LOCTEXT("RefreshRetargetSourceActionTooltip", "Refreshs the selected retarget sources.");
-			MenuBuilder.AddMenuEntry( Label, ToolTipText, FSlateIcon(), Action);
-		}
-		MenuBuilder.EndSection();
+	const bool bShouldCloseWindowAfterMenuSelection = true;
+	FMenuBuilder MenuBuilder( bShouldCloseWindowAfterMenuSelection, NULL);
 
-		return MenuBuilder.MakeWidget();
+	MenuBuilder.BeginSection("RetargetSourceAction", LOCTEXT( "New", "New" ) );
+	{
+		FUIAction Action = FUIAction( FExecuteAction::CreateSP( this, &SRetargetSourceWindow::OnAddRetargetSource ) );
+		const FText Label = LOCTEXT("AddRetargetSourceActionLabel", "Add...");
+		const FText ToolTipText = LOCTEXT("AddRetargetSourceActionTooltip", "Add new retarget source.");
+		MenuBuilder.AddMenuEntry( Label, ToolTipText, FSlateIcon(), Action);
 	}
+	MenuBuilder.EndSection();
 
-	return TSharedPtr<SWidget>();
-
+	MenuBuilder.BeginSection("RetargetSourceAction", LOCTEXT( "Selected", "Selected Item Actions" ) );
+	{
+		FUIAction Action = FUIAction( FExecuteAction::CreateSP( this, &SRetargetSourceWindow::OnRenameRetargetSource ), 
+			FCanExecuteAction::CreateSP( this, &SRetargetSourceWindow::CanPerformRename ) );
+		const FText Label = LOCTEXT("RenameRetargetSourceActionLabel", "Rename");
+		const FText ToolTipText = LOCTEXT("RenameRetargetSourceActionTooltip", "Rename the selected retarget source.");
+		MenuBuilder.AddMenuEntry( Label, ToolTipText, FSlateIcon(), Action);
+	}
+	{
+		FUIAction Action = FUIAction( FExecuteAction::CreateSP( this, &SRetargetSourceWindow::OnDeleteRetargetSource ), 
+			FCanExecuteAction::CreateSP( this, &SRetargetSourceWindow::CanPerformDelete ) );
+		const FText Label = LOCTEXT("DeleteRetargetSourceActionLabel", "Delete");
+		const FText ToolTipText = LOCTEXT("DeleteRetargetSourceActionTooltip", "Deletes the selected retarget sources.");
+		MenuBuilder.AddMenuEntry( Label, ToolTipText, FSlateIcon(), Action);
+	}
+	{
+		FUIAction Action = FUIAction( FExecuteAction::CreateSP( this, &SRetargetSourceWindow::OnRefreshRetargetSource ), 
+			FCanExecuteAction::CreateSP( this, &SRetargetSourceWindow::CanPerformRefresh ) );
+		const FText Label = LOCTEXT("RefreshRetargetSourceActionLabel", "Refresh");
+		const FText ToolTipText = LOCTEXT("RefreshRetargetSourceActionTooltip", "Refreshs the selected retarget sources.");
+		MenuBuilder.AddMenuEntry( Label, ToolTipText, FSlateIcon(), Action);
+	}
+	MenuBuilder.EndSection();
 
+	return MenuBuilder.MakeWidget();
 }
 
 /** @TODO: FIXME: Item to rename. Only valid for adding and this is so hacky, I know**/
@@ -434,37 +353,34 @@ TSharedPtr<SWidget> SRetargetSourceWindow::OnGetContextMenuContent() const
 void SRetargetSourceWindow::CreateRetargetSourceList( const FString& SearchText, const FName  NewName )
 {
 	RetargetSourceList.Empty();
+	bool bDoFiltering = !SearchText.IsEmpty();
 
-	if ( Skeleton )
+	const USkeleton& Skeleton = EditableSkeletonPtr.Pin()->GetSkeleton();
+	for (auto Iter=Skeleton.AnimRetargetSources.CreateConstIterator(); Iter; ++Iter)
 	{
-		bool bDoFiltering = !SearchText.IsEmpty();
+		const FName& Name = Iter.Key();
+		const FReferencePose& RefPose = Iter.Value();
 
-		for (auto Iter=Skeleton->AnimRetargetSources.CreateConstIterator(); Iter; ++Iter)
+		if ( bDoFiltering )
 		{
-			const FName& Name = Iter.Key();
-			const FReferencePose& RefPose = Iter.Value();
-
-			if ( bDoFiltering )
+			if (!Name.ToString().Contains( SearchText ))
 			{
-				if (!Name.ToString().Contains( SearchText ))
-				{
-					continue; // Skip items that don't match our filter
-				}
-				if (RefPose.ReferenceMesh && RefPose.ReferenceMesh->GetPathName().Contains( SearchText ))
-				{
-					continue;
-				}
+				continue; // Skip items that don't match our filter
 			}
-
-			TSharedRef<FDisplayedRetargetSourceInfo> Info = FDisplayedRetargetSourceInfo::Make( Name, RefPose.ReferenceMesh );
-
-			if (Name == NewName)
+			if (RefPose.ReferenceMesh && RefPose.ReferenceMesh->GetPathName().Contains( SearchText ))
 			{
-				ItemToRename = Info;
+				continue;
 			}
+		}
+
+		TSharedRef<FDisplayedRetargetSourceInfo> Info = FDisplayedRetargetSourceInfo::Make( Name, RefPose.ReferenceMesh );
 
-			RetargetSourceList.Add( Info );
+		if (Name == NewName)
+		{
+			ItemToRename = Info;
 		}
+
+		RetargetSourceList.Add( Info );
 	}
 
 	RetargetSourceListView->RequestListRefresh();
@@ -486,49 +402,13 @@ int32 SRetargetSourceWindow::OnPaint( const FPaintArgs& Args, const FGeometry& A
 
 void SRetargetSourceWindow::AddRetargetSource( const FName Name, USkeletalMesh * ReferenceMesh  )
 {
-	// need to verify if the name is unique, if not create unique name
-	if ( Skeleton )
-	{
-		int32 IntSuffix = 1;
-		FReferencePose * ExistingPose = NULL;
-		FString NewSourceName = TEXT("");
-		do
-		{
-			if ( IntSuffix <= 1 )
-			{
-				NewSourceName = FString::Printf( TEXT("%s"), *Name.ToString());
-			}
-			else
-			{
-				NewSourceName = FString::Printf( TEXT("%s%d"), *Name.ToString(), IntSuffix );
-			}
-
-			ExistingPose = Skeleton->AnimRetargetSources.Find(FName(*NewSourceName));
-			IntSuffix++;
-		}
-		while (ExistingPose != NULL);
-
-		// add new one
-		// remap to skeleton refpose
-		// we have to do this whenever skeleton changes
-		FReferencePose RefPose;
-		RefPose.PoseName = FName(*NewSourceName);
-		RefPose.ReferenceMesh = ReferenceMesh;
+	EditableSkeletonPtr.Pin()->AddRetargetSource(Name, ReferenceMesh);
 
-		const FScopedTransaction Transaction( LOCTEXT("RetargetSourceWindow_Add", "Add New Retarget Source") );
-		Skeleton->Modify();
+	FAssetNotifications::SkeletonNeedsToBeSaved(&EditableSkeletonPtr.Pin()->GetSkeleton());
 
-		Skeleton->AnimRetargetSources.Add(Name, RefPose);
-		// ask skeleton to update retarget source for the given name
-		Skeleton->UpdateRetargetSource(Name);
-
-		Skeleton->CallbackRetargetSourceChanged();
-		FAssetNotifications::SkeletonNeedsToBeSaved(Skeleton);
-
-		// clear search filter
-		NameFilterBox->SetText(FText::GetEmpty());
-		CreateRetargetSourceList( NameFilterBox->GetText().ToString(), Name );
-	}
+	// clear search filter
+	NameFilterBox->SetText(FText::GetEmpty());
+	CreateRetargetSourceList( NameFilterBox->GetText().ToString(), Name );
 }
 
 void SRetargetSourceWindow::OnAddRetargetSource()
@@ -542,11 +422,10 @@ void SRetargetSourceWindow::OnAddRetargetSource()
 	AssetPickerConfig.bAllowNullSelection = false;
 	AssetPickerConfig.InitialAssetViewType = EAssetViewType::Tile;
 
-	if(Skeleton)
-	{
-		FString SkeletonString = FAssetData(Skeleton).GetExportTextName();
-		AssetPickerConfig.Filter.TagsAndValues.Add(GET_MEMBER_NAME_CHECKED(USkeletalMesh, Skeleton), SkeletonString);
-	}
+	const USkeleton& Skeleton = EditableSkeletonPtr.Pin()->GetSkeleton();
+
+	FString SkeletonString = FAssetData(&Skeleton).GetExportTextName();
+	AssetPickerConfig.Filter.TagsAndValues.Add(GET_MEMBER_NAME_CHECKED(USkeletalMesh, Skeleton), SkeletonString);
 
 	TSharedRef<SWidget> Widget = SNew(SBox)
 		.WidthOverride(384)
@@ -579,7 +458,8 @@ void SRetargetSourceWindow::OnAssetSelectedFromMeshPicker(const FAssetData& Asse
 	USkeletalMesh * SelectedMesh = CastChecked<USkeletalMesh> (AssetData.GetAsset());
 
 	// make sure you don't have any more retarget sources from the same mesh
-	for (auto Iter=Skeleton->AnimRetargetSources.CreateConstIterator(); Iter; ++Iter)
+	const USkeleton& Skeleton = EditableSkeletonPtr.Pin()->GetSkeleton();
+	for (auto Iter=Skeleton.AnimRetargetSources.CreateConstIterator(); Iter; ++Iter)
 	{
 		const FName& Name = Iter.Key();
 		const FReferencePose& RefPose = Iter.Value();
@@ -606,80 +486,16 @@ bool SRetargetSourceWindow::CanPerformDelete() const
 
 void SRetargetSourceWindow::OnDeleteRetargetSource()
 {
+	TArray<FName> SelectedNames;
 	TArray< TSharedPtr< FDisplayedRetargetSourceInfo > > SelectedRows = RetargetSourceListView->GetSelectedItems();
-	
-	const FScopedTransaction Transaction( LOCTEXT("RetargetSourceWindow_Delete", "Delete Retarget Source") );
 	for (int RowIndex = 0; RowIndex < SelectedRows.Num(); ++RowIndex)
 	{
-		const FReferencePose* PoseFound = Skeleton->AnimRetargetSources.Find(SelectedRows[RowIndex]->Name);
-		if(PoseFound)
-		{
-			// need to verify if these pose is used by anybody else
-			FAssetRegistryModule& AssetRegistryModule = FModuleManager::LoadModuleChecked<FAssetRegistryModule>(TEXT("AssetRegistry"));
-
-			TArray<FAssetData> AssetList;
-			TMultiMap<FName, FString> TagsAndValues;
-			TagsAndValues.Add(GET_MEMBER_NAME_CHECKED(UAnimSequence, RetargetSource), PoseFound->PoseName.ToString());
-			AssetRegistryModule.Get().GetAssetsByTagValues(TagsAndValues, AssetList);
+		SelectedNames.Add(SelectedRows[RowIndex]->Name);
+	}
 
-			// ask users if they'd like to continue and/or fix up
-			if (AssetList.Num() > 0 )
-			{
-				FString ListOfAssets;
-				// if users is sure to delete, delete
-				for (auto Iter=AssetList.CreateConstIterator(); Iter; ++Iter)
-				{
-					ListOfAssets += Iter->AssetName.ToString();
-					ListOfAssets += TEXT("\n");
-				}
-
-				// ask if they'd like to continue deleting this pose regardless animation references
-				FString Message = NSLOCTEXT("RetargetSourceEditor", "RetargetSourceDeleteMessage", "Following animation(s) is(are) referencing this pose. Are you sure if you'd like to delete this pose?").ToString();
-				Message += TEXT("\n\n");
-				Message += ListOfAssets;
-
-				EAppReturnType::Type Result = FMessageDialog::Open( EAppMsgType::YesNo, FText::FromString(Message) );
-
-				if (Result == EAppReturnType::No)
-				{
-					continue;
-				}
-
-				// if so, ask if they'd like us to fix the animations as well
-				Message = NSLOCTEXT("RetargetSourceEditor", "RetargetSourceDelete_FixUpAnimation_Message", "Would you like to fix up the following animation(s)? You'll have to save all the assets in the list.").ToString();
-				Message += TEXT("\n");
-				Message += ListOfAssets;
-
-				Result = FMessageDialog::Open( EAppMsgType::YesNo, FText::FromString(Message) );
-
-				if (Result == EAppReturnType::No)
-				{
-					continue;
-				}
-
-				// now fix up all assets
-				TArray<UObject *> ObjectsToUpdate;
-				for (auto Iter=AssetList.CreateConstIterator(); Iter; ++Iter)
-				{
-					UAnimSequence * AnimSequence = Cast<UAnimSequence>(Iter->GetAsset());
-					if (AnimSequence)
-					{
-						ObjectsToUpdate.Add(AnimSequence);
-
-						AnimSequence->Modify();
-						// clear name
-						AnimSequence->RetargetSource = NAME_None;
-					}
-				}
-			}
+	EditableSkeletonPtr.Pin()->DeleteRetargetSources(SelectedNames);
 
-			Skeleton->Modify();
-			// delete now
-			Skeleton->AnimRetargetSources.Remove(SelectedRows[RowIndex]->Name);
-			Skeleton->CallbackRetargetSourceChanged();
-			FAssetNotifications::SkeletonNeedsToBeSaved(Skeleton);
-		}
-	}
+	FAssetNotifications::SkeletonNeedsToBeSaved(&EditableSkeletonPtr.Pin()->GetSkeleton());
 
 	CreateRetargetSourceList( NameFilterBox->GetText().ToString() );
 }
@@ -697,7 +513,8 @@ void SRetargetSourceWindow::OnRenameRetargetSource()
 	if (ensure (SelectedRows.Num() == 1))
 	{
 		int32 RowIndex = 0;
-		const FReferencePose* PoseFound = Skeleton->AnimRetargetSources.Find(SelectedRows[RowIndex]->Name);
+		const USkeleton& Skeleton = EditableSkeletonPtr.Pin()->GetSkeleton();
+		const FReferencePose* PoseFound = Skeleton.AnimRetargetSources.Find(SelectedRows[RowIndex]->Name);
 		if(PoseFound)
 		{
 			// we used to verify if there is any animation referencing and warn them, but that doesn't really help
@@ -716,66 +533,20 @@ bool SRetargetSourceWindow::CanPerformRefresh() const
 
 void SRetargetSourceWindow::OnRefreshRetargetSource()
 {
+	TArray<FName> SelectedNames;
 	TArray< TSharedPtr< FDisplayedRetargetSourceInfo > > SelectedRows = RetargetSourceListView->GetSelectedItems();
-
-	const FScopedTransaction Transaction( LOCTEXT("RetargetSourceWindow_Refresh", "Refresh Retarget Source") );
 	for (int RowIndex = 0; RowIndex < SelectedRows.Num(); ++RowIndex)
 	{
-		Skeleton->Modify();
-		// refresh pose now
-		FName RetargetSourceName = SelectedRows[RowIndex]->Name;
-		Skeleton->UpdateRetargetSource(RetargetSourceName);
-
-		// feedback of pose has been updated
-		FFormatNamedArguments Args;
-		Args.Add( TEXT("RetargetSourceName"), FText::FromString( RetargetSourceName.ToString() ) );
-		FNotificationInfo Info( FText::Format( LOCTEXT("RetargetSourceWindow_RefreshPose", "Retarget Source {RetargetSourceName} is refreshed"), Args ) );
-		Info.ExpireDuration = 5.0f;
-		Info.bUseLargeFont = false;
-		TSharedPtr<SNotificationItem> Notification = FSlateNotificationManager::Get().AddNotification(Info);
-		if ( Notification.IsValid() )
-		{
-			Notification->SetCompletionState( SNotificationItem::CS_None );
-		}
+		SelectedNames.Add(SelectedRows[RowIndex]->Name);
 	}
 
-	FAssetNotifications::SkeletonNeedsToBeSaved(Skeleton);
-}
+	EditableSkeletonPtr.Pin()->RefreshRetargetSources(SelectedNames);
 
-SRetargetSourceWindow::~SRetargetSourceWindow()
-{
-	if (PersonaPtr.IsValid())
-	{
-		PersonaPtr.Pin()->UnregisterOnPostUndo(this);
-	}
+	FAssetNotifications::SkeletonNeedsToBeSaved(&EditableSkeletonPtr.Pin()->GetSkeleton());
 }
 
 void SRetargetSourceWindow::PostUndo()
 {
-	// handle renamed pose because currently adding old one without removing new one when the user did undo action
-	// As a result, both exist in the list so should remove new one 
-	if (RenamedPoseNameStack.Num() > 0 && Skeleton->AnimRetargetSources.Num() > 0)
-	{
-		FPoseNameForUndo& RenamedPoseName = RenamedPoseNameStack.Top();
-
-		int32 NumSources = Skeleton->AnimRetargetSources.Num();
-
-		FReferencePose * Pose = Skeleton->AnimRetargetSources.Find(RenamedPoseName.OldPoseName);
-		if (Pose)
-		{
-			FReferencePose * NewPose = Skeleton->AnimRetargetSources.Find(RenamedPoseName.NewPoseName);
-
-			// if both the old pose name and the new pose name exist 
-			if (NewPose)
-			{
-				Skeleton->Modify();
-				Skeleton->AnimRetargetSources.Remove(RenamedPoseName.NewPoseName);
-				// remove the last one because we handled it
-				RenamedPoseNameStack.Pop();
-			}
-		}
-	}
-
 	CreateRetargetSourceList();
 }
 
