Copying //UE4/Dev-Framework to //UE4/Dev-Main (Source: //UE4/Dev-Framework @ 3136612)

diff --git a/Engine/Source/Editor/Persona/Private/SRetargetSourceWindow.h b/Engine/Source/Editor/Persona/Private/SRetargetSourceWindow.h
--- a/Engine/Source/Editor/Persona/Private/SRetargetSourceWindow.h
+++ b/Engine/Source/Editor/Persona/Private/SRetargetSourceWindow.h
@@ -57,11 +57,7 @@ class SRetargetSourceWindow : public SCompoundWidget
 {
 public:
 	SLATE_BEGIN_ARGS( SRetargetSourceWindow )
-		: _Persona()
 	{}
-		
-		/* The Persona that owns this table */
-		SLATE_ARGUMENT( TWeakPtr< FPersona >, Persona )
 
 	SLATE_END_ARGS()
 
@@ -71,13 +67,7 @@ public:
 	* @param InArgs - Arguments passed from Slate
 	*
 	*/
-	void Construct( const FArguments& InArgs );
-
-	/**
-	* Destructor - resets the morph targets
-	*
-	*/
-	virtual ~SRetargetSourceWindow();
+	void Construct( const FArguments& InArgs, const TSharedRef<IEditableSkeleton>& InEditableSkeleton, FSimpleMulticastDelegate& InOnPostUndo);
 
 	/**
 	* Filters the SListView when the user changes the search text box (NameFilterBox)
@@ -181,8 +171,8 @@ private:
 	*/
 	void CreateRetargetSourceList( const FString& SearchText = FString(""), const FName  NewName = NAME_None );
 
-	/** Pointer back to the Persona that owns us */
-	TWeakPtr<FPersona> PersonaPtr;
+	/** The editable skeleton */
+	TWeakPtr<class IEditableSkeleton> EditableSkeletonPtr;
 
 	/** Box to filter to a specific morph target name */
 	TSharedPtr<SSearchBox>	NameFilterBox;
@@ -193,31 +183,14 @@ private:
 	/** A list of retarget sources. Used by the RetargetSourceListView. */
 	TArray< TSharedPtr<FDisplayedRetargetSourceInfo> > RetargetSourceList;
 
-	/** The skeletal mesh that we grab the morph targets from */
-	USkeleton* Skeleton;
-
 	/** Current text typed into NameFilterBox */
 	FText FilterText;
 
 	/** Item to rename. Only valid for adding **/
 	mutable TSharedPtr<FDisplayedRetargetSourceInfo> ItemToRename;
 
-	/** data to handle renamed pose when undo action is executed */
-	struct FPoseNameForUndo
-	{
-		FName OldPoseName;
-		FName NewPoseName;
-
-		FPoseNameForUndo(FName OldName, FName NewName)
-			: OldPoseName(OldName), NewPoseName(NewName)
-		{}
-	};
-
-	/** stack for renamed poses for multiple undo actions **/
-	TArray<FPoseNameForUndo> RenamedPoseNameStack;
-
 	/** Rename committed. Called by Inline Widget **/
-	void OnRenameCommit( const FName& OldName,  const FString& NewName );
+	void OnRenameCommit(const FName& InOldName, const FString& InNewName);
 	/** Verify Rename is legit or not **/
 	bool OnVerifyRenameCommit( const FName& OldName, const FString& NewName, FText& OutErrorMessage);
 
diff --git a/Engine/Source/Editor/Persona/Private/SRigWindow.cpp b/Engine/Source/Editor/Persona/Private/SRigWindow.cpp
--- a/Engine/Source/Editor/Persona/Private/SRigWindow.cpp
+++ b/Engine/Source/Editor/Persona/Private/SRigWindow.cpp
@@ -18,6 +18,7 @@
 #include "AnimationRuntime.h"
 #include "BoneMappingHelper.h"
 #include "SSkeletonWidget.h"
+#include "IEditableSkeleton.h"
 
 #define LOCTEXT_NAMESPACE "SRigWindow"
 
@@ -48,9 +49,6 @@ public:
 		/* Widget used to display the list of retarget sources*/
 		SLATE_ARGUMENT( TSharedPtr<SBoneMappingListType>, BoneMappingListView )
 
-		/* Persona used to update the viewport when a weight slider is dragged */
-		SLATE_ARGUMENT( TWeakPtr<FPersona>, Persona )
-
 		SLATE_EVENT( FOnBoneMappingChanged, OnBoneMappingChanged)
 
 		SLATE_EVENT( FOnGetBoneMapping, OnGetBoneMapping)
@@ -94,8 +92,6 @@ void SBoneMappingListRow::Construct( const FArguments& InArgs, const TSharedRef<
 	OnBoneMappingChanged = InArgs._OnBoneMappingChanged;
 	OnGetBoneMapping = InArgs._OnGetBoneMapping;
 
-	PersonaPtr = InArgs._Persona;
-
 	check( Item.IsValid() );
 
 	SMultiColumnTableRow< FDisplayedBoneMappingInfoPtr >::Construct( FSuperRowType::FArguments(), InOwnerTableView );
@@ -125,8 +121,6 @@ TSharedRef< SWidget > SBoneMappingListRow::GenerateWidgetForColumn( const FName&
 	}
 	else
 	{
-		check (Item->Skeleton);
-
 		// show bone list
 		// Encase the SSpinbox in an SVertical box so we can apply padding. Setting ItemHeight on the containing SListView has no effect :-(
 		return
@@ -141,7 +135,7 @@ TSharedRef< SWidget > SBoneMappingListRow::GenerateWidgetForColumn( const FName&
 
 				+SHorizontalBox::Slot()
 				[
-					SNew(SBoneSelectionWidget, Item->Skeleton)
+					SNew(SBoneSelectionWidget, Item->EditableSkeletonPtr.Pin().ToSharedRef())
 					.Tooltip(FText::Format(LOCTEXT("BoneSelectinWidget", "Select Bone for node {0}"), FText::FromString(Item->GetDisplayName())))
 					.OnBoneSelectionChanged(this, &SBoneMappingListRow::OnBoneSelectionChanged)
 					.OnGetSelectedBone(this, &SBoneMappingListRow::GetSelectedBone)
@@ -192,22 +186,14 @@ FName SBoneMappingListRow::GetSelectedBone() const
 //////////////////////////////////////////////////////////////////////////
 // SRigWindow
 
-void SRigWindow::Construct(const FArguments& InArgs)
+void SRigWindow::Construct(const FArguments& InArgs, const TSharedRef<class IEditableSkeleton>& InEditableSkeleton, FSimpleMulticastDelegate& InOnPostUndo)
 {
-	PersonaPtr = InArgs._Persona;
-	Skeleton = nullptr;
+	EditableSkeletonPtr = InEditableSkeleton;
 	bDisplayAdvanced = false;
 
-	if ( PersonaPtr.IsValid() )
-	{
-		Skeleton = PersonaPtr.Pin()->GetSkeleton();
-		PersonaPtr.Pin()->RegisterOnPostUndo(FPersona::FOnPostUndo::CreateSP( this, &SRigWindow::PostUndo ) );
-	}
+	InOnPostUndo.Add(FSimpleDelegate::CreateSP( this, &SRigWindow::PostUndo ) );
 	
-	// @todo it will crash right nwo without Skeleton
-	check (Skeleton);
-
-	Skeleton->RefreshRigConfig();
+	InEditableSkeleton->RefreshRigConfig();
 
 	// show list of skeletalmeshes that they can choose from
 	FContentBrowserModule& ContentBrowserModule = FModuleManager::Get().LoadModuleChecked<FContentBrowserModule>(TEXT("ContentBrowser"));
@@ -371,7 +357,6 @@ TSharedRef<ITableRow> SRigWindow::GenerateBoneMappingRow(TSharedPtr<FDisplayedBo
 
 	return
 		SNew( SBoneMappingListRow, OwnerTable )
-		.Persona( PersonaPtr )
 		.Item( InInfo )
 		.RigWindow( this )
 		.BoneMappingListView( BoneMappingListView )
@@ -383,7 +368,8 @@ void SRigWindow::CreateBoneMappingList( const FString& SearchText)
 {
 	BoneMappingList.Empty();
 
-	const URig* Rig = Skeleton->GetRig();
+	const USkeleton& Skeleton = EditableSkeletonPtr.Pin()->GetSkeleton();
+	const URig* Rig = Skeleton.GetRig();
 
 	if ( Rig )
 	{
@@ -394,7 +380,7 @@ void SRigWindow::CreateBoneMappingList( const FString& SearchText)
 		{
 			const FName& Name = Node.Name;
 			const FString& DisplayName = Node.DisplayName;
-			const FName& BoneName = Skeleton->GetRigBoneMapping(Name);
+			const FName& BoneName = Skeleton.GetRigBoneMapping(Name);
 
 			if (Node.bAdvanced == bDisplayAdvanced)
 			{
@@ -407,7 +393,7 @@ void SRigWindow::CreateBoneMappingList( const FString& SearchText)
 					}
 				}
 
-				TSharedRef<FDisplayedBoneMappingInfo> Info = FDisplayedBoneMappingInfo::Make(Name, DisplayName, Skeleton);
+				TSharedRef<FDisplayedBoneMappingInfo> Info = FDisplayedBoneMappingInfo::Make(Name, DisplayName, EditableSkeletonPtr.Pin().ToSharedRef());
 
 				BoneMappingList.Add(Info);
 			}
@@ -420,17 +406,13 @@ void SRigWindow::CreateBoneMappingList( const FString& SearchText)
 
 void SRigWindow::OnAssetSelected(UObject* Object)
 {
-	if (Skeleton)
-	{
-		AssetComboButton->SetIsOpen(false);
+	AssetComboButton->SetIsOpen(false);
 
-		const FScopedTransaction Transaction(LOCTEXT("RigAssetChanged", "Select Rig"));
-		Skeleton->Modify();
-		Skeleton->SetRigConfig(Cast<URig>(Object));
-		CreateBoneMappingList(FilterText.ToString());
+	EditableSkeletonPtr.Pin()->SetRigConfig(Cast<URig>(Object));
 
-		FAssetNotifications::SkeletonNeedsToBeSaved(Skeleton);
-	}
+	CreateBoneMappingList(FilterText.ToString());
+
+	FAssetNotifications::SkeletonNeedsToBeSaved(&EditableSkeletonPtr.Pin()->GetSkeleton());
 }
 
 /** Returns true if the asset shouldn't show  */
@@ -441,15 +423,8 @@ bool SRigWindow::ShouldFilterAsset(const class FAssetData& AssetData)
 
 URig* SRigWindow::GetRigObject() const
 {
-	return (Skeleton)? Skeleton->GetRig() : nullptr;
-}
-
-SRigWindow::~SRigWindow()
-{
-	if (PersonaPtr.IsValid())
-	{
-		PersonaPtr.Pin()->UnregisterOnPostUndo(this);
-	}
+	const USkeleton& Skeleton = EditableSkeletonPtr.Pin()->GetSkeleton();
+	return Skeleton.GetRig();
 }
 
 void SRigWindow::PostUndo()
@@ -459,15 +434,13 @@ void SRigWindow::PostUndo()
 
 void SRigWindow::OnBoneMappingChanged(FName NodeName, FName BoneName)
 {
-	const FScopedTransaction Transaction(LOCTEXT("BoneMappingChanged", "Change Bone Mapping"));
-	Skeleton->Modify();
-
-	Skeleton->SetRigBoneMapping(NodeName, BoneName);
+	EditableSkeletonPtr.Pin()->SetRigBoneMapping(NodeName, BoneName);
 }
 
 FName SRigWindow::GetBoneMapping(FName NodeName)
 {
-	return Skeleton->GetRigBoneMapping(NodeName);
+	const USkeleton& Skeleton = EditableSkeletonPtr.Pin()->GetSkeleton();
+	return Skeleton.GetRigBoneMapping(NodeName);
 }
 
 FReply SRigWindow::OnToggleAdvanced()
@@ -491,10 +464,12 @@ FText SRigWindow::GetAdvancedButtonText() const
 
 TSharedRef<SWidget> SRigWindow::MakeRigPickerWithMenu()
 {
+	const USkeleton& Skeleton = EditableSkeletonPtr.Pin()->GetSkeleton();
+
 	// rig asset picker
 	return	
 		SNew(SRigPicker)
-		.InitialObject(Skeleton->GetRig())
+		.InitialObject(Skeleton.GetRig())
 		.OnShouldFilterAsset(this, &SRigWindow::ShouldFilterAsset)
 		.OnSetReference(this, &SRigWindow::OnAssetSelected)
 		.OnClose(this, &SRigWindow::CloseComboButton );
@@ -593,14 +568,12 @@ FReply SRigWindow::OnAutoMapping()
 		}
 
 		FReferenceSkeleton RigReferenceSkeleton = Rig->GetSourceReferenceSkeleton();
-		FBoneMappingHelper Helper(RigReferenceSkeleton, Skeleton->GetReferenceSkeleton());
+		const USkeleton& Skeleton = EditableSkeletonPtr.Pin()->GetSkeleton();
+		FBoneMappingHelper Helper(RigReferenceSkeleton, Skeleton.GetReferenceSkeleton());
 		TMap<FName, FName> BestMatches;
 		Helper.TryMatch(BestMatches);
 
-		for (const TPair<FName, FName>& BestCandidate : BestMatches)
-		{
-			Skeleton->SetRigBoneMapping(BestCandidate.Key, BestCandidate.Value);
-		}
+		EditableSkeletonPtr.Pin()->SetRigBoneMappings(BestMatches);
 
 		BoneMappingListView->RequestListRefresh();
 	}
@@ -613,12 +586,15 @@ FReply SRigWindow::OnClearMapping()
 	URig* Rig = GetRigObject();
 	if (Rig)
 	{
-		const TArray<FNode> Nodes = Rig->GetNodes();
+		const TArray<FNode>& Nodes = Rig->GetNodes();
+		TMap<FName, FName> Mappings;
 		for (const auto& Node : Nodes)
 		{
-			Skeleton->SetRigBoneMapping(Node.Name, NAME_None);
+			Mappings.Add(Node.Name, NAME_None);
 		}
 
+		EditableSkeletonPtr.Pin()->SetRigBoneMappings(Mappings);
+
 		BoneMappingListView->RequestListRefresh();
 	}
 	return FReply::Handled();
diff --git a/Engine/Source/Editor/Persona/Private/SRigWindow.h b/Engine/Source/Editor/Persona/Private/SRigWindow.h
--- a/Engine/Source/Editor/Persona/Private/SRigWindow.h
+++ b/Engine/Source/Editor/Persona/Private/SRigWindow.h
@@ -2,8 +2,6 @@
 
 #pragma once
 
-#include "Persona.h"
-
 //////////////////////////////////////////////////////////////////////////
 // FDisplayedBoneMappingInfo
 
@@ -12,12 +10,12 @@ class FDisplayedBoneMappingInfo
 public:
 	FName Name;
 	FString DisplayName;
-	USkeleton* Skeleton;
+	TWeakPtr<class IEditableSkeleton> EditableSkeletonPtr;
 
 	/** Static function for creating a new item, but ensures that you can only have a TSharedRef to one */
-	static TSharedRef<FDisplayedBoneMappingInfo> Make(const FName NodeName, const FString DisplayName, USkeleton* InSkeleton)
+	static TSharedRef<FDisplayedBoneMappingInfo> Make(const FName NodeName, const FString DisplayName, const TSharedRef<class IEditableSkeleton>& InEditableSkeleton)
 	{
-		return MakeShareable(new FDisplayedBoneMappingInfo(NodeName, DisplayName, InSkeleton));
+		return MakeShareable(new FDisplayedBoneMappingInfo(NodeName, DisplayName, InEditableSkeleton));
 	}
 
 	FName GetNodeName() const
@@ -32,10 +30,10 @@ public:
 
 protected:
 	/** Hidden constructor, always use Make above */
-	FDisplayedBoneMappingInfo(const FName InNodeName, const FString InDisplayName, USkeleton* InSkeleton)
+	FDisplayedBoneMappingInfo(const FName InNodeName, const FString InDisplayName, const TSharedRef<class IEditableSkeleton>& InEditableSkeleton)
 		: Name( InNodeName )
 		, DisplayName( InDisplayName )
-		, Skeleton( InSkeleton )
+		, EditableSkeletonPtr(InEditableSkeleton)
 	{}
 
 	/** Hidden constructor, always use Make above */
@@ -51,11 +49,7 @@ class SRigWindow : public SCompoundWidget
 {
 public:
 	SLATE_BEGIN_ARGS( SRigWindow )
-		: _Persona()
 	{}
-		
-		/* The Persona that owns this table */
-		SLATE_ARGUMENT( TWeakPtr< FPersona >, Persona )
 
 	SLATE_END_ARGS()
 
@@ -65,13 +59,7 @@ public:
 	* @param InArgs - Arguments passed from Slate
 	*
 	*/
-	void Construct( const FArguments& InArgs );
-
-	/**
-	* Destructor - resets the morph targets
-	*
-	*/
-	virtual ~SRigWindow();
+	void Construct( const FArguments& InArgs, const TSharedRef<class IEditableSkeleton>& InEditableSkeleton, FSimpleMulticastDelegate& InOnPostUndo );
 
 	/**
 	* Filters the SListView when the user changes the search text box (NameFilterBox)
@@ -128,7 +116,7 @@ private:
 	void CreateBoneMappingList( const FString& SearchText = FString("") );
 
 	/** Pointer back to the Persona that owns us */
-	TWeakPtr<FPersona> PersonaPtr;
+	TWeakPtr<class IEditableSkeleton> EditableSkeletonPtr;
 
 	/** Box to filter to a specific morph target name */
 	TSharedPtr<SSearchBox>	NameFilterBox;
@@ -139,9 +127,6 @@ private:
 	/** A list of retarget sources. Used by the BoneMappingListView. */
 	TArray< TSharedPtr<FDisplayedBoneMappingInfo> > BoneMappingList;
 
-	/** The skeletal mesh that we grab the morph targets from */
-	USkeleton* Skeleton;
-
 	/** Current text typed into NameFilterBox */
 	FText FilterText;
 
diff --git a/Engine/Source/Editor/Persona/Private/SSequenceEditor.cpp b/Engine/Source/Editor/Persona/Private/SSequenceEditor.cpp
--- a/Engine/Source/Editor/Persona/Private/SSequenceEditor.cpp
+++ b/Engine/Source/Editor/Persona/Private/SSequenceEditor.cpp
@@ -17,24 +17,24 @@
 //////////////////////////////////////////////////////////////////////////
 // SSequenceEditor
 
-void SSequenceEditor::Construct(const FArguments& InArgs)
+void SSequenceEditor::Construct(const FArguments& InArgs, TSharedRef<class IPersonaPreviewScene> InPreviewScene, TSharedRef<class IEditableSkeleton> InEditableSkeleton, FSimpleMulticastDelegate& OnPostUndo, FSimpleMulticastDelegate& OnCurvesChanged)
 {
 	SequenceObj = InArgs._Sequence;
 	check(SequenceObj);
+	PreviewScenePtr = InPreviewScene;
 
 	SAnimEditorBase::Construct( SAnimEditorBase::FArguments()
-		.Persona(InArgs._Persona)
-		);
-
-	PersonaPtr.Pin()->RegisterOnPostUndo(FPersona::FOnPostUndo::CreateSP( this, &SSequenceEditor::PostUndo ) );
-	PersonaPtr.Pin()->RegisterOnChangeTrackCurves(FPersona::FOnTrackCurvesChanged::CreateSP( this, &SSequenceEditor::OnTrackCurveChanged ) );
+		.OnObjectsSelected(InArgs._OnObjectsSelected), 
+		InPreviewScene );
 
+	OnPostUndo.Add(FPersona::FOnPostUndo::CreateSP( this, &SSequenceEditor::PostUndo ) );
+	OnCurvesChanged.Add(FPersona::FOnTrackCurvesChanged::CreateSP( this, &SSequenceEditor::HandleCurvesChanged) );
+	
 	EditorPanels->AddSlot()
 	.AutoHeight()
 	.Padding(0, 10)
 	[
-		SAssignNew( AnimNotifyPanel, SAnimNotifyPanel )
-		.Persona(InArgs._Persona)
+		SAssignNew( AnimNotifyPanel, SAnimNotifyPanel, OnPostUndo)
 		.Sequence(SequenceObj)
 		.WidgetWidth(S2ColumnWidget::DEFAULT_RIGHT_COLUMN_WIDTH)
 		.ViewInputMin(this, &SAnimEditorBase::GetViewMinInput)
@@ -44,14 +44,15 @@ void SSequenceEditor::Construct(const FArguments& InArgs)
 		.OnSetInputViewRange(this, &SAnimEditorBase::SetInputViewRange)
 		.OnGetScrubValue(this, &SAnimEditorBase::GetScrubValue)
 		.OnSelectionChanged(this, &SAnimEditorBase::OnSelectionChanged)
+		.OnAnimNotifiesChanged(InArgs._OnAnimNotifiesChanged)
+		.OnInvokeTab(InArgs._OnInvokeTab)
 	];
 
 	EditorPanels->AddSlot()
 	.AutoHeight()
 	.Padding(0, 10)
 	[
-		SAssignNew( AnimCurvePanel, SAnimCurvePanel )
-		.Persona(InArgs._Persona)
+		SAssignNew( AnimCurvePanel, SAnimCurvePanel, InEditableSkeleton )
 		.Sequence(SequenceObj)
 		.WidgetWidth(S2ColumnWidget::DEFAULT_RIGHT_COLUMN_WIDTH)
 		.ViewInputMin(this, &SAnimEditorBase::GetViewMinInput)
@@ -70,7 +71,6 @@ void SSequenceEditor::Construct(const FArguments& InArgs)
 		.Padding(0, 10)
 		[
 			SAssignNew(AnimTrackCurvePanel, SAnimTrackCurvePanel)
-			.Persona(InArgs._Persona)
 			.Sequence(AnimSeq)
 			.WidgetWidth(S2ColumnWidget::DEFAULT_RIGHT_COLUMN_WIDTH)
 			.ViewInputMin(this, &SAnimEditorBase::GetViewMinInput)
@@ -83,22 +83,9 @@ void SSequenceEditor::Construct(const FArguments& InArgs)
 	}
 }
 
-SSequenceEditor::~SSequenceEditor()
-{
-	if (PersonaPtr.IsValid())
-	{
-		PersonaPtr.Pin()->UnregisterOnPostUndo(this);
-		PersonaPtr.Pin()->UnregisterOnChangeTrackCurves(this);
-	}
-}
-
 void SSequenceEditor::PostUndo()
 {
-	TSharedPtr<FPersona> SharedPersona = PersonaPtr.Pin();
-	if(SharedPersona.IsValid())
-	{
-		SharedPersona->SetPreviewAnimationAsset(SequenceObj);
-	}
+	GetPreviewScene()->SetPreviewAnimationAsset(SequenceObj);
 
 	if( SequenceObj )
 	{
@@ -113,7 +100,7 @@ void SSequenceEditor::PostUndo()
 	}
 }
 
-void SSequenceEditor::OnTrackCurveChanged()
+void SSequenceEditor::HandleCurvesChanged()
 {
 	if (AnimTrackCurvePanel.IsValid())
 	{
diff --git a/Engine/Source/Editor/Persona/Private/SSequenceEditor.h b/Engine/Source/Editor/Persona/Private/SSequenceEditor.h
--- a/Engine/Source/Editor/Persona/Private/SSequenceEditor.h
+++ b/Engine/Source/Editor/Persona/Private/SSequenceEditor.h
@@ -17,24 +17,24 @@ class SSequenceEditor : public SAnimEditorBase
 {
 public:
 	SLATE_BEGIN_ARGS( SSequenceEditor )
-		: _Persona()
-		, _Sequence(NULL)
+		: _Sequence(NULL)
 		{}
 
-		SLATE_ARGUMENT( TSharedPtr<FPersona>, Persona )
-		SLATE_ARGUMENT( UAnimSequenceBase*, Sequence )
+		SLATE_ARGUMENT(UAnimSequenceBase*, Sequence)
+		SLATE_EVENT(FOnObjectsSelected, OnObjectsSelected)
+		SLATE_EVENT(FSimpleDelegate, OnAnimNotifiesChanged)
+		SLATE_EVENT(FOnInvokeTab, OnInvokeTab)
+
 	SLATE_END_ARGS()
 
 private:
-	/** Persona reference **/
 	TSharedPtr<class SAnimNotifyPanel>	AnimNotifyPanel;
 	TSharedPtr<class SAnimCurvePanel>	AnimCurvePanel;
 	TSharedPtr<class SAnimTrackCurvePanel>	AnimTrackCurvePanel;
 	TSharedPtr<class SAnimationScrubPanel> AnimScrubPanel;
-
+	TWeakPtr<class IPersonaPreviewScene> PreviewScenePtr;
 public:
-	void Construct(const FArguments& InArgs);
-	virtual ~SSequenceEditor();
+	void Construct(const FArguments& InArgs, TSharedRef<class IPersonaPreviewScene> InPreviewScene, TSharedRef<class IEditableSkeleton> InEditableSkeleton, FSimpleMulticastDelegate& OnPostUndo, FSimpleMulticastDelegate& OnCurvesChanged);
 
 	virtual UAnimationAsset* GetEditorObject() const override { return SequenceObj; }
 
@@ -44,5 +44,5 @@ private:
 
 	/** Post undo **/
 	void PostUndo();
-	void OnTrackCurveChanged();
+	void HandleCurvesChanged();
 };
diff --git a/Engine/Source/Editor/Persona/Private/SSkeletonAnimNotifies.cpp b/Engine/Source/Editor/Persona/Private/SSkeletonAnimNotifies.cpp
--- a/Engine/Source/Editor/Persona/Private/SSkeletonAnimNotifies.cpp
+++ b/Engine/Source/Editor/Persona/Private/SSkeletonAnimNotifies.cpp
@@ -5,13 +5,12 @@
 
 #include "Persona.h"
 #include "SSkeletonAnimNotifies.h"
-#include "AssetRegistryModule.h"
 #include "ScopedTransaction.h"
 #include "SSearchBox.h"
 #include "SInlineEditableTextBlock.h"
 #include "SNotificationList.h"
 #include "NotificationManager.h"
-#include "BlueprintActionDatabase.h"
+#include "IEditableSkeleton.h"
 
 #define LOCTEXT_NAMESPACE "SkeletonAnimNotifies"
 
@@ -35,9 +34,6 @@ public:
 	/* Widget used to display the list of morph targets */
 	SLATE_ARGUMENT( TSharedPtr<SSkeletonAnimNotifies>, NotifiesListView )
 
-	/* Persona used to update the viewport when a weight slider is dragged */
-	SLATE_ARGUMENT( TWeakPtr<FPersona>, Persona )
-
 	SLATE_END_ARGS()
 
 	void Construct( const FArguments& InArgs, const TSharedRef<STableViewBase>& OwnerTableView );
@@ -52,16 +48,12 @@ private:
 
 	/** The notify being displayed by this row */
 	FDisplayedAnimNotifyInfoPtr	Item;
-
-	/** Pointer back to the Persona that owns us */
-	TWeakPtr<FPersona> PersonaPtr;
 };
 
 void SAnimNotifyListRow::Construct( const FArguments& InArgs, const TSharedRef<STableViewBase>& InOwnerTableView )
 {
 	Item = InArgs._Item;
 	NotifiesListView = InArgs._NotifiesListView;
-	PersonaPtr = InArgs._Persona;
 
 	check( Item.IsValid() );
 
@@ -89,8 +81,12 @@ TSharedRef< SWidget > SAnimNotifyListRow::GenerateWidgetForColumn( const FName&
 /////////////////////////////////////////////////////
 // FSkeletonAnimNotifiesSummoner
 
-FSkeletonAnimNotifiesSummoner::FSkeletonAnimNotifiesSummoner(TSharedPtr<class FAssetEditorToolkit> InHostingApp)
+FSkeletonAnimNotifiesSummoner::FSkeletonAnimNotifiesSummoner(TSharedPtr<class FAssetEditorToolkit> InHostingApp, const TSharedRef<class IEditableSkeleton>& InEditableSkeleton, FSimpleMulticastDelegate& InOnChangeAnimNotifies, FSimpleMulticastDelegate& InOnPostUndo, FOnObjectsSelected InOnObjectsSelected)
 	: FWorkflowTabFactory(FPersonaTabs::SkeletonAnimNotifiesID, InHostingApp)
+	, EditableSkeleton(InEditableSkeleton)
+	, OnChangeAnimNotifies(InOnChangeAnimNotifies)
+	, OnPostUndo(InOnPostUndo)
+	, OnObjectsSelected(InOnObjectsSelected)
 {
 	TabLabel = LOCTEXT("SkeletonAnimNotifiesTabTitle", "Animation Notifies");
 	TabIcon = FSlateIcon(FEditorStyle::GetStyleSetName(), "Persona.Tabs.AnimationNotifies");
@@ -104,20 +100,20 @@ FSkeletonAnimNotifiesSummoner::FSkeletonAnimNotifiesSummoner(TSharedPtr<class FA
 
 TSharedRef<SWidget> FSkeletonAnimNotifiesSummoner::CreateTabBody(const FWorkflowTabSpawnInfo& Info) const
 {
-	return SNew(SSkeletonAnimNotifies)
-		.Persona(StaticCastSharedPtr<FPersona>(HostingApp.Pin()));
+	return SNew(SSkeletonAnimNotifies, EditableSkeleton.Pin().ToSharedRef(), OnChangeAnimNotifies, OnPostUndo)
+		.OnObjectsSelected(OnObjectsSelected);
 }
 
 /////////////////////////////////////////////////////
 // SSkeletonAnimNotifies
 
-void SSkeletonAnimNotifies::Construct(const FArguments& InArgs)
+void SSkeletonAnimNotifies::Construct(const FArguments& InArgs, const TSharedRef<IEditableSkeleton>& InEditableSkeleton, FSimpleMulticastDelegate& InOnChangeAnimNotifies, FSimpleMulticastDelegate& InOnPostUndo)
 {
-	PersonaPtr = InArgs._Persona;
-	TargetSkeleton = PersonaPtr.Pin()->GetSkeleton();
+	OnObjectsSelected = InArgs._OnObjectsSelected;
+	EditableSkeleton = InEditableSkeleton;
 
-	PersonaPtr.Pin()->RegisterOnPostUndo(FPersona::FOnPostUndo::CreateSP( this, &SSkeletonAnimNotifies::PostUndo ) );
-	PersonaPtr.Pin()->RegisterOnChangeAnimNotifies(FPersona::FOnAnimNotifiesChanged::CreateSP( this, &SSkeletonAnimNotifies::RefreshNotifiesListWithFilter ) );
+	InOnChangeAnimNotifies.Add(FSimpleDelegate::CreateSP( this, &SSkeletonAnimNotifies::PostUndo ) );
+	InOnPostUndo.Add(FSimpleDelegate::CreateSP( this, &SSkeletonAnimNotifies::RefreshNotifiesListWithFilter ) );
 
 	this->ChildSlot
 	[
@@ -155,15 +151,6 @@ void SSkeletonAnimNotifies::Construct(const FArguments& InArgs)
 	CreateNotifiesList();
 }
 
-SSkeletonAnimNotifies::~SSkeletonAnimNotifies()
-{
-	if ( PersonaPtr.IsValid() )
-	{
-		PersonaPtr.Pin()->UnregisterOnPostUndo( this );
-		PersonaPtr.Pin()->UnregisterOnChangeAnimNotifies( this );
-	}
-}
-
 void SSkeletonAnimNotifies::OnFilterTextChanged( const FText& SearchText )
 {
 	FilterText = SearchText;
@@ -183,7 +170,6 @@ TSharedRef<ITableRow> SSkeletonAnimNotifies::GenerateNotifyRow(TSharedPtr<FDispl
 
 	return
 		SNew( SAnimNotifyListRow, OwnerTable )
-		.Persona( PersonaPtr )
 		.Item( InInfo )
 		.NotifiesListView( SharedThis( this ) );
 }
@@ -240,50 +226,15 @@ void SSkeletonAnimNotifies::OnDeleteAnimNotify()
 {
 	TArray< TSharedPtr< FDisplayedAnimNotifyInfo > > SelectedRows = NotifiesListView->GetSelectedItems();
 
-	const FScopedTransaction Transaction( LOCTEXT("DeleteAnimNotify", "Delete Anim Notify") );
-
 	// this one deletes all notifies with same name. 
 	TArray<FName> SelectedNotifyNames;
-	TargetSkeleton->Modify();
 
 	for(int Selection = 0; Selection < SelectedRows.Num(); ++Selection)
 	{
-		FName& SelectedName = SelectedRows[Selection]->Name;
-		TargetSkeleton->AnimationNotifies.Remove(SelectedName);
-		SelectedNotifyNames.Add(SelectedName);
+		SelectedNotifyNames.Add(SelectedRows[Selection]->Name);
 	}
 
-	TArray<FAssetData> CompatibleAnimSequences;
-	GetCompatibleAnimSequences(CompatibleAnimSequences);
-
-	int32 NumAnimationsModified = 0;
-
-	for( int32 AssetIndex = 0; AssetIndex < CompatibleAnimSequences.Num(); ++AssetIndex )
-	{
-		const FAssetData& PossibleAnimSequence = CompatibleAnimSequences[AssetIndex];
-		UAnimSequenceBase* Sequence = Cast<UAnimSequenceBase>(PossibleAnimSequence.GetAsset());
-
-		bool SequenceModified = false;
-		for(int32 NotifyIndex = Sequence->Notifies.Num()-1; NotifyIndex >= 0; --NotifyIndex)
-		{
-			FAnimNotifyEvent& AnimNotify = Sequence->Notifies[NotifyIndex];
-			if(SelectedNotifyNames.Contains(AnimNotify.NotifyName))
-			{
-				if(!SequenceModified)
-				{
-					Sequence->Modify();
-					++NumAnimationsModified;
-					SequenceModified = true;
-				}
-				Sequence->Notifies.RemoveAtSwap(NotifyIndex);
-			}
-		}
-
-		if(SequenceModified)
-		{
-			Sequence->MarkPackageDirty();
-		}
-	}
+	int32 NumAnimationsModified = EditableSkeleton->DeleteAnimNotifies(SelectedNotifyNames);
 
 	if(NumAnimationsModified > 0)
 	{
@@ -298,8 +249,6 @@ void SSkeletonAnimNotifies::OnDeleteAnimNotify()
 		NotifyUser( Info );
 	}
 
-	FBlueprintActionDatabase::Get().RefreshAssetActions(TargetSkeleton);
-
 	CreateNotifiesList( NameFilterBox->GetText().ToString() );
 }
 
@@ -325,7 +274,7 @@ bool SSkeletonAnimNotifies::OnVerifyNotifyNameCommit( const FText& NewName, FTex
 	FName NotifyName( *NewName.ToString() );
 	if(NotifyName != Item->Name)
 	{
-		if(TargetSkeleton->AnimationNotifies.Contains(NotifyName))
+		if(EditableSkeleton->GetSkeleton().AnimationNotifies.Contains(NotifyName))
 		{
 			OutErrorMessage = FText::Format( LOCTEXT("AlreadyInUseMessage", "'{0}' is already in use."), NewName );
 			bValid = false;
@@ -337,47 +286,7 @@ bool SSkeletonAnimNotifies::OnVerifyNotifyNameCommit( const FText& NewName, FTex
 
 void SSkeletonAnimNotifies::OnNotifyNameCommitted( const FText& NewName, ETextCommit::Type, TSharedPtr<FDisplayedAnimNotifyInfo> Item )
 {
-	const FScopedTransaction Transaction( LOCTEXT("RenameAnimNotify", "Rename Anim Notify") );
-
-	FName NewNotifyName = FName( *NewName.ToString() );
-	FName NotifyToRename = Item->Name;
-
-	// rename notify in skeleton
-	TargetSkeleton->Modify();
-	int32 Index = TargetSkeleton->AnimationNotifies.IndexOfByKey(NotifyToRename);
-	TargetSkeleton->AnimationNotifies[Index] = NewNotifyName;
-
-	TArray<FAssetData> CompatibleAnimSequences;
-	GetCompatibleAnimSequences(CompatibleAnimSequences);
-
-	int32 NumAnimationsModified = 0;
-
-	for( int32 AssetIndex = 0; AssetIndex < CompatibleAnimSequences.Num(); ++AssetIndex )
-	{
-		const FAssetData& PossibleAnimSequence = CompatibleAnimSequences[AssetIndex];
-		UAnimSequenceBase* Sequence = Cast<UAnimSequenceBase>(PossibleAnimSequence.GetAsset());
-
-		bool SequenceModified = false;
-		for(int32 NotifyIndex = Sequence->Notifies.Num()-1; NotifyIndex >= 0; --NotifyIndex)
-		{
-			FAnimNotifyEvent& AnimNotify = Sequence->Notifies[NotifyIndex];
-			if(NotifyToRename == AnimNotify.NotifyName)
-			{
-				if(!SequenceModified)
-				{
-					Sequence->Modify();
-					++NumAnimationsModified;
-					SequenceModified = true;
-				}	
-				AnimNotify.NotifyName = NewNotifyName;
-			}
-		}
-
-		if(SequenceModified)
-		{
-			Sequence->MarkPackageDirty();
-		}
-	}
+	int32 NumAnimationsModified = EditableSkeleton->RenameNotify(FName(*NewName.ToString()), Item->Name);
 
 	if(NumAnimationsModified > 0)
 	{
@@ -404,23 +313,20 @@ void SSkeletonAnimNotifies::CreateNotifiesList( const FString& SearchText )
 {
 	NotifyList.Empty();
 
-	if ( TargetSkeleton )
+	for(int i = 0; i < EditableSkeleton->GetSkeleton().AnimationNotifies.Num(); ++i)
 	{
-		for(int i = 0; i < TargetSkeleton->AnimationNotifies.Num(); ++i)
+		const FName& NotifyName = EditableSkeleton->GetSkeleton().AnimationNotifies[i];
+		if ( !SearchText.IsEmpty() )
 		{
-			FName& NotifyName = TargetSkeleton->AnimationNotifies[i];
-			if ( !SearchText.IsEmpty() )
-			{
-				if ( NotifyName.ToString().Contains( SearchText ) )
-				{
-					NotifyList.Add( FDisplayedAnimNotifyInfo::Make( NotifyName ) );
-				}
-			}
-			else
+			if ( NotifyName.ToString().Contains( SearchText ) )
 			{
 				NotifyList.Add( FDisplayedAnimNotifyInfo::Make( NotifyName ) );
 			}
 		}
+		else
+		{
+			NotifyList.Add( FDisplayedAnimNotifyInfo::Make( NotifyName ) );
+		}
 	}
 
 	NotifiesListView->RequestListRefresh();
@@ -434,7 +340,7 @@ void SSkeletonAnimNotifies::ShowNotifyInDetailsView(FName NotifyName)
 		Obj->AnimationNames.Empty();
 
 		TArray<FAssetData> CompatibleAnimSequences;
-		GetCompatibleAnimSequences(CompatibleAnimSequences);
+		EditableSkeleton->GetCompatibleAnimSequences(CompatibleAnimSequences);
 
 		for( int32 AssetIndex = 0; AssetIndex < CompatibleAnimSequences.Num(); ++AssetIndex )
 		{
@@ -455,43 +361,23 @@ void SSkeletonAnimNotifies::ShowNotifyInDetailsView(FName NotifyName)
 	}
 }
 
-void SSkeletonAnimNotifies::GetCompatibleAnimSequences(TArray<class FAssetData>& OutAssets)
-{
-	//Get the skeleton tag to search for
-	FString SkeletonExportName = FAssetData(TargetSkeleton).GetExportTextName();
-
-	// Load the asset registry module
-	FAssetRegistryModule& AssetRegistryModule = FModuleManager::LoadModuleChecked<FAssetRegistryModule>(TEXT("AssetRegistry"));
-
-	TArray<FAssetData> AssetDataList;
-	AssetRegistryModule.Get().GetAssetsByClass(UAnimSequenceBase::StaticClass()->GetFName(), AssetDataList, true);
-
-	OutAssets.Empty(AssetDataList.Num());
-
-	for( int32 AssetIndex = 0; AssetIndex < AssetDataList.Num(); ++AssetIndex )
-	{
-		const FAssetData& PossibleAnimSequence = AssetDataList[AssetIndex];
-		if( SkeletonExportName == PossibleAnimSequence.GetTagValueRef<FString>("Skeleton") )
-		{
-			OutAssets.Add( PossibleAnimSequence );
-		}
-	}
-}
-
 UObject* SSkeletonAnimNotifies::ShowInDetailsView( UClass* EdClass )
 {
 	UObject *Obj = EditorObjectTracker.GetEditorObjectForClass(EdClass);
 
 	if(Obj != NULL)
 	{
-		PersonaPtr.Pin()->SetDetailObject(Obj);
+		TArray<UObject*> Objects;
+		Objects.Add(Obj);
+		OnObjectsSelected.ExecuteIfBound(Objects);
 	}
 	return Obj;
 }
 
 void SSkeletonAnimNotifies::ClearDetailsView()
 {
-	PersonaPtr.Pin()->SetDetailObject(NULL);
+	TArray<UObject*> Objects;
+	OnObjectsSelected.ExecuteIfBound(Objects);
 }
 
 void SSkeletonAnimNotifies::PostUndo()
diff --git a/Engine/Source/Editor/Persona/Private/SSkeletonAnimNotifies.h b/Engine/Source/Editor/Persona/Private/SSkeletonAnimNotifies.h
--- a/Engine/Source/Editor/Persona/Private/SSkeletonAnimNotifies.h
+++ b/Engine/Source/Editor/Persona/Private/SSkeletonAnimNotifies.h
@@ -11,7 +11,7 @@
 struct FSkeletonAnimNotifiesSummoner : public FWorkflowTabFactory
 {
 public:
-	FSkeletonAnimNotifiesSummoner(TSharedPtr<class FAssetEditorToolkit> InHostingApp);
+	FSkeletonAnimNotifiesSummoner(TSharedPtr<class FAssetEditorToolkit> InHostingApp, const TSharedRef<class IEditableSkeleton>& InEditableSkeleton, FSimpleMulticastDelegate& InOnChangeAnimNotifies, FSimpleMulticastDelegate& InOnPostUndo, FOnObjectsSelected InOnObjectsSelected);
 
 	virtual TSharedRef<SWidget> CreateTabBody(const FWorkflowTabSpawnInfo& Info) const override;
 
@@ -20,6 +20,12 @@ public:
 	{
 		return  IDocumentation::Get()->CreateToolTip(LOCTEXT("AnimationNotifierTooltip", "This tab lets you modify custom animation notifies"), NULL, TEXT("Shared/Editors/Persona"), TEXT("AnimationNotifies_Window"));
 	}
+
+private:
+	TWeakPtr<class IEditableSkeleton> EditableSkeleton;
+	FSimpleMulticastDelegate& OnChangeAnimNotifies;
+	FSimpleMulticastDelegate& OnPostUndo;
+	FOnObjectsSelected OnObjectsSelected;
 };
 
 //////////////////////////////////////////////////////////////////////////
@@ -56,16 +62,13 @@ class SSkeletonAnimNotifies : public SCompoundWidget, public FGCObject
 {
 public:
 	SLATE_BEGIN_ARGS( SSkeletonAnimNotifies )
-		: _Persona()
 	{}
 
-	SLATE_ARGUMENT( TWeakPtr<FPersona>, Persona )
+	SLATE_EVENT(FOnObjectsSelected, OnObjectsSelected)
 		
-		SLATE_END_ARGS()
+	SLATE_END_ARGS()
 public:
-	void Construct(const FArguments& InArgs);
-
-	~SSkeletonAnimNotifies();
+	void Construct(const FArguments& InArgs, const TSharedRef<IEditableSkeleton>& InEditableSkeleton, FSimpleMulticastDelegate& InOnChangeAnimNotifies, FSimpleMulticastDelegate& InOnPostUndo);
 
 	/**
 	* Accessor so our rows can grab the filter text for highlighting
@@ -132,17 +135,11 @@ private:
 	/** handler for user selecting a Notify in NotifiesListView - populates the details panel */
 	void ShowNotifyInDetailsView( FName NotifyName );
 
-	/** Populates OutAssets with the AnimSequences that match Personas current skeleton */
-	void GetCompatibleAnimSequences( TArray<class FAssetData>& OutAssets );
-
 	/** Utility function to display notifications to the user */
 	void NotifyUser( FNotificationInfo& NotificationInfo );
 
-	/** The owning Persona instance */
-	TWeakPtr<class FPersona> PersonaPtr;
-
 	/** The skeleton we are currently editing */
-	USkeleton* TargetSkeleton;
+	TSharedPtr<class IEditableSkeleton> EditableSkeleton;
 
 	/** SSearchBox to filter the notify list */
 	TSharedPtr<SSearchBox>	NameFilterBox;
@@ -158,6 +155,9 @@ private:
 
 	/** Tracks objects created for displaying in the details panel*/
 	FEditorObjectTracker EditorObjectTracker;
+
+	/** Delegate called to select an object in the details panel */
+	FOnObjectsSelected OnObjectsSelected;
 };
 
 #undef LOCTEXT_NAMESPACE
\ No newline at end of file
diff --git a/Engine/Source/Editor/Persona/Private/SSkeletonBlendProfiles.cpp b/Engine/Source/Editor/Persona/Private/SSkeletonBlendProfiles.cpp
--- a/Engine/Source/Editor/Persona/Private/SSkeletonBlendProfiles.cpp
+++ b/Engine/Source/Editor/Persona/Private/SSkeletonBlendProfiles.cpp
@@ -7,6 +7,7 @@
 #include "SSearchBox.h"
 #include "STextEntryPopup.h"
 #include "ScopedTransaction.h"
+#include "IPersonaToolkit.h"
 
 #define LOCTEXT_NAMESPACE "SkeletonBlendProfiles"
 
@@ -117,7 +118,7 @@ void SSkeletonBlendProfiles::Construct(const FArguments& InArgs)
 	TSharedPtr<FPersona> SharedPersona = WeakPersona.Pin();
 	if(SharedPersona.IsValid())
 	{
-		CurrentSkeleton = SharedPersona->GetSkeleton();
+		CurrentSkeleton = SharedPersona->GetPersonaToolkit()->GetSkeleton();
 
 		ChildSlot
 		[
diff --git a/Engine/Source/Editor/Persona/Private/SSkeletonSlotNames.cpp b/Engine/Source/Editor/Persona/Private/SSkeletonSlotNames.cpp
--- a/Engine/Source/Editor/Persona/Private/SSkeletonSlotNames.cpp
+++ b/Engine/Source/Editor/Persona/Private/SSkeletonSlotNames.cpp
@@ -18,6 +18,7 @@
 #include "BlueprintEditorUtils.h"
 #include "SSlotNameReferenceWindow.h"
 #include "IMainFrameModule.h"
+#include "IEditableSkeleton.h"
 
 #define LOCTEXT_NAMESPACE "SkeletonSlotNames"
 
@@ -41,9 +42,6 @@ public:
 	/* Widget used to display the list of morph targets */
 	SLATE_ARGUMENT( TSharedPtr<SSkeletonSlotNames>, SlotNameListView )
 
-	/* Persona used to update the viewport when a weight slider is dragged */
-	SLATE_ARGUMENT( TWeakPtr<FPersona>, Persona )
-
 	SLATE_END_ARGS()
 
 	void Construct( const FArguments& InArgs, const TSharedRef<STableViewBase>& OwnerTableView );
@@ -58,16 +56,12 @@ private:
 
 	/** The notify being displayed by this row */
 	FDisplayedSlotNameInfoPtr	Item;
-
-	/** Pointer back to the Persona that owns us */
-	TWeakPtr<FPersona> PersonaPtr;
 };
 
 void SSlotNameListRow::Construct( const FArguments& InArgs, const TSharedRef<STableViewBase>& InOwnerTableView )
 {
 	Item = InArgs._Item;
 	SlotNameListView = InArgs._SlotNameListView;
-	PersonaPtr = InArgs._Persona;
 
 	check( Item.IsValid() );
 
@@ -100,8 +94,11 @@ TSharedRef< SWidget > SSlotNameListRow::GenerateWidgetForColumn( const FName& Co
 /////////////////////////////////////////////////////
 // FSkeletonSlotNamesSummoner
 
-FSkeletonSlotNamesSummoner::FSkeletonSlotNamesSummoner(TSharedPtr<class FAssetEditorToolkit> InHostingApp)
+FSkeletonSlotNamesSummoner::FSkeletonSlotNamesSummoner(TSharedPtr<class FAssetEditorToolkit> InHostingApp, const TSharedRef<IEditableSkeleton>& InEditableSkeleton, FSimpleMulticastDelegate& InOnPostUndo, FOnObjectSelected InOnObjectSelected)
 	: FWorkflowTabFactory(FPersonaTabs::SkeletonSlotNamesID, InHostingApp)
+	, EditableSkeleton(InEditableSkeleton)
+	, OnPostUndo(InOnPostUndo)
+	, OnObjectSelected(InOnObjectSelected)
 {
 	TabLabel = LOCTEXT("AnimSlotManagerTabTitle", "Anim Slot Manager");
 	TabIcon = FSlateIcon(FEditorStyle::GetStyleSetName(), "Persona.Tabs.AnimSlotManager");
@@ -115,19 +112,19 @@ FSkeletonSlotNamesSummoner::FSkeletonSlotNamesSummoner(TSharedPtr<class FAssetEd
 
 TSharedRef<SWidget> FSkeletonSlotNamesSummoner::CreateTabBody(const FWorkflowTabSpawnInfo& Info) const
 {
-	return SNew(SSkeletonSlotNames)
-		.Persona(StaticCastSharedPtr<FPersona>(HostingApp.Pin()));
+	return SNew(SSkeletonSlotNames, EditableSkeleton.Pin().ToSharedRef(), OnPostUndo)
+		.OnObjectSelected(OnObjectSelected);
 }
 
 /////////////////////////////////////////////////////
 // SSkeletonSlotNames
 
-void SSkeletonSlotNames::Construct(const FArguments& InArgs)
+void SSkeletonSlotNames::Construct(const FArguments& InArgs, const TSharedRef<IEditableSkeleton>& InEditableSkeleton, FSimpleMulticastDelegate& InOnPostUndo)
 {
-	PersonaPtr = InArgs._Persona;
-	TargetSkeleton = PersonaPtr.Pin()->GetSkeleton();
+	EditableSkeletonPtr = InEditableSkeleton;
+	OnObjectSelected = InArgs._OnObjectSelected;
 
-	PersonaPtr.Pin()->RegisterOnPostUndo(FPersona::FOnPostUndo::CreateSP( this, &SSkeletonSlotNames::PostUndo ) );
+	InOnPostUndo.Add(FSimpleDelegate::CreateSP( this, &SSkeletonSlotNames::PostUndo ) );
 
 	// Toolbar
 	FToolBarBuilder ToolbarBuilder(TSharedPtr< FUICommandList >(), FMultiBoxCustomization::None);
@@ -205,14 +202,6 @@ void SSkeletonSlotNames::Construct(const FArguments& InArgs)
 	CreateSlotNameList();
 }
 
-SSkeletonSlotNames::~SSkeletonSlotNames()
-{
-	if ( PersonaPtr.IsValid() )
-	{
-		PersonaPtr.Pin()->UnregisterOnPostUndo( this );
-	}
-}
-
 void SSkeletonSlotNames::OnFilterTextChanged( const FText& SearchText )
 {
 	FilterText = SearchText;
@@ -232,7 +221,6 @@ TSharedRef<ITableRow> SSkeletonSlotNames::GenerateNotifyRow(TSharedPtr<FDisplaye
 
 	return
 		SNew( SSlotNameListRow, OwnerTable )
-		.Persona( PersonaPtr )
 		.Item( InInfo )
 		.SlotNameListView( SharedThis( this ) );
 }
@@ -326,7 +314,7 @@ TSharedPtr<SWidget> SSkeletonSlotNames::OnGetContextMenuContent() const
 
 void SSkeletonSlotNames::FillSetSlotGroupSubMenu(FMenuBuilder& MenuBuilder)
 {
-	const TArray<FAnimSlotGroup>& SlotGroups = TargetSkeleton->GetSlotGroups();
+	const TArray<FAnimSlotGroup>& SlotGroups = EditableSkeletonPtr.Pin()->GetSkeleton().GetSlotGroups();
 	for (auto SlotGroup : SlotGroups)
 	{
 		const FName& GroupName = SlotGroup.GroupName;
@@ -351,11 +339,9 @@ void SSkeletonSlotNames::ContextMenuOnSetSlot(FName InNewGroupName)
 	if (bShowSlotItem)
 	{
 		const FName SlotName = SelectedItems[0].Get()->Name;
-		if (TargetSkeleton->ContainsSlotName(SlotName))
+		if (EditableSkeletonPtr.Pin()->GetSkeleton().ContainsSlotName(SlotName))
 		{
-			const FScopedTransaction Transaction(LOCTEXT("AnimSlotManager_ContextMenuOnSetSlot", "Set Slot Group Name"));
-			TargetSkeleton->SetSlotGroupName(SlotName, InNewGroupName);
-			TargetSkeleton->Modify(true);
+			EditableSkeletonPtr.Pin()->SetSlotGroupName(SlotName, InNewGroupName);
 
 			RefreshSlotNameListWithFilter();
 		}
@@ -381,13 +367,10 @@ void SSkeletonSlotNames::OnNotifySelectionChanged(TSharedPtr<FDisplayedSlotNameI
 
 void SSkeletonSlotNames::OnSaveSkeleton()
 {
-	if (TargetSkeleton)
-	{
-		TArray< UPackage* > PackagesToSave;
-		PackagesToSave.Add(TargetSkeleton->GetOutermost());
+	TArray< UPackage* > PackagesToSave;
+	PackagesToSave.Add(EditableSkeletonPtr.Pin()->GetSkeleton().GetOutermost());
 
-		FEditorFileUtils::PromptForCheckoutAndSave(PackagesToSave, /*bCheckDirty=*/ false, /*bPromptToSave=*/ false);
-	}
+	FEditorFileUtils::PromptForCheckoutAndSave(PackagesToSave, /*bCheckDirty=*/ false, /*bPromptToSave=*/ false);
 }
 
 void SSkeletonSlotNames::OnAddSlot()
@@ -431,15 +414,15 @@ void SSkeletonSlotNames::AddSlotPopUpOnCommit(const FText & InNewSlotText, EText
 		const FScopedTransaction Transaction(LOCTEXT("NewSlotName_AddSlotName", "Add New Slot Node Name"));
 
 		FName NewSlotName = FName(*InNewSlotText.ToString());
+		const USkeleton& Skeleton = EditableSkeletonPtr.Pin()->GetSkeleton();
 		// Keep slot and group names unique
-		if (!TargetSkeleton->ContainsSlotName(NewSlotName) && (TargetSkeleton->FindAnimSlotGroup(NewSlotName) == NULL))
+		if (!Skeleton.ContainsSlotName(NewSlotName) && (Skeleton.FindAnimSlotGroup(NewSlotName) == nullptr))
 		{
 			TArray< TSharedPtr< FDisplayedSlotNameInfo > > SelectedItems = SlotNameListView->GetSelectedItems();
 			bool bHasSelectedItem = (SelectedItems.Num() > 0);
 			bool bShowGroupItem = bHasSelectedItem && SelectedItems[0].Get()->bIsGroupItem;
 
-			TargetSkeleton->SetSlotGroupName(NewSlotName, bShowGroupItem ? SelectedItems[0].Get()->Name : FAnimSlotGroup::DefaultGroupName);
-			TargetSkeleton->Modify(true);
+			EditableSkeletonPtr.Pin()->SetSlotGroupName(NewSlotName, bShowGroupItem ? SelectedItems[0].Get()->Name : FAnimSlotGroup::DefaultGroupName);
 
 			RefreshSlotNameListWithFilter();
 		}
@@ -459,13 +442,11 @@ void SSkeletonSlotNames::AddGroupPopUpOnCommit(const FText & InNewGroupText, ETe
 {
 	if (!InNewGroupText.IsEmpty())
 	{
-		const FScopedTransaction Transaction(LOCTEXT("NewGroupName_AddGroupName", "Add New Slot Group Name"));
-
 		FName NewGroupName = FName(*InNewGroupText.ToString());
+		const USkeleton& Skeleton = EditableSkeletonPtr.Pin()->GetSkeleton();
 		// Keep slot and group names unique
-		if (!TargetSkeleton->ContainsSlotName(NewGroupName) && TargetSkeleton->AddSlotGroupName(NewGroupName))
+		if (!Skeleton.ContainsSlotName(NewGroupName) && EditableSkeletonPtr.Pin()->AddSlotGroupName(NewGroupName))
 		{
-			TargetSkeleton->Modify(true);
 			RefreshSlotNameListWithFilter();
 		}
 
@@ -483,7 +464,8 @@ void SSkeletonSlotNames::AddGroupPopUpOnCommit(const FText & InNewGroupText, ETe
 void SSkeletonSlotNames::GetCompatibleAnimBlueprints( TArray<FAssetData>& OutAssets )
 {
 	//Get the skeleton tag to search for
-	FString SkeletonExportName = FAssetData(TargetSkeleton).GetExportTextName();
+	const USkeleton& Skeleton = EditableSkeletonPtr.Pin()->GetSkeleton();
+	FString SkeletonExportName = FAssetData(&Skeleton).GetExportTextName();
 
 	// Load the asset registry module
 	FAssetRegistryModule& AssetRegistryModule = FModuleManager::LoadModuleChecked<FAssetRegistryModule>(TEXT("AssetRegistry"));
@@ -512,29 +494,27 @@ void SSkeletonSlotNames::CreateSlotNameList(const FString& SearchText)
 {
 	NotifyList.Empty();
 
-	if ( TargetSkeleton )
+	const USkeleton& Skeleton = EditableSkeletonPtr.Pin()->GetSkeleton();
+	const TArray<FAnimSlotGroup>& SlotGroups = Skeleton.GetSlotGroups();
+	for (auto SlotGroup : SlotGroups)
 	{
-		const TArray<FAnimSlotGroup>& SlotGroups = TargetSkeleton->GetSlotGroups();
-		for (auto SlotGroup : SlotGroups)
-		{
-			const FName& GroupName = SlotGroup.GroupName;
+		const FName& GroupName = SlotGroup.GroupName;
 			
-			TSharedRef<FDisplayedSlotNameInfo> GroupItem = FDisplayedSlotNameInfo::Make(GroupName, true);
-			SlotNameListView->SetItemExpansion(GroupItem, true);
-			NotifyList.Add(GroupItem);
+		TSharedRef<FDisplayedSlotNameInfo> GroupItem = FDisplayedSlotNameInfo::Make(GroupName, true);
+		SlotNameListView->SetItemExpansion(GroupItem, true);
+		NotifyList.Add(GroupItem);
 
-			for (auto SlotName : SlotGroup.SlotNames)
+		for (auto SlotName : SlotGroup.SlotNames)
+		{
+			if (SearchText.IsEmpty() || GroupName.ToString().Contains(SearchText) || SlotName.ToString().Contains(SearchText))
 			{
-				if (SearchText.IsEmpty() || GroupName.ToString().Contains(SearchText) || SlotName.ToString().Contains(SearchText))
-				{
-					TSharedRef<FDisplayedSlotNameInfo> SlotItem = FDisplayedSlotNameInfo::Make(SlotName, false);
-					SlotNameListView->SetItemExpansion(SlotItem, true);
-					NotifyList[NotifyList.Num() - 1]->Children.Add(SlotItem);
-				}
+				TSharedRef<FDisplayedSlotNameInfo> SlotItem = FDisplayedSlotNameInfo::Make(SlotName, false);
+				SlotNameListView->SetItemExpansion(SlotItem, true);
+				NotifyList[NotifyList.Num() - 1]->Children.Add(SlotItem);
 			}
 		}
 	}
-
+	
 	SlotNameListView->RequestTreeRefresh();
 }
 
@@ -567,7 +547,8 @@ void SSkeletonSlotNames::ShowNotifyInDetailsView(FName NotifyName)
 void SSkeletonSlotNames::GetCompatibleAnimMontages(TArray<class FAssetData>& OutAssets)
 {
 	//Get the skeleton tag to search for
-	FString SkeletonExportName = FAssetData(TargetSkeleton).GetExportTextName();
+	const USkeleton& Skeleton = EditableSkeletonPtr.Pin()->GetSkeleton();
+	FString SkeletonExportName = FAssetData(&Skeleton).GetExportTextName();
 
 	// Load the asset registry module
 	FAssetRegistryModule& AssetRegistryModule = FModuleManager::LoadModuleChecked<FAssetRegistryModule>(TEXT("AssetRegistry"));
@@ -593,14 +574,14 @@ UObject* SSkeletonSlotNames::ShowInDetailsView( UClass* EdClass )
 
 	if(Obj != NULL)
 	{
-		PersonaPtr.Pin()->SetDetailObject(Obj);
+		OnObjectSelected.ExecuteIfBound(Obj);
 	}
 	return Obj;
 }
 
 void SSkeletonSlotNames::ClearDetailsView()
 {
-	PersonaPtr.Pin()->SetDetailObject(NULL);
+	OnObjectSelected.ExecuteIfBound(nullptr);
 }
 
 void SSkeletonSlotNames::PostUndo()
@@ -687,10 +668,10 @@ void SSkeletonSlotNames::OnDeleteSlot(FName SlotName)
 
 void SSkeletonSlotNames::DeleteSlot(const FName& SlotName)
 {
-	if(TargetSkeleton->ContainsSlotName(SlotName))
+	const USkeleton& Skeleton = EditableSkeletonPtr.Pin()->GetSkeleton();
+	if(Skeleton.ContainsSlotName(SlotName))
 	{
-		TargetSkeleton->Modify();
-		TargetSkeleton->RemoveSlotName(SlotName);
+		EditableSkeletonPtr.Pin()->DeleteSlotName(SlotName);
 		RefreshSlotNameListWithFilter();
 	}
 }
@@ -718,9 +699,9 @@ void SSkeletonSlotNames::GetAnimMontagesUsingSlot(FName SlotName, TArray<FAssetD
 
 void SSkeletonSlotNames::GetAnimMontagesUsingSlotGroup(FName SlotGroupName, TArray<FAssetData>& OutMontages)
 {
-	if(FAnimSlotGroup* Group = TargetSkeleton->FindAnimSlotGroup(SlotGroupName))
+	if(const FAnimSlotGroup* Group = EditableSkeletonPtr.Pin()->GetSkeleton().FindAnimSlotGroup(SlotGroupName))
 	{
-		for(FName& SlotName : Group->SlotNames)
+		for(const FName& SlotName : Group->SlotNames)
 		{
 			GetAnimMontagesUsingSlot(SlotName, OutMontages);
 		}
@@ -802,14 +783,11 @@ void SSkeletonSlotNames::GetMontagesAndNodesUsingSlot(const FName& SlotName, TAr
 
 void SSkeletonSlotNames::GetMontagesAndNodesUsingSlotGroup(const FName& SlotGroupName, TArray<FAssetData>& OutMontages, TMultiMap<UAnimBlueprint*, UAnimGraphNode_Slot*> &OutBlueprintSlotMap)
 {
-	if(TargetSkeleton)
+	if(const FAnimSlotGroup* SlotGroup = EditableSkeletonPtr.Pin()->GetSkeleton().FindAnimSlotGroup(SlotGroupName))
 	{
-		if(FAnimSlotGroup* SlotGroup = TargetSkeleton->FindAnimSlotGroup(SlotGroupName))
+		for(const auto& SlotName : SlotGroup->SlotNames)
 		{
-			for(auto& SlotName : SlotGroup->SlotNames)
-			{
-				GetMontagesAndNodesUsingSlot(SlotName, OutMontages, OutBlueprintSlotMap);
-			}
+			GetMontagesAndNodesUsingSlot(SlotName, OutMontages, OutBlueprintSlotMap);
 		}
 	}
 }
@@ -823,81 +801,78 @@ void SSkeletonSlotNames::OnRenameSlotPopupCommitted(const FText & InNewSlotText,
 		// Need to dismiss menus early or the slow task in GetMontagesAndNodesUsingSlot will fail to show onscreen
 		FSlateApplication::Get().DismissAllMenus();
 
-		if(TargetSkeleton)
+		// Make sure the name doesn't already exist
+		if(EditableSkeletonPtr.Pin()->GetSkeleton().ContainsSlotName(NewName))
 		{
-			// Make sure the name doesn't already exist
-			if(TargetSkeleton->ContainsSlotName(NewName))
-			{
-				FNotificationInfo Notification(FText::Format(LOCTEXT("ToastRenameFailDesc", "Rename Failed! Slot name {0} already exists in the target skeleton."), FText::FromName(NewName)));
-				Notification.ExpireDuration = 3.0f;
-				Notification.bFireAndForget = true;
+			FNotificationInfo Notification(FText::Format(LOCTEXT("ToastRenameFailDesc", "Rename Failed! Slot name {0} already exists in the target skeleton."), FText::FromName(NewName)));
+			Notification.ExpireDuration = 3.0f;
+			Notification.bFireAndForget = true;
 
-				NotifyUser(Notification);
+			NotifyUser(Notification);
 
-				return;
-			}
+			return;
+		}
 
-			// Validate references
-			TArray<FAssetData> CompatibleMontages;
-			TMultiMap<UAnimBlueprint*, UAnimGraphNode_Slot*> CompatibleSlotNodes;
-			GetMontagesAndNodesUsingSlot(OldName, CompatibleMontages, CompatibleSlotNodes);
+		// Validate references
+		TArray<FAssetData> CompatibleMontages;
+		TMultiMap<UAnimBlueprint*, UAnimGraphNode_Slot*> CompatibleSlotNodes;
+		GetMontagesAndNodesUsingSlot(OldName, CompatibleMontages, CompatibleSlotNodes);
 
-			if(CompatibleMontages.Num() > 0 || CompatibleSlotNodes.Num() > 0)
+		if(CompatibleMontages.Num() > 0 || CompatibleSlotNodes.Num() > 0)
+		{
+			// We can't rename here - still have references. Give the user a chance to fix.
+			if(!ReferenceWindow.IsValid())
 			{
-				// We can't rename here - still have references. Give the user a chance to fix.
-				if(!ReferenceWindow.IsValid())
-				{
-					// No existing window
-					SAssignNew(ReferenceWindow, SWindow)
-						.AutoCenter(EAutoCenter::PreferredWorkArea)
-						.SizingRule(ESizingRule::Autosized)
-						.Title(LOCTEXT("ReferenceWindowTitle", "Slot References"));
-
-					ReferenceWindow->SetContent
-						(
-						SNew(SBorder)
-						.Padding(FMargin(3))
-						.BorderImage(FEditorStyle::GetBrush("ToolPanel.GroupBorder"))
-						[
-							SAssignNew(ReferenceWidget, SSlotNameReferenceWindow)
-							.ReferencingMontages(&CompatibleMontages)
-							.ReferencingNodes(&CompatibleSlotNodes)
-							.SlotName(OldName.ToString())
-							.OperationText(LOCTEXT("RenameOperation", "Rename"))
-							.WidgetWindow(ReferenceWindow)
-							.OnRetry(FSimpleDelegate::CreateSP(this, &SSkeletonSlotNames::RetryRenameSlot, OldName, NewName))
-						]
-					);
-
-					TSharedPtr<SWindow> ParentWindow;
-					IMainFrameModule& MainFrameModule = FModuleManager::LoadModuleChecked<IMainFrameModule>(TEXT("MainFrame"));
-					ParentWindow = MainFrameModule.GetParentWindow();
-
-					FSlateApplication::Get().AddWindowAsNativeChild(ReferenceWindow.ToSharedRef(), ParentWindow.ToSharedRef());
-					ReferenceWindow->SetOnWindowClosed(FOnWindowClosed::CreateSP(this, &SSkeletonSlotNames::ReferenceWindowClosed));
-				}
-				else
-				{
-					TSharedPtr<SSlotNameReferenceWindow> RefWidgetPinned = ReferenceWidget.Pin();
-					if(RefWidgetPinned.IsValid())
-					{
-						FReferenceWindowInfo WindowInfo;
-						WindowInfo.ReferencingMontages = &CompatibleMontages;
-						WindowInfo.ReferencingNodes = &CompatibleSlotNodes;
-						WindowInfo.ItemText = FText::FromName(OldName);
-						WindowInfo.OperationText = LOCTEXT("RenameOperation", "Rename");
-						WindowInfo.RetryDelegate = FSimpleDelegate::CreateSP(this, &SSkeletonSlotNames::RetryRenameSlot, OldName, NewName);
-
-						RefWidgetPinned->UpdateInfo(WindowInfo);
-						ReferenceWindow->BringToFront();
-					}
-				}
+				// No existing window
+				SAssignNew(ReferenceWindow, SWindow)
+					.AutoCenter(EAutoCenter::PreferredWorkArea)
+					.SizingRule(ESizingRule::Autosized)
+					.Title(LOCTEXT("ReferenceWindowTitle", "Slot References"));
+
+				ReferenceWindow->SetContent
+					(
+					SNew(SBorder)
+					.Padding(FMargin(3))
+					.BorderImage(FEditorStyle::GetBrush("ToolPanel.GroupBorder"))
+					[
+						SAssignNew(ReferenceWidget, SSlotNameReferenceWindow)
+						.ReferencingMontages(&CompatibleMontages)
+						.ReferencingNodes(&CompatibleSlotNodes)
+						.SlotName(OldName.ToString())
+						.OperationText(LOCTEXT("RenameOperation", "Rename"))
+						.WidgetWindow(ReferenceWindow)
+						.OnRetry(FSimpleDelegate::CreateSP(this, &SSkeletonSlotNames::RetryRenameSlot, OldName, NewName))
+					]
+				);
+
+				TSharedPtr<SWindow> ParentWindow;
+				IMainFrameModule& MainFrameModule = FModuleManager::LoadModuleChecked<IMainFrameModule>(TEXT("MainFrame"));
+				ParentWindow = MainFrameModule.GetParentWindow();
+
+				FSlateApplication::Get().AddWindowAsNativeChild(ReferenceWindow.ToSharedRef(), ParentWindow.ToSharedRef());
+				ReferenceWindow->SetOnWindowClosed(FOnWindowClosed::CreateSP(this, &SSkeletonSlotNames::ReferenceWindowClosed));
 			}
 			else
 			{
-				RenameSlot(OldName, NewName);
+				TSharedPtr<SSlotNameReferenceWindow> RefWidgetPinned = ReferenceWidget.Pin();
+				if(RefWidgetPinned.IsValid())
+				{
+					FReferenceWindowInfo WindowInfo;
+					WindowInfo.ReferencingMontages = &CompatibleMontages;
+					WindowInfo.ReferencingNodes = &CompatibleSlotNodes;
+					WindowInfo.ItemText = FText::FromName(OldName);
+					WindowInfo.OperationText = LOCTEXT("RenameOperation", "Rename");
+					WindowInfo.RetryDelegate = FSimpleDelegate::CreateSP(this, &SSkeletonSlotNames::RetryRenameSlot, OldName, NewName);
+
+					RefWidgetPinned->UpdateInfo(WindowInfo);
+					ReferenceWindow->BringToFront();
+				}
 			}
 		}
+		else
+		{
+			RenameSlot(OldName, NewName);
+		}
 	}
 }
 
@@ -920,10 +895,9 @@ void SSkeletonSlotNames::OnRenameSlot(FName CurrentName)
 
 void SSkeletonSlotNames::RenameSlot(FName CurrentName, FName NewName)
 {
-	if(TargetSkeleton->ContainsSlotName(CurrentName))
+	if(EditableSkeletonPtr.Pin()->GetSkeleton().ContainsSlotName(CurrentName))
 	{
-		TargetSkeleton->Modify();
-		TargetSkeleton->RenameSlotName(CurrentName, NewName);
+		EditableSkeletonPtr.Pin()->RenameSlotName(CurrentName, NewName);
 		RefreshSlotNameListWithFilter();
 	}
 }
@@ -1060,10 +1034,9 @@ bool SSkeletonSlotNames::CanDeleteSlotGroup(FName GroupName)
 
 void SSkeletonSlotNames::DeleteSlotGroup(const FName& GroupName)
 {
-	if(TargetSkeleton && TargetSkeleton->FindAnimSlotGroup(GroupName))
+	if(EditableSkeletonPtr.Pin()->GetSkeleton().FindAnimSlotGroup(GroupName) != nullptr)
 	{
-		TargetSkeleton->Modify();
-		TargetSkeleton->RemoveSlotGroup(GroupName);
+		EditableSkeletonPtr.Pin()->DeleteSlotGroup(GroupName);
 		RefreshSlotNameListWithFilter();
 	}
 }
diff --git a/Engine/Source/Editor/Persona/Private/SSkeletonSlotNames.h b/Engine/Source/Editor/Persona/Private/SSkeletonSlotNames.h
--- a/Engine/Source/Editor/Persona/Private/SSkeletonSlotNames.h
+++ b/Engine/Source/Editor/Persona/Private/SSkeletonSlotNames.h
@@ -14,7 +14,7 @@ class UAnimGraphNode_Slot;
 struct FSkeletonSlotNamesSummoner : public FWorkflowTabFactory
 {
 public:
-	FSkeletonSlotNamesSummoner(TSharedPtr<class FAssetEditorToolkit> InHostingApp);
+	FSkeletonSlotNamesSummoner(TSharedPtr<class FAssetEditorToolkit> InHostingApp, const TSharedRef<class IEditableSkeleton>& InEditableSkeleton, FSimpleMulticastDelegate& InOnPostUndo, FOnObjectSelected InOnObjectSelected);
 
 	virtual TSharedRef<SWidget> CreateTabBody(const FWorkflowTabSpawnInfo& Info) const override;
 
@@ -23,6 +23,10 @@ public:
 	{
 		return  IDocumentation::Get()->CreateToolTip(LOCTEXT("WindowTooltip", "This tab lets you modify custom animation SlotName"), NULL, TEXT("Shared/Editors/Persona"), TEXT("AnimationSlotName_Window"));
 	}
+public:
+	TWeakPtr<class IEditableSkeleton> EditableSkeleton;
+	FSimpleMulticastDelegate& OnPostUndo;
+	FOnObjectSelected OnObjectSelected;
 };
 
 //////////////////////////////////////////////////////////////////////////
@@ -64,16 +68,13 @@ class SSkeletonSlotNames : public SCompoundWidget, public FGCObject
 {
 public:
 	SLATE_BEGIN_ARGS( SSkeletonSlotNames )
-		: _Persona()
 	{}
 
-	SLATE_ARGUMENT( TWeakPtr<FPersona>, Persona )
-		
-		SLATE_END_ARGS()
-public:
-	void Construct(const FArguments& InArgs);
+	SLATE_ARGUMENT(FOnObjectSelected, OnObjectSelected)
 
-	~SSkeletonSlotNames();
+	SLATE_END_ARGS()
+public:
+	void Construct(const FArguments& InArgs, const TSharedRef<IEditableSkeleton>& InEditableSkeleton, FSimpleMulticastDelegate& InOnPostUndo);
 
 	/**
 	* Accessor so our rows can grab the filter text for highlighting
@@ -187,11 +188,11 @@ private:
 	// Callback for popup reference window closing
 	void ReferenceWindowClosed(const TSharedRef<SWindow>& Window);
 
-	/** The owning Persona instance */
-	TWeakPtr<class FPersona> PersonaPtr;
-
 	/** The skeleton we are currently editing */
-	USkeleton* TargetSkeleton;
+	TWeakPtr<class IEditableSkeleton> EditableSkeletonPtr;
+
+	/** Delegate call to select an object & display its details */
+	FOnObjectSelected OnObjectSelected;
 
 	/** SSearchBox to filter the notify list */
 	TSharedPtr<SSearchBox>	NameFilterBox;
diff --git a/Engine/Source/Editor/Persona/Private/STrack.h b/Engine/Source/Editor/Persona/Private/STrack.h
--- a/Engine/Source/Editor/Persona/Private/STrack.h
+++ b/Engine/Source/Editor/Persona/Private/STrack.h
@@ -6,7 +6,7 @@
 
 //////////////////////////////////////////////////////////////////////////
 DECLARE_DELEGATE_RetVal( float, FOnGetScrubValue )
-DECLARE_DELEGATE_OneParam( FOnSelectionChanged, const FGraphPanelSelectionSet& )
+DECLARE_DELEGATE_OneParam( FOnSelectionChanged, const TArray<UObject*>& )
 DECLARE_DELEGATE( FOnNodeSelectionChanged )
 DECLARE_DELEGATE( FOnUpdatePanel )
 
diff --git a/Engine/Source/Editor/Persona/Private/Shared/BoneSelectionWidget.cpp b/Engine/Source/Editor/Persona/Private/Shared/BoneSelectionWidget.cpp
--- a/Engine/Source/Editor/Persona/Private/Shared/BoneSelectionWidget.cpp
+++ b/Engine/Source/Editor/Persona/Private/Shared/BoneSelectionWidget.cpp
@@ -5,13 +5,14 @@
 #include "BoneSelectionWidget.h"
 #include "SSearchBox.h"
 #include "Editor/PropertyEditor/Public/DetailLayoutBuilder.h"
+#include "IEditableSkeleton.h"
 
 #define LOCTEXT_NAMESPACE "SBoneSelectionWidget"
 
 /////////////////////////////////////////////////////
-void SBoneTreeMenu::Construct(const FArguments& InArgs, TWeakObjectPtr<const USkeleton> Skeleton)
+void SBoneTreeMenu::Construct(const FArguments& InArgs, const TSharedRef<class IEditableSkeleton>& InEditableSkeleton)
 {
-	TargetSkeleton = Skeleton;
+	EditableSkeletonPtr = InEditableSkeleton;
 	OnSelectionChangedDelegate = InArgs._OnBoneSelectionChanged;
 
 	FText TitleToUse = !InArgs._Title.IsEmpty() ? InArgs._Title  : LOCTEXT("BonePickerTitle", "Pick Bone...");
@@ -105,12 +106,8 @@ void SBoneTreeMenu::RebuildBoneList()
 {
 	SkeletonTreeInfo.Empty();
 	SkeletonTreeInfoFlat.Empty();
-	if (!TargetSkeleton.IsValid())
-	{
-		return;
-	}
 
-	const FReferenceSkeleton& RefSkeleton = TargetSkeleton->GetReferenceSkeleton();
+	const FReferenceSkeleton& RefSkeleton = EditableSkeletonPtr.Pin()->GetSkeleton().GetReferenceSkeleton();
 	for(int32 BoneIdx = 0; BoneIdx < RefSkeleton.GetNum(); ++BoneIdx)
 	{
 		TSharedRef<FBoneNameInfo> BoneInfo = MakeShareable(new FBoneNameInfo(RefSkeleton.GetBoneName(BoneIdx)));
@@ -163,11 +160,9 @@ void SBoneTreeMenu::RebuildBoneList()
 
 /////////////////////////////////////////////////////
 
-void SBoneSelectionWidget::Construct(const FArguments& InArgs, TWeakObjectPtr<const USkeleton> Skeleton)
+void SBoneSelectionWidget::Construct(const FArguments& InArgs, const TSharedRef<class IEditableSkeleton>& InEditableSkeleton)
 {
-	TargetSkeleton = Skeleton;
-
-	check(TargetSkeleton.IsValid());
+	EditableSkeletonPtr = InEditableSkeleton;
 
 	OnBoneSelectionChanged = InArgs._OnBoneSelectionChanged;
 	OnGetSelectedBone = InArgs._OnGetSelectedBone;
@@ -191,7 +186,7 @@ void SBoneSelectionWidget::Construct(const FArguments& InArgs, TWeakObjectPtr<co
 
 TSharedRef<SWidget> SBoneSelectionWidget::CreateSkeletonWidgetMenu()
 {
-	TSharedRef<SBoneTreeMenu> MenuWidget = SNew(SBoneTreeMenu, TargetSkeleton)
+	TSharedRef<SBoneTreeMenu> MenuWidget = SNew(SBoneTreeMenu, EditableSkeletonPtr.Pin().ToSharedRef())
 									.OnBoneSelectionChanged(this, &SBoneSelectionWidget::OnSelectionChanged);
 
 	BonePickerButton->SetMenuContentWidgetToFocus(MenuWidget->FilterTextWidget);
diff --git a/Engine/Source/Editor/Persona/Private/Shared/PersonaMode.cpp b/Engine/Source/Editor/Persona/Private/Shared/PersonaMode.cpp
--- a/Engine/Source/Editor/Persona/Private/Shared/PersonaMode.cpp
+++ b/Engine/Source/Editor/Persona/Private/Shared/PersonaMode.cpp
@@ -10,6 +10,11 @@
 #include "SSkeletonBlendProfiles.h"
 #include "SDockTab.h"
 #include "SAdvancedPreviewDetailsTab.h"
+#include "ISkeletonEditorModule.h"
+#include "SPersonaDetails.h"
+#include "Editor/KismetWidgets/Public/SSingleObjectDetailsPanel.h"
+#include "Kismet2/BlueprintEditorUtils.h"
+#include "PersonaUtils.h"
 
 #define LOCTEXT_NAMESPACE "PersonaModes"
 
@@ -46,6 +51,8 @@ const FName FPersonaTabs::BlendProfileManagerID("BlendProfileManager");
 
 const FName FPersonaTabs::AdvancedPreviewSceneSettingsID("AdvancedPreviewTab");
 
+const FName FPersonaTabs::DetailsID("DetailsTab");
+
 /////////////////////////////////////////////////////
 // FPersonaMode
 
@@ -65,15 +72,15 @@ FPersonaAppMode::FPersonaAppMode(TSharedPtr<class FPersona> InPersona, FName InM
 	MyPersona = InPersona;
 
 	PersonaTabFactories.RegisterFactory(MakeShareable(new FSkeletonTreeSummoner(InPersona)));
-	PersonaTabFactories.RegisterFactory(MakeShareable(new FAnimationAssetBrowserSummoner(InPersona)));
-	PersonaTabFactories.RegisterFactory(MakeShareable(new FPreviewViewportSummoner(InPersona)));
-	PersonaTabFactories.RegisterFactory(MakeShareable(new FMorphTargetTabSummoner(InPersona)));
-	PersonaTabFactories.RegisterFactory(MakeShareable(new FAnimCurveViewerTabSummoner(InPersona)));
-	PersonaTabFactories.RegisterFactory(MakeShareable(new FSkeletonAnimNotifiesSummoner(InPersona)));
-	PersonaTabFactories.RegisterFactory(MakeShareable(new FRetargetManagerTabSummoner(InPersona)));
-	PersonaTabFactories.RegisterFactory(MakeShareable(new FSkeletonSlotNamesSummoner(InPersona)));
+	PersonaTabFactories.RegisterFactory(MakeShareable(new FAnimationAssetBrowserSummoner(InPersona, InPersona->GetPersonaToolkit(), FOnOpenNewAsset::CreateSP(InPersona.Get(), &FPersona::HandleOpenNewAsset), FOnAnimationSequenceBrowserCreated::CreateSP(InPersona.Get(), &FPersona::HandleAnimationSequenceBrowserCreated), true)));
+	PersonaTabFactories.RegisterFactory(MakeShareable(new FPreviewViewportSummoner(InPersona, InPersona->GetSkeletonTree(), InPersona->GetPreviewScene(), InPersona->OnPostUndo, InPersona, FOnViewportCreated::CreateSP(InPersona.Get(), &FPersona::HandleViewportCreated), InModeName != FPersonaModes::AnimationEditMode, true)));
+	PersonaTabFactories.RegisterFactory(MakeShareable(new FMorphTargetTabSummoner(InPersona, InPersona->GetPreviewScene(), InPersona->OnPostUndo)));
+	PersonaTabFactories.RegisterFactory(MakeShareable(new FAnimCurveViewerTabSummoner(InPersona, InPersona->GetSkeletonTree()->GetEditableSkeleton(), InPersona->GetPreviewScene(), InPersona->OnCurvesChanged, InPersona->OnPostUndo)));
+	PersonaTabFactories.RegisterFactory(MakeShareable(new FSkeletonAnimNotifiesSummoner(InPersona, InPersona->GetSkeletonTree()->GetEditableSkeleton(), InPersona->OnAnimNotifiesChanged, InPersona->OnPostUndo, FOnObjectsSelected::CreateSP(InPersona.Get(), &FPersona::HandleObjectsSelected))));
+	PersonaTabFactories.RegisterFactory(MakeShareable(new FRetargetManagerTabSummoner(InPersona, InPersona->GetSkeletonTree()->GetEditableSkeleton(), InPersona->GetPreviewScene(), InPersona->OnPostUndo)));
+	PersonaTabFactories.RegisterFactory(MakeShareable(new FSkeletonSlotNamesSummoner(InPersona, InPersona->GetSkeletonTree()->GetEditableSkeleton(), InPersona->OnPostUndo, FOnObjectSelected::CreateSP(InPersona.Get(), &FPersona::HandleObjectSelected))));
 	PersonaTabFactories.RegisterFactory(MakeShareable(new FSkeletonBlendProfilesSummoner(InPersona)));
-	PersonaTabFactories.RegisterFactory(MakeShareable(new FAdvancedPreviewSceneTabSummoner(InPersona)));
+	PersonaTabFactories.RegisterFactory(MakeShareable(new FAdvancedPreviewSceneTabSummoner(InPersona, InPersona->GetPreviewScene())));
 }
 
 void FPersonaAppMode::RegisterTabFactories(TSharedPtr<FTabManager> InTabManager)
@@ -108,15 +115,47 @@ FPersonaModeSharedData::FPersonaModeSharedData()
 	, bShowSockets(false)
 	, bShowBound(false)
 	, ViewportType(0)
-	, PlaybackSpeedMode(0)
+	, PlaybackSpeedMode(EAnimationPlaybackSpeeds::Normal)
 	, LocalAxesMode(0)
 {}
 
+void FPersonaModeSharedData::Save(const TSharedRef<FAnimationViewportClient>& InFromViewport)
+{
+	ViewLocation = InFromViewport->GetViewLocation();
+	ViewRotation = InFromViewport->GetViewRotation();
+	LookAtLocation = InFromViewport->GetLookAtLocation();
+	OrthoZoom = InFromViewport->GetOrthoZoom();
+	bCameraLock = InFromViewport->IsCameraLocked();
+	bCameraFollow = InFromViewport->IsSetCameraFollowChecked();
+	bShowBound = InFromViewport->IsSetShowBoundsChecked();
+	LocalAxesMode = InFromViewport->GetLocalAxesMode();
+	ViewportType = InFromViewport->ViewportType;
+	PlaybackSpeedMode = InFromViewport->GetPlaybackSpeedMode();
+}
+
+void FPersonaModeSharedData::Restore(const TSharedRef<FAnimationViewportClient>& InToViewport)
+{	
+	InToViewport->SetViewportType((ELevelViewportType)ViewportType);
+	InToViewport->SetViewLocation(ViewLocation);
+	InToViewport->SetViewRotation(ViewRotation);
+	InToViewport->SetShowBounds(bShowBound);
+	InToViewport->SetLocalAxesMode((ELocalAxesMode::Type)LocalAxesMode);
+	InToViewport->SetOrthoZoom(OrthoZoom);
+	InToViewport->SetPlaybackSpeedMode(PlaybackSpeedMode);
+	
+	if (bCameraLock)
+	{
+		InToViewport->SetLookAtLocation(LookAtLocation);
+	}
+	else if(bCameraFollow)
+	{
+		InToViewport->SetCameraFollow();
+	}
+}
+
 /////////////////////////////////////////////////////
 // FSkeletonTreeSummoner
 
-#include "SSkeletonTree.h"
-
 FSkeletonTreeSummoner::FSkeletonTreeSummoner(TSharedPtr<class FAssetEditorToolkit> InHostingApp)
 	: FWorkflowTabFactory(FPersonaTabs::SkeletonTreeViewID, InHostingApp)
 {
@@ -132,8 +171,8 @@ FSkeletonTreeSummoner::FSkeletonTreeSummoner(TSharedPtr<class FAssetEditorToolki
 
 TSharedRef<SWidget> FSkeletonTreeSummoner::CreateTabBody(const FWorkflowTabSpawnInfo& Info) const
 {
-	return SNew(SSkeletonTree)
-		.Persona(StaticCastSharedPtr<FPersona>(HostingApp.Pin()));
+	TSharedPtr<FPersona> Persona = StaticCastSharedPtr<FPersona>(HostingApp.Pin());
+	return Persona->GetSkeletonTree();
 }
 
 /////////////////////////////////////////////////////
@@ -141,8 +180,10 @@ TSharedRef<SWidget> FSkeletonTreeSummoner::CreateTabBody(const FWorkflowTabSpawn
 
 #include "SMorphTargetViewer.h"
 
-FMorphTargetTabSummoner::FMorphTargetTabSummoner(TSharedPtr<class FAssetEditorToolkit> InHostingApp)
+FMorphTargetTabSummoner::FMorphTargetTabSummoner(TSharedPtr<class FAssetEditorToolkit> InHostingApp, const TSharedRef<IPersonaPreviewScene>& InPreviewScene, FSimpleMulticastDelegate& InOnPostUndo)
 	: FWorkflowTabFactory(FPersonaTabs::MorphTargetsID, InHostingApp)
+	, PreviewScene(InPreviewScene)
+	, OnPostUndo(InOnPostUndo)
 {
 	TabLabel = LOCTEXT("MorphTargetTabTitle", "Morph Target Previewer");
 	TabIcon = FSlateIcon(FEditorStyle::GetStyleSetName(), "Persona.Tabs.MorphTargetPreviewer");
@@ -156,16 +197,19 @@ FMorphTargetTabSummoner::FMorphTargetTabSummoner(TSharedPtr<class FAssetEditorTo
 
 TSharedRef<SWidget> FMorphTargetTabSummoner::CreateTabBody(const FWorkflowTabSpawnInfo& Info) const
 {
-	return SNew(SMorphTargetViewer)
-		.Persona(StaticCastSharedPtr<FPersona>(HostingApp.Pin()));
+	return SNew(SMorphTargetViewer, PreviewScene.Pin().ToSharedRef(), OnPostUndo);
 }
 /////////////////////////////////////////////////////
 // FAnimCurveViewerTabSummoner
 
 #include "SAnimCurveViewer.h"
 
-FAnimCurveViewerTabSummoner::FAnimCurveViewerTabSummoner(TSharedPtr<class FAssetEditorToolkit> InHostingApp)
+FAnimCurveViewerTabSummoner::FAnimCurveViewerTabSummoner(TSharedPtr<class FAssetEditorToolkit> InHostingApp, const TSharedRef<IEditableSkeleton>& InEditableSkeleton, const TSharedRef<IPersonaPreviewScene>& InPreviewScene, FSimpleMulticastDelegate& InOnCurvesChanged, FSimpleMulticastDelegate& InOnPostUndo)
 	: FWorkflowTabFactory(FPersonaTabs::AnimCurveViewID, InHostingApp)
+	, EditableSkeleton(InEditableSkeleton)
+	, PreviewScene(InPreviewScene)
+	, OnCurvesChanged(InOnCurvesChanged)
+	, OnPostUndo(InOnPostUndo)
 {
 	TabLabel = LOCTEXT("AnimCurveViewTabTitle", "Anim Curves");
 	TabIcon = FSlateIcon(FEditorStyle::GetStyleSetName(), "Persona.Tabs.AnimCurvePreviewer");
@@ -179,8 +223,7 @@ FAnimCurveViewerTabSummoner::FAnimCurveViewerTabSummoner(TSharedPtr<class FAsset
 
 TSharedRef<SWidget> FAnimCurveViewerTabSummoner::CreateTabBody(const FWorkflowTabSpawnInfo& Info) const
 {
-	return SNew(SAnimCurveViewer)
-		.Persona(StaticCastSharedPtr<FPersona>(HostingApp.Pin()));
+	return SNew(SAnimCurveViewer, EditableSkeleton.Pin().ToSharedRef(), PreviewScene.Pin().ToSharedRef(), OnCurvesChanged, OnPostUndo);
 }
 
 /////////////////////////////////////////////////////
@@ -188,13 +231,16 @@ TSharedRef<SWidget> FAnimCurveViewerTabSummoner::CreateTabBody(const FWorkflowTa
 
 #include "SAnimationSequenceBrowser.h"
 
-FAnimationAssetBrowserSummoner::FAnimationAssetBrowserSummoner(TSharedPtr<class FAssetEditorToolkit> InHostingApp)
+FAnimationAssetBrowserSummoner::FAnimationAssetBrowserSummoner(TSharedPtr<class FAssetEditorToolkit> InHostingApp, const TSharedRef<IPersonaToolkit>& InPersonaToolkit, FOnOpenNewAsset InOnOpenNewAsset, FOnAnimationSequenceBrowserCreated InOnAnimationSequenceBrowserCreated, bool bInShowHistory)
 	: FWorkflowTabFactory(FPersonaTabs::AssetBrowserID, InHostingApp)
+	, PersonaToolkit(InPersonaToolkit)
+	, OnOpenNewAsset(InOnOpenNewAsset)
+	, OnAnimationSequenceBrowserCreated(InOnAnimationSequenceBrowserCreated)
+	, bShowHistory(bInShowHistory)
 {
 	TabLabel = LOCTEXT("AssetBrowserTabTitle", "Asset Browser");
 	TabIcon = FSlateIcon(FEditorStyle::GetStyleSetName(), "ContentBrowser.TabIcon");
 
-	EnableTabPadding();
 	bIsSingleton = true;
 
 	ViewMenuDescription = LOCTEXT("AssetBrowser", "Asset Browser");
@@ -203,8 +249,13 @@ FAnimationAssetBrowserSummoner::FAnimationAssetBrowserSummoner(TSharedPtr<class
 
 TSharedRef<SWidget> FAnimationAssetBrowserSummoner::CreateTabBody(const FWorkflowTabSpawnInfo& Info) const
 {
-	return SNew(SAnimationSequenceBrowser)
-		.Persona(StaticCastSharedPtr<FPersona>(HostingApp.Pin()));
+	TSharedRef<SAnimationSequenceBrowser> Widget = SNew(SAnimationSequenceBrowser, PersonaToolkit.Pin().ToSharedRef())
+		.OnOpenNewAsset(OnOpenNewAsset)
+		.ShowHistory(bShowHistory);
+
+	OnAnimationSequenceBrowserCreated.ExecuteIfBound(Widget);
+
+	return Widget;
 }
 
 /////////////////////////////////////////////////////
@@ -212,45 +263,35 @@ TSharedRef<SWidget> FAnimationAssetBrowserSummoner::CreateTabBody(const FWorkflo
 
 #include "SAnimationEditorViewport.h"
 
-FPreviewViewportSummoner::FPreviewViewportSummoner(TSharedPtr<class FAssetEditorToolkit> InHostingApp)
+FPreviewViewportSummoner::FPreviewViewportSummoner(TSharedPtr<class FAssetEditorToolkit> InHostingApp, const TSharedRef<ISkeletonTree>& InSkeletonTree, const TSharedRef<IPersonaPreviewScene>& InPreviewScene, FSimpleMulticastDelegate& InOnPostUndo, const TSharedPtr<FBlueprintEditor>& InBlueprintEditor, FOnViewportCreated InOnViewportCreated, bool bInShowTimeline, bool bInShowStats)
 	: FWorkflowTabFactory(FPersonaTabs::PreviewViewportID, InHostingApp)
+	, SkeletonTree(InSkeletonTree)
+	, PreviewScene(InPreviewScene)
+	, OnPostUndo(InOnPostUndo)
+	, BlueprintEditor(InBlueprintEditor)
+	, OnViewportCreated(InOnViewportCreated)
+	, bShowTimeline(bInShowTimeline)
+	, bShowStats(bInShowStats)
 {
 	TabLabel = LOCTEXT("ViewportTabTitle", "Viewport");
 	TabIcon = FSlateIcon(FEditorStyle::GetStyleSetName(), "LevelEditor.Tabs.Viewports");
 
 	bIsSingleton = true;
 
-	EnableTabPadding();
-
 	ViewMenuDescription = LOCTEXT("ViewportView", "Viewport");
 	ViewMenuTooltip = LOCTEXT("ViewportView_ToolTip", "Shows the viewport");
 }
 
 TSharedRef<SWidget> FPreviewViewportSummoner::CreateTabBody(const FWorkflowTabSpawnInfo& Info) const
 {
-	TSharedPtr<FPersona> PersonaPtr = StaticCastSharedPtr<FPersona>(HostingApp.Pin());
-
-	TSharedRef<SAnimationEditorViewportTabBody> NewViewport = SNew(SAnimationEditorViewportTabBody)
-		.Persona(PersonaPtr)
-		.Skeleton(PersonaPtr->GetSkeleton())
-		.AddMetaData<FTagMetaData>(TEXT("Persona.Viewport"));
-
-	//@TODO:MODES:check(!PersonaPtr->Viewport.IsValid());
-	// mode switch data sharing
-	// we saves data from previous viewport
-	// and restore after switched
-	bool bRestoreData = PersonaPtr->Viewport.IsValid();
-	if (bRestoreData)
-	{
-		NewViewport.Get().SaveData(PersonaPtr->Viewport.Pin().Get());
-	}
+	TSharedRef<SAnimationEditorViewportTabBody> NewViewport = SNew(SAnimationEditorViewportTabBody, SkeletonTree.Pin().ToSharedRef(), PreviewScene.Pin().ToSharedRef(), HostingApp.Pin().ToSharedRef(), OnPostUndo)
+		.BlueprintEditor(BlueprintEditor.Pin())
+		.OnInvokeTab(FOnInvokeTab::CreateSP(HostingApp.Pin().Get(), &FAssetEditorToolkit::InvokeTab))
+		.AddMetaData<FTagMetaData>(TEXT("Persona.Viewport"))
+		.ShowTimeline(bShowTimeline)
+		.ShowStats(bShowStats);
 
-	PersonaPtr->SetViewport(NewViewport);
-
-	if (bRestoreData)
-	{
-		NewViewport.Get().RestoreData();
-	}
+	OnViewportCreated.ExecuteIfBound(NewViewport);
 
 	return NewViewport;
 }
@@ -260,8 +301,11 @@ TSharedRef<SWidget> FPreviewViewportSummoner::CreateTabBody(const FWorkflowTabSp
 
 #include "SRetargetManager.h"
 
-FRetargetManagerTabSummoner::FRetargetManagerTabSummoner(TSharedPtr<class FAssetEditorToolkit> InHostingApp)
+FRetargetManagerTabSummoner::FRetargetManagerTabSummoner(TSharedPtr<class FAssetEditorToolkit> InHostingApp, const TSharedRef<IEditableSkeleton>& InEditableSkeleton, const TSharedRef<IPersonaPreviewScene>& InPreviewScene, FSimpleMulticastDelegate& InOnPostUndo)
 	: FWorkflowTabFactory(FPersonaTabs::RetargetManagerID, InHostingApp)
+	, EditableSkeleton(InEditableSkeleton)
+	, PreviewScene(InPreviewScene)
+	, OnPostUndo(InOnPostUndo)
 {
 	TabLabel = LOCTEXT("RetargetManagerTabTitle", "Retarget Manager");
 	TabIcon = FSlateIcon(FEditorStyle::GetStyleSetName(), "Persona.Tabs.RetargetManager");
@@ -275,18 +319,120 @@ FRetargetManagerTabSummoner::FRetargetManagerTabSummoner(TSharedPtr<class FAsset
 
 TSharedRef<SWidget> FRetargetManagerTabSummoner::CreateTabBody(const FWorkflowTabSpawnInfo& Info) const
 {
-	return SNew(SRetargetManager)
-		.Persona(StaticCastSharedPtr<FPersona>(HostingApp.Pin()));
+	return SNew(SRetargetManager, EditableSkeleton.Pin().ToSharedRef(), PreviewScene.Pin().ToSharedRef(), OnPostUndo);
 }
 
 
+/////////////////////////////////////////////////////
+// SPersonaPreviewPropertyEditor
+
+void SPersonaPreviewPropertyEditor::Construct(const FArguments& InArgs, TSharedRef<IPersonaPreviewScene> InPreviewScene)
+{
+	PreviewScene = InPreviewScene;
+	bPropertyEdited = false;
+
+	SSingleObjectDetailsPanel::Construct(SSingleObjectDetailsPanel::FArguments(), /*bAutomaticallyObserveViaGetObjectToObserve*/ true, /*bAllowSearch*/ true);
+
+	PropertyView->SetIsPropertyEditingEnabledDelegate(FIsPropertyEditingEnabled::CreateStatic([] { return !GIntraFrameDebuggingGameThread; }));
+	PropertyView->OnFinishedChangingProperties().Add(FOnFinishedChangingProperties::FDelegate::CreateSP(this, &SPersonaPreviewPropertyEditor::HandlePropertyChanged));
+}
+
+UObject* SPersonaPreviewPropertyEditor::GetObjectToObserve() const
+{
+	if (UDebugSkelMeshComponent* PreviewMeshComponent = PreviewScene.Pin()->GetPreviewMeshComponent())
+	{
+		if (PreviewMeshComponent->GetAnimInstance() != nullptr)
+		{
+			return PreviewMeshComponent->GetAnimInstance();
+		}
+	}
+
+	return nullptr;
+}
+
+TSharedRef<SWidget> SPersonaPreviewPropertyEditor::PopulateSlot(TSharedRef<SWidget> PropertyEditorWidget)
+{
+	return SNew(SVerticalBox)
+		+SVerticalBox::Slot()
+		.FillHeight(1.0f)
+		[
+			PropertyEditorWidget
+		]
+		+SVerticalBox::Slot()
+		.AutoHeight()
+		[
+			SNew(SBorder)
+			.BorderImage(FEditorStyle::GetBrush("Docking.Tab.ContentAreaBrush"))
+			.Visibility_Lambda([this]() { return bPropertyEdited ? EVisibility::Visible : EVisibility::Collapsed; })
+			[
+				SNew(SHorizontalBox)
+				+SHorizontalBox::Slot()
+				.FillWidth(1.0f)
+				.VAlign(VAlign_Center)
+				.Padding(2.0f)
+				[
+					SNew(STextBlock)
+					.Text(LOCTEXT("AnimBlueprintEditPreviewText", "Changes made to preview only. Changes will not be saved!"))
+					.ColorAndOpacity(FLinearColor::Yellow)
+					.ShadowOffset(FVector2D::UnitVector)
+					.AutoWrapText(true)
+				]
+				+SHorizontalBox::Slot()
+				.AutoWidth()
+				.HAlign(HAlign_Right)
+				.VAlign(VAlign_Center)
+				.Padding(2.0f)
+				[
+					SNew(SButton)
+					.OnClicked(this, &SPersonaPreviewPropertyEditor::HandleApplyChanges)
+					.ToolTipText(LOCTEXT("AnimBlueprintEditApplyChanges_Tooltip", "Apply any changes that have been made to the preview to the defaults."))
+					[
+						SNew(STextBlock)
+						.Text(LOCTEXT("AnimBlueprintEditApplyChanges", "Apply"))
+					]
+				]
+			]
+		];
+}
+
+void SPersonaPreviewPropertyEditor::HandlePropertyChanged(const FPropertyChangedEvent& InPropertyChangedEvent)
+{
+	if (UDebugSkelMeshComponent* PreviewMeshComponent = PreviewScene.Pin()->GetPreviewMeshComponent())
+	{
+		if (UAnimInstance* AnimInstance = PreviewMeshComponent->GetAnimInstance())
+		{
+			// check to see how many properties have changed
+			const int32 NumChangedProperties = PersonaUtils::CopyPropertiesToCDO(AnimInstance, PersonaUtils::FCopyOptions(PersonaUtils::ECopyOptions::PreviewOnly));
+			bPropertyEdited = (NumChangedProperties > 0);
+		}
+	}
+}
+
+FReply SPersonaPreviewPropertyEditor::HandleApplyChanges()
+{
+	// copy preview properties into CDO
+	if (UDebugSkelMeshComponent* PreviewMeshComponent = PreviewScene.Pin()->GetPreviewMeshComponent())
+	{
+		if (UAnimInstance* AnimInstance = PreviewMeshComponent->GetAnimInstance())
+		{
+			PersonaUtils::CopyPropertiesToCDO(AnimInstance);
+
+			bPropertyEdited = false;
+		}
+	}
+
+	return FReply::Handled();
+}
+
 /////////////////////////////////////////////////////
 // FAnimBlueprintPreviewEditorSummoner
 
 #include "SKismetInspector.h"
 
-FAnimBlueprintPreviewEditorSummoner::FAnimBlueprintPreviewEditorSummoner(TSharedPtr<class FAssetEditorToolkit> InHostingApp)
-	: FWorkflowTabFactory(FPersonaTabs::AnimBlueprintPreviewEditorID, InHostingApp)
+FAnimBlueprintPreviewEditorSummoner::FAnimBlueprintPreviewEditorSummoner(TSharedPtr<class FBlueprintEditor> InBlueprintEditor, const TSharedRef<class IPersonaPreviewScene>& InPreviewScene)
+	: FWorkflowTabFactory(FPersonaTabs::AnimBlueprintPreviewEditorID, InBlueprintEditor)
+	, BlueprintEditor(InBlueprintEditor)
+	, PreviewScene(InPreviewScene)
 {
 	TabLabel = LOCTEXT("AnimBlueprintPreviewTabTitle", "Anim Preview Editor");
 
@@ -300,8 +446,6 @@ FAnimBlueprintPreviewEditorSummoner::FAnimBlueprintPreviewEditorSummoner(TShared
 
 TSharedRef<SWidget> FAnimBlueprintPreviewEditorSummoner::CreateTabBody(const FWorkflowTabSpawnInfo& Info) const
 {
-	TSharedPtr<FPersona> PersonaPtr = StaticCastSharedPtr<FPersona>(HostingApp.Pin());
-
 	return	SNew(SVerticalBox)
 			+SVerticalBox::Slot()
 			.AutoHeight()
@@ -362,17 +506,17 @@ TSharedRef<SWidget> FAnimBlueprintPreviewEditorSummoner::CreateTabBody(const FWo
 					.BorderImage( FEditorStyle::GetBrush("NoBorder") )
 					.Visibility(this, &FAnimBlueprintPreviewEditorSummoner::IsEditorVisible, EAnimBlueprintEditorMode::PreviewMode)
 					[
-						PersonaPtr->GetPreviewEditor()
+						SNew(SPersonaPreviewPropertyEditor, PreviewScene.Pin().ToSharedRef())
 					]
 				]
 				+SOverlay::Slot()
 				[
 					SNew(SBorder)
-					.Padding(0)
+					.Padding(FMargin(3.0f, 2.0f))
 					.BorderImage( FEditorStyle::GetBrush("NoBorder") )
 					.Visibility(this, &FAnimBlueprintPreviewEditorSummoner::IsEditorVisible, EAnimBlueprintEditorMode::DefaultsMode)
 					[
-						PersonaPtr->GetDefaultEditor()
+						BlueprintEditor.Pin()->GetDefaultEditor()
 					]
 				]
 			];
@@ -416,8 +560,10 @@ void FAnimBlueprintPreviewEditorSummoner::OnCheckedChanged(ECheckBoxState NewTyp
 //////////////////////////////////////////////////////////////////////////
 // FAnimBlueprintParentPlayerEditorSummoner
 
-FAnimBlueprintParentPlayerEditorSummoner::FAnimBlueprintParentPlayerEditorSummoner(TSharedPtr<class FAssetEditorToolkit> InHostingApp) 
-: FWorkflowTabFactory(FPersonaTabs::AnimBlueprintParentPlayerEditorID, InHostingApp)
+FAnimBlueprintParentPlayerEditorSummoner::FAnimBlueprintParentPlayerEditorSummoner(TSharedPtr<class FBlueprintEditor> InBlueprintEditor, FSimpleMulticastDelegate& InOnPostUndo)
+	: FWorkflowTabFactory(FPersonaTabs::AnimBlueprintParentPlayerEditorID, InBlueprintEditor)
+	, BlueprintEditor(InBlueprintEditor)
+	, OnPostUndo(InOnPostUndo)
 {
 	TabLabel = LOCTEXT("ParentPlayerOverrideEditor", "Asset Override Editor");
 	bIsSingleton = true;
@@ -425,8 +571,7 @@ FAnimBlueprintParentPlayerEditorSummoner::FAnimBlueprintParentPlayerEditorSummon
 
 TSharedRef<SWidget> FAnimBlueprintParentPlayerEditorSummoner::CreateTabBody(const FWorkflowTabSpawnInfo& Info) const
 {
-	TWeakPtr<FPersona> PersonaPtr = StaticCastSharedPtr<FPersona>(HostingApp.Pin());
-	return SNew(SAnimBlueprintParentPlayerList).Persona(PersonaPtr);
+	return SNew(SAnimBlueprintParentPlayerList, BlueprintEditor.Pin().ToSharedRef(), OnPostUndo);
 }
 
 FText FAnimBlueprintParentPlayerEditorSummoner::GetTabToolTipText(const FWorkflowTabSpawnInfo& Info) const
@@ -437,8 +582,9 @@ FText FAnimBlueprintParentPlayerEditorSummoner::GetTabToolTipText(const FWorkflo
 /////////////////////////////////////////////////////
 // FAdvancedPreviewSceneTabSummoner
 
-FAdvancedPreviewSceneTabSummoner::FAdvancedPreviewSceneTabSummoner(TSharedPtr<class FAssetEditorToolkit> InHostingApp)
+FAdvancedPreviewSceneTabSummoner::FAdvancedPreviewSceneTabSummoner(TSharedPtr<class FAssetEditorToolkit> InHostingApp, const TSharedRef<IPersonaPreviewScene>& InPreviewScene)
 	: FWorkflowTabFactory(FPersonaTabs::AdvancedPreviewSceneSettingsID, InHostingApp)
+	, PreviewScene(InPreviewScene)
 {
 	TabLabel = LOCTEXT("PreviewSceneSettingsTab", "Preview Scene Settings");
 	TabIcon = FSlateIcon(FEditorStyle::GetStyleSetName(), "LevelEditor.Tabs.Details");	
@@ -450,8 +596,7 @@ FAdvancedPreviewSceneTabSummoner::FAdvancedPreviewSceneTabSummoner(TSharedPtr<cl
 
 TSharedRef<SWidget> FAdvancedPreviewSceneTabSummoner::CreateTabBody(const FWorkflowTabSpawnInfo& Info) const
 {
-	TSharedPtr<FPersona> PersonaPtr = StaticCastSharedPtr<FPersona>(HostingApp.Pin());	
-	return SNew(SAdvancedPreviewDetailsTab).PreviewScenePtr(&(PersonaPtr->GetPreviewScene()));
+	return SNew(SAdvancedPreviewDetailsTab, PreviewScene.Pin().ToSharedRef());
 }
 
 FText FAdvancedPreviewSceneTabSummoner::GetTabToolTipText(const FWorkflowTabSpawnInfo& Info) const
@@ -459,5 +604,35 @@ FText FAdvancedPreviewSceneTabSummoner::GetTabToolTipText(const FWorkflowTabSpaw
 	return LOCTEXT("AdvancedPreviewSettingsToolTip", "The Advanced Preview Settings tab will let you alter the preview scene's settings.");
 }
 
+/////////////////////////////////////////////////////
+// FDetailsTabSummoner
+
+FDetailsTabSummoner::FDetailsTabSummoner(TSharedPtr<class FAssetEditorToolkit> InHostingApp, FOnDetailsCreated InOnDetailsCreated)
+	: FWorkflowTabFactory(FPersonaTabs::DetailsID, InHostingApp)
+	, OnDetailsCreated(InOnDetailsCreated)
+{
+	TabLabel = LOCTEXT("PersonaDetailsTab", "Details");
+	TabIcon = FSlateIcon(FEditorStyle::GetStyleSetName(), "LevelEditor.Tabs.Details");
+	bIsSingleton = true;
+
+	ViewMenuDescription = LOCTEXT("DetailsDescription", "Details");
+	ViewMenuTooltip = LOCTEXT("DetailsToolTip", "Shows the details tab for selected objects.");
+
+	PersonaDetails = SNew(SPersonaDetails);
+
+	OnDetailsCreated.ExecuteIfBound(PersonaDetails->DetailsView.ToSharedRef());
+}
+
+TSharedRef<SWidget> FDetailsTabSummoner::CreateTabBody(const FWorkflowTabSpawnInfo& Info) const
+{
+	return PersonaDetails.ToSharedRef();
+}
+
+FText FDetailsTabSummoner::GetTabToolTipText(const FWorkflowTabSpawnInfo& Info) const
+{
+	return LOCTEXT("PersonaDetailsToolTip", "Edit the details of selected objects.");
+}
+
+
 #undef LOCTEXT_NAMESPACE
 
diff --git a/Engine/Source/Editor/Persona/Private/Shared/PersonaMode.h b/Engine/Source/Editor/Persona/Private/Shared/PersonaMode.h
--- a/Engine/Source/Editor/Persona/Private/Shared/PersonaMode.h
+++ b/Engine/Source/Editor/Persona/Private/Shared/PersonaMode.h
@@ -3,6 +3,10 @@
 #pragma once
 
 #include "IDocumentation.h"
+#include "AnimationEditorViewportClient.h"
+#include "IPersonaViewport.h"
+#include "PersonaModule.h"
+#include "SSingleObjectDetailsPanel.h"
 
 #define LOCTEXT_NAMESPACE "PersonaMode"
 
@@ -44,6 +48,8 @@ struct FPersonaTabs
 	// Advanced Preview Scene
 	static const FName AdvancedPreviewSceneSettingsID;
 
+	static const FName DetailsID;
+
 	// Blueprint Document
 
 	// Inherited from blueprint editor
@@ -95,10 +101,13 @@ private:
 /////////////////////////////////////////////////////
 // FPersonaModeSharedData
 
-struct FPersonaModeSharedData 
+struct FPersonaModeSharedData : public IPersonaViewportState
 {
 	FPersonaModeSharedData();
 
+	void Save(const TSharedRef<FAnimationViewportClient>& InFromViewport);
+	void Restore(const TSharedRef<FAnimationViewportClient>& InToViewport);
+
 	// camera setup
 	FVector				ViewLocation;
 	FRotator			ViewRotation;
@@ -119,11 +128,8 @@ struct FPersonaModeSharedData
 
 	// viewport setup
 	int32				ViewportType;
-	int32				PlaybackSpeedMode;
+	EAnimationPlaybackSpeeds::Type PlaybackSpeedMode;
 	int32				LocalAxesMode;
-
-	// Playback state
-	FSingleAnimationPlayData PlaybackData;
 };
 
 /////////////////////////////////////////////////////
@@ -170,7 +176,7 @@ public:
 struct FMorphTargetTabSummoner : public FWorkflowTabFactory
 {
 public:
-	FMorphTargetTabSummoner(TSharedPtr<class FAssetEditorToolkit> InHostingApp);
+	FMorphTargetTabSummoner(TSharedPtr<class FAssetEditorToolkit> InHostingApp, const TSharedRef<IPersonaPreviewScene>& InPreviewScene, FSimpleMulticastDelegate& InOnPostUndo);
 
 	virtual TSharedRef<SWidget> CreateTabBody(const FWorkflowTabSpawnInfo& Info) const override;
 
@@ -179,14 +185,19 @@ public:
 	{
 		return  IDocumentation::Get()->CreateToolTip(LOCTEXT("MorphTargetTooltip", "The Morph Target tab lets you preview any morph targets (aka blend shapes) available for the current mesh."), NULL, TEXT("Shared/Editors/Persona"), TEXT("MorphTarget_Window"));
 	}
+
+private:
+	TWeakPtr<class IPersonaPreviewScene> PreviewScene;
+	FSimpleMulticastDelegate& OnPostUndo;
 };
+
 /////////////////////////////////////////////////////
 // FAnimCurveViewerTabSummoner
 
 struct FAnimCurveViewerTabSummoner : public FWorkflowTabFactory
 {
 public:
-	FAnimCurveViewerTabSummoner(TSharedPtr<class FAssetEditorToolkit> InHostingApp);
+	FAnimCurveViewerTabSummoner(TSharedPtr<class FAssetEditorToolkit> InHostingApp, const TSharedRef<IEditableSkeleton>& InEditableSkeleton, const TSharedRef<IPersonaPreviewScene>& InPreviewScene, FSimpleMulticastDelegate& InOnCurvesChanged, FSimpleMulticastDelegate& InOnPostUndo);
 
 	virtual TSharedRef<SWidget> CreateTabBody(const FWorkflowTabSpawnInfo& Info) const override;
 
@@ -195,6 +206,12 @@ public:
 	{
 		return  IDocumentation::Get()->CreateToolTip(LOCTEXT("AnimCurveViewTooltip", "The Anim Curve Viewer tab lets you preview any animation curves available for the current mesh from preview asset."), NULL, TEXT("Shared/Editors/Persona"), TEXT("AnimCurveView_Window"));
 	}
+
+private:
+	TWeakPtr<class IEditableSkeleton> EditableSkeleton;
+	TWeakPtr<class IPersonaPreviewScene> PreviewScene;
+	FSimpleMulticastDelegate& OnCurvesChanged;
+	FSimpleMulticastDelegate& OnPostUndo;
 };
 
 
@@ -203,7 +220,7 @@ public:
 
 struct FAnimationAssetBrowserSummoner : public FWorkflowTabFactory
 {
-	FAnimationAssetBrowserSummoner(TSharedPtr<class FAssetEditorToolkit> InHostingApp);
+	FAnimationAssetBrowserSummoner(TSharedPtr<class FAssetEditorToolkit> InHostingApp, const TSharedRef<IPersonaToolkit>& InPersonaToolkit, FOnOpenNewAsset InOnOpenNewAsset, FOnAnimationSequenceBrowserCreated InOnAnimationSequenceBrowserCreated, bool bInShowHistory);
 
 	virtual TSharedRef<SWidget> CreateTabBody(const FWorkflowTabSpawnInfo& Info) const override;
 
@@ -212,6 +229,12 @@ struct FAnimationAssetBrowserSummoner : public FWorkflowTabFactory
 	{
 		return  IDocumentation::Get()->CreateToolTip(LOCTEXT("AnimAssetBrowserTooltip", "The Asset Browser lets you browse all animation-related assets (animations, blend spaces etc)."), NULL, TEXT("Shared/Editors/Persona"), TEXT("AssetBrowser_Window"));
 	}
+
+private:
+	TWeakPtr<class IPersonaToolkit> PersonaToolkit;
+	FOnOpenNewAsset OnOpenNewAsset;
+	FOnAnimationSequenceBrowserCreated OnAnimationSequenceBrowserCreated;
+	bool bShowHistory;
 };
 
 /////////////////////////////////////////////////////
@@ -219,9 +242,17 @@ struct FAnimationAssetBrowserSummoner : public FWorkflowTabFactory
 
 struct FPreviewViewportSummoner : public FWorkflowTabFactory
 {
-	FPreviewViewportSummoner(TSharedPtr<class FAssetEditorToolkit> InHostingApp);
+	FPreviewViewportSummoner(TSharedPtr<class FAssetEditorToolkit> InHostingApp, const TSharedRef<ISkeletonTree>& InSkeletonTree, const TSharedRef<IPersonaPreviewScene>& InPreviewScene, FSimpleMulticastDelegate& InOnPostUndo, const TSharedPtr<FBlueprintEditor>& InBlueprintEditor, FOnViewportCreated InOnViewportCreated, bool bShowTimeline, bool bInShowStats);
 
 	virtual TSharedRef<SWidget> CreateTabBody(const FWorkflowTabSpawnInfo& Info) const override;
+
+	TWeakPtr<ISkeletonTree> SkeletonTree;
+	TWeakPtr<IPersonaPreviewScene> PreviewScene;
+	FSimpleMulticastDelegate& OnPostUndo;
+	TWeakPtr<FBlueprintEditor> BlueprintEditor;
+	FOnViewportCreated OnViewportCreated;
+	bool bShowTimeline;
+	bool bShowStats;
 };
 
 /////////////////////////////////////////////////////
@@ -230,7 +261,7 @@ struct FPreviewViewportSummoner : public FWorkflowTabFactory
 struct FRetargetManagerTabSummoner : public FWorkflowTabFactory
 {
 public:
-	FRetargetManagerTabSummoner(TSharedPtr<class FAssetEditorToolkit> InHostingApp);
+	FRetargetManagerTabSummoner(TSharedPtr<class FAssetEditorToolkit> InHostingApp, const TSharedRef<IEditableSkeleton>& InEditableSkeleton, const TSharedRef<IPersonaPreviewScene>& InPreviewScene, FSimpleMulticastDelegate& InOnPostUndo);
 
 	virtual TSharedRef<SWidget> CreateTabBody(const FWorkflowTabSpawnInfo& Info) const override;
 
@@ -239,6 +270,40 @@ public:
 	{
 		return  IDocumentation::Get()->CreateToolTip(LOCTEXT("RetargetSourceTooltip", "In this panel, you can manage retarget sources for different body types"), NULL, TEXT("Shared/Editors/Persona"), TEXT("RetargetManager"));
 	}
+
+private:
+	TWeakPtr<class IEditableSkeleton> EditableSkeleton;
+	TWeakPtr<class IPersonaPreviewScene> PreviewScene;
+	FSimpleMulticastDelegate& OnPostUndo;
+};
+
+/////////////////////////////////////////////////////
+// SPersonaPreviewPropertyEditor
+
+class SPersonaPreviewPropertyEditor : public SSingleObjectDetailsPanel
+{
+public:
+	SLATE_BEGIN_ARGS(SPersonaPreviewPropertyEditor) {}
+	SLATE_END_ARGS()
+
+private:
+	// Pointer to preview scene
+	TWeakPtr<IPersonaPreviewScene> PreviewScene;
+
+public:
+	void Construct(const FArguments& InArgs, TSharedRef<IPersonaPreviewScene> InPreviewScene);
+
+	// SSingleObjectDetailsPanel interface
+	virtual UObject* GetObjectToObserve() const override;
+	virtual TSharedRef<SWidget> PopulateSlot(TSharedRef<SWidget> PropertyEditorWidget) override;
+	// End of SSingleObjectDetailsPanel interface
+
+private:
+	void HandlePropertyChanged(const FPropertyChangedEvent& InPropertyChangedEvent);
+	FReply HandleApplyChanges();
+
+private:
+	bool bPropertyEdited;
 };
 
 /////////////////////////////////////////////////////
@@ -256,7 +321,7 @@ namespace EAnimBlueprintEditorMode
 struct FAnimBlueprintPreviewEditorSummoner : public FWorkflowTabFactory
 {
 public:
-	FAnimBlueprintPreviewEditorSummoner(TSharedPtr<class FAssetEditorToolkit> InHostingApp);
+	FAnimBlueprintPreviewEditorSummoner(TSharedPtr<class FBlueprintEditor> InBlueprintEditor, const TSharedRef<class IPersonaPreviewScene>& InPreviewScene);
 
 	virtual TSharedRef<SWidget> CreateTabBody(const FWorkflowTabSpawnInfo& Info) const override;
 
@@ -271,6 +336,9 @@ private:
 	/** Handle changing of editor mode */
 	void OnCheckedChanged(ECheckBoxState NewType, EAnimBlueprintEditorMode::Type Mode);
 	EAnimBlueprintEditorMode::Type CurrentMode;
+
+	TWeakPtr<class FBlueprintEditor> BlueprintEditor;
+	TWeakPtr<class IPersonaPreviewScene> PreviewScene;
 };
 
 //////////////////////////////////////////////////////////////////////////
@@ -278,13 +346,14 @@ private:
 class FAnimBlueprintParentPlayerEditorSummoner : public FWorkflowTabFactory
 {
 public:
-	FAnimBlueprintParentPlayerEditorSummoner(TSharedPtr<class FAssetEditorToolkit> InHostingApp);
+	FAnimBlueprintParentPlayerEditorSummoner(TSharedPtr<class FBlueprintEditor> InBlueprintEditor, FSimpleMulticastDelegate& InOnPostUndo);
 
 	virtual TSharedRef<SWidget> CreateTabBody(const FWorkflowTabSpawnInfo& Info) const override;
 	virtual FText GetTabToolTipText(const FWorkflowTabSpawnInfo& Info) const override;
 
 private:
-
+	TWeakPtr<class FBlueprintEditor> BlueprintEditor;
+	FSimpleMulticastDelegate& OnPostUndo;
 };
 
 /////////////////////////////////////////////////////
@@ -293,9 +362,27 @@ private:
 struct FAdvancedPreviewSceneTabSummoner : public FWorkflowTabFactory
 {
 public:
- 	FAdvancedPreviewSceneTabSummoner(TSharedPtr<class FAssetEditorToolkit> InHostingApp);
+ 	FAdvancedPreviewSceneTabSummoner(TSharedPtr<class FAssetEditorToolkit> InHostingApp, const TSharedRef<class IPersonaPreviewScene>& InPreviewScene);
 	virtual TSharedRef<SWidget> CreateTabBody(const FWorkflowTabSpawnInfo& Info) const override;
 	virtual FText GetTabToolTipText(const FWorkflowTabSpawnInfo& Info) const override;
+
+private:
+	TWeakPtr<class IPersonaPreviewScene> PreviewScene;
+};
+
+/////////////////////////////////////////////////////
+// FDetailsTabSummoner
+
+struct FDetailsTabSummoner : public FWorkflowTabFactory
+{
+public:
+	FDetailsTabSummoner(TSharedPtr<class FAssetEditorToolkit> InHostingApp, FOnDetailsCreated InOnDetailsCreated);
+	virtual TSharedRef<SWidget> CreateTabBody(const FWorkflowTabSpawnInfo& Info) const override;
+	virtual FText GetTabToolTipText(const FWorkflowTabSpawnInfo& Info) const override;
+
+private:
+	FOnDetailsCreated OnDetailsCreated;
+	TSharedPtr<class SPersonaDetails> PersonaDetails;
 };
 
 #undef LOCTEXT_NAMESPACE
\ No newline at end of file
