Copying //UE4/Dev-Framework to //UE4/Dev-Main (Source: //UE4/Dev-Framework @ 3136612)

diff --git a/Engine/Source/Editor/Persona/Private/SkeletalMeshSocketDetails.cpp b/Engine/Source/Editor/Persona/Private/SkeletalMeshSocketDetails.cpp
--- a/Engine/Source/Editor/Persona/Private/SkeletalMeshSocketDetails.cpp
+++ b/Engine/Source/Editor/Persona/Private/SkeletalMeshSocketDetails.cpp
@@ -12,6 +12,9 @@
 #include "Persona.h"
 #include "AssetSearchBoxUtilPersona.h"
 #include "Engine/SkeletalMeshSocket.h"
+#include "ISkeletonTree.h"
+#include "IEditableSkeleton.h"
+#include "ISkeletonEditorModule.h"
 
 TSharedRef<IDetailCustomization> FSkeletalMeshSocketDetails::MakeInstance()
 {
@@ -22,7 +25,7 @@ void FSkeletalMeshSocketDetails::CustomizeDetails( IDetailLayoutBuilder& DetailB
 {
 	TargetSocket = NULL;
 
-	IDetailCategoryBuilder& SocketCategory = DetailBuilder.EditCategory( TEXT("SkeletalMeshSocket"), FText::GetEmpty(), ECategoryPriority::TypeSpecific );
+	IDetailCategoryBuilder& SocketCategory = DetailBuilder.EditCategory( TEXT("Socket Parameters"), FText::GetEmpty(), ECategoryPriority::TypeSpecific );
 
 	SocketNameProperty = DetailBuilder.GetProperty( TEXT("SocketName") );
 	if( SocketNameProperty.IsValid() && SocketNameProperty->GetProperty() )
@@ -119,13 +122,10 @@ void FSkeletalMeshSocketDetails::OnParentBoneNameCommitted(const FText& InSearch
 	{
 		if( Skeleton->GetReferenceSkeleton().FindBoneIndex( *InSearchText.ToString() ) != INDEX_NONE )
 		{
-			IAssetEditorInstance* EditorInstance = FAssetEditorManager::Get().FindEditorForAsset(Outer, false);
-			if( EditorInstance->GetEditorName() == TEXT("Persona") )
-			{
-				FPersona* PersonaInstance = static_cast<FPersona*> (EditorInstance);
-				PersonaInstance->ChangeSocketParent( TargetSocket, *InSearchText.ToString() );
-				PersonaInstance->SetSelectedSocket( FSelectedSocketInfo( TargetSocket, true ) );
-			}
+			ISkeletonEditorModule& SkeletonEditorModule = FModuleManager::LoadModuleChecked<ISkeletonEditorModule>("SkeletonEditor");
+			TSharedRef<IEditableSkeleton> EditableSkeleton = SkeletonEditorModule.CreateEditableSkeleton(Skeleton);
+
+			EditableSkeleton->SetSocketParent( TargetSocket->SocketName, *InSearchText.ToString() );
 
 			ParentBoneProperty->SetValue( InSearchText.ToString() );
 		}
@@ -134,52 +134,46 @@ void FSkeletalMeshSocketDetails::OnParentBoneNameCommitted(const FText& InSearch
 
 void FSkeletalMeshSocketDetails::OnSocketNameChanged(const FText& InSearchText)
 {
-	if( UObject* Outer = TargetSocket->GetOuter() )
+	if(USkeleton* Skeleton = Cast<USkeleton>(TargetSocket->GetOuter()))
 	{
-		IAssetEditorInstance* EditorInstance = FAssetEditorManager::Get().FindEditorForAsset(Outer, false);
-		if( EditorInstance->GetEditorName() == TEXT("Persona") )
+		ISkeletonEditorModule& SkeletonEditorModule = FModuleManager::LoadModuleChecked<ISkeletonEditorModule>("SkeletonEditor");
+		TSharedRef<IEditableSkeleton> EditableSkeleton = SkeletonEditorModule.CreateEditableSkeleton(Skeleton);
+
+		FText OutErrorMessage;
+		if( VerifySocketName(EditableSkeleton, TargetSocket, InSearchText, OutErrorMessage) )
 		{
-			FPersona* PersonaInstance = static_cast<FPersona*> (EditorInstance);
-			FText OutErrorMessage;
-			if( VerifySocketName(PersonaInstance, TargetSocket, InSearchText, OutErrorMessage) )
-			{
-				SocketNameTextBox->SetError( FText::GetEmpty() );
-			}
-			else
-			{
-				SocketNameTextBox->SetError( OutErrorMessage );
-			}
+			SocketNameTextBox->SetError( FText::GetEmpty() );
+		}
+		else
+		{
+			SocketNameTextBox->SetError( OutErrorMessage );
 		}
 	}
 }
 
 void FSkeletalMeshSocketDetails::OnSocketNameCommitted(const FText& InSearchText, ETextCommit::Type CommitInfo)
 {
-	if( UObject* Outer = TargetSocket->GetOuter() )
+	if (USkeleton* Skeleton = Cast<USkeleton>(TargetSocket->GetOuter()))
 	{
-		IAssetEditorInstance* EditorInstance = FAssetEditorManager::Get().FindEditorForAsset(Outer, false);
-		if( EditorInstance->GetEditorName() == TEXT("Persona") )
-		{
-			FPersona* PersonaInstance = static_cast<FPersona*> (EditorInstance);
+		ISkeletonEditorModule& SkeletonEditorModule = FModuleManager::LoadModuleChecked<ISkeletonEditorModule>("SkeletonEditor");
+		TSharedRef<IEditableSkeleton> EditableSkeleton = SkeletonEditorModule.CreateEditableSkeleton(Skeleton);
 
-			FText NewText = FText::TrimPrecedingAndTrailing(InSearchText);
+		FText NewText = FText::TrimPrecedingAndTrailing(InSearchText);
 
-			FText OutErrorMessage;
-			if (VerifySocketName(PersonaInstance, TargetSocket, NewText, OutErrorMessage))
-			{
-				// tell persona to rename the socket.
-				PersonaInstance->RenameSocket(TargetSocket, *NewText.ToString());
-				PersonaInstance->SetSelectedSocket( FSelectedSocketInfo( TargetSocket, true ) );
+		FText OutErrorMessage;
+		if (VerifySocketName(EditableSkeleton, TargetSocket, NewText, OutErrorMessage))
+		{
+			// tell rename the socket.
+			EditableSkeleton->RenameSocket(TargetSocket->SocketName, FName(*NewText.ToString()));
 				
-				// update the pre-edit socket name and the property
-				PreEditSocketName = NewText;
-				SocketNameTextBox->SetText(NewText);
-			}
-			else
-			{
-				// restore the pre-edit name to the socket text box.
-				SocketNameTextBox->SetText(PreEditSocketName);
-			}
+			// update the pre-edit socket name and the property
+			PreEditSocketName = NewText;
+			SocketNameTextBox->SetText(NewText);
+		}
+		else
+		{
+			// restore the pre-edit name to the socket text box.
+			SocketNameTextBox->SetText(PreEditSocketName);
 		}
 	}
 
@@ -189,7 +183,7 @@ void FSkeletalMeshSocketDetails::OnSocketNameCommitted(const FText& InSearchText
 	}
 }
 
-bool FSkeletalMeshSocketDetails::VerifySocketName( FPersona* PersonaPtr, const USkeletalMeshSocket* Socket, const FText& InText, FText& OutErrorMessage ) const
+bool FSkeletalMeshSocketDetails::VerifySocketName(TSharedRef<IEditableSkeleton> EditableSkeleton, const USkeletalMeshSocket* Socket, const FText& InText, FText& OutErrorMessage ) const
 {
 	// You can't have two sockets with the same name on the mesh, nor on the skeleton,
 	// but you can have a socket with the same name on the mesh *and* the skeleton.
@@ -204,12 +198,12 @@ bool FSkeletalMeshSocketDetails::VerifySocketName( FPersona* PersonaPtr, const U
 	}
 	else
 	{
-		if (PersonaPtr->DoesSocketAlreadyExist(Socket, NewText, PersonaPtr->GetSkeleton()->Sockets))
+		if (EditableSkeleton->DoesSocketAlreadyExist(Socket, NewText, ESocketParentType::Skeleton))
 		{
 			bVerifyName = false;
 		}
 
-		if (bVerifyName && PersonaPtr->DoesSocketAlreadyExist(Socket, NewText, PersonaPtr->GetMesh()->GetMeshOnlySocketList()))
+		if (bVerifyName && EditableSkeleton->DoesSocketAlreadyExist(Socket, NewText, ESocketParentType::Mesh))
 		{
 			bVerifyName = false;
 		}
diff --git a/Engine/Source/Editor/Persona/Private/SkeletalMeshSocketDetails.h b/Engine/Source/Editor/Persona/Private/SkeletalMeshSocketDetails.h
--- a/Engine/Source/Editor/Persona/Private/SkeletalMeshSocketDetails.h
+++ b/Engine/Source/Editor/Persona/Private/SkeletalMeshSocketDetails.h
@@ -25,7 +25,7 @@ private:
 	void OnSocketNameCommitted(const FText& InSearchText, ETextCommit::Type CommitInfo);
 
 	/** Verifies the socket name and supplies and error message if it's invalid **/
-	bool VerifySocketName( FPersona* PersonaPtr, const USkeletalMeshSocket* Socket, const FText& InText, FText& OutErrorMessage ) const;
+	bool VerifySocketName(TSharedRef<class IEditableSkeleton> EditableSkeleton, const USkeletalMeshSocket* Socket, const FText& InText, FText& OutErrorMessage ) const;
 
 	/** Get the search suggestions */
 	TArray<FString> GetSearchSuggestions() const;
diff --git a/Engine/Source/Editor/Persona/Public/BoneDragDropOp.h b/Engine/Source/Editor/Persona/Public/BoneDragDropOp.h
--- a/Engine/Source/Editor/Persona/Public/BoneDragDropOp.h
+++ b/Engine/Source/Editor/Persona/Public/BoneDragDropOp.h
@@ -8,7 +8,7 @@ class FBoneDragDropOp : public FDragDropOperation
 public:	
 	DRAG_DROP_OPERATOR_TYPE(FBoneDragDropOp, FDragDropOperation)
 	
-	USkeleton* TargetSkeleton;
+	TWeakPtr<class IEditableSkeleton> EditableSkeleton;
 	FName BoneName;
 
 	/** The widget decorator to use */
@@ -50,11 +50,11 @@ public:
 		CurrentIconBrush = InIcon;
 	}
 
-	static TSharedRef<FBoneDragDropOp> New(USkeleton* Skeleton, const FName& InBoneName)
+	static TSharedRef<FBoneDragDropOp> New(TSharedRef<class IEditableSkeleton> EditableSkeleton, const FName& InBoneName)
 	{
 		TSharedRef<FBoneDragDropOp> Operation = MakeShareable(new FBoneDragDropOp);
 		Operation->BoneName = InBoneName;
-		Operation->TargetSkeleton = Skeleton;
+		Operation->EditableSkeleton = EditableSkeleton;
 		Operation->SetIcon( FEditorStyle::GetBrush(TEXT("Graph.ConnectorFeedback.Error")) );
 		Operation->Construct();
 		return Operation;
diff --git a/Engine/Source/Editor/Persona/Public/BoneSelectionWidget.h b/Engine/Source/Editor/Persona/Public/BoneSelectionWidget.h
--- a/Engine/Source/Editor/Persona/Public/BoneSelectionWidget.h
+++ b/Engine/Source/Editor/Persona/Public/BoneSelectionWidget.h
@@ -32,7 +32,7 @@ public:
 	*
 	* @param	InArgs	The declaration data for this widget
 	*/
-	void Construct(const FArguments& InArgs, TWeakObjectPtr<const USkeleton> Skeleton);
+		void Construct(const FArguments& InArgs, const TSharedRef<class IEditableSkeleton>& InEditableSkeleton);
 
 	//Filter text widget
 	TSharedPtr<SSearchBox> FilterTextWidget;
@@ -62,7 +62,7 @@ private:
 	FText FilterText;
 
 	// Skeleton to search
-	TWeakObjectPtr<const USkeleton> TargetSkeleton;
+	TWeakPtr<class IEditableSkeleton> EditableSkeletonPtr;
 
 	// Tree view used in the button menu
 	TSharedPtr<STreeView<TSharedPtr<FBoneNameInfo>>> TreeView;
@@ -79,6 +79,7 @@ public:
 		,_OnBoneSelectionChanged()
 		,_OnGetSelectedBone()
 	{}
+
 		/** Set tooltip attribute */
 		SLATE_ARGUMENT(FText, Tooltip);
 
@@ -95,7 +96,7 @@ public:
 	 *
 	 * @param	InArgs	The declaration data for this widget
 	 */
-	void Construct( const FArguments& InArgs, TWeakObjectPtr<const USkeleton> Skeleton);
+	void Construct( const FArguments& InArgs, const TSharedRef<class IEditableSkeleton>& InEditableSkeleton );
 
 private: 
 
@@ -110,7 +111,7 @@ private:
 	TSharedPtr<SComboButton> BonePickerButton;
 
 	// Skeleton to search
-	TWeakObjectPtr<const USkeleton> TargetSkeleton;
+	TWeakPtr<class IEditableSkeleton> EditableSkeletonPtr;
 
 	// delegates
 	FOnBoneSelectionChanged OnBoneSelectionChanged;
diff --git a/Engine/Source/Editor/Persona/Public/IAnimationSequenceBrowser.h b/Engine/Source/Editor/Persona/Public/IAnimationSequenceBrowser.h
new file mode 100644
--- /dev/null
+++ b/Engine/Source/Editor/Persona/Public/IAnimationSequenceBrowser.h
@@ -0,0 +1,12 @@
+// Copyright 1998-2016 Epic Games, Inc. All Rights Reserved.
+
+#pragma once
+
+class IAnimationSequenceBrowser : public SCompoundWidget
+{
+public:
+	/** Select the specified asset */
+	virtual void SelectAsset(UAnimationAsset* InAsset) = 0;
+	/** Add to navigation history */
+	virtual void AddToHistory(UAnimationAsset * AnimAsset) = 0;
+};
\ No newline at end of file
diff --git a/Engine/Source/Editor/Persona/Public/IPersonaEditMode.h b/Engine/Source/Editor/Persona/Public/IPersonaEditMode.h
new file mode 100644
--- /dev/null
+++ b/Engine/Source/Editor/Persona/Public/IPersonaEditMode.h
@@ -0,0 +1,27 @@
+// Copyright 1998-2016 Epic Games, Inc. All Rights Reserved.
+
+#pragma once
+
+#include "EdMode.h"
+
+class IPersonaEditMode : public FEdMode
+{
+public:
+	/** 
+	 * Get a camera target used to focus the viewport on an object when a user presses 'F' (default). 
+	 * @param	OutTarget	The target object
+	 * @return true if the target sphere was filled-in
+	 */
+	virtual bool GetCameraTarget(FSphere& OutTarget) const = 0;
+
+	/** @return the anim preview scene */
+	virtual class IPersonaPreviewScene& GetAnimPreviewScene() const = 0;
+
+	/**
+	* Function to collect strings from nodes to display in the viewport.
+	* Use this rather than DrawHUD when adding general text to the viewport.
+	* Display of this assumes that this will be mostly used by skeletal controls.
+	* @param	OutDebugInfo	Text to display
+	*/
+	virtual void GetOnScreenDebugInfo(TArray<FText>& OutDebugInfo) const = 0;
+};
\ No newline at end of file
diff --git a/Engine/Source/Editor/Persona/Public/IPersonaToolkit.h b/Engine/Source/Editor/Persona/Public/IPersonaToolkit.h
new file mode 100644
--- /dev/null
+++ b/Engine/Source/Editor/Persona/Public/IPersonaToolkit.h
@@ -0,0 +1,34 @@
+// Copyright 1998-2016 Epic Games, Inc. All Rights Reserved.
+
+#pragma once
+
+/**
+ * Reusable 'Persona' features for asset editors concerned with USkeleton-related assets
+ */
+class IPersonaToolkit
+{
+public:
+	/** Get the skeleton that we are editing */
+	virtual class USkeleton* GetSkeleton() const = 0;
+
+	/** Get the preview component that we are using */
+	virtual class UDebugSkelMeshComponent* GetPreviewMeshComponent() const = 0;
+
+	/** Get the skeletal mesh that we are editing */
+	virtual class USkeletalMesh* GetMesh() const = 0;
+
+	/** Set the skeletal mesh we are editing */
+	virtual void SetMesh(class USkeletalMesh* InSkeletalMesh) = 0;
+
+	/** Get the anim blueprint that we are editing */
+	virtual class UAnimBlueprint* GetAnimBlueprint() const = 0;
+
+	/** Get the animation asset that we are editing */
+	virtual class UAnimationAsset* GetAnimationAsset() const = 0;
+
+	/** Set the animation asset we are editing */
+	virtual void SetAnimationAsset(class UAnimationAsset* InAnimationAsset) = 0;
+
+	/** Get the preview scene that we are using */
+	virtual TSharedRef<class IPersonaPreviewScene> GetPreviewScene() const = 0;
+};
\ No newline at end of file
diff --git a/Engine/Source/Editor/Persona/Public/PersonaModule.h b/Engine/Source/Editor/Persona/Public/PersonaModule.h
--- a/Engine/Source/Editor/Persona/Public/PersonaModule.h
+++ b/Engine/Source/Editor/Persona/Public/PersonaModule.h
@@ -1,24 +1,79 @@
 // Copyright 1998-2016 Epic Games, Inc. All Rights Reserved.
 
-
-#ifndef __PersonaModule_h__
-#define __PersonaModule_h__
-
 #pragma once
 
 #include "Toolkits/IToolkit.h"		// For EToolkitMode::Type
 #include "Toolkits/AssetEditorToolkit.h" 
 #include "ModuleInterface.h"
+#include "PersonaDelegates.h"
+#include "Editor.h"
+#include "Factories/FbxImportUI.h"
 
 extern const FName PersonaAppName;
 
+// Editor mode constants
+struct PERSONA_API FPersonaEditModes
+{
+	/** Selection/manipulation of bones & sockets */
+	const static FEditorModeID SkeletonSelection;
+};
+
 DECLARE_DELEGATE_TwoParams(FIsRecordingActive, USkeletalMeshComponent* /*Component*/, bool& /* bIsRecording */);
 DECLARE_DELEGATE_OneParam(FRecord, USkeletalMeshComponent* /*Component*/);
 DECLARE_DELEGATE_OneParam(FStopRecording, USkeletalMeshComponent* /*Component*/);
-DECLARE_DELEGATE_TwoParams(FGetCurrentRecording, USkeletalMeshComponent* /*Component*/, UAnimSequence*& /* OutRecording */);
+DECLARE_DELEGATE_TwoParams(FGetCurrentRecording, USkeletalMeshComponent* /*Component*/, class UAnimSequence*& /* OutRecording */);
 DECLARE_DELEGATE_TwoParams(FGetCurrentRecordingTime, USkeletalMeshComponent* /*Component*/, float& /* OutTime */);
 DECLARE_DELEGATE_TwoParams(FTickRecording, USkeletalMeshComponent* /*Component*/, float /* DeltaSeconds */);
 
+/** Called back when a viewport is created */
+DECLARE_DELEGATE_OneParam(FOnViewportCreated, const TSharedRef<class IPersonaViewport>&);
+
+/** Called back when a details panel is created */
+DECLARE_DELEGATE_OneParam(FOnDetailsCreated, const TSharedRef<class IDetailsView>&);
+
+/** Called back when an anim sequence browser is created */
+DECLARE_DELEGATE_OneParam(FOnAnimationSequenceBrowserCreated, const TSharedRef<class IAnimationSequenceBrowser>&);
+
+/** Initialization parameters for persona toolkits */
+struct FPersonaToolkitArgs
+{
+	/** Whether to create a preview scene */
+	bool bCreatePreviewScene;
+
+	FPersonaToolkitArgs()
+		: bCreatePreviewScene(true)
+	{}
+};
+
+struct FAnimDocumentArgs
+{
+	FAnimDocumentArgs(const TSharedRef<class IPersonaPreviewScene>& InPreviewScene, const TSharedRef<class IPersonaToolkit>& InPersonaToolkit, const TSharedRef<class IEditableSkeleton>& InEditableSkeleton, FSimpleMulticastDelegate& InOnPostUndo, FSimpleMulticastDelegate& InOnCurvesChanged, FSimpleMulticastDelegate& InOnAnimNotifiesChanged, FSimpleMulticastDelegate& InOnSectionsChanged)
+		: PreviewScene(InPreviewScene)
+		, PersonaToolkit(InPersonaToolkit)
+		, EditableSkeleton(InEditableSkeleton)
+		, OnPostUndo(InOnPostUndo)
+		, OnCurvesChanged(InOnCurvesChanged)
+		, OnAnimNotifiesChanged(InOnAnimNotifiesChanged)
+		, OnSectionsChanged(InOnSectionsChanged)
+	{}
+
+	/** Required args */
+	TWeakPtr<class IPersonaPreviewScene> PreviewScene;
+	TWeakPtr<class IPersonaToolkit> PersonaToolkit;
+	TWeakPtr<class IEditableSkeleton> EditableSkeleton;
+	FSimpleMulticastDelegate& OnPostUndo;
+	FSimpleMulticastDelegate& OnCurvesChanged;
+	FSimpleMulticastDelegate& OnAnimNotifiesChanged;
+	FSimpleMulticastDelegate& OnSectionsChanged;
+
+	/** Optional args */
+	FOnObjectsSelected OnDespatchObjectsSelected;
+	FOnInvokeTab OnDespatchInvokeTab;
+	FSimpleDelegate OnDespatchCurvesChanged;
+	FSimpleDelegate OnDespatchSectionsChanged;
+	FSimpleDelegate OnDespatchAnimNotifiesChanged;
+};
+
 /**
  * Persona module manages the lifetime of all instances of Persona editors.
  */
@@ -53,10 +108,85 @@ public:
 	 */
 	virtual TSharedRef<class IBlueprintEditor> CreatePersona( const EToolkitMode::Type Mode, const TSharedPtr< IToolkitHost >& InitToolkitHost, class USkeleton* Skeleton, class UAnimBlueprint* Blueprint, class UAnimationAsset* AnimationAsset, class USkeletalMesh * Mesh );
 
+	/** Create a re-usable toolkit that multiple asset editors that are concerned with USkeleton-related data can use */
+	virtual TSharedRef<class IPersonaToolkit> CreatePersonaToolkit(USkeleton* InSkeleton, const FPersonaToolkitArgs& PersonaToolkitArgs = FPersonaToolkitArgs()) const;
+	virtual TSharedRef<class IPersonaToolkit> CreatePersonaToolkit(UAnimationAsset* InAnimationAsset, const FPersonaToolkitArgs& PersonaToolkitArgs = FPersonaToolkitArgs()) const;
+	virtual TSharedRef<class IPersonaToolkit> CreatePersonaToolkit(USkeletalMesh* InSkeletalMesh, const FPersonaToolkitArgs& PersonaToolkitArgs = FPersonaToolkitArgs()) const;
+	virtual TSharedRef<class IPersonaToolkit> CreatePersonaToolkit(UAnimBlueprint* InAnimBlueprint, const FPersonaToolkitArgs& PersonaToolkitArgs = FPersonaToolkitArgs()) const;
+
+	/** Create an asset family for the supplied persona asset */
+	virtual TSharedRef<class IAssetFamily> CreatePersonaAssetFamily(const UObject* InAsset) const;
+
+	/** Create a shortcut widget for an asset family */
+	virtual TSharedRef<SWidget> CreateAssetFamilyShortcutWidget(const TSharedRef<class FWorkflowCentricApplication>& InHostingApp, const TSharedRef<class IAssetFamily>& InAssetFamily) const;
+
+	/** Create a details panel tab factory */
+	virtual TSharedRef<class FWorkflowTabFactory> CreateDetailsTabFactory(const TSharedRef<class FWorkflowCentricApplication>& InHostingApp, FOnDetailsCreated InOnDetailsCreated) const;
+
+	/** Create a persona viewport tab factory */
+	virtual TSharedRef<class FWorkflowTabFactory> CreatePersonaViewportTabFactory(const TSharedRef<class FWorkflowCentricApplication>& InHostingApp, const TSharedRef<class ISkeletonTree>& InSkeletonTree, const TSharedRef<class IPersonaPreviewScene>& InPreviewScene, FSimpleMulticastDelegate& OnPostUndo, const TSharedPtr<class FBlueprintEditor>& InBlueprintEditor, FOnViewportCreated InOnViewportCreated, bool bInShowTimeline, bool bInShowStats) const;
+
+	/** Create an anim notifies tab factory */
+	virtual TSharedRef<class FWorkflowTabFactory> CreateAnimNotifiesTabFactory(const TSharedRef<class FWorkflowCentricApplication>& InHostingApp, const TSharedRef<class IEditableSkeleton>& InEditableSkeleton, FSimpleMulticastDelegate& InOnChangeAnimNotifies, FSimpleMulticastDelegate& InOnPostUndo, FOnObjectsSelected InOnObjectsSelected) const;
+
+	/** Create a skeleton cuve viewer tab factory */
+	virtual TSharedRef<class FWorkflowTabFactory> CreateCurveViewerTabFactory(const TSharedRef<class FWorkflowCentricApplication>& InHostingApp, const TSharedRef<class IEditableSkeleton>& InEditableSkeleton, const TSharedRef<class IPersonaPreviewScene>& InPreviewScene, FSimpleMulticastDelegate& InOnCurvesChanged, FSimpleMulticastDelegate& InOnPostUndo) const;
+
+	/** Create a retarget manager tab factory */
+	virtual TSharedRef<class FWorkflowTabFactory> CreateRetargetManagerTabFactory(const TSharedRef<class FWorkflowCentricApplication>& InHostingApp, const TSharedRef<class IEditableSkeleton>& InEditableSkeleton, const TSharedRef<IPersonaPreviewScene>& InPreviewScene, FSimpleMulticastDelegate& InOnPostUndo) const;
+
+	/** Create a tab factory used to configure preview scene settings */
+	virtual TSharedRef<class FWorkflowTabFactory> CreateAdvancedPreviewSceneTabFactory(const TSharedRef<class FWorkflowCentricApplication>& InHostingApp, const TSharedRef<IPersonaPreviewScene>& InPreviewScene) const;
+
+	/** Create a tab factory for the animation asset browser */
+	virtual TSharedRef<class FWorkflowTabFactory> CreateAnimationAssetBrowserTabFactory(const TSharedRef<class FWorkflowCentricApplication>& InHostingApp, const TSharedRef<IPersonaToolkit>& InPersonaToolkit, FOnOpenNewAsset InOnOpenNewAsset, FOnAnimationSequenceBrowserCreated InOnAnimationSequenceBrowserCreated, bool bInShowHistory) const;
+
+	/** Create a tab factory for editing a single object (like an animation asset) */
+	virtual TSharedRef<class FWorkflowTabFactory> CreateAssetDetailsTabFactory(const TSharedRef<class FWorkflowCentricApplication>& InHostingApp, FOnGetAsset InOnGetAsset, FOnDetailsCreated InOnDetailsCreated) const;
+
+	/** Create a tab factory for for previewing morph targets */
+	virtual TSharedRef<class FWorkflowTabFactory> CreateMorphTargetTabFactory(const TSharedRef<class FWorkflowCentricApplication>& InHostingApp, const TSharedRef<IPersonaPreviewScene>& InPreviewScene, FSimpleMulticastDelegate& OnPostUndo) const;
+
+	/** Create a tab factory for editing anim blueprint preview & defaults */
+	virtual TSharedRef<class FWorkflowTabFactory> CreateAnimBlueprintPreviewTabFactory(const TSharedRef<class FBlueprintEditor>& InBlueprintEditor, const TSharedRef<IPersonaPreviewScene>& InPreviewScene) const;
+
+	/** Create a tab factory for editing anim blueprint parent overrides */
+	virtual TSharedRef<class FWorkflowTabFactory> CreateAnimBlueprintAssetOverridesTabFactory(const TSharedRef<class FBlueprintEditor>& InBlueprintEditor, UAnimBlueprint* InAnimBlueprint, FSimpleMulticastDelegate& InOnPostUndo) const;
+
+	/** Create a tab factory for editing slot names and groups */
+	virtual TSharedRef<class FWorkflowTabFactory> CreateSkeletonSlotNamesTabFactory(const TSharedRef<class FWorkflowCentricApplication>& InHostingApp, const TSharedRef<class IEditableSkeleton>& InEditableSkeleton, FSimpleMulticastDelegate& InOnPostUndo, FOnObjectSelected InOnObjectSelected) const;
+
+	/** Create a widget that acts as a document for an animation asset */
+	virtual TSharedRef<SWidget> CreateEditorWidgetForAnimDocument(const TSharedRef<class FWorkflowCentricApplication>& InHostingApp, UObject* InAnimAsset, const FAnimDocumentArgs& InArgs, FString& OutDocumentLink);
+
+	/** Customize a skeletal mesh details panel */
+	virtual void CustomizeMeshDetails(const TSharedRef<IDetailsView>& InDetailsView, const TSharedRef<IPersonaToolkit>& InPersonaToolkit);
+
 	/** Gets the extensibility managers for outside entities to extend persona editor's menus and toolbars */
 	virtual TSharedPtr<FExtensibilityManager> GetMenuExtensibilityManager() {return MenuExtensibilityManager;}
 	virtual TSharedPtr<FExtensibilityManager> GetToolBarExtensibilityManager() {return ToolBarExtensibilityManager;}
 
+	/** Import a new asset using the supplied skeleton */
+	virtual void ImportNewAsset(USkeleton* InSkeleton, EFBXImportType DefaultImportType);
+
+	/** Check all animations & skeletal meshes for curve usage */
+	virtual void TestSkeletonCurveNamesForUse(const TSharedRef<IEditableSkeleton>& InEditableSkeleton) const;
+
+	/** Apply Compression to list of animations */
+	virtual void ApplyCompression(TArray<TWeakObjectPtr<class UAnimSequence>>& AnimSequences);
+
+	/** Export to FBX files of the list of animations */
+	virtual void ExportToFBX(TArray<TWeakObjectPtr<class UAnimSequence>>& AnimSequences, USkeletalMesh* SkeletalMesh);
+
+	/** Add looping interpolation to the list of animations */
+	virtual void AddLoopingInterpolation(TArray<TWeakObjectPtr<class UAnimSequence>>& AnimSequences);
+
+	/** Customize the details of a slot node for the specified details view */
+	virtual void CustomizeSlotNodeDetails(const TSharedRef<class IDetailsView>& InDetailsView, FOnInvokeTab InOnInvokeTab);
+
+	/** Create a Persona editor mode manager. Should be destroyed using plain ol' delete. Note: Only FPersonaEditMode-derived modes should be used with this manager! */
+	virtual class IPersonaEditorModeManager* CreatePersonaEditorModeManager();
+
 	/** Delegate used to query whether recording is active */
 	FIsRecordingActive& OnIsRecordingActive() { return IsRecordingActiveDelegate; }
 
@@ -99,6 +229,3 @@ private:
 	FTickRecording TickRecordingDelegate;
 };
 
-
-
-#endif		//__PersonaModule_h__
diff --git a/Engine/Source/Editor/PropertyEditor/Private/DetailPropertyRow.cpp b/Engine/Source/Editor/PropertyEditor/Private/DetailPropertyRow.cpp
--- a/Engine/Source/Editor/PropertyEditor/Private/DetailPropertyRow.cpp
+++ b/Engine/Source/Editor/PropertyEditor/Private/DetailPropertyRow.cpp
@@ -497,17 +497,19 @@ void FDetailPropertyRow::MakeValueWidget( FDetailWidgetRow& Row, const TSharedPt
 			];
 		}
 
-
-		ValueWidget->AddSlot()
-		.Padding( 2.0f, 0.0f )
-		.AutoWidth()
-		.VAlign(VAlign_Center)
-		.HAlign(HAlign_Left)
-		[
-			SNew( SResetToDefaultPropertyEditor, PropertyEditor.ToSharedRef() )
-			.IsEnabled( IsEnabledAttrib )
-			.CustomResetToDefault(CustomResetToDefault)
-		];
+		if (!PropertyHandle->HasMetaData(TEXT("NoResetToDefault")))
+		{
+			ValueWidget->AddSlot()
+			.Padding( 2.0f, 0.0f )
+			.AutoWidth()
+			.VAlign(VAlign_Center)
+			.HAlign(HAlign_Left)
+			[
+				SNew( SResetToDefaultPropertyEditor, PropertyEditor.ToSharedRef() )
+				.IsEnabled( IsEnabledAttrib )
+				.CustomResetToDefault(CustomResetToDefault)
+			];
+		}
 	}
 
 	Row.ValueContent()
diff --git a/Engine/Source/Editor/PropertyEditor/Private/ItemPropertyNode.cpp b/Engine/Source/Editor/PropertyEditor/Private/ItemPropertyNode.cpp
--- a/Engine/Source/Editor/PropertyEditor/Private/ItemPropertyNode.cpp
+++ b/Engine/Source/Editor/PropertyEditor/Private/ItemPropertyNode.cpp
@@ -138,7 +138,8 @@ void FItemPropertyNode::InitExpansionFlags (void)
 		|| ( ( Cast<UArrayProperty>(MyProperty) || Cast<USetProperty>(MyProperty) || Cast<UMapProperty>(MyProperty) ) && GetReadAddress(false,Addresses) );
 
 	if(	bExpandableType
-		||  HasNodeFlags(EPropertyNodeFlags::EditInline)
+		|| HasNodeFlags(EPropertyNodeFlags::EditInlineNew)
+		|| HasNodeFlags(EPropertyNodeFlags::ShowInnerObjectProperties)
 		||	( MyProperty->ArrayDim > 1 && ArrayIndex == -1 ) )
 	{
 		SetNodeFlags(EPropertyNodeFlags::CanBeExpanded, true);
diff --git a/Engine/Source/Editor/PropertyEditor/Private/PropertyEditorHelpers.cpp b/Engine/Source/Editor/PropertyEditor/Private/PropertyEditorHelpers.cpp
--- a/Engine/Source/Editor/PropertyEditor/Private/PropertyEditorHelpers.cpp
+++ b/Engine/Source/Editor/PropertyEditor/Private/PropertyEditorHelpers.cpp
@@ -59,7 +59,7 @@ void SPropertyNameWidget::Construct( const FArguments& InArgs, TSharedPtr<FPrope
 	
 	];
 
-	if( InArgs._DisplayResetToDefault )
+	if( InArgs._DisplayResetToDefault && !PropertyEditor->GetPropertyHandle()->HasMetaData(TEXT("NoResetToDefault")) )
 	{
 		HorizontalBox->AddSlot()
 		.AutoWidth()
@@ -83,7 +83,7 @@ void SPropertyValueWidget::Construct( const FArguments& InArgs, TSharedPtr<FProp
 
 	if ( !ValueEditorWidget->GetToolTip().IsValid() )
 	{
-		ValueEditorWidget->SetToolTipText( PropertyEditor->GetToolTipText() );
+	ValueEditorWidget->SetToolTipText( PropertyEditor->GetToolTipText() );
 	}
 
 
@@ -650,12 +650,12 @@ namespace PropertyEditorHelpers
 		if( NodeProperty->IsA(UArrayProperty::StaticClass()) || NodeProperty->IsA(USetProperty::StaticClass()) || NodeProperty->IsA(UMapProperty::StaticClass()) )
 		{
 			if (!NodeProperty->IsA(UArrayProperty::StaticClass()))
-			{
+		{
 				// Only Sets and Maps get a Documentation widget
 				OutRequiredButtons.Add(EPropertyButton::Documentation);
 			}
 			
-			if (!(NodeProperty->PropertyFlags & CPF_EditFixedSize))
+			if( !(NodeProperty->PropertyFlags & CPF_EditFixedSize) )
 			{
 				OutRequiredButtons.Add( EPropertyButton::Add );
 				OutRequiredButtons.Add( EPropertyButton::Empty );
@@ -672,59 +672,59 @@ namespace PropertyEditorHelpers
 			bool bStaticSizedArray = (NodeProperty->ArrayDim > 1) && (PropertyNode->GetArrayIndex() == -1);
 			if (!bStaticSizedArray)
 			{
-				if( PropertyNode->HasNodeFlags(EPropertyNodeFlags::EditInline) )
-				{
-					// hmmm, seems like this code could be removed and the code inside the 'if <UClassProperty>' check
-					// below could be moved outside the else....but is there a reason to allow class properties to have the
-					// following buttons if the class property is marked 'editinline' (which is effectively what this logic is doing)
-					if( !(NodeProperty->PropertyFlags & CPF_NoClear) )
+					if( PropertyNode->HasNodeFlags(EPropertyNodeFlags::EditInlineNew) )
 					{
-						OutRequiredButtons.Add( EPropertyButton::Clear );
-					}
-				}
-				else
-				{
-					// ignore class properties
-					if( (Cast<const UClassProperty>( NodeProperty ) == NULL) && (Cast<const UAssetClassProperty>( NodeProperty ) == NULL) )
-					{
-						UObjectPropertyBase* ObjectProperty = Cast<UObjectPropertyBase>( NodeProperty );
-
-						if( ObjectProperty && ObjectProperty->PropertyClass->IsChildOf( AActor::StaticClass() ) )
-						{
-							// add button for picking the actor from the viewport
-							OutRequiredButtons.Add( EPropertyButton::PickActorInteractive );
-						}
-						else
-						{
-							// add button for filling the value of this item with the selected object from the GB
-							OutRequiredButtons.Add( EPropertyButton::Use );
-						}
-
-						// add button to display the generic browser
-						OutRequiredButtons.Add( EPropertyButton::Browse );
-
-						// reference to object resource that isn't dynamically created (i.e. some content package)
+						// hmmm, seems like this code could be removed and the code inside the 'if <UClassProperty>' check
+						// below could be moved outside the else....but is there a reason to allow class properties to have the
+						// following buttons if the class property is marked 'editinline' (which is effectively what this logic is doing)
 						if( !(NodeProperty->PropertyFlags & CPF_NoClear) )
 						{
-							// add button to clear the text
 							OutRequiredButtons.Add( EPropertyButton::Clear );
 						}
-							
-						// Do not allow actor object properties to show the asset picker
-						if( ( ObjectProperty && !ObjectProperty->PropertyClass->IsChildOf( AActor::StaticClass() ) ) || IsStringAssetReference(NodeProperty) )
-						{
-							// add button for picking the asset from an asset picker
-							OutRequiredButtons.Add( EPropertyButton::PickAsset );
-						}
-						else if( ObjectProperty && ObjectProperty->PropertyClass->IsChildOf( AActor::StaticClass() ) )
+					}
+					else
+					{
+						// ignore class properties
+						if( (Cast<const UClassProperty>( NodeProperty ) == NULL) && (Cast<const UAssetClassProperty>( NodeProperty ) == NULL) )
 						{
-							// add button for picking the actor from the scene outliner
-							OutRequiredButtons.Add( EPropertyButton::PickActor );
+							UObjectPropertyBase* ObjectProperty = Cast<UObjectPropertyBase>( NodeProperty );
+
+							if( ObjectProperty && ObjectProperty->PropertyClass->IsChildOf( AActor::StaticClass() ) )
+							{
+								// add button for picking the actor from the viewport
+								OutRequiredButtons.Add( EPropertyButton::PickActorInteractive );
+							}
+							else
+							{
+								// add button for filling the value of this item with the selected object from the GB
+								OutRequiredButtons.Add( EPropertyButton::Use );
+							}
+
+							// add button to display the generic browser
+							OutRequiredButtons.Add( EPropertyButton::Browse );
+
+							// reference to object resource that isn't dynamically created (i.e. some content package)
+							if( !(NodeProperty->PropertyFlags & CPF_NoClear) )
+							{
+								// add button to clear the text
+								OutRequiredButtons.Add( EPropertyButton::Clear );
+							}
+							
+							// Do not allow actor object properties to show the asset picker
+							if( ( ObjectProperty && !ObjectProperty->PropertyClass->IsChildOf( AActor::StaticClass() ) ) || IsStringAssetReference(NodeProperty) )
+							{
+								// add button for picking the asset from an asset picker
+								OutRequiredButtons.Add( EPropertyButton::PickAsset );
+							}
+							else if( ObjectProperty && ObjectProperty->PropertyClass->IsChildOf( AActor::StaticClass() ) )
+							{
+								// add button for picking the actor from the scene outliner
+								OutRequiredButtons.Add( EPropertyButton::PickActor );
+							}
 						}
 					}
 				}
 			}
-		}
 
 		//////////////////////////////
 		// Handle a class property.
@@ -940,7 +940,7 @@ namespace PropertyEditorHelpers
 					InsertAction = FExecuteAction::CreateSP(PropertyEditor, &FPropertyEditor::InsertItem);
 				}
 
-				if (ButtonType == EPropertyButton::Insert_Delete_Duplicate )
+				if (ButtonType == EPropertyButton::Insert_Delete_Duplicate)
 				{
 					DuplicateAction = FExecuteAction::CreateSP( PropertyEditor, &FPropertyEditor::DuplicateItem );
 				}
diff --git a/Engine/Source/Editor/PropertyEditor/Private/PropertyNode.cpp b/Engine/Source/Editor/PropertyEditor/Private/PropertyNode.cpp
--- a/Engine/Source/Editor/PropertyEditor/Private/PropertyNode.cpp
+++ b/Engine/Source/Editor/PropertyEditor/Private/PropertyNode.cpp
@@ -124,10 +124,12 @@ void FPropertyNode::InitNode( const FPropertyNodeInitParams& InitParams )
 	//needs to happen after the above SetNodeFlags calls so that ObjectPropertyNode can properly respond to CollapseCategories
 	InitBeforeNodeFlags();
 
+	bool bIsEditInlineNew = false;
+	bool bShowInnerObjectProperties = false;
 	if ( !Property.IsValid() )
 	{
 		// Disable all flags if no property is bound.
-		SetNodeFlags(EPropertyNodeFlags::SingleSelectOnly | EPropertyNodeFlags::EditInline , false);
+		SetNodeFlags(EPropertyNodeFlags::SingleSelectOnly | EPropertyNodeFlags::EditInlineNew | EPropertyNodeFlags::ShowInnerObjectProperties, false);
 	}
 	else
 	{
@@ -142,8 +144,19 @@ void FPropertyNode::InitNode( const FPropertyNodeInitParams& InitParams )
 		// true if the property can be expanded into the property window; that is, instead of seeing
 		// a pointer to the object, you see the object's properties.
 		static const FName Name_EditInline("EditInline");
-		const bool bEditInline = bIsObjectOrInterface && GotReadAddresses && MyProperty->HasMetaData(Name_EditInline);
-		SetNodeFlags(EPropertyNodeFlags::EditInline, bEditInline);
+		static const FName Name_ShowInnerProperties("ShowInnerProperties");
+
+		bIsEditInlineNew = bIsObjectOrInterface && GotReadAddresses && MyProperty->HasMetaData(Name_EditInline);
+		bShowInnerObjectProperties = bIsObjectOrInterface && GotReadAddresses && MyProperty->HasMetaData(Name_ShowInnerProperties);
+
+		if(bIsEditInlineNew)
+		{
+			SetNodeFlags(EPropertyNodeFlags::EditInlineNew, true);
+		}
+		else if(bShowInnerObjectProperties)
+		{
+			SetNodeFlags(EPropertyNodeFlags::ShowInnerObjectProperties, true);
+		}
 
 		//Get the property max child depth
 		static const FName Name_MaxPropertyDepth("MaxPropertyDepth");
@@ -171,9 +184,7 @@ void FPropertyNode::InitNode( const FPropertyNodeInitParams& InitParams )
 
 	UProperty* MyProperty = Property.Get();
 
-	bool bIsEditInlineNew = MyProperty && !( MyProperty->PropertyFlags & CPF_EditConst ) && HasNodeFlags( EPropertyNodeFlags::EditInline ) != 0;
-
-	bool bRequiresValidation = bIsEditInlineNew || ( MyProperty && (MyProperty->IsA<UArrayProperty>() || MyProperty->IsA<USetProperty>() || MyProperty->IsA<UMapProperty>() ));
+	bool bRequiresValidation = bIsEditInlineNew || bShowInnerObjectProperties || ( MyProperty && (MyProperty->IsA<UArrayProperty>() || MyProperty->IsA<USetProperty>() || MyProperty->IsA<UMapProperty>() ));
 
 	// We require validation if our parent also needs validation (if an array parent was resized all the addresses of children are invalid)
 	bRequiresValidation |= (GetParentNode() && GetParentNode()->HasNodeFlags( EPropertyNodeFlags::RequiresValidation ) != 0);
@@ -389,7 +400,7 @@ EPropertyDataValidationResult FPropertyNode::EnsureDataIsValid()
 			bool bObjectPropertyNull = true;
 
 			//Edit inline properties can change underneath the window
-			bool bIgnoreChangingChildren = !HasNodeFlags(EPropertyNodeFlags::EditInline);
+			bool bIgnoreChangingChildren = !(HasNodeFlags(EPropertyNodeFlags::EditInlineNew) || HasNodeFlags(EPropertyNodeFlags::ShowInnerObjectProperties)) ;
 			//ignore this node if the consistency check should happen for the children
 			bool bIgnoreStaticArray = (Property->ArrayDim > 1) && (ArrayIndex == -1);
 
@@ -409,6 +420,8 @@ EPropertyDataValidationResult FPropertyNode::EnsureDataIsValid()
 				return EPropertyDataValidationResult::ObjectInvalid;
 			}
 
+			// If an object property with ShowInnerProperties changed object values out from under the property
+			bool bShowInnerObjectPropertiesObjectChanged = false;
 
 			//check for null, if we find one, there is a problem.
 			for (int32 Scan = 0; Scan < ReadAddresses.Num(); ++Scan)
@@ -468,12 +481,33 @@ EPropertyDataValidationResult FPropertyNode::EnsureDataIsValid()
 
 				if (ObjectProperty && !bIgnoreAllMismatch)
 				{
-					UObject* obj = ObjectProperty->GetObjectPropertyValue(Addr);
-					if (obj != NULL)
+					UObject* Obj = ObjectProperty->GetObjectPropertyValue(Addr);
+
+					if (!bShowInnerObjectPropertiesObjectChanged && HasNodeFlags(EPropertyNodeFlags::ShowInnerObjectProperties) && ChildNodes.Num() == 1)
+					{
+						bool bChildObjectFound = false;
+						// should never have more than one node (0 is ok if the object property is null)
+						check(ChildNodes.Num() == 1);
+						bool bNeedRebuild = false;
+						FObjectPropertyNode* ChildObjectNode = ChildNodes[0]->AsObjectNode();
+						for(int32 ObjectIndex = 0; ObjectIndex < ChildObjectNode->GetNumObjects(); ++ObjectIndex)
+						{
+							if(Obj == ChildObjectNode->GetUObject(ObjectIndex))
+							{
+								bChildObjectFound = true;
+								break;
+							}
+						}
+						bShowInnerObjectPropertiesObjectChanged = !bChildObjectFound;
+					}
+
+					if (Obj != NULL)
 					{
 						bObjectPropertyNull = false;
 						break;
 					}
+
+					
 				}
 			}
 
@@ -496,6 +530,12 @@ EPropertyDataValidationResult FPropertyNode::EnsureDataIsValid()
 				return EPropertyDataValidationResult::ArraySizeChanged;
 			}
 
+			if(bShowInnerObjectPropertiesObjectChanged)
+			{
+				RebuildChildren();
+				return EPropertyDataValidationResult::EditInlineNewValueChanged;
+			}
+
 			const bool bHasChildren = (GetNumChildNodes() != 0);
 			// If the object property is not null and has no children, its children need to be rebuilt
 			// If the object property is null and this node has children, the node needs to be rebuilt
@@ -511,7 +551,7 @@ EPropertyDataValidationResult FPropertyNode::EnsureDataIsValid()
 	{
 		RebuildChildren();
 		// If this property is editinline and not edit const then its editinline new and we can optimize some of the refreshing in some cases.  Otherwise we need to refresh all properties in the view
-		return HasNodeFlags(EPropertyNodeFlags::EditInline) && !IsEditConst() ? EPropertyDataValidationResult::EditInlineNewValueChanged : EPropertyDataValidationResult::PropertiesChanged;
+		return HasNodeFlags(EPropertyNodeFlags::ShowInnerObjectProperties) || (HasNodeFlags(EPropertyNodeFlags::EditInlineNew) && !IsEditConst()) ? EPropertyDataValidationResult::EditInlineNewValueChanged : EPropertyDataValidationResult::PropertiesChanged;
 	}
 	
 	EPropertyDataValidationResult FinalResult = EPropertyDataValidationResult::DataValid;
@@ -1602,10 +1642,10 @@ FString FPropertyNode::GetDefaultValueAsStringForObject( FPropertyItemValueDataT
 				{
 					if ( !OuterSetProperty && !OuterMapProperty )
 					{
-						// if either are NULL, we had a dynamic array somewhere in our parent chain and the array doesn't
-						// have enough elements in either the default or the object
-						DefaultValue = NSLOCTEXT("PropertyEditor", "DifferentArrayLength", "Array has different length than the default.").ToString();
-					}
+					// if either are NULL, we had a dynamic array somewhere in our parent chain and the array doesn't
+					// have enough elements in either the default or the object
+					DefaultValue = NSLOCTEXT("PropertyEditor", "DifferentArrayLength", "Array has different length than the default.").ToString();
+				}
 				}
 				else if ( GetArrayIndex() == INDEX_NONE && InProperty->ArrayDim > 1 )
 				{
@@ -1690,7 +1730,7 @@ void FPropertyNode::ResetToDefault( FNotifyHook* InNotifyHook )
 	UProperty* TheProperty = GetProperty();
 	check(TheProperty);
 
-	// Get an iterator for the enclosing objects.
+		// Get an iterator for the enclosing objects.
 	FObjectPropertyNode* ObjectNode = FindObjectItemParent();
 
 	if( ObjectNode )
@@ -2663,7 +2703,7 @@ void FPropertyNode::PropagateContainerPropertyChange( UObject* ModifiedObject, c
 					case EPropertyArrayChangeType::Duplicate:
 						check(false);	// Duplicate not supported on sets
 						break;
-					}
+			}
 
 					if (ElementToInitialize >= 0)
 					{
diff --git a/Engine/Source/Editor/PropertyEditor/Private/PropertyNode.h b/Engine/Source/Editor/PropertyEditor/Private/PropertyNode.h
--- a/Engine/Source/Editor/PropertyEditor/Private/PropertyNode.h
+++ b/Engine/Source/Editor/PropertyEditor/Private/PropertyNode.h
@@ -19,7 +19,7 @@ namespace EPropertyNodeFlags
 	const Type	Expanded						= 1 << 5;		/** true if this node should display its children*/
 	const Type	CanBeExpanded					= 1 << 6;		/** true if this node is able to be expanded */
 
-	const Type	EditInline						= 1 << 7;		/** true if the property can be expanded into the property window. */
+	const Type	EditInlineNew					= 1 << 7;		/** true if the property can be expanded into the property window. */
 
 	const Type	SingleSelectOnly				= 1 << 8;		/** true if only a single object is selected. */
 	const Type  ShowCategories					= 1 << 9;		/** true if this node should show categories.  Different*/
@@ -45,6 +45,8 @@ namespace EPropertyNodeFlags
 
 	const Type	SkipChildValidation				= 1 << 21; /** true if this node should skip child validation */
 
+	const Type  ShowInnerObjectProperties		= 1 << 22;
+
 	const Type 	NoFlags							= 0;
 
 };
diff --git a/Engine/Source/Editor/PropertyEditor/Private/SDetailsView.cpp b/Engine/Source/Editor/PropertyEditor/Private/SDetailsView.cpp
--- a/Engine/Source/Editor/PropertyEditor/Private/SDetailsView.cpp
+++ b/Engine/Source/Editor/PropertyEditor/Private/SDetailsView.cpp
@@ -268,7 +268,8 @@ EVisibility SDetailsView::GetActorNameAreaVisibility() const
 EVisibility SDetailsView::GetScrollBarVisibility() const
 {
 	const bool bHasObjects = RootPropertyNodes.Num() > 0;
-	return bHasObjects ? EVisibility::Visible : EVisibility::Collapsed; 
+	const bool bShowScrollBar = DetailsViewArgs.bShowScrollBar && bHasObjects;
+	return bShowScrollBar ? EVisibility::Visible : EVisibility::Collapsed;
 }
 
 void SDetailsView::ForceRefresh()
diff --git a/Engine/Source/Editor/PropertyEditor/Private/SDetailsViewBase.cpp b/Engine/Source/Editor/PropertyEditor/Private/SDetailsViewBase.cpp
--- a/Engine/Source/Editor/PropertyEditor/Private/SDetailsViewBase.cpp
+++ b/Engine/Source/Editor/PropertyEditor/Private/SDetailsViewBase.cpp
@@ -460,7 +460,7 @@ void SDetailsViewBase::UpdateSinglePropertyMapRecursive(FPropertyNode& InNode, F
 				bool bPushOutStructProps = bIsStruct && !bIsCustomizedStruct && !ParentStructProp && Property->HasMetaData(ShowOnlyInners);
 
 				// Is the property edit inline new 
-				const bool bIsEditInlineNew = SPropertyEditorEditInline::Supports(&ChildNode, ChildNode.GetArrayIndex());
+				const bool bIsEditInlineNew = ChildNode.HasNodeFlags(EPropertyNodeFlags::ShowInnerObjectProperties) || SPropertyEditorEditInline::Supports(&ChildNode, ChildNode.GetArrayIndex());
 
 				// Is this a property of a container property
 				bool bIsChildOfContainer = PropertyEditorHelpers::IsChildOfArray(ChildNode) || PropertyEditorHelpers::IsChildOfSet(ChildNode) || PropertyEditorHelpers::IsChildOfMap(ChildNode);
diff --git a/Engine/Source/Editor/PropertyEditor/Private/SSingleProperty.cpp b/Engine/Source/Editor/PropertyEditor/Private/SSingleProperty.cpp
--- a/Engine/Source/Editor/PropertyEditor/Private/SSingleProperty.cpp
+++ b/Engine/Source/Editor/PropertyEditor/Private/SSingleProperty.cpp
@@ -172,13 +172,16 @@ void SSingleProperty::SetObject( UObject* InObject )
 			SNew( SPropertyValueWidget, PropertyEditor, PropertyUtilities.ToSharedRef() )
 		];
 
-		HorizontalBox->AddSlot()
-		.Padding( 2.0f )
-		.AutoWidth()
-		.VAlign( VAlign_Center )
-		[
-			SNew( SResetToDefaultPropertyEditor,  PropertyEditor )
-		];
+		if (!PropertyEditor->GetPropertyHandle()->HasMetaData(TEXT("NoResetToDefault")))
+		{
+			HorizontalBox->AddSlot()
+			.Padding( 2.0f )
+			.AutoWidth()
+			.VAlign( VAlign_Center )
+			[
+				SNew( SResetToDefaultPropertyEditor,  PropertyEditor )
+			];
+		}
 	}
 	else
 	{
diff --git a/Engine/Source/Editor/PropertyEditor/Private/UserInterface/PropertyEditor/SPropertyEditorAsset.cpp b/Engine/Source/Editor/PropertyEditor/Private/UserInterface/PropertyEditor/SPropertyEditorAsset.cpp
--- a/Engine/Source/Editor/PropertyEditor/Private/UserInterface/PropertyEditor/SPropertyEditorAsset.cpp
+++ b/Engine/Source/Editor/PropertyEditor/Private/UserInterface/PropertyEditor/SPropertyEditorAsset.cpp
@@ -421,7 +421,7 @@ void SPropertyEditorAsset::Tick( const FGeometry& AllottedGeometry, const double
 bool SPropertyEditorAsset::Supports( const TSharedRef< FPropertyEditor >& InPropertyEditor )
 {
 	const TSharedRef< FPropertyNode > PropertyNode = InPropertyEditor->GetPropertyNode();
-	if(	PropertyNode->HasNodeFlags(EPropertyNodeFlags::EditInline) )
+	if(	PropertyNode->HasNodeFlags(EPropertyNodeFlags::EditInlineNew) )
 	{
 		return false;
 	}
diff --git a/Engine/Source/Editor/PropertyEditor/Private/UserInterface/PropertyEditor/SPropertyEditorEditInline.cpp b/Engine/Source/Editor/PropertyEditor/Private/UserInterface/PropertyEditor/SPropertyEditorEditInline.cpp
--- a/Engine/Source/Editor/PropertyEditor/Private/UserInterface/PropertyEditor/SPropertyEditorEditInline.cpp
+++ b/Engine/Source/Editor/PropertyEditor/Private/UserInterface/PropertyEditor/SPropertyEditorEditInline.cpp
@@ -131,7 +131,7 @@ void SPropertyEditorEditInline::GetDesiredWidth( float& OutMinDesiredWidth, floa
 bool SPropertyEditorEditInline::Supports( const FPropertyNode* InTreeNode, int32 InArrayIdx )
 {
 	return InTreeNode
-		&& InTreeNode->HasNodeFlags(EPropertyNodeFlags::EditInline)
+		&& InTreeNode->HasNodeFlags(EPropertyNodeFlags::EditInlineNew)
 		&& InTreeNode->FindObjectItemParent()
 		&& !InTreeNode->IsEditConst();
 }
diff --git a/Engine/Source/Editor/PropertyEditor/Private/UserInterface/PropertyEditor/SPropertyEditorMap.cpp b/Engine/Source/Editor/PropertyEditor/Private/UserInterface/PropertyEditor/SPropertyEditorMap.cpp
--- a/Engine/Source/Editor/PropertyEditor/Private/UserInterface/PropertyEditor/SPropertyEditorMap.cpp
+++ b/Engine/Source/Editor/PropertyEditor/Private/UserInterface/PropertyEditor/SPropertyEditorMap.cpp
@@ -33,7 +33,7 @@ bool SPropertyEditorMap::Supports(const TSharedRef< FPropertyEditor >& InPropert
 	const TSharedRef< FPropertyNode > PropertyNode = InPropertyEditor->GetPropertyNode();
 	const UProperty* Property = InPropertyEditor->GetProperty();
 
-	if (!PropertyNode->HasNodeFlags(EPropertyNodeFlags::EditInline)
+	if (!PropertyNode->HasNodeFlags(EPropertyNodeFlags::EditInlineNew)
 		&& Property->IsA<UMapProperty>())
 	{
 		return true;
diff --git a/Engine/Source/Editor/PropertyEditor/Private/UserInterface/PropertyEditor/SPropertyEditorNumeric.h b/Engine/Source/Editor/PropertyEditor/Private/UserInterface/PropertyEditor/SPropertyEditorNumeric.h
--- a/Engine/Source/Editor/PropertyEditor/Private/UserInterface/PropertyEditor/SPropertyEditorNumeric.h
+++ b/Engine/Source/Editor/PropertyEditor/Private/UserInterface/PropertyEditor/SPropertyEditorNumeric.h
@@ -341,7 +341,7 @@ public:
 	{ 
 		const TSharedRef< FPropertyNode > PropertyNode = PropertyEditor->GetPropertyNode();
 		
-		if (!PropertyNode->HasNodeFlags(EPropertyNodeFlags::EditInline))
+		if (!PropertyNode->HasNodeFlags(EPropertyNodeFlags::EditInlineNew))
 		{
 			return TTypeToProperty<NumericType>::Match(PropertyEditor->GetProperty());
 		}
diff --git a/Engine/Source/Editor/PropertyEditor/Private/UserInterface/PropertyEditor/SPropertyEditorSet.cpp b/Engine/Source/Editor/PropertyEditor/Private/UserInterface/PropertyEditor/SPropertyEditorSet.cpp
--- a/Engine/Source/Editor/PropertyEditor/Private/UserInterface/PropertyEditor/SPropertyEditorSet.cpp
+++ b/Engine/Source/Editor/PropertyEditor/Private/UserInterface/PropertyEditor/SPropertyEditorSet.cpp
@@ -33,7 +33,7 @@ bool SPropertyEditorSet::Supports( const TSharedRef< FPropertyEditor >& InProper
 	const TSharedRef< FPropertyNode > PropertyNode = InPropertyEditor->GetPropertyNode();
 	const UProperty* Property = InPropertyEditor->GetProperty();
 
-	if (!PropertyNode->HasNodeFlags(EPropertyNodeFlags::EditInline)
+	if (!PropertyNode->HasNodeFlags(EPropertyNodeFlags::EditInlineNew)
 		&& Property->IsA<USetProperty>())
 	{
 		return true;
diff --git a/Engine/Source/Editor/PropertyEditor/Private/UserInterface/PropertyEditor/SPropertyEditorText.cpp b/Engine/Source/Editor/PropertyEditor/Private/UserInterface/PropertyEditor/SPropertyEditorText.cpp
--- a/Engine/Source/Editor/PropertyEditor/Private/UserInterface/PropertyEditor/SPropertyEditorText.cpp
+++ b/Engine/Source/Editor/PropertyEditor/Private/UserInterface/PropertyEditor/SPropertyEditorText.cpp
@@ -96,7 +96,7 @@ bool SPropertyEditorText::Supports( const TSharedRef< FPropertyEditor >& InPrope
 	const TSharedRef< FPropertyNode > PropertyNode = InPropertyEditor->GetPropertyNode();
 	const UProperty* Property = InPropertyEditor->GetProperty();
 
-	if(	!PropertyNode->HasNodeFlags(EPropertyNodeFlags::EditInline)
+	if(	!PropertyNode->HasNodeFlags(EPropertyNodeFlags::EditInlineNew)
 		&&	( (Property->IsA(UNameProperty::StaticClass()) && Property->GetFName() != NAME_InitialState)
 		||	Property->IsA(UStrProperty::StaticClass())
 		||	Property->IsA(UTextProperty::StaticClass())
diff --git a/Engine/Source/Editor/PropertyEditor/Private/UserInterface/PropertyTable/BooleanPropertyTableCellPresenter.cpp b/Engine/Source/Editor/PropertyEditor/Private/UserInterface/PropertyTable/BooleanPropertyTableCellPresenter.cpp
--- a/Engine/Source/Editor/PropertyEditor/Private/UserInterface/PropertyTable/BooleanPropertyTableCellPresenter.cpp
+++ b/Engine/Source/Editor/PropertyEditor/Private/UserInterface/PropertyTable/BooleanPropertyTableCellPresenter.cpp
@@ -14,14 +14,16 @@ FBooleanPropertyTableCellPresenter::FBooleanPropertyTableCellPresenter( const TS
 
 TSharedRef< class SWidget > FBooleanPropertyTableCellPresenter::ConstructDisplayWidget()
 {
-	return SNew( SBorder )
+	TSharedPtr<SHorizontalBox> HorizontalBox;
+
+	TSharedRef<SBorder> Border = SNew( SBorder )
 		.Padding( 0 )
 		.VAlign( VAlign_Center )
 		.HAlign( HAlign_Center )
 		.BorderImage( FEditorStyle::GetBrush("NoBorder") )
 		.Content()
 		[	
-			SNew( SHorizontalBox )
+			SAssignNew( HorizontalBox, SHorizontalBox )
 			+SHorizontalBox::Slot()
 			.AutoWidth()
 			.VAlign( VAlign_Center )
@@ -37,15 +39,21 @@ TSharedRef< class SWidget > FBooleanPropertyTableCellPresenter::ConstructDisplay
 					.ToolTipText( PropertyEditor->GetToolTipText() )
 				]
 			]
-			+SHorizontalBox::Slot()
-			.AutoWidth()
-			.VAlign( VAlign_Center )
-			.HAlign( HAlign_Center )
-			.Padding( FMargin( 0, 0, 2, 0 ) )
-			[
-				SNew( SResetToDefaultPropertyEditor, PropertyEditor )
-			]
 		];
+
+	if (!PropertyEditor->GetPropertyHandle()->HasMetaData(TEXT("NoResetToDefault")))
+	{
+		HorizontalBox->AddSlot()
+		.AutoWidth()
+		.VAlign(VAlign_Center)
+		.HAlign(HAlign_Center)
+		.Padding(FMargin(0, 0, 2, 0))
+		[
+			SNew(SResetToDefaultPropertyEditor, PropertyEditor)
+		];
+	}
+
+	return Border;
 }
 
 bool FBooleanPropertyTableCellPresenter::RequiresDropDown()
diff --git a/Engine/Source/Editor/PropertyEditor/Private/UserInterface/PropertyTable/ColorPropertyTableCellPresenter.cpp b/Engine/Source/Editor/PropertyEditor/Private/UserInterface/PropertyTable/ColorPropertyTableCellPresenter.cpp
--- a/Engine/Source/Editor/PropertyEditor/Private/UserInterface/PropertyTable/ColorPropertyTableCellPresenter.cpp
+++ b/Engine/Source/Editor/PropertyEditor/Private/UserInterface/PropertyTable/ColorPropertyTableCellPresenter.cpp
@@ -16,7 +16,7 @@ FColorPropertyTableCellPresenter::FColorPropertyTableCellPresenter( const TShare
 
 TSharedRef< class SWidget > FColorPropertyTableCellPresenter::ConstructDisplayWidget()
 {
-	return SNew( SHorizontalBox )
+	TSharedRef<SHorizontalBox> HorizontalBox = SNew(SHorizontalBox)
 		+SHorizontalBox::Slot()
 		.Padding( FMargin( 3.0, 0, 3.0, 0 ) )
 		.FillWidth( 1.0 )
@@ -30,8 +30,11 @@ TSharedRef< class SWidget > FColorPropertyTableCellPresenter::ConstructDisplayWi
 				SAssignNew( FocusWidget, SPropertyEditorColor, PropertyEditor, PropertyUtilities )
 				.ToolTipText( PropertyEditor->GetToolTipText() )
 			]
-		]
-		+SHorizontalBox::Slot()
+		];
+		
+	if (!PropertyEditor->GetPropertyHandle()->HasMetaData(TEXT("NoResetToDefault")))
+	{
+		HorizontalBox->AddSlot()
 		.AutoWidth()
 		.VAlign( VAlign_Center )
 		.HAlign( HAlign_Center )
@@ -39,6 +42,9 @@ TSharedRef< class SWidget > FColorPropertyTableCellPresenter::ConstructDisplayWi
 		[
 			SNew( SResetToDefaultPropertyEditor, PropertyEditor )
 		];
+	}
+
+	return HorizontalBox;
 }
 
 bool FColorPropertyTableCellPresenter::RequiresDropDown()
diff --git a/Engine/Source/Editor/PropertyEditor/Private/UserInterface/PropertyTable/TextPropertyTableCellPresenter.cpp b/Engine/Source/Editor/PropertyEditor/Private/UserInterface/PropertyTable/TextPropertyTableCellPresenter.cpp
--- a/Engine/Source/Editor/PropertyEditor/Private/UserInterface/PropertyTable/TextPropertyTableCellPresenter.cpp
+++ b/Engine/Source/Editor/PropertyEditor/Private/UserInterface/PropertyTable/TextPropertyTableCellPresenter.cpp
@@ -34,7 +34,7 @@ FTextPropertyTableCellPresenter::FTextPropertyTableCellPresenter( const TSharedR
 
 TSharedRef< class SWidget > FTextPropertyTableCellPresenter::ConstructDisplayWidget()
 {
-	return SNew( SHorizontalBox )
+	TSharedRef<SHorizontalBox> HorizontalBox = SNew( SHorizontalBox )
 		+SHorizontalBox::Slot()
 		.FillWidth( 1.0 )
 		.VAlign( VAlign_Center )
@@ -44,14 +44,20 @@ TSharedRef< class SWidget > FTextPropertyTableCellPresenter::ConstructDisplayWid
 			.Text( PropertyEditor->GetValueAsText() )
 			.ToolTipText( PropertyEditor->GetToolTipText() )
 			.Font( Font )
-		]
-	+SHorizontalBox::Slot()
+		];
+
+	if (!PropertyEditor->GetPropertyHandle()->HasMetaData(TEXT("NoResetToDefault")))
+	{
+		HorizontalBox->AddSlot()
 		.AutoWidth()
-		.VAlign( VAlign_Center )
-		.Padding( FMargin( 0, 0, 2, 0 ) )
+		.VAlign(VAlign_Center)
+		.Padding(FMargin(0, 0, 2, 0))
 		[
-			SNew( SResetToDefaultPropertyEditor, PropertyEditor )
+			SNew(SResetToDefaultPropertyEditor, PropertyEditor)
 		];
+	}
+	
+	return HorizontalBox;
 }
 
 bool FTextPropertyTableCellPresenter::RequiresDropDown()
diff --git a/Engine/Source/Editor/PropertyEditor/Public/IDetailsView.h b/Engine/Source/Editor/PropertyEditor/Public/IDetailsView.h
--- a/Engine/Source/Editor/PropertyEditor/Public/IDetailsView.h
+++ b/Engine/Source/Editor/PropertyEditor/Public/IDetailsView.h
@@ -68,7 +68,8 @@ struct FDetailsViewArgs
 	uint32 bAllowFavoriteSystem : 1;
 	/** If true the details panel will assume each object passed in through SetObjects will be a unique object shown in the tree and not combined with other objects */
 	uint32 bAllowMultipleTopLevelObjects : 1;
-
+	/** If false, the details panel's scrollbar will always be hidden. Useful when embedding details panels in widgets that either grow to accommodate them, or with scrollbars of their own. */
+	uint32 bShowScrollBar : 1;
 
 public:
 	/** Default constructor */
@@ -97,6 +98,7 @@ public:
 		, bCustomFilterAreaLocation(false)
 		, bAllowFavoriteSystem(true)
 		, bAllowMultipleTopLevelObjects(false)
+		, bShowScrollBar(true)
 	{
 	}
 };
diff --git a/Engine/Source/Editor/SequenceRecorder/Private/AnimationRecorder.cpp b/Engine/Source/Editor/SequenceRecorder/Private/AnimationRecorder.cpp
--- a/Engine/Source/Editor/SequenceRecorder/Private/AnimationRecorder.cpp
+++ b/Engine/Source/Editor/SequenceRecorder/Private/AnimationRecorder.cpp
@@ -325,7 +325,8 @@ UAnimSequence* FAnimationRecorder::StopRecord(bool bShowMessage)
 						FSmartName CurveName;
 						if (SkeletonObj->GetSmartNameByUID(USkeleton::AnimCurveMappingName, UID, CurveName))
 						{
-							AnimationObject->RawCurveData.AddFloatCurveKey(CurveName, CurCurve.Flags, TimeToRecord, CurCurve.Value);
+							// give default curve flag for recording 
+							AnimationObject->RawCurveData.AddFloatCurveKey(CurveName, ACF_DefaultCurve, TimeToRecord, CurCurve.Value);
 							FloatCurveData = static_cast<FFloatCurve*>(AnimationObject->RawCurveData.GetCurveData(UID, FRawCurveTracks::FloatType));
 						}
 					}
diff --git a/Engine/Source/Editor/SequenceRecorder/Private/AnimationRecorder.h b/Engine/Source/Editor/SequenceRecorder/Private/AnimationRecorder.h
--- a/Engine/Source/Editor/SequenceRecorder/Private/AnimationRecorder.h
+++ b/Engine/Source/Editor/SequenceRecorder/Private/AnimationRecorder.h
@@ -75,7 +75,7 @@ private:
 
 	// recording curve data 
 	TArray< TArray<FCurveElement> >  RecordedCurves;
-	TArray<FSmartNameMapping::UID> const * UIDList;
+	TArray<SmartName::UID_Type> const * UIDList;
 };
 
 //////////////////////////////////////////////////////////////////////////
diff --git a/Engine/Source/Editor/SequenceRecorder/Private/SequenceRecorderModule.cpp b/Engine/Source/Editor/SequenceRecorder/Private/SequenceRecorderModule.cpp
--- a/Engine/Source/Editor/SequenceRecorder/Private/SequenceRecorderModule.cpp
+++ b/Engine/Source/Editor/SequenceRecorder/Private/SequenceRecorderModule.cpp
@@ -86,11 +86,15 @@ class FSequenceRecorderModule : public ISequenceRecorder, private FSelfRegisteri
 			LevelEditorModule.OnCaptureSingleFrameAnimSequence().BindStatic(&FSequenceRecorderModule::HandleCaptureSingleFrameAnimSequence);
 
 			// register standalone UI
-			FGlobalTabmanager::Get()->RegisterNomadTabSpawner(SequenceRecorderTabName, FOnSpawnTab::CreateStatic(&FSequenceRecorderModule::SpawnSequenceRecorderTab))
+			LevelEditorTabManagerChangedHandle = LevelEditorModule.OnTabManagerChanged().AddLambda([]()
+			{
+				FLevelEditorModule& LocalLevelEditorModule = FModuleManager::GetModuleChecked<FLevelEditorModule>(TEXT("LevelEditor"));
+				LocalLevelEditorModule.GetLevelEditorTabManager()->RegisterTabSpawner(SequenceRecorderTabName, FOnSpawnTab::CreateStatic(&FSequenceRecorderModule::SpawnSequenceRecorderTab))
 				.SetGroup(WorkspaceMenu::GetMenuStructure().GetLevelEditorCategory())
 				.SetDisplayName(LOCTEXT("SequenceRecorderTabTitle", "Sequence Recorder"))
 				.SetTooltipText(LOCTEXT("SequenceRecorderTooltipText", "Open the Sequence Recorder tab."))
 				.SetIcon(FSlateIcon(FEditorStyle::GetStyleSetName(), "SequenceRecorder.TabIcon"));
+			});
 
 			// register for debug drawing
 			DrawDebugDelegateHandle = UDebugDrawService::Register(TEXT("Decals"), FDebugDrawDelegate::CreateStatic(&FSequenceRecorderModule::DrawDebug));
@@ -123,6 +127,7 @@ class FSequenceRecorderModule : public ISequenceRecorder, private FSelfRegisteri
 			{
 				FLevelEditorModule& LevelEditorModule = FModuleManager::GetModuleChecked<FLevelEditorModule>(TEXT("LevelEditor"));
 				LevelEditorModule.OnCaptureSingleFrameAnimSequence().Unbind();
+				LevelEditorModule.OnTabManagerChanged().Remove(LevelEditorTabManagerChangedHandle);
 			}
 
 			if (FModuleManager::Get().IsModuleLoaded(TEXT("Persona")))
@@ -639,6 +644,8 @@ class FSequenceRecorderModule : public ISequenceRecorder, private FSelfRegisteri
 
 	FDelegateHandle DrawDebugDelegateHandle;
 
+	FDelegateHandle LevelEditorTabManagerChangedHandle;
+
 	TFunction<TUniquePtr<ISequenceAudioRecorder>()> AudioFactory;
 
 	FDelegateHandle AudioFactoryHandle;
diff --git a/Engine/Source/Editor/Sequencer/Private/Sequencer.cpp b/Engine/Source/Editor/Sequencer/Private/Sequencer.cpp
--- a/Engine/Source/Editor/Sequencer/Private/Sequencer.cpp
+++ b/Engine/Source/Editor/Sequencer/Private/Sequencer.cpp
@@ -2182,6 +2182,7 @@ TSharedRef<SWidget> FSequencer::MakeTransportControls(bool bExtended)
 		TransportControlArgs.OnToggleLooping.BindSP( this, &FSequencer::OnToggleLooping );
 		TransportControlArgs.OnGetLooping.BindSP( this, &FSequencer::IsLooping );
 		TransportControlArgs.OnGetPlaybackMode.BindSP( this, &FSequencer::GetPlaybackMode );
+		TransportControlArgs.OnGetRecording.BindSP(this, &FSequencer::IsRecording);
 
 		if(bExtended)
 		{
@@ -2851,14 +2852,14 @@ EPlaybackMode::Type FSequencer::GetPlaybackMode() const
 			return EPlaybackMode::PlayingReverse;
 		}
 	}
-	else if (PlaybackState == EMovieScenePlayerStatus::Recording)
-	{
-		return EPlaybackMode::Recording;
-	}
 		
 	return EPlaybackMode::Stopped;
 }
 
+bool FSequencer::IsRecording() const
+{
+	return PlaybackState == EMovieScenePlayerStatus::Recording;
+}
 
 void FSequencer::UpdateTimeBoundsToFocusedMovieScene()
 {
diff --git a/Engine/Source/Editor/Sequencer/Private/Sequencer.h b/Engine/Source/Editor/Sequencer/Private/Sequencer.h
--- a/Engine/Source/Editor/Sequencer/Private/Sequencer.h
+++ b/Engine/Source/Editor/Sequencer/Private/Sequencer.h
@@ -380,6 +380,8 @@ public:
 
 	EPlaybackMode::Type GetPlaybackMode() const;
 
+	bool IsRecording() const;
+
 	/** Called to determine whether a frame number is set so that frame numbers can be shown */
 	bool CanShowFrameNumbers() const;
 
diff --git a/Engine/Source/Editor/SkeletalMeshEditor/Private/MeshHierarchyCmd.h b/Engine/Source/Editor/SkeletalMeshEditor/Private/MeshHierarchyCmd.h
new file mode 100644
--- /dev/null
+++ b/Engine/Source/Editor/SkeletalMeshEditor/Private/MeshHierarchyCmd.h
@@ -0,0 +1,12 @@
+// Copyright 1998-2016 Epic Games, Inc. All Rights Reserved.
+
+#pragma once
+
+#include "CoreMisc.h"
+
+class FMeshHierarchyCmd : private FSelfRegisteringExec
+{
+public:
+	/** Console commands, see embeded usage statement **/
+	virtual bool Exec(UWorld*, const TCHAR* Cmd, FOutputDevice& Ar) override;
+};
\ No newline at end of file
diff --git a/Engine/Source/Editor/SkeletalMeshEditor/Private/SkeletalMeshEditorCommands.cpp b/Engine/Source/Editor/SkeletalMeshEditor/Private/SkeletalMeshEditorCommands.cpp
new file mode 100644
--- /dev/null
+++ b/Engine/Source/Editor/SkeletalMeshEditor/Private/SkeletalMeshEditorCommands.cpp
@@ -0,0 +1,14 @@
+// Copyright 1998-2016 Epic Games, Inc. All Rights Reserved.
+
+
+#include "SkeletalMeshEditorPrivatePCH.h"
+#include "SkeletalMeshEditorCommands.h"
+
+#define LOCTEXT_NAMESPACE "SkeletalMeshEditorCommands"
+
+void FSkeletalMeshEditorCommands::RegisterCommands()
+{
+	UI_COMMAND(ReimportMesh, "Reimport Mesh", "Reimport the current mesh.", EUserInterfaceActionType::Button, FInputChord());
+}
+
+#undef LOCTEXT_NAMESPACE
\ No newline at end of file
diff --git a/Engine/Source/Editor/SkeletalMeshEditor/Private/SkeletalMeshEditorCommands.h b/Engine/Source/Editor/SkeletalMeshEditor/Private/SkeletalMeshEditorCommands.h
new file mode 100644
--- /dev/null
+++ b/Engine/Source/Editor/SkeletalMeshEditor/Private/SkeletalMeshEditorCommands.h
@@ -0,0 +1,19 @@
+// Copyright 1998-2016 Epic Games, Inc. All Rights Reserved.
+
+#pragma once
+
+class FSkeletalMeshEditorCommands : public TCommands<FSkeletalMeshEditorCommands>
+{
+public:
+	FSkeletalMeshEditorCommands()
+		: TCommands<FSkeletalMeshEditorCommands>(TEXT("SkeletalMeshEditor"), NSLOCTEXT("Contexts", "SkeletalMeshEditor", "Skeletal Mesh Editor"), NAME_None, FEditorStyle::GetStyleSetName())
+	{
+	}
+
+	virtual void RegisterCommands() override;
+
+public:
+
+	// reimport current mesh
+	TSharedPtr<FUICommandInfo> ReimportMesh;
+};
diff --git a/Engine/Source/Editor/SkeletalMeshEditor/Private/SkeletalMeshEditorMode.cpp b/Engine/Source/Editor/SkeletalMeshEditor/Private/SkeletalMeshEditorMode.cpp
new file mode 100644
--- /dev/null
+++ b/Engine/Source/Editor/SkeletalMeshEditor/Private/SkeletalMeshEditorMode.cpp
@@ -0,0 +1,78 @@
+// Copyright 1998-2016 Epic Games, Inc. All Rights Reserved.
+
+#include "SkeletalMeshEditorPrivatePCH.h"
+#include "SkeletalMeshEditorMode.h"
+#include "SkeletalMeshEditor.h"
+#include "ISkeletonEditorModule.h"
+#include "IPersonaToolkit.h"
+
+FSkeletalMeshEditorMode::FSkeletalMeshEditorMode(TSharedRef<FWorkflowCentricApplication> InHostingApp, TSharedRef<ISkeletonTree> InSkeletonTree)
+	: FApplicationMode(SkeletalMeshEditorModes::SkeletalMeshEditorMode)
+{
+	HostingAppPtr = InHostingApp;
+
+	TSharedRef<FSkeletalMeshEditor> SkeletalMeshEditor = StaticCastSharedRef<FSkeletalMeshEditor>(InHostingApp);
+
+	ISkeletonEditorModule& SkeletonEditorModule = FModuleManager::LoadModuleChecked<ISkeletonEditorModule>("SkeletonEditor");
+	TabFactories.RegisterFactory(SkeletonEditorModule.CreateSkeletonTreeTabFactory(InHostingApp, InSkeletonTree));
+
+	FPersonaModule& PersonaModule = FModuleManager::LoadModuleChecked<FPersonaModule>("Persona");
+	TabFactories.RegisterFactory(PersonaModule.CreateDetailsTabFactory(InHostingApp, FOnDetailsCreated::CreateSP(&SkeletalMeshEditor.Get(), &FSkeletalMeshEditor::HandleDetailsCreated)));
+	TabFactories.RegisterFactory(PersonaModule.CreatePersonaViewportTabFactory(InHostingApp, InSkeletonTree, SkeletalMeshEditor->GetPersonaToolkit()->GetPreviewScene(), SkeletalMeshEditor->OnPostUndo, nullptr, FOnViewportCreated(), true, true));
+	TabFactories.RegisterFactory(PersonaModule.CreateAdvancedPreviewSceneTabFactory(InHostingApp, SkeletalMeshEditor->GetPersonaToolkit()->GetPreviewScene()));
+	TabFactories.RegisterFactory(PersonaModule.CreateAssetDetailsTabFactory(InHostingApp, FOnGetAsset::CreateSP(&SkeletalMeshEditor.Get(), &FSkeletalMeshEditor::HandleGetAsset), FOnDetailsCreated::CreateSP(&SkeletalMeshEditor.Get(), &FSkeletalMeshEditor::HandleMeshDetailsCreated)));
+	TabFactories.RegisterFactory(PersonaModule.CreateMorphTargetTabFactory(InHostingApp, SkeletalMeshEditor->GetPersonaToolkit()->GetPreviewScene(), SkeletalMeshEditor->OnPostUndo));
+
+	TabLayout = FTabManager::NewLayout("Standalone_SkeletalMeshEditor_Layout_v3")
+		->AddArea
+		(
+			FTabManager::NewPrimaryArea()
+			->SetOrientation(Orient_Vertical)
+			->Split
+			(
+				FTabManager::NewStack()
+				->SetSizeCoefficient(0.1f)
+				->SetHideTabWell(true)
+				->AddTab(InHostingApp->GetToolbarTabId(), ETabState::OpenedTab)
+			)
+			->Split
+			(
+				FTabManager::NewSplitter()
+				->SetSizeCoefficient(0.9f)
+				->SetOrientation(Orient_Horizontal)
+				->Split
+				(
+					FTabManager::NewStack()
+					->SetSizeCoefficient(0.2f)
+					->SetHideTabWell(false)
+					->AddTab(SkeletalMeshEditorTabs::SkeletonTreeTab, ETabState::ClosedTab)
+					->AddTab(SkeletalMeshEditorTabs::AssetDetailsTab, ETabState::OpenedTab)
+				)
+				->Split
+				(
+					FTabManager::NewStack()
+					->SetSizeCoefficient(0.6f)
+					->SetHideTabWell(true)
+					->AddTab(SkeletalMeshEditorTabs::ViewportTab, ETabState::OpenedTab)
+				)
+				->Split
+				(
+					FTabManager::NewStack()
+					->SetSizeCoefficient(0.2f)
+					->SetHideTabWell(false)
+					->AddTab(SkeletalMeshEditorTabs::MorphTargetsTab, ETabState::OpenedTab)
+					->AddTab(SkeletalMeshEditorTabs::DetailsTab, ETabState::ClosedTab)
+					->AddTab(SkeletalMeshEditorTabs::AdvancedPreviewTab, ETabState::ClosedTab)
+				)
+			)
+		);
+}
+
+void FSkeletalMeshEditorMode::RegisterTabFactories(TSharedPtr<FTabManager> InTabManager)
+{
+	TSharedPtr<FWorkflowCentricApplication> HostingApp = HostingAppPtr.Pin();
+	HostingApp->RegisterTabSpawners(InTabManager.ToSharedRef());
+	HostingApp->PushTabFactories(TabFactories);
+
+	FApplicationMode::RegisterTabFactories(InTabManager);
+}
\ No newline at end of file
diff --git a/Engine/Source/Editor/SkeletalMeshEditor/Private/SkeletalMeshEditorMode.h b/Engine/Source/Editor/SkeletalMeshEditor/Private/SkeletalMeshEditorMode.h
new file mode 100644
--- /dev/null
+++ b/Engine/Source/Editor/SkeletalMeshEditor/Private/SkeletalMeshEditorMode.h
@@ -0,0 +1,42 @@
+// Copyright 1998-2016 Epic Games, Inc. All Rights Reserved.
+
+#pragma once
+
+#include "ApplicationMode.h"
+#include "WorkflowCentricApplication.h"
+#include "PersonaModule.h"
+#include "IDocumentation.h"
+
+class FSkeletalMeshEditorMode : public FApplicationMode
+{
+public:
+	FSkeletalMeshEditorMode(TSharedRef<class FWorkflowCentricApplication> InHostingApp, TSharedRef<class ISkeletonTree> InSkeletonTree);
+
+	/** FApplicationMode interface */
+	virtual void RegisterTabFactories(TSharedPtr<FTabManager> InTabManager) override;
+
+protected:
+	/** The hosting app */
+	TWeakPtr<class FWorkflowCentricApplication> HostingAppPtr;
+
+	/** The tab factories we support */
+	FWorkflowAllowedTabSet TabFactories;
+};
+
+struct FMeshPropertiesSummoner : public FWorkflowTabFactory
+{
+	FMeshPropertiesSummoner(TSharedPtr<class FAssetEditorToolkit> InHostingApp, FOnGetAsset InOnGetAsset);
+
+	// FWorkflowTabFactory interface
+	virtual TSharedRef<SWidget> CreateTabBody(const FWorkflowTabSpawnInfo& Info) const override;
+
+	// Create a tooltip widget for the tab
+	virtual TSharedPtr<SToolTip> CreateTabToolTipWidget(const FWorkflowTabSpawnInfo& Info) const override
+	{
+		return  IDocumentation::Get()->CreateToolTip(NSLOCTEXT("MeshProperties", "MeshDetailsTooltip", "The Mesh Details tab lets you edit properties (materials etc) of the current Skeletal Mesh."), NULL, TEXT("Shared/Editors/Persona"), TEXT("MeshDetail_Window"));
+	}
+	// FWorkflowTabFactory interface
+
+private:
+	FOnGetAsset OnGetAsset;
+};
\ No newline at end of file
diff --git a/Engine/Source/Editor/SkeletalMeshEditor/Private/SkeletalMeshEditorPrivatePCH.h b/Engine/Source/Editor/SkeletalMeshEditor/Private/SkeletalMeshEditorPrivatePCH.h
new file mode 100644
--- /dev/null
+++ b/Engine/Source/Editor/SkeletalMeshEditor/Private/SkeletalMeshEditorPrivatePCH.h
@@ -0,0 +1,8 @@
+// Copyright 1998-2016 Epic Games, Inc. All Rights Reserved.
+
+#pragma once
+
+#include "UnrealEd.h"
+#include "SlateBasics.h"
+#include "Toolkits/AssetEditorToolkit.h"
+#include "Toolkits/IToolkit.h"
\ No newline at end of file
diff --git a/Engine/Source/Editor/SkeletalMeshEditor/Public/ISkeletalMeshEditor.h b/Engine/Source/Editor/SkeletalMeshEditor/Public/ISkeletalMeshEditor.h
new file mode 100644
--- /dev/null
+++ b/Engine/Source/Editor/SkeletalMeshEditor/Public/ISkeletalMeshEditor.h
@@ -0,0 +1,11 @@
+// Copyright 1998-2016 Epic Games, Inc. All Rights Reserved.
+
+#pragma once
+
+#include "WorkflowCentricApplication.h"
+#include "IHasPersonaToolkit.h"
+
+class ISkeletalMeshEditor : public FWorkflowCentricApplication, public IHasPersonaToolkit
+{
+};
+
diff --git a/Engine/Source/Editor/SkeletalMeshEditor/SkeletalMeshEditor.Build.cs b/Engine/Source/Editor/SkeletalMeshEditor/SkeletalMeshEditor.Build.cs
new file mode 100644
--- /dev/null
+++ b/Engine/Source/Editor/SkeletalMeshEditor/SkeletalMeshEditor.Build.cs
@@ -0,0 +1,43 @@
+// Copyright 1998-2016 Epic Games, Inc. All Rights Reserved.
+
+using UnrealBuildTool;
+
+public class SkeletalMeshEditor : ModuleRules
+{
+	public SkeletalMeshEditor(TargetInfo Target)
+	{
+		PrivateDependencyModuleNames.AddRange(
+			new string[] {
+				"Core",
+				"CoreUObject",
+				"Engine",
+                "InputCore",
+				"Slate",
+				"SlateCore",
+                "EditorStyle",
+                "UnrealEd",
+                "Persona",
+                "SkeletonEditor",
+                "Kismet",
+                "KismetWidgets",
+                "ActorPickerMode",
+                "SceneDepthPickerMode",
+                "MainFrame",
+                "DesktopPlatform",
+                "PropertyEditor",
+                "RHI",
+            }
+		);
+
+        PrivateIncludePathModuleNames.AddRange(
+            new string[] {
+                "PropertyEditor",
+            }
+        );
+
+		DynamicallyLoadedModuleNames.AddRange(
+			new string[] {
+            }
+		);
+	}
+}
diff --git a/Engine/Source/Editor/SkeletonEditor/Private/EditableSkeleton.h b/Engine/Source/Editor/SkeletonEditor/Private/EditableSkeleton.h
new file mode 100644
--- /dev/null
+++ b/Engine/Source/Editor/SkeletonEditor/Private/EditableSkeleton.h
@@ -0,0 +1,144 @@
+// Copyright 1998-2016 Epic Games, Inc. All Rights Reserved.
+
+#pragma once
+
+#include "IEditableSkeleton.h"
+#include "Animation/Skeleton.h"
+
+/** View-model for a skeleton tree */
+class FEditableSkeleton : public IEditableSkeleton, public FGCObject, public TSharedFromThis<FEditableSkeleton>
+{
+public:
+	/** String used as a header for text based copy-paste of sockets */
+	static const FString SocketCopyPasteHeader;
+
+public:
+	FEditableSkeleton(USkeleton* InSkeleton, USkeletalMesh* InSkeletalMesh = nullptr);
+
+	/** IEditableSkeleton interface */
+	virtual const USkeleton& GetSkeleton() const override;
+	virtual const TArray<class UBlendProfile*>& GetBlendProfiles() const override;
+	virtual class UBlendProfile* GetBlendProfile(const FName& InBlendProfileName) override;
+	virtual class UBlendProfile* CreateNewBlendProfile(const FName& InBlendProfileName) override;
+	virtual void RemoveBlendProfile(UBlendProfile* InBlendProfile) override;
+	virtual void SetBlendProfileScale(const FName& InBlendProfileName, const FName& InBoneName, float InNewScale, bool bInRecurse) override;
+	virtual USkeletalMeshSocket* DuplicateSocket(const FSelectedSocketInfo& SocketInfoToDuplicate, const FName& NewParentBoneName) override;
+	virtual int32 ValidatePreviewAttachedObjects() override;
+	virtual int32 DeleteAnimNotifies(const TArray<FName>& InSelectedNotifyNames) override;
+	virtual int32 RenameNotify(const FName& NewName, const FName& OldName) override;
+	virtual void GetCompatibleAnimSequences(TArray<class FAssetData>& OutAssets) override;
+	virtual void RenameSocket(const FName& OldSocketName, const FName& NewSocketName) override;
+	virtual void SetSocketParent(const FName& SocketName, const FName& NewParentName) override;
+	virtual bool DoesSocketAlreadyExist(const class USkeletalMeshSocket* InSocket, const FText& InSocketName, ESocketParentType SocketParentType) const override;
+	virtual bool AddSmartname(const FName& InContainerName, const FName& InNewName, FSmartName& OutSmartName) override;
+	virtual void RenameSmartname(const FName& InContainerName, SmartName::UID_Type InNameUid, const FName& InNewName) override;
+	virtual void RemoveSmartname(const FName& InContainerName, SmartName::UID_Type InNameUid) override;
+	virtual void RemoveSmartnamesAndFixupAnimations(const FName& InContainerName, const TArray<SmartName::UID_Type>& InNameUids) override;
+	virtual void SetCurveMetaDataMaterial(const FSmartName& CurveName, bool bOverrideMaterial) override;
+	virtual void SetPreviewMesh(class USkeletalMesh* InSkeletalMesh) override;
+	virtual void LoadAdditionalPreviewSkeletalMeshes() override;
+	virtual void SetAdditionalPreviewSkeletalMeshes(class UPreviewMeshCollection* InSkeletalMeshes) override;
+	virtual void RenameRetargetSource(const FName& InOldName, const FName& InNewName) override;
+	virtual void AddRetargetSource(const FName& InName, USkeletalMesh* InReferenceMesh) override;
+	virtual void DeleteRetargetSources(const TArray<FName>& InRetargetSourceNames) override;
+	virtual void RefreshRetargetSources(const TArray<FName>& InRetargetSourceNames) override;
+	virtual void RefreshRigConfig() override;
+	virtual void SetRigConfig(URig* InRig) override;
+	virtual void SetRigBoneMapping(const FName& InNodeName, const FName& InBoneName) override;
+	virtual void SetRigBoneMappings(const TMap<FName, FName>& InMappings) override;
+	virtual void RemoveUnusedBones() override;
+	virtual void UpdateSkeletonReferencePose(USkeletalMesh* InSkeletalMesh) override;
+	virtual bool AddSlotGroupName(const FName& InSlotName) override;
+	virtual void SetSlotGroupName(const FName& InSlotName, const FName& InGroupName) override;
+	virtual void DeleteSlotName(const FName& InSlotName) override;
+	virtual void DeleteSlotGroup(const FName& InGroupName) override;
+	virtual void RenameSlotName(const FName& InOldSlotName, const FName& InNewSlotName) override;
+	virtual FDelegateHandle RegisterOnSmartNameRemoved(const FOnSmartNameRemoved::FDelegate& InOnSmartNameRemoved) override;
+	virtual void UnregisterOnSmartNameRemoved(FDelegateHandle InHandle) override;
+
+	/** FGCObject interface */
+	virtual void AddReferencedObjects(FReferenceCollector& Collector) override;
+
+	/** Wrap USkeleton::SetBoneTranslationRetargetingMode */
+	void SetBoneTranslationRetargetingMode(FName InBoneName, EBoneTranslationRetargetingMode::Type NewRetargetingMode);
+
+	/** Wrap USkeleton::GetBoneTranslationRetargetingMode */
+	EBoneTranslationRetargetingMode::Type GetBoneTranslationRetargetingMode(FName InBoneName);
+
+	/** Generates a unique socket name from the input name, by changing the FName's number */
+	FName GenerateUniqueSocketName(FName InName);
+
+	/** Handle the user pasting sockets */
+	void HandlePasteSockets(const FName& InBoneName);
+
+	/** Handles adding a socket to the specified bone (i.e. skeleton, not mesh) */
+	USkeletalMeshSocket* HandleAddSocket(const FName& InBoneName);
+
+	/** Function to customize a socket - this essentially copies a socket from the skeleton to the mesh */
+	void HandleCustomizeSocket(USkeletalMeshSocket* InSocketToCustomize);
+
+	/** Function to promote a socket - this essentially copies a socket from the mesh to the skeleton */
+	void HandlePromoteSocket(USkeletalMeshSocket* InSocketToPromote);
+
+	/** Handle removing all attached assets, optionally keeping a preview scene in sync */
+	void HandleRemoveAllAssets(TSharedPtr<class IPersonaPreviewScene> InPreviewScene);
+
+	/** Handle attaching assets to the skeleton or mesh, optionally keeping a preview scene in sync */
+	void HandleAttachAssets(const TArray<UObject*>& InObjects, const FName& InAttachToName, bool bAttachToMesh, TSharedPtr<class IPersonaPreviewScene> InPreviewScene);
+
+	/** Handle deleting attached assets, optionally keeping a preview scene in sync */
+	void HandleDeleteAttachedAssets(const TArray<FPreviewAttachedObjectPair>& InAttachedObjects, TSharedPtr<class IPersonaPreviewScene> InPreviewScene);
+
+	/** Handle deleting sockets, optionally keeping a preview scene in sync */
+	void HandleDeleteSockets(const TArray<FSelectedSocketInfo>& InSocketInfo, TSharedPtr<class IPersonaPreviewScene> InPreviewScene);
+
+	/** Set Bone Translation Retargeting Mode for the passed-in bones and their children. */
+	void SetBoneTranslationRetargetingModeRecursive(const TArray<FName>& InBoneNames, EBoneTranslationRetargetingMode::Type NewRetargetingMode);
+
+	/** Sets the blend scale for the selected bones and all of their children */
+	void RecursiveSetBlendProfileScales(const FName& InBlendProfileName, const TArray<FName>& InBoneNames, float InScaleToSet);
+
+	/** Create a new skeleton tree to edit this editable skeleton */
+	TSharedRef<class ISkeletonTree> CreateSkeletonTree(const struct FSkeletonTreeArgs& InSkeletonTreeArgs);
+
+	/** Create a new blend profile picker to edit this editable skeleton's blend profiles */
+	TSharedRef<class SWidget> CreateBlendProfilePicker(const struct FBlendProfilePickerArgs& InArgs);
+
+	/** Check whether we have any widgets editing our data */
+	bool IsEdited() const { return SkeletonTrees.Num() > 0 || BlendProfilePickers.Num() > 0; }
+
+	/** Access a const version of the skeletal mesh we are optionally editing */
+	class USkeletalMesh* const GetSkeletalMesh() const { return SkeletalMesh;  }
+
+	/** Register for skeleton changes */
+	void RegisterOnSkeletonHierarchyChanged(const USkeleton::FOnSkeletonHierarchyChanged& InDelegate);
+
+	/** Unregister for skeleton changes */
+	void UnregisterOnSkeletonHierarchyChanged(void* Thing);
+
+	/** Wrap USkeleton::RecreateBoneTree */
+	void RecreateBoneTree(USkeletalMesh* NewPreviewMesh);
+
+private:
+	/** Helper function for deleting attached objects */
+	void DeleteAttachedObjects(FPreviewAssetAttachContainer& AttachedAssets, TSharedPtr<class IPersonaPreviewScene> InPreviewScene);
+
+private:
+	/** The skeleton we are editing */
+	class USkeleton* Skeleton;
+
+	/** Optionally the skeletal mesh we are editing */
+	class USkeletalMesh* SkeletalMesh;
+
+	/** All skeleton tree widgets that are editing this skeleton */
+	TArray<TWeakPtr<class SSkeletonTree>> SkeletonTrees;
+
+	/** All blend profile widgets that are editing this skeleton */
+	TArray<TWeakPtr<class SBlendProfilePicker>> BlendProfilePickers;
+
+	/** Delegate called when trees need refreshing */
+	FSimpleMulticastDelegate OnTreeRefresh;
+
+	/** Delegate called when a smart name is removed */
+	FOnSmartNameRemoved OnSmartNameRemoved;
+};
\ No newline at end of file
diff --git a/Engine/Source/Editor/SkeletonEditor/Private/SBlendProfilePicker.cpp b/Engine/Source/Editor/SkeletonEditor/Private/SBlendProfilePicker.cpp
new file mode 100644
--- /dev/null
+++ b/Engine/Source/Editor/SkeletonEditor/Private/SBlendProfilePicker.cpp
@@ -0,0 +1,356 @@
+// Copyright 1998-2016 Epic Games, Inc. All Rights Reserved.
+
+#include "SkeletonEditorPrivatePCH.h"
+#include "SBlendProfilePicker.h"
+#include "Animation/BlendProfile.h"
+#include "STextEntryPopup.h"
+#include "ScopedTransaction.h"
+#include "IEditableSkeleton.h"
+
+#define LOCTEXT_NAMESPACE "BlendProfilePicker"
+
+
+class SBlendProfileMenuEntry : public SCompoundWidget
+{
+public:
+
+	SLATE_BEGIN_ARGS(SBlendProfileMenuEntry){}
+		SLATE_ARGUMENT( FText, LabelOverride )
+		/** Called to when an entry is clicked */
+		SLATE_EVENT( FExecuteAction, OnOpenClickedDelegate )
+		/** Called to when the button remove an entry is clicked */
+		SLATE_EVENT( FExecuteAction, OnRemoveClickedDelegate )
+	SLATE_END_ARGS()
+
+	void Construct( const FArguments& InArgs )
+	{
+		const FText DisplayName = InArgs._LabelOverride;
+		OnOpenClickedDelegate = InArgs._OnOpenClickedDelegate;
+		OnRemoveClickedDelegate = InArgs._OnRemoveClickedDelegate;
+
+		FSlateFontInfo MenuEntryFont = FEditorStyle::GetFontStyle( "Menu.Label.Font" );
+
+		ChildSlot
+		[
+			SNew(SButton)
+			.ButtonStyle( FEditorStyle::Get(), "Menu.Button" )
+			.ForegroundColor( TAttribute<FSlateColor>::Create( TAttribute<FSlateColor>::FGetter::CreateRaw( this, &SBlendProfileMenuEntry::InvertOnHover ) ) )
+			.Text(DisplayName)
+			.ToolTipText(LOCTEXT("OpenBlendProfileToolTip", "Select this profile for editing."))
+			.OnClicked(this, &SBlendProfileMenuEntry::OnOpen)
+			.VAlign(VAlign_Center)
+			.HAlign(HAlign_Fill)
+			.ContentPadding( FMargin(4.0, 2.0) )
+			[
+				SNew( SOverlay )
+
+				+SOverlay::Slot()
+				.Padding( FMargin( 12.0, 0.0 ) )
+				.VAlign(VAlign_Center)
+				.HAlign(HAlign_Left)
+				[
+					SNew( STextBlock )
+					.Font( MenuEntryFont )
+					.ColorAndOpacity( TAttribute<FSlateColor>::Create( TAttribute<FSlateColor>::FGetter::CreateRaw( this, &SBlendProfileMenuEntry::InvertOnHover ) ) )
+					.Text( DisplayName )
+				]
+
+				+SOverlay::Slot()
+				.Padding( FMargin( 0.0, 0.0 ) )
+				.VAlign(VAlign_Center)
+				.HAlign(HAlign_Right)
+				[
+					SNew(SButton)
+					.ContentPadding( FMargin(4.0, 0.0) )
+					.ButtonStyle( FEditorStyle::Get(), "Docking.Tab.CloseButton" )
+					.ToolTipText( FText::Format( LOCTEXT("RemoveBlendProfileToolTipFmt", "Remove {0}"), DisplayName ) )
+					.OnClicked(this, &SBlendProfileMenuEntry::OnRemove)
+				]
+			]
+		];
+	}
+
+	FReply OnOpen()
+	{
+		OnOpenClickedDelegate.ExecuteIfBound();
+		FSlateApplication::Get().DismissAllMenus();
+		return FReply::Handled();
+	}
+
+	FReply OnRemove()
+	{
+		OnRemoveClickedDelegate.ExecuteIfBound();
+		FSlateApplication::Get().DismissAllMenus();
+		return FReply::Handled();
+	}
+
+private:
+	FSlateColor InvertOnHover() const
+	{
+		if ( this->IsHovered() )
+		{
+			return FLinearColor::Black;
+		}
+		else
+		{
+			return FSlateColor::UseForeground();
+		}
+	}
+
+	FExecuteAction OnOpenClickedDelegate;
+	FExecuteAction OnRemoveClickedDelegate;
+};
+
+//////////////////////////////////////////////////////////////////////////
+
+BEGIN_SLATE_FUNCTION_BUILD_OPTIMIZATION
+void SBlendProfilePicker::Construct(const FArguments& InArgs, TSharedRef<class IEditableSkeleton> InEditableSkeleton)
+{
+	bShowNewOption = InArgs._AllowNew;
+	bShowClearOption = InArgs._AllowClear;
+	bIsStandalone = InArgs._Standalone;
+	EditableSkeleton = InEditableSkeleton;
+
+	UEditorEngine* Editor = Cast<UEditorEngine>(GEngine);
+	if (Editor != nullptr)
+	{
+		Editor->RegisterForUndo(this);
+	}
+
+	if(InArgs._InitialProfile != nullptr && InEditableSkeleton->GetBlendProfiles().Contains(InArgs._InitialProfile))
+	{
+		SelectedProfileName = InArgs._InitialProfile->GetFName();
+	}
+	else
+	{
+		SelectedProfileName = NAME_None;
+	}
+
+	BlendProfileSelectedDelegate = InArgs._OnBlendProfileSelected;
+
+	TSharedRef<SWidget> TextBlock = SNew(STextBlock)
+		.TextStyle(FEditorStyle::Get(), "PropertyEditor.AssetClass")
+		.Font(FEditorStyle::GetFontStyle("PropertyWindow.NormalFont"))
+		.Text(this, &SBlendProfilePicker::GetSelectedProfileName);
+
+	TSharedPtr<SWidget> ButtonContent;
+
+	if (bIsStandalone)
+	{
+		ButtonContent = SNew(SHorizontalBox)
+			+SHorizontalBox::Slot()
+			.AutoWidth()
+			.Padding(0.0f, 2.0f)
+			.VAlign(VAlign_Center)
+			[
+				SNew(SImage)
+				.Image(FEditorStyle::GetBrush("SkeletonTree.BlendProfile"))
+			]
+			+SHorizontalBox::Slot()
+			.AutoWidth()
+			.Padding(2, 0, 0, 0)
+			.VAlign(VAlign_Center)
+			[
+				TextBlock
+			];
+	}
+	else
+	{
+		ButtonContent = TextBlock;
+	}
+
+	ChildSlot
+	[
+		SNew(SComboButton)
+		.ButtonStyle(FEditorStyle::Get(), "PropertyEditor.AssetComboStyle")
+		.ForegroundColor(FEditorStyle::GetColor("PropertyEditor.AssetName.ColorAndOpacity"))
+		.ContentPadding(2.0f)
+		.OnGetMenuContent(this, &SBlendProfilePicker::GetMenuContent)
+		.ButtonContent()
+		[
+			ButtonContent.ToSharedRef()
+		]
+	];
+}
+END_SLATE_FUNCTION_BUILD_OPTIMIZATION
+
+SBlendProfilePicker::~SBlendProfilePicker()
+{
+	UEditorEngine* Editor = Cast<UEditorEngine>(GEngine);
+	if (Editor != nullptr)
+	{
+		Editor->UnregisterForUndo(this);
+	}
+}
+
+FText SBlendProfilePicker::GetSelectedProfileName() const
+{
+	UBlendProfile* SelectedProfile = EditableSkeleton.Pin()->GetBlendProfile(SelectedProfileName);
+	if(SelectedProfile)
+	{
+		if (bIsStandalone)
+		{
+			return FText::Format(FText(LOCTEXT("SelectedNameEntryStandalone", "Blend Profile: {0}")), FText::FromName(SelectedProfileName));
+		}
+		else
+		{
+			return FText::Format(FText(LOCTEXT("SelectedNameEntry", "{0}")), FText::FromName(SelectedProfileName));
+		}
+	}
+	if (bIsStandalone)
+	{
+		return FText(LOCTEXT("NoSelectionEntryStandalone", "Blend Profile: None"));
+	}
+	else
+	{
+		return FText(LOCTEXT("NoSelectionEntry", "None"));
+	}
+}
+
+TSharedRef<SWidget> SBlendProfilePicker::GetMenuContent()
+{
+	FMenuBuilder MenuBuilder(true, nullptr);
+
+	const bool bHasSettingsSection = bShowNewOption || bShowClearOption;
+
+	if(bHasSettingsSection)
+	{
+		MenuBuilder.BeginSection(NAME_None, LOCTEXT("MenuSettings", "Settings"));
+		{
+			if(bShowNewOption)
+			{
+				MenuBuilder.AddMenuEntry(
+					LOCTEXT("CreateNew", "Create New Blend Profile"),
+					LOCTEXT("CreateNew_ToolTip", "Creates a new blend profile inside the skeleton."),
+					FSlateIcon(),
+					FUIAction(FExecuteAction::CreateSP(this, &SBlendProfilePicker::OnCreateNewProfile)));
+			}
+
+			if(bShowClearOption)
+			{
+				MenuBuilder.AddMenuEntry(
+					LOCTEXT("Clear", "Clear"),
+					LOCTEXT("Clear_ToolTip", "Clear the selected blend profile."),
+					FSlateIcon(),
+					FUIAction(FExecuteAction::CreateSP(this, &SBlendProfilePicker::OnClearSelection)));
+			}
+		}
+		MenuBuilder.EndSection();
+	}
+
+	MenuBuilder.BeginSection(NAME_None, LOCTEXT("Profiles", "Available Blend Profiles"));
+	{
+		if (EditableSkeleton.IsValid())
+		{
+			for (UBlendProfile* Profile : EditableSkeleton.Pin()->GetBlendProfiles())
+			{
+				MenuBuilder.AddWidget(
+					SNew(SBlendProfileMenuEntry)
+						.LabelOverride(FText::FromString(Profile->GetName()))
+						.OnOpenClickedDelegate(FExecuteAction::CreateSP(this, &SBlendProfilePicker::OnProfileSelected, Profile->GetFName()))
+						.OnRemoveClickedDelegate(FExecuteAction::CreateSP(this, &SBlendProfilePicker::OnProfileRemoved, Profile->GetFName())),
+					FText(),
+					true
+				);
+			}
+		}
+	}
+	MenuBuilder.EndSection();
+
+	return MenuBuilder.MakeWidget();
+}
+
+void SBlendProfilePicker::OnProfileSelected(FName InBlendProfileName)
+{
+	SelectedProfileName = InBlendProfileName;
+	BlendProfileSelectedDelegate.ExecuteIfBound(EditableSkeleton.Pin()->GetBlendProfile(SelectedProfileName));
+}
+
+void SBlendProfilePicker::OnClearSelection()
+{
+	SelectedProfileName = NAME_None;
+	BlendProfileSelectedDelegate.ExecuteIfBound(nullptr);
+}
+
+void SBlendProfilePicker::OnProfileRemoved(FName InBlendProfileName)
+{
+	EditableSkeleton.Pin()->RemoveBlendProfile(EditableSkeleton.Pin()->GetBlendProfile(InBlendProfileName));
+	SelectedProfileName = NAME_None;
+	BlendProfileSelectedDelegate.ExecuteIfBound(nullptr);
+}
+
+void SBlendProfilePicker::OnCreateNewProfile()
+{
+	TSharedRef<STextEntryPopup> TextEntry = SNew(STextEntryPopup)
+		.Label(LOCTEXT("NewProfileName", "Profile Name"))
+		.OnTextCommitted(this, &SBlendProfilePicker::OnCreateNewProfileComitted);
+
+	FSlateApplication::Get().PushMenu(
+		AsShared(),
+		FWidgetPath(),
+		TextEntry,
+		FSlateApplication::Get().GetCursorPos(),
+		FPopupTransitionEffect(FPopupTransitionEffect::TypeInPopup));
+}
+
+void SBlendProfilePicker::OnCreateNewProfileComitted(const FText& NewName, ETextCommit::Type CommitType)
+{
+	FSlateApplication::Get().DismissAllMenus();
+
+	if(CommitType == ETextCommit::OnEnter && EditableSkeleton.IsValid())
+	{
+		FScopedTransaction Transaction(LOCTEXT("Trans_NewProfile", "Create new blend profile."));
+
+		FName NameToUse = FName(*NewName.ToString());
+
+		// Only create if we don't have a matching profile
+		if(UBlendProfile* FoundProfile = EditableSkeleton.Pin()->GetBlendProfile(NameToUse))
+		{
+			OnProfileSelected(FoundProfile->GetFName());
+		}
+		else if(UBlendProfile* NewProfile = EditableSkeleton.Pin()->CreateNewBlendProfile(NameToUse))
+		{
+			OnProfileSelected(NewProfile->GetFName());
+		}
+	}
+}
+
+void SBlendProfilePicker::SetSelectedProfile(UBlendProfile* InProfile, bool bBroadcast /*= true*/)
+{
+	if(EditableSkeleton.Pin()->GetBlendProfiles().Contains(InProfile))
+	{
+		SelectedProfileName = InProfile->GetFName();
+		if(bBroadcast)
+		{
+			BlendProfileSelectedDelegate.ExecuteIfBound(InProfile);
+		}
+	}
+	else if(!InProfile)
+	{
+		OnClearSelection();
+	}
+}
+
+UBlendProfile* const SBlendProfilePicker::GetSelectedBlendProfile() const
+{
+	return EditableSkeleton.Pin()->GetBlendProfile(SelectedProfileName);
+}
+
+FName SBlendProfilePicker::GetSelectedBlendProfileName() const
+{
+	UBlendProfile* SelectedProfile = EditableSkeleton.Pin()->GetBlendProfile(SelectedProfileName);
+
+	return SelectedProfile ? SelectedProfile->GetFName() : NAME_None;
+}
+
+void SBlendProfilePicker::PostUndo(bool bSuccess)
+{
+	BlendProfileSelectedDelegate.ExecuteIfBound(EditableSkeleton.Pin()->GetBlendProfile(SelectedProfileName));
+}
+
+void SBlendProfilePicker::PostRedo(bool bSuccess)
+{
+	BlendProfileSelectedDelegate.ExecuteIfBound(EditableSkeleton.Pin()->GetBlendProfile(SelectedProfileName));
+}
+
+#undef LOCTEXT_NAMESPACE
