Copying //UE4/Dev-Editor to //UE4/Dev-Main (Source: //UE4/Dev-Editor @ 3152045)

diff --git a/Engine/Source/Developer/FunctionalTesting/Private/FunctionalTest.cpp b/Engine/Source/Developer/FunctionalTesting/Private/FunctionalTest.cpp
--- a/Engine/Source/Developer/FunctionalTesting/Private/FunctionalTest.cpp
+++ b/Engine/Source/Developer/FunctionalTesting/Private/FunctionalTest.cpp
@@ -73,14 +73,13 @@ namespace
 AFunctionalTest::AFunctionalTest( const FObjectInitializer& ObjectInitializer )
 	: Super(ObjectInitializer)
 	, bIsEnabled(true)
+	, bWarningsAsErrors(false)
 	, Result(EFunctionalTestResult::Invalid)
 	, PreparationTimeLimit(15.0f)
 	, TimeLimit(60.0f)
 	, TimesUpMessage( NSLOCTEXT("FunctionalTest", "DefaultTimesUpMessage", "Time's up!") )
 	, TimesUpResult(EFunctionalTestResult::Failed)
 	, bIsRunning(false)
-	, bInStep(false)
-	, CurrentStepName()
 	, TotalTime(0.f)
 	, RunFrame(0)
 	, StartFrame(0)
@@ -162,6 +161,8 @@ void AFunctionalTest::OnConstruction(const FTransform& Transform)
 
 bool AFunctionalTest::RunTest(const TArray<FString>& Params)
 {
+	FAutomationTestFramework::Get().SetTreatWarningsAsErrors(bWarningsAsErrors);
+
 	FailureMessage = TEXT("");
 	
 	//Do not collect garbage during the test. We force GC at the end.
@@ -313,6 +314,8 @@ void AFunctionalTest::FinishTest(EFunctionalTestResult TestResult, const FString
 	}
 
 	TestFinishedObserver.ExecuteIfBound(this);
+
+	FAutomationTestFramework::Get().SetTreatWarningsAsErrors(TOptional<bool>());
 }
 
 void AFunctionalTest::EndPlay(const EEndPlayReason::Type EndPlayReason)
@@ -477,11 +480,11 @@ void AFunctionalTest::AssertTrue(bool Condition, FString Message, const UObject*
 {
 	if ( !Condition )
 	{
-		LogAssertFail(FString::Printf(TEXT("Assertion failed: '%s' for context '%s'"), *Message, ContextObject ? *ContextObject->GetName() : TEXT("")));
+		LogStep(ELogVerbosity::Error, FString::Printf(TEXT("Assertion failed: '%s' for context '%s'"), *Message, ContextObject ? *ContextObject->GetName() : TEXT("")));
 	}
 	else
 	{
-		LogAssertPass(FString::Printf(TEXT("Assertion passed (%s)"), *Message));
+		LogStep(ELogVerbosity::Log, FString::Printf(TEXT("Assertion passed (%s)"), *Message));
 	}
 }
 
@@ -494,11 +497,11 @@ void AFunctionalTest::AssertIsValid(UObject* Object, FString Message, const UObj
 {
 	if ( !IsValid(Object) )
 	{
-		LogAssertFail(FString::Printf(TEXT("Invalid object: '%s' for context '%s'"), *Message, ContextObject ? *ContextObject->GetName() : TEXT("")));
+		LogStep(ELogVerbosity::Error, FString::Printf(TEXT("Invalid object: '%s' for context '%s'"), *Message, ContextObject ? *ContextObject->GetName() : TEXT("")));
 	}
 	else
 	{
-		LogAssertPass(FString::Printf(TEXT("Valid object: (%s)"), *Message));
+		LogStep(ELogVerbosity::Log, FString::Printf(TEXT("Valid object: (%s)"), *Message));
 	}
 }
 
@@ -506,11 +509,11 @@ void AFunctionalTest::AssertValue_Int(int32 Actual, EComparisonMethod ShouldBe,
 {
 	if ( !PerformComparison(Actual, Expected, ShouldBe) )
 	{
-		LogAssertFail(FString::Printf(TEXT("%s: expected {%d} to be %s {%d} for context '%s'"), *What, Actual, *GetComparisonAsString(ShouldBe), Expected, ContextObject ? *ContextObject->GetName() : TEXT("")));
+		LogStep(ELogVerbosity::Error, FString::Printf(TEXT("%s: expected {%d} to be %s {%d} for context '%s'"), *What, Actual, *GetComparisonAsString(ShouldBe), Expected, ContextObject ? *ContextObject->GetName() : TEXT("")));
 	}
 	else
 	{
-		LogAssertPass(FString::Printf(TEXT("Int assertion passed (%s)"), *What));
+		LogStep(ELogVerbosity::Log, FString::Printf(TEXT("Int assertion passed (%s)"), *What));
 	}
 }
 
@@ -518,11 +521,11 @@ void AFunctionalTest::AssertValue_Float(float Actual, EComparisonMethod ShouldBe
 {
 	if ( !PerformComparison(Actual, Expected, ShouldBe) )
 	{
-		LogAssertFail(FString::Printf(TEXT("%s: expected {%f} to be %s {%f} for context '%s'"), *What, Actual, *GetComparisonAsString(ShouldBe), Expected, ContextObject ? *ContextObject->GetName() : TEXT("")));
+		LogStep(ELogVerbosity::Error, FString::Printf(TEXT("%s: expected {%f} to be %s {%f} for context '%s'"), *What, Actual, *GetComparisonAsString(ShouldBe), Expected, ContextObject ? *ContextObject->GetName() : TEXT("")));
 	}
 	else
 	{
-		LogAssertPass(FString::Printf(TEXT("Float assertion passed (%s)"), *What));
+		LogStep(ELogVerbosity::Log, FString::Printf(TEXT("Float assertion passed (%s)"), *What));
 	}
 }
 
@@ -530,11 +533,11 @@ void AFunctionalTest::AssertValue_DateTime(FDateTime Actual, EComparisonMethod S
 {
 	if ( !PerformComparison(Actual, Expected, ShouldBe) )
 	{
-		LogAssertFail(FString::Printf(TEXT("%s: expected {%s} to be %s {%s} for context '%s'"), *What, *Actual.ToString(), *GetComparisonAsString(ShouldBe), *Expected.ToString(), ContextObject ? *ContextObject->GetName() : TEXT("")));
+		LogStep(ELogVerbosity::Error, FString::Printf(TEXT("%s: expected {%s} to be %s {%s} for context '%s'"), *What, *Actual.ToString(), *GetComparisonAsString(ShouldBe), *Expected.ToString(), ContextObject ? *ContextObject->GetName() : TEXT("")));
 	}
 	else
 	{
-		LogAssertPass(FString::Printf(TEXT("DateTime assertion passed (%s)"), *What));
+		LogStep(ELogVerbosity::Log, FString::Printf(TEXT("DateTime assertion passed (%s)"), *What));
 	}
 }
 
@@ -543,11 +546,11 @@ void AFunctionalTest::AssertEqual_Float(const float Actual, const float Expected
 {
 	if ( !FMath::IsNearlyEqual(Actual, Expected, Tolerance) )
 	{
-		LogAssertFail(FString::Printf(TEXT("Expected '%s' to be {%f}, but it was {%f} within tolerance {%f} for context '%s'"), *What, Expected, Actual, Tolerance, ContextObject ? *ContextObject->GetName() : TEXT("")));
+		LogStep(ELogVerbosity::Error, FString::Printf(TEXT("Expected '%s' to be {%f}, but it was {%f} within tolerance {%f} for context '%s'"), *What, Expected, Actual, Tolerance, ContextObject ? *ContextObject->GetName() : TEXT("")));
 	}
 	else
 	{
-		LogAssertPass(FString::Printf(TEXT("Float assertion passed (%s)"), *What));
+		LogStep(ELogVerbosity::Log, FString::Printf(TEXT("Float assertion passed (%s)"), *What));
 	}
 }
 
@@ -555,11 +558,11 @@ void AFunctionalTest::AssertEqual_Transform(const FTransform Actual, const FTran
 {
 	if ( !Expected.Equals(Actual) )
 	{
-		LogAssertFail(FString::Printf(TEXT("Expected '%s' to be {%s}, but it was {%s} for context '%s'"), *What, *TransformToString(Expected), *TransformToString(Actual), ContextObject ? *ContextObject->GetName() : TEXT("")));
+		LogStep(ELogVerbosity::Error, FString::Printf(TEXT("Expected '%s' to be {%s}, but it was {%s} for context '%s'"), *What, *TransformToString(Expected), *TransformToString(Actual), ContextObject ? *ContextObject->GetName() : TEXT("")));
 	}
 	else
 	{
-		LogAssertPass(FString::Printf(TEXT("Transform assertion passed (%s)"), *What));
+		LogStep(ELogVerbosity::Log, FString::Printf(TEXT("Transform assertion passed (%s)"), *What));
 	}
 }
 
@@ -567,11 +570,11 @@ void AFunctionalTest::AssertNotEqual_Transform(const FTransform Actual, const FT
 {
 	if ( NotExpected.Equals(Actual) )
 	{
-		LogAssertFail(FString::Printf(TEXT("Expected '%s' not to be {%s} for context '%s'"), *What, *TransformToString(NotExpected), ContextObject ? *ContextObject->GetName() : TEXT("")));
+		LogStep(ELogVerbosity::Error, FString::Printf(TEXT("Expected '%s' not to be {%s} for context '%s'"), *What, *TransformToString(NotExpected), ContextObject ? *ContextObject->GetName() : TEXT("")));
 	}
 	else
 	{
-		LogAssertPass(FString::Printf(TEXT("Transform assertion passed (%s)"), *What));
+		LogStep(ELogVerbosity::Log, FString::Printf(TEXT("Transform assertion passed (%s)"), *What));
 	}
 }
 
@@ -579,11 +582,11 @@ void AFunctionalTest::AssertEqual_Rotator(const FRotator Actual, const FRotator
 {
 	if ( !Expected.Equals(Actual) )
 	{
-		LogAssertFail(FString::Printf(TEXT("Expected '%s' to be {%s} but it was {%s} for context '%s'"), *What, *Expected.ToString(), *Actual.ToString(), ContextObject ? *ContextObject->GetName() : TEXT("")));
+		LogStep(ELogVerbosity::Error, FString::Printf(TEXT("Expected '%s' to be {%s} but it was {%s} for context '%s'"), *What, *Expected.ToString(), *Actual.ToString(), ContextObject ? *ContextObject->GetName() : TEXT("")));
 	}
 	else
 	{
-		LogAssertPass(FString::Printf(TEXT("Rotator assertion passed (%s)"), *What));
+		LogStep(ELogVerbosity::Log, FString::Printf(TEXT("Rotator assertion passed (%s)"), *What));
 	}
 }
 
@@ -591,11 +594,11 @@ void AFunctionalTest::AssertNotEqual_Rotator(const FRotator Actual, const FRotat
 {
 	if ( NotExpected.Equals(Actual) )
 	{
-		LogAssertFail(FString::Printf(TEXT("Expected '%s' not to be {%s} for context '%s'"), *What, *NotExpected.ToString(), ContextObject ? *ContextObject->GetName() : TEXT("")));
+		LogStep(ELogVerbosity::Error, FString::Printf(TEXT("Expected '%s' not to be {%s} for context '%s'"), *What, *NotExpected.ToString(), ContextObject ? *ContextObject->GetName() : TEXT("")));
 	}
 	else
 	{
-		LogAssertPass(FString::Printf(TEXT("Rotator assertion passed (%s)"), *What));
+		LogStep(ELogVerbosity::Log, FString::Printf(TEXT("Rotator assertion passed (%s)"), *What));
 	}
 }
 
@@ -603,11 +606,11 @@ void AFunctionalTest::AssertEqual_Vector(const FVector Actual, const FVector Exp
 {
 	if ( !Expected.Equals(Actual, Tolerance) )
 	{
-		LogAssertFail(FString::Printf(TEXT("Expected '%s' to be {%s} but it was {%s} within tolerance {%f} for context '%s'"), *What, *Expected.ToString(), *Actual.ToString(), Tolerance, ContextObject ? *ContextObject->GetName() : TEXT("")));
+		LogStep(ELogVerbosity::Error, FString::Printf(TEXT("Expected '%s' to be {%s} but it was {%s} within tolerance {%f} for context '%s'"), *What, *Expected.ToString(), *Actual.ToString(), Tolerance, ContextObject ? *ContextObject->GetName() : TEXT("")));
 	}
 	else
 	{
-		LogAssertPass(FString::Printf(TEXT("Vector assertion passed (%s)"), *What));
+		LogStep(ELogVerbosity::Log, FString::Printf(TEXT("Vector assertion passed (%s)"), *What));
 	}
 }
 
@@ -615,11 +618,11 @@ void AFunctionalTest::AssertNotEqual_Vector(const FVector Actual, const FVector
 {
 	if ( NotExpected.Equals(Actual) )
 	{
-		LogAssertFail(FString::Printf(TEXT("Expected '%s' not to be {%s} for context '%s'"), *What, *NotExpected.ToString(), ContextObject ? *ContextObject->GetName() : TEXT("")));
+		LogStep(ELogVerbosity::Error, FString::Printf(TEXT("Expected '%s' not to be {%s} for context '%s'"), *What, *NotExpected.ToString(), ContextObject ? *ContextObject->GetName() : TEXT("")));
 	}
 	else
 	{
-		LogAssertPass(FString::Printf(TEXT("Vector assertion passed (%s)"), *What));
+		LogStep(ELogVerbosity::Log, FString::Printf(TEXT("Vector assertion passed (%s)"), *What));
 	}
 }
 
@@ -627,11 +630,11 @@ void AFunctionalTest::AssertEqual_String(const FString Actual, const FString Exp
 {
 	if ( !Expected.Equals(Actual) )
 	{
-		LogAssertFail(FString::Printf(TEXT("Expected '%s' to be {%s} but it was {%s} for context '%s'"), *What, *Expected, *Actual, ContextObject ? *ContextObject->GetName() : TEXT("")));
+		LogStep(ELogVerbosity::Error, FString::Printf(TEXT("Expected '%s' to be {%s} but it was {%s} for context '%s'"), *What, *Expected, *Actual, ContextObject ? *ContextObject->GetName() : TEXT("")));
 	}
 	else
 	{
-		LogAssertPass(FString::Printf(TEXT("String assertion passed (%s)"), *What));
+		LogStep(ELogVerbosity::Log, FString::Printf(TEXT("String assertion passed (%s)"), *What));
 	}
 }
 
@@ -639,15 +642,25 @@ void AFunctionalTest::AssertNotEqual_String(const FString Actual, const FString
 {
 	if ( NotExpected.Equals(Actual) )
 	{
-		LogAssertFail(FString::Printf(TEXT("Expected '%s' not to be {%s} for context '%s'"), *What, *NotExpected, ContextObject ? *ContextObject->GetName() : TEXT("")));
+		LogStep(ELogVerbosity::Error, FString::Printf(TEXT("Expected '%s' not to be {%s} for context '%s'"), *What, *NotExpected, ContextObject ? *ContextObject->GetName() : TEXT("")));
 	}
 	else
 	{
-		LogAssertPass(FString::Printf(TEXT("String assertion passed (%s)"), *What));
+		LogStep(ELogVerbosity::Log, FString::Printf(TEXT("String assertion passed (%s)"), *What));
 	}
 }
 
-void AFunctionalTest::LogAssertFail(const FString& Message)
+void AFunctionalTest::AddWarning(const FString Message)
+{
+	LogStep(ELogVerbosity::Warning, Message);
+}
+
+void AFunctionalTest::AddError(const FString Message)
+{
+	LogStep(ELogVerbosity::Error, Message);
+}
+
+void AFunctionalTest::LogStep(ELogVerbosity::Type Verbosity, const FString& Message)
 {
 	FString FullMessage(Message);
 	if ( IsInStep() )
@@ -660,25 +673,49 @@ void AFunctionalTest::LogAssertFail(const FString& Message)
 		}
 		FullMessage.Append(StepName);
 	}
-	UE_VLOG(this, LogFunctionalTest, Error, TEXT("%s"), *FullMessage);
-	UE_LOG(LogFunctionalTest, Error, TEXT("%s"), *FullMessage);
+
+	switch ( Verbosity )
+	{
+		case ELogVerbosity::Log:
+			UE_VLOG(this, LogFunctionalTest, Log, TEXT("%s"), *FullMessage);
+			UE_LOG(LogFunctionalTest, Log, TEXT("%s"), *FullMessage);
+		break;
+		case ELogVerbosity::Warning:
+			UE_VLOG(this, LogFunctionalTest, Warning, TEXT("%s"), *FullMessage);
+			UE_LOG(LogFunctionalTest, Warning, TEXT("%s"), *FullMessage);
+		break;
+		case ELogVerbosity::Error:
+			UE_VLOG(this, LogFunctionalTest, Error, TEXT("%s"), *FullMessage);
+			UE_LOG(LogFunctionalTest, Error, TEXT("%s"), *FullMessage);
+		break;
+	}
 }
 
-void AFunctionalTest::LogAssertPass(const FString& Message)
+FString	AFunctionalTest::GetCurrentStepName() const
 {
-	FString FullMessage(Message);
-	if ( IsInStep() )
+	return IsInStep() ? Steps.Top() : FString();
+}
+
+void AFunctionalTest::StartStep(const FString& StepName)
+{
+	Steps.Push(StepName);
+}
+
+void AFunctionalTest::FinishStep()
+{
+	if ( Steps.Num() > 0 )
 	{
-		FullMessage.Append(TEXT(" in step: "));
-		FString StepName = TEXT("");
-		if ( StepName.IsEmpty() )
-		{
-			StepName = TEXT("<UN-NAMED STEP>");
-		}
-		FullMessage.Append(StepName);
+		Steps.Pop();
+	}
+	else
+	{
+		AddWarning(TEXT("FinishStep was called when no steps were currently in progress."));
 	}
-	UE_VLOG(this, LogFunctionalTest, Log, TEXT("%s"), *FullMessage);
-	UE_LOG(LogFunctionalTest, Log, TEXT("%s"), *FullMessage);
+}
+
+bool AFunctionalTest::IsInStep() const
+{
+	return Steps.Num() > 0;
 }
 
 //////////////////////////////////////////////////////////////////////////
diff --git a/Engine/Source/Developer/FunctionalTesting/Private/FunctionalTestingPrivatePCH.h b/Engine/Source/Developer/FunctionalTesting/Private/FunctionalTestingPrivatePCH.h
--- a/Engine/Source/Developer/FunctionalTesting/Private/FunctionalTestingPrivatePCH.h
+++ b/Engine/Source/Developer/FunctionalTesting/Private/FunctionalTestingPrivatePCH.h
@@ -22,3 +22,4 @@
 #include "FunctionalTestingManager.h"
 #include "FunctionalTestingModule.h"
 #include "FuncTestRenderingComponent.h"
+#include "TestPhaseComponent.h"
\ No newline at end of file
diff --git a/Engine/Source/Developer/FunctionalTesting/Private/ScreenshotFunctionalTest.cpp b/Engine/Source/Developer/FunctionalTesting/Private/ScreenshotFunctionalTest.cpp
--- a/Engine/Source/Developer/FunctionalTesting/Private/ScreenshotFunctionalTest.cpp
+++ b/Engine/Source/Developer/FunctionalTesting/Private/ScreenshotFunctionalTest.cpp
@@ -5,21 +5,13 @@
 #include "ScreenshotFunctionalTest.h"
 #include "AutomationBlueprintFunctionLibrary.h"
 #include "BufferVisualizationData.h"
+#include "AutomationScreenshotOptions.h"
 
 AScreenshotFunctionalTest::AScreenshotFunctionalTest( const FObjectInitializer& ObjectInitializer )
 	: AFunctionalTest(ObjectInitializer)
 {
 	ScreenshotCamera = CreateDefaultSubobject<UCameraComponent>(TEXT("Camera"));
 	ScreenshotCamera->SetupAttachment(RootComponent);
-
-	VisualizeBuffer = NAME_None;
-
-	Tolerance = EComparisonTolerance::Low;
-	ToleranceAmount = FComparisonToleranceAmount(16, 16, 16, 16, 16, 240);
-	MaximumAllowedError = 0.02f; // 2% Allowed Error, after passing whatever tolerance levels are set for pixel variation.
-
-	bIgnoreAntiAliasing = true;
-	bIgnoreColors = false;
 }
 
 void AScreenshotFunctionalTest::PrepareTest()
@@ -41,7 +33,7 @@ void AScreenshotFunctionalTest::StartTest()
 {
 	Super::StartTest();
 
-	UAutomationBlueprintFunctionLibrary::TakeAutomationScreenshotInternal(GetName(), Resolution);
+	UAutomationBlueprintFunctionLibrary::TakeAutomationScreenshotInternal(GetName(), ScreenshotOptions);
 
 	GEngine->GameViewport->OnScreenshotCaptured().AddUObject(this, &AScreenshotFunctionalTest::OnScreenshotTaken);
 
@@ -68,8 +60,8 @@ void AScreenshotFunctionalTest::SetupVisualizeBuffer()
 			{
 				if ( ViewportClient->GetEngineShowFlags() )
 				{
-					ViewportClient->GetEngineShowFlags()->SetVisualizeBuffer( VisualizeBuffer == NAME_None ? false : true);
-					ICVar->Set(*VisualizeBuffer.ToString());
+					ViewportClient->GetEngineShowFlags()->SetVisualizeBuffer( ScreenshotOptions.VisualizeBuffer == NAME_None ? false : true);
+					ICVar->Set(*ScreenshotOptions.VisualizeBuffer.ToString());
 				}
 			}
 		}
@@ -85,9 +77,9 @@ bool AScreenshotFunctionalTest::CanEditChange(const UProperty* InProperty) const
 	{
 		const FName PropertyName = InProperty->GetFName();
 
-		if ( PropertyName == GET_MEMBER_NAME_CHECKED(AScreenshotFunctionalTest, ToleranceAmount) )
+		if ( PropertyName == GET_MEMBER_NAME_CHECKED(FAutomationScreenshotOptions, ToleranceAmount) )
 		{
-			bIsEditable = Tolerance == EComparisonTolerance::Custom;
+			bIsEditable = ScreenshotOptions.Tolerance == EComparisonTolerance::Custom;
 		}
 		else if ( PropertyName == TEXT("ObservationPoint") )
 		{
@@ -107,20 +99,9 @@ void AScreenshotFunctionalTest::PostEditChangeProperty(FPropertyChangedEvent& Pr
 	{
 		const FName PropertyName = PropertyChangedEvent.Property->GetFName();
 
-		if ( PropertyName == GET_MEMBER_NAME_CHECKED(AScreenshotFunctionalTest, Tolerance) )
+		if ( PropertyName == GET_MEMBER_NAME_CHECKED(FAutomationScreenshotOptions, Tolerance) )
 		{
-			switch ( Tolerance )
-			{
-			case EComparisonTolerance::Zero:
-				ToleranceAmount = FComparisonToleranceAmount(0, 0, 0, 0, 0, 255);
-				break;
-			case EComparisonTolerance::Low:
-				ToleranceAmount = FComparisonToleranceAmount(16, 16, 16, 16, 16, 240);
-				break;
-			case EComparisonTolerance::High:
-				ToleranceAmount = FComparisonToleranceAmount(32, 32, 32, 32, 64, 96);
-				break;
-			}
+			ScreenshotOptions.SetToleranceAmounts(ScreenshotOptions.Tolerance);
 		}
 	}
 }
diff --git a/Engine/Source/Developer/FunctionalTesting/Private/TakeScreenshotAfterTimeLatentAction.cpp b/Engine/Source/Developer/FunctionalTesting/Private/TakeScreenshotAfterTimeLatentAction.cpp
--- a/Engine/Source/Developer/FunctionalTesting/Private/TakeScreenshotAfterTimeLatentAction.cpp
+++ b/Engine/Source/Developer/FunctionalTesting/Private/TakeScreenshotAfterTimeLatentAction.cpp
@@ -3,20 +3,22 @@
 #include "AutomationBlueprintFunctionLibrary.h"
 #include "TakeScreenshotAfterTimeLatentAction.h"
 
-FTakeScreenshotAfterTimeLatentAction::FTakeScreenshotAfterTimeLatentAction(const FLatentActionInfo& LatentInfo, const FString& InScreenshotName, FIntPoint Resolution, float Seconds)
+FTakeScreenshotAfterTimeLatentAction::FTakeScreenshotAfterTimeLatentAction(const FLatentActionInfo& LatentInfo, const FString& InScreenshotName, FAutomationScreenshotOptions InOptions)
 	: ExecutionFunction(LatentInfo.ExecutionFunction)
 	, OutputLink(LatentInfo.Linkage)
 	, CallbackTarget(LatentInfo.CallbackTarget)
+	, DelegateHandle()
 	, ScreenshotName(InScreenshotName)
-	, SecondsRemaining(Seconds)
+	, SecondsRemaining(InOptions.Delay)
 	, IssuedScreenshotCapture(false)
 	, TakenScreenshot(false)
-	, DesiredResolution(Resolution)
+	, Options(InOptions)
 {
 }
 
 FTakeScreenshotAfterTimeLatentAction::~FTakeScreenshotAfterTimeLatentAction()
 {
+	GEngine->GameViewport->OnScreenshotCaptured().RemoveAll(this);
 }
 
 void FTakeScreenshotAfterTimeLatentAction::OnScreenshotTaken(int32 InSizeX, int32 InSizeY, const TArray<FColor>& InImageData)
@@ -35,7 +37,7 @@ void FTakeScreenshotAfterTimeLatentAction::UpdateOperation(FLatentResponse& Resp
 			{
 				DelegateHandle = GEngine->GameViewport->OnScreenshotCaptured().AddRaw(this, &FTakeScreenshotAfterTimeLatentAction::OnScreenshotTaken);
 
-				if ( UAutomationBlueprintFunctionLibrary::TakeAutomationScreenshotInternal(ScreenshotName, DesiredResolution) )
+				if ( UAutomationBlueprintFunctionLibrary::TakeAutomationScreenshotInternal(ScreenshotName, Options) )
 				{
 					IssuedScreenshotCapture = true;
 				}
@@ -49,6 +51,7 @@ void FTakeScreenshotAfterTimeLatentAction::UpdateOperation(FLatentResponse& Resp
 	}
 	else
 	{
+		GEngine->GameViewport->OnScreenshotCaptured().RemoveAll(this);
 		Response.FinishAndTriggerIf(true, ExecutionFunction, OutputLink, CallbackTarget);
 	}
 }
diff --git a/Engine/Source/Developer/FunctionalTesting/Private/TestPhaseComponent.cpp b/Engine/Source/Developer/FunctionalTesting/Private/TestPhaseComponent.cpp
new file mode 100644
--- /dev/null
+++ b/Engine/Source/Developer/FunctionalTesting/Private/TestPhaseComponent.cpp
@@ -0,0 +1,8 @@
+// Copyright 1998-2016 Epic Games, Inc. All Rights Reserved.
+
+#include "FunctionalTestingPrivatePCH.h"
+
+UTestPhaseComponent::UTestPhaseComponent(const FObjectInitializer& ObjectInitializer)
+	: Super(ObjectInitializer)
+{
+}
diff --git a/Engine/Source/Developer/FunctionalTesting/Public/AutomationBlueprintFunctionLibrary.h b/Engine/Source/Developer/FunctionalTesting/Public/AutomationBlueprintFunctionLibrary.h
--- a/Engine/Source/Developer/FunctionalTesting/Public/AutomationBlueprintFunctionLibrary.h
+++ b/Engine/Source/Developer/FunctionalTesting/Public/AutomationBlueprintFunctionLibrary.h
@@ -3,6 +3,7 @@
 
 #include "Core.h"
 #include "Kismet/BlueprintFunctionLibrary.h"
+#include "AutomationScreenshotOptions.h"
 
 #include "AutomationBlueprintFunctionLibrary.generated.h"
 
@@ -12,17 +13,26 @@ class FUNCTIONALTESTING_API UAutomationBlueprintFunctionLibrary : public UBluepr
 	GENERATED_UCLASS_BODY()
 	
 public:
-	static bool TakeAutomationScreenshotInternal(const FString& Name, FIntPoint Resolution);
+	static bool TakeAutomationScreenshotInternal(const FString& Name, FAutomationScreenshotOptions Options);
 
-	UFUNCTION(BlueprintCallable, Category = "StudiosAutomation", meta = (Latent, HidePin = "WorldContextObject", DefaultToSelf = "WorldContextObject", LatentInfo = "LatentInfo", Name = "", DelayBeforeScreenshotSeconds = "0.2" ))
-	static void TakeAutomationScreenshot(UObject* WorldContextObject, FLatentActionInfo LatentInfo, const FString& Name, FIntPoint Resolution, float DelayBeforeScreenshotSeconds);
+	UFUNCTION(BlueprintCallable, Category = "Automation", meta = (Latent, HidePin = "WorldContextObject", DefaultToSelf = "WorldContextObject", LatentInfo = "LatentInfo", Name = "" ))
+	static void TakeAutomationScreenshot(UObject* WorldContextObject, FLatentActionInfo LatentInfo, const FString& Name, FAutomationScreenshotOptions Options);
 
-	UFUNCTION(BlueprintCallable, Category = "StudiosAutomation|Helpers", meta = (Latent, HidePin = "WorldContextObject", DefaultToSelf = "WorldContextObject", LatentInfo = "LatentInfo", NameOverride = "", DelayBeforeScreenshotSeconds = "0.2"))
-	static void TakeAutomationScreenshotAtCamera(UObject* WorldContextObject, FLatentActionInfo LatentInfo, ACameraActor* Camera, const FString& NameOverride, FIntPoint Resolution, float DelayBeforeScreenshotSeconds);
+	UFUNCTION(BlueprintCallable, Category = "Automation", meta = (Latent, HidePin = "WorldContextObject", DefaultToSelf = "WorldContextObject", LatentInfo = "LatentInfo", NameOverride = "" ))
+	static void TakeAutomationScreenshotAtCamera(UObject* WorldContextObject, FLatentActionInfo LatentInfo, ACameraActor* Camera, const FString& NameOverride, FAutomationScreenshotOptions Options);
 
-    UFUNCTION(BlueprintCallable, Category = "StudiosAutomation|Performance")
-    static void BeginPerformanceCapture();
+    //UFUNCTION(BlueprintCallable, Category = "Automation|Performance")
+    //static void BeginPerformanceCapture();
 
-    UFUNCTION(BlueprintCallable, Category = "StudiosAutomation|Performance")
-    static void EndPerformanceCapture();
+    //UFUNCTION(BlueprintCallable, Category = "Automation|Performance")
+    //static void EndPerformanceCapture();
+
+	UFUNCTION(BlueprintPure, Category="Automation")
+	static bool AreAutomatedTestsRunning();
+
+	UFUNCTION(BlueprintPure, Category="Automation")
+	static FAutomationScreenshotOptions GetDefaultScreenshotOptionsForGameplay(EComparisonTolerance Tolerance = EComparisonTolerance::Low);
+
+	UFUNCTION(BlueprintPure, Category="Automation")
+	static FAutomationScreenshotOptions GetDefaultScreenshotOptionsForRendering(EComparisonTolerance Tolerance = EComparisonTolerance::Low);
 };
\ No newline at end of file
diff --git a/Engine/Source/Developer/FunctionalTesting/Public/TakeScreenshotAfterTimeLatentAction.h b/Engine/Source/Developer/FunctionalTesting/Public/TakeScreenshotAfterTimeLatentAction.h
--- a/Engine/Source/Developer/FunctionalTesting/Public/TakeScreenshotAfterTimeLatentAction.h
+++ b/Engine/Source/Developer/FunctionalTesting/Public/TakeScreenshotAfterTimeLatentAction.h
@@ -1,10 +1,12 @@
 // Copyright 1998-2016 Epic Games, Inc. All Rights Reserved.
 #pragma once
 
+#include "AutomationScreenshotOptions.h"
+
 class FTakeScreenshotAfterTimeLatentAction : public FPendingLatentAction
 {
 public:
-	FTakeScreenshotAfterTimeLatentAction(const FLatentActionInfo& LatentInfo, const FString& InScreenshotName, FIntPoint Resolution, float Seconds);
+	FTakeScreenshotAfterTimeLatentAction(const FLatentActionInfo& LatentInfo, const FString& InScreenshotName, FAutomationScreenshotOptions InOptions);
 	virtual ~FTakeScreenshotAfterTimeLatentAction();
 
 	virtual void UpdateOperation(FLatentResponse& Response) override;
@@ -26,5 +28,5 @@ private:
 	float SecondsRemaining;
 	bool IssuedScreenshotCapture;
 	bool TakenScreenshot;
-	FIntPoint DesiredResolution;
+	FAutomationScreenshotOptions Options;
 };
diff --git a/Engine/Source/Developer/LogVisualizer/Private/SVisualLoggerFilters.cpp b/Engine/Source/Developer/LogVisualizer/Private/SVisualLoggerFilters.cpp
--- a/Engine/Source/Developer/LogVisualizer/Private/SVisualLoggerFilters.cpp
+++ b/Engine/Source/Developer/LogVisualizer/Private/SVisualLoggerFilters.cpp
@@ -194,7 +194,7 @@ TSharedRef<SWidget> SVisualLoggerFilters::MakeGraphsFilterMenu()
 		SNew(SVerticalBox)
 
 		+ SVerticalBox::Slot()
-		.MaxHeight(DisplaySize.Y * 0.5)
+		.MaxHeight(DisplaySize.Y * 0.9)
 		[
 			MenuBuilder.MakeWidget()
 		];
diff --git a/Engine/Source/Developer/LogVisualizer/Private/SVisualLoggerTimelinesContainer.cpp b/Engine/Source/Developer/LogVisualizer/Private/SVisualLoggerTimelinesContainer.cpp
--- a/Engine/Source/Developer/LogVisualizer/Private/SVisualLoggerTimelinesContainer.cpp
+++ b/Engine/Source/Developer/LogVisualizer/Private/SVisualLoggerTimelinesContainer.cpp
@@ -33,7 +33,7 @@ TSharedRef<SWidget> SVisualLoggerTimelinesContainer::GetRightClickMenuContent()
 		SNew(SVerticalBox)
 
 		+ SVerticalBox::Slot()
-		.MaxHeight(DisplaySize.Y * 0.5)
+		.MaxHeight(DisplaySize.Y * 0.9)
 		[
 			MenuBuilder.MakeWidget()
 		];
diff --git a/Engine/Source/Editor/BehaviorTreeEditor/Private/BehaviorTreeEditor.cpp b/Engine/Source/Editor/BehaviorTreeEditor/Private/BehaviorTreeEditor.cpp
--- a/Engine/Source/Editor/BehaviorTreeEditor/Private/BehaviorTreeEditor.cpp
+++ b/Engine/Source/Editor/BehaviorTreeEditor/Private/BehaviorTreeEditor.cpp
@@ -38,6 +38,7 @@
 #include "SDockTab.h"
 #include "GenericCommands.h"
 #include "Engine/BlueprintGeneratedClass.h"
+#include "Factories/DataAssetFactory.h"
 
 #define LOCTEXT_NAMESPACE "BehaviorTreeEditor"
 
diff --git a/Engine/Source/Editor/BlueprintGraph/Classes/EdGraphSchema_K2.h b/Engine/Source/Editor/BlueprintGraph/Classes/EdGraphSchema_K2.h
--- a/Engine/Source/Editor/BlueprintGraph/Classes/EdGraphSchema_K2.h
+++ b/Engine/Source/Editor/BlueprintGraph/Classes/EdGraphSchema_K2.h
@@ -410,7 +410,7 @@ public:
 	};
 
 public:
-	void SelectAllInputNodes(UEdGraph* Graph, UEdGraphPin* InGraphPin);
+	void SelectAllNodesInDirection(TEnumAsByte<enum EEdGraphPinDirection> InDirection, UEdGraph* Graph, UEdGraphPin* InGraphPin);
 
 	//~ Begin EdGraphSchema Interface
 	virtual void GetContextMenuActions(const UEdGraph* CurrentGraph, const UEdGraphNode* InGraphNode, const UEdGraphPin* InGraphPin, class FMenuBuilder* MenuBuilder, bool bIsDebugging) const override;
diff --git a/Engine/Source/Editor/BlueprintGraph/Private/EdGraphSchema_K2.cpp b/Engine/Source/Editor/BlueprintGraph/Private/EdGraphSchema_K2.cpp
--- a/Engine/Source/Editor/BlueprintGraph/Private/EdGraphSchema_K2.cpp
+++ b/Engine/Source/Editor/BlueprintGraph/Private/EdGraphSchema_K2.cpp
@@ -1277,7 +1277,7 @@ bool UEdGraphSchema_K2::PinDefaultValueIsEditable(const UEdGraphPin& InGraphPin)
 	return true;
 }
 
-void UEdGraphSchema_K2::SelectAllInputNodes(UEdGraph* Graph, UEdGraphPin* InGraphPin)
+void UEdGraphSchema_K2::SelectAllNodesInDirection(TEnumAsByte<enum EEdGraphPinDirection> InDirection, UEdGraph* Graph, UEdGraphPin* InGraphPin)
 {
 	TArray<UEdGraphPin*> AllPins = InGraphPin->LinkedTo;
 
@@ -1294,13 +1294,9 @@ void UEdGraphSchema_K2::SelectAllInputNodes(UEdGraph* Graph, UEdGraphPin* InGrap
 		TArray<UEdGraphPin*> LinkedPins = Pin->GetOwningNode()->GetAllPins();
 		for (UEdGraphPin* InputPin : LinkedPins)
 		{
-			if (InputPin->Direction == EEdGraphPinDirection::EGPD_Output)
+			if (InputPin->Direction == InDirection)
 			{
-				continue;
-			}
-			else
-			{
-				SelectAllInputNodes(Graph, InputPin);
+				SelectAllNodesInDirection(InDirection, Graph, InputPin);
 			}
 		}
 	}
@@ -1333,10 +1329,10 @@ void UEdGraphSchema_K2::GetContextMenuActions(const UEdGraph* CurrentGraph, cons
 				if (InGraphPin->LinkedTo.Num() > 0)
 				{
 					MenuBuilder->AddMenuEntry(
-						LOCTEXT("SelectAllInputNodes", "Select All Input Nodes"),
-						LOCTEXT("SelectAllInputNodesTooltip", "Adds all input Nodes linked to this Pin to selection"),
+						InGraphPin->Direction == EEdGraphPinDirection::EGPD_Input ? LOCTEXT("SelectAllInputNodes", "Select All Input Nodes") : LOCTEXT("SelectAllOutputNodes", "Select All Output Nodes"),
+						InGraphPin->Direction == EEdGraphPinDirection::EGPD_Input ? LOCTEXT("SelectAllInputNodesTooltip", "Adds all input Nodes linked to this Pin to selection") : LOCTEXT("SelectAllOutputNodesTooltip", "Adds all output Nodes linked to this Pin to selection"),
 						FSlateIcon(),
-						FUIAction(FExecuteAction::CreateUObject((UEdGraphSchema_K2*const)this, &UEdGraphSchema_K2::SelectAllInputNodes, const_cast<UEdGraph*>(CurrentGraph), const_cast<UEdGraphPin*>(InGraphPin))));
+						FUIAction(FExecuteAction::CreateUObject((UEdGraphSchema_K2*const)this, &UEdGraphSchema_K2::SelectAllNodesInDirection, InGraphPin->Direction, const_cast<UEdGraph*>(CurrentGraph), const_cast<UEdGraphPin*>(InGraphPin))));
 
 					if(InGraphPin->LinkedTo.Num() > 1)
 					{
diff --git a/Engine/Source/Editor/Cascade/Private/CascadePreviewViewportClient.cpp b/Engine/Source/Editor/Cascade/Private/CascadePreviewViewportClient.cpp
--- a/Engine/Source/Editor/Cascade/Private/CascadePreviewViewportClient.cpp
+++ b/Engine/Source/Editor/Cascade/Private/CascadePreviewViewportClient.cpp
@@ -730,7 +730,7 @@ void FCascadeEdPreviewViewportClient::UpdateMemoryInformation()
 	{
 		FArchiveCountMem ComponentMemCount(ParticleSystemComponent);
 		PSysCompRootSize = ComponentMemCount.GetMax();
-		PSysCompResourceSize = ParticleSystemComponent->GetResourceSize(EResourceSizeMode::Exclusive);
+		PSysCompResourceSize = ParticleSystemComponent->GetResourceSizeBytes(EResourceSizeMode::Exclusive);
 	}
 }
 
diff --git a/Engine/Source/Editor/Cascade/Private/CascadePreviewViewportClient.h b/Engine/Source/Editor/Cascade/Private/CascadePreviewViewportClient.h
--- a/Engine/Source/Editor/Cascade/Private/CascadePreviewViewportClient.h
+++ b/Engine/Source/Editor/Cascade/Private/CascadePreviewViewportClient.h
@@ -63,7 +63,7 @@ public:
 	FEditorCommonDrawHelper& GetDrawHelper();
 	float& GetWireSphereRadius();
 
-	FCascade* GetCascade(){ return CascadePtr.Pin().Get(); }
+	FCascade* GetCascade(){ return CascadePtr.IsValid() ? CascadePtr.Pin().Get() : nullptr; }
 private:
 	/** Pointer back to the ParticleSystem editor tool that owns us */
 	TWeakPtr<FCascade> CascadePtr;
diff --git a/Engine/Source/Editor/ContentBrowser/Private/AssetViewWidgets.cpp b/Engine/Source/Editor/ContentBrowser/Private/AssetViewWidgets.cpp
--- a/Engine/Source/Editor/ContentBrowser/Private/AssetViewWidgets.cpp
+++ b/Engine/Source/Editor/ContentBrowser/Private/AssetViewWidgets.cpp
@@ -667,7 +667,7 @@ TSharedRef<SWidget> SAssetViewItem::CreateToolTipWidget() const
 
 			int32 PackageNameLengthForCooking = ContentBrowserUtils::GetPackageLengthForCooking(AssetData.PackageName.ToString(), FEngineBuildSettings::IsInternalBuild());
 
-			AddToToolTipInfoBox( InfoBox, LOCTEXT("TileViewTooltipPathLengthForCookingKey", "Cooking Filepath length"), FText::Format(LOCTEXT("TileViewTooltipPathLengthForCookingValue", "{0} / {1}"), 
+			AddToToolTipInfoBox( InfoBox, LOCTEXT("TileViewTooltipPathLengthForCookingKey", "Cooking Filepath Length"), FText::Format(LOCTEXT("TileViewTooltipPathLengthForCookingValue", "{0} / {1}"), 
 								 FText::AsNumber(PackageNameLengthForCooking), FText::AsNumber(ContentBrowserUtils::MaxCookPathLen)), PackageNameLengthForCooking > ContentBrowserUtils::MaxCookPathLen ? true : false);
 
 			
diff --git a/Engine/Source/Editor/ContentBrowser/Private/SAssetView.cpp b/Engine/Source/Editor/ContentBrowser/Private/SAssetView.cpp
--- a/Engine/Source/Editor/ContentBrowser/Private/SAssetView.cpp
+++ b/Engine/Source/Editor/ContentBrowser/Private/SAssetView.cpp
@@ -169,6 +169,7 @@ void SAssetView::Construct( const FArguments& InArgs )
 	bSortByPathInColumnView = bShowPathInColumnView & InArgs._SortByPathInColumnView;
 
 	bPendingUpdateThumbnails = false;
+	bShouldNotifyNextAssetSync = true;
 	CurrentThumbnailSize = TileViewThumbnailSize;
 
 	SourcesData = InArgs._InitialSourcesData;
@@ -406,6 +407,7 @@ void SAssetView::Construct( const FArguments& InArgs )
 	if( InArgs._InitialAssetSelection.IsValid() )
 	{
 		// sync to the initial item without notifying of selection
+		bShouldNotifyNextAssetSync = false;
 		TArray<FAssetData> AssetsToSync;
 		AssetsToSync.Add( InArgs._InitialAssetSelection );
 		SyncToAssets( AssetsToSync );
@@ -917,10 +919,11 @@ void SAssetView::Tick( const FGeometry& AllottedGeometry, const double InCurrent
 			// Don't sync to selection because we are just going to do it below
 			SortList(/*bSyncToSelection=*/false);
 		}
-
+		
 		bBulkSelecting = true;
 		ClearSelection();
 		bool bFoundScrollIntoViewTarget = false;
+
 		for ( auto ItemIt = FilteredAssetItems.CreateConstIterator(); ItemIt; ++ItemIt )
 		{
 			const auto& Item = *ItemIt;
@@ -940,9 +943,18 @@ void SAssetView::Tick( const FGeometry& AllottedGeometry, const double InCurrent
 				}
 			}
 		}
-		
+	
 		bBulkSelecting = false;
 
+		if (bShouldNotifyNextAssetSync && !bUserSearching)
+		{
+			AssetSelectionChanged(TSharedPtr<FAssetViewAsset>(), ESelectInfo::Direct);
+		}
+
+		// Default to always notifying
+		bShouldNotifyNextAssetSync = true;
+
+
 		PendingSyncAssets.Empty();
 
 		if (bAllowFocusOnSync && bPendingFocusOnSync)
diff --git a/Engine/Source/Editor/DetailCustomizations/Private/ComponentMaterialCategory.cpp b/Engine/Source/Editor/DetailCustomizations/Private/ComponentMaterialCategory.cpp
--- a/Engine/Source/Editor/DetailCustomizations/Private/ComponentMaterialCategory.cpp
+++ b/Engine/Source/Editor/DetailCustomizations/Private/ComponentMaterialCategory.cpp
@@ -156,8 +156,6 @@ void FComponentMaterialCategory::Create( IDetailLayoutBuilder& DetailBuilder )
 	FMaterialListDelegates MaterialListDelegates;
 	MaterialListDelegates.OnGetMaterials.BindSP( this, &FComponentMaterialCategory::OnGetMaterialsForView );
 	MaterialListDelegates.OnMaterialChanged.BindSP( this, &FComponentMaterialCategory::OnMaterialChanged );
-	MaterialListDelegates.OnGenerateCustomNameWidgets.BindSP(this, &FComponentMaterialCategory::OnGetMaterialNameForView);
-
 	TSharedRef<FMaterialList> MaterialList = MakeShareable( new FMaterialList( DetailBuilder, MaterialListDelegates ) );
 
 	bool bAnyMaterialsToDisplay = false;
@@ -182,98 +180,6 @@ void FComponentMaterialCategory::Create( IDetailLayoutBuilder& DetailBuilder )
 	
 }
 
-TSharedRef<SWidget> FComponentMaterialCategory::OnGetMaterialNameForView(UMaterialInterface* Material, int32 MaterialIndex)
-{
-	if (SelectedComponents.Num() <= 0)
-		return SNullWidget::NullWidget;
-
-	bool IsValueInitialize = false;
-	bool IsAllSameName = true;
-	FName MaterialSlotName(NAME_None);
-	for (FMaterialIterator It(SelectedComponents); It; ++It)
-	{
-		//Just compare the parameter material and MaterialIndex
-		//Skip the others
-		int32 IteratorMaterialIndex = It.GetMaterialIndex();
-		if (IteratorMaterialIndex != MaterialIndex)
-			continue;
-		UMaterialInterface *IteratorMaterial = It.GetMaterial();
-		if (Material != nullptr && IteratorMaterial != Material)
-			continue;
-
-		UActorComponent* CurrentComponent = It.GetComponent();
-		if (!CurrentComponent)
-		{
-			IsAllSameName = false;
-			break;
-		}
-
-		UStaticMeshComponent *StaticMeshComponent = Cast<UStaticMeshComponent>(CurrentComponent);
-		USkeletalMeshComponent *SkeletalMeshComponent = Cast<USkeletalMeshComponent>(CurrentComponent);
-		FName CurrentMaterialSlotName(NAME_None);
-		if (StaticMeshComponent && StaticMeshComponent->StaticMesh && StaticMeshComponent->StaticMesh->StaticMaterials.IsValidIndex(MaterialIndex))
-		{
-			CurrentMaterialSlotName = StaticMeshComponent->StaticMesh->StaticMaterials[MaterialIndex].MaterialSlotName;
-		}
-		else if (SkeletalMeshComponent && SkeletalMeshComponent->SkeletalMesh && SkeletalMeshComponent->SkeletalMesh->Materials.IsValidIndex(MaterialIndex))
-		{
-			CurrentMaterialSlotName = SkeletalMeshComponent->SkeletalMesh->Materials[MaterialIndex].MaterialSlotName;
-		}
-		else
-		{
-			IsAllSameName = false;
-			break;
-		}
-		if (!IsValueInitialize)
-		{
-			IsValueInitialize = true;
-			MaterialSlotName = CurrentMaterialSlotName;
-		}
-		else if (MaterialSlotName != CurrentMaterialSlotName)
-		{
-			IsAllSameName = false;
-			break;
-		}
-	}
-
-	if (!IsValueInitialize)
-	{
-		return SNullWidget::NullWidget;
-	}
-
-	if (!IsAllSameName)
-	{
-		MaterialSlotName = FName(TEXT("Multiple Values"));
-	}
-
-	TSharedRef<SVerticalBox> CustomNameWidget = SNew(SVerticalBox)
-		+ SVerticalBox::Slot()
-		.AutoHeight()
-		.Padding(0, 2, 0, 0)
-		[
-			SNew(SHorizontalBox)
-			+ SHorizontalBox::Slot()
-			.VAlign(VAlign_Center)
-			.Padding(0.0f, 3.0f, 0.0f, 0.0f)
-			.FillWidth(0.4f)
-			[
-				SNew(STextBlock)
-				.Font(IDetailLayoutBuilder::GetDetailFont())
-				.Text(LOCTEXT("MaterialArrayNameLabelStringKey", "Name"))
-			]
-			+ SHorizontalBox::Slot()
-			.Padding(5.0f, 3.0f, 5.0f, 0.0f)
-			.VAlign(VAlign_Center)
-			.FillWidth(1.0f)
-			[
-				SNew(SEditableTextBox)
-				.Text(FText::FromName(MaterialSlotName))
-				.IsEnabled(false)
-			]
-		];
-	return CustomNameWidget;
-}
-
 void FComponentMaterialCategory::OnGetMaterialsForView( IMaterialListBuilder& MaterialList )
 {
 	const bool bAllowNullEntries = true;
@@ -335,6 +241,12 @@ void FComponentMaterialCategory::OnMaterialChanged( UMaterialInterface* NewMater
 		}
 	};
 
+	struct FObjectAndProperty
+	{
+		UObject* Object;
+		UProperty* PropertyThatChanged;
+	};
+	TArray<FObjectAndProperty> ObjectsThatChanged;
 	// Scan the selected actors mesh components for the old material and swap it with the new material 
 	for( FMaterialIterator It( SelectedComponents ); It; ++It )
 	{
@@ -398,16 +310,15 @@ void FComponentMaterialCategory::OnMaterialChanged( UMaterialInterface* NewMater
 				{
 					NotifyHook->NotifyPreChange( MaterialProperty );
 				}
+				FObjectAndProperty ObjectAndProperty;
+				ObjectAndProperty.Object = EditChangeObject;
+				ObjectAndProperty.PropertyThatChanged = MaterialProperty;
 
-				SwapMaterialLambda( CurrentComponent, It.GetMaterialIndex(), NewMaterial );
+				ObjectsThatChanged.Add(ObjectAndProperty);
 
-				FPropertyChangedEvent PropertyChangedEvent( MaterialProperty );
-				EditChangeObject->PostEditChangeProperty( PropertyChangedEvent );
+				FPropertyChangedEvent PropertyChangedEvent(MaterialProperty);
 
-				if( NotifyHook && MaterialProperty )
-				{
-					NotifyHook->NotifyPostChange( PropertyChangedEvent, MaterialProperty );
-				}
+				SwapMaterialLambda( CurrentComponent, It.GetMaterialIndex(), NewMaterial );
 
 				// Propagate material change to instances of the edited component template
 				if( !FApp::IsGame() )
@@ -457,6 +368,20 @@ void FComponentMaterialCategory::OnMaterialChanged( UMaterialInterface* NewMater
 		}
 	}
 
+	// Route post edit change after all components have had their values changed.  This is to avoid 
+	// construction scripts from re-running in the middle of setting values and wiping out components we need to modify
+	for( FObjectAndProperty& ObjectData : ObjectsThatChanged)
+	{
+		FPropertyChangedEvent PropertyChangeEvent(ObjectData.PropertyThatChanged, EPropertyChangeType::ValueSet);
+		ObjectData.Object->PostEditChangeProperty(PropertyChangeEvent);
+
+		if(NotifyHook && ObjectData.PropertyThatChanged)
+		{
+			NotifyHook->NotifyPostChange(PropertyChangeEvent, ObjectData.PropertyThatChanged);
+		}
+	}
+
+	
 	if( bMadeTransaction )
 	{
 		// End the transation if we created one
diff --git a/Engine/Source/Editor/DetailCustomizations/Private/ComponentMaterialCategory.h b/Engine/Source/Editor/DetailCustomizations/Private/ComponentMaterialCategory.h
--- a/Engine/Source/Editor/DetailCustomizations/Private/ComponentMaterialCategory.h
+++ b/Engine/Source/Editor/DetailCustomizations/Private/ComponentMaterialCategory.h
@@ -18,12 +18,6 @@ public:
 	 */
 	void Create( IDetailLayoutBuilder& DetailBuilder );
 private:
-
-	/**
-	* Called by the material list widget when we need to get the material slot name
-	*/
-	TSharedRef<SWidget> OnGetMaterialNameForView(UMaterialInterface* Material, int32 MaterialIndex);
-	
 	/**
 	 * Called by the material list widget when we need to get new materials for the list
 	 *
diff --git a/Engine/Source/Editor/FoliageEdit/Private/FoliageEdMode.cpp b/Engine/Source/Editor/FoliageEdit/Private/FoliageEdMode.cpp
--- a/Engine/Source/Editor/FoliageEdit/Private/FoliageEdMode.cpp
+++ b/Engine/Source/Editor/FoliageEdit/Private/FoliageEdMode.cpp
@@ -158,7 +158,7 @@ bool FFoliagePaintingGeometryFilter::operator() (const UPrimitiveComponent* Comp
 		bool bAllowed =
 			(bAllowLandscape   && Component->IsA(ULandscapeHeightfieldCollisionComponent::StaticClass())) ||
 			(bAllowStaticMesh  && Component->IsA(UStaticMeshComponent::StaticClass()) && !Component->IsA(UFoliageInstancedStaticMeshComponent::StaticClass())) ||
-			(bAllowBSP         && Component->IsA(UModelComponent::StaticClass())) ||
+			(bAllowBSP         && (Component->IsA(UBrushComponent::StaticClass()) || Component->IsA(UModelComponent::StaticClass()))) ||
 			(bAllowFoliage     && Component->IsA(UFoliageInstancedStaticMeshComponent::StaticClass()));
 
 		// Blacklist
diff --git a/Engine/Source/Editor/GameplayAbilitiesEditor/Private/SGameplayCueEditor.cpp b/Engine/Source/Editor/GameplayAbilitiesEditor/Private/SGameplayCueEditor.cpp
--- a/Engine/Source/Editor/GameplayAbilitiesEditor/Private/SGameplayCueEditor.cpp
+++ b/Engine/Source/Editor/GameplayAbilitiesEditor/Private/SGameplayCueEditor.cpp
@@ -575,7 +575,7 @@ BEGIN_SLATE_FUNCTION_BUILD_OPTIMIZATION
 			SNew(STableRow< TSharedPtr<FGCTreeItem> >, OwnerTable)
 			[
 				SNew(STextBlock)
-				.Text(LOCTEXT("UnknownItemType", "Unknown Filter Type"))
+				.Text(LOCTEXT("UnknownFilterType", "Unknown Filter Type"))
 			];
 		}
 	}
diff --git a/Engine/Source/Editor/Kismet/Private/BlueprintEditor.cpp b/Engine/Source/Editor/Kismet/Private/BlueprintEditor.cpp
--- a/Engine/Source/Editor/Kismet/Private/BlueprintEditor.cpp
+++ b/Engine/Source/Editor/Kismet/Private/BlueprintEditor.cpp
@@ -3124,6 +3124,9 @@ void FBlueprintEditor::OnBlueprintCompiled(UBlueprint* InBlueprint)
 {	
 	if( InBlueprint )
 	{
+		// Compiling will invalidate any cached components in the component visualizer, so clear out active components here
+		GUnrealEd->ComponentVisManager.ClearActiveComponentVis();
+
 		// This could be made more efficient by tracking which nodes change
 		// their bHasCompilerMessage flag, or immediately updating the error info
 		// when we assign the flag:
diff --git a/Engine/Source/Editor/Kismet/Private/SSCSEditor.cpp b/Engine/Source/Editor/Kismet/Private/SSCSEditor.cpp
--- a/Engine/Source/Editor/Kismet/Private/SSCSEditor.cpp
+++ b/Engine/Source/Editor/Kismet/Private/SSCSEditor.cpp
@@ -2790,6 +2790,8 @@ void SSCS_RowWidget::OnMakeNewRootDropAction(FSCSEditorTreeNodePtrType DroppedNo
 
 void SSCS_RowWidget::PostDragDropAction(bool bRegenerateTreeNodes)
 {
+	GUnrealEd->ComponentVisManager.ClearActiveComponentVis();
+
 	FSCSEditorTreeNodePtrType NodePtr = GetNode();
 
 	TSharedPtr<SSCSEditor> PinnedEditor = SCSEditor.Pin();
@@ -5016,6 +5018,9 @@ bool SSCSEditor::CanDeleteNodes() const
 
 void SSCSEditor::OnDeleteNodes()
 {
+	// Invalidate any active component in the visualizer
+	GUnrealEd->ComponentVisManager.ClearActiveComponentVis();
+
 	const FScopedTransaction Transaction( LOCTEXT("RemoveComponents", "Remove Components") );
 
 	if (EditorMode == EComponentEditorMode::BlueprintSCS)
diff --git a/Engine/Source/Editor/LandscapeEditor/Private/LandscapeEdMode.cpp b/Engine/Source/Editor/LandscapeEditor/Private/LandscapeEdMode.cpp
--- a/Engine/Source/Editor/LandscapeEditor/Private/LandscapeEdMode.cpp
+++ b/Engine/Source/Editor/LandscapeEditor/Private/LandscapeEdMode.cpp
@@ -528,9 +528,6 @@ void FEdModeLandscape::Exit()
 	LandscapeList.Empty();
 	LandscapeTargetList.Empty();
 
-	// Clear any import landscape data
-	UISettings->ClearImportLandscapeData();
-
 	// Save UI settings to config file
 	UISettings->Save();
 	GLandscapeViewMode = ELandscapeViewMode::Normal;
diff --git a/Engine/Source/Editor/LandscapeEditor/Private/LandscapeEdMode.h b/Engine/Source/Editor/LandscapeEditor/Private/LandscapeEdMode.h
--- a/Engine/Source/Editor/LandscapeEditor/Private/LandscapeEdMode.h
+++ b/Engine/Source/Editor/LandscapeEditor/Private/LandscapeEdMode.h
@@ -12,6 +12,7 @@
 DECLARE_LOG_CATEGORY_EXTERN(LogLandscapeEdMode, Log, All);
 
 // Forward declarations
+class ULandscapeEditorObject;
 class ULandscapeLayerInfoObject;
 class FLandscapeToolSplines;
 class UViewportInteractor;
@@ -108,7 +109,7 @@ struct FLandscapeTargetListInfo
 			}
 			else
 			{
-				int32 Index = Proxy->EditorLayerSettings.Add(LayerInfoObj.Get());
+				int32 Index = Proxy->EditorLayerSettings.Add(FLandscapeEditorLayerSettings(LayerInfoObj.Get()));
 				return &Proxy->EditorLayerSettings[Index];
 			}
 		}
@@ -220,7 +221,7 @@ class FEdModeLandscape : public FEdMode
 {
 public:
 
-	class ULandscapeEditorObject* UISettings;
+	ULandscapeEditorObject* UISettings;
 
 	FLandscapeToolMode* CurrentToolMode;
 	FLandscapeTool* CurrentTool;
diff --git a/Engine/Source/Editor/LandscapeEditor/Private/LandscapeEdModeBrushes.cpp b/Engine/Source/Editor/LandscapeEditor/Private/LandscapeEdModeBrushes.cpp
--- a/Engine/Source/Editor/LandscapeEditor/Private/LandscapeEdModeBrushes.cpp
+++ b/Engine/Source/Editor/LandscapeEditor/Private/LandscapeEdModeBrushes.cpp
@@ -73,16 +73,34 @@ public:
 		// Allow any currently unused material instances to be GC'd
 		BrushMaterialFreeInstances.Empty();
 
+		Collector.AddReferencedObjects(BrushMaterialComponents);
 		Collector.AddReferencedObjects(BrushMaterialInstanceMap);
+
+		// If a user tool removes any components then we will have bad (null) entries in our TSet/TMap, remove them
+		// We can't just call .Remove(nullptr) because the entries were hashed as non-null values so a hash lookup of nullptr won't find them
+		for (auto It = BrushMaterialComponents.CreateIterator(); It; ++It)
+		{
+			if (*It == nullptr)
+			{
+				It.RemoveCurrent();
+			}
+		}
+		for (auto It = BrushMaterialInstanceMap.CreateIterator(); It; ++It)
+		{
+			if (It->Key == nullptr || It->Value == nullptr)
+			{
+				It.RemoveCurrent();
+			}
+		}
 	}
 
 	virtual void LeaveBrush() override
 	{
-		for (TSet<ULandscapeComponent*>::TIterator It(BrushMaterialComponents); It; ++It)
+		for (ULandscapeComponent* Component : BrushMaterialComponents)
 		{
-			if ((*It) != nullptr && (*It)->EditToolRenderData != nullptr)
+			if (Component && Component->EditToolRenderData)
 			{
-				(*It)->EditToolRenderData->Update(nullptr);
+				Component->EditToolRenderData->Update(nullptr);
 			}
 		}
 		TArray<UMaterialInstanceDynamic*> BrushMaterialInstances;
@@ -126,19 +144,11 @@ public:
 		TSet<ULandscapeComponent*> RemovedComponents = BrushMaterialComponents.Difference(NewComponents);
 		for (ULandscapeComponent* RemovedComponent : RemovedComponents)
 		{
-			if (RemovedComponent == nullptr)
-			{
-				// can occur if a component was removed by a user tool
-				BrushMaterialInstanceMap.Remove(nullptr);
-			}
-			else
-			{
-				BrushMaterialFreeInstances.Push(BrushMaterialInstanceMap.FindAndRemoveChecked(RemovedComponent));
+			BrushMaterialFreeInstances.Push(BrushMaterialInstanceMap.FindAndRemoveChecked(RemovedComponent));
 
-				if (ensure(RemovedComponent->EditToolRenderData != nullptr))
-				{
-					RemovedComponent->EditToolRenderData->Update(nullptr);
-				}
+			if (ensure(RemovedComponent->EditToolRenderData != nullptr))
+			{
+				RemovedComponent->EditToolRenderData->Update(nullptr);
 			}
 		}
 
diff --git a/Engine/Source/Editor/LandscapeEditor/Private/LandscapeEdModeComponentTools.cpp b/Engine/Source/Editor/LandscapeEditor/Private/LandscapeEdModeComponentTools.cpp
--- a/Engine/Source/Editor/LandscapeEditor/Private/LandscapeEdModeComponentTools.cpp
+++ b/Engine/Source/Editor/LandscapeEditor/Private/LandscapeEdModeComponentTools.cpp
@@ -14,155 +14,100 @@
 #include "ComponentReregisterContext.h"
 #include "PhysicalMaterials/PhysicalMaterial.h"
 #include "Materials/MaterialExpressionLandscapeVisibilityMask.h"
+#include "Algo/Copy.h"
 
 #define LOCTEXT_NAMESPACE "Landscape"
 
+//
+// FLandscapeToolSelect
+//
 class FLandscapeToolStrokeSelect : public FLandscapeToolStrokeBase
 {
 	bool bInitializedComponentInvert;
-	bool bComponentInvert;
+	bool bInvert;
+	bool bNeedsSelectionUpdate;
 
 public:
 	FLandscapeToolStrokeSelect(FEdModeLandscape* InEdMode, FEditorViewportClient* InViewportClient, const FLandscapeToolTarget& InTarget)
 		: FLandscapeToolStrokeBase(InEdMode, InViewportClient, InTarget)
 		, bInitializedComponentInvert(false)
+		, bNeedsSelectionUpdate(false)
 		, Cache(InTarget)
 	{
 	}
 
+	~FLandscapeToolStrokeSelect()
+	{
+		if (bNeedsSelectionUpdate)
+		{
+			TArray<UObject*> Objects;
+			if (LandscapeInfo)
+			{
+				TSet<ULandscapeComponent*> SelectedComponents = LandscapeInfo->GetSelectedComponents();
+				Objects.Reset(SelectedComponents.Num());
+				Algo::Copy(SelectedComponents, Objects);
+			}
+			FPropertyEditorModule& PropertyModule = FModuleManager::Get().LoadModuleChecked<FPropertyEditorModule>(TEXT("PropertyEditor"));
+			PropertyModule.UpdatePropertyViews(Objects);
+		}
+	}
+
 	void Apply(FEditorViewportClient* ViewportClient, FLandscapeBrush* Brush, const ULandscapeEditorObject* UISettings, const TArray<FLandscapeToolInteractorPosition>& InteractorPositions)
 	{
 		if (LandscapeInfo)
 		{
 			LandscapeInfo->Modify();
 
-			// Invert when holding Shift
-			bool bInvert = InteractorPositions[InteractorPositions.Num() - 1].bModifierPressed;
-
-			if (Brush->GetBrushType() == ELandscapeBrushType::Component)
+			// TODO - only retrieve bounds as we don't need the data
+			const FLandscapeBrushData BrushInfo = Brush->ApplyBrush(InteractorPositions);
+			if (!BrushInfo)
 			{
-				// TODO - only retrieve bounds as we don't need the data
-				const FLandscapeBrushData BrushInfo = Brush->ApplyBrush(InteractorPositions);
-				if (!BrushInfo)
-				{
-					return;
-				}
-
-				int32 X1, Y1, X2, Y2;
-				BrushInfo.GetInclusiveBounds(X1, Y1, X2, Y2);
-
-				// Shrink bounds by 1,1 to avoid GetComponentsInRegion picking up extra components on all sides due to the overlap between components
-				TSet<ULandscapeComponent*> NewComponents;
-				LandscapeInfo->GetComponentsInRegion(X1 + 1, Y1 + 1, X2 - 1, Y2 - 1, NewComponents);
-
-				if (!bInitializedComponentInvert)
-				{
-					// Get the component under the mouse location. Copied from FLandscapeBrushComponent::ApplyBrush()
-					const float MouseX = InteractorPositions[0].Position.X;
-					const float MouseY = InteractorPositions[0].Position.Y;
-					const int32 MouseComponentIndexX = (MouseX >= 0.0f) ? FMath::FloorToInt(MouseX / LandscapeInfo->ComponentSizeQuads) : FMath::CeilToInt(MouseX / LandscapeInfo->ComponentSizeQuads);
-					const int32 MouseComponentIndexY = (MouseY >= 0.0f) ? FMath::FloorToInt(MouseY / LandscapeInfo->ComponentSizeQuads) : FMath::CeilToInt(MouseY / LandscapeInfo->ComponentSizeQuads);
-					ULandscapeComponent* MouseComponent = LandscapeInfo->XYtoComponentMap.FindRef(FIntPoint(MouseComponentIndexX, MouseComponentIndexY));
-
-					if (MouseComponent != nullptr)
-					{
-						bComponentInvert = LandscapeInfo->GetSelectedComponents().Contains(MouseComponent);
-					}
-					else
-					{
-						bComponentInvert = false;
-					}
+				return;
+			}
 
-					bInitializedComponentInvert = true;
-				}
+			int32 X1, Y1, X2, Y2;
+			BrushInfo.GetInclusiveBounds(X1, Y1, X2, Y2);
 
-				bInvert = bComponentInvert;
+			// Shrink bounds by 1,1 to avoid GetComponentsInRegion picking up extra components on all sides due to the overlap between components
+			TSet<ULandscapeComponent*> NewComponents;
+			LandscapeInfo->GetComponentsInRegion(X1 + 1, Y1 + 1, X2 - 1, Y2 - 1, NewComponents);
 
-				TSet<ULandscapeComponent*> NewSelection;
-				if (bInvert)
+			if (!bInitializedComponentInvert)
+			{
+				// Get the component under the mouse location. Copied from FLandscapeBrushComponent::ApplyBrush()
+				const float MouseX = InteractorPositions[0].Position.X;
+				const float MouseY = InteractorPositions[0].Position.Y;
+				const int32 MouseComponentIndexX = (MouseX >= 0.0f) ? FMath::FloorToInt(MouseX / LandscapeInfo->ComponentSizeQuads) : FMath::CeilToInt(MouseX / LandscapeInfo->ComponentSizeQuads);
+				const int32 MouseComponentIndexY = (MouseY >= 0.0f) ? FMath::FloorToInt(MouseY / LandscapeInfo->ComponentSizeQuads) : FMath::CeilToInt(MouseY / LandscapeInfo->ComponentSizeQuads);
+				ULandscapeComponent* MouseComponent = LandscapeInfo->XYtoComponentMap.FindRef(FIntPoint(MouseComponentIndexX, MouseComponentIndexY));
+
+				if (MouseComponent != nullptr)
 				{
-					NewSelection = LandscapeInfo->GetSelectedComponents().Difference(NewComponents);
+					bInvert = LandscapeInfo->GetSelectedComponents().Contains(MouseComponent);
 				}
 				else
 				{
-					NewSelection = LandscapeInfo->GetSelectedComponents().Union(NewComponents);
+					bInvert = false;
 				}
 
-				LandscapeInfo->Modify();
-				LandscapeInfo->UpdateSelectedComponents(NewSelection);
+				bInitializedComponentInvert = true;
+			}
 
-				// Update Details tab with selection
-				TArray<UObject*> Objects;
-				Objects.Reset(NewSelection.Num());
-				for (auto It = NewSelection.CreateConstIterator(); It; ++It)
-				{
-					Objects.Add(*It);
-				}
-				FPropertyEditorModule& PropertyModule = FModuleManager::Get().LoadModuleChecked<FPropertyEditorModule>(TEXT("PropertyEditor"));
-				PropertyModule.UpdatePropertyViews(Objects);
+			TSet<ULandscapeComponent*> NewSelection;
+			if (bInvert)
+			{
+				NewSelection = LandscapeInfo->GetSelectedComponents().Difference(NewComponents);
 			}
-			else // Select various shape regions
+			else
 			{
-				const FLandscapeBrushData BrushInfo = Brush->ApplyBrush(InteractorPositions);
-				if (!BrushInfo)
-				{
-					return;
-				}
-
-				int32 X1, Y1, X2, Y2;
-				BrushInfo.GetInclusiveBounds(X1, Y1, X2, Y2);
-
-				// Tablet pressure
-				float Pressure = ViewportClient->Viewport->IsPenActive() ? ViewportClient->Viewport->GetTabletPressure() : 1.0f;
-
-				Cache.CacheData(X1, Y1, X2, Y2);
-				TArray<uint8> Data;
-				Cache.GetCachedData(X1, Y1, X2, Y2, Data);
-
-				TSet<ULandscapeComponent*> NewComponents;
-				LandscapeInfo->GetComponentsInRegion(X1, Y1, X2, Y2, NewComponents);
-				LandscapeInfo->UpdateSelectedComponents(NewComponents, false);
-
-				for (int32 Y = BrushInfo.GetBounds().Min.Y; Y < BrushInfo.GetBounds().Max.Y; Y++)
-				{
-					const float* BrushScanline = BrushInfo.GetDataPtr(FIntPoint(0, Y));
-					uint8* DataScanline = Data.GetData() + (Y - Y1) * (X2 - X1 + 1) + (0 - X1);
-
-					for (int32 X = BrushInfo.GetBounds().Min.X; X < BrushInfo.GetBounds().Max.X; X++)
-					{
-						const FIntPoint Key = FIntPoint(X, Y);
-						const float BrushValue = BrushScanline[X];
-
-						if (BrushValue > 0.0f && LandscapeInfo->IsValidPosition(X, Y))
-						{
-							float PaintValue = BrushValue * UISettings->ToolStrength * Pressure;
-							float Value = DataScanline[X] / 255.0f;
-							checkSlow(FMath::IsNearlyEqual(Value, LandscapeInfo->SelectedRegion.FindRef(Key), 1 / 255.0f));
-							if (bInvert)
-							{
-								Value = FMath::Max(Value - PaintValue, 0.0f);
-							}
-							else
-							{
-								Value = FMath::Min(Value + PaintValue, 1.0f);
-							}
-							if (Value > 0.0f)
-							{
-								LandscapeInfo->SelectedRegion.Add(Key, Value);
-							}
-							else
-							{
-								LandscapeInfo->SelectedRegion.Remove(Key);
-							}
+				NewSelection = LandscapeInfo->GetSelectedComponents().Union(NewComponents);
+			}
 
-							DataScanline[X] = FMath::Clamp<int32>(FMath::RoundToInt(Value * 255), 0, 255);
-						}
-					}
-				}
+			LandscapeInfo->Modify();
+			LandscapeInfo->UpdateSelectedComponents(NewSelection);
 
-				Cache.SetCachedData(X1, Y1, X2, Y2, Data);
-				Cache.Flush();
-			}
+			// Update Details tab with selection
+			bNeedsSelectionUpdate = true;
 		}
 	}
 
@@ -170,15 +115,11 @@ protected:
 	FLandscapeDataCache Cache;
 };
 
-// 
-// FLandscapeToolSelect
-//
-template<class TStrokeClass>
-class FLandscapeToolSelect : public FLandscapeToolBase<TStrokeClass>
+class FLandscapeToolSelect : public FLandscapeToolBase<FLandscapeToolStrokeSelect>
 {
 public:
 	FLandscapeToolSelect(FEdModeLandscape* InEdMode)
-		: FLandscapeToolBase<TStrokeClass>(InEdMode)
+		: FLandscapeToolBase<FLandscapeToolStrokeSelect>(InEdMode)
 	{
 	}
 
@@ -275,12 +216,11 @@ protected:
 	FLandscapeDataCache Cache;
 };
 
-template<class TStrokeClass>
-class FLandscapeToolMask : public FLandscapeToolSelect<TStrokeClass>
+class FLandscapeToolMask : public FLandscapeToolBase<FLandscapeToolStrokeMask>
 {
 public:
 	FLandscapeToolMask(FEdModeLandscape* InEdMode)
-		: FLandscapeToolSelect<TStrokeClass>(InEdMode)
+		: FLandscapeToolBase<FLandscapeToolStrokeMask>(InEdMode)
 	{
 	}
 
@@ -288,8 +228,13 @@ public:
 	virtual FText GetDisplayName() override { return NSLOCTEXT("UnrealEd", "LandscapeMode_Mask", "Region Selection"); };
 	virtual void SetEditRenderType() override { GLandscapeEditRenderMode = ELandscapeEditRenderMode::SelectRegion | (GLandscapeEditRenderMode & ELandscapeEditRenderMode::BitMaskForMask); }
 	virtual bool SupportsMask() override { return true; }
+
+	virtual ELandscapeToolType GetToolType() override { return ELandscapeToolType::Mask; }
 };
 
+//
+// FLandscapeToolVisibility
+//
 class FLandscapeToolStrokeVisibility : public FLandscapeToolStrokeBase
 {
 public:
@@ -350,9 +295,6 @@ protected:
 	FLandscapeVisCache Cache;
 };
 
-// 
-// FLandscapeToolVisibility
-//
 class FLandscapeToolVisibility : public FLandscapeToolBase<FLandscapeToolStrokeVisibility>
 {
 public:
@@ -391,6 +333,9 @@ public:
 	}
 };
 
+//
+// FLandscapeToolMoveToLevel
+//
 class FLandscapeToolStrokeMoveToLevel : public FLandscapeToolStrokeBase
 {
 public:
@@ -824,9 +769,6 @@ public:
 	}
 };
 
-// 
-// FLandscapeToolMoveToLevel
-//
 class FLandscapeToolMoveToLevel : public FLandscapeToolBase<FLandscapeToolStrokeMoveToLevel>
 {
 public:
@@ -842,7 +784,9 @@ public:
 	virtual bool SupportsMask() override { return false; }
 };
 
-
+//
+// FLandscapeToolAddComponent
+//
 class FLandscapeToolStrokeAddComponent : public FLandscapeToolStrokeBase
 {
 public:
@@ -1020,9 +964,6 @@ protected:
 	FLandscapeXYOffsetCache<true> XYOffsetCache;
 };
 
-// 
-// FLandscapeToolAddComponent
-//
 class FLandscapeToolAddComponent : public FLandscapeToolBase<FLandscapeToolStrokeAddComponent>
 {
 public:
@@ -1052,6 +993,9 @@ public:
 	}
 };
 
+//
+// FLandscapeToolDeleteComponent
+//
 class FLandscapeToolStrokeDeleteComponent : public FLandscapeToolStrokeBase
 {
 public:
@@ -1088,9 +1032,6 @@ public:
 	}
 };
 
-// 
-// FLandscapeToolDeleteComponent
-//
 class FLandscapeToolDeleteComponent : public FLandscapeToolBase<FLandscapeToolStrokeDeleteComponent>
 {
 public:
@@ -1106,6 +1047,9 @@ public:
 	virtual bool SupportsMask() override { return false; }
 };
 
+//
+// FLandscapeToolCopy
+//
 template<class ToolTarget>
 class FLandscapeToolStrokeCopy : public FLandscapeToolStrokeBase
 {
@@ -1411,9 +1355,6 @@ protected:
 	FLandscapeFullWeightCache WeightCache;
 };
 
-// 
-// FLandscapeToolCopy
-//
 template<class ToolTarget>
 class FLandscapeToolCopy : public FLandscapeToolBase<FLandscapeToolStrokeCopy<ToolTarget>>
 {
@@ -1461,6 +1402,9 @@ protected:
 	FLandscapeBrush* BackupCurrentBrush;
 };
 
+//
+// FLandscapeToolPaste
+//
 template<class ToolTarget>
 class FLandscapeToolStrokePaste : public FLandscapeToolStrokeBase
 {
@@ -1726,9 +1670,6 @@ protected:
 	FLandscapeFullWeightCache WeightCache;
 };
 
-// 
-// FLandscapeToolPaste
-//
 template<class ToolTarget>
 class FLandscapeToolPaste : public FLandscapeToolBase<FLandscapeToolStrokePaste<ToolTarget>>
 {
@@ -1800,6 +1741,9 @@ protected:
 	FLandscapeBrush* BackupCurrentBrush;
 };
 
+//
+// FLandscapeToolCopyPaste
+//
 template<class ToolTarget>
 class FLandscapeToolCopyPaste : public FLandscapeToolPaste<ToolTarget>
 {
@@ -1872,7 +1816,7 @@ void FEdModeLandscape::PasteDataFromGizmo()
 	}
 }
 
-// 
+//
 // FLandscapeToolNewLandscape
 //
 class FLandscapeToolNewLandscape : public FLandscapeTool
@@ -1897,12 +1841,14 @@ public:
 	virtual void EnterTool() override
 	{
 		EdMode->NewLandscapePreviewMode = NewLandscapePreviewMode;
+		EdMode->UISettings->ImportLandscapeData();
 	}
 
 	virtual void ExitTool() override
 	{
 		NewLandscapePreviewMode = EdMode->NewLandscapePreviewMode;
 		EdMode->NewLandscapePreviewMode = ENewLandscapePreviewMode::None;
+		EdMode->UISettings->ClearImportLandscapeData();
 	}
 
 	virtual bool BeginTool(FEditorViewportClient* ViewportClient, const FLandscapeToolTarget& Target, const FVector& InHitLocation, const UViewportInteractor* Interactor = nullptr) override
@@ -1924,7 +1870,7 @@ public:
 };
 
 
-// 
+//
 // FLandscapeToolResizeLandscape
 //
 class FLandscapeToolResizeLandscape : public FLandscapeTool
@@ -2005,7 +1951,7 @@ void FEdModeLandscape::InitializeTool_ResizeLandscape()
 
 void FEdModeLandscape::InitializeTool_Select()
 {
-	auto Tool_Select = MakeUnique<FLandscapeToolSelect<FLandscapeToolStrokeSelect>>(this);
+	auto Tool_Select = MakeUnique<FLandscapeToolSelect>(this);
 	Tool_Select->ValidBrushes.Add("BrushSet_Component");
 	LandscapeTools.Add(MoveTemp(Tool_Select));
 }
@@ -2033,7 +1979,7 @@ void FEdModeLandscape::InitializeTool_MoveToLevel()
 
 void FEdModeLandscape::InitializeTool_Mask()
 {
-	auto Tool_Mask = MakeUnique<FLandscapeToolMask<FLandscapeToolStrokeSelect>>(this);
+	auto Tool_Mask = MakeUnique<FLandscapeToolMask>(this);
 	Tool_Mask->ValidBrushes.Add("BrushSet_Circle");
 	Tool_Mask->ValidBrushes.Add("BrushSet_Alpha");
 	Tool_Mask->ValidBrushes.Add("BrushSet_Pattern");
diff --git a/Engine/Source/Editor/LandscapeEditor/Private/LandscapeEdModeSplineTools.cpp b/Engine/Source/Editor/LandscapeEditor/Private/LandscapeEdModeSplineTools.cpp
--- a/Engine/Source/Editor/LandscapeEditor/Private/LandscapeEdModeSplineTools.cpp
+++ b/Engine/Source/Editor/LandscapeEditor/Private/LandscapeEdModeSplineTools.cpp
@@ -20,6 +20,7 @@
 #include "ControlPointMeshComponent.h"
 #include "EditorUndoClient.h"
 #include "EngineUtils.h"
+#include "Algo/Copy.h"
 
 
 #define LOCTEXT_NAMESPACE "Landscape"
@@ -86,16 +87,8 @@ public:
 		{
 			TArray<UObject*> Objects;
 			Objects.Reset(SelectedSplineControlPoints.Num() + SelectedSplineSegments.Num());
-
-			for (ULandscapeSplineControlPoint* ControlPoint : SelectedSplineControlPoints)
-			{
-				Objects.Add(ControlPoint);
-			}
-			for (ULandscapeSplineSegment* Segment : SelectedSplineSegments)
-			{
-				Objects.Add(Segment);
-			}
-
+			Algo::Copy(SelectedSplineControlPoints, Objects);
+			Algo::Copy(SelectedSplineSegments, Objects);
 
 			FPropertyEditorModule& PropertyModule = FModuleManager::Get().LoadModuleChecked<FPropertyEditorModule>(TEXT("PropertyEditor"));
 			PropertyModule.UpdatePropertyViews(Objects);
@@ -817,15 +810,8 @@ public:
 	{
 		TArray<UObject*> Objects;
 		Objects.Reset(SelectedSplineControlPoints.Num() + SelectedSplineSegments.Num());
-
-		for (ULandscapeSplineControlPoint* ControlPoint : SelectedSplineControlPoints)
-		{
-			Objects.Add(ControlPoint);
-		}
-		for (ULandscapeSplineSegment* Segment : SelectedSplineSegments)
-		{
-			Objects.Add(Segment);
-		}
+		Algo::Copy(SelectedSplineControlPoints, Objects);
+		Algo::Copy(SelectedSplineSegments, Objects);
 
 		FPropertyEditorModule& PropertyModule = FModuleManager::Get().LoadModuleChecked<FPropertyEditorModule>(TEXT("PropertyEditor"));
 		if (!PropertyModule.HasUnlockedDetailViews())
@@ -1366,86 +1352,39 @@ public:
 		SelectedSplineControlPoints.Empty();
 		SelectedSplineSegments.Empty();
 
-		if (EdMode->CurrentTool != NULL && EdMode->CurrentTool == this)
+		if (EdMode->CurrentTool != nullptr && EdMode->CurrentTool == this)
 		{
 			for (const FLandscapeListInfo& Info : EdMode->GetLandscapeList())
 			{
-				ALandscape* Landscape = Info.Info->LandscapeActor.Get();
-				if (Landscape != NULL && Landscape->SplineComponent != NULL)
-				{
-					for (ULandscapeSplineControlPoint* ControlPoint : Landscape->SplineComponent->ControlPoints)
-					{
-						if (ControlPoint->IsSplineSelected())
-						{
-							SelectedSplineControlPoints.Add(ControlPoint);
-						}
-					}
-
-					for (ULandscapeSplineSegment* Segment : Landscape->SplineComponent->Segments)
-					{
-						if (Segment->IsSplineSelected())
-						{
-							SelectedSplineSegments.Add(Segment);
-						}
-					}
-				}
-
-				for (ALandscapeProxy* LandscapeProxy : Info.Info->Proxies)
+				Info.Info->ForAllLandscapeProxies([this](ALandscapeProxy* Proxy)
 				{
-					if (LandscapeProxy != NULL && LandscapeProxy->SplineComponent != NULL)
+					if (Proxy->SplineComponent)
 					{
-						for (ULandscapeSplineControlPoint* ControlPoint : LandscapeProxy->SplineComponent->ControlPoints)
-						{
-							if (ControlPoint->IsSplineSelected())
-							{
-								SelectedSplineControlPoints.Add(ControlPoint);
-							}
-						}
-
-						for (ULandscapeSplineSegment* Segment : LandscapeProxy->SplineComponent->Segments)
-						{
-							if (Segment->IsSplineSelected())
-							{
-								SelectedSplineSegments.Add(Segment);
-							}
-						}
+						Algo::CopyIf(Proxy->SplineComponent->ControlPoints, SelectedSplineControlPoints, &ULandscapeSplineControlPoint::IsSplineSelected);
+						Algo::CopyIf(Proxy->SplineComponent->Segments,      SelectedSplineSegments,      &ULandscapeSplineSegment::IsSplineSelected);
 					}
-				}
+				});
 			}
 		}
 		else
 		{
 			for (const FLandscapeListInfo& Info : EdMode->GetLandscapeList())
 			{
-				ALandscape* Landscape = Info.Info->LandscapeActor.Get();
-				if (Landscape != NULL && Landscape->SplineComponent != NULL)
-				{
-					for (ULandscapeSplineControlPoint* ControlPoint : Landscape->SplineComponent->ControlPoints)
-					{
-						ControlPoint->SetSplineSelected(false);
-					}
-
-					for (ULandscapeSplineSegment* Segment : Landscape->SplineComponent->Segments)
-					{
-						Segment->SetSplineSelected(false);
-					}
-				}
-
-				for (ALandscapeProxy* LandscapeProxy : Info.Info->Proxies)
+				Info.Info->ForAllLandscapeProxies([](ALandscapeProxy* Proxy)
 				{
-					if (LandscapeProxy != NULL && LandscapeProxy->SplineComponent != NULL)
+					if (Proxy->SplineComponent)
 					{
-						for (ULandscapeSplineControlPoint* ControlPoint : LandscapeProxy->SplineComponent->ControlPoints)
+						for (ULandscapeSplineControlPoint* ControlPoint : Proxy->SplineComponent->ControlPoints)
 						{
 							ControlPoint->SetSplineSelected(false);
 						}
 
-						for (ULandscapeSplineSegment* Segment : LandscapeProxy->SplineComponent->Segments)
+						for (ULandscapeSplineSegment* Segment : Proxy->SplineComponent->Segments)
 						{
 							Segment->SetSplineSelected(false);
 						}
 					}
-				}
+				});
 			}
 		}
 	}
@@ -1462,19 +1401,13 @@ public:
 
 		for (const FLandscapeListInfo& Info : EdMode->GetLandscapeList())
 		{
-			ALandscape* Landscape = Info.Info->LandscapeActor.Get();
-			if (Landscape && Landscape->SplineComponent)
-			{
-				Landscape->SplineComponent->ShowSplineEditorMesh(true);
-			}
-
-			for (ALandscapeProxy* LandscapeProxy : Info.Info->Proxies)
+			Info.Info->ForAllLandscapeProxies([](ALandscapeProxy* Proxy)
 			{
-				if (LandscapeProxy->SplineComponent)
+				if (Proxy->SplineComponent)
 				{
-					LandscapeProxy->SplineComponent->ShowSplineEditorMesh(true);
+					Proxy->SplineComponent->ShowSplineEditorMesh(true);
 				}
-			}
+			});
 		}
 	}
 
@@ -1485,19 +1418,13 @@ public:
 
 		for (const FLandscapeListInfo& Info : EdMode->GetLandscapeList())
 		{
-			ALandscape* Landscape = Info.Info->LandscapeActor.Get();
-			if (Landscape && Landscape->SplineComponent)
+			Info.Info->ForAllLandscapeProxies([](ALandscapeProxy* Proxy)
 			{
-				Landscape->SplineComponent->ShowSplineEditorMesh(false);
-			}
-
-			for (ALandscapeProxy* LandscapeProxy : Info.Info->Proxies)
-			{
-				if (LandscapeProxy->SplineComponent)
+				if (Proxy->SplineComponent)
 				{
-					LandscapeProxy->SplineComponent->ShowSplineEditorMesh(false);
+					Proxy->SplineComponent->ShowSplineEditorMesh(false);
 				}
-			}
+			});
 		}
 	}
 
diff --git a/Engine/Source/Editor/LandscapeEditor/Private/LandscapeEditorDetailCustomization_NewLandscape.cpp b/Engine/Source/Editor/LandscapeEditor/Private/LandscapeEditorDetailCustomization_NewLandscape.cpp
--- a/Engine/Source/Editor/LandscapeEditor/Private/LandscapeEditorDetailCustomization_NewLandscape.cpp
+++ b/Engine/Source/Editor/LandscapeEditor/Private/LandscapeEditorDetailCustomization_NewLandscape.cpp
@@ -1101,6 +1101,7 @@ void FLandscapeEditorDetailCustomization_NewLandscape::OnImportHeightmapFilename
 			int32 i = ImportResolutions.Num() / 2;
 			LandscapeEdMode->UISettings->ImportLandscape_Width = ImportResolutions[i].Width;
 			LandscapeEdMode->UISettings->ImportLandscape_Height = ImportResolutions[i].Height;
+			LandscapeEdMode->UISettings->ImportLandscapeData();
 			ChooseBestComponentSizeForImport(LandscapeEdMode);
 		}
 	}
diff --git a/Engine/Source/Editor/LandscapeEditor/Private/LandscapeEditorObject.cpp b/Engine/Source/Editor/LandscapeEditor/Private/LandscapeEditorObject.cpp
--- a/Engine/Source/Editor/LandscapeEditor/Private/LandscapeEditorObject.cpp
+++ b/Engine/Source/Editor/LandscapeEditor/Private/LandscapeEditorObject.cpp
@@ -528,7 +528,7 @@ bool ULandscapeEditorObject::SetAlphaTexture(UTexture2D* InTexture, EColorChanne
 	return Result;
 }
 
-const TArray<uint16>& ULandscapeEditorObject::GetImportLandscapeData()
+void ULandscapeEditorObject::ImportLandscapeData()
 {
 	ILandscapeEditorModule& LandscapeEditorModule = FModuleManager::GetModuleChecked<ILandscapeEditorModule>("LandscapeEditor");
 	const ILandscapeHeightmapFileFormat* HeightmapFormat = LandscapeEditorModule.GetHeightmapFormatByExtension(*FPaths::GetExtension(ImportLandscape_HeightmapFilename, true));
@@ -550,8 +550,6 @@ const TArray<uint16>& ULandscapeEditorObject::GetImportLandscapeData()
 	{
 		ImportLandscape_Data.Empty();
 	}
-
-	return ImportLandscape_Data;
 }
 
 void ULandscapeEditorObject::RefreshImportLayersList()
diff --git a/Engine/Source/Editor/LandscapeEditor/Public/LandscapeEditorObject.h b/Engine/Source/Editor/LandscapeEditor/Public/LandscapeEditorObject.h
--- a/Engine/Source/Editor/LandscapeEditor/Public/LandscapeEditorObject.h
+++ b/Engine/Source/Editor/LandscapeEditor/Public/LandscapeEditorObject.h
@@ -489,8 +489,11 @@ class ULandscapeEditorObject : public UObject
 	uint32 ImportLandscape_Width;
 	UPROPERTY(NonTransactional)
 	uint32 ImportLandscape_Height;
+
+private:
 	UPROPERTY(NonTransactional)
 	TArray<uint16> ImportLandscape_Data;
+public:
 
 	// Whether the imported alpha maps are to be interpreted as "layered" or "additive" (UE4 uses additive internally)
 	UPROPERTY(Category="New Landscape", EditAnywhere, NonTransactional, meta=(DisplayName="Layer Alphamap Type", ShowForTools="NewLandscape"))
@@ -593,9 +596,10 @@ class ULandscapeEditorObject : public UObject
 	// New Landscape
 	FString LastImportPath;
 
-	const TArray<uint16>& GetImportLandscapeData();
+	const TArray<uint16>& GetImportLandscapeData() const { return ImportLandscape_Data; }
 	void ClearImportLandscapeData() { ImportLandscape_Data.Empty(); }
 
+	void ImportLandscapeData();
 	void RefreshImportLayersList();
 	
 	void UpdateComponentLayerWhitelist();
diff --git a/Engine/Source/Editor/LevelEditor/Private/LevelEditor.cpp b/Engine/Source/Editor/LevelEditor/Private/LevelEditor.cpp
--- a/Engine/Source/Editor/LevelEditor/Private/LevelEditor.cpp
+++ b/Engine/Source/Editor/LevelEditor/Private/LevelEditor.cpp
@@ -53,8 +53,8 @@ public:
 
 	void Construct(const FArguments& InArgs)
 	{
-		FString OptionalBranchPrefix;
-		GConfig->GetString(TEXT("LevelEditor"), TEXT("ProjectNameWatermarkPrefix"), /*out*/ OptionalBranchPrefix, GEditorPerProjectIni);
+		FString ProjectNameWatermarkPrefix;
+		GConfig->GetString(TEXT("LevelEditor"), TEXT("ProjectNameWatermarkPrefix"), /*out*/ ProjectNameWatermarkPrefix, GEditorPerProjectIni);
 
 		FColor BadgeBackgroundColor = FColor::Black;
 		GConfig->GetColor(TEXT("LevelEditor"), TEXT("ProjectBadgeBackgroundColor"), /*out*/ BadgeBackgroundColor, GEditorPerProjectIni);
@@ -65,7 +65,9 @@ public:
 		const FString EngineVersionString = FEngineVersion::Current().ToString(FEngineVersion::Current().HasChangelist() ? EVersionComponent::Changelist : EVersionComponent::Patch);
 		
 		FFormatNamedArguments Args;
-		Args.Add(TEXT("Branch"), FText::FromString(OptionalBranchPrefix));
+
+		Args.Add(TEXT("ProjectNameWatermarkPrefix"), FText::FromString(ProjectNameWatermarkPrefix));
+		Args.Add(TEXT("Branch"), FText::FromString(FApp::GetBranchName()));
 		Args.Add(TEXT("GameName"), FText::FromString(FString(FApp::GetGameName())));
 		Args.Add(TEXT("EngineVersion"), (GetDefault<UEditorPerProjectUserSettings>()->bDisplayEngineVersionInBadge) ? FText::FromString("(" + EngineVersionString + ")") : FText());
 
@@ -76,11 +78,11 @@ public:
 		if (BuildConfig != EBuildConfigurations::Shipping && BuildConfig != EBuildConfigurations::Development && BuildConfig != EBuildConfigurations::Unknown)
 		{
 			Args.Add(TEXT("Config"), EBuildConfigurations::ToText(BuildConfig));
-			RightContentText = FText::Format(NSLOCTEXT("UnrealEditor", "TitleBarRightContentAndConfig", "{Branch}{GameName} [{Config}] {EngineVersion}"), Args);
+			RightContentText = FText::Format(NSLOCTEXT("UnrealEditor", "TitleBarRightContentAndConfig", "{ProjectNameWatermarkPrefix} {GameName} [{Config}] {Branch} {EngineVersion}"), Args);
 		}
 		else
 		{
-			RightContentText = FText::Format(NSLOCTEXT("UnrealEditor", "TitleBarRightContent", "{Branch}{GameName} {EngineVersion}"), Args);
+			RightContentText = FText::Format(NSLOCTEXT("UnrealEditor", "TitleBarRightContent", "{ProjectNameWatermarkPrefix} {GameName} {Branch} {EngineVersion}"), Args);
 		}
 
 		// Create the tooltip showing more detailed information
diff --git a/Engine/Source/Editor/MainFrame/Private/Frame/MainFrameActions.cpp b/Engine/Source/Editor/MainFrame/Private/Frame/MainFrameActions.cpp
--- a/Engine/Source/Editor/MainFrame/Private/Frame/MainFrameActions.cpp
+++ b/Engine/Source/Editor/MainFrame/Private/Frame/MainFrameActions.cpp
@@ -138,7 +138,7 @@ void FMainFrameCommands::RegisterCommands()
 	UI_COMMAND( VisitSearchForAnswersPage, "Answer Hub...", "Go to the AnswerHub to ask questions, search existing answers, and share your knowledge with other UE4 developers.", EUserInterfaceActionType::Button, FInputChord() );
 	ActionList->MapAction( VisitSearchForAnswersPage, FExecuteAction::CreateStatic( &FMainFrameActionCallbacks::VisitSearchForAnswersPage ) );
 
-	UI_COMMAND( VisitSupportWebSite, "Unreal Engine Support Web Site...", "Navigates to the Unreal Engine Support web site's main page.", EUserInterfaceActionType::Button, FInputChord() );
+	UI_COMMAND( VisitSupportWebSite, "Support...", "Navigates to the Unreal Engine Support web site's main page.", EUserInterfaceActionType::Button, FInputChord() );
 	ActionList->MapAction( VisitSupportWebSite, FExecuteAction::CreateStatic( &FMainFrameActionCallbacks::VisitSupportWebSite ) );
 
 	UI_COMMAND( VisitEpicGamesDotCom, "Visit UnrealEngine.com...", "Navigates to UnrealEngine.com where you can learn more about Unreal Technology.", EUserInterfaceActionType::Button, FInputChord() );
diff --git a/Engine/Source/Editor/MainFrame/Private/Menus/MainMenu.cpp b/Engine/Source/Editor/MainFrame/Private/Menus/MainMenu.cpp
--- a/Engine/Source/Editor/MainFrame/Private/Menus/MainMenu.cpp
+++ b/Engine/Source/Editor/MainFrame/Private/Menus/MainMenu.cpp
@@ -267,6 +267,7 @@ void FMainMenu::FillHelpMenu( FMenuBuilder& MenuBuilder, const TSharedRef< FExte
 {
 	MenuBuilder.BeginSection("HelpOnline", NSLOCTEXT("MainHelpMenu", "Online", "Online"));
 	{
+		MenuBuilder.AddMenuEntry(FMainFrameCommands::Get().VisitSupportWebSite);
 		MenuBuilder.AddMenuEntry(FMainFrameCommands::Get().VisitForums);
 		MenuBuilder.AddMenuEntry(FMainFrameCommands::Get().VisitSearchForAnswersPage);
 		MenuBuilder.AddMenuEntry(FMainFrameCommands::Get().VisitWiki);
diff --git a/Engine/Source/Editor/Persona/Private/PersonaAssetFamily.h b/Engine/Source/Editor/Persona/Private/PersonaAssetFamily.h
--- a/Engine/Source/Editor/Persona/Private/PersonaAssetFamily.h
+++ b/Engine/Source/Editor/Persona/Private/PersonaAssetFamily.h
@@ -8,6 +8,7 @@ class FPersonaAssetFamily : public IAssetFamily
 {
 public:
 	FPersonaAssetFamily(const UObject* InFromObject);
+	virtual ~FPersonaAssetFamily() {}
 
 	/** IAssetFamily interface */
 	virtual void GetAssetTypes(TArray<UClass*>& OutAssetTypes) const override;
diff --git a/Engine/Source/Editor/Persona/Private/PersonaMeshDetails.cpp b/Engine/Source/Editor/Persona/Private/PersonaMeshDetails.cpp
--- a/Engine/Source/Editor/Persona/Private/PersonaMeshDetails.cpp
+++ b/Engine/Source/Editor/Persona/Private/PersonaMeshDetails.cpp
@@ -484,7 +484,8 @@ void FPersonaMeshDetails::AddLODLevelCategories(IDetailLayoutBuilder& DetailLayo
 				.VAlign(VAlign_Center)
 				[
 					SNew(STextBlock)
-					.Text(LOCTEXT("AddLODLevelCategories_MaterialArrayOperations", "Materials Operations"))
+					.Font(IDetailLayoutBuilder::GetDetailFont())
+					.Text(LOCTEXT("AddLODLevelCategories_MaterialArrayOperations", "Material Slots"))
 				]
 				.ValueContent()
 				.HAlign(HAlign_Left)
@@ -500,6 +501,7 @@ void FPersonaMeshDetails::AddLODLevelCategories(IDetailLayoutBuilder& DetailLayo
 						.VAlign(VAlign_Center)
 						[
 							SNew(STextBlock)
+							.Font(IDetailLayoutBuilder::GetDetailFont())
 							.Text(this, &FPersonaMeshDetails::GetMaterialArrayText)
 						]
 						+ SHorizontalBox::Slot()
@@ -530,7 +532,6 @@ void FPersonaMeshDetails::AddLODLevelCategories(IDetailLayoutBuilder& DetailLayo
 
 				MaterialListDelegates.OnGetMaterials.BindSP(this, &FPersonaMeshDetails::OnGetMaterialsForArray, 0);
 				MaterialListDelegates.OnMaterialChanged.BindSP(this, &FPersonaMeshDetails::OnMaterialArrayChanged, 0);
-				MaterialListDelegates.OnGenerateCustomNameWidgets.BindSP(this, &FPersonaMeshDetails::OnGenerateCustomNameWidgetsForMaterialArray, 0);
 				MaterialListDelegates.OnGenerateCustomMaterialWidgets.BindSP(this, &FPersonaMeshDetails::OnGenerateCustomMaterialWidgetsForMaterialArray, 0);
 				MaterialListDelegates.OnMaterialListDirty.BindSP(this, &FPersonaMeshDetails::OnMaterialListDirty);
 				MaterialCategory.AddCustomBuilder(MakeShareable(new FMaterialList(MaterialCategory.GetParentLayout(), MaterialListDelegates, false)));
@@ -1135,7 +1136,7 @@ void FPersonaMeshDetails::OnSetPostProcessBlueprint(const FAssetData& AssetData,
 
 void FPersonaMeshDetails::OnGetMaterialsForArray(class IMaterialListBuilder& OutMaterials, int32 LODIndex)
 {
-	USkeletalMesh* SkelMesh = SkeletalMeshPtr.Get();
+	USkeletalMesh* SkelMesh = GetPersonaToolkit()->GetMesh();
 
 	if (!SkelMesh)
 		return;
@@ -1148,7 +1149,7 @@ void FPersonaMeshDetails::OnGetMaterialsForArray(class IMaterialListBuilder& Out
 
 void FPersonaMeshDetails::OnMaterialArrayChanged(UMaterialInterface* NewMaterial, UMaterialInterface* PrevMaterial, int32 SlotIndex, bool bReplaceAll, int32 LODIndex)
 {
-	USkeletalMesh* Mesh = SkeletalMeshPtr.Get();
+	USkeletalMesh* Mesh = GetPersonaToolkit()->GetMesh();
 	if (Mesh)
 	{
 		// Whether or not we made a transaction and need to end it
@@ -1189,20 +1190,23 @@ void FPersonaMeshDetails::OnMaterialArrayChanged(UMaterialInterface* NewMaterial
 FReply FPersonaMeshDetails::AddMaterialSlot()
 {
 	if (!SkeletalMeshPtr.IsValid())
+	{
 		return FReply::Handled();
-	GEditor->BeginTransaction(LOCTEXT("PersonaChangedMaterialIndex", "Change material index on mesh"));
+	}
+
+	FScopedTransaction Transaction(LOCTEXT("PersonaChangedMaterialIndex", "Change material index on mesh"));
+
 	SkeletalMeshPtr->Modify();
 	SkeletalMeshPtr->Materials.Add(FSkeletalMaterial());
 
 	SkeletalMeshPtr->PostEditChange();
-	// End the transation if we created one
-	GEditor->EndTransaction();
+
 	return FReply::Handled();
 }
 
 FText FPersonaMeshDetails::GetMaterialArrayText() const
 {
-	FString MaterialArrayText = TEXT(" Materials Slots");
+	FString MaterialArrayText = TEXT(" Material Slots");
 	int32 SlotNumber = 0;
 	if (SkeletalMeshPtr.IsValid())
 	{
@@ -1311,108 +1315,23 @@ void FPersonaMeshDetails::OnMaterialNameChanged(const FText& InValue, int32 Mate
 	}
 }
 
-TSharedRef<SWidget> FPersonaMeshDetails::OnGenerateCustomNameWidgetsForMaterialArray(UMaterialInterface* Material, int32 MaterialIndex, int32 LODIndexGarbage)
+
+TSharedRef<SWidget> FPersonaMeshDetails::OnGenerateCustomMaterialWidgetsForMaterialArray(UMaterialInterface* Material, int32 MaterialIndex, int32 LODIndex)
 {
-	bool MaterialIsUsed = false;
-	if (SkeletalMeshPtr.IsValid() && MaterialUsedMap.Contains(MaterialIndex))
+	bool bMaterialIsUsed = false;
+	if(SkeletalMeshPtr.IsValid() && MaterialUsedMap.Contains(MaterialIndex))
 	{
-		MaterialIsUsed = MaterialUsedMap.Find(MaterialIndex)->Num() > 0;
+		bMaterialIsUsed = MaterialUsedMap.Find(MaterialIndex)->Num() > 0;
 	}
-	TSharedRef<SVerticalBox> CustomNameWidget = SNew(SVerticalBox)
-		+ SVerticalBox::Slot()
-		.AutoHeight()
-		.Padding(0, 2, 0, 0)
-		[
-			SNew(SHorizontalBox)
-			+ SHorizontalBox::Slot()
-			.VAlign(VAlign_Center)
-			.Padding(0.0f, 3.0f, 0.0f, 0.0f)
-			.FillWidth(0.4f)
-		[
-				SNew(STextBlock)
-				.Font(IDetailLayoutBuilder::GetDetailFont())
-				.Text(LOCTEXT("MaterialArrayNameLabelStringKey", "Name"))
-				.ToolTipText(this, &FPersonaMeshDetails::GetOriginalImportMaterialNameText, MaterialIndex)
-			]
-			+ SHorizontalBox::Slot()
-			.Padding(5.0f, 3.0f, 5.0f, 0.0f)
-			.VAlign(VAlign_Center)
-			.FillWidth(1.0f)
-			[
-				SNew(SEditableTextBox)
-				.Text(this, &FPersonaMeshDetails::GetMaterialNameText, MaterialIndex)
-				.OnTextChanged(this, &FPersonaMeshDetails::OnMaterialNameChanged, MaterialIndex)
-				.OnTextCommitted(this, &FPersonaMeshDetails::OnMaterialNameCommitted, MaterialIndex)
-			]
-		];
 
-	if (!MaterialIsUsed)
-	{
-		CustomNameWidget->AddSlot()
-			.AutoHeight()
-			.Padding(0, 2, 0, 0)
-			[
-				SNew(SHorizontalBox)
-				+ SHorizontalBox::Slot()
-				.Padding(0.0f, 3.0f, 0.0f, 0.0f)
-				.VAlign(VAlign_Center)
-				.FillWidth(0.4f)
-			[
-				SNew(STextBlock)
-				.Font(IDetailLayoutBuilder::GetDetailFont())
-					.Text(LOCTEXT("CustomNameMaterialNotUsed", "Unused"))
-					.ToolTipText(LOCTEXT("CustomNameMaterialNotUsedTooltip", "This material is not use by any LODs section! It will be set to None when the skeletal mesh will be saved."))
-			]
-				+ SHorizontalBox::Slot()
-				.Padding(5.0f, 3.0f, 5.0f, 0.0f)
-				.VAlign(VAlign_Center)
-				.FillWidth(1.0f)
-				[
-					SNew(SButton)
-					.Text(LOCTEXT("CustomNameMaterialNotUsedDelete", "Delete"))
-					.ToolTipText(LOCTEXT("CustomNameMaterialNotUsedDeleteTooltip", "Delete this material slot."))
-					.IsEnabled(this, &FPersonaMeshDetails::CanDeleteMaterialSlot, MaterialIndex)
-					.OnClicked(this, &FPersonaMeshDetails::OnDeleteMaterialSlot, MaterialIndex)
-		]
-			];
-	}
-	else
-	{
-		CustomNameWidget->AddSlot()
-		.AutoHeight()
-		.Padding(0, 2, 0, 0)
-		[
-				SNew(SHorizontalBox)
-				+ SHorizontalBox::Slot()
-				.VAlign(VAlign_Center)
-				.Padding(0.0f, 3.0f, 0.0f, 0.0f)
-				.FillWidth(0.4f)
-				[
-					SNew(STextBlock)
-					.Font(IDetailLayoutBuilder::GetDetailFont())
-					.Text(LOCTEXT("MaterialSlotUsedByLabel", "Used By"))
-					.ToolTipText(LOCTEXT("MaterialSlotUsedByLabelTooltip", "List all sections that use this material slot"))
-				]
-				+ SHorizontalBox::Slot()
-				.Padding(5.0f, 3.0f, 5.0f, 0.0f)
-				.VAlign(VAlign_Center)
-				.FillWidth(1.0f)
-				[
-					//Material Slot Name
-					SNew(SComboButton)
-					.OnGetMenuContent(this, &FPersonaMeshDetails::OnGetMaterialSlotUsedByMenuContent, MaterialIndex)
-					.VAlign(VAlign_Center)
-					.ContentPadding(2)
-					.ButtonContent()
-			[
-				SNew(STextBlock)
-				.Font(IDetailLayoutBuilder::GetDetailFont())
-						.Text(this, &FPersonaMeshDetails::GetFirstMaterialSlotUsedBySection, MaterialIndex)
-					]
-			]
-		];
-}
-	return CustomNameWidget;
+	return
+		SNew(SMaterialSlotWidget, MaterialIndex, bMaterialIsUsed)
+		.MaterialName(this, &FPersonaMeshDetails::GetMaterialNameText, MaterialIndex)
+		.OnMaterialNameChanged(this, &FPersonaMeshDetails::OnMaterialNameChanged, MaterialIndex)
+		.OnMaterialNameCommitted(this, &FPersonaMeshDetails::OnMaterialNameCommitted, MaterialIndex)
+		.CanDeleteMaterialSlot(this, &FPersonaMeshDetails::CanDeleteMaterialSlot, MaterialIndex)
+		.OnDeleteMaterialSlot(this, &FPersonaMeshDetails::OnDeleteMaterialSlot, MaterialIndex)
+		.ToolTipText(this, &FPersonaMeshDetails::GetOriginalImportMaterialNameText, MaterialIndex);
 }
 
 FText FPersonaMeshDetails::GetFirstMaterialSlotUsedBySection(int32 MaterialIndex) const
@@ -1454,24 +1373,25 @@ TSharedRef<SWidget> FPersonaMeshDetails::OnGetMaterialSlotUsedByMenuContent(int3
 bool FPersonaMeshDetails::CanDeleteMaterialSlot(int32 MaterialIndex) const
 	{
 	if (!SkeletalMeshPtr.IsValid())
+	{
 		return false;
+	}
+
 	return (MaterialIndex + 1) == SkeletalMeshPtr->Materials.Num();
 	}
 	
-FReply FPersonaMeshDetails::OnDeleteMaterialSlot(int32 MaterialIndex)
+void FPersonaMeshDetails::OnDeleteMaterialSlot(int32 MaterialIndex)
 {
 	if (!SkeletalMeshPtr.IsValid() || !CanDeleteMaterialSlot(MaterialIndex))
-		return FReply::Handled();
+	{
+		return;
+	}
 
-	GEditor->BeginTransaction(LOCTEXT("PersonaDeletedMaterialSlot", "Deleted material slot on skeletal mesh"));
+	FScopedTransaction Transaction(LOCTEXT("PersonaDeletedMaterialSlot", "Delete material slot on Skeletal Mesh"));
 	SkeletalMeshPtr->Modify();
 	SkeletalMeshPtr->Materials.RemoveAt(MaterialIndex);
 
 	SkeletalMeshPtr->PostEditChange();
-	// End the transation if we created one
-	GEditor->EndTransaction();
-
-	return FReply::Handled();
 }
 
 bool FPersonaMeshDetails::OnMaterialListDirty()
@@ -1551,42 +1471,6 @@ bool FPersonaMeshDetails::OnMaterialListDirty()
 	return ForceMaterialListRefresh;
 }
 
-TSharedRef<SWidget> FPersonaMeshDetails::OnGenerateCustomMaterialWidgetsForMaterialArray(UMaterialInterface* Material, int32 MaterialIndex, int32 LODIndex)
-{
-	TSharedRef<SWidget> MaterialWidget
-	= SNew(SVerticalBox)
-
-		+SVerticalBox::Slot()
-		.Padding(0,2,0,0)
-		[
-			SNew(SCheckBox)
-			.IsChecked(this, &FPersonaMeshDetails::IsShadowCastingEnabled, MaterialIndex)
-			.OnCheckStateChanged(this, &FPersonaMeshDetails::OnShadowCastingChanged, MaterialIndex)
-			[
-				SNew(STextBlock)
-				.Font(FEditorStyle::GetFontStyle("StaticMeshEditor.NormalFont"))
-				.Text(LOCTEXT("Cast Shadows", "Cast Shadows"))
-			]
-		]
-
-		+SVerticalBox::Slot()
-		.Padding(0,2,0,0)
-		[
-			SNew(SCheckBox)
-			.IsChecked(this, &FPersonaMeshDetails::IsRecomputeTangentEnabled, MaterialIndex)
-			.OnCheckStateChanged(this, &FPersonaMeshDetails::OnRecomputeTangentChanged, MaterialIndex)
-			[
-				SNew(STextBlock)
-				.Font(FEditorStyle::GetFontStyle("StaticMeshEditor.NormalFont"))
-				.Text(LOCTEXT("RecomputeTangent_Title", "Recompute Tangent"))
-				.ToolTipText(LOCTEXT("RecomputeTangent_Tooltip", "This feature only works if you enable skin cache (r.SkinCache.Mode) and recompute tangent console variable(r.SkinCache.RecomputeTangents). Please note that skin cache is an experimental feature and only works if you have compute shaders."))
-			]
-		];
-
-	return MaterialWidget;
-}
-
-
 TSharedRef<SWidget> FPersonaMeshDetails::OnGenerateCustomNameWidgetsForSection(int32 LodIndex, int32 SectionIndex)
 {
 	return SNew(SVerticalBox)
@@ -1715,7 +1599,7 @@ TSharedRef<SWidget> FPersonaMeshDetails::OnGenerateCustomSectionWidgetsForSectio
 ECheckBoxState FPersonaMeshDetails::IsSectionSelected(int32 SectionIndex) const
 {
 	ECheckBoxState State = ECheckBoxState::Unchecked;
-	const USkeletalMesh* Mesh = SkeletalMeshPtr.Get();
+	const USkeletalMesh* Mesh = GetPersonaToolkit()->GetMesh();
 	if (Mesh)
 	{
 		State = Mesh->SelectedEditorSection == SectionIndex ? ECheckBoxState::Checked : ECheckBoxState::Unchecked;
@@ -1726,7 +1610,7 @@ ECheckBoxState FPersonaMeshDetails::IsSectionSelected(int32 SectionIndex) const
 
 void FPersonaMeshDetails::OnSectionSelectedChanged(ECheckBoxState NewState, int32 SectionIndex)
 {
-	USkeletalMesh* Mesh = SkeletalMeshPtr.Get();
+	USkeletalMesh* Mesh = GetPersonaToolkit()->GetMesh();
 
 	// Currently assumes that we only ever have one preview mesh in Persona.
 	UDebugSkelMeshComponent* MeshComponent = GetPersonaToolkit()->GetPreviewScene()->GetPreviewMeshComponent();
@@ -1763,7 +1647,7 @@ ECheckBoxState FPersonaMeshDetails::IsIsolateSectionEnabled(int32 SectionIndex)
 
 void FPersonaMeshDetails::OnSectionIsolatedChanged(ECheckBoxState NewState, int32 SectionIndex)
 {
-	USkeletalMesh* Mesh = SkeletalMeshPtr.Get();
+	USkeletalMesh* Mesh = GetPersonaToolkit()->GetMesh();
 	UDebugSkelMeshComponent * MeshComponent = GetPersonaToolkit()->GetPreviewScene()->GetPreviewMeshComponent();
 	if (Mesh && MeshComponent)
 	{
@@ -1787,7 +1671,7 @@ void FPersonaMeshDetails::OnSectionIsolatedChanged(ECheckBoxState NewState, int3
 ECheckBoxState FPersonaMeshDetails::IsShadowCastingEnabled(int32 MaterialIndex) const
 {
 	ECheckBoxState State = ECheckBoxState::Unchecked;
-	const USkeletalMesh* Mesh = SkeletalMeshPtr.Get();
+	const USkeletalMesh* Mesh = GetPersonaToolkit()->GetMesh();
 	bool FirstValueSet = false;
 	bool AllValueState = false;
 	bool AllValueSame = true;
@@ -1827,7 +1711,7 @@ ECheckBoxState FPersonaMeshDetails::IsShadowCastingEnabled(int32 MaterialIndex)
 
 void FPersonaMeshDetails::OnShadowCastingChanged(ECheckBoxState NewState, int32 MaterialIndex)
 {
-	USkeletalMesh* Mesh = SkeletalMeshPtr.Get();
+	USkeletalMesh* Mesh = GetPersonaToolkit()->GetMesh();
 
 	if (Mesh)
 	{
@@ -1878,7 +1762,7 @@ void FPersonaMeshDetails::OnShadowCastingChanged(ECheckBoxState NewState, int32
 ECheckBoxState FPersonaMeshDetails::IsRecomputeTangentEnabled(int32 MaterialIndex) const
 {
 	ECheckBoxState State = ECheckBoxState::Unchecked;
-	const USkeletalMesh* Mesh = SkeletalMeshPtr.Get();
+	const USkeletalMesh* Mesh = GetPersonaToolkit()->GetMesh();
 	if (Mesh == nullptr)
 		return State;
 
@@ -1920,7 +1804,7 @@ ECheckBoxState FPersonaMeshDetails::IsRecomputeTangentEnabled(int32 MaterialInde
 
 void FPersonaMeshDetails::OnRecomputeTangentChanged(ECheckBoxState NewState, int32 MaterialIndex)
 {
-	USkeletalMesh* Mesh = SkeletalMeshPtr.Get();
+	USkeletalMesh* Mesh = GetPersonaToolkit()->GetMesh();
 
 	if (Mesh)
 	{
@@ -1956,7 +1840,7 @@ void FPersonaMeshDetails::OnRecomputeTangentChanged(ECheckBoxState NewState, int
 ECheckBoxState FPersonaMeshDetails::IsSectionShadowCastingEnabled(int32 LODIndex, int32 SectionIndex) const
 {
 	ECheckBoxState State = ECheckBoxState::Unchecked;
-	const USkeletalMesh* Mesh = SkeletalMeshPtr.Get();
+	const USkeletalMesh* Mesh = GetPersonaToolkit()->GetMesh();
 	if (Mesh == nullptr)
 		return State;
 	
@@ -1979,7 +1863,7 @@ ECheckBoxState FPersonaMeshDetails::IsSectionShadowCastingEnabled(int32 LODIndex
 void FPersonaMeshDetails::OnSectionShadowCastingChanged(ECheckBoxState NewState, int32 LODIndex, int32 SectionIndex)
 {
 	ECheckBoxState State = ECheckBoxState::Unchecked;
-	USkeletalMesh* Mesh = SkeletalMeshPtr.Get();
+	USkeletalMesh* Mesh = GetPersonaToolkit()->GetMesh();
 	if (Mesh == nullptr)
 		return;
 
@@ -2024,7 +1908,7 @@ void FPersonaMeshDetails::OnSectionShadowCastingChanged(ECheckBoxState NewState,
 ECheckBoxState FPersonaMeshDetails::IsSectionRecomputeTangentEnabled(int32 LODIndex, int32 SectionIndex) const
 {
 	ECheckBoxState State = ECheckBoxState::Unchecked;
-	const USkeletalMesh* Mesh = SkeletalMeshPtr.Get();
+	const USkeletalMesh* Mesh = GetPersonaToolkit()->GetMesh();
 	if (Mesh == nullptr)
 		return State;
 
@@ -2047,7 +1931,7 @@ ECheckBoxState FPersonaMeshDetails::IsSectionRecomputeTangentEnabled(int32 LODIn
 void FPersonaMeshDetails::OnSectionRecomputeTangentChanged(ECheckBoxState NewState, int32 LODIndex, int32 SectionIndex)
 {
 	ECheckBoxState State = ECheckBoxState::Unchecked;
-	USkeletalMesh* Mesh = SkeletalMeshPtr.Get();
+	USkeletalMesh* Mesh = GetPersonaToolkit()->GetMesh();
 	if (Mesh == nullptr)
 		return;
 
@@ -2152,7 +2036,7 @@ bool FPersonaMeshDetails::IsDuplicatedMaterialIndex(int32 LODIndex, int32 Materi
 
 void FPersonaMeshDetails::OnSectionChanged(int32 LODIndex, int32 SectionIndex, int32 NewMaterialSlotIndex, FName NewMaterialSlotName)
 {
-	USkeletalMesh* Mesh = SkeletalMeshPtr.Get();
+	USkeletalMesh* Mesh = GetPersonaToolkit()->GetMesh();
 	if(Mesh)
 	{
 		FSkeletalMeshResource* ImportedResource = Mesh->GetImportedResource();
@@ -2169,32 +2053,32 @@ void FPersonaMeshDetails::OnSectionChanged(int32 LODIndex, int32 SectionIndex, i
 				NewSkeletalMaterialIndex = SkeletalMaterialIndex;
 				break;
 			}
-			}
+		}
 
 		check(NewSkeletalMaterialIndex != INDEX_NONE);
 
 		// Begin a transaction for undo/redo the first time we encounter a material to replace.  
 		// There is only one transaction for all replacement
-		GEditor->BeginTransaction(LOCTEXT("PersonaChangedMaterialIndex", "Change material index on mesh"));
+		FScopedTransaction Transaction(LOCTEXT("PersonaChangedMaterialIndex", "Change material index on mesh"));
 		Mesh->Modify();
 
 		FSkeletalMeshLODInfo& Info = Mesh->LODInfo[LODIndex];
 		if (LODIndex == 0 || Info.LODMaterialMap.Num() == 0)
-				{
+		{
 			ImportedResource->LODModels[LODIndex].Sections[SectionIndex].MaterialIndex = NewSkeletalMaterialIndex;
-				}
-				else
-				{
+		}
+		else
+		{
 			check(SectionIndex < Info.LODMaterialMap.Num());
 			Info.LODMaterialMap[SectionIndex] = NewSkeletalMaterialIndex;
-				}
-		Mesh->PostEditChange();
-			// End the transation if we created one
-			GEditor->EndTransaction();
-			// Redraw viewports to reflect the material changes 
-			GUnrealEd->RedrawLevelEditingViewports();
 		}
+
+		Mesh->PostEditChange();
+
+		// Redraw viewports to reflect the material changes 
+		GUnrealEd->RedrawLevelEditingViewports();
 	}
+}
 
 #if WITH_APEX_CLOTHING
 
diff --git a/Engine/Source/Editor/Persona/Private/PersonaMeshDetails.h b/Engine/Source/Editor/Persona/Private/PersonaMeshDetails.h
--- a/Engine/Source/Editor/Persona/Private/PersonaMeshDetails.h
+++ b/Engine/Source/Editor/Persona/Private/PersonaMeshDetails.h
@@ -190,14 +190,6 @@ private:
 
 	FText GetOriginalImportMaterialNameText(int32 MaterialIndex)const;
 
-	/**
-	* Called by the material list widget on generating each name widget
-	*
-	* @param Material		The material that is being displayed
-	* @param MaterialIndex	The index of the material slot
-	*/
-	TSharedRef<SWidget> OnGenerateCustomNameWidgetsForMaterialArray(UMaterialInterface* Material, int32 MaterialIndex, int32 LODIndex);
-
 	/**
 	* Called by the material list widget on generating each thumbnail widget
 	*
@@ -211,7 +203,7 @@ private:
 
 	bool CanDeleteMaterialSlot(int32 MaterialIndex) const;
 
-	FReply OnDeleteMaterialSlot(int32 MaterialIndex);
+	void OnDeleteMaterialSlot(int32 MaterialIndex);
 
 	TSharedRef<SWidget> OnGetMaterialSlotUsedByMenuContent(int32 MaterialIndex);
 
diff --git a/Engine/Source/Editor/Persona/Private/PersonaModule.cpp b/Engine/Source/Editor/Persona/Private/PersonaModule.cpp
--- a/Engine/Source/Editor/Persona/Private/PersonaModule.cpp
+++ b/Engine/Source/Editor/Persona/Private/PersonaModule.cpp
@@ -455,7 +455,7 @@ void FPersonaModule::TestSkeletonCurveNamesForUse(const TSharedRef<IEditableSkel
 
 			bool bFoundIssue = false;
 
-			const FText ProcessingAnimStatusUpdate = FText::Format(LOCTEXT("VerifyCurves_ProcessingSkeletalMeshes", "Finding animations that reference unused curves on skeleton '{0}'"), FText::FromString(Skeleton.GetName()));
+			const FText ProcessingAnimStatusUpdate = FText::Format(LOCTEXT("FindUnusedCurves_ProcessingSkeletalMeshes", "Finding animations that reference unused curves on skeleton '{0}'"), FText::FromString(Skeleton.GetName()));
 			{
 				FScopedSlowTask ProcessingAnimationsSlowTask(Animations.Num(), ProcessingAnimStatusUpdate);
 				ProcessingAnimationsSlowTask.MakeDialog();
diff --git a/Engine/Source/Editor/Persona/Private/PersonaToolkit.h b/Engine/Source/Editor/Persona/Private/PersonaToolkit.h
--- a/Engine/Source/Editor/Persona/Private/PersonaToolkit.h
+++ b/Engine/Source/Editor/Persona/Private/PersonaToolkit.h
@@ -9,6 +9,7 @@ class FPersonaToolkit : public IPersonaToolkit, public TSharedFromThis<FPersonaT
 {
 public:
 	FPersonaToolkit();
+	virtual ~FPersonaToolkit() {}
 
 	/** Initialize from a various sources */
 	void Initialize(USkeleton* InSkeleton);
diff --git a/Engine/Source/Editor/Persona/Private/SAnimationDlgs.cpp b/Engine/Source/Editor/Persona/Private/SAnimationDlgs.cpp
--- a/Engine/Source/Editor/Persona/Private/SAnimationDlgs.cpp
+++ b/Engine/Source/Editor/Persona/Private/SAnimationDlgs.cpp
@@ -66,12 +66,6 @@ void SImportPathDialog::Construct(const FArguments& InArgs)
 					[
 						ContentBrowserModule.Get().CreatePathPicker(PathPickerConfig)
 					]
-
-					+SVerticalBox::Slot()
-					.AutoHeight()
-					[
-						SNew(SSeparator)
-					]
 				]
 			]
 
@@ -91,6 +85,7 @@ void SImportPathDialog::Construct(const FArguments& InArgs)
 					.ContentPadding(FEditorStyle::GetMargin("StandardDialog.ContentPadding"))
 					.Text(LOCTEXT("OK", "OK"))
 					.OnClicked(this, &SImportPathDialog::OnButtonClick, EAppReturnType::Ok)
+					.IsEnabled(this, &SImportPathDialog::IsOkButtonEnabled)
 				]
 				+SUniformGridPanel::Slot(1, 0)
 				[
@@ -118,6 +113,11 @@ FReply SImportPathDialog::OnButtonClick(EAppReturnType::Type ButtonID)
 	return FReply::Handled();
 }
 
+bool SImportPathDialog::IsOkButtonEnabled() const
+{
+	return !AssetPath.IsEmptyOrWhitespace();
+}
+
 EAppReturnType::Type SImportPathDialog::ShowModal()
 {
 	GEditor->EditorAddModalWindow(SharedThis(this));
diff --git a/Engine/Source/Editor/Persona/Private/SAnimationDlgs.h b/Engine/Source/Editor/Persona/Private/SAnimationDlgs.h
--- a/Engine/Source/Editor/Persona/Private/SAnimationDlgs.h
+++ b/Engine/Source/Editor/Persona/Private/SAnimationDlgs.h
@@ -29,7 +29,7 @@ public:
 protected:
 	void OnPathChange(const FString& NewPath);
 	FReply OnButtonClick(EAppReturnType::Type ButtonID);
-
+	bool IsOkButtonEnabled() const;
 	EAppReturnType::Type UserResponse;
 	FText AssetPath;
 	FText AssetName;
