Fix for CIS linux error
#fix incorrect case in include path

diff --git a/Engine/Source/Editor/Persona/Private/SAnimationBlendSpaceGridWidget.cpp b/Engine/Source/Editor/Persona/Private/SAnimationBlendSpaceGridWidget.cpp
--- a/Engine/Source/Editor/Persona/Private/SAnimationBlendSpaceGridWidget.cpp
+++ b/Engine/Source/Editor/Persona/Private/SAnimationBlendSpaceGridWidget.cpp
@@ -7,7 +7,7 @@
 #include "Animation/BlendSpaceBase.h"
 #include "Animation/BlendSpace1D.h"
 
-#include "STooltip.h"
+#include "SToolTip.h"
 #include "SNumericEntryBox.h"
 
 #include "PropertyEditing.h"
