Fix IWYU non-unity CIS failures from Framework copy up

diff --git a/Engine/Source/Editor/LevelEditor/Private/LevelEditorToolBar.cpp b/Engine/Source/Editor/LevelEditor/Private/LevelEditorToolBar.cpp
--- a/Engine/Source/Editor/LevelEditor/Private/LevelEditorToolBar.cpp
+++ b/Engine/Source/Editor/LevelEditor/Private/LevelEditorToolBar.cpp
@@ -52,6 +52,7 @@
 #include "EngineUtils.h"
 #include "ScopedTransaction.h"
 #include "Features/EditorFeatures.h"
+#include "ConfigCacheIni.h"
 
 namespace LevelEditorActionHelpers
 {
diff --git a/Engine/Source/Editor/Persona/Private/SPoseEditor.cpp b/Engine/Source/Editor/Persona/Private/SPoseEditor.cpp
--- a/Engine/Source/Editor/Persona/Private/SPoseEditor.cpp
+++ b/Engine/Source/Editor/Persona/Private/SPoseEditor.cpp
@@ -5,6 +5,7 @@
 #include "Misc/MessageDialog.h"
 #include "Framework/MultiBox/MultiBoxBuilder.h"
 #include "Widgets/Input/SSpinBox.h"
+#include "Widgets/Layout/SSplitter.h"
 #include "Animation/DebugSkelMeshComponent.h"
 #include "ScopedTransaction.h"
 #include "Widgets/Input/SSearchBox.h"
diff --git a/Engine/Source/Editor/UnrealEd/Private/FbxExporter.h b/Engine/Source/Editor/UnrealEd/Private/FbxExporter.h
--- a/Engine/Source/Editor/UnrealEd/Private/FbxExporter.h
+++ b/Engine/Source/Editor/UnrealEd/Private/FbxExporter.h
@@ -35,6 +35,7 @@ class USkeletalMesh;
 class USkeletalMeshComponent;
 class USplineMeshComponent;
 class UStaticMeshComponent;
+class FColorVertexBuffer;
 struct FAnimControlTrackKey;
 struct FExpressionInput;
 struct FRichCurve;
