Copying //UE4/Dev-AnimPhys to //UE4/Dev-Main (Source: //UE4/Dev-AnimPhys @ 3624379)

diff --git a/Engine/Source/Editor/Persona/Private/SAnimCurvePanel.cpp b/Engine/Source/Editor/Persona/Private/SAnimCurvePanel.cpp
--- a/Engine/Source/Editor/Persona/Private/SAnimCurvePanel.cpp
+++ b/Engine/Source/Editor/Persona/Private/SAnimCurvePanel.cpp
@@ -38,13 +38,12 @@ class FAnimCurveBaseInterface : public FCurveOwnerInterface
 private:
 	FAnimCurveBase*	CurveData;
 
-	void UpdateNameInternal(FRawCurveTracks& RawCurveData, const SmartName::UID_Type& RequestedNameUID, const FGuid& RequestedNameGuid, FName RequestedName)
+	void UpdateNameInternal(FRawCurveTracks& RawCurveData, const SmartName::UID_Type& RequestedNameUID, FName RequestedName)
 	{
 		FAnimCurveBase* CurrentCurveData = RawCurveData.GetCurveData(CurveUID);
 		if (CurrentCurveData)
 		{
 			CurrentCurveData->Name.UID = RequestedNameUID;
-			CurrentCurveData->Name.Guid = RequestedNameGuid;
 			CurrentCurveData->Name.DisplayName = RequestedName;
 		}
 	}
@@ -160,12 +159,12 @@ public:
 		CurrentCurveData->FloatCurve.AddKey(0.0f, 1.0f);
 	}
 
-	void UpdateName(const SmartName::UID_Type& RequestedNameUID, const FGuid& RequestedNameGuid, FName RequestedName)
+	void UpdateName(const SmartName::UID_Type& RequestedNameUID, FName RequestedName)
 	{
-		UpdateNameInternal(AnimSequenceBase.Get()->RawCurveData, RequestedNameUID, RequestedNameGuid, RequestedName);
+		UpdateNameInternal(AnimSequenceBase.Get()->RawCurveData, RequestedNameUID, RequestedName);
 		if (UAnimSequence* Seq = AnimSequence.Get())
 		{
-			UpdateNameInternal(Seq->CompressedCurveData, RequestedNameUID, RequestedNameGuid, RequestedName);
+			UpdateNameInternal(Seq->CompressedCurveData, RequestedNameUID, RequestedName);
 		}
 		CurveUID = RequestedNameUID;
 	}
@@ -437,23 +436,16 @@ void SCurveEdTrack::NewCurveNameEntered( const FText& NewText, ETextCommit::Type
 
 				const FSmartNameMapping* NameMapping = Skeleton->GetSmartNameContainer(USkeleton::AnimCurveMappingName);
 
-				// If requested name exists, make sure it's not currently in use in this sequence.
-				SmartName::UID_Type RequestedNameUID = NameMapping->FindUID(RequestedName);
-				FGuid RequestedNameGuid;
-
 				FScopedTransaction Transaction(LOCTEXT("CurveEditor_RenameCurve", "Rename Curve"));
 
-				if (RequestedNameUID != SmartName::MaxUID)
+				FSmartName NewSmartName;
+				if (NameMapping->FindSmartName(RequestedName, NewSmartName))
 				{
-					FSmartName ExistingName;
-					bool bFoundSmartName = NameMapping->FindSmartNameByUID(RequestedNameUID, ExistingName);
-					check(bFoundSmartName);
-					RequestedNameGuid = ExistingName.Guid;
-
 					// Already in use in this sequence, and if it's not my UID
-					if (RequestedNameUID != CurveInterface->CurveUID && CurveInterface->AnimSequenceBase->RawCurveData.GetCurveData(RequestedNameUID) != nullptr)
+					if (NewSmartName.UID != CurveInterface->CurveUID && CurveInterface->AnimSequenceBase->RawCurveData.GetCurveData(NewSmartName.UID) != nullptr)
 					{
 						Transaction.Cancel();
+
 						FFormatNamedArguments Args;
 						Args.Add(TEXT("InvalidName"), FText::FromName(RequestedName));
 						FNotificationInfo Info(FText::Format(LOCTEXT("AnimCurveRenamedInUse", "The name \"{InvalidName}\" is already used."), Args));
@@ -466,44 +458,14 @@ void SCurveEdTrack::NewCurveNameEntered( const FText& NewText, ETextCommit::Type
 						{
 							Notification->SetCompletionState(SNotificationItem::CS_Fail);
 						}
-					}
-					else
-					{
-						CurveInterface->AnimSequenceBase->Modify(true);
-						CurveInterface->UpdateName(RequestedNameUID, RequestedNameGuid, RequestedName);
-
-						// Refresh the panel
-						TSharedPtr<SAnimCurvePanel> SharedPanel = PanelPtr.Pin();
-						if (SharedPanel.IsValid())
-						{
-							SharedPanel.Get()->UpdatePanel();
-						}
+						return;
 					}
 				}
 				else
 				{
-					const bool bRenamed = Skeleton->RenameSmartnameAndModify(USkeleton::AnimCurveMappingName, CurveInterface->CurveUID, RequestedName);
-					if (bRenamed)
-					{
-						FSmartName NewName;
-						Skeleton->GetSmartNameByName(USkeleton::AnimCurveMappingName, RequestedName, NewName);
-
-						RequestedNameUID = NewName.UID;
-						RequestedNameGuid = NewName.Guid;
-						CurveInterface->UpdateName(RequestedNameUID, RequestedNameGuid, RequestedName);
-
-						// Refresh the panel
-						TSharedPtr<SAnimCurvePanel> SharedPanel = PanelPtr.Pin();
-						if (SharedPanel.IsValid())
-						{
-							SharedPanel.Get()->UpdatePanel();
-						}
-					}
-					else
+					if(!Skeleton->AddSmartNameAndModify(USkeleton::AnimCurveMappingName, RequestedName, NewSmartName))
 					{
-						// Cancel the rename transaction
 						Transaction.Cancel();
-
 						FNotificationInfo Info(LOCTEXT("AnimCurveRenamedError", "Failed to rename curve smart name, check the log for errors."));
 
 						Info.bUseLargeFont = false;
@@ -513,9 +475,20 @@ void SCurveEdTrack::NewCurveNameEntered( const FText& NewText, ETextCommit::Type
 						if (Notification.IsValid())
 						{
 							Notification->SetCompletionState(SNotificationItem::CS_Fail);
-						}						
+						}
+						return;
 					}
 				}
+
+				CurveInterface->ModifyOwner();
+				CurveInterface->UpdateName(NewSmartName.UID, RequestedName);
+
+				// Refresh the panel
+				TSharedPtr<SAnimCurvePanel> SharedPanel = PanelPtr.Pin();
+				if (SharedPanel.IsValid())
+				{
+					SharedPanel.Get()->UpdatePanel();
+				}
 			}
 		}
 	}
@@ -635,7 +608,7 @@ void SAnimCurvePanel::Construct(const FArguments& InArgs, const TSharedRef<class
 	WidgetWidth = InArgs._WidgetWidth;
 	OnGetScrubValue = InArgs._OnGetScrubValue;
 
-	InEditableSkeleton->RegisterOnSmartNameRemoved(FOnSmartNameRemoved::FDelegate::CreateSP(this, &SAnimCurvePanel::HandleSmartNameRemoved));
+	InEditableSkeleton->RegisterOnSmartNameChanged(FOnSmartNameChanged::FDelegate::CreateSP(this, &SAnimCurvePanel::HandleSmartNamesChange));
 
 	Sequence->RegisterOnAnimCurvesChanged(UAnimSequenceBase::FOnAnimCurvesChanged::CreateSP(this, &SAnimCurvePanel::UpdatePanel));
 
@@ -1377,7 +1350,7 @@ void SAnimCurvePanel::AddVariableCurve(USkeleton::AnimCurveUID CurveUid)
 	UpdatePanel();
 }
 
-void SAnimCurvePanel::HandleSmartNameRemoved(const FName& InContainerName, const TArray<SmartName::UID_Type>& InNameUids)
+void SAnimCurvePanel::HandleSmartNamesChange(const FName& InContainerName)
 {
 	UpdatePanel();
 }
diff --git a/Engine/Source/Editor/Persona/Private/SAnimCurvePanel.h b/Engine/Source/Editor/Persona/Private/SAnimCurvePanel.h
--- a/Engine/Source/Editor/Persona/Private/SAnimCurvePanel.h
+++ b/Engine/Source/Editor/Persona/Private/SAnimCurvePanel.h
@@ -241,5 +241,5 @@ private:
 	void ToggleCurveTypeMenuCallback(FAnimCurveBaseInterface* Curve);
 
 	/** Handle smart names getting removed (and regenerate widgets) */
-	void HandleSmartNameRemoved(const FName& InContainerName, const TArray<SmartName::UID_Type>& InNameUids);
+	void HandleSmartNamesChange(const FName& InContainerName);
 };
diff --git a/Engine/Source/Editor/Persona/Private/SAnimCurveViewer.cpp b/Engine/Source/Editor/Persona/Private/SAnimCurveViewer.cpp
--- a/Engine/Source/Editor/Persona/Private/SAnimCurveViewer.cpp
+++ b/Engine/Source/Editor/Persona/Private/SAnimCurveViewer.cpp
@@ -545,7 +545,7 @@ void SAnimCurveViewer::Construct(const FArguments& InArgs, const TSharedRef<clas
 	InPreviewScene->RegisterOnAnimChanged(FOnAnimChanged::CreateSP(this, &SAnimCurveViewer::OnPreviewAssetChanged));
 	InOnPostUndo.Add(FSimpleDelegate::CreateSP(this, &SAnimCurveViewer::OnPostUndo));
 
-	SmartNameRemovedHandle = InEditableSkeleton->RegisterOnSmartNameRemoved(FOnSmartNameRemoved::FDelegate::CreateSP(this, &SAnimCurveViewer::HandleSmartNameRemoved));
+	SmartNameChangedHandle = InEditableSkeleton->RegisterOnSmartNameChanged(FOnSmartNameChanged::FDelegate::CreateSP(this, &SAnimCurveViewer::HandleSmartNamesChange));
 
 	// Register and bind all our menu commands
 	FCurveViewerCommands::Register();
@@ -637,7 +637,7 @@ SAnimCurveViewer::~SAnimCurveViewer()
 
 	if (EditableSkeletonPtr.IsValid())
 	{
-		EditableSkeletonPtr.Pin()->UnregisterOnSmartNameRemoved(SmartNameRemovedHandle);
+		EditableSkeletonPtr.Pin()->UnregisterOnSmartNameChanged(SmartNameChangedHandle);
 	}
 }
 
@@ -1075,14 +1075,14 @@ void SAnimCurveViewer::OnNameCommitted(const FText& InNewName, ETextCommit::Type
 void SAnimCurveViewer::OnDeleteNameClicked()
 {
 	TArray< TSharedPtr<FDisplayedAnimCurveInfo> > SelectedItems = AnimCurveListView->GetSelectedItems();
-	TArray<SmartName::UID_Type> SelectedUids;
+	TArray<FName> SelectedNames;
 
 	for (TSharedPtr<FDisplayedAnimCurveInfo> Item : SelectedItems)
 	{
-		SelectedUids.Add(Item->SmartName.UID);
+		SelectedNames.Add(Item->SmartName.DisplayName);
 	}
 
-	EditableSkeletonPtr.Pin()->RemoveSmartnamesAndFixupAnimations(ContainerName, SelectedUids);
+	EditableSkeletonPtr.Pin()->RemoveSmartnamesAndFixupAnimations(ContainerName, SelectedNames);
 }
 
 bool SAnimCurveViewer::CanDelete()
@@ -1106,12 +1106,14 @@ void SAnimCurveViewer::OnSelectionChanged(TSharedPtr<FDisplayedAnimCurveInfo> In
 			TArray<FBoneReference> BoneLinks;
 			FSmartName CurrentName = RowItem->SmartName;
 			const FCurveMetaData* CurveMetaData = EditableSkeletonPtr.Pin()->GetSkeleton().GetCurveMetaData(CurrentName);
+			uint32 MaxLOD = 0xFF;
 			if (CurveMetaData)
 			{
 				BoneLinks = CurveMetaData->LinkedBones;
+				MaxLOD = CurveMetaData->MaxLOD;
 			}
 
-			EditorMirrorObj->Refresh(CurrentName, BoneLinks);
+			EditorMirrorObj->Refresh(CurrentName, BoneLinks, MaxLOD);
 		}
 
 		SelectedObjects.Add(EditorMirrorObj);
@@ -1124,11 +1126,11 @@ void SAnimCurveViewer::ApplyCurveBoneLinks(class UEditorAnimCurveBoneLinks* Edit
 {
 	if (EditorObj)
 	{
-		EditableSkeletonPtr.Pin()->SetCurveMetaBoneLinks(EditorObj->CurveName, EditorObj->ConnectedBones);
+		EditableSkeletonPtr.Pin()->SetCurveMetaBoneLinks(EditorObj->CurveName, EditorObj->ConnectedBones, EditorObj->MaxLOD);
 	}
 }
 
-void SAnimCurveViewer::HandleSmartNameRemoved(const FName& InContainerName, const TArray<SmartName::UID_Type>& InNameUids)
+void SAnimCurveViewer::HandleSmartNamesChange(const FName& InContainerName)
 {
 	AnimCurveList.Empty();
 	RefreshCurveList();
diff --git a/Engine/Source/Editor/Persona/Private/SAnimCurveViewer.h b/Engine/Source/Editor/Persona/Private/SAnimCurveViewer.h
--- a/Engine/Source/Editor/Persona/Private/SAnimCurveViewer.h
+++ b/Engine/Source/Editor/Persona/Private/SAnimCurveViewer.h
@@ -284,7 +284,7 @@ private:
 	void CreateNewNameEntry(const FText& CommittedText, ETextCommit::Type CommitType);
 
 	/** Handle smart name (i.e. curve) removal */
-	void HandleSmartNameRemoved(const FName& InContainerName, const TArray<SmartName::UID_Type>& InNameUids);
+	void HandleSmartNamesChange(const FName& InContainerName);
 
 	/** Get the SmartNameMapping for anim curves */
 	const struct FSmartNameMapping* GetAnimCurveMapping();
@@ -338,5 +338,5 @@ private:
 	void ApplyCurveBoneLinks(class UEditorAnimCurveBoneLinks* EditorObj);
 
 	/** Delegate handle for HandleSmartNameRemoved callback */
-	FDelegateHandle SmartNameRemovedHandle;
+	FDelegateHandle SmartNameChangedHandle;
 };
\ No newline at end of file
diff --git a/Engine/Source/Editor/Persona/Private/SAnimMontagePanel.cpp b/Engine/Source/Editor/Persona/Private/SAnimMontagePanel.cpp
--- a/Engine/Source/Editor/Persona/Private/SAnimMontagePanel.cpp
+++ b/Engine/Source/Editor/Persona/Private/SAnimMontagePanel.cpp
@@ -376,6 +376,7 @@ void SAnimMontagePanel::SummonTrackContextMenu( FMenuBuilder& MenuBuilder, float
 						.MinValue(0.0f)
 						.MaxValue(Montage->SequenceLength)
 						.Value(Section.GetTime())
+						.AllowSpin(true)
 						.OnValueCommitted_Lambda([this, SectionIndex](float InValue, ETextCommit::Type InCommitType)
 						{
 							if (Montage->CompositeSections.IsValidIndex(SectionIndex))
@@ -405,6 +406,7 @@ void SAnimMontagePanel::SummonTrackContextMenu( FMenuBuilder& MenuBuilder, float
 						.MinValue(0)
 						.MaxValue(Montage->GetNumberOfFrames())
 						.Value(Montage->GetFrameAtTime(Section.GetTime()))
+						.AllowSpin(true)						
 						.OnValueCommitted_Lambda([this, SectionIndex](int32 InValue, ETextCommit::Type InCommitType)
 						{
 							if (Montage->CompositeSections.IsValidIndex(SectionIndex))
diff --git a/Engine/Source/Editor/Persona/Private/SAnimNotifyPanel.cpp b/Engine/Source/Editor/Persona/Private/SAnimNotifyPanel.cpp
--- a/Engine/Source/Editor/Persona/Private/SAnimNotifyPanel.cpp
+++ b/Engine/Source/Editor/Persona/Private/SAnimNotifyPanel.cpp
@@ -35,6 +35,9 @@
 #include "TabSpawners.h"
 #include "SInlineEditableTextBlock.h"
 #include "HAL/PlatformApplicationMisc.h"
+#include "Modules/ModuleManager.h"
+#include "IEditableSkeleton.h"
+#include "ISkeletonEditorModule.h"
 
 // Track Panel drawing
 const float NotificationTrackHeight = 20.0f;
@@ -3207,7 +3210,11 @@ void SAnimNotifyTrack::AddNewNotify(const FText& NewNotifyName, ETextCommit::Typ
 	{
 		const FScopedTransaction Transaction( LOCTEXT("AddNewNotifyEvent", "Add New Anim Notify") );
 		FName NewName = FName( *NewNotifyName.ToString() );
-		SeqSkeleton->AddNewAnimationNotify(NewName);
+
+		ISkeletonEditorModule& SkeletonEditorModule = FModuleManager::LoadModuleChecked<ISkeletonEditorModule>("SkeletonEditor");
+		TSharedRef<IEditableSkeleton> EditableSkeleton = SkeletonEditorModule.CreateEditableSkeleton(SeqSkeleton);
+
+		EditableSkeleton->AddNotify(NewName);
 
 		FBlueprintActionDatabase::Get().RefreshAssetActions(SeqSkeleton);
 
@@ -4877,6 +4884,11 @@ void SAnimNotifyPanel::OnGetNotifyBlueprintData(TArray<FAssetData>& OutNotifyDat
 			if(InOutAllowedClassNames->Contains(TagValue))
 			{
 				FString GenClass = AssetData.GetTagValueRef<FString>(BPGenClassName);
+				const uint32 ClassFlags = AssetData.GetTagValueRef<uint32>("ClassFlags");
+				if (ClassFlags & CLASS_Abstract)
+				{
+					continue;
+				}
 
 				if(!OutNotifyData.Contains(AssetData))
 				{
diff --git a/Engine/Source/Editor/Persona/Private/SAnimViewportToolBar.cpp b/Engine/Source/Editor/Persona/Private/SAnimViewportToolBar.cpp
--- a/Engine/Source/Editor/Persona/Private/SAnimViewportToolBar.cpp
+++ b/Engine/Source/Editor/Persona/Private/SAnimViewportToolBar.cpp
@@ -39,121 +39,6 @@
 
 #define LOCTEXT_NAMESPACE "AnimViewportToolBar"
 
-
-//Class definition which represents widget to modify viewport's background color.
-class SBackgroundColorSettings : public SCompoundWidget
-{
-public:
-
-	SLATE_BEGIN_ARGS(SBackgroundColorSettings)
-		{}
-		SLATE_ARGUMENT( TWeakPtr<SAnimationEditorViewportTabBody>, AnimEditorViewport )
-	SLATE_END_ARGS()
-
-	/** Constructs this widget from its declaration */
-	void Construct(const FArguments& InArgs )
-	{
-		AnimViewportPtr = InArgs._AnimEditorViewport;
-
-		TSharedPtr<SWidget> ExtraWidget = 
-			SNew(SBorder)
-			.BorderImage(FEditorStyle::GetBrush("FilledBorder"))
-			[
-				SNew(SHorizontalBox)
-				+SHorizontalBox::Slot()
-				.AutoWidth()
-				.Padding(1)
-				[
-					SNew( SColorBlock )
-					.Color(AnimViewportPtr.Pin().ToSharedRef(), &SAnimationEditorViewportTabBody::GetViewportBackgroundColor)
-					.IgnoreAlpha(true)
-					.ToolTipText(LOCTEXT("ColorBlock_ToolTip", "Select background color"))
-					.OnMouseButtonDown(this, &SBackgroundColorSettings::OnColorBoxClicked)
-				]
-			];
-
-		this->ChildSlot
-		[
-			SNew(SAnimPlusMinusSlider)
-				.Label( LOCTEXT("BrightNess", "Brightness:")  )
-				.IsEnabled( this, &SBackgroundColorSettings::IsBrightnessSliderEnabled )
-				.OnMinusClicked( this, &SBackgroundColorSettings::OnDecreaseBrightness )
-				.MinusTooltip( LOCTEXT("DecreaseBrightness_ToolTip", "Decrease brightness") )
-				.SliderValue(AnimViewportPtr.Pin().ToSharedRef(), &SAnimationEditorViewportTabBody::GetBackgroundBrightness)
-				.OnSliderValueChanged(AnimViewportPtr.Pin().ToSharedRef(), &SAnimationEditorViewportTabBody::SetBackgroundBrightness)
-				.SliderTooltip( LOCTEXT("BackgroundBrightness_ToolTip", "Change background brightness") )
-				.OnPlusClicked( this, &SBackgroundColorSettings::OnIncreaseBrightness )
-				.PlusTooltip( LOCTEXT("IncreaseBrightness_ToolTip", "Increase brightness") )
-				.ExtraWidget( ExtraWidget )
-		];
-	}
-
-protected:
-
-	/** Function to open color picker window when selected from context menu */
-	FReply OnColorBoxClicked(const FGeometry& MyGeometry, const FPointerEvent& MouseEvent)
-	{
-		if (MouseEvent.GetEffectingButton() != EKeys::LeftMouseButton)
-		{
-			return FReply::Unhandled();
-		}
-
-		FSlateApplication::Get().DismissAllMenus();
-
-		FLinearColor BackgroundColor = AnimViewportPtr.Pin()->GetViewportBackgroundColor();
-		TArray<FLinearColor*> LinearColorArray;
-		LinearColorArray.Add(&BackgroundColor);
-
-		FColorPickerArgs PickerArgs;
-		PickerArgs.bIsModal = true;
-		PickerArgs.ParentWidget = AnimViewportPtr.Pin();
-		PickerArgs.bOnlyRefreshOnOk = true;
-		PickerArgs.DisplayGamma = TAttribute<float>::Create( TAttribute<float>::FGetter::CreateUObject(GEngine, &UEngine::GetDisplayGamma) );
-		PickerArgs.LinearColorArray = &LinearColorArray;
-		PickerArgs.OnColorCommitted = FOnLinearColorValueChanged::CreateSP(AnimViewportPtr.Pin().ToSharedRef(), &SAnimationEditorViewportTabBody::SetViewportBackgroundColor);
-
-		if (OpenColorPicker(PickerArgs))
-		{
-			AnimViewportPtr.Pin()->RefreshViewport();
-		}
-
-		return FReply::Handled();
-	}
-
-	/** Callback function for decreasing color brightness */
-	FReply OnDecreaseBrightness( )
-	{
-		const float DeltaValue = 0.05f;
-		AnimViewportPtr.Pin()->SetBackgroundBrightness( AnimViewportPtr.Pin()->GetBackgroundBrightness() - DeltaValue );
-		return FReply::Handled();
-	}
-
-	/** Callback function for increasing color brightness */
-	FReply OnIncreaseBrightness( )
-	{
-		const float DeltaValue = 0.05f;
-		AnimViewportPtr.Pin()->SetBackgroundBrightness( AnimViewportPtr.Pin()->GetBackgroundBrightness() + DeltaValue );
-		return FReply::Handled();
-	}
-
-	/** Callback function which determines whether this widget is enabled */
-	bool IsBrightnessSliderEnabled() const
-	{
-		const UAssetViewerSettings* Settings = UAssetViewerSettings::Get();
-		const UEditorPerProjectUserSettings* PerProjectUserSettings = GetDefault<UEditorPerProjectUserSettings>();
-		const int32 ProfileIndex = Settings->Profiles.IsValidIndex(PerProjectUserSettings->AssetViewerProfileIndex) ? PerProjectUserSettings->AssetViewerProfileIndex : 0;
-
-		ensureMsgf(Settings->Profiles.IsValidIndex(PerProjectUserSettings->AssetViewerProfileIndex), TEXT("Invalid default settings pointer or current profile index"));
-		
-		return !Settings->Profiles[ProfileIndex].bShowEnvironment;
-	}
-
-protected:
-
-	/** The viewport hosting this widget */
-	TWeakPtr<SAnimationEditorViewportTabBody> AnimViewportPtr;
-};
-
 //Class definition which represents widget to modify strength of wind for clothing
 class SClothWindSettings : public SCompoundWidget
 {
@@ -276,87 +161,39 @@ protected:
 
 void SAnimViewportToolBar::Construct(const FArguments& InArgs, TSharedPtr<class SAnimationEditorViewportTabBody> InViewport, TSharedPtr<class SEditorViewport> InRealViewport)
 {
-	Viewport = InViewport;
+	bShowShowMenu = InArgs._ShowShowMenu;
+	bShowLODMenu = InArgs._ShowLODMenu;
+	bShowPlaySpeedMenu = InArgs._ShowPlaySpeedMenu;
+	bShowFloorOptions = InArgs._ShowFloorOptions;
+	bShowTurnTable = InArgs._ShowTurnTable;
 
-	TSharedRef<SAnimationEditorViewportTabBody> ViewportRef = Viewport.Pin().ToSharedRef();
+	CommandList = InRealViewport->GetCommandList();
+	Extenders = InArgs._Extenders;
 
+	// If we have no extender, make an empty one
+	if (!Extenders.IsValid())
+	{
+		Extenders = MakeShared<FExtender>();
+	}
+	
 	TSharedRef<SHorizontalBox> LeftToolbar = SNew(SHorizontalBox)
-			// Generic viewport options
-			+ SHorizontalBox::Slot()
-			.AutoWidth()
-			.Padding( 2.0f, 2.0f )
-			[
-				//Menu
-				SNew( SEditorViewportToolbarMenu )
-				.ParentToolBar( SharedThis( this ) )
-				.Image("EditorViewportToolBar.MenuDropdown")
-				.OnGetMenuContent( this, &SAnimViewportToolBar::GenerateViewMenu ) 
-			]
-
-			// Camera Type (Perspective/Top/etc...)
-			+ SHorizontalBox::Slot()
-			.AutoWidth()
-			.Padding( 2.0f, 2.0f )
-			[
-				SNew( SEditorViewportToolbarMenu )
-				.ParentToolBar( SharedThis( this ) )
-				.Label(this, &SAnimViewportToolBar::GetCameraMenuLabel)
-				.LabelIcon(this, &SAnimViewportToolBar::GetCameraMenuLabelIcon)
-				.OnGetMenuContent(this, &SAnimViewportToolBar::GenerateViewportTypeMenu)
-			]
-
-			// View menu (lit, unlit, etc...)
-			+ SHorizontalBox::Slot()
-			.AutoWidth()
-			.Padding(2.0f, 2.0f)
-			[
-				SNew( SEditorViewportViewMenu, InRealViewport.ToSharedRef(), SharedThis(this) )
-			]
-
-			// Show flags menu
-			+ SHorizontalBox::Slot()
-			.AutoWidth()
-			.Padding( 2.0f, 2.0f )
-			[
-				SNew( SEditorViewportToolbarMenu )
-				.ParentToolBar( SharedThis( this ) )
-				.Label( LOCTEXT("ShowMenu", "Show") )
-				.OnGetMenuContent( this, &SAnimViewportToolBar::GenerateShowMenu ) 
-			]
-
-			// LOD menu
-			+ SHorizontalBox::Slot()
-			.AutoWidth()
-			.Padding( 2.0f, 2.0f )
-			[
-				//LOD
-				SNew( SEditorViewportToolbarMenu )
-				.ParentToolBar( SharedThis( this ) )
-				.Label( this, &SAnimViewportToolBar::GetLODMenuLabel )
-				.OnGetMenuContent( this, &SAnimViewportToolBar::GenerateLODMenu ) 
-			]
+		+SHorizontalBox::Slot()
+		.AutoWidth()
+		.Padding(0.0f)
+		[
+			MakeViewportToolbar(InRealViewport)
+		]
 
-			// Playback speed menu
-			+ SHorizontalBox::Slot()
-			.AutoWidth()
-			.Padding( 2.0f, 2.0f )
-			[
-				SNew( SEditorViewportToolbarMenu )
-				.ParentToolBar( SharedThis( this ) )
-				.Label( this, &SAnimViewportToolBar::GetPlaybackMenuLabel )
-				.LabelIcon(FEditorStyle::GetBrush("AnimViewportMenu.PlayBackSpeed"))
-				.OnGetMenuContent( this, &SAnimViewportToolBar::GeneratePlaybackMenu ) 
-			]
-				
- 			+SHorizontalBox::Slot()
-			.Padding(3.0f, 1.0f)
-			.HAlign(HAlign_Right)
-			[
-				SNew(STransformViewportToolBar)
-				.Viewport(InRealViewport)
-				.CommandList(InRealViewport->GetCommandList())
-				.Visibility(this, &SAnimViewportToolBar::GetTransformToolbarVisibility)
-			];
+		// Transform toolbar
+		+SHorizontalBox::Slot()
+		.HAlign(HAlign_Right)
+		.Padding(3.0f, 0.0f)
+		[
+			SNew(STransformViewportToolBar)
+			.Viewport(InRealViewport)
+			.CommandList(CommandList)
+			.Visibility(this, &SAnimViewportToolBar::GetTransformToolbarVisibility)
+		];
 	//@TODO: Need clipping horizontal box: LeftToolbar->AddWrapButton();
 
 	static const FName DefaultForegroundName("DefaultForeground");
@@ -373,7 +210,6 @@ void SAnimViewportToolBar::Construct(const FArguments& InArgs, TSharedPtr<class
 			+ SVerticalBox::Slot()
 			.AutoHeight()
 			[
-
 				LeftToolbar
 			]
 			+SVerticalBox::Slot()
@@ -381,7 +217,7 @@ void SAnimViewportToolBar::Construct(const FArguments& InArgs, TSharedPtr<class
 			[
 				// Display text (e.g., item being previewed)
 				SNew(STextBlock)
-				.Text(ViewportRef, &SAnimationEditorViewportTabBody::GetDisplayString)
+				.Text(InViewport.Get(), &SAnimationEditorViewportTabBody::GetDisplayString)
 				.Font(FEditorStyle::GetFontStyle(TEXT("AnimViewport.MessageFont")))
 				.ShadowOffset(FVector2D(0.5f, 0.5f))
 				.ShadowColorAndOpacity(FLinearColor(0.3f, 0.3f, 0.3f))
@@ -391,6 +227,69 @@ void SAnimViewportToolBar::Construct(const FArguments& InArgs, TSharedPtr<class
 	];
 
 	SViewportToolBar::Construct(SViewportToolBar::FArguments());
+
+	// We assign the viewport pointer her rather that initially, as SViewportToolbar::Construct 
+	// ends up calling through and attempting to perform operations on the not-yet-full-constructed viewport
+	Viewport = InViewport;
+}
+
+TSharedRef<SWidget> SAnimViewportToolBar::MakeViewportToolbar(TSharedPtr<class SEditorViewport> InRealViewport)
+{
+	FMenuBarBuilder MenuBarBuilder(CommandList, Extenders, &FEditorStyle::Get(), false);
+
+	FName MenuBarStyle("ViewportMenu");
+	MenuBarBuilder.SetStyle(&FEditorStyle::Get(), MenuBarStyle);
+
+	MenuBarBuilder.AddPullDownMenu(FText(), LOCTEXT("ViewMenuToolTip", "Show viewport options"), FSlateIcon("EditorStyle", "EditorViewportToolBar.MenuDropdown"), FNewMenuDelegate::CreateSP(this, &SAnimViewportToolBar::GenerateViewMenu), "AnimViewportViewMenu");
+
+	MenuBarBuilder.AddPullDownMenu(
+		TAttribute<FText>::Create(TAttribute<FText>::FGetter::CreateSP(this, &SAnimViewportToolBar::GetCameraMenuLabel)), 
+		LOCTEXT("CameraMenuToolTip", "Viewport camera options"), 
+		TAttribute<FSlateIcon>::Create(TAttribute<FSlateIcon>::FGetter::CreateSP(this, &SAnimViewportToolBar::GetCameraMenuLabelIcon)), 
+		FNewMenuDelegate::CreateSP(this, &SAnimViewportToolBar::GenerateViewportTypeMenu), 
+		"AnimViewportCameraMenu");
+
+	TWeakPtr<SEditorViewport> WeakViewport = InRealViewport;
+	TWeakPtr<SViewportToolBar> WeakToolbar = SharedThis(this);
+
+	MenuBarBuilder.AddPullDownMenu(
+		LOCTEXT("ViewMenu", "View"),
+		LOCTEXT("ViewMenuToolTip", "Change the rendering mode"),
+		TAttribute<FSlateIcon>::Create(TAttribute<FSlateIcon>::FGetter::CreateStatic(&SEditorViewportViewMenu::GetViewMenuLabelIcon, WeakViewport)),
+		FNewMenuDelegate::CreateStatic(&SEditorViewportViewMenu::GenerateViewMenu, WeakToolbar),
+		"AnimViewportViewMenu");
+
+	if(bShowShowMenu)
+	{
+		MenuBarBuilder.AddPullDownMenu(
+			LOCTEXT("ShowMenu", "Show"),
+			LOCTEXT("ShowMenuToolTip", "Various options for rendering"),
+			FSlateIcon(),
+			FNewMenuDelegate::CreateSP(this, &SAnimViewportToolBar::GenerateShowMenu),
+			"AnimViewportShowMenu");
+	}
+
+	if (bShowLODMenu)
+	{
+		MenuBarBuilder.AddPullDownMenu(
+			TAttribute<FText>::Create(TAttribute<FText>::FGetter::CreateSP(this, &SAnimViewportToolBar::GetLODMenuLabel)),
+			LOCTEXT("LODMenuToolTip", "Change the current LOD mode"),
+			FSlateIcon(),
+			FNewMenuDelegate::CreateSP(this, &SAnimViewportToolBar::GenerateLODMenu),
+			"AnimViewportLODMenu");
+	}
+
+	if(bShowPlaySpeedMenu)
+	{
+		MenuBarBuilder.AddPullDownMenu(
+			TAttribute<FText>::Create(TAttribute<FText>::FGetter::CreateSP(this, &SAnimViewportToolBar::GetPlaybackMenuLabel)),
+			LOCTEXT("PlabackSpeedMenuToolTip", "Adjust the animation playback speed"),
+			FSlateIcon("EditorStyle", "AnimViewportMenu.PlayBackSpeed"),
+			FNewMenuDelegate::CreateSP(this, &SAnimViewportToolBar::GeneratePlaybackMenu),
+			"AnimViewportPlaybackSpeedMenu");
+	}
+
+	return MenuBarBuilder.MakeWidget();
 }
 
 EVisibility SAnimViewportToolBar::GetTransformToolbarVisibility() const
@@ -398,18 +297,48 @@ EVisibility SAnimViewportToolBar::GetTransformToolbarVisibility() const
 	return Viewport.Pin()->CanUseGizmos() ? EVisibility::Visible : EVisibility::Hidden;
 }
 
-TSharedRef<SWidget> SAnimViewportToolBar::GenerateViewMenu() const
+void SAnimViewportToolBar::GenerateViewMenu(FMenuBuilder& InMenuBuilder) const
 {
 	const FAnimViewportMenuCommands& Actions = FAnimViewportMenuCommands::Get();
 
-	const bool bInShouldCloseWindowAfterMenuSelection = true;
-	FMenuBuilder ViewMenuBuilder( bInShouldCloseWindowAfterMenuSelection, Viewport.Pin()->GetCommandList() );
+	InMenuBuilder.PushCommandList(CommandList.ToSharedRef());
+	InMenuBuilder.PushExtender(Extenders.ToSharedRef());
+
+	InMenuBuilder.BeginSection("AnimViewportSceneSetup", LOCTEXT("ViewMenu_SceneSetupLabel", "Scene Setup"));
 	{
-		ViewMenuBuilder.BeginSection("AnimViewportSceneSetup", LOCTEXT("ViewMenu_SceneSetupLabel", "Scene Setup"));
-		{
-			ViewMenuBuilder.AddMenuEntry(FAnimViewportMenuCommands::Get().PreviewSceneSettings);
+		InMenuBuilder.AddMenuEntry(FAnimViewportMenuCommands::Get().PreviewSceneSettings);
 
-			ViewMenuBuilder.AddSubMenu(
+		if(bShowFloorOptions)
+		{
+			TSharedPtr<SWidget> FloorOffsetWidget = 
+				SNew(SBox)
+				.HAlign(HAlign_Right)
+				[
+					SNew(SBox)
+					.Padding(FMargin(4.0f, 0.0f, 0.0f, 0.0f))
+					.WidthOverride(100.0f)
+					[
+						SNew(SNumericEntryBox<float>)
+						.Font(FEditorStyle::GetFontStyle(TEXT("MenuItem.Font")))
+						.AllowSpin(true)
+						.MinSliderValue(-100.0f)
+						.MaxSliderValue(100.0f)
+						.Value(this, &SAnimViewportToolBar::OnGetFloorOffset)
+						.OnValueChanged(this, &SAnimViewportToolBar::OnFloorOffsetChanged)
+						.ToolTipText(LOCTEXT("FloorOffsetToolTip", "Height offset for the floor mesh (stored per-mesh)"))
+					]
+				];
+
+			InMenuBuilder.AddWidget(FloorOffsetWidget.ToSharedRef(), LOCTEXT("FloorHeightOffset", "Floor Height Offset"));
+
+			InMenuBuilder.PushCommandList(Viewport.Pin()->GetCommandList().ToSharedRef());
+			InMenuBuilder.AddMenuEntry(FAnimViewportShowCommands::Get().AutoAlignFloorToMesh);
+			InMenuBuilder.PopCommandList();
+		}
+			
+		if (bShowTurnTable)
+		{
+			InMenuBuilder.AddSubMenu(
 				LOCTEXT("TurnTableLabel", "Turn Table"),
 				LOCTEXT("TurnTableTooltip", "Set up auto-rotation of preview."),
 				FNewMenuDelegate::CreateRaw(this, &SAnimViewportToolBar::GenerateTurnTableMenu),
@@ -417,201 +346,178 @@ TSharedRef<SWidget> SAnimViewportToolBar::GenerateViewMenu() const
 				FSlateIcon(FEditorStyle::GetStyleSetName(), "AnimViewportMenu.TurnTableSpeed")
 				);
 		}
-		ViewMenuBuilder.EndSection();
-
-		ViewMenuBuilder.BeginSection("AnimViewportCamera", LOCTEXT("ViewMenu_CameraLabel", "Camera"));
-		{
-			ViewMenuBuilder.AddMenuEntry(FAnimViewportMenuCommands::Get().CameraFollow);
-			ViewMenuBuilder.AddMenuEntry(FEditorViewportCommands::Get().FocusViewportToSelection);
-		}
-		ViewMenuBuilder.EndSection();
-
-		ViewMenuBuilder.BeginSection("AnimViewportDefaultCamera", LOCTEXT("ViewMenu_DefaultCameraLabel", "Default Camera"));
-		{
-			ViewMenuBuilder.AddMenuEntry(FAnimViewportMenuCommands::Get().JumpToDefaultCamera);
-			ViewMenuBuilder.AddMenuEntry(FAnimViewportMenuCommands::Get().SaveCameraAsDefault);
-			ViewMenuBuilder.AddMenuEntry(FAnimViewportMenuCommands::Get().ClearDefaultCamera);
-		}
-		ViewMenuBuilder.EndSection();
 	}
+	InMenuBuilder.EndSection();
 
-	return ViewMenuBuilder.MakeWidget();
-}
-
-TSharedRef<SWidget> SAnimViewportToolBar::GenerateShowMenu() const
-{
-	const FAnimViewportShowCommands& Actions = FAnimViewportShowCommands::Get();
-
-	const bool bInShouldCloseWindowAfterMenuSelection = true;
-	FMenuBuilder ShowMenuBuilder(bInShouldCloseWindowAfterMenuSelection, Viewport.Pin()->GetCommandList());
+	InMenuBuilder.BeginSection("AnimViewportCamera", LOCTEXT("ViewMenu_CameraLabel", "Camera"));
 	{
-			ShowMenuBuilder.BeginSection("AnimViewportFOV", LOCTEXT("Viewport_FOVLabel", "Field Of View"));
-			{
-				const float FOVMin = 5.f;
-				const float FOVMax = 170.f;
+		InMenuBuilder.AddMenuEntry(FAnimViewportMenuCommands::Get().CameraFollow);
+		InMenuBuilder.AddMenuEntry(FEditorViewportCommands::Get().FocusViewportToSelection);
+
+		const float FOVMin = 5.f;
+		const float FOVMax = 170.f;
 
-				TSharedPtr<SWidget> FOVWidget = SNew(SSpinBox<float>)
+		TSharedPtr<SWidget> FOVWidget = 
+			SNew(SBox)
+			.HAlign(HAlign_Right)
+			[
+				SNew(SBox)
+				.Padding(FMargin(4.0f, 0.0f, 0.0f, 0.0f))
+				.WidthOverride(100.0f)
+				[
+					SNew(SSpinBox<float>)
 					.Font(FEditorStyle::GetFontStyle(TEXT("MenuItem.Font")))
 					.MinValue(FOVMin)
 					.MaxValue(FOVMax)
 					.Value(this, &SAnimViewportToolBar::OnGetFOVValue)
 					.OnValueChanged(this, &SAnimViewportToolBar::OnFOVValueChanged)
-					.OnValueCommitted(this, &SAnimViewportToolBar::OnFOVValueCommitted);
+					.OnValueCommitted(this, &SAnimViewportToolBar::OnFOVValueCommitted)
+				]
+			];
 
-				ShowMenuBuilder.AddWidget(FOVWidget.ToSharedRef(), FText());
-			}
-			ShowMenuBuilder.EndSection();
+		InMenuBuilder.AddWidget(FOVWidget.ToSharedRef(), LOCTEXT("Viewport_FOVLabel", "Field Of View"));
+	}
+	InMenuBuilder.EndSection();
 
-			ShowMenuBuilder.BeginSection("AnimViewportAudio", LOCTEXT("Viewport_AudioLabel", "Audio"));
-			{
-				ShowMenuBuilder.AddMenuEntry(Actions.MuteAudio);
-				ShowMenuBuilder.AddMenuEntry(Actions.UseAudioAttenuation);
-			}
-			ShowMenuBuilder.EndSection();
+	InMenuBuilder.BeginSection("AnimViewportDefaultCamera", LOCTEXT("ViewMenu_DefaultCameraLabel", "Default Camera"));
+	{
+		InMenuBuilder.AddMenuEntry(FAnimViewportMenuCommands::Get().JumpToDefaultCamera);
+		InMenuBuilder.AddMenuEntry(FAnimViewportMenuCommands::Get().SaveCameraAsDefault);
+		InMenuBuilder.AddMenuEntry(FAnimViewportMenuCommands::Get().ClearDefaultCamera);
+	}
+	InMenuBuilder.EndSection();
 
-			ShowMenuBuilder.BeginSection("AnimViewportRootMotion", LOCTEXT("Viewport_RootMotionLabel", "Root Motion"));
-			{
-				ShowMenuBuilder.AddMenuEntry(Actions.ProcessRootMotion);
-			}
-			ShowMenuBuilder.EndSection();
+	InMenuBuilder.PopCommandList();
+	InMenuBuilder.PopExtender();
+}
 
-			ShowMenuBuilder.BeginSection("AnimViewportMesh", LOCTEXT("ShowMenu_Actions_Mesh", "Mesh"));
-			{
-				ShowMenuBuilder.AddMenuEntry( Actions.ShowRetargetBasePose );
-				ShowMenuBuilder.AddMenuEntry( Actions.ShowBound );
-				ShowMenuBuilder.AddMenuEntry( Actions.UseInGameBound );
-				ShowMenuBuilder.AddMenuEntry( Actions.ShowPreviewMesh );
-				ShowMenuBuilder.AddMenuEntry( Actions.ShowMorphTargets );
-				ShowMenuBuilder.AddMenuEntry( Actions.ShowVertexColors );
-			}
-			ShowMenuBuilder.EndSection();
+void SAnimViewportToolBar::GenerateShowMenu(FMenuBuilder& InMenuBuilder) const
+{
+	const FAnimViewportShowCommands& Actions = FAnimViewportShowCommands::Get();
 
-			ShowMenuBuilder.BeginSection("AnimViewportAnimation", LOCTEXT("ShowMenu_Actions_Asset", "Asset"));
-			{
-				ShowMenuBuilder.AddMenuEntry( Actions.ShowRawAnimation );
-				ShowMenuBuilder.AddMenuEntry( Actions.ShowNonRetargetedAnimation );
-				ShowMenuBuilder.AddMenuEntry( Actions.ShowAdditiveBaseBones );
-				ShowMenuBuilder.AddMenuEntry( Actions.ShowSourceRawAnimation);
-				ShowMenuBuilder.AddMenuEntry( Actions.ShowBakedAnimation );
-			}
-			ShowMenuBuilder.EndSection();
+	InMenuBuilder.PushCommandList(Viewport.Pin()->GetCommandList().ToSharedRef());
+	InMenuBuilder.PushExtender(Extenders.ToSharedRef());
 
-			ShowMenuBuilder.BeginSection("AnimViewportPreviewBones", LOCTEXT("ShowMenu_Actions_Bones", "Hierarchy") );
-			{
-				ShowMenuBuilder.AddSubMenu(
-					LOCTEXT("AnimViewportBoneDrawSubMenu", "Bone"),
-					LOCTEXT("AnimViewportBoneDrawSubMenuToolTip", "Bone Draw Option"),
-					FNewMenuDelegate::CreateRaw(this, &SAnimViewportToolBar::FillShowBoneDrawMenu)
-					);
-
-				ShowMenuBuilder.AddMenuEntry( Actions.ShowSockets );
-				ShowMenuBuilder.AddMenuEntry( Actions.ShowBoneNames );
-
-				ShowMenuBuilder.AddSubMenu(
-					LOCTEXT("AnimViewportOverlayDrawSubMenu", "Mesh Overlay"),
-					LOCTEXT("AnimViewportOverlayDrawSubMenuToolTip", "Options"),
-					FNewMenuDelegate::CreateRaw(this, &SAnimViewportToolBar::FillShowOverlayDrawMenu)
-				);
-			}
-			ShowMenuBuilder.EndSection();
+	{
+		InMenuBuilder.BeginSection("AnimViewportGeneralShowFlags", LOCTEXT("ViewMenu_GeneralShowFlags", "General Show Flags"));
+		{
+			InMenuBuilder.AddMenuEntry(FAnimViewportShowCommands::Get().ToggleGrid);
+		}
+		InMenuBuilder.EndSection();
 
-			ShowMenuBuilder.AddMenuSeparator();
-			ShowMenuBuilder.AddSubMenu(
-				LOCTEXT("AnimviewportInfo", "Display Info"),
-				LOCTEXT("AnimviewportInfoSubMenuToolTip", "Display Mesh Info in Viewport"),
-				FNewMenuDelegate::CreateRaw(this, &SAnimViewportToolBar::FillShowDisplayInfoMenu));
+		InMenuBuilder.BeginSection("AnimViewportSceneElements", LOCTEXT("ViewMenu_SceneElements", "Scene Elements"));
+		{
+			InMenuBuilder.AddSubMenu(
+				LOCTEXT("AnimViewportMeshSubMenu", "Mesh"),
+				LOCTEXT("AnimViewportMeshSubMenuToolTip", "Mesh-related options"),
+				FNewMenuDelegate::CreateLambda([](FMenuBuilder& SubMenuBuilder)
+				{
+					SubMenuBuilder.BeginSection("AnimViewportMesh", LOCTEXT("ShowMenu_Actions_Mesh", "Mesh"));
+					{
+						SubMenuBuilder.AddMenuEntry(FAnimViewportShowCommands::Get().ShowRetargetBasePose );
+						SubMenuBuilder.AddMenuEntry(FAnimViewportShowCommands::Get().ShowBound );
+						SubMenuBuilder.AddMenuEntry(FAnimViewportShowCommands::Get().UseInGameBound );
+						SubMenuBuilder.AddMenuEntry(FAnimViewportShowCommands::Get().ShowPreviewMesh );
+						SubMenuBuilder.AddMenuEntry(FAnimViewportShowCommands::Get().ShowMorphTargets );
+						SubMenuBuilder.AddMenuEntry(FAnimViewportShowCommands::Get().ShowVertexColors );
+					}
+					SubMenuBuilder.EndSection();
+
+					SubMenuBuilder.BeginSection("AnimViewportMeshInfo", LOCTEXT("ShowMenu_Actions_MeshInfo", "Mesh Info"));
+					{
+						SubMenuBuilder.AddMenuEntry(FAnimViewportShowCommands::Get().ShowDisplayInfoBasic);
+						SubMenuBuilder.AddMenuEntry(FAnimViewportShowCommands::Get().ShowDisplayInfoDetailed);
+						SubMenuBuilder.AddMenuEntry(FAnimViewportShowCommands::Get().ShowDisplayInfoSkelControls);
+						SubMenuBuilder.AddMenuEntry(FAnimViewportShowCommands::Get().HideDisplayInfo);
+					}
+					SubMenuBuilder.EndSection();
+
+					SubMenuBuilder.BeginSection("AnimViewportPreviewOverlayDraw", LOCTEXT("ShowMenu_Actions_Overlay", "Mesh Overlay Drawing"));
+					{
+						SubMenuBuilder.AddMenuEntry(FAnimViewportShowCommands::Get().ShowOverlayNone);
+						SubMenuBuilder.AddMenuEntry(FAnimViewportShowCommands::Get().ShowBoneWeight);
+						SubMenuBuilder.AddMenuEntry(FAnimViewportShowCommands::Get().ShowMorphTargetVerts);
+					}
+					SubMenuBuilder.EndSection();
+				})
+			);
+
+			InMenuBuilder.AddSubMenu(
+				LOCTEXT("AnimViewportAnimationSubMenu", "Animation"),
+				LOCTEXT("AnimViewportAnimationSubMenuToolTip", "Animation-related options"),
+				FNewMenuDelegate::CreateLambda([](FMenuBuilder& SubMenuBuilder)
+				{
+					SubMenuBuilder.BeginSection("AnimViewportRootMotion", LOCTEXT("Viewport_RootMotionLabel", "Root Motion"));
+					{
+						SubMenuBuilder.AddMenuEntry(FAnimViewportShowCommands::Get().ProcessRootMotion);
+					}
+					SubMenuBuilder.EndSection();
+
+					SubMenuBuilder.BeginSection("AnimViewportAnimation", LOCTEXT("ShowMenu_Actions_AnimationAsset", "Animation"));
+					{
+						SubMenuBuilder.AddMenuEntry(FAnimViewportShowCommands::Get().ShowRawAnimation);
+						SubMenuBuilder.AddMenuEntry(FAnimViewportShowCommands::Get().ShowNonRetargetedAnimation);
+						SubMenuBuilder.AddMenuEntry(FAnimViewportShowCommands::Get().ShowAdditiveBaseBones);
+						SubMenuBuilder.AddMenuEntry(FAnimViewportShowCommands::Get().ShowSourceRawAnimation);
+						SubMenuBuilder.AddMenuEntry(FAnimViewportShowCommands::Get().ShowBakedAnimation);
+					}
+					SubMenuBuilder.EndSection();
+				})
+			);
+
+			InMenuBuilder.AddSubMenu(
+				LOCTEXT("AnimViewportBoneDrawSubMenu", "Bones"),
+				LOCTEXT("AnimViewportBoneDrawSubMenuToolTip", "Bone Drawing Options"),
+				FNewMenuDelegate::CreateLambda([](FMenuBuilder& SubMenuBuilder)
+				{
+					SubMenuBuilder.BeginSection("BonesAndSockets", LOCTEXT("Viewport_BonesAndSocketsLabel", "Bones & Sockets"));
+					{
+						SubMenuBuilder.AddMenuEntry(FAnimViewportShowCommands::Get().ShowSockets);
+						SubMenuBuilder.AddMenuEntry(FAnimViewportShowCommands::Get().ShowBoneNames);
+					}
+					SubMenuBuilder.EndSection();
+
+					SubMenuBuilder.BeginSection("AnimViewportPreviewHierarchyBoneDraw", LOCTEXT("ShowMenu_Actions_BoneDrawing", "Bone Drawing"));
+					{
+						SubMenuBuilder.AddMenuEntry(FAnimViewportShowCommands::Get().ShowBoneDrawAll);
+						SubMenuBuilder.AddMenuEntry(FAnimViewportShowCommands::Get().ShowBoneDrawSelected);
+						SubMenuBuilder.AddMenuEntry(FAnimViewportShowCommands::Get().ShowBoneDrawSelectedAndParents);
+						SubMenuBuilder.AddMenuEntry(FAnimViewportShowCommands::Get().ShowBoneDrawNone);
+					}
+					SubMenuBuilder.EndSection();
+				})
+				);
 
 #if WITH_APEX_CLOTHING
 			UDebugSkelMeshComponent* PreviewComp = Viewport.Pin()->GetPreviewScene()->GetPreviewMeshComponent();
 
 			if(PreviewComp)
 			{
-				ShowMenuBuilder.AddMenuSeparator();
-				ShowMenuBuilder.AddSubMenu(
+				InMenuBuilder.AddSubMenu(
 					LOCTEXT("AnimViewportClothingSubMenu", "Clothing"),
 					LOCTEXT("AnimViewportClothingSubMenuToolTip", "Options relating to clothing"),
 					FNewMenuDelegate::CreateRaw(this, &SAnimViewportToolBar::FillShowClothingMenu));
 			}
 #endif // #if WITH_APEX_CLOTHING
+		}
 
-		ShowMenuBuilder.AddMenuSeparator();
-
-		ShowMenuBuilder.AddSubMenu(
-			LOCTEXT("AnimViewportSceneSubMenu", "Scene Setup"),
-			LOCTEXT("AnimViewportSceneSubMenuToolTip", "Options relating to the preview scene"),
-			FNewMenuDelegate::CreateRaw( this, &SAnimViewportToolBar::FillShowSceneMenu ) );
-
-		ShowMenuBuilder.AddSubMenu(
+		InMenuBuilder.AddSubMenu(
 			LOCTEXT("AnimViewportAdvancedSubMenu", "Advanced"),
 			LOCTEXT("AnimViewportAdvancedSubMenuToolTip", "Advanced options"),
-			FNewMenuDelegate::CreateRaw( this, &SAnimViewportToolBar::FillShowAdvancedMenu ) );
-	}
-
-	return ShowMenuBuilder.MakeWidget();
-
-}
-
-void SAnimViewportToolBar::FillShowSceneMenu(FMenuBuilder& MenuBuilder) const
-{
-	const FAnimViewportShowCommands& Actions = FAnimViewportShowCommands::Get();
-
-	MenuBuilder.BeginSection("AnimViewportAccessory", LOCTEXT("Viewport_AccessoryLabel", "Accessory"));
-	{
-		MenuBuilder.AddMenuEntry(Actions.AutoAlignFloorToMesh);
-	}
-	MenuBuilder.EndSection();
-	
-	MenuBuilder.BeginSection("AnimViewportFloorOffset", LOCTEXT("Viewport_FloorOffsetLabel", "Floor Height Offset"));
-	{
-		TSharedPtr<SWidget> FloorOffsetWidget = SNew(SNumericEntryBox<float>)
-			.Font(FEditorStyle::GetFontStyle(TEXT("MenuItem.Font")))
-			.Value(this, &SAnimViewportToolBar::OnGetFloorOffset)
-			.OnValueChanged(this, &SAnimViewportToolBar::OnFloorOffsetChanged)
-			.ToolTipText(LOCTEXT("FloorOffsetToolTip", "Height offset for the floor mesh (stored per-mesh)"));
-
-		MenuBuilder.AddWidget(FloorOffsetWidget.ToSharedRef(), FText());
-	}
-	MenuBuilder.EndSection();
-
-	MenuBuilder.BeginSection("AnimViewportGrid", LOCTEXT("Viewport_GridLabel", "Grid"));
-	{
-		MenuBuilder.AddMenuEntry(Actions.ToggleGrid);
-	}
-	MenuBuilder.EndSection();
-
-	MenuBuilder.BeginSection("AnimViewportBackground", LOCTEXT("Viewport_BackgroundLabel", "Background"));
-	{
-		TSharedPtr<SWidget> BackgroundColorWidget = SNew(SBackgroundColorSettings).AnimEditorViewport(Viewport);
-		MenuBuilder.AddWidget(BackgroundColorWidget.ToSharedRef(), FText());
-	}
-	MenuBuilder.EndSection();
-}
+			FNewMenuDelegate::CreateRaw(this, &SAnimViewportToolBar::FillShowAdvancedMenu));
 
-void SAnimViewportToolBar::FillShowBoneDrawMenu(FMenuBuilder& MenuBuilder) const
-{
-	const FAnimViewportShowCommands& Actions = FAnimViewportShowCommands::Get();
+		InMenuBuilder.EndSection();
 
-	MenuBuilder.BeginSection("AnimViewportPreviewHierarchyBoneDraw", LOCTEXT("ShowMenu_Actions_BoneDrawing", "Bone Drawing"));
-	{
-		MenuBuilder.AddMenuEntry(Actions.ShowBoneDrawAll);
-		MenuBuilder.AddMenuEntry(Actions.ShowBoneDrawSelected);
-		MenuBuilder.AddMenuEntry(Actions.ShowBoneDrawSelectedAndParents);
-		MenuBuilder.AddMenuEntry(Actions.ShowBoneDrawNone);
+		InMenuBuilder.BeginSection("AnimViewportOtherFlags", LOCTEXT("ViewMenu_OtherFlags", "Other Flags"));
+		{
+			InMenuBuilder.AddMenuEntry(Actions.MuteAudio);
+			InMenuBuilder.AddMenuEntry(Actions.UseAudioAttenuation);
+		}
+		InMenuBuilder.EndSection();
 	}
-	MenuBuilder.EndSection();
-}
 
-void SAnimViewportToolBar::FillShowOverlayDrawMenu(FMenuBuilder& MenuBuilder) const
-{
-	const FAnimViewportShowCommands& Actions = FAnimViewportShowCommands::Get();
-
-	MenuBuilder.BeginSection("AnimViewportPreviewOverlayDraw", LOCTEXT("ShowMenu_Actions_Overlay", "Overlay Draw"));
-	{
-		MenuBuilder.AddMenuEntry(Actions.ShowOverlayNone);
-		MenuBuilder.AddMenuEntry(Actions.ShowBoneWeight);
-		MenuBuilder.AddMenuEntry(Actions.ShowMorphTargetVerts);
-	}
-	MenuBuilder.EndSection();
+	InMenuBuilder.PopCommandList();
+	InMenuBuilder.PopExtender();
 }
 
 void SAnimViewportToolBar::FillShowAdvancedMenu(FMenuBuilder& MenuBuilder) const
@@ -696,19 +602,6 @@ void SAnimViewportToolBar::FillShowClothingMenu(FMenuBuilder& MenuBuilder)
 #endif // #if WITH_APEX_CLOTHING
 }
 
-void SAnimViewportToolBar::FillShowDisplayInfoMenu(FMenuBuilder& MenuBuilder) const
-{
-	const FAnimViewportShowCommands& Actions = FAnimViewportShowCommands::Get();
-
-	// display info levels
-	{
-		MenuBuilder.AddMenuEntry(Actions.ShowDisplayInfoBasic);
-		MenuBuilder.AddMenuEntry(Actions.ShowDisplayInfoDetailed);
-		MenuBuilder.AddMenuEntry(Actions.ShowDisplayInfoSkelControls);
-		MenuBuilder.AddMenuEntry(Actions.HideDisplayInfo);
-	}
-}
-
 FText SAnimViewportToolBar::GetLODMenuLabel() const
 {
 	FText Label = LOCTEXT("LODMenu_AutoLabel", "LOD Auto");
@@ -725,18 +618,18 @@ FText SAnimViewportToolBar::GetLODMenuLabel() const
 	return Label;
 }
 
-TSharedRef<SWidget> SAnimViewportToolBar::GenerateLODMenu() const
+void SAnimViewportToolBar::GenerateLODMenu(FMenuBuilder& InMenuBuilder) const
 {
 	const FAnimViewportLODCommands& Actions = FAnimViewportLODCommands::Get();
 
-	const bool bInShouldCloseWindowAfterMenuSelection = true;
-	FMenuBuilder ShowMenuBuilder( bInShouldCloseWindowAfterMenuSelection, Viewport.Pin()->GetCommandList() );
+	InMenuBuilder.PushCommandList(Viewport.Pin()->GetCommandList().ToSharedRef());
+	InMenuBuilder.PushExtender(Extenders.ToSharedRef());
 	{
 		// LOD Models
-		ShowMenuBuilder.BeginSection("AnimViewportPreviewLODs", LOCTEXT("ShowLOD_PreviewLabel", "Preview LODs") );
+		InMenuBuilder.BeginSection("AnimViewportPreviewLODs", LOCTEXT("ShowLOD_PreviewLabel", "Preview LODs") );
 		{
-			ShowMenuBuilder.AddMenuEntry( Actions.LODAuto );
-			ShowMenuBuilder.AddMenuEntry( Actions.LOD0 );
+			InMenuBuilder.AddMenuEntry( Actions.LODAuto );
+			InMenuBuilder.AddMenuEntry( Actions.LOD0 );
 
 			int32 LODCount = Viewport.Pin()->GetLODModelCount();
 			for (int32 LODId = 1; LODId < LODCount; ++LODId)
@@ -747,58 +640,59 @@ TSharedRef<SWidget> SAnimViewportToolBar::GenerateLODMenu() const
 					FCanExecuteAction(),
 					FIsActionChecked::CreateSP(Viewport.Pin().ToSharedRef(), &SAnimationEditorViewportTabBody::IsLODModelSelected, LODId + 1));
 
-				ShowMenuBuilder.AddMenuEntry(FText::FromString(TitleLabel), FText::GetEmpty(), FSlateIcon(), Action, NAME_None, EUserInterfaceActionType::RadioButton);
+				InMenuBuilder.AddMenuEntry(FText::FromString(TitleLabel), FText::GetEmpty(), FSlateIcon(), Action, NAME_None, EUserInterfaceActionType::RadioButton);
 			}
 		}
-		ShowMenuBuilder.EndSection();
+		InMenuBuilder.EndSection();
 	}
 
-	return ShowMenuBuilder.MakeWidget();
+	InMenuBuilder.PopCommandList();
+	InMenuBuilder.PopExtender();
 }
 
-TSharedRef<SWidget> SAnimViewportToolBar::GenerateViewportTypeMenu() const
+void SAnimViewportToolBar::GenerateViewportTypeMenu(FMenuBuilder& InMenuBuilder) const
 {
-	const bool bInShouldCloseWindowAfterMenuSelection = true;
-	FMenuBuilder CameraMenuBuilder(bInShouldCloseWindowAfterMenuSelection, Viewport.Pin()->GetViewportWidget()->GetCommandList());
+	InMenuBuilder.SetStyle(&FEditorStyle::Get(), "Menu");
+	InMenuBuilder.PushCommandList(Viewport.Pin()->GetViewportWidget()->GetCommandList().ToSharedRef());
+	InMenuBuilder.PushExtender(Extenders.ToSharedRef());
 
 	// Camera types
-	CameraMenuBuilder.AddMenuEntry(FEditorViewportCommands::Get().Perspective);
-
-	CameraMenuBuilder.BeginSection("LevelViewportCameraType_Ortho", LOCTEXT("CameraTypeHeader_Ortho", "Orthographic"));
-	CameraMenuBuilder.AddMenuEntry(FEditorViewportCommands::Get().Top);
-	CameraMenuBuilder.AddMenuEntry(FEditorViewportCommands::Get().Bottom);
-	CameraMenuBuilder.AddMenuEntry(FEditorViewportCommands::Get().Left);
-	CameraMenuBuilder.AddMenuEntry(FEditorViewportCommands::Get().Right);
-	CameraMenuBuilder.AddMenuEntry(FEditorViewportCommands::Get().Front);
-	CameraMenuBuilder.AddMenuEntry(FEditorViewportCommands::Get().Back);
-	CameraMenuBuilder.EndSection();
-
-	return CameraMenuBuilder.MakeWidget();
+	InMenuBuilder.AddMenuEntry(FEditorViewportCommands::Get().Perspective);
+
+	InMenuBuilder.BeginSection("LevelViewportCameraType_Ortho", LOCTEXT("CameraTypeHeader_Ortho", "Orthographic"));
+	InMenuBuilder.AddMenuEntry(FEditorViewportCommands::Get().Top);
+	InMenuBuilder.AddMenuEntry(FEditorViewportCommands::Get().Bottom);
+	InMenuBuilder.AddMenuEntry(FEditorViewportCommands::Get().Left);
+	InMenuBuilder.AddMenuEntry(FEditorViewportCommands::Get().Right);
+	InMenuBuilder.AddMenuEntry(FEditorViewportCommands::Get().Front);
+	InMenuBuilder.AddMenuEntry(FEditorViewportCommands::Get().Back);
+	InMenuBuilder.EndSection();
+
+	InMenuBuilder.PopCommandList();
+	InMenuBuilder.PopExtender();
 }
 
-TSharedRef<SWidget> SAnimViewportToolBar::GeneratePlaybackMenu() const
+void SAnimViewportToolBar::GeneratePlaybackMenu(FMenuBuilder& InMenuBuilder) const
 {
 	const FAnimViewportPlaybackCommands& Actions = FAnimViewportPlaybackCommands::Get();
 
-	const bool bInShouldCloseWindowAfterMenuSelection = true;
-
-	FMenuBuilder PlaybackMenuBuilder( bInShouldCloseWindowAfterMenuSelection, Viewport.Pin()->GetCommandList() );
+	InMenuBuilder.PushCommandList(Viewport.Pin()->GetCommandList().ToSharedRef());
+	InMenuBuilder.PushExtender(Extenders.ToSharedRef());
 	{
 		// View modes
 		{
-			PlaybackMenuBuilder.BeginSection("AnimViewportPlaybackSpeed", LOCTEXT("PlaybackMenu_SpeedLabel", "Playback Speed") );
+			InMenuBuilder.BeginSection("AnimViewportPlaybackSpeed", LOCTEXT("PlaybackMenu_SpeedLabel", "Playback Speed") );
 			{
-				for(int i = 0; i < EAnimationPlaybackSpeeds::NumPlaybackSpeeds; ++i)
+				for(int32 PlaybackSpeedIndex = 0; PlaybackSpeedIndex < EAnimationPlaybackSpeeds::NumPlaybackSpeeds; ++PlaybackSpeedIndex)
 				{
-					PlaybackMenuBuilder.AddMenuEntry( Actions.PlaybackSpeedCommands[i] );
+					InMenuBuilder.AddMenuEntry( Actions.PlaybackSpeedCommands[PlaybackSpeedIndex] );
 				}
 			}
-			PlaybackMenuBuilder.EndSection();
+			InMenuBuilder.EndSection();
 		}
 	}
-
-	return PlaybackMenuBuilder.MakeWidget();
-
+	InMenuBuilder.PopCommandList();
+	InMenuBuilder.PopExtender();
 }
 
 void SAnimViewportToolBar::GenerateTurnTableMenu(FMenuBuilder& MenuBuilder) const
@@ -840,10 +734,10 @@ FSlateColor SAnimViewportToolBar::GetFontColor() const
 	}
 	else
 	{
-		FLinearColor BackgroundColorInHSV = Viewport.Pin()->GetViewportBackgroundColor().LinearRGBToHSV();
+		FLinearColor Color = Settings->Profiles[ProfileIndex].EnvironmentColor * Settings->Profiles[ProfileIndex].EnvironmentIntensity;
 
 		// see if it's dark, if V is less than 0.2
-		if ( BackgroundColorInHSV.B < 0.3f )
+		if (Color.B < 0.3f )
 		{
 			FontColor = FLinearColor::White;
 		}
@@ -919,7 +813,7 @@ FText SAnimViewportToolBar::GetCameraMenuLabel() const
 	return Label;
 }
 
-const FSlateBrush* SAnimViewportToolBar::GetCameraMenuLabelIcon() const
+FSlateIcon SAnimViewportToolBar::GetCameraMenuLabelIcon() const
 {
 	FName Icon = NAME_None;
 	TSharedPtr< SAnimationEditorViewportTabBody > PinnedViewport(Viewport.Pin());
@@ -959,7 +853,7 @@ const FSlateBrush* SAnimViewportToolBar::GetCameraMenuLabelIcon() const
 		}
 	}
 
-	return FEditorStyle::GetBrush(Icon);
+	return FSlateIcon("EditorStyle", Icon);
 }
 
 float SAnimViewportToolBar::OnGetFOVValue( ) const
diff --git a/Engine/Source/Editor/Persona/Private/SAnimViewportToolBar.h b/Engine/Source/Editor/Persona/Private/SAnimViewportToolBar.h
--- a/Engine/Source/Editor/Persona/Private/SAnimViewportToolBar.h
+++ b/Engine/Source/Editor/Persona/Private/SAnimViewportToolBar.h
@@ -19,7 +19,24 @@ class FMenuBuilder;
 class SAnimViewportToolBar : public SViewportToolBar
 {
 public:
-	SLATE_BEGIN_ARGS( SAnimViewportToolBar ){}
+	SLATE_BEGIN_ARGS( SAnimViewportToolBar )
+		: _ShowShowMenu(true)
+		, _ShowLODMenu(true)
+		, _ShowPlaySpeedMenu(true)
+	{}
+
+	SLATE_ARGUMENT(TSharedPtr<FExtender>, Extenders)
+
+	SLATE_ARGUMENT(bool, ShowShowMenu)
+
+	SLATE_ARGUMENT(bool, ShowLODMenu)
+
+	SLATE_ARGUMENT(bool, ShowPlaySpeedMenu)
+
+	SLATE_ARGUMENT(bool, ShowFloorOptions)
+
+	SLATE_ARGUMENT(bool, ShowTurnTable)
+
 	SLATE_END_ARGS()
 
 	void Construct(const FArguments& InArgs, TSharedPtr<class SAnimationEditorViewportTabBody> InViewport, TSharedPtr<class SEditorViewport> InRealViewport);
@@ -27,17 +44,13 @@ public:
 private:
 	/**
 	 * Generates the toolbar view menu content 
-	 *
-	 * @return The widget containing the view menu content
 	 */
-	TSharedRef<SWidget> GenerateViewMenu() const;
+	void GenerateViewMenu(FMenuBuilder& InMenuBuilder) const;
 
 	/**
 	 * Generates the toolbar show menu content 
-	 *
-	 * @return The widget containing the show menu content
 	 */
-	TSharedRef<SWidget> GenerateShowMenu() const;
+	void GenerateShowMenu(FMenuBuilder& InMenuBuilder) const;
 
 	/**
 	 * Generates the Show -> Scene sub menu content
@@ -49,32 +62,15 @@ private:
 	 */
 	void FillShowAdvancedMenu(FMenuBuilder& MenuBuilder) const;
 
-	/**
-	* Generates the Show -> Bone sub menu content
-	*/
-	void FillShowBoneDrawMenu(FMenuBuilder& MenuBuilder) const;
-
-	/**
-	* Generates the Show -> Overlay sub menu content
-	*/
-	void FillShowOverlayDrawMenu(FMenuBuilder& MenuBuilder) const;
-
 	/**
 	* Generates the Show -> Clothing sub menu content
 	*/
 	void FillShowClothingMenu(FMenuBuilder& MenuBuilder);
 
-	/**
-	* Generates the Show -> Display Info sub menu content
-	*/
-	void FillShowDisplayInfoMenu(FMenuBuilder& MenuBuilder) const;
-
 	/**
 	 * Generates the toolbar LOD menu content 
-	 *
-	 * @return The widget containing the LOD menu content based on LOD model count
 	 */
-	TSharedRef<SWidget> GenerateLODMenu() const;
+	void GenerateLODMenu(FMenuBuilder& InMenuBuilder) const;
 
 	/**
 	 * Returns the label for the "LOD" tool bar menu, which changes depending on the current LOD selection
@@ -85,17 +81,13 @@ private:
 
 	/**
 	 * Generates the toolbar viewport type menu content 
-	 *
-	 * @return The widget containing the viewport type menu content
 	 */
-	TSharedRef<SWidget> GenerateViewportTypeMenu() const;
+	void GenerateViewportTypeMenu(FMenuBuilder& InMenuBuilder) const;
 
 	/**
 	 * Generates the toolbar playback menu content 
-	 *
-	 * @return The widget containing the playback menu content
 	 */
-	TSharedRef<SWidget> GeneratePlaybackMenu() const;
+	void GeneratePlaybackMenu(FMenuBuilder& InMenuBuilder) const;
 
 	/** Generate the turntable menu entries */
 	void GenerateTurnTableMenu(FMenuBuilder& MenuBuilder) const;
@@ -118,7 +110,7 @@ private:
 	 * @return	Label to use for this Menu
 	 */
 	FText GetCameraMenuLabel() const;
-	const FSlateBrush* GetCameraMenuLabelIcon() const;
+	FSlateIcon GetCameraMenuLabelIcon() const;
 
 	/** Called by the FOV slider in the perspective viewport to get the FOV value */
 	float OnGetFOVValue() const;
@@ -135,7 +127,32 @@ private:
 	// Called to determine if the gizmos can be used in the current preview
 	EVisibility GetTransformToolbarVisibility() const;
 
+private:
+	/** Build the main toolbar */
+	TSharedRef<SWidget> MakeViewportToolbar(TSharedPtr<class SEditorViewport> InRealViewport);
+
 private:
 	/** The viewport that we are in */
 	TWeakPtr<class SAnimationEditorViewportTabBody> Viewport;
+
+	/** Command list */
+	TSharedPtr<FUICommandList> CommandList;
+
+	/** Extenders */
+	TSharedPtr<FExtender> Extenders;
+
+	/** Whether to show the 'Show' menu */
+	bool bShowShowMenu;
+
+	/** Whether to show the 'LOD' menu */
+	bool bShowLODMenu;
+
+	/** Whether to show the 'Play Speed' menu */
+	bool bShowPlaySpeedMenu;
+
+	/** Whether to show options relating to floor height */
+	bool bShowFloorOptions;
+
+	/** Whether to show options relating to turntable */
+	bool bShowTurnTable;
 };
diff --git a/Engine/Source/Editor/Persona/Private/SAnimationBlendSpaceGridWidget.cpp b/Engine/Source/Editor/Persona/Private/SAnimationBlendSpaceGridWidget.cpp
--- a/Engine/Source/Editor/Persona/Private/SAnimationBlendSpaceGridWidget.cpp
+++ b/Engine/Source/Editor/Persona/Private/SAnimationBlendSpaceGridWidget.cpp
@@ -92,6 +92,8 @@ void SBlendSpaceGridWidget::Construct(const FArguments& InArgs)
 	TextMargin = 16.0f;
 	GridMargin = FMargin(MaxVerticalAxisTextWidth + (TextMargin * 2.0f), TextMargin, (HorizontalAxisMaxTextWidth * 0.5f) + TextMargin, MaxHorizontalAxisTextHeight + (TextMargin * 2.0f));
 
+	bPreviewToolTipHidden = false;
+
 	const bool bShowInputBoxLabel = true;
 	// Widget construction
 	this->ChildSlot
@@ -115,92 +117,119 @@ void SBlendSpaceGridWidget::Construct(const FArguments& InArgs)
 					.DesiredSizeScale(FVector2D(1.0f, 1.0f))
 					.Padding_Lambda([&]() { return FMargin(GridMargin.Left + 6, GridMargin.Top + 6, 0, 0) + GridRatioMargin; })
 					[
-						SNew(SHorizontalBox)
-						+ SHorizontalBox::Slot()
-						.AutoWidth()
+						SNew(SVerticalBox)
+						+ SVerticalBox::Slot()
+						.AutoHeight()
 						[
-							SNew(SBorder)
-							.BorderImage(FEditorStyle::GetBrush("NoBorder"))
-							.Visibility(TAttribute<EVisibility>::Create(TAttribute<EVisibility>::FGetter::CreateSP(this, &SBlendSpaceGridWidget::GetTriangulationButtonVisibility)))		
-							.VAlign(VAlign_Center)
+							SNew(SHorizontalBox)
+							+ SHorizontalBox::Slot()
+							.AutoWidth()
 							[
-								SNew(SButton)
-								.ToolTipText(LOCTEXT("ShowTriangulation", "Show Triangulation"))
-								.OnClicked(this, &SBlendSpaceGridWidget::ToggleTriangulationVisibility)
-								.ButtonColorAndOpacity_Lambda([this]() -> FLinearColor { return bShowTriangulation ? FEditorStyle::GetSlateColor("SelectionColor").GetSpecifiedColor() : FLinearColor::White; })
-								.ContentPadding(1)
+								SNew(SBorder)
+								.BorderImage(FEditorStyle::GetBrush("NoBorder"))
+								.Visibility(TAttribute<EVisibility>::Create(TAttribute<EVisibility>::FGetter::CreateSP(this, &SBlendSpaceGridWidget::GetTriangulationButtonVisibility)))		
+								.VAlign(VAlign_Center)
 								[
-									SNew(SImage)
-									.Image(FEditorStyle::GetBrush("BlendSpaceEditor.ToggleTriangulation"))
-									.ColorAndOpacity(FSlateColor::UseForeground())
+									SNew(SButton)
+									.ToolTipText(LOCTEXT("ShowTriangulation", "Show Triangulation"))
+									.OnClicked(this, &SBlendSpaceGridWidget::ToggleTriangulationVisibility)
+									.ButtonColorAndOpacity_Lambda([this]() -> FLinearColor { return bShowTriangulation ? FEditorStyle::GetSlateColor("SelectionColor").GetSpecifiedColor() : FLinearColor::White; })
+									.ContentPadding(1)
+									[
+										SNew(SImage)
+										.Image(FEditorStyle::GetBrush("BlendSpaceEditor.ToggleTriangulation"))
+										.ColorAndOpacity(FSlateColor::UseForeground())										
+									]
 								]
 							]
-						]
 	
-						+ SHorizontalBox::Slot()
-						.AutoWidth()
-						[
-							SNew(SBorder)
-							.BorderImage(FEditorStyle::GetBrush("NoBorder"))
-							.VAlign(VAlign_Center)
+							+ SHorizontalBox::Slot()
+							.AutoWidth()
 							[
-								SNew(SButton)
-								.ToolTipText(LOCTEXT("ShowAnimationNames", "Show Animation Names"))
-								.OnClicked(this, &SBlendSpaceGridWidget::ToggleShowAnimationNames)
-								.ButtonColorAndOpacity_Lambda([this]() -> FLinearColor { return bShowAnimationNames ? FEditorStyle::GetSlateColor("SelectionColor").GetSpecifiedColor() : FLinearColor::White; })
-								.ContentPadding(1)
+								SNew(SBorder)
+								.BorderImage(FEditorStyle::GetBrush("NoBorder"))
+								.VAlign(VAlign_Center)
 								[
-									SNew(SImage)
-									.Image(FEditorStyle::GetBrush("BlendSpaceEditor.ToggleLabels"))
-									.ColorAndOpacity(FSlateColor::UseForeground())
+									SNew(SButton)
+									.ToolTipText(LOCTEXT("ShowAnimationNames", "Show Animation Names"))
+									.OnClicked(this, &SBlendSpaceGridWidget::ToggleShowAnimationNames)
+									.ButtonColorAndOpacity_Lambda([this]() -> FLinearColor { return bShowAnimationNames ? FEditorStyle::GetSlateColor("SelectionColor").GetSpecifiedColor() : FLinearColor::White; })
+									.ContentPadding(1)
+									[
+										SNew(SImage)
+										.Image(FEditorStyle::GetBrush("BlendSpaceEditor.ToggleLabels"))
+										.ColorAndOpacity(FSlateColor::UseForeground())
+									]
 								]
 							]
-						]
 
-						+ SHorizontalBox::Slot()
-						.AutoWidth()
-						[
-							SNew(SBorder)
-							.BorderImage(FEditorStyle::GetBrush("NoBorder"))
-							.Visibility(TAttribute<EVisibility>::Create(TAttribute<EVisibility>::FGetter::CreateSP(this, &SBlendSpaceGridWidget::GetFittingButtonVisibility)))
-							.VAlign(VAlign_Center)
+							+ SHorizontalBox::Slot()
+							.AutoWidth()
 							[
-								SNew(SButton)
-								.ToolTipText(this, &SBlendSpaceGridWidget::GetFittingTypeButtonToolTipText)
-								.OnClicked(this, &SBlendSpaceGridWidget::ToggleFittingType)
-								.ContentPadding(1)
-								.ButtonColorAndOpacity_Lambda([this]() -> FLinearColor { return bStretchToFit ? FEditorStyle::GetSlateColor("SelectionColor").GetSpecifiedColor() : FLinearColor::White; })
+								SNew(SBorder)
+								.BorderImage(FEditorStyle::GetBrush("NoBorder"))
+								.Visibility(TAttribute<EVisibility>::Create(TAttribute<EVisibility>::FGetter::CreateSP(this, &SBlendSpaceGridWidget::GetFittingButtonVisibility)))
+								.VAlign(VAlign_Center)
 								[
-									SNew(SImage)
-									.Image(FEditorStyle::GetBrush("WidgetDesigner.ZoomToFit"))
-									.ColorAndOpacity(FSlateColor::UseForeground())
+									SNew(SButton)
+									.ToolTipText(this, &SBlendSpaceGridWidget::GetFittingTypeButtonToolTipText)
+									.OnClicked(this, &SBlendSpaceGridWidget::ToggleFittingType)
+									.ContentPadding(1)
+									.ButtonColorAndOpacity_Lambda([this]() -> FLinearColor { return bStretchToFit ? FEditorStyle::GetSlateColor("SelectionColor").GetSpecifiedColor() : FLinearColor::White; })
+									[
+										SNew(SImage)
+										.Image(FEditorStyle::GetBrush("WidgetDesigner.ZoomToFit"))
+										.ColorAndOpacity(FSlateColor::UseForeground())
+									]
 								]
 							]
-						]
 
-						+ SHorizontalBox::Slot()
-						.AutoWidth()
-						[
-							SNew(SBorder)
-							.BorderImage(FEditorStyle::GetBrush("NoBorder"))
-							.Visibility(TAttribute<EVisibility>::Create(TAttribute<EVisibility>::FGetter::CreateSP(this, &SBlendSpaceGridWidget::GetInputBoxVisibility, 0)))
-							.VAlign(VAlign_Center)
+							+ SHorizontalBox::Slot()
+							.AutoWidth()
 							[
-								CreateGridEntryBox(0, bShowInputBoxLabel).ToSharedRef()
+								SNew(SBorder)
+								.BorderImage(FEditorStyle::GetBrush("NoBorder"))
+								.Visibility(TAttribute<EVisibility>::Create(TAttribute<EVisibility>::FGetter::CreateSP(this, &SBlendSpaceGridWidget::GetInputBoxVisibility, 0)))
+								.VAlign(VAlign_Center)
+								[
+									CreateGridEntryBox(0, bShowInputBoxLabel).ToSharedRef()
+								]
 							]
-						]
 	
-						+ SHorizontalBox::Slot()
-						.AutoWidth()
-						[
-							SNew(SBorder)
-							.BorderImage(FEditorStyle::GetBrush("NoBorder"))
-							.Visibility(TAttribute<EVisibility>::Create(TAttribute<EVisibility>::FGetter::CreateSP(this, &SBlendSpaceGridWidget::GetInputBoxVisibility, 1)))
-							.VAlign(VAlign_Center)
+							+ SHorizontalBox::Slot()
+							.AutoWidth()
 							[
-								CreateGridEntryBox(1, bShowInputBoxLabel).ToSharedRef()
+								SNew(SBorder)
+								.BorderImage(FEditorStyle::GetBrush("NoBorder"))
+								.Visibility(TAttribute<EVisibility>::Create(TAttribute<EVisibility>::FGetter::CreateSP(this, &SBlendSpaceGridWidget::GetInputBoxVisibility, 1)))
+								.VAlign(VAlign_Center)
+								[
+									CreateGridEntryBox(1, bShowInputBoxLabel).ToSharedRef()
+								]
 							]
 						]
+						
+						+ SVerticalBox::Slot()
+						.AutoHeight()
+						.Padding(FMargin(2.0f, 3.0f, 0.0f, 0.0f ))
+						[
+							SNew(STextBlock)
+							.Text(LOCTEXT("BlendSpaceSamplesToolTip", "Drag and Drop Animations from the Asset Browser to place Sample Points"))
+							.Font(FEditorStyle::GetFontStyle(TEXT("AnimViewport.MessageFont")))
+							.ColorAndOpacity(FLinearColor(1.0f, 1.0f, 1.0f, 0.7f))
+							.Visibility(TAttribute<EVisibility>::Create(TAttribute<EVisibility>::FGetter::CreateSP(this, &SBlendSpaceGridWidget::GetSampleToolTipVisibility)))
+						]
+
+						+ SVerticalBox::Slot()
+						.AutoHeight()
+						.Padding(FMargin(2.0f, 3.0f, 0.0f, 0.0f))
+						[
+							SNew(STextBlock)
+							.Text(LOCTEXT("BlendspacePreviewToolTip", "Click and Drag or Hold Shift to move the Preview (Green) Pin" ))
+							.Font(FEditorStyle::GetFontStyle(TEXT("AnimViewport.MessageFont")))
+							.ColorAndOpacity(FLinearColor(1.0f, 1.0f, 1.0f, 0.7f))
+							.Visibility(TAttribute<EVisibility>::Create(TAttribute<EVisibility>::FGetter::CreateSP(this, &SBlendSpaceGridWidget::GetPreviewToolTipVisibility)))
+						]
 					]
 				]
 			]
@@ -454,11 +483,12 @@ void SBlendSpaceGridWidget::PaintAnimationNames(const FGeometry& AllottedGeometr
 	const TSharedRef< FSlateFontMeasure > FontMeasure = FSlateApplication::Get().GetRenderer()->GetFontMeasureService();
 	const FVector2D GridCenter = CachedGridRectangle.GetCenter();
 	const TArray<FBlendSample>& Samples = BlendSpace->GetBlendSamples();
-	for (const FBlendSample& Sample : Samples)
+	for (int32 SampleIndex = 0; SampleIndex < Samples.Num(); ++SampleIndex)
 	{
+		const FBlendSample& Sample = Samples[SampleIndex];
 		if (Sample.Animation)
 		{
-			const FString Name = Sample.Animation->GetName();
+			const FString Name = Sample.Animation->GetName() + FString::Printf(TEXT(" (%i)"), SampleIndex);
 			const FVector2D TextSize = FontMeasure->Measure(Name, FontInfo);
 
 			FVector2D GridPosition = SampleValueToGridPosition(Sample.SampleValue);
@@ -563,6 +593,12 @@ FReply SBlendSpaceGridWidget::OnMouseButtonUp(const FGeometry& MyGeometry, const
 			// Process drag ending			
 			ResetToolTip();
 		}
+		else if (DragState == EDragState::DragPreview)
+		{
+			// Process preview drag ending			
+			ResetToolTip();
+			StopPreviewing();
+		}
 
 		// Reset drag state and index
 		DragState = EDragState::None;
@@ -583,7 +619,18 @@ FReply SBlendSpaceGridWidget::OnMouseButtonDown(const FGeometry& MyGeometry, con
 	if (MouseEvent.GetEffectingButton() == EKeys::LeftMouseButton)
 	{
 		// If we are over a sample, make it our currently (dragged) sample
-		if (HighlightedSampleIndex != INDEX_NONE)
+		if (bHighlightPreviewPin)
+		{
+			ResetToolTip();
+			DragState = EDragState::DragPreview;
+			MouseDownPosition = LocalMousePosition;
+			StartPreviewing();
+			ShowToolTip();
+
+			// Start mouse capture
+			return FReply::Handled().CaptureMouse(SharedThis(this));
+		}
+		else if (HighlightedSampleIndex != INDEX_NONE)
 		{
 			DraggedSampleIndex = SelectedSampleIndex = HighlightedSampleIndex;
 			HighlightedSampleIndex = INDEX_NONE;
@@ -720,7 +767,7 @@ FReply SBlendSpaceGridWidget::OnKeyDown(const FGeometry& MyGeometry, const FKeyE
 	// Start previewing when either one of the shift keys is pressed
 	if (bMouseIsOverGeometry)
 	{
-		if ((InKeyEvent.GetKey() == EKeys::LeftShift) || (InKeyEvent.GetKey() == EKeys::RightShift))
+		if (DragState != EDragState::DragPreview && ((InKeyEvent.GetKey() == EKeys::LeftShift) || (InKeyEvent.GetKey() == EKeys::RightShift)))
 		{
 			StartPreviewing();
 			DragState = EDragState::Preview;
@@ -743,7 +790,7 @@ FReply SBlendSpaceGridWidget::OnKeyDown(const FGeometry& MyGeometry, const FKeyE
 FReply SBlendSpaceGridWidget::OnKeyUp(const FGeometry& MyGeometry, const FKeyEvent& InKeyEvent)
 {	
 	// Stop previewing when shift keys are released 
-	if ((InKeyEvent.GetKey() == EKeys::LeftShift) || (InKeyEvent.GetKey() == EKeys::RightShift))
+	if (DragState != EDragState::DragPreview && ((InKeyEvent.GetKey() == EKeys::LeftShift) || (InKeyEvent.GetKey() == EKeys::RightShift)))
 	{
 		StopPreviewing();
 		DragState = EDragState::None;
@@ -947,6 +994,7 @@ void SBlendSpaceGridWidget::StartPreviewing()
 	LastPreviewingMousePosition = LocalMousePosition;
 	LastPreviewingSampleValue = GridPositionToSampleValue(LastPreviewingMousePosition);
 	bPreviewPositionSet = true;	
+	bPreviewToolTipHidden = true;
 }
 
 void SBlendSpaceGridWidget::StopPreviewing()
@@ -965,7 +1013,7 @@ FText SBlendSpaceGridWidget::GetToolTipAnimationName() const
 		{		
 			if (bHighlightPreviewPin)
 			{
-				const FText PreviewPinText = LOCTEXT("HighlightPreviewPinTooltipText", "Preview Value (Hold Shift to Move)");
+				const FText PreviewPinText = LOCTEXT("HighlightPreviewPinTooltipText", "Preview Value (Click and Drag or Hold Shift to Move)");
 				ToolTipText = PreviewPinText;
 			}
 			else if (HighlightedSampleIndex != INDEX_NONE && BlendSpace->IsValidBlendSampleIndex(HighlightedSampleIndex))
@@ -1013,6 +1061,7 @@ FText SBlendSpaceGridWidget::GetToolTipAnimationName() const
 		
 		// If we are previewing return a descriptive label
 		case EDragState::Preview:
+		case EDragState::DragPreview:
 		{
 			const FText AdvancedPreviewText = LOCTEXT("AdvancedPreviewValueTooltip", "Preview Value");
 			const FText BasicPreviewText = LOCTEXT("BasicPreviewValueTooltip", "Preview Value (Hold Ctrl for Weight Details)");	
@@ -1099,6 +1148,7 @@ FText SBlendSpaceGridWidget::GetToolTipSampleValue() const
 		}
 
 		// If we are setting the preview value return the current preview sample value
+		case EDragState::DragPreview:
 		case EDragState::Preview:
 		{
 			FFormatOrderedArguments PreviewTextArguments;
@@ -1230,6 +1280,18 @@ void SBlendSpaceGridWidget::OnInputBoxValueChanged(const float NewValue, const i
 	}
 }
 
+EVisibility SBlendSpaceGridWidget::GetSampleToolTipVisibility() const
+{
+	// Show tool tip when the grid is empty
+	return BlendSpace->GetNumberOfBlendSamples() == 0 ? EVisibility::Visible : EVisibility::Collapsed;
+}
+
+EVisibility SBlendSpaceGridWidget::GetPreviewToolTipVisibility() const
+{
+	// Only show preview tooltip until the user discovers the functionality
+	return !bPreviewToolTipHidden ? EVisibility::Visible : EVisibility::Collapsed;
+}
+
 EVisibility SBlendSpaceGridWidget::GetTriangulationButtonVisibility() const
 {
 	return (GridType == EGridType::TwoAxis) ? EVisibility::Visible : EVisibility::Collapsed;
diff --git a/Engine/Source/Editor/Persona/Private/SAnimationBlendSpaceGridWidget.h b/Engine/Source/Editor/Persona/Private/SAnimationBlendSpaceGridWidget.h
--- a/Engine/Source/Editor/Persona/Private/SAnimationBlendSpaceGridWidget.h
+++ b/Engine/Source/Editor/Persona/Private/SAnimationBlendSpaceGridWidget.h
@@ -55,6 +55,8 @@ protected:
 		PreDrag,
 		/** The user is dragging the selected sample. */
 		DragSample,
+		/** The user is dragging the preview pin */
+		DragPreview,
 		/** The user is setting the preview value */
 		Preview,
 		/** The user is dropping a new sample onto the grid */
@@ -178,6 +180,10 @@ protected:
 	void OnInputBoxValueCommited(const float NewValue, ETextCommit::Type CommitType, const int32 ParameterIndex);
 	void OnInputBoxValueChanged(const float NewValue, const int32 ParameterIndex, bool bIsInteractive);
 
+	/** Returns whether or not the sample tool tips should be visible */
+	EVisibility GetSampleToolTipVisibility() const;
+	EVisibility GetPreviewToolTipVisibility() const;
+
 	/** Updates the cached blend parameter data */
 	void UpdateCachedBlendParameterData();
 private:
@@ -285,4 +291,6 @@ private:
 
 	bool bStretchToFit;
 	FMargin GridRatioMargin;
+
+	bool bPreviewToolTipHidden;
 };
diff --git a/Engine/Source/Editor/Persona/Private/SAnimationEditorViewport.cpp b/Engine/Source/Editor/Persona/Private/SAnimationEditorViewport.cpp
--- a/Engine/Source/Editor/Persona/Private/SAnimationEditorViewport.cpp
+++ b/Engine/Source/Editor/Persona/Private/SAnimationEditorViewport.cpp
@@ -36,7 +36,13 @@ void SAnimationEditorViewport::Construct(const FArguments& InArgs, const FAnimat
 	PreviewScenePtr = InRequiredArgs.PreviewScene;
 	TabBodyPtr = InRequiredArgs.TabBody;
 	AssetEditorToolkitPtr = InRequiredArgs.AssetEditorToolkit;
+	Extenders = InArgs._Extenders;
+	bShowShowMenu = InArgs._ShowShowMenu;
+	bShowLODMenu = InArgs._ShowLODMenu;
+	bShowPlaySpeedMenu = InArgs._ShowPlaySpeedMenu;
 	bShowStats = InArgs._ShowStats;
+	bShowFloorOptions = InArgs._ShowFloorOptions;
+	bShowTurnTable = InArgs._ShowTurnTable;
 
 	InRequiredArgs.OnPostUndo.Add(FSimpleDelegate::CreateSP(this, &SAnimationEditorViewport::OnUndoRedo));
 
@@ -65,7 +71,13 @@ TSharedRef<FEditorViewportClient> SAnimationEditorViewport::MakeEditorViewportCl
 TSharedPtr<SWidget> SAnimationEditorViewport::MakeViewportToolbar()
 {
 	return SNew(SAnimViewportToolBar, TabBodyPtr.Pin(), SharedThis(this))
-		.Cursor(EMouseCursor::Default);
+		.Cursor(EMouseCursor::Default)
+		.Extenders(Extenders)
+		.ShowShowMenu(bShowShowMenu)
+		.ShowLODMenu(bShowLODMenu)
+		.ShowPlaySpeedMenu(bShowPlaySpeedMenu)
+		.ShowFloorOptions(bShowFloorOptions)
+		.ShowTurnTable(bShowTurnTable);
 }
 
 void SAnimationEditorViewport::OnUndoRedo()
@@ -79,105 +91,6 @@ void SAnimationEditorViewport::OnFocusViewportToSelection()
 	AnimViewportClient->FocusViewportOnPreviewMesh(false);
 }
 
-//////////////////////////////////////////////////////////////////////////
-// 
-/** //@TODO MODES: Simple text entry popup used it to get 3 vector position*/
-
-class STranslationInputWindow : public SCompoundWidget
-{
-public:
-	SLATE_BEGIN_ARGS( STranslationInputWindow )
-		: _X_Value(0.0f)
-		, _Y_Value(0.0f)
-		, _Z_Value(0.0f)
-	{}
-
-
-	SLATE_ARGUMENT( float, X_Value )
-	SLATE_ARGUMENT( float, Y_Value )
-	SLATE_ARGUMENT( float, Z_Value )
-	SLATE_EVENT( FOnTextCommitted, OnXModified )
-	SLATE_EVENT( FOnTextCommitted, OnYModified )
-	SLATE_EVENT( FOnTextCommitted, OnZModified )
-	SLATE_END_ARGS()
-
-	void Construct( const FArguments& InArgs )
-	{
-		FString XString = FString::Printf(TEXT("%0.2f"), InArgs._X_Value);
-		FString YString = FString::Printf(TEXT("%0.2f"), InArgs._Y_Value);
-		FString ZString = FString::Printf(TEXT("%0.2f"), InArgs._Z_Value);
-
-		this->ChildSlot
-		[
-			SNew(SBorder)
-			. BorderImage(FEditorStyle::GetBrush(TEXT("Menu.Background")))
-			. Padding(10)
-			[
-				SNew(SVerticalBox)
-				+SVerticalBox::Slot()
-				.AutoHeight()
-				[
-					SNew(SHorizontalBox)
-					+SHorizontalBox::Slot()
-					.AutoWidth()
-					[
-						SNew(STextBlock)
-						.Text(NSLOCTEXT("AnimationEditorViewport", "XValueLabel", "X:"))
-					]
-					+SHorizontalBox::Slot()
-					.AutoWidth()
-					[
-						SNew(SEditableTextBox)
-						.MinDesiredWidth(10.0f)
-						.Text( FText::FromString(XString) )
-						.OnTextCommitted( InArgs._OnXModified )
-					]
-				]
-
-				+SVerticalBox::Slot()
-				.AutoHeight()
-				[
-					SNew(SHorizontalBox)
-					+SHorizontalBox::Slot()
-					.AutoWidth()
-					[
-						SNew(STextBlock)
-						.Text(NSLOCTEXT("AnimationEditorViewport", "YValueLabel", "Y:"))
-					]
-					+SHorizontalBox::Slot()
-					.AutoWidth()
-					[
-						SNew(SEditableTextBox)
-						.MinDesiredWidth(10.0f)
-						.Text( FText::FromString(YString) )
-						.OnTextCommitted( InArgs._OnYModified )
-					]
-				]
-
-				+SVerticalBox::Slot()
-				.AutoHeight()
-				[
-					SNew(SHorizontalBox)
-					+SHorizontalBox::Slot()
-					.AutoWidth()
-					[
-						SNew(STextBlock)
-						.Text(NSLOCTEXT("AnimationEditorViewport", "ZValueLabel", "Z: "))
-					]
-					+SHorizontalBox::Slot()
-					.AutoWidth()
-					[
-						SNew(SEditableTextBox)
-						.MinDesiredWidth(10.0f)
-						.Text( FText::FromString(ZString) )
-						.OnTextCommitted( InArgs._OnZModified )
-					]
-				]
-			]
-		];
-	}
-};
-
 //////////////////////////////////////////////////////////////////////////
 // SAnimationEditorViewportTabBody
 
@@ -201,6 +114,11 @@ SAnimationEditorViewportTabBody::~SAnimationEditorViewportTabBody()
 
 bool SAnimationEditorViewportTabBody::CanUseGizmos() const
 {
+	if (bAlwaysShowTransformToolbar)
+	{
+		return true;
+	}
+
 	class UDebugSkelMeshComponent* Component = GetPreviewScene()->GetPreviewMeshComponent();
 
 	if (Component != NULL)
@@ -249,15 +167,9 @@ FText SAnimationEditorViewportTabBody::GetDisplayString() const
 		{
 			return FText::Format(LOCTEXT("NoMeshFound", "No skeletal mesh found for skeleton '{0}'"), FText::FromName(TargetSkeletonName));
 		}
-		else
-		{
-			return LOCTEXT("NothingToPlay", "Nothing to play");
-		}
-	}
-	else
-	{
-		return FText::Format(LOCTEXT("NoMeshFound", "No skeletal mesh found for skeleton '{0}'"), FText::FromName(TargetSkeletonName));
 	}
+
+	return FText();
 }
 
 TSharedRef<IPersonaViewportState> SAnimationEditorViewportTabBody::SaveState() const
@@ -307,6 +219,7 @@ void SAnimationEditorViewportTabBody::Construct(const FArguments& InArgs, const
 	AssetEditorToolkitPtr = InAssetEditorToolkit;
 	BlueprintEditorPtr = InArgs._BlueprintEditor;
 	bShowTimeline = InArgs._ShowTimeline;
+	bAlwaysShowTransformToolbar = InArgs._AlwaysShowTransformToolbar;
 	OnInvokeTab = InArgs._OnInvokeTab;
 
 	// register delegates for change notifications
@@ -328,7 +241,13 @@ void SAnimationEditorViewportTabBody::Construct(const FArguments& InArgs, const
 	FAnimationEditorViewportRequiredArgs ViewportArgs(InSkeletonTree, InPreviewScene, SharedThis(this), InAssetEditorToolkit, InOnUndoRedo);
 
 	ViewportWidget = SNew(SAnimationEditorViewport, ViewportArgs)
-		.ShowStats(InArgs._ShowStats);
+		.Extenders(InArgs._Extenders)
+		.ShowShowMenu(InArgs._ShowShowMenu)
+		.ShowLODMenu(InArgs._ShowLODMenu)
+		.ShowPlaySpeedMenu(InArgs._ShowPlaySpeedMenu)
+		.ShowStats(InArgs._ShowStats)
+		.ShowFloorOptions(InArgs._ShowFloorOptions)
+		.ShowTurnTable(InArgs._ShowTurnTable);
 
 	TSharedPtr<SVerticalBox> ViewportContainer = nullptr;
 	this->ChildSlot
@@ -1168,7 +1087,7 @@ void SAnimationEditorViewportTabBody::UseInGameBound()
 bool SAnimationEditorViewportTabBody::CanUseInGameBound() const
 {
 	UDebugSkelMeshComponent* PreviewComponent = GetPreviewScene()->GetPreviewMeshComponent();
-	return PreviewComponent != NULL;
+	return PreviewComponent != NULL && IsShowBoundEnabled();
 }
 
 bool SAnimationEditorViewportTabBody::IsUsingInGameBound() const
@@ -1383,31 +1302,6 @@ void SAnimationEditorViewportTabBody::OnLODModelChanged()
 	}
 }
 
-FLinearColor SAnimationEditorViewportTabBody::GetViewportBackgroundColor() const
-{
-	TSharedRef<FAnimationViewportClient> AnimViewportClient = StaticCastSharedRef<FAnimationViewportClient>(LevelViewportClient.ToSharedRef());
-	return AnimViewportClient->GetBackgroundColor();
-}
-
-void SAnimationEditorViewportTabBody::SetViewportBackgroundColor(FLinearColor InColor)
-{
-	TSharedRef<FAnimationViewportClient> AnimViewportClient = StaticCastSharedRef<FAnimationViewportClient>(LevelViewportClient.ToSharedRef());
-	AnimViewportClient->SetBackgroundColor( InColor );
-}
-
-float SAnimationEditorViewportTabBody::GetBackgroundBrightness() const
-{
-	TSharedRef<FAnimationViewportClient> AnimViewportClient = StaticCastSharedRef<FAnimationViewportClient>(LevelViewportClient.ToSharedRef());
-	return AnimViewportClient->GetBrightnessValue();
-}
-
-void SAnimationEditorViewportTabBody::SetBackgroundBrightness(float Value)
-{
-	TSharedRef<FAnimationViewportClient> AnimViewportClient = StaticCastSharedRef<FAnimationViewportClient>(LevelViewportClient.ToSharedRef());
-	AnimViewportClient->SetBrightnessValue(Value);
-	RefreshViewport();
-}
-
 TSharedRef<FAnimationViewportClient> SAnimationEditorViewportTabBody::GetAnimationViewportClient() const
 {
 	return StaticCastSharedRef<FAnimationViewportClient>(LevelViewportClient.ToSharedRef());
@@ -1488,7 +1382,7 @@ void SAnimationEditorViewportTabBody::OnTogglePreviewRootMotion()
 
 	if (PreviewComponent)
 	{
-		PreviewComponent->bPreviewRootMotion = !PreviewComponent->bPreviewRootMotion;
+		PreviewComponent->SetPreviewRootMotion(!PreviewComponent->GetPreviewRootMotion());
 	}
 }
 
@@ -1498,7 +1392,7 @@ bool SAnimationEditorViewportTabBody::IsPreviewingRootMotion() const
 
 	if (PreviewComponent)
 	{
-		return PreviewComponent->bPreviewRootMotion;
+		return PreviewComponent->GetPreviewRootMotion();
 	}
 	return false;
 }
diff --git a/Engine/Source/Editor/Persona/Private/SAnimationEditorViewport.h b/Engine/Source/Editor/Persona/Private/SAnimationEditorViewport.h
--- a/Engine/Source/Editor/Persona/Private/SAnimationEditorViewport.h
+++ b/Engine/Source/Editor/Persona/Private/SAnimationEditorViewport.h
@@ -18,6 +18,7 @@
 #include "ISkeletonTree.h"
 #include "AnimationEditorPreviewScene.h"
 #include "SEditorViewport.h"
+#include "PersonaModule.h"
 
 class SAnimationEditorViewportTabBody;
 
@@ -49,11 +50,28 @@ class SAnimationEditorViewport : public SEditorViewport
 {
 public:
 	SLATE_BEGIN_ARGS(SAnimationEditorViewport) 
-		: _ShowStats(true)
+		: _ShowShowMenu(true)
+		, _ShowLODMenu(true)
+		, _ShowPlaySpeedMenu(true)
+		, _ShowStats(true)
+		, _ShowFloorOptions(true)
+		, _ShowTurnTable(true)
 	{}
 
+	SLATE_ARGUMENT(TSharedPtr<FExtender>, Extenders)
+
+	SLATE_ARGUMENT(bool, ShowShowMenu)
+
+	SLATE_ARGUMENT(bool, ShowLODMenu)
+
+	SLATE_ARGUMENT(bool, ShowPlaySpeedMenu)
+
 	SLATE_ARGUMENT(bool, ShowStats)
 
+	SLATE_ARGUMENT(bool, ShowFloorOptions)
+
+	SLATE_ARGUMENT(bool, ShowTurnTable)
+
 	SLATE_END_ARGS()
 
 	void Construct(const FArguments& InArgs, const FAnimationEditorViewportRequiredArgs& InRequiredArgs);
@@ -84,8 +102,26 @@ protected:
 	// The asset editor we are embedded in
 	TWeakPtr<class FAssetEditorToolkit> AssetEditorToolkitPtr;
 
+	/** Menu extenders */
+	TSharedPtr<FExtender> Extenders;
+
+	/** Whether to show the 'Show' menu */
+	bool bShowShowMenu;
+
+	/** Whether to show the 'LOD' menu */
+	bool bShowLODMenu;
+
+	/** Whether to show the 'Play Speed' menu */
+	bool bShowPlaySpeedMenu;
+
 	/** Whether we should show stats for this viewport */
 	bool bShowStats;
+
+	/** Whether to show options relating to floor height */
+	bool bShowFloorOptions;
+
+	/** Whether to show options relating to turntable */
+	bool bShowTurnTable;
 };
 
 //////////////////////////////////////////////////////////////////////////
@@ -96,16 +132,37 @@ class SAnimationEditorViewportTabBody : public IPersonaViewport
 public:
 	SLATE_BEGIN_ARGS( SAnimationEditorViewportTabBody )
 		: _BlueprintEditor()
+		, _ShowShowMenu(true)
+		, _ShowLODMenu(true)
+		, _ShowPlaySpeedMenu(true)
 		, _ShowTimeline(true)
+		, _ShowStats(true)
+		, _AlwaysShowTransformToolbar(false)
+		, _ShowFloorOptions(true)
+		, _ShowTurnTable(true)
 		{}
 
 		SLATE_ARGUMENT(TWeakPtr<FBlueprintEditor>, BlueprintEditor)
 
 		SLATE_ARGUMENT(FOnInvokeTab, OnInvokeTab)
 
+		SLATE_ARGUMENT(TSharedPtr<FExtender>, Extenders)
+
+		SLATE_ARGUMENT(bool, ShowShowMenu)
+
+		SLATE_ARGUMENT(bool, ShowLODMenu)
+
+		SLATE_ARGUMENT(bool, ShowPlaySpeedMenu)
+
 		SLATE_ARGUMENT(bool, ShowTimeline)
 
 		SLATE_ARGUMENT(bool, ShowStats)
+
+		SLATE_ARGUMENT(bool, AlwaysShowTransformToolbar)
+
+		SLATE_ARGUMENT(bool, ShowFloorOptions)
+
+		SLATE_ARGUMENT(bool, ShowTurnTable)
 	SLATE_END_ARGS()
 public:
 
@@ -147,18 +204,6 @@ public:
 	/** Function to get anim viewport widget */
 	TSharedPtr<class SEditorViewport> GetViewportWidget() const { return ViewportWidget; }
 
-	/** Function to get viewport's current background color */
-	FLinearColor GetViewportBackgroundColor( ) const;
-
-	/** Function to set viewport's new background color */
-	void SetViewportBackgroundColor( FLinearColor InColor );
-
-	/** Function to get viewport's background color brightness */
-	float GetBackgroundBrightness( ) const;
-
-	/** Function to set viewport's background color brightness */
-	void SetBackgroundBrightness( float Value );
-
 	/** Function to check whether grid is displayed or not */
 	bool IsShowingGrid() const;
 
@@ -435,6 +480,9 @@ private:
 	/** Whether to show the timeline */
 	bool bShowTimeline;
 
+	/** Whether we should always show the transform toolbar for this viewport */
+	bool bAlwaysShowTransformToolbar;
+
 	/** Level viewport client */
 	TSharedPtr<FEditorViewportClient> LevelViewportClient;
 
diff --git a/Engine/Source/Editor/Persona/Private/SPersonaDetails.cpp b/Engine/Source/Editor/Persona/Private/SPersonaDetails.cpp
--- a/Engine/Source/Editor/Persona/Private/SPersonaDetails.cpp
+++ b/Engine/Source/Editor/Persona/Private/SPersonaDetails.cpp
@@ -11,8 +11,34 @@ void SPersonaDetails::Construct(const FArguments& InArgs)
 	FPropertyEditorModule& PropertyEditorModule = FModuleManager::LoadModuleChecked<FPropertyEditorModule>("PropertyEditor");
 	DetailsView = PropertyEditorModule.CreateDetailView(DetailsViewArgs);
 
-	ChildSlot
+	TSharedRef<SVerticalBox> Content = SNew(SVerticalBox);
+
+	if (InArgs._TopContent.IsValid())
+	{
+		Content->AddSlot()
+		.AutoHeight()
+		[
+			InArgs._TopContent.ToSharedRef()
+		];
+	}
+
+	Content->AddSlot()
+	.FillHeight(1.0f)
 	[
 		DetailsView.ToSharedRef()
 	];
+
+	if (InArgs._BottomContent.IsValid())
+	{
+		Content->AddSlot()
+		.AutoHeight()
+		[
+			InArgs._BottomContent.ToSharedRef()
+		];
+	}
+
+	ChildSlot
+	[
+		Content
+	];
 }
diff --git a/Engine/Source/Editor/Persona/Private/SPersonaDetails.h b/Engine/Source/Editor/Persona/Private/SPersonaDetails.h
--- a/Engine/Source/Editor/Persona/Private/SPersonaDetails.h
+++ b/Engine/Source/Editor/Persona/Private/SPersonaDetails.h
@@ -12,6 +12,12 @@ class SPersonaDetails : public SCompoundWidget
 public:
 	SLATE_BEGIN_ARGS(SPersonaDetails) {}
 
+	/** Optional content to display above the details panel */
+	SLATE_ARGUMENT(TSharedPtr<SWidget>, TopContent)
+
+	/** Optional content to display below the details panel */
+	SLATE_ARGUMENT(TSharedPtr<SWidget>, BottomContent)
+
 	SLATE_END_ARGS()
 
 	void Construct(const FArguments& InArgs);
diff --git a/Engine/Source/Editor/Persona/Private/TabSpawners.cpp b/Engine/Source/Editor/Persona/Private/TabSpawners.cpp
--- a/Engine/Source/Editor/Persona/Private/TabSpawners.cpp
+++ b/Engine/Source/Editor/Persona/Private/TabSpawners.cpp
@@ -203,15 +203,22 @@ TSharedRef<SWidget> FAnimationAssetBrowserSummoner::CreateTabBody(const FWorkflo
 /////////////////////////////////////////////////////
 // FPreviewViewportSummoner
 
-FPreviewViewportSummoner::FPreviewViewportSummoner(TSharedPtr<class FAssetEditorToolkit> InHostingApp, const TSharedRef<ISkeletonTree>& InSkeletonTree, const TSharedRef<IPersonaPreviewScene>& InPreviewScene, FSimpleMulticastDelegate& InOnPostUndo, const TSharedPtr<FBlueprintEditor>& InBlueprintEditor, FOnViewportCreated InOnViewportCreated, bool bInShowTimeline, bool bInShowStats)
+FPreviewViewportSummoner::FPreviewViewportSummoner(TSharedPtr<class FAssetEditorToolkit> InHostingApp, const FPersonaViewportArgs& InArgs)
 	: FWorkflowTabFactory(FPersonaTabs::PreviewViewportID, InHostingApp)
-	, SkeletonTree(InSkeletonTree)
-	, PreviewScene(InPreviewScene)
-	, OnPostUndo(InOnPostUndo)
-	, BlueprintEditor(InBlueprintEditor)
-	, OnViewportCreated(InOnViewportCreated)
-	, bShowTimeline(bInShowTimeline)
-	, bShowStats(bInShowStats)
+	, SkeletonTree(InArgs.SkeletonTree)
+	, PreviewScene(InArgs.PreviewScene)
+	, OnPostUndo(InArgs.OnPostUndo)
+	, BlueprintEditor(InArgs.BlueprintEditor)
+	, OnViewportCreated(InArgs.OnViewportCreated)
+	, Extenders(InArgs.Extenders)
+	, bShowShowMenu(InArgs.bShowShowMenu)
+	, bShowLODMenu(InArgs.bShowLODMenu)
+	, bShowPlaySpeedMenu(InArgs.bShowPlaySpeedMenu)
+	, bShowTimeline(InArgs.bShowTimeline)
+	, bShowStats(InArgs.bShowStats)
+	, bAlwaysShowTransformToolbar(InArgs.bAlwaysShowTransformToolbar)
+	, bShowFloorOptions(InArgs.bShowFloorOptions)
+	, bShowTurnTable(InArgs.bShowTurnTable)
 {
 	TabLabel = LOCTEXT("ViewportTabTitle", "Viewport");
 	TabIcon = FSlateIcon(FEditorStyle::GetStyleSetName(), "LevelEditor.Tabs.Viewports");
@@ -228,8 +235,15 @@ TSharedRef<SWidget> FPreviewViewportSummoner::CreateTabBody(const FWorkflowTabSp
 		.BlueprintEditor(BlueprintEditor.Pin())
 		.OnInvokeTab(FOnInvokeTab::CreateSP(HostingApp.Pin().Get(), &FAssetEditorToolkit::InvokeTab))
 		.AddMetaData<FTagMetaData>(TEXT("Persona.Viewport"))
+		.Extenders(Extenders)
+		.ShowShowMenu(bShowShowMenu)
+		.ShowLODMenu(bShowLODMenu)
+		.ShowPlaySpeedMenu(bShowPlaySpeedMenu)
 		.ShowTimeline(bShowTimeline)
-		.ShowStats(bShowStats);
+		.ShowStats(bShowStats)
+		.AlwaysShowTransformToolbar(bAlwaysShowTransformToolbar)
+		.ShowFloorOptions(bShowFloorOptions)
+		.ShowTurnTable(bShowTurnTable);
 
 	OnViewportCreated.ExecuteIfBound(NewViewport);
 
diff --git a/Engine/Source/Editor/Persona/Private/TabSpawners.h b/Engine/Source/Editor/Persona/Private/TabSpawners.h
--- a/Engine/Source/Editor/Persona/Private/TabSpawners.h
+++ b/Engine/Source/Editor/Persona/Private/TabSpawners.h
@@ -242,7 +242,7 @@ private:
 
 struct FPreviewViewportSummoner : public FWorkflowTabFactory
 {
-	FPreviewViewportSummoner(TSharedPtr<class FAssetEditorToolkit> InHostingApp, const TSharedRef<ISkeletonTree>& InSkeletonTree, const TSharedRef<IPersonaPreviewScene>& InPreviewScene, FSimpleMulticastDelegate& InOnPostUndo, const TSharedPtr<FBlueprintEditor>& InBlueprintEditor, FOnViewportCreated InOnViewportCreated, bool bShowTimeline, bool bInShowStats);
+	FPreviewViewportSummoner(TSharedPtr<class FAssetEditorToolkit> InHostingApp, const FPersonaViewportArgs& InArgs);
 
 	virtual TSharedRef<SWidget> CreateTabBody(const FWorkflowTabSpawnInfo& Info) const override;
 
@@ -251,8 +251,15 @@ struct FPreviewViewportSummoner : public FWorkflowTabFactory
 	FSimpleMulticastDelegate& OnPostUndo;
 	TWeakPtr<FBlueprintEditor> BlueprintEditor;
 	FOnViewportCreated OnViewportCreated;
+	TSharedPtr<FExtender> Extenders;
+	bool bShowShowMenu;
+	bool bShowLODMenu;
+	bool bShowPlaySpeedMenu;
 	bool bShowTimeline;
 	bool bShowStats;
+	bool bAlwaysShowTransformToolbar;
+	bool bShowFloorOptions;
+	bool bShowTurnTable;
 };
 
 /////////////////////////////////////////////////////
diff --git a/Engine/Source/Editor/Persona/Public/IPersonaPreviewScene.h b/Engine/Source/Editor/Persona/Public/IPersonaPreviewScene.h
--- a/Engine/Source/Editor/Persona/Public/IPersonaPreviewScene.h
+++ b/Engine/Source/Editor/Persona/Public/IPersonaPreviewScene.h
@@ -63,6 +63,9 @@ public:
 	/** Get the skeletal mesh component we are using for preview, if any. */
 	virtual UDebugSkelMeshComponent* GetPreviewMeshComponent() const = 0;
 
+	/** Set the skeletal mesh component we are going to preview. */
+	virtual void SetPreviewMeshComponent(UDebugSkelMeshComponent* InSkeletalMeshComponent) = 0;
+
 	/** Set the additional meshes used by this preview scene (sets the additional meshes on the skeleton) */
 	virtual void SetAdditionalMeshes(class UDataAsset* InAdditionalMeshes) = 0;
 
@@ -192,6 +195,9 @@ public:
 	/** Get the main actor */
 	virtual AActor* GetActor() const = 0;
 
+	/** Set the main actor */
+	virtual void SetActor(AActor* InActor) = 0;
+
 	/** Get whether or not to ignore mesh hit proxies */
 	virtual bool AllowMeshHitProxies() const = 0;
 
diff --git a/Engine/Source/Editor/Persona/Public/PersonaModule.h b/Engine/Source/Editor/Persona/Public/PersonaModule.h
--- a/Engine/Source/Editor/Persona/Public/PersonaModule.h
+++ b/Engine/Source/Editor/Persona/Public/PersonaModule.h
@@ -18,6 +18,7 @@ class IPersonaPreviewScene;
 class IPersonaToolkit;
 class UAnimBlueprint;
 class USkeletalMeshComponent;
+class UPhysicsAsset;
 
 extern const FName PersonaAppName;
 
@@ -50,12 +51,19 @@ DECLARE_MULTICAST_DELEGATE_OneParam(FOnPreviewSceneCreated, const TSharedRef<cla
 /** Initialization parameters for persona toolkits */
 struct FPersonaToolkitArgs
 {
+	/** 
+	 * Delegate called when the preview scene is created, used to setup the scene 
+	 * If this is not set, then a default scene will be set up.
+	 */
+	FOnPreviewSceneCreated::FDelegate OnPreviewSceneCreated;
+
 	/** Whether to create a preview scene */
 	bool bCreatePreviewScene;
 
 	FPersonaToolkitArgs()
 		: bCreatePreviewScene(true)
-	{}
+	{
+	}
 };
 
 struct FAnimDocumentArgs
@@ -84,6 +92,62 @@ struct FAnimDocumentArgs
 	FSimpleDelegate OnDespatchAnimNotifiesChanged;
 };
 
+/** Arguments used to create a persona viewport tab */
+struct FPersonaViewportArgs
+{
+	FPersonaViewportArgs(const TSharedRef<class ISkeletonTree>& InSkeletonTree, const TSharedRef<class IPersonaPreviewScene>& InPreviewScene, FSimpleMulticastDelegate& InOnPostUndo)
+		: SkeletonTree(InSkeletonTree)
+		, PreviewScene(InPreviewScene)
+		, OnPostUndo(InOnPostUndo)
+		, bShowShowMenu(true)
+		, bShowLODMenu(true)
+		, bShowPlaySpeedMenu(true)
+		, bShowTimeline(true)
+		, bShowStats(true)
+		, bAlwaysShowTransformToolbar(false)
+		, bShowFloorOptions(true)
+		, bShowTurnTable(true)
+	{}
+
+	/** Required args */
+	TSharedRef<class ISkeletonTree> SkeletonTree;
+	TSharedRef<class IPersonaPreviewScene> PreviewScene;
+	FSimpleMulticastDelegate& OnPostUndo;
+
+	/** Optional blueprint editor that we can be embedded in */
+	TSharedPtr<class FBlueprintEditor> BlueprintEditor;
+
+	/** Delegate fired when the viewport is created */
+	FOnViewportCreated OnViewportCreated;
+	
+	/** Menu extenders */
+	TSharedPtr<FExtender> Extenders;
+
+	/** Whether to show the 'Show' menu */
+	bool bShowShowMenu;
+
+	/** Whether to show the 'LOD' menu */
+	bool bShowLODMenu;
+
+	/** Whether to show the 'Play Speed' menu */
+	bool bShowPlaySpeedMenu;
+
+	/** Whether to show the animation timeline */
+	bool bShowTimeline;
+
+	/** Whether to show in-viewport stats */
+	bool bShowStats;
+
+	/** Whether we should always show the transform toolbar for this viewport */
+	bool bAlwaysShowTransformToolbar;
+
+	/** Whether to show options relating to floor height */
+	bool bShowFloorOptions;
+
+	/** Whether to show options relating to turntable */
+	bool bShowTurnTable;
+};
+
 /**
  * Persona module manages the lifetime of all instances of Persona editors.
  */
@@ -106,6 +170,7 @@ public:
 	virtual TSharedRef<class IPersonaToolkit> CreatePersonaToolkit(UAnimationAsset* InAnimationAsset, const FPersonaToolkitArgs& PersonaToolkitArgs = FPersonaToolkitArgs()) const;
 	virtual TSharedRef<class IPersonaToolkit> CreatePersonaToolkit(USkeletalMesh* InSkeletalMesh, const FPersonaToolkitArgs& PersonaToolkitArgs = FPersonaToolkitArgs()) const;
 	virtual TSharedRef<class IPersonaToolkit> CreatePersonaToolkit(UAnimBlueprint* InAnimBlueprint, const FPersonaToolkitArgs& PersonaToolkitArgs = FPersonaToolkitArgs()) const;
+	virtual TSharedRef<class IPersonaToolkit> CreatePersonaToolkit(UPhysicsAsset* InPhysicsAsset, const FPersonaToolkitArgs& PersonaToolkitArgs = FPersonaToolkitArgs()) const;
 
 	/** Create an asset family for the supplied persona asset */
 	virtual TSharedRef<class IAssetFamily> CreatePersonaAssetFamily(const UObject* InAsset) const;
@@ -117,7 +182,7 @@ public:
 	virtual TSharedRef<class FWorkflowTabFactory> CreateDetailsTabFactory(const TSharedRef<class FWorkflowCentricApplication>& InHostingApp, FOnDetailsCreated InOnDetailsCreated) const;
 
 	/** Create a persona viewport tab factory */
-	virtual TSharedRef<class FWorkflowTabFactory> CreatePersonaViewportTabFactory(const TSharedRef<class FWorkflowCentricApplication>& InHostingApp, const TSharedRef<class ISkeletonTree>& InSkeletonTree, const TSharedRef<class IPersonaPreviewScene>& InPreviewScene, FSimpleMulticastDelegate& OnPostUndo, const TSharedPtr<class FBlueprintEditor>& InBlueprintEditor, FOnViewportCreated InOnViewportCreated, bool bInShowTimeline, bool bInShowStats) const;
+	virtual TSharedRef<class FWorkflowTabFactory> CreatePersonaViewportTabFactory(const TSharedRef<class FWorkflowCentricApplication>& InHostingApp, const FPersonaViewportArgs& InArgs) const;
 
 	/** Create an anim notifies tab factory */
 	virtual TSharedRef<class FWorkflowTabFactory> CreateAnimNotifiesTabFactory(const TSharedRef<class FWorkflowCentricApplication>& InHostingApp, const TSharedRef<class IEditableSkeleton>& InEditableSkeleton, FSimpleMulticastDelegate& InOnChangeAnimNotifies, FSimpleMulticastDelegate& InOnPostUndo, FOnObjectsSelected InOnObjectsSelected) const;
diff --git a/Engine/Source/Editor/PhysicsAssetEditor/PhysicsAssetEditor.Build.cs b/Engine/Source/Editor/PhysicsAssetEditor/PhysicsAssetEditor.Build.cs
new file mode 100644
--- /dev/null
+++ b/Engine/Source/Editor/PhysicsAssetEditor/PhysicsAssetEditor.Build.cs
@@ -0,0 +1,48 @@
+// Copyright 1998-2017 Epic Games, Inc. All Rights Reserved.
+
+using UnrealBuildTool;
+
+public class PhysicsAssetEditor : ModuleRules
+{
+	public PhysicsAssetEditor(ReadOnlyTargetRules Target) : base(Target)
+	{
+		PublicIncludePaths.Add("Editor/UnrealEd/Public");
+
+        PrivateIncludePaths.AddRange(
+            new string[] {
+                "Editor/PhysicsAssetEditor/Private",
+                "Editor/PhysicsAssetEditor/Private/PhysicsAssetGraph",
+            }
+        );
+
+        PrivateDependencyModuleNames.AddRange(
+			new string[] {
+				"Core",
+				"CoreUObject",
+				"Engine",
+                "InputCore",
+				"RenderCore",
+				"Slate",
+				"SlateCore",
+                "EditorStyle",
+				"LevelEditor",
+				"UnrealEd",
+                "Kismet",
+                "Persona",
+                "SkeletonEditor",
+                "GraphEditor",
+                "AnimGraph",
+                "AnimGraphRuntime",
+                "AdvancedPreviewScene",
+            }
+        );
+
+		DynamicallyLoadedModuleNames.AddRange(
+			new string[] {
+				"MainFrame",
+				"PropertyEditor",
+                "MeshUtilities",
+			}
+		);
+	}
+}
