Fix localization warnings

Fallout from Dev-AnimPhys merge

diff --git a/Engine/Source/Editor/Persona/Private/PersonaMeshDetails.cpp b/Engine/Source/Editor/Persona/Private/PersonaMeshDetails.cpp
--- a/Engine/Source/Editor/Persona/Private/PersonaMeshDetails.cpp
+++ b/Engine/Source/Editor/Persona/Private/PersonaMeshDetails.cpp
@@ -1338,7 +1338,7 @@ FText FPersonaMeshDetails::GetLODCustomModeNameContent(int32 LODIndex) const
 	int32 RealCurrentLODIndex = (CurrentLodIndex == 0 ? 0 : CurrentLodIndex - 1);
 	if (LODIndex == INDEX_NONE)
 	{
-		return LOCTEXT("GetLODCustomModeNameContent", "Custom");
+		return LOCTEXT("GetLODCustomModeNameContent_None", "Custom");
 	}
 	return FText::Format(LOCTEXT("GetLODCustomModeNameContent", "LOD{0}"), LODIndex);
 }
diff --git a/Engine/Source/Editor/Persona/Private/SAnimViewportToolBar.cpp b/Engine/Source/Editor/Persona/Private/SAnimViewportToolBar.cpp
--- a/Engine/Source/Editor/Persona/Private/SAnimViewportToolBar.cpp
+++ b/Engine/Source/Editor/Persona/Private/SAnimViewportToolBar.cpp
@@ -240,7 +240,7 @@ TSharedRef<SWidget> SAnimViewportToolBar::MakeViewportToolbar(TSharedPtr<class S
 	FName MenuBarStyle("ViewportMenu");
 	MenuBarBuilder.SetStyle(&FEditorStyle::Get(), MenuBarStyle);
 
-	MenuBarBuilder.AddPullDownMenu(FText(), LOCTEXT("ViewMenuToolTip", "Show viewport options"), FSlateIcon("EditorStyle", "EditorViewportToolBar.MenuDropdown"), FNewMenuDelegate::CreateSP(this, &SAnimViewportToolBar::GenerateViewMenu), "AnimViewportViewMenu");
+	MenuBarBuilder.AddPullDownMenu(FText(), LOCTEXT("ViewportMenuToolTip", "Show viewport options"), FSlateIcon("EditorStyle", "EditorViewportToolBar.MenuDropdown"), FNewMenuDelegate::CreateSP(this, &SAnimViewportToolBar::GenerateViewMenu), "AnimViewportViewMenu");
 
 	MenuBarBuilder.AddPullDownMenu(
 		TAttribute<FText>::Create(TAttribute<FText>::FGetter::CreateSP(this, &SAnimViewportToolBar::GetCameraMenuLabel)), 
