Fix crash with automapping

diff --git a/Engine/Source/Editor/Persona/Private/BoneMappingHelper.cpp b/Engine/Source/Editor/Persona/Private/BoneMappingHelper.cpp
--- a/Engine/Source/Editor/Persona/Private/BoneMappingHelper.cpp
+++ b/Engine/Source/Editor/Persona/Private/BoneMappingHelper.cpp
@@ -266,7 +266,7 @@ void FBoneMappingHelper::TryMatch(TMap<FName, FName>& OutBestMatches)
 			UE_LOG(LogAnimation, Log, TEXT("Bone Match [%s] - [%s] (score %0.2f)"), *Bone0Name.ToString(), *Bone1Name.ToString(), Score);
 
 			// print next 3 best score, clear current one
-			for (int32 LoopCount = 0; LoopCount < MAX_CANDIDATE-1; ++LoopCount)
+			for (int32 LoopCount = 0; LoopCount < MAX_CANDIDATE-1 && BoneIndex1 != INDEX_NONE; ++LoopCount)
 			{
 				BoneDescArray0[BoneIndex0].SetScore(BoneIndex1, 0.f);
 
