Merging //UE4/Release-4.18 @ CL 3709383 to Release-Staging-4.18 (//UE4/Release-Staging-4.18) (merged by hand)

diff --git a/Engine/Plugins/Runtime/AppleARKit/AppleARKit.uplugin b/Engine/Plugins/Runtime/AppleARKit/AppleARKit.uplugin
--- a/Engine/Plugins/Runtime/AppleARKit/AppleARKit.uplugin
+++ b/Engine/Plugins/Runtime/AppleARKit/AppleARKit.uplugin
@@ -12,19 +12,13 @@
 	"SupportURL": "",
 	"CanContainContent": true,
 	"IsBetaVersion": true,
-	"EnabledByDefault": false,
+	"EnabledByDefault": true,
 	"Modules": [
 		{
 			"Name": "AppleARKit",
 			"Type": "Runtime",
 			"LoadingPhase": "PostConfigInit",
 			"WhitelistPlatforms" : [ "Mac", "iOS", "Win64" ]
-		},
-		{
-			"Name": "AppleARKitEditor",
-			"Type": "Editor",
-			"LoadingPhase": "Default",
-			"WhitelistPlatforms" : [ "Mac", "iOS", "Win64" ]
 		}
 	]
 }
diff --git a/Engine/Shaders/Public/ShaderVersion.ush b/Engine/Shaders/Public/ShaderVersion.ush
--- a/Engine/Shaders/Public/ShaderVersion.ush
+++ b/Engine/Shaders/Public/ShaderVersion.ush
@@ -2,5 +2,4 @@
 // Each time the console command is executed it generates a new GUID. As this file is included
 // in Platform.ush (which should be included in any shader) it allows to invalidate the shader DDC.
 // 
-// GUID = 5E666AF6-CE05434383FFAD743B606CC6
-
+// GUID = 4C6AEAAADB0F4C96A1C49C75F8B4EB16
\ No newline at end of file
diff --git a/Engine/Source/Editor/Persona/Private/SAnimViewportToolBar.cpp b/Engine/Source/Editor/Persona/Private/SAnimViewportToolBar.cpp
--- a/Engine/Source/Editor/Persona/Private/SAnimViewportToolBar.cpp
+++ b/Engine/Source/Editor/Persona/Private/SAnimViewportToolBar.cpp
@@ -37,6 +37,8 @@
 #include "ClothingSimulationFactoryInterface.h"
 #include "ClothingSystemEditorInterfaceModule.h"
 #include "Widgets/SWidget.h"
+#include "ISlateMetaData.h"
+#include "Textures/SlateIcon.h"
 
 #define LOCTEXT_NAMESPACE "AnimViewportToolBar"
 
@@ -167,38 +169,108 @@ void SAnimViewportToolBar::Construct(const FArguments& InArgs, TSharedPtr<class
 	bShowPlaySpeedMenu = InArgs._ShowPlaySpeedMenu;
 	bShowFloorOptions = InArgs._ShowFloorOptions;
 	bShowTurnTable = InArgs._ShowTurnTable;
+	bShowPhysicsMenu = InArgs._ShowPhysicsMenu;
 
 	CommandList = InRealViewport->GetCommandList();
 	Extenders = InArgs._Extenders;
 
 	// If we have no extender, make an empty one
-	if (!Extenders.IsValid())
+	if (Extenders.Num() == 0)
 	{
-		Extenders = MakeShared<FExtender>();
+		Extenders.Add(MakeShared<FExtender>());
 	}
-	
+
+	const FMargin ToolbarSlotPadding(2.0f, 2.0f);
+	const FMargin ToolbarButtonPadding(2.0f, 0.0f);
+
+	static const FName DefaultForegroundName("DefaultForeground");
+
+
 	TSharedRef<SHorizontalBox> LeftToolbar = SNew(SHorizontalBox)
 		+SHorizontalBox::Slot()
 		.AutoWidth()
-		.Padding(0.0f)
+		.Padding(ToolbarSlotPadding)
 		[
-			MakeViewportToolbar(InRealViewport)
+			SNew(SEditorViewportToolbarMenu)
+			.ParentToolBar(SharedThis(this))
+			.Cursor(EMouseCursor::Default)
+			.Image("EditorViewportToolBar.MenuDropdown")
+			.AddMetaData<FTagMetaData>(FTagMetaData(TEXT("EditorViewportToolBar.MenuDropdown")))
+			.OnGetMenuContent(this, &SAnimViewportToolBar::GenerateViewMenu)
 		]
-
-		// Transform toolbar
 		+SHorizontalBox::Slot()
+		.AutoWidth()
+		.Padding(ToolbarSlotPadding)
+		[
+			SNew(SEditorViewportToolbarMenu)
+			.ParentToolBar(SharedThis(this))
+			.Cursor(EMouseCursor::Default)
+			.Label(this, &SAnimViewportToolBar::GetCameraMenuLabel)
+			.LabelIcon(this, &SAnimViewportToolBar::GetCameraMenuLabelIcon)
+			.AddMetaData<FTagMetaData>(FTagMetaData(TEXT("EditorViewportToolBar.CameraMenu")))
+			.OnGetMenuContent(this, &SAnimViewportToolBar::GenerateViewportTypeMenu)
+		]
+		// View menu (lit, unlit, etc...)
+		+ SHorizontalBox::Slot()
+		.AutoWidth()
+		.Padding(ToolbarSlotPadding)
+		[
+			SNew(SEditorViewportViewMenu, InRealViewport.ToSharedRef(), SharedThis(this))
+		]
+		+ SHorizontalBox::Slot()
+		.AutoWidth()
+		.Padding(ToolbarSlotPadding)
+		[
+			SNew(SEditorViewportToolbarMenu)
+			.ParentToolBar(SharedThis(this))
+			.Cursor(EMouseCursor::Default)
+			.Label(LOCTEXT("ShowMenu", "Show"))
+			.AddMetaData<FTagMetaData>(FTagMetaData(TEXT("ViewMenuButton")))
+			.OnGetMenuContent(this, &SAnimViewportToolBar::GenerateShowMenu)
+		]
+		+ SHorizontalBox::Slot()
+		.AutoWidth()
+		.Padding(ToolbarSlotPadding)
+		[
+			SNew(SEditorViewportToolbarMenu)
+			.ParentToolBar(SharedThis(this))
+			.Label(LOCTEXT("Physics", "Physics"))
+			.OnGetMenuContent(this, &SAnimViewportToolBar::GeneratePhysicsMenu)
+			.Visibility(bShowPhysicsMenu ? EVisibility::Visible : EVisibility::Collapsed)
+		]
+		+ SHorizontalBox::Slot()
+		.AutoWidth()
+		.Padding(ToolbarSlotPadding)
+		[
+			//LOD
+			SNew(SEditorViewportToolbarMenu)
+			.ParentToolBar(SharedThis(this))
+			.Label(this, &SAnimViewportToolBar::GetLODMenuLabel)
+			.OnGetMenuContent(this, &SAnimViewportToolBar::GenerateLODMenu)
+		]
+		+ SHorizontalBox::Slot()
+		.AutoWidth()
+		.Padding(ToolbarSlotPadding)
+		[
+			SNew(SEditorViewportToolbarMenu)
+			.ParentToolBar(SharedThis(this))
+			.Label(this, &SAnimViewportToolBar::GetPlaybackMenuLabel)
+			.LabelIcon(FEditorStyle::GetBrush("AnimViewportMenu.PlayBackSpeed"))
+			.OnGetMenuContent(this, &SAnimViewportToolBar::GeneratePlaybackMenu)
+		]
+		+ SHorizontalBox::Slot()
+		.Padding(ToolbarSlotPadding)
 		.HAlign(HAlign_Right)
-		.Padding(3.0f, 0.0f)
 		[
 			SNew(STransformViewportToolBar)
 			.Viewport(InRealViewport)
-			.CommandList(CommandList)
+			.CommandList(InRealViewport->GetCommandList())
 			.Visibility(this, &SAnimViewportToolBar::GetTransformToolbarVisibility)
 		];
+			
+	
 	//@TODO: Need clipping horizontal box: LeftToolbar->AddWrapButton();
 
-	static const FName DefaultForegroundName("DefaultForeground");
-
 	ChildSlot
 	[
 		SNew( SBorder )
@@ -226,7 +298,7 @@ void SAnimViewportToolBar::Construct(const FArguments& InArgs, TSharedPtr<class
 			]
 		]
 	];
-
+	
 	SViewportToolBar::Construct(SViewportToolBar::FArguments());
 
 	// We assign the viewport pointer her rather that initially, as SViewportToolbar::Construct 
@@ -234,80 +306,28 @@ void SAnimViewportToolBar::Construct(const FArguments& InArgs, TSharedPtr<class
 	Viewport = InViewport;
 }
 
-TSharedRef<SWidget> SAnimViewportToolBar::MakeViewportToolbar(TSharedPtr<class SEditorViewport> InRealViewport)
-{
-	FMenuBarBuilder MenuBarBuilder(CommandList, Extenders, &FEditorStyle::Get(), false, EMultiBoxType::ToolMenuBar);
-
-	FName MenuBarStyle("ViewportMenu");
-	MenuBarBuilder.SetStyle(&FEditorStyle::Get(), MenuBarStyle);
-
-	MenuBarBuilder.AddPullDownMenu(FText(), LOCTEXT("ViewportMenuToolTip", "Show viewport options"), FSlateIcon("EditorStyle", "EditorViewportToolBar.MenuDropdown"), FNewMenuDelegate::CreateSP(this, &SAnimViewportToolBar::GenerateViewMenu), "AnimViewportViewMenu");
-
-	MenuBarBuilder.AddPullDownMenu(
-		TAttribute<FText>::Create(TAttribute<FText>::FGetter::CreateSP(this, &SAnimViewportToolBar::GetCameraMenuLabel)), 
-		LOCTEXT("CameraMenuToolTip", "Viewport camera options"), 
-		TAttribute<FSlateIcon>::Create(TAttribute<FSlateIcon>::FGetter::CreateSP(this, &SAnimViewportToolBar::GetCameraMenuLabelIcon)), 
-		FNewMenuDelegate::CreateSP(this, &SAnimViewportToolBar::GenerateViewportTypeMenu), 
-		"AnimViewportCameraMenu");
-
-	TWeakPtr<SEditorViewport> WeakViewport = InRealViewport;
-	TWeakPtr<SViewportToolBar> WeakToolbar = SharedThis(this);
-
-	MenuBarBuilder.AddPullDownMenu(
-		LOCTEXT("ViewMenu", "View"),
-		LOCTEXT("ViewMenuToolTip", "View options"),
-		TAttribute<FSlateIcon>::Create(TAttribute<FSlateIcon>::FGetter::CreateStatic(&SEditorViewportViewMenu::GetViewMenuLabelIcon, WeakViewport)),
-		FNewMenuDelegate::CreateStatic(&SEditorViewportViewMenu::GenerateViewMenu, WeakToolbar, InRealViewport->BuildFixedEV100Menu(), false),
-		"AnimViewportViewMenu");
-
-	if(bShowShowMenu)
-	{
-		MenuBarBuilder.AddPullDownMenu(
-			LOCTEXT("ShowMenu", "Show"),
-			LOCTEXT("ShowMenuToolTip", "Various options for rendering"),
-			FSlateIcon(),
-			FNewMenuDelegate::CreateSP(this, &SAnimViewportToolBar::GenerateShowMenu),
-			"AnimViewportShowMenu");
-	}
-
-	if (bShowLODMenu)
-	{
-		MenuBarBuilder.AddPullDownMenu(
-			TAttribute<FText>::Create(TAttribute<FText>::FGetter::CreateSP(this, &SAnimViewportToolBar::GetLODMenuLabel)),
-			LOCTEXT("LODMenuToolTip", "Change the current LOD mode"),
-			FSlateIcon(),
-			FNewMenuDelegate::CreateSP(this, &SAnimViewportToolBar::GenerateLODMenu),
-			"AnimViewportLODMenu");
-	}
-
-	if(bShowPlaySpeedMenu)
-	{
-		MenuBarBuilder.AddPullDownMenu(
-			TAttribute<FText>::Create(TAttribute<FText>::FGetter::CreateSP(this, &SAnimViewportToolBar::GetPlaybackMenuLabel)),
-			LOCTEXT("PlabackSpeedMenuToolTip", "Adjust the animation playback speed"),
-			FSlateIcon("EditorStyle", "AnimViewportMenu.PlayBackSpeed"),
-			FNewMenuDelegate::CreateSP(this, &SAnimViewportToolBar::GeneratePlaybackMenu),
-			"AnimViewportPlaybackSpeedMenu");
-	}
-
-	return MenuBarBuilder.MakeWidget();
-}
-
 EVisibility SAnimViewportToolBar::GetTransformToolbarVisibility() const
 {
 	return Viewport.Pin()->CanUseGizmos() ? EVisibility::Visible : EVisibility::Hidden;
 }
 
-void SAnimViewportToolBar::GenerateViewMenu(FMenuBuilder& InMenuBuilder) const
+TSharedRef<SWidget> SAnimViewportToolBar::GenerateViewMenu() const
 {
 	const FAnimViewportMenuCommands& Actions = FAnimViewportMenuCommands::Get();
 
-	InMenuBuilder.PushCommandList(CommandList.ToSharedRef());
-	InMenuBuilder.PushExtender(Extenders.ToSharedRef());
+	TSharedPtr<FExtender> MenuExtender = FExtender::Combine(Extenders);
+
+	const bool bInShouldCloseWindowAfterMenuSelection = true;
+	FMenuBuilder InMenuBuilder(bInShouldCloseWindowAfterMenuSelection, Viewport.Pin()->GetCommandList(), MenuExtender);
+
+	InMenuBuilder.PushCommandList(Viewport.Pin()->GetCommandList().ToSharedRef());
+	InMenuBuilder.PushExtender(MenuExtender.ToSharedRef());
 
 	InMenuBuilder.BeginSection("AnimViewportSceneSetup", LOCTEXT("ViewMenu_SceneSetupLabel", "Scene Setup"));
 	{
+		InMenuBuilder.PushCommandList(Viewport.Pin()->GetCommandList().ToSharedRef());
 		InMenuBuilder.AddMenuEntry(FAnimViewportMenuCommands::Get().PreviewSceneSettings);
+		InMenuBuilder.PopCommandList();
 
 		if(bShowFloorOptions)
 		{
@@ -352,7 +372,9 @@ void SAnimViewportToolBar::GenerateViewMenu(FMenuBuilder& InMenuBuilder) const
 
 	InMenuBuilder.BeginSection("AnimViewportCamera", LOCTEXT("ViewMenu_CameraLabel", "Camera"));
 	{
+		InMenuBuilder.PushCommandList(Viewport.Pin()->GetCommandList().ToSharedRef());
 		InMenuBuilder.AddMenuEntry(FAnimViewportMenuCommands::Get().CameraFollow);
+		InMenuBuilder.PopCommandList();
 		InMenuBuilder.AddMenuEntry(FEditorViewportCommands::Get().FocusViewportToSelection);
 
 		const float FOVMin = 5.f;
@@ -382,22 +404,51 @@ void SAnimViewportToolBar::GenerateViewMenu(FMenuBuilder& InMenuBuilder) const
 
 	InMenuBuilder.BeginSection("AnimViewportDefaultCamera", LOCTEXT("ViewMenu_DefaultCameraLabel", "Default Camera"));
 	{
+		InMenuBuilder.PushCommandList(Viewport.Pin()->GetCommandList().ToSharedRef());
 		InMenuBuilder.AddMenuEntry(FAnimViewportMenuCommands::Get().JumpToDefaultCamera);
 		InMenuBuilder.AddMenuEntry(FAnimViewportMenuCommands::Get().SaveCameraAsDefault);
 		InMenuBuilder.AddMenuEntry(FAnimViewportMenuCommands::Get().ClearDefaultCamera);
+		InMenuBuilder.PopCommandList();
 	}
 	InMenuBuilder.EndSection();
 
 	InMenuBuilder.PopCommandList();
 	InMenuBuilder.PopExtender();
+
+	return InMenuBuilder.MakeWidget();
+}
+
+TSharedRef<SWidget> SAnimViewportToolBar::GeneratePhysicsMenu() const
+{
+	const FAnimViewportShowCommands& Actions = FAnimViewportShowCommands::Get();
+
+	TSharedPtr<FExtender> MenuExtender = FExtender::Combine(Extenders);
+
+	const bool bInShouldCloseWindowAfterMenuSelection = true;
+	FMenuBuilder InMenuBuilder(bInShouldCloseWindowAfterMenuSelection, Viewport.Pin()->GetCommandList(), MenuExtender);
+
+	InMenuBuilder.PushCommandList(Viewport.Pin()->GetCommandList().ToSharedRef());
+	InMenuBuilder.PushExtender(MenuExtender.ToSharedRef());
+	{
+		InMenuBuilder.BeginSection("AnimViewportShowMenu", LOCTEXT("ViewMenu_AnimViewportShowMenu", "Anim Viewport Show Menu"));
+		InMenuBuilder.EndSection();
+	}
+	InMenuBuilder.PopCommandList();
+	InMenuBuilder.PopExtender();
+	return InMenuBuilder.MakeWidget();
 }
 
-void SAnimViewportToolBar::GenerateShowMenu(FMenuBuilder& InMenuBuilder) const
+TSharedRef<SWidget> SAnimViewportToolBar::GenerateShowMenu() const
 {
 	const FAnimViewportShowCommands& Actions = FAnimViewportShowCommands::Get();
 
+	TSharedPtr<FExtender> MenuExtender = FExtender::Combine(Extenders);
+
+	const bool bInShouldCloseWindowAfterMenuSelection = true;
+	FMenuBuilder InMenuBuilder(bInShouldCloseWindowAfterMenuSelection, Viewport.Pin()->GetCommandList(), MenuExtender);
+
 	InMenuBuilder.PushCommandList(Viewport.Pin()->GetCommandList().ToSharedRef());
-	InMenuBuilder.PushExtender(Extenders.ToSharedRef());
+	InMenuBuilder.PushExtender(MenuExtender.ToSharedRef());
 
 	{
 		InMenuBuilder.BeginSection("AnimViewportGeneralShowFlags", LOCTEXT("ViewMenu_GeneralShowFlags", "General Show Flags"));
@@ -519,6 +570,8 @@ void SAnimViewportToolBar::GenerateShowMenu(FMenuBuilder& InMenuBuilder) const
 
 	InMenuBuilder.PopCommandList();
 	InMenuBuilder.PopExtender();
+
+	return InMenuBuilder.MakeWidget();
 }
 
 void SAnimViewportToolBar::FillShowAdvancedMenu(FMenuBuilder& MenuBuilder) const
@@ -619,12 +672,18 @@ FText SAnimViewportToolBar::GetLODMenuLabel() const
 	return Label;
 }
 
-void SAnimViewportToolBar::GenerateLODMenu(FMenuBuilder& InMenuBuilder) const
+TSharedRef<SWidget> SAnimViewportToolBar::GenerateLODMenu() const
 {
 	const FAnimViewportLODCommands& Actions = FAnimViewportLODCommands::Get();
 
+	TSharedPtr<FExtender> MenuExtender = FExtender::Combine(Extenders);
+
+	const bool bInShouldCloseWindowAfterMenuSelection = true;
+	FMenuBuilder InMenuBuilder(bInShouldCloseWindowAfterMenuSelection, Viewport.Pin()->GetCommandList(), MenuExtender);
+
 	InMenuBuilder.PushCommandList(Viewport.Pin()->GetCommandList().ToSharedRef());
-	InMenuBuilder.PushExtender(Extenders.ToSharedRef());
+	InMenuBuilder.PushExtender(MenuExtender.ToSharedRef());
+
 	{
 		// LOD Models
 		InMenuBuilder.BeginSection("AnimViewportPreviewLODs", LOCTEXT("ShowLOD_PreviewLabel", "Preview LODs") );
@@ -649,13 +708,20 @@ void SAnimViewportToolBar::GenerateLODMenu(FMenuBuilder& InMenuBuilder) const
 
 	InMenuBuilder.PopCommandList();
 	InMenuBuilder.PopExtender();
+
+	return InMenuBuilder.MakeWidget();
 }
 
-void SAnimViewportToolBar::GenerateViewportTypeMenu(FMenuBuilder& InMenuBuilder) const
+TSharedRef<SWidget> SAnimViewportToolBar::GenerateViewportTypeMenu() const
 {
+
+	TSharedPtr<FExtender> MenuExtender = FExtender::Combine(Extenders);
+
+	const bool bInShouldCloseWindowAfterMenuSelection = true;
+	FMenuBuilder InMenuBuilder(bInShouldCloseWindowAfterMenuSelection, Viewport.Pin()->GetCommandList(), MenuExtender);
 	InMenuBuilder.SetStyle(&FEditorStyle::Get(), "Menu");
-	InMenuBuilder.PushCommandList(Viewport.Pin()->GetViewportWidget()->GetCommandList().ToSharedRef());
-	InMenuBuilder.PushExtender(Extenders.ToSharedRef());
+	InMenuBuilder.PushCommandList(Viewport.Pin()->GetCommandList().ToSharedRef());
+	InMenuBuilder.PushExtender(MenuExtender.ToSharedRef());
 
 	// Camera types
 	InMenuBuilder.AddMenuEntry(FEditorViewportCommands::Get().Perspective);
@@ -671,14 +737,21 @@ void SAnimViewportToolBar::GenerateViewportTypeMenu(FMenuBuilder& InMenuBuilder)
 
 	InMenuBuilder.PopCommandList();
 	InMenuBuilder.PopExtender();
+
+	return InMenuBuilder.MakeWidget();
 }
 
-void SAnimViewportToolBar::GeneratePlaybackMenu(FMenuBuilder& InMenuBuilder) const
+TSharedRef<SWidget> SAnimViewportToolBar::GeneratePlaybackMenu() const
 {
 	const FAnimViewportPlaybackCommands& Actions = FAnimViewportPlaybackCommands::Get();
 
+	TSharedPtr<FExtender> MenuExtender = FExtender::Combine(Extenders);
+
+	const bool bInShouldCloseWindowAfterMenuSelection = true;
+	FMenuBuilder InMenuBuilder(bInShouldCloseWindowAfterMenuSelection, Viewport.Pin()->GetCommandList(), MenuExtender);
+
 	InMenuBuilder.PushCommandList(Viewport.Pin()->GetCommandList().ToSharedRef());
-	InMenuBuilder.PushExtender(Extenders.ToSharedRef());
+	InMenuBuilder.PushExtender(MenuExtender.ToSharedRef());
 	{
 		// View modes
 		{
@@ -694,6 +767,8 @@ void SAnimViewportToolBar::GeneratePlaybackMenu(FMenuBuilder& InMenuBuilder) con
 	}
 	InMenuBuilder.PopCommandList();
 	InMenuBuilder.PopExtender();
+
+	return InMenuBuilder.MakeWidget();
 }
 
 void SAnimViewportToolBar::GenerateTurnTableMenu(FMenuBuilder& MenuBuilder) const
@@ -702,6 +777,7 @@ void SAnimViewportToolBar::GenerateTurnTableMenu(FMenuBuilder& MenuBuilder) cons
 
 	const bool bInShouldCloseWindowAfterMenuSelection = true;
 
+	MenuBuilder.PushCommandList(Viewport.Pin()->GetCommandList().ToSharedRef());
 	MenuBuilder.BeginSection("AnimViewportTurnTableMode", LOCTEXT("TurnTableMenu_ModeLabel", "Turn Table Mode"));
 	{
 		MenuBuilder.AddMenuEntry(Actions.PersonaTurnTablePlay);
@@ -718,6 +794,7 @@ void SAnimViewportToolBar::GenerateTurnTableMenu(FMenuBuilder& MenuBuilder) cons
 		}
 	}
 	MenuBuilder.EndSection();
+	MenuBuilder.PopCommandList();
 }
 
 FSlateColor SAnimViewportToolBar::GetFontColor() const
@@ -814,47 +891,55 @@ FText SAnimViewportToolBar::GetCameraMenuLabel() const
 	return Label;
 }
 
-FSlateIcon SAnimViewportToolBar::GetCameraMenuLabelIcon() const
+const FSlateBrush* SAnimViewportToolBar::GetCameraMenuLabelIcon() const
 {
 	FName Icon = NAME_None;
 	TSharedPtr< SAnimationEditorViewportTabBody > PinnedViewport(Viewport.Pin());
 	if (PinnedViewport.IsValid())
 	{
+		static FName PerspectiveIcon("EditorViewport.Perspective");
+		static FName TopIcon("EditorViewport.Top");
+		static FName LeftIcon("EditorViewport.Left");
+		static FName FrontIcon("EditorViewport.Front");
+		static FName BottomIcon("EditorViewport.Bottom");
+		static FName RightIcon("EditorViewport.Right");
+		static FName BackIcon("EditorViewport.Back");
+
 		switch (PinnedViewport->GetLevelViewportClient().ViewportType)
 		{
 		case LVT_Perspective:
-			Icon = FName("EditorViewport.Perspective");
+			Icon = PerspectiveIcon;
 			break;
 
 		case LVT_OrthoXY:
-			Icon = FName("EditorViewport.Top");
+			Icon = TopIcon;
 			break;
 
 		case LVT_OrthoYZ:
-			Icon = FName("EditorViewport.Left");
+			Icon = LeftIcon;
 			break;
 
 		case LVT_OrthoXZ:
-			Icon = FName("EditorViewport.Front");
+			Icon = FrontIcon;
 			break;
 
 		case LVT_OrthoNegativeXY:
-			Icon = FName("EditorViewport.Bottom");
+			Icon = BottomIcon;
 			break;
 
 		case LVT_OrthoNegativeYZ:
-			Icon = FName("EditorViewport.Right");
+			Icon = RightIcon;
 			break;
 
 		case LVT_OrthoNegativeXZ:
-			Icon = FName("EditorViewport.Back");
+			Icon = BackIcon;
 			break;
 		case LVT_OrthoFreelook:
 			break;
 		}
 	}
 
-	return FSlateIcon("EditorStyle", Icon);
+	return FEditorStyle::GetBrush(Icon);
 }
 
 float SAnimViewportToolBar::OnGetFOVValue( ) const
diff --git a/Engine/Source/Editor/UnrealEd/Private/Fbx/FbxMainExport.cpp b/Engine/Source/Editor/UnrealEd/Private/Fbx/FbxMainExport.cpp
--- a/Engine/Source/Editor/UnrealEd/Private/Fbx/FbxMainExport.cpp
+++ b/Engine/Source/Editor/UnrealEd/Private/Fbx/FbxMainExport.cpp
@@ -102,7 +102,8 @@ TSharedPtr<FFbxExporter> FFbxExporter::StaticInstance;
 
 FFbxExporter::FFbxExporter()
 {
-	ExportOptions = NewObject<UFbxExportOption>(GetTransientPackage(), NAME_None);
+	//We use the FGCObject pattern to keep the fbx export option alive during the editor session
+	ExportOptions = NewObject<UFbxExportOption>();
 	//Load the option from the user save ini file
 	ExportOptions->LoadOptions();
 
@@ -142,12 +143,10 @@ void FFbxExporter::DeleteInstance()
 void FFbxExporter::FillExportOptions(bool BatchMode, bool bShowOptionDialog, const FString& FullPath, bool& OutOperationCanceled, bool& bOutExportAll)
 {
 	OutOperationCanceled = false;
-
-	if (ExportOptions == nullptr)
-	{
-		ExportOptions = NewObject<UFbxExportOption>(GetTransientPackage(), NAME_None);
-	}
-
+	
+	//Export option should have been set in the constructor
+	check(ExportOptions != nullptr);
+	
 	//Load the option from the user save ini file
 	ExportOptions->LoadOptions();
 	
@@ -3365,7 +3364,7 @@ FbxNode* FFbxExporter::ExportStaticMeshToFbx(const UStaticMesh* StaticMesh, int3
 		else
 		{
 			// Do not weld verts
-			VertRemap.Add(RenderMesh.VertexBuffers.StaticMeshVertexBuffer.GetNumVertices());
+			VertRemap.AddUninitialized(RenderMesh.VertexBuffers.StaticMeshVertexBuffer.GetNumVertices());
 			for (int32 i = 0; i < VertRemap.Num(); i++)
 			{
 				VertRemap[i] = i;
@@ -3711,7 +3710,7 @@ void FFbxExporter::ExportSplineMeshToFbx(const USplineMeshComponent* SplineMeshC
 	else
 	{
 		// Do not weld verts
-		VertRemap.Add(RenderMesh.VertexBuffers.StaticMeshVertexBuffer.GetNumVertices());
+		VertRemap.AddUninitialized(RenderMesh.VertexBuffers.StaticMeshVertexBuffer.GetNumVertices());
 		for (int32 i = 0; i < VertRemap.Num(); i++)
 		{
 			VertRemap[i] = i;
diff --git a/Engine/Source/Editor/UnrealEd/Private/SEditorViewportViewMenu.cpp b/Engine/Source/Editor/UnrealEd/Private/SEditorViewportViewMenu.cpp
--- a/Engine/Source/Editor/UnrealEd/Private/SEditorViewportViewMenu.cpp
+++ b/Engine/Source/Editor/UnrealEd/Private/SEditorViewportViewMenu.cpp
@@ -18,17 +18,17 @@ void SEditorViewportViewMenu::Construct( const FArguments& InArgs, TSharedRef<SE
 		SEditorViewportToolbarMenu::FArguments()
 			.ParentToolBar( InParentToolBar)
 			.Cursor( EMouseCursor::Default )
-			.Label_Static( &SEditorViewportViewMenu::GetViewMenuLabel, Viewport)
-			.LabelIcon( this, &SEditorViewportViewMenu::GetViewMenuLabelBrush)
+			.Label(this, &SEditorViewportViewMenu::GetViewMenuLabel)
+			.LabelIcon(this, &SEditorViewportViewMenu::GetViewMenuLabelIcon)
 			.OnGetMenuContent( this, &SEditorViewportViewMenu::GenerateViewMenuContent )
 	);
 		
 }
 
-FText SEditorViewportViewMenu::GetViewMenuLabel(TWeakPtr<SEditorViewport> InViewport)
+FText SEditorViewportViewMenu::GetViewMenuLabel() const
 {
 	FText Label = LOCTEXT("ViewMenuTitle_Default", "View");
-	TSharedPtr< SEditorViewport > PinnedViewport = InViewport.Pin();
+	TSharedPtr< SEditorViewport > PinnedViewport = Viewport.Pin();
 	if( PinnedViewport.IsValid() )
 	{
 		switch( PinnedViewport->GetViewportClient()->GetViewMode() )
@@ -128,15 +128,10 @@ FText SEditorViewportViewMenu::GetViewMenuLabel(TWeakPtr<SEditorViewport> InView
 	return Label;
 }
 
-const FSlateBrush* SEditorViewportViewMenu::GetViewMenuLabelBrush() const
-{
-	return GetViewMenuLabelIcon(Viewport).GetIcon();
-}
-
-FSlateIcon SEditorViewportViewMenu::GetViewMenuLabelIcon(TWeakPtr<SEditorViewport> InViewport)
+const FSlateBrush* SEditorViewportViewMenu::GetViewMenuLabelIcon() const
 {
 	FName Icon = NAME_None;
-	TSharedPtr< SEditorViewport > PinnedViewport = InViewport.Pin();
+	TSharedPtr< SEditorViewport > PinnedViewport = Viewport.Pin();
 	if( PinnedViewport.IsValid() )
 	{
 		static FName WireframeIcon( "EditorViewport.WireframeMode" );
@@ -259,23 +254,16 @@ FSlateIcon SEditorViewportViewMenu::GetViewMenuLabelIcon(TWeakPtr<SEditorViewpor
 		}
 	}
 
-	return FSlateIcon( "EditorStyle", Icon );
+	return FEditorStyle::GetBrush(Icon);
 }
 
 TSharedRef<SWidget> SEditorViewportViewMenu::GenerateViewMenuContent() const
 {
-	const bool bInShouldCloseWindowAfterMenuSelection = true;
-	TSharedPtr<FEditorViewportClient> EditorViewPostClient = Viewport.Pin()->GetViewportClient();
+	const FEditorViewportCommands& BaseViewportActions = FEditorViewportCommands::Get();
 
+	const bool bInShouldCloseWindowAfterMenuSelection = true;
 	FMenuBuilder ViewMenuBuilder(bInShouldCloseWindowAfterMenuSelection, Viewport.Pin()->GetCommandList(), MenuExtenders);
-	GenerateViewMenu(ViewMenuBuilder, ParentToolBar, Viewport.Pin()->BuildFixedEV100Menu(), EditorViewPostClient.IsValid() && EditorViewPostClient->IsLevelEditorClient());
 
-	return ViewMenuBuilder.MakeWidget();
-}
-
-void SEditorViewportViewMenu::GenerateViewMenu(FMenuBuilder& ViewMenuBuilder, TWeakPtr<SViewportToolBar> ParentToolBar, TSharedRef<SWidget> FixedEV100Menu, bool bIsLevelEditor)
-{
-	const FEditorViewportCommands& BaseViewportActions = FEditorViewportCommands::Get();
 	{
 		// View modes
 		{
@@ -344,12 +332,17 @@ void SEditorViewportViewMenu::GenerateViewMenu(FMenuBuilder& ViewMenuBuilder, TW
 		{
 			const FEditorViewportCommands& BaseViewportCommands = FEditorViewportCommands::Get();
 
+			TSharedRef<SWidget> FixedEV100Menu = Viewport.Pin()->BuildFixedEV100Menu();
+			TSharedPtr<FEditorViewportClient> EditorViewPostClient = Viewport.Pin()->GetViewportClient();
+			const bool bIsLevelEditor = EditorViewPostClient.IsValid() && EditorViewPostClient->IsLevelEditorClient();
+
 			ViewMenuBuilder.BeginSection("Exposure", LOCTEXT("ExposureHeader", "Exposure"));
 			ViewMenuBuilder.AddMenuEntry( bIsLevelEditor ? BaseViewportCommands.ToggleInGameExposure : BaseViewportCommands.ToggleAutoExposure, NAME_None );
 			ViewMenuBuilder.AddWidget( FixedEV100Menu, LOCTEXT("FixedEV100", "EV100") );
 			ViewMenuBuilder.EndSection();
 		}
 	}
+	return ViewMenuBuilder.MakeWidget();
 }
 
 #undef LOCTEXT_NAMESPACE
diff --git a/Engine/Source/Editor/UnrealEd/Public/SEditorViewportViewMenu.h b/Engine/Source/Editor/UnrealEd/Public/SEditorViewportViewMenu.h
--- a/Engine/Source/Editor/UnrealEd/Public/SEditorViewportViewMenu.h
+++ b/Engine/Source/Editor/UnrealEd/Public/SEditorViewportViewMenu.h
@@ -22,15 +22,11 @@ public:
 
 	void Construct( const FArguments& InArgs, TSharedRef<SEditorViewport> InViewport, TSharedRef<class SViewportToolBar> InParentToolBar );
 
-	/** Generate the view menu */
-	static void GenerateViewMenu(FMenuBuilder& InMenuBuilder, TWeakPtr<SViewportToolBar> ParentToolBar, TSharedRef<SWidget> FixedEV100Menu, bool bIsLevelEditor);
-
-	static FText GetViewMenuLabel(TWeakPtr<SEditorViewport> InViewport);
-
-	static FSlateIcon GetViewMenuLabelIcon(TWeakPtr<SEditorViewport> InViewport);
+private:
+	FText GetViewMenuLabel() const;
+	const FSlateBrush* GetViewMenuLabelIcon() const;
 
 protected:
-	const FSlateBrush* GetViewMenuLabelBrush() const;
 	virtual TSharedRef<SWidget> GenerateViewMenuContent() const;
 	TWeakPtr<SEditorViewport> Viewport;
 
diff --git a/Engine/Source/Runtime/Engine/Private/Components/ReflectionCaptureComponent.cpp b/Engine/Source/Runtime/Engine/Private/Components/ReflectionCaptureComponent.cpp
--- a/Engine/Source/Runtime/Engine/Private/Components/ReflectionCaptureComponent.cpp
+++ b/Engine/Source/Runtime/Engine/Private/Components/ReflectionCaptureComponent.cpp
@@ -39,6 +39,9 @@
 #include "Engine/MapBuildDataRegistry.h"
 #include "ComponentRecreateRenderStateContext.h"
 
+// ES3.0+ devices support seamless cubemap filtering, averaging edges will produce artifacts on those devices
+#define MOBILE_AVERAGE_CUBEMAP_EDGES 0 
+
 /** 
  * Size of all reflection captures.
  * Reflection capture derived data versions must be changed if modifying this
@@ -477,8 +480,8 @@ void GenerateEncodedHDRData(const TArray<uint8>& FullHDRData, int32 CubemapSize,
 		const FFloat16Color*	MipSrcData = (const FFloat16Color*)&FullHDRData[SourceMipBaseIndex];
 		FColor*					MipDstData = (FColor*)&OutEncodedHDRData[DestMipBaseIndex];
 
+#if MOBILE_AVERAGE_CUBEMAP_EDGES
 		// Fix cubemap seams by averaging colors across edges
-
 		int32 CornerTable[4] =
 		{
 			0,
@@ -552,6 +555,7 @@ void GenerateEncodedHDRData(const TArray<uint8>& FullHDRData, int32 CubemapSize,
 				FaceDstDataA[ EdgeTexelA ] = FaceDstDataB[ EdgeTexelB ] = RGBMEncode( AvgColor * Brightness );
 			}
 		}
+#endif // MOBILE_AVERAGE_CUBEMAP_EDGES
 		
 		// Encode rest of texels
 		for (int32 CubeFace = 0; CubeFace < CubeFace_MAX; CubeFace++)
@@ -564,9 +568,11 @@ void GenerateEncodedHDRData(const TArray<uint8>& FullHDRData, int32 CubemapSize,
 			// Convert each texel from linear space FP16 to RGBM FColor
 			// Note: Brightness on the capture is baked into the encoded HDR data
 			// Skip edges
-			for( int32 y = 1; y < MipSize - 1; y++ )
+			const int32 SkipEdges = MOBILE_AVERAGE_CUBEMAP_EDGES ? 1 : 0;
+
+			for( int32 y = SkipEdges; y < MipSize - SkipEdges; y++ )
 			{
-				for( int32 x = 1; x < MipSize - 1; x++ )
+				for( int32 x = SkipEdges; x < MipSize - SkipEdges; x++ )
 				{
 					int32 TexelIndex = x + y * MipSize;
 					const FLinearColor LinearColor = FLinearColor( FaceSourceData[ TexelIndex ]) * Brightness;
diff --git a/Engine/Source/Runtime/Engine/Private/InstancedStaticMesh.h b/Engine/Source/Runtime/Engine/Private/InstancedStaticMesh.h
--- a/Engine/Source/Runtime/Engine/Private/InstancedStaticMesh.h
+++ b/Engine/Source/Runtime/Engine/Private/InstancedStaticMesh.h
@@ -85,9 +85,6 @@ public:
 	 */
 	ENGINE_API void InitFromPreallocatedData(UInstancedStaticMeshComponent* InComponent, FStaticMeshInstanceData& Other, bool InRequireCPUAccess);
 
-	/** Propagates instance selection state and hit proxy colors */
-	void SetPerInstanceEditorData(UInstancedStaticMeshComponent* InComponent, const TArray<TRefCountPtr<HHitProxy>>& InHitProxies, int32 InUpdateInstanceStartingIndex, int32 InUpdateInstanceIndexCount);
-
 	/**
 	 * Update the RHI vertex buffer (called on render thread)
 	 * @param InIndexList - List of items to update in the buffer (an empty list means we want to update all the buffer)
@@ -108,6 +105,11 @@ public:
 		return NumInstances;
 	}
 
+	FORCEINLINE uint32 GetCurrentNumInstances() const
+	{
+		return InstanceData->GetNumInstances();
+	}
+
 	FORCEINLINE  void GetInstanceTransform(int32 InstanceIndex, FMatrix& Transform) const
 	{
 		InstanceData->GetInstanceTransform(InstanceIndex, Transform);
@@ -498,6 +500,19 @@ struct FPerInstanceRenderData
 		InstanceBuffer.UpdateInstanceData(InComponent, HitProxies, InUpdateInstanceStartingIndex, InUpdateInstanceIndexCount);
 	}
 
+	void UpdateAllInstanceData(UInstancedStaticMeshComponent* InComponent, bool InUpdateProxyData = true)
+	{
+		QUICK_SCOPE_CYCLE_COUNTER(STAT_FoliageBufferUpdate);
+
+		if (InUpdateProxyData)
+		{
+			AddHitProxyData(InComponent, 0, InComponent->PerInstanceSMData.Num());
+		}
+
+		// Force full refresh of ALL the buffer instance (including the removed one as we might need to re locate them)
+		InstanceBuffer.UpdateInstanceData(InComponent, HitProxies, 0, FMath::Max((int32)InstanceBuffer.GetNumInstances(), InComponent->PerInstanceSMData.Num()));
+	}
+
 	/**
 	 * Remove a single instance from the instance buffer
 	 * @param InComponent - The owning component.
diff --git a/Engine/Source/Runtime/Engine/Public/StaticMeshResources.h b/Engine/Source/Runtime/Engine/Public/StaticMeshResources.h
--- a/Engine/Source/Runtime/Engine/Public/StaticMeshResources.h
+++ b/Engine/Source/Runtime/Engine/Public/StaticMeshResources.h
@@ -722,6 +722,7 @@ class FStaticMeshInstanceData
 		F InstanceTransform1[4];
 		F InstanceTransform2[4];
 		F InstanceTransform3[4];
+		bool IsUsed;
 
 		friend FArchive& operator<<(FArchive& Ar, FInstanceTransformMatrix& V)
 		{
@@ -882,6 +883,37 @@ public:
 		GetInstanceOriginInternal(InstanceIndex, InstanceOrigin);
 	}
 
+	FORCEINLINE int32 GetNextAvailableInstanceIndex() const
+	{
+		for (int32 i = 0; i < NumInstances; ++i)
+		{
+			if (bUseHalfFloat)
+			{
+				FInstanceTransformMatrix<FFloat16>* ElementData = reinterpret_cast<FInstanceTransformMatrix<FFloat16>*>(InstanceTransformData[i].GetDataPointer());
+				check((void*)((&ElementData[i]) + 1) <= (void*)(InstanceTransformDataPtr + InstanceTransformData->GetResourceSize()));
+				check((void*)((&ElementData[i]) + 0) >= (void*)(InstanceTransformDataPtr));
+
+				if (!ElementData->IsUsed)
+				{
+					return i;
+				}
+			}
+			else
+			{
+				FInstanceTransformMatrix<float>* ElementData = reinterpret_cast<FInstanceTransformMatrix<float>*>(InstanceTransformData[i].GetDataPointer());
+				check((void*)((&ElementData[i]) + 1) <= (void*)(InstanceTransformDataPtr + InstanceTransformData->GetResourceSize()));
+				check((void*)((&ElementData[i]) + 0) >= (void*)(InstanceTransformDataPtr));
+
+				if (!ElementData->IsUsed)
+				{
+					return i;
+				}
+			}
+		}
+
+		return INDEX_NONE;
+	}
+
 	FORCEINLINE void SetInstance(int32 InstanceIndex, const FMatrix& Transform, float RandomInstanceID)
 	{
 		FVector4 Origin(Transform.M[3][0], Transform.M[3][1], Transform.M[3][2], RandomInstanceID);
@@ -894,11 +926,11 @@ public:
 
 		if (bUseHalfFloat)
 		{
-			SetInstanceTransformInternal<FFloat16>(InstanceIndex, InstanceTransform);
+			SetInstanceTransformInternal<FFloat16>(InstanceIndex, InstanceTransform, true);
 		}
 		else
 		{
-			SetInstanceTransformInternal<float>(InstanceIndex, InstanceTransform);
+			SetInstanceTransformInternal<float>(InstanceIndex, InstanceTransform, true);
 		}
 
 		SetInstanceLightMapDataInternal(InstanceIndex, FVector4(0, 0, 0, 0));
@@ -916,11 +948,11 @@ public:
 
 		if (bUseHalfFloat)
 		{
-			SetInstanceTransformInternal<FFloat16>(InstanceIndex, InstanceTransform);
+			SetInstanceTransformInternal<FFloat16>(InstanceIndex, InstanceTransform, true);
 		}
 		else
 		{
-			SetInstanceTransformInternal<float>(InstanceIndex, InstanceTransform);
+			SetInstanceTransformInternal<float>(InstanceIndex, InstanceTransform, true);
 		}
 
 		SetInstanceLightMapDataInternal(InstanceIndex, FVector4(LightmapUVBias.X, LightmapUVBias.Y, ShadowmapUVBias.X, ShadowmapUVBias.Y));
@@ -937,11 +969,11 @@ public:
 
 		if (bUseHalfFloat)
 		{
-			SetInstanceTransformInternal<FFloat16>(InstanceIndex, InstanceTransform);
+			SetInstanceTransformInternal<FFloat16>(InstanceIndex, InstanceTransform, false);
 		}
 		else
 		{
-			SetInstanceTransformInternal<float>(InstanceIndex, InstanceTransform);
+			SetInstanceTransformInternal<float>(InstanceIndex, InstanceTransform, false);
 		}
 
 		SetInstanceLightMapDataInternal(InstanceIndex, FVector4(0, 0, 0, 0));
@@ -956,7 +988,7 @@ public:
 			InstanceTransform[0][3] = ((float)HitProxyColor.R) + (bSelected ? 256.f : 0.0f);
 			InstanceTransform[1][3] = (float)HitProxyColor.G;
 			InstanceTransform[2][3] = (float)HitProxyColor.B;
-			SetInstanceTransformInternal<FFloat16>(InstanceIndex, InstanceTransform);
+			SetInstanceTransformInternal<FFloat16>(InstanceIndex, InstanceTransform, true);
 		}
 		else
 		{
@@ -964,7 +996,7 @@ public:
 			InstanceTransform[0][3] = ((float)HitProxyColor.R) + (bSelected ? 256.f : 0.0f);
 			InstanceTransform[1][3] = (float)HitProxyColor.G;
 			InstanceTransform[2][3] = (float)HitProxyColor.B;
-			SetInstanceTransformInternal<float>(InstanceIndex, InstanceTransform);
+			SetInstanceTransformInternal<float>(InstanceIndex, InstanceTransform, true);
 		}
 	}
 
@@ -1118,7 +1150,7 @@ private:
 	}
 
 	template<typename T>
-	FORCEINLINE void SetInstanceTransformInternal(int32 InstanceIndex, FVector4(Transform)[3]) const
+	FORCEINLINE void SetInstanceTransformInternal(int32 InstanceIndex, FVector4(Transform)[3], bool InIsUsed) const
 	{
 		FInstanceTransformMatrix<T>* ElementData = reinterpret_cast<FInstanceTransformMatrix<T>*>(InstanceTransformDataPtr);
 		check((void*)((&ElementData[InstanceIndex]) + 1) <= (void*)(InstanceTransformDataPtr + InstanceTransformData->GetResourceSize()));
@@ -1138,6 +1170,8 @@ private:
 		ElementData[InstanceIndex].InstanceTransform3[1] = Transform[2][1];
 		ElementData[InstanceIndex].InstanceTransform3[2] = Transform[2][2];
 		ElementData[InstanceIndex].InstanceTransform3[3] = Transform[2][3];
+
+		ElementData->IsUsed = InIsUsed;
 	}
 
 	FORCEINLINE void SetInstanceOriginInternal(int32 InstanceIndex, const FVector4& Origin) const
