Fix for Timed Notifies scrub incorrectly on timeline on Mac.  This is also an issue if running in Hi DPI on windows.

diff --git a/Engine/Source/Editor/Persona/Private/SAnimNotifyPanel.cpp b/Engine/Source/Editor/Persona/Private/SAnimNotifyPanel.cpp
--- a/Engine/Source/Editor/Persona/Private/SAnimNotifyPanel.cpp
+++ b/Engine/Source/Editor/Persona/Private/SAnimNotifyPanel.cpp
@@ -1601,7 +1601,7 @@ void SAnimNotifyNode::UpdateSizeAndPosition(const FGeometry& AllottedGeometry)
 	FTrackScaleInfo ScaleInfo(ViewInputMin.Get(), ViewInputMax.Get(), 0, 0, AllottedGeometry.Size);
 
 	// Cache the geometry information, the alloted geometry is the same size as the track.
-	CachedAllotedGeometrySize = AllottedGeometry.Size;
+	CachedAllotedGeometrySize = AllottedGeometry.Size * AllottedGeometry.Scale;
 
 	NotifyTimePositionX = ScaleInfo.InputToLocalX(NodeObjectInterface->GetTime());
 	NotifyDurationSizeX = ScaleInfo.PixelsPerInput * NodeObjectInterface->GetDuration();
