Fixed incorrect preview mesh being applied after retargeting

On first setup, the preview mesh was obtained using a legacy path that didnt use the skeleton fallback

diff --git a/Engine/Source/Editor/Persona/Private/PersonaToolkit.cpp b/Engine/Source/Editor/Persona/Private/PersonaToolkit.cpp
--- a/Engine/Source/Editor/Persona/Private/PersonaToolkit.cpp
+++ b/Engine/Source/Editor/Persona/Private/PersonaToolkit.cpp
@@ -136,27 +136,7 @@ void FPersonaToolkit::CreatePreviewScene(const FPersonaToolkitArgs& PersonaToolk
 		bool bSetMesh = false;
 
 		// Set the mesh
-		// first check assets first
-		// mesh exists with anim BP
-		if (AnimationAsset != nullptr)
-		{
-			USkeletalMesh* AssetMesh = AnimationAsset->GetPreviewMesh();
-			if (AssetMesh)
-			{
-				PreviewScene->SetPreviewMesh(AssetMesh);
-				bSetMesh = true;
-			}
-		}
-		else if (AnimBlueprint != nullptr)
-		{
-			USkeletalMesh* AssetMesh = AnimBlueprint->GetPreviewMesh();
-			if (AssetMesh)
-			{
-				PreviewScene->SetPreviewMesh(AssetMesh);
-				bSetMesh = true;
-			}
-		}
-		else if (Mesh != nullptr)
+		if (Mesh)
 		{
 			PreviewScene->SetPreviewMesh(Mesh);
 			bSetMesh = true;
