Fix blend sample context menu appering offset on high dpi monitors

diff --git a/Engine/Source/Editor/Persona/Private/SAnimationBlendSpaceGridWidget.cpp b/Engine/Source/Editor/Persona/Private/SAnimationBlendSpaceGridWidget.cpp
--- a/Engine/Source/Editor/Persona/Private/SAnimationBlendSpaceGridWidget.cpp
+++ b/Engine/Source/Editor/Persona/Private/SAnimationBlendSpaceGridWidget.cpp
@@ -751,7 +751,7 @@ FReply SBlendSpaceGridWidget::ProcessClick(const FGeometry& MyGeometry, const FP
 				const FVector2D MousePosition = MouseEvent.GetScreenSpacePosition();
 				// This is of a fixed size atm since MenuContent->GetDesiredSize() will not take the detail customization into account and return an incorrect (small) size
 				const FVector2D ExpectedSize(300, 100);				
-				const FVector2D MenuPosition = FSlateApplication::Get().CalculatePopupWindowPosition(FSlateRect(MousePosition.X, MousePosition.Y, MousePosition.X, MousePosition.Y), ExpectedSize);
+				const FVector2D MenuPosition = FSlateApplication::Get().CalculatePopupWindowPosition(FSlateRect(MousePosition.X, MousePosition.Y, MousePosition.X, MousePosition.Y), ExpectedSize, false);
 
 				FSlateApplication::Get().PushMenu(
 					AsShared(),
