Make sure the material slot name are unique when the user add one slot and assign a material. For both staticmesh and skeletal mesh editor

diff --git a/Engine/Source/Editor/Persona/Private/PersonaMeshDetails.cpp b/Engine/Source/Editor/Persona/Private/PersonaMeshDetails.cpp
--- a/Engine/Source/Editor/Persona/Private/PersonaMeshDetails.cpp
+++ b/Engine/Source/Editor/Persona/Private/PersonaMeshDetails.cpp
@@ -1784,12 +1784,43 @@ void FPersonaMeshDetails::OnMaterialArrayChanged(UMaterialInterface* NewMaterial
 			{
 				if (Mesh->Materials[SlotIndex].MaterialSlotName == NAME_None)
 				{
-					Mesh->Materials[SlotIndex].MaterialSlotName = FName(*(NewMaterial->GetName()));
+					
+					Mesh->Materials[SlotIndex].MaterialSlotName = NewMaterial->GetFName();
 				}
+
+				//Ensure the imported material slot name is unique
 				if (Mesh->Materials[SlotIndex].ImportedMaterialSlotName == NAME_None)
 				{
-					//Add an imported material slot name so when we reimport we can keep the user changes
-					Mesh->Materials[SlotIndex].ImportedMaterialSlotName = Mesh->Materials[SlotIndex].MaterialSlotName;
+					auto IsMaterialNameUnique = [&Mesh, SlotIndex](const FName TestName)
+					{
+						for (int32 MaterialIndex = 0; MaterialIndex < Mesh->Materials.Num(); ++MaterialIndex)
+						{
+							if (MaterialIndex == SlotIndex)
+							{
+								continue;
+							}
+							if (Mesh->Materials[MaterialIndex].ImportedMaterialSlotName == TestName)
+							{
+								return false;
+							}
+						}
+						return true;
+					};
+					int32 MatchNameCounter = 0;
+					//Make sure the name is unique for imported material slot name
+					bool bUniqueName = false;
+					FString MaterialSlotName = NewMaterial->GetName();
+					while (!bUniqueName)
+					{
+						bUniqueName = true;
+						if (!IsMaterialNameUnique(FName(*MaterialSlotName)))
+						{
+							bUniqueName = false;
+							MatchNameCounter++;
+							MaterialSlotName = NewMaterial->GetName() + TEXT("_") + FString::FromInt(MatchNameCounter);
+						}
+					}
+					Mesh->Materials[SlotIndex].ImportedMaterialSlotName = FName(*MaterialSlotName);
 				}
 			}
 		}
diff --git a/Engine/Source/Runtime/Engine/Private/StaticMesh.cpp b/Engine/Source/Runtime/Engine/Private/StaticMesh.cpp
--- a/Engine/Source/Runtime/Engine/Private/StaticMesh.cpp
+++ b/Engine/Source/Runtime/Engine/Private/StaticMesh.cpp
@@ -4222,10 +4222,41 @@ void UStaticMesh::SetMaterial(int32 MaterialIndex, UMaterialInterface* NewMateri
 			{
 				StaticMaterials[MaterialIndex].MaterialSlotName = NewMaterial->GetFName();
 			}
-			//Set the original fbx material name so we can re-import correctly
+			
+			//Set the original fbx material name so we can re-import correctly, ensure the name is unique
 			if (StaticMaterials[MaterialIndex].ImportedMaterialSlotName == NAME_None)
 			{
-				StaticMaterials[MaterialIndex].ImportedMaterialSlotName = NewMaterial->GetFName();
+				auto IsMaterialNameUnique = [this, MaterialIndex](const FName TestName)
+				{
+					for (int32 MatIndex = 0; MatIndex < StaticMaterials.Num(); ++MatIndex)
+					{
+						if (MatIndex == MaterialIndex)
+						{
+							continue;
+						}
+						if (StaticMaterials[MatIndex].ImportedMaterialSlotName == TestName)
+						{
+							return false;
+						}
+					}
+					return true;
+				};
+
+				int32 MatchNameCounter = 0;
+				//Make sure the name is unique for imported material slot name
+				bool bUniqueName = false;
+				FString MaterialSlotName = NewMaterial->GetName();
+				while (!bUniqueName)
+				{
+					bUniqueName = true;
+					if (!IsMaterialNameUnique(FName(*MaterialSlotName)))
+					{
+						bUniqueName = false;
+						MatchNameCounter++;
+						MaterialSlotName = NewMaterial->GetName() + TEXT("_") + FString::FromInt(MatchNameCounter);
+					}
+				}
+				StaticMaterials[MaterialIndex].ImportedMaterialSlotName = FName(*MaterialSlotName);
 			}
 		}
 
