Default coordinate space for animation editors is now local space

diff --git a/Engine/Source/Editor/Persona/Private/AnimationEditorViewportClient.cpp b/Engine/Source/Editor/Persona/Private/AnimationEditorViewportClient.cpp
--- a/Engine/Source/Editor/Persona/Private/AnimationEditorViewportClient.cpp
+++ b/Engine/Source/Editor/Persona/Private/AnimationEditorViewportClient.cpp
@@ -93,6 +93,9 @@ FAnimationViewportClient::FAnimationViewportClient(const TSharedRef<IPersonaPrev
 	((FAssetEditorModeManager*)ModeTools)->SetPreviewScene(&InPreviewScene.Get());
 	((FAssetEditorModeManager*)ModeTools)->SetDefaultMode(FPersonaEditModes::SkeletonSelection);
 
+	// Default to local space
+	SetWidgetCoordSystemSpace(COORD_Local);
+
 	// load config
 	ConfigOption = UPersonaOptions::StaticClass()->GetDefaultObject<UPersonaOptions>();
 	check (ConfigOption);
