Updated notify add/replace menus to use class pickers for better searching of BP and native notifies

Also updated class picker to allow access to the native parent class for unloaded Blueprints
Improved skeleton notify management:
- Added 'Add' option to the manager
- Editor no longer freezes when selecting a notify (it was scanning all animations on the skeleton to populate the details panel).

diff --git a/Engine/Source/Editor/AnimationEditor/Private/AnimationEditor.cpp b/Engine/Source/Editor/AnimationEditor/Private/AnimationEditor.cpp
--- a/Engine/Source/Editor/AnimationEditor/Private/AnimationEditor.cpp
+++ b/Engine/Source/Editor/AnimationEditor/Private/AnimationEditor.cpp
@@ -394,9 +394,8 @@ TSharedPtr<SDockTab> FAnimationEditor::OpenNewAnimationDocumentTab(UAnimationAss
 	{
 		FString	DocumentLink;
 
-		FAnimDocumentArgs Args(PersonaToolkit->GetPreviewScene(), GetPersonaToolkit(), GetSkeletonTree()->GetEditableSkeleton(), OnPostUndo, OnChangeAnimNotifies, OnSectionsChanged);
+		FAnimDocumentArgs Args(PersonaToolkit->GetPreviewScene(), GetPersonaToolkit(), GetSkeletonTree()->GetEditableSkeleton(), OnPostUndo, OnSectionsChanged);
 		Args.OnDespatchObjectsSelected = FOnObjectsSelected::CreateSP(this, &FAnimationEditor::HandleObjectsSelected);
-		Args.OnDespatchAnimNotifiesChanged = FSimpleDelegate::CreateSP(this, &FAnimationEditor::HandleAnimNotifiesChanged);
 		Args.OnDespatchInvokeTab = FOnInvokeTab::CreateSP(this, &FAssetEditorToolkit::InvokeTab);
 		Args.OnDespatchSectionsChanged = FSimpleDelegate::CreateSP(this, &FAnimationEditor::HandleSectionsChanged);
 
@@ -462,11 +461,6 @@ TSharedPtr<SDockTab> FAnimationEditor::OpenNewAnimationDocumentTab(UAnimationAss
 	return OpenedTab;
 }
 
-void FAnimationEditor::HandleAnimNotifiesChanged()
-{
-	OnChangeAnimNotifies.Broadcast();
-}
-
 void FAnimationEditor::HandleSectionsChanged()
 {
 	OnSectionsChanged.Broadcast();
diff --git a/Engine/Source/Editor/AnimationEditor/Private/AnimationEditor.h b/Engine/Source/Editor/AnimationEditor/Private/AnimationEditor.h
--- a/Engine/Source/Editor/AnimationEditor/Private/AnimationEditor.h
+++ b/Engine/Source/Editor/AnimationEditor/Private/AnimationEditor.h
@@ -114,8 +114,6 @@ private:
 		Max
 	};
 
-	void HandleAnimNotifiesChanged();
-
 	void HandleSectionsChanged();
 
 	bool HasValidAnimationSequence() const;
@@ -181,9 +179,6 @@ private:
 
 	bool RecordMeshToAnimation(USkeletalMeshComponent* PreviewComponent, UAnimSequence* NewAsset) const;
 public:
-	/** Multicast delegate fired on anim notifies changing */
-	FSimpleMulticastDelegate OnChangeAnimNotifies;
-
 	/** Multicast delegate fired on global undo/redo */
 	FSimpleMulticastDelegate OnPostUndo;
 
diff --git a/Engine/Source/Editor/AnimationEditor/Private/AnimationEditorMode.cpp b/Engine/Source/Editor/AnimationEditor/Private/AnimationEditorMode.cpp
--- a/Engine/Source/Editor/AnimationEditor/Private/AnimationEditorMode.cpp
+++ b/Engine/Source/Editor/AnimationEditor/Private/AnimationEditorMode.cpp
@@ -34,7 +34,7 @@ FAnimationEditorMode::FAnimationEditorMode(TSharedRef<FWorkflowCentricApplicatio
 	TabFactories.RegisterFactory(PersonaModule.CreateAssetDetailsTabFactory(InHostingApp, FOnGetAsset::CreateSP(&AnimationEditor.Get(), &FAnimationEditor::HandleGetAsset), FOnDetailsCreated()));
 	TabFactories.RegisterFactory(PersonaModule.CreateCurveViewerTabFactory(InHostingApp, InSkeletonTree->GetEditableSkeleton(), AnimationEditor->GetPersonaToolkit()->GetPreviewScene(), AnimationEditor->OnPostUndo, OnObjectsSelected));
 	TabFactories.RegisterFactory(PersonaModule.CreateSkeletonSlotNamesTabFactory(InHostingApp, InSkeletonTree->GetEditableSkeleton(), AnimationEditor->OnPostUndo, FOnObjectSelected::CreateSP(&AnimationEditor.Get(), &FAnimationEditor::HandleObjectSelected)));
-	TabFactories.RegisterFactory(PersonaModule.CreateAnimNotifiesTabFactory(InHostingApp, InSkeletonTree->GetEditableSkeleton(), AnimationEditor->OnChangeAnimNotifies, AnimationEditor->OnPostUndo, OnObjectsSelected));
+	TabFactories.RegisterFactory(PersonaModule.CreateAnimNotifiesTabFactory(InHostingApp, InSkeletonTree->GetEditableSkeleton(), OnObjectsSelected));
 
 	TabLayout = FTabManager::NewLayout("Standalone_AnimationEditor_Layout_v1.2")
 		->AddArea
diff --git a/Engine/Source/Editor/ClassViewer/Private/UnloadedBlueprintData.cpp b/Engine/Source/Editor/ClassViewer/Private/UnloadedBlueprintData.cpp
--- a/Engine/Source/Editor/ClassViewer/Private/UnloadedBlueprintData.cpp
+++ b/Engine/Source/Editor/ClassViewer/Private/UnloadedBlueprintData.cpp
@@ -96,6 +96,27 @@ const UClass* FUnloadedBlueprintData::GetClassWithin() const
 	return nullptr;
 }
 
+const UClass* FUnloadedBlueprintData::GetNativeParent() const
+{
+	TSharedPtr<FClassViewerNode> CurrentNode = ClassViewerNode.Pin()->ParentNode.Pin();
+
+	while (CurrentNode.IsValid())
+	{
+		// The class field will be invalid for unloaded classes.
+		// However, it should be valid once we've hit a loaded class or a Native class.
+		// Assuming BP cannot change ClassWithin data, this should be safe.
+		if (CurrentNode->Class.IsValid() && CurrentNode->Class->HasAnyClassFlags(CLASS_Native))
+		{
+			return CurrentNode->Class.Get();
+		}
+
+		CurrentNode = CurrentNode->ParentNode.Pin();
+	}
+
+	return nullptr;
+}
+
+
 const TWeakPtr<FClassViewerNode>& FUnloadedBlueprintData::GetClassViewerNode() const
 {
 	return ClassViewerNode;
diff --git a/Engine/Source/Editor/ClassViewer/Private/UnloadedBlueprintData.h b/Engine/Source/Editor/ClassViewer/Private/UnloadedBlueprintData.h
--- a/Engine/Source/Editor/ClassViewer/Private/UnloadedBlueprintData.h
+++ b/Engine/Source/Editor/ClassViewer/Private/UnloadedBlueprintData.h
@@ -26,6 +26,8 @@ public:
 
 	virtual const UClass* GetClassWithin() const override;
 
+	virtual const UClass* GetNativeParent() const override;
+
 	/** Retrieves the Class Viewer node this data is associated with. */
 	const TWeakPtr<FClassViewerNode>& GetClassViewerNode() const;
 
diff --git a/Engine/Source/Editor/ClassViewer/Public/ClassViewerFilter.h b/Engine/Source/Editor/ClassViewer/Public/ClassViewerFilter.h
--- a/Engine/Source/Editor/ClassViewer/Public/ClassViewerFilter.h
+++ b/Engine/Source/Editor/ClassViewer/Public/ClassViewerFilter.h
@@ -220,4 +220,11 @@ public:
 	 * @return ClassWithin of the child most Native class in the hierarchy.
 	 */
 	virtual const UClass* GetClassWithin() const = 0;
+
+	/**
+	 * Attempts to get the child-most Native class in the hierarchy.
+	 *
+	 * @return The child-most Native class in the hierarchy.
+	 */
+	virtual const UClass* GetNativeParent() const = 0;
 };
diff --git a/Engine/Source/Editor/DetailCustomizations/Private/SkeletonNotifyDetails.cpp b/Engine/Source/Editor/DetailCustomizations/Private/SkeletonNotifyDetails.cpp
--- a/Engine/Source/Editor/DetailCustomizations/Private/SkeletonNotifyDetails.cpp
+++ b/Engine/Source/Editor/DetailCustomizations/Private/SkeletonNotifyDetails.cpp
@@ -13,6 +13,12 @@
 #include "IDetailPropertyRow.h"
 #include "DetailCategoryBuilder.h"
 #include "IDetailsView.h"
+#include "Widgets/Input/SButton.h"
+#include "Misc/ScopedSlowTask.h"
+#include "IEditableSkeleton.h"
+#include "Animation/AnimSequenceBase.h"
+#include "Widgets/Layout/SBox.h"
+#include "Widgets/Layout/SScrollBox.h"
 
 #define LOCTEXT_NAMESPACE "SkeletonNotifyDetails"
 
@@ -45,6 +51,8 @@ void FSkeletonNotifyDetails::CustomizeDetails( IDetailLayoutBuilder& DetailBuild
 
 	if(EdObj)
 	{
+		NotifyObject = EdObj;
+
 		Category.AddCustomRow(LOCTEXT("AnimationsLabel","Animations"))
 		.NameContent()
 		[
@@ -55,9 +63,33 @@ void FSkeletonNotifyDetails::CustomizeDetails( IDetailLayoutBuilder& DetailBuild
 		]
 		.ValueContent()
 		[
-			SNew(SListView<TSharedPtr<FString>>)
-			.ListItemsSource(&EdObj->AnimationNames)
-			.OnGenerateRow(this, &FSkeletonNotifyDetails::MakeAnimationRow)
+			SNew(SHorizontalBox)
+			+SHorizontalBox::Slot()
+			.AutoWidth()
+			[
+				SNew(SButton)
+				.Visibility_Lambda([this](){ return AnimationNames.Num() > 0 ? EVisibility::Collapsed : EVisibility::Visible; })
+				.Text(LOCTEXT("ScanForAnimations", "Scan"))
+				.ToolTipText(LOCTEXT("ScanForAnimationsTooltip", "Scan for animations that reference this notify"))
+				.OnClicked(this, &FSkeletonNotifyDetails::CollectSequencesUsingNotify)
+			]
+			+SHorizontalBox::Slot()
+			.AutoWidth()
+			[
+				SNew(SBox)
+				.MaxDesiredHeight(300.0f)
+				.MaxDesiredWidth(200.0f)
+				[
+					SNew(SScrollBox)
+					+SScrollBox::Slot()
+					[
+						SAssignNew(ListView, SListView<TSharedPtr<FString>>)
+						.Visibility_Lambda([this](){ return AnimationNames.Num() > 0 ? EVisibility::Visible : EVisibility::Collapsed; })
+						.ListItemsSource(&AnimationNames)
+						.OnGenerateRow(this, &FSkeletonNotifyDetails::MakeAnimationRow)
+					]
+				]
+			]
 		];
 	}
 }
@@ -66,8 +98,59 @@ TSharedRef< ITableRow > FSkeletonNotifyDetails::MakeAnimationRow( TSharedPtr<FSt
 {
 	return SNew(STableRow<TSharedPtr<FString>>, OwnerTable)
 	[
-		SNew(STextBlock).Text(FText::FromString(*Item.Get()))
+		SNew(STextBlock)
+		.ToolTipText(FText::FromString(*Item.Get()))
+		.Text(FText::FromString(*Item.Get()))
 	];
 }
 
+FReply FSkeletonNotifyDetails::CollectSequencesUsingNotify()
+{
+	if(NotifyObject.IsValid() && NotifyObject->EditableSkeleton.IsValid())
+	{
+		FScopedSlowTask SlowTask(1.0f, FText::Format(LOCTEXT("ScanningAnimationMessage", "Looking for animations that reference notify '{0}'."), FText::FromName(NotifyObject->Name)));
+		SlowTask.MakeDialog(true);
+
+		AnimationNames.Empty();
+
+		TArray<FAssetData> CompatibleAnimSequences;
+		NotifyObject->EditableSkeleton.Pin()->GetCompatibleAnimSequences(CompatibleAnimSequences);
+		SlowTask.TotalAmountOfWork = CompatibleAnimSequences.Num();
+
+		for( int32 AssetIndex = 0; AssetIndex < CompatibleAnimSequences.Num(); ++AssetIndex )
+		{
+			SlowTask.EnterProgressFrame(1.0f);
+
+			if(SlowTask.ShouldCancel())
+			{
+				break;
+			}
+
+			const FAssetData& PossibleAnimSequence = CompatibleAnimSequences[AssetIndex];
+			if(UObject* AnimSeqAsset = PossibleAnimSequence.GetAsset())
+			{
+				UAnimSequenceBase* Sequence = CastChecked<UAnimSequenceBase>(AnimSeqAsset);
+				for (int32 NotifyIndex = 0; NotifyIndex < Sequence->Notifies.Num(); ++NotifyIndex)
+				{
+					FAnimNotifyEvent& NotifyEvent = Sequence->Notifies[NotifyIndex];
+					if (NotifyEvent.NotifyName == NotifyObject->Name)
+					{
+						AnimationNames.Add(MakeShareable(new FString(PossibleAnimSequence.AssetName.ToString())));
+						break;
+					}
+				}
+			}
+		}
+
+		if(SlowTask.ShouldCancel())
+		{
+			AnimationNames.Empty();
+		}
+
+		ListView->RequestListRefresh();
+	}
+
+	return FReply::Handled();
+}
+
 #undef LOCTEXT_NAMESPACE
diff --git a/Engine/Source/Editor/DetailCustomizations/Private/SkeletonNotifyDetails.h b/Engine/Source/Editor/DetailCustomizations/Private/SkeletonNotifyDetails.h
--- a/Engine/Source/Editor/DetailCustomizations/Private/SkeletonNotifyDetails.h
+++ b/Engine/Source/Editor/DetailCustomizations/Private/SkeletonNotifyDetails.h
@@ -4,10 +4,12 @@
 
 #include "CoreMinimal.h"
 #include "IDetailCustomization.h"
+#include "Widgets/Views/SListView.h"
 
 class IDetailLayoutBuilder;
 class ITableRow;
 class STableViewBase;
+class UEditorSkeletonNotifyObj;
 
 class FSkeletonNotifyDetails : public IDetailCustomization
 {
@@ -20,4 +22,18 @@ public:
 
 	/** Delegate to handle creating rows for the animations slate list */
 	TSharedRef< ITableRow > MakeAnimationRow( TSharedPtr<FString> Item, const TSharedRef< STableViewBase >& OwnerTable );
+
+private:
+	/** Look for all aniamtions that reference our notify */
+	FReply CollectSequencesUsingNotify();
+
+private:
+	/** The object we are customizing */
+	TWeakObjectPtr<UEditorSkeletonNotifyObj> NotifyObject;
+
+	/** The names of any animations that reference the notify we are displaying */
+	TArray< TSharedPtr<FString> > AnimationNames;
+
+	/** The list view widget */
+	TSharedPtr<SListView<TSharedPtr<FString>>> ListView;
 };
diff --git a/Engine/Source/Editor/Persona/Private/PersonaModule.cpp b/Engine/Source/Editor/Persona/Private/PersonaModule.cpp
--- a/Engine/Source/Editor/Persona/Private/PersonaModule.cpp
+++ b/Engine/Source/Editor/Persona/Private/PersonaModule.cpp
@@ -271,9 +271,9 @@ void FPersonaModule::RegisterPersonaViewportTabFactories(FWorkflowAllowedTabSet&
 	TabSet.RegisterFactory(MakeShareable(new FPreviewViewportSummoner(InHostingApp, InArgs, 3)));
 }
 
-TSharedRef<class FWorkflowTabFactory> FPersonaModule::CreateAnimNotifiesTabFactory(const TSharedRef<class FWorkflowCentricApplication>& InHostingApp, const TSharedRef<class IEditableSkeleton>& InEditableSkeleton, FSimpleMulticastDelegate& InOnChangeAnimNotifies, FSimpleMulticastDelegate& InOnPostUndo, FOnObjectsSelected InOnObjectsSelected) const
+TSharedRef<class FWorkflowTabFactory> FPersonaModule::CreateAnimNotifiesTabFactory(const TSharedRef<class FWorkflowCentricApplication>& InHostingApp, const TSharedRef<class IEditableSkeleton>& InEditableSkeleton, FOnObjectsSelected InOnObjectsSelected) const
 {
-	return MakeShareable(new FSkeletonAnimNotifiesSummoner(InHostingApp, InEditableSkeleton, InOnChangeAnimNotifies, InOnPostUndo, InOnObjectsSelected));
+	return MakeShareable(new FSkeletonAnimNotifiesSummoner(InHostingApp, InEditableSkeleton, InOnObjectsSelected));
 }
 
 TSharedRef<class FWorkflowTabFactory> FPersonaModule::CreateCurveViewerTabFactory(const TSharedRef<class FWorkflowCentricApplication>& InHostingApp, const TSharedRef<class IEditableSkeleton>& InEditableSkeleton, const TSharedRef<IPersonaPreviewScene>& InPreviewScene, FSimpleMulticastDelegate& InOnPostUndo, FOnObjectsSelected InOnObjectsSelected) const
@@ -328,34 +328,31 @@ TSharedRef<SWidget> FPersonaModule::CreateEditorWidgetForAnimDocument(const TSha
 	{
 		if (UAnimSequence* Sequence = Cast<UAnimSequence>(InAnimAsset))
 		{
-			Result = SNew(SSequenceEditor, InArgs.PreviewScene.Pin().ToSharedRef(), InArgs.EditableSkeleton.Pin().ToSharedRef(), InArgs.OnPostUndo)
+			Result = SNew(SSequenceEditor, InArgs.PreviewScene.Pin().ToSharedRef(), InArgs.EditableSkeleton.Pin().ToSharedRef())
 				.Sequence(Sequence)
 				.OnObjectsSelected(InArgs.OnDespatchObjectsSelected)
-				.OnAnimNotifiesChanged(InArgs.OnDespatchAnimNotifiesChanged)
 				.OnInvokeTab(InArgs.OnDespatchInvokeTab);
 
 			OutDocumentLink = TEXT("Engine/Animation/Sequences");
 		}
 		else if (UAnimComposite* Composite = Cast<UAnimComposite>(InAnimAsset))
 		{
-			Result = SNew(SAnimCompositeEditor, InArgs.PreviewScene.Pin().ToSharedRef(), InArgs.EditableSkeleton.Pin().ToSharedRef(), InArgs.OnPostUndo)
+			Result = SNew(SAnimCompositeEditor, InArgs.PreviewScene.Pin().ToSharedRef(), InArgs.EditableSkeleton.Pin().ToSharedRef())
 				.Composite(Composite)
 				.OnObjectsSelected(InArgs.OnDespatchObjectsSelected)
-				.OnAnimNotifiesChanged(InArgs.OnDespatchAnimNotifiesChanged)
 				.OnInvokeTab(InArgs.OnDespatchInvokeTab);
 
 			OutDocumentLink = TEXT("Engine/Animation/AnimationComposite");
 		}
 		else if (UAnimMontage* Montage = Cast<UAnimMontage>(InAnimAsset))
 		{
-			FMontageEditorRequiredArgs RequiredArgs(InArgs.PreviewScene.Pin().ToSharedRef(), InArgs.EditableSkeleton.Pin().ToSharedRef(), InArgs.OnPostUndo, InArgs.OnAnimNotifiesChanged, InArgs.OnSectionsChanged);
+			FMontageEditorRequiredArgs RequiredArgs(InArgs.PreviewScene.Pin().ToSharedRef(), InArgs.EditableSkeleton.Pin().ToSharedRef(), InArgs.OnSectionsChanged);
 
 			Result = SNew(SMontageEditor, RequiredArgs)
 				.Montage(Montage)
 				.OnSectionsChanged(InArgs.OnDespatchSectionsChanged)
 				.OnInvokeTab(InArgs.OnDespatchInvokeTab)
-				.OnObjectsSelected(InArgs.OnDespatchObjectsSelected)
-				.OnAnimNotifiesChanged(InArgs.OnDespatchAnimNotifiesChanged);
+				.OnObjectsSelected(InArgs.OnDespatchObjectsSelected);
 
 			OutDocumentLink = TEXT("Engine/Animation/AnimMontage");
 		}
diff --git a/Engine/Source/Editor/Persona/Private/SAnimCompositeEditor.cpp b/Engine/Source/Editor/Persona/Private/SAnimCompositeEditor.cpp
--- a/Engine/Source/Editor/Persona/Private/SAnimCompositeEditor.cpp
+++ b/Engine/Source/Editor/Persona/Private/SAnimCompositeEditor.cpp
@@ -6,6 +6,7 @@
 #include "IDocumentation.h"
 
 #include "SAnimNotifyPanel.h"
+#include "Editor.h"
 
 //////////////////////////////////////////////////////////////////////////
 // SAnimCompositeEditor
@@ -15,7 +16,7 @@ TSharedRef<SWidget> SAnimCompositeEditor::CreateDocumentAnchor()
 	return IDocumentation::Get()->CreateAnchor(TEXT("Engine/Animation/AnimationComposite"));
 }
 
-void SAnimCompositeEditor::Construct(const FArguments& InArgs, const TSharedRef<class IPersonaPreviewScene>& InPreviewScene, const TSharedRef<class IEditableSkeleton>& InEditableSkeleton, FSimpleMulticastDelegate& OnPostUndo)
+void SAnimCompositeEditor::Construct(const FArguments& InArgs, const TSharedRef<class IPersonaPreviewScene>& InPreviewScene, const TSharedRef<class IEditableSkeleton>& InEditableSkeleton)
 {
 	bIsActiveTimerRegistered = false;
 	CompositeObj = InArgs._Composite;
@@ -25,7 +26,10 @@ void SAnimCompositeEditor::Construct(const FArguments& InArgs, const TSharedRef<
 		.OnObjectsSelected(InArgs._OnObjectsSelected), 
 		InPreviewScene );
 
-	OnPostUndo.Add(FSimpleDelegate::CreateSP( this, &SAnimCompositeEditor::PostUndo ) );
+	if(GEditor)
+	{
+		GEditor->RegisterForUndo(this);
+	}
 
 	EditorPanels->AddSlot()
 		.AutoHeight()
@@ -44,7 +48,7 @@ void SAnimCompositeEditor::Construct(const FArguments& InArgs, const TSharedRef<
 		.AutoHeight()
 		.Padding(0, 10)
 		[
-			SAssignNew( AnimNotifyPanel, SAnimNotifyPanel, OnPostUndo)
+			SAssignNew( AnimNotifyPanel, SAnimNotifyPanel, InEditableSkeleton )
 			.Sequence(CompositeObj)
 			.WidgetWidth(S2ColumnWidget::DEFAULT_RIGHT_COLUMN_WIDTH)
 			.InputMin(this, &SAnimEditorBase::GetMinInput)
@@ -54,7 +58,6 @@ void SAnimCompositeEditor::Construct(const FArguments& InArgs, const TSharedRef<
 			.OnSetInputViewRange(this, &SAnimEditorBase::SetInputViewRange)
 			.OnGetScrubValue(this, &SAnimEditorBase::GetScrubValue)
 			.OnSelectionChanged(this, &SAnimEditorBase::OnSelectionChanged)
-			.OnAnimNotifiesChanged(InArgs._OnAnimNotifiesChanged)
 			.OnInvokeTab(InArgs._OnInvokeTab)
 		];
 
@@ -76,6 +79,14 @@ void SAnimCompositeEditor::Construct(const FArguments& InArgs, const TSharedRef<
 	CollapseComposite();
 }
 
+SAnimCompositeEditor::~SAnimCompositeEditor()
+{
+	if(GEditor)
+	{
+		GEditor->UnregisterForUndo(this);
+	}
+}
+
 void SAnimCompositeEditor::PreAnimUpdate()
 {
 	CompositeObj->Modify();
@@ -125,7 +136,17 @@ void SAnimCompositeEditor::CollapseComposite()
 	RecalculateSequenceLength();
 }
 
-void SAnimCompositeEditor::PostUndo()
+void SAnimCompositeEditor::PostUndo( bool bSuccess )
+{
+	PostUndoRedo();
+}
+
+void SAnimCompositeEditor::PostRedo( bool bSuccess )
+{
+	PostUndoRedo();
+}
+
+void SAnimCompositeEditor::PostUndoRedo()
 {
 	if (!bIsActiveTimerRegistered)
 	{
diff --git a/Engine/Source/Editor/Persona/Private/SAnimCompositeEditor.h b/Engine/Source/Editor/Persona/Private/SAnimCompositeEditor.h
--- a/Engine/Source/Editor/Persona/Private/SAnimCompositeEditor.h
+++ b/Engine/Source/Editor/Persona/Private/SAnimCompositeEditor.h
@@ -22,7 +22,7 @@ class SAnimNotifyPanel;
 	portion of the composite tool and registering callbacks to the SAnimCompositeEditor to do the actual editing.
 	
 */
-class SAnimCompositeEditor : public SAnimEditorBase
+class SAnimCompositeEditor : public SAnimEditorBase, public FEditorUndoClient
 {
 public:
 	SLATE_BEGIN_ARGS( SAnimCompositeEditor )
@@ -31,7 +31,6 @@ public:
 
 		SLATE_ARGUMENT( UAnimComposite*, Composite)
 		SLATE_EVENT(FOnObjectsSelected, OnObjectsSelected)
-		SLATE_EVENT(FSimpleDelegate, OnAnimNotifiesChanged)
 		SLATE_EVENT(FOnInvokeTab, OnInvokeTab)
 
 	SLATE_END_ARGS()
@@ -56,8 +55,10 @@ protected:
 	virtual void InitDetailsViewEditorObject(class UEditorAnimBaseObj* EdObj) override;
 
 public:
-	void Construct(const FArguments& InArgs, const TSharedRef<class IPersonaPreviewScene>& InPreviewScene, const TSharedRef<class IEditableSkeleton>& InEditableSkeleton, FSimpleMulticastDelegate& OnPostUndo);
-	
+	void Construct(const FArguments& InArgs, const TSharedRef<class IPersonaPreviewScene>& InPreviewScene, const TSharedRef<class IEditableSkeleton>& InEditableSkeleton);
+
+	~SAnimCompositeEditor();
+
 	/** Return the animation composite being edited */
 	UAnimComposite* GetCompositeObj() const { return CompositeObj; }
 	virtual UAnimationAsset* GetEditorObject() const override { return GetCompositeObj(); }
@@ -66,8 +67,11 @@ private:
 	/** Pointer to the animation composite being edited */
 	UAnimComposite* CompositeObj;
 	
-	void PostUndo();
-	
+	/** FEditorUndoClient interface */
+	virtual void PostUndo( bool bSuccess ) override;
+	virtual void PostRedo( bool bSuccess ) override;
+	void PostUndoRedo();
+
 	virtual float CalculateSequenceLengthOfEditorObject() const override;
 
 	/** This will sort all components of the montage and update (recreate) the UI */
diff --git a/Engine/Source/Editor/Persona/Private/SAnimMontagePanel.cpp b/Engine/Source/Editor/Persona/Private/SAnimMontagePanel.cpp
--- a/Engine/Source/Editor/Persona/Private/SAnimMontagePanel.cpp
+++ b/Engine/Source/Editor/Persona/Private/SAnimMontagePanel.cpp
@@ -19,13 +19,17 @@
 #include "TabSpawners.h"
 #include "Widgets/Input/SNumericEntryBox.h"
 #include "Styling/CoreStyle.h"
+#include "ISkeletonEditorModule.h"
+#include "Modules/ModuleManager.h"
+#include "IEditableSkeleton.h"
+#include "Editor.h"
 
 #define LOCTEXT_NAMESPACE "AnimMontagePanel"
 
 //////////////////////////////////////////////////////////////////////////
 // SAnimMontagePanel
 
-void SAnimMontagePanel::Construct(const FArguments& InArgs, FSimpleMulticastDelegate& OnAnimNotifiesChanged, FSimpleMulticastDelegate& OnSectionsChanged)
+void SAnimMontagePanel::Construct(const FArguments& InArgs, FSimpleMulticastDelegate& OnSectionsChanged)
 {
 	SAnimTrackPanel::Construct( SAnimTrackPanel::FArguments()
 		.WidgetWidth(InArgs._WidgetWidth)
@@ -65,7 +69,10 @@ void SAnimMontagePanel::Construct(const FArguments& InArgs, FSimpleMulticastDele
 		]
 	];
 
-	OnAnimNotifiesChanged.Add(FSimpleDelegate::CreateSP(this, &SAnimMontagePanel::Update));
+	ISkeletonEditorModule& SkeletonEditorModule = FModuleManager::LoadModuleChecked<ISkeletonEditorModule>("SkeletonEditor");
+	TSharedPtr<IEditableSkeleton> EditableSkeleton = SkeletonEditorModule.CreateEditableSkeleton(Montage->GetSkeleton());
+	EditableSkeleton->RegisterOnNotifiesChanged(FSimpleDelegate::CreateSP(this, &SAnimMontagePanel::Update));
+
 	OnSectionsChanged.Add(FSimpleDelegate::CreateSP(this, &SAnimMontagePanel::Update));
 
 	Update();
diff --git a/Engine/Source/Editor/Persona/Private/SAnimMontagePanel.h b/Engine/Source/Editor/Persona/Private/SAnimMontagePanel.h
--- a/Engine/Source/Editor/Persona/Private/SAnimMontagePanel.h
+++ b/Engine/Source/Editor/Persona/Private/SAnimMontagePanel.h
@@ -13,6 +13,7 @@
 #include "SMontageEditor.h"
 #include "Animation/EditorAnimBaseObj.h"
 #include "STrack.h"
+#include "EditorUndoClient.h"
 
 class FMenuBuilder;
 class SBorder;
@@ -68,7 +69,8 @@ public:
 	SLATE_ARGUMENT(bool, bChildAnimMontage)
 	SLATE_END_ARGS()
 
-	void Construct(const FArguments& InArgs, FSimpleMulticastDelegate& OnAnimNotifiesChanged, FSimpleMulticastDelegate& OnSectionsChanged);
+	void Construct(const FArguments& InArgs, FSimpleMulticastDelegate& OnSectionsChanged);
+
 	void SetMontage(class UAnimMontage * InMontage);
 
 	// SWidget interface
diff --git a/Engine/Source/Editor/Persona/Private/SAnimNotifyPanel.cpp b/Engine/Source/Editor/Persona/Private/SAnimNotifyPanel.cpp
--- a/Engine/Source/Editor/Persona/Private/SAnimNotifyPanel.cpp
+++ b/Engine/Source/Editor/Persona/Private/SAnimNotifyPanel.cpp
@@ -40,6 +40,9 @@
 #include "IEditableSkeleton.h"
 #include "ISkeletonEditorModule.h"
 #include "Widgets/Input/SNumericEntryBox.h"
+#include "ClassViewerFilter.h"
+#include "ClassViewerModule.h"
+#include "SSkeletonAnimNotifies.h"
 
 // Track Panel drawing
 const float NotificationTrackHeight = 20.0f;
@@ -756,7 +759,7 @@ public:
 	static TSubclassOf<UObject> GetBlueprintClassFromPath(FString BlueprintPath);
 
 	// Get the default Notify Name for a given blueprint notify asset
-	FString MakeBlueprintNotifyName(FAssetData& NotifyAssetData);
+	FString MakeBlueprintNotifyName(const FString& InNotifyClassName);
 
 	// Need to make sure tool tips are cleared during node clear up so slate system won't
 	// call into invalid notify.
@@ -767,6 +770,8 @@ protected:
 	void CreateCommands();
 
 	// Build up a "New Notify..." menu
+	template<typename NotifyTypeClass>
+	void MakeNewNotifyPicker(FMenuBuilder& MenuBuilder, bool bIsReplaceWithMenu = false);
 	void FillNewNotifyMenu(FMenuBuilder& MenuBuilderbool, bool bIsReplaceWithMenu = false);
 	void FillNewNotifyStateMenu(FMenuBuilder& MenuBuilder, bool bIsReplaceWithMenu  = false);
 
@@ -863,9 +868,6 @@ private:
 		UClass* BaseClass;
 	};
 
-	// Format notify asset data into the information needed for menu display
-	void GetNotifyMenuData(TArray<FAssetData>& NotifyAssetData, TArray<BlueprintNotifyMenuInfo>& OutNotifyMenuData);
-
 	// Store the tracks geometry for later use
 	void UpdateCachedGeometry(const FGeometry& InGeometry);
 
@@ -2324,231 +2326,139 @@ FCursorReply SAnimNotifyTrack::OnCursorQuery(const FGeometry& MyGeometry, const
 	return FCursorReply::Unhandled();
 }
 
-// Fill "new notify state" menu, or "replace with notify state menu"
-void SAnimNotifyTrack::FillNewNotifyStateMenu(FMenuBuilder& MenuBuilder, bool bIsReplaceWithMenu /* = false */)
+template<typename NotifyTypeClass>
+void SAnimNotifyTrack::MakeNewNotifyPicker(FMenuBuilder& MenuBuilder, bool bIsReplaceWithMenu /* = false */)
 {
-	// Run the native query first to update the allowed classes for blueprints.
-	TArray<UClass*> NotifyStateClasses;
-	OnGetNotifyStateNativeClasses.ExecuteIfBound(NotifyStateClasses);
+	FText TypeName = NotifyTypeClass::StaticClass() == UAnimNotify::StaticClass() ? LOCTEXT("AnimNotifyName", "anim notify") : LOCTEXT("AnimNotifyStateName", "anim notify state");
+	FText SectionHeaderFormat = bIsReplaceWithMenu ? LOCTEXT("ReplaceWithAnExistingAnimNotify", "Replace with an existing {0}") : LOCTEXT("AddsAnExistingAnimNotify", "Add an existing {0}");
 
-	// Collect blueprint notify data
-	TArray<FAssetData> NotifyAssetData;
-	TArray<BlueprintNotifyMenuInfo> NotifyMenuData;
-	OnGetNotifyStateBlueprintData.ExecuteIfBound(NotifyAssetData);
-	GetNotifyMenuData(NotifyAssetData, NotifyMenuData);
-
-	for(BlueprintNotifyMenuInfo& NotifyData : NotifyMenuData)
+	class FNotifyStateClassFilter : public IClassViewerFilter
 	{
-		const FText LabelText = FText::FromString(NotifyData.NotifyName);
+	public:
+		FNotifyStateClassFilter(UAnimSequenceBase* InSequence)
+			: Sequence(InSequence)
+		{}
 
-		FUIAction UIAction;
-		FText Description = FText::GetEmpty();
-		if (!bIsReplaceWithMenu)
+		bool IsClassAllowed(const FClassViewerInitializationOptions& InInitOptions, const UClass* InClass, TSharedRef< FClassViewerFilterFuncs > InFilterFuncs) override
 		{
-			Description = LOCTEXT("AddsAnExistingAnimNotify", "Add an existing notify");
-			UIAction.ExecuteAction.BindRaw(
-				this, &SAnimNotifyTrack::CreateNewBlueprintNotifyAtCursor,
-				NotifyData.NotifyName,
-				NotifyData.BlueprintPath);
-			UIAction.CanExecuteAction.BindRaw(
-				this, &SAnimNotifyTrack::IsValidToPlace,
-				NotifyData.BaseClass);
-		}
-		else
-		{
-			Description = LOCTEXT("ReplaceWithAnExistingAnimNotify", "Replace with an existing notify");
-			UIAction.ExecuteAction.BindRaw(
-				this, &SAnimNotifyTrack::ReplaceSelectedWithBlueprintNotify,
-				NotifyData.NotifyName,
-				NotifyData.BlueprintPath);
-			UIAction.CanExecuteAction.BindRaw(
-				this, &SAnimNotifyTrack::IsValidToPlace,
-				NotifyData.BaseClass);
+			const bool bChildOfObjectClass = InClass->IsChildOf(NotifyTypeClass::StaticClass());
+			const bool bMatchesFlags = !InClass->HasAnyClassFlags(CLASS_Hidden | CLASS_HideDropDown | CLASS_Deprecated | CLASS_Abstract);
+			return bChildOfObjectClass && bMatchesFlags && CastChecked<NotifyTypeClass>(InClass->ClassDefaultObject)->CanBePlaced(Sequence);
 		}
 
-		MenuBuilder.AddMenuEntry(LabelText, Description, FSlateIcon(), UIAction);
-	}
-
-	MenuBuilder.BeginSection("NativeNotifyStates", LOCTEXT("NewStateNotifyMenu_Native", "Native Notify States"));
-	{
-		for(UClass* Class : NotifyStateClasses)
+		virtual bool IsUnloadedClassAllowed(const FClassViewerInitializationOptions& InInitOptions, const TSharedRef< const IUnloadedBlueprintData > InUnloadedClassData, TSharedRef< FClassViewerFilterFuncs > InFilterFuncs) override
 		{
-			if (Class->HasAllClassFlags(CLASS_Abstract))
+			const bool bChildOfObjectClass = InUnloadedClassData->IsChildOf(NotifyTypeClass::StaticClass());
+			const bool bMatchesFlags = !InUnloadedClassData->HasAnyClassFlags(CLASS_Hidden | CLASS_HideDropDown | CLASS_Deprecated | CLASS_Abstract);
+			bool bValidToPlace = false;
+			if(bChildOfObjectClass)
 			{
-				continue; // skip abstract classes
+				if (const UClass* NativeBaseClass = InUnloadedClassData->GetNativeParent())
+				{
+					bValidToPlace = CastChecked<NotifyTypeClass>(NativeBaseClass->ClassDefaultObject)->CanBePlaced(Sequence);
+				}
 			}
 
-			const FText Description = LOCTEXT("NewNotifyStateSubMenu_NativeToolTip", "Add an existing native notify state");
-			const FText LabelText = Class->GetDisplayNameText();
-			const FString Label = LabelText.ToString();
+			return bChildOfObjectClass && bMatchesFlags && bValidToPlace;
+		}
 
-			FUIAction UIAction;
-			if (!bIsReplaceWithMenu)
-			{
-				UIAction.ExecuteAction.BindRaw(
-					this, &SAnimNotifyTrack::CreateNewNotifyAtCursor,
-					Label,
-					Class);
-				UIAction.CanExecuteAction.BindRaw(
-					this, &SAnimNotifyTrack::IsValidToPlace,
-					Class);
-			}
-			else
-			{
-				UIAction.ExecuteAction.BindRaw(
-					this, &SAnimNotifyTrack::ReplaceSelectedWithNotify,
-					Label,
-					Class);
-				UIAction.CanExecuteAction.BindRaw(
-					this, &SAnimNotifyTrack::IsValidToPlace,
-					Class);
-			}
+		/** Sequence referenced by outer panel */
+		UAnimSequenceBase* Sequence;
+	};
 
-			MenuBuilder.AddMenuEntry(LabelText, Description, FSlateIcon(), UIAction);
-		}
-	}
-	MenuBuilder.EndSection();
+	FClassViewerInitializationOptions InitOptions;
+	InitOptions.Mode = EClassViewerMode::ClassPicker;
+	InitOptions.bShowObjectRootClass = false;
+	InitOptions.bShowUnloadedBlueprints = true;
+	InitOptions.bShowNoneOption = false;
+	InitOptions.bEnableClassDynamicLoading = true;
+	InitOptions.bExpandRootNodes = true;
+	InitOptions.bShowDisplayNames = true;
+	InitOptions.ClassFilter = MakeShared<FNotifyStateClassFilter>(Sequence);
+
+	FClassViewerModule& ClassViewerModule = FModuleManager::LoadModuleChecked<FClassViewerModule>("ClassViewer");
+	MenuBuilder.AddWidget(
+		SNew(SBox)
+		.MinDesiredWidth(300.0f)
+		.MaxDesiredHeight(400.0f)
+		[
+			ClassViewerModule.CreateClassViewer(InitOptions,
+				FOnClassPicked::CreateLambda([this, bIsReplaceWithMenu](UClass* InClass)
+				{
+					FSlateApplication::Get().DismissAllMenus();
+					if(bIsReplaceWithMenu)
+					{
+						ReplaceSelectedWithNotify(MakeBlueprintNotifyName(InClass->GetName()), InClass);
+					}
+					else
+					{
+						CreateNewNotifyAtCursor(MakeBlueprintNotifyName(InClass->GetName()), InClass);
+					}
+				}
+			))
+		],
+		FText(), true, false);
 }
 
-// Fill "new notify" menu, or "replace with notify menu"
-void SAnimNotifyTrack::FillNewNotifyMenu(FMenuBuilder& MenuBuilder, bool bIsReplaceWithMenu /* = false */)
+void SAnimNotifyTrack::FillNewNotifyStateMenu(FMenuBuilder& MenuBuilder, bool bIsReplaceWithMenu /* = false */)
 {
-	TArray<UClass*> NativeNotifyClasses;
-	OnGetNotifyNativeClasses.ExecuteIfBound(NativeNotifyClasses);
-
-	TArray<FAssetData> NotifyAssetData;
-	TArray<BlueprintNotifyMenuInfo> NotifyMenuData;
-	OnGetNotifyBlueprintData.ExecuteIfBound(NotifyAssetData);
-	GetNotifyMenuData(NotifyAssetData, NotifyMenuData);
-
-	for(BlueprintNotifyMenuInfo& NotifyData : NotifyMenuData)
-	{
-		const FText LabelText = FText::FromString( NotifyData.NotifyName );
-
-		FUIAction UIAction;
-		FText Description = FText::GetEmpty();
-		if (!bIsReplaceWithMenu)
-		{
-			Description = LOCTEXT("NewNotifySubMenu_ToolTip", "Add an existing notify");
-			UIAction.ExecuteAction.BindRaw(
-				this, &SAnimNotifyTrack::CreateNewBlueprintNotifyAtCursor,
-				NotifyData.NotifyName,
-				NotifyData.BlueprintPath);
-			UIAction.CanExecuteAction.BindRaw(
-				this, &SAnimNotifyTrack::IsValidToPlace,
-				NotifyData.BaseClass);
-		}
-		else
-		{
-			Description = LOCTEXT("ReplaceWithNotifySubMenu_ToolTip", "Replace with an existing notify");
-			UIAction.ExecuteAction.BindRaw(
-				this, &SAnimNotifyTrack::ReplaceSelectedWithBlueprintNotify,
-				NotifyData.NotifyName,
-				NotifyData.BlueprintPath);
-			UIAction.CanExecuteAction.BindRaw(
-				this, &SAnimNotifyTrack::IsValidToPlace,
-				NotifyData.BaseClass);
-		}
-		
-		MenuBuilder.AddMenuEntry(LabelText, Description, FSlateIcon(), UIAction);
-	}
+	MakeNewNotifyPicker<UAnimNotifyState>(MenuBuilder, bIsReplaceWithMenu);
+}
 
-	MenuBuilder.BeginSection("NativeNotifies", LOCTEXT("NewNotifyMenu_Native", "Native Notifies"));
+void SAnimNotifyTrack::FillNewNotifyMenu(FMenuBuilder& MenuBuilder, bool bIsReplaceWithMenu /* = false */)
+{
+	// now add custom anim notifiers
+	USkeleton* SeqSkeleton = Sequence->GetSkeleton();
+	if (SeqSkeleton)
 	{
-		for(UClass* Class : NativeNotifyClasses)
+		MenuBuilder.BeginSection("AnimNotifySkeletonSubMenu", LOCTEXT("NewNotifySubMenu_Skeleton", "Skeleton Notifies"));
 		{
-			if (Class->HasAllClassFlags(CLASS_Abstract))
-			{
-				continue; // skip abstract classes
-			}
-
-			const FText LabelText = Class->GetDisplayNameText();
-			const FString Label = LabelText.ToString();
-
-			FUIAction UIAction;
-			FText Description = FText::GetEmpty();
 			if (!bIsReplaceWithMenu)
 			{
-				Description = LOCTEXT("NewNotifySubMenu_NativeToolTip", "Add an existing native notify");
-				UIAction.ExecuteAction.BindRaw(
-					this, &SAnimNotifyTrack::CreateNewNotifyAtCursor,
-					Label,
-					Class);
-				UIAction.CanExecuteAction.BindRaw(
-					this, &SAnimNotifyTrack::IsValidToPlace,
-					Class);
-			}
-			else
-			{
-				Description = LOCTEXT("ReplaceWithNotifySubMenu_NativeToolTip", "Replace with an existing native notify");
-				UIAction.ExecuteAction.BindRaw(
-					this, &SAnimNotifyTrack::ReplaceSelectedWithNotify,
-					Label,
-					Class);
-				UIAction.CanExecuteAction.BindRaw(
-					this, &SAnimNotifyTrack::IsValidToPlace,
-					Class);
+				FUIAction UIAction;
+				UIAction.ExecuteAction.BindSP(
+					this, &SAnimNotifyTrack::OnNewNotifyClicked);
+				MenuBuilder.AddMenuEntry(LOCTEXT("NewNotify", "New Notify..."), LOCTEXT("NewNotifyToolTip", "Create a new animation notify on the skeleton"), FSlateIcon(), UIAction);
 			}
 
-			MenuBuilder.AddMenuEntry(LabelText, Description, FSlateIcon(), UIAction);
-		}
-	}
-	MenuBuilder.EndSection();
-
-	MenuBuilder.BeginSection("AnimNotifyCustom", LOCTEXT("NewNotifySubMenu_Custom", "Custom"));
-	{
-		// now add custom anim notifiers
-		USkeleton* SeqSkeleton = Sequence->GetSkeleton();
-		if (SeqSkeleton)
-		{
-			for (int32 I = 0; I<SeqSkeleton->AnimationNotifies.Num(); ++I)
-			{
-				FName NotifyName = SeqSkeleton->AnimationNotifies[I];
-				FString Label = NotifyName.ToString();
-
-				FText Description = FText::GetEmpty();
-				if (!bIsReplaceWithMenu)
-				{
-					Description = LOCTEXT("NewNotifySubMenu_ToolTip", "Add an existing notify");
-				}
-				else
-				{
-					Description = LOCTEXT("ReplaceWithNotifySubMenu_ToolTip", "Replace with an existing notify");
-				}
-
-				FUIAction UIAction;
-				if (!bIsReplaceWithMenu)
-				{
-					UIAction.ExecuteAction.BindRaw(
-						this, &SAnimNotifyTrack::CreateNewNotifyAtCursor,
-						Label,
-						(UClass*)nullptr);
-				}
-				else
+			MenuBuilder.AddSubMenu(
+				LOCTEXT("NewNotifySubMenu_Skeleton", "Skeleton Notifies"),
+				LOCTEXT("NewNotifySubMenu_Skeleton_Tooltip", "Choose from custom notifies on the skeleton"),
+				FNewMenuDelegate::CreateLambda([this, SeqSkeleton, bIsReplaceWithMenu](FMenuBuilder& InSubMenuBuilder)
 				{
-					UIAction.ExecuteAction.BindRaw(
-						this, &SAnimNotifyTrack::ReplaceSelectedWithNotify,
-						Label,
-						(UClass*)nullptr);
-				}
+					ISkeletonEditorModule& SkeletonEditorModule = FModuleManager::LoadModuleChecked<ISkeletonEditorModule>("SkeletonEditor");
+					TSharedRef<IEditableSkeleton> EditableSkeleton = SkeletonEditorModule.CreateEditableSkeleton(SeqSkeleton);
 
-				MenuBuilder.AddMenuEntry( FText::FromString( Label ), Description, FSlateIcon(), UIAction);
-			}
-		}
-	}
-	MenuBuilder.EndSection();
+					InSubMenuBuilder.AddWidget(
+						SNew(SBox)
+						.MinDesiredWidth(300.0f)
+						.MaxDesiredHeight(400.0f)
+						[
+							SNew(SSkeletonAnimNotifies, EditableSkeleton)
+							.IsPicker(true)
+							.OnNotifySelected_Lambda([this, bIsReplaceWithMenu](const FName& InNotifyName)
+							{
+								FSlateApplication::Get().DismissAllMenus();
 
-	if (!bIsReplaceWithMenu)
-	{
-		MenuBuilder.BeginSection("AnimNotifyCreateNew");
-		{
-			FUIAction UIAction;
-			UIAction.ExecuteAction.BindRaw(
-				this, &SAnimNotifyTrack::OnNewNotifyClicked);
-			MenuBuilder.AddMenuEntry(LOCTEXT("NewNotify", "New Notify"), LOCTEXT("NewNotifyToolTip", "Create a new animation notify"), FSlateIcon(), UIAction);
+								if (!bIsReplaceWithMenu)
+								{
+									CreateNewNotifyAtCursor(InNotifyName.ToString(), nullptr);
+								}
+								else
+								{
+									ReplaceSelectedWithNotify(InNotifyName.ToString(), nullptr);
+								}
+							})
+						],
+						FText(), true, false
+					);
+				}));
 		}
 		MenuBuilder.EndSection();
 	}
+
+	// Add a notify picker
+	MakeNewNotifyPicker<UAnimNotify>(MenuBuilder, bIsReplaceWithMenu);
 }
 
 FAnimNotifyEvent& SAnimNotifyTrack::CreateNewBlueprintNotify(FString NewNotifyName, FString BlueprintPath, float StartTime)
@@ -3755,9 +3665,9 @@ void SAnimNotifyTrack::RefreshMarqueeSelectedNodes(FSlateRect& Rect, FNotifyMarq
 	}
 }
 
-FString SAnimNotifyTrack::MakeBlueprintNotifyName(FAssetData& NotifyAssetData)
+FString SAnimNotifyTrack::MakeBlueprintNotifyName(const FString& InNotifyClassName)
 {
-	FString DefaultNotifyName = NotifyAssetData.AssetName.ToString();
+	FString DefaultNotifyName = InNotifyClassName;
 	DefaultNotifyName = DefaultNotifyName.Replace(TEXT("AnimNotify_"), TEXT(""), ESearchCase::CaseSensitive);
 	DefaultNotifyName = DefaultNotifyName.Replace(TEXT("AnimNotifyState_"), TEXT(""), ESearchCase::CaseSensitive);
 
@@ -3774,32 +3684,6 @@ void SAnimNotifyTrack::ClearNodeTooltips()
 	}
 }
 
-void SAnimNotifyTrack::GetNotifyMenuData(TArray<FAssetData>& NotifyAssetData, TArray<BlueprintNotifyMenuInfo>& OutNotifyMenuData)
-{
-	for(FAssetData& NotifyData : NotifyAssetData)
-	{
-		OutNotifyMenuData.AddZeroed();
-		BlueprintNotifyMenuInfo& MenuInfo = OutNotifyMenuData.Last();
-
-		MenuInfo.BlueprintPath = NotifyData.ObjectPath.ToString();
-		MenuInfo.NotifyName = MakeBlueprintNotifyName(NotifyData);
-		// this functionality is only available in native class
-		// so we don't have to call BP function but just call native on the check of validity
-		FString NativeParentClassName;
-		if (NotifyData.GetTagValue(FBlueprintTags::NativeParentClassPath, NativeParentClassName))
-		{
-			UObject* Outer = nullptr;
-			ResolveName(Outer, NativeParentClassName, false, false);
-			MenuInfo.BaseClass = FindObject<UClass>(ANY_PACKAGE, *NativeParentClassName);
-		}
-	}
-
-	OutNotifyMenuData.Sort([](const BlueprintNotifyMenuInfo& A, const BlueprintNotifyMenuInfo& B)
-	{
-		return A.NotifyName < B.NotifyName;
-	});
-}
-
 const EVisibility SAnimNotifyTrack::GetTimingNodeVisibility(TSharedPtr<SAnimNotifyNode> NotifyNode)
 {
 	if(OnGetTimingNodeVisibility.IsBound())
@@ -3942,7 +3826,7 @@ void FAnimNotifyPanelCommands::RegisterCommands()
 //////////////////////////////////////////////////////////////////////////
 // SAnimNotifyPanel
 
-void SAnimNotifyPanel::Construct(const FArguments& InArgs, FSimpleMulticastDelegate& OnPostUndo)
+void SAnimNotifyPanel::Construct(const FArguments& InArgs, const TSharedRef<class IEditableSkeleton>& InEditableSkeleton)
 {
 	SAnimTrackPanel::Construct( SAnimTrackPanel::FArguments()
 		.WidgetWidth(InArgs._WidgetWidth)
@@ -3954,7 +3838,6 @@ void SAnimNotifyPanel::Construct(const FArguments& InArgs, FSimpleMulticastDeleg
 
 	Sequence = InArgs._Sequence;
 	MarkerBars = InArgs._MarkerBars;
-	OnAnimNotifiesChanged = InArgs._OnAnimNotifiesChanged;
 	OnInvokeTab = InArgs._OnInvokeTab;
 
 	FAnimNotifyPanelCommands::Register();
@@ -3966,7 +3849,12 @@ void SAnimNotifyPanel::Construct(const FArguments& InArgs, FSimpleMulticastDeleg
 	Sequence->InitializeNotifyTrack();
 	Sequence->RegisterOnNotifyChanged(UAnimSequenceBase::FOnNotifyChanged::CreateSP( this, &SAnimNotifyPanel::RefreshNotifyTracks )  );
 
-	OnPostUndo.Add(FSimpleDelegate::CreateSP( this, &SAnimNotifyPanel::PostUndo ) );
+	InEditableSkeleton->RegisterOnNotifiesChanged(FSimpleDelegate::CreateSP(this, &SAnimNotifyPanel::RefreshNotifyTracks));
+
+	if(GEditor)
+	{
+		GEditor->RegisterForUndo(this);
+	}
 
 	CurrentPosition = InArgs._CurrentPosition;
 	OnSelectionChanged = InArgs._OnSelectionChanged;
@@ -4038,6 +3926,11 @@ SAnimNotifyPanel::~SAnimNotifyPanel()
 	Sequence->UnregisterOnNotifyChanged(this);
 
 	FCoreUObjectDelegates::OnObjectPropertyChanged.Remove(OnPropertyChangedHandleDelegateHandle);
+
+	if(GEditor)
+	{
+		GEditor->UnregisterForUndo(this);
+	}
 }
 
 FName SAnimNotifyPanel::GetNewTrackName() const
@@ -4165,8 +4058,6 @@ void SAnimNotifyPanel::Update()
 	{
 		Sequence->RefreshCacheData();
 	}
-
-	OnAnimNotifiesChanged.ExecuteIfBound();
 }
 
 void SAnimNotifyPanel::RefreshNotifyTracks()
@@ -4289,7 +4180,15 @@ FReply SAnimNotifyPanel::OnNotifyNodeDragStarted(TArray<TSharedPtr<SAnimNotifyNo
 	return FReply::Handled().BeginDragDrop(FNotifyDragDropOp::New(Nodes, NodeDragDecorator, NotifyAnimTracks, Sequence, ScreenCursorPos, OverlayOrigin, OverlayExtents, CurrentDragXPosition, PanRequestDelegate, MarkerBars, UpdateDelegate));
 }
 
-void SAnimNotifyPanel::PostUndo()
+void SAnimNotifyPanel::PostUndo( bool bSuccess )
+{
+	if(Sequence != NULL)
+	{
+		Sequence->RefreshCacheData();
+	}
+}
+
+void SAnimNotifyPanel::PostRedo( bool bSuccess )
 {
 	if(Sequence != NULL)
 	{
@@ -4389,9 +4288,6 @@ void SAnimNotifyPanel::DeselectAllNotifies()
 		Track->DeselectAllNotifyNodes(false);
 	}
 
-	// Broadcast the change so the editor can update
-	OnAnimNotifiesChanged.ExecuteIfBound();
-
 	OnTrackSelectionChanged();
 }
 
@@ -4664,9 +4560,6 @@ void SAnimNotifyPanel::OnPropertyChanged(UObject* ChangedObject, FPropertyChange
 				RefreshNotifyTracks();
 			}
 		}
-
-		// Broadcast the change so the editor can update
-		OnAnimNotifiesChanged.ExecuteIfBound();
 	}
 }
 
@@ -4921,9 +4814,6 @@ void SAnimNotifyPanel::OnNotifyObjectChanged(UObject* EditorBaseObj, bool bRebui
 		{
 			NotifyAnimTracks[WidgetTrackIdx]->Update();
 		}
-
-		// Broadcast the change so the editor can update
-		OnAnimNotifiesChanged.ExecuteIfBound();
 	}
 }
 
diff --git a/Engine/Source/Editor/Persona/Private/SAnimNotifyPanel.h b/Engine/Source/Editor/Persona/Private/SAnimNotifyPanel.h
--- a/Engine/Source/Editor/Persona/Private/SAnimNotifyPanel.h
+++ b/Engine/Source/Editor/Persona/Private/SAnimNotifyPanel.h
@@ -16,6 +16,7 @@
 #include "AssetData.h"
 #include "Framework/Commands/Commands.h"
 #include "SAnimTimingPanel.h"
+#include "EditorUndoClient.h"
 
 class FSlateWindowElementList;
 class SAnimNotifyNode;
@@ -153,7 +154,7 @@ public:
 
 // @todo anim : register when it's opened for the animsequence
 // broadcast when animsequence changed, so that we refresh for multiple window
-class SAnimNotifyPanel: public SAnimTrackPanel
+class SAnimNotifyPanel: public SAnimTrackPanel, public FEditorUndoClient
 {
 public:
 	SLATE_BEGIN_ARGS( SAnimNotifyPanel )
@@ -182,12 +183,11 @@ public:
 	SLATE_EVENT( FOnGetScrubValue, OnGetScrubValue )
 	SLATE_EVENT( FRefreshOffsetsRequest, OnRequestRefreshOffsets )
 	SLATE_EVENT( FOnGetTimingNodeVisibility, OnGetTimingNodeVisibility )
-	SLATE_EVENT( FSimpleDelegate, OnAnimNotifiesChanged )
 	SLATE_EVENT( FOnInvokeTab, OnInvokeTab )
 
 	SLATE_END_ARGS()
 
-	void Construct(const FArguments& InArgs, FSimpleMulticastDelegate& OnPostUndo);
+	void Construct(const FArguments& InArgs, const TSharedRef<class IEditableSkeleton>& InEditableSkeleton);
 	virtual ~SAnimNotifyPanel();
 
 	void SetSequence(class UAnimSequenceBase *	InSequence);
@@ -273,7 +273,10 @@ private:
 	// this just refresh notify tracks - UI purpose only
 	// do not call this from here. This gets called by asset. 
 	void RefreshNotifyTracks();
-	void PostUndo();
+
+	/** FEditorUndoClient interface */
+	virtual void PostUndo( bool bSuccess ) override;
+	virtual void PostRedo( bool bSuccess ) override;
 
 	/** Handler for delete command */
 	void OnDeletePressed();
@@ -335,9 +338,6 @@ private:
 	/** Handle to the registered OnPropertyChangedHandle delegate */
 	FDelegateHandle OnPropertyChangedHandleDelegateHandle;
 
-	/** Delegate fired when anim notifies are changed */
-	FSimpleDelegate OnAnimNotifiesChanged;
-
 	/** Delegate used to invoke a tab */
 	FOnInvokeTab OnInvokeTab;
 };
diff --git a/Engine/Source/Editor/Persona/Private/SAnimTimingPanel.cpp b/Engine/Source/Editor/Persona/Private/SAnimTimingPanel.cpp
--- a/Engine/Source/Editor/Persona/Private/SAnimTimingPanel.cpp
+++ b/Engine/Source/Editor/Persona/Private/SAnimTimingPanel.cpp
@@ -18,6 +18,10 @@
 #include "Preferences/PersonaOptions.h"
 #include "Widgets/Layout/SExpandableArea.h"
 #include "STimingTrack.h"
+#include "ISkeletonEditorModule.h"
+#include "Modules/ModuleManager.h"
+#include "IEditableSkeleton.h"
+#include "Editor.h"
 
 #define LOCTEXT_NAMESPACE "AnimTimingPanel"
 
@@ -188,7 +192,7 @@ void SAnimTimingTrackNode::Construct(const FArguments& InArgs)
 		);
 }
 
-void SAnimTimingPanel::Construct(const FArguments& InArgs, FSimpleMulticastDelegate& OnAnimNotifiesChanged, FSimpleMulticastDelegate& OnSectionsChanged)
+void SAnimTimingPanel::Construct(const FArguments& InArgs, FSimpleMulticastDelegate& OnSectionsChanged)
 {
 	SAnimTrackPanel::Construct(SAnimTrackPanel::FArguments()
 		.WidgetWidth(InArgs._WidgetWidth)
@@ -223,7 +227,10 @@ void SAnimTimingPanel::Construct(const FArguments& InArgs, FSimpleMulticastDeleg
 	Update();
 
 	// Register to some delegates to update the interface
-	OnAnimNotifiesChanged.Add(FSimpleDelegate::CreateSP(this, &SAnimTimingPanel::RefreshTrackNodes));
+	ISkeletonEditorModule& SkeletonEditorModule = FModuleManager::LoadModuleChecked<ISkeletonEditorModule>("SkeletonEditor");
+	TSharedPtr<IEditableSkeleton> EditableSkeleton = SkeletonEditorModule.CreateEditableSkeleton(AnimSequence->GetSkeleton());
+	EditableSkeleton->RegisterOnNotifiesChanged(FSimpleDelegate::CreateSP(this, &SAnimTimingPanel::RefreshTrackNodes));
+
 	OnSectionsChanged.Add(FSimpleDelegate::CreateSP(this, &SAnimTimingPanel::RefreshTrackNodes));
 
 	// Clear display flags
diff --git a/Engine/Source/Editor/Persona/Private/SAnimTimingPanel.h b/Engine/Source/Editor/Persona/Private/SAnimTimingPanel.h
--- a/Engine/Source/Editor/Persona/Private/SAnimTimingPanel.h
+++ b/Engine/Source/Editor/Persona/Private/SAnimTimingPanel.h
@@ -193,7 +193,7 @@ public:
 
 	// Construct the panel
 	// @param InArgs - Slate arguments
-	void Construct(const FArguments& InArgs, FSimpleMulticastDelegate& OnAnimNotifiesChanged, FSimpleMulticastDelegate& OnSectionsChanged);
+	void Construct(const FArguments& InArgs, FSimpleMulticastDelegate& OnSectionsChanged);
 
 	// Updates panel widgets
 	void Update();
diff --git a/Engine/Source/Editor/Persona/Private/SMontageEditor.cpp b/Engine/Source/Editor/Persona/Private/SMontageEditor.cpp
--- a/Engine/Source/Editor/Persona/Private/SMontageEditor.cpp
+++ b/Engine/Source/Editor/Persona/Private/SMontageEditor.cpp
@@ -21,6 +21,7 @@
 #include "Widgets/Input/SButton.h"
 #include "Widgets/Images/SImage.h"
 #include "Factories/AnimMontageFactory.h"
+#include "Editor.h"
 
 #define LOCTEXT_NAMESPACE "AnimSequenceEditor"
 
@@ -52,9 +53,12 @@ void SMontageEditor::Construct(const FArguments& InArgs, const FMontageEditorReq
 		.OnObjectsSelected(InArgs._OnObjectsSelected), 
 		InRequiredArgs.PreviewScene );
 
-	InRequiredArgs.OnPostUndo.Add(FSimpleDelegate::CreateSP(this, &SMontageEditor::PostUndo));
+	if(GEditor)
+	{
+		GEditor->RegisterForUndo(this);
+	}
 
-	SAssignNew(AnimTimingPanel, SAnimTimingPanel, InRequiredArgs.OnAnimNotifiesChanged, InRequiredArgs.OnSectionsChanged)
+	SAssignNew(AnimTimingPanel, SAnimTimingPanel, InRequiredArgs.OnSectionsChanged)
 		.IsEnabled(!bChildAnimMontage)
 		.InSequence(MontageObj)
 		.WidgetWidth(S2ColumnWidget::DEFAULT_RIGHT_COLUMN_WIDTH)
@@ -144,7 +148,7 @@ void SMontageEditor::Construct(const FArguments& InArgs, const FMontageEditorReq
 	.AutoHeight()
 	.Padding(0, 5)
 	[
-		SAssignNew( AnimMontagePanel, SAnimMontagePanel, InRequiredArgs.OnAnimNotifiesChanged, InRequiredArgs.OnSectionsChanged)
+		SAssignNew( AnimMontagePanel, SAnimMontagePanel, InRequiredArgs.OnSectionsChanged)
 		.Montage(MontageObj)
 		.bChildAnimMontage(bChildAnimMontage)
 		.MontageEditor(SharedThis(this))
@@ -180,7 +184,7 @@ void SMontageEditor::Construct(const FArguments& InArgs, const FMontageEditorReq
 	.AutoHeight()
 	.Padding(0, 10)
 	[
-		SAssignNew( AnimNotifyPanel, SAnimNotifyPanel, InRequiredArgs.OnPostUndo )
+		SAssignNew( AnimNotifyPanel, SAnimNotifyPanel, InRequiredArgs.EditableSkeleton )
 		.Sequence(MontageObj)
 		.IsEnabled(!bChildAnimMontage)
 		.WidgetWidth(S2ColumnWidget::DEFAULT_RIGHT_COLUMN_WIDTH)
@@ -194,7 +198,6 @@ void SMontageEditor::Construct(const FArguments& InArgs, const FMontageEditorReq
 		.MarkerBars(this, &SMontageEditor::GetMarkerBarInformation)
 		.OnRequestRefreshOffsets(this, &SMontageEditor::RefreshNotifyTriggerOffsets)
 		.OnGetTimingNodeVisibility(TimingNodeVisibilityDelegate)
-		.OnAnimNotifiesChanged(InArgs._OnAnimNotifiesChanged)
 		.OnInvokeTab(InArgs._OnInvokeTab)
 	];
 
@@ -253,6 +256,11 @@ SMontageEditor::~SMontageEditor()
 	{
 		MontageObj->UnregisterOnMontageChanged(this);
 	}
+
+	if(GEditor)
+	{
+		GEditor->UnregisterForUndo(this);
+	}
 }
 
 TSharedRef<class SAnimationScrubPanel> SMontageEditor::ConstructAnimScrubPanel()
@@ -886,7 +894,17 @@ void SMontageEditor::ClearSquenceOrdering()
 	RestartPreview();
 }
 
-void SMontageEditor::PostUndo()
+void SMontageEditor::PostUndo( bool bSuccess )
+{
+	PostRedoUndo();
+}
+
+void SMontageEditor::PostRedo( bool bSuccess )
+{
+	PostRedoUndo();
+}
+
+void SMontageEditor::PostRedoUndo()
 {
 	// when undo or redo happens, we still have to recalculate length, so we can't rely on sequence length changes or not
 	if (MontageObj->SequenceLength)
diff --git a/Engine/Source/Editor/Persona/Private/SMontageEditor.h b/Engine/Source/Editor/Persona/Private/SMontageEditor.h
--- a/Engine/Source/Editor/Persona/Private/SMontageEditor.h
+++ b/Engine/Source/Editor/Persona/Private/SMontageEditor.h
@@ -18,18 +18,14 @@ class SAnimTimingPanel;
 
 struct FMontageEditorRequiredArgs
 {
-	FMontageEditorRequiredArgs(const TSharedRef<class IPersonaPreviewScene>& InPreviewScene, const TSharedRef<class IEditableSkeleton>& InEditableSkeleton, FSimpleMulticastDelegate& InOnPostUndo, FSimpleMulticastDelegate& InOnAnimNotifiesChanged, FSimpleMulticastDelegate& InOnSectionsChanged)
+	FMontageEditorRequiredArgs(const TSharedRef<class IPersonaPreviewScene>& InPreviewScene, const TSharedRef<class IEditableSkeleton>& InEditableSkeleton, FSimpleMulticastDelegate& InOnSectionsChanged)
 		: PreviewScene(InPreviewScene)
 		, EditableSkeleton(InEditableSkeleton)
-		, OnPostUndo(InOnPostUndo)
-		, OnAnimNotifiesChanged(InOnAnimNotifiesChanged)
 		, OnSectionsChanged(InOnSectionsChanged)
 	{}
 
 	TSharedRef<class IPersonaPreviewScene> PreviewScene;
 	TSharedRef<class IEditableSkeleton> EditableSkeleton;
-	FSimpleMulticastDelegate& OnPostUndo;
-	FSimpleMulticastDelegate& OnAnimNotifiesChanged;
 	FSimpleMulticastDelegate& OnSectionsChanged;
 };
 
@@ -42,7 +38,7 @@ struct FMontageEditorRequiredArgs
 	portion of the Montage tool and registering callbacks to the SMontageEditor to do the actual editing.
 	
 */
-class SMontageEditor : public SAnimEditorBase 
+class SMontageEditor : public SAnimEditorBase, public FEditorUndoClient
 {
 public:
 	SLATE_BEGIN_ARGS( SMontageEditor )
@@ -54,7 +50,6 @@ public:
 		SLATE_EVENT(FSimpleDelegate, OnSectionsChanged)
 		SLATE_ARGUMENT( UAnimMontage*, Montage )
 		SLATE_EVENT(FOnObjectsSelected, OnObjectsSelected)
-		SLATE_EVENT(FSimpleDelegate, OnAnimNotifiesChanged)
 	SLATE_END_ARGS()
 
 	~SMontageEditor();
@@ -121,7 +116,12 @@ private:
 	void EnsureStartingSection();
 	void EnsureSlotNode();
 	virtual bool ClampToEndTime(float NewEndTime) override;
-	void PostUndo();
+
+	/** FEditorUndoClient interface */
+	virtual void PostUndo( bool bSuccess ) override;
+	virtual void PostRedo( bool bSuccess ) override;
+
+	void PostRedoUndo();
 	
 	bool GetSectionTime( int32 SectionIndex, float &OutTime ) const;
 
diff --git a/Engine/Source/Editor/Persona/Private/SSequenceEditor.cpp b/Engine/Source/Editor/Persona/Private/SSequenceEditor.cpp
--- a/Engine/Source/Editor/Persona/Private/SSequenceEditor.cpp
+++ b/Engine/Source/Editor/Persona/Private/SSequenceEditor.cpp
@@ -7,13 +7,14 @@
 #include "SAnimNotifyPanel.h"
 #include "SAnimTrackCurvePanel.h"
 #include "AnimPreviewInstance.h"
+#include "Editor.h"
 
 #define LOCTEXT_NAMESPACE "AnimSequenceEditor"
 
 //////////////////////////////////////////////////////////////////////////
 // SSequenceEditor
 
-void SSequenceEditor::Construct(const FArguments& InArgs, TSharedRef<class IPersonaPreviewScene> InPreviewScene, TSharedRef<class IEditableSkeleton> InEditableSkeleton, FSimpleMulticastDelegate& OnPostUndo)
+void SSequenceEditor::Construct(const FArguments& InArgs, TSharedRef<class IPersonaPreviewScene> InPreviewScene, TSharedRef<class IEditableSkeleton> InEditableSkeleton)
 {
 	SequenceObj = InArgs._Sequence;
 	check(SequenceObj);
@@ -23,13 +24,16 @@ void SSequenceEditor::Construct(const FArguments& InArgs, TSharedRef<class IPers
 		.OnObjectsSelected(InArgs._OnObjectsSelected), 
 		InPreviewScene );
 
-	OnPostUndo.Add(FSimpleDelegate::CreateSP( this, &SSequenceEditor::PostUndo ) );
+	if(GEditor)
+	{
+		GEditor->RegisterForUndo(this);
+	}
 
 	EditorPanels->AddSlot()
 	.AutoHeight()
 	.Padding(0, 10)
 	[
-		SAssignNew( AnimNotifyPanel, SAnimNotifyPanel, OnPostUndo)
+		SAssignNew( AnimNotifyPanel, SAnimNotifyPanel, InEditableSkeleton )
 		.Sequence(SequenceObj)
 		.WidgetWidth(S2ColumnWidget::DEFAULT_RIGHT_COLUMN_WIDTH)
 		.ViewInputMin(this, &SAnimEditorBase::GetViewMinInput)
@@ -39,7 +43,6 @@ void SSequenceEditor::Construct(const FArguments& InArgs, TSharedRef<class IPers
 		.OnSetInputViewRange(this, &SAnimEditorBase::SetInputViewRange)
 		.OnGetScrubValue(this, &SAnimEditorBase::GetScrubValue)
 		.OnSelectionChanged(this, &SAnimEditorBase::OnSelectionChanged)
-		.OnAnimNotifiesChanged(InArgs._OnAnimNotifiesChanged)
 		.OnInvokeTab(InArgs._OnInvokeTab)
 	];
 
@@ -78,7 +81,25 @@ void SSequenceEditor::Construct(const FArguments& InArgs, TSharedRef<class IPers
 	}
 }
 
-void SSequenceEditor::PostUndo()
+SSequenceEditor::~SSequenceEditor()
+{
+	if(GEditor)
+	{
+		GEditor->UnregisterForUndo(this);
+	}
+}
+
+void SSequenceEditor::PostUndo( bool bSuccess )
+{
+	PostUndoRedo();
+}
+
+void SSequenceEditor::PostRedo( bool bSuccess )
+{
+	PostUndoRedo();
+}
+
+void SSequenceEditor::PostUndoRedo()
 {
 	GetPreviewScene()->SetPreviewAnimationAsset(SequenceObj);
 
diff --git a/Engine/Source/Editor/Persona/Private/SSequenceEditor.h b/Engine/Source/Editor/Persona/Private/SSequenceEditor.h
--- a/Engine/Source/Editor/Persona/Private/SSequenceEditor.h
+++ b/Engine/Source/Editor/Persona/Private/SSequenceEditor.h
@@ -9,6 +9,7 @@
 #include "SAnimCurvePanel.h"
 #include "SAnimationScrubPanel.h"
 #include "SAnimEditorBase.h"
+#include "EditorUndoClient.h"
 
 class SAnimNotifyPanel;
 class SAnimTrackCurvePanel;
@@ -17,7 +18,7 @@ class SAnimTrackCurvePanel;
 // SSequenceEditor
 
 /** Overall animation sequence editing widget */
-class SSequenceEditor : public SAnimEditorBase
+class SSequenceEditor : public SAnimEditorBase, public FEditorUndoClient
 {
 public:
 	SLATE_BEGIN_ARGS( SSequenceEditor )
@@ -26,7 +27,6 @@ public:
 
 		SLATE_ARGUMENT(UAnimSequenceBase*, Sequence)
 		SLATE_EVENT(FOnObjectsSelected, OnObjectsSelected)
-		SLATE_EVENT(FSimpleDelegate, OnAnimNotifiesChanged)
 		SLATE_EVENT(FOnInvokeTab, OnInvokeTab)
 
 	SLATE_END_ARGS()
@@ -38,7 +38,9 @@ private:
 	TSharedPtr<class SAnimationScrubPanel> AnimScrubPanel;
 	TWeakPtr<class IPersonaPreviewScene> PreviewScenePtr;
 public:
-	void Construct(const FArguments& InArgs, TSharedRef<class IPersonaPreviewScene> InPreviewScene, TSharedRef<class IEditableSkeleton> InEditableSkeleton, FSimpleMulticastDelegate& OnPostUndo);
+	void Construct(const FArguments& InArgs, TSharedRef<class IPersonaPreviewScene> InPreviewScene, TSharedRef<class IEditableSkeleton> InEditableSkeleton);
+
+	~SSequenceEditor();
 
 	virtual UAnimationAsset* GetEditorObject() const override { return SequenceObj; }
 
@@ -46,6 +48,10 @@ private:
 	/** Pointer to the animation sequence being edited */
 	UAnimSequenceBase* SequenceObj;
 
+	/** FEditorUndoClient interface */
+	virtual void PostUndo( bool bSuccess ) override;
+	virtual void PostRedo( bool bSuccess ) override;
+
 	/** Post undo **/
-	void PostUndo();
+	void PostUndoRedo();
 };
diff --git a/Engine/Source/Editor/Persona/Private/SSkeletonAnimNotifies.cpp b/Engine/Source/Editor/Persona/Private/SSkeletonAnimNotifies.cpp
--- a/Engine/Source/Editor/Persona/Private/SSkeletonAnimNotifies.cpp
+++ b/Engine/Source/Editor/Persona/Private/SSkeletonAnimNotifies.cpp
@@ -15,81 +15,21 @@
 #include "ScopedTransaction.h"
 #include "IEditableSkeleton.h"
 #include "TabSpawners.h"
+#include "Editor.h"
 
 #define LOCTEXT_NAMESPACE "SkeletonAnimNotifies"
 
-static const FName ColumnId_AnimNotifyNameLabel( "AnimNotifyName" );
-
 //////////////////////////////////////////////////////////////////////////
 // SMorphTargetListRow
 
 typedef TSharedPtr< FDisplayedAnimNotifyInfo > FDisplayedAnimNotifyInfoPtr;
 
-class SAnimNotifyListRow
-	: public SMultiColumnTableRow< FDisplayedAnimNotifyInfoPtr >
-{
-public:
-
-	SLATE_BEGIN_ARGS( SAnimNotifyListRow ) {}
-
-	/** The item for this row **/
-	SLATE_ARGUMENT( FDisplayedAnimNotifyInfoPtr, Item )
-
-	/* Widget used to display the list of morph targets */
-	SLATE_ARGUMENT( TSharedPtr<SSkeletonAnimNotifies>, NotifiesListView )
-
-	SLATE_END_ARGS()
-
-	void Construct( const FArguments& InArgs, const TSharedRef<STableViewBase>& OwnerTableView );
-
-	/** Overridden from SMultiColumnTableRow.  Generates a widget for this column of the tree row. */
-	virtual TSharedRef<SWidget> GenerateWidgetForColumn( const FName& ColumnName ) override;
-
-private:
-
-	/** Widget used to display the list of notifies */
-	TSharedPtr<SSkeletonAnimNotifies> NotifiesListView;
-
-	/** The notify being displayed by this row */
-	FDisplayedAnimNotifyInfoPtr	Item;
-};
-
-void SAnimNotifyListRow::Construct( const FArguments& InArgs, const TSharedRef<STableViewBase>& InOwnerTableView )
-{
-	Item = InArgs._Item;
-	NotifiesListView = InArgs._NotifiesListView;
-
-	check( Item.IsValid() );
-
-	SMultiColumnTableRow< FDisplayedAnimNotifyInfoPtr >::Construct( FSuperRowType::FArguments(), InOwnerTableView );
-}
-
-TSharedRef< SWidget > SAnimNotifyListRow::GenerateWidgetForColumn( const FName& ColumnName )
-{
-	check( ColumnName == ColumnId_AnimNotifyNameLabel );
-
-	return SNew( SVerticalBox )
-	+ SVerticalBox::Slot()
-	.AutoHeight()
-	.Padding( 0.0f, 4.0f )
-	.VAlign( VAlign_Center )
-	[
-		SAssignNew(Item->InlineEditableText, SInlineEditableTextBlock)
-		.Text( FText::FromName(Item->Name) )
-		.OnVerifyTextChanged(NotifiesListView.Get(), &SSkeletonAnimNotifies::OnVerifyNotifyNameCommit, Item)
-		.OnTextCommitted(NotifiesListView.Get(), &SSkeletonAnimNotifies::OnNotifyNameCommitted, Item)
-		.IsSelected(NotifiesListView.Get(), &SSkeletonAnimNotifies::IsSelected)
-	];
-}
-
 /////////////////////////////////////////////////////
 // FSkeletonAnimNotifiesSummoner
 
-FSkeletonAnimNotifiesSummoner::FSkeletonAnimNotifiesSummoner(TSharedPtr<class FAssetEditorToolkit> InHostingApp, const TSharedRef<class IEditableSkeleton>& InEditableSkeleton, FSimpleMulticastDelegate& InOnChangeAnimNotifies, FSimpleMulticastDelegate& InOnPostUndo, FOnObjectsSelected InOnObjectsSelected)
+FSkeletonAnimNotifiesSummoner::FSkeletonAnimNotifiesSummoner(TSharedPtr<class FAssetEditorToolkit> InHostingApp, const TSharedRef<class IEditableSkeleton>& InEditableSkeleton, FOnObjectsSelected InOnObjectsSelected)
 	: FWorkflowTabFactory(FPersonaTabs::SkeletonAnimNotifiesID, InHostingApp)
 	, EditableSkeleton(InEditableSkeleton)
-	, OnChangeAnimNotifies(InOnChangeAnimNotifies)
-	, OnPostUndo(InOnPostUndo)
 	, OnObjectsSelected(InOnObjectsSelected)
 {
 	TabLabel = LOCTEXT("SkeletonAnimNotifiesTabTitle", "Animation Notifies");
@@ -104,23 +44,30 @@ FSkeletonAnimNotifiesSummoner::FSkeletonAnimNotifiesSummoner(TSharedPtr<class FA
 
 TSharedRef<SWidget> FSkeletonAnimNotifiesSummoner::CreateTabBody(const FWorkflowTabSpawnInfo& Info) const
 {
-	return SNew(SSkeletonAnimNotifies, EditableSkeleton.Pin().ToSharedRef(), OnChangeAnimNotifies, OnPostUndo)
+	return SNew(SSkeletonAnimNotifies, EditableSkeleton.Pin().ToSharedRef())
 		.OnObjectsSelected(OnObjectsSelected);
 }
 
 /////////////////////////////////////////////////////
 // SSkeletonAnimNotifies
 
-void SSkeletonAnimNotifies::Construct(const FArguments& InArgs, const TSharedRef<IEditableSkeleton>& InEditableSkeleton, FSimpleMulticastDelegate& InOnChangeAnimNotifies, FSimpleMulticastDelegate& InOnPostUndo)
+void SSkeletonAnimNotifies::Construct(const FArguments& InArgs, const TSharedRef<IEditableSkeleton>& InEditableSkeleton)
 {
 	OnObjectsSelected = InArgs._OnObjectsSelected;
-	EditableSkeleton = InEditableSkeleton;
+	OnNotifySelected = InArgs._OnNotifySelected;
+	bIsPicker = InArgs._IsPicker;
 
-	InOnChangeAnimNotifies.Add(FSimpleDelegate::CreateSP( this, &SSkeletonAnimNotifies::PostUndo ) );
-	InOnPostUndo.Add(FSimpleDelegate::CreateSP( this, &SSkeletonAnimNotifies::RefreshNotifiesListWithFilter ) );
+	EditableSkeleton = InEditableSkeleton;
 
 	EditableSkeleton->RegisterOnNotifiesChanged(FSimpleDelegate::CreateSP(this, &SSkeletonAnimNotifies::OnNotifiesChanged));
 
+	if(GEditor)
+	{
+		GEditor->RegisterForUndo(this);
+	}
+
+	FOnContextMenuOpening OnContextMenuOpening = !bIsPicker ? FOnContextMenuOpening::CreateSP(this, &SSkeletonAnimNotifies::OnGetContextMenuContent) : FOnContextMenuOpening();
+
 	this->ChildSlot
 	[
 		SNew( SVerticalBox )
@@ -142,21 +89,24 @@ void SSkeletonAnimNotifies::Construct(const FArguments& InArgs, const TSharedRef
 			SAssignNew( NotifiesListView, SAnimNotifyListType )
 			.ListItemsSource( &NotifyList )
 			.OnGenerateRow( this, &SSkeletonAnimNotifies::GenerateNotifyRow )
-			.OnContextMenuOpening( this, &SSkeletonAnimNotifies::OnGetContextMenuContent )
+			.OnContextMenuOpening( OnContextMenuOpening )
 			.OnSelectionChanged( this, &SSkeletonAnimNotifies::OnNotifySelectionChanged )
 			.ItemHeight( 22.0f )
-			.HeaderRow
-			(
-				SNew( SHeaderRow )
-				+ SHeaderRow::Column( ColumnId_AnimNotifyNameLabel )
-				.DefaultLabel( LOCTEXT( "AnimNotifyNameLabel", "Notify Name" ) )
-			)
+			.OnItemScrolledIntoView( this, &SSkeletonAnimNotifies::OnItemScrolledIntoView )
 		]
 	];
 
 	CreateNotifiesList();
 }
 
+SSkeletonAnimNotifies::~SSkeletonAnimNotifies()
+{
+	if(GEditor)
+	{
+		GEditor->UnregisterForUndo(this);
+	}
+}
+
 void SSkeletonAnimNotifies::OnNotifiesChanged()
 {
 	RefreshNotifiesListWithFilter();
@@ -180,9 +130,23 @@ TSharedRef<ITableRow> SSkeletonAnimNotifies::GenerateNotifyRow(TSharedPtr<FDispl
 	check( InInfo.IsValid() );
 
 	return
-		SNew( SAnimNotifyListRow, OwnerTable )
-		.Item( InInfo )
-		.NotifiesListView( SharedThis( this ) );
+		SNew( STableRow<TSharedPtr<FDisplayedAnimNotifyInfo>>, OwnerTable )
+		[
+			SNew( SVerticalBox )
+			+ SVerticalBox::Slot()
+			.AutoHeight()
+			.Padding( 0.0f, 4.0f )
+			.VAlign( VAlign_Center )
+			[
+				SAssignNew(InInfo->InlineEditableText, SInlineEditableTextBlock)
+				.Text(FText::FromName(InInfo->Name))
+				.OnVerifyTextChanged(this, &SSkeletonAnimNotifies::OnVerifyNotifyNameCommit, InInfo)
+				.OnTextCommitted(this, &SSkeletonAnimNotifies::OnNotifyNameCommitted, InInfo)
+				.IsSelected(this, &SSkeletonAnimNotifies::IsSelected)
+				.HighlightText_Lambda([this](){ return FilterText; })
+				.IsReadOnly(bIsPicker)
+			]
+		];
 }
 
 TSharedPtr<SWidget> SSkeletonAnimNotifies::OnGetContextMenuContent() const
@@ -190,16 +154,17 @@ TSharedPtr<SWidget> SSkeletonAnimNotifies::OnGetContextMenuContent() const
 	const bool bShouldCloseWindowAfterMenuSelection = true;
 	FMenuBuilder MenuBuilder( bShouldCloseWindowAfterMenuSelection, NULL);
 
-	MenuBuilder.BeginSection("AnimNotifyAction", LOCTEXT( "AnimNotifyActions", "Selected Notify Actions" ) );
+	MenuBuilder.BeginSection("AnimNotifyAction", LOCTEXT( "AnimNotifyActions", "Notifies" ) );
 	{
-		{
-			FUIAction Action = FUIAction( FExecuteAction::CreateSP( this, &SSkeletonAnimNotifies::OnDeleteAnimNotify ), 
-				FCanExecuteAction::CreateSP( this, &SSkeletonAnimNotifies::CanPerformDelete ) );
-			const FText Label = LOCTEXT("DeleteAnimNotifyButtonLabel", "Delete");
-			const FText ToolTipText = LOCTEXT("DeleteAnimNotifyButtonTooltip", "Deletes the selected anim notifies.");
-			MenuBuilder.AddMenuEntry( Label, ToolTipText, FSlateIcon(), Action);
-		}
+		FUIAction Action = FUIAction( FExecuteAction::CreateSP( this, &SSkeletonAnimNotifies::OnAddAnimNotify ) );
+		const FText Label = LOCTEXT("NewAnimNotifyButtonLabel", "New...");
+		const FText ToolTipText = LOCTEXT("NewAnimNotifyButtonTooltip", "Creates a new anim notify.");
+		MenuBuilder.AddMenuEntry( Label, ToolTipText, FSlateIcon(), Action);
+	}
+	MenuBuilder.EndSection();
 
+	MenuBuilder.BeginSection("AnimNotifyAction", LOCTEXT( "SelectedAnimNotifyActions", "Selected Notify Actions" ) );
+	{
 		{
 			FUIAction Action = FUIAction( FExecuteAction::CreateSP( this, &SSkeletonAnimNotifies::OnRenameAnimNotify ), 
 				FCanExecuteAction::CreateSP( this, &SSkeletonAnimNotifies::CanPerformRename ) );
@@ -207,6 +172,16 @@ TSharedPtr<SWidget> SSkeletonAnimNotifies::OnGetContextMenuContent() const
 			const FText ToolTipText = LOCTEXT("RenameAnimNotifyButtonTooltip", "Renames the selected anim notifies.");
 			MenuBuilder.AddMenuEntry( Label, ToolTipText, FSlateIcon(), Action);
 		}
+
+		{
+			FUIAction Action = FUIAction( FExecuteAction::CreateSP( this, &SSkeletonAnimNotifies::OnDeleteAnimNotify ), 
+				FCanExecuteAction::CreateSP( this, &SSkeletonAnimNotifies::CanPerformDelete ) );
+			const FText Label = LOCTEXT("DeleteAnimNotifyButtonLabel", "Delete");
+			const FText ToolTipText = LOCTEXT("DeleteAnimNotifyButtonTooltip", "Deletes the selected anim notifies.");
+			MenuBuilder.AddMenuEntry( Label, ToolTipText, FSlateIcon(), Action);
+		}
+
+
 	}
 	MenuBuilder.EndSection();
 
@@ -218,19 +193,50 @@ void SSkeletonAnimNotifies::OnNotifySelectionChanged(TSharedPtr<FDisplayedAnimNo
 	if(Selection.IsValid())
 	{
 		ShowNotifyInDetailsView(Selection->Name);
+
+		OnNotifySelected.ExecuteIfBound(Selection->Name);
 	}
 }
 
 bool SSkeletonAnimNotifies::CanPerformDelete() const
 {
-	TArray< TSharedPtr< FDisplayedAnimNotifyInfo > > SelectedRows = NotifiesListView->GetSelectedItems();
-	return SelectedRows.Num() > 0;
+	return NotifiesListView->GetNumItemsSelected() > 0;
 }
 
 bool SSkeletonAnimNotifies::CanPerformRename() const
 {
-	TArray< TSharedPtr< FDisplayedAnimNotifyInfo > > SelectedRows = NotifiesListView->GetSelectedItems();
-	return SelectedRows.Num() == 1;
+	return NotifiesListView->GetNumItemsSelected() == 1;
+}
+
+void SSkeletonAnimNotifies::OnAddAnimNotify()
+{
+	// Find a unique name for this notify
+	const TCHAR* BaseNotifyString = TEXT("NewNotify");
+	FString NewNotifyString = BaseNotifyString;
+	int32 NumericExtension = 0;
+
+	while(EditableSkeleton->GetSkeleton().AnimationNotifies.ContainsByPredicate([&NewNotifyString](const FName& InNotifyName){ return InNotifyName.ToString() == NewNotifyString; }))
+	{
+		NewNotifyString = FString::Printf(TEXT("%s_%d"), BaseNotifyString, NumericExtension);
+		NumericExtension++;
+	}
+
+	// Add an item. The subsequent rename will commit the item.
+	TSharedPtr<FDisplayedAnimNotifyInfo> NewItem = FDisplayedAnimNotifyInfo::Make(*NewNotifyString);
+	NewItem->bIsNew = true;
+	NotifyList.Add(NewItem);
+
+	NotifiesListView->ClearSelection();
+	NotifiesListView->RequestListRefresh();
+	NotifiesListView->RequestScrollIntoView(NewItem);
+}
+
+void SSkeletonAnimNotifies::OnItemScrolledIntoView(TSharedPtr<FDisplayedAnimNotifyInfo> InItem, const TSharedPtr<ITableRow>& InTableRow)
+{
+	if(InItem.IsValid() && InItem->InlineEditableText.IsValid() && InItem->bIsNew)
+	{
+		InItem->InlineEditableText->EnterEditingMode();
+	}
 }
 
 void SSkeletonAnimNotifies::OnDeleteAnimNotify()
@@ -283,7 +289,7 @@ bool SSkeletonAnimNotifies::OnVerifyNotifyNameCommit( const FText& NewName, FTex
 	}
 
 	FName NotifyName( *NewName.ToString() );
-	if(NotifyName != Item->Name)
+	if(NotifyName != Item->Name || Item->bIsNew)
 	{
 		if(EditableSkeleton->GetSkeleton().AnimationNotifies.Contains(NotifyName))
 		{
@@ -297,22 +303,34 @@ bool SSkeletonAnimNotifies::OnVerifyNotifyNameCommit( const FText& NewName, FTex
 
 void SSkeletonAnimNotifies::OnNotifyNameCommitted( const FText& NewName, ETextCommit::Type, TSharedPtr<FDisplayedAnimNotifyInfo> Item )
 {
-	int32 NumAnimationsModified = EditableSkeleton->RenameNotify(FName(*NewName.ToString()), Item->Name);
-
-	if(NumAnimationsModified > 0)
+	FName NewFName = FName(*NewName.ToString());
+	if(Item->bIsNew)
 	{
-		// Tell the user that the socket is a duplicate
-		FFormatNamedArguments Args;
-		Args.Add( TEXT("NumAnimationsModified"), NumAnimationsModified );
-		FNotificationInfo Info( FText::Format( LOCTEXT( "AnimNotifiesRenamed", "{NumAnimationsModified} animation(s) modified to rename notification" ), Args ) );
+		EditableSkeleton->AddNotify(NewFName);
+		Item->bIsNew = false;
+	}
+	else
+	{
+		if(NewFName != Item->Name)
+		{
+			int32 NumAnimationsModified = EditableSkeleton->RenameNotify(FName(*NewName.ToString()), Item->Name);
 
-		Info.bUseLargeFont = false;
-		Info.ExpireDuration = 5.0f;
+			if(NumAnimationsModified > 0)
+			{
+				// Tell the user that the socket is a duplicate
+				FFormatNamedArguments Args;
+				Args.Add( TEXT("NumAnimationsModified"), NumAnimationsModified );
+				FNotificationInfo Info( FText::Format( LOCTEXT( "AnimNotifiesRenamed", "{NumAnimationsModified} animation(s) modified to rename notification" ), Args ) );
 
-		NotifyUser( Info );
-	}
+				Info.bUseLargeFont = false;
+				Info.ExpireDuration = 5.0f;
 
-	RefreshNotifiesListWithFilter();
+				NotifyUser( Info );
+			}
+
+			RefreshNotifiesListWithFilter();
+		}
+	}
 }
 
 void SSkeletonAnimNotifies::RefreshNotifiesListWithFilter()
@@ -345,33 +363,16 @@ void SSkeletonAnimNotifies::CreateNotifiesList( const FString& SearchText )
 
 void SSkeletonAnimNotifies::ShowNotifyInDetailsView(FName NotifyName)
 {
-	UEditorSkeletonNotifyObj *Obj = Cast<UEditorSkeletonNotifyObj>(ShowInDetailsView(UEditorSkeletonNotifyObj::StaticClass()));
-	if(Obj != NULL)
+	if(OnObjectsSelected.IsBound())
 	{
-		Obj->AnimationNames.Empty();
-
-		TArray<FAssetData> CompatibleAnimSequences;
-		EditableSkeleton->GetCompatibleAnimSequences(CompatibleAnimSequences);
+		ClearDetailsView();
 
-		for( int32 AssetIndex = 0; AssetIndex < CompatibleAnimSequences.Num(); ++AssetIndex )
+		UEditorSkeletonNotifyObj *Obj = Cast<UEditorSkeletonNotifyObj>(ShowInDetailsView(UEditorSkeletonNotifyObj::StaticClass()));
+		if(Obj != NULL)
 		{
-			const FAssetData& PossibleAnimSequence = CompatibleAnimSequences[AssetIndex];
-			if(UObject* AnimSeqAsset = PossibleAnimSequence.GetAsset())
-			{
-				UAnimSequenceBase* Sequence = CastChecked<UAnimSequenceBase>(AnimSeqAsset);
-				for (int32 NotifyIndex = 0; NotifyIndex < Sequence->Notifies.Num(); ++NotifyIndex)
-				{
-					FAnimNotifyEvent& NotifyEvent = Sequence->Notifies[NotifyIndex];
-					if (NotifyEvent.NotifyName == NotifyName)
-					{
-						Obj->AnimationNames.Add(MakeShareable(new FString(PossibleAnimSequence.AssetName.ToString())));
-						break;
-					}
-				}
-			}
+			Obj->EditableSkeleton = EditableSkeleton;
+			Obj->Name = NotifyName;
 		}
-
-		Obj->Name = NotifyName;
 	}
 }
 
@@ -394,7 +395,12 @@ void SSkeletonAnimNotifies::ClearDetailsView()
 	OnObjectsSelected.ExecuteIfBound(Objects);
 }
 
-void SSkeletonAnimNotifies::PostUndo()
+void SSkeletonAnimNotifies::PostUndo( bool bSuccess )
+{
+	RefreshNotifiesListWithFilter();
+}
+
+void SSkeletonAnimNotifies::PostRedo( bool bSuccess )
 {
 	RefreshNotifiesListWithFilter();
 }
diff --git a/Engine/Source/Editor/Persona/Private/SSkeletonAnimNotifies.h b/Engine/Source/Editor/Persona/Private/SSkeletonAnimNotifies.h
--- a/Engine/Source/Editor/Persona/Private/SSkeletonAnimNotifies.h
+++ b/Engine/Source/Editor/Persona/Private/SSkeletonAnimNotifies.h
@@ -15,6 +15,7 @@
 #include "Widgets/Views/STableRow.h"
 #include "Widgets/Views/SListView.h"
 #include "EditorObjectsTracker.h"
+#include "EditorUndoClient.h"
 
 #define LOCTEXT_NAMESPACE "SkeletonAnimnotifies"
 
@@ -22,12 +23,15 @@ class IEditableSkeleton;
 class SToolTip;
 struct FNotificationInfo;
 
+/** Delegate fired when a notify is selected */
+DECLARE_DELEGATE_OneParam(FOnNotifySelected, const FName& /*InSelectedNotify*/);
+
 /////////////////////////////////////////////////////
 // FSkeletonAnimNotifiesSummoner
 struct FSkeletonAnimNotifiesSummoner : public FWorkflowTabFactory
 {
 public:
-	FSkeletonAnimNotifiesSummoner(TSharedPtr<class FAssetEditorToolkit> InHostingApp, const TSharedRef<class IEditableSkeleton>& InEditableSkeleton, FSimpleMulticastDelegate& InOnChangeAnimNotifies, FSimpleMulticastDelegate& InOnPostUndo, FOnObjectsSelected InOnObjectsSelected);
+	FSkeletonAnimNotifiesSummoner(TSharedPtr<class FAssetEditorToolkit> InHostingApp, const TSharedRef<class IEditableSkeleton>& InEditableSkeleton, FOnObjectsSelected InOnObjectsSelected);
 
 	virtual TSharedRef<SWidget> CreateTabBody(const FWorkflowTabSpawnInfo& Info) const override;
 
@@ -39,8 +43,6 @@ public:
 
 private:
 	TWeakPtr<class IEditableSkeleton> EditableSkeleton;
-	FSimpleMulticastDelegate& OnChangeAnimNotifies;
-	FSimpleMulticastDelegate& OnPostUndo;
 	FOnObjectsSelected OnObjectsSelected;
 };
 
@@ -55,6 +57,9 @@ public:
 	/** Handle to editable text block for rename */
 	TSharedPtr<SInlineEditableTextBlock> InlineEditableText;
 
+	/** Flag to say whether this is a new item we are creating */
+	bool bIsNew;
+
 	/** Static function for creating a new item, but ensures that you can only have a TSharedRef to one */
 	static TSharedRef<FDisplayedAnimNotifyInfo> Make(const FName& NotifyName)
 	{
@@ -65,6 +70,7 @@ protected:
 	/** Hidden constructor, always use Make above */
 	FDisplayedAnimNotifyInfo(const FName& InNotifyName)
 		: Name( InNotifyName )
+		, bIsNew(false)
 	{}
 
 	/** Hidden constructor, always use Make above */
@@ -74,17 +80,27 @@ protected:
 /** Widgets list type */
 typedef SListView< TSharedPtr<FDisplayedAnimNotifyInfo> > SAnimNotifyListType;
 
-class SSkeletonAnimNotifies : public SCompoundWidget, public FGCObject
+class SSkeletonAnimNotifies : public SCompoundWidget, public FGCObject, public FEditorUndoClient
 {
 public:
 	SLATE_BEGIN_ARGS( SSkeletonAnimNotifies )
+		: _IsPicker(false)
 	{}
 
+	/** Delegate called to select an object in the details panel */
 	SLATE_EVENT(FOnObjectsSelected, OnObjectsSelected)
-		
+
+	/** Delegate fired when a notify is selected */
+	SLATE_EVENT(FOnNotifySelected, OnNotifySelected)
+
+	/** Whether we should use this dialog as a picker or an editor. In picker mode we cant add, remove or rename notifies. */
+	SLATE_ARGUMENT(bool, IsPicker)
+
 	SLATE_END_ARGS()
 public:
-	void Construct(const FArguments& InArgs, const TSharedRef<IEditableSkeleton>& InEditableSkeleton, FSimpleMulticastDelegate& InOnChangeAnimNotifies, FSimpleMulticastDelegate& InOnPostUndo);
+	void Construct(const FArguments& InArgs, const TSharedRef<IEditableSkeleton>& InEditableSkeleton);
+
+	~SSkeletonAnimNotifies();
 
 	/**
 	* Accessor so our rows can grab the filter text for highlighting
@@ -98,8 +114,9 @@ public:
 	/** Clears the detail view of whatever we displayed last */
 	void ClearDetailsView();
 
-	/** This triggers a UI repopulation after undo has been called */
-	void PostUndo();
+	/** FEditorUndoClient interface */
+	virtual void PostUndo( bool bSuccess ) override;
+	virtual void PostRedo( bool bSuccess ) override;
 
 	// FGCObject interface start
 	virtual void AddReferencedObjects( FReferenceCollector& Collector ) override;
@@ -139,6 +156,9 @@ private:
 	/** Delegate handler for determining whether we can show the rename menu options */
 	bool CanPerformRename() const;
 
+	/** Delegate handler for adding anim notifies */
+	void OnAddAnimNotify();
+
 	/** Delegate handler for renaming anim notifies */
 	void OnRenameAnimNotify();
 
@@ -157,6 +177,9 @@ private:
 	/** Handler function for when notifies are modified on the skeleton */
 	void OnNotifiesChanged();
 
+	/** Handle when an item is scrolled into view, triggers a rename for new items */
+	void OnItemScrolledIntoView(TSharedPtr<FDisplayedAnimNotifyInfo> InItem, const TSharedPtr<ITableRow>& InTableRow);
+
 	/** The skeleton we are currently editing */
 	TSharedPtr<class IEditableSkeleton> EditableSkeleton;
 
@@ -177,6 +200,12 @@ private:
 
 	/** Delegate called to select an object in the details panel */
 	FOnObjectsSelected OnObjectsSelected;
+
+	/** Delegate fired when a notify is selected */
+	FOnNotifySelected OnNotifySelected;
+
+	/** Whether we should use this dialog as a picker or an editor. In picker mode we cant add, remove or rename notifies. */
+	bool bIsPicker;
 };
 
 #undef LOCTEXT_NAMESPACE
diff --git a/Engine/Source/Editor/Persona/Public/PersonaModule.h b/Engine/Source/Editor/Persona/Public/PersonaModule.h
--- a/Engine/Source/Editor/Persona/Public/PersonaModule.h
+++ b/Engine/Source/Editor/Persona/Public/PersonaModule.h
@@ -78,12 +78,11 @@ struct FPersonaToolkitArgs
 
 struct FAnimDocumentArgs
 {
-	FAnimDocumentArgs(const TSharedRef<class IPersonaPreviewScene>& InPreviewScene, const TSharedRef<class IPersonaToolkit>& InPersonaToolkit, const TSharedRef<class IEditableSkeleton>& InEditableSkeleton, FSimpleMulticastDelegate& InOnPostUndo, FSimpleMulticastDelegate& InOnAnimNotifiesChanged, FSimpleMulticastDelegate& InOnSectionsChanged)
+	FAnimDocumentArgs(const TSharedRef<class IPersonaPreviewScene>& InPreviewScene, const TSharedRef<class IPersonaToolkit>& InPersonaToolkit, const TSharedRef<class IEditableSkeleton>& InEditableSkeleton, FSimpleMulticastDelegate& InOnPostUndo, FSimpleMulticastDelegate& InOnSectionsChanged)
 		: PreviewScene(InPreviewScene)
 		, PersonaToolkit(InPersonaToolkit)
 		, EditableSkeleton(InEditableSkeleton)
 		, OnPostUndo(InOnPostUndo)
-		, OnAnimNotifiesChanged(InOnAnimNotifiesChanged)
 		, OnSectionsChanged(InOnSectionsChanged)
 	{}
 
@@ -92,14 +91,12 @@ struct FAnimDocumentArgs
 	TWeakPtr<class IPersonaToolkit> PersonaToolkit;
 	TWeakPtr<class IEditableSkeleton> EditableSkeleton;
 	FSimpleMulticastDelegate& OnPostUndo;
-	FSimpleMulticastDelegate& OnAnimNotifiesChanged;
 	FSimpleMulticastDelegate& OnSectionsChanged;
 
 	/** Optional args */
 	FOnObjectsSelected OnDespatchObjectsSelected;
 	FOnInvokeTab OnDespatchInvokeTab;
 	FSimpleDelegate OnDespatchSectionsChanged;
-	FSimpleDelegate OnDespatchAnimNotifiesChanged;
 };
 
 /** Places that viewport text can be placed */
@@ -218,7 +215,7 @@ public:
 	virtual void RegisterPersonaViewportTabFactories(FWorkflowAllowedTabSet& TabSet, const TSharedRef<class FWorkflowCentricApplication>& InHostingApp, const FPersonaViewportArgs& InArgs) const;
 
 	/** Create an anim notifies tab factory */
-	virtual TSharedRef<FWorkflowTabFactory> CreateAnimNotifiesTabFactory(const TSharedRef<class FWorkflowCentricApplication>& InHostingApp, const TSharedRef<class IEditableSkeleton>& InEditableSkeleton, FSimpleMulticastDelegate& InOnChangeAnimNotifies, FSimpleMulticastDelegate& InOnPostUndo, FOnObjectsSelected InOnObjectsSelected) const;
+	virtual TSharedRef<FWorkflowTabFactory> CreateAnimNotifiesTabFactory(const TSharedRef<class FWorkflowCentricApplication>& InHostingApp, const TSharedRef<class IEditableSkeleton>& InEditableSkeleton, FOnObjectsSelected InOnObjectsSelected) const;
 
 	/** Create a skeleton cuve viewer tab factory */
 	virtual TSharedRef<FWorkflowTabFactory> CreateCurveViewerTabFactory(const TSharedRef<class FWorkflowCentricApplication>& InHostingApp, const TSharedRef<class IEditableSkeleton>& InEditableSkeleton, const TSharedRef<class IPersonaPreviewScene>& InPreviewScene, FSimpleMulticastDelegate& InOnPostUndo, FOnObjectsSelected InOnObjectsSelected) const;
diff --git a/Engine/Source/Editor/SkeletalMeshEditor/Private/SkeletalMeshEditor.h b/Engine/Source/Editor/SkeletalMeshEditor/Private/SkeletalMeshEditor.h
--- a/Engine/Source/Editor/SkeletalMeshEditor/Private/SkeletalMeshEditor.h
+++ b/Engine/Source/Editor/SkeletalMeshEditor/Private/SkeletalMeshEditor.h
@@ -137,9 +137,6 @@ private:
 	void BindCommands();
 
 public:
-	/** Multicast delegate fired on anim notifies changing */
-	FSimpleMulticastDelegate OnChangeAnimNotifies;
-
 	/** Multicast delegate fired on global undo/redo */
 	FSimpleMulticastDelegate OnPostUndo;
 
diff --git a/Engine/Source/Editor/SkeletonEditor/Private/EditableSkeleton.cpp b/Engine/Source/Editor/SkeletonEditor/Private/EditableSkeleton.cpp
--- a/Engine/Source/Editor/SkeletonEditor/Private/EditableSkeleton.cpp
+++ b/Engine/Source/Editor/SkeletonEditor/Private/EditableSkeleton.cpp
@@ -1151,6 +1151,11 @@ int32 FEditableSkeleton::RenameNotify(const FName& NewName, const FName& OldName
 	return NumAnimationsModified;
 }
 
+void FEditableSkeleton::BroadcastNotifyChanged()
+{
+	OnNotifiesChanged.Broadcast();
+}
+
 void FEditableSkeleton::GetCompatibleAnimSequences(TArray<struct FAssetData>& OutAssets)
 {
 	//Get the skeleton tag to search for
diff --git a/Engine/Source/Editor/SkeletonEditor/Private/EditableSkeleton.h b/Engine/Source/Editor/SkeletonEditor/Private/EditableSkeleton.h
--- a/Engine/Source/Editor/SkeletonEditor/Private/EditableSkeleton.h
+++ b/Engine/Source/Editor/SkeletonEditor/Private/EditableSkeleton.h
@@ -39,6 +39,7 @@ public:
 	virtual int32 DeleteAnimNotifies(const TArray<FName>& InSelectedNotifyNames) override;
 	virtual void AddNotify(FName NewName) override;
 	virtual int32 RenameNotify(const FName& NewName, const FName& OldName) override;
+	virtual void BroadcastNotifyChanged() override;
 	virtual void GetCompatibleAnimSequences(TArray<struct FAssetData>& OutAssets) override;
 	virtual void RenameSocket(const FName& OldSocketName, const FName& NewSocketName, USkeletalMesh* InSkeletalMesh) override;
 	virtual void SetSocketParent(const FName& SocketName, const FName& NewParentName, USkeletalMesh* InSkeletalMesh) override;
diff --git a/Engine/Source/Editor/SkeletonEditor/Private/SkeletonEditor.h b/Engine/Source/Editor/SkeletonEditor/Private/SkeletonEditor.h
--- a/Engine/Source/Editor/SkeletonEditor/Private/SkeletonEditor.h
+++ b/Engine/Source/Editor/SkeletonEditor/Private/SkeletonEditor.h
@@ -106,9 +106,6 @@ private:
 	void OnImportAsset();
 
 public:
-	/** Multicast delegate fired on anim notifies changing */
-	FSimpleMulticastDelegate OnChangeAnimNotifies;
-
 	/** Multicast delegate fired on global undo/redo */
 	FSimpleMulticastDelegate OnPostUndo;
 
diff --git a/Engine/Source/Editor/SkeletonEditor/Private/SkeletonEditorMode.cpp b/Engine/Source/Editor/SkeletonEditor/Private/SkeletonEditorMode.cpp
--- a/Engine/Source/Editor/SkeletonEditor/Private/SkeletonEditorMode.cpp
+++ b/Engine/Source/Editor/SkeletonEditor/Private/SkeletonEditorMode.cpp
@@ -28,7 +28,7 @@ FSkeletonEditorMode::FSkeletonEditorMode(TSharedRef<FWorkflowCentricApplication>
 
 	PersonaModule.RegisterPersonaViewportTabFactories(TabFactories, InHostingApp, ViewportArgs);
 
-	TabFactories.RegisterFactory(PersonaModule.CreateAnimNotifiesTabFactory(InHostingApp, InSkeletonTree->GetEditableSkeleton(), SkeletonEditor->OnChangeAnimNotifies, SkeletonEditor->OnPostUndo, OnObjectsSelected));
+	TabFactories.RegisterFactory(PersonaModule.CreateAnimNotifiesTabFactory(InHostingApp, InSkeletonTree->GetEditableSkeleton(),  OnObjectsSelected));
 	TabFactories.RegisterFactory(PersonaModule.CreateAdvancedPreviewSceneTabFactory(InHostingApp, SkeletonEditor->GetPersonaToolkit()->GetPreviewScene()));
 	TabFactories.RegisterFactory(PersonaModule.CreateRetargetManagerTabFactory(InHostingApp, InSkeletonTree->GetEditableSkeleton(), SkeletonEditor->GetPersonaToolkit()->GetPreviewScene(), SkeletonEditor->OnPostUndo));
 	TabFactories.RegisterFactory(PersonaModule.CreateCurveViewerTabFactory(InHostingApp, InSkeletonTree->GetEditableSkeleton(), SkeletonEditor->GetPersonaToolkit()->GetPreviewScene(), SkeletonEditor->OnPostUndo, OnObjectsSelected));
diff --git a/Engine/Source/Editor/SkeletonEditor/Public/IEditableSkeleton.h b/Engine/Source/Editor/SkeletonEditor/Public/IEditableSkeleton.h
--- a/Engine/Source/Editor/SkeletonEditor/Public/IEditableSkeleton.h
+++ b/Engine/Source/Editor/SkeletonEditor/Public/IEditableSkeleton.h
@@ -103,6 +103,9 @@ public:
 	 */	
 	virtual int32 RenameNotify(const FName& NewName, const FName& OldName) = 0;
 
+	/** Inform the system that something about a notify changed */	
+	virtual void BroadcastNotifyChanged() = 0;
+
 	/** Populates OutAssets with the AnimSequences that match this current skeleton */
 	virtual void GetCompatibleAnimSequences(TArray<struct FAssetData>& OutAssets) = 0;
 
diff --git a/Engine/Source/Editor/UnrealEd/Classes/Animation/EditorSkeletonNotifyObj.h b/Engine/Source/Editor/UnrealEd/Classes/Animation/EditorSkeletonNotifyObj.h
--- a/Engine/Source/Editor/UnrealEd/Classes/Animation/EditorSkeletonNotifyObj.h
+++ b/Engine/Source/Editor/UnrealEd/Classes/Animation/EditorSkeletonNotifyObj.h
@@ -10,7 +10,7 @@
 #include "UObject/Object.h"
 #include "EditorSkeletonNotifyObj.generated.h"
 
-class UAnimSequenceBase;
+class IEditableSkeleton;
 
 DECLARE_DELEGATE_TwoParams( FOnAnimObjectChange, class UObject*, bool)
 
@@ -19,11 +19,10 @@ class UEditorSkeletonNotifyObj : public UObject
 {
 	GENERATED_UCLASS_BODY()
 public:
-	
-	UObject* EditedObject;
+	/** Current skeleton to try to grab objects from */
+	TWeakPtr<IEditableSkeleton> EditableSkeleton;
 
+	/** The name of the notify we represent */
 	UPROPERTY(VisibleAnywhere, Category=SkeletonNotifies)
 	FName Name;
-
-	TArray< TSharedPtr<FString> > AnimationNames;
 };
