Added notification to the viewport when a min LOD is being applied

diff --git a/Engine/Source/Editor/Persona/Private/SAnimationEditorViewport.cpp b/Engine/Source/Editor/Persona/Private/SAnimationEditorViewport.cpp
--- a/Engine/Source/Editor/Persona/Private/SAnimationEditorViewport.cpp
+++ b/Engine/Source/Editor/Persona/Private/SAnimationEditorViewport.cpp
@@ -500,6 +500,8 @@ void SAnimationEditorViewportTabBody::Construct(const FArguments& InArgs, const
 	GetPreviewScene()->OnRecordingStateChanged().AddSP(this, &SAnimationEditorViewportTabBody::AddRecordingNotification);
 
 	AddPostProcessNotification();
+
+	AddMinLODNotification();
 }
 
 void SAnimationEditorViewportTabBody::BindCommands()
@@ -2159,6 +2161,55 @@ void SAnimationEditorViewportTabBody::AddPostProcessNotification()
 	);
 }
 
+void SAnimationEditorViewportTabBody::AddMinLODNotification()
+{
+	if(WeakMinLODNotification.IsValid())
+	{
+		return;
+	}
+
+	auto GetMinLODNotificationVisibility = [this]()
+	{
+		if (GetPreviewScene()->GetPreviewMesh() && GetPreviewScene()->GetPreviewMesh()->MinLod.Default != 0)
+		{
+			return EVisibility::Visible;
+		}
+
+		return EVisibility::Collapsed;
+	};
+
+	WeakMinLODNotification = AddNotification(EMessageSeverity::Info,
+		true,
+		SNew(SHorizontalBox)
+		.Visibility_Lambda(GetMinLODNotificationVisibility)
+		.ToolTipText(LOCTEXT("MinLODNotificationTooltip", "This asset has a minimum LOD applied."))
+		+SHorizontalBox::Slot()
+		.FillWidth(1.0f)
+		.Padding(2.0f, 4.0f)
+		[
+			SNew(SHorizontalBox)
+			+SHorizontalBox::Slot()
+			.AutoWidth()
+			.VAlign(VAlign_Center)
+			.Padding(0.0f, 0.0f, 4.0f, 0.0f)
+			[
+				SNew(STextBlock)
+				.TextStyle(FEditorStyle::Get(), "AnimViewport.MessageText")
+				.Font(FEditorStyle::Get().GetFontStyle("FontAwesome.9"))
+				.Text(FEditorFontGlyphs::Level_Down)
+			]
+			+SHorizontalBox::Slot()
+			.VAlign(VAlign_Center)
+			.FillWidth(1.0f)
+			[
+				SNew(STextBlock)
+				.Text(LOCTEXT("MinLODNotification", "Min LOD applied"))
+				.TextStyle(FEditorStyle::Get(), "AnimViewport.MessageText")
+			]
+		]
+	);
+}
+
 void SAnimationEditorViewportTabBody::HandleFocusCamera()
 {
 	TSharedRef<FAnimationViewportClient> AnimViewportClient = StaticCastSharedRef<FAnimationViewportClient>(LevelViewportClient.ToSharedRef());
diff --git a/Engine/Source/Editor/Persona/Private/SAnimationEditorViewport.h b/Engine/Source/Editor/Persona/Private/SAnimationEditorViewport.h
--- a/Engine/Source/Editor/Persona/Private/SAnimationEditorViewport.h
+++ b/Engine/Source/Editor/Persona/Private/SAnimationEditorViewport.h
@@ -560,6 +560,9 @@ private:
 	/** Recording notification */
 	TWeakPtr<SWidget> WeakRecordingNotification;
 
+	/** Min LOD notification */
+	TWeakPtr<SWidget> WeakMinLODNotification;
+
 	/** Current LOD selection*/
 	int32 LODSelection;
 
@@ -584,4 +587,7 @@ private:
 
 	/** Adds a persistent notification to display post process graph state */
 	void AddPostProcessNotification();
+
+	/** Add a notification to tell the user a min LOD is being applied */
+	void AddMinLODNotification();
 };
