Add caching for "sync marker" names

diff --git a/Engine/Source/Editor/Persona/Private/SAnimNotifyPanel.cpp b/Engine/Source/Editor/Persona/Private/SAnimNotifyPanel.cpp
--- a/Engine/Source/Editor/Persona/Private/SAnimNotifyPanel.cpp
+++ b/Engine/Source/Editor/Persona/Private/SAnimNotifyPanel.cpp
@@ -774,6 +774,7 @@ protected:
 	void MakeNewNotifyPicker(FMenuBuilder& MenuBuilder, bool bIsReplaceWithMenu = false);
 	void FillNewNotifyMenu(FMenuBuilder& MenuBuilderbool, bool bIsReplaceWithMenu = false);
 	void FillNewNotifyStateMenu(FMenuBuilder& MenuBuilder, bool bIsReplaceWithMenu  = false);
+	void FillNewSyncMarkerMenu(FMenuBuilder& MenuBuilder);
 
 	// New notify functions
 	void CreateNewBlueprintNotifyAtCursor(FString NewNotifyName, FString BlueprintPath);
@@ -2436,7 +2437,7 @@ void SAnimNotifyTrack::FillNewNotifyMenu(FMenuBuilder& MenuBuilder, bool bIsRepl
 						[
 							SNew(SSkeletonAnimNotifies, EditableSkeleton)
 							.IsPicker(true)
-							.OnNotifySelected_Lambda([this, bIsReplaceWithMenu](const FName& InNotifyName)
+							.OnItemSelected_Lambda([this, bIsReplaceWithMenu](const FName& InNotifyName)
 							{
 								FSlateApplication::Get().DismissAllMenus();
 
@@ -2461,6 +2462,48 @@ void SAnimNotifyTrack::FillNewNotifyMenu(FMenuBuilder& MenuBuilder, bool bIsRepl
 	MakeNewNotifyPicker<UAnimNotify>(MenuBuilder, bIsReplaceWithMenu);
 }
 
+void SAnimNotifyTrack::FillNewSyncMarkerMenu(FMenuBuilder& MenuBuilder)
+{
+	USkeleton* SeqSkeleton = Sequence->GetSkeleton();
+	if (SeqSkeleton)
+	{
+		MenuBuilder.BeginSection("AnimSyncMarkerSubMenu", LOCTEXT("NewSyncMarkerSubMenu_Skeleton", "Sync Markers"));
+		{
+			FUIAction UIAction;
+			UIAction.ExecuteAction.BindSP(
+				this, &SAnimNotifyTrack::OnNewSyncMarkerClicked);
+			MenuBuilder.AddMenuEntry(LOCTEXT("NewSyncMarker", "New Sync Marker..."), LOCTEXT("NewSyncMarkerToolTip", "Create a new animation sync marker"), FSlateIcon(), UIAction);
+
+			MenuBuilder.AddSubMenu(
+				LOCTEXT("NewSyncMarkerSubMenu_Existing", "Existing Sync Markers"),
+				LOCTEXT("NewSyncMarkerSubMenu_Existing_Tooltip", "Choose from existing sync marker names on the skeleton"),
+				FNewMenuDelegate::CreateLambda([this, SeqSkeleton](FMenuBuilder& InSubMenuBuilder)
+			{
+				ISkeletonEditorModule& SkeletonEditorModule = FModuleManager::LoadModuleChecked<ISkeletonEditorModule>("SkeletonEditor");
+				TSharedRef<IEditableSkeleton> EditableSkeleton = SkeletonEditorModule.CreateEditableSkeleton(SeqSkeleton);
+
+				InSubMenuBuilder.AddWidget(
+					SNew(SBox)
+					.MinDesiredWidth(300.0f)
+					.MaxDesiredHeight(400.0f)
+					[
+						SNew(SSkeletonAnimNotifies, EditableSkeleton)
+						.IsSyncMarker(true)
+						.OnItemSelected_Lambda([this](const FName& InNotifyName)
+						{
+							FSlateApplication::Get().DismissAllMenus();
+
+							CreateNewSyncMarkerAtCursor(InNotifyName.ToString(), nullptr);
+						})
+					],
+					FText(), true, false
+					);
+			}));
+		}
+		MenuBuilder.EndSection();
+	}
+}
+
 FAnimNotifyEvent& SAnimNotifyTrack::CreateNewBlueprintNotify(FString NewNotifyName, FString BlueprintPath, float StartTime)
 {
 	TSubclassOf<UObject> BlueprintClass = GetBlueprintClassFromPath(BlueprintPath);
@@ -2931,11 +2974,10 @@ TSharedPtr<SWidget> SAnimNotifyTrack::SummonContextMenu(const FGeometry& MyGeome
 
 			if (Sequence->IsA(UAnimSequence::StaticClass()))
 			{
-				MenuBuilder.AddMenuEntry(
-					LOCTEXT("NewSyncMarker", "Add Sync Marker"),
-					LOCTEXT("NewSyncMarkerToolTip", "Create a new animation sync marker"),
-					FSlateIcon(),
-					FUIAction(FExecuteAction::CreateSP(this, &SAnimNotifyTrack::OnNewSyncMarkerClicked)));
+				MenuBuilder.AddSubMenu(
+					NSLOCTEXT("NewSyncMarkerSubMenu", "NewSyncMarkerSubMenuAddNotifyState", "Add Sync Marker..."),
+					NSLOCTEXT("NewSyncMarkerSubMenu", "NewSyncMarkerSubMenuAddNotifyStateToolTip", "Create a new animation sync marker"),
+					FNewMenuDelegate::CreateRaw(this, &SAnimNotifyTrack::FillNewSyncMarkerMenu));
 			}
 
 			MenuBuilder.AddMenuEntry(
@@ -3194,8 +3236,13 @@ void SAnimNotifyTrack::AddNewSyncMarker(const FText& NewNotifyName, ETextCommit:
 	if ((CommitInfo == ETextCommit::OnEnter) && SeqSkeleton)
 	{
 		const FScopedTransaction Transaction(LOCTEXT("AddNewSyncMarker", "Add New Sync Marker"));
-		//FName NewName = FName(*NewNotifyName.ToString());
-		//SeqSkeleton->AddNewAnimationNotify(NewName);
+
+		FName NewName = FName(*NewNotifyName.ToString());
+
+		ISkeletonEditorModule& SkeletonEditorModule = FModuleManager::LoadModuleChecked<ISkeletonEditorModule>("SkeletonEditor");
+		TSharedRef<IEditableSkeleton> EditableSkeleton = SkeletonEditorModule.CreateEditableSkeleton(SeqSkeleton);
+
+		EditableSkeleton->AddSyncMarker(NewName);
 
 		FBlueprintActionDatabase::Get().RefreshAssetActions(SeqSkeleton);
 
diff --git a/Engine/Source/Editor/Persona/Private/SSkeletonAnimNotifies.cpp b/Engine/Source/Editor/Persona/Private/SSkeletonAnimNotifies.cpp
--- a/Engine/Source/Editor/Persona/Private/SSkeletonAnimNotifies.cpp
+++ b/Engine/Source/Editor/Persona/Private/SSkeletonAnimNotifies.cpp
@@ -54,19 +54,27 @@ TSharedRef<SWidget> FSkeletonAnimNotifiesSummoner::CreateTabBody(const FWorkflow
 void SSkeletonAnimNotifies::Construct(const FArguments& InArgs, const TSharedRef<IEditableSkeleton>& InEditableSkeleton)
 {
 	OnObjectsSelected = InArgs._OnObjectsSelected;
-	OnNotifySelected = InArgs._OnNotifySelected;
+	OnItemSelected = InArgs._OnItemSelected;
 	bIsPicker = InArgs._IsPicker;
+	bIsSyncMarker = InArgs._IsSyncMarker;
 
 	EditableSkeleton = InEditableSkeleton;
 
-	EditableSkeleton->RegisterOnNotifiesChanged(FSimpleDelegate::CreateSP(this, &SSkeletonAnimNotifies::OnNotifiesChanged));
+	if (bIsSyncMarker)
+	{
+		bIsPicker = false; // Sync Markers are never pickers
+	}
+	else
+	{
+		EditableSkeleton->RegisterOnNotifiesChanged(FSimpleDelegate::CreateSP(this, &SSkeletonAnimNotifies::OnNotifiesChanged));
+	}
 
 	if(GEditor)
 	{
 		GEditor->RegisterForUndo(this);
 	}
 
-	FOnContextMenuOpening OnContextMenuOpening = !bIsPicker ? FOnContextMenuOpening::CreateSP(this, &SSkeletonAnimNotifies::OnGetContextMenuContent) : FOnContextMenuOpening();
+	FOnContextMenuOpening OnContextMenuOpening = (!bIsPicker && !bIsSyncMarker) ? FOnContextMenuOpening::CreateSP(this, &SSkeletonAnimNotifies::OnGetContextMenuContent) : FOnContextMenuOpening();
 
 	this->ChildSlot
 	[
@@ -154,36 +162,50 @@ TSharedPtr<SWidget> SSkeletonAnimNotifies::OnGetContextMenuContent() const
 	const bool bShouldCloseWindowAfterMenuSelection = true;
 	FMenuBuilder MenuBuilder( bShouldCloseWindowAfterMenuSelection, NULL);
 
-	MenuBuilder.BeginSection("AnimNotifyAction", LOCTEXT( "AnimNotifyActions", "Notifies" ) );
+	if (bIsSyncMarker)
 	{
-		FUIAction Action = FUIAction( FExecuteAction::CreateSP( this, &SSkeletonAnimNotifies::OnAddAnimNotify ) );
-		const FText Label = LOCTEXT("NewAnimNotifyButtonLabel", "New...");
-		const FText ToolTipText = LOCTEXT("NewAnimNotifyButtonTooltip", "Creates a new anim notify.");
-		MenuBuilder.AddMenuEntry( Label, ToolTipText, FSlateIcon(), Action);
+		MenuBuilder.BeginSection("AnimNotifyAction", LOCTEXT("SelectedSyncMarkerActions", "Selected Sync Marker Actions"));
+		{
+			FUIAction Action = FUIAction(FExecuteAction::CreateSP(this, &SSkeletonAnimNotifies::OnDeleteSyncMarker));
+			const FText Label = LOCTEXT("DeleteSyncMarkerButtonLabel", "Delete");
+			const FText ToolTipText = LOCTEXT("DeleteSyncMarkerButtonTooltip", "Deletes the sync marker from the suggestions");
+			MenuBuilder.AddMenuEntry(Label, ToolTipText, FSlateIcon(), Action);
+		}
+		MenuBuilder.EndSection();
 	}
-	MenuBuilder.EndSection();
-
-	MenuBuilder.BeginSection("AnimNotifyAction", LOCTEXT( "SelectedAnimNotifyActions", "Selected Notify Actions" ) );
+	else
 	{
+		MenuBuilder.BeginSection("AnimNotifyAction", LOCTEXT("AnimNotifyActions", "Notifies"));
 		{
-			FUIAction Action = FUIAction( FExecuteAction::CreateSP( this, &SSkeletonAnimNotifies::OnRenameAnimNotify ), 
-				FCanExecuteAction::CreateSP( this, &SSkeletonAnimNotifies::CanPerformRename ) );
-			const FText Label = LOCTEXT("RenameAnimNotifyButtonLabel", "Rename");
-			const FText ToolTipText = LOCTEXT("RenameAnimNotifyButtonTooltip", "Renames the selected anim notifies.");
-			MenuBuilder.AddMenuEntry( Label, ToolTipText, FSlateIcon(), Action);
+			FUIAction Action = FUIAction(FExecuteAction::CreateSP(this, &SSkeletonAnimNotifies::OnAddAnimNotify));
+			const FText Label = LOCTEXT("NewAnimNotifyButtonLabel", "New...");
+			const FText ToolTipText = LOCTEXT("NewAnimNotifyButtonTooltip", "Creates a new anim notify.");
+			MenuBuilder.AddMenuEntry(Label, ToolTipText, FSlateIcon(), Action);
 		}
+		MenuBuilder.EndSection();
 
+		MenuBuilder.BeginSection("AnimNotifyAction", LOCTEXT("SelectedAnimNotifyActions", "Selected Notify Actions"));
 		{
-			FUIAction Action = FUIAction( FExecuteAction::CreateSP( this, &SSkeletonAnimNotifies::OnDeleteAnimNotify ), 
-				FCanExecuteAction::CreateSP( this, &SSkeletonAnimNotifies::CanPerformDelete ) );
-			const FText Label = LOCTEXT("DeleteAnimNotifyButtonLabel", "Delete");
-			const FText ToolTipText = LOCTEXT("DeleteAnimNotifyButtonTooltip", "Deletes the selected anim notifies.");
-			MenuBuilder.AddMenuEntry( Label, ToolTipText, FSlateIcon(), Action);
-		}
+			{
+				FUIAction Action = FUIAction(FExecuteAction::CreateSP(this, &SSkeletonAnimNotifies::OnRenameAnimNotify),
+					FCanExecuteAction::CreateSP(this, &SSkeletonAnimNotifies::CanPerformRename));
+				const FText Label = LOCTEXT("RenameAnimNotifyButtonLabel", "Rename");
+				const FText ToolTipText = LOCTEXT("RenameAnimNotifyButtonTooltip", "Renames the selected anim notifies.");
+				MenuBuilder.AddMenuEntry(Label, ToolTipText, FSlateIcon(), Action);
+			}
+
+			{
+				FUIAction Action = FUIAction(FExecuteAction::CreateSP(this, &SSkeletonAnimNotifies::OnDeleteAnimNotify),
+					FCanExecuteAction::CreateSP(this, &SSkeletonAnimNotifies::CanPerformDelete));
+				const FText Label = LOCTEXT("DeleteAnimNotifyButtonLabel", "Delete");
+				const FText ToolTipText = LOCTEXT("DeleteAnimNotifyButtonTooltip", "Deletes the selected anim notifies.");
+				MenuBuilder.AddMenuEntry(Label, ToolTipText, FSlateIcon(), Action);
+			}
 
 
+		}
+		MenuBuilder.EndSection();
 	}
-	MenuBuilder.EndSection();
 
 	return MenuBuilder.MakeWidget();
 }
@@ -192,9 +214,12 @@ void SSkeletonAnimNotifies::OnNotifySelectionChanged(TSharedPtr<FDisplayedAnimNo
 {
 	if(Selection.IsValid())
 	{
-		ShowNotifyInDetailsView(Selection->Name);
+		if (!bIsSyncMarker)
+		{
+			ShowNotifyInDetailsView(Selection->Name);
+		}
 
-		OnNotifySelected.ExecuteIfBound(Selection->Name);
+		OnItemSelected.ExecuteIfBound(Selection->Name);
 	}
 }
 
@@ -269,6 +294,22 @@ void SSkeletonAnimNotifies::OnDeleteAnimNotify()
 	CreateNotifiesList( NameFilterBox->GetText().ToString() );
 }
 
+void SSkeletonAnimNotifies::OnDeleteSyncMarker()
+{
+	TArray< TSharedPtr< FDisplayedAnimNotifyInfo > > SelectedRows = NotifiesListView->GetSelectedItems();
+
+	TArray<FName> SelectedSyncMarkerNames;
+
+	for (int Selection = 0; Selection < SelectedRows.Num(); ++Selection)
+	{
+		SelectedSyncMarkerNames.Add(SelectedRows[Selection]->Name);
+	}
+
+	EditableSkeleton->DeleteSyncMarkers(SelectedSyncMarkerNames);
+
+	CreateNotifiesList(NameFilterBox->GetText().ToString());
+}
+
 void SSkeletonAnimNotifies::OnRenameAnimNotify()
 {
 	TArray< TSharedPtr< FDisplayedAnimNotifyInfo > > SelectedRows = NotifiesListView->GetSelectedItems();
@@ -342,19 +383,22 @@ void SSkeletonAnimNotifies::CreateNotifiesList( const FString& SearchText )
 {
 	NotifyList.Empty();
 
-	for(int i = 0; i < EditableSkeleton->GetSkeleton().AnimationNotifies.Num(); ++i)
+	const USkeleton& TargetSkeleton = EditableSkeleton->GetSkeleton();
+
+	const TArray<FName>& ItemNames = bIsSyncMarker ? TargetSkeleton.GetExistingMarkerNames() : TargetSkeleton.AnimationNotifies;
+
+	for(const FName& ItemName : ItemNames)
 	{
-		const FName& NotifyName = EditableSkeleton->GetSkeleton().AnimationNotifies[i];
 		if ( !SearchText.IsEmpty() )
 		{
-			if ( NotifyName.ToString().Contains( SearchText ) )
+			if (ItemName.ToString().Contains( SearchText ) )
 			{
-				NotifyList.Add( FDisplayedAnimNotifyInfo::Make( NotifyName ) );
+				NotifyList.Add( FDisplayedAnimNotifyInfo::Make(ItemName) );
 			}
 		}
 		else
 		{
-			NotifyList.Add( FDisplayedAnimNotifyInfo::Make( NotifyName ) );
+			NotifyList.Add( FDisplayedAnimNotifyInfo::Make(ItemName) );
 		}
 	}
 
diff --git a/Engine/Source/Editor/Persona/Private/SSkeletonAnimNotifies.h b/Engine/Source/Editor/Persona/Private/SSkeletonAnimNotifies.h
--- a/Engine/Source/Editor/Persona/Private/SSkeletonAnimNotifies.h
+++ b/Engine/Source/Editor/Persona/Private/SSkeletonAnimNotifies.h
@@ -24,7 +24,7 @@ class SToolTip;
 struct FNotificationInfo;
 
 /** Delegate fired when a notify is selected */
-DECLARE_DELEGATE_OneParam(FOnNotifySelected, const FName& /*InSelectedNotify*/);
+DECLARE_DELEGATE_OneParam(FOnItemSelected, const FName& /*InSelectedNotify*/);
 
 /////////////////////////////////////////////////////
 // FSkeletonAnimNotifiesSummoner
@@ -85,17 +85,21 @@ class SSkeletonAnimNotifies : public SCompoundWidget, public FGCObject, public F
 public:
 	SLATE_BEGIN_ARGS( SSkeletonAnimNotifies )
 		: _IsPicker(false)
+		, _IsSyncMarker(false)
 	{}
 
 	/** Delegate called to select an object in the details panel */
 	SLATE_EVENT(FOnObjectsSelected, OnObjectsSelected)
 
 	/** Delegate fired when a notify is selected */
-	SLATE_EVENT(FOnNotifySelected, OnNotifySelected)
+	SLATE_EVENT(FOnItemSelected, OnItemSelected)
 
 	/** Whether we should use this dialog as a picker or an editor. In picker mode we cant add, remove or rename notifies. */
 	SLATE_ARGUMENT(bool, IsPicker)
 
+	/** Whether we should use this dialog for sync markers instead of notifies */
+	SLATE_ARGUMENT(bool, IsSyncMarker)
+
 	SLATE_END_ARGS()
 public:
 	void Construct(const FArguments& InArgs, const TSharedRef<IEditableSkeleton>& InEditableSkeleton);
@@ -153,6 +157,9 @@ private:
 	/** Delegate handler for deleting anim notifies */
 	void OnDeleteAnimNotify();
 
+	/** Delegate handler for deleting a sync marker */
+	void OnDeleteSyncMarker();
+
 	/** Delegate handler for determining whether we can show the rename menu options */
 	bool CanPerformRename() const;
 
@@ -202,10 +209,13 @@ private:
 	FOnObjectsSelected OnObjectsSelected;
 
 	/** Delegate fired when a notify is selected */
-	FOnNotifySelected OnNotifySelected;
+	FOnItemSelected OnItemSelected;
 
 	/** Whether we should use this dialog as a picker or an editor. In picker mode we cant add, remove or rename notifies. */
 	bool bIsPicker;
+
+	/** Whether we are using this dialog for sync markers instead of notifies */
+	bool bIsSyncMarker;
 };
 
 #undef LOCTEXT_NAMESPACE
diff --git a/Engine/Source/Editor/SkeletonEditor/Private/EditableSkeleton.cpp b/Engine/Source/Editor/SkeletonEditor/Private/EditableSkeleton.cpp
--- a/Engine/Source/Editor/SkeletonEditor/Private/EditableSkeleton.cpp
+++ b/Engine/Source/Editor/SkeletonEditor/Private/EditableSkeleton.cpp
@@ -1095,6 +1095,18 @@ int32 FEditableSkeleton::DeleteAnimNotifies(const TArray<FName>& InNotifyNames)
 	return NumAnimationsModified;
 }
 
+void FEditableSkeleton::DeleteSyncMarkers(const TArray<FName>& InSyncMarkerNames)
+{
+	const FScopedTransaction Transaction(LOCTEXT("DeleteSyncMarkers", "Delete Sync Markers"));
+
+	Skeleton->Modify();
+
+	for (FName Marker : InSyncMarkerNames)
+	{
+		Skeleton->RemoveMarkerName(Marker);
+	}
+}
+
 void FEditableSkeleton::AddNotify(FName NewName)
 {
 	const FScopedTransaction Transaction(LOCTEXT("AddNewNotifyToSkeleton", "Add New Anim Notify To Skeleton"));
@@ -1103,6 +1115,13 @@ void FEditableSkeleton::AddNotify(FName NewName)
 	OnNotifiesChanged.Broadcast();
 }
 
+void FEditableSkeleton::AddSyncMarker(FName NewName)
+{
+	const FScopedTransaction Transaction(LOCTEXT("AddNewSyncMarkerToSkeleton", "Add New Sync Marker To Skeleton"));
+	Skeleton->Modify();
+	Skeleton->RegisterMarkerName(NewName);
+}
+
 int32 FEditableSkeleton::RenameNotify(const FName& NewName, const FName& OldName)
 {
 	const FScopedTransaction Transaction(LOCTEXT("RenameAnimNotify", "Rename Anim Notify"));
diff --git a/Engine/Source/Editor/SkeletonEditor/Private/EditableSkeleton.h b/Engine/Source/Editor/SkeletonEditor/Private/EditableSkeleton.h
--- a/Engine/Source/Editor/SkeletonEditor/Private/EditableSkeleton.h
+++ b/Engine/Source/Editor/SkeletonEditor/Private/EditableSkeleton.h
@@ -37,7 +37,9 @@ public:
 	virtual USkeletalMeshSocket* DuplicateSocket(const FSelectedSocketInfo& SocketInfoToDuplicate, const FName& NewParentBoneName, USkeletalMesh* InSkeletalMesh) override;
 	virtual int32 ValidatePreviewAttachedObjects() override;
 	virtual int32 DeleteAnimNotifies(const TArray<FName>& InSelectedNotifyNames) override;
+	virtual void DeleteSyncMarkers(const TArray<FName>& ISyncMarkerNames) override;
 	virtual void AddNotify(FName NewName) override;
+	virtual void AddSyncMarker(FName NewName) override;
 	virtual int32 RenameNotify(const FName& NewName, const FName& OldName) override;
 	virtual void BroadcastNotifyChanged() override;
 	virtual void GetCompatibleAnimSequences(TArray<struct FAssetData>& OutAssets) override;
diff --git a/Engine/Source/Editor/SkeletonEditor/Public/IEditableSkeleton.h b/Engine/Source/Editor/SkeletonEditor/Public/IEditableSkeleton.h
--- a/Engine/Source/Editor/SkeletonEditor/Public/IEditableSkeleton.h
+++ b/Engine/Source/Editor/SkeletonEditor/Public/IEditableSkeleton.h
@@ -92,6 +92,11 @@ public:
 	 */	
 	virtual int32 DeleteAnimNotifies(const TArray<FName>& InotifyNames) = 0;
 
+	/**
+	 * Delete sync markers from the skeleton cache by name
+	 */
+	virtual void DeleteSyncMarkers(const TArray<FName>& ISyncMarkerNames) = 0;
+
 	/**
 	* Add a notify
 	*/
@@ -103,6 +108,11 @@ public:
 	 */	
 	virtual int32 RenameNotify(const FName& NewName, const FName& OldName) = 0;
 
+	/**
+	* Add a sync marker
+	*/
+	virtual void AddSyncMarker(FName NewName) = 0;
+
 	/** Inform the system that something about a notify changed */	
 	virtual void BroadcastNotifyChanged() = 0;
 
diff --git a/Engine/Source/Runtime/Core/Public/UObject/AnimObjectVersion.h b/Engine/Source/Runtime/Core/Public/UObject/AnimObjectVersion.h
--- a/Engine/Source/Runtime/Core/Public/UObject/AnimObjectVersion.h
+++ b/Engine/Source/Runtime/Core/Public/UObject/AnimObjectVersion.h
@@ -15,6 +15,9 @@ struct CORE_API FAnimObjectVersion
 		// Reworked how anim blueprint root nodes are recovered
 		LinkTimeAnimBlueprintRootDiscovery,
 
+		// Cached marker sync names on skeleton for editor
+		StoreMarkerNamesOnSkeleton,
+
 		// -----<new versions can be added above this line>-------------------------------------------------
 		VersionPlusOne,
 		LatestVersion = VersionPlusOne - 1
diff --git a/Engine/Source/Runtime/Engine/Classes/Animation/Skeleton.h b/Engine/Source/Runtime/Engine/Classes/Animation/Skeleton.h
--- a/Engine/Source/Runtime/Engine/Classes/Animation/Skeleton.h
+++ b/Engine/Source/Runtime/Engine/Classes/Animation/Skeleton.h
@@ -366,6 +366,18 @@ public:
 	// return version of AnimCurveUidVersion
 	uint16 GetAnimCurveUidVersion() const { return AnimCurveUidVersion;  }
 	const TArray<uint16>& GetDefaultCurveUIDList() const { return DefaultCurveUIDList; }
+
+#if WITH_EDITOR
+	// Get existing (seen) sync marker names for this Skeleton
+	const TArray<FName>& GetExistingMarkerNames() const { return ExistingMarkerNames; }
+
+	// Register a new sync marker name
+	void RegisterMarkerName(FName MarkerName) { ExistingMarkerNames.AddUnique(MarkerName); ExistingMarkerNames.Sort(); }
+
+	// Remove a sync marker name
+	void RemoveMarkerName(FName MarkerName) { ExistingMarkerNames.Remove(MarkerName); }
+#endif
+
 protected:
 	// Container for smart name mappings
 	UPROPERTY()
@@ -376,6 +388,9 @@ protected:
 	// FBoneContainer contains only list that is used by current LOD
 	TArray<uint16> DefaultCurveUIDList;
 
+	//Cached marker sync marker names (stripped for non editor)
+	TArray<FName> ExistingMarkerNames;
+
 private:
 	/** Increase the AnimCurveUidVersion so that instances can get the latest information */
 	void IncreaseAnimCurveUidVersion();
diff --git a/Engine/Source/Runtime/Engine/Private/Animation/AnimSequence.cpp b/Engine/Source/Runtime/Engine/Private/Animation/AnimSequence.cpp
--- a/Engine/Source/Runtime/Engine/Private/Animation/AnimSequence.cpp
+++ b/Engine/Source/Runtime/Engine/Private/Animation/AnimSequence.cpp
@@ -536,6 +536,11 @@ void UAnimSequence::PostLoad()
 
 #if WITH_EDITOR
 		VerifyCurveNames<FTransformCurve>(*CurrentSkeleton, USkeleton::AnimTrackCurveMappingName, RawCurveData.TransformCurves);
+
+		for (const FAnimSyncMarker& SyncMarker : AuthoredSyncMarkers)
+		{
+			CurrentSkeleton->RegisterMarkerName(SyncMarker.MarkerName);
+		}
 #endif
 	}
 
diff --git a/Engine/Source/Runtime/Engine/Private/Animation/Skeleton.cpp b/Engine/Source/Runtime/Engine/Private/Animation/Skeleton.cpp
--- a/Engine/Source/Runtime/Engine/Private/Animation/Skeleton.cpp
+++ b/Engine/Source/Runtime/Engine/Private/Animation/Skeleton.cpp
@@ -27,6 +27,8 @@
 #include "Engine/PreviewMeshCollection.h"
 #include "Misc/ScopedSlowTask.h"
 #include "Animation/AnimBlueprint.h"
+#include "UObject/AnimObjectVersion.h"
+#include "EngineUtils.h"
 
 #define LOCTEXT_NAMESPACE "Skeleton"
 #define ROOT_BONE_PARENT	INDEX_NONE
@@ -153,6 +155,8 @@ void USkeleton::PostDuplicate(bool bDuplicateForPIE)
 
 void USkeleton::Serialize( FArchive& Ar )
 {
+	Ar.UsingCustomVersion(FAnimObjectVersion::GUID);
+
 	DECLARE_SCOPE_CYCLE_COUNTER( TEXT("USkeleton::Serialize"), STAT_Skeleton_Serialize, STATGROUP_LoadTime );
 
 	Super::Serialize(Ar);
@@ -237,6 +241,15 @@ void USkeleton::Serialize( FArchive& Ar )
 	}
 #endif
 
+	if (Ar.CustomVer(FAnimObjectVersion::GUID) >= FAnimObjectVersion::StoreMarkerNamesOnSkeleton)
+	{
+		FStripDataFlags StripFlags(Ar);
+		if (!StripFlags.IsEditorDataStripped())
+		{
+			Ar << ExistingMarkerNames;
+		}
+	}
+
 	const bool bRebuildNameMap = false;
 	ReferenceSkeleton.RebuildRefSkeleton(this, bRebuildNameMap);
 }
