Copying //UE4/Dev-Anim to //UE4/Dev-Main (Source: //UE4/Dev-Anim @ 4293397)

============================
  MAJOR FEATURES & CHANGES
============================

Change 4070996 by [REDACTED_PERSON]

	Merging using ControlRig<->Dev_Anim

Change 4072896 by [REDACTED_PERSON]

	Merging using Dev-AnimPhys<->Dev-Anim

	CCDIK node

Change 4073877 by [REDACTED_PERSON]

	fix compile error

Change 4081360 by [REDACTED_PERSON]

	Fix orion cooking issue

Change 4117810 by [REDACTED_PERSON]

	Fix issue where master bone map does not exists when master pose component is serialized
	- allow force update master pose component

Change 4186009 by [REDACTED_PERSON]

	Fix trashed indentation in RefreshBoneTransforms

Change 4240655 by [REDACTED_PERSON]

	Fix out of bounds access when meshes have bones removed under physics assets

Change 4257469 by [REDACTED_PERSON]

	Live anim BP recompilation is now non-experimental

Change 4257474 by [REDACTED_PERSON]

	Default coordinate space for animation editors is now local space

Change 4264307 by [REDACTED_PERSON]

	Updated notify add/replace menus to use class pickers for better searching of BP and native notifies

	Also updated class picker to allow access to the native parent class for unloaded Blueprints
	Improved skeleton notify management:
	- Added 'Add' option to the manager
	- Editor no longer freezes when selecting a notify (it was scanning all animations on the skeleton to populate the details panel).

Change 4264790 by [REDACTED_PERSON]

	Added the ability to access sub instance nodes by tag

	Sub instances now have the tag as their object FName. We prevent duplicates in the compiler.

Change 4275035 by [REDACTED_PERSON]

	Anim graph root node is now hooked up at class Link() time

	This ensures that the index/ptr is self-consistent with the compiled graph properties.
	Added a new anim graph node and anim node that StateResult node suse now. No funcitonal change, but a different inheritance/data hierarchy. Added some mismatched tag handling for serialization to deal with upgrading old content.
	Added a Dev-Anim object version header.

Change 4275053 by [REDACTED_PERSON]

	Added notification to the viewport when a min LOD is being applied

Change 4275057 by [REDACTED_PERSON]

	Added toggle test for bitfields

Change 4275895 by [REDACTED_PERSON]

	Adding more info for the case where getting bone transform fails due to out of bound index

Change 4278795 by [REDACTED_PERSON]

	Fix negative play rate montages

Change 4278796 by [REDACTED_PERSON]

	Fix bloom drawing of pose link when weight is greater than 1

Change 4278846 by [REDACTED_PERSON]

	Add caching for "sync marker" names

Change 4278884 by [REDACTED_PERSON]

	Add option to maintain original scale of root motion

	PR #4940

Change 4278903 by [REDACTED_PERSON]

	Native notifies for dynamics and cloth management

Change 4283011 by [REDACTED_PERSON]

	Add animation sequence framerate to editor viewport and content browser tooltip

Change 4283284 by [REDACTED_PERSON]

	Fix cis

Change 4286977 by [REDACTED_PERSON]

	Fixup fallout from Main merge

Change 4071832 by [REDACTED_PERSON]

	Merging using Fortnight-Main<->Dev-Anim

	- Rename MeshComponentUpdateFlag to VisibilityBasedAnimTickOption for better description
	- Made it configurable from INI file so that you can change engine default for your project
	- This doesn't include INI file change yet. I'm waiting for this to go through without issue, and then will check in the INI file change

Change 4074887 by [REDACTED_PERSON]

	Fix build error

Change 4097208 by [REDACTED_PERSON]

	Merging //UE4/Dev-Main to Dev-Anim (//UE4/Dev-Anim) @ CL 4096586

Change 4121923 by [REDACTED_PERSON]

	Merging //UE4/Dev-Main to Dev-Anim (//UE4/Dev-Anim) @ CL 4118128

Change 4289832 by [REDACTED_PERSON]

	Merging //UE4/Dev-Main to Dev-Anim (//UE4/Dev-Anim) @ 4289353

diff --git a/Engine/Config/BaseEngine.ini b/Engine/Config/BaseEngine.ini
--- a/Engine/Config/BaseEngine.ini
+++ b/Engine/Config/BaseEngine.ini
@@ -581,6 +581,8 @@ bUseStreamingPause=false
 +StructRedirects=(OldName="Timecode",NewName="/Script/CoreUObject.Timecode")
 
 +FunctionRedirects=(OldName="SkeletalMeshComponent.SetAnimInstanceClass",NewName="SkeletalMeshComponent.K2_SetAnimInstanceClass")
++EnumRedirects=(OldName="EMeshComponentUpdateFlag",NewName="/Script/Engine.EVisibilityBasedAnimTickOption")
++PropertyRedirects=(OldName="SkinnedMeshComponent.MeshComponentUpdateFlag", NewName="VisibilityBasedAnimTickOption")
 
 ; Move immediate mode anim nodes into engine as we now rely on it there
 +PackageRedirects=(OldName="/Script/ImmediatePhysicsEditor", NewName="/Script/AnimGraph")
diff --git a/Engine/Source/Editor/AnimGraph/Classes/AnimGraphNode_CCDIK.h b/Engine/Source/Editor/AnimGraph/Classes/AnimGraphNode_CCDIK.h
new file mode 100644
--- /dev/null
+++ b/Engine/Source/Editor/AnimGraph/Classes/AnimGraphNode_CCDIK.h
@@ -0,0 +1,40 @@
+// Copyright 1998-2016 Epic Games, Inc. All Rights Reserved.
+
+#pragma once
+#include "AnimGraphNode_SkeletalControlBase.h"
+#include "BoneControllers/AnimNode_CCDIK.h"
+#include "EdGraph/EdGraphNodeUtils.h"
+#include "AnimGraphNode_CCDIK.generated.h"
+
+// Editor node for CCDIK IK skeletal controller
+UCLASS(MinimalAPI, Experimental)
+class UAnimGraphNode_CCDIK : public UAnimGraphNode_SkeletalControlBase
+{
+	GENERATED_UCLASS_BODY()
+
+	UPROPERTY(EditAnywhere, Category = Settings)
+	FAnimNode_CCDIK Node;
+
+public:
+	// UEdGraphNode interface
+	virtual FText GetNodeTitle(ENodeTitleType::Type TitleType) const override;
+	// End of UEdGraphNode interface
+
+	// UAnimGraphNode_Base interface
+	virtual void CopyNodeDataToPreviewNode(FAnimNode_Base* AnimNode) override;
+	virtual FEditorModeID GetEditorMode() const override;
+	// End of UAnimGraphNode_Base interface
+
+	// UObject interface
+	virtual void PostEditChangeProperty(struct FPropertyChangedEvent& PropertyChangedEvent) override;
+
+protected:
+	// UAnimGraphNode_SkeletalControlBase interface
+	virtual FText GetControllerDescription() const override;
+	virtual const FAnimNode_SkeletalControlBase* GetNode() const override { return &Node; }
+	// End of UAnimGraphNode_SkeletalControlBase interface
+
+private:
+	/** Constructing FText strings can be costly, so we cache the node's title */
+	FNodeTitleTextTable CachedNodeTitles;
+};
\ No newline at end of file
diff --git a/Engine/Source/Editor/AnimGraph/Classes/AnimGraphNode_StateResult.h b/Engine/Source/Editor/AnimGraph/Classes/AnimGraphNode_StateResult.h
--- a/Engine/Source/Editor/AnimGraph/Classes/AnimGraphNode_StateResult.h
+++ b/Engine/Source/Editor/AnimGraph/Classes/AnimGraphNode_StateResult.h
@@ -5,11 +5,32 @@
 
 #include "CoreMinimal.h"
 #include "UObject/ObjectMacros.h"
-#include "AnimGraphNode_Root.h"
+#include "AnimGraphNode_Base.h"
+#include "AnimNodes/AnimNode_StateResult.h"
 #include "AnimGraphNode_StateResult.generated.h"
 
 UCLASS(MinimalAPI)
-class UAnimGraphNode_StateResult : public UAnimGraphNode_Root
+class UAnimGraphNode_StateResult : public UAnimGraphNode_Base
 {
 	GENERATED_UCLASS_BODY()
+
+	UPROPERTY(EditAnywhere, Category=Settings)
+	FAnimNode_StateResult Node;
+
+	//~ Begin UEdGraphNode Interface.
+	virtual FLinearColor GetNodeTitleColor() const override;
+	virtual FText GetTooltipText() const override;
+	virtual FText GetNodeTitle(ENodeTitleType::Type TitleType) const override;
+	virtual bool CanUserDeleteNode() const override { return false; }
+	virtual bool CanDuplicateNode() const override { return false; }
+	virtual void GetMenuActions(FBlueprintActionDatabaseRegistrar& ActionRegistrar) const override;
+	//~ End UEdGraphNode Interface.
+
+	//~ Begin UAnimGraphNode_Base Interface
+	virtual bool IsSinkNode() const override;
+
+	// Get the link to the documentation
+	virtual FString GetDocumentationLink() const override;
+
+	//~ End UAnimGraphNode_Base Interface
 };
diff --git a/Engine/Source/Editor/AnimGraph/Private/AnimGraphModule.cpp b/Engine/Source/Editor/AnimGraph/Private/AnimGraphModule.cpp
--- a/Engine/Source/Editor/AnimGraph/Private/AnimGraphModule.cpp
+++ b/Engine/Source/Editor/AnimGraph/Private/AnimGraphModule.cpp
@@ -18,6 +18,7 @@
 #include "EditModes/PoseDriverEditMode.h"
 #include "EditModes/SplineIKEditMode.h"
 #include "EditModes/LookAtEditMode.h"
+#include "EditModes/CCDIKEditMode.h"
 
 IMPLEMENT_MODULE(FAnimGraphModule, AnimGraph);
 
@@ -36,6 +37,7 @@ void FAnimGraphModule::StartupModule()
 	FEditorModeRegistry::Get().RegisterMode<FPoseDriverEditMode>(AnimNodeEditModes::PoseDriver, LOCTEXT("PoseDriverEditMode", "PoseDriver"), FSlateIcon(), false);
 	FEditorModeRegistry::Get().RegisterMode<FSplineIKEditMode>(AnimNodeEditModes::SplineIK, LOCTEXT("SplineIKEditMode", "Spline IK"), FSlateIcon(), false);
 	FEditorModeRegistry::Get().RegisterMode<FLookAtEditMode>(AnimNodeEditModes::LookAt, LOCTEXT("LookAtEditMode", "LookAt"), FSlateIcon(), false);
+	FEditorModeRegistry::Get().RegisterMode<FCCDIKEditMode>(AnimNodeEditModes::CCDIK, LOCTEXT("CCDIKEditMode", "CCD IK"), FSlateIcon(), false);
 
 	// Register details customization
 	FPropertyEditorModule& PropertyModule = FModuleManager::LoadModuleChecked<FPropertyEditorModule>("PropertyEditor");
@@ -45,6 +47,7 @@ void FAnimGraphModule::StartupModule()
 void FAnimGraphModule::ShutdownModule()
 {
 	// Unregister the editor modes
+	FEditorModeRegistry::Get().UnregisterMode(AnimNodeEditModes::CCDIK);
 	FEditorModeRegistry::Get().UnregisterMode(AnimNodeEditModes::SplineIK);
 	FEditorModeRegistry::Get().UnregisterMode(AnimNodeEditModes::PoseDriver);
 	FEditorModeRegistry::Get().UnregisterMode(AnimNodeEditModes::Fabrik);
diff --git a/Engine/Source/Editor/AnimGraph/Private/AnimGraphNode_CCDIK.cpp b/Engine/Source/Editor/AnimGraph/Private/AnimGraphNode_CCDIK.cpp
new file mode 100644
--- /dev/null
+++ b/Engine/Source/Editor/AnimGraph/Private/AnimGraphNode_CCDIK.cpp
@@ -0,0 +1,101 @@
+// Copyright 1998-2016 Epic Games, Inc. All Rights Reserved.
+
+#include "AnimGraphNode_CCDIK.h"
+#include "Animation/AnimInstance.h"
+#include "AnimNodeEditModes.h"
+
+/////////////////////////////////////////////////////
+// UAnimGraphNode_CCDIK 
+
+#define LOCTEXT_NAMESPACE "A3Nodes"
+
+UAnimGraphNode_CCDIK::UAnimGraphNode_CCDIK(const FObjectInitializer& ObjectInitializer)
+: Super(ObjectInitializer)
+{
+}
+
+FText UAnimGraphNode_CCDIK::GetControllerDescription() const
+{
+	return LOCTEXT("CCDIK", "CCDIK");
+}
+
+FText UAnimGraphNode_CCDIK::GetNodeTitle(ENodeTitleType::Type TitleType) const
+{
+	if ((TitleType == ENodeTitleType::ListView || TitleType == ENodeTitleType::MenuTitle))
+	{
+		return GetControllerDescription();
+	}
+	// @TODO: the bone can be altered in the property editor, so we have to 
+	//        choose to mark this dirty when that happens for this to properly work
+	else //if (!CachedNodeTitles.IsTitleCached(TitleType, this))
+	{
+		FFormatNamedArguments Args;
+		Args.Add(TEXT("ControllerDescription"), GetControllerDescription());
+		Args.Add(TEXT("RootBoneName"), FText::FromName(Node.RootBone.BoneName));
+		Args.Add(TEXT("TipBoneName"), FText::FromName(Node.TipBone.BoneName));
+
+		// FText::Format() is slow, so we cache this to save on performance
+		if (TitleType == ENodeTitleType::ListView || TitleType == ENodeTitleType::MenuTitle)
+		{
+			CachedNodeTitles.SetCachedTitle(TitleType, FText::Format(LOCTEXT("AnimGraphNode_CCDIKBone_ListTitle", "{ControllerDescription} - Root: {RootBoneName}, Tip: {TipBoneName} "), Args), this);
+		}
+		else
+		{
+			CachedNodeTitles.SetCachedTitle(TitleType, FText::Format(LOCTEXT("AnimGraphNode_CCDIKBone_Title", "{ControllerDescription}\nRoot: {RootBoneName} - Tip: {TipBoneName} "), Args), this);
+		}
+	}
+	return CachedNodeTitles[TitleType];
+}
+
+void UAnimGraphNode_CCDIK::CopyNodeDataToPreviewNode(FAnimNode_Base* InPreviewNode)
+{
+	FAnimNode_CCDIK* CCDIK = static_cast<FAnimNode_CCDIK*>(InPreviewNode);
+
+	// copies Pin values from the internal node to get data which are not compiled yet
+	CCDIK->EffectorLocation = Node.EffectorLocation;
+}
+
+FEditorModeID UAnimGraphNode_CCDIK::GetEditorMode() const
+{
+	return AnimNodeEditModes::CCDIK;
+}
+
+void UAnimGraphNode_CCDIK::PostEditChangeProperty(struct FPropertyChangedEvent& PropertyChangedEvent)
+{
+	Super::PostEditChangeProperty(PropertyChangedEvent);
+
+	FName PropertyName = (PropertyChangedEvent.Property != NULL) ? PropertyChangedEvent.Property->GetFName() : NAME_None;
+	if (PropertyName == GET_MEMBER_NAME_CHECKED(FBoneReference, BoneName))
+	{
+		USkeleton* Skeleton = GetAnimBlueprint()->TargetSkeleton;
+		if (Node.TipBone.BoneName != NAME_None && Node.RootBone.BoneName != NAME_None)
+		{
+			const FReferenceSkeleton& RefSkeleton = Skeleton->GetReferenceSkeleton();
+			const int32 TipBoneIndex = RefSkeleton.FindBoneIndex(Node.TipBone.BoneName);
+			const int32 RootBoneIndex = RefSkeleton.FindBoneIndex(Node.RootBone.BoneName);
+
+			if (TipBoneIndex != INDEX_NONE && RootBoneIndex != INDEX_NONE)
+			{
+				const int32 Depth = RefSkeleton.GetDepthBetweenBones(TipBoneIndex, RootBoneIndex);
+				if (Depth >= 0)
+				{
+					Node.ResizeRotationLimitPerJoints(Depth + 1);
+				}
+				else
+				{
+					Node.ResizeRotationLimitPerJoints(0);
+				}
+			}
+			else
+			{
+				Node.ResizeRotationLimitPerJoints(0);
+			}
+		}
+		else
+		{
+			Node.ResizeRotationLimitPerJoints(0);
+		}
+	}
+}
+
+#undef LOCTEXT_NAMESPACE
\ No newline at end of file
diff --git a/Engine/Source/Editor/AnimGraph/Private/AnimGraphNode_StateResult.cpp b/Engine/Source/Editor/AnimGraph/Private/AnimGraphNode_StateResult.cpp
--- a/Engine/Source/Editor/AnimGraph/Private/AnimGraphNode_StateResult.cpp
+++ b/Engine/Source/Editor/AnimGraph/Private/AnimGraphNode_StateResult.cpp
@@ -1,6 +1,9 @@
 // Copyright 1998-2018 Epic Games, Inc. All Rights Reserved.
 
 #include "AnimGraphNode_StateResult.h"
+#include "GraphEditorSettings.h"
+
+#define LOCTEXT_NAMESPACE "A3Nodes"
 
 /////////////////////////////////////////////////////
 // UAnimGraphNode_StateResult
@@ -9,3 +12,35 @@ UAnimGraphNode_StateResult::UAnimGraphNode_StateResult(const FObjectInitializer&
 	: Super(ObjectInitializer)
 {
 }
+
+FLinearColor UAnimGraphNode_StateResult::GetNodeTitleColor() const
+{
+	return GetDefault<UGraphEditorSettings>()->ResultNodeTitleColor;
+}
+
+FText UAnimGraphNode_StateResult::GetNodeTitle(ENodeTitleType::Type TitleType) const
+{
+	return LOCTEXT("AnimGraphNodeStateResult_Title", "Output Animation Pose");
+}
+
+FText UAnimGraphNode_StateResult::GetTooltipText() const
+{
+	return LOCTEXT("AnimGraphNodeStateResult_Tooltip", "This is the output of this animation state");
+}
+
+bool UAnimGraphNode_StateResult::IsSinkNode() const
+{
+	return true;
+}
+
+void UAnimGraphNode_StateResult::GetMenuActions(FBlueprintActionDatabaseRegistrar& ActionRegistrar) const
+{
+	// Intentionally empty. This node is auto-generated when a new graph is created.
+}
+
+FString UAnimGraphNode_StateResult::GetDocumentationLink() const
+{
+	return TEXT("Shared/GraphNodes/AnimationStateMachine");
+}
+
+#undef LOCTEXT_NAMESPACE
diff --git a/Engine/Source/Editor/AnimGraph/Private/AnimGraphNode_SubInstance.cpp b/Engine/Source/Editor/AnimGraph/Private/AnimGraphNode_SubInstance.cpp
--- a/Engine/Source/Editor/AnimGraph/Private/AnimGraphNode_SubInstance.cpp
+++ b/Engine/Source/Editor/AnimGraph/Private/AnimGraphNode_SubInstance.cpp
@@ -44,11 +44,27 @@ FText UAnimGraphNode_SubInstance::GetNodeTitle(ENodeTitleType::Type TitleType) c
 	}
 	if(TitleType == ENodeTitleType::ListView)
 	{
-		return FText::Format(LOCTEXT("TitleListFormat", "{NodeTitle} - Target Class: {TargetClass}"), Args);
+		if(Node.Tag != NAME_None)
+		{
+			Args.Add(TEXT("Tag"), FText::FromName(Node.Tag));
+			return FText::Format(LOCTEXT("TitleListFormatTagged", "{NodeTitle} - Target Class: {TargetClass} - Tag: {Tag}"), Args);
+		}
+		else
+		{
+			return FText::Format(LOCTEXT("TitleListFormat", "{NodeTitle} - Target Class: {TargetClass}"), Args);
+		}
 	}
 	else
 	{
-		return FText::Format(LOCTEXT("TitleFormat", "{NodeTitle}\nTarget Class: {TargetClass}"), Args);
+		if(Node.Tag != NAME_None)
+		{
+			Args.Add(TEXT("Tag"), FText::FromName(Node.Tag));
+			return FText::Format(LOCTEXT("TitleFormatTagged", "{NodeTitle}\nTarget Class: {TargetClass}\nTag: {Tag}"), Args);
+		}
+		else
+		{
+			return FText::Format(LOCTEXT("TitleFormat", "{NodeTitle}\nTarget Class: {TargetClass}"), Args);
+		}
 	}
 }
 
@@ -96,6 +112,28 @@ void UAnimGraphNode_SubInstance::ValidateAnimNodeDuringCompilation(USkeleton* Fo
 		MessageLog.Error(TEXT("Detected loop in sub instance chain starting at @@ inside class @@"), this, AnimBP->GetAnimBlueprintGeneratedClass());
 	}
 
+	// Check for duplicate tags in this anim blueprint
+	for(UEdGraph* Graph : Graphs)
+	{
+		TArray<UAnimGraphNode_SubInstance*> SubInstanceNodes;
+		Graph->GetNodesOfClass(SubInstanceNodes);
+
+		for(UAnimGraphNode_SubInstance* SubInstanceNode : SubInstanceNodes)
+		{
+			if(SubInstanceNode == OriginalNode)
+			{
+				continue;
+			}
+
+			FAnimNode_SubInstance& InnerNode = SubInstanceNode->Node;
+
+			if(InnerNode.Tag != NAME_None && InnerNode.Tag == Node.Tag)
+			{
+				MessageLog.Error(*FText::Format(LOCTEXT("DuplicateTagErrorFormat", "Node @@ and node @@ both have the same tag '{0}'."), FText::FromName(Node.Tag)).ToString(), this, SubInstanceNode);
+			}
+		}
+	}
+
 	// Check we don't try to spawn our own blueprint
 	if(*Node.InstanceClass == AnimBP->GetAnimBlueprintGeneratedClass())
 	{
diff --git a/Engine/Source/Editor/AnimGraph/Private/EditModes/AnimNodeEditModes.cpp b/Engine/Source/Editor/AnimGraph/Private/EditModes/AnimNodeEditModes.cpp
--- a/Engine/Source/Editor/AnimGraph/Private/EditModes/AnimNodeEditModes.cpp
+++ b/Engine/Source/Editor/AnimGraph/Private/EditModes/AnimNodeEditModes.cpp
@@ -8,5 +8,6 @@ const FEditorModeID AnimNodeEditModes::ObserveBone("AnimGraph.SkeletalControl.Ob
 const FEditorModeID AnimNodeEditModes::ModifyBone("AnimGraph.SkeletalControl.ModifyBone");
 const FEditorModeID AnimNodeEditModes::Fabrik("AnimGraph.SkeletalControl.Fabrik");
 const FEditorModeID AnimNodeEditModes::PoseDriver("AnimGraph.AnimNode.PoseDriver");
+const FEditorModeID AnimNodeEditModes::CCDIK("AnimGraph.SkeletalControl.CCDIK");
 const FEditorModeID AnimNodeEditModes::SplineIK("AnimGraph.AnimNode.SplineIK");
 const FEditorModeID AnimNodeEditModes::LookAt("AnimGraph.SkeletalControl.LookAt");
diff --git a/Engine/Source/Editor/AnimGraph/Private/EditModes/CCDIKEditMode.cpp b/Engine/Source/Editor/AnimGraph/Private/EditModes/CCDIKEditMode.cpp
new file mode 100644
--- /dev/null
+++ b/Engine/Source/Editor/AnimGraph/Private/EditModes/CCDIKEditMode.cpp
@@ -0,0 +1,78 @@
+// Copyright 1998-2016 Epic Games, Inc. All Rights Reserved.
+
+#include "CCDIKEditMode.h"
+#include "AnimGraphNode_CCDIK.h"
+#include "IPersonaPreviewScene.h"
+#include "Animation/DebugSkelMeshComponent.h"
+
+void FCCDIKEditMode::EnterMode(class UAnimGraphNode_Base* InEditorNode, struct FAnimNode_Base* InRuntimeNode)
+{
+	RuntimeNode = static_cast<FAnimNode_CCDIK*>(InRuntimeNode);
+	GraphNode = CastChecked<UAnimGraphNode_CCDIK>(InEditorNode);
+
+	FAnimNodeEditMode::EnterMode(InEditorNode, InRuntimeNode);
+}
+
+void FCCDIKEditMode::ExitMode()
+{
+	RuntimeNode = nullptr;
+	GraphNode = nullptr;
+
+	FAnimNodeEditMode::ExitMode();
+}
+
+FVector FCCDIKEditMode::GetWidgetLocation() const
+{
+	EBoneControlSpace Space = RuntimeNode->EffectorLocationSpace;
+	FVector Location = RuntimeNode->EffectorLocation;
+	FBoneSocketTarget Target = RuntimeNode->EffectorTarget;
+
+	USkeletalMeshComponent* SkelComp = GetAnimPreviewScene().GetPreviewMeshComponent();
+	return ConvertWidgetLocation(SkelComp, RuntimeNode->ForwardedPose, Target, Location, Space);
+}
+
+FWidget::EWidgetMode FCCDIKEditMode::GetWidgetMode() const
+{
+	USkeletalMeshComponent* SkelComp = GetAnimPreviewScene().GetPreviewMeshComponent();
+	int32 TipBoneIndex = SkelComp->GetBoneIndex(RuntimeNode->TipBone.BoneName);
+	int32 RootBoneIndex = SkelComp->GetBoneIndex(RuntimeNode->RootBone.BoneName);
+
+	if (TipBoneIndex!= INDEX_NONE && RootBoneIndex != INDEX_NONE)
+	{
+		return FWidget::WM_Translate;
+	}
+
+	return FWidget::WM_None;
+}
+
+void FCCDIKEditMode::DoTranslation(FVector& InTranslation)
+{
+	USkeletalMeshComponent* SkelComp = GetAnimPreviewScene().GetPreviewMeshComponent();
+
+	FVector Offset = ConvertCSVectorToBoneSpace(SkelComp, InTranslation, RuntimeNode->ForwardedPose, RuntimeNode->EffectorTarget, GraphNode->Node.EffectorLocationSpace);
+
+	RuntimeNode->EffectorLocation += Offset;
+	GraphNode->Node.EffectorLocation = RuntimeNode->EffectorLocation;
+	GraphNode->SetDefaultValue(GET_MEMBER_NAME_STRING_CHECKED(FAnimNode_CCDIK, EffectorLocation), RuntimeNode->EffectorLocation);
+}
+
+void FCCDIKEditMode::Render(const FSceneView* View, FViewport* Viewport, FPrimitiveDrawInterface* PDI)
+{
+#if !(UE_BUILD_SHIPPING || UE_BUILD_TEST)
+	// draw line from root bone to tip bone if available
+	if (RuntimeNode && RuntimeNode->DebugLines.Num() > 0)
+	{
+		USkeletalMeshComponent* SkelComp = GetAnimPreviewScene().GetPreviewMeshComponent();
+		FTransform CompToWorld = SkelComp->GetComponentToWorld();
+
+		// no component space
+		for (int32 Index = 1; Index < RuntimeNode->DebugLines.Num(); ++Index)
+		{
+			FVector Start = CompToWorld.TransformPosition(RuntimeNode->DebugLines[Index - 1]);
+			FVector End = CompToWorld.TransformPosition(RuntimeNode->DebugLines[Index]);
+
+			PDI->DrawLine(Start, End, FLinearColor::Red, SDPG_Foreground);
+		}
+	}
+#endif // #if !(UE_BUILD_SHIPPING || UE_BUILD_TEST)
+}
\ No newline at end of file
diff --git a/Engine/Source/Editor/AnimGraph/Private/EditModes/CCDIKEditMode.h b/Engine/Source/Editor/AnimGraph/Private/EditModes/CCDIKEditMode.h
new file mode 100644
--- /dev/null
+++ b/Engine/Source/Editor/AnimGraph/Private/EditModes/CCDIKEditMode.h
@@ -0,0 +1,20 @@
+// Copyright 1998-2016 Epic Games, Inc. All Rights Reserved.
+
+#pragma once
+
+#include "AnimNodeEditMode.h"
+
+class FCCDIKEditMode : public FAnimNodeEditMode
+{
+public:
+	/** IAnimNodeEditMode interface */
+	virtual void EnterMode(class UAnimGraphNode_Base* InEditorNode, struct FAnimNode_Base* InRuntimeNode) override;
+	virtual void ExitMode() override;
+	virtual FVector GetWidgetLocation() const override;
+	virtual FWidget::EWidgetMode GetWidgetMode() const override;
+	virtual void DoTranslation(FVector& InTranslation) override;
+	virtual void Render(const FSceneView* View, FViewport* Viewport, FPrimitiveDrawInterface* PDI) override;
+private:
+	struct FAnimNode_CCDIK* RuntimeNode;
+	class UAnimGraphNode_CCDIK* GraphNode;
+};
\ No newline at end of file
diff --git a/Engine/Source/Editor/AnimGraph/Public/AnimNodeEditModes.h b/Engine/Source/Editor/AnimGraph/Public/AnimNodeEditModes.h
--- a/Engine/Source/Editor/AnimGraph/Public/AnimNodeEditModes.h
+++ b/Engine/Source/Editor/AnimGraph/Public/AnimNodeEditModes.h
@@ -12,6 +12,7 @@ struct ANIMGRAPH_API AnimNodeEditModes
 	const static FEditorModeID ObserveBone;
 	const static FEditorModeID ModifyBone;
 	const static FEditorModeID Fabrik;
+	const static FEditorModeID CCDIK;
 	const static FEditorModeID PoseDriver;
 	const static FEditorModeID SplineIK;
 	const static FEditorModeID LookAt;
diff --git a/Engine/Source/Editor/AnimationBlueprintEditor/Private/AnimGraphConnectionDrawingPolicy.cpp b/Engine/Source/Editor/AnimationBlueprintEditor/Private/AnimGraphConnectionDrawingPolicy.cpp
--- a/Engine/Source/Editor/AnimationBlueprintEditor/Private/AnimGraphConnectionDrawingPolicy.cpp
+++ b/Engine/Source/Editor/AnimationBlueprintEditor/Private/AnimGraphConnectionDrawingPolicy.cpp
@@ -66,7 +66,7 @@ void FAnimGraphConnectionDrawingPolicy::BuildExecutionRoadmap()
 						FExecPairingMap& Predecessors = PredecessorPins.FindOrAdd((UEdGraphNode*)SourceNode);
 						FTimePair& Timings = Predecessors.FindOrAdd(PoseNet);
 						Timings.PredExecTime = 0.0;
-						Timings.ThisExecTime = VisitRecord.Weight;
+						Timings.ThisExecTime = FMath::Clamp(VisitRecord.Weight, 0.f, 1.f);
 					}
 				}
 			}
diff --git a/Engine/Source/Editor/AnimationEditor/Private/AnimationEditor.cpp b/Engine/Source/Editor/AnimationEditor/Private/AnimationEditor.cpp
--- a/Engine/Source/Editor/AnimationEditor/Private/AnimationEditor.cpp
+++ b/Engine/Source/Editor/AnimationEditor/Private/AnimationEditor.cpp
@@ -394,9 +394,8 @@ TSharedPtr<SDockTab> FAnimationEditor::OpenNewAnimationDocumentTab(UAnimationAss
 	{
 		FString	DocumentLink;
 
-		FAnimDocumentArgs Args(PersonaToolkit->GetPreviewScene(), GetPersonaToolkit(), GetSkeletonTree()->GetEditableSkeleton(), OnPostUndo, OnChangeAnimNotifies, OnSectionsChanged);
+		FAnimDocumentArgs Args(PersonaToolkit->GetPreviewScene(), GetPersonaToolkit(), GetSkeletonTree()->GetEditableSkeleton(), OnPostUndo, OnSectionsChanged);
 		Args.OnDespatchObjectsSelected = FOnObjectsSelected::CreateSP(this, &FAnimationEditor::HandleObjectsSelected);
-		Args.OnDespatchAnimNotifiesChanged = FSimpleDelegate::CreateSP(this, &FAnimationEditor::HandleAnimNotifiesChanged);
 		Args.OnDespatchInvokeTab = FOnInvokeTab::CreateSP(this, &FAssetEditorToolkit::InvokeTab);
 		Args.OnDespatchSectionsChanged = FSimpleDelegate::CreateSP(this, &FAnimationEditor::HandleSectionsChanged);
 
@@ -462,11 +461,6 @@ TSharedPtr<SDockTab> FAnimationEditor::OpenNewAnimationDocumentTab(UAnimationAss
 	return OpenedTab;
 }
 
-void FAnimationEditor::HandleAnimNotifiesChanged()
-{
-	OnChangeAnimNotifies.Broadcast();
-}
-
 void FAnimationEditor::HandleSectionsChanged()
 {
 	OnSectionsChanged.Broadcast();
diff --git a/Engine/Source/Editor/AnimationEditor/Private/AnimationEditor.h b/Engine/Source/Editor/AnimationEditor/Private/AnimationEditor.h
--- a/Engine/Source/Editor/AnimationEditor/Private/AnimationEditor.h
+++ b/Engine/Source/Editor/AnimationEditor/Private/AnimationEditor.h
@@ -114,8 +114,6 @@ private:
 		Max
 	};
 
-	void HandleAnimNotifiesChanged();
-
 	void HandleSectionsChanged();
 
 	bool HasValidAnimationSequence() const;
@@ -181,9 +179,6 @@ private:
 
 	bool RecordMeshToAnimation(USkeletalMeshComponent* PreviewComponent, UAnimSequence* NewAsset) const;
 public:
-	/** Multicast delegate fired on anim notifies changing */
-	FSimpleMulticastDelegate OnChangeAnimNotifies;
-
 	/** Multicast delegate fired on global undo/redo */
 	FSimpleMulticastDelegate OnPostUndo;
 
diff --git a/Engine/Source/Editor/AnimationEditor/Private/AnimationEditorMode.cpp b/Engine/Source/Editor/AnimationEditor/Private/AnimationEditorMode.cpp
--- a/Engine/Source/Editor/AnimationEditor/Private/AnimationEditorMode.cpp
+++ b/Engine/Source/Editor/AnimationEditor/Private/AnimationEditorMode.cpp
@@ -34,7 +34,7 @@ FAnimationEditorMode::FAnimationEditorMode(TSharedRef<FWorkflowCentricApplicatio
 	TabFactories.RegisterFactory(PersonaModule.CreateAssetDetailsTabFactory(InHostingApp, FOnGetAsset::CreateSP(&AnimationEditor.Get(), &FAnimationEditor::HandleGetAsset), FOnDetailsCreated()));
 	TabFactories.RegisterFactory(PersonaModule.CreateCurveViewerTabFactory(InHostingApp, InSkeletonTree->GetEditableSkeleton(), AnimationEditor->GetPersonaToolkit()->GetPreviewScene(), AnimationEditor->OnPostUndo, OnObjectsSelected));
 	TabFactories.RegisterFactory(PersonaModule.CreateSkeletonSlotNamesTabFactory(InHostingApp, InSkeletonTree->GetEditableSkeleton(), AnimationEditor->OnPostUndo, FOnObjectSelected::CreateSP(&AnimationEditor.Get(), &FAnimationEditor::HandleObjectSelected)));
-	TabFactories.RegisterFactory(PersonaModule.CreateAnimNotifiesTabFactory(InHostingApp, InSkeletonTree->GetEditableSkeleton(), AnimationEditor->OnChangeAnimNotifies, AnimationEditor->OnPostUndo, OnObjectsSelected));
+	TabFactories.RegisterFactory(PersonaModule.CreateAnimNotifiesTabFactory(InHostingApp, InSkeletonTree->GetEditableSkeleton(), OnObjectsSelected));
 
 	TabLayout = FTabManager::NewLayout("Standalone_AnimationEditor_Layout_v1.2")
 		->AddArea
diff --git a/Engine/Source/Editor/ClassViewer/Private/UnloadedBlueprintData.cpp b/Engine/Source/Editor/ClassViewer/Private/UnloadedBlueprintData.cpp
--- a/Engine/Source/Editor/ClassViewer/Private/UnloadedBlueprintData.cpp
+++ b/Engine/Source/Editor/ClassViewer/Private/UnloadedBlueprintData.cpp
@@ -96,6 +96,27 @@ const UClass* FUnloadedBlueprintData::GetClassWithin() const
 	return nullptr;
 }
 
+const UClass* FUnloadedBlueprintData::GetNativeParent() const
+{
+	TSharedPtr<FClassViewerNode> CurrentNode = ClassViewerNode.Pin()->ParentNode.Pin();
+
+	while (CurrentNode.IsValid())
+	{
+		// The class field will be invalid for unloaded classes.
+		// However, it should be valid once we've hit a loaded class or a Native class.
+		// Assuming BP cannot change ClassWithin data, this should be safe.
+		if (CurrentNode->Class.IsValid() && CurrentNode->Class->HasAnyClassFlags(CLASS_Native))
+		{
+			return CurrentNode->Class.Get();
+		}
+
+		CurrentNode = CurrentNode->ParentNode.Pin();
+	}
+
+	return nullptr;
+}
+
+
 const TWeakPtr<FClassViewerNode>& FUnloadedBlueprintData::GetClassViewerNode() const
 {
 	return ClassViewerNode;
diff --git a/Engine/Source/Editor/ClassViewer/Private/UnloadedBlueprintData.h b/Engine/Source/Editor/ClassViewer/Private/UnloadedBlueprintData.h
--- a/Engine/Source/Editor/ClassViewer/Private/UnloadedBlueprintData.h
+++ b/Engine/Source/Editor/ClassViewer/Private/UnloadedBlueprintData.h
@@ -26,6 +26,8 @@ public:
 
 	virtual const UClass* GetClassWithin() const override;
 
+	virtual const UClass* GetNativeParent() const override;
+
 	/** Retrieves the Class Viewer node this data is associated with. */
 	const TWeakPtr<FClassViewerNode>& GetClassViewerNode() const;
 
diff --git a/Engine/Source/Editor/ClassViewer/Public/ClassViewerFilter.h b/Engine/Source/Editor/ClassViewer/Public/ClassViewerFilter.h
--- a/Engine/Source/Editor/ClassViewer/Public/ClassViewerFilter.h
+++ b/Engine/Source/Editor/ClassViewer/Public/ClassViewerFilter.h
@@ -220,4 +220,11 @@ public:
 	 * @return ClassWithin of the child most Native class in the hierarchy.
 	 */
 	virtual const UClass* GetClassWithin() const = 0;
+
+	/**
+	 * Attempts to get the child-most Native class in the hierarchy.
+	 *
+	 * @return The child-most Native class in the hierarchy.
+	 */
+	virtual const UClass* GetNativeParent() const = 0;
 };
diff --git a/Engine/Source/Editor/DetailCustomizations/Private/AnimMontageSegmentDetails.cpp b/Engine/Source/Editor/DetailCustomizations/Private/AnimMontageSegmentDetails.cpp
--- a/Engine/Source/Editor/DetailCustomizations/Private/AnimMontageSegmentDetails.cpp
+++ b/Engine/Source/Editor/DetailCustomizations/Private/AnimMontageSegmentDetails.cpp
@@ -252,7 +252,7 @@ void SAnimationSegmentViewport::Construct(const FArguments& InArgs)
 	ViewportWidget->SetViewportInterface( SceneViewport.ToSharedRef() );
 	
 	PreviewComponent = NewObject<UDebugSkelMeshComponent>();
-	PreviewComponent->bEnablePhysicsOnDedicatedServer = EMeshComponentUpdateFlag::AlwaysTickPoseAndRefreshBones;
+	PreviewComponent->VisibilityBasedAnimTickOption = EVisibilityBasedAnimTickOption::AlwaysTickPoseAndRefreshBones;
 	PreviewScene.AddComponent(PreviewComponent, FTransform::Identity);
 
 	InitSkeleton();
diff --git a/Engine/Source/Editor/DetailCustomizations/Private/AnimSequenceDetails.cpp b/Engine/Source/Editor/DetailCustomizations/Private/AnimSequenceDetails.cpp
--- a/Engine/Source/Editor/DetailCustomizations/Private/AnimSequenceDetails.cpp
+++ b/Engine/Source/Editor/DetailCustomizations/Private/AnimSequenceDetails.cpp
@@ -400,7 +400,7 @@ void SAnimationRefPoseViewport::Construct(const FArguments& InArgs)
 
 	// Create the preview component
 	PreviewComponent = NewObject<UDebugSkelMeshComponent>();
-	PreviewComponent->MeshComponentUpdateFlag = EMeshComponentUpdateFlag::AlwaysTickPoseAndRefreshBones;
+	PreviewComponent->VisibilityBasedAnimTickOption = EVisibilityBasedAnimTickOption::AlwaysTickPoseAndRefreshBones;
 	PreviewScene.AddComponent( PreviewComponent, FTransform::Identity );
 
 	this->ChildSlot
diff --git a/Engine/Source/Editor/DetailCustomizations/Private/SkeletonNotifyDetails.cpp b/Engine/Source/Editor/DetailCustomizations/Private/SkeletonNotifyDetails.cpp
--- a/Engine/Source/Editor/DetailCustomizations/Private/SkeletonNotifyDetails.cpp
+++ b/Engine/Source/Editor/DetailCustomizations/Private/SkeletonNotifyDetails.cpp
@@ -13,6 +13,12 @@
 #include "IDetailPropertyRow.h"
 #include "DetailCategoryBuilder.h"
 #include "IDetailsView.h"
+#include "Widgets/Input/SButton.h"
+#include "Misc/ScopedSlowTask.h"
+#include "IEditableSkeleton.h"
+#include "Animation/AnimSequenceBase.h"
+#include "Widgets/Layout/SBox.h"
+#include "Widgets/Layout/SScrollBox.h"
 
 #define LOCTEXT_NAMESPACE "SkeletonNotifyDetails"
 
@@ -45,6 +51,8 @@ void FSkeletonNotifyDetails::CustomizeDetails( IDetailLayoutBuilder& DetailBuild
 
 	if(EdObj)
 	{
+		NotifyObject = EdObj;
+
 		Category.AddCustomRow(LOCTEXT("AnimationsLabel","Animations"))
 		.NameContent()
 		[
@@ -55,9 +63,33 @@ void FSkeletonNotifyDetails::CustomizeDetails( IDetailLayoutBuilder& DetailBuild
 		]
 		.ValueContent()
 		[
-			SNew(SListView<TSharedPtr<FString>>)
-			.ListItemsSource(&EdObj->AnimationNames)
-			.OnGenerateRow(this, &FSkeletonNotifyDetails::MakeAnimationRow)
+			SNew(SHorizontalBox)
+			+SHorizontalBox::Slot()
+			.AutoWidth()
+			[
+				SNew(SButton)
+				.Visibility_Lambda([this](){ return AnimationNames.Num() > 0 ? EVisibility::Collapsed : EVisibility::Visible; })
+				.Text(LOCTEXT("ScanForAnimations", "Scan"))
+				.ToolTipText(LOCTEXT("ScanForAnimationsTooltip", "Scan for animations that reference this notify"))
+				.OnClicked(this, &FSkeletonNotifyDetails::CollectSequencesUsingNotify)
+			]
+			+SHorizontalBox::Slot()
+			.AutoWidth()
+			[
+				SNew(SBox)
+				.MaxDesiredHeight(300.0f)
+				.MaxDesiredWidth(200.0f)
+				[
+					SNew(SScrollBox)
+					+SScrollBox::Slot()
+					[
+						SAssignNew(ListView, SListView<TSharedPtr<FString>>)
+						.Visibility_Lambda([this](){ return AnimationNames.Num() > 0 ? EVisibility::Visible : EVisibility::Collapsed; })
+						.ListItemsSource(&AnimationNames)
+						.OnGenerateRow(this, &FSkeletonNotifyDetails::MakeAnimationRow)
+					]
+				]
+			]
 		];
 	}
 }
@@ -66,8 +98,59 @@ TSharedRef< ITableRow > FSkeletonNotifyDetails::MakeAnimationRow( TSharedPtr<FSt
 {
 	return SNew(STableRow<TSharedPtr<FString>>, OwnerTable)
 	[
-		SNew(STextBlock).Text(FText::FromString(*Item.Get()))
+		SNew(STextBlock)
+		.ToolTipText(FText::FromString(*Item.Get()))
+		.Text(FText::FromString(*Item.Get()))
 	];
 }
 
+FReply FSkeletonNotifyDetails::CollectSequencesUsingNotify()
+{
+	if(NotifyObject.IsValid() && NotifyObject->EditableSkeleton.IsValid())
+	{
+		FScopedSlowTask SlowTask(1.0f, FText::Format(LOCTEXT("ScanningAnimationMessage", "Looking for animations that reference notify '{0}'."), FText::FromName(NotifyObject->Name)));
+		SlowTask.MakeDialog(true);
+
+		AnimationNames.Empty();
+
+		TArray<FAssetData> CompatibleAnimSequences;
+		NotifyObject->EditableSkeleton.Pin()->GetCompatibleAnimSequences(CompatibleAnimSequences);
+		SlowTask.TotalAmountOfWork = CompatibleAnimSequences.Num();
+
+		for( int32 AssetIndex = 0; AssetIndex < CompatibleAnimSequences.Num(); ++AssetIndex )
+		{
+			SlowTask.EnterProgressFrame(1.0f);
+
+			if(SlowTask.ShouldCancel())
+			{
+				break;
+			}
+
+			const FAssetData& PossibleAnimSequence = CompatibleAnimSequences[AssetIndex];
+			if(UObject* AnimSeqAsset = PossibleAnimSequence.GetAsset())
+			{
+				UAnimSequenceBase* Sequence = CastChecked<UAnimSequenceBase>(AnimSeqAsset);
+				for (int32 NotifyIndex = 0; NotifyIndex < Sequence->Notifies.Num(); ++NotifyIndex)
+				{
+					FAnimNotifyEvent& NotifyEvent = Sequence->Notifies[NotifyIndex];
+					if (NotifyEvent.NotifyName == NotifyObject->Name)
+					{
+						AnimationNames.Add(MakeShareable(new FString(PossibleAnimSequence.AssetName.ToString())));
+						break;
+					}
+				}
+			}
+		}
+
+		if(SlowTask.ShouldCancel())
+		{
+			AnimationNames.Empty();
+		}
+
+		ListView->RequestListRefresh();
+	}
+
+	return FReply::Handled();
+}
+
 #undef LOCTEXT_NAMESPACE
diff --git a/Engine/Source/Editor/DetailCustomizations/Private/SkeletonNotifyDetails.h b/Engine/Source/Editor/DetailCustomizations/Private/SkeletonNotifyDetails.h
--- a/Engine/Source/Editor/DetailCustomizations/Private/SkeletonNotifyDetails.h
+++ b/Engine/Source/Editor/DetailCustomizations/Private/SkeletonNotifyDetails.h
@@ -4,10 +4,12 @@
 
 #include "CoreMinimal.h"
 #include "IDetailCustomization.h"
+#include "Widgets/Views/SListView.h"
 
 class IDetailLayoutBuilder;
 class ITableRow;
 class STableViewBase;
+class UEditorSkeletonNotifyObj;
 
 class FSkeletonNotifyDetails : public IDetailCustomization
 {
@@ -20,4 +22,18 @@ public:
 
 	/** Delegate to handle creating rows for the animations slate list */
 	TSharedRef< ITableRow > MakeAnimationRow( TSharedPtr<FString> Item, const TSharedRef< STableViewBase >& OwnerTable );
+
+private:
+	/** Look for all aniamtions that reference our notify */
+	FReply CollectSequencesUsingNotify();
+
+private:
+	/** The object we are customizing */
+	TWeakObjectPtr<UEditorSkeletonNotifyObj> NotifyObject;
+
+	/** The names of any animations that reference the notify we are displaying */
+	TArray< TSharedPtr<FString> > AnimationNames;
+
+	/** The list view widget */
+	TSharedPtr<SListView<TSharedPtr<FString>>> ListView;
 };
diff --git a/Engine/Source/Editor/KismetCompiler/Private/AnimBlueprintCompiler.cpp b/Engine/Source/Editor/KismetCompiler/Private/AnimBlueprintCompiler.cpp
--- a/Engine/Source/Editor/KismetCompiler/Private/AnimBlueprintCompiler.cpp
+++ b/Engine/Source/Editor/KismetCompiler/Private/AnimBlueprintCompiler.cpp
@@ -1148,7 +1148,6 @@ void FAnimBlueprintCompilerContext::ProcessAllAnimationNodes()
 	TArray<UAnimGraphNode_Base*> RootSet;
 
 	AllocateNodeIndexCounter = 0;
-	NewAnimBlueprintClass->RootAnimNodeIndex = 0;//INDEX_NONE;
 
 	for (auto SourceNodeIt = AnimNodeList.CreateIterator(); SourceNodeIt; ++SourceNodeIt)
 	{
@@ -1202,8 +1201,6 @@ void FAnimBlueprintCompilerContext::ProcessAllAnimationNodes()
 		{
 			AutoWireAnimGetter(Getter, nullptr);
 		}
-
-		NewAnimBlueprintClass->RootAnimNodeIndex = GetAllocationIndexOfNode(PrePhysicsRoot);
 	}
 	else
 	{
@@ -2063,8 +2060,6 @@ void FAnimBlueprintCompilerContext::CleanAndSanitizeClass(UBlueprintGeneratedCla
 	//@TODO: Move this into PurgeClass
 	NewAnimBlueprintClass->BakedStateMachines.Empty();
 	NewAnimBlueprintClass->AnimNotifies.Empty();
-
-	NewAnimBlueprintClass->RootAnimNodeIndex = INDEX_NONE;
 	NewAnimBlueprintClass->RootAnimNodeProperty = NULL;
 	NewAnimBlueprintClass->OrderedSavedPoseIndices.Empty();
 	NewAnimBlueprintClass->AnimNodeProperties.Empty();
@@ -2082,7 +2077,6 @@ void FAnimBlueprintCompilerContext::CleanAndSanitizeClass(UBlueprintGeneratedCla
 
 		NewAnimBlueprintClass->BakedStateMachines.Append(RootAnimClass->BakedStateMachines);
 		NewAnimBlueprintClass->AnimNotifies.Append(RootAnimClass->AnimNotifies);
-		NewAnimBlueprintClass->RootAnimNodeIndex = RootAnimClass->RootAnimNodeIndex;
 		NewAnimBlueprintClass->OrderedSavedPoseIndices = RootAnimClass->OrderedSavedPoseIndices;
 	}
 }
diff --git a/Engine/Source/Editor/Persona/Private/AnimationEditorPreviewScene.cpp b/Engine/Source/Editor/Persona/Private/AnimationEditorPreviewScene.cpp
--- a/Engine/Source/Editor/Persona/Private/AnimationEditorPreviewScene.cpp
+++ b/Engine/Source/Editor/Persona/Private/AnimationEditorPreviewScene.cpp
@@ -239,7 +239,7 @@ void FAnimationEditorPreviewScene::SetPreviewMeshInternal(USkeletalMesh* NewPrev
 
 		AddPreviewAttachedObjects();
 
-		SkeletalMeshComponent->MeshComponentUpdateFlag = EMeshComponentUpdateFlag::AlwaysTickPoseAndRefreshBones;
+		SkeletalMeshComponent->VisibilityBasedAnimTickOption = EVisibilityBasedAnimTickOption::AlwaysTickPoseAndRefreshBones;
 	}
 
 	for (auto Iter = AdditionalMeshes.CreateIterator(); Iter; ++Iter)
diff --git a/Engine/Source/Editor/Persona/Private/AnimationEditorViewportClient.cpp b/Engine/Source/Editor/Persona/Private/AnimationEditorViewportClient.cpp
--- a/Engine/Source/Editor/Persona/Private/AnimationEditorViewportClient.cpp
+++ b/Engine/Source/Editor/Persona/Private/AnimationEditorViewportClient.cpp
@@ -93,6 +93,9 @@ FAnimationViewportClient::FAnimationViewportClient(const TSharedRef<IPersonaPrev
 	((FAssetEditorModeManager*)ModeTools)->SetPreviewScene(&InPreviewScene.Get());
 	((FAssetEditorModeManager*)ModeTools)->SetDefaultMode(FPersonaEditModes::SkeletonSelection);
 
+	// Default to local space
+	SetWidgetCoordSystemSpace(COORD_Local);
+
 	// load config
 	ConfigOption = UPersonaOptions::StaticClass()->GetDefaultObject<UPersonaOptions>();
 	check (ConfigOption);
@@ -943,6 +946,11 @@ FText FAnimationViewportClient::GetDisplayInfo(bool bDisplayAllInfo) const
 		TextValue = ConcatenateLine(TextValue, LOCTEXT("MeshMaterialHiddenWarning", "Mesh Materials Hidden"));
 	}
 
+	if (const UAnimSequence* AnimSequence = Cast<UAnimSequence>(GetAnimPreviewScene()->GetPreviewAnimationAsset()))
+	{
+		TextValue = ConcatenateLine(TextValue, FText::Format(LOCTEXT("FramerateFormat", "Framerate: {0}"), FText::AsNumber(AnimSequence->GetFrameRate())));
+	}
+
 	return TextValue;
 }
 void FAnimationViewportClient::DrawNodeDebugLines(TArray<FText>& Lines, FCanvas* Canvas, FSceneView* View)
diff --git a/Engine/Source/Editor/Persona/Private/PersonaModule.cpp b/Engine/Source/Editor/Persona/Private/PersonaModule.cpp
--- a/Engine/Source/Editor/Persona/Private/PersonaModule.cpp
+++ b/Engine/Source/Editor/Persona/Private/PersonaModule.cpp
@@ -271,9 +271,9 @@ void FPersonaModule::RegisterPersonaViewportTabFactories(FWorkflowAllowedTabSet&
 	TabSet.RegisterFactory(MakeShareable(new FPreviewViewportSummoner(InHostingApp, InArgs, 3)));
 }
 
-TSharedRef<class FWorkflowTabFactory> FPersonaModule::CreateAnimNotifiesTabFactory(const TSharedRef<class FWorkflowCentricApplication>& InHostingApp, const TSharedRef<class IEditableSkeleton>& InEditableSkeleton, FSimpleMulticastDelegate& InOnChangeAnimNotifies, FSimpleMulticastDelegate& InOnPostUndo, FOnObjectsSelected InOnObjectsSelected) const
+TSharedRef<class FWorkflowTabFactory> FPersonaModule::CreateAnimNotifiesTabFactory(const TSharedRef<class FWorkflowCentricApplication>& InHostingApp, const TSharedRef<class IEditableSkeleton>& InEditableSkeleton, FOnObjectsSelected InOnObjectsSelected) const
 {
-	return MakeShareable(new FSkeletonAnimNotifiesSummoner(InHostingApp, InEditableSkeleton, InOnChangeAnimNotifies, InOnPostUndo, InOnObjectsSelected));
+	return MakeShareable(new FSkeletonAnimNotifiesSummoner(InHostingApp, InEditableSkeleton, InOnObjectsSelected));
 }
 
 TSharedRef<class FWorkflowTabFactory> FPersonaModule::CreateCurveViewerTabFactory(const TSharedRef<class FWorkflowCentricApplication>& InHostingApp, const TSharedRef<class IEditableSkeleton>& InEditableSkeleton, const TSharedRef<IPersonaPreviewScene>& InPreviewScene, FSimpleMulticastDelegate& InOnPostUndo, FOnObjectsSelected InOnObjectsSelected) const
@@ -328,34 +328,31 @@ TSharedRef<SWidget> FPersonaModule::CreateEditorWidgetForAnimDocument(const TSha
 	{
 		if (UAnimSequence* Sequence = Cast<UAnimSequence>(InAnimAsset))
 		{
-			Result = SNew(SSequenceEditor, InArgs.PreviewScene.Pin().ToSharedRef(), InArgs.EditableSkeleton.Pin().ToSharedRef(), InArgs.OnPostUndo)
+			Result = SNew(SSequenceEditor, InArgs.PreviewScene.Pin().ToSharedRef(), InArgs.EditableSkeleton.Pin().ToSharedRef())
 				.Sequence(Sequence)
 				.OnObjectsSelected(InArgs.OnDespatchObjectsSelected)
-				.OnAnimNotifiesChanged(InArgs.OnDespatchAnimNotifiesChanged)
 				.OnInvokeTab(InArgs.OnDespatchInvokeTab);
 
 			OutDocumentLink = TEXT("Engine/Animation/Sequences");
 		}
 		else if (UAnimComposite* Composite = Cast<UAnimComposite>(InAnimAsset))
 		{
-			Result = SNew(SAnimCompositeEditor, InArgs.PreviewScene.Pin().ToSharedRef(), InArgs.EditableSkeleton.Pin().ToSharedRef(), InArgs.OnPostUndo)
+			Result = SNew(SAnimCompositeEditor, InArgs.PreviewScene.Pin().ToSharedRef(), InArgs.EditableSkeleton.Pin().ToSharedRef())
 				.Composite(Composite)
 				.OnObjectsSelected(InArgs.OnDespatchObjectsSelected)
-				.OnAnimNotifiesChanged(InArgs.OnDespatchAnimNotifiesChanged)
 				.OnInvokeTab(InArgs.OnDespatchInvokeTab);
 
 			OutDocumentLink = TEXT("Engine/Animation/AnimationComposite");
 		}
 		else if (UAnimMontage* Montage = Cast<UAnimMontage>(InAnimAsset))
 		{
-			FMontageEditorRequiredArgs RequiredArgs(InArgs.PreviewScene.Pin().ToSharedRef(), InArgs.EditableSkeleton.Pin().ToSharedRef(), InArgs.OnPostUndo, InArgs.OnAnimNotifiesChanged, InArgs.OnSectionsChanged);
+			FMontageEditorRequiredArgs RequiredArgs(InArgs.PreviewScene.Pin().ToSharedRef(), InArgs.EditableSkeleton.Pin().ToSharedRef(), InArgs.OnSectionsChanged);
 
 			Result = SNew(SMontageEditor, RequiredArgs)
 				.Montage(Montage)
 				.OnSectionsChanged(InArgs.OnDespatchSectionsChanged)
 				.OnInvokeTab(InArgs.OnDespatchInvokeTab)
-				.OnObjectsSelected(InArgs.OnDespatchObjectsSelected)
-				.OnAnimNotifiesChanged(InArgs.OnDespatchAnimNotifiesChanged);
+				.OnObjectsSelected(InArgs.OnDespatchObjectsSelected);
 
 			OutDocumentLink = TEXT("Engine/Animation/AnimMontage");
 		}
diff --git a/Engine/Source/Editor/Persona/Private/SAnimCompositeEditor.cpp b/Engine/Source/Editor/Persona/Private/SAnimCompositeEditor.cpp
--- a/Engine/Source/Editor/Persona/Private/SAnimCompositeEditor.cpp
+++ b/Engine/Source/Editor/Persona/Private/SAnimCompositeEditor.cpp
@@ -6,6 +6,7 @@
 #include "IDocumentation.h"
 
 #include "SAnimNotifyPanel.h"
+#include "Editor.h"
 
 //////////////////////////////////////////////////////////////////////////
 // SAnimCompositeEditor
@@ -15,7 +16,7 @@ TSharedRef<SWidget> SAnimCompositeEditor::CreateDocumentAnchor()
 	return IDocumentation::Get()->CreateAnchor(TEXT("Engine/Animation/AnimationComposite"));
 }
 
-void SAnimCompositeEditor::Construct(const FArguments& InArgs, const TSharedRef<class IPersonaPreviewScene>& InPreviewScene, const TSharedRef<class IEditableSkeleton>& InEditableSkeleton, FSimpleMulticastDelegate& OnPostUndo)
+void SAnimCompositeEditor::Construct(const FArguments& InArgs, const TSharedRef<class IPersonaPreviewScene>& InPreviewScene, const TSharedRef<class IEditableSkeleton>& InEditableSkeleton)
 {
 	bIsActiveTimerRegistered = false;
 	CompositeObj = InArgs._Composite;
@@ -25,7 +26,10 @@ void SAnimCompositeEditor::Construct(const FArguments& InArgs, const TSharedRef<
 		.OnObjectsSelected(InArgs._OnObjectsSelected), 
 		InPreviewScene );
 
-	OnPostUndo.Add(FSimpleDelegate::CreateSP( this, &SAnimCompositeEditor::PostUndo ) );
+	if(GEditor)
+	{
+		GEditor->RegisterForUndo(this);
+	}
 
 	EditorPanels->AddSlot()
 		.AutoHeight()
@@ -44,7 +48,7 @@ void SAnimCompositeEditor::Construct(const FArguments& InArgs, const TSharedRef<
 		.AutoHeight()
 		.Padding(0, 10)
 		[
-			SAssignNew( AnimNotifyPanel, SAnimNotifyPanel, OnPostUndo)
+			SAssignNew( AnimNotifyPanel, SAnimNotifyPanel, InEditableSkeleton )
 			.Sequence(CompositeObj)
 			.WidgetWidth(S2ColumnWidget::DEFAULT_RIGHT_COLUMN_WIDTH)
 			.InputMin(this, &SAnimEditorBase::GetMinInput)
@@ -54,7 +58,6 @@ void SAnimCompositeEditor::Construct(const FArguments& InArgs, const TSharedRef<
 			.OnSetInputViewRange(this, &SAnimEditorBase::SetInputViewRange)
 			.OnGetScrubValue(this, &SAnimEditorBase::GetScrubValue)
 			.OnSelectionChanged(this, &SAnimEditorBase::OnSelectionChanged)
-			.OnAnimNotifiesChanged(InArgs._OnAnimNotifiesChanged)
 			.OnInvokeTab(InArgs._OnInvokeTab)
 		];
 
@@ -76,6 +79,14 @@ void SAnimCompositeEditor::Construct(const FArguments& InArgs, const TSharedRef<
 	CollapseComposite();
 }
 
+SAnimCompositeEditor::~SAnimCompositeEditor()
+{
+	if(GEditor)
+	{
+		GEditor->UnregisterForUndo(this);
+	}
+}
+
 void SAnimCompositeEditor::PreAnimUpdate()
 {
 	CompositeObj->Modify();
@@ -125,7 +136,17 @@ void SAnimCompositeEditor::CollapseComposite()
 	RecalculateSequenceLength();
 }
 
-void SAnimCompositeEditor::PostUndo()
+void SAnimCompositeEditor::PostUndo( bool bSuccess )
+{
+	PostUndoRedo();
+}
+
+void SAnimCompositeEditor::PostRedo( bool bSuccess )
+{
+	PostUndoRedo();
+}
+
+void SAnimCompositeEditor::PostUndoRedo()
 {
 	if (!bIsActiveTimerRegistered)
 	{
diff --git a/Engine/Source/Editor/Persona/Private/SAnimCompositeEditor.h b/Engine/Source/Editor/Persona/Private/SAnimCompositeEditor.h
--- a/Engine/Source/Editor/Persona/Private/SAnimCompositeEditor.h
+++ b/Engine/Source/Editor/Persona/Private/SAnimCompositeEditor.h
@@ -22,7 +22,7 @@ class SAnimNotifyPanel;
 	portion of the composite tool and registering callbacks to the SAnimCompositeEditor to do the actual editing.
 	
 */
-class SAnimCompositeEditor : public SAnimEditorBase
+class SAnimCompositeEditor : public SAnimEditorBase, public FEditorUndoClient
 {
 public:
 	SLATE_BEGIN_ARGS( SAnimCompositeEditor )
@@ -31,7 +31,6 @@ public:
 
 		SLATE_ARGUMENT( UAnimComposite*, Composite)
 		SLATE_EVENT(FOnObjectsSelected, OnObjectsSelected)
-		SLATE_EVENT(FSimpleDelegate, OnAnimNotifiesChanged)
 		SLATE_EVENT(FOnInvokeTab, OnInvokeTab)
 
 	SLATE_END_ARGS()
@@ -56,8 +55,10 @@ protected:
 	virtual void InitDetailsViewEditorObject(class UEditorAnimBaseObj* EdObj) override;
 
 public:
-	void Construct(const FArguments& InArgs, const TSharedRef<class IPersonaPreviewScene>& InPreviewScene, const TSharedRef<class IEditableSkeleton>& InEditableSkeleton, FSimpleMulticastDelegate& OnPostUndo);
-	
+	void Construct(const FArguments& InArgs, const TSharedRef<class IPersonaPreviewScene>& InPreviewScene, const TSharedRef<class IEditableSkeleton>& InEditableSkeleton);
+
+	~SAnimCompositeEditor();
+
 	/** Return the animation composite being edited */
 	UAnimComposite* GetCompositeObj() const { return CompositeObj; }
 	virtual UAnimationAsset* GetEditorObject() const override { return GetCompositeObj(); }
@@ -66,8 +67,11 @@ private:
 	/** Pointer to the animation composite being edited */
 	UAnimComposite* CompositeObj;
 	
-	void PostUndo();
-	
+	/** FEditorUndoClient interface */
+	virtual void PostUndo( bool bSuccess ) override;
+	virtual void PostRedo( bool bSuccess ) override;
+	void PostUndoRedo();
+
 	virtual float CalculateSequenceLengthOfEditorObject() const override;
 
 	/** This will sort all components of the montage and update (recreate) the UI */
diff --git a/Engine/Source/Editor/Persona/Private/SAnimMontagePanel.cpp b/Engine/Source/Editor/Persona/Private/SAnimMontagePanel.cpp
--- a/Engine/Source/Editor/Persona/Private/SAnimMontagePanel.cpp
+++ b/Engine/Source/Editor/Persona/Private/SAnimMontagePanel.cpp
@@ -19,13 +19,17 @@
 #include "TabSpawners.h"
 #include "Widgets/Input/SNumericEntryBox.h"
 #include "Styling/CoreStyle.h"
+#include "ISkeletonEditorModule.h"
+#include "Modules/ModuleManager.h"
+#include "IEditableSkeleton.h"
+#include "Editor.h"
 
 #define LOCTEXT_NAMESPACE "AnimMontagePanel"
 
 //////////////////////////////////////////////////////////////////////////
 // SAnimMontagePanel
 
-void SAnimMontagePanel::Construct(const FArguments& InArgs, FSimpleMulticastDelegate& OnAnimNotifiesChanged, FSimpleMulticastDelegate& OnSectionsChanged)
+void SAnimMontagePanel::Construct(const FArguments& InArgs, FSimpleMulticastDelegate& OnSectionsChanged)
 {
 	SAnimTrackPanel::Construct( SAnimTrackPanel::FArguments()
 		.WidgetWidth(InArgs._WidgetWidth)
@@ -65,7 +69,10 @@ void SAnimMontagePanel::Construct(const FArguments& InArgs, FSimpleMulticastDele
 		]
 	];
 
-	OnAnimNotifiesChanged.Add(FSimpleDelegate::CreateSP(this, &SAnimMontagePanel::Update));
+	ISkeletonEditorModule& SkeletonEditorModule = FModuleManager::LoadModuleChecked<ISkeletonEditorModule>("SkeletonEditor");
+	TSharedPtr<IEditableSkeleton> EditableSkeleton = SkeletonEditorModule.CreateEditableSkeleton(Montage->GetSkeleton());
+	EditableSkeleton->RegisterOnNotifiesChanged(FSimpleDelegate::CreateSP(this, &SAnimMontagePanel::Update));
+
 	OnSectionsChanged.Add(FSimpleDelegate::CreateSP(this, &SAnimMontagePanel::Update));
 
 	Update();
diff --git a/Engine/Source/Editor/Persona/Private/SAnimMontagePanel.h b/Engine/Source/Editor/Persona/Private/SAnimMontagePanel.h
--- a/Engine/Source/Editor/Persona/Private/SAnimMontagePanel.h
+++ b/Engine/Source/Editor/Persona/Private/SAnimMontagePanel.h
@@ -13,6 +13,7 @@
 #include "SMontageEditor.h"
 #include "Animation/EditorAnimBaseObj.h"
 #include "STrack.h"
+#include "EditorUndoClient.h"
 
 class FMenuBuilder;
 class SBorder;
@@ -68,7 +69,8 @@ public:
 	SLATE_ARGUMENT(bool, bChildAnimMontage)
 	SLATE_END_ARGS()
 
-	void Construct(const FArguments& InArgs, FSimpleMulticastDelegate& OnAnimNotifiesChanged, FSimpleMulticastDelegate& OnSectionsChanged);
+	void Construct(const FArguments& InArgs, FSimpleMulticastDelegate& OnSectionsChanged);
+
 	void SetMontage(class UAnimMontage * InMontage);
 
 	// SWidget interface
diff --git a/Engine/Source/Editor/Persona/Private/SAnimNotifyPanel.cpp b/Engine/Source/Editor/Persona/Private/SAnimNotifyPanel.cpp
--- a/Engine/Source/Editor/Persona/Private/SAnimNotifyPanel.cpp
+++ b/Engine/Source/Editor/Persona/Private/SAnimNotifyPanel.cpp
@@ -40,6 +40,9 @@
 #include "IEditableSkeleton.h"
 #include "ISkeletonEditorModule.h"
 #include "Widgets/Input/SNumericEntryBox.h"
+#include "ClassViewerFilter.h"
+#include "ClassViewerModule.h"
+#include "SSkeletonAnimNotifies.h"
 
 // Track Panel drawing
 const float NotificationTrackHeight = 20.0f;
@@ -756,7 +759,7 @@ public:
 	static TSubclassOf<UObject> GetBlueprintClassFromPath(FString BlueprintPath);
 
 	// Get the default Notify Name for a given blueprint notify asset
-	FString MakeBlueprintNotifyName(FAssetData& NotifyAssetData);
+	FString MakeBlueprintNotifyName(const FString& InNotifyClassName);
 
 	// Need to make sure tool tips are cleared during node clear up so slate system won't
 	// call into invalid notify.
@@ -767,8 +770,11 @@ protected:
 	void CreateCommands();
 
 	// Build up a "New Notify..." menu
+	template<typename NotifyTypeClass>
+	void MakeNewNotifyPicker(FMenuBuilder& MenuBuilder, bool bIsReplaceWithMenu = false);
 	void FillNewNotifyMenu(FMenuBuilder& MenuBuilderbool, bool bIsReplaceWithMenu = false);
 	void FillNewNotifyStateMenu(FMenuBuilder& MenuBuilder, bool bIsReplaceWithMenu  = false);
+	void FillNewSyncMarkerMenu(FMenuBuilder& MenuBuilder);
 
 	// New notify functions
 	void CreateNewBlueprintNotifyAtCursor(FString NewNotifyName, FString BlueprintPath);
@@ -863,9 +869,6 @@ private:
 		UClass* BaseClass;
 	};
 
-	// Format notify asset data into the information needed for menu display
-	void GetNotifyMenuData(TArray<FAssetData>& NotifyAssetData, TArray<BlueprintNotifyMenuInfo>& OutNotifyMenuData);
-
 	// Store the tracks geometry for later use
 	void UpdateCachedGeometry(const FGeometry& InGeometry);
 
@@ -2324,228 +2327,178 @@ FCursorReply SAnimNotifyTrack::OnCursorQuery(const FGeometry& MyGeometry, const
 	return FCursorReply::Unhandled();
 }
 
-// Fill "new notify state" menu, or "replace with notify state menu"
-void SAnimNotifyTrack::FillNewNotifyStateMenu(FMenuBuilder& MenuBuilder, bool bIsReplaceWithMenu /* = false */)
+template<typename NotifyTypeClass>
+void SAnimNotifyTrack::MakeNewNotifyPicker(FMenuBuilder& MenuBuilder, bool bIsReplaceWithMenu /* = false */)
 {
-	// Run the native query first to update the allowed classes for blueprints.
-	TArray<UClass*> NotifyStateClasses;
-	OnGetNotifyStateNativeClasses.ExecuteIfBound(NotifyStateClasses);
-
-	// Collect blueprint notify data
-	TArray<FAssetData> NotifyAssetData;
-	TArray<BlueprintNotifyMenuInfo> NotifyMenuData;
-	OnGetNotifyStateBlueprintData.ExecuteIfBound(NotifyAssetData);
-	GetNotifyMenuData(NotifyAssetData, NotifyMenuData);
+	FText TypeName = NotifyTypeClass::StaticClass() == UAnimNotify::StaticClass() ? LOCTEXT("AnimNotifyName", "anim notify") : LOCTEXT("AnimNotifyStateName", "anim notify state");
+	FText SectionHeaderFormat = bIsReplaceWithMenu ? LOCTEXT("ReplaceWithAnExistingAnimNotify", "Replace with an existing {0}") : LOCTEXT("AddsAnExistingAnimNotify", "Add an existing {0}");
 
-	for(BlueprintNotifyMenuInfo& NotifyData : NotifyMenuData)
+	class FNotifyStateClassFilter : public IClassViewerFilter
 	{
-		const FText LabelText = FText::FromString(NotifyData.NotifyName);
+	public:
+		FNotifyStateClassFilter(UAnimSequenceBase* InSequence)
+			: Sequence(InSequence)
+		{}
 
-		FUIAction UIAction;
-		FText Description = FText::GetEmpty();
-		if (!bIsReplaceWithMenu)
+		bool IsClassAllowed(const FClassViewerInitializationOptions& InInitOptions, const UClass* InClass, TSharedRef< FClassViewerFilterFuncs > InFilterFuncs) override
 		{
-			Description = LOCTEXT("AddsAnExistingAnimNotify", "Add an existing notify");
-			UIAction.ExecuteAction.BindRaw(
-				this, &SAnimNotifyTrack::CreateNewBlueprintNotifyAtCursor,
-				NotifyData.NotifyName,
-				NotifyData.BlueprintPath);
-			UIAction.CanExecuteAction.BindRaw(
-				this, &SAnimNotifyTrack::IsValidToPlace,
-				NotifyData.BaseClass);
+			const bool bChildOfObjectClass = InClass->IsChildOf(NotifyTypeClass::StaticClass());
+			const bool bMatchesFlags = !InClass->HasAnyClassFlags(CLASS_Hidden | CLASS_HideDropDown | CLASS_Deprecated | CLASS_Abstract);
+			return bChildOfObjectClass && bMatchesFlags && CastChecked<NotifyTypeClass>(InClass->ClassDefaultObject)->CanBePlaced(Sequence);
 		}
-		else
-		{
-			Description = LOCTEXT("ReplaceWithAnExistingAnimNotify", "Replace with an existing notify");
-			UIAction.ExecuteAction.BindRaw(
-				this, &SAnimNotifyTrack::ReplaceSelectedWithBlueprintNotify,
-				NotifyData.NotifyName,
-				NotifyData.BlueprintPath);
-			UIAction.CanExecuteAction.BindRaw(
-				this, &SAnimNotifyTrack::IsValidToPlace,
-				NotifyData.BaseClass);
-		}
-
-		MenuBuilder.AddMenuEntry(LabelText, Description, FSlateIcon(), UIAction);
-	}
 
-	MenuBuilder.BeginSection("NativeNotifyStates", LOCTEXT("NewStateNotifyMenu_Native", "Native Notify States"));
-	{
-		for(UClass* Class : NotifyStateClasses)
+		virtual bool IsUnloadedClassAllowed(const FClassViewerInitializationOptions& InInitOptions, const TSharedRef< const IUnloadedBlueprintData > InUnloadedClassData, TSharedRef< FClassViewerFilterFuncs > InFilterFuncs) override
 		{
-			if (Class->HasAllClassFlags(CLASS_Abstract))
+			const bool bChildOfObjectClass = InUnloadedClassData->IsChildOf(NotifyTypeClass::StaticClass());
+			const bool bMatchesFlags = !InUnloadedClassData->HasAnyClassFlags(CLASS_Hidden | CLASS_HideDropDown | CLASS_Deprecated | CLASS_Abstract);
+			bool bValidToPlace = false;
+			if(bChildOfObjectClass)
 			{
-				continue; // skip abstract classes
+				if (const UClass* NativeBaseClass = InUnloadedClassData->GetNativeParent())
+				{
+					bValidToPlace = CastChecked<NotifyTypeClass>(NativeBaseClass->ClassDefaultObject)->CanBePlaced(Sequence);
+				}
 			}
 
-			const FText Description = LOCTEXT("NewNotifyStateSubMenu_NativeToolTip", "Add an existing native notify state");
-			const FText LabelText = Class->GetDisplayNameText();
-			const FString Label = LabelText.ToString();
+			return bChildOfObjectClass && bMatchesFlags && bValidToPlace;
+		}
 
-			FUIAction UIAction;
-			if (!bIsReplaceWithMenu)
-			{
-				UIAction.ExecuteAction.BindRaw(
-					this, &SAnimNotifyTrack::CreateNewNotifyAtCursor,
-					Label,
-					Class);
-				UIAction.CanExecuteAction.BindRaw(
-					this, &SAnimNotifyTrack::IsValidToPlace,
-					Class);
-			}
-			else
-			{
-				UIAction.ExecuteAction.BindRaw(
-					this, &SAnimNotifyTrack::ReplaceSelectedWithNotify,
-					Label,
-					Class);
-				UIAction.CanExecuteAction.BindRaw(
-					this, &SAnimNotifyTrack::IsValidToPlace,
-					Class);
-			}
+		/** Sequence referenced by outer panel */
+		UAnimSequenceBase* Sequence;
+	};
 
-			MenuBuilder.AddMenuEntry(LabelText, Description, FSlateIcon(), UIAction);
-		}
-	}
-	MenuBuilder.EndSection();
+	FClassViewerInitializationOptions InitOptions;
+	InitOptions.Mode = EClassViewerMode::ClassPicker;
+	InitOptions.bShowObjectRootClass = false;
+	InitOptions.bShowUnloadedBlueprints = true;
+	InitOptions.bShowNoneOption = false;
+	InitOptions.bEnableClassDynamicLoading = true;
+	InitOptions.bExpandRootNodes = true;
+	InitOptions.bShowDisplayNames = true;
+	InitOptions.ClassFilter = MakeShared<FNotifyStateClassFilter>(Sequence);
+
+	FClassViewerModule& ClassViewerModule = FModuleManager::LoadModuleChecked<FClassViewerModule>("ClassViewer");
+	MenuBuilder.AddWidget(
+		SNew(SBox)
+		.MinDesiredWidth(300.0f)
+		.MaxDesiredHeight(400.0f)
+		[
+			ClassViewerModule.CreateClassViewer(InitOptions,
+				FOnClassPicked::CreateLambda([this, bIsReplaceWithMenu](UClass* InClass)
+				{
+					FSlateApplication::Get().DismissAllMenus();
+					if(bIsReplaceWithMenu)
+					{
+						ReplaceSelectedWithNotify(MakeBlueprintNotifyName(InClass->GetName()), InClass);
+					}
+					else
+					{
+						CreateNewNotifyAtCursor(MakeBlueprintNotifyName(InClass->GetName()), InClass);
+					}
+				}
+			))
+		],
+		FText(), true, false);
 }
 
-// Fill "new notify" menu, or "replace with notify menu"
-void SAnimNotifyTrack::FillNewNotifyMenu(FMenuBuilder& MenuBuilder, bool bIsReplaceWithMenu /* = false */)
+void SAnimNotifyTrack::FillNewNotifyStateMenu(FMenuBuilder& MenuBuilder, bool bIsReplaceWithMenu /* = false */)
 {
-	TArray<UClass*> NativeNotifyClasses;
-	OnGetNotifyNativeClasses.ExecuteIfBound(NativeNotifyClasses);
-
-	TArray<FAssetData> NotifyAssetData;
-	TArray<BlueprintNotifyMenuInfo> NotifyMenuData;
-	OnGetNotifyBlueprintData.ExecuteIfBound(NotifyAssetData);
-	GetNotifyMenuData(NotifyAssetData, NotifyMenuData);
-
-	for(BlueprintNotifyMenuInfo& NotifyData : NotifyMenuData)
-	{
-		const FText LabelText = FText::FromString( NotifyData.NotifyName );
-
-		FUIAction UIAction;
-		FText Description = FText::GetEmpty();
-		if (!bIsReplaceWithMenu)
-		{
-			Description = LOCTEXT("NewNotifySubMenu_ToolTip", "Add an existing notify");
-			UIAction.ExecuteAction.BindRaw(
-				this, &SAnimNotifyTrack::CreateNewBlueprintNotifyAtCursor,
-				NotifyData.NotifyName,
-				NotifyData.BlueprintPath);
-			UIAction.CanExecuteAction.BindRaw(
-				this, &SAnimNotifyTrack::IsValidToPlace,
-				NotifyData.BaseClass);
-		}
-		else
-		{
-			Description = LOCTEXT("ReplaceWithNotifySubMenu_ToolTip", "Replace with an existing notify");
-			UIAction.ExecuteAction.BindRaw(
-				this, &SAnimNotifyTrack::ReplaceSelectedWithBlueprintNotify,
-				NotifyData.NotifyName,
-				NotifyData.BlueprintPath);
-			UIAction.CanExecuteAction.BindRaw(
-				this, &SAnimNotifyTrack::IsValidToPlace,
-				NotifyData.BaseClass);
-		}
-		
-		MenuBuilder.AddMenuEntry(LabelText, Description, FSlateIcon(), UIAction);
-	}
+	MakeNewNotifyPicker<UAnimNotifyState>(MenuBuilder, bIsReplaceWithMenu);
+}
 
-	MenuBuilder.BeginSection("NativeNotifies", LOCTEXT("NewNotifyMenu_Native", "Native Notifies"));
+void SAnimNotifyTrack::FillNewNotifyMenu(FMenuBuilder& MenuBuilder, bool bIsReplaceWithMenu /* = false */)
+{
+	// now add custom anim notifiers
+	USkeleton* SeqSkeleton = Sequence->GetSkeleton();
+	if (SeqSkeleton)
 	{
-		for(UClass* Class : NativeNotifyClasses)
+		MenuBuilder.BeginSection("AnimNotifySkeletonSubMenu", LOCTEXT("NewNotifySubMenu_Skeleton", "Skeleton Notifies"));
 		{
-			if (Class->HasAllClassFlags(CLASS_Abstract))
-			{
-				continue; // skip abstract classes
-			}
-
-			const FText LabelText = Class->GetDisplayNameText();
-			const FString Label = LabelText.ToString();
-
-			FUIAction UIAction;
-			FText Description = FText::GetEmpty();
 			if (!bIsReplaceWithMenu)
 			{
-				Description = LOCTEXT("NewNotifySubMenu_NativeToolTip", "Add an existing native notify");
-				UIAction.ExecuteAction.BindRaw(
-					this, &SAnimNotifyTrack::CreateNewNotifyAtCursor,
-					Label,
-					Class);
-				UIAction.CanExecuteAction.BindRaw(
-					this, &SAnimNotifyTrack::IsValidToPlace,
-					Class);
-			}
-			else
-			{
-				Description = LOCTEXT("ReplaceWithNotifySubMenu_NativeToolTip", "Replace with an existing native notify");
-				UIAction.ExecuteAction.BindRaw(
-					this, &SAnimNotifyTrack::ReplaceSelectedWithNotify,
-					Label,
-					Class);
-				UIAction.CanExecuteAction.BindRaw(
-					this, &SAnimNotifyTrack::IsValidToPlace,
-					Class);
+				FUIAction UIAction;
+				UIAction.ExecuteAction.BindSP(
+					this, &SAnimNotifyTrack::OnNewNotifyClicked);
+				MenuBuilder.AddMenuEntry(LOCTEXT("NewNotify", "New Notify..."), LOCTEXT("NewNotifyToolTip", "Create a new animation notify on the skeleton"), FSlateIcon(), UIAction);
 			}
 
-			MenuBuilder.AddMenuEntry(LabelText, Description, FSlateIcon(), UIAction);
-		}
-	}
-	MenuBuilder.EndSection();
-
-	MenuBuilder.BeginSection("AnimNotifyCustom", LOCTEXT("NewNotifySubMenu_Custom", "Custom"));
-	{
-		// now add custom anim notifiers
-		USkeleton* SeqSkeleton = Sequence->GetSkeleton();
-		if (SeqSkeleton)
-		{
-			for (int32 I = 0; I<SeqSkeleton->AnimationNotifies.Num(); ++I)
-			{
-				FName NotifyName = SeqSkeleton->AnimationNotifies[I];
-				FString Label = NotifyName.ToString();
-
-				FText Description = FText::GetEmpty();
-				if (!bIsReplaceWithMenu)
-				{
-					Description = LOCTEXT("NewNotifySubMenu_ToolTip", "Add an existing notify");
-				}
-				else
+			MenuBuilder.AddSubMenu(
+				LOCTEXT("NewNotifySubMenu_Skeleton", "Skeleton Notifies"),
+				LOCTEXT("NewNotifySubMenu_Skeleton_Tooltip", "Choose from custom notifies on the skeleton"),
+				FNewMenuDelegate::CreateLambda([this, SeqSkeleton, bIsReplaceWithMenu](FMenuBuilder& InSubMenuBuilder)
 				{
-					Description = LOCTEXT("ReplaceWithNotifySubMenu_ToolTip", "Replace with an existing notify");
-				}
+					ISkeletonEditorModule& SkeletonEditorModule = FModuleManager::LoadModuleChecked<ISkeletonEditorModule>("SkeletonEditor");
+					TSharedRef<IEditableSkeleton> EditableSkeleton = SkeletonEditorModule.CreateEditableSkeleton(SeqSkeleton);
 
-				FUIAction UIAction;
-				if (!bIsReplaceWithMenu)
-				{
-					UIAction.ExecuteAction.BindRaw(
-						this, &SAnimNotifyTrack::CreateNewNotifyAtCursor,
-						Label,
-						(UClass*)nullptr);
-				}
-				else
-				{
-					UIAction.ExecuteAction.BindRaw(
-						this, &SAnimNotifyTrack::ReplaceSelectedWithNotify,
-						Label,
-						(UClass*)nullptr);
-				}
+					InSubMenuBuilder.AddWidget(
+						SNew(SBox)
+						.MinDesiredWidth(300.0f)
+						.MaxDesiredHeight(400.0f)
+						[
+							SNew(SSkeletonAnimNotifies, EditableSkeleton)
+							.IsPicker(true)
+							.OnItemSelected_Lambda([this, bIsReplaceWithMenu](const FName& InNotifyName)
+							{
+								FSlateApplication::Get().DismissAllMenus();
 
-				MenuBuilder.AddMenuEntry( FText::FromString( Label ), Description, FSlateIcon(), UIAction);
-			}
+								if (!bIsReplaceWithMenu)
+								{
+									CreateNewNotifyAtCursor(InNotifyName.ToString(), nullptr);
+								}
+								else
+								{
+									ReplaceSelectedWithNotify(InNotifyName.ToString(), nullptr);
+								}
+							})
+						],
+						FText(), true, false
+					);
+				}));
 		}
+		MenuBuilder.EndSection();
 	}
-	MenuBuilder.EndSection();
 
-	if (!bIsReplaceWithMenu)
+	// Add a notify picker
+	MakeNewNotifyPicker<UAnimNotify>(MenuBuilder, bIsReplaceWithMenu);
+}
+
+void SAnimNotifyTrack::FillNewSyncMarkerMenu(FMenuBuilder& MenuBuilder)
+{
+	USkeleton* SeqSkeleton = Sequence->GetSkeleton();
+	if (SeqSkeleton)
 	{
-		MenuBuilder.BeginSection("AnimNotifyCreateNew");
+		MenuBuilder.BeginSection("AnimSyncMarkerSubMenu", LOCTEXT("NewSyncMarkerSubMenu_Skeleton", "Sync Markers"));
 		{
 			FUIAction UIAction;
-			UIAction.ExecuteAction.BindRaw(
-				this, &SAnimNotifyTrack::OnNewNotifyClicked);
-			MenuBuilder.AddMenuEntry(LOCTEXT("NewNotify", "New Notify"), LOCTEXT("NewNotifyToolTip", "Create a new animation notify"), FSlateIcon(), UIAction);
+			UIAction.ExecuteAction.BindSP(
+				this, &SAnimNotifyTrack::OnNewSyncMarkerClicked);
+			MenuBuilder.AddMenuEntry(LOCTEXT("NewSyncMarker", "New Sync Marker..."), LOCTEXT("NewSyncMarkerToolTip", "Create a new animation sync marker"), FSlateIcon(), UIAction);
+
+			MenuBuilder.AddSubMenu(
+				LOCTEXT("NewSyncMarkerSubMenu_Existing", "Existing Sync Markers"),
+				LOCTEXT("NewSyncMarkerSubMenu_Existing_Tooltip", "Choose from existing sync marker names on the skeleton"),
+				FNewMenuDelegate::CreateLambda([this, SeqSkeleton](FMenuBuilder& InSubMenuBuilder)
+			{
+				ISkeletonEditorModule& SkeletonEditorModule = FModuleManager::LoadModuleChecked<ISkeletonEditorModule>("SkeletonEditor");
+				TSharedRef<IEditableSkeleton> EditableSkeleton = SkeletonEditorModule.CreateEditableSkeleton(SeqSkeleton);
+
+				InSubMenuBuilder.AddWidget(
+					SNew(SBox)
+					.MinDesiredWidth(300.0f)
+					.MaxDesiredHeight(400.0f)
+					[
+						SNew(SSkeletonAnimNotifies, EditableSkeleton)
+						.IsSyncMarker(true)
+						.OnItemSelected_Lambda([this](const FName& InNotifyName)
+						{
+							FSlateApplication::Get().DismissAllMenus();
+
+							CreateNewSyncMarkerAtCursor(InNotifyName.ToString(), nullptr);
+						})
+					],
+					FText(), true, false
+					);
+			}));
 		}
 		MenuBuilder.EndSection();
 	}
@@ -3021,11 +2974,10 @@ TSharedPtr<SWidget> SAnimNotifyTrack::SummonContextMenu(const FGeometry& MyGeome
 
 			if (Sequence->IsA(UAnimSequence::StaticClass()))
 			{
-				MenuBuilder.AddMenuEntry(
-					LOCTEXT("NewSyncMarker", "Add Sync Marker"),
-					LOCTEXT("NewSyncMarkerToolTip", "Create a new animation sync marker"),
-					FSlateIcon(),
-					FUIAction(FExecuteAction::CreateSP(this, &SAnimNotifyTrack::OnNewSyncMarkerClicked)));
+				MenuBuilder.AddSubMenu(
+					NSLOCTEXT("NewSyncMarkerSubMenu", "NewSyncMarkerSubMenuAddNotifyState", "Add Sync Marker..."),
+					NSLOCTEXT("NewSyncMarkerSubMenu", "NewSyncMarkerSubMenuAddNotifyStateToolTip", "Create a new animation sync marker"),
+					FNewMenuDelegate::CreateRaw(this, &SAnimNotifyTrack::FillNewSyncMarkerMenu));
 			}
 
 			MenuBuilder.AddMenuEntry(
@@ -3284,8 +3236,13 @@ void SAnimNotifyTrack::AddNewSyncMarker(const FText& NewNotifyName, ETextCommit:
 	if ((CommitInfo == ETextCommit::OnEnter) && SeqSkeleton)
 	{
 		const FScopedTransaction Transaction(LOCTEXT("AddNewSyncMarker", "Add New Sync Marker"));
-		//FName NewName = FName(*NewNotifyName.ToString());
-		//SeqSkeleton->AddNewAnimationNotify(NewName);
+
+		FName NewName = FName(*NewNotifyName.ToString());
+
+		ISkeletonEditorModule& SkeletonEditorModule = FModuleManager::LoadModuleChecked<ISkeletonEditorModule>("SkeletonEditor");
+		TSharedRef<IEditableSkeleton> EditableSkeleton = SkeletonEditorModule.CreateEditableSkeleton(SeqSkeleton);
+
+		EditableSkeleton->AddSyncMarker(NewName);
 
 		FBlueprintActionDatabase::Get().RefreshAssetActions(SeqSkeleton);
 
@@ -3755,9 +3712,9 @@ void SAnimNotifyTrack::RefreshMarqueeSelectedNodes(FSlateRect& Rect, FNotifyMarq
 	}
 }
 
-FString SAnimNotifyTrack::MakeBlueprintNotifyName(FAssetData& NotifyAssetData)
+FString SAnimNotifyTrack::MakeBlueprintNotifyName(const FString& InNotifyClassName)
 {
-	FString DefaultNotifyName = NotifyAssetData.AssetName.ToString();
+	FString DefaultNotifyName = InNotifyClassName;
 	DefaultNotifyName = DefaultNotifyName.Replace(TEXT("AnimNotify_"), TEXT(""), ESearchCase::CaseSensitive);
 	DefaultNotifyName = DefaultNotifyName.Replace(TEXT("AnimNotifyState_"), TEXT(""), ESearchCase::CaseSensitive);
 
@@ -3774,32 +3731,6 @@ void SAnimNotifyTrack::ClearNodeTooltips()
 	}
 }
 
-void SAnimNotifyTrack::GetNotifyMenuData(TArray<FAssetData>& NotifyAssetData, TArray<BlueprintNotifyMenuInfo>& OutNotifyMenuData)
-{
-	for(FAssetData& NotifyData : NotifyAssetData)
-	{
-		OutNotifyMenuData.AddZeroed();
-		BlueprintNotifyMenuInfo& MenuInfo = OutNotifyMenuData.Last();
-
-		MenuInfo.BlueprintPath = NotifyData.ObjectPath.ToString();
-		MenuInfo.NotifyName = MakeBlueprintNotifyName(NotifyData);
-		// this functionality is only available in native class
-		// so we don't have to call BP function but just call native on the check of validity
-		FString NativeParentClassName;
-		if (NotifyData.GetTagValue(FBlueprintTags::NativeParentClassPath, NativeParentClassName))
-		{
-			UObject* Outer = nullptr;
-			ResolveName(Outer, NativeParentClassName, false, false);
-			MenuInfo.BaseClass = FindObject<UClass>(ANY_PACKAGE, *NativeParentClassName);
-		}
-	}
-
-	OutNotifyMenuData.Sort([](const BlueprintNotifyMenuInfo& A, const BlueprintNotifyMenuInfo& B)
-	{
-		return A.NotifyName < B.NotifyName;
-	});
-}
-
 const EVisibility SAnimNotifyTrack::GetTimingNodeVisibility(TSharedPtr<SAnimNotifyNode> NotifyNode)
 {
 	if(OnGetTimingNodeVisibility.IsBound())
@@ -3942,7 +3873,7 @@ void FAnimNotifyPanelCommands::RegisterCommands()
 //////////////////////////////////////////////////////////////////////////
 // SAnimNotifyPanel
 
-void SAnimNotifyPanel::Construct(const FArguments& InArgs, FSimpleMulticastDelegate& OnPostUndo)
+void SAnimNotifyPanel::Construct(const FArguments& InArgs, const TSharedRef<class IEditableSkeleton>& InEditableSkeleton)
 {
 	SAnimTrackPanel::Construct( SAnimTrackPanel::FArguments()
 		.WidgetWidth(InArgs._WidgetWidth)
@@ -3954,7 +3885,6 @@ void SAnimNotifyPanel::Construct(const FArguments& InArgs, FSimpleMulticastDeleg
 
 	Sequence = InArgs._Sequence;
 	MarkerBars = InArgs._MarkerBars;
-	OnAnimNotifiesChanged = InArgs._OnAnimNotifiesChanged;
 	OnInvokeTab = InArgs._OnInvokeTab;
 
 	FAnimNotifyPanelCommands::Register();
@@ -3966,7 +3896,12 @@ void SAnimNotifyPanel::Construct(const FArguments& InArgs, FSimpleMulticastDeleg
 	Sequence->InitializeNotifyTrack();
 	Sequence->RegisterOnNotifyChanged(UAnimSequenceBase::FOnNotifyChanged::CreateSP( this, &SAnimNotifyPanel::RefreshNotifyTracks )  );
 
-	OnPostUndo.Add(FSimpleDelegate::CreateSP( this, &SAnimNotifyPanel::PostUndo ) );
+	InEditableSkeleton->RegisterOnNotifiesChanged(FSimpleDelegate::CreateSP(this, &SAnimNotifyPanel::RefreshNotifyTracks));
+
+	if(GEditor)
+	{
+		GEditor->RegisterForUndo(this);
+	}
 
 	CurrentPosition = InArgs._CurrentPosition;
 	OnSelectionChanged = InArgs._OnSelectionChanged;
@@ -4038,6 +3973,11 @@ SAnimNotifyPanel::~SAnimNotifyPanel()
 	Sequence->UnregisterOnNotifyChanged(this);
 
 	FCoreUObjectDelegates::OnObjectPropertyChanged.Remove(OnPropertyChangedHandleDelegateHandle);
+
+	if(GEditor)
+	{
+		GEditor->UnregisterForUndo(this);
+	}
 }
 
 FName SAnimNotifyPanel::GetNewTrackName() const
@@ -4165,8 +4105,6 @@ void SAnimNotifyPanel::Update()
 	{
 		Sequence->RefreshCacheData();
 	}
-
-	OnAnimNotifiesChanged.ExecuteIfBound();
 }
 
 void SAnimNotifyPanel::RefreshNotifyTracks()
@@ -4289,7 +4227,15 @@ FReply SAnimNotifyPanel::OnNotifyNodeDragStarted(TArray<TSharedPtr<SAnimNotifyNo
 	return FReply::Handled().BeginDragDrop(FNotifyDragDropOp::New(Nodes, NodeDragDecorator, NotifyAnimTracks, Sequence, ScreenCursorPos, OverlayOrigin, OverlayExtents, CurrentDragXPosition, PanRequestDelegate, MarkerBars, UpdateDelegate));
 }
 
-void SAnimNotifyPanel::PostUndo()
+void SAnimNotifyPanel::PostUndo( bool bSuccess )
+{
+	if(Sequence != NULL)
+	{
+		Sequence->RefreshCacheData();
+	}
+}
+
+void SAnimNotifyPanel::PostRedo( bool bSuccess )
 {
 	if(Sequence != NULL)
 	{
@@ -4389,9 +4335,6 @@ void SAnimNotifyPanel::DeselectAllNotifies()
 		Track->DeselectAllNotifyNodes(false);
 	}
 
-	// Broadcast the change so the editor can update
-	OnAnimNotifiesChanged.ExecuteIfBound();
-
 	OnTrackSelectionChanged();
 }
 
@@ -4664,9 +4607,6 @@ void SAnimNotifyPanel::OnPropertyChanged(UObject* ChangedObject, FPropertyChange
 				RefreshNotifyTracks();
 			}
 		}
-
-		// Broadcast the change so the editor can update
-		OnAnimNotifiesChanged.ExecuteIfBound();
 	}
 }
 
@@ -4921,9 +4861,6 @@ void SAnimNotifyPanel::OnNotifyObjectChanged(UObject* EditorBaseObj, bool bRebui
 		{
 			NotifyAnimTracks[WidgetTrackIdx]->Update();
 		}
-
-		// Broadcast the change so the editor can update
-		OnAnimNotifiesChanged.ExecuteIfBound();
 	}
 }
 
diff --git a/Engine/Source/Editor/Persona/Private/SAnimNotifyPanel.h b/Engine/Source/Editor/Persona/Private/SAnimNotifyPanel.h
--- a/Engine/Source/Editor/Persona/Private/SAnimNotifyPanel.h
+++ b/Engine/Source/Editor/Persona/Private/SAnimNotifyPanel.h
@@ -16,6 +16,7 @@
 #include "AssetData.h"
 #include "Framework/Commands/Commands.h"
 #include "SAnimTimingPanel.h"
+#include "EditorUndoClient.h"
 
 class FSlateWindowElementList;
 class SAnimNotifyNode;
@@ -153,7 +154,7 @@ public:
 
 // @todo anim : register when it's opened for the animsequence
 // broadcast when animsequence changed, so that we refresh for multiple window
-class SAnimNotifyPanel: public SAnimTrackPanel
+class SAnimNotifyPanel: public SAnimTrackPanel, public FEditorUndoClient
 {
 public:
 	SLATE_BEGIN_ARGS( SAnimNotifyPanel )
@@ -182,12 +183,11 @@ public:
 	SLATE_EVENT( FOnGetScrubValue, OnGetScrubValue )
 	SLATE_EVENT( FRefreshOffsetsRequest, OnRequestRefreshOffsets )
 	SLATE_EVENT( FOnGetTimingNodeVisibility, OnGetTimingNodeVisibility )
-	SLATE_EVENT( FSimpleDelegate, OnAnimNotifiesChanged )
 	SLATE_EVENT( FOnInvokeTab, OnInvokeTab )
 
 	SLATE_END_ARGS()
 
-	void Construct(const FArguments& InArgs, FSimpleMulticastDelegate& OnPostUndo);
+	void Construct(const FArguments& InArgs, const TSharedRef<class IEditableSkeleton>& InEditableSkeleton);
 	virtual ~SAnimNotifyPanel();
 
 	void SetSequence(class UAnimSequenceBase *	InSequence);
@@ -273,7 +273,10 @@ private:
 	// this just refresh notify tracks - UI purpose only
 	// do not call this from here. This gets called by asset. 
 	void RefreshNotifyTracks();
-	void PostUndo();
+
+	/** FEditorUndoClient interface */
+	virtual void PostUndo( bool bSuccess ) override;
+	virtual void PostRedo( bool bSuccess ) override;
 
 	/** Handler for delete command */
 	void OnDeletePressed();
@@ -335,9 +338,6 @@ private:
 	/** Handle to the registered OnPropertyChangedHandle delegate */
 	FDelegateHandle OnPropertyChangedHandleDelegateHandle;
 
-	/** Delegate fired when anim notifies are changed */
-	FSimpleDelegate OnAnimNotifiesChanged;
-
 	/** Delegate used to invoke a tab */
 	FOnInvokeTab OnInvokeTab;
 };
diff --git a/Engine/Source/Editor/Persona/Private/SAnimTimingPanel.cpp b/Engine/Source/Editor/Persona/Private/SAnimTimingPanel.cpp
--- a/Engine/Source/Editor/Persona/Private/SAnimTimingPanel.cpp
+++ b/Engine/Source/Editor/Persona/Private/SAnimTimingPanel.cpp
@@ -18,6 +18,10 @@
 #include "Preferences/PersonaOptions.h"
 #include "Widgets/Layout/SExpandableArea.h"
 #include "STimingTrack.h"
+#include "ISkeletonEditorModule.h"
+#include "Modules/ModuleManager.h"
+#include "IEditableSkeleton.h"
+#include "Editor.h"
 
 #define LOCTEXT_NAMESPACE "AnimTimingPanel"
 
@@ -188,7 +192,7 @@ void SAnimTimingTrackNode::Construct(const FArguments& InArgs)
 		);
 }
 
-void SAnimTimingPanel::Construct(const FArguments& InArgs, FSimpleMulticastDelegate& OnAnimNotifiesChanged, FSimpleMulticastDelegate& OnSectionsChanged)
+void SAnimTimingPanel::Construct(const FArguments& InArgs, FSimpleMulticastDelegate& OnSectionsChanged)
 {
 	SAnimTrackPanel::Construct(SAnimTrackPanel::FArguments()
 		.WidgetWidth(InArgs._WidgetWidth)
@@ -223,7 +227,10 @@ void SAnimTimingPanel::Construct(const FArguments& InArgs, FSimpleMulticastDeleg
 	Update();
 
 	// Register to some delegates to update the interface
-	OnAnimNotifiesChanged.Add(FSimpleDelegate::CreateSP(this, &SAnimTimingPanel::RefreshTrackNodes));
+	ISkeletonEditorModule& SkeletonEditorModule = FModuleManager::LoadModuleChecked<ISkeletonEditorModule>("SkeletonEditor");
+	TSharedPtr<IEditableSkeleton> EditableSkeleton = SkeletonEditorModule.CreateEditableSkeleton(AnimSequence->GetSkeleton());
+	EditableSkeleton->RegisterOnNotifiesChanged(FSimpleDelegate::CreateSP(this, &SAnimTimingPanel::RefreshTrackNodes));
+
 	OnSectionsChanged.Add(FSimpleDelegate::CreateSP(this, &SAnimTimingPanel::RefreshTrackNodes));
 
 	// Clear display flags
diff --git a/Engine/Source/Editor/Persona/Private/SAnimTimingPanel.h b/Engine/Source/Editor/Persona/Private/SAnimTimingPanel.h
--- a/Engine/Source/Editor/Persona/Private/SAnimTimingPanel.h
+++ b/Engine/Source/Editor/Persona/Private/SAnimTimingPanel.h
@@ -193,7 +193,7 @@ public:
 
 	// Construct the panel
 	// @param InArgs - Slate arguments
-	void Construct(const FArguments& InArgs, FSimpleMulticastDelegate& OnAnimNotifiesChanged, FSimpleMulticastDelegate& OnSectionsChanged);
+	void Construct(const FArguments& InArgs, FSimpleMulticastDelegate& OnSectionsChanged);
 
 	// Updates panel widgets
 	void Update();
diff --git a/Engine/Source/Editor/Persona/Private/SAnimationEditorViewport.cpp b/Engine/Source/Editor/Persona/Private/SAnimationEditorViewport.cpp
--- a/Engine/Source/Editor/Persona/Private/SAnimationEditorViewport.cpp
+++ b/Engine/Source/Editor/Persona/Private/SAnimationEditorViewport.cpp
@@ -500,6 +500,8 @@ void SAnimationEditorViewportTabBody::Construct(const FArguments& InArgs, const
 	GetPreviewScene()->OnRecordingStateChanged().AddSP(this, &SAnimationEditorViewportTabBody::AddRecordingNotification);
 
 	AddPostProcessNotification();
+
+	AddMinLODNotification();
 }
 
 void SAnimationEditorViewportTabBody::BindCommands()
@@ -2159,6 +2161,55 @@ void SAnimationEditorViewportTabBody::AddPostProcessNotification()
 	);
 }
 
+void SAnimationEditorViewportTabBody::AddMinLODNotification()
+{
+	if(WeakMinLODNotification.IsValid())
+	{
+		return;
+	}
+
+	auto GetMinLODNotificationVisibility = [this]()
+	{
+		if (GetPreviewScene()->GetPreviewMesh() && GetPreviewScene()->GetPreviewMesh()->MinLod.Default != 0)
+		{
+			return EVisibility::Visible;
+		}
+
+		return EVisibility::Collapsed;
+	};
+
+	WeakMinLODNotification = AddNotification(EMessageSeverity::Info,
+		true,
+		SNew(SHorizontalBox)
+		.Visibility_Lambda(GetMinLODNotificationVisibility)
+		.ToolTipText(LOCTEXT("MinLODNotificationTooltip", "This asset has a minimum LOD applied."))
+		+SHorizontalBox::Slot()
+		.FillWidth(1.0f)
+		.Padding(2.0f, 4.0f)
+		[
+			SNew(SHorizontalBox)
+			+SHorizontalBox::Slot()
+			.AutoWidth()
+			.VAlign(VAlign_Center)
+			.Padding(0.0f, 0.0f, 4.0f, 0.0f)
+			[
+				SNew(STextBlock)
+				.TextStyle(FEditorStyle::Get(), "AnimViewport.MessageText")
+				.Font(FEditorStyle::Get().GetFontStyle("FontAwesome.9"))
+				.Text(FEditorFontGlyphs::Level_Down)
+			]
+			+SHorizontalBox::Slot()
+			.VAlign(VAlign_Center)
+			.FillWidth(1.0f)
+			[
+				SNew(STextBlock)
+				.Text(LOCTEXT("MinLODNotification", "Min LOD applied"))
+				.TextStyle(FEditorStyle::Get(), "AnimViewport.MessageText")
+			]
+		]
+	);
+}
+
 void SAnimationEditorViewportTabBody::HandleFocusCamera()
 {
 	TSharedRef<FAnimationViewportClient> AnimViewportClient = StaticCastSharedRef<FAnimationViewportClient>(LevelViewportClient.ToSharedRef());
diff --git a/Engine/Source/Editor/Persona/Private/SAnimationEditorViewport.h b/Engine/Source/Editor/Persona/Private/SAnimationEditorViewport.h
--- a/Engine/Source/Editor/Persona/Private/SAnimationEditorViewport.h
+++ b/Engine/Source/Editor/Persona/Private/SAnimationEditorViewport.h
@@ -560,6 +560,9 @@ private:
 	/** Recording notification */
 	TWeakPtr<SWidget> WeakRecordingNotification;
 
+	/** Min LOD notification */
+	TWeakPtr<SWidget> WeakMinLODNotification;
+
 	/** Current LOD selection*/
 	int32 LODSelection;
 
@@ -584,4 +587,7 @@ private:
 
 	/** Adds a persistent notification to display post process graph state */
 	void AddPostProcessNotification();
+
+	/** Add a notification to tell the user a min LOD is being applied */
+	void AddMinLODNotification();
 };
diff --git a/Engine/Source/Editor/Persona/Private/SAnimationSequenceBrowser.cpp b/Engine/Source/Editor/Persona/Private/SAnimationSequenceBrowser.cpp
--- a/Engine/Source/Editor/Persona/Private/SAnimationSequenceBrowser.cpp
+++ b/Engine/Source/Editor/Persona/Private/SAnimationSequenceBrowser.cpp
@@ -1147,7 +1147,7 @@ void SAnimationSequenceBrowser::CreateAssetTooltipResources()
 	ViewportWidget->SetViewportInterface(SceneViewport.ToSharedRef());
 
 	// Setup the preview component to ensure an animation will update when requested
-	PreviewComponent->MeshComponentUpdateFlag = EMeshComponentUpdateFlag::AlwaysTickPoseAndRefreshBones;
+	PreviewComponent->VisibilityBasedAnimTickOption = EVisibilityBasedAnimTickOption::AlwaysTickPoseAndRefreshBones;
 	PreviewScene.AddComponent(PreviewComponent, FTransform::Identity);
 
 	const USkeletalMeshEditorSettings* Options = GetDefault<USkeletalMeshEditorSettings>();
diff --git a/Engine/Source/Editor/Persona/Private/SMontageEditor.cpp b/Engine/Source/Editor/Persona/Private/SMontageEditor.cpp
--- a/Engine/Source/Editor/Persona/Private/SMontageEditor.cpp
+++ b/Engine/Source/Editor/Persona/Private/SMontageEditor.cpp
@@ -21,6 +21,7 @@
 #include "Widgets/Input/SButton.h"
 #include "Widgets/Images/SImage.h"
 #include "Factories/AnimMontageFactory.h"
+#include "Editor.h"
 
 #define LOCTEXT_NAMESPACE "AnimSequenceEditor"
 
@@ -52,9 +53,12 @@ void SMontageEditor::Construct(const FArguments& InArgs, const FMontageEditorReq
 		.OnObjectsSelected(InArgs._OnObjectsSelected), 
 		InRequiredArgs.PreviewScene );
 
-	InRequiredArgs.OnPostUndo.Add(FSimpleDelegate::CreateSP(this, &SMontageEditor::PostUndo));
+	if(GEditor)
+	{
+		GEditor->RegisterForUndo(this);
+	}
 
-	SAssignNew(AnimTimingPanel, SAnimTimingPanel, InRequiredArgs.OnAnimNotifiesChanged, InRequiredArgs.OnSectionsChanged)
+	SAssignNew(AnimTimingPanel, SAnimTimingPanel, InRequiredArgs.OnSectionsChanged)
 		.IsEnabled(!bChildAnimMontage)
 		.InSequence(MontageObj)
 		.WidgetWidth(S2ColumnWidget::DEFAULT_RIGHT_COLUMN_WIDTH)
@@ -144,7 +148,7 @@ void SMontageEditor::Construct(const FArguments& InArgs, const FMontageEditorReq
 	.AutoHeight()
 	.Padding(0, 5)
 	[
-		SAssignNew( AnimMontagePanel, SAnimMontagePanel, InRequiredArgs.OnAnimNotifiesChanged, InRequiredArgs.OnSectionsChanged)
+		SAssignNew( AnimMontagePanel, SAnimMontagePanel, InRequiredArgs.OnSectionsChanged)
 		.Montage(MontageObj)
 		.bChildAnimMontage(bChildAnimMontage)
 		.MontageEditor(SharedThis(this))
@@ -180,7 +184,7 @@ void SMontageEditor::Construct(const FArguments& InArgs, const FMontageEditorReq
 	.AutoHeight()
 	.Padding(0, 10)
 	[
-		SAssignNew( AnimNotifyPanel, SAnimNotifyPanel, InRequiredArgs.OnPostUndo )
+		SAssignNew( AnimNotifyPanel, SAnimNotifyPanel, InRequiredArgs.EditableSkeleton )
 		.Sequence(MontageObj)
 		.IsEnabled(!bChildAnimMontage)
 		.WidgetWidth(S2ColumnWidget::DEFAULT_RIGHT_COLUMN_WIDTH)
@@ -194,7 +198,6 @@ void SMontageEditor::Construct(const FArguments& InArgs, const FMontageEditorReq
 		.MarkerBars(this, &SMontageEditor::GetMarkerBarInformation)
 		.OnRequestRefreshOffsets(this, &SMontageEditor::RefreshNotifyTriggerOffsets)
 		.OnGetTimingNodeVisibility(TimingNodeVisibilityDelegate)
-		.OnAnimNotifiesChanged(InArgs._OnAnimNotifiesChanged)
 		.OnInvokeTab(InArgs._OnInvokeTab)
 	];
 
@@ -253,6 +256,11 @@ SMontageEditor::~SMontageEditor()
 	{
 		MontageObj->UnregisterOnMontageChanged(this);
 	}
+
+	if(GEditor)
+	{
+		GEditor->UnregisterForUndo(this);
+	}
 }
 
 TSharedRef<class SAnimationScrubPanel> SMontageEditor::ConstructAnimScrubPanel()
@@ -886,7 +894,17 @@ void SMontageEditor::ClearSquenceOrdering()
 	RestartPreview();
 }
 
-void SMontageEditor::PostUndo()
+void SMontageEditor::PostUndo( bool bSuccess )
+{
+	PostRedoUndo();
+}
+
+void SMontageEditor::PostRedo( bool bSuccess )
+{
+	PostRedoUndo();
+}
+
+void SMontageEditor::PostRedoUndo()
 {
 	// when undo or redo happens, we still have to recalculate length, so we can't rely on sequence length changes or not
 	if (MontageObj->SequenceLength)
diff --git a/Engine/Source/Editor/Persona/Private/SMontageEditor.h b/Engine/Source/Editor/Persona/Private/SMontageEditor.h
--- a/Engine/Source/Editor/Persona/Private/SMontageEditor.h
+++ b/Engine/Source/Editor/Persona/Private/SMontageEditor.h
@@ -18,18 +18,14 @@ class SAnimTimingPanel;
 
 struct FMontageEditorRequiredArgs
 {
-	FMontageEditorRequiredArgs(const TSharedRef<class IPersonaPreviewScene>& InPreviewScene, const TSharedRef<class IEditableSkeleton>& InEditableSkeleton, FSimpleMulticastDelegate& InOnPostUndo, FSimpleMulticastDelegate& InOnAnimNotifiesChanged, FSimpleMulticastDelegate& InOnSectionsChanged)
+	FMontageEditorRequiredArgs(const TSharedRef<class IPersonaPreviewScene>& InPreviewScene, const TSharedRef<class IEditableSkeleton>& InEditableSkeleton, FSimpleMulticastDelegate& InOnSectionsChanged)
 		: PreviewScene(InPreviewScene)
 		, EditableSkeleton(InEditableSkeleton)
-		, OnPostUndo(InOnPostUndo)
-		, OnAnimNotifiesChanged(InOnAnimNotifiesChanged)
 		, OnSectionsChanged(InOnSectionsChanged)
 	{}
 
 	TSharedRef<class IPersonaPreviewScene> PreviewScene;
 	TSharedRef<class IEditableSkeleton> EditableSkeleton;
-	FSimpleMulticastDelegate& OnPostUndo;
-	FSimpleMulticastDelegate& OnAnimNotifiesChanged;
 	FSimpleMulticastDelegate& OnSectionsChanged;
 };
 
@@ -42,7 +38,7 @@ struct FMontageEditorRequiredArgs
 	portion of the Montage tool and registering callbacks to the SMontageEditor to do the actual editing.
 	
 */
-class SMontageEditor : public SAnimEditorBase 
+class SMontageEditor : public SAnimEditorBase, public FEditorUndoClient
 {
 public:
 	SLATE_BEGIN_ARGS( SMontageEditor )
@@ -54,7 +50,6 @@ public:
 		SLATE_EVENT(FSimpleDelegate, OnSectionsChanged)
 		SLATE_ARGUMENT( UAnimMontage*, Montage )
 		SLATE_EVENT(FOnObjectsSelected, OnObjectsSelected)
-		SLATE_EVENT(FSimpleDelegate, OnAnimNotifiesChanged)
 	SLATE_END_ARGS()
 
 	~SMontageEditor();
@@ -121,7 +116,12 @@ private:
 	void EnsureStartingSection();
 	void EnsureSlotNode();
 	virtual bool ClampToEndTime(float NewEndTime) override;
-	void PostUndo();
+
+	/** FEditorUndoClient interface */
+	virtual void PostUndo( bool bSuccess ) override;
+	virtual void PostRedo( bool bSuccess ) override;
+
+	void PostRedoUndo();
 	
 	bool GetSectionTime( int32 SectionIndex, float &OutTime ) const;
 
diff --git a/Engine/Source/Editor/Persona/Private/SSequenceEditor.cpp b/Engine/Source/Editor/Persona/Private/SSequenceEditor.cpp
--- a/Engine/Source/Editor/Persona/Private/SSequenceEditor.cpp
+++ b/Engine/Source/Editor/Persona/Private/SSequenceEditor.cpp
@@ -7,13 +7,14 @@
 #include "SAnimNotifyPanel.h"
 #include "SAnimTrackCurvePanel.h"
 #include "AnimPreviewInstance.h"
+#include "Editor.h"
 
 #define LOCTEXT_NAMESPACE "AnimSequenceEditor"
 
 //////////////////////////////////////////////////////////////////////////
 // SSequenceEditor
 
-void SSequenceEditor::Construct(const FArguments& InArgs, TSharedRef<class IPersonaPreviewScene> InPreviewScene, TSharedRef<class IEditableSkeleton> InEditableSkeleton, FSimpleMulticastDelegate& OnPostUndo)
+void SSequenceEditor::Construct(const FArguments& InArgs, TSharedRef<class IPersonaPreviewScene> InPreviewScene, TSharedRef<class IEditableSkeleton> InEditableSkeleton)
 {
 	SequenceObj = InArgs._Sequence;
 	check(SequenceObj);
@@ -23,13 +24,16 @@ void SSequenceEditor::Construct(const FArguments& InArgs, TSharedRef<class IPers
 		.OnObjectsSelected(InArgs._OnObjectsSelected), 
 		InPreviewScene );
 
-	OnPostUndo.Add(FSimpleDelegate::CreateSP( this, &SSequenceEditor::PostUndo ) );
+	if(GEditor)
+	{
+		GEditor->RegisterForUndo(this);
+	}
 
 	EditorPanels->AddSlot()
 	.AutoHeight()
 	.Padding(0, 10)
 	[
-		SAssignNew( AnimNotifyPanel, SAnimNotifyPanel, OnPostUndo)
+		SAssignNew( AnimNotifyPanel, SAnimNotifyPanel, InEditableSkeleton )
 		.Sequence(SequenceObj)
 		.WidgetWidth(S2ColumnWidget::DEFAULT_RIGHT_COLUMN_WIDTH)
 		.ViewInputMin(this, &SAnimEditorBase::GetViewMinInput)
@@ -39,7 +43,6 @@ void SSequenceEditor::Construct(const FArguments& InArgs, TSharedRef<class IPers
 		.OnSetInputViewRange(this, &SAnimEditorBase::SetInputViewRange)
 		.OnGetScrubValue(this, &SAnimEditorBase::GetScrubValue)
 		.OnSelectionChanged(this, &SAnimEditorBase::OnSelectionChanged)
-		.OnAnimNotifiesChanged(InArgs._OnAnimNotifiesChanged)
 		.OnInvokeTab(InArgs._OnInvokeTab)
 	];
 
@@ -78,7 +81,25 @@ void SSequenceEditor::Construct(const FArguments& InArgs, TSharedRef<class IPers
 	}
 }
 
-void SSequenceEditor::PostUndo()
+SSequenceEditor::~SSequenceEditor()
+{
+	if(GEditor)
+	{
+		GEditor->UnregisterForUndo(this);
+	}
+}
+
+void SSequenceEditor::PostUndo( bool bSuccess )
+{
+	PostUndoRedo();
+}
+
+void SSequenceEditor::PostRedo( bool bSuccess )
+{
+	PostUndoRedo();
+}
+
+void SSequenceEditor::PostUndoRedo()
 {
 	GetPreviewScene()->SetPreviewAnimationAsset(SequenceObj);
 
diff --git a/Engine/Source/Editor/Persona/Private/SSequenceEditor.h b/Engine/Source/Editor/Persona/Private/SSequenceEditor.h
--- a/Engine/Source/Editor/Persona/Private/SSequenceEditor.h
+++ b/Engine/Source/Editor/Persona/Private/SSequenceEditor.h
@@ -9,6 +9,7 @@
 #include "SAnimCurvePanel.h"
 #include "SAnimationScrubPanel.h"
 #include "SAnimEditorBase.h"
+#include "EditorUndoClient.h"
 
 class SAnimNotifyPanel;
 class SAnimTrackCurvePanel;
@@ -17,7 +18,7 @@ class SAnimTrackCurvePanel;
 // SSequenceEditor
 
 /** Overall animation sequence editing widget */
-class SSequenceEditor : public SAnimEditorBase
+class SSequenceEditor : public SAnimEditorBase, public FEditorUndoClient
 {
 public:
 	SLATE_BEGIN_ARGS( SSequenceEditor )
@@ -26,7 +27,6 @@ public:
 
 		SLATE_ARGUMENT(UAnimSequenceBase*, Sequence)
 		SLATE_EVENT(FOnObjectsSelected, OnObjectsSelected)
-		SLATE_EVENT(FSimpleDelegate, OnAnimNotifiesChanged)
 		SLATE_EVENT(FOnInvokeTab, OnInvokeTab)
 
 	SLATE_END_ARGS()
@@ -38,7 +38,9 @@ private:
 	TSharedPtr<class SAnimationScrubPanel> AnimScrubPanel;
 	TWeakPtr<class IPersonaPreviewScene> PreviewScenePtr;
 public:
-	void Construct(const FArguments& InArgs, TSharedRef<class IPersonaPreviewScene> InPreviewScene, TSharedRef<class IEditableSkeleton> InEditableSkeleton, FSimpleMulticastDelegate& OnPostUndo);
+	void Construct(const FArguments& InArgs, TSharedRef<class IPersonaPreviewScene> InPreviewScene, TSharedRef<class IEditableSkeleton> InEditableSkeleton);
+
+	~SSequenceEditor();
 
 	virtual UAnimationAsset* GetEditorObject() const override { return SequenceObj; }
 
@@ -46,6 +48,10 @@ private:
 	/** Pointer to the animation sequence being edited */
 	UAnimSequenceBase* SequenceObj;
 
+	/** FEditorUndoClient interface */
+	virtual void PostUndo( bool bSuccess ) override;
+	virtual void PostRedo( bool bSuccess ) override;
+
 	/** Post undo **/
-	void PostUndo();
+	void PostUndoRedo();
 };
diff --git a/Engine/Source/Editor/Persona/Private/SSkeletonAnimNotifies.cpp b/Engine/Source/Editor/Persona/Private/SSkeletonAnimNotifies.cpp
--- a/Engine/Source/Editor/Persona/Private/SSkeletonAnimNotifies.cpp
+++ b/Engine/Source/Editor/Persona/Private/SSkeletonAnimNotifies.cpp
@@ -15,81 +15,21 @@
 #include "ScopedTransaction.h"
 #include "IEditableSkeleton.h"
 #include "TabSpawners.h"
+#include "Editor.h"
 
 #define LOCTEXT_NAMESPACE "SkeletonAnimNotifies"
 
-static const FName ColumnId_AnimNotifyNameLabel( "AnimNotifyName" );
-
 //////////////////////////////////////////////////////////////////////////
 // SMorphTargetListRow
 
 typedef TSharedPtr< FDisplayedAnimNotifyInfo > FDisplayedAnimNotifyInfoPtr;
 
-class SAnimNotifyListRow
-	: public SMultiColumnTableRow< FDisplayedAnimNotifyInfoPtr >
-{
-public:
-
-	SLATE_BEGIN_ARGS( SAnimNotifyListRow ) {}
-
-	/** The item for this row **/
-	SLATE_ARGUMENT( FDisplayedAnimNotifyInfoPtr, Item )
-
-	/* Widget used to display the list of morph targets */
-	SLATE_ARGUMENT( TSharedPtr<SSkeletonAnimNotifies>, NotifiesListView )
-
-	SLATE_END_ARGS()
-
-	void Construct( const FArguments& InArgs, const TSharedRef<STableViewBase>& OwnerTableView );
-
-	/** Overridden from SMultiColumnTableRow.  Generates a widget for this column of the tree row. */
-	virtual TSharedRef<SWidget> GenerateWidgetForColumn( const FName& ColumnName ) override;
-
-private:
-
-	/** Widget used to display the list of notifies */
-	TSharedPtr<SSkeletonAnimNotifies> NotifiesListView;
-
-	/** The notify being displayed by this row */
-	FDisplayedAnimNotifyInfoPtr	Item;
-};
-
-void SAnimNotifyListRow::Construct( const FArguments& InArgs, const TSharedRef<STableViewBase>& InOwnerTableView )
-{
-	Item = InArgs._Item;
-	NotifiesListView = InArgs._NotifiesListView;
-
-	check( Item.IsValid() );
-
-	SMultiColumnTableRow< FDisplayedAnimNotifyInfoPtr >::Construct( FSuperRowType::FArguments(), InOwnerTableView );
-}
-
-TSharedRef< SWidget > SAnimNotifyListRow::GenerateWidgetForColumn( const FName& ColumnName )
-{
-	check( ColumnName == ColumnId_AnimNotifyNameLabel );
-
-	return SNew( SVerticalBox )
-	+ SVerticalBox::Slot()
-	.AutoHeight()
-	.Padding( 0.0f, 4.0f )
-	.VAlign( VAlign_Center )
-	[
-		SAssignNew(Item->InlineEditableText, SInlineEditableTextBlock)
-		.Text( FText::FromName(Item->Name) )
-		.OnVerifyTextChanged(NotifiesListView.Get(), &SSkeletonAnimNotifies::OnVerifyNotifyNameCommit, Item)
-		.OnTextCommitted(NotifiesListView.Get(), &SSkeletonAnimNotifies::OnNotifyNameCommitted, Item)
-		.IsSelected(NotifiesListView.Get(), &SSkeletonAnimNotifies::IsSelected)
-	];
-}
-
 /////////////////////////////////////////////////////
 // FSkeletonAnimNotifiesSummoner
 
-FSkeletonAnimNotifiesSummoner::FSkeletonAnimNotifiesSummoner(TSharedPtr<class FAssetEditorToolkit> InHostingApp, const TSharedRef<class IEditableSkeleton>& InEditableSkeleton, FSimpleMulticastDelegate& InOnChangeAnimNotifies, FSimpleMulticastDelegate& InOnPostUndo, FOnObjectsSelected InOnObjectsSelected)
+FSkeletonAnimNotifiesSummoner::FSkeletonAnimNotifiesSummoner(TSharedPtr<class FAssetEditorToolkit> InHostingApp, const TSharedRef<class IEditableSkeleton>& InEditableSkeleton, FOnObjectsSelected InOnObjectsSelected)
 	: FWorkflowTabFactory(FPersonaTabs::SkeletonAnimNotifiesID, InHostingApp)
 	, EditableSkeleton(InEditableSkeleton)
-	, OnChangeAnimNotifies(InOnChangeAnimNotifies)
-	, OnPostUndo(InOnPostUndo)
 	, OnObjectsSelected(InOnObjectsSelected)
 {
 	TabLabel = LOCTEXT("SkeletonAnimNotifiesTabTitle", "Animation Notifies");
@@ -104,22 +44,37 @@ FSkeletonAnimNotifiesSummoner::FSkeletonAnimNotifiesSummoner(TSharedPtr<class FA
 
 TSharedRef<SWidget> FSkeletonAnimNotifiesSummoner::CreateTabBody(const FWorkflowTabSpawnInfo& Info) const
 {
-	return SNew(SSkeletonAnimNotifies, EditableSkeleton.Pin().ToSharedRef(), OnChangeAnimNotifies, OnPostUndo)
+	return SNew(SSkeletonAnimNotifies, EditableSkeleton.Pin().ToSharedRef())
 		.OnObjectsSelected(OnObjectsSelected);
 }
 
 /////////////////////////////////////////////////////
 // SSkeletonAnimNotifies
 
-void SSkeletonAnimNotifies::Construct(const FArguments& InArgs, const TSharedRef<IEditableSkeleton>& InEditableSkeleton, FSimpleMulticastDelegate& InOnChangeAnimNotifies, FSimpleMulticastDelegate& InOnPostUndo)
+void SSkeletonAnimNotifies::Construct(const FArguments& InArgs, const TSharedRef<IEditableSkeleton>& InEditableSkeleton)
 {
 	OnObjectsSelected = InArgs._OnObjectsSelected;
+	OnItemSelected = InArgs._OnItemSelected;
+	bIsPicker = InArgs._IsPicker;
+	bIsSyncMarker = InArgs._IsSyncMarker;
+
 	EditableSkeleton = InEditableSkeleton;
 
-	InOnChangeAnimNotifies.Add(FSimpleDelegate::CreateSP( this, &SSkeletonAnimNotifies::PostUndo ) );
-	InOnPostUndo.Add(FSimpleDelegate::CreateSP( this, &SSkeletonAnimNotifies::RefreshNotifiesListWithFilter ) );
+	if (bIsSyncMarker)
+	{
+		bIsPicker = false; // Sync Markers are never pickers
+	}
+	else
+	{
+		EditableSkeleton->RegisterOnNotifiesChanged(FSimpleDelegate::CreateSP(this, &SSkeletonAnimNotifies::OnNotifiesChanged));
+	}
+
+	if(GEditor)
+	{
+		GEditor->RegisterForUndo(this);
+	}
 
-	EditableSkeleton->RegisterOnNotifiesChanged(FSimpleDelegate::CreateSP(this, &SSkeletonAnimNotifies::OnNotifiesChanged));
+	FOnContextMenuOpening OnContextMenuOpening = (!bIsPicker && !bIsSyncMarker) ? FOnContextMenuOpening::CreateSP(this, &SSkeletonAnimNotifies::OnGetContextMenuContent) : FOnContextMenuOpening();
 
 	this->ChildSlot
 	[
@@ -142,21 +97,24 @@ void SSkeletonAnimNotifies::Construct(const FArguments& InArgs, const TSharedRef
 			SAssignNew( NotifiesListView, SAnimNotifyListType )
 			.ListItemsSource( &NotifyList )
 			.OnGenerateRow( this, &SSkeletonAnimNotifies::GenerateNotifyRow )
-			.OnContextMenuOpening( this, &SSkeletonAnimNotifies::OnGetContextMenuContent )
+			.OnContextMenuOpening( OnContextMenuOpening )
 			.OnSelectionChanged( this, &SSkeletonAnimNotifies::OnNotifySelectionChanged )
 			.ItemHeight( 22.0f )
-			.HeaderRow
-			(
-				SNew( SHeaderRow )
-				+ SHeaderRow::Column( ColumnId_AnimNotifyNameLabel )
-				.DefaultLabel( LOCTEXT( "AnimNotifyNameLabel", "Notify Name" ) )
-			)
+			.OnItemScrolledIntoView( this, &SSkeletonAnimNotifies::OnItemScrolledIntoView )
 		]
 	];
 
 	CreateNotifiesList();
 }
 
+SSkeletonAnimNotifies::~SSkeletonAnimNotifies()
+{
+	if(GEditor)
+	{
+		GEditor->UnregisterForUndo(this);
+	}
+}
+
 void SSkeletonAnimNotifies::OnNotifiesChanged()
 {
 	RefreshNotifiesListWithFilter();
@@ -180,9 +138,23 @@ TSharedRef<ITableRow> SSkeletonAnimNotifies::GenerateNotifyRow(TSharedPtr<FDispl
 	check( InInfo.IsValid() );
 
 	return
-		SNew( SAnimNotifyListRow, OwnerTable )
-		.Item( InInfo )
-		.NotifiesListView( SharedThis( this ) );
+		SNew( STableRow<TSharedPtr<FDisplayedAnimNotifyInfo>>, OwnerTable )
+		[
+			SNew( SVerticalBox )
+			+ SVerticalBox::Slot()
+			.AutoHeight()
+			.Padding( 0.0f, 4.0f )
+			.VAlign( VAlign_Center )
+			[
+				SAssignNew(InInfo->InlineEditableText, SInlineEditableTextBlock)
+				.Text(FText::FromName(InInfo->Name))
+				.OnVerifyTextChanged(this, &SSkeletonAnimNotifies::OnVerifyNotifyNameCommit, InInfo)
+				.OnTextCommitted(this, &SSkeletonAnimNotifies::OnNotifyNameCommitted, InInfo)
+				.IsSelected(this, &SSkeletonAnimNotifies::IsSelected)
+				.HighlightText_Lambda([this](){ return FilterText; })
+				.IsReadOnly(bIsPicker)
+			]
+		];
 }
 
 TSharedPtr<SWidget> SSkeletonAnimNotifies::OnGetContextMenuContent() const
@@ -190,25 +162,50 @@ TSharedPtr<SWidget> SSkeletonAnimNotifies::OnGetContextMenuContent() const
 	const bool bShouldCloseWindowAfterMenuSelection = true;
 	FMenuBuilder MenuBuilder( bShouldCloseWindowAfterMenuSelection, NULL);
 
-	MenuBuilder.BeginSection("AnimNotifyAction", LOCTEXT( "AnimNotifyActions", "Selected Notify Actions" ) );
+	if (bIsSyncMarker)
+	{
+		MenuBuilder.BeginSection("AnimNotifyAction", LOCTEXT("SelectedSyncMarkerActions", "Selected Sync Marker Actions"));
+		{
+			FUIAction Action = FUIAction(FExecuteAction::CreateSP(this, &SSkeletonAnimNotifies::OnDeleteSyncMarker));
+			const FText Label = LOCTEXT("DeleteSyncMarkerButtonLabel", "Delete");
+			const FText ToolTipText = LOCTEXT("DeleteSyncMarkerButtonTooltip", "Deletes the sync marker from the suggestions");
+			MenuBuilder.AddMenuEntry(Label, ToolTipText, FSlateIcon(), Action);
+		}
+		MenuBuilder.EndSection();
+	}
+	else
 	{
+		MenuBuilder.BeginSection("AnimNotifyAction", LOCTEXT("AnimNotifyActions", "Notifies"));
 		{
-			FUIAction Action = FUIAction( FExecuteAction::CreateSP( this, &SSkeletonAnimNotifies::OnDeleteAnimNotify ), 
-				FCanExecuteAction::CreateSP( this, &SSkeletonAnimNotifies::CanPerformDelete ) );
-			const FText Label = LOCTEXT("DeleteAnimNotifyButtonLabel", "Delete");
-			const FText ToolTipText = LOCTEXT("DeleteAnimNotifyButtonTooltip", "Deletes the selected anim notifies.");
-			MenuBuilder.AddMenuEntry( Label, ToolTipText, FSlateIcon(), Action);
+			FUIAction Action = FUIAction(FExecuteAction::CreateSP(this, &SSkeletonAnimNotifies::OnAddAnimNotify));
+			const FText Label = LOCTEXT("NewAnimNotifyButtonLabel", "New...");
+			const FText ToolTipText = LOCTEXT("NewAnimNotifyButtonTooltip", "Creates a new anim notify.");
+			MenuBuilder.AddMenuEntry(Label, ToolTipText, FSlateIcon(), Action);
 		}
+		MenuBuilder.EndSection();
 
+		MenuBuilder.BeginSection("AnimNotifyAction", LOCTEXT("SelectedAnimNotifyActions", "Selected Notify Actions"));
 		{
-			FUIAction Action = FUIAction( FExecuteAction::CreateSP( this, &SSkeletonAnimNotifies::OnRenameAnimNotify ), 
-				FCanExecuteAction::CreateSP( this, &SSkeletonAnimNotifies::CanPerformRename ) );
-			const FText Label = LOCTEXT("RenameAnimNotifyButtonLabel", "Rename");
-			const FText ToolTipText = LOCTEXT("RenameAnimNotifyButtonTooltip", "Renames the selected anim notifies.");
-			MenuBuilder.AddMenuEntry( Label, ToolTipText, FSlateIcon(), Action);
+			{
+				FUIAction Action = FUIAction(FExecuteAction::CreateSP(this, &SSkeletonAnimNotifies::OnRenameAnimNotify),
+					FCanExecuteAction::CreateSP(this, &SSkeletonAnimNotifies::CanPerformRename));
+				const FText Label = LOCTEXT("RenameAnimNotifyButtonLabel", "Rename");
+				const FText ToolTipText = LOCTEXT("RenameAnimNotifyButtonTooltip", "Renames the selected anim notifies.");
+				MenuBuilder.AddMenuEntry(Label, ToolTipText, FSlateIcon(), Action);
+			}
+
+			{
+				FUIAction Action = FUIAction(FExecuteAction::CreateSP(this, &SSkeletonAnimNotifies::OnDeleteAnimNotify),
+					FCanExecuteAction::CreateSP(this, &SSkeletonAnimNotifies::CanPerformDelete));
+				const FText Label = LOCTEXT("DeleteAnimNotifyButtonLabel", "Delete");
+				const FText ToolTipText = LOCTEXT("DeleteAnimNotifyButtonTooltip", "Deletes the selected anim notifies.");
+				MenuBuilder.AddMenuEntry(Label, ToolTipText, FSlateIcon(), Action);
+			}
+
+
 		}
+		MenuBuilder.EndSection();
 	}
-	MenuBuilder.EndSection();
 
 	return MenuBuilder.MakeWidget();
 }
@@ -217,20 +214,54 @@ void SSkeletonAnimNotifies::OnNotifySelectionChanged(TSharedPtr<FDisplayedAnimNo
 {
 	if(Selection.IsValid())
 	{
-		ShowNotifyInDetailsView(Selection->Name);
+		if (!bIsSyncMarker)
+		{
+			ShowNotifyInDetailsView(Selection->Name);
+		}
+
+		OnItemSelected.ExecuteIfBound(Selection->Name);
 	}
 }
 
 bool SSkeletonAnimNotifies::CanPerformDelete() const
 {
-	TArray< TSharedPtr< FDisplayedAnimNotifyInfo > > SelectedRows = NotifiesListView->GetSelectedItems();
-	return SelectedRows.Num() > 0;
+	return NotifiesListView->GetNumItemsSelected() > 0;
 }
 
 bool SSkeletonAnimNotifies::CanPerformRename() const
 {
-	TArray< TSharedPtr< FDisplayedAnimNotifyInfo > > SelectedRows = NotifiesListView->GetSelectedItems();
-	return SelectedRows.Num() == 1;
+	return NotifiesListView->GetNumItemsSelected() == 1;
+}
+
+void SSkeletonAnimNotifies::OnAddAnimNotify()
+{
+	// Find a unique name for this notify
+	const TCHAR* BaseNotifyString = TEXT("NewNotify");
+	FString NewNotifyString = BaseNotifyString;
+	int32 NumericExtension = 0;
+
+	while(EditableSkeleton->GetSkeleton().AnimationNotifies.ContainsByPredicate([&NewNotifyString](const FName& InNotifyName){ return InNotifyName.ToString() == NewNotifyString; }))
+	{
+		NewNotifyString = FString::Printf(TEXT("%s_%d"), BaseNotifyString, NumericExtension);
+		NumericExtension++;
+	}
+
+	// Add an item. The subsequent rename will commit the item.
+	TSharedPtr<FDisplayedAnimNotifyInfo> NewItem = FDisplayedAnimNotifyInfo::Make(*NewNotifyString);
+	NewItem->bIsNew = true;
+	NotifyList.Add(NewItem);
+
+	NotifiesListView->ClearSelection();
+	NotifiesListView->RequestListRefresh();
+	NotifiesListView->RequestScrollIntoView(NewItem);
+}
+
+void SSkeletonAnimNotifies::OnItemScrolledIntoView(TSharedPtr<FDisplayedAnimNotifyInfo> InItem, const TSharedPtr<ITableRow>& InTableRow)
+{
+	if(InItem.IsValid() && InItem->InlineEditableText.IsValid() && InItem->bIsNew)
+	{
+		InItem->InlineEditableText->EnterEditingMode();
+	}
 }
 
 void SSkeletonAnimNotifies::OnDeleteAnimNotify()
@@ -263,6 +294,22 @@ void SSkeletonAnimNotifies::OnDeleteAnimNotify()
 	CreateNotifiesList( NameFilterBox->GetText().ToString() );
 }
 
+void SSkeletonAnimNotifies::OnDeleteSyncMarker()
+{
+	TArray< TSharedPtr< FDisplayedAnimNotifyInfo > > SelectedRows = NotifiesListView->GetSelectedItems();
+
+	TArray<FName> SelectedSyncMarkerNames;
+
+	for (int Selection = 0; Selection < SelectedRows.Num(); ++Selection)
+	{
+		SelectedSyncMarkerNames.Add(SelectedRows[Selection]->Name);
+	}
+
+	EditableSkeleton->DeleteSyncMarkers(SelectedSyncMarkerNames);
+
+	CreateNotifiesList(NameFilterBox->GetText().ToString());
+}
+
 void SSkeletonAnimNotifies::OnRenameAnimNotify()
 {
 	TArray< TSharedPtr< FDisplayedAnimNotifyInfo > > SelectedRows = NotifiesListView->GetSelectedItems();
@@ -283,7 +330,7 @@ bool SSkeletonAnimNotifies::OnVerifyNotifyNameCommit( const FText& NewName, FTex
 	}
 
 	FName NotifyName( *NewName.ToString() );
-	if(NotifyName != Item->Name)
+	if(NotifyName != Item->Name || Item->bIsNew)
 	{
 		if(EditableSkeleton->GetSkeleton().AnimationNotifies.Contains(NotifyName))
 		{
@@ -297,22 +344,34 @@ bool SSkeletonAnimNotifies::OnVerifyNotifyNameCommit( const FText& NewName, FTex
 
 void SSkeletonAnimNotifies::OnNotifyNameCommitted( const FText& NewName, ETextCommit::Type, TSharedPtr<FDisplayedAnimNotifyInfo> Item )
 {
-	int32 NumAnimationsModified = EditableSkeleton->RenameNotify(FName(*NewName.ToString()), Item->Name);
-
-	if(NumAnimationsModified > 0)
+	FName NewFName = FName(*NewName.ToString());
+	if(Item->bIsNew)
 	{
-		// Tell the user that the socket is a duplicate
-		FFormatNamedArguments Args;
-		Args.Add( TEXT("NumAnimationsModified"), NumAnimationsModified );
-		FNotificationInfo Info( FText::Format( LOCTEXT( "AnimNotifiesRenamed", "{NumAnimationsModified} animation(s) modified to rename notification" ), Args ) );
+		EditableSkeleton->AddNotify(NewFName);
+		Item->bIsNew = false;
+	}
+	else
+	{
+		if(NewFName != Item->Name)
+		{
+			int32 NumAnimationsModified = EditableSkeleton->RenameNotify(FName(*NewName.ToString()), Item->Name);
 
-		Info.bUseLargeFont = false;
-		Info.ExpireDuration = 5.0f;
+			if(NumAnimationsModified > 0)
+			{
+				// Tell the user that the socket is a duplicate
+				FFormatNamedArguments Args;
+				Args.Add( TEXT("NumAnimationsModified"), NumAnimationsModified );
+				FNotificationInfo Info( FText::Format( LOCTEXT( "AnimNotifiesRenamed", "{NumAnimationsModified} animation(s) modified to rename notification" ), Args ) );
 
-		NotifyUser( Info );
-	}
+				Info.bUseLargeFont = false;
+				Info.ExpireDuration = 5.0f;
 
-	RefreshNotifiesListWithFilter();
+				NotifyUser( Info );
+			}
+
+			RefreshNotifiesListWithFilter();
+		}
+	}
 }
 
 void SSkeletonAnimNotifies::RefreshNotifiesListWithFilter()
@@ -324,19 +383,22 @@ void SSkeletonAnimNotifies::CreateNotifiesList( const FString& SearchText )
 {
 	NotifyList.Empty();
 
-	for(int i = 0; i < EditableSkeleton->GetSkeleton().AnimationNotifies.Num(); ++i)
+	const USkeleton& TargetSkeleton = EditableSkeleton->GetSkeleton();
+
+	const TArray<FName>& ItemNames = bIsSyncMarker ? TargetSkeleton.GetExistingMarkerNames() : TargetSkeleton.AnimationNotifies;
+
+	for(const FName& ItemName : ItemNames)
 	{
-		const FName& NotifyName = EditableSkeleton->GetSkeleton().AnimationNotifies[i];
 		if ( !SearchText.IsEmpty() )
 		{
-			if ( NotifyName.ToString().Contains( SearchText ) )
+			if (ItemName.ToString().Contains( SearchText ) )
 			{
-				NotifyList.Add( FDisplayedAnimNotifyInfo::Make( NotifyName ) );
+				NotifyList.Add( FDisplayedAnimNotifyInfo::Make(ItemName) );
 			}
 		}
 		else
 		{
-			NotifyList.Add( FDisplayedAnimNotifyInfo::Make( NotifyName ) );
+			NotifyList.Add( FDisplayedAnimNotifyInfo::Make(ItemName) );
 		}
 	}
 
@@ -345,33 +407,16 @@ void SSkeletonAnimNotifies::CreateNotifiesList( const FString& SearchText )
 
 void SSkeletonAnimNotifies::ShowNotifyInDetailsView(FName NotifyName)
 {
-	UEditorSkeletonNotifyObj *Obj = Cast<UEditorSkeletonNotifyObj>(ShowInDetailsView(UEditorSkeletonNotifyObj::StaticClass()));
-	if(Obj != NULL)
+	if(OnObjectsSelected.IsBound())
 	{
-		Obj->AnimationNames.Empty();
-
-		TArray<FAssetData> CompatibleAnimSequences;
-		EditableSkeleton->GetCompatibleAnimSequences(CompatibleAnimSequences);
+		ClearDetailsView();
 
-		for( int32 AssetIndex = 0; AssetIndex < CompatibleAnimSequences.Num(); ++AssetIndex )
+		UEditorSkeletonNotifyObj *Obj = Cast<UEditorSkeletonNotifyObj>(ShowInDetailsView(UEditorSkeletonNotifyObj::StaticClass()));
+		if(Obj != NULL)
 		{
-			const FAssetData& PossibleAnimSequence = CompatibleAnimSequences[AssetIndex];
-			if(UObject* AnimSeqAsset = PossibleAnimSequence.GetAsset())
-			{
-				UAnimSequenceBase* Sequence = CastChecked<UAnimSequenceBase>(AnimSeqAsset);
-				for (int32 NotifyIndex = 0; NotifyIndex < Sequence->Notifies.Num(); ++NotifyIndex)
-				{
-					FAnimNotifyEvent& NotifyEvent = Sequence->Notifies[NotifyIndex];
-					if (NotifyEvent.NotifyName == NotifyName)
-					{
-						Obj->AnimationNames.Add(MakeShareable(new FString(PossibleAnimSequence.AssetName.ToString())));
-						break;
-					}
-				}
-			}
+			Obj->EditableSkeleton = EditableSkeleton;
+			Obj->Name = NotifyName;
 		}
-
-		Obj->Name = NotifyName;
 	}
 }
 
@@ -394,7 +439,12 @@ void SSkeletonAnimNotifies::ClearDetailsView()
 	OnObjectsSelected.ExecuteIfBound(Objects);
 }
 
-void SSkeletonAnimNotifies::PostUndo()
+void SSkeletonAnimNotifies::PostUndo( bool bSuccess )
+{
+	RefreshNotifiesListWithFilter();
+}
+
+void SSkeletonAnimNotifies::PostRedo( bool bSuccess )
 {
 	RefreshNotifiesListWithFilter();
 }
diff --git a/Engine/Source/Editor/Persona/Private/SSkeletonAnimNotifies.h b/Engine/Source/Editor/Persona/Private/SSkeletonAnimNotifies.h
--- a/Engine/Source/Editor/Persona/Private/SSkeletonAnimNotifies.h
+++ b/Engine/Source/Editor/Persona/Private/SSkeletonAnimNotifies.h
@@ -15,6 +15,7 @@
 #include "Widgets/Views/STableRow.h"
 #include "Widgets/Views/SListView.h"
 #include "EditorObjectsTracker.h"
+#include "EditorUndoClient.h"
 
 #define LOCTEXT_NAMESPACE "SkeletonAnimnotifies"
 
@@ -22,12 +23,15 @@ class IEditableSkeleton;
 class SToolTip;
 struct FNotificationInfo;
 
+/** Delegate fired when a notify is selected */
+DECLARE_DELEGATE_OneParam(FOnItemSelected, const FName& /*InSelectedNotify*/);
+
 /////////////////////////////////////////////////////
 // FSkeletonAnimNotifiesSummoner
 struct FSkeletonAnimNotifiesSummoner : public FWorkflowTabFactory
 {
 public:
-	FSkeletonAnimNotifiesSummoner(TSharedPtr<class FAssetEditorToolkit> InHostingApp, const TSharedRef<class IEditableSkeleton>& InEditableSkeleton, FSimpleMulticastDelegate& InOnChangeAnimNotifies, FSimpleMulticastDelegate& InOnPostUndo, FOnObjectsSelected InOnObjectsSelected);
+	FSkeletonAnimNotifiesSummoner(TSharedPtr<class FAssetEditorToolkit> InHostingApp, const TSharedRef<class IEditableSkeleton>& InEditableSkeleton, FOnObjectsSelected InOnObjectsSelected);
 
 	virtual TSharedRef<SWidget> CreateTabBody(const FWorkflowTabSpawnInfo& Info) const override;
 
@@ -39,8 +43,6 @@ public:
 
 private:
 	TWeakPtr<class IEditableSkeleton> EditableSkeleton;
-	FSimpleMulticastDelegate& OnChangeAnimNotifies;
-	FSimpleMulticastDelegate& OnPostUndo;
 	FOnObjectsSelected OnObjectsSelected;
 };
 
@@ -55,6 +57,9 @@ public:
 	/** Handle to editable text block for rename */
 	TSharedPtr<SInlineEditableTextBlock> InlineEditableText;
 
+	/** Flag to say whether this is a new item we are creating */
+	bool bIsNew;
+
 	/** Static function for creating a new item, but ensures that you can only have a TSharedRef to one */
 	static TSharedRef<FDisplayedAnimNotifyInfo> Make(const FName& NotifyName)
 	{
@@ -65,6 +70,7 @@ protected:
 	/** Hidden constructor, always use Make above */
 	FDisplayedAnimNotifyInfo(const FName& InNotifyName)
 		: Name( InNotifyName )
+		, bIsNew(false)
 	{}
 
 	/** Hidden constructor, always use Make above */
@@ -74,17 +80,31 @@ protected:
 /** Widgets list type */
 typedef SListView< TSharedPtr<FDisplayedAnimNotifyInfo> > SAnimNotifyListType;
 
-class SSkeletonAnimNotifies : public SCompoundWidget, public FGCObject
+class SSkeletonAnimNotifies : public SCompoundWidget, public FGCObject, public FEditorUndoClient
 {
 public:
 	SLATE_BEGIN_ARGS( SSkeletonAnimNotifies )
+		: _IsPicker(false)
+		, _IsSyncMarker(false)
 	{}
 
+	/** Delegate called to select an object in the details panel */
 	SLATE_EVENT(FOnObjectsSelected, OnObjectsSelected)
-		
+
+	/** Delegate fired when a notify is selected */
+	SLATE_EVENT(FOnItemSelected, OnItemSelected)
+
+	/** Whether we should use this dialog as a picker or an editor. In picker mode we cant add, remove or rename notifies. */
+	SLATE_ARGUMENT(bool, IsPicker)
+
+	/** Whether we should use this dialog for sync markers instead of notifies */
+	SLATE_ARGUMENT(bool, IsSyncMarker)
+
 	SLATE_END_ARGS()
 public:
-	void Construct(const FArguments& InArgs, const TSharedRef<IEditableSkeleton>& InEditableSkeleton, FSimpleMulticastDelegate& InOnChangeAnimNotifies, FSimpleMulticastDelegate& InOnPostUndo);
+	void Construct(const FArguments& InArgs, const TSharedRef<IEditableSkeleton>& InEditableSkeleton);
+
+	~SSkeletonAnimNotifies();
 
 	/**
 	* Accessor so our rows can grab the filter text for highlighting
@@ -98,8 +118,9 @@ public:
 	/** Clears the detail view of whatever we displayed last */
 	void ClearDetailsView();
 
-	/** This triggers a UI repopulation after undo has been called */
-	void PostUndo();
+	/** FEditorUndoClient interface */
+	virtual void PostUndo( bool bSuccess ) override;
+	virtual void PostRedo( bool bSuccess ) override;
 
 	// FGCObject interface start
 	virtual void AddReferencedObjects( FReferenceCollector& Collector ) override;
@@ -136,9 +157,15 @@ private:
 	/** Delegate handler for deleting anim notifies */
 	void OnDeleteAnimNotify();
 
+	/** Delegate handler for deleting a sync marker */
+	void OnDeleteSyncMarker();
+
 	/** Delegate handler for determining whether we can show the rename menu options */
 	bool CanPerformRename() const;
 
+	/** Delegate handler for adding anim notifies */
+	void OnAddAnimNotify();
+
 	/** Delegate handler for renaming anim notifies */
 	void OnRenameAnimNotify();
 
@@ -157,6 +184,9 @@ private:
 	/** Handler function for when notifies are modified on the skeleton */
 	void OnNotifiesChanged();
 
+	/** Handle when an item is scrolled into view, triggers a rename for new items */
+	void OnItemScrolledIntoView(TSharedPtr<FDisplayedAnimNotifyInfo> InItem, const TSharedPtr<ITableRow>& InTableRow);
+
 	/** The skeleton we are currently editing */
 	TSharedPtr<class IEditableSkeleton> EditableSkeleton;
 
@@ -177,6 +207,15 @@ private:
 
 	/** Delegate called to select an object in the details panel */
 	FOnObjectsSelected OnObjectsSelected;
+
+	/** Delegate fired when a notify is selected */
+	FOnItemSelected OnItemSelected;
+
+	/** Whether we should use this dialog as a picker or an editor. In picker mode we cant add, remove or rename notifies. */
+	bool bIsPicker;
+
+	/** Whether we are using this dialog for sync markers instead of notifies */
+	bool bIsSyncMarker;
 };
 
 #undef LOCTEXT_NAMESPACE
diff --git a/Engine/Source/Editor/Persona/Public/PersonaModule.h b/Engine/Source/Editor/Persona/Public/PersonaModule.h
--- a/Engine/Source/Editor/Persona/Public/PersonaModule.h
+++ b/Engine/Source/Editor/Persona/Public/PersonaModule.h
@@ -78,12 +78,11 @@ struct FPersonaToolkitArgs
 
 struct FAnimDocumentArgs
 {
-	FAnimDocumentArgs(const TSharedRef<class IPersonaPreviewScene>& InPreviewScene, const TSharedRef<class IPersonaToolkit>& InPersonaToolkit, const TSharedRef<class IEditableSkeleton>& InEditableSkeleton, FSimpleMulticastDelegate& InOnPostUndo, FSimpleMulticastDelegate& InOnAnimNotifiesChanged, FSimpleMulticastDelegate& InOnSectionsChanged)
+	FAnimDocumentArgs(const TSharedRef<class IPersonaPreviewScene>& InPreviewScene, const TSharedRef<class IPersonaToolkit>& InPersonaToolkit, const TSharedRef<class IEditableSkeleton>& InEditableSkeleton, FSimpleMulticastDelegate& InOnPostUndo, FSimpleMulticastDelegate& InOnSectionsChanged)
 		: PreviewScene(InPreviewScene)
 		, PersonaToolkit(InPersonaToolkit)
 		, EditableSkeleton(InEditableSkeleton)
 		, OnPostUndo(InOnPostUndo)
-		, OnAnimNotifiesChanged(InOnAnimNotifiesChanged)
 		, OnSectionsChanged(InOnSectionsChanged)
 	{}
 
@@ -92,14 +91,12 @@ struct FAnimDocumentArgs
 	TWeakPtr<class IPersonaToolkit> PersonaToolkit;
 	TWeakPtr<class IEditableSkeleton> EditableSkeleton;
 	FSimpleMulticastDelegate& OnPostUndo;
-	FSimpleMulticastDelegate& OnAnimNotifiesChanged;
 	FSimpleMulticastDelegate& OnSectionsChanged;
 
 	/** Optional args */
 	FOnObjectsSelected OnDespatchObjectsSelected;
 	FOnInvokeTab OnDespatchInvokeTab;
 	FSimpleDelegate OnDespatchSectionsChanged;
-	FSimpleDelegate OnDespatchAnimNotifiesChanged;
 };
 
 /** Places that viewport text can be placed */
@@ -218,7 +215,7 @@ public:
 	virtual void RegisterPersonaViewportTabFactories(FWorkflowAllowedTabSet& TabSet, const TSharedRef<class FWorkflowCentricApplication>& InHostingApp, const FPersonaViewportArgs& InArgs) const;
 
 	/** Create an anim notifies tab factory */
-	virtual TSharedRef<FWorkflowTabFactory> CreateAnimNotifiesTabFactory(const TSharedRef<class FWorkflowCentricApplication>& InHostingApp, const TSharedRef<class IEditableSkeleton>& InEditableSkeleton, FSimpleMulticastDelegate& InOnChangeAnimNotifies, FSimpleMulticastDelegate& InOnPostUndo, FOnObjectsSelected InOnObjectsSelected) const;
+	virtual TSharedRef<FWorkflowTabFactory> CreateAnimNotifiesTabFactory(const TSharedRef<class FWorkflowCentricApplication>& InHostingApp, const TSharedRef<class IEditableSkeleton>& InEditableSkeleton, FOnObjectsSelected InOnObjectsSelected) const;
 
 	/** Create a skeleton cuve viewer tab factory */
 	virtual TSharedRef<FWorkflowTabFactory> CreateCurveViewerTabFactory(const TSharedRef<class FWorkflowCentricApplication>& InHostingApp, const TSharedRef<class IEditableSkeleton>& InEditableSkeleton, const TSharedRef<class IPersonaPreviewScene>& InPreviewScene, FSimpleMulticastDelegate& InOnPostUndo, FOnObjectsSelected InOnObjectsSelected) const;
diff --git a/Engine/Source/Editor/PhysicsAssetEditor/Private/PhysicsAssetEditorEditMode.cpp b/Engine/Source/Editor/PhysicsAssetEditor/Private/PhysicsAssetEditorEditMode.cpp
--- a/Engine/Source/Editor/PhysicsAssetEditor/Private/PhysicsAssetEditorEditMode.cpp
+++ b/Engine/Source/Editor/PhysicsAssetEditor/Private/PhysicsAssetEditorEditMode.cpp
@@ -604,12 +604,14 @@ bool FPhysicsAssetEditorEditMode::GetCustomDrawingCoordinateSystem(FMatrix& InMa
 	if (SharedData->GetSelectedBody())
 	{
 		int32 BoneIndex = SharedData->EditorSkelComp->GetBoneIndex(SharedData->PhysicsAsset->SkeletalBodySetups[SharedData->GetSelectedBody()->Index]->BoneName);
+		if(BoneIndex != INDEX_NONE)
+		{
+			FTransform BoneTM = SharedData->EditorSkelComp->GetBoneTransform(BoneIndex);
+			BoneTM.RemoveScaling();
 
-		FTransform BoneTM = SharedData->EditorSkelComp->GetBoneTransform(BoneIndex);
-		BoneTM.RemoveScaling();
-
-		InMatrix = SharedData->EditorSkelComp->GetPrimitiveTransform(BoneTM, SharedData->GetSelectedBody()->Index, SharedData->GetSelectedBody()->PrimitiveType, SharedData->GetSelectedBody()->PrimitiveIndex, 1.f).ToMatrixNoScale().RemoveTranslation();
-		return true;
+			InMatrix = SharedData->EditorSkelComp->GetPrimitiveTransform(BoneTM, SharedData->GetSelectedBody()->Index, SharedData->GetSelectedBody()->PrimitiveType, SharedData->GetSelectedBody()->PrimitiveIndex, 1.f).ToMatrixNoScale().RemoveTranslation();
+			return true;
+		}
 	}
 	else if (SharedData->GetSelectedConstraint())
 	{
@@ -631,12 +633,14 @@ FVector FPhysicsAssetEditorEditMode::GetWidgetLocation() const
 	if (SharedData->GetSelectedBody())
 	{
 		int32 BoneIndex = SharedData->EditorSkelComp->GetBoneIndex(SharedData->PhysicsAsset->SkeletalBodySetups[SharedData->GetSelectedBody()->Index]->BoneName);
+		if(BoneIndex != INDEX_NONE)
+		{
+			FTransform BoneTM = SharedData->EditorSkelComp->GetBoneTransform(BoneIndex);
+			const float Scale = BoneTM.GetScale3D().GetAbsMax();
+			BoneTM.RemoveScaling();
 
-		FTransform BoneTM = SharedData->EditorSkelComp->GetBoneTransform(BoneIndex);
-		const float Scale = BoneTM.GetScale3D().GetAbsMax();
-		BoneTM.RemoveScaling();
-
-		return SharedData->EditorSkelComp->GetPrimitiveTransform(BoneTM, SharedData->GetSelectedBody()->Index, SharedData->GetSelectedBody()->PrimitiveType, SharedData->GetSelectedBody()->PrimitiveIndex, Scale).GetTranslation();
+			return SharedData->EditorSkelComp->GetPrimitiveTransform(BoneTM, SharedData->GetSelectedBody()->Index, SharedData->GetSelectedBody()->PrimitiveType, SharedData->GetSelectedBody()->PrimitiveIndex, Scale).GetTranslation();
+		}
 	}
 	else if (SharedData->GetSelectedConstraint())
 	{
diff --git a/Engine/Source/Editor/SequenceRecorder/Private/ActorRecording.cpp b/Engine/Source/Editor/SequenceRecorder/Private/ActorRecording.cpp
--- a/Engine/Source/Editor/SequenceRecorder/Private/ActorRecording.cpp
+++ b/Engine/Source/Editor/SequenceRecorder/Private/ActorRecording.cpp
@@ -483,7 +483,7 @@ void UActorRecording::StartRecordingActorProperties(ULevelSequence* CurrentSeque
 						{
 							SkeletalMeshComponent->SetAnimationMode(EAnimationMode::AnimationSingleNode);
 							SkeletalMeshComponent->bEnableUpdateRateOptimizations = false;
-							SkeletalMeshComponent->MeshComponentUpdateFlag = EMeshComponentUpdateFlag::AlwaysTickPoseAndRefreshBones;
+							SkeletalMeshComponent->VisibilityBasedAnimTickOption = EVisibilityBasedAnimTickOption::AlwaysTickPoseAndRefreshBones;
 							SkeletalMeshComponent->ForcedLodModel = 1;
 						}
 					}
diff --git a/Engine/Source/Editor/SequenceRecorder/Private/AnimationRecorder.cpp b/Engine/Source/Editor/SequenceRecorder/Private/AnimationRecorder.cpp
--- a/Engine/Source/Editor/SequenceRecorder/Private/AnimationRecorder.cpp
+++ b/Engine/Source/Editor/SequenceRecorder/Private/AnimationRecorder.cpp
@@ -774,7 +774,7 @@ FAnimationRecorderManager& FAnimationRecorderManager::Get()
 FAnimRecorderInstance::FAnimRecorderInstance()
 	: SkelComp(nullptr)
 	, Recorder(nullptr)
-	, CachedMeshComponentUpdateFlag(EMeshComponentUpdateFlag::AlwaysTickPoseAndRefreshBones)
+	, CachedVisibilityBasedAnimTickOption(EVisibilityBasedAnimTickOption::AlwaysTickPoseAndRefreshBones)
 	, bCachedEnableUpdateRateOptimizations(false)
 {
 }
@@ -813,10 +813,10 @@ void FAnimRecorderInstance::InitInternal(USkeletalMeshComponent* InComponent, co
 
 		// turn off URO and make sure we always update even if out of view
 		bCachedEnableUpdateRateOptimizations = InComponent->bEnableUpdateRateOptimizations;
-		CachedMeshComponentUpdateFlag = InComponent->MeshComponentUpdateFlag;
+		CachedVisibilityBasedAnimTickOption = InComponent->VisibilityBasedAnimTickOption;
 
 		InComponent->bEnableUpdateRateOptimizations = false;
-		InComponent->MeshComponentUpdateFlag = EMeshComponentUpdateFlag::AlwaysTickPoseAndRefreshBones;
+		InComponent->VisibilityBasedAnimTickOption = EVisibilityBasedAnimTickOption::AlwaysTickPoseAndRefreshBones;
 	}
 }
 
@@ -865,7 +865,7 @@ void FAnimRecorderInstance::FinishRecording(bool bShowMessage)
 
 		// restore update flags
 		SkelComp->bEnableUpdateRateOptimizations = bCachedEnableUpdateRateOptimizations;
-		SkelComp->MeshComponentUpdateFlag = CachedMeshComponentUpdateFlag;
+		SkelComp->VisibilityBasedAnimTickOption = CachedVisibilityBasedAnimTickOption;
 	}
 }
 
diff --git a/Engine/Source/Editor/SequenceRecorder/Private/AnimationRecorder.h b/Engine/Source/Editor/SequenceRecorder/Private/AnimationRecorder.h
--- a/Engine/Source/Editor/SequenceRecorder/Private/AnimationRecorder.h
+++ b/Engine/Source/Editor/SequenceRecorder/Private/AnimationRecorder.h
@@ -127,7 +127,7 @@ public:
 	TSharedPtr<FAnimationRecorder> Recorder;
 
 	/** Used to store/restore update flag when recording */
-	EMeshComponentUpdateFlag::Type CachedMeshComponentUpdateFlag;
+	EVisibilityBasedAnimTickOption CachedVisibilityBasedAnimTickOption;
 
 	/** Used to store/restore URO when recording */
 	bool bCachedEnableUpdateRateOptimizations;
diff --git a/Engine/Source/Editor/SequenceRecorder/Private/Sections/MovieSceneAnimationSectionRecorder.cpp b/Engine/Source/Editor/SequenceRecorder/Private/Sections/MovieSceneAnimationSectionRecorder.cpp
--- a/Engine/Source/Editor/SequenceRecorder/Private/Sections/MovieSceneAnimationSectionRecorder.cpp
+++ b/Engine/Source/Editor/SequenceRecorder/Private/Sections/MovieSceneAnimationSectionRecorder.cpp
@@ -170,6 +170,6 @@ void FMovieSceneAnimationSectionRecorder::Record(float CurrentTime)
 	{
 		// re-force updates on as gameplay can sometimes turn these back off!
 		SkeletalMeshComponent->bEnableUpdateRateOptimizations = false;
-		SkeletalMeshComponent->MeshComponentUpdateFlag = EMeshComponentUpdateFlag::AlwaysTickPoseAndRefreshBones;
+		SkeletalMeshComponent->VisibilityBasedAnimTickOption = EVisibilityBasedAnimTickOption::AlwaysTickPoseAndRefreshBones;
 	}
 }
diff --git a/Engine/Source/Editor/SkeletalMeshEditor/Private/SkeletalMeshEditor.h b/Engine/Source/Editor/SkeletalMeshEditor/Private/SkeletalMeshEditor.h
--- a/Engine/Source/Editor/SkeletalMeshEditor/Private/SkeletalMeshEditor.h
+++ b/Engine/Source/Editor/SkeletalMeshEditor/Private/SkeletalMeshEditor.h
@@ -137,9 +137,6 @@ private:
 	void BindCommands();
 
 public:
-	/** Multicast delegate fired on anim notifies changing */
-	FSimpleMulticastDelegate OnChangeAnimNotifies;
-
 	/** Multicast delegate fired on global undo/redo */
 	FSimpleMulticastDelegate OnPostUndo;
 
diff --git a/Engine/Source/Editor/SkeletonEditor/Private/EditableSkeleton.cpp b/Engine/Source/Editor/SkeletonEditor/Private/EditableSkeleton.cpp
--- a/Engine/Source/Editor/SkeletonEditor/Private/EditableSkeleton.cpp
+++ b/Engine/Source/Editor/SkeletonEditor/Private/EditableSkeleton.cpp
@@ -1095,6 +1095,18 @@ int32 FEditableSkeleton::DeleteAnimNotifies(const TArray<FName>& InNotifyNames)
 	return NumAnimationsModified;
 }
 
+void FEditableSkeleton::DeleteSyncMarkers(const TArray<FName>& InSyncMarkerNames)
+{
+	const FScopedTransaction Transaction(LOCTEXT("DeleteSyncMarkers", "Delete Sync Markers"));
+
+	Skeleton->Modify();
+
+	for (FName Marker : InSyncMarkerNames)
+	{
+		Skeleton->RemoveMarkerName(Marker);
+	}
+}
+
 void FEditableSkeleton::AddNotify(FName NewName)
 {
 	const FScopedTransaction Transaction(LOCTEXT("AddNewNotifyToSkeleton", "Add New Anim Notify To Skeleton"));
@@ -1103,6 +1115,13 @@ void FEditableSkeleton::AddNotify(FName NewName)
 	OnNotifiesChanged.Broadcast();
 }
 
+void FEditableSkeleton::AddSyncMarker(FName NewName)
+{
+	const FScopedTransaction Transaction(LOCTEXT("AddNewSyncMarkerToSkeleton", "Add New Sync Marker To Skeleton"));
+	Skeleton->Modify();
+	Skeleton->RegisterMarkerName(NewName);
+}
+
 int32 FEditableSkeleton::RenameNotify(const FName& NewName, const FName& OldName)
 {
 	const FScopedTransaction Transaction(LOCTEXT("RenameAnimNotify", "Rename Anim Notify"));
@@ -1151,6 +1170,11 @@ int32 FEditableSkeleton::RenameNotify(const FName& NewName, const FName& OldName
 	return NumAnimationsModified;
 }
 
+void FEditableSkeleton::BroadcastNotifyChanged()
+{
+	OnNotifiesChanged.Broadcast();
+}
+
 void FEditableSkeleton::GetCompatibleAnimSequences(TArray<struct FAssetData>& OutAssets)
 {
 	//Get the skeleton tag to search for
diff --git a/Engine/Source/Editor/SkeletonEditor/Private/EditableSkeleton.h b/Engine/Source/Editor/SkeletonEditor/Private/EditableSkeleton.h
--- a/Engine/Source/Editor/SkeletonEditor/Private/EditableSkeleton.h
+++ b/Engine/Source/Editor/SkeletonEditor/Private/EditableSkeleton.h
@@ -37,8 +37,11 @@ public:
 	virtual USkeletalMeshSocket* DuplicateSocket(const FSelectedSocketInfo& SocketInfoToDuplicate, const FName& NewParentBoneName, USkeletalMesh* InSkeletalMesh) override;
 	virtual int32 ValidatePreviewAttachedObjects() override;
 	virtual int32 DeleteAnimNotifies(const TArray<FName>& InSelectedNotifyNames) override;
+	virtual void DeleteSyncMarkers(const TArray<FName>& ISyncMarkerNames) override;
 	virtual void AddNotify(FName NewName) override;
+	virtual void AddSyncMarker(FName NewName) override;
 	virtual int32 RenameNotify(const FName& NewName, const FName& OldName) override;
+	virtual void BroadcastNotifyChanged() override;
 	virtual void GetCompatibleAnimSequences(TArray<struct FAssetData>& OutAssets) override;
 	virtual void RenameSocket(const FName& OldSocketName, const FName& NewSocketName, USkeletalMesh* InSkeletalMesh) override;
 	virtual void SetSocketParent(const FName& SocketName, const FName& NewParentName, USkeletalMesh* InSkeletalMesh) override;
diff --git a/Engine/Source/Editor/SkeletonEditor/Private/SkeletonEditor.h b/Engine/Source/Editor/SkeletonEditor/Private/SkeletonEditor.h
--- a/Engine/Source/Editor/SkeletonEditor/Private/SkeletonEditor.h
+++ b/Engine/Source/Editor/SkeletonEditor/Private/SkeletonEditor.h
@@ -106,9 +106,6 @@ private:
 	void OnImportAsset();
 
 public:
-	/** Multicast delegate fired on anim notifies changing */
-	FSimpleMulticastDelegate OnChangeAnimNotifies;
-
 	/** Multicast delegate fired on global undo/redo */
 	FSimpleMulticastDelegate OnPostUndo;
 
diff --git a/Engine/Source/Editor/SkeletonEditor/Private/SkeletonEditorMode.cpp b/Engine/Source/Editor/SkeletonEditor/Private/SkeletonEditorMode.cpp
--- a/Engine/Source/Editor/SkeletonEditor/Private/SkeletonEditorMode.cpp
+++ b/Engine/Source/Editor/SkeletonEditor/Private/SkeletonEditorMode.cpp
@@ -28,7 +28,7 @@ FSkeletonEditorMode::FSkeletonEditorMode(TSharedRef<FWorkflowCentricApplication>
 
 	PersonaModule.RegisterPersonaViewportTabFactories(TabFactories, InHostingApp, ViewportArgs);
 
-	TabFactories.RegisterFactory(PersonaModule.CreateAnimNotifiesTabFactory(InHostingApp, InSkeletonTree->GetEditableSkeleton(), SkeletonEditor->OnChangeAnimNotifies, SkeletonEditor->OnPostUndo, OnObjectsSelected));
+	TabFactories.RegisterFactory(PersonaModule.CreateAnimNotifiesTabFactory(InHostingApp, InSkeletonTree->GetEditableSkeleton(),  OnObjectsSelected));
 	TabFactories.RegisterFactory(PersonaModule.CreateAdvancedPreviewSceneTabFactory(InHostingApp, SkeletonEditor->GetPersonaToolkit()->GetPreviewScene()));
 	TabFactories.RegisterFactory(PersonaModule.CreateRetargetManagerTabFactory(InHostingApp, InSkeletonTree->GetEditableSkeleton(), SkeletonEditor->GetPersonaToolkit()->GetPreviewScene(), SkeletonEditor->OnPostUndo));
 	TabFactories.RegisterFactory(PersonaModule.CreateCurveViewerTabFactory(InHostingApp, InSkeletonTree->GetEditableSkeleton(), SkeletonEditor->GetPersonaToolkit()->GetPreviewScene(), SkeletonEditor->OnPostUndo, OnObjectsSelected));
diff --git a/Engine/Source/Editor/SkeletonEditor/Public/IEditableSkeleton.h b/Engine/Source/Editor/SkeletonEditor/Public/IEditableSkeleton.h
--- a/Engine/Source/Editor/SkeletonEditor/Public/IEditableSkeleton.h
+++ b/Engine/Source/Editor/SkeletonEditor/Public/IEditableSkeleton.h
@@ -92,6 +92,11 @@ public:
 	 */	
 	virtual int32 DeleteAnimNotifies(const TArray<FName>& InotifyNames) = 0;
 
+	/**
+	 * Delete sync markers from the skeleton cache by name
+	 */
+	virtual void DeleteSyncMarkers(const TArray<FName>& ISyncMarkerNames) = 0;
+
 	/**
 	* Add a notify
 	*/
@@ -103,6 +108,14 @@ public:
 	 */	
 	virtual int32 RenameNotify(const FName& NewName, const FName& OldName) = 0;
 
+	/**
+	* Add a sync marker
+	*/
+	virtual void AddSyncMarker(FName NewName) = 0;
+
+	/** Inform the system that something about a notify changed */	
+	virtual void BroadcastNotifyChanged() = 0;
+
 	/** Populates OutAssets with the AnimSequences that match this current skeleton */
 	virtual void GetCompatibleAnimSequences(TArray<struct FAssetData>& OutAssets) = 0;
 
diff --git a/Engine/Source/Editor/UnrealEd/Classes/Animation/EditorSkeletonNotifyObj.h b/Engine/Source/Editor/UnrealEd/Classes/Animation/EditorSkeletonNotifyObj.h
--- a/Engine/Source/Editor/UnrealEd/Classes/Animation/EditorSkeletonNotifyObj.h
+++ b/Engine/Source/Editor/UnrealEd/Classes/Animation/EditorSkeletonNotifyObj.h
@@ -10,7 +10,7 @@
 #include "UObject/Object.h"
 #include "EditorSkeletonNotifyObj.generated.h"
 
-class UAnimSequenceBase;
+class IEditableSkeleton;
 
 DECLARE_DELEGATE_TwoParams( FOnAnimObjectChange, class UObject*, bool)
 
@@ -19,11 +19,10 @@ class UEditorSkeletonNotifyObj : public UObject
 {
 	GENERATED_UCLASS_BODY()
 public:
-	
-	UObject* EditedObject;
+	/** Current skeleton to try to grab objects from */
+	TWeakPtr<IEditableSkeleton> EditableSkeleton;
 
+	/** The name of the notify we represent */
 	UPROPERTY(VisibleAnywhere, Category=SkeletonNotifies)
 	FName Name;
-
-	TArray< TSharedPtr<FString> > AnimationNames;
 };
diff --git a/Engine/Source/Editor/UnrealEd/Classes/Settings/EditorExperimentalSettings.h b/Engine/Source/Editor/UnrealEd/Classes/Settings/EditorExperimentalSettings.h
--- a/Engine/Source/Editor/UnrealEd/Classes/Settings/EditorExperimentalSettings.h
+++ b/Engine/Source/Editor/UnrealEd/Classes/Settings/EditorExperimentalSettings.h
@@ -113,10 +113,6 @@ public:
 	UPROPERTY(EditAnywhere, config, Category = Tools, meta = (ConfigRestartRequired = true))
 	bool bFacialAnimationImporter;
 
-	/** Allow animation blueprints to be recompiled while a PIE session is running */
-	UPROPERTY(EditAnywhere, config, Category = Tools)
-	bool bEnableLiveRecompilationOfAnimationBlueprints;
-
 	/** Enable experimental PIE preview device launch */
 	UPROPERTY(EditAnywhere, config, Category = PIE, meta = (DisplayName = "Enable mobile PIE with preview device launch options."))
 	bool bMobilePIEPreviewDeviceLaunch;
diff --git a/Engine/Source/Editor/UnrealEd/Private/SSkeletonWidget.cpp b/Engine/Source/Editor/UnrealEd/Private/SSkeletonWidget.cpp
--- a/Engine/Source/Editor/UnrealEd/Private/SSkeletonWidget.cpp
+++ b/Engine/Source/Editor/UnrealEd/Private/SSkeletonWidget.cpp
@@ -1373,7 +1373,7 @@ void SBasePoseViewport::Construct(const FArguments& InArgs)
 	SEditorViewport::Construct(SEditorViewport::FArguments());
 
 	PreviewComponent = NewObject<UDebugSkelMeshComponent>();
-	PreviewComponent->MeshComponentUpdateFlag = EMeshComponentUpdateFlag::AlwaysTickPoseAndRefreshBones;
+	PreviewComponent->VisibilityBasedAnimTickOption = EVisibilityBasedAnimTickOption::AlwaysTickPoseAndRefreshBones;
 	PreviewScene.AddComponent(PreviewComponent, FTransform::Identity);
 
 	SetSkeleton(InArgs._Skeleton);
diff --git a/Engine/Source/Runtime/AnimGraphRuntime/Private/BoneControllers/AnimNode_CCDIK.cpp b/Engine/Source/Runtime/AnimGraphRuntime/Private/BoneControllers/AnimNode_CCDIK.cpp
new file mode 100644
--- /dev/null
+++ b/Engine/Source/Runtime/AnimGraphRuntime/Private/BoneControllers/AnimNode_CCDIK.cpp
@@ -0,0 +1,318 @@
+// Copyright 1998-2016 Epic Games, Inc. All Rights Reserved.
+
+#include "BoneControllers/AnimNode_CCDIK.h"
+#include "Animation/AnimTypes.h"
+#include "AnimationRuntime.h"
+#include "DrawDebugHelpers.h"
+#include "Animation/AnimInstanceProxy.h"
+
+PRAGMA_DISABLE_OPTIMIZATION
+/////////////////////////////////////////////////////
+// AnimNode_CCDIK
+// Implementation of the CCDIK IK Algorithm
+
+FAnimNode_CCDIK::FAnimNode_CCDIK()
+	: EffectorLocation(FVector::ZeroVector)
+	, EffectorLocationSpace(BCS_ComponentSpace)
+	, Precision(1.f)
+	, MaxIterations(10)
+	, bStartFromTail(true)
+	, bEnableRotationLimit(false)
+{
+}
+
+FVector FAnimNode_CCDIK::GetCurrentLocation(FCSPose<FCompactPose>& MeshBases, const FCompactPoseBoneIndex& BoneIndex)
+{
+	return MeshBases.GetComponentSpaceTransform(BoneIndex).GetLocation();
+}
+
+FTransform FAnimNode_CCDIK::GetTargetTransform(const FTransform& InComponentTransform, FCSPose<FCompactPose>& MeshBases, FBoneSocketTarget& InTarget, EBoneControlSpace Space, const FVector& InOffset)
+{
+	FTransform OutTransform;
+	if (Space == BCS_BoneSpace)
+	{
+		OutTransform = InTarget.GetTargetTransform(InOffset, MeshBases, InComponentTransform);
+	}
+	else
+	{
+		// parent bone space still goes through this way
+		// if your target is socket, it will try find parents of joint that socket belongs to
+		OutTransform.SetLocation(InOffset);
+		FAnimationRuntime::ConvertBoneSpaceTransformToCS(InComponentTransform, MeshBases, OutTransform, InTarget.GetCompactPoseBoneIndex(), Space);
+	}
+
+	return OutTransform;
+}
+
+void FAnimNode_CCDIK::EvaluateSkeletalControl_AnyThread(FComponentSpacePoseContext& Output, TArray<FBoneTransform>& OutBoneTransforms)
+{
+	const FBoneContainer& BoneContainer = Output.Pose.GetPose().GetBoneContainer();
+
+	// Update EffectorLocation if it is based off a bone position
+	FTransform CSEffectorTransform = GetTargetTransform(Output.AnimInstanceProxy->GetComponentTransform(), Output.Pose, EffectorTarget, EffectorLocationSpace, EffectorLocation);
+	FVector const CSEffectorLocation = CSEffectorTransform.GetLocation();
+
+	// Gather all bone indices between root and tip.
+	TArray<FCompactPoseBoneIndex> BoneIndices;
+	
+	{
+		const FCompactPoseBoneIndex RootIndex = RootBone.GetCompactPoseIndex(BoneContainer);
+		FCompactPoseBoneIndex BoneIndex = TipBone.GetCompactPoseIndex(BoneContainer);
+		do
+		{
+			BoneIndices.Insert(BoneIndex, 0);
+			BoneIndex = Output.Pose.GetPose().GetParentBoneIndex(BoneIndex);
+		} while (BoneIndex != RootIndex);
+		BoneIndices.Insert(BoneIndex, 0);
+	}
+
+	// Gather transforms
+	int32 const NumTransforms = BoneIndices.Num();
+	OutBoneTransforms.AddUninitialized(NumTransforms);
+
+	// Gather chain links. These are non zero length bones.
+	TArray<CCDIKChainLink> Chain;
+	Chain.Reserve(NumTransforms);
+	// Start with Root Bone
+	{
+		const FCompactPoseBoneIndex& RootBoneIndex = BoneIndices[0];
+		const FTransform& LocalTransform = Output.Pose.GetLocalSpaceTransform(RootBoneIndex);
+		const FTransform& BoneCSTransform = Output.Pose.GetComponentSpaceTransform(RootBoneIndex);
+
+		OutBoneTransforms[0] = FBoneTransform(RootBoneIndex, BoneCSTransform);
+		Chain.Add(CCDIKChainLink(BoneCSTransform, LocalTransform, RootBoneIndex, 0));
+	}
+
+	// Go through remaining transforms
+	for (int32 TransformIndex = 1; TransformIndex < NumTransforms; TransformIndex++)
+	{
+		const FCompactPoseBoneIndex& BoneIndex = BoneIndices[TransformIndex];
+
+		const FTransform& LocalTransform = Output.Pose.GetLocalSpaceTransform(BoneIndex);
+		const FTransform& BoneCSTransform = Output.Pose.GetComponentSpaceTransform(BoneIndex);
+		FVector const BoneCSPosition = BoneCSTransform.GetLocation();
+
+		OutBoneTransforms[TransformIndex] = FBoneTransform(BoneIndex, BoneCSTransform);
+
+		// Calculate the combined length of this segment of skeleton
+		float const BoneLength = FVector::Dist(BoneCSPosition, OutBoneTransforms[TransformIndex - 1].Transform.GetLocation());
+
+		if (!FMath::IsNearlyZero(BoneLength))
+		{
+			Chain.Add(CCDIKChainLink(BoneCSTransform, LocalTransform, BoneIndex, TransformIndex));
+		}
+		else
+		{
+			// Mark this transform as a zero length child of the last link.
+			// It will inherit position and delta rotation from parent link.
+			CCDIKChainLink & ParentLink = Chain[Chain.Num() - 1];
+			ParentLink.ChildZeroLengthTransformIndices.Add(TransformIndex);
+		}
+	}
+
+	bool bBoneLocationUpdated = false;
+	int32 const NumChainLinks = BoneIndices.Num();
+
+	// iterate
+	{
+		int32 const TipBoneLinkIndex = NumChainLinks - 1;
+
+		// @todo optimize locally if no update, stop?
+		bool bLocalUpdated = false;
+		// check how far
+		const FVector TargetPos = CSEffectorLocation;
+		FVector TipPos = Chain[TipBoneLinkIndex].Transform.GetLocation();
+		float Distance = FVector::Dist(TargetPos, TipPos);
+		int32 IterationCount = 0;
+		while ((Distance > Precision) && (IterationCount++ < MaxIterations))
+		{
+			// iterate from tip to root
+			if (bStartFromTail)
+			{
+				for (int32 LinkIndex = TipBoneLinkIndex - 1; LinkIndex > 0; --LinkIndex)
+				{
+					bLocalUpdated |= UpdateChainLink(Chain, LinkIndex, TargetPos);
+				}
+			}
+			else
+			{
+				for (int32 LinkIndex = 1; LinkIndex < TipBoneLinkIndex; ++LinkIndex)
+				{
+					bLocalUpdated |= UpdateChainLink(Chain, LinkIndex, TargetPos);
+				}
+			}
+
+			Distance = FVector::Dist(Chain[TipBoneLinkIndex].Transform.GetLocation(), CSEffectorLocation);
+
+			bBoneLocationUpdated |= bLocalUpdated;
+
+			// no more update in this iteration
+			if (!bLocalUpdated)
+			{
+				break;
+			}
+		}
+	}
+
+	// If we moved some bones, update bone transforms.
+	if (bBoneLocationUpdated)
+	{
+		// First step: update bone transform positions from chain links.
+		for (int32 LinkIndex = 0; LinkIndex < NumChainLinks; LinkIndex++)
+		{
+			CCDIKChainLink const & ChainLink = Chain[LinkIndex];
+			OutBoneTransforms[ChainLink.TransformIndex].Transform = ChainLink.Transform;
+
+			// If there are any zero length children, update position of those
+			int32 const NumChildren = ChainLink.ChildZeroLengthTransformIndices.Num();
+			for (int32 ChildIndex = 0; ChildIndex < NumChildren; ChildIndex++)
+			{
+				OutBoneTransforms[ChainLink.ChildZeroLengthTransformIndices[ChildIndex]].Transform = ChainLink.Transform;
+			}
+		}
+
+#if WITH_EDITOR
+		DebugLines.Reset(OutBoneTransforms.Num());
+		DebugLines.AddUninitialized(OutBoneTransforms.Num());
+		for (int32 Index = 0; Index < OutBoneTransforms.Num(); ++Index)
+		{
+			DebugLines[Index] = OutBoneTransforms[Index].Transform.GetLocation();
+		}
+#endif // WITH_EDITOR
+
+	}
+}
+
+bool FAnimNode_CCDIK::UpdateChainLink(TArray<CCDIKChainLink>& Chain, int32 LinkIndex, const FVector& TargetPos) const
+{
+	int32 const TipBoneLinkIndex = Chain.Num() - 1;
+
+	ensure(Chain.IsValidIndex(TipBoneLinkIndex));
+	CCDIKChainLink& CurrentLink = Chain[LinkIndex];
+
+	// update new tip pos
+	FVector TipPos = Chain[TipBoneLinkIndex].Transform.GetLocation();
+
+	FTransform& CurrentLinkTransform = CurrentLink.Transform;
+	FVector ToEnd = TipPos - CurrentLinkTransform.GetLocation();
+	FVector ToTarget = TargetPos - CurrentLinkTransform.GetLocation();
+
+	ToEnd.Normalize();
+	ToTarget.Normalize();
+
+	float RotationLimitPerJointInRadian = FMath::DegreesToRadians(RotationLimitPerJoints[LinkIndex]);
+	float Angle = FMath::ClampAngle(FMath::Acos(FVector::DotProduct(ToEnd, ToTarget)), -RotationLimitPerJointInRadian, RotationLimitPerJointInRadian);
+	bool bCanRotate = (FMath::Abs(Angle) > KINDA_SMALL_NUMBER) && (!bEnableRotationLimit || RotationLimitPerJointInRadian > CurrentLink.CurrentAngleDelta);
+	if (bCanRotate)
+	{
+		// check rotation limit first, if fails, just abort
+		if (bEnableRotationLimit)
+		{
+			if (RotationLimitPerJointInRadian < CurrentLink.CurrentAngleDelta + Angle)
+			{
+				Angle = RotationLimitPerJointInRadian - CurrentLink.CurrentAngleDelta;
+				if (Angle <= KINDA_SMALL_NUMBER)
+				{
+					return false;
+				}
+			}
+
+			CurrentLink.CurrentAngleDelta += Angle;
+		}
+
+		// continue with rotating toward to target
+		FVector RotationAxis = FVector::CrossProduct(ToEnd, ToTarget);
+		if (RotationAxis.SizeSquared() > 0.f)
+		{
+			RotationAxis.Normalize();
+			// Delta Rotation is the rotation to target
+			FQuat DeltaRotation(RotationAxis, Angle);
+
+			FQuat NewRotation = DeltaRotation * CurrentLinkTransform.GetRotation();
+			NewRotation.Normalize();
+			CurrentLinkTransform.SetRotation(NewRotation);
+
+			// if I have parent, make sure to refresh local transform since my current transform has changed
+			if (LinkIndex > 0)
+			{
+				CCDIKChainLink const & Parent = Chain[LinkIndex - 1];
+				CurrentLink.LocalTransform = CurrentLinkTransform.GetRelativeTransform(Parent.Transform);
+				CurrentLink.LocalTransform.NormalizeRotation();
+			}
+
+			// now update all my children to have proper transform
+			FTransform CurrentParentTransform = CurrentLinkTransform;
+
+			// now update all chain
+			for (int32 ChildLinkIndex = LinkIndex + 1; ChildLinkIndex <= TipBoneLinkIndex; ++ChildLinkIndex)
+			{
+				CCDIKChainLink& ChildIterLink = Chain[ChildLinkIndex];
+				FCompactPoseBoneIndex ChildBoneIndex = ChildIterLink.BoneIndex;
+				const FTransform LocalTransform = ChildIterLink.LocalTransform;
+				ChildIterLink.Transform = LocalTransform * CurrentParentTransform;
+				ChildIterLink.Transform.NormalizeRotation();
+				CurrentParentTransform = ChildIterLink.Transform;
+			}
+
+			return true;
+		}
+	}
+
+	return false;
+}
+
+bool FAnimNode_CCDIK::IsValidToEvaluate(const USkeleton* Skeleton, const FBoneContainer& RequiredBones)
+{
+	if (EffectorLocationSpace == BCS_ParentBoneSpace || EffectorLocationSpace == BCS_BoneSpace)
+	{
+		if (!EffectorTarget.IsValidToEvaluate(RequiredBones))
+		{
+			return false;
+		}
+	}
+
+	// Allow evaluation if all parameters are initialized and TipBone is child of RootBone
+	return
+		(
+		TipBone.IsValidToEvaluate(RequiredBones)
+		&& RootBone.IsValidToEvaluate(RequiredBones)
+		&& Precision > 0
+		&& RequiredBones.BoneIsChildOf(TipBone.BoneIndex, RootBone.BoneIndex)
+		);
+}
+
+#if WITH_EDITOR
+void FAnimNode_CCDIK::ResizeRotationLimitPerJoints(int32 NewSize)
+{
+	if (NewSize == 0)
+	{
+		RotationLimitPerJoints.Reset();
+	}
+	else if (RotationLimitPerJoints.Num() != NewSize)
+	{
+		int32 StartIndex = RotationLimitPerJoints.Num();
+		RotationLimitPerJoints.SetNum(NewSize);
+		for (int32 Index = StartIndex; Index < RotationLimitPerJoints.Num(); ++Index)
+		{
+			RotationLimitPerJoints[Index] = 30.f;
+		}
+	}
+}
+#endif 
+
+void FAnimNode_CCDIK::InitializeBoneReferences(const FBoneContainer& RequiredBones)
+{
+	TipBone.Initialize(RequiredBones);
+	RootBone.Initialize(RequiredBones);
+	EffectorTarget.InitializeBoneReferences(RequiredBones);
+}
+
+void FAnimNode_CCDIK::GatherDebugData(FNodeDebugData& DebugData)
+{
+	FString DebugLine = DebugData.GetNodeName(this);
+
+	DebugData.AddDebugItem(DebugLine);
+	ComponentPose.GatherDebugData(DebugData);
+}
+
+PRAGMA_ENABLE_OPTIMIZATION
\ No newline at end of file
diff --git a/Engine/Source/Runtime/AnimGraphRuntime/Public/AnimCustomInstance.h b/Engine/Source/Runtime/AnimGraphRuntime/Public/AnimCustomInstance.h
--- a/Engine/Source/Runtime/AnimGraphRuntime/Public/AnimCustomInstance.h
+++ b/Engine/Source/Runtime/AnimGraphRuntime/Public/AnimCustomInstance.h
@@ -25,7 +25,7 @@ class ANIMGRAPHRUNTIME_API UAnimCustomInstance : public UAnimInstance
 	{
 		// make sure to tick and refresh all the time when ticks
 		// @TODO: this needs restoring post-binding
-		InSkeletalMeshComponent->MeshComponentUpdateFlag = EMeshComponentUpdateFlag::AlwaysTickPoseAndRefreshBones;
+		InSkeletalMeshComponent->VisibilityBasedAnimTickOption = EVisibilityBasedAnimTickOption::AlwaysTickPoseAndRefreshBones;
 #if WITH_EDITOR
 		InSkeletalMeshComponent->SetUpdateAnimationInEditor(true);
 #endif
diff --git a/Engine/Source/Runtime/AnimGraphRuntime/Public/AnimNodes/AnimNode_StateResult.cpp b/Engine/Source/Runtime/AnimGraphRuntime/Public/AnimNodes/AnimNode_StateResult.cpp
new file mode 100644
--- /dev/null
+++ b/Engine/Source/Runtime/AnimGraphRuntime/Public/AnimNodes/AnimNode_StateResult.cpp
@@ -0,0 +1,17 @@
+// Copyright 1998-2018 Epic Games, Inc. All Rights Reserved.
+
+#include "AnimNode_StateResult.h"
+
+bool FAnimNode_StateResult::SerializeFromMismatchedTag(const FPropertyTag& Tag, FStructuredArchive::FSlot Slot)
+{
+	if(Tag.Type == NAME_StructProperty && Tag.StructName == FAnimNode_Root::StaticStruct()->GetFName())
+	{
+		FAnimNode_Root OldValue;
+		FAnimNode_Root::StaticStruct()->SerializeItem(Slot, &OldValue, nullptr);
+		*static_cast<FAnimNode_Root*>(this) = OldValue;
+
+		return true;
+	}
+
+	return false;
+}
\ No newline at end of file
diff --git a/Engine/Source/Runtime/AnimGraphRuntime/Public/AnimNodes/AnimNode_StateResult.h b/Engine/Source/Runtime/AnimGraphRuntime/Public/AnimNodes/AnimNode_StateResult.h
new file mode 100644
--- /dev/null
+++ b/Engine/Source/Runtime/AnimGraphRuntime/Public/AnimNodes/AnimNode_StateResult.h
@@ -0,0 +1,29 @@
+// Copyright 1998-2018 Epic Games, Inc. All Rights Reserved.
+
+#pragma once
+
+#include "CoreMinimal.h"
+#include "UObject/ObjectMacros.h"
+#include "AnimNode_Root.h"
+#include "AnimNode_StateResult.generated.h"
+
+// Root node of an state machine state (sink node).
+// We dont use AnimNode_Root to let us distinguish these nodes in the property list at link time.
+USTRUCT(BlueprintInternalUseOnly)
+struct ANIMGRAPHRUNTIME_API FAnimNode_StateResult : public FAnimNode_Root
+{
+	GENERATED_USTRUCT_BODY()
+
+	/** Used to upgrade old FAnimNode_Roots to FAnimNode_StateResult */
+	bool SerializeFromMismatchedTag(const FPropertyTag& Tag, FStructuredArchive::FSlot Slot);
+};
+
+template<>
+struct TStructOpsTypeTraits<FAnimNode_StateResult>
+	: public TStructOpsTypeTraitsBase2<FAnimNode_StateResult>
+{
+	enum
+	{
+		WithSerializeFromMismatchedTag = true,
+	};
+};
diff --git a/Engine/Source/Runtime/AnimGraphRuntime/Public/BoneControllers/AnimNode_CCDIK.h b/Engine/Source/Runtime/AnimGraphRuntime/Public/BoneControllers/AnimNode_CCDIK.h
new file mode 100644
--- /dev/null
+++ b/Engine/Source/Runtime/AnimGraphRuntime/Public/BoneControllers/AnimNode_CCDIK.h
@@ -0,0 +1,129 @@
+// Copyright 1998-2016 Epic Games, Inc. All Rights Reserved.
+
+#pragma once
+
+#include "AnimNode_SkeletalControlBase.h"
+#include "AnimNode_CCDIK.generated.h"
+
+/**
+*	Controller which implements the CCDIK IK approximation algorithm
+*/
+
+/** Transient structure for CCDIK node evaluation */
+struct CCDIKChainLink
+{
+public:
+	/** Transform of bone in component space. */
+	FTransform Transform;
+
+	/** Transform of bone in local space. This is mutable as their component space changes or parents*/
+	FTransform LocalTransform;
+
+	/** Bone Index in SkeletalMesh */
+	FCompactPoseBoneIndex BoneIndex;
+
+	/** Transform Index that this control will output */
+	int32 TransformIndex;
+
+	/** Child bones which are overlapping this bone. 
+	 * They have a zero length distance, so they will inherit this bone's transformation. */
+	TArray<int32> ChildZeroLengthTransformIndices;
+
+	float CurrentAngleDelta;
+
+	CCDIKChainLink()
+		: BoneIndex(INDEX_NONE)
+		, TransformIndex(INDEX_NONE)
+		, CurrentAngleDelta(0.f)
+	{
+	}
+
+	CCDIKChainLink(const FTransform& InTransform, const FTransform& InLocalTransform,  const FCompactPoseBoneIndex& InBoneIndex, const int32& InTransformIndex)
+		: Transform(InTransform)
+		, LocalTransform(InLocalTransform)
+		, BoneIndex(InBoneIndex)
+		, TransformIndex(InTransformIndex)
+		, CurrentAngleDelta(0.f)
+	{
+	}
+};
+
+USTRUCT()
+struct ANIMGRAPHRUNTIME_API FAnimNode_CCDIK : public FAnimNode_SkeletalControlBase
+{
+	GENERATED_USTRUCT_BODY()
+
+	/** Coordinates for target location of tip bone - if EffectorLocationSpace is bone, this is the offset from Target Bone to use as target location*/
+	UPROPERTY(EditAnywhere, Category = Effector, meta = (PinShownByDefault))
+	FVector EffectorLocation;
+
+	/** Reference frame of Effector Transform. */
+	UPROPERTY(EditAnywhere, Category = Effector)
+	TEnumAsByte<enum EBoneControlSpace> EffectorLocationSpace;
+
+	/** If EffectorTransformSpace is a bone, this is the bone to use. **/
+	UPROPERTY(EditAnywhere, Category = Effector)
+	FBoneSocketTarget EffectorTarget;
+
+	/** Name of tip bone */
+	UPROPERTY(EditAnywhere, Category = Solver)
+	FBoneReference TipBone;
+
+	/** Name of the root bone*/
+	UPROPERTY(EditAnywhere, Category = Solver)
+	FBoneReference RootBone;
+
+	/** Tolerance for final tip location delta from EffectorLocation*/
+	UPROPERTY(EditAnywhere, Category = Solver)
+	float Precision;
+
+	/** Maximum number of iterations allowed, to control performance. */
+	UPROPERTY(EditAnywhere, Category = Solver)
+	int32 MaxIterations;
+
+	/** Toggle drawing of axes to debug joint rotation*/
+	UPROPERTY(EditAnywhere, Category = Solver)
+	bool bStartFromTail;
+
+	/** Tolerance for final tip location delta from EffectorLocation*/
+	UPROPERTY(EditAnywhere, Category = Solver)
+	bool bEnableRotationLimit;
+
+private:
+	/** symmetry rotation limit per joint. Index 0 matches with root bone and last index matches with tip bone. */
+	UPROPERTY(EditAnywhere, EditFixedSize, Category = Solver)
+	TArray<float> RotationLimitPerJoints;
+
+public:
+	FAnimNode_CCDIK();
+
+	// FAnimNode_Base interface
+	virtual void GatherDebugData(FNodeDebugData& DebugData) override;
+	// End of FAnimNode_Base interface
+
+	// FAnimNode_SkeletalControlBase interface
+	virtual void EvaluateSkeletalControl_AnyThread(FComponentSpacePoseContext& Output, TArray<FBoneTransform>& OutBoneTransforms) override;
+	virtual bool IsValidToEvaluate(const USkeleton* Skeleton, const FBoneContainer& RequiredBones) override;
+	// End of FAnimNode_SkeletalControlBase interface
+
+private:
+	// FAnimNode_SkeletalControlBase interface
+	virtual void InitializeBoneReferences(const FBoneContainer& RequiredBones) override;
+	// End of FAnimNode_SkeletalControlBase interface
+
+	// Convenience function to get current (pre-translation iteration) component space location of bone by bone index
+	FVector GetCurrentLocation(FCSPose<FCompactPose>& MeshBases, const FCompactPoseBoneIndex& BoneIndex);
+
+	static FTransform GetTargetTransform(const FTransform& InComponentTransform, FCSPose<FCompactPose>& MeshBases, FBoneSocketTarget& InTarget, EBoneControlSpace Space, const FVector& InOffset);
+
+	// return true if updated
+	bool UpdateChainLink(TArray<CCDIKChainLink>& Chain, int32 LinkIndex, const FVector& TargetPos) const;
+public:
+#if WITH_EDITOR
+#if !(UE_BUILD_SHIPPING || UE_BUILD_TEST)
+	TArray<FVector> DebugLines;
+#endif // #if !(UE_BUILD_SHIPPING || UE_BUILD_TEST)
+	// resize rotation limit array based on set up
+	void ResizeRotationLimitPerJoints(int32 NewSize);
+#endif
+};
\ No newline at end of file
diff --git a/Engine/Source/Runtime/Core/Private/UObject/DevObjectVersion.cpp b/Engine/Source/Runtime/Core/Private/UObject/DevObjectVersion.cpp
--- a/Engine/Source/Runtime/Core/Private/UObject/DevObjectVersion.cpp
+++ b/Engine/Source/Runtime/Core/Private/UObject/DevObjectVersion.cpp
@@ -17,6 +17,7 @@
 #include "UObject/VRObjectVersion.h"
 #include "UObject/GeometryObjectVersion.h"
 #include "UObject/AnimPhysObjectVersion.h"
+#include "UObject/AnimObjectVersion.h"
 #include "UObject/FortniteMainBranchObjectVersion.h"
 #include "UObject/ReflectionCaptureObjectVersion.h"
 #include "UObject/LoadTimesObjectVersion.h"
@@ -126,6 +127,11 @@ const FGuid FAnimPhysObjectVersion::GUID(0x29E575DD, 0xE0A34627, 0x9D10D276, 0x2
 // Register AnimPhys custom version with Core
 FDevVersionRegistration GRegisterAnimPhysObjectVersion(FAnimPhysObjectVersion::GUID, FAnimPhysObjectVersion::LatestVersion, TEXT("Dev-AnimPhys"));
 
+// Unique Anim Object version id
+const FGuid FAnimObjectVersion::GUID(0xAF43A65D, 0x7FD34947, 0x98733E8E, 0xD9C1BB05);
+// Register AnimPhys custom version with Core
+FDevVersionRegistration GRegisterAnimObjectVersion(FAnimObjectVersion::GUID, FAnimObjectVersion::LatestVersion, TEXT("Dev-Anim"));
+
 // Unique ReflectionCapture Object version id
 const FGuid FReflectionCaptureObjectVersion::GUID(0x6B266CEC, 0x1EC74B8F, 0xA30BE4D9, 0x0942FC07);
 // Register Rendering custom version with Core
diff --git a/Engine/Source/Runtime/Core/Public/UObject/AnimObjectVersion.h b/Engine/Source/Runtime/Core/Public/UObject/AnimObjectVersion.h
new file mode 100644
--- /dev/null
+++ b/Engine/Source/Runtime/Core/Public/UObject/AnimObjectVersion.h
@@ -0,0 +1,31 @@
+// Copyright 1998-2018 Epic Games, Inc. All Rights Reserved.
+#pragma once
+
+#include "CoreTypes.h"
+#include "Misc/Guid.h"
+
+// Custom serialization version for changes made in Dev-Anim stream
+struct CORE_API FAnimObjectVersion
+{
+	enum Type
+	{
+		// Before any version changes were made
+		BeforeCustomVersionWasAdded,
+
+		// Reworked how anim blueprint root nodes are recovered
+		LinkTimeAnimBlueprintRootDiscovery,
+
+		// Cached marker sync names on skeleton for editor
+		StoreMarkerNamesOnSkeleton,
+
+		// -----<new versions can be added above this line>-------------------------------------------------
+		VersionPlusOne,
+		LatestVersion = VersionPlusOne - 1
+	};
+
+	// The GUID for this custom version number
+	const static FGuid GUID;
+
+private:
+	FAnimObjectVersion() {}
+};
diff --git a/Engine/Source/Runtime/CoreUObject/Private/UObject/CoreRedirects.cpp b/Engine/Source/Runtime/CoreUObject/Private/UObject/CoreRedirects.cpp
--- a/Engine/Source/Runtime/CoreUObject/Private/UObject/CoreRedirects.cpp
+++ b/Engine/Source/Runtime/CoreUObject/Private/UObject/CoreRedirects.cpp
@@ -1312,7 +1312,6 @@ static void RegisterNativeRedirects40(TArray<FCoreRedirect>& Redirects)
 	PROPERTY_REDIRECT("SceneComponent.RelativeTranslation", "SceneComponent.RelativeLocation");
 	PROPERTY_REDIRECT("SceneComponent.bAbsoluteTranslation", "SceneComponent.bAbsoluteLocation");
 	PROPERTY_REDIRECT("SkeletalMeshComponent.AnimationBlueprint", "SkeletalMeshComponent.AnimBlueprintGeneratedClass");
-	PROPERTY_REDIRECT("SkinnedMeshComponent.SkinnedMeshUpdateFlag", "SkinnedMeshComponent.MeshComponentUpdateFlag");
 	PROPERTY_REDIRECT("SlateBrush.TextureName", "SlateBrush.ResourceName");
 	PROPERTY_REDIRECT("SlateBrush.TextureObject", "SlateBrush.ResourceObject");
 	PROPERTY_REDIRECT("WorldSettings.DefaultGameType", "WorldSettings.DefaultGameMode");
@@ -1381,10 +1380,6 @@ static void RegisterNativeRedirects40(TArray<FCoreRedirect>& Redirects)
 	FCoreRedirect* EMaxConcurrentResolutionRule = ENUM_REDIRECT("EMaxConcurrentResolutionRule", "EMaxConcurrentResolutionRule");
 	EMaxConcurrentResolutionRule->ValueChanges.Add(TEXT("EMaxConcurrentResolutionRule::StopFarthest"), TEXT("EMaxConcurrentResolutionRule::StopFarthestThenPreventNew"));
 
-	FCoreRedirect* EMeshComponentUpdateFlag = ENUM_REDIRECT("EMeshComponentUpdateFlag", "EMeshComponentUpdateFlag");
-	EMeshComponentUpdateFlag->ValueChanges.Add(TEXT("SMU_AlwaysTickPose"), TEXT("EMeshComponentUpdateFlag::AlwaysTickPose"));
-	EMeshComponentUpdateFlag->ValueChanges.Add(TEXT("SMU_AlwaysTickPoseAndRefreshBones"), TEXT("EMeshComponentUpdateFlag::AlwaysTickPoseAndRefreshBones"));
-	EMeshComponentUpdateFlag->ValueChanges.Add(TEXT("SMU_OnlyTickPoseWhenRendered"), TEXT("EMeshComponentUpdateFlag::OnlyTickPoseWhenRendered"));
 
 	FCoreRedirect* EParticleEventType = ENUM_REDIRECT("EParticleEventType", "EParticleEventType");
 	EParticleEventType->ValueChanges.Add(TEXT("EPET_Kismet"), TEXT("EPET_Blueprint"));
diff --git a/Engine/Source/Runtime/Engine/Classes/Animation/AnimBlueprintGeneratedClass.h b/Engine/Source/Runtime/Engine/Classes/Animation/AnimBlueprintGeneratedClass.h
--- a/Engine/Source/Runtime/Engine/Classes/Animation/AnimBlueprintGeneratedClass.h
+++ b/Engine/Source/Runtime/Engine/Classes/Animation/AnimBlueprintGeneratedClass.h
@@ -189,8 +189,7 @@ class ENGINE_API UAnimBlueprintGeneratedClass : public UBlueprintGeneratedClass,
 	UPROPERTY()
 	TArray<FAnimNotifyEvent> AnimNotifies;
 
-	// The index of the root node in the animation tree
-	UPROPERTY()
+	// The index of the root node in the animation tree (created during link)
 	int32 RootAnimNodeIndex;
 
 	// Indices for each of the saved pose nodes that require updating, in the order they need to get updates.
@@ -334,6 +333,10 @@ public:
 
 #endif
 
+	// UObject interface
+	virtual void Serialize(FArchive& Ar) override;
+	// End of UObject interface
+
 	// UStruct interface
 	virtual void Link(FArchive& Ar, bool bRelinkExistingProperties) override;
 	// End of UStruct interface
diff --git a/Engine/Source/Runtime/Engine/Classes/Animation/AnimNode_SubInstance.h b/Engine/Source/Runtime/Engine/Classes/Animation/AnimNode_SubInstance.h
--- a/Engine/Source/Runtime/Engine/Classes/Animation/AnimNode_SubInstance.h
+++ b/Engine/Source/Runtime/Engine/Classes/Animation/AnimNode_SubInstance.h
@@ -28,9 +28,14 @@ public:
 	UPROPERTY()
 	FPoseLink InPose;
 
+	/** The class spawned for this sub-instance */
 	UPROPERTY(EditAnywhere, BlueprintReadWrite, Category = Settings)
 	TSubclassOf<UAnimInstance> InstanceClass;
 
+	/** Optional tag used to identify this sub-instance */
+	UPROPERTY(EditAnywhere, BlueprintReadWrite, Category = Settings)
+	FName Tag;
+
 	/** This is the actual instance allocated at runtime that will run */
 	UPROPERTY(transient)
 	UAnimInstance* InstanceToRun;
diff --git a/Engine/Source/Runtime/Engine/Classes/Animation/AnimNotifies/AnimNotify_PauseClothingSimulation.h b/Engine/Source/Runtime/Engine/Classes/Animation/AnimNotifies/AnimNotify_PauseClothingSimulation.h
new file mode 100644
--- /dev/null
+++ b/Engine/Source/Runtime/Engine/Classes/Animation/AnimNotifies/AnimNotify_PauseClothingSimulation.h
@@ -0,0 +1,26 @@
+// Copyright 1998-2018 Epic Games, Inc. All Rights Reserved.
+
+#pragma once
+
+#include "AnimNotify.h"
+
+#include "AnimNotify_PauseClothingSimulation.generated.h"
+
+UCLASS(const, hidecategories=Object, collapsecategories, meta=(DisplayName="Pause Clothing Simulation"))
+class ENGINE_API UAnimNotify_PauseClothingSimulation : public UAnimNotify
+{
+	GENERATED_BODY()
+
+public:
+
+	UAnimNotify_PauseClothingSimulation();
+
+	// Begin UAnimNotify interface
+	virtual FString GetNotifyName_Implementation() const override;
+	virtual void Notify(USkeletalMeshComponent* MeshComp, UAnimSequenceBase* Animation) override;
+	// End UAnimNotify interface
+
+};
+
+
+
diff --git a/Engine/Source/Runtime/Engine/Classes/Animation/AnimNotifies/AnimNotify_ResetDynamics.h b/Engine/Source/Runtime/Engine/Classes/Animation/AnimNotifies/AnimNotify_ResetDynamics.h
new file mode 100644
--- /dev/null
+++ b/Engine/Source/Runtime/Engine/Classes/Animation/AnimNotifies/AnimNotify_ResetDynamics.h
@@ -0,0 +1,26 @@
+// Copyright 1998-2018 Epic Games, Inc. All Rights Reserved.
+
+#pragma once
+
+#include "AnimNotify.h"
+
+#include "AnimNotify_ResetDynamics.generated.h"
+
+UCLASS(const, hidecategories=Object, collapsecategories, meta=(DisplayName="Reset Dynamics"))
+class ENGINE_API UAnimNotify_ResetDynamics : public UAnimNotify
+{
+	GENERATED_BODY()
+
+public:
+
+	UAnimNotify_ResetDynamics();
+
+	// Begin UAnimNotify interface
+	virtual FString GetNotifyName_Implementation() const override;
+	virtual void Notify(USkeletalMeshComponent* MeshComp, UAnimSequenceBase* Animation) override;
+	// End UAnimNotify interface
+
+};
+
+
+
diff --git a/Engine/Source/Runtime/Engine/Classes/Animation/AnimNotifies/AnimNotify_ResumeClothingSimulation.h b/Engine/Source/Runtime/Engine/Classes/Animation/AnimNotifies/AnimNotify_ResumeClothingSimulation.h
new file mode 100644
--- /dev/null
+++ b/Engine/Source/Runtime/Engine/Classes/Animation/AnimNotifies/AnimNotify_ResumeClothingSimulation.h
@@ -0,0 +1,26 @@
+// Copyright 1998-2018 Epic Games, Inc. All Rights Reserved.
+
+#pragma once
+
+#include "AnimNotify.h"
+
+#include "AnimNotify_ResumeClothingSimulation.generated.h"
+
+UCLASS(const, hidecategories=Object, collapsecategories, meta=(DisplayName="Resume Clothing Simulation"))
+class ENGINE_API UAnimNotify_ResumeClothingSimulation : public UAnimNotify
+{
+	GENERATED_BODY()
+
+public:
+
+	UAnimNotify_ResumeClothingSimulation();
+
+	// Begin UAnimNotify interface
+	virtual FString GetNotifyName_Implementation() const override;
+	virtual void Notify(USkeletalMeshComponent* MeshComp, UAnimSequenceBase* Animation) override;
+	// End UAnimNotify interface
+
+};
+
+
+
diff --git a/Engine/Source/Runtime/Engine/Classes/Animation/AnimSequence.h b/Engine/Source/Runtime/Engine/Classes/Animation/AnimSequence.h
--- a/Engine/Source/Runtime/Engine/Classes/Animation/AnimSequence.h
+++ b/Engine/Source/Runtime/Engine/Classes/Animation/AnimSequence.h
@@ -514,6 +514,10 @@ public:
 	UPROPERTY(EditAnywhere, Category = RootMotion)
 	bool bForceRootLock;
 
+	/** If this is on, it will use a normalized scale value for the root motion extracted: FVector(1.0, 1.0, 1.0) **/
+	UPROPERTY(EditAnywhere, AssetRegistrySearchable, Category = RootMotion, meta = (DisplayName = "Use Normalized Root Motion Scale"))
+	bool bUseNormalizedRootMotionScale;
+
 	/** Have we copied root motion settings from an owning montage */
 	UPROPERTY()
 	bool bRootMotionSettingsCopiedFromMontage;
@@ -606,6 +610,9 @@ public:
 #endif
 	//~ End UAnimSequenceBase Interface
 
+	// Returns the framerate of the animation
+	float GetFrameRate() const { return (float)(FMath::Max(NumFrames - 1, 1)) / (SequenceLength > 0.f ? SequenceLength : 1.f); }
+
 	// Extract Root Motion transform from the animation
 	FTransform ExtractRootMotion(float StartTime, float DeltaTime, bool bAllowLooping) const;
 
diff --git a/Engine/Source/Runtime/Engine/Classes/Animation/Skeleton.h b/Engine/Source/Runtime/Engine/Classes/Animation/Skeleton.h
--- a/Engine/Source/Runtime/Engine/Classes/Animation/Skeleton.h
+++ b/Engine/Source/Runtime/Engine/Classes/Animation/Skeleton.h
@@ -366,6 +366,18 @@ public:
 	// return version of AnimCurveUidVersion
 	uint16 GetAnimCurveUidVersion() const { return AnimCurveUidVersion;  }
 	const TArray<uint16>& GetDefaultCurveUIDList() const { return DefaultCurveUIDList; }
+
+#if WITH_EDITOR
+	// Get existing (seen) sync marker names for this Skeleton
+	const TArray<FName>& GetExistingMarkerNames() const { return ExistingMarkerNames; }
+
+	// Register a new sync marker name
+	void RegisterMarkerName(FName MarkerName) { ExistingMarkerNames.AddUnique(MarkerName); ExistingMarkerNames.Sort(); }
+
+	// Remove a sync marker name
+	void RemoveMarkerName(FName MarkerName) { ExistingMarkerNames.Remove(MarkerName); }
+#endif
+
 protected:
 	// Container for smart name mappings
 	UPROPERTY()
@@ -376,6 +388,9 @@ protected:
 	// FBoneContainer contains only list that is used by current LOD
 	TArray<uint16> DefaultCurveUIDList;
 
+	//Cached marker sync marker names (stripped for non editor)
+	TArray<FName> ExistingMarkerNames;
+
 private:
 	/** Increase the AnimCurveUidVersion so that instances can get the latest information */
 	void IncreaseAnimCurveUidVersion();
diff --git a/Engine/Source/Runtime/Engine/Classes/Components/SkeletalMeshComponent.h b/Engine/Source/Runtime/Engine/Classes/Components/SkeletalMeshComponent.h
--- a/Engine/Source/Runtime/Engine/Classes/Components/SkeletalMeshComponent.h
+++ b/Engine/Source/Runtime/Engine/Classes/Components/SkeletalMeshComponent.h
@@ -663,6 +663,13 @@ public:
 	UFUNCTION(BlueprintCallable, Category = "Components|SkeletalMesh", meta = (Keywords = "AnimBlueprint"))
 	UAnimInstance* GetPostProcessInstance() const;
 
+	/**
+	 * Returns the a tagged sub-instance node. If non sub instances are found or none are tagged with the
+	 * supplied name, this will return NULL.
+	 */
+	UFUNCTION(BlueprintCallable, Category = "Components|SkeletalMesh", meta = (Keywords = "AnimBlueprint"))
+	UAnimInstance* GetSubInstanceByName(FName InName) const;
+
 	/** 
 	 * Returns whether there are any valid instances to run, currently this means whether we have
 	 * have an animation instance or a post process instance available to process.
@@ -1159,7 +1166,7 @@ public:
 	FOnAnimInitialized OnAnimInitialized;
 
 	/**
-		If MeshComponentUpdateFlag == EMeshComponentUpdateFlag::OnlyTickMontagesWhenNotRendered
+		If VisibilityBasedAnimTickOption == EVisibilityBasedAnimTickOption::OnlyTickMontagesWhenNotRendered
 		Should we tick Montages only?
 	*/
 	bool ShouldOnlyTickMontages(const float DeltaTime) const;
diff --git a/Engine/Source/Runtime/Engine/Classes/Components/SkinnedMeshComponent.h b/Engine/Source/Runtime/Engine/Classes/Components/SkinnedMeshComponent.h
--- a/Engine/Source/Runtime/Engine/Classes/Components/SkinnedMeshComponent.h
+++ b/Engine/Source/Runtime/Engine/Classes/Components/SkinnedMeshComponent.h
@@ -50,26 +50,38 @@ enum EPhysBodyOp
 	PBO_MAX,
 };
 
-/** Skinned Mesh Update Flag based on rendered or not. */
-UENUM()
+/** Skinned Mesh Animation Tick option based on rendered or not. This dictates "TickPose and RefreshBoneTransforms" */
+UENUM(BlueprintType)
+enum class EVisibilityBasedAnimTickOption : uint8
+{
+	/** Always Tick and Refresh BoneTransforms whether rendered or not. */
+	AlwaysTickPoseAndRefreshBones,
+	/** Always Tick, but Refresh BoneTransforms only when rendered. */
+	AlwaysTickPose,
+	/**
+		When rendered Tick Pose and Refresh Bone Transforms,
+		otherwise, just update montages and skip everything else.
+		(AnimBP graph will not be updated).
+	*/
+	OnlyTickMontagesWhenNotRendered,
+	/** Tick only when rendered, and it will only RefreshBoneTransforms when rendered. */
+	OnlyTickPoseWhenRendered,
+};
+
+/** Previous deprecated animation tick option. */
 namespace EMeshComponentUpdateFlag
 {
-	enum Type
-	{
-		/** Always Tick and Refresh BoneTransforms whether rendered or not. */
-		AlwaysTickPoseAndRefreshBones,
-		/** Always Tick, but Refresh BoneTransforms only when rendered. */
-		AlwaysTickPose,
-		/**
-			When rendered Tick Pose and Refresh Bone Transforms,
-			otherwise, just update montages and skip everything else.
-			(AnimBP graph will not be updated).
-		*/
-		OnlyTickMontagesWhenNotRendered,
-		/** Tick only when rendered, and it will only RefreshBoneTransforms when rendered. */
-		OnlyTickPoseWhenRendered,
-	};
-}
+	DEPRECATED(4.20, "EMeshComponentUpdateFlag has been deprecated use EVisibilityBasedAnimTickOption instead")
+	typedef int32 Type;
+	DEPRECATED(4.20, "EMeshComponentUpdateFlag has been deprecated use EVisibilityBasedAnimTickOption instead")
+	static const uint8 AlwaysTickPoseAndRefreshBones = 0;
+	DEPRECATED(4.20, "EMeshComponentUpdateFlag has been deprecated use EVisibilityBasedAnimTickOption instead")
+	static const uint8 AlwaysTickPose = 1;
+	DEPRECATED(4.20, "EMeshComponentUpdateFlag has been deprecated use EVisibilityBasedAnimTickOption instead")
+	static const uint8 OnlyTickMontagesWhenNotRendered = 2;
+	DEPRECATED(4.20, "EMeshComponentUpdateFlag has been deprecated use EVisibilityBasedAnimTickOption instead")
+	static const uint8 OnlyTickPoseWhenRendered = 3;
+};
 
 /** Values for specifying bone space. */
 UENUM()
@@ -358,15 +370,13 @@ public:
 	/** Array of bone visibilities (containing one of the values in EBoneVisibilityStatus for each bone).  A bone is only visible if it is *exactly* 1 (BVS_Visible) */
 	TArray<uint8> BoneVisibilityStates;
 
-	/** This is update frequency flag even when our Owner has not been rendered recently
-	 * 
-	 * SMU_AlwaysTickPoseAndRefreshBones,			// Always Tick and Refresh BoneTransforms whether rendered or not
-	 * SMU_AlwaysTickPose,							// Always Tick, but Refresh BoneTransforms only when rendered
-	 * SMU_OnlyTickPoseWhenRendered,				// Tick only when rendered, and it will only RefreshBoneTransforms when rendered
-	 * 
+	/*
+	 * This is tick animation frequency option based on this component rendered or not or using montage
+	 *  You can change this default value in the INI file 
+	 * Mostly related with performance
 	 */
-	UPROPERTY(EditAnywhere, AdvancedDisplay, BlueprintReadWrite, Category=SkeletalMesh)
-	TEnumAsByte<EMeshComponentUpdateFlag::Type> MeshComponentUpdateFlag;
+	UPROPERTY(EditAnywhere, AdvancedDisplay, BlueprintReadWrite, Config, Category=Optimization)
+	EVisibilityBasedAnimTickOption VisibilityBasedAnimTickOption;
 
 	/** 
 	 * When true, we will just using the bounds from our MasterPoseComponent.  This is useful for when we have a Mesh Parented
@@ -402,11 +412,11 @@ public:
 	//
 	
 	/** When true, skip using the physics asset etc. and always use the fixed bounds defined in the SkeletalMesh. */
-	UPROPERTY(EditAnywhere, AdvancedDisplay, BlueprintReadWrite, Category=SkeletalMesh)
+	UPROPERTY(EditAnywhere, AdvancedDisplay, BlueprintReadWrite, Category=Optimization)
 	uint8 bComponentUseFixedSkelBounds:1;
 
 	/** If true, when updating bounds from a PhysicsAsset, consider _all_ BodySetups, not just those flagged with bConsiderForBounds. */
-	UPROPERTY(EditAnywhere, AdvancedDisplay, BlueprintReadWrite, Category=SkeletalMesh)
+	UPROPERTY(EditAnywhere, AdvancedDisplay, BlueprintReadWrite, Category=Optimization)
 	uint8 bConsiderAllBodiesForBounds:1;
 
 	/** If true, this component uses its parents LOD when attached if available
diff --git a/Engine/Source/Runtime/Engine/Private/AnimBlueprint.cpp b/Engine/Source/Runtime/Engine/Private/AnimBlueprint.cpp
--- a/Engine/Source/Runtime/Engine/Private/AnimBlueprint.cpp
+++ b/Engine/Source/Runtime/Engine/Private/AnimBlueprint.cpp
@@ -178,7 +178,7 @@ void UAnimBlueprint::PostLoad()
 
 bool UAnimBlueprint::CanRecompileWhilePlayingInEditor() const
 {
-	return GetDefault<UEditorExperimentalSettings>()->bEnableLiveRecompilationOfAnimationBlueprints;
+	return true;
 }
 #endif
 
diff --git a/Engine/Source/Runtime/Engine/Private/Animation/AnimBlueprintGeneratedClass.cpp b/Engine/Source/Runtime/Engine/Private/Animation/AnimBlueprintGeneratedClass.cpp
--- a/Engine/Source/Runtime/Engine/Private/Animation/AnimBlueprintGeneratedClass.cpp
+++ b/Engine/Source/Runtime/Engine/Private/Animation/AnimBlueprintGeneratedClass.cpp
@@ -10,6 +10,7 @@
 #include "BonePose.h"
 #include "Animation/AnimNodeBase.h"
 #include "Animation/AnimInstance.h"
+#include "UObject/AnimObjectVersion.h"
 
 /////////////////////////////////////////////////////
 // FStateMachineDebugData
@@ -234,10 +235,20 @@ UAnimBlueprintGeneratedClass::UAnimBlueprintGeneratedClass(const FObjectInitiali
 	RootAnimNodeIndex = INDEX_NONE;
 }
 
+void UAnimBlueprintGeneratedClass::Serialize(FArchive& Ar)
+{
+	Ar.UsingCustomVersion(FAnimObjectVersion::GUID);
+
+	Super::Serialize(Ar);
+}
+
 void UAnimBlueprintGeneratedClass::Link(FArchive& Ar, bool bRelinkExistingProperties)
 {
 	Super::Link(Ar, bRelinkExistingProperties);
 
+	// We cant reference FAnimNode_Root directly as it is in AnimGraphRuntime, so we just hard-code the name instead.
+	static const FName NAME_AnimNode_Root(TEXT("AnimNode_Root"));
+
 	// @TODO: Shouldn't be necessary to clear these, but currently the class gets linked twice during compilation
 	AnimNodeProperties.Empty();
 	RootAnimNodeProperty = NULL;
@@ -249,6 +260,19 @@ void UAnimBlueprintGeneratedClass::Link(FArchive& Ar, bool bRelinkExistingProper
 		{
 			if (StructProp->Struct->IsChildOf(FAnimNode_Base::StaticStruct()))
 			{
+				if(StructProp->Struct->GetFName() == NAME_AnimNode_Root)
+				{
+					// If we are loading from a newer version, we must verify that there is only one root here.
+					// When linking an older version on load we may find multiple roots until the class is recompiled or saved
+#if DO_CHECK
+					if(Ar.IsLoading() && Ar.CustomVer(FAnimObjectVersion::GUID) >= FAnimObjectVersion::LinkTimeAnimBlueprintRootDiscovery)
+					{
+						check(RootAnimNodeProperty == nullptr);
+					}
+#endif
+					RootAnimNodeProperty = StructProp;
+					RootAnimNodeIndex = AnimNodeProperties.Num();
+				}
 				AnimNodeProperties.Add(StructProp);
 			}
 		}
@@ -273,11 +297,7 @@ void UAnimBlueprintGeneratedClass::Link(FArchive& Ar, bool bRelinkExistingProper
 	if (AnimNodeProperties.Num() > 0)
 	{
 		const bool bValidRootIndex = (RootAnimNodeIndex >= 0) && (RootAnimNodeIndex < AnimNodeProperties.Num());
-		if (bValidRootIndex)
-		{
-			RootAnimNodeProperty = AnimNodeProperties[AnimNodeProperties.Num() - 1 - RootAnimNodeIndex];
-		}
-		else
+		if (!bValidRootIndex)
 		{
 			UE_LOG(LogAnimation, Warning, TEXT("Invalid animation root node index %d on '%s' (only %d nodes)"), RootAnimNodeIndex, *GetPathName(), AnimNodeProperties.Num());
 			AnimNodeProperties.Empty();
diff --git a/Engine/Source/Runtime/Engine/Private/Animation/AnimInstanceProxy.cpp b/Engine/Source/Runtime/Engine/Private/Animation/AnimInstanceProxy.cpp
--- a/Engine/Source/Runtime/Engine/Private/Animation/AnimInstanceProxy.cpp
+++ b/Engine/Source/Runtime/Engine/Private/Animation/AnimInstanceProxy.cpp
@@ -1188,7 +1188,7 @@ void FAnimInstanceProxy::SlotEvaluatePose(const FName& SlotNodeName, const FComp
 	{
 		// If MontageEvaluationData is not valid anymore, pass-through AnimSlot.
 		// This can happen if InitAnim pushes a RefreshBoneTransforms when not rendered,
-		// with EMeshComponentUpdateFlag::OnlyTickMontagesWhenNotRendered set.
+		// with EVisibilityBasedAnimTickOption::OnlyTickMontagesWhenNotRendered set.
 		if (!EvalState.Montage.IsValid())
 		{
 			BlendedPose = SourcePose;
diff --git a/Engine/Source/Runtime/Engine/Private/Animation/AnimMontage.cpp b/Engine/Source/Runtime/Engine/Private/Animation/AnimMontage.cpp
--- a/Engine/Source/Runtime/Engine/Private/Animation/AnimMontage.cpp
+++ b/Engine/Source/Runtime/Engine/Private/Animation/AnimMontage.cpp
@@ -2268,21 +2268,27 @@ void FAnimMontageInstance::Advance(float DeltaTime, struct FRootMotionMovementPa
 							, *RootMotion.GetTranslation().ToCompactString()
 							, *OutRootMotionParams->GetRootMotionTransform().GetTranslation().ToCompactString()
 							, bBlendRootMotion
-							);
+						);
 					}
+				}
 
-					// Delegate has to be called last in this loop
-					// so that if this changes position, the new position will be applied in the next loop
-					// first need to have event handler to handle it
-					// Save off position before triggering events, in case they cause a jump to another position
-					const float PositionBeforeFiringEvents = Position;
+				// Delegate has to be called last in this loop
+				// so that if this changes position, the new position will be applied in the next loop
+				// first need to have event handler to handle it
+				// Save off position before triggering events, in case they cause a jump to another position
+				const float PositionBeforeFiringEvents = Position;
 
+				if(bHaveMoved)
+				{
 					// Save position before firing events.
 					if (!bInterrupted)
 					{
 						HandleEvents(PreviousSubStepPosition, Position, BranchingPointMarker);
 					}
+				}
 
+				if(MontageSubStepper.HasTimeRemaining())
+				{
 					// if we reached end of section, and we were not processing a branching point, and no events has messed with out current position..
 					// .. Move to next section.
 					// (this also handles looping, the same as jumping to a different section).
@@ -2299,6 +2305,7 @@ void FAnimMontageInstance::Advance(float DeltaTime, struct FRootMotionMovementPa
 							// Jump to next section's appropriate starting point (start or end).
 							const float EndOffset = KINDA_SMALL_NUMBER / 2.f; //KINDA_SMALL_NUMBER/2 because we use KINDA_SMALL_NUMBER to offset notifies for triggering and SMALL_NUMBER is too small
 							Position = bPlayingForward ? LatestNextSectionStartTime : (LatestNextSectionEndTime - EndOffset);
+							SubStepResult = EMontageSubStepResult::Moved;
 						}
 						else
 						{
diff --git a/Engine/Source/Runtime/Engine/Private/Animation/AnimNode_SubInstance.cpp b/Engine/Source/Runtime/Engine/Private/Animation/AnimNode_SubInstance.cpp
--- a/Engine/Source/Runtime/Engine/Private/Animation/AnimNode_SubInstance.cpp
+++ b/Engine/Source/Runtime/Engine/Private/Animation/AnimNode_SubInstance.cpp
@@ -7,6 +7,7 @@
 
 FAnimNode_SubInstance::FAnimNode_SubInstance()
 	: InstanceClass(nullptr)
+	, Tag(NAME_None)
 	, InstanceToRun(nullptr)
 {
 
@@ -142,8 +143,10 @@ void FAnimNode_SubInstance::OnInitializeAnimInstance(const FAnimInstanceProxy* I
 			InstanceToRun = nullptr;
 		}
 
-		// Need an instance to run
-		InstanceToRun = NewObject<UAnimInstance>(MeshComp, InstanceClass);
+		// Need an instance to run, so create it now
+		// We use the tag to name the object, but as we verify there are no duplicates in the compiler we
+		// dont need to verify it is unique here.
+		InstanceToRun = NewObject<UAnimInstance>(MeshComp, InstanceClass, Tag);
 
 		// Set up bone transform array
 		AllocateBoneTransforms(InstanceToRun);
diff --git a/Engine/Source/Runtime/Engine/Private/Animation/AnimNotify_PauseClothingSimulation.cpp b/Engine/Source/Runtime/Engine/Private/Animation/AnimNotify_PauseClothingSimulation.cpp
new file mode 100644
--- /dev/null
+++ b/Engine/Source/Runtime/Engine/Private/Animation/AnimNotify_PauseClothingSimulation.cpp
@@ -0,0 +1,22 @@
+// Copyright 1998-2018 Epic Games, Inc. All Rights Reserved.
+
+#include "Animation/AnimNotifies/AnimNotify_PauseClothingSimulation.h"
+#include "Components/SkeletalMeshComponent.h"
+
+UAnimNotify_PauseClothingSimulation::UAnimNotify_PauseClothingSimulation()
+	: Super()
+{
+#if WITH_EDITORONLY_DATA
+	NotifyColor = FColor(90, 220, 255, 255);
+#endif // WITH_EDITORONLY_DATA
+}
+
+void UAnimNotify_PauseClothingSimulation::Notify(class USkeletalMeshComponent* MeshComp, class UAnimSequenceBase* Animation)
+{
+	MeshComp->SuspendClothingSimulation();
+}
+
+FString UAnimNotify_PauseClothingSimulation::GetNotifyName_Implementation() const
+{
+	return TEXT("Pause Clothing Sim");
+}
\ No newline at end of file
diff --git a/Engine/Source/Runtime/Engine/Private/Animation/AnimNotify_ResetDynamics.cpp b/Engine/Source/Runtime/Engine/Private/Animation/AnimNotify_ResetDynamics.cpp
new file mode 100644
--- /dev/null
+++ b/Engine/Source/Runtime/Engine/Private/Animation/AnimNotify_ResetDynamics.cpp
@@ -0,0 +1,22 @@
+// Copyright 1998-2018 Epic Games, Inc. All Rights Reserved.
+
+#include "Animation/AnimNotifies/AnimNotify_ResetDynamics.h"
+#include "Components/SkeletalMeshComponent.h"
+
+UAnimNotify_ResetDynamics::UAnimNotify_ResetDynamics()
+	: Super()
+{
+#if WITH_EDITORONLY_DATA
+	NotifyColor = FColor(90, 220, 255, 255);
+#endif // WITH_EDITORONLY_DATA
+}
+
+void UAnimNotify_ResetDynamics::Notify(class USkeletalMeshComponent* MeshComp, class UAnimSequenceBase* Animation)
+{
+	MeshComp->ResetAnimInstanceDynamics();
+}
+
+FString UAnimNotify_ResetDynamics::GetNotifyName_Implementation() const
+{
+	return TEXT("Reset Dynamics");
+}
\ No newline at end of file
diff --git a/Engine/Source/Runtime/Engine/Private/Animation/AnimNotify_ResumeClothingSimulation.cpp b/Engine/Source/Runtime/Engine/Private/Animation/AnimNotify_ResumeClothingSimulation.cpp
new file mode 100644
--- /dev/null
+++ b/Engine/Source/Runtime/Engine/Private/Animation/AnimNotify_ResumeClothingSimulation.cpp
@@ -0,0 +1,22 @@
+// Copyright 1998-2018 Epic Games, Inc. All Rights Reserved.
+
+#include "Animation/AnimNotifies/AnimNotify_ResumeClothingSimulation.h"
+#include "Components/SkeletalMeshComponent.h"
+
+UAnimNotify_ResumeClothingSimulation::UAnimNotify_ResumeClothingSimulation()
+	: Super()
+{
+#if WITH_EDITORONLY_DATA
+	NotifyColor = FColor(90, 220, 255, 255);
+#endif // WITH_EDITORONLY_DATA
+}
+
+void UAnimNotify_ResumeClothingSimulation::Notify(class USkeletalMeshComponent* MeshComp, class UAnimSequenceBase* Animation)
+{
+	MeshComp->ResumeClothingSimulation();
+}
+
+FString UAnimNotify_ResumeClothingSimulation::GetNotifyName_Implementation() const
+{
+	return TEXT("Resume Clothing Sim");
+}
\ No newline at end of file
diff --git a/Engine/Source/Runtime/Engine/Private/Animation/AnimSequence.cpp b/Engine/Source/Runtime/Engine/Private/Animation/AnimSequence.cpp
--- a/Engine/Source/Runtime/Engine/Private/Animation/AnimSequence.cpp
+++ b/Engine/Source/Runtime/Engine/Private/Animation/AnimSequence.cpp
@@ -130,6 +130,7 @@ UAnimSequence::UAnimSequence(const FObjectInitializer& ObjectInitializer)
 	, Interpolation(EAnimInterpolationType::Linear)
 	, bEnableRootMotion(false)
 	, RootMotionRootLock(ERootMotionRootLock::RefPose)
+	, bUseNormalizedRootMotionScale(true)
 	, bRootMotionSettingsCopiedFromMontage(false)
 	, bUseRawDataOnly(!FPlatformProperties::RequiresCookedData())
 {
@@ -161,7 +162,7 @@ void UAnimSequence::GetAssetRegistryTags(TArray<FAssetRegistryTag>& OutTags) con
 
 	OutTags.Add(FAssetRegistryTag(TEXT("Compression Ratio"), FString::Printf(TEXT("%.03f"), (float)GetApproxCompressedSize() / (float)GetUncompressedRawSize()), FAssetRegistryTag::TT_Numerical));
 	OutTags.Add(FAssetRegistryTag(TEXT("Compressed Size (KB)"), FString::Printf(TEXT("%.02f"), (float)GetApproxCompressedSize() / 1024.0f), FAssetRegistryTag::TT_Numerical));
-
+	OutTags.Add(FAssetRegistryTag(TEXT("FrameRate"), FString::Printf(TEXT("%.2f"), GetFrameRate()), FAssetRegistryTag::TT_Numerical));
 	Super::GetAssetRegistryTags(OutTags);
 }
 
@@ -536,6 +537,11 @@ void UAnimSequence::PostLoad()
 
 #if WITH_EDITOR
 		VerifyCurveNames<FTransformCurve>(*CurrentSkeleton, USkeleton::AnimTrackCurveMappingName, RawCurveData.TransformCurves);
+
+		for (const FAnimSyncMarker& SyncMarker : AuthoredSyncMarkers)
+		{
+			CurrentSkeleton->RegisterMarkerName(SyncMarker.MarkerName);
+		}
 #endif
 	}
 
@@ -1003,9 +1009,21 @@ FTransform UAnimSequence::ExtractRootMotionFromRange(float StartTrackPosition, f
 	FTransform StartTransform = ExtractRootTrackTransform(StartTrackPosition, NULL);
 	FTransform EndTransform = ExtractRootTrackTransform(EndTrackPosition, NULL);
 
-	//Clear scale as it will muck up GetRelativeTransform
-	StartTransform.SetScale3D(FVector(1.f));
-	EndTransform.SetScale3D(FVector(1.f));
+	// Use old calculation if needed.
+	if (bUseNormalizedRootMotionScale)
+	{
+		//Clear scale as it will muck up GetRelativeTransform
+		StartTransform.SetScale3D(FVector(1.f));
+		EndTransform.SetScale3D(FVector(1.f));
+	}
+	else
+	{
+		if (IsValidAdditive())
+		{
+			StartTransform.SetScale3D(StartTransform.GetScale3D() + DefaultScale);
+			EndTransform.SetScale3D(EndTransform.GetScale3D() + DefaultScale);
+		}
+	}
 
 	// Transform to Component Space Rotation (inverse root transform from first frame)
 	const FTransform RootToComponentRot = FTransform(InitialTransform.GetRotation().Inverse());
diff --git a/Engine/Source/Runtime/Engine/Private/Animation/SkeletalMeshActor.cpp b/Engine/Source/Runtime/Engine/Private/Animation/SkeletalMeshActor.cpp
--- a/Engine/Source/Runtime/Engine/Private/Animation/SkeletalMeshActor.cpp
+++ b/Engine/Source/Runtime/Engine/Private/Animation/SkeletalMeshActor.cpp
@@ -18,7 +18,7 @@ ASkeletalMeshActor::ASkeletalMeshActor(const FObjectInitializer& ObjectInitializ
 {
 
 	SkeletalMeshComponent = CreateDefaultSubobject<USkeletalMeshComponent>(TEXT("SkeletalMeshComponent0"));
-	SkeletalMeshComponent->MeshComponentUpdateFlag = EMeshComponentUpdateFlag::AlwaysTickPose;
+	SkeletalMeshComponent->VisibilityBasedAnimTickOption = EVisibilityBasedAnimTickOption::AlwaysTickPose;
 	// check BaseEngine.ini for profile setup
 	SkeletalMeshComponent->SetCollisionProfileName(UCollisionProfile::PhysicsActor_ProfileName);
 	RootComponent = SkeletalMeshComponent;
diff --git a/Engine/Source/Runtime/Engine/Private/Animation/Skeleton.cpp b/Engine/Source/Runtime/Engine/Private/Animation/Skeleton.cpp
--- a/Engine/Source/Runtime/Engine/Private/Animation/Skeleton.cpp
+++ b/Engine/Source/Runtime/Engine/Private/Animation/Skeleton.cpp
@@ -27,6 +27,8 @@
 #include "Engine/PreviewMeshCollection.h"
 #include "Misc/ScopedSlowTask.h"
 #include "Animation/AnimBlueprint.h"
+#include "UObject/AnimObjectVersion.h"
+#include "EngineUtils.h"
 
 #define LOCTEXT_NAMESPACE "Skeleton"
 #define ROOT_BONE_PARENT	INDEX_NONE
@@ -153,6 +155,8 @@ void USkeleton::PostDuplicate(bool bDuplicateForPIE)
 
 void USkeleton::Serialize( FArchive& Ar )
 {
+	Ar.UsingCustomVersion(FAnimObjectVersion::GUID);
+
 	DECLARE_SCOPE_CYCLE_COUNTER( TEXT("USkeleton::Serialize"), STAT_Skeleton_Serialize, STATGROUP_LoadTime );
 
 	Super::Serialize(Ar);
@@ -237,6 +241,15 @@ void USkeleton::Serialize( FArchive& Ar )
 	}
 #endif
 
+	if (Ar.CustomVer(FAnimObjectVersion::GUID) >= FAnimObjectVersion::StoreMarkerNamesOnSkeleton)
+	{
+		FStripDataFlags StripFlags(Ar);
+		if (!StripFlags.IsEditorDataStripped())
+		{
+			Ar << ExistingMarkerNames;
+		}
+	}
+
 	const bool bRebuildNameMap = false;
 	ReferenceSkeleton.RebuildRefSkeleton(this, bRebuildNameMap);
 }
diff --git a/Engine/Source/Runtime/Engine/Private/Character.cpp b/Engine/Source/Runtime/Engine/Private/Character.cpp
--- a/Engine/Source/Runtime/Engine/Private/Character.cpp
+++ b/Engine/Source/Runtime/Engine/Private/Character.cpp
@@ -96,7 +96,7 @@ ACharacter::ACharacter(const FObjectInitializer& ObjectInitializer)
 		Mesh->AlwaysLoadOnClient = true;
 		Mesh->AlwaysLoadOnServer = true;
 		Mesh->bOwnerNoSee = false;
-		Mesh->MeshComponentUpdateFlag = EMeshComponentUpdateFlag::AlwaysTickPose;
+		Mesh->VisibilityBasedAnimTickOption = EVisibilityBasedAnimTickOption::AlwaysTickPose;
 		Mesh->bCastDynamicShadow = true;
 		Mesh->bAffectDynamicIndirectLighting = true;
 		Mesh->PrimaryComponentTick.TickGroup = TG_PrePhysics;
diff --git a/Engine/Source/Runtime/Engine/Private/Components/SkeletalMeshComponent.cpp b/Engine/Source/Runtime/Engine/Private/Components/SkeletalMeshComponent.cpp
--- a/Engine/Source/Runtime/Engine/Private/Components/SkeletalMeshComponent.cpp
+++ b/Engine/Source/Runtime/Engine/Private/Components/SkeletalMeshComponent.cpp
@@ -161,7 +161,7 @@ USkeletalMeshComponent::USkeletalMeshComponent(const FObjectInitializer& ObjectI
 	bWantsInitializeComponent = true;
 	GlobalAnimRateScale = 1.0f;
 	bNoSkeletonUpdate = false;
-	MeshComponentUpdateFlag = EMeshComponentUpdateFlag::AlwaysTickPoseAndRefreshBones;
+	VisibilityBasedAnimTickOption = EVisibilityBasedAnimTickOption::AlwaysTickPoseAndRefreshBones;
 	KinematicBonesUpdateType = EKinematicBonesUpdateToPhysics::SkipSimulatingBones;
 	PhysicsTransformUpdateMode = EPhysicsTransformUpdateMode::SimulationUpatesComponentTransform;
 	SetGenerateOverlapEvents(false);
@@ -516,7 +516,7 @@ void USkeletalMeshComponent::OnRegister()
 	// so we need to force initialize them before we begin to tick.
 	InitAnim(true);
 
-	if (bRenderStatic || (MeshComponentUpdateFlag == EMeshComponentUpdateFlag::OnlyTickPoseWhenRendered && !FApp::CanEverRender()))
+	if (bRenderStatic || (VisibilityBasedAnimTickOption == EVisibilityBasedAnimTickOption::OnlyTickPoseWhenRendered && !FApp::CanEverRender()))
 	{
 		SetComponentTickEnabled(false);
 	}
@@ -981,7 +981,7 @@ bool USkeletalMeshComponent::ShouldOnlyTickMontages(const float DeltaTime) const
 {
 	// Ignore DeltaSeconds == 0.f, as that is used when we want to force an update followed by RefreshBoneTransforms.
 	// RefreshBoneTransforms will need an updated graph.
-	return (MeshComponentUpdateFlag == EMeshComponentUpdateFlag::OnlyTickMontagesWhenNotRendered)
+	return (VisibilityBasedAnimTickOption == EVisibilityBasedAnimTickOption::OnlyTickMontagesWhenNotRendered)
 		&& !bRecentlyRendered 
 		&& (DeltaTime > 0.f);
 }
@@ -1108,7 +1108,7 @@ bool USkeletalMeshComponent::ShouldTickPose() const
 	const bool bShouldTickBasedOnAutonomousCheck = bIsAutonomousTickPose || (!bOnlyAllowAutonomousTickPose && !bAlreadyTickedThisFrame);
 	// When playing networked Root Motion Montages, we want these to play on dedicated servers and remote clients for networking and position correction purposes.
 	// So we force pose updates in that case to keep root motion and position in sync.
-	const bool bShouldTickBasedOnVisibility = ((MeshComponentUpdateFlag < EMeshComponentUpdateFlag::OnlyTickPoseWhenRendered) || bRecentlyRendered || IsPlayingNetworkedRootMotionMontage());
+	const bool bShouldTickBasedOnVisibility = ((VisibilityBasedAnimTickOption < EVisibilityBasedAnimTickOption::OnlyTickPoseWhenRendered) || bRecentlyRendered || IsPlayingNetworkedRootMotionMontage());
 
 	return (bShouldTickBasedOnVisibility && bShouldTickBasedOnAutonomousCheck && IsRegistered() && (AnimScriptInstance || PostProcessAnimInstance) && !bPauseAnims && GetWorld()->AreActorsInitialized() && !bNoSkeletonUpdate);
 }
@@ -1958,10 +1958,10 @@ void USkeletalMeshComponent::RefreshBoneTransforms(FActorComponentTickFunction*
 
 	const bool bDoPAE = !!CVarUseParallelAnimationEvaluation.GetValueOnGameThread() && FApp::ShouldUseThreadingForPerformance();
 
-const bool bMainInstanceValidForParallelWork = AnimScriptInstance == nullptr || AnimScriptInstance->CanRunParallelWork();
-const bool bPostInstanceValidForParallelWork = PostProcessAnimInstance == nullptr || PostProcessAnimInstance->CanRunParallelWork();
-const bool bHasValidInstanceForParallelWork = HasValidAnimationInstance() && bMainInstanceValidForParallelWork && bPostInstanceValidForParallelWork;
-const bool bDoParallelEvaluation = bHasValidInstanceForParallelWork && bDoPAE && bShouldDoEvaluation && TickFunction && (TickFunction->GetActualTickGroup() == TickFunction->TickGroup) && TickFunction->IsCompletionHandleValid();
+	const bool bMainInstanceValidForParallelWork = AnimScriptInstance == nullptr || AnimScriptInstance->CanRunParallelWork();
+	const bool bPostInstanceValidForParallelWork = PostProcessAnimInstance == nullptr || PostProcessAnimInstance->CanRunParallelWork();
+	const bool bHasValidInstanceForParallelWork = HasValidAnimationInstance() && bMainInstanceValidForParallelWork && bPostInstanceValidForParallelWork;
+	const bool bDoParallelEvaluation = bHasValidInstanceForParallelWork && bDoPAE && bShouldDoEvaluation && TickFunction && (TickFunction->GetActualTickGroup() == TickFunction->TickGroup) && TickFunction->IsCompletionHandleValid();
 	const bool bBlockOnTask = !bDoParallelEvaluation;  // If we aren't trying to do parallel evaluation then we
 															// will need to wait on an existing task.
 
@@ -2008,44 +2008,44 @@ const bool bDoParallelEvaluation = bHasValidInstanceForParallelWork && bDoPAE &&
 
 	if (bShouldDoEvaluation)
 	{
-	// If we need to eval the graph, and we're not going to update it.
-	// make sure it's been ticked at least once!
-	{
-		bool bShouldTickAnimation = false;		
-		if (AnimScriptInstance && !AnimScriptInstance->NeedsUpdate())
+		// If we need to eval the graph, and we're not going to update it.
+		// make sure it's been ticked at least once!
 		{
-			bShouldTickAnimation = bShouldTickAnimation || !AnimScriptInstance->GetUpdateCounter().HasEverBeenUpdated();
-			for (const UAnimInstance* SubInstance : SubInstances)
+			bool bShouldTickAnimation = false;		
+			if (AnimScriptInstance && !AnimScriptInstance->NeedsUpdate())
 			{
-				bShouldTickAnimation = bShouldTickAnimation || (SubInstance && !SubInstance->GetUpdateCounter().HasEverBeenUpdated());
+				bShouldTickAnimation = bShouldTickAnimation || !AnimScriptInstance->GetUpdateCounter().HasEverBeenUpdated();
+				for (const UAnimInstance* SubInstance : SubInstances)
+				{
+					bShouldTickAnimation = bShouldTickAnimation || (SubInstance && !SubInstance->GetUpdateCounter().HasEverBeenUpdated());
+				}
 			}
-		}
 
-		bShouldTickAnimation = bShouldTickAnimation || (ShouldPostUpdatePostProcessInstance() && !PostProcessAnimInstance->GetUpdateCounter().HasEverBeenUpdated());
+			bShouldTickAnimation = bShouldTickAnimation || (ShouldPostUpdatePostProcessInstance() && !PostProcessAnimInstance->GetUpdateCounter().HasEverBeenUpdated());
 
-		if (bShouldTickAnimation)
-		{
-			// We bypass TickPose() and call TickAnimation directly, so URO doesn't intercept us.
-			TickAnimation(0.f, false);
+			if (bShouldTickAnimation)
+			{
+				// We bypass TickPose() and call TickAnimation directly, so URO doesn't intercept us.
+				TickAnimation(0.f, false);
+			}
 		}
-	}
 
 		// If we're going to evaluate animation, call PreEvaluateAnimation()
 		{
-	if(AnimScriptInstance)
-	{
-		AnimScriptInstance->PreEvaluateAnimation();
+			if(AnimScriptInstance)
+			{
+				AnimScriptInstance->PreEvaluateAnimation();
 
-		for(UAnimInstance* SubInstance : SubInstances)
-		{
-			SubInstance->PreEvaluateAnimation();
-		}
-	}
+				for(UAnimInstance* SubInstance : SubInstances)
+				{
+					SubInstance->PreEvaluateAnimation();
+				}
+			}
 
-	if(ShouldEvaluatePostProcessInstance())
-	{
-		PostProcessAnimInstance->PreEvaluateAnimation();
-	}
+			if(ShouldEvaluatePostProcessInstance())
+			{
+				PostProcessAnimInstance->PreEvaluateAnimation();
+			}
 		}
 	}
 
@@ -2595,6 +2595,19 @@ UAnimInstance* USkeletalMeshComponent::GetPostProcessInstance() const
 	return PostProcessAnimInstance;
 }
 
+UAnimInstance* USkeletalMeshComponent::GetSubInstanceByName(FName InName) const
+{
+	for(UAnimInstance* SubInstance : SubInstances)
+	{
+		if(SubInstance->GetFName() == InName)
+		{
+			return SubInstance;
+		}
+	}
+
+	return nullptr;
+}
+
 bool USkeletalMeshComponent::HasValidAnimationInstance() const
 {
 	return AnimScriptInstance || PostProcessAnimInstance;
@@ -3435,7 +3448,25 @@ void USkeletalMeshComponent::FinalizeBoneTransform()
 
 void USkeletalMeshComponent::GetCurrentRefToLocalMatrices(TArray<FMatrix>& OutRefToLocals, int32 InLodIdx)
 {
-	UpdateRefToLocalMatrices(OutRefToLocals, this, SkeletalMesh->GetResourceForRendering(), InLodIdx, nullptr);
+	if(SkeletalMesh)
+	{
+		FSkeletalMeshRenderData* RenderData = SkeletalMesh->GetResourceForRendering();
+		if (ensureMsgf(RenderData->LODRenderData.IsValidIndex(InLodIdx),
+			TEXT("GetCurrentRefToLocalMatrices (SkelMesh :%s) input LODIndex (%d) doesn't match with render data size (%d)."),
+			*SkeletalMesh->GetPathName(), InLodIdx, RenderData->LODRenderData.Num()))
+		{
+			UpdateRefToLocalMatrices(OutRefToLocals, this, RenderData, InLodIdx, nullptr);
+		}
+		else
+		{
+			const FReferenceSkeleton& RefSkeleton = SkeletalMesh->RefSkeleton;
+			OutRefToLocals.AddUninitialized(RefSkeleton.GetNum());
+			for (int32 Index = 0; Index < OutRefToLocals.Num(); ++Index)
+			{
+				OutRefToLocals[Index] = FMatrix::Identity;
+			}
+		}
+	}
 }
 
 bool USkeletalMeshComponent::ShouldUpdatePostProcessInstance() const
diff --git a/Engine/Source/Runtime/Engine/Private/Components/SkinnedMeshComponent.cpp b/Engine/Source/Runtime/Engine/Private/Components/SkinnedMeshComponent.cpp
--- a/Engine/Source/Runtime/Engine/Private/Components/SkinnedMeshComponent.cpp
+++ b/Engine/Source/Runtime/Engine/Private/Components/SkinnedMeshComponent.cpp
@@ -324,7 +324,7 @@ USkinnedMeshComponent::USkinnedMeshComponent(const FObjectInitializer& ObjectIni
 	PrimaryComponentTick.bCanEverTick = true;
 	PrimaryComponentTick.TickGroup = TG_PrePhysics;
 
-	MeshComponentUpdateFlag = EMeshComponentUpdateFlag::AlwaysTickPoseAndRefreshBones;
+	VisibilityBasedAnimTickOption = EVisibilityBasedAnimTickOption::AlwaysTickPoseAndRefreshBones;
 
 	SetCollisionProfileName(UCollisionProfile::NoCollision_ProfileName);
 
@@ -595,7 +595,7 @@ void USkinnedMeshComponent::SendRenderDynamicData_Concurrent()
 
 	// if we have not updated the transforms then no need to send them to the rendering thread
 	// @todo GIsEditor used to be bUpdateSkelWhenNotRendered. Look into it further to find out why it doesn't update animations in the AnimSetViewer, when a level is loaded in UED (like POC_Cover.gear).
-	if( MeshObject && SkeletalMesh && (bForceMeshObjectUpdate || (bRecentlyRendered || MeshComponentUpdateFlag == EMeshComponentUpdateFlag::AlwaysTickPoseAndRefreshBones || GIsEditor || MeshObject->bHasBeenUpdatedAtLeastOnce == false)) )
+	if( MeshObject && SkeletalMesh && (bForceMeshObjectUpdate || (bRecentlyRendered || VisibilityBasedAnimTickOption == EVisibilityBasedAnimTickOption::AlwaysTickPoseAndRefreshBones || GIsEditor || MeshObject->bHasBeenUpdatedAtLeastOnce == false)) )
 	{
 		SCOPE_CYCLE_COUNTER(STAT_MeshObjectUpdate);
 
@@ -679,12 +679,12 @@ void USkinnedMeshComponent::InitLODInfos()
 
 bool USkinnedMeshComponent::ShouldTickPose() const
 {
-	return ((MeshComponentUpdateFlag < EMeshComponentUpdateFlag::OnlyTickPoseWhenRendered) || bRecentlyRendered);
+	return ((VisibilityBasedAnimTickOption < EVisibilityBasedAnimTickOption::OnlyTickPoseWhenRendered) || bRecentlyRendered);
 }
 
 bool USkinnedMeshComponent::ShouldUpdateTransform(bool bLODHasChanged) const
 {
-	return (bRecentlyRendered || (MeshComponentUpdateFlag == EMeshComponentUpdateFlag::AlwaysTickPoseAndRefreshBones));
+	return (bRecentlyRendered || (VisibilityBasedAnimTickOption == EVisibilityBasedAnimTickOption::AlwaysTickPoseAndRefreshBones));
 }
 
 bool USkinnedMeshComponent::ShouldUseUpdateRateOptimizations() const
@@ -753,7 +753,7 @@ void USkinnedMeshComponent::TickComponent(float DeltaTime, enum ELevelTick TickT
 			RefreshBoneTransforms(ThisTickFunction);
 		}
 	}
-	else if(MeshComponentUpdateFlag == EMeshComponentUpdateFlag::AlwaysTickPose)
+	else if(VisibilityBasedAnimTickOption == EVisibilityBasedAnimTickOption::AlwaysTickPose)
 	{
 		// We are not refreshing bone transforms, but we do want to tick pose. We may need to kick off a parallel task
 		DispatchParallelTickPose(ThisTickFunction);
@@ -1047,7 +1047,8 @@ FMatrix USkinnedMeshComponent::GetBoneMatrix(int32 BoneIdx) const
 	}
 	else
 	{
-		if(GetNumComponentSpaceTransforms() && BoneIdx < GetNumComponentSpaceTransforms() )
+		const int32 NumTransforms = GetNumComponentSpaceTransforms();
+		if(BoneIdx >= 0 && BoneIdx < NumTransforms)
 		{
 			return GetComponentSpaceTransforms()[BoneIdx].ToMatrixWithScale() * GetComponentTransform().ToMatrixWithScale();
 		}
@@ -1108,7 +1109,7 @@ FTransform USkinnedMeshComponent::GetBoneTransform(int32 BoneIdx, const FTransfo
 	else
 	{
 		const int32 NumTransforms = GetNumComponentSpaceTransforms();
-		if(NumTransforms > 0 && BoneIdx < NumTransforms)
+		if(BoneIdx >= 0 && BoneIdx < NumTransforms)
 		{
 			return GetComponentSpaceTransforms()[BoneIdx] * LocalToWorld;
 		}
diff --git a/Engine/Source/Runtime/Engine/Private/SkeletalMeshComponentPhysics.cpp b/Engine/Source/Runtime/Engine/Private/SkeletalMeshComponentPhysics.cpp
--- a/Engine/Source/Runtime/Engine/Private/SkeletalMeshComponentPhysics.cpp
+++ b/Engine/Source/Runtime/Engine/Private/SkeletalMeshComponentPhysics.cpp
@@ -3088,7 +3088,7 @@ void USkeletalMeshComponent::TickClothing(float DeltaTime, FTickFunction& ThisTi
 	}
 
 	// Use the component update flag to gate simulation to respect the always tick options
-	bool bShouldTick = ((MeshComponentUpdateFlag < EMeshComponentUpdateFlag::OnlyTickPoseWhenRendered) || bRecentlyRendered);
+	bool bShouldTick = ((VisibilityBasedAnimTickOption < EVisibilityBasedAnimTickOption::OnlyTickPoseWhenRendered) || bRecentlyRendered);
 
 	if (bShouldTick)
 	{
diff --git a/Engine/Source/Runtime/MovieSceneTracks/Private/Evaluation/MovieSceneSkeletalAnimationTemplate.cpp b/Engine/Source/Runtime/MovieSceneTracks/Private/Evaluation/MovieSceneSkeletalAnimationTemplate.cpp
--- a/Engine/Source/Runtime/MovieSceneTracks/Private/Evaluation/MovieSceneSkeletalAnimationTemplate.cpp
+++ b/Engine/Source/Runtime/MovieSceneTracks/Private/Evaluation/MovieSceneSkeletalAnimationTemplate.cpp
@@ -77,7 +77,7 @@ struct FPreAnimatedAnimationTokenProducer : IMovieScenePreAnimatedTokenProducer
 			FToken(USkeletalMeshComponent* InComponent)
 			{
 				// Cache this object's current update flag and animation mode
-				MeshComponentUpdateFlag = InComponent->MeshComponentUpdateFlag;
+				VisibilityBasedAnimTickOption = InComponent->VisibilityBasedAnimTickOption;
 				AnimationMode = InComponent->GetAnimationMode();
 			}
 
@@ -94,7 +94,7 @@ struct FPreAnimatedAnimationTokenProducer : IMovieScenePreAnimatedTokenProducer
 				UAnimSequencerInstance::UnbindFromSkeletalMeshComponent(Component);
 
 				// Reset the mesh component update flag and animation mode to what they were before we animated the object
-				Component->MeshComponentUpdateFlag = MeshComponentUpdateFlag;
+				Component->VisibilityBasedAnimTickOption = VisibilityBasedAnimTickOption;
 				if (Component->GetAnimationMode() != AnimationMode)
 				{
 					// this SetAnimationMode reinitializes even if the mode is same
@@ -103,7 +103,7 @@ struct FPreAnimatedAnimationTokenProducer : IMovieScenePreAnimatedTokenProducer
 				}
 			}
 
-			EMeshComponentUpdateFlag::Type MeshComponentUpdateFlag;
+			EVisibilityBasedAnimTickOption VisibilityBasedAnimTickOption;
 			EAnimationMode::Type AnimationMode;
 		};
 
