Fix crash when multi select with bone driven controller

- make sure you can't continue with mismatching types of classes, code becomes very vulnerable

diff --git a/Engine/Source/Editor/AnimGraph/Private/AnimGraphNode_BoneDrivenController.cpp b/Engine/Source/Editor/AnimGraph/Private/AnimGraphNode_BoneDrivenController.cpp
--- a/Engine/Source/Editor/AnimGraph/Private/AnimGraphNode_BoneDrivenController.cpp
+++ b/Engine/Source/Editor/AnimGraph/Private/AnimGraphNode_BoneDrivenController.cpp
@@ -357,6 +357,12 @@ void UAnimGraphNode_BoneDrivenController::CustomizeDetails(IDetailLayoutBuilder&
 
 	TSharedRef<IPropertyHandle> NodeHandle = DetailBuilder.GetProperty(GET_MEMBER_NAME_CHECKED(UAnimGraphNode_BoneDrivenController, Node), GetClass());
 
+	// if it doesn't have this node, that means this isn't really bone driven controller
+	if (!NodeHandle->IsValidHandle())
+	{
+		return;
+	}
+
 	TAttribute<EVisibility> NotUsingCurveVisibility = TAttribute<EVisibility>::Create(TAttribute<EVisibility>::FGetter::CreateStatic(&UAnimGraphNode_BoneDrivenController::AreNonCurveMappingValuesVisible, &DetailBuilder));
 	TAttribute<EVisibility> MapRangeVisiblity = TAttribute<EVisibility>::Create(TAttribute<EVisibility>::FGetter::CreateStatic(&UAnimGraphNode_BoneDrivenController::AreRemappingValuesVisible, &DetailBuilder));
 
diff --git a/Engine/Source/Editor/Persona/Private/AnimGraphNodeDetails.cpp b/Engine/Source/Editor/Persona/Private/AnimGraphNodeDetails.cpp
--- a/Engine/Source/Editor/Persona/Private/AnimGraphNodeDetails.cpp
+++ b/Engine/Source/Editor/Persona/Private/AnimGraphNodeDetails.cpp
@@ -63,17 +63,29 @@ void FAnimGraphNodeDetails::CustomizeDetails(class IDetailLayoutBuilder& DetailB
 
 	// get first animgraph nodes
 	UAnimGraphNode_Base* AnimGraphNode = Cast<UAnimGraphNode_Base>(SelectedObjectsList[0].Get());
-	if (AnimGraphNode == NULL)
+	if (AnimGraphNode == nullptr)
 	{
 		return;
 	}
 
+	// make sure type matches with all the nodes. 
+	const UAnimGraphNode_Base* FirstNodeType = AnimGraphNode;
+	for (int32 Index = 1; Index < SelectedObjectsList.Num(); ++Index)
+	{
+		UAnimGraphNode_Base* CurrentNode = Cast<UAnimGraphNode_Base>(SelectedObjectsList[Index].Get());
+		if (!CurrentNode || CurrentNode->GetClass() != FirstNodeType->GetClass())
+		{
+			// if type mismatches, multi selection doesn't work, just return
+			return;
+		}
+	}
+
 	TargetSkeleton = AnimGraphNode->GetAnimBlueprint()->TargetSkeleton;
 	TargetSkeletonName = FString::Printf(TEXT("%s'%s'"), *TargetSkeleton->GetClass()->GetName(), *TargetSkeleton->GetPathName());
 
 	// Get the node property
 	const UStructProperty* NodeProperty = AnimGraphNode->GetFNodeProperty();
-	if (NodeProperty == NULL)
+	if (NodeProperty == nullptr)
 	{
 		return;
 	}
