Fixed crash when creating a new notify in certain circumstances

When creating a notify causes a timing order change, then the timing panel wasnt getting refreshed and was accessing the wrong notify to do its drawing.

diff --git a/Engine/Source/Editor/Persona/Private/SAnimNotifyPanel.cpp b/Engine/Source/Editor/Persona/Private/SAnimNotifyPanel.cpp
--- a/Engine/Source/Editor/Persona/Private/SAnimNotifyPanel.cpp
+++ b/Engine/Source/Editor/Persona/Private/SAnimNotifyPanel.cpp
@@ -3886,6 +3886,7 @@ void SAnimNotifyPanel::Construct(const FArguments& InArgs, const TSharedRef<clas
 	Sequence = InArgs._Sequence;
 	MarkerBars = InArgs._MarkerBars;
 	OnInvokeTab = InArgs._OnInvokeTab;
+	OnNotifiesChanged = InArgs._OnNotifiesChanged;
 
 	FAnimNotifyPanelCommands::Register();
 	BindCommands();
@@ -4105,6 +4106,10 @@ void SAnimNotifyPanel::Update()
 	{
 		Sequence->RefreshCacheData();
 	}
+
+	RefreshNotifyTracks();
+
+	OnNotifiesChanged.ExecuteIfBound();
 }
 
 void SAnimNotifyPanel::RefreshNotifyTracks()
diff --git a/Engine/Source/Editor/Persona/Private/SAnimNotifyPanel.h b/Engine/Source/Editor/Persona/Private/SAnimNotifyPanel.h
--- a/Engine/Source/Editor/Persona/Private/SAnimNotifyPanel.h
+++ b/Engine/Source/Editor/Persona/Private/SAnimNotifyPanel.h
@@ -184,6 +184,7 @@ public:
 	SLATE_EVENT( FRefreshOffsetsRequest, OnRequestRefreshOffsets )
 	SLATE_EVENT( FOnGetTimingNodeVisibility, OnGetTimingNodeVisibility )
 	SLATE_EVENT( FOnInvokeTab, OnInvokeTab )
+	SLATE_EVENT( FSimpleDelegate, OnNotifiesChanged )
 
 	SLATE_END_ARGS()
 
@@ -340,4 +341,7 @@ private:
 
 	/** Delegate used to invoke a tab */
 	FOnInvokeTab OnInvokeTab;
+
+	/** Delegate used to inform others that notifies have changed (for timing) */
+	FSimpleDelegate OnNotifiesChanged;
 };
diff --git a/Engine/Source/Editor/Persona/Private/SMontageEditor.cpp b/Engine/Source/Editor/Persona/Private/SMontageEditor.cpp
--- a/Engine/Source/Editor/Persona/Private/SMontageEditor.cpp
+++ b/Engine/Source/Editor/Persona/Private/SMontageEditor.cpp
@@ -199,6 +199,7 @@ void SMontageEditor::Construct(const FArguments& InArgs, const FMontageEditorReq
 		.OnRequestRefreshOffsets(this, &SMontageEditor::RefreshNotifyTriggerOffsets)
 		.OnGetTimingNodeVisibility(TimingNodeVisibilityDelegate)
 		.OnInvokeTab(InArgs._OnInvokeTab)
+		.OnNotifiesChanged(this, &SMontageEditor::HandleNotifiesChanged)
 	];
 
 	EditorPanels->AddSlot()
@@ -333,6 +334,11 @@ void SMontageEditor::RefreshNotifyTriggerOffsets()
 	}
 }
 
+void SMontageEditor::HandleNotifiesChanged()
+{
+	AnimTimingPanel->Update();
+}
+
 bool SMontageEditor::GetSectionTime( int32 SectionIndex, float &OutTime ) const
 {
 	if (MontageObj != NULL && MontageObj->CompositeSections.IsValidIndex(SectionIndex))
diff --git a/Engine/Source/Editor/Persona/Private/SMontageEditor.h b/Engine/Source/Editor/Persona/Private/SMontageEditor.h
--- a/Engine/Source/Editor/Persona/Private/SMontageEditor.h
+++ b/Engine/Source/Editor/Persona/Private/SMontageEditor.h
@@ -131,6 +131,9 @@ private:
 	/** Updates Notify trigger offsets to take into account current montage state */
 	void RefreshNotifyTriggerOffsets();
 
+	/** Handle notifies changing, so we rebuild timing */
+	void HandleNotifiesChanged();
+
 	/** One-off active timer to trigger a montage panel rebuild */
 	EActiveTimerReturnType TriggerRebuildMontagePanel(double InCurrentTime, float InDeltaTime);
 
