Enable or disable the morph target weight slider depending of the project settings.

diff --git a/Engine/Source/Editor/Persona/Private/SMorphTargetViewer.cpp b/Engine/Source/Editor/Persona/Private/SMorphTargetViewer.cpp
--- a/Engine/Source/Editor/Persona/Private/SMorphTargetViewer.cpp
+++ b/Engine/Source/Editor/Persona/Private/SMorphTargetViewer.cpp
@@ -13,6 +13,8 @@
 #include "Animation/MorphTarget.h"
 #include "Animation/AnimInstance.h"
 #include "HAL/PlatformApplicationMisc.h"
+#include "GPUSkinCache.h"
+#include "Engine/RendererSettings.h"
 
 #define LOCTEXT_NAMESPACE "SMorphTargetViewer"
 
@@ -60,6 +62,7 @@ private:
 	*
 	*/
 	void OnMorphTargetWeightChanged( float NewWeight );
+	
 	/**
 	* Called when the user types the value and enters
 	*
@@ -67,6 +70,16 @@ private:
 	*
 	*/
 	void OnMorphTargetWeightValueCommitted( float NewWeight, ETextCommit::Type CommitType);
+	
+	/**
+	* Called to know if we enable or disable the weight sliders
+	*/
+	bool IsMorphTargetWeightSliderEnabled() const;
+	
+	/**
+	* Show the tooltip for the weight widget
+	*/
+	FText GetMorphTargetWeightSliderToolTip() const;
 
 	/** Auto fill check call back functions */
 	void OnMorphTargetAutoFillChecked(ECheckBoxState InState);
@@ -141,6 +154,8 @@ TSharedRef< SWidget > SMorphTargetListRow::GenerateWidgetForColumn( const FName&
 				.Value( this, &SMorphTargetListRow::GetWeight )
 				.OnValueChanged( this, &SMorphTargetListRow::OnMorphTargetWeightChanged )
 				.OnValueCommitted( this, &SMorphTargetListRow::OnMorphTargetWeightValueCommitted )
+				.IsEnabled(this, &SMorphTargetListRow::IsMorphTargetWeightSliderEnabled)
+				.ToolTipText(this, &SMorphTargetListRow::GetMorphTargetWeightSliderToolTip)
 			];
 	}
 	else if ( ColumnName == ColumnID_MorphTargetEditLabel )
@@ -266,6 +281,21 @@ void SMorphTargetListRow::OnMorphTargetWeightValueCommitted( float NewWeight, ET
 	}
 }
 
+bool SMorphTargetListRow::IsMorphTargetWeightSliderEnabled() const
+{
+	const uint32 CVarMorphTargetModeValue = GetDefault<URendererSettings>()->bUseGPUMorphTargets;
+	return GEnableGPUSkinCache > 0 ? CVarMorphTargetModeValue > 0 : true;
+}
+
+FText SMorphTargetListRow::GetMorphTargetWeightSliderToolTip() const
+{
+	if (!IsMorphTargetWeightSliderEnabled())
+	{
+		return LOCTEXT("MorphTargetWeightSliderTooltip", "When using skin cache, the morph target must use the GPU to affect the mesh");
+	}
+	return FText();
+}
+
 float SMorphTargetListRow::GetWeight() const 
 { 
 	if (Item->bAutoFillData)
