Merging using Dev-Anim<->Release-4.22

Original CL 5117200

diff --git a/Engine/Source/Editor/Persona/Private/AnimGraphNodeDetails.cpp b/Engine/Source/Editor/Persona/Private/AnimGraphNodeDetails.cpp
--- a/Engine/Source/Editor/Persona/Private/AnimGraphNodeDetails.cpp
+++ b/Engine/Source/Editor/Persona/Private/AnimGraphNodeDetails.cpp
@@ -40,6 +40,7 @@
 #include "Widgets/Input/SNumericEntryBox.h"
 #include "BoneControllers/AnimNode_SkeletalControlBase.h"
 #include "AnimGraphNode_StateMachine.h"
+#include "AnimGraphNode_SequencePlayer.h"
 #include "Engine/SkeletalMeshSocket.h"
 #include "Styling/CoreStyle.h"
 #include "Widgets/Input/SCheckBox.h"
@@ -884,6 +885,15 @@ void FAnimGraphParentPlayerDetails::CustomizeDetails(class IDetailLayoutBuilder&
 				{
 					// Don't create a node for these, the graph speaks for it
 				}
+				else if (UAnimGraphNode_SequencePlayer* SequencePlayer = Cast<UAnimGraphNode_SequencePlayer>(Object))
+				{
+					UAnimationAsset * AnimAsset = SequencePlayer->Node.GetAnimAsset();
+					if (AnimAsset != nullptr)
+					{
+						FString Title = FString::Printf(TEXT("Play %s"), *AnimAsset->GetName());
+						ThisNode = MakeShareable(new FPlayerTreeViewEntry(Title, EPlayerTreeViewEntryType::Node));
+					}
+				}
 				else if (UAnimGraphNode_Base* Node = Cast<UAnimGraphNode_Base>(Object))
 				{
 					ThisNode = MakeShareable(new FPlayerTreeViewEntry(Node->GetName(), EPlayerTreeViewEntryType::Node));
