Merging using Dev-Anim<->Release-4.22

Original CL 5202790

diff --git a/Engine/Source/Editor/Persona/Private/AnimGraphNodeDetails.cpp b/Engine/Source/Editor/Persona/Private/AnimGraphNodeDetails.cpp
--- a/Engine/Source/Editor/Persona/Private/AnimGraphNodeDetails.cpp
+++ b/Engine/Source/Editor/Persona/Private/AnimGraphNodeDetails.cpp
@@ -885,14 +885,10 @@ void FAnimGraphParentPlayerDetails::CustomizeDetails(class IDetailLayoutBuilder&
 				{
 					// Don't create a node for these, the graph speaks for it
 				}
-				else if (UAnimGraphNode_SequencePlayer* SequencePlayer = Cast<UAnimGraphNode_SequencePlayer>(Object))
+				else if (UAnimGraphNode_AssetPlayerBase * AssetPlayerBase = Cast<UAnimGraphNode_AssetPlayerBase>(Object))
 				{
-					UAnimationAsset * AnimAsset = SequencePlayer->Node.GetAnimAsset();
-					if (AnimAsset != nullptr)
-					{
-						FString Title = FString::Printf(TEXT("Play %s"), *AnimAsset->GetName());
-						ThisNode = MakeShareable(new FPlayerTreeViewEntry(Title, EPlayerTreeViewEntryType::Node));
-					}
+					FString Title = AssetPlayerBase->GetNodeTitle(ENodeTitleType::FullTitle).ToString();
+					ThisNode = MakeShareable(new FPlayerTreeViewEntry(Title, EPlayerTreeViewEntryType::Node));
 				}
 				else if (UAnimGraphNode_Base* Node = Cast<UAnimGraphNode_Base>(Object))
 				{
