Fix for bug where you could select virtual bones as rig targets

diff --git a/Engine/Source/Editor/Persona/Private/SBoneMappingBase.cpp b/Engine/Source/Editor/Persona/Private/SBoneMappingBase.cpp
--- a/Engine/Source/Editor/Persona/Private/SBoneMappingBase.cpp
+++ b/Engine/Source/Editor/Persona/Private/SBoneMappingBase.cpp
@@ -83,6 +83,7 @@ TSharedRef< SWidget > SBoneMappingListRow::GenerateWidgetForColumn( const FName&
 					.OnBoneSelectionChanged(this, &SBoneMappingListRow::OnBoneSelectionChanged)
 					.OnGetSelectedBone(this, &SBoneMappingListRow::GetSelectedBone)
 					.OnGetReferenceSkeleton(OnGetReferenceSkeleton)
+					.bShowVirtualBones(false)
 				]
 
 				+SHorizontalBox::Slot()
diff --git a/Engine/Source/Editor/Persona/Private/Shared/BoneSelectionWidget.cpp b/Engine/Source/Editor/Persona/Private/Shared/BoneSelectionWidget.cpp
--- a/Engine/Source/Editor/Persona/Private/Shared/BoneSelectionWidget.cpp
+++ b/Engine/Source/Editor/Persona/Private/Shared/BoneSelectionWidget.cpp
@@ -233,6 +233,7 @@ void SBoneSelectionWidget::Construct(const FArguments& InArgs)
 	OnGetReferenceSkeleton = InArgs._OnGetReferenceSkeleton;
 	OnGetSocketList = InArgs._OnGetSocketList;
 	bShowSocket = InArgs._bShowSocket;
+	bShowVirtualBones = InArgs._bShowVirtualBones;
 
 	SuppliedToolTip = InArgs._ToolTipText.Get();
 
@@ -265,7 +266,8 @@ TSharedRef<SWidget> SBoneSelectionWidget::CreateSkeletonWidgetMenu()
 		.OnGetReferenceSkeleton(OnGetReferenceSkeleton)
 		.OnGetSocketList(OnGetSocketList)
 		.bShowSocket(bShowSocket)
-		.SelectedBone(CurrentBoneName);
+		.SelectedBone(CurrentBoneName)
+		.bShowVirtualBones(bShowVirtualBones);
 
 	BonePickerButton->SetMenuContentWidgetToFocus(MenuWidget->GetFilterTextWidget());
 
diff --git a/Engine/Source/Editor/Persona/Public/BoneSelectionWidget.h b/Engine/Source/Editor/Persona/Public/BoneSelectionWidget.h
--- a/Engine/Source/Editor/Persona/Public/BoneSelectionWidget.h
+++ b/Engine/Source/Editor/Persona/Public/BoneSelectionWidget.h
@@ -103,6 +103,7 @@ public:
 
 	SLATE_BEGIN_ARGS( SBoneSelectionWidget )
 		: _bShowSocket(false)
+		, _bShowVirtualBones(true)
 		, _OnBoneSelectionChanged()
 		, _OnGetSelectedBone()
 		, _OnGetReferenceSkeleton()
@@ -110,6 +111,10 @@ public:
 	{}
 
 		SLATE_ARGUMENT(bool, bShowSocket)
+
+		/** Should show skeletons virtual bones in tree */
+		SLATE_ARGUMENT(bool, bShowVirtualBones)
+
 		/** set selected bone name */
 		SLATE_EVENT(FOnBoneSelectionChanged, OnBoneSelectionChanged);
 
@@ -152,6 +157,7 @@ private:
 	FGetReferenceSkeleton	OnGetReferenceSkeleton;
 	FGetSocketList			OnGetSocketList;
 	bool bShowSocket;
+	bool bShowVirtualBones;
 
 	// Cache supplied tooltip
 	FText SuppliedToolTip;
