Fix crash when editing notifies in a montage wwhen one of its child montages is also open for editing

diff --git a/Engine/Source/Editor/Persona/Private/SAnimNotifyPanel.cpp b/Engine/Source/Editor/Persona/Private/SAnimNotifyPanel.cpp
--- a/Engine/Source/Editor/Persona/Private/SAnimNotifyPanel.cpp
+++ b/Engine/Source/Editor/Persona/Private/SAnimNotifyPanel.cpp
@@ -262,6 +262,7 @@ struct FNotifyNodeInterface : public INodeObjectInterface
 			if (NotifyEvent == &(Seq->Notifies[I]))
 			{
 				Seq->Notifies.RemoveAt(I);
+				Seq->PostEditChange();
 				Seq->MarkPackageDirty();
 				break;
 			}
@@ -338,6 +339,7 @@ struct FSyncMarkerNodeInterface : public INodeObjectInterface
 				if (SyncMarker == &(Seq->AuthoredSyncMarkers[I]))
 				{
 					Seq->AuthoredSyncMarkers.RemoveAt(I);
+					Seq->PostEditChange();
 					Seq->MarkPackageDirty();
 					break;
 				}
@@ -1078,7 +1080,9 @@ public:
 				Node->DropCancelled();
 			}
 
+			Sequence->PostEditChange();
 			Sequence->MarkPackageDirty();
+			
 			OnUpdatePanel.ExecuteIfBound();
 		}
 		
@@ -1959,6 +1963,9 @@ FReply SAnimNotifyNode::OnMouseButtonUp( const FGeometry& MyGeometry, const FPoi
 		GEditor->EndTransaction();
 		DragMarkerTransactionIdx = INDEX_NONE;
 
+		Sequence->PostEditChange();
+		Sequence->MarkPackageDirty();
+
 		return FReply::Handled().ReleaseMouseCapture();
 	}
 
@@ -2579,6 +2586,7 @@ FAnimNotifyEvent& SAnimNotifyTrack::CreateNewNotify(FString NewNotifyName, UClas
 		NewEvent.NotifyStateClass->OnAnimNotifyCreatedInEditor(NewEvent);
 	}
 
+	Sequence->PostEditChange();
 	Sequence->MarkPackageDirty();
 
 	return NewEvent;
@@ -2610,6 +2618,9 @@ void SAnimNotifyTrack::CreateNewSyncMarkerAtCursor(FString NewSyncMarkerName, UC
 	SyncMarker.MarkerName = FName(*NewSyncMarkerName);
 	SyncMarker.TrackIndex = TrackIndex;
 	SyncMarker.Time = LastClickedTime;
+
+	Seq->PostEditChange();
+	Seq->MarkPackageDirty();
 	OnUpdatePanel.ExecuteIfBound();
 
 	UBlendSpaceBase::UpdateBlendSpacesUsingAnimSequence(Seq);
@@ -3624,6 +3635,7 @@ void SAnimNotifyTrack::PasteSingleNotify(FString& NotifyString, float PasteTime)
 	}
 
 	OnDeselectAllNotifies.ExecuteIfBound();
+	Sequence->PostEditChange();
 	Sequence->MarkPackageDirty();
 	OnUpdatePanel.ExecuteIfBound();
 }
@@ -3660,6 +3672,7 @@ void SAnimNotifyTrack::PasteSingleSyncMarker(FString& MarkerString, float PasteT
 		UBlendSpaceBase::UpdateBlendSpacesUsingAnimSequence(Sequence);
 
 		OnDeselectAllNotifies.ExecuteIfBound();
+		Sequence->PostEditChange();
 		Sequence->MarkPackageDirty();
 		OnUpdatePanel.ExecuteIfBound();
 	}
@@ -4035,6 +4048,7 @@ FReply SAnimNotifyPanel::InsertTrack(int32 TrackIndexToInsert)
 	NewItem.TrackColor = FLinearColor::White;
 
 	Sequence->AnimNotifyTracks.Insert(NewItem, TrackIndexToInsert);
+	Sequence->PostEditChange();
 	Sequence->MarkPackageDirty();
 
 	Update();
@@ -4077,6 +4091,7 @@ FReply SAnimNotifyPanel::DeleteTrack(int32 TrackIndexToDelete)
 			}
 
 			Sequence->AnimNotifyTracks.RemoveAt(TrackIndexToDelete);
+			Sequence->PostEditChange();
 			Sequence->MarkPackageDirty();
 			Update();
 		}
@@ -4511,6 +4526,10 @@ void SAnimNotifyPanel::OnReplaceSelectedWithNotify(FString NewNotifyName, UClass
 	OnSelectionChanged.ExecuteIfBound(Objects);
 	// TODO: set selection to new notifies?
 	// update the panel
+
+	Sequence->PostEditChange();
+	Sequence->MarkPackageDirty();
+
 	Update();
 }
 
diff --git a/Engine/Source/Runtime/Engine/Private/Animation/AnimMontage.cpp b/Engine/Source/Runtime/Engine/Private/Animation/AnimMontage.cpp
--- a/Engine/Source/Runtime/Engine/Private/Animation/AnimMontage.cpp
+++ b/Engine/Source/Runtime/Engine/Private/Animation/AnimMontage.cpp
@@ -1097,10 +1097,7 @@ void UAnimMontage::RefreshParentAssetData()
 		}
 	}
 
-	// this delegate causes it to reconstrucand this code can be called by UI. 
-	// that is dangerous as it can cause the UI to reconstruct in the middle of it. 
-	// until this, the multi window won't work well. 
-	//OnMontageChanged.Broadcast();
+	OnMontageChanged.Broadcast();
 }
 
 #endif
