Fixed incorrect displayed LOD level when editing skeletal mesh LOD counts in external sub-editors

diff --git a/Engine/Source/Editor/Persona/Private/SAnimationEditorViewport.cpp b/Engine/Source/Editor/Persona/Private/SAnimationEditorViewport.cpp
--- a/Engine/Source/Editor/Persona/Private/SAnimationEditorViewport.cpp
+++ b/Engine/Source/Editor/Persona/Private/SAnimationEditorViewport.cpp
@@ -1564,20 +1564,23 @@ int32 SAnimationEditorViewportTabBody::GetLODSelection() const
 
 	if (PreviewComponent)
 	{
-		return PreviewComponent->ForcedLodModel;
+		// If we are forcing a LOD level, report the actual LOD level we are displaying
+		// as the mesh can potentially change LOD count under the viewport.
+		if(PreviewComponent->ForcedLodModel > 0)
+		{
+			return PreviewComponent->PredictedLODLevel + 1;
+		}
+		else
+		{
+			return PreviewComponent->ForcedLodModel;
+		}
 	}
 	return 0;
 }
 
 bool SAnimationEditorViewportTabBody::IsLODModelSelected(int32 LODSelectionType) const
 {
-	UDebugSkelMeshComponent* PreviewComponent = GetPreviewScene()->GetPreviewMeshComponent();
-
-	if (PreviewComponent)
-	{
-		return (PreviewComponent->ForcedLodModel == LODSelectionType) ? true : false;
-	}
-	return false;
+	return GetLODSelection() == LODSelectionType;
 }
 
 void SAnimationEditorViewportTabBody::OnSetLODModel(int32 LODSelectionType)
