Adding pre-skinned bounds debug drawing

diff --git a/Engine/Source/Editor/Persona/Private/AnimViewportShowCommands.cpp b/Engine/Source/Editor/Persona/Private/AnimViewportShowCommands.cpp
--- a/Engine/Source/Editor/Persona/Private/AnimViewportShowCommands.cpp
+++ b/Engine/Source/Editor/Persona/Private/AnimViewportShowCommands.cpp
@@ -17,6 +17,7 @@ void FAnimViewportShowCommands::RegisterCommands()
 	UI_COMMAND( ShowBound, "Bound", "Show bound on preview mesh", EUserInterfaceActionType::ToggleButton, FInputChord() );
 	UI_COMMAND( UseInGameBound, "In-game Bound", "Use in-game bound on preview mesh when showing bounds. Otherwise bounds will always be calculated from bones alone.", EUserInterfaceActionType::ToggleButton, FInputChord());
 	UI_COMMAND( UseFixedBounds, "Fixed Bounds", "Preview using the 'Fixed Bounds' option, which will use bounds from skel mesh and not bones at all, for speed.", EUserInterfaceActionType::ToggleButton, FInputChord());
+	UI_COMMAND( UsePreSkinnedBounds, "Pre-skinned Bounds", "Preview using the 'Pre-skinned' option, which will use local space pre-skinned(base pose) bounds from skel mesh. Overrides all other bounds preview settings.", EUserInterfaceActionType::ToggleButton, FInputChord());
 	UI_COMMAND( ShowPreviewMesh, "Mesh", "Show the preview mesh", EUserInterfaceActionType::ToggleButton, FInputChord() );
 	UI_COMMAND( ShowMorphTargets, "Morph Targets", "Display applied morph targets of the mesh", EUserInterfaceActionType::ToggleButton, FInputChord() );
 
diff --git a/Engine/Source/Editor/Persona/Private/AnimViewportShowCommands.h b/Engine/Source/Editor/Persona/Private/AnimViewportShowCommands.h
--- a/Engine/Source/Editor/Persona/Private/AnimViewportShowCommands.h
+++ b/Engine/Source/Editor/Persona/Private/AnimViewportShowCommands.h
@@ -51,6 +51,9 @@ public:
 	/** Use in-game Bound of preview mesh */
 	TSharedPtr< FUICommandInfo > UseFixedBounds;
 
+	/** Use pre-skinned Bound of preview mesh */
+	TSharedPtr< FUICommandInfo > UsePreSkinnedBounds;
+
 	/** Show/hide the preview mesh */
 	TSharedPtr< FUICommandInfo > ShowPreviewMesh;
 
diff --git a/Engine/Source/Editor/Persona/Private/SAnimViewportToolBar.cpp b/Engine/Source/Editor/Persona/Private/SAnimViewportToolBar.cpp
--- a/Engine/Source/Editor/Persona/Private/SAnimViewportToolBar.cpp
+++ b/Engine/Source/Editor/Persona/Private/SAnimViewportToolBar.cpp
@@ -602,6 +602,7 @@ TSharedRef<SWidget> SAnimViewportToolBar::GenerateCharacterMenu() const
 						SubMenuBuilder.AddMenuEntry(FAnimViewportShowCommands::Get().ShowBound );
 						SubMenuBuilder.AddMenuEntry(FAnimViewportShowCommands::Get().UseInGameBound);
 						SubMenuBuilder.AddMenuEntry(FAnimViewportShowCommands::Get().UseFixedBounds);
+						SubMenuBuilder.AddMenuEntry(FAnimViewportShowCommands::Get().UsePreSkinnedBounds);
 						SubMenuBuilder.AddMenuEntry(FAnimViewportShowCommands::Get().ShowPreviewMesh );
 						SubMenuBuilder.AddMenuEntry(FAnimViewportShowCommands::Get().ShowMorphTargets );
 					}
diff --git a/Engine/Source/Editor/Persona/Private/SAnimationEditorViewport.cpp b/Engine/Source/Editor/Persona/Private/SAnimationEditorViewport.cpp
--- a/Engine/Source/Editor/Persona/Private/SAnimationEditorViewport.cpp
+++ b/Engine/Source/Editor/Persona/Private/SAnimationEditorViewport.cpp
@@ -640,6 +640,12 @@ void SAnimationEditorViewportTabBody::BindCommands()
 		FCanExecuteAction::CreateSP(this, &SAnimationEditorViewportTabBody::CanUseFixedBounds),
 		FIsActionChecked::CreateSP(this, &SAnimationEditorViewportTabBody::IsUsingFixedBounds));
 
+	CommandList.MapAction(
+		ViewportShowMenuCommands.UsePreSkinnedBounds,
+		FExecuteAction::CreateSP(this, &SAnimationEditorViewportTabBody::UsePreSkinnedBounds),
+		FCanExecuteAction::CreateSP(this, &SAnimationEditorViewportTabBody::CanUsePreSkinnedBounds),
+		FIsActionChecked::CreateSP(this, &SAnimationEditorViewportTabBody::IsUsingPreSkinnedBounds));
+
 	CommandList.MapAction(
 		ViewportShowMenuCommands.ShowPreviewMesh,
 		FExecuteAction::CreateSP(this, &SAnimationEditorViewportTabBody::ToggleShowPreviewMesh),
@@ -1388,6 +1394,27 @@ bool SAnimationEditorViewportTabBody::IsUsingFixedBounds() const
 	return PreviewComponent != NULL && PreviewComponent->bComponentUseFixedSkelBounds;
 }
 
+void SAnimationEditorViewportTabBody::UsePreSkinnedBounds()
+{
+	UDebugSkelMeshComponent* PreviewComponent = GetPreviewScene()->GetPreviewMeshComponent();
+	if (PreviewComponent != NULL)
+	{
+		PreviewComponent->UsePreSkinnedBounds(!PreviewComponent->IsUsingPreSkinnedBounds());
+	}
+}
+
+bool SAnimationEditorViewportTabBody::CanUsePreSkinnedBounds() const
+{
+	UDebugSkelMeshComponent* PreviewComponent = GetPreviewScene()->GetPreviewMeshComponent();
+	return PreviewComponent != NULL && IsShowBoundEnabled();
+}
+
+bool SAnimationEditorViewportTabBody::IsUsingPreSkinnedBounds() const
+{
+	UDebugSkelMeshComponent* PreviewComponent = GetPreviewScene()->GetPreviewMeshComponent();
+	return PreviewComponent != NULL && PreviewComponent->IsUsingPreSkinnedBounds();
+}
+
 void SAnimationEditorViewportTabBody::HandlePreviewMeshChanged(class USkeletalMesh* OldSkeletalMesh, class USkeletalMesh* NewSkeletalMesh)
 {
 	PopulateNumUVChannels();
diff --git a/Engine/Source/Editor/Persona/Private/SAnimationEditorViewport.h b/Engine/Source/Editor/Persona/Private/SAnimationEditorViewport.h
--- a/Engine/Source/Editor/Persona/Private/SAnimationEditorViewport.h
+++ b/Engine/Source/Editor/Persona/Private/SAnimationEditorViewport.h
@@ -415,6 +415,11 @@ private:
 	bool CanUseFixedBounds() const;
 	bool IsUsingFixedBounds() const;
 
+	/** Called to toggle 'pre-skinned' option to preview */
+	void UsePreSkinnedBounds();
+	bool CanUsePreSkinnedBounds() const;
+	bool IsUsingPreSkinnedBounds() const;
+
 	/** Called by UV channel combo box on selection change */
 	void ComboBoxSelectionChanged( TSharedPtr<FString> NewSelection, ESelectInfo::Type SelectInfo );
 
diff --git a/Engine/Source/Editor/UnrealEd/Classes/Animation/DebugSkelMeshComponent.h b/Engine/Source/Editor/UnrealEd/Classes/Animation/DebugSkelMeshComponent.h
--- a/Engine/Source/Editor/UnrealEd/Classes/Animation/DebugSkelMeshComponent.h
+++ b/Engine/Source/Editor/UnrealEd/Classes/Animation/DebugSkelMeshComponent.h
@@ -269,6 +269,10 @@ class UNREALED_API UDebugSkelMeshComponent : public USkeletalMeshComponent
 	UPROPERTY(transient)
 	bool bIsUsingInGameBounds;
 	
+	/** Does this component use pre-skinned bounds? This overrides other bounds settings */
+	UPROPERTY(transient)
+	bool bIsUsingPreSkinnedBounds;
+
 	/** Base skel mesh has support for suspending clothing, but single ticks are more of a debug feature when stepping through an animation
 	 *  So we control that using this flag
 	 */
@@ -342,6 +346,16 @@ class UNREALED_API UDebugSkelMeshComponent : public USkeletalMeshComponent
 	 */
 	void UseInGameBounds(bool bUseInGameBounds);
 
+	/**
+	 * Does it use pre-skinned bounds
+	 */
+	bool IsUsingPreSkinnedBounds() const;
+
+	/**
+	 * Set to use pre-skinned bounds
+	 */
+	void UsePreSkinnedBounds(bool bUsePreSkinnedBounds);
+
 	/**
 	 * Test if in-game bounds are as big as preview bounds
 	 */
diff --git a/Engine/Source/Editor/UnrealEd/Private/Animation/DebugSkelMeshComponent.cpp b/Engine/Source/Editor/UnrealEd/Private/Animation/DebugSkelMeshComponent.cpp
--- a/Engine/Source/Editor/UnrealEd/Private/Animation/DebugSkelMeshComponent.cpp
+++ b/Engine/Source/Editor/UnrealEd/Private/Animation/DebugSkelMeshComponent.cpp
@@ -47,6 +47,12 @@ UDebugSkelMeshComponent::UDebugSkelMeshComponent(const FObjectInitializer& Objec
 
 FBoxSphereBounds UDebugSkelMeshComponent::CalcBounds(const FTransform& LocalToWorld) const
 {
+	// Override bounds with pre-skinned bounds if asking for them
+	if (IsUsingPreSkinnedBounds())
+	{
+		return GetPreSkinnedLocalBounds();
+	}
+
 	FBoxSphereBounds Result = Super::CalcBounds(LocalToWorld);
 
 	if (!IsUsingInGameBounds())
@@ -92,6 +98,16 @@ void UDebugSkelMeshComponent::UseInGameBounds(bool bUseInGameBounds)
 	bIsUsingInGameBounds = bUseInGameBounds;
 }
 
+bool UDebugSkelMeshComponent::IsUsingPreSkinnedBounds() const
+{
+	return bIsUsingPreSkinnedBounds;
+}
+
+void UDebugSkelMeshComponent::UsePreSkinnedBounds(bool bUsePreSkinnedBounds)
+{
+	bIsUsingPreSkinnedBounds = bUsePreSkinnedBounds;
+}
+
 bool UDebugSkelMeshComponent::CheckIfBoundsAreCorrrect()
 {
 	if (GetPhysicsAsset())
