PR #5734: Only spawn SEditableTextBox if TargetSocket != nullptr (Contributed by [REDACTED_PERSON])

diff --git a/Engine/Source/Editor/Persona/Private/SkeletalMeshSocketDetails.cpp b/Engine/Source/Editor/Persona/Private/SkeletalMeshSocketDetails.cpp
--- a/Engine/Source/Editor/Persona/Private/SkeletalMeshSocketDetails.cpp
+++ b/Engine/Source/Editor/Persona/Private/SkeletalMeshSocketDetails.cpp
@@ -42,39 +42,42 @@ void FSkeletalMeshSocketDetails::CustomizeDetails( IDetailLayoutBuilder& DetailB
 			}
 		}
 
-		// get the currently chosen bone/socket, if any
-		const FText PropertyName = SocketNameProperty->GetPropertyDisplayName();
-		FString CurValue;
-		SocketNameProperty->GetValue(CurValue);
-		if( CurValue == FString("None") )
+		if (TargetSocket)
 		{
-			CurValue.Empty();
-		}
-		PreEditSocketName = FText::FromString(CurValue);
-
-		// create the editable text box
-		SocketCategory.AddProperty( SocketNameProperty.ToSharedRef() )
-		.CustomWidget()
-		.NameContent()
-		[
-			SNew(SHorizontalBox)
-			+SHorizontalBox::Slot()
-			.Padding( FMargin( 2, 1, 0, 1 ) )
-			[
-				SNew(STextBlock)
-				.Text(PropertyName)
+			// get the currently chosen bone/socket, if any
+			const FText PropertyName = SocketNameProperty->GetPropertyDisplayName();
+			FString CurValue;
+			SocketNameProperty->GetValue(CurValue);
+			if (CurValue == FString("None"))
+			{
+				CurValue.Empty();
+			}
+			PreEditSocketName = FText::FromString(CurValue);
+
+			// create the editable text box
+			SocketCategory.AddProperty(SocketNameProperty.ToSharedRef())
+				.CustomWidget()
+				.NameContent()
+				[
+					SNew(SHorizontalBox)
+					+ SHorizontalBox::Slot()
+				.Padding(FMargin(2, 1, 0, 1))
+				[
+					SNew(STextBlock)
+					.Text(PropertyName)
 				.Font(FEditorStyle::GetFontStyle(TEXT("PropertyWindow.NormalFont")))
-			]
-		]
-		.ValueContent()
-		[
-			SAssignNew(SocketNameTextBox, SEditableTextBox)
-			.Text(FText::FromString(CurValue))
-			.HintText( NSLOCTEXT( "SkeletalMeshSocketDetails", "SkeletalMeshSocketDetailsHintTextSocketName", "Socket Name..." ) )
-			.OnTextCommitted( this, &FSkeletalMeshSocketDetails::OnSocketNameCommitted )
-			.OnTextChanged( this, &FSkeletalMeshSocketDetails::OnSocketNameChanged )
-			.ClearKeyboardFocusOnCommit(false)
-		];
+				]
+				]
+			.ValueContent()
+				[
+					SAssignNew(SocketNameTextBox, SEditableTextBox)
+					.Text(FText::FromString(CurValue))
+				.HintText(NSLOCTEXT("SkeletalMeshSocketDetails", "SkeletalMeshSocketDetailsHintTextSocketName", "Socket Name..."))
+				.OnTextCommitted(this, &FSkeletalMeshSocketDetails::OnSocketNameCommitted)
+				.OnTextChanged(this, &FSkeletalMeshSocketDetails::OnSocketNameChanged)
+				.ClearKeyboardFocusOnCommit(false)
+				];
+		}
 	}
 
 	ParentBoneProperty = DetailBuilder.GetProperty( TEXT("BoneName") );
