Fix incremental build CIS to use include rather than forward declaration for USkeletalMesh

diff --git a/Engine/Source/Editor/Persona/Public/SSkinWeightProfileImportOptions.h b/Engine/Source/Editor/Persona/Public/SSkinWeightProfileImportOptions.h
--- a/Engine/Source/Editor/Persona/Public/SSkinWeightProfileImportOptions.h
+++ b/Engine/Source/Editor/Persona/Public/SSkinWeightProfileImportOptions.h
@@ -12,12 +12,12 @@
 #include "IDetailCustomization.h"
 #include "Widgets/Input/SEditableTextBox.h"
 #include "PropertyHandle.h"
+#include "Engine/SkeletalMesh.h"
 
 #include "SSkinWeightProfileImportOptions.generated.h"
 
 // Forward declares
 class IDetailsView;
-class USkeletalMesh;
 
 UCLASS(Config = EditorPerProjectUserSettings)
 class USkinWeightImportOptions : public UObject
