Fix crash in Control Rig Editor

diff --git a/Engine/Source/Editor/Persona/Private/SAnimationEditorViewport.cpp b/Engine/Source/Editor/Persona/Private/SAnimationEditorViewport.cpp
--- a/Engine/Source/Editor/Persona/Private/SAnimationEditorViewport.cpp
+++ b/Engine/Source/Editor/Persona/Private/SAnimationEditorViewport.cpp
@@ -580,7 +580,10 @@ void SAnimationEditorViewportTabBody::Construct(const FArguments& InArgs, const
 	PopulateSkinWeightProfileNames();
 
 	GetPreviewScene()->OnRecordingStateChanged().AddSP(this, &SAnimationEditorViewportTabBody::AddRecordingNotification);
-	GetPreviewScene()->GetPreviewMesh()->OnPostMeshCached().AddSP(this, &SAnimationEditorViewportTabBody::UpdateSkinWeightSelection);
+	if (GetPreviewScene()->GetPreviewMesh())
+	{
+		GetPreviewScene()->GetPreviewMesh()->OnPostMeshCached().AddSP(this, &SAnimationEditorViewportTabBody::UpdateSkinWeightSelection);
+	}
 
 	AddPostProcessNotification();
 
