In order to minimize the addition of new hardcoded entries in this list, am adding a metadata field to properties to let the details panel for AnimNotifies know to make it a bone search field.

diff --git a/Engine/Source/Editor/Persona/Private/AnimNotifyDetails.cpp b/Engine/Source/Editor/Persona/Private/AnimNotifyDetails.cpp
--- a/Engine/Source/Editor/Persona/Private/AnimNotifyDetails.cpp
+++ b/Engine/Source/Editor/Persona/Private/AnimNotifyDetails.cpp
@@ -330,6 +330,7 @@ bool FAnimNotifyDetails::CustomizeProperty(IDetailCategoryBuilder& CategoryBuild
 	{
 		FString ClassName = Notify->GetClass()->GetName();
 		FString PropertyName = Property->GetProperty()->GetName();
+		bool bIsBoneName = Property->GetBoolMetaData(TEXT("AnimNotifyBoneName"));
 
 		if(ClassName.Find(TEXT("AnimNotify_PlayParticleEffect")) != INDEX_NONE && PropertyName == TEXT("SocketName"))
 		{
@@ -359,6 +360,11 @@ bool FAnimNotifyDetails::CustomizeProperty(IDetailCategoryBuilder& CategoryBuild
 				return true;
 			}
 		}
+		else if (bIsBoneName)
+		{
+			AddBoneNameProperty(CategoryBuilder, Notify, Property);
+			return true;
+		}
 	}
 	return false;
 }
