Editor crashed when deleting a skeleton hierarchy in control rig when show bone name is enabled

diff --git a/Engine/Source/Editor/Persona/Private/AnimationEditorViewportClient.cpp b/Engine/Source/Editor/Persona/Private/AnimationEditorViewportClient.cpp
--- a/Engine/Source/Editor/Persona/Private/AnimationEditorViewportClient.cpp
+++ b/Engine/Source/Editor/Persona/Private/AnimationEditorViewportClient.cpp
@@ -597,6 +597,13 @@ void FAnimationViewportClient::ShowBoneNames( FCanvas* Canvas, FSceneView* View
 	const int32 LODIndex = FMath::Clamp(PreviewMeshComponent->PredictedLODLevel, 0, SkelMeshRenderData->LODRenderData.Num()-1);
 	FSkeletalMeshLODRenderData& LODData = SkelMeshRenderData->LODRenderData[ LODIndex ];
 
+	// Check if our reference skeleton is out of synch with the one on the loddata
+	const FReferenceSkeleton& ReferenceSkeleton = PreviewMeshComponent->GetReferenceSkeleton();
+	if (ReferenceSkeleton.GetNum() < LODData.RequiredBones.Num())
+	{
+		return;
+	}
+
 	const int32 HalfX = Viewport->GetSizeXY().X/2 / GetDPIScale();
 	const int32 HalfY = Viewport->GetSizeXY().Y/2 / GetDPIScale();
 
@@ -649,7 +656,7 @@ void FAnimationViewportClient::ShowBoneNames( FCanvas* Canvas, FSceneView* View
 				const int32 XPos = HalfX + ( HalfX * proj.X );
 				const int32 YPos = HalfY + ( HalfY * (proj.Y * -1) );
 
-				const FName BoneName = PreviewMeshComponent->GetReferenceSkeleton().GetBoneName(BoneIndex);
+				const FName BoneName = ReferenceSkeleton.GetBoneName(BoneIndex);
 				const FString BoneString = FString::Printf( TEXT("%d: %s"), BoneIndex, *BoneName.ToString() );
 				FCanvasTextItem TextItem( FVector2D( XPos, YPos), FText::FromString( BoneString ), GEngine->GetSmallFont(), BoneColor );
 				TextItem.EnableShadow(FLinearColor::Black);
