Fix for non-unity CIS issue

diff --git a/Engine/Source/Editor/Persona/Private/SkinWeightProfileHelpers.cpp b/Engine/Source/Editor/Persona/Private/SkinWeightProfileHelpers.cpp
--- a/Engine/Source/Editor/Persona/Private/SkinWeightProfileHelpers.cpp
+++ b/Engine/Source/Editor/Persona/Private/SkinWeightProfileHelpers.cpp
@@ -17,6 +17,7 @@
 #include "EditorFramework/AssetImportData.h"
 #include "Widgets/Notifications/SNotificationList.h"
 #include "Framework/Notifications/NotificationManager.h"
+#include "Misc/MessageDialog.h"
 
 #define LOCTEXT_NAMESPACE "SkinWeightProfileHelpers"
 
