Fixed scrollbar always being shown in animation notify track

diff --git a/Engine/Source/Editor/Persona/Private/SAnimNotifyPanel.cpp b/Engine/Source/Editor/Persona/Private/SAnimNotifyPanel.cpp
--- a/Engine/Source/Editor/Persona/Private/SAnimNotifyPanel.cpp
+++ b/Engine/Source/Editor/Persona/Private/SAnimNotifyPanel.cpp
@@ -3940,7 +3940,6 @@ void SAnimNotifyPanel::Construct(const FArguments& InArgs, const TSharedRef<clas
 					[
 						SAssignNew(NotifyTrackScrollBar, SScrollBar)
 						.Orientation(EOrientation::Orient_Horizontal)
-						.AlwaysShowScrollbar(true)
 						.OnUserScrolled(this, &SAnimNotifyPanel::OnNotifyTrackScrolled)
 					]
 					+SHorizontalBox::Slot()
