Display Vert Colors in Persona Viewport
#fix Added vertex color to 'Show' menu and moved behaviour when toggling the show-flag
#misc Removed custom entry from 'Character / Mesh'

diff --git a/Engine/Source/Editor/Persona/Private/AnimViewportShowCommands.cpp b/Engine/Source/Editor/Persona/Private/AnimViewportShowCommands.cpp
--- a/Engine/Source/Editor/Persona/Private/AnimViewportShowCommands.cpp
+++ b/Engine/Source/Editor/Persona/Private/AnimViewportShowCommands.cpp
@@ -32,8 +32,6 @@ void FAnimViewportShowCommands::RegisterCommands()
 	UI_COMMAND( ShowBoneWeight, "Selected Bone Weight", "Display color overlay of the weight from selected bone in the viewport", EUserInterfaceActionType::RadioButton, FInputChord() );
 	UI_COMMAND( ShowMorphTargetVerts, "Selected Morph Target Vertices", "Display color overlay with the change of selected morph target in the viewport", EUserInterfaceActionType::RadioButton, FInputChord());
 
-	UI_COMMAND(ShowVertexColors, "Vertex Colors", "Display mesh vertex colors", EUserInterfaceActionType::ToggleButton, FInputChord());
-
 	UI_COMMAND( ShowRawAnimation, "Uncompressed Animation", "Display skeleton with uncompressed animation data", EUserInterfaceActionType::ToggleButton, FInputChord() );
 	UI_COMMAND( ShowNonRetargetedAnimation, "Non-Retargeted Animation", "Display Skeleton With non-retargeted animation data", EUserInterfaceActionType::ToggleButton, FInputChord() );
 	UI_COMMAND( ShowAdditiveBaseBones, "Additive Base", "Display skeleton in additive base pose", EUserInterfaceActionType::ToggleButton, FInputChord() );
diff --git a/Engine/Source/Editor/Persona/Private/AnimViewportShowCommands.h b/Engine/Source/Editor/Persona/Private/AnimViewportShowCommands.h
--- a/Engine/Source/Editor/Persona/Private/AnimViewportShowCommands.h
+++ b/Engine/Source/Editor/Persona/Private/AnimViewportShowCommands.h
@@ -98,9 +98,6 @@ public:
 	TSharedPtr< FUICommandInfo > ShowBoneWeight;
 	TSharedPtr< FUICommandInfo > ShowMorphTargetVerts;
 
-	/** Show mesh vertex colors */
-	TSharedPtr< FUICommandInfo > ShowVertexColors;
-
 	/** Show socket hit point diamonds */
 	TSharedPtr< FUICommandInfo > ShowSockets;
 
diff --git a/Engine/Source/Editor/Persona/Private/AnimationEditorViewportClient.cpp b/Engine/Source/Editor/Persona/Private/AnimationEditorViewportClient.cpp
--- a/Engine/Source/Editor/Persona/Private/AnimationEditorViewportClient.cpp
+++ b/Engine/Source/Editor/Persona/Private/AnimationEditorViewportClient.cpp
@@ -1909,6 +1909,12 @@ void FAnimationViewportClient::SetupViewForRendering( FSceneViewFamily& ViewFami
 void FAnimationViewportClient::HandleToggleShowFlag(FEngineShowFlags::EShowFlag EngineShowFlagIndex)
 {
 	FEditorViewportClient::HandleToggleShowFlag(EngineShowFlagIndex);
+	
+	if (UDebugSkelMeshComponent* Component = GetAnimPreviewScene()->GetPreviewMeshComponent())
+	{
+		Component->bDisplayVertexColors = EngineShowFlags.VertexColors;
+		Component->MarkRenderStateDirty();
+	}
 
 	ConfigOption->SetShowGrid(EngineShowFlags.Grid);
 }
diff --git a/Engine/Source/Editor/Persona/Private/SAnimViewportToolBar.cpp b/Engine/Source/Editor/Persona/Private/SAnimViewportToolBar.cpp
--- a/Engine/Source/Editor/Persona/Private/SAnimViewportToolBar.cpp
+++ b/Engine/Source/Editor/Persona/Private/SAnimViewportToolBar.cpp
@@ -604,7 +604,6 @@ TSharedRef<SWidget> SAnimViewportToolBar::GenerateCharacterMenu() const
 						SubMenuBuilder.AddMenuEntry(FAnimViewportShowCommands::Get().UseFixedBounds);
 						SubMenuBuilder.AddMenuEntry(FAnimViewportShowCommands::Get().ShowPreviewMesh );
 						SubMenuBuilder.AddMenuEntry(FAnimViewportShowCommands::Get().ShowMorphTargets );
-						SubMenuBuilder.AddMenuEntry(FAnimViewportShowCommands::Get().ShowVertexColors );
 					}
 					SubMenuBuilder.EndSection();
 
@@ -842,6 +841,7 @@ TSharedRef<SWidget> SAnimViewportToolBar::GenerateShowMenu() const
 		.IncludeFlag(FEngineShowFlags::SF_SeparateTranslucency)
 		.IncludeFlag(FEngineShowFlags::SF_TemporalAA)
 		.IncludeFlag(FEngineShowFlags::SF_Tessellation)
+		.IncludeFlag(FEngineShowFlags::SF_VertexColors)		
 		;
 
 	FShowFlagMenuCommands::Get().BuildShowFlagsMenu(InMenuBuilder, ShowFlagFilter);
diff --git a/Engine/Source/Editor/Persona/Private/SAnimationEditorViewport.cpp b/Engine/Source/Editor/Persona/Private/SAnimationEditorViewport.cpp
--- a/Engine/Source/Editor/Persona/Private/SAnimationEditorViewport.cpp
+++ b/Engine/Source/Editor/Persona/Private/SAnimationEditorViewport.cpp
@@ -692,13 +692,7 @@ void SAnimationEditorViewportTabBody::BindCommands()
 		FExecuteAction::CreateSP(this, &SAnimationEditorViewportTabBody::OnShowOverlayMorphTargetVert),
 		FCanExecuteAction(),
 		FIsActionChecked::CreateSP(this, &SAnimationEditorViewportTabBody::IsShowingOverlayMorphTargetVerts));
-
-	CommandList.MapAction(
-		ViewportShowMenuCommands.ShowVertexColors,
-		FExecuteAction::CreateSP(this, &SAnimationEditorViewportTabBody::OnShowVertexColorsChanged),
-		FCanExecuteAction(),
-		FIsActionChecked::CreateSP(this, &SAnimationEditorViewportTabBody::IsShowingVertexColors));
-
+	
 	CommandList.EndGroup();
 
 	// Show sockets
@@ -1670,38 +1664,6 @@ bool SAnimationEditorViewportTabBody::IsPreviewingRootMotion() const
 	return false;
 }
 
-bool SAnimationEditorViewportTabBody::IsShowingVertexColors() const
-{
-	return GetAnimationViewportClient()->EngineShowFlags.VertexColors;
-}
-
-void SAnimationEditorViewportTabBody::OnShowVertexColorsChanged()
-{
-	FEngineShowFlags& ShowFlags = GetAnimationViewportClient()->EngineShowFlags;
-
-	if(UDebugSkelMeshComponent* PreviewComponent = GetPreviewScene()->GetPreviewMeshComponent())
-	{
-		if(!ShowFlags.VertexColors)
-		{
-			ShowFlags.SetVertexColors(true);
-			ShowFlags.SetLighting(false);
-			ShowFlags.SetIndirectLightingCache(false);
-			PreviewComponent->bDisplayVertexColors = true;
-		}
-		else
-		{
-			ShowFlags.SetVertexColors(false);
-			ShowFlags.SetLighting(true);
-			ShowFlags.SetIndirectLightingCache(true);
-			PreviewComponent->bDisplayVertexColors = false;
-		}
-
-		PreviewComponent->RecreateRenderState_Concurrent();
-	}
-
-	RefreshViewport();
-}
-
 #if WITH_APEX_CLOTHING
 bool SAnimationEditorViewportTabBody::IsClothSimulationEnabled() const
 {
diff --git a/Engine/Source/Editor/Persona/Private/SAnimationEditorViewport.h b/Engine/Source/Editor/Persona/Private/SAnimationEditorViewport.h
--- a/Engine/Source/Editor/Persona/Private/SAnimationEditorViewport.h
+++ b/Engine/Source/Editor/Persona/Private/SAnimationEditorViewport.h
@@ -456,13 +456,6 @@ private:
 	
 	/** Whether or not we are previewing root motion */
 	bool IsPreviewingRootMotion() const;
-
-	/** Callback when user checks the vertex colors box in the show menu */
-	void OnShowVertexColorsChanged();
-
-	/** Whether or not vertex color display is enabled */
-	bool IsShowingVertexColors() const;
-
 private:
 	/** Selected Turn Table speed  */
 	EAnimationPlaybackSpeeds::Type SelectedTurnTableSpeed;
