Fixed sorting & layout issues with anim nodes

Fixed anim node details customization so that it respects code re-ordered parameters and ordering metadata.
Fixed EditCondition and (As Pin) indicators causing double-checkboxes.
Moved (As Pin) indicators to a right-aligned checkbox instead of right-aligning the label.
Reordered anim node properties: AnimDynamics, ApplyAdditive, BlendBoneByChannel.

Also manually resized the 'eye' icon (as it was box filtered form 40x40) & fixed up asymmetric padding for its checkbox style. This means 'public' checkboxes in Blueprints are now square!

diff --git a/Engine/Source/Editor/EditorStyle/Private/SlateEditorStyle.cpp b/Engine/Source/Editor/EditorStyle/Private/SlateEditorStyle.cpp
--- a/Engine/Source/Editor/EditorStyle/Private/SlateEditorStyle.cpp
+++ b/Engine/Source/Editor/EditorStyle/Private/SlateEditorStyle.cpp
@@ -520,7 +520,8 @@ void FSlateEditorStyle::FStyle::SetupGeneralStyles()
 			.SetCheckedPressedImage( IMAGE_BRUSH( "Common/CheckBox_Checked", Icon16x16 ) )
 			.SetUndeterminedImage( IMAGE_BRUSH( "Common/CheckBox_Undetermined", Icon16x16 ) )
 			.SetUndeterminedHoveredImage( IMAGE_BRUSH( "Common/CheckBox_Undetermined_Hovered", Icon16x16 ) )
-			.SetUndeterminedPressedImage( IMAGE_BRUSH( "Common/CheckBox_Undetermined_Hovered", Icon16x16, FLinearColor( 0.5f, 0.5f, 0.5f ) ) );
+			.SetUndeterminedPressedImage( IMAGE_BRUSH( "Common/CheckBox_Undetermined_Hovered", Icon16x16, FLinearColor( 0.5f, 0.5f, 0.5f ) ) )
+			.SetPadding(1.0f);
 		/* ... and set new style */
 		Set( "CheckboxLookToggleButtonCheckbox", CheckboxLookingToggleButtonStyle );
 
@@ -5556,8 +5557,8 @@ void FSlateEditorStyle::FStyle::SetupPersonaStyle()
 		Set( "Kismet.VariableList.SetTypeIconLarge", new IMAGE_BRUSH( "/Icons/pillset_40x", Icon40x40 ) );
 		Set( "Kismet.VariableList.MapValueTypeIcon", new IMAGE_BRUSH( "/Icons/pillmapvalue_16x", Icon16x16 ) );
 		Set( "Kismet.VariableList.MapKeyTypeIcon", new IMAGE_BRUSH( "/Icons/pillmapkey_16x", Icon16x16 ) );
-		Set( "Kismet.VariableList.ExposeForInstance", new IMAGE_BRUSH( "/Icons/icon_layer_visible", Icon16x16 ) );
-		Set( "Kismet.VariableList.HideForInstance", new IMAGE_BRUSH( "/Icons/icon_layer_not_visible", Icon16x16 ) );
+		Set( "Kismet.VariableList.ExposeForInstance", new IMAGE_BRUSH( "/Icons/icon_layer_visible_16x", Icon16x16 ) );
+		Set( "Kismet.VariableList.HideForInstance", new IMAGE_BRUSH( "/Icons/icon_layer_not_visible_16x", Icon16x16 ) );
 		Set( "Kismet.VariableList.VariableIsUsed", new IMAGE_BRUSH( "/Icons/icon_variable_used_16x", Icon16x16 ) );
 		Set( "Kismet.VariableList.VariableNotUsed", new IMAGE_BRUSH( "/Icons/icon_variable_not_used_16x", Icon16x16 ) );
 
diff --git a/Engine/Source/Editor/Persona/Private/AnimGraphNodeDetails.cpp b/Engine/Source/Editor/Persona/Private/AnimGraphNodeDetails.cpp
--- a/Engine/Source/Editor/Persona/Private/AnimGraphNodeDetails.cpp
+++ b/Engine/Source/Editor/Persona/Private/AnimGraphNodeDetails.cpp
@@ -42,6 +42,7 @@
 #include "AnimGraphNode_StateMachine.h"
 #include "Engine/SkeletalMeshSocket.h"
 #include "Styling/CoreStyle.h"
+#include "Widgets/Input/SCheckBox.h"
 
 #define LOCTEXT_NAMESPACE "KismetNodeWithOptionalPinsDetails"
 
@@ -98,142 +99,216 @@ void FAnimGraphNodeDetails::CustomizeDetails(class IDetailLayoutBuilder& DetailB
 	TSharedRef<IPropertyHandle> NodePropertyHandle = DetailBuilder.GetProperty(NodeProperty->GetFName(), AnimGraphNode->GetClass());
 	DetailBuilder.HideProperty(NodePropertyHandle);
 
-	// Now customize each property in the pins array
-	for (int CustomPinIndex = 0; CustomPinIndex < AnimGraphNode->ShowPinForProperties.Num(); ++CustomPinIndex)
+	uint32 NumChildHandles = 0;
+	FPropertyAccess::Result Result = NodePropertyHandle->GetNumChildren(NumChildHandles);
+	if(Result != FPropertyAccess::Fail)
 	{
-		const FOptionalPinFromProperty& OptionalPin = AnimGraphNode->ShowPinForProperties[CustomPinIndex];
+		for(uint32 ChildHandleIndex = 0; ChildHandleIndex < NumChildHandles; ++ChildHandleIndex)
+		{
+			TSharedPtr<IPropertyHandle> TargetPropertyHandle = NodePropertyHandle->GetChildHandle(ChildHandleIndex);
+			if(TargetPropertyHandle.IsValid())
+			{
+				UProperty* TargetProperty = TargetPropertyHandle->GetProperty();
+				IDetailCategoryBuilder& CurrentCategory = DetailBuilder.EditCategory(FObjectEditorUtils::GetCategoryFName(TargetProperty));
+			
+				int32 CustomPinIndex = AnimGraphNode->ShowPinForProperties.IndexOfByPredicate([TargetProperty](const FOptionalPinFromProperty& InOptionalPin)
+				{
+					return TargetProperty->GetFName() == InOptionalPin.PropertyName;
+				});
 
-		UProperty* TargetProperty = FindField<UProperty>(NodeProperty->Struct, OptionalPin.PropertyName);
+				if(CustomPinIndex != INDEX_NONE)
+				{
+					const FOptionalPinFromProperty& OptionalPin = AnimGraphNode->ShowPinForProperties[CustomPinIndex];
 
-		IDetailCategoryBuilder& CurrentCategory = DetailBuilder.EditCategory(FObjectEditorUtils::GetCategoryFName(TargetProperty));
+					// Not optional
+					if (!OptionalPin.bCanToggleVisibility && OptionalPin.bShowPin)
+					{
+						// Always displayed as a pin, so hide the property
+						DetailBuilder.HideProperty(TargetPropertyHandle);
+						continue;
+					}
 
-		const FName TargetPropertyPath(*FString::Printf(TEXT("%s.%s"), *NodeProperty->GetName(), *TargetProperty->GetName()));
-		TSharedRef<IPropertyHandle> TargetPropertyHandle = DetailBuilder.GetProperty(TargetPropertyPath, AnimGraphNode->GetClass() );
+					if(!TargetPropertyHandle->GetProperty())
+					{
+						continue;
+					}
 
-		// Not optional
-		if (!OptionalPin.bCanToggleVisibility && OptionalPin.bShowPin)
-		{
-			// Always displayed as a pin, so hide the property
-			DetailBuilder.HideProperty(TargetPropertyHandle);
-			continue;
-		}
+					// if customized, do not do anything
+					if(TargetPropertyHandle->IsCustomized())
+					{
+						continue;
+					}
+		
+					// sometimes because of order of customization
+					// this gets called first for the node you'd like to customize
+					// then the above statement won't work
+					// so you can mark certain property to have meta data "CustomizeProperty"
+					// which will trigger below statement
+					if (OptionalPin.bPropertyIsCustomized)
+					{
+						continue;
+					}
 
-		if(!TargetPropertyHandle->GetProperty())
-		{
-			continue;
-		}
+					TSharedRef<SWidget> InternalCustomWidget = CreatePropertyWidget(TargetProperty, TargetPropertyHandle.ToSharedRef(), AnimGraphNode->GetClass());
 
-		// if customized, do not do anything
-		if(TargetPropertyHandle->IsCustomized())
-		{
-			continue;
-		}
-		
-		// sometimes because of order of customization
-		// this gets called first for the node you'd like to customize
-		// then the above statement won't work
-		// so you can mark certain property to have meta data "CustomizeProperty"
-		// which will trigger below statement
-		if (OptionalPin.bPropertyIsCustomized)
-		{
-			continue;
+					if (OptionalPin.bCanToggleVisibility)
+					{
+						IDetailPropertyRow& PropertyRow = CurrentCategory.AddProperty(TargetPropertyHandle);
+
+						TSharedPtr<SWidget> NameWidget; 
+						TSharedPtr<SWidget> ValueWidget;
+						FDetailWidgetRow Row;
+						PropertyRow.GetDefaultWidgets( NameWidget, ValueWidget, Row );
+						
+						ValueWidget = (InternalCustomWidget == SNullWidget::NullWidget) ? ValueWidget : InternalCustomWidget;
+
+						const FName OptionalPinArrayEntryName(*FString::Printf(TEXT("ShowPinForProperties[%d].bShowPin"), CustomPinIndex));
+						TSharedRef<IPropertyHandle> ShowHidePropertyHandle = DetailBuilder.GetProperty(OptionalPinArrayEntryName);
+
+						ShowHidePropertyHandle->MarkHiddenByCustomization();
+
+						TSharedRef<SWidget> ShowHidePropertyWidget = CreateAsPinWidget(ShowHidePropertyHandle);
+
+						ValueWidget->SetVisibility(TAttribute<EVisibility>::Create(TAttribute<EVisibility>::FGetter::CreateSP(this, &FAnimGraphNodeDetails::GetVisibilityOfProperty, ShowHidePropertyHandle)));
+
+						// If we have an edit condition, that comes as part of the default name widget, so just use a text block to avoid duplicate checkboxes
+						TSharedPtr<SWidget> PropertyNameWidget;
+						if(TargetProperty->HasMetaData(TEXT("EditCondition")))
+						{
+							PropertyNameWidget = SNew(STextBlock)
+							.Text(TargetProperty->GetDisplayNameText())
+							.Font(IDetailLayoutBuilder::GetDetailFont())
+							.ToolTipText(TargetProperty->GetToolTipText());
+						}
+						else
+						{
+							PropertyNameWidget = NameWidget;
+						}
+
+						NameWidget = 
+							SNew(SHorizontalBox)
+							+SHorizontalBox::Slot()
+							.HAlign(HAlign_Left)
+							.VAlign(VAlign_Center)
+							.FillWidth(1)
+							[
+								SNew(SBox)
+								.Clipping(EWidgetClipping::ClipToBounds)
+								[
+									PropertyNameWidget.ToSharedRef()
+								]
+							]
+							+SHorizontalBox::Slot()
+							.AutoWidth()
+							.HAlign(HAlign_Right)
+							.VAlign(VAlign_Center)
+							[
+								ShowHidePropertyWidget
+							];
+
+						// we only show children if visibility is one
+						// whenever toggles, this gets called, so it will be refreshed
+						const bool bShowChildren = GetVisibilityOfProperty(ShowHidePropertyHandle) == EVisibility::Visible;
+						PropertyRow.CustomWidget(bShowChildren)
+						.NameContent()
+						.MinDesiredWidth(Row.NameWidget.MinWidth)
+						.MaxDesiredWidth(Row.NameWidget.MaxWidth)
+						[
+							NameWidget.ToSharedRef()
+						]
+						.ValueContent()
+						.MinDesiredWidth(Row.ValueWidget.MinWidth)
+						.MaxDesiredWidth(Row.ValueWidget.MaxWidth)
+						[
+							ValueWidget.ToSharedRef()
+						];
+					}
+					else if(InternalCustomWidget != SNullWidget::NullWidget)
+					{
+						// A few properties are internally customized within this customization. Here we
+						// catch instances of these that don't have an optional pin flag.
+						IDetailPropertyRow& PropertyRow = CurrentCategory.AddProperty(TargetPropertyHandle);
+						PropertyRow.CustomWidget()
+						.NameContent()
+						[
+							TargetPropertyHandle->CreatePropertyNameWidget()
+						]
+						.ValueContent()
+						[
+							InternalCustomWidget
+						];
+					}
+					else
+					{
+						CurrentCategory.AddProperty(TargetPropertyHandle);
+					}
+				}
+			}
 		}
+	}
+}
 
-		IDetailPropertyRow& PropertyRow = CurrentCategory.AddProperty( TargetPropertyHandle );
-
-		TSharedRef<SWidget> InternalCustomWidget = CreatePropertyWidget(TargetProperty, TargetPropertyHandle, AnimGraphNode->GetClass());
+TSharedRef<SWidget> FAnimGraphNodeDetails::CreateAsPinWidget(TSharedRef<IPropertyHandle> InPropertyHandle)
+{
+	TWeakPtr<IPropertyHandle> WeakPropertyHandle = InPropertyHandle;
 
-		if (OptionalPin.bCanToggleVisibility)
+	auto IsCheckedLambda = [WeakPropertyHandle]()
+	{
+		if(WeakPropertyHandle.IsValid())
 		{
-			TSharedPtr<SWidget> NameWidget; 
-			TSharedPtr<SWidget> ValueWidget;
-			FDetailWidgetRow Row;
-			PropertyRow.GetDefaultWidgets( NameWidget, ValueWidget, Row );
-
-			ValueWidget = (InternalCustomWidget == SNullWidget::NullWidget) ? ValueWidget : InternalCustomWidget;
-
-			const FName OptionalPinArrayEntryName(*FString::Printf(TEXT("ShowPinForProperties[%d].bShowPin"), CustomPinIndex));
-			TSharedRef<IPropertyHandle> ShowHidePropertyHandle = DetailBuilder.GetProperty(OptionalPinArrayEntryName);
-
-			ShowHidePropertyHandle->MarkHiddenByCustomization();
-
-			const FText AsPinTooltip = LOCTEXT("AsPinTooltip", "Show this property as a pin on the node");
-
-			TSharedRef<SWidget> ShowHidePropertyWidget = ShowHidePropertyHandle->CreatePropertyValueWidget();
-			ShowHidePropertyWidget->SetToolTipText(AsPinTooltip);
+			bool bValue;
+			FPropertyAccess::Result Result = WeakPropertyHandle.Pin()->GetValue(bValue);
+			if(Result == FPropertyAccess::MultipleValues)
+			{
+				return ECheckBoxState::Undetermined;
+			}
+			else
+			{
+				return bValue ? ECheckBoxState::Checked : ECheckBoxState::Unchecked;
+			}
+		}
 
-			ValueWidget->SetVisibility(TAttribute<EVisibility>::Create(TAttribute<EVisibility>::FGetter::CreateSP(this, &FAnimGraphNodeDetails::GetVisibilityOfProperty, ShowHidePropertyHandle)));
+		return ECheckBoxState::Unchecked;
+	};
 
-			NameWidget = SNew(SHorizontalBox)
-				+SHorizontalBox::Slot()
-				.HAlign(HAlign_Fill)
-				.FillWidth(1)
-				[
-					SNew(SHorizontalBox)
-					+SHorizontalBox::Slot()
-					.AutoWidth()
-					.HAlign(HAlign_Left)
-					.VAlign(VAlign_Center)
-					[
-						ShowHidePropertyWidget
-					]
-					+SHorizontalBox::Slot()
-					.AutoWidth()
-					.HAlign(HAlign_Left)
-					.VAlign(VAlign_Center)
-					[
-						SNew(STextBlock)
-						.Text(LOCTEXT("AsPin", " (As pin) "))
-						.Font(DetailBuilder.IDetailLayoutBuilder::GetDetailFont())
-						.ToolTipText(AsPinTooltip)
-					]
-					+SHorizontalBox::Slot()
-					.AutoWidth()
-					[
-						SNew(SSpacer)
-					]
-					+SHorizontalBox::Slot()
-					.HAlign(HAlign_Right)
-					.VAlign(VAlign_Center)
-					.FillWidth(1)
-					.Padding(FMargin(0,0,4,0))
-					[
-						NameWidget.ToSharedRef()
-					]
-				];
-
-			// we only show children if visilibity is one
-			// whenever toggles, this gets called, so it will be refreshed
-			const bool bShowChildren = GetVisibilityOfProperty(ShowHidePropertyHandle) == EVisibility::Visible;
-			PropertyRow.CustomWidget(bShowChildren)
-			.NameContent()
-			.MinDesiredWidth(Row.NameWidget.MinWidth)
-			.MaxDesiredWidth(Row.NameWidget.MaxWidth)
-			[
-				NameWidget.ToSharedRef()
-			]
-			.ValueContent()
-			.MinDesiredWidth(Row.ValueWidget.MinWidth)
-			.MaxDesiredWidth(Row.ValueWidget.MaxWidth)
-			[
-				ValueWidget.ToSharedRef()
-			];
+	auto OnCheckStateChangedLambda = [WeakPropertyHandle](ECheckBoxState InCheckBoxState)
+	{
+		if(WeakPropertyHandle.IsValid())
+		{
+			bool bValue = InCheckBoxState == ECheckBoxState::Checked;
+			WeakPropertyHandle.Pin()->SetValue(bValue);
 		}
-		else if(InternalCustomWidget != SNullWidget::NullWidget)
+	};
+
+	auto ImageLambda = [WeakPropertyHandle]()
+	{
+		if(WeakPropertyHandle.IsValid())
 		{
-			// A few properties are internally customized within this customization. Here we
-			// catch instances of these that don't have an optional pin flag.
-			PropertyRow.CustomWidget()
-				.NameContent()
-				[
-					TargetPropertyHandle->CreatePropertyNameWidget()
-				]
-				.ValueContent()
-				[
-					InternalCustomWidget
-				];
+			bool bValue;
+			FPropertyAccess::Result Result = WeakPropertyHandle.Pin()->GetValue(bValue);
+			if(Result == FPropertyAccess::MultipleValues)
+			{
+				return FEditorStyle::GetBrush("Kismet.VariableList.HideForInstance");
+			}
+			else
+			{
+				return bValue ? FEditorStyle::GetBrush("Kismet.VariableList.ExposeForInstance") : FEditorStyle::GetBrush("Kismet.VariableList.HideForInstance");
+			}
 		}
-	}
+
+		return FEditorStyle::GetBrush("Kismet.VariableList.HideForInstance");
+	};
+
+	return SNew(SCheckBox)
+		.ToolTipText(LOCTEXT("AsPinTooltip", "Show/hide this property as a pin on the node"))
+		.IsChecked_Lambda(IsCheckedLambda)
+		.OnCheckStateChanged_Lambda(OnCheckStateChangedLambda)
+		.Style(FEditorStyle::Get(), "CheckboxLookToggleButtonCheckbox")
+		[
+			SNew(SImage)
+			.Image_Lambda(ImageLambda)
+			.ColorAndOpacity(FLinearColor::Black)
+		];
 }
 
 TSharedRef<SWidget> FAnimGraphNodeDetails::CreatePropertyWidget(UProperty* TargetProperty, TSharedRef<IPropertyHandle> TargetPropertyHandle, UClass* NodeClass)
diff --git a/Engine/Source/Editor/Persona/Private/AnimGraphNodeDetails.h b/Engine/Source/Editor/Persona/Private/AnimGraphNodeDetails.h
--- a/Engine/Source/Editor/Persona/Private/AnimGraphNodeDetails.h
+++ b/Engine/Source/Editor/Persona/Private/AnimGraphNodeDetails.h
@@ -39,6 +39,9 @@ protected:
 	// Creates a widget for the supplied property
 	TSharedRef<SWidget> CreatePropertyWidget(UProperty* TargetProperty, TSharedRef<IPropertyHandle> TargetPropertyHandle, UClass* NodeClass);
 
+	// Creates the 'as pin' toggle widget for a property
+	TSharedRef<SWidget> CreateAsPinWidget(TSharedRef<IPropertyHandle> InPropertyHandle);
+
 	EVisibility GetVisibilityOfProperty(TSharedRef<IPropertyHandle> Handle) const;
 
 	/** Delegate to handle filtering of asset pickers */
diff --git a/Engine/Source/Runtime/AnimGraphRuntime/Public/AnimNodes/AnimNode_ApplyAdditive.h b/Engine/Source/Runtime/AnimGraphRuntime/Public/AnimNodes/AnimNode_ApplyAdditive.h
--- a/Engine/Source/Runtime/AnimGraphRuntime/Public/AnimNodes/AnimNode_ApplyAdditive.h
+++ b/Engine/Source/Runtime/AnimGraphRuntime/Public/AnimNodes/AnimNode_ApplyAdditive.h
@@ -34,7 +34,7 @@ struct ANIMGRAPHRUNTIME_API FAnimNode_ApplyAdditive : public FAnimNode_Base
 	UPROPERTY(EditAnywhere, BlueprintReadWrite, Category = Performance, meta=(DisplayName="LOD Threshold"))
 	int32 LODThreshold;
 
-	UPROPERTY(EditAnywhere, BlueprintReadWrite, Category = Alpha, meta = (DisplayName = "Blend Settings"))
+	UPROPERTY(EditAnywhere, BlueprintReadWrite, Category = Alpha, meta = (DisplayName = "Blend Settings", DisplayAfter = "bAlphaBoolEnabled"))
 	FInputAlphaBoolBlend AlphaBoolBlend;
 
 	UPROPERTY(EditAnywhere, BlueprintReadWrite, Category = Alpha, meta = (PinShownByDefault))
@@ -45,10 +45,10 @@ struct ANIMGRAPHRUNTIME_API FAnimNode_ApplyAdditive : public FAnimNode_Base
 
 	float ActualAlpha;
 
-	UPROPERTY(EditAnywhere, BlueprintReadWrite, Category = Alpha)
+	UPROPERTY(EditAnywhere, BlueprintReadWrite, Category = Alpha, meta = (DisplayAfter = "AlphaScaleBias"))
 	EAnimAlphaInputType AlphaInputType;
 
-	UPROPERTY(EditAnywhere, BlueprintReadWrite, Category = Alpha, meta = (PinShownByDefault, DisplayName = "bEnabled"))
+	UPROPERTY(EditAnywhere, BlueprintReadWrite, Category = Alpha, meta = (PinShownByDefault, DisplayName = "bEnabled", DisplayAfter = "AlphaInputType"))
 	bool bAlphaBoolEnabled;
 
 public:	
diff --git a/Engine/Source/Runtime/AnimGraphRuntime/Public/AnimNodes/AnimNode_BlendBoneByChannel.h b/Engine/Source/Runtime/AnimGraphRuntime/Public/AnimNodes/AnimNode_BlendBoneByChannel.h
--- a/Engine/Source/Runtime/AnimGraphRuntime/Public/AnimNodes/AnimNode_BlendBoneByChannel.h
+++ b/Engine/Source/Runtime/AnimGraphRuntime/Public/AnimNodes/AnimNode_BlendBoneByChannel.h
@@ -55,7 +55,7 @@ public:
 	UPROPERTY(EditAnywhere, BlueprintReadWrite, Category = Links)
 	FPoseLink B;
 
-	UPROPERTY(EditAnywhere, Category = Blend)
+	UPROPERTY(EditAnywhere, Category = Blend, meta = (DisplayAfter = "AlphaScaleBias"))
 	TArray<FBlendBoneByChannelEntry> BoneDefinitions;
 
 private:
@@ -64,7 +64,7 @@ private:
 	TArray<FBlendBoneByChannelEntry> ValidBoneEntries;
 
 public:
-	UPROPERTY(EditAnywhere, BlueprintReadWrite, Category = Settings, meta = (PinShownByDefault))
+	UPROPERTY(EditAnywhere, BlueprintReadWrite, Category = Settings, meta = (PinShownByDefault, DisplayAfter = "B"))
 	float Alpha;
 
 private:
diff --git a/Engine/Source/Runtime/AnimGraphRuntime/Public/BoneControllers/AnimNode_AnimDynamics.h b/Engine/Source/Runtime/AnimGraphRuntime/Public/BoneControllers/AnimNode_AnimDynamics.h
--- a/Engine/Source/Runtime/AnimGraphRuntime/Public/BoneControllers/AnimNode_AnimDynamics.h
+++ b/Engine/Source/Runtime/AnimGraphRuntime/Public/BoneControllers/AnimNode_AnimDynamics.h
@@ -140,7 +140,7 @@ struct FAnimPhysConstraintSetup
 	AnimPhysTwistAxis TwistAxis;
 
 	/** Axis on body1 to match to the angular target direction. */
-	UPROPERTY(EditAnywhere, Category = Angular)
+	UPROPERTY(EditAnywhere, Category = Angular, meta=(DisplayAfter=AngularLimitsMax))
 	AnimPhysTwistAxis AngularTargetAxis;
 
 	/** Angle to use when constraining using a cone */
@@ -232,11 +232,11 @@ struct ANIMGRAPHRUNTIME_API FAnimNode_AnimDynamics : public FAnimNode_SkeletalCo
 	FAnimNode_AnimDynamics();
 
 	/** Overridden linear damping value */
-	UPROPERTY(EditAnywhere, AdvancedDisplay, BlueprintReadWrite, Category = Setup, meta = (PinHiddenByDefault))
+	UPROPERTY(EditAnywhere, AdvancedDisplay, BlueprintReadWrite, Category = Setup, meta = (PinHiddenByDefault, DisplayAfter="bOverrideLinearDamping"))
 	float LinearDampingOverride;
 
 	/** Overridden angular damping value */
-	UPROPERTY(EditAnywhere, AdvancedDisplay, BlueprintReadWrite, Category = Setup, meta = (PinHiddenByDefault))
+	UPROPERTY(EditAnywhere, AdvancedDisplay, BlueprintReadWrite, Category = Setup, meta = (PinHiddenByDefault, DisplayAfter="bOverrideAngularDamping"))
 	float AngularDampingOverride;
 
 	// Previous component & actor transforms, used to account for teleports
@@ -244,11 +244,11 @@ struct ANIMGRAPHRUNTIME_API FAnimNode_AnimDynamics : public FAnimNode_SkeletalCo
 	FTransform PreviousActorWorldSpaceTM;
 
 	/** When in BoneRelative sim space, the simulation will use this bone as the origin */
-	UPROPERTY(EditAnywhere, Category = Setup)
+	UPROPERTY(EditAnywhere, Category = Setup, meta=(DisplayAfter="SimulationSpace"))
 	FBoneReference RelativeSpaceBone;
 
 	/** The bone to attach the physics body to, if bChain is true this is the top of the chain */
-	UPROPERTY(EditAnywhere, Category = Setup)
+	UPROPERTY(EditAnywhere, Category = Setup, meta=(DisplayAfter="bChain"))
 	FBoneReference BoundBone;
 	
 	/** If bChain is true this is the bottom of the chain, otherwise ignored */
@@ -272,7 +272,7 @@ struct ANIMGRAPHRUNTIME_API FAnimNode_AnimDynamics : public FAnimNode_SkeletalCo
 	FVector GravityOverride;
 
 	/** Spring constant to use when calculating linear springs, higher values mean a stronger spring.*/
-	UPROPERTY(EditAnywhere, BlueprintReadWrite, Category = Setup, meta = (PinHiddenByDefault))
+	UPROPERTY(EditAnywhere, BlueprintReadWrite, Category = Setup, meta = (PinHiddenByDefault, DisplayAfter="bAngularSpring"))
 	float LinearSpringConstant;
 
 	/** Spring constant to use when calculating angular springs, higher values mean a stronger spring */
@@ -280,7 +280,7 @@ struct ANIMGRAPHRUNTIME_API FAnimNode_AnimDynamics : public FAnimNode_SkeletalCo
 	float AngularSpringConstant;
 
 	/** Scale to apply to calculated wind velocities in the solver */
-	UPROPERTY(EditAnywhere, Category = Wind)
+	UPROPERTY(EditAnywhere, Category = Wind, meta=(DisplayAfter="bEnableWind"))
 	float WindScale;
 
 	/** When using non-world-space sim, this controls how much of the components world-space acceleration is passed on to the local-space simulation. */
@@ -300,11 +300,11 @@ struct ANIMGRAPHRUNTIME_API FAnimNode_AnimDynamics : public FAnimNode_SkeletalCo
 	 *  in check. When using single-body systems sometimes angular forces will look like they are "catching-up" with
 	 *  the mesh, if that's the case override this and push it towards 1.0f until it settles correctly
 	 */
-	UPROPERTY(EditAnywhere, AdvancedDisplay, BlueprintReadWrite, Category = Setup, meta = (PinHiddenByDefault))
+	UPROPERTY(EditAnywhere, AdvancedDisplay, BlueprintReadWrite, Category = Setup, meta = (PinHiddenByDefault, DisplayAfter="bOverrideAngularBias"))
 	float AngularBiasOverride;
 
 	/** Number of update passes on the linear and angular limits before we solve the position of the bodies recommended to be four times the value of NumSolverIterationsPostUpdate */
-	UPROPERTY(EditAnywhere, AdvancedDisplay, Category = Setup)
+	UPROPERTY(EditAnywhere, AdvancedDisplay, Category = Setup, meta=(DisplayAfter="bDoEval"))
 	int32 NumSolverIterationsPreUpdate;
 
 	/** Number of update passes on the linear and angular limits after we solve the position of the bodies, recommended to be around a quarter of NumSolverIterationsPreUpdate */
@@ -316,11 +316,11 @@ struct ANIMGRAPHRUNTIME_API FAnimNode_AnimDynamics : public FAnimNode_SkeletalCo
 	FAnimPhysConstraintSetup ConstraintSetup;
 
 	/** List of available spherical limits for this node */
-	UPROPERTY(EditAnywhere, Category = SphericalLimit)
+	UPROPERTY(EditAnywhere, Category = SphericalLimit, meta=(DisplayAfter="bUseSphericalLimits"))
 	TArray<FAnimPhysSphericalLimit> SphericalLimits;
 
 	/** Radius to use if CollisionType is set to CustomSphere */
-	UPROPERTY(EditAnywhere, Category = Collision, meta = (UIMin = "1", ClampMin = "1"))
+	UPROPERTY(EditAnywhere, Category = Collision, meta = (UIMin = "1", ClampMin = "1", DisplayAfter="CollisionType"))
 	float SphereCollisionRadius;
 
 	/** An external force to apply to all bodies in the simulation when ticked, specified in world space */
@@ -328,7 +328,7 @@ struct ANIMGRAPHRUNTIME_API FAnimNode_AnimDynamics : public FAnimNode_SkeletalCo
 	FVector ExternalForce;
 
 	/** List of available planar limits for this node */
-	UPROPERTY(EditAnywhere, Category=PlanarLimit)
+	UPROPERTY(EditAnywhere, Category=PlanarLimit, meta=(DisplayAfter="bUsePlanarLimit"))
 	TArray<FAnimPhysPlanarLimit> PlanarLimits;
 
 	/** Resolution method for planar limits */
@@ -336,7 +336,7 @@ struct ANIMGRAPHRUNTIME_API FAnimNode_AnimDynamics : public FAnimNode_SkeletalCo
 	AnimPhysCollisionType CollisionType;
 
 	/** The space used to run the simulation */
-	UPROPERTY(EditAnywhere, BlueprintReadWrite, Category = Setup, meta = (PinHiddenByDefault))
+	UPROPERTY(EditAnywhere, BlueprintReadWrite, Category = Setup, meta = (PinHiddenByDefault, DisplayPriority=0))
 	AnimPhysSimSpaceType SimulationSpace;
 
 	// Cached sim space that we last used
@@ -355,7 +355,7 @@ struct ANIMGRAPHRUNTIME_API FAnimNode_AnimDynamics : public FAnimNode_SkeletalCo
 	uint8 bUsePlanarLimit:1;
 
 	/** If true we will perform physics update, otherwise skip - allows visualisation of the initial state of the bodies */
-	UPROPERTY(EditAnywhere, AdvancedDisplay, Category = Setup)
+	UPROPERTY(EditAnywhere, AdvancedDisplay, Category = Setup, meta=(DisplayAfter="AngularBiasOverride"))
 	uint8 bDoUpdate:1;
 
 	/** If true we will perform bone transform evaluation, otherwise skip - allows visualisation of the initial anim state compared to the physics sim */
@@ -363,7 +363,7 @@ struct ANIMGRAPHRUNTIME_API FAnimNode_AnimDynamics : public FAnimNode_SkeletalCo
 	uint8 bDoEval:1;
 
 	/** If true, the override value will be used for linear damping */
-	UPROPERTY(EditAnywhere, AdvancedDisplay, Category = Setup)
+	UPROPERTY(EditAnywhere, AdvancedDisplay, Category = Setup, meta=(DisplayAfter="AngularSpringConstraint"))
 	uint8 bOverrideLinearDamping:1;
 
 	/** If true, the override value will be used for the angular bias for bodies in this node. 
@@ -371,11 +371,11 @@ struct ANIMGRAPHRUNTIME_API FAnimNode_AnimDynamics : public FAnimNode_SkeletalCo
 	 *  in check. When using single-body systems sometimes angular forces will look like they are "catching-up" with
 	 *  the mesh, if that's the case override this and push it towards 1.0f until it settles correctly
 	 */
-	UPROPERTY(EditAnywhere, AdvancedDisplay, Category = Setup)
+	UPROPERTY(EditAnywhere, AdvancedDisplay, Category = Setup, meta=(DisplayAfter="AngularDampingOverride"))
 	uint8 bOverrideAngularBias:1;
 
 	/** If true, the override value will be used for angular damping */
-	UPROPERTY(EditAnywhere, AdvancedDisplay, Category = Setup)
+	UPROPERTY(EditAnywhere, AdvancedDisplay, Category = Setup, meta=(DisplayAfter="LinearDampingOverride"))
 	uint8 bOverrideAngularDamping:1;
 
 	/** Whether or not wind is enabled for the bodies in this simulation */
@@ -389,15 +389,15 @@ struct ANIMGRAPHRUNTIME_API FAnimNode_AnimDynamics : public FAnimNode_SkeletalCo
 	uint8 bUseGravityOverride:1;
 
 	/** If true the body will attempt to spring back to its initial position */
-	UPROPERTY(EditAnywhere, Category = Setup)
+	UPROPERTY(EditAnywhere, Category = Setup, meta=(DisplayAfter="GravityOverride"))
 	uint8 bLinearSpring:1;
 
 	/** If true the body will attempt to align itself with the specified angular target */
-	UPROPERTY(EditAnywhere, Category = Setup)
+	UPROPERTY(EditAnywhere, Category = Setup, meta=(DisplayAfter="bLinearSpring"))
 	uint8 bAngularSpring:1;
 
 	/** Set to true to use the solver to simulate a connected chain */
-	UPROPERTY(EditAnywhere, Category = Setup)
+	UPROPERTY(EditAnywhere, Category = Setup, meta=(DisplayAfter="RelativeSpaceBone"))
 	uint8 bChain:1;
 
 	// FAnimNode_SkeletalControlBase interface
