Add exportPreviewMesh in the animation category of the fbx export options dialog. This allow the user to control if the preview mesh should be exported with animation or not. It replace a dialog that was showing before the fbx export options dialog.

diff --git a/Engine/Source/Editor/Persona/Private/PersonaModule.cpp b/Engine/Source/Editor/Persona/Private/PersonaModule.cpp
--- a/Engine/Source/Editor/Persona/Private/PersonaModule.cpp
+++ b/Engine/Source/Editor/Persona/Private/PersonaModule.cpp
@@ -678,8 +678,6 @@ bool FPersonaModule::ExportToFBX(TArray<TWeakObjectPtr<UAnimSequence>>& AnimSequ
 				FEditorDirectories::Get().SetLastDirectory(ELastDirectory::GENERIC_EXPORT, DestinationFolder);
 			}
 
-			EAppReturnType::Type DialogReturn = FMessageDialog::Open(EAppMsgType::YesNo, LOCTEXT("ExportToFBXExportSkeletalMeshToo", "Would you like to export the current skeletal mesh with the animation(s)?"));
-			bool bSaveSkeletalMesh = EAppReturnType::Yes == DialogReturn;
 
 			const bool bShowCancel = false;
 			const bool bShowProgressDialog = true;
@@ -698,7 +696,7 @@ bool FPersonaModule::ExportToFBX(TArray<TWeakObjectPtr<UAnimSequence>>& AnimSequ
 
 				FString FileName = FString::Printf(TEXT("%s/%s"), *DestinationFolder, *AnimFileNames[i]);
 
-				FbxAnimUtils::ExportAnimFbx(*FileName, AnimSequence, SkeletalMesh, bSaveSkeletalMesh, ExportBatch, ExportAll, ExportCancel);
+				FbxAnimUtils::ExportAnimFbx(*FileName, AnimSequence, SkeletalMesh, ExportBatch, ExportAll, ExportCancel);
 				if (ExportBatch && ExportCancel)
 				{
 					//The user cancel the batch export
diff --git a/Engine/Source/Editor/UnrealEd/Classes/Exporters/FbxExportOption.h b/Engine/Source/Editor/UnrealEd/Classes/Exporters/FbxExportOption.h
--- a/Engine/Source/Editor/UnrealEd/Classes/Exporters/FbxExportOption.h
+++ b/Engine/Source/Editor/UnrealEd/Classes/Exporters/FbxExportOption.h
@@ -39,7 +39,7 @@ public:
 	UPROPERTY(EditAnywhere, BlueprintReadWrite, AdvancedDisplay, config, Category = Exporter)
 	uint32 bForceFrontXAxis : 1;
 
-	/** If enable, export vertex color */
+	/** If enabled, export vertex color */
 	UPROPERTY(EditAnywhere, BlueprintReadWrite, config, category = Mesh)
 	uint32 VertexColor : 1;
 
@@ -51,6 +51,10 @@ public:
 	UPROPERTY(EditAnywhere, BlueprintReadWrite, config, category = StaticMesh)
 	uint32 Collision : 1;
 
+	/** If enable, the preview mesh link to the exported animations will be also exported. */
+	UPROPERTY(EditAnywhere, BlueprintReadWrite, config, category = Animation)
+	uint32 bExportPreviewMesh : 1;
+
 	/** If enable, Map skeletal actor motion to the root bone of the skeleton. */
 	UPROPERTY(EditAnywhere, BlueprintReadWrite, config, category = Animation)
 	uint32 MapSkeletalMotionToRoot : 1;
diff --git a/Engine/Source/Editor/UnrealEd/Private/EditorExporters.cpp b/Engine/Source/Editor/UnrealEd/Private/EditorExporters.cpp
--- a/Engine/Source/Editor/UnrealEd/Private/EditorExporters.cpp
+++ b/Engine/Source/Editor/UnrealEd/Private/EditorExporters.cpp
@@ -2130,7 +2130,7 @@ bool UAnimSequenceExporterFBX::ExportBinary( UObject* Object, const TCHAR* Type,
 		}
 
 		Exporter->CreateDocument();
-		Exporter->ExportAnimSequence(AnimSequence, PreviewMesh, false);
+		Exporter->ExportAnimSequence(AnimSequence, PreviewMesh, Exporter->GetExportOptions()->bExportPreviewMesh);
 		Exporter->WriteToFile( *UExporter::CurrentFilename );
 		Exporter->SetExportOptionsOverride(nullptr);
 
diff --git a/Engine/Source/Editor/UnrealEd/Private/Fbx/FbxSkeletalMeshExport.cpp b/Engine/Source/Editor/UnrealEd/Private/Fbx/FbxSkeletalMeshExport.cpp
--- a/Engine/Source/Editor/UnrealEd/Private/Fbx/FbxSkeletalMeshExport.cpp
+++ b/Engine/Source/Editor/UnrealEd/Private/Fbx/FbxSkeletalMeshExport.cpp
@@ -539,8 +539,7 @@ FbxNode* FFbxExporter::ExportSkeletalMeshToFbx(const USkeletalMesh* SkelMesh, co
 {
 	if(AnimSeq)
 	{
-		return ExportAnimSequence(AnimSeq, SkelMesh, true, MeshName, ActorRootNode);
-
+		return ExportAnimSequence(AnimSeq, SkelMesh, GetExportOptions()->bExportPreviewMesh, MeshName, ActorRootNode);
 	}
 	else
 	{
diff --git a/Engine/Source/Editor/UnrealEd/Private/FbxAnimUtils.cpp b/Engine/Source/Editor/UnrealEd/Private/FbxAnimUtils.cpp
--- a/Engine/Source/Editor/UnrealEd/Private/FbxAnimUtils.cpp
+++ b/Engine/Source/Editor/UnrealEd/Private/FbxAnimUtils.cpp
@@ -7,10 +7,11 @@
 #include "Animation/AnimTypes.h"
 #include "Curves/RichCurve.h"
 #include "Engine/CurveTable.h"
+#include "Exporters/FbxExportOption.h"
 
 namespace FbxAnimUtils
 {
-	void ExportAnimFbx( const FString& ExportFilename, UAnimSequence* AnimSequence, USkeletalMesh* Mesh, bool bSaveSkeletalMesh, bool BatchMode, bool &OutExportAll, bool &OutCancelExport )
+	void ExportAnimFbx( const FString& ExportFilename, UAnimSequence* AnimSequence, USkeletalMesh* Mesh, bool BatchMode, bool &OutExportAll, bool &OutCancelExport )
 	{
 		if( !ExportFilename.IsEmpty() && AnimSequence && Mesh )
 		{
@@ -24,7 +25,7 @@ namespace FbxAnimUtils
 			{
 				Exporter->CreateDocument();
 
-				Exporter->ExportAnimSequence(AnimSequence, Mesh, bSaveSkeletalMesh);
+				Exporter->ExportAnimSequence(AnimSequence, Mesh, Exporter->GetExportOptions()->bExportPreviewMesh);
 
 				// Save to disk
 				Exporter->WriteToFile(*ExportFilename);
diff --git a/Engine/Source/Editor/UnrealEd/Private/FbxExporter.h b/Engine/Source/Editor/UnrealEd/Private/FbxExporter.h
--- a/Engine/Source/Editor/UnrealEd/Private/FbxExporter.h
+++ b/Engine/Source/Editor/UnrealEd/Private/FbxExporter.h
@@ -466,11 +466,12 @@ private:
 	 */
 	void CreateAnimatableUserProperty(FbxNode* Node, float Value, const char* Name, const char* Label);
 
-	/** Returns currently active FBX export options. Automation or UI diaog based options. */
-	UFbxExportOption* GetExportOptions();
-
 	/** Exports all the object's FBX metadata to the FBX node */
 	void ExportObjectMetadata(const UObject* ObjectToExport, FbxNode* Node);
+
+public:
+	/** Returns currently active FBX export options. Automation or UI dialog based options. */
+	UFbxExportOption* GetExportOptions();
 };
 
 
diff --git a/Engine/Source/Editor/UnrealEd/Public/FbxAnimUtils.h b/Engine/Source/Editor/UnrealEd/Public/FbxAnimUtils.h
--- a/Engine/Source/Editor/UnrealEd/Public/FbxAnimUtils.h
+++ b/Engine/Source/Editor/UnrealEd/Public/FbxAnimUtils.h
@@ -18,7 +18,7 @@ namespace fbxsdk
 namespace FbxAnimUtils
 {
 	//Function to export fbx animation
-	UNREALED_API void ExportAnimFbx( const FString& ExportFilename, UAnimSequence* AnimSequence, USkeletalMesh* Mesh, bool bSaveSkeletalMesh, bool BatchMode, bool &OutExportAll, bool &OutCancelExport);
+	UNREALED_API void ExportAnimFbx( const FString& ExportFilename, UAnimSequence* AnimSequence, USkeletalMesh* Mesh, bool BatchMode, bool &OutExportAll, bool &OutCancelExport);
 
 	/** 
 	 * Import curves from a named node into the supplied curve table.
