PR #5417: Hide unnecessary border in animation context menu (Contributed by [REDACTED_PERSON])

diff --git a/Engine/Source/Editor/ClassViewer/Private/SClassViewer.cpp b/Engine/Source/Editor/ClassViewer/Private/SClassViewer.cpp
--- a/Engine/Source/Editor/ClassViewer/Private/SClassViewer.cpp
+++ b/Engine/Source/Editor/ClassViewer/Private/SClassViewer.cpp
@@ -2166,7 +2166,7 @@ void SClassViewer::Construct(const FArguments& InArgs, const FClassViewerInitial
 	.MaxDesiredHeight(800.0f)
 	[
 		SNew(SBorder)
-		.BorderImage(FEditorStyle::GetBrush("Menu.Background"))
+		.BorderImage(FEditorStyle::GetBrush(InitOptions.bShowBackgroundBorder ? "ToolPanel.GroupBorder" : "NoBorder"))
 		[
 			SNew(SVerticalBox)
 			+SVerticalBox::Slot()
diff --git a/Engine/Source/Editor/ClassViewer/Public/ClassViewerModule.h b/Engine/Source/Editor/ClassViewer/Public/ClassViewerModule.h
--- a/Engine/Source/Editor/ClassViewer/Public/ClassViewerModule.h
+++ b/Engine/Source/Editor/ClassViewer/Public/ClassViewerModule.h
@@ -101,8 +101,11 @@ public:
 	/** The property this class viewer be working on. */
 	TSharedPtr<class IPropertyHandle> PropertyHandle;
 
-	/** true (the default) shows the view options at the bottom of the class picker*/
+	/** true (the default) shows the view options at the bottom of the class picker */
 	bool bAllowViewOptions;
+
+	/** true (the default) shows a background border behind the class viewer widget. */
+	bool bShowBackgroundBorder = true;
 public:
 
 	/** Constructor */
diff --git a/Engine/Source/Editor/Persona/Private/SAnimNotifyPanel.cpp b/Engine/Source/Editor/Persona/Private/SAnimNotifyPanel.cpp
--- a/Engine/Source/Editor/Persona/Private/SAnimNotifyPanel.cpp
+++ b/Engine/Source/Editor/Persona/Private/SAnimNotifyPanel.cpp
@@ -2367,6 +2367,12 @@ void SAnimNotifyTrack::MakeNewNotifyPicker(FMenuBuilder& MenuBuilder, bool bIsRe
 		UAnimSequenceBase* Sequence;
 	};
 
+	// MenuBuilder always has a search widget added to it by default, hence if larger then 1 then something else has been added to it
+	if (MenuBuilder.GetMultiBox()->GetBlocks().Num() > 1)
+	{
+		MenuBuilder.AddMenuSeparator();
+	}
+
 	FClassViewerInitializationOptions InitOptions;
 	InitOptions.Mode = EClassViewerMode::ClassPicker;
 	InitOptions.bShowObjectRootClass = false;
@@ -2376,6 +2382,7 @@ void SAnimNotifyTrack::MakeNewNotifyPicker(FMenuBuilder& MenuBuilder, bool bIsRe
 	InitOptions.bExpandRootNodes = true;
 	InitOptions.NameTypeToDisplay = EClassViewerNameTypeToDisplay::DisplayName;
 	InitOptions.ClassFilter = MakeShared<FNotifyStateClassFilter>(Sequence);
+	InitOptions.bShowBackgroundBorder = false;
 
 	FClassViewerModule& ClassViewerModule = FModuleManager::LoadModuleChecked<FClassViewerModule>("ClassViewer");
 	MenuBuilder.AddWidget(
