Fixed notify panel not updating when parent montages are edited

diff --git a/Engine/Source/Editor/Persona/Private/SMontageEditor.cpp b/Engine/Source/Editor/Persona/Private/SMontageEditor.cpp
--- a/Engine/Source/Editor/Persona/Private/SMontageEditor.cpp
+++ b/Engine/Source/Editor/Persona/Private/SMontageEditor.cpp
@@ -617,6 +617,7 @@ void SMontageEditor::SortAndUpdateMontage()
 	AnimMontagePanel->Update();
 	AnimMontageSectionsPanel->Update();
 	AnimTimingPanel->Update();
+	AnimNotifyPanel->Update();
 
 	// Restart the preview instance of the montage
 	RestartPreview();
