PR #5352: Editor: Unregistering some detail layout delegates (Contributed by [REDACTED_PERSON])

diff --git a/Engine/Source/Developer/MeshUtilities/Private/MeshUtilities.cpp b/Engine/Source/Developer/MeshUtilities/Private/MeshUtilities.cpp
--- a/Engine/Source/Developer/MeshUtilities/Private/MeshUtilities.cpp
+++ b/Engine/Source/Developer/MeshUtilities/Private/MeshUtilities.cpp
@@ -5596,6 +5596,8 @@ void FMeshUtilities::ShutdownModule()
 		FPropertyEditorModule& PropertyEditorModule = FModuleManager::Get().GetModuleChecked<FPropertyEditorModule>(PropertyEditorModuleName);
 
 		PropertyEditorModule.UnregisterCustomClassLayout("MeshSimplificationSettings");
+		PropertyEditorModule.UnregisterCustomClassLayout("SkeletalMeshSimplificationSettings");
+		PropertyEditorModule.UnregisterCustomClassLayout("ProxyLODMeshSimplificationSettings");
 	}
 
 	RemoveLevelViewportMenuExtender();
diff --git a/Engine/Source/Editor/Persona/Private/PersonaModule.cpp b/Engine/Source/Editor/Persona/Private/PersonaModule.cpp
--- a/Engine/Source/Editor/Persona/Private/PersonaModule.cpp
+++ b/Engine/Source/Editor/Persona/Private/PersonaModule.cpp
@@ -169,6 +169,7 @@ void FPersonaModule::ShutdownModule()
 		PropertyModule.UnregisterCustomClassLayout("SkeletalMeshSocket");
 		PropertyModule.UnregisterCustomClassLayout("EditorNotifyObject");
 		PropertyModule.UnregisterCustomClassLayout("AnimGraphNode_Base");
+		PropertyModule.UnregisterCustomClassLayout("AnimInstance");
 		PropertyModule.UnregisterCustomClassLayout("BlendSpaceBase");
 
 		PropertyModule.UnregisterCustomPropertyTypeLayout("InputScaleBias");
@@ -177,8 +178,8 @@ void FPersonaModule::ShutdownModule()
 		PropertyModule.UnregisterCustomPropertyTypeLayout("BlendParameter");
 		PropertyModule.UnregisterCustomPropertyTypeLayout("InterpolationParameter");
 
-		PropertyModule.UnregisterCustomClassLayout("SkeletalMeshSamplingRegionBoneFilter");
-		PropertyModule.UnregisterCustomClassLayout("SkeletalMeshSamplingRegionMaterialFilter");
+		PropertyModule.UnregisterCustomPropertyTypeLayout("SkeletalMeshSamplingRegionBoneFilter");
+		PropertyModule.UnregisterCustomPropertyTypeLayout("SkeletalMeshSamplingRegionMaterialFilter");
 	}
 }
 
