Removed copy skin weights functionality, because it only introduces problems

AltSkinWeights - Text referencing "LOD Index" does not contain a space
#fix Added metadata to improve property text

diff --git a/Engine/Source/Editor/Persona/Private/PersonaMeshDetails.cpp b/Engine/Source/Editor/Persona/Private/PersonaMeshDetails.cpp
--- a/Engine/Source/Editor/Persona/Private/PersonaMeshDetails.cpp
+++ b/Engine/Source/Editor/Persona/Private/PersonaMeshDetails.cpp
@@ -1267,21 +1267,7 @@ TSharedRef<SWidget> FPersonaMeshDetails::CreateSkinWeightProfileMenuContent()
 			MeshDetailLayout->ForceRefreshDetails();
 		}
 	})));
-
-	// Menu entry for copying skin weights from another Skeletal Mesh
-	AddProfileMenuBuilder.AddMenuEntry(LOCTEXT("CopyOverrideLabel", "Copy Skin Weights"), LOCTEXT("CopyOverrideToolTip", "Copy Skin Weights from another Skeletal Mesh"),
-		FSlateIcon(), FUIAction(FExecuteAction::CreateLambda([this, WeakSkeletalMeshPtr = SkeletalMeshPtr]()
-	{
-		if (USkeletalMesh* SkeletalMesh = WeakSkeletalMeshPtr.Get())
-		{
-			FScopedTransaction ScopedTransaction(LOCTEXT("CopySkinWeightProfile", "Create Skin Weight Profile with Skin Weights from another Skeletal Mesh"));
-			SkeletalMesh->Modify();
-
-			FSkinWeightProfileHelpers::CopySkinWeightProfile(SkeletalMesh);
-			MeshDetailLayout->ForceRefreshDetails();
-		}
-	})));
-
+	
 	// Add extra (sub)-menus for previously added Skin Weight Profiles
 	if (USkeletalMesh* Mesh = SkeletalMeshPtr.Get())
 	{
diff --git a/Engine/Source/Editor/Persona/Private/SSkinWeightProfileImportOptions.cpp b/Engine/Source/Editor/Persona/Private/SSkinWeightProfileImportOptions.cpp
--- a/Engine/Source/Editor/Persona/Private/SSkinWeightProfileImportOptions.cpp
+++ b/Engine/Source/Editor/Persona/Private/SSkinWeightProfileImportOptions.cpp
@@ -19,26 +19,8 @@
 
 void FSkinWeightImportOptionsCustomization::CustomizeDetails(IDetailLayoutBuilder& DetailBuilder)
 {
-	TSharedRef<IPropertyHandle> FilePathHandle = DetailBuilder.GetProperty(GET_MEMBER_NAME_CHECKED(USkinWeightImportOptions, FilePath));
-	TSharedRef<IPropertyHandle> SourceMeshHandle = DetailBuilder.GetProperty(GET_MEMBER_NAME_CHECKED(USkinWeightImportOptions, SourceMesh));
 	TSharedPtr<IPropertyHandle> LODIndexHandle = DetailBuilder.GetProperty(GET_MEMBER_NAME_CHECKED(USkinWeightImportOptions, LODIndex));
-
-	FString FilePath;
-
-	// In case we are importing from a file do not show the Skeletal Mesh property and vice-versa
-	if (FilePathHandle->GetValue(FilePath) == FPropertyAccess::Success)
-	{
-		if (!FilePath.IsEmpty())
-		{
-			DetailBuilder.HideProperty(SourceMeshHandle);
-		}
-		else
-		{
-			DetailBuilder.HideProperty(FilePathHandle);
-			DetailBuilder.HideProperty(LODIndexHandle);
-		}
-	}
-
+	
 	ProfileNameHandle = DetailBuilder.GetProperty(GET_MEMBER_NAME_CHECKED(USkinWeightImportOptions, ProfileName));
 	DetailBuilder.HideProperty(ProfileNameHandle);
 
@@ -164,12 +146,8 @@ void SSkinWeightProfileImportOptions::Construct(const FArguments& InArgs)
 				.OnClicked(this, &SSkinWeightProfileImportOptions::OnImport)
 				.IsEnabled_Lambda([this]() -> bool 
 				{
-					const bool bCopyOperation = ImportSettings->FilePath.IsEmpty();
 					const bool bValidProfileName = DetailsCustomization->IsProfileNameValid(DetailsCustomization->OnGetProfileName().ToString());
-					const bool bValidSkeletalMesh = ImportSettings->SourceMesh != nullptr;
-
-					// Need a valid profile name, and in case we are copying we'll need to have a valid source skeletal mesh as well 
-					return bValidProfileName && (!bCopyOperation || bValidSkeletalMesh);
+					return bValidProfileName;
 				})
 			]
 			+ SUniformGridPanel::Slot(1, 0)
diff --git a/Engine/Source/Editor/Persona/Private/SkinWeightProfileHelpers.cpp b/Engine/Source/Editor/Persona/Private/SkinWeightProfileHelpers.cpp
--- a/Engine/Source/Editor/Persona/Private/SkinWeightProfileHelpers.cpp
+++ b/Engine/Source/Editor/Persona/Private/SkinWeightProfileHelpers.cpp
@@ -39,7 +39,6 @@ void FSkinWeightProfileHelpers::ImportSkinWeightProfile(USkeletalMesh* InSkeleta
 
 			TSharedPtr<SSkinWeightProfileImportOptions> OptionsWidget;
 			ImportSettings->FilePath = PickedFileName;
-			ImportSettings->SourceMesh = nullptr;
 			ImportSettings->LODIndex = 0;
 			Window->SetContent
 			(
@@ -65,87 +64,30 @@ void FSkinWeightProfileHelpers::ImportSkinWeightProfile(USkeletalMesh* InSkeleta
 		
 		if (bShouldImport)
 		{
-			const FName ProfileName(*ImportSettings->ProfileName);
-			if (ImportSettings->SourceMesh == nullptr)
+			const FName ProfileName(*ImportSettings->ProfileName);			
+			// Try and import the skin weight profile from the provided FBX file path
+			const bool bReregisterComponent = true;
+			const bool bResult = FLODUtilities::ImportAlternateSkinWeight(InSkeletalMesh, PickedFileName, ImportSettings->LODIndex, ProfileName, bReregisterComponent);
+			if (bResult)
 			{
-				// Try and import the skin weight profile from the provided FBX file path
-				const bool bReregisterComponent = true;
-				const bool bResult = FLODUtilities::ImportAlternateSkinWeight(InSkeletalMesh, PickedFileName, ImportSettings->LODIndex, ProfileName, bReregisterComponent);
-				if (bResult)
-				{
-					FNotificationInfo NotificationInfo(FText::GetEmpty());
-					NotificationInfo.Text = LOCTEXT("ImportSuccessful", "Skin Weights imported successfully!");
-					NotificationInfo.ExpireDuration = 2.5f;
-					FSlateNotificationManager::Get().AddNotification(NotificationInfo);
-				}
-				else
-				{
-					FNotificationInfo NotificationInfo(FText::GetEmpty());
-					NotificationInfo.Text = LOCTEXT("ImportFailed", "Failed to import Skin Weights!");
-					NotificationInfo.ExpireDuration = 2.5f;
-					FSlateNotificationManager::Get().AddNotification(NotificationInfo);
-				}
+				FNotificationInfo NotificationInfo(FText::GetEmpty());
+				NotificationInfo.Text = LOCTEXT("ImportSuccessful", "Skin Weights imported successfully!");
+				NotificationInfo.ExpireDuration = 2.5f;
+				FSlateNotificationManager::Get().AddNotification(NotificationInfo);
 			}
 			else
 			{
-				if (USkeletalMesh* Mesh = ImportSettings->SourceMesh)
-				{
-					if (CopySkinWeightsToProfile(Mesh, InSkeletalMesh, LOD0Index, ProfileName))
-					{
-						// Add a new profile entry
-						FSkinWeightProfileInfo SkeletalMeshProfile;
-						SkeletalMeshProfile.DefaultProfile = (InSkeletalMesh->GetNumSkinWeightProfiles() == 0);
-						SkeletalMeshProfile.DefaultProfileFromLODIndex = LOD0Index;
-						SkeletalMeshProfile.Name = ProfileName;
-						InSkeletalMesh->AddSkinWeightProfile(SkeletalMeshProfile);
-
-						// Enforce a post-edit call to ensure the mesh data is cached
-						FPropertyChangedEvent PropertyChangeEvent(nullptr);
-						InSkeletalMesh->PostEditChangeProperty(PropertyChangeEvent);
-					}
-				}				
+				FNotificationInfo NotificationInfo(FText::GetEmpty());
+				NotificationInfo.Text = LOCTEXT("ImportFailed", "Failed to import Skin Weights!");
+				NotificationInfo.ExpireDuration = 2.5f;
+				FSlateNotificationManager::Get().AddNotification(NotificationInfo);
 			}
+			
 			FLODUtilities::RegenerateDependentLODs(InSkeletalMesh, 0);
 		}
 	}
 }
 
-bool FSkinWeightProfileHelpers::CopySkinWeightsToProfile(USkeletalMesh* SourceMesh, USkeletalMesh* TargetMesh, const int32 LODIndex, const FName ProfileName)
-{
-	bool bCopied = false;
-	const FSkeletalMeshModel* SourceModel = SourceMesh->GetImportedModel();
-	FSkeletalMeshModel* TargetModel = TargetMesh->GetImportedModel();
-	if (SourceModel && TargetModel)
-	{
-		if (SourceModel->LODModels.IsValidIndex(LODIndex) && TargetModel->LODModels.IsValidIndex(LODIndex))
-		{
-			const FSkeletalMeshLODModel& SourceLODModel = SourceModel->LODModels[LODIndex];
-			FSkeletalMeshLODModel& TargetLODModel = TargetModel->LODModels[LODIndex];
-			if (SourceLODModel.NumVertices == TargetLODModel.NumVertices)
-			{
-				// Copy actual per-lod data into profile entry
-				TArray<FSoftSkinVertex> SourceVertices;
-				SourceLODModel.GetVertices(SourceVertices);
-
-				// Add a new imported skin weight data entry
-				FImportedSkinWeightProfileData& ProfileData = TargetLODModel.SkinWeightProfiles.Add(ProfileName);
-				TArray<FRawSkinWeight>& SkinWeights = ProfileData.SkinWeights;
-				SkinWeights.Reserve(SourceVertices.Num());
-				for (const FSoftSkinVertex& Vertex : SourceVertices)
-				{
-					FRawSkinWeight& SkinWeight = SkinWeights.AddDefaulted_GetRef();
-					FMemory::Memcpy(SkinWeight.InfluenceBones, Vertex.InfluenceBones);
-					FMemory::Memcpy(SkinWeight.InfluenceWeights, Vertex.InfluenceWeights);
-				}
-
-				bCopied = true;
-			}
-		}
-	}
-
-	return bCopied;
-}
-
 void FSkinWeightProfileHelpers::ImportSkinWeightProfileLOD(USkeletalMesh* InSkeletalMesh, FName ProfileName, int32 LODIndex)
 {
 	if (InSkeletalMesh)
@@ -174,68 +116,6 @@ void FSkinWeightProfileHelpers::ImportSkinWeightProfileLOD(USkeletalMesh* InSkel
 	}
 }
 
-void FSkinWeightProfileHelpers::CopySkinWeightProfile(USkeletalMesh* InSkeletalMesh)
-{
-	if (InSkeletalMesh)
-	{
-		bool bShouldImport = false;
-		USkinWeightImportOptions* ImportSettings = GetMutableDefault<USkinWeightImportOptions>();
-
-		// Show settings dialog	
-		TSharedRef<SWindow> Window = SNew(SWindow)
-			.Title(LOCTEXT("WindowTitle", "Skin Weight Profile Import Options"))
-			.SizingRule(ESizingRule::Autosized);
-
-		TSharedPtr<SSkinWeightProfileImportOptions> OptionsWidget;
-		ImportSettings->FilePath.Empty();
-		ImportSettings->SourceMesh = nullptr;
-		Window->SetContent
-		(
-			SAssignNew(OptionsWidget, SSkinWeightProfileImportOptions).WidgetWindow(Window)
-			.ImportSettings(ImportSettings)
-			.WidgetWindow(Window)
-			.SkeletalMesh(InSkeletalMesh)
-		);
-
-		TSharedPtr<SWindow> ParentWindow;
-
-		if (FModuleManager::Get().IsModuleLoaded("MainFrame"))
-		{
-			IMainFrameModule& MainFrame = FModuleManager::LoadModuleChecked<IMainFrameModule>("MainFrame");
-			ParentWindow = MainFrame.GetParentWindow();
-		}
-
-		FSlateApplication::Get().AddModalWindow(Window, ParentWindow, false);
-
-		// Poll the result of the window interaction, user could have cancelled operation
-		bShouldImport = OptionsWidget->ShouldImport();
-
-		if (bShouldImport)
-		{
-			const int32 LOD0Index = 0;
-			const FName ProfileName(*ImportSettings->ProfileName);
-			if (USkeletalMesh* Mesh = ImportSettings->SourceMesh)
-			{				
-				if (CopySkinWeightsToProfile(Mesh, InSkeletalMesh, LOD0Index, ProfileName))
-				{
-					// Add a new profile entry
-					FSkinWeightProfileInfo SkeletalMeshProfile;
-					SkeletalMeshProfile.DefaultProfile = (InSkeletalMesh->GetNumSkinWeightProfiles() == 0);
-					SkeletalMeshProfile.DefaultProfileFromLODIndex = LOD0Index;
-					SkeletalMeshProfile.Name = ProfileName;
-					InSkeletalMesh->AddSkinWeightProfile(SkeletalMeshProfile);
-
-					// Enforce a post-edit call to ensure the mesh data is cached
-					FPropertyChangedEvent PropertyChangeEvent(nullptr);
-					InSkeletalMesh->PostEditChangeProperty(PropertyChangeEvent);
-
-					// Make sure we regenerate any LOD data that is based off the now re imported LOD
-					FLODUtilities::RegenerateDependentLODs(InSkeletalMesh, LOD0Index);
-				}
-			}
-		}
-	}
-}
 
 void FSkinWeightProfileHelpers::ReimportSkinWeightProfileLOD(USkeletalMesh* InSkeletalMesh, const FName& InProfileName, const int32 LODIndex)
 {
diff --git a/Engine/Source/Editor/Persona/Public/SSkinWeightProfileImportOptions.h b/Engine/Source/Editor/Persona/Public/SSkinWeightProfileImportOptions.h
--- a/Engine/Source/Editor/Persona/Public/SSkinWeightProfileImportOptions.h
+++ b/Engine/Source/Editor/Persona/Public/SSkinWeightProfileImportOptions.h
@@ -32,13 +32,9 @@ public:
 	/** File path to FBX file containing Mesh with alternative set of Skin Weights */
 	UPROPERTY(VisibleAnywhere, Category = SkinWeights)
 	FString FilePath;
-
-	/** Skeletal Mesh from which to copy the Skin Weights used for this Profile */
-	UPROPERTY(EditAnywhere, Category = SkinWeights)
-	USkeletalMesh* SourceMesh;
-
+	
 	/** Target LOD index this file corresponds to */
-	UPROPERTY(EditAnywhere, Category = SkinWeights)
+	UPROPERTY(EditAnywhere, Category = SkinWeights, meta=(DisplayName="LOD Index"))
 	int32 LODIndex;
 };
 
diff --git a/Engine/Source/Editor/Persona/Public/SkinWeightProfileHelpers.h b/Engine/Source/Editor/Persona/Public/SkinWeightProfileHelpers.h
--- a/Engine/Source/Editor/Persona/Public/SkinWeightProfileHelpers.h
+++ b/Engine/Source/Editor/Persona/Public/SkinWeightProfileHelpers.h
@@ -16,9 +16,6 @@ struct FSkinWeightProfileHelpers
 	/** Tries to import a new set of Skin Weights for the given Skeletal Mesh at the given LOD index from an FBX file */
 	static void ImportSkinWeightProfileLOD(USkeletalMesh* InSkeletalMesh, FName ProfileName, int32 LODIndex);
 
-	/** Tries to copy a new set of Skin Weights for the given Skeletal Mesh from another Skeletal Mesh */
-	static void CopySkinWeightProfile(USkeletalMesh* InSkeletalMesh);
-
 	/** Tries to re import the previously imported skin weights for the given Skeletal Mesh, Profile name and LOD Index */
 	static void ReimportSkinWeightProfileLOD(USkeletalMesh* InSkeletalMesh, const FName& InProfileName, const int32 LODIndex);
 
