Update skin weight profile import button tooltip to reflect that copy is no longer an option
#rnx

diff --git a/Engine/Source/Editor/Persona/Private/PersonaMeshDetails.cpp b/Engine/Source/Editor/Persona/Private/PersonaMeshDetails.cpp
--- a/Engine/Source/Editor/Persona/Private/PersonaMeshDetails.cpp
+++ b/Engine/Source/Editor/Persona/Private/PersonaMeshDetails.cpp
@@ -1247,8 +1247,8 @@ void FPersonaMeshDetails::CustomizeSkinWeightProfiles(IDetailLayoutBuilder& Deta
 				.Image(FEditorStyle::GetBrush("PropertyWindow.Button_AddToArray"))
 			]
 			.OnGetMenuContent(this, &FPersonaMeshDetails::CreateSkinWeightProfileMenuContent)
-			.ToolTipText(LOCTEXT("ImportSkinWeightButtonToolTip", "Import a new Skin Weight Profile, or copy Skin Weights from another Skeletal Mesh"))
-		]		
+			.ToolTipText(LOCTEXT("ImportSkinWeightButtonToolTip", "Import a new Skin Weight Profile"))
+		]
 	];	
 }
 
