Fixed robomerge conflict

Add the ability to preview of layer overlays and sub-instances in the animation blueprint editor

This change adds a new 'preview animation blueprint' to UAnimBlueprint. This is a blueprint that hosts a sub-instance/layer BP for preview purposes. Setting this in a BP will show a preview in the viewport of the BP running as a sub-instance of the preview instance, thus allowing for sub-instances to be previewed in context.

diff --git a/Engine/Source/Editor/AnimationBlueprintEditor/Private/AnimationBlueprintEditor.cpp b/Engine/Source/Editor/AnimationBlueprintEditor/Private/AnimationBlueprintEditor.cpp
--- a/Engine/Source/Editor/AnimationBlueprintEditor/Private/AnimationBlueprintEditor.cpp
+++ b/Engine/Source/Editor/AnimationBlueprintEditor/Private/AnimationBlueprintEditor.cpp
@@ -74,6 +74,7 @@
 #include "EdGraphNode_Comment.h"
 #include "AnimStateNodeBase.h"
 #include "AnimStateEntryNode.h"
+#include "PersonaUtils.h"
 #include "Subsystems/AssetEditorSubsystem.h"
 
 #define LOCTEXT_NAMESPACE "AnimationBlueprintEditor"
@@ -282,10 +283,19 @@ void FAnimationBlueprintEditor::InitAnimationBlueprintEditor(const EToolkitMode:
 
 		UDebugSkelMeshComponent* PreviewMeshComponent = PersonaToolkit->GetPreviewMeshComponent();
 		UAnimBlueprint* AnimBlueprint = PersonaToolkit->GetAnimBlueprint();
-		PreviewMeshComponent->SetAnimInstanceClass(AnimBlueprint ? AnimBlueprint->GeneratedClass : NULL);
+		UAnimBlueprint* PreviewAnimBlueprint = AnimBlueprint->GetPreviewAnimationBlueprint();
+		
+		if (PreviewAnimBlueprint)
+		{
+			PersonaToolkit->GetPreviewScene()->SetPreviewAnimationBlueprint(PreviewAnimBlueprint, AnimBlueprint);
+			PreviewAnimBlueprint->OnCompiled().AddSP(this, &FAnimationBlueprintEditor::HandlePreviewAnimBlueprintCompiled);
+		}
+		else
+		{
+			PersonaToolkit->GetPreviewScene()->SetPreviewAnimationBlueprint(AnimBlueprint, nullptr);
+		}
 
-		// Make sure the object being debugged is the preview instance
-		AnimBlueprint->SetObjectBeingDebugged(PreviewMeshComponent->GetAnimInstance());
+		PersonaUtils::SetObjectBeingDebugged(AnimBlueprint, PreviewMeshComponent->GetAnimInstance());
 
 		ExtendMenu();
 		ExtendToolbar();
@@ -1152,12 +1162,19 @@ void FAnimationBlueprintEditor::Compile()
 			DebuggedMeshComponent->InitAnim(true);
 		}
 
-		UAnimInstance* NewInstance = DebuggedMeshComponent->GetAnimInstance();
-		UBlueprint* Blueprint = GetBlueprintObj();
+		// re-apply preview anim bp if needed
+		UAnimBlueprint* AnimBlueprint = GetAnimBlueprint();
+		UAnimBlueprint* PreviewAnimBlueprint = AnimBlueprint->GetPreviewAnimationBlueprint();
+		
+		if (PreviewAnimBlueprint)
+		{
+			PersonaToolkit->GetPreviewScene()->SetPreviewAnimationBlueprint(PreviewAnimBlueprint, AnimBlueprint);
+		}
 
-		if(NewInstance->IsA(Blueprint->GeneratedClass))
+		UAnimInstance* NewInstance = DebuggedMeshComponent->GetAnimInstance();
+		if(AnimBlueprint && (NewInstance->IsA(AnimBlueprint->GeneratedClass) || (PreviewAnimBlueprint != nullptr && NewInstance->IsA(PreviewAnimBlueprint->GeneratedClass))))
 		{
-			Blueprint->SetObjectBeingDebugged(NewInstance);
+			PersonaUtils::SetObjectBeingDebugged(AnimBlueprint, NewInstance);
 		}
 	}
 
@@ -1268,7 +1285,23 @@ void FAnimationBlueprintEditor::GetCustomDebugObjects(TArray<FCustomDebugObject>
 	UDebugSkelMeshComponent* PreviewMeshComponent = PersonaToolkit->GetPreviewMeshComponent();
 	if (PreviewMeshComponent->IsAnimBlueprintInstanced())
 	{
-		new (DebugList) FCustomDebugObject(PreviewMeshComponent->GetAnimInstance(), LOCTEXT("PreviewObjectLabel", "Preview Instance").ToString());
+		UAnimInstance* PreviewInstance = PreviewMeshComponent->GetAnimInstance();
+		UAnimBlueprint* AnimBlueprint = GetAnimBlueprint();
+		UAnimBlueprint* PreviewAnimBlueprint = AnimBlueprint->GetPreviewAnimationBlueprint();
+		if (PreviewAnimBlueprint)
+		{
+			EPreviewAnimationBlueprintApplicationMethod ApplicationMethod = AnimBlueprint->GetPreviewAnimationBlueprintApplicationMethod();
+			if(ApplicationMethod == EPreviewAnimationBlueprintApplicationMethod::OverlayLayer)
+			{
+				PreviewInstance = PreviewInstance->GetLayerSubInstanceByClass(AnimBlueprint->GeneratedClass.Get());
+			}
+			else if(ApplicationMethod == EPreviewAnimationBlueprintApplicationMethod::SubInstance)
+			{
+				PreviewInstance = PreviewInstance->GetSubInstanceByTag(AnimBlueprint->GetPreviewAnimationBlueprintTag());
+			}
+		}
+
+		new (DebugList) FCustomDebugObject(PreviewInstance, LOCTEXT("PreviewObjectLabel", "Preview Instance").ToString());
 	}
 }
 
@@ -1770,5 +1803,15 @@ void FAnimationBlueprintEditor::SaveEditorSettings()
 	}
 }
 
+void FAnimationBlueprintEditor::HandlePreviewAnimBlueprintCompiled(UBlueprint* InBlueprint)
+{
+	UAnimBlueprint* AnimBlueprint = GetAnimBlueprint();
+	UAnimBlueprint* PreviewAnimBlueprint = AnimBlueprint->GetPreviewAnimationBlueprint();
+	if (PreviewAnimBlueprint)
+	{
+		GetPreviewScene()->SetPreviewAnimationBlueprint(PreviewAnimBlueprint, AnimBlueprint);
+	}
+}
+
 #undef LOCTEXT_NAMESPACE
 
diff --git a/Engine/Source/Editor/AnimationBlueprintEditor/Private/AnimationBlueprintEditor.h b/Engine/Source/Editor/AnimationBlueprintEditor/Private/AnimationBlueprintEditor.h
--- a/Engine/Source/Editor/AnimationBlueprintEditor/Private/AnimationBlueprintEditor.h
+++ b/Engine/Source/Editor/AnimationBlueprintEditor/Private/AnimationBlueprintEditor.h
@@ -274,6 +274,9 @@ private:
 	/** Handle the viewport being created */
 	void HandleViewportCreated(const TSharedRef<IPersonaViewport>& InPersonaViewport);
 
+	/** Handle the preview anim blueprint being compiled */
+	void HandlePreviewAnimBlueprintCompiled(UBlueprint* InBlueprint);
+
     /**
 	 * Load editor settings from disk (docking state, window pos/size, option state, etc).
 	 */
diff --git a/Engine/Source/Editor/AnimationBlueprintEditor/Private/AnimationBlueprintEditorMode.cpp b/Engine/Source/Editor/AnimationBlueprintEditor/Private/AnimationBlueprintEditorMode.cpp
--- a/Engine/Source/Editor/AnimationBlueprintEditor/Private/AnimationBlueprintEditorMode.cpp
+++ b/Engine/Source/Editor/AnimationBlueprintEditor/Private/AnimationBlueprintEditorMode.cpp
@@ -12,6 +12,8 @@
 #include "PersonaModule.h"
 #include "SBlueprintEditorToolbar.h"
 #include "IPersonaPreviewScene.h"
+#include "PersonaUtils.h"
+
 /////////////////////////////////////////////////////
 // FAnimationBlueprintEditorMode
 
@@ -173,7 +175,7 @@ void FAnimationBlueprintEditorMode::PostActivateMode()
 		UDebugSkelMeshComponent* PreviewComponent = PreviewScenePtr.Pin()->GetPreviewMeshComponent();
 		if ((AnimBlueprint->GetObjectBeingDebugged() == NULL) && (PreviewComponent->IsAnimBlueprintInstanced()))
 		{
-			AnimBlueprint->SetObjectBeingDebugged(PreviewComponent->GetAnimInstance());
+			PersonaUtils::SetObjectBeingDebugged(AnimBlueprint, PreviewComponent->GetAnimInstance());
 		}
 
 		// If we are a derived anim blueprint always show the overrides tab
diff --git a/Engine/Source/Editor/Persona/Private/AnimationEditorPreviewScene.cpp b/Engine/Source/Editor/Persona/Private/AnimationEditorPreviewScene.cpp
--- a/Engine/Source/Editor/Persona/Private/AnimationEditorPreviewScene.cpp
+++ b/Engine/Source/Editor/Persona/Private/AnimationEditorPreviewScene.cpp
@@ -76,6 +76,13 @@ FAnimationEditorPreviewScene::FAnimationEditorPreviewScene(const ConstructionVal
 		PreviewSceneDescription->AdditionalMeshes = InEditableSkeleton->GetSkeleton().GetAdditionalPreviewSkeletalMeshes();
 	}
 
+	if(UAnimBlueprint* AnimBlueprint = InPersonaToolkit->GetAnimBlueprint())
+	{
+		PreviewSceneDescription->PreviewAnimationBlueprint = AnimBlueprint->GetPreviewAnimationBlueprint();
+		PreviewSceneDescription->ApplicationMethod = AnimBlueprint->GetPreviewAnimationBlueprintApplicationMethod();
+		PreviewSceneDescription->SubInstanceTag = AnimBlueprint->GetPreviewAnimationBlueprintTag();
+	}
+
 	// create a default additional mesh collection so we dont always have to create an asset to edit additional meshes
 	UPreviewMeshCollectionFactory* FactoryToUse = NewObject<UPreviewMeshCollectionFactory>();
 	if(InEditableSkeleton.IsValid())
@@ -167,6 +174,27 @@ void FAnimationEditorPreviewScene::SetPreviewMesh(USkeletalMesh* NewPreviewMesh,
 	RefreshAdditionalMeshes(bAllowOverrideBaseMesh);
 }
 
+void FAnimationEditorPreviewScene::SetPreviewAnimationBlueprint(UAnimBlueprint* InAnimBlueprint, UAnimBlueprint* InOverlayOrSubAnimBlueprint)
+{
+	if (InOverlayOrSubAnimBlueprint)
+	{
+		SkeletalMeshComponent->SetAnimInstanceClass(InAnimBlueprint ? InAnimBlueprint->GeneratedClass : nullptr);
+		EPreviewAnimationBlueprintApplicationMethod ApplicationMethod = InOverlayOrSubAnimBlueprint->GetPreviewAnimationBlueprintApplicationMethod();
+		if(ApplicationMethod == EPreviewAnimationBlueprintApplicationMethod::OverlayLayer)
+		{
+			SkeletalMeshComponent->SetLayerOverlay(InOverlayOrSubAnimBlueprint ? InOverlayOrSubAnimBlueprint->GeneratedClass.Get() : nullptr);
+		}
+		else if(ApplicationMethod == EPreviewAnimationBlueprintApplicationMethod::SubInstance)
+		{
+			SkeletalMeshComponent->SetSubInstanceClassByTag(InOverlayOrSubAnimBlueprint->GetPreviewAnimationBlueprintTag(), InOverlayOrSubAnimBlueprint ? InOverlayOrSubAnimBlueprint->GeneratedClass.Get() : nullptr);
+		}
+	}
+	else
+	{
+		SkeletalMeshComponent->SetAnimInstanceClass(InAnimBlueprint ? InAnimBlueprint->GeneratedClass : nullptr);
+	}
+}
+
 USkeletalMesh* FAnimationEditorPreviewScene::GetPreviewMesh() const
 {
 	return PreviewSceneDescription->PreviewMesh.Get();
@@ -246,7 +274,7 @@ void FAnimationEditorPreviewScene::SetPreviewMeshInternal(USkeletalMesh* NewPrev
 	if (DebuggedSkeletalMeshComponent)
 	{
 		UAnimBlueprint* SourceBlueprint = PersonaToolkit.Pin()->GetAnimBlueprint();
-		SourceBlueprint->SetObjectBeingDebugged(DebuggedSkeletalMeshComponent->GetAnimInstance());
+		PersonaUtils::SetObjectBeingDebugged(SourceBlueprint, DebuggedSkeletalMeshComponent->GetAnimInstance());
 	}
 
 	OnPreviewMeshChanged.Broadcast(OldPreviewMesh, NewPreviewMesh);
@@ -673,10 +701,19 @@ void FAnimationEditorPreviewScene::ShowDefaultMode()
 		{
 			SkeletalMeshComponent->EnablePreview(false, nullptr);
 
-			UAnimBlueprint* AnimBP = PersonaToolkit.Pin()->GetAnimBlueprint();
-			if (AnimBP)
+			UAnimBlueprint* AnimBlueprint = PersonaToolkit.Pin()->GetAnimBlueprint();
+			if (AnimBlueprint)
 			{
-				SkeletalMeshComponent->SetAnimInstanceClass(AnimBP->GeneratedClass);
+				UAnimBlueprint* PreviewAnimBlueprint = AnimBlueprint->GetPreviewAnimationBlueprint();
+		
+				if (PreviewAnimBlueprint)
+				{
+					SetPreviewAnimationBlueprint(PreviewAnimBlueprint, AnimBlueprint);
+				}
+				else
+				{
+					SetPreviewAnimationBlueprint(AnimBlueprint, nullptr);
+				}
 			}
 		}
 		break;
diff --git a/Engine/Source/Editor/Persona/Private/AnimationEditorPreviewScene.h b/Engine/Source/Editor/Persona/Private/AnimationEditorPreviewScene.h
--- a/Engine/Source/Editor/Persona/Private/AnimationEditorPreviewScene.h
+++ b/Engine/Source/Editor/Persona/Private/AnimationEditorPreviewScene.h
@@ -26,6 +26,7 @@ public:
 	virtual UAnimationAsset* GetPreviewAnimationAsset() const override;
 	virtual void SetPreviewMesh(USkeletalMesh* NewPreviewMesh, bool bAllowOverrideBaseMesh = true) override;
 	virtual USkeletalMesh* GetPreviewMesh() const override;
+	virtual void SetPreviewAnimationBlueprint(UAnimBlueprint* InAnimBlueprint, UAnimBlueprint* InOverlayOrSubAnimBlueprint) override;
 	virtual bool AttachObjectToPreviewComponent(UObject* Object, FName AttachTo) override;
 	virtual void RemoveAttachedObjectFromPreviewComponent(UObject* Object, FName AttachedTo) override;
 	virtual void InvalidateViews() override;
diff --git a/Engine/Source/Editor/Persona/Private/PersonaPreviewSceneDescription.h b/Engine/Source/Editor/Persona/Private/PersonaPreviewSceneDescription.h
--- a/Engine/Source/Editor/Persona/Private/PersonaPreviewSceneDescription.h
+++ b/Engine/Source/Editor/Persona/Private/PersonaPreviewSceneDescription.h
@@ -16,6 +16,7 @@ class UPreviewMeshCollection;
 class USkeletalMesh;
 class UDataAsset;
 class UPreviewMeshCollection;
+enum class EPreviewAnimationBlueprintApplicationMethod : uint8;
 
 UCLASS()
 class UPersonaPreviewSceneDescription : public UObject
@@ -37,6 +38,18 @@ public:
 	UPROPERTY(EditAnywhere, Category = "Mesh", meta=(DisplayThumbnail=true))
 	TSoftObjectPtr<USkeletalMesh> PreviewMesh;
 
+	/** The preview anim blueprint to use */
+	UPROPERTY(EditAnywhere, Category = "Animation Blueprint", meta=(DisplayThumbnail=true))
+	TSoftObjectPtr<UAnimBlueprint> PreviewAnimationBlueprint;
+
+	/** The method by which a preview animation blueprint is applied, either as an overlay layer, or as a sub-instance */
+	UPROPERTY(EditAnywhere, Category = "Animation Blueprint")
+	EPreviewAnimationBlueprintApplicationMethod ApplicationMethod;
+
+	/** The tag to use when applying a preview animation blueprint via SetSubInstanceClassByTag */
+	UPROPERTY(EditAnywhere, Category = "Animation Blueprint")
+	FName SubInstanceTag;
+
 	UPROPERTY(EditAnywhere, Category = "Additional Meshes")
 	TSoftObjectPtr<UDataAsset> AdditionalMeshes;
 
diff --git a/Engine/Source/Editor/Persona/Private/PersonaToolkit.cpp b/Engine/Source/Editor/Persona/Private/PersonaToolkit.cpp
--- a/Engine/Source/Editor/Persona/Private/PersonaToolkit.cpp
+++ b/Engine/Source/Editor/Persona/Private/PersonaToolkit.cpp
@@ -330,6 +330,32 @@ void FPersonaToolkit::SetPreviewMesh(class USkeletalMesh* InSkeletalMesh, bool b
 	}
 }
 
+void FPersonaToolkit::SetPreviewAnimationBlueprint(UAnimBlueprint* InAnimBlueprint)
+{
+	// Only allowed for anim blueprints
+	if (InitialAssetClass == UAnimBlueprint::StaticClass())
+	{
+		FScopedTransaction Transaction(NSLOCTEXT("PersonaToolkit", "SetAnimBlueprintPreviewBlueprint", "Set Animation Blueprint Preview Blueprint"));
+
+		check(AnimBlueprint);
+		AnimBlueprint->SetPreviewAnimationBlueprint(InAnimBlueprint);
+
+		// Note setting the 'edited' blueprint as an overlay here
+		GetPreviewScene()->SetPreviewAnimationBlueprint(InAnimBlueprint, AnimBlueprint);
+	}
+}
+
+UAnimBlueprint* FPersonaToolkit::GetPreviewAnimationBlueprint() const
+{
+	if (InitialAssetClass == UAnimBlueprint::StaticClass())
+	{
+		check(AnimBlueprint);
+		return AnimBlueprint->GetPreviewAnimationBlueprint();
+	}
+
+	return nullptr;
+}
+
 int32 FPersonaToolkit::GetCustomData(const int32 Key) const
 {
 	if (!CustomEditorData.Contains(Key))
diff --git a/Engine/Source/Editor/Persona/Private/PersonaToolkit.h b/Engine/Source/Editor/Persona/Private/PersonaToolkit.h
--- a/Engine/Source/Editor/Persona/Private/PersonaToolkit.h
+++ b/Engine/Source/Editor/Persona/Private/PersonaToolkit.h
@@ -43,6 +43,8 @@ public:
 	virtual class TSharedRef<IPersonaPreviewScene> GetPreviewScene() const override;
 	virtual class USkeletalMesh* GetPreviewMesh() const override;
 	virtual void SetPreviewMesh(class USkeletalMesh* InSkeletalMesh, bool bSetPreviewMeshInAsset = true) override;
+	virtual void SetPreviewAnimationBlueprint(UAnimBlueprint* InAnimBlueprint) override;
+	virtual UAnimBlueprint* GetPreviewAnimationBlueprint() const override;
 	virtual int32 GetCustomData(const int32 Key) const override;
 	virtual void SetCustomData(const int32 Key, const int32 CustomData) override;
 	virtual FName GetContext() const override;
diff --git a/Engine/Source/Editor/Persona/Private/PersonaUtils.cpp b/Engine/Source/Editor/Persona/Private/PersonaUtils.cpp
--- a/Engine/Source/Editor/Persona/Private/PersonaUtils.cpp
+++ b/Engine/Source/Editor/Persona/Private/PersonaUtils.cpp
@@ -9,6 +9,7 @@
 #include "Animation/AnimInstance.h"
 #include "Kismet2/BlueprintEditorUtils.h"
 #include "Animation/AnimNodeBase.h"
+#include "Animation/AnimBlueprint.h"
 
 namespace PersonaUtils
 {
@@ -97,4 +98,29 @@ int32 CopyPropertiesToCDO(UAnimInstance* InAnimInstance, const FCopyOptions& Opt
 	return CopiedPropertyCount;
 }
 
+void SetObjectBeingDebugged(UAnimBlueprint* InAnimBlueprint, UAnimInstance* InAnimInstance)
+{
+	UAnimBlueprint* PreviewAnimBlueprint = InAnimBlueprint->GetPreviewAnimationBlueprint();
+		
+	if (PreviewAnimBlueprint)
+	{
+		EPreviewAnimationBlueprintApplicationMethod ApplicationMethod = InAnimBlueprint->GetPreviewAnimationBlueprintApplicationMethod();
+		if(ApplicationMethod == EPreviewAnimationBlueprintApplicationMethod::OverlayLayer)
+		{
+			// Make sure the object being debugged is the overlay instance
+			InAnimBlueprint->SetObjectBeingDebugged(InAnimInstance->GetLayerSubInstanceByClass(InAnimBlueprint->GeneratedClass.Get()));
+		}
+		else if(ApplicationMethod == EPreviewAnimationBlueprintApplicationMethod::SubInstance)
+		{
+			// Make sure the object being debugged is the sub instance
+			InAnimBlueprint->SetObjectBeingDebugged(InAnimInstance->GetSubInstanceByTag(InAnimBlueprint->GetPreviewAnimationBlueprintTag()));
+		}
+	}
+	else
+	{
+		// Make sure the object being debugged is the preview instance
+		InAnimBlueprint->SetObjectBeingDebugged(InAnimInstance);
+	}
+}
+
 }
diff --git a/Engine/Source/Editor/Persona/Private/PreviewSceneCustomizations.cpp b/Engine/Source/Editor/Persona/Private/PreviewSceneCustomizations.cpp
--- a/Engine/Source/Editor/Persona/Private/PreviewSceneCustomizations.cpp
+++ b/Engine/Source/Editor/Persona/Private/PreviewSceneCustomizations.cpp
@@ -210,11 +210,93 @@ void FPreviewSceneDescriptionCustomization::CustomizeDetails(IDetailLayoutBuilde
 			SNew(SObjectPropertyEntryBox)
 			.AllowedClass(USkeletalMesh::StaticClass())
 			.PropertyHandle(SkeletalMeshProperty)
-			.OnShouldFilterAsset(this, &FPreviewSceneDescriptionCustomization::HandleShouldFilterAsset, PersonaToolkit.Pin()->GetContext() == UPhysicsAsset::StaticClass()->GetFName())
+			.OnShouldFilterAsset(this, &FPreviewSceneDescriptionCustomization::HandleShouldFilterAsset, FName("Skeleton"), PersonaToolkit.Pin()->GetContext() == UPhysicsAsset::StaticClass()->GetFName())
 			.OnObjectChanged(this, &FPreviewSceneDescriptionCustomization::HandleMeshChanged)
 			.ThumbnailPool(DetailBuilder.GetThumbnailPool())
 		];
 
+		// Customize animation blueprint preview
+		TSharedRef<IPropertyHandle> PreviewAnimationBlueprintProperty = DetailBuilder.GetProperty(GET_MEMBER_NAME_CHECKED(UPersonaPreviewSceneDescription, PreviewAnimationBlueprint));
+		TSharedRef<IPropertyHandle> ApplicationMethodProperty = DetailBuilder.GetProperty(GET_MEMBER_NAME_CHECKED(UPersonaPreviewSceneDescription, ApplicationMethod));
+		TSharedRef<IPropertyHandle> SubInstanceTagProperty = DetailBuilder.GetProperty(GET_MEMBER_NAME_CHECKED(UPersonaPreviewSceneDescription, SubInstanceTag));
+		
+		if (PersonaToolkit.Pin()->GetContext() == UAnimBlueprint::StaticClass()->GetFName())
+		{
+			DetailBuilder.EditCategory("Animation Blueprint")
+			.AddProperty(PreviewAnimationBlueprintProperty)
+			.CustomWidget()
+			.NameContent()
+			[
+				SNew(SVerticalBox)
+				+SVerticalBox::Slot()
+				.AutoHeight()
+				[
+					PreviewAnimationBlueprintProperty->CreatePropertyNameWidget()
+				]
+			]
+			.ValueContent()
+			.MaxDesiredWidth(250.0f)
+			.MinDesiredWidth(250.0f)
+			[
+				SNew(SObjectPropertyEntryBox)
+				.AllowedClass(UAnimBlueprint::StaticClass())
+				.PropertyHandle(PreviewAnimationBlueprintProperty)
+				.OnShouldFilterAsset(this, &FPreviewSceneDescriptionCustomization::HandleShouldFilterAsset, FName("TargetSkeleton"), false)
+				.OnObjectChanged(this, &FPreviewSceneDescriptionCustomization::HandlePreviewAnimBlueprintChanged)
+				.ThumbnailPool(DetailBuilder.GetThumbnailPool())
+			];
+
+			ApplicationMethodProperty->SetOnPropertyValueChanged(FSimpleDelegate::CreateLambda([this]()
+			{
+				FScopedTransaction Transaction(LOCTEXT("SetAnimationBlueprintApplicationMethod", "Set Application Method"));
+
+				TSharedPtr<IPersonaToolkit> PinnedPersonaToolkit = PersonaToolkit.Pin();
+				TSharedRef<FAnimationEditorPreviewScene> PreviewScene = StaticCastSharedRef<FAnimationEditorPreviewScene>(PinnedPersonaToolkit->GetPreviewScene());
+				UPersonaPreviewSceneDescription* PersonaPreviewSceneDescription = PreviewScene->GetPreviewSceneDescription();
+				PinnedPersonaToolkit->GetAnimBlueprint()->SetPreviewAnimationBlueprintApplicationMethod(PersonaPreviewSceneDescription->ApplicationMethod);
+				PreviewScene->SetPreviewAnimationBlueprint(PersonaPreviewSceneDescription->PreviewAnimationBlueprint.Get(), PinnedPersonaToolkit->GetAnimBlueprint());
+			}));
+
+			DetailBuilder.EditCategory("Animation Blueprint")
+			.AddProperty(ApplicationMethodProperty)
+			.IsEnabled(MakeAttributeLambda([this]()
+			{
+				TSharedPtr<IPersonaToolkit> PinnedPersonaToolkit = PersonaToolkit.Pin();
+				TSharedRef<FAnimationEditorPreviewScene> PreviewScene = StaticCastSharedRef<FAnimationEditorPreviewScene>(PinnedPersonaToolkit->GetPreviewScene());
+				UPersonaPreviewSceneDescription* PersonaPreviewSceneDescription = PreviewScene->GetPreviewSceneDescription();
+				
+				return PersonaPreviewSceneDescription->PreviewAnimationBlueprint.IsValid();
+			}));
+		
+			SubInstanceTagProperty->SetOnPropertyValueChanged(FSimpleDelegate::CreateLambda([this]()
+			{
+				FScopedTransaction Transaction(LOCTEXT("SetAnimationBlueprintTag", "Set Sub Instance Tag"));
+
+				TSharedPtr<IPersonaToolkit> PinnedPersonaToolkit = PersonaToolkit.Pin();
+				TSharedRef<FAnimationEditorPreviewScene> PreviewScene = StaticCastSharedRef<FAnimationEditorPreviewScene>(PinnedPersonaToolkit->GetPreviewScene());
+				UPersonaPreviewSceneDescription* PersonaPreviewSceneDescription = PreviewScene->GetPreviewSceneDescription();
+				PinnedPersonaToolkit->GetAnimBlueprint()->SetPreviewAnimationBlueprintTag(PersonaPreviewSceneDescription->SubInstanceTag);
+				PreviewScene->SetPreviewAnimationBlueprint(PersonaPreviewSceneDescription->PreviewAnimationBlueprint.Get(), PinnedPersonaToolkit->GetAnimBlueprint());
+			}));
+
+			DetailBuilder.EditCategory("Animation Blueprint")
+			.AddProperty(SubInstanceTagProperty)
+			.IsEnabled(MakeAttributeLambda([this]()
+			{
+				TSharedPtr<IPersonaToolkit> PinnedPersonaToolkit = PersonaToolkit.Pin();
+				TSharedRef<FAnimationEditorPreviewScene> PreviewScene = StaticCastSharedRef<FAnimationEditorPreviewScene>(PinnedPersonaToolkit->GetPreviewScene());
+				UPersonaPreviewSceneDescription* PersonaPreviewSceneDescription = PreviewScene->GetPreviewSceneDescription();
+				
+				return PersonaPreviewSceneDescription->PreviewAnimationBlueprint.IsValid() && PersonaPreviewSceneDescription->ApplicationMethod == EPreviewAnimationBlueprintApplicationMethod::SubInstance;
+			}));
+		}
+		else
+		{
+			PreviewAnimationBlueprintProperty->MarkHiddenByCustomization();
+			ApplicationMethodProperty->MarkHiddenByCustomization();
+			SubInstanceTagProperty->MarkHiddenByCustomization();
+		}
+
 		// set the skeleton to use in our factory as we shouldn't be picking one here
 		FactoryToUse->CurrentSkeleton = EditableSkeleton.IsValid() ? MakeWeakObjectPtr(const_cast<USkeleton*>(&EditableSkeleton.Pin()->GetSkeleton())) : nullptr;
 		TArray<UFactory*> FactoriesToUse({ FactoryToUse });
@@ -380,17 +462,17 @@ bool FPreviewSceneDescriptionCustomization::HandleShouldFilterAdditionalMesh(con
 		return true;
 	}
 
-	return HandleShouldFilterAsset(InAssetData, bCanUseDifferentSkeleton);
+	return HandleShouldFilterAsset(InAssetData, FName("Skeleton"), bCanUseDifferentSkeleton);
 }
 
-bool FPreviewSceneDescriptionCustomization::HandleShouldFilterAsset(const FAssetData& InAssetData, bool bCanUseDifferentSkeleton)
+bool FPreviewSceneDescriptionCustomization::HandleShouldFilterAsset(const FAssetData& InAssetData, FName InTag, bool bCanUseDifferentSkeleton)
 {
 	if (bCanUseDifferentSkeleton && GetDefault<UPersonaOptions>()->bAllowPreviewMeshCollectionsToSelectFromDifferentSkeletons)
 	{
 		return false;
 	}
 
-	FString SkeletonTag = InAssetData.GetTagValueRef<FString>("Skeleton");
+	FString SkeletonTag = InAssetData.GetTagValueRef<FString>(InTag);
 	if (SkeletonName.IsEmpty() || SkeletonTag == SkeletonName)
 	{
 		return false;
@@ -437,6 +519,12 @@ void FPreviewSceneDescriptionCustomization::HandleMeshChanged(const FAssetData&
 	PersonaToolkit.Pin()->SetPreviewMesh(NewPreviewMesh, false);
 }
 
+void FPreviewSceneDescriptionCustomization::HandlePreviewAnimBlueprintChanged(const FAssetData& InAssetData)   
+{
+	UAnimBlueprint* NewAnimBlueprint = Cast<UAnimBlueprint>(InAssetData.GetAsset());
+	PersonaToolkit.Pin()->SetPreviewAnimationBlueprint(NewAnimBlueprint);
+}
+
 void FPreviewSceneDescriptionCustomization::HandleAdditionalMeshesChanged(const FAssetData& InAssetData, IDetailLayoutBuilder* DetailLayoutBuilder)
 {
 	UDataAsset* MeshCollection = Cast<UDataAsset>(InAssetData.GetAsset());
diff --git a/Engine/Source/Editor/Persona/Private/PreviewSceneCustomizations.h b/Engine/Source/Editor/Persona/Private/PreviewSceneCustomizations.h
--- a/Engine/Source/Editor/Persona/Private/PreviewSceneCustomizations.h
+++ b/Engine/Source/Editor/Persona/Private/PreviewSceneCustomizations.h
@@ -47,7 +47,7 @@ private:
 
 	FReply OnSaveCollectionClicked(TSharedRef<IPropertyHandle> InAdditionalMeshesProperty, IDetailLayoutBuilder* DetailLayoutBuilder);
 
-	bool HandleShouldFilterAsset(const FAssetData& InAssetData, bool bCanUseDifferentSkeleton);
+	bool HandleShouldFilterAsset(const FAssetData& InAssetData, FName InTag, bool bCanUseDifferentSkeleton);
 
 	bool HandleShouldFilterAdditionalMesh(const FAssetData& InAssetData, bool bCanUseDifferentSkeleton);
 
@@ -65,6 +65,8 @@ private:
 
 	void HandleMeshChanged(const FAssetData& InAssetData);
 
+	void HandlePreviewAnimBlueprintChanged(const FAssetData& InAssetData);
+
 	void HandleAdditionalMeshesChanged(const FAssetData& InAssetData, IDetailLayoutBuilder* DetailLayoutBuilder);
 
 	void HandleAllowDifferentSkeletonsCheckedStateChanged(ECheckBoxState CheckState);
diff --git a/Engine/Source/Editor/Persona/Public/IPersonaPreviewScene.h b/Engine/Source/Editor/Persona/Public/IPersonaPreviewScene.h
--- a/Engine/Source/Editor/Persona/Public/IPersonaPreviewScene.h
+++ b/Engine/Source/Editor/Persona/Public/IPersonaPreviewScene.h
@@ -6,6 +6,7 @@
 #include "AdvancedPreviewScene.h"
 #include "Containers/ArrayView.h"
 #include "Types/SlateEnums.h"
+#include "Animation/AnimBlueprint.h"
 
 class UAnimationAsset;
 class UDebugSkelMeshComponent;
@@ -98,6 +99,9 @@ public:
 	/** Get the preview mesh for this scene (does go via skeleton/asset) */
 	virtual USkeletalMesh* GetPreviewMesh() const = 0;
 
+	/** Set the preview animation blueprint and an optional overlay (for sub-layers) */
+	virtual void SetPreviewAnimationBlueprint(UAnimBlueprint* InAnimBlueprint, UAnimBlueprint* InOverlayOrSubAnimBlueprint) = 0;
+
 	/** Show the reference pose of the displayed skeletal mesh. Otherwise display the default. Optionally reset bone transforms, if any. */
 	virtual void ShowReferencePose(bool bShowRefPose, bool bResetBoneTransforms = false) = 0;
 
diff --git a/Engine/Source/Editor/Persona/Public/IPersonaToolkit.h b/Engine/Source/Editor/Persona/Public/IPersonaToolkit.h
--- a/Engine/Source/Editor/Persona/Public/IPersonaToolkit.h
+++ b/Engine/Source/Editor/Persona/Public/IPersonaToolkit.h
@@ -4,6 +4,14 @@
 
 #include "CoreMinimal.h"
 
+class USkeleton;
+class IEditableSkeleton;
+class UDebugSkelMeshComponent;
+class USkeletalMesh;
+class UAnimBlueprint;
+class UAnimationAsset;
+class IPersonaPreviewScene;
+
 /**
  * Reusable 'Persona' features for asset editors concerned with USkeleton-related assets
  */
@@ -14,34 +22,34 @@ public:
 	virtual ~IPersonaToolkit() {}
 
 	/** Get the skeleton that we are editing */
-	virtual class USkeleton* GetSkeleton() const = 0;
+	virtual USkeleton* GetSkeleton() const = 0;
 
 	/** Get the editable skeleton that we are editing */
-	virtual TSharedPtr<class IEditableSkeleton> GetEditableSkeleton() const = 0;
+	virtual TSharedPtr<IEditableSkeleton> GetEditableSkeleton() const = 0;
 
 	/** Get the preview component that we are using */
-	virtual class UDebugSkelMeshComponent* GetPreviewMeshComponent() const = 0;
+	virtual UDebugSkelMeshComponent* GetPreviewMeshComponent() const = 0;
 
 	/** Get the skeletal mesh that we are editing */
-	virtual class USkeletalMesh* GetMesh() const = 0;
+	virtual USkeletalMesh* GetMesh() const = 0;
 
 	/** Set the skeletal mesh we are editing */
-	virtual void SetMesh(class USkeletalMesh* InSkeletalMesh) = 0;
+	virtual void SetMesh(USkeletalMesh* InSkeletalMesh) = 0;
 
 	/** Get the anim blueprint that we are editing */
-	virtual class UAnimBlueprint* GetAnimBlueprint() const = 0;
+	virtual UAnimBlueprint* GetAnimBlueprint() const = 0;
 
 	/** Get the animation asset that we are editing */
-	virtual class UAnimationAsset* GetAnimationAsset() const = 0;
+	virtual UAnimationAsset* GetAnimationAsset() const = 0;
 
 	/** Set the animation asset we are editing */
-	virtual void SetAnimationAsset(class UAnimationAsset* InAnimationAsset) = 0;
+	virtual void SetAnimationAsset(UAnimationAsset* InAnimationAsset) = 0;
 
 	/** Get the preview scene that we are using */
-	virtual TSharedRef<class IPersonaPreviewScene> GetPreviewScene() const = 0;
+	virtual TSharedRef<IPersonaPreviewScene> GetPreviewScene() const = 0;
 
 	/** Set the preview mesh, according to context (mesh, skeleton or animation etc.) */
-	virtual class USkeletalMesh* GetPreviewMesh() const = 0;
+	virtual USkeletalMesh* GetPreviewMesh() const = 0;
 
 	/** 
 	 * Set the preview mesh, according to context (mesh, skeleton or animation etc.) 
@@ -49,7 +57,13 @@ public:
 	 * @param	bSetPreviewMeshInAsset	If true, the mesh will be written to the asset so it can be permanently saved. 
 	 *									Otherwise the change is merely transient and will reset next time the editor is opened.
 	 */
-	virtual void SetPreviewMesh(class USkeletalMesh* InSkeletalMesh, bool bSetPreviewMeshInAsset = true) = 0;
+	virtual void SetPreviewMesh(USkeletalMesh* InSkeletalMesh, bool bSetPreviewMeshInAsset = true) = 0;
+
+	/** Set the preview anim blueprint, used to preview sub layers in context. */
+	virtual void SetPreviewAnimationBlueprint(UAnimBlueprint* InAnimBlueprint) = 0;
+
+	/** Get the preview anim blueprint, used to preview sub layers in context. */
+	virtual UAnimBlueprint* GetPreviewAnimationBlueprint() const = 0;
 
 	/** Retrieve editor custom data. Return INDEX_NONE if the key is invalid */
 	virtual int32 GetCustomData(const int32 Key) const { return INDEX_NONE; }
diff --git a/Engine/Source/Editor/Persona/Public/PersonaUtils.h b/Engine/Source/Editor/Persona/Public/PersonaUtils.h
--- a/Engine/Source/Editor/Persona/Public/PersonaUtils.h
+++ b/Engine/Source/Editor/Persona/Public/PersonaUtils.h
@@ -54,4 +54,7 @@ struct FCopyOptions
 /** Copy modified properties from the specified anim instance back to its CDO */
 PERSONA_API int32 CopyPropertiesToCDO(UAnimInstance* InAnimInstance, const FCopyOptions& Options = FCopyOptions(ECopyOptions::Default));
 
+/** Set the correct object to debug depending on whether we are running a sub-layer etc. */
+PERSONA_API void SetObjectBeingDebugged(UAnimBlueprint* InAnimBlueprint, UAnimInstance* InAnimInstance);
+
 }
diff --git a/Engine/Source/Runtime/Engine/Classes/Animation/AnimBlueprint.h b/Engine/Source/Runtime/Engine/Classes/Animation/AnimBlueprint.h
--- a/Engine/Source/Runtime/Engine/Classes/Animation/AnimBlueprint.h
+++ b/Engine/Source/Runtime/Engine/Classes/Animation/AnimBlueprint.h
@@ -57,6 +57,17 @@ struct FAnimParentNodeAssetOverride
 	}
 };
 
+/** The method by which a preview animation blueprint is applied */
+UENUM()
+enum class EPreviewAnimationBlueprintApplicationMethod : uint8
+{
+	/** Apply the preview animation blueprint using SetLayerOverlay */
+	OverlayLayer,
+
+	/** Apply the preview animation blueprint using SetSubInstanceClassByTag */
+	SubInstance,
+};
+
 /**
  * An Anim Blueprint is essentially a specialized Blueprint whose graphs control the animation of a Skeletal Mesh.
  * It can perform blending of animations, directly control the bones of the skeleton, and output a final pose
@@ -166,6 +177,16 @@ public:
 	virtual USkeletalMesh* GetPreviewMesh(bool bFindIfNotSet = false) override;
 	virtual USkeletalMesh* GetPreviewMesh() const override;
 
+	/** Preview anim blueprint support */
+	void SetPreviewAnimationBlueprint(UAnimBlueprint* InPreviewAnimationBlueprint);
+	UAnimBlueprint* GetPreviewAnimationBlueprint() const;
+
+	void SetPreviewAnimationBlueprintApplicationMethod(EPreviewAnimationBlueprintApplicationMethod InMethod);
+	EPreviewAnimationBlueprintApplicationMethod GetPreviewAnimationBlueprintApplicationMethod() const;
+
+	void SetPreviewAnimationBlueprintTag(FName InTag);
+	FName GetPreviewAnimationBlueprintTag() const;
+
 public:
 	/** Check if the anim instance is the active debug object for this anim BP */
 	bool IsObjectBeingDebugged(const UObject* AnimInstance) const;
@@ -188,5 +209,20 @@ private:
 	/** The default skeletal mesh to use when previewing this asset - this only applies when you open Persona using this asset*/
 	UPROPERTY(duplicatetransient, AssetRegistrySearchable)
 	TSoftObjectPtr<class USkeletalMesh> PreviewSkeletalMesh;
+
+	/** 
+	 * An animation Blueprint to overlay with this Blueprint. When working on layers, this allows this Blueprint to be previewed in the context of another 'outer' anim blueprint. 
+	 * Setting this is the equivalent of running the preview animation blueprint on the preview mesh, then calling SetLayerOverlay with this anim blueprint.
+	 */
+	UPROPERTY(duplicatetransient, AssetRegistrySearchable)
+	TSoftObjectPtr<class UAnimBlueprint> PreviewAnimationBlueprint;
+
+	/** The method by which a preview animation blueprint is applied, either as an overlay layer, or as a sub-instance */
+	UPROPERTY()
+	EPreviewAnimationBlueprintApplicationMethod PreviewAnimationBlueprintApplicationMethod;
+
+	/** The tag to use when applying a preview animation blueprint via SetSubInstanceClassByTag */
+	UPROPERTY()
+	FName PreviewAnimationBlueprintTag;
 #endif // WITH_EDITORONLY_DATA
 };
diff --git a/Engine/Source/Runtime/Engine/Classes/Components/SkeletalMeshComponent.h b/Engine/Source/Runtime/Engine/Classes/Components/SkeletalMeshComponent.h
--- a/Engine/Source/Runtime/Engine/Classes/Components/SkeletalMeshComponent.h
+++ b/Engine/Source/Runtime/Engine/Classes/Components/SkeletalMeshComponent.h
@@ -777,15 +777,19 @@ public:
 	 * Returns the a tagged sub-instance node. If no sub instances are found or none are tagged with the
 	 * supplied name, this will return NULL.
 	 */
-	UFUNCTION(BlueprintCallable, Category = "Components|SkeletalMesh", meta = (Keywords = "AnimBlueprint"))
+	UFUNCTION(BlueprintCallable, Category = "Components|SkeletalMesh|Sub-Instances", meta = (Keywords = "AnimBlueprint"))
 	UAnimInstance* GetSubInstanceByTag(FName InTag) const;
 
 	/**
 	 * Returns all tagged sub-instance nodes that match the tag.
 	 */
-	UFUNCTION(BlueprintCallable, Category = "Components|SkeletalMesh", meta = (Keywords = "AnimBlueprint"))
+	UFUNCTION(BlueprintCallable, Category = "Components|SkeletalMesh|Sub-Instances", meta = (Keywords = "AnimBlueprint"))
 	void GetSubInstancesByTag(FName InTag, TArray<UAnimInstance*>& OutSubInstances) const;
 
+	/** Runs through all nodes, attempting to find sub-instance by name/tag, then sets the class of each node if the tag matches */
+	UFUNCTION(BlueprintCallable, Category = "Components|SkeletalMesh|Sub-Instances", meta = (Keywords = "AnimBlueprint"))
+	void SetSubInstanceClassByTag(FName InTag, TSubclassOf<UAnimInstance> InClass);
+
 	/** 
 	 * Runs through all layer nodes, attempting to find layer nodes that are implemented by the specified class, then sets up a sub instance of the class for each.
 	 * Allocates one sub instance to run each of the groups specified in the class, so state is shared. If a layer is not grouped (ie. NAME_None), then state is not shared
diff --git a/Engine/Source/Runtime/Engine/Private/AnimBlueprint.cpp b/Engine/Source/Runtime/Engine/Private/AnimBlueprint.cpp
--- a/Engine/Source/Runtime/Engine/Private/AnimBlueprint.cpp
+++ b/Engine/Source/Runtime/Engine/Private/AnimBlueprint.cpp
@@ -236,6 +236,60 @@ void UAnimBlueprint::SetPreviewMesh(USkeletalMesh* PreviewMesh, bool bMarkAsDirt
 #endif
 }
 
+void UAnimBlueprint::SetPreviewAnimationBlueprint(UAnimBlueprint* InPreviewAnimationBlueprint)
+{
+#if WITH_EDITORONLY_DATA
+	Modify();
+	PreviewAnimationBlueprint = InPreviewAnimationBlueprint;
+#endif
+}
+
+UAnimBlueprint* UAnimBlueprint::GetPreviewAnimationBlueprint() const
+{
+#if WITH_EDITORONLY_DATA
+	if (!PreviewAnimationBlueprint.IsValid())
+	{
+		PreviewAnimationBlueprint.LoadSynchronous();
+	}
+	return PreviewAnimationBlueprint.Get();
+#else
+	return nullptr;
+#endif
+}
+
+void UAnimBlueprint::SetPreviewAnimationBlueprintApplicationMethod(EPreviewAnimationBlueprintApplicationMethod InMethod) 
+{ 
+#if WITH_EDITORONLY_DATA
+	Modify();
+	PreviewAnimationBlueprintApplicationMethod = InMethod; 
+#endif
+}
+
+EPreviewAnimationBlueprintApplicationMethod UAnimBlueprint::GetPreviewAnimationBlueprintApplicationMethod() const 
+{ 
+#if WITH_EDITORONLY_DATA
+	return PreviewAnimationBlueprintApplicationMethod; 
+#else
+	return EPreviewAnimationBlueprintApplicationMethod::OverlayLayer;
+#endif
+}
+
+void UAnimBlueprint::SetPreviewAnimationBlueprintTag(FName InTag) 
+{ 
+#if WITH_EDITORONLY_DATA
+	Modify();
+	PreviewAnimationBlueprintTag = InTag; 
+#endif
+}
+
+FName UAnimBlueprint::GetPreviewAnimationBlueprintTag() const 
+{ 
+#if WITH_EDITORONLY_DATA
+	return PreviewAnimationBlueprintTag; 
+#else
+	return NAME_None;
+#endif
+}
 
 bool UAnimBlueprint::IsObjectBeingDebugged(const UObject* Object) const
 {
diff --git a/Engine/Source/Runtime/Engine/Private/Components/SkeletalMeshComponent.cpp b/Engine/Source/Runtime/Engine/Private/Components/SkeletalMeshComponent.cpp
--- a/Engine/Source/Runtime/Engine/Private/Components/SkeletalMeshComponent.cpp
+++ b/Engine/Source/Runtime/Engine/Private/Components/SkeletalMeshComponent.cpp
@@ -2761,6 +2761,14 @@ void USkeletalMeshComponent::GetSubInstancesByTag(FName InTag, TArray<UAnimInsta
 	}
 }
 
+void USkeletalMeshComponent::SetSubInstanceClassByTag(FName InTag, TSubclassOf<UAnimInstance> InClass)
+{
+	if(AnimScriptInstance)
+	{
+		AnimScriptInstance->SetSubInstanceClassByTag(InTag, InClass);
+	}
+}
+
 void USkeletalMeshComponent::SetLayerOverlay(TSubclassOf<UAnimInstance> InClass)
 {
 	if(AnimScriptInstance)
