Fix non-unity build issues

diff --git a/Engine/Source/Editor/Persona/Private/PersonaPreviewSceneDescription.h b/Engine/Source/Editor/Persona/Private/PersonaPreviewSceneDescription.h
--- a/Engine/Source/Editor/Persona/Private/PersonaPreviewSceneDescription.h
+++ b/Engine/Source/Editor/Persona/Private/PersonaPreviewSceneDescription.h
@@ -16,6 +16,7 @@ class UPreviewMeshCollection;
 class USkeletalMesh;
 class UDataAsset;
 class UPreviewMeshCollection;
+class UAnimBlueprint;
 enum class EPreviewAnimationBlueprintApplicationMethod : uint8;
 
 UCLASS()
diff --git a/Engine/Source/Editor/Persona/Private/PreviewSceneCustomizations.cpp b/Engine/Source/Editor/Persona/Private/PreviewSceneCustomizations.cpp
--- a/Engine/Source/Editor/Persona/Private/PreviewSceneCustomizations.cpp
+++ b/Engine/Source/Editor/Persona/Private/PreviewSceneCustomizations.cpp
@@ -27,6 +27,7 @@
 #include "Widgets/Input/SComboBox.h"
 #include "PhysicsEngine/PhysicsAsset.h"
 #include "Algo/Sort.h"
+#include "ScopedTransaction.h"
 
 #define LOCTEXT_NAMESPACE "PreviewSceneCustomizations"
 
