File I missed

diff --git a/Engine/Source/Editor/Persona/Public/PersonaUtils.h b/Engine/Source/Editor/Persona/Public/PersonaUtils.h
--- a/Engine/Source/Editor/Persona/Public/PersonaUtils.h
+++ b/Engine/Source/Editor/Persona/Public/PersonaUtils.h
@@ -6,6 +6,7 @@
 
 class UAnimInstance;
 class USceneComponent;
+class UAnimBlueprint;
 
 namespace PersonaUtils
 {
