#UE4 Fix shadow variable

diff --git a/Engine/Source/Editor/Persona/Private/PreviewSceneCustomizations.cpp b/Engine/Source/Editor/Persona/Private/PreviewSceneCustomizations.cpp
--- a/Engine/Source/Editor/Persona/Private/PreviewSceneCustomizations.cpp
+++ b/Engine/Source/Editor/Persona/Private/PreviewSceneCustomizations.cpp
@@ -252,10 +252,10 @@ void FPreviewSceneDescriptionCustomization::CustomizeDetails(IDetailLayoutBuilde
 				FScopedTransaction Transaction(LOCTEXT("SetAnimationBlueprintApplicationMethod", "Set Application Method"));
 
 				TSharedPtr<IPersonaToolkit> PinnedPersonaToolkit = PersonaToolkit.Pin();
-				TSharedRef<FAnimationEditorPreviewScene> PreviewScene = StaticCastSharedRef<FAnimationEditorPreviewScene>(PinnedPersonaToolkit->GetPreviewScene());
-				UPersonaPreviewSceneDescription* PersonaPreviewSceneDescription = PreviewScene->GetPreviewSceneDescription();
+				TSharedRef<FAnimationEditorPreviewScene> LocalPreviewScene = StaticCastSharedRef<FAnimationEditorPreviewScene>(PinnedPersonaToolkit->GetPreviewScene());
+				UPersonaPreviewSceneDescription* PersonaPreviewSceneDescription = LocalPreviewScene->GetPreviewSceneDescription();
 				PinnedPersonaToolkit->GetAnimBlueprint()->SetPreviewAnimationBlueprintApplicationMethod(PersonaPreviewSceneDescription->ApplicationMethod);
-				PreviewScene->SetPreviewAnimationBlueprint(PersonaPreviewSceneDescription->PreviewAnimationBlueprint.Get(), PinnedPersonaToolkit->GetAnimBlueprint());
+				LocalPreviewScene->SetPreviewAnimationBlueprint(PersonaPreviewSceneDescription->PreviewAnimationBlueprint.Get(), PinnedPersonaToolkit->GetAnimBlueprint());
 			}));
 
 			DetailBuilder.EditCategory("Animation Blueprint")
@@ -263,8 +263,8 @@ void FPreviewSceneDescriptionCustomization::CustomizeDetails(IDetailLayoutBuilde
 			.IsEnabled(MakeAttributeLambda([this]()
 			{
 				TSharedPtr<IPersonaToolkit> PinnedPersonaToolkit = PersonaToolkit.Pin();
-				TSharedRef<FAnimationEditorPreviewScene> PreviewScene = StaticCastSharedRef<FAnimationEditorPreviewScene>(PinnedPersonaToolkit->GetPreviewScene());
-				UPersonaPreviewSceneDescription* PersonaPreviewSceneDescription = PreviewScene->GetPreviewSceneDescription();
+				TSharedRef<FAnimationEditorPreviewScene> LocalPreviewScene = StaticCastSharedRef<FAnimationEditorPreviewScene>(PinnedPersonaToolkit->GetPreviewScene());
+				UPersonaPreviewSceneDescription* PersonaPreviewSceneDescription = LocalPreviewScene->GetPreviewSceneDescription();
 				
 				return PersonaPreviewSceneDescription->PreviewAnimationBlueprint.IsValid();
 			}));
@@ -274,10 +274,10 @@ void FPreviewSceneDescriptionCustomization::CustomizeDetails(IDetailLayoutBuilde
 				FScopedTransaction Transaction(LOCTEXT("SetAnimationBlueprintTag", "Set Sub Instance Tag"));
 
 				TSharedPtr<IPersonaToolkit> PinnedPersonaToolkit = PersonaToolkit.Pin();
-				TSharedRef<FAnimationEditorPreviewScene> PreviewScene = StaticCastSharedRef<FAnimationEditorPreviewScene>(PinnedPersonaToolkit->GetPreviewScene());
-				UPersonaPreviewSceneDescription* PersonaPreviewSceneDescription = PreviewScene->GetPreviewSceneDescription();
+				TSharedRef<FAnimationEditorPreviewScene> LocalPreviewScene = StaticCastSharedRef<FAnimationEditorPreviewScene>(PinnedPersonaToolkit->GetPreviewScene());
+				UPersonaPreviewSceneDescription* PersonaPreviewSceneDescription = LocalPreviewScene->GetPreviewSceneDescription();
 				PinnedPersonaToolkit->GetAnimBlueprint()->SetPreviewAnimationBlueprintTag(PersonaPreviewSceneDescription->SubInstanceTag);
-				PreviewScene->SetPreviewAnimationBlueprint(PersonaPreviewSceneDescription->PreviewAnimationBlueprint.Get(), PinnedPersonaToolkit->GetAnimBlueprint());
+				LocalPreviewScene->SetPreviewAnimationBlueprint(PersonaPreviewSceneDescription->PreviewAnimationBlueprint.Get(), PinnedPersonaToolkit->GetAnimBlueprint());
 			}));
 
 			DetailBuilder.EditCategory("Animation Blueprint")
@@ -285,8 +285,8 @@ void FPreviewSceneDescriptionCustomization::CustomizeDetails(IDetailLayoutBuilde
 			.IsEnabled(MakeAttributeLambda([this]()
 			{
 				TSharedPtr<IPersonaToolkit> PinnedPersonaToolkit = PersonaToolkit.Pin();
-				TSharedRef<FAnimationEditorPreviewScene> PreviewScene = StaticCastSharedRef<FAnimationEditorPreviewScene>(PinnedPersonaToolkit->GetPreviewScene());
-				UPersonaPreviewSceneDescription* PersonaPreviewSceneDescription = PreviewScene->GetPreviewSceneDescription();
+				TSharedRef<FAnimationEditorPreviewScene> LocalPreviewScene = StaticCastSharedRef<FAnimationEditorPreviewScene>(PinnedPersonaToolkit->GetPreviewScene());
+				UPersonaPreviewSceneDescription* PersonaPreviewSceneDescription = LocalPreviewScene->GetPreviewSceneDescription();
 				
 				return PersonaPreviewSceneDescription->PreviewAnimationBlueprint.IsValid() && PersonaPreviewSceneDescription->ApplicationMethod == EPreviewAnimationBlueprintApplicationMethod::SubInstance;
 			}));
