Remerge cl: 9413572 into main it was backout in cl: 9418047
Undefine all portions of code that deal with the Apex clothing file import unless the project is built with Apex.
This is to prevent unnecessary Apex dependencies leaking into the Chaos Cloth asset code, and will allow the removal of PhysX when needed.

diff --git a/Engine/Source/Developer/AssetTools/Private/AssetTypeActions/AssetTypeActions_SkeletalMesh.cpp b/Engine/Source/Developer/AssetTools/Private/AssetTypeActions/AssetTypeActions_SkeletalMesh.cpp
--- a/Engine/Source/Developer/AssetTools/Private/AssetTypeActions/AssetTypeActions_SkeletalMesh.cpp
+++ b/Engine/Source/Developer/AssetTools/Private/AssetTypeActions/AssetTypeActions_SkeletalMesh.cpp
@@ -473,12 +473,14 @@ void FAssetTypeActions_SkeletalMesh::GetActions(const TArray<UObject*>& InObject
 		FNewMenuDelegate::CreateSP(this, &FAssetTypeActions_SkeletalMesh::GetLODMenu, Meshes)
 		);
 	
+#if WITH_APEX_CLOTHING
 	Section.AddMenuEntry(
 		"ImportClothing_Entry",
 		LOCTEXT("ImportClothing_Entry", "Import Clothing Asset..."),
 		LOCTEXT("ImportClothing_ToolTip", "Import a clothing asset from a supported file on disk into this skeletal mesh."),
 		FSlateIcon(),
 		FUIAction(FExecuteAction::CreateSP(this, &FAssetTypeActions_SkeletalMesh::ExecuteImportClothing, Meshes)));
+#endif  // #if WITH_APEX_CLOTHING
 
 	// skeleton menu
 	Section.AddSubMenu(
@@ -829,6 +831,7 @@ void FAssetTypeActions_SkeletalMesh::ExecuteImportMeshLOD(UObject* Mesh, int32 L
 	}
 }
 
+#if WITH_APEX_CLOTHING
 void FAssetTypeActions_SkeletalMesh::ExecuteImportClothing(TArray<TWeakObjectPtr<USkeletalMesh>> Objects)
 {
 	if(Objects.Num() > 0)
@@ -841,6 +844,7 @@ void FAssetTypeActions_SkeletalMesh::ExecuteImportClothing(TArray<TWeakObjectPtr
 		}
 	}
 }
+#endif  // #if WITH_APEX_CLOTHING
 
 void FAssetTypeActions_SkeletalMesh::FillSkeletonMenu(FMenuBuilder& MenuBuilder, TArray<TWeakObjectPtr<USkeletalMesh>> Meshes) const
 {
diff --git a/Engine/Source/Developer/AssetTools/Private/AssetTypeActions/AssetTypeActions_SkeletalMesh.h b/Engine/Source/Developer/AssetTools/Private/AssetTypeActions/AssetTypeActions_SkeletalMesh.h
--- a/Engine/Source/Developer/AssetTools/Private/AssetTypeActions/AssetTypeActions_SkeletalMesh.h
+++ b/Engine/Source/Developer/AssetTools/Private/AssetTypeActions/AssetTypeActions_SkeletalMesh.h
@@ -55,8 +55,10 @@ private:
 	/** Handler for skeletal mesh import */
 	static void ExecuteImportMeshLOD(class UObject* Mesh, int32 LOD);
 
+#if WITH_APEX_CLOTHING
 	/** Handler to import clothing to a mesh */
 	void ExecuteImportClothing(TArray<TWeakObjectPtr<USkeletalMesh>> Objects);
+#endif
 
 	// Helper functions
 private:
diff --git a/Engine/Source/Editor/ClothPainter/Private/ClothingAssetListCommands.cpp b/Engine/Source/Editor/ClothPainter/Private/ClothingAssetListCommands.cpp
--- a/Engine/Source/Editor/ClothPainter/Private/ClothingAssetListCommands.cpp
+++ b/Engine/Source/Editor/ClothPainter/Private/ClothingAssetListCommands.cpp
@@ -7,7 +7,9 @@
 void FClothingAssetListCommands::RegisterCommands()
 {
 	UI_COMMAND(DeleteAsset, "Delete Asset", "Deletes a clothing asset from the mesh.", EUserInterfaceActionType::Button, FInputChord());
+#if WITH_APEX_CLOTHING
 	UI_COMMAND(ReimportAsset, "Reimport Asset", "If an asset was originally imported from an external file, reimport from that file.", EUserInterfaceActionType::Button, FInputChord());
+#endif
 	UI_COMMAND(RebuildAssetParams, "Rebuild Asset Parameter Masks", "Takes the parameter masks in LOD0 and creates masks in all lower LODs to match, casting those parameter masks to the new mesh.", EUserInterfaceActionType::Button, FInputChord());
 }
 
diff --git a/Engine/Source/Editor/ClothPainter/Private/ClothingAssetListCommands.h b/Engine/Source/Editor/ClothPainter/Private/ClothingAssetListCommands.h
--- a/Engine/Source/Editor/ClothPainter/Private/ClothingAssetListCommands.h
+++ b/Engine/Source/Editor/ClothPainter/Private/ClothingAssetListCommands.h
@@ -21,7 +21,9 @@ public:
 
 	TSharedPtr<FUICommandInfo> DeleteAsset;
 
+#if WITH_APEX_CLOTHING
 	TSharedPtr<FUICommandInfo> ReimportAsset;
+#endif
 
 	TSharedPtr<FUICommandInfo> RebuildAssetParams;
 
diff --git a/Engine/Source/Editor/ClothPainter/Private/SClothAssetSelector.cpp b/Engine/Source/Editor/ClothPainter/Private/SClothAssetSelector.cpp
--- a/Engine/Source/Editor/ClothPainter/Private/SClothAssetSelector.cpp
+++ b/Engine/Source/Editor/ClothPainter/Private/SClothAssetSelector.cpp
@@ -142,11 +142,13 @@ public:
 			FExecuteAction::CreateSP(this, &SAssetListRow::DeleteAsset)
 		);
 
+#if WITH_APEX_CLOTHING
 		UICommandList->MapAction(
 			Commands.ReimportAsset,
 			FExecuteAction::CreateSP(this, &SAssetListRow::ReimportAsset),
 			FCanExecuteAction::CreateSP(this, &SAssetListRow::CanReimportAsset)
 		);
+#endif
 
 		UICommandList->MapAction(
 			Commands.RebuildAssetParams,
@@ -165,7 +167,9 @@ public:
 			Builder.BeginSection(NAME_None, LOCTEXT("AssetActions_SectionName", "Actions"));
 			{
 				Builder.AddMenuEntry(FGenericCommands::Get().Delete);
+#if WITH_APEX_CLOTHING
 				Builder.AddMenuEntry(Commands.ReimportAsset);
+#endif
 				Builder.AddMenuEntry(Commands.RebuildAssetParams);
 			}
 			Builder.EndSection();
@@ -230,6 +234,7 @@ private:
 		}
 	}
 
+#if WITH_APEX_CLOTHING
 	void ReimportAsset()
 	{
 		if(UClothingAssetCommon* Asset = Item->ClothingAsset.Get())
@@ -283,6 +288,7 @@ private:
 	{
 		return Item.IsValid() && !Item->ClothingAsset->ImportedFilePath.IsEmpty();
 	}
+#endif  // #if WITH_APEX_CLOTHING
 
 	// Using LOD0 of an asset, rebuild the other LOD masks by mapping the LOD0 parameters onto their meshes
 	void RebuildLODParameters()
@@ -726,13 +732,14 @@ void SClothAssetSelector::Construct(const FArguments& InArgs, USkeletalMesh* InM
 					.Font(FEditorStyle::GetFontStyle("DetailsView.CategoryFontStyle"))
 					.ShadowOffset(FVector2D(1.0f, 1.0f))
 				]
+#if WITH_APEX_CLOTHING
 				+SHorizontalBox::Slot()
 				.AutoWidth()
 				.VAlign(VAlign_Center)
 				.HAlign(HAlign_Right)
 				.Padding(0.0f, 0.0f, 4.0f, 0.0f)
 				[
-					SAssignNew(ImportApexButton, SButton)
+					SNew(SButton)
 					.ButtonStyle(FEditorStyle::Get(), "RoundButton")
 					.ForegroundColor(FEditorStyle::GetSlateColor("DefaultForeground"))
 					.ContentPadding(FMargin(2, 0))
@@ -763,6 +770,7 @@ void SClothAssetSelector::Construct(const FArguments& InArgs, USkeletalMesh* InM
 						]
 					]
 				]
+#endif  // #if WITH_APEX_CLOTHING
 				+ SHorizontalBox::Slot()
 					.AutoWidth()
 					.VAlign(VAlign_Center)
@@ -952,6 +960,7 @@ void SClothAssetSelector::PostUndo(bool bSuccess)
 	OnRefresh();
 }
 
+#if WITH_APEX_CLOTHING
 FReply SClothAssetSelector::OnImportApexFileClicked()
 {
 	if(Mesh)
@@ -964,6 +973,7 @@ FReply SClothAssetSelector::OnImportApexFileClicked()
 
 	return FReply::Unhandled();
 }
+#endif  // #if WITH_APEX_CLOTHING
 
 void SClothAssetSelector::OnCopyClothingAssetSelected(const FAssetData& AssetData)
 {
diff --git a/Engine/Source/Editor/ClothPainter/Private/SClothAssetSelector.h b/Engine/Source/Editor/ClothPainter/Private/SClothAssetSelector.h
--- a/Engine/Source/Editor/ClothPainter/Private/SClothAssetSelector.h
+++ b/Engine/Source/Editor/ClothPainter/Private/SClothAssetSelector.h
@@ -62,7 +62,9 @@ public:
 
 protected:
 
+#if WITH_APEX_CLOTHING
 	FReply OnImportApexFileClicked();
+#endif
 
 	/* Copies clothing setup from source SkelMesh */
 	void OnCopyClothingAssetSelected(const FAssetData& AssetData);
@@ -100,7 +102,6 @@ protected:
 
 	USkeletalMesh* Mesh;
 
-	TSharedPtr<SButton> ImportApexButton;
 	TSharedPtr<SButton> NewMaskButton;
 	TSharedPtr<SAssetList> AssetList;
 	TSharedPtr<SMaskList> MaskList;
diff --git a/Engine/Source/Editor/Persona/Private/PersonaMeshDetails.cpp b/Engine/Source/Editor/Persona/Private/PersonaMeshDetails.cpp
--- a/Engine/Source/Editor/Persona/Private/PersonaMeshDetails.cpp
+++ b/Engine/Source/Editor/Persona/Private/PersonaMeshDetails.cpp
@@ -5018,6 +5018,7 @@ void FPersonaMeshDetails::OnGenerateElementForClothingAsset( TSharedRef<IPropert
 		+ SHorizontalBox::Slot()
 		.FillWidth(1)
 
+#if WITH_APEX_CLOTHING
 		// re-import button
 		+ SHorizontalBox::Slot()
 		.VAlign( VAlign_Center )
@@ -5038,6 +5039,7 @@ void FPersonaMeshDetails::OnGenerateElementForClothingAsset( TSharedRef<IPropert
 				.ColorAndOpacity( FSlateColor::UseForeground() )
 			]
 		]
+#endif  // #if WITH_APEX_CLOTHING
 
 		// remove button
 		+ SHorizontalBox::Slot()
@@ -5047,7 +5049,7 @@ void FPersonaMeshDetails::OnGenerateElementForClothingAsset( TSharedRef<IPropert
 		[
 			SNew( SButton )
 			.Text( LOCTEXT("ClearButtonLabel", "Remove") )
-			.OnClicked( this, &FPersonaMeshDetails::OnRemoveApexFileClicked, ElementIndex, DetailLayout )
+			.OnClicked( this, &FPersonaMeshDetails::OnRemoveClothingAssetClicked, ElementIndex, DetailLayout )
 			.IsFocusable( false )
 			.ContentPadding(0)
 			.ForegroundColor( FSlateColor::UseForeground() )
@@ -5252,9 +5254,9 @@ TSharedRef<SUniformGridPanel> FPersonaMeshDetails::MakeClothingDetailsWidget(int
 	return Grid;
 }
 
+#if WITH_APEX_CLOTHING
 FReply FPersonaMeshDetails::OnReimportApexFileClicked(int32 AssetIndex, IDetailLayoutBuilder* DetailLayout)
 {
-#if WITH_APEX_CLOTHING
 	USkeletalMesh* SkelMesh = GetPersonaToolkit()->GetMesh();
 
 	check(SkelMesh && SkelMesh->MeshClothingAssets.IsValidIndex(AssetIndex));
@@ -5305,12 +5307,12 @@ FReply FPersonaMeshDetails::OnReimportApexFileClicked(int32 AssetIndex, IDetailL
 		// Force layout to refresh
 		DetailLayout->ForceRefreshDetails();
 	}
-#endif
 
 	return FReply::Handled();
 }
+#endif
 
-FReply FPersonaMeshDetails::OnRemoveApexFileClicked(int32 AssetIndex, IDetailLayoutBuilder* DetailLayout)
+FReply FPersonaMeshDetails::OnRemoveClothingAssetClicked(int32 AssetIndex, IDetailLayoutBuilder* DetailLayout)
 {
 	USkeletalMesh* SkelMesh = GetPersonaToolkit()->GetMesh();
 	check(SkelMesh);
@@ -5367,9 +5369,9 @@ FReply FPersonaMeshDetails::OnRemoveApexFileClicked(int32 AssetIndex, IDetailLay
 	return FReply::Handled();
 }
 
+#if WITH_APEX_CLOTHING
 FReply FPersonaMeshDetails::OnOpenClothingFileClicked(IDetailLayoutBuilder* DetailLayout)
 {
-#if WITH_APEX_CLOTHING
 	USkeletalMesh* SkelMesh = GetPersonaToolkit()->GetMesh();
 
 	if(SkelMesh)
@@ -5379,10 +5381,10 @@ FReply FPersonaMeshDetails::OnOpenClothingFileClicked(IDetailLayoutBuilder* Deta
 		UpdateClothingEntries();
 		RefreshClothingComboBoxes();
 	}
-#endif
 
 	return FReply::Handled();
 }
+#endif
 
 void FPersonaMeshDetails::UpdateClothingEntries()
 {
diff --git a/Engine/Source/Editor/Persona/Private/PersonaMeshDetails.h b/Engine/Source/Editor/Persona/Private/PersonaMeshDetails.h
--- a/Engine/Source/Editor/Persona/Private/PersonaMeshDetails.h
+++ b/Engine/Source/Editor/Persona/Private/PersonaMeshDetails.h
@@ -703,14 +703,16 @@ private:
 	/* Make uniform grid widget for Apex details */
 	TSharedRef<SUniformGridPanel> MakeClothingDetailsWidget(int32 AssetIndex) const;
 
+#if WITH_APEX_CLOTHING
 	/* Opens dialog to add a new clothing asset */
 	FReply OnOpenClothingFileClicked(IDetailLayoutBuilder* DetailLayout);
 
 	/* Reimports a clothing asset */ 
 	FReply OnReimportApexFileClicked(int32 AssetIndex, IDetailLayoutBuilder* DetailLayout);
+#endif
 
 	/* Removes a clothing asset */ 
-	FReply OnRemoveApexFileClicked(int32 AssetIndex, IDetailLayoutBuilder* DetailLayout);
+	FReply OnRemoveClothingAssetClicked(int32 AssetIndex, IDetailLayoutBuilder* DetailLayout);
 
 	/* Create LOD setting assets from current setting */
 	FReply OnSaveLODSettings();
diff --git a/Engine/Source/Editor/UnrealEd/Private/ApexClothingUtils.cpp b/Engine/Source/Editor/UnrealEd/Private/ApexClothingUtils.cpp
--- a/Engine/Source/Editor/UnrealEd/Private/ApexClothingUtils.cpp
+++ b/Engine/Source/Editor/UnrealEd/Private/ApexClothingUtils.cpp
@@ -54,10 +54,11 @@ void RefreshSkelMeshComponents(USkeletalMesh* SkelMesh)
 	}
 }
 
+#if WITH_APEX_CLOTHING
 FString PromptForClothingFile()
 {
 	if(IDesktopPlatform* Platform = FDesktopPlatformModule::Get())
-{
+	{
 		const void* ParentWindowWindowHandle = FSlateApplication::Get().FindBestParentWindowHandleForDialogs(nullptr);
 
 		TArray<FString> OpenFilenames;
@@ -72,7 +73,7 @@ FString PromptForClothingFile()
 			EFileDialogFlags::None,
 			OpenFilenames
 			))
-	{
+		{
 			return OpenFilenames[0];
 		}
 	}
@@ -88,19 +89,18 @@ void PromptAndImportClothing(USkeletalMesh* SkelMesh)
 	FString Filename = PromptForClothingFile();
 
 	if(!Filename.IsEmpty())
-{
+	{
 		FEditorDirectories::Get().SetLastDirectory(ELastDirectory::MESH_IMPORT_EXPORT, Filename);
 
-		UClothingAssetFactory* Factory = UClothingAssetFactory::StaticClass()->GetDefaultObject<UClothingAssetFactory>();
+		UClothingAssetFactory* Factory = UClothingAssetFactory::StaticClass()->GetDefaultObject<UClothingAssetFactory>();  // TODO: Rename this class to UClothingAssetFactoryNv
 
 		if(Factory && Factory->CanImport(Filename))
-	{					
+		{
 			Factory->Import(Filename, SkelMesh);
-	}
+		}
 	}
 }
 
-#if WITH_APEX_CLOTHING
 apex::ClothingAsset* CreateApexClothingAssetFromPxStream(physx::PxFileBuf& Stream)
 {
 	// Peek into the buffer to see what kind of data it is (binary or xml)
diff --git a/Engine/Source/Editor/UnrealEd/Public/ApexClothingUtils.h b/Engine/Source/Editor/UnrealEd/Public/ApexClothingUtils.h
--- a/Engine/Source/Editor/UnrealEd/Public/ApexClothingUtils.h
+++ b/Engine/Source/Editor/UnrealEd/Public/ApexClothingUtils.h
@@ -7,6 +7,7 @@
 
 class USkeletalMesh;
 
+#if WITH_APEX_CLOTHING
 namespace nvidia
 {
 	namespace apex
@@ -14,10 +15,13 @@ namespace nvidia
 		class ClothingAsset;
 	}
 }
+#endif  // #if WITH_APEX_CLOTHING
 
 // Define interface for importing apex clothing asset
 namespace ApexClothingUtils
 {
+#if WITH_APEX_CLOTHING
+	// Prompt the user to select an APEX file
 	UNREALED_API FString PromptForClothingFile();
 
 	// Prompt the user to select an APEX file that will be imported to the new UClothingAssetCommon format
@@ -25,7 +29,7 @@ namespace ApexClothingUtils
 
 	// Given a buffer, build an apex clothing asset
 	UNREALED_API nvidia::apex::ClothingAsset* CreateApexClothingAssetFromBuffer(const uint8* Buffer, int32 BufferSize);
-
+#endif  // #if WITH_APEX_CLOTHING
 
 	// Functions below remain from previous clothing system and only remain to remove the bound
 	// data from a skeletal mesh. This is done when postloading USkeletalMesh when upgrading the assets
