Fix crash when removing alternate skinning profile and the preview is active
#rnx

diff --git a/Engine/Source/Editor/Persona/Private/SkinWeightProfileCustomization.cpp b/Engine/Source/Editor/Persona/Private/SkinWeightProfileCustomization.cpp
--- a/Engine/Source/Editor/Persona/Private/SkinWeightProfileCustomization.cpp
+++ b/Engine/Source/Editor/Persona/Private/SkinWeightProfileCustomization.cpp
@@ -10,6 +10,8 @@
 #include "Rendering/SkeletalMeshLODModel.h"
 #include "Rendering/SkeletalMeshModel.h"
 #include "Components/SkinnedMeshComponent.h"
+#include "Animation/DebugSkelMeshComponent.h"
+#include "SkeletalMeshTypes.h"
 
 #include "PropertyHandle.h"
 #include "DetailWidgetRow.h"
@@ -412,12 +414,15 @@ TSharedRef<class SWidget> FSkinWeightProfileCustomization::GenerateRemoveMenu()
 		{
 			if (USkeletalMesh* SkeletalMesh = WeakSkeletalMesh.Get())
 			{
-				FSkinnedMeshComponentRecreateRenderStateContext ReregisterContext(SkeletalMesh);
-				FScopedTransaction ScopedTransaction(LOCTEXT("RemoveSkinProfileTransaction", "Remove Skin Weight Profile"));
-				SkeletalMesh->Modify();
-				FSkinWeightProfileHelpers::RemoveSkinWeightProfile(SkeletalMesh, LastKnownProfileName);
-				SkeletalMesh->PostEditChange();
-				UpdateNameRestriction();
+				//Scope the postedit change
+				{
+					FScopedSuspendAlternateSkinWeightPreview ScopedSuspendAlternateSkinnWeightPreview(SkeletalMesh);
+					FScopedSkeletalMeshPostEditChange ScopedPostEditChange(SkeletalMesh);
+					FScopedTransaction ScopedTransaction(LOCTEXT("RemoveSkinProfileTransaction", "Remove Skin Weight Profile"));
+					SkeletalMesh->Modify();
+					FSkinWeightProfileHelpers::RemoveSkinWeightProfile(SkeletalMesh, LastKnownProfileName);
+					UpdateNameRestriction();
+				}
 
 				PropertyUtilities->ForceRefresh();
 			}
diff --git a/Engine/Source/Editor/Persona/Private/SkinWeightProfileHelpers.cpp b/Engine/Source/Editor/Persona/Private/SkinWeightProfileHelpers.cpp
--- a/Engine/Source/Editor/Persona/Private/SkinWeightProfileHelpers.cpp
+++ b/Engine/Source/Editor/Persona/Private/SkinWeightProfileHelpers.cpp
@@ -15,6 +15,7 @@
 #include "IMeshReductionInterfaces.h"
 #include "ComponentReregisterContext.h"
 #include "Components/SkinnedMeshComponent.h"
+#include "Animation/DebugSkelMeshComponent.h"
 #include "EditorFramework/AssetImportData.h"
 #include "Widgets/Notifications/SNotificationList.h"
 #include "Framework/Notifications/NotificationManager.h"
@@ -66,6 +67,7 @@ void FSkinWeightProfileHelpers::ImportSkinWeightProfile(USkeletalMesh* InSkeleta
 		
 		if (bShouldImport)
 		{
+			FScopedSuspendAlternateSkinWeightPreview ScopedSuspendAlternateSkinnWeightPreview(InSkeletalMesh);
 			FScopedSkeletalMeshPostEditChange ScopedPostEditChange(InSkeletalMesh);
 			const FName ProfileName(*ImportSettings->ProfileName);			
 			// Try and import the skin weight profile from the provided FBX file path
@@ -96,7 +98,8 @@ void FSkinWeightProfileHelpers::ImportSkinWeightProfileLOD(USkeletalMesh* InSkel
 	{
 		// Pick a FBX file to import the weights from 
 		const FString PickedFileName = FSkinWeightsUtilities::PickSkinWeightFBXPath(LODIndex);
-		
+		FScopedSuspendAlternateSkinWeightPreview ScopedSuspendAlternateSkinnWeightPreview(InSkeletalMesh);
+		FScopedSkeletalMeshPostEditChange ScopedPostEditChange(InSkeletalMesh);
 		// Try and import skin weights for a specific mesh LOD
 		const bool bResult = FSkinWeightsUtilities::ImportAlternateSkinWeight(InSkeletalMesh, PickedFileName, LODIndex, ProfileName);
 		if (bResult)
@@ -126,6 +129,7 @@ void FSkinWeightProfileHelpers::ReimportSkinWeightProfileLOD(USkeletalMesh* InSk
 	{
 		if (const FString* PathNamePtr = Profile->PerLODSourceFiles.Find(LODIndex))
 		{
+			FScopedSuspendAlternateSkinWeightPreview ScopedSuspendAlternateSkinnWeightPreview(InSkeletalMesh);
 			FScopedSkeletalMeshPostEditChange ScopedPostEditChange(InSkeletalMesh);
 			bool bResult = false;
 			const FString& PathName = UAssetImportData::ResolveImportFilename(*PathNamePtr, InSkeletalMesh->GetOutermost());
@@ -177,6 +181,7 @@ void FSkinWeightProfileHelpers::RemoveSkinWeightProfile(USkeletalMesh* InSkeleta
 	// Check whether or not we actually removed a profile
 	if (NumRemoved)
 	{
+		FScopedSuspendAlternateSkinWeightPreview ScopedSuspendAlternateSkinnWeightPreview(InSkeletalMesh);
 		FScopedSkeletalMeshPostEditChange ScopedPostEditChange(InSkeletalMesh);
 		FSkeletalMeshUpdateContext UpdateContext;
 		UpdateContext.SkeletalMesh = InSkeletalMesh;
@@ -226,7 +231,7 @@ void FSkinWeightProfileHelpers::RemoveSkinWeightProfileLOD(USkeletalMesh* InSkel
 			//It will be remove when the base LOD use to reduce will get his skin weights profile remove
 			return;
 		}
-		
+		FScopedSuspendAlternateSkinWeightPreview ScopedSuspendAlternateSkinnWeightPreview(InSkeletalMesh);
 		FScopedSkeletalMeshPostEditChange ScopedPostEditChange(InSkeletalMesh);
 
 		if (FSkinWeightsUtilities::RemoveSkinnedWeightProfileData(InSkeletalMesh, InProfileName, LODIndex))
@@ -254,6 +259,7 @@ void FSkinWeightProfileHelpers::RemoveSkinWeightProfileLOD(USkeletalMesh* InSkel
 void FSkinWeightProfileHelpers::ClearSkinWeightProfileInstanceOverrides(USkeletalMesh* InSkeletalMesh, FName InProfileName)
 {
 	//Make sure all component are unregister (world and render data) avoid post edit change
+	FScopedSuspendAlternateSkinWeightPreview ScopedSuspendAlternateSkinnWeightPreview(InSkeletalMesh);
 	FScopedSkeletalMeshPostEditChange ScopedPostEditChange(InSkeletalMesh, false, true);
 	for (TObjectIterator<USkinnedMeshComponent> It; It; ++It)
 	{
diff --git a/Engine/Source/Editor/UnrealEd/Private/SkinWeightsUtilities.cpp b/Engine/Source/Editor/UnrealEd/Private/SkinWeightsUtilities.cpp
--- a/Engine/Source/Editor/UnrealEd/Private/SkinWeightsUtilities.cpp
+++ b/Engine/Source/Editor/UnrealEd/Private/SkinWeightsUtilities.cpp
@@ -4,6 +4,7 @@
 #include "LODUtilities.h"
 #include "Modules/ModuleManager.h"
 #include "Components/SkinnedMeshComponent.h"
+#include "Animation/DebugSkelMeshComponent.h"
 #include "Rendering/SkeletalMeshModel.h"
 #include "Rendering/SkeletalMeshLODModel.h"
 #include "Engine/SkeletalMesh.h"
@@ -53,6 +54,7 @@ bool FSkinWeightsUtilities::ImportAlternateSkinWeight(USkeletalMesh* SkeletalMes
 		UE_LOG(LogSkinWeightsUtilities, Error, TEXT("Path containing Skin Weight Profile data does not exist (%s)."), *Path);
 		return false;
 	}
+	FScopedSuspendAlternateSkinWeightPreview ScopedSuspendAlternateSkinnWeightPreview(SkeletalMesh);
 	FScopedSkeletalMeshPostEditChange ScopePostEditChange(SkeletalMesh);
 
 	UnFbx::FBXImportOptions ImportOptions;
@@ -249,7 +251,7 @@ bool FSkinWeightsUtilities::ReimportAlternateSkinWeight(USkeletalMesh* SkeletalM
 	{
 		return bResult;
 	}
-
+	FScopedSuspendAlternateSkinWeightPreview ScopedSuspendAlternateSkinnWeightPreview(SkeletalMesh);
 	FScopedSkeletalMeshPostEditChange ScopePostEditChange(SkeletalMesh);
 
 	for (int32 ProfileIndex = 0; ProfileIndex < SkinWeightProfiles.Num(); ++ProfileIndex)
