Restore user change when we restore a reduce custom LOD
#rnx

diff --git a/Engine/Source/Editor/Persona/Private/PersonaMeshDetails.cpp b/Engine/Source/Editor/Persona/Private/PersonaMeshDetails.cpp
--- a/Engine/Source/Editor/Persona/Private/PersonaMeshDetails.cpp
+++ b/Engine/Source/Editor/Persona/Private/PersonaMeshDetails.cpp
@@ -3002,7 +3002,52 @@ void FPersonaMeshDetails::RegenerateOneLOD(int32 LODIndex)
 			CurrentLODInfo.bHasBeenSimplified = false;
 			if (!bIsLODModelbuildDataAvailable)
 			{
+				TMap<int32, FSkelMeshSourceSectionUserData> OldUserSectionsData;
+				TArray<FSkelMeshSection> OldSections;
+				if (SkelMesh->GetImportedModel() && SkelMesh->GetImportedModel()->LODModels.IsValidIndex(LODIndex))
+				{
+					OldSections = SkelMesh->GetImportedModel()->LODModels[LODIndex].Sections;
+					OldUserSectionsData = SkelMesh->GetImportedModel()->LODModels[LODIndex].UserSectionsData;
+				}
+				
 				FLODUtilities::RestoreSkeletalMeshLODImportedData(SkelMesh, LODIndex);
+				if (OldUserSectionsData.Num())
+				{
+					check(SkelMesh->GetImportedModel()->LODModels.IsValidIndex(LODIndex));
+					FSkeletalMeshLODModel& RestoredLODModel = SkelMesh->GetImportedModel()->LODModels[LODIndex];
+					for (int32 OldSectionIndex = 0; OldSectionIndex < OldSections.Num(); ++OldSectionIndex)
+					{
+						FSkelMeshSection& OldSection = OldSections[OldSectionIndex];
+						if (!OldUserSectionsData.Contains(OldSection.OriginalDataSectionIndex))
+						{
+							continue;
+						}
+						FSkelMeshSourceSectionUserData& OldSectionUserData = OldUserSectionsData.FindChecked(OldSection.OriginalDataSectionIndex);
+						for (int32 SectionIndex = 0; SectionIndex < RestoredLODModel.Sections.Num(); ++SectionIndex)
+						{
+							FSkelMeshSection& Section = RestoredLODModel.Sections[SectionIndex];
+							if (Section.MaterialIndex == OldSection.MaterialIndex)
+							{
+								FSkelMeshSourceSectionUserData& RestoredSectionUserData = RestoredLODModel.UserSectionsData.FindOrAdd(Section.OriginalDataSectionIndex);
+								RestoredSectionUserData.bDisabled = OldSectionUserData.bDisabled;
+								RestoredSectionUserData.bCastShadow = OldSectionUserData.bCastShadow;
+								RestoredSectionUserData.bRecomputeTangent = OldSectionUserData.bRecomputeTangent;
+								RestoredSectionUserData.GenerateUpToLodIndex = OldSectionUserData.GenerateUpToLodIndex;
+								RestoredSectionUserData.CorrespondClothAssetIndex = OldSectionUserData.CorrespondClothAssetIndex;
+								RestoredSectionUserData.ClothingData = OldSectionUserData.ClothingData;
+
+								Section.bDisabled = OldSection.bDisabled;
+								Section.bCastShadow = OldSection.bCastShadow;
+								Section.bRecomputeTangent = OldSection.bRecomputeTangent;
+								Section.GenerateUpToLodIndex = OldSection.GenerateUpToLodIndex;
+								Section.CorrespondClothAssetIndex = OldSection.CorrespondClothAssetIndex;
+								Section.ClothingData = OldSection.ClothingData;
+								break;
+							}
+						}
+						SkelMesh->GetImportedModel()->LODModels[LODIndex].UserSectionsData;
+					}
+				}
 			}
 		}
 
