Fix linux build
#rnx

diff --git a/Engine/Source/Editor/Persona/Private/PersonaMeshDetails.cpp b/Engine/Source/Editor/Persona/Private/PersonaMeshDetails.cpp
--- a/Engine/Source/Editor/Persona/Private/PersonaMeshDetails.cpp
+++ b/Engine/Source/Editor/Persona/Private/PersonaMeshDetails.cpp
@@ -3045,7 +3045,6 @@ void FPersonaMeshDetails::RegenerateOneLOD(int32 LODIndex)
 								break;
 							}
 						}
-						SkelMesh->GetImportedModel()->LODModels[LODIndex].UserSectionsData;
 					}
 				}
 			}
