Fixed crash when deleting a skeleton

Re-implementing fix from 4.21 that didnt get merged back

diff --git a/Engine/Source/Editor/Persona/Private/SAnimationSequenceBrowser.cpp b/Engine/Source/Editor/Persona/Private/SAnimationSequenceBrowser.cpp
--- a/Engine/Source/Editor/Persona/Private/SAnimationSequenceBrowser.cpp
+++ b/Engine/Source/Editor/Persona/Private/SAnimationSequenceBrowser.cpp
@@ -1377,8 +1377,10 @@ bool SAnimationSequenceBrowser::OnVisualizeAssetToolTip(const TSharedPtr<SWidget
 		// one viewport between all of the assets in the sequence browser.
 		UAnimationAsset* Asset = StaticCast<UAnimationAsset*>(AssetData.GetAsset());
 		USkeleton* Skeleton = Asset->GetSkeleton();
-		
-		MeshToUse = Skeleton->GetAssetPreviewMesh(Asset);
+		if(Skeleton)
+		{
+			MeshToUse = Skeleton->GetAssetPreviewMesh(Asset);
+		}
 		
 		if(MeshToUse)
 		{
