Animation BP preview meshes now correctly persist between sessions

diff --git a/Engine/Source/Editor/Persona/Private/PersonaAssetFamily.cpp b/Engine/Source/Editor/Persona/Private/PersonaAssetFamily.cpp
--- a/Engine/Source/Editor/Persona/Private/PersonaAssetFamily.cpp
+++ b/Engine/Source/Editor/Persona/Private/PersonaAssetFamily.cpp
@@ -364,8 +364,9 @@ void FPersonaAssetFamily::FindCounterpartAssets(const UObject* InAsset, const US
 	{
 		const UAnimBlueprint* AnimBlueprint = CastChecked<const UAnimBlueprint>(InAsset);
 		OutSkeleton = AnimBlueprint->TargetSkeleton;
+		OutMesh = AnimBlueprint->GetPreviewMesh();
 		check(AnimBlueprint->BlueprintType == BPTYPE_Interface || AnimBlueprint->TargetSkeleton != nullptr);
-		if(AnimBlueprint->TargetSkeleton)
+		if(OutMesh == nullptr && AnimBlueprint->TargetSkeleton)
 		{
 			OutMesh = AnimBlueprint->TargetSkeleton->GetPreviewMesh();
 			if(OutMesh == nullptr)
