Values rapidly go back and forth when hovering the edges of the graph in the Animation Editor\
#fix Clamp blend space sample position to min/max per-axis values
#misc ensure that constrain to axis uses atleast one axis to constrain to if both blendspace value ranges are equal

diff --git a/Engine/Source/Editor/Persona/Private/SAnimationBlendSpaceGridWidget.cpp b/Engine/Source/Editor/Persona/Private/SAnimationBlendSpaceGridWidget.cpp
--- a/Engine/Source/Editor/Persona/Private/SAnimationBlendSpaceGridWidget.cpp
+++ b/Engine/Source/Editor/Persona/Private/SAnimationBlendSpaceGridWidget.cpp
@@ -997,9 +997,13 @@ const FVector SBlendSpaceGridWidget::GridPositionToSampleValue(const FVector2D&
 	Position *= 0.5f;
 
 	// Calculate the sample value by mapping it to the blend parameter range
-	const FVector SampleValue((Position.X * SampleValueRange.X) + SampleValueMin.X,
-							  (GridType == EGridType::TwoAxis) ? SampleValueMax.Y - (Position.Y * SampleValueRange.Y) : 0.0f,
-							  0.0f );
+	const FVector SampleValue
+	(		
+		FMath::Clamp((Position.X * SampleValueRange.X) + SampleValueMin.X, SampleValueMin.X, SampleValueMax.X),
+		FMath::Clamp(((GridType == EGridType::TwoAxis) ? SampleValueMax.Y - (Position.Y * SampleValueRange.Y) : 0.0f),
+			SampleValueMin.Y, SampleValueMax.Y),
+		0.f	
+	);
 	return SampleValue;
 }
 
@@ -1364,7 +1368,7 @@ void SBlendSpaceGridWidget::UpdateGridRationMargin(const FVector2D& GeometrySize
 		// Reset values first
 		GridRatioMargin.Top = GridRatioMargin.Bottom = GridRatioMargin.Left = GridRatioMargin.Right = 0.0f;
 
-		if (SampleValueRange.X > SampleValueRange.Y)
+		if (SampleValueRange.X >= SampleValueRange.Y)
 		{
 			if (GeometrySize.Y > GeometrySize.X)
 			{
