Adding more info and filter options for animation assets
#feature Add the ability to filter according to any curves or notifies rather than a specifically named one

diff --git a/Engine/Source/Editor/Persona/Private/SAnimationSequenceBrowser.cpp b/Engine/Source/Editor/Persona/Private/SAnimationSequenceBrowser.cpp
--- a/Engine/Source/Editor/Persona/Private/SAnimationSequenceBrowser.cpp
+++ b/Engine/Source/Editor/Persona/Private/SAnimationSequenceBrowser.cpp
@@ -91,13 +91,39 @@ public:
 
 	virtual FText GetToolTipText() const override
 	{
-		return NotifyString.IsEmpty() ? LOCTEXT("FFrontendFilter_SkeletonNotifyAssetsEmptyTooltip", "Show assets that use a specified skeleton notify") : FText::Format(LOCTEXT("FFrontendFilter_SkeletonNotifyAssetsTooltip", "Show assets that use skeleton notify '{0}'"), FText::FromString(NotifyString));
+		return NotifyString.IsEmpty() ? LOCTEXT("FFrontendFilter_SkeletonNotifyAssetsEmptyTooltip", "Show assets that contains a (specified) skeleton notify") : FText::Format(LOCTEXT("FFrontendFilter_SkeletonNotifyAssetsTooltip", "Show assets that use skeleton notify '{0}'"), FText::FromString(NotifyString));
 	}
 
 	virtual void ModifyContextMenu(FMenuBuilder& MenuBuilder) override
 	{
 		MenuBuilder.BeginSection(TEXT("NotifySection"), LOCTEXT("NotifySectionHeading", "Choose Notify"));
 
+		MenuBuilder.AddMenuEntry(LOCTEXT("AllNotifyFilterName", "Any Animation Notify"),
+			LOCTEXT("AllNotifyFilterName_ToolTip", "Consider all Animation Notifies, rather than a specific one, for filtering."),
+			FSlateIcon(),
+			FUIAction(
+				FExecuteAction::CreateLambda([this]()
+				{
+					if (!NotifyString.IsEmpty())
+					{
+						NotifyString.Empty();
+					}
+
+					OnChanged().Broadcast();
+				}), 
+				FCanExecuteAction::CreateLambda([this]()
+				{
+					return !NotifyString.IsEmpty();
+				}),
+				FGetActionCheckState::CreateLambda([this]()
+				{
+					return NotifyString.IsEmpty() ? ECheckBoxState::Checked : ECheckBoxState::Unchecked;
+				})
+			), 
+			NAME_None,
+			EUserInterfaceActionType::ToggleButton
+		);
+		
 		TSharedRef<SWidget> Widget =
 			SNew(SBox)
 			.HeightOverride(300.0f)
@@ -126,10 +152,10 @@ public:
 	// IFilter implementation
 	virtual bool PassesFilter(FAssetFilterType InItem) const override
 	{
-		if (!NotifyString.IsEmpty())
+		const FString TagValue = InItem.GetTagValueRef<FString>(USkeleton::AnimNotifyTag);
+		if (!TagValue.IsEmpty())
 		{
-			const FString TagValue = InItem.GetTagValueRef<FString>(USkeleton::AnimNotifyTag);
-			if (!TagValue.IsEmpty())
+			if (!NotifyString.IsEmpty())
 			{
 				// parse notifies
 				TArray<FString> NotifyValues;
@@ -145,10 +171,10 @@ public:
 				}
 			}
 
-			return false;
+			return NotifyString.IsEmpty();
 		}
 
-		return true;
+		return false;
 	}
 
 	void SetNotifyFilter(const FName& InNotifyFilter)
@@ -188,13 +214,39 @@ public:
 
 	virtual FText GetToolTipText() const override
 	{
-		return CurveString.IsEmpty() ? LOCTEXT("FFrontendFilter_CurveAssetsEmptyTooltip", "Show assets that use a specified curve") : FText::Format(LOCTEXT("FFrontendFilter_CurveAssetsTooltip", "Show assets that use curve '{0}'"), FText::FromString(CurveString));
+		return CurveString.IsEmpty() ? LOCTEXT("FFrontendFilter_CurveAssetsEmptyTooltip", "Show assets that contains a (specified) curve") : FText::Format(LOCTEXT("FFrontendFilter_CurveAssetsTooltip", "Show assets that use curve '{0}'"), FText::FromString(CurveString));
 	}
 
 	virtual void ModifyContextMenu(FMenuBuilder& MenuBuilder) override
 	{
 		MenuBuilder.BeginSection(TEXT("CurveSection"), LOCTEXT("CurveSectionHeading", "Choose Curve"));
+		
+		MenuBuilder.AddMenuEntry(LOCTEXT("AllCurveFilterName", "Any Animation Curve"),
+			LOCTEXT("AllCurveFilterName_ToolTip", "Consider all Animation Curves, rather than a specific one, for filtering."),
+			FSlateIcon(),
+			FUIAction(
+				FExecuteAction::CreateLambda([this]()
+				{
+					if (!CurveString.IsEmpty())
+					{
+						CurveString.Empty();
+					}
 
+					OnChanged().Broadcast();
+				}), 
+				FCanExecuteAction::CreateLambda([this]()
+				{
+					return !CurveString.IsEmpty();
+				}),
+				FGetActionCheckState::CreateLambda([this]()
+				{
+					return CurveString.IsEmpty() ? ECheckBoxState::Checked : ECheckBoxState::Unchecked;
+				})
+			), 
+			NAME_None,
+			EUserInterfaceActionType::ToggleButton
+		);
+		
 		TSharedRef<SWidget> Widget = 
 			SNew(SBox)
 			.HeightOverride(300.0f)
@@ -222,10 +274,10 @@ public:
 	// IFilter implementation
 	virtual bool PassesFilter(FAssetFilterType InItem) const override
 	{
-		if (!CurveString.IsEmpty())
+		const FString TagValue = InItem.GetTagValueRef<FString>(USkeleton::CurveNameTag);
+		if (!TagValue.IsEmpty())
 		{
-			const FString TagValue = InItem.GetTagValueRef<FString>(USkeleton::CurveNameTag);
-			if (!TagValue.IsEmpty())
+			if (!CurveString.IsEmpty())
 			{
 				// parse curves
 				TArray<FString> CurveValues;
@@ -241,10 +293,10 @@ public:
 				}
 			}
 
-			return false;
+			return CurveString.IsEmpty();
 		}
 
-		return true;
+		return false;
 	}
 
 	void SetCurveFilter(const FName& InCurveFilter)
