Make animation preview scene additional meshes non-selectable in phat.

diff --git a/Engine/Source/Editor/Persona/Private/AnimationEditorPreviewScene.cpp b/Engine/Source/Editor/Persona/Private/AnimationEditorPreviewScene.cpp
--- a/Engine/Source/Editor/Persona/Private/AnimationEditorPreviewScene.cpp
+++ b/Engine/Source/Editor/Persona/Private/AnimationEditorPreviewScene.cpp
@@ -51,6 +51,7 @@ FAnimationEditorPreviewScene::FAnimationEditorPreviewScene(const ConstructionVal
 	, LastTickTime(0.0)
 	, bSelecting(false)
 	, bAllowAdditionalMeshes(true)
+	, bAdditionalMeshesSelectable(true)
 {
 	if (GEditor)
 	{
@@ -317,6 +318,11 @@ void FAnimationEditorPreviewScene::SetAdditionalMeshes(class UDataAsset* InAddit
 	RefreshAdditionalMeshes(true);
 }
 
+void FAnimationEditorPreviewScene::SetAdditionalMeshesSelectable(bool bSelectable)
+{
+	bAdditionalMeshesSelectable = bSelectable;
+}
+
 void FAnimationEditorPreviewScene::RefreshAdditionalMeshes(bool bAllowOverrideBaseMesh)
 {
 	// remove all components
@@ -369,11 +375,11 @@ void FAnimationEditorPreviewScene::RefreshAdditionalMeshes(bool bAllowOverrideBa
 					if (SkeletalMesh)
 					{
 						USkeletalMeshComponent* NewComp = NewObject<USkeletalMeshComponent>(Actor);
+						NewComp->bSelectable = bAdditionalMeshesSelectable;
 						NewComp->RegisterComponent();
 						NewComp->SetSkeletalMesh(SkeletalMesh);
 						NewComp->bUseAttachParentBound = true;
 						AddComponent(NewComp, FTransform::Identity, true);
-
 						if (bUseCustomAnimBP && AnimInstances.IsValidIndex(MeshIndex) && AnimInstances[MeshIndex] != nullptr)
 						{
 							NewComp->SetAnimInstanceClass(AnimInstances[MeshIndex]);
@@ -383,7 +389,6 @@ void FAnimationEditorPreviewScene::RefreshAdditionalMeshes(bool bAllowOverrideBa
 							bool bWasCreated = false;
 							UAnimCustomInstance::BindToSkeletalMeshComponent<UAnimPreviewAttacheInstance>(NewComp,bWasCreated);
 						}
-
 						AdditionalMeshes.Add(NewComp);
 					}
 				}
diff --git a/Engine/Source/Editor/Persona/Private/AnimationEditorPreviewScene.h b/Engine/Source/Editor/Persona/Private/AnimationEditorPreviewScene.h
--- a/Engine/Source/Editor/Persona/Private/AnimationEditorPreviewScene.h
+++ b/Engine/Source/Editor/Persona/Private/AnimationEditorPreviewScene.h
@@ -34,6 +34,7 @@ public:
 	virtual UDebugSkelMeshComponent* GetPreviewMeshComponent() const override { return SkeletalMeshComponent; }
 	virtual void SetPreviewMeshComponent(UDebugSkelMeshComponent* InSkeletalMeshComponent) override;
 	virtual void SetAdditionalMeshes(class UDataAsset* InAdditionalMeshes) override;
+	virtual void SetAdditionalMeshesSelectable(bool bSelectable) override;
 	virtual void RefreshAdditionalMeshes(bool bAllowOverrideBaseMesh) override;
 	virtual void ShowReferencePose(bool bShowRefPose, bool bResetBoneTransforms = false) override;
 	virtual bool IsShowReferencePoseEnabled() const override;
@@ -423,4 +424,7 @@ private:
 
 	/** Allow additional meshes to be attached */
 	bool bAllowAdditionalMeshes;
+
+	/** Allow additional meshes to be selectable */
+	bool bAdditionalMeshesSelectable;
 };
diff --git a/Engine/Source/Editor/Persona/Public/IPersonaPreviewScene.h b/Engine/Source/Editor/Persona/Public/IPersonaPreviewScene.h
--- a/Engine/Source/Editor/Persona/Public/IPersonaPreviewScene.h
+++ b/Engine/Source/Editor/Persona/Public/IPersonaPreviewScene.h
@@ -84,6 +84,9 @@ public:
 	/** Set the additional meshes used by this preview scene (sets the additional meshes on the skeleton) */
 	virtual void SetAdditionalMeshes(class UDataAsset* InAdditionalMeshes) = 0;
 
+	/** Set whether additional meshes are selectable */
+	virtual void SetAdditionalMeshesSelectable(bool bSelectable) = 0;
+
 	/** Refreshes the additional meshes displayed in this preview scene */
 	virtual void RefreshAdditionalMeshes(bool bAllowOverrideBaseMesh) = 0;
 
diff --git a/Engine/Source/Editor/PhysicsAssetEditor/Private/PhysicsAssetEditor.cpp b/Engine/Source/Editor/PhysicsAssetEditor/Private/PhysicsAssetEditor.cpp
--- a/Engine/Source/Editor/PhysicsAssetEditor/Private/PhysicsAssetEditor.cpp
+++ b/Engine/Source/Editor/PhysicsAssetEditor/Private/PhysicsAssetEditor.cpp
@@ -2908,7 +2908,7 @@ void FPhysicsAssetEditor::HandlePreviewSceneCreated(const TSharedRef<IPersonaPre
 	SharedData->EditorSkelComp->SetDisablePostProcessBlueprint(true);
 	InPersonaPreviewScene->SetPreviewMeshComponent(SharedData->EditorSkelComp);
 	InPersonaPreviewScene->AddComponent(SharedData->EditorSkelComp, FTransform::Identity);
-
+	InPersonaPreviewScene->SetAdditionalMeshesSelectable(false);
 	// set root component, so we can attach to it. 
 	Actor->SetRootComponent(SharedData->EditorSkelComp);
 
