Fix LODMaterialMap issues, some place where using the material index instead of the section index to acces it
Also make sure that the material index and the LODMaterial map are use with the same design everywhere.
The importer, the reduction and the edition are now using the same concept for the material IDs.

diff --git a/Engine/Plugins/Experimental/ChaosClothEditor/Source/Chaos/Private/ChaosClothEditor/ChaosSimulationEditorExtender.cpp b/Engine/Plugins/Experimental/ChaosClothEditor/Source/Chaos/Private/ChaosClothEditor/ChaosSimulationEditorExtender.cpp
--- a/Engine/Plugins/Experimental/ChaosClothEditor/Source/Chaos/Private/ChaosClothEditor/ChaosSimulationEditorExtender.cpp
+++ b/Engine/Plugins/Experimental/ChaosClothEditor/Source/Chaos/Private/ChaosClothEditor/ChaosSimulationEditorExtender.cpp
@@ -158,7 +158,7 @@ void FSimulationEditorExtender::ShowClothSections(USkeletalMeshComponent* MeshCo
 
 				if (SkelMeshRenderSection.HasClothingData())
 				{
-					MeshComponent->ShowMaterialSection(SkelMeshRenderSection.MaterialIndex, bIsClothSectionsVisible, LODIndex);
+					MeshComponent->ShowMaterialSection(SkelMeshRenderSection.MaterialIndex, SectionIndex, bIsClothSectionsVisible, LODIndex);
 				}
 			}
 		}
diff --git a/Engine/Plugins/Experimental/SkeletalReduction/Source/Private/SkeletalMeshReductionPlugin.cpp b/Engine/Plugins/Experimental/SkeletalReduction/Source/Private/SkeletalMeshReductionPlugin.cpp
--- a/Engine/Plugins/Experimental/SkeletalReduction/Source/Private/SkeletalMeshReductionPlugin.cpp
+++ b/Engine/Plugins/Experimental/SkeletalReduction/Source/Private/SkeletalMeshReductionPlugin.cpp
@@ -1826,42 +1826,11 @@ void FQuadricSkeletalMeshReduction::ReduceSkeletalMesh(USkeletalMesh& SkeletalMe
 		int32 OriginalDataSectionIndex;
 	};
 
-	TMap<int32, FSectionData> BackupSectionIndexToSectionData;
 	TMap<int32, FSkelMeshSourceSectionUserData> BackupUserSectionsData;
 	FString BackupLodModelBuildStringID = TEXT("");
-	//Store the section data. Store the source LOD model in case we add a LOD, or the target LOD model in case the LOD already exist
-	FSkeletalMeshLODModel& BackupSectionLODModel = bLODModelAdded ? SkeletalMeshResource.LODModels[Settings.BaseLOD] : SkeletalMeshResource.LODModels[LODIndex];
-	if (!bLODModelAdded)
-	{
-		BackupLodModelBuildStringID = BackupSectionLODModel.BuildStringID;
-		BackupUserSectionsData = BackupSectionLODModel.UserSectionsData;
-	}
 
-	BackupSectionIndexToSectionData.Reserve(BackupSectionLODModel.Sections.Num());
 
-	for (int32 SectionIndex = 0; SectionIndex < BackupSectionLODModel.Sections.Num(); ++SectionIndex)
-	{
-		//Skip disable and not generated section
-		int32 MaxGeneratedLODIndex = BackupSectionLODModel.Sections[SectionIndex].GenerateUpToLodIndex;
-		if (BackupSectionLODModel.Sections[SectionIndex].bDisabled || (MaxGeneratedLODIndex != -1 && MaxGeneratedLODIndex < LODIndex))
-		{
-			continue;
-		}
-		FSectionData& SectionData = BackupSectionIndexToSectionData.FindOrAdd(SectionIndex);
-		SectionData.MaterialIndex = BackupSectionLODModel.Sections[SectionIndex].MaterialIndex;
-		SectionData.bCastShadow = BackupSectionLODModel.Sections[SectionIndex].bCastShadow;
-		SectionData.bRecomputeTangent = BackupSectionLODModel.Sections[SectionIndex].bRecomputeTangent;
-		SectionData.bDisabled = BackupSectionLODModel.Sections[SectionIndex].bDisabled;
-		SectionData.GenerateUpToLodIndex = BackupSectionLODModel.Sections[SectionIndex].GenerateUpToLodIndex;
-		SectionData.ChunkedParentSectionIndex = BackupSectionLODModel.Sections[SectionIndex].ChunkedParentSectionIndex;
-		SectionData.OriginalDataSectionIndex = BackupSectionLODModel.Sections[SectionIndex].OriginalDataSectionIndex;
-		SectionData.MaterialMap = LODInfo->LODMaterialMap.IsValidIndex(SectionIndex) ? LODInfo->LODMaterialMap[SectionIndex] : INDEX_NONE;
-		if (SectionData.MaterialMap == SectionData.MaterialIndex)
-		{
-			//Remove any override if the value is the same
-			SectionData.MaterialMap = INDEX_NONE;
-		}
-	}
+	
 
 	auto FillSectionMaterialSlot = [&SkeletalMeshResource, &LODIndex, bLODModelAdded](TArray<int32>& SectionMaterialSlot)
 	{
@@ -1887,12 +1856,22 @@ void FQuadricSkeletalMeshReduction::ReduceSkeletalMesh(USkeletalMesh& SkeletalMe
 		FLODUtilities::UnbindClothingAndBackup(&SkeletalMesh, ClothingBindings, LODIndex);
 	}
 
+	if (!bLODModelAdded)
+	{
+		FSkeletalMeshLODModel& DstBackupSectionLODModel = SkeletalMeshResource.LODModels[LODIndex];
+		BackupLodModelBuildStringID = DstBackupSectionLODModel.BuildStringID;
+		BackupUserSectionsData = DstBackupSectionLODModel.UserSectionsData;
+	}
+
 	bool bReducingSourceModel = false;
 	//Reducing source LOD, we need to use the temporary data so it can be iterative
 	if (BaseLOD == LODIndex && (SkelResource->OriginalReductionSourceMeshData.IsValidIndex(BaseLOD) && !SkelResource->OriginalReductionSourceMeshData[BaseLOD]->IsEmpty()))
 	{
 		TMap<FString, TArray<FMorphTargetDelta>> TempLODMorphTargetData;
 		SkelResource->OriginalReductionSourceMeshData[BaseLOD]->LoadReductionData(*SrcModel, TempLODMorphTargetData, &SkeletalMesh);
+		//Rebackup the source model since we update it, source always have empty LODMaterial map
+		//If you swap a material ID and after you do inline reduction, you have to remap it again, but not if you reduce and then remap the materialID
+		//this is by design currently
 		bReducingSourceModel = true;
 	}
 	else
@@ -1902,6 +1881,37 @@ void FQuadricSkeletalMeshReduction::ReduceSkeletalMesh(USkeletalMesh& SkeletalMe
 
 	check(SrcModel);
 
+	//We backup the section data to put keep the LODModel Section in a good state after the reduction
+	TMap<int32, FSectionData> SrcBackupSectionIndexToSectionData;
+	{
+		FSkeletalMeshLODInfo* BackupLODInfo = SkeletalMesh.GetLODInfo(Settings.BaseLOD);
+		FSkeletalMeshLODModel& LODModelToBackup = *SrcModel;
+		SrcBackupSectionIndexToSectionData.Empty(LODModelToBackup.Sections.Num());
+		for (int32 SectionIndex = 0; SectionIndex < LODModelToBackup.Sections.Num(); ++SectionIndex)
+		{
+			//Skip disable and not generated section
+			int32 MaxGeneratedLODIndex = LODModelToBackup.Sections[SectionIndex].GenerateUpToLodIndex;
+			if (LODModelToBackup.Sections[SectionIndex].bDisabled || (MaxGeneratedLODIndex != -1 && MaxGeneratedLODIndex < LODIndex))
+			{
+				continue;
+			}
+			FSectionData& SectionData = SrcBackupSectionIndexToSectionData.FindOrAdd(SectionIndex);
+			SectionData.MaterialIndex = LODModelToBackup.Sections[SectionIndex].MaterialIndex;
+			SectionData.bCastShadow = LODModelToBackup.Sections[SectionIndex].bCastShadow;
+			SectionData.bRecomputeTangent = LODModelToBackup.Sections[SectionIndex].bRecomputeTangent;
+			SectionData.bDisabled = LODModelToBackup.Sections[SectionIndex].bDisabled;
+			SectionData.GenerateUpToLodIndex = LODModelToBackup.Sections[SectionIndex].GenerateUpToLodIndex;
+			SectionData.ChunkedParentSectionIndex = LODModelToBackup.Sections[SectionIndex].ChunkedParentSectionIndex;
+			SectionData.OriginalDataSectionIndex = LODModelToBackup.Sections[SectionIndex].OriginalDataSectionIndex;
+			SectionData.MaterialMap = (BackupLODInfo != nullptr && BackupLODInfo->LODMaterialMap.IsValidIndex(SectionIndex)) ? BackupLODInfo->LODMaterialMap[SectionIndex] : INDEX_NONE;
+			if (SectionData.MaterialMap == SectionData.MaterialIndex)
+			{
+				//Remove any override if the value is the same
+				SectionData.MaterialMap = INDEX_NONE;
+			}
+		}
+	}
+
 	// now try bone reduction process if it's setup
 	TMap<FBoneIndexType, FBoneIndexType> BonesToRemove;
 	
@@ -2038,51 +2048,52 @@ void FQuadricSkeletalMeshReduction::ReduceSkeletalMesh(USkeletalMesh& SkeletalMe
 		// Flag this LOD as having been simplified.
 		ReducedLODInfoPtr->bHasBeenSimplified = true;
 		SkeletalMesh.bHasBeenSimplified = true;
-		//Restore section data
-		FSkeletalMeshLODModel& ImportedModelLOD = SkeletalMesh.GetImportedModel()->LODModels[LODIndex];
-		TArray<bool> SectionMatched;
-		SectionMatched.AddZeroed(ImportedModelLOD.Sections.Num());
 		
-		if (!bLODModelAdded)
-		{
-			ImportedModelLOD.UserSectionsData = BackupUserSectionsData;
-			ImportedModelLOD.BuildStringID = BackupLodModelBuildStringID;
-		}
-
-		for (auto Kvp : BackupSectionIndexToSectionData)
+		//Restore the source sections data
 		{
-			const int32 SourceSectionIndex = Kvp.Key;
-
-			const FSectionData& SectionData = Kvp.Value;
-			int32 MaxSectionIndex = FMath::Min(ImportedModelLOD.Sections.Num(), SourceSectionIndex+1);
-			for (int32 SectionIndex = 0; SectionIndex < MaxSectionIndex; ++SectionIndex)
+			FSkeletalMeshLODModel& ImportedModelLOD = SkeletalMesh.GetImportedModel()->LODModels[LODIndex];
+			TArray<bool> SectionMatched;
+			SectionMatched.AddZeroed(SrcBackupSectionIndexToSectionData.Num());
+			int32 CurrentParentSectionIndex = INDEX_NONE;
+			int32 OriginalSectionIndex = INDEX_NONE;
+			for (int32 SectionIndex = 0; SectionIndex < ImportedModelLOD.Sections.Num(); ++SectionIndex)
 			{
-				if (SectionMatched[SectionIndex])
-				{
-					continue;
-				}
-				if (SectionData.MaterialIndex == ImportedModelLOD.Sections[SectionIndex].MaterialIndex)
+				FSkelMeshSection& Section = ImportedModelLOD.Sections[SectionIndex];
+				for (auto Kvp : SrcBackupSectionIndexToSectionData)
 				{
-					//Restore the LOD material map for this section
-					if (SectionData.MaterialMap != INDEX_NONE)
+					const int32 SourceSectionIndex = Kvp.Key;
+					if (SectionMatched[Kvp.Key])
+					{
+						continue;
+					}
+					const FSectionData& SectionData = Kvp.Value;
+					//We use the material index to match the section
+					if (Section.MaterialIndex == SectionData.MaterialIndex)
 					{
-						while (ReducedLODInfoPtr->LODMaterialMap.Num() <= SectionIndex)
+						bool bIsChunkedSection = SectionData.ChunkedParentSectionIndex != INDEX_NONE;
+						if (!bIsChunkedSection)
 						{
-							ReducedLODInfoPtr->LODMaterialMap.Add(INDEX_NONE);
+							CurrentParentSectionIndex = SectionIndex;
+							OriginalSectionIndex++;
 						}
-						ReducedLODInfoPtr->LODMaterialMap[SectionIndex] = SectionData.MaterialMap;
+						Section.bCastShadow = SectionData.bCastShadow;
+						Section.bRecomputeTangent = SectionData.bRecomputeTangent;
+						Section.bDisabled = SectionData.bDisabled;
+						Section.GenerateUpToLodIndex = SectionData.GenerateUpToLodIndex;
+						Section.ChunkedParentSectionIndex = bIsChunkedSection ? CurrentParentSectionIndex : INDEX_NONE;
+						Section.OriginalDataSectionIndex = OriginalSectionIndex;
+						SectionMatched[Kvp.Key] = true; //a backup section can be restore only once
+						break;
 					}
-					//Restore the LODModel Section
-					ImportedModelLOD.Sections[SectionIndex].bCastShadow = SectionData.bCastShadow;
-					ImportedModelLOD.Sections[SectionIndex].bRecomputeTangent = SectionData.bRecomputeTangent;
-					ImportedModelLOD.Sections[SectionIndex].bDisabled = SectionData.bDisabled;
-					ImportedModelLOD.Sections[SectionIndex].GenerateUpToLodIndex = SectionData.GenerateUpToLodIndex;
-					ImportedModelLOD.Sections[SectionIndex].ChunkedParentSectionIndex = SectionData.ChunkedParentSectionIndex;
-					ImportedModelLOD.Sections[SectionIndex].OriginalDataSectionIndex = SectionData.OriginalDataSectionIndex;
-					SectionMatched[SectionIndex] = true;
-					break;
 				}
 			}
+
+			if (!bLODModelAdded)
+			{
+				//If its an existing LOD re-apply the UserSectionData
+				ImportedModelLOD.UserSectionsData = BackupUserSectionsData;
+				ImportedModelLOD.BuildStringID = BackupLodModelBuildStringID;
+			}
 		}
 	}
 	else
diff --git a/Engine/Source/Developer/SkeletalMeshUtilitiesCommon/Private/LODUtilities.cpp b/Engine/Source/Developer/SkeletalMeshUtilitiesCommon/Private/LODUtilities.cpp
--- a/Engine/Source/Developer/SkeletalMeshUtilitiesCommon/Private/LODUtilities.cpp
+++ b/Engine/Source/Developer/SkeletalMeshUtilitiesCommon/Private/LODUtilities.cpp
@@ -824,7 +824,32 @@ void FLODUtilities::ApplyMorphTargetsToLOD(USkeletalMesh* SkeletalMesh, int32 So
 	}
 
 	const FSkeletalMeshLODModel& BaseLODModel = bReduceBaseLOD ? TempBaseLODModel : SkeletalMeshResource->LODModels[SourceLOD];
+	const FSkeletalMeshLODInfo* BaseLODInfo = SkeletalMesh->GetLODInfo(SourceLOD);
 	const FSkeletalMeshLODModel& TargetLODModel = SkeletalMeshResource->LODModels[DestinationLOD];
+	const FSkeletalMeshLODInfo* TargetLODInfo = SkeletalMesh->GetLODInfo(DestinationLOD);
+
+	TArray<int32> BaseLODMaterialMap = BaseLODInfo ? BaseLODInfo->LODMaterialMap : TArray<int32>();
+	TArray<int32> TargetLODMaterialMap = TargetLODInfo ? TargetLODInfo->LODMaterialMap : TArray<int32>();
+
+	auto InternalGetSectionMaterialIndex = [](const FSkeletalMeshLODModel& LODModel, int32 SectionIndex)->int32
+	{
+		if (!LODModel.Sections.IsValidIndex(SectionIndex))
+		{
+			return 0;
+		}
+		return LODModel.Sections[SectionIndex].MaterialIndex;
+	};
+
+	auto GetBaseSectionMaterialIndex = [&BaseLODModel, &InternalGetSectionMaterialIndex](int32 SectionIndex)->int32
+	{
+		return InternalGetSectionMaterialIndex(BaseLODModel, SectionIndex);
+	};
+
+	auto GetTargetSectionMaterialIndex = [&TargetLODModel, &InternalGetSectionMaterialIndex](int32 SectionIndex)->int32
+	{
+		return InternalGetSectionMaterialIndex(TargetLODModel, SectionIndex);
+	};
+
 	//Make sure we have some morph for this LOD
 	bool bContainsMorphTargets = false;
 	for (UMorphTarget *MorphTarget : SkeletalMesh->MorphTargets)
@@ -848,23 +873,26 @@ void FLODUtilities::ApplyMorphTargetsToLOD(USkeletalMesh* SkeletalMesh, int32 So
 	{
 		TargetSectionMatchBaseIndex[TargetSectionIndex] = INDEX_NONE;
 	}
+	TBitArray<> BaseSectionMatch;
+	BaseSectionMatch.Init(false, BaseLODModel.Sections.Num());
 	//Find corresponding section indices from Source LOD for Target LOD
-	for (int32 BaseSectionIndex = 0; BaseSectionIndex < BaseLODModel.Sections.Num(); ++BaseSectionIndex)
+	for (int32 TargetSectionIndex = 0; TargetSectionIndex < TargetLODModel.Sections.Num(); ++TargetSectionIndex)
 	{
-		int32 TargetSectionIndexMatch = INDEX_NONE;
-		for (int32 TargetSectionIndex = 0; TargetSectionIndex < TargetLODModel.Sections.Num(); ++TargetSectionIndex)
+		int32 TargetSectionMaterialIndex = GetTargetSectionMaterialIndex(TargetSectionIndex);
+		for (int32 BaseSectionIndex = 0; BaseSectionIndex < BaseLODModel.Sections.Num(); ++BaseSectionIndex)
 		{
-			if (TargetLODModel.Sections[TargetSectionIndex].MaterialIndex == BaseLODModel.Sections[BaseSectionIndex].MaterialIndex && TargetSectionMatchBaseIndex[TargetSectionIndex] == INDEX_NONE)
+			if (BaseSectionMatch[BaseSectionIndex])
 			{
-				TargetSectionIndexMatch = TargetSectionIndex;
+				continue;
+			}
+			int32 BaseSectionMaterialIndex = GetBaseSectionMaterialIndex(BaseSectionIndex);
+			if (TargetSectionMaterialIndex == BaseSectionMaterialIndex)
+			{
+				TargetSectionMatchBaseIndex[TargetSectionIndex] = BaseSectionIndex;
+				BaseSectionMatch[BaseSectionIndex] = true;
 				break;
 			}
 		}
-		//We can set the data only once. There should be no clash
-		if (TargetSectionMatchBaseIndex.IsValidIndex(TargetSectionIndexMatch) && TargetSectionMatchBaseIndex[TargetSectionIndexMatch] == INDEX_NONE)
-		{
-			TargetSectionMatchBaseIndex[TargetSectionIndexMatch] = BaseSectionIndex;
-		}
 	}
 	//We should have match all the target sections
 	check(!TargetSectionMatchBaseIndex.Contains(INDEX_NONE));
@@ -1123,11 +1151,11 @@ void FLODUtilities::RestoreSkeletalMeshLODImportedData(USkeletalMesh* SkeletalMe
 	TMap<FString, TArray<FMorphTargetDelta>> ImportedBaseLODMorphTargetData;
 	SkeletalMesh->GetImportedModel()->OriginalReductionSourceMeshData[LodIndex]->LoadReductionData(ImportedBaseLODModel, ImportedBaseLODMorphTargetData, SkeletalMesh);
 	{
-		FSkeletalMeshUpdateContext UpdateContext;
-		UpdateContext.SkeletalMesh = SkeletalMesh;
 		TArray<int32> EmptyLodInfoMaterialMap;
 		ImportedBaseLODModel.UpdateChunkedSectionInfo(SkeletalMesh->GetName(), EmptyLodInfoMaterialMap);
-
+		//When we restore a LOD we destroy the LODMaterialMap (user manual section material slot assignation)
+		FSkeletalMeshLODInfo* LODInfo = SkeletalMesh->GetLODInfo(LodIndex);
+		LODInfo->LODMaterialMap.Empty();
 		//Copy the SkeletalMeshLODModel
 		SkeletalMesh->GetImportedModel()->LODModels[LodIndex] = ImportedBaseLODModel;
 		//Copy the morph target deltas
@@ -1150,12 +1178,6 @@ void FLODUtilities::RestoreSkeletalMeshLODImportedData(USkeletalMesh* SkeletalMe
 
 		//Put back the clothing for the restore LOD
 		FLODUtilities::RestoreClothingFromBackup(SkeletalMesh, ClothingBindings);
-
-		if (UpdateContext.OnLODChanged.IsBound())
-		{
-			//Notify calling system of change
-			UpdateContext.OnLODChanged.ExecuteIfBound();
-		}
 	}
 }
 
diff --git a/Engine/Source/Editor/Persona/Private/PersonaMeshDetails.cpp b/Engine/Source/Editor/Persona/Private/PersonaMeshDetails.cpp
--- a/Engine/Source/Editor/Persona/Private/PersonaMeshDetails.cpp
+++ b/Engine/Source/Editor/Persona/Private/PersonaMeshDetails.cpp
@@ -108,6 +108,24 @@ namespace PersonaMeshDetailsConstants
 	const static int32 LodSliderExtension = 5;
 }
 
+static int32 GetDefaultMaterialIndex(const USkeletalMesh* SkeletalMesh, int32 LODIndex, int32 SectionIndex)
+{
+	int32 DefaultMaterialIndex = INDEX_NONE;
+	if (!SkeletalMesh || !SkeletalMesh->GetImportedModel() || !SkeletalMesh->GetImportedModel()->LODModels.IsValidIndex(LODIndex))
+	{
+		return DefaultMaterialIndex;
+	}
+
+	const FSkeletalMeshLODModel& LODModel = SkeletalMesh->GetImportedModel()->LODModels[LODIndex];
+
+	if (LODModel.Sections.IsValidIndex(SectionIndex))
+	{
+		DefaultMaterialIndex = LODModel.Sections[SectionIndex].MaterialIndex;
+	}
+	
+	return DefaultMaterialIndex;
+}
+
 /** Returns true if automatic mesh reduction is available. */
 static bool IsAutoMeshReductionAvailable()
 {
@@ -2992,65 +3010,6 @@ void FPersonaMeshDetails::RegenerateOneLOD(int32 LODIndex)
 			return;
 		}
 
-		//If we are doing inline reduction, restore the skeletalmesh LOD if it was already reduced and we want to reduced it again (user change the options)
-		//This will ensure the data is the same when starting the reduction
-		if (LODIndex == CurrentLODInfo.ReductionSettings.BaseLOD
-			&& CurrentLODInfo.bHasBeenSimplified
-			&& SkelMesh->IsReductionActive(LODIndex)
-			&& (bIsLODModelbuildDataAvailable || bIsReductionDataPresent))
-		{
-			//Restore the base LOD data
-			CurrentLODInfo.bHasBeenSimplified = false;
-			if (!bIsLODModelbuildDataAvailable)
-			{
-				TMap<int32, FSkelMeshSourceSectionUserData> OldUserSectionsData;
-				TArray<FSkelMeshSection> OldSections;
-				if (SkelMesh->GetImportedModel() && SkelMesh->GetImportedModel()->LODModels.IsValidIndex(LODIndex))
-				{
-					OldSections = SkelMesh->GetImportedModel()->LODModels[LODIndex].Sections;
-					OldUserSectionsData = SkelMesh->GetImportedModel()->LODModels[LODIndex].UserSectionsData;
-				}
-				
-				FLODUtilities::RestoreSkeletalMeshLODImportedData(SkelMesh, LODIndex);
-				if (OldUserSectionsData.Num())
-				{
-					check(SkelMesh->GetImportedModel()->LODModels.IsValidIndex(LODIndex));
-					FSkeletalMeshLODModel& RestoredLODModel = SkelMesh->GetImportedModel()->LODModels[LODIndex];
-					for (int32 OldSectionIndex = 0; OldSectionIndex < OldSections.Num(); ++OldSectionIndex)
-					{
-						FSkelMeshSection& OldSection = OldSections[OldSectionIndex];
-						if (!OldUserSectionsData.Contains(OldSection.OriginalDataSectionIndex))
-						{
-							continue;
-						}
-						FSkelMeshSourceSectionUserData& OldSectionUserData = OldUserSectionsData.FindChecked(OldSection.OriginalDataSectionIndex);
-						for (int32 SectionIndex = 0; SectionIndex < RestoredLODModel.Sections.Num(); ++SectionIndex)
-						{
-							FSkelMeshSection& Section = RestoredLODModel.Sections[SectionIndex];
-							if (Section.MaterialIndex == OldSection.MaterialIndex)
-							{
-								FSkelMeshSourceSectionUserData& RestoredSectionUserData = RestoredLODModel.UserSectionsData.FindOrAdd(Section.OriginalDataSectionIndex);
-								RestoredSectionUserData.bDisabled = OldSectionUserData.bDisabled;
-								RestoredSectionUserData.bCastShadow = OldSectionUserData.bCastShadow;
-								RestoredSectionUserData.bRecomputeTangent = OldSectionUserData.bRecomputeTangent;
-								RestoredSectionUserData.GenerateUpToLodIndex = OldSectionUserData.GenerateUpToLodIndex;
-								RestoredSectionUserData.CorrespondClothAssetIndex = OldSectionUserData.CorrespondClothAssetIndex;
-								RestoredSectionUserData.ClothingData = OldSectionUserData.ClothingData;
-
-								Section.bDisabled = OldSection.bDisabled;
-								Section.bCastShadow = OldSection.bCastShadow;
-								Section.bRecomputeTangent = OldSection.bRecomputeTangent;
-								Section.GenerateUpToLodIndex = OldSection.GenerateUpToLodIndex;
-								Section.CorrespondClothAssetIndex = OldSection.CorrespondClothAssetIndex;
-								Section.ClothingData = OldSection.ClothingData;
-								break;
-							}
-						}
-					}
-				}
-			}
-		}
-
 		FSkeletalMeshUpdateContext UpdateContext;
 		UpdateContext.SkeletalMesh = SkelMesh;
 		UpdateContext.AssociatedComponents.Push(GetPersonaToolkit()->GetPreviewMeshComponent());
@@ -3723,11 +3682,11 @@ void FPersonaMeshDetails::OnGetSectionsForView(ISectionListBuilder& OutSections,
 		int32 NumSections = Model.Sections.Num();
 		for (int32 SectionIdx = 0; SectionIdx < NumSections; SectionIdx++)
 		{
-			int32 DefaultSectionMaterialIndex = Model.Sections[SectionIdx].MaterialIndex;
-			int32 MaterialIndex = DefaultSectionMaterialIndex;
-			if (MaterialMap.IsValidIndex(MaterialIndex) && SkelMesh->Materials.IsValidIndex(MaterialMap[MaterialIndex]))
+			int32 DefaultSectionMaterialIndex = GetDefaultMaterialIndex(SkelMesh, LODIndex, SectionIdx);
+			int32 MaterialIndex = Model.Sections[SectionIdx].MaterialIndex;;
+			if (MaterialMap.IsValidIndex(SectionIdx) && SkelMesh->Materials.IsValidIndex(MaterialMap[SectionIdx]))
 			{
-				MaterialIndex = MaterialMap[MaterialIndex];
+				MaterialIndex = MaterialMap[SectionIdx];
 			}
 
 			if (SkelMesh->Materials.IsValidIndex(MaterialIndex))
@@ -3917,33 +3876,34 @@ void FPersonaMeshDetails::OnDeleteMaterialSlot(int32 MaterialIndex)
 	SkeletalMeshPtr->Modify();
 	{
 		FScopedSkeletalMeshPostEditChange ScopedPostEditChange(SkeletalMeshPtr.Get());
+		//When we delete a material slot we must invalidate the DDC because material index is not part of the DDC key by design
 		SkeletalMeshPtr->Materials.RemoveAt(MaterialIndex);
 		FSkeletalMeshModel* Model = SkeletalMeshPtr->GetImportedModel();
 
 		int32 NumLODInfos = SkeletalMeshPtr->GetLODNum();
 
-		//When we delete a material slot we need a full LODMaterialMap for all section of all LODs
+		//When we delete a material slot we need to fix all MaterialIndex after the deleted index
 		for (int32 LODInfoIdx = 0; LODInfoIdx < NumLODInfos; LODInfoIdx++)
 		{
+			TArray<int32>& LODMaterialMap = SkeletalMeshPtr->GetLODInfo(LODInfoIdx)->LODMaterialMap;
 			for (int32 SectionIndex = 0; SectionIndex < Model->LODModels[LODInfoIdx].Sections.Num(); ++SectionIndex)
 			{
-				int32 SectionOriginalMaterialIndex = Model->LODModels[LODInfoIdx].Sections[SectionIndex].MaterialIndex;
-				if (!SkeletalMeshPtr->GetLODInfo(LODInfoIdx)->LODMaterialMap.IsValidIndex(SectionOriginalMaterialIndex))
+				int32 SectionMaterialIndex = Model->LODModels[LODInfoIdx].Sections[SectionIndex].MaterialIndex;
+				if (LODMaterialMap.IsValidIndex(SectionIndex) && LODMaterialMap[SectionIndex] != INDEX_NONE)
 				{
-					SkeletalMeshPtr->GetLODInfo(LODInfoIdx)->LODMaterialMap.Add(SectionOriginalMaterialIndex);
+					SectionMaterialIndex = LODMaterialMap[SectionIndex];
 				}
-
-				int32& CurrentMaterialIndexRemap = SkeletalMeshPtr->GetLODInfo(LODInfoIdx)->LODMaterialMap[SectionOriginalMaterialIndex];
-
-				if (CurrentMaterialIndexRemap == MaterialIndex)
+				if (SectionMaterialIndex > MaterialIndex)
 				{
-					// Set to first material
-					CurrentMaterialIndexRemap = 0;
+					SectionMaterialIndex--;
 				}
-				else if (CurrentMaterialIndexRemap > MaterialIndex)
+				if (SectionMaterialIndex != Model->LODModels[LODInfoIdx].Sections[SectionIndex].MaterialIndex)
 				{
-					// Decrement to keep correct reference after removal
-					CurrentMaterialIndexRemap--;
+					while(!LODMaterialMap.IsValidIndex(SectionIndex))
+					{
+						LODMaterialMap.Add(INDEX_NONE);
+					}
+					LODMaterialMap[SectionIndex] = SectionMaterialIndex;
 				}
 			}
 		}
@@ -4916,54 +4876,12 @@ int32 FPersonaMeshDetails::GetMaterialIndex(int32 LODIndex, int32 SectionIndex)
 	check(ImportedResource && ImportedResource->LODModels.IsValidIndex(LODIndex));
 	int32 MaterialIndex = ImportedResource->LODModels[LODIndex].Sections[SectionIndex].MaterialIndex;
 	FSkeletalMeshLODInfo& Info = *(SkelMesh->GetLODInfo(LODIndex));
-	if (!Info.LODMaterialMap.IsValidIndex(MaterialIndex) || !SkelMesh->Materials.IsValidIndex(Info.LODMaterialMap[MaterialIndex]))
+	if (Info.LODMaterialMap.IsValidIndex(SectionIndex) && SkelMesh->Materials.IsValidIndex(Info.LODMaterialMap[SectionIndex]))
 	{
-		return MaterialIndex;
-	}
-	else
-	{
-		return Info.LODMaterialMap[MaterialIndex];
-	}
-}
-
-bool FPersonaMeshDetails::IsDuplicatedMaterialIndex(int32 LODIndex, int32 MaterialIndex)
-{
-	USkeletalMesh* SkelMesh = GetPersonaToolkit()->GetMesh();
-
-	// finding whether this material index is being used in parent LODs
-	for (int32 LODInfoIdx = 0; LODInfoIdx < LODIndex; LODInfoIdx++)
-	{
-		FSkeletalMeshLODInfo& Info = *(SkelMesh->GetLODInfo(LODInfoIdx));
-		if (LODIndex == 0 || Info.LODMaterialMap.Num() == 0)
-		{
-			FSkeletalMeshModel* ImportedResource = SkelMesh->GetImportedModel();
-
-			if (ImportedResource && ImportedResource->LODModels.IsValidIndex(LODInfoIdx))
-			{
-				FSkeletalMeshLODModel& Model = ImportedResource->LODModels[LODInfoIdx];
-
-				for (int32 SectionIdx = 0; SectionIdx < Model.Sections.Num(); SectionIdx++)
-				{
-					if (MaterialIndex == Model.Sections[SectionIdx].MaterialIndex)
-					{
-						return true;
-					}
-				}
-			}
-		}
-		else // if LODMaterialMap exists
-		{
-			for (int32 MapIdx = 0; MapIdx < Info.LODMaterialMap.Num(); MapIdx++)
-			{
-				if (MaterialIndex == Info.LODMaterialMap[MapIdx])
-				{
-					return true;
-				}
-			}
-		}
+		return Info.LODMaterialMap[SectionIndex];
+		
 	}
-
-	return false;
+	return MaterialIndex;
 }
 
 void FPersonaMeshDetails::OnSectionChanged(int32 LODIndex, int32 SectionIndex, int32 NewMaterialSlotIndex, FName NewMaterialSlotName)
@@ -5004,24 +4922,23 @@ void FPersonaMeshDetails::OnSectionChanged(int32 LODIndex, int32 SectionIndex, i
 			int32 NumSections = ImportedResource->LODModels[LODIndex].Sections.Num();
 			FSkeletalMeshLODInfo& Info = *(Mesh->GetLODInfo(LODIndex));
 
-			auto SetLODMaterialMapValue = [&LODIndex, &Info, &ImportedResource](int32 OriginalSectionMaterialIndex, int32 OverrideMaterialIndex)
+			auto SetLODMaterialMapValue = [&LODIndex, &Info, &ImportedResource](int32 InSectionIndex, int32 OverrideMaterialIndex)
 			{
-				if (ImportedResource->LODModels[LODIndex].Sections[OriginalSectionMaterialIndex].MaterialIndex == OverrideMaterialIndex)
+				if (ImportedResource->LODModels[LODIndex].Sections[InSectionIndex].MaterialIndex == OverrideMaterialIndex)
 				{
-					if (Info.LODMaterialMap.IsValidIndex(OriginalSectionMaterialIndex))
+					if (Info.LODMaterialMap.IsValidIndex(InSectionIndex))
 					{
-						Info.LODMaterialMap[OriginalSectionMaterialIndex] = INDEX_NONE;
+						Info.LODMaterialMap[InSectionIndex] = INDEX_NONE;
 					}
 				}
 				else
 				{
-					while (Info.LODMaterialMap.Num() <= OriginalSectionMaterialIndex)
+					while (Info.LODMaterialMap.Num() <= InSectionIndex)
 					{
-						int32 NewSectionIndex = Info.LODMaterialMap.Num();
 						Info.LODMaterialMap.Add(INDEX_NONE);
 					}
-					check(OriginalSectionMaterialIndex < Info.LODMaterialMap.Num());
-					Info.LODMaterialMap[OriginalSectionMaterialIndex] = OverrideMaterialIndex;
+					check(InSectionIndex < Info.LODMaterialMap.Num());
+					Info.LODMaterialMap[InSectionIndex] = OverrideMaterialIndex;
 				}
 			};
 
@@ -5648,73 +5565,6 @@ bool FPersonaMeshDetails::CanDeleteMaterialElement(int32 LODIndex, int32 Section
 	return SectionIndex != 0;
 }
 
-FReply FPersonaMeshDetails::OnDeleteButtonClicked(int32 LODIndex, int32 SectionIndex)
-{
-	ensure(SectionIndex != 0);
-
-	int32 MaterialIndex = GetMaterialIndex(LODIndex, SectionIndex);
-
-	USkeletalMesh* SkelMesh = GetPersonaToolkit()->GetMesh();
-
-	// Move any mappings pointing to the requested material to point to the first
-	// and decrement any above it
-	if(SkelMesh)
-	{
-
-		const FScopedTransaction Transaction(LOCTEXT("PersonaOnDeleteButtonClickedTransaction", "Persona editor: Delete material slot"));
-		UProperty* MaterialProperty = FindField<UProperty>( USkeletalMesh::StaticClass(), "Materials" );
-		SkelMesh->PreEditChange( MaterialProperty );
-
-		// Patch up LOD mapping indices
-		int32 NumLODInfos = SkelMesh->GetLODNum();
-		for(int32 LODInfoIdx=0; LODInfoIdx < NumLODInfos; LODInfoIdx++)
-		{
-			for(auto LodMaterialIter = SkelMesh->GetLODInfo(LODInfoIdx)->LODMaterialMap.CreateIterator() ; LodMaterialIter ; ++LodMaterialIter)
-			{
-				int32 CurrentMapping = *LodMaterialIter;
-
-				if (CurrentMapping == MaterialIndex)
-				{
-					// Set to first material
-					*LodMaterialIter = 0;
-				}
-				else if (CurrentMapping > MaterialIndex)
-				{
-					// Decrement to keep correct reference after removal
-					*LodMaterialIter = CurrentMapping - 1;
-				}
-			}
-		}
-		
-		// Patch up section indices
-		for(auto ModelIter = SkelMesh->GetImportedModel()->LODModels.CreateIterator() ; ModelIter ; ++ModelIter)
-		{
-			FSkeletalMeshLODModel& Model = *ModelIter;
-			for(auto SectionIter = Model.Sections.CreateIterator() ; SectionIter ; ++SectionIter)
-			{
-				FSkelMeshSection& Section = *SectionIter;
-
-				if (Section.MaterialIndex == MaterialIndex)
-				{
-					Section.MaterialIndex = 0;
-				}
-				else if (Section.MaterialIndex > MaterialIndex)
-				{
-					Section.MaterialIndex--;
-				}
-			}
-		}
-
-		SkelMesh->Materials.RemoveAt(MaterialIndex);
-
-		// Notify the change in material
-		FPropertyChangedEvent PropertyChangedEvent( MaterialProperty );
-		SkelMesh->PostEditChangeProperty(PropertyChangedEvent);
-	}
-
-	return FReply::Handled();
-}
-
 void FPersonaMeshDetails::OnPreviewMeshChanged(USkeletalMesh* OldSkeletalMesh, USkeletalMesh* NewMesh)
 {
 	if (IsApplyNeeded())
diff --git a/Engine/Source/Editor/Persona/Private/PersonaMeshDetails.h b/Engine/Source/Editor/Persona/Private/PersonaMeshDetails.h
--- a/Engine/Source/Editor/Persona/Private/PersonaMeshDetails.h
+++ b/Engine/Source/Editor/Persona/Private/PersonaMeshDetails.h
@@ -489,18 +489,9 @@ private:
 	 */
 	bool CanDeleteMaterialElement(int32 LODIndex, int32 SectionIndex) const;
 
-	/**
-	 * Handler for deleting material elements
-	 * 
-	 * @Param SectionIndex - material section to remove
-	 */
-	FReply OnDeleteButtonClicked(int32 LODIndex, int32 SectionIndex);
-
 	/** Creates the UI for Current LOD panel */
 	void AddLODLevelCategories(IDetailLayoutBuilder& DetailLayout);
 
-	bool IsDuplicatedMaterialIndex(int32 LODIndex, int32 MaterialIndex);
-
 	/** Get a material index from LOD index and section index */
 	int32 GetMaterialIndex(int32 LODIndex, int32 SectionIndex) const;
 
diff --git a/Engine/Source/Editor/UnrealEd/Private/Animation/DebugSkelMeshComponent.cpp b/Engine/Source/Editor/UnrealEd/Private/Animation/DebugSkelMeshComponent.cpp
--- a/Engine/Source/Editor/UnrealEd/Private/Animation/DebugSkelMeshComponent.cpp
+++ b/Engine/Source/Editor/UnrealEd/Private/Animation/DebugSkelMeshComponent.cpp
@@ -661,11 +661,11 @@ void UDebugSkelMeshComponent::ToggleClothSectionsVisibility(bool bShowOnlyClothS
 
 				if(Section.HasClothingData())
 				{
-					ShowMaterialSection(Section.MaterialIndex, bShowOnlyClothSections, LODIndex);
+					ShowMaterialSection(Section.MaterialIndex, SecIdx, bShowOnlyClothSections, LODIndex);
 				}
 				else
 				{
-					ShowMaterialSection(Section.MaterialIndex, !bShowOnlyClothSections, LODIndex);
+					ShowMaterialSection(Section.MaterialIndex, SecIdx, !bShowOnlyClothSections, LODIndex);
 				}
 			}
 		}
@@ -708,7 +708,7 @@ void UDebugSkelMeshComponent::SetMeshSectionVisibilityForCloth(FGuid InClothGuid
 				// disables cloth section and also corresponding original section for matching cloth asset
 				if(Section.HasClothingData() && Section.ClothingData.AssetGuid == InClothGuid)
 				{
-					ShowMaterialSection(Section.MaterialIndex, bVisibility, LODIndex);
+					ShowMaterialSection(Section.MaterialIndex, SecIdx, bVisibility, LODIndex);
 				}
 			}
 		}
diff --git a/Engine/Source/Editor/UnrealEd/Private/Fbx/FbxSkeletalMeshImport.cpp b/Engine/Source/Editor/UnrealEd/Private/Fbx/FbxSkeletalMeshImport.cpp
--- a/Engine/Source/Editor/UnrealEd/Private/Fbx/FbxSkeletalMeshImport.cpp
+++ b/Engine/Source/Editor/UnrealEd/Private/Fbx/FbxSkeletalMeshImport.cpp
@@ -1972,19 +1972,15 @@ USkeletalMesh* UnFbx::FFbxImporter::ImportSkeletalMesh(FImportSkeletalMeshArgs &
 		}
 		if (ImportSkeletalMeshArgs.ImportMeshSectionsData)
 		{
-			if (NewLODInfo.LODMaterialMap.Num() >= NumSections)
+			for (int32 SectionIndex = 0; SectionIndex < ImportedResource->LODModels[0].Sections.Num(); ++SectionIndex)
 			{
-				for (int32 SectionMaterialIndex : SkeletalMesh->GetLODInfo(0)->LODMaterialMap)
+				const FSkelMeshSection &SkelMeshSection = ImportedResource->LODModels[0].Sections[SectionIndex];
+				int32 MaterialIndex = SkelMeshSection.MaterialIndex;
+				if (SkeletalMesh->GetLODInfo(0)->LODMaterialMap.IsValidIndex(SectionIndex) && SkeletalMesh->GetLODInfo(0)->LODMaterialMap[SectionIndex] != INDEX_NONE)
 				{
-					ImportSkeletalMeshArgs.ImportMeshSectionsData->SectionOriginalMaterialName.Add(SkeletalMesh->Materials[SectionMaterialIndex].ImportedMaterialSlotName);
-				}
-			}
-			else
-			{
-				for (const FSkelMeshSection &SkelMeshSection : ImportedResource->LODModels[0].Sections)
-				{
-					ImportSkeletalMeshArgs.ImportMeshSectionsData->SectionOriginalMaterialName.Add(SkeletalMesh->Materials[SkelMeshSection.MaterialIndex].ImportedMaterialSlotName);
+					MaterialIndex = SkeletalMesh->GetLODInfo(0)->LODMaterialMap[SectionIndex];
 				}
+				ImportSkeletalMeshArgs.ImportMeshSectionsData->SectionOriginalMaterialName.Add(SkeletalMesh->Materials[MaterialIndex].ImportedMaterialSlotName);
 			}
 		}
 
@@ -2273,7 +2269,7 @@ void UnFbx::FFbxImporter::UpdateSkeletalMeshImportData(USkeletalMesh *SkeletalMe
 				{
 					int32 MaterialLodSectionIndex = LODModel.Sections[SectionIndex].MaterialIndex;
 					//Is this LOD use the LODMaterialMap override
-					if (LODInfo.LODMaterialMap.Num() > SectionIndex)
+					if (LODInfo.LODMaterialMap.IsValidIndex(SectionIndex) && LODInfo.LODMaterialMap[SectionIndex] != INDEX_NONE)
 					{
 						MaterialLodSectionIndex = LODInfo.LODMaterialMap[SectionIndex];
 					}
@@ -3819,6 +3815,12 @@ void UnFbx::FFbxImporter::InsertNewLODToBaseSkeletalMesh(USkeletalMesh* InSkelet
 				}
 				//Get the current skelmesh section slot import name
 				int32 ExistRemapMaterialIndex = ExistingSkelMeshLodInfo.LODMaterialMap[SectionIndex];
+				if (!BaseSkeletalMesh->Materials.IsValidIndex(ExistRemapMaterialIndex))
+				{
+					//LODMaterialMap can have INDEX_NONE value or an invalid index in the material array
+					continue;
+				}
+
 				ExistingMeshSectionSlotNames.Add(BaseSkeletalMesh->Materials[ExistRemapMaterialIndex].ImportedMaterialSlotName);
 
 				//Get the Last imported skelmesh section slot import name
@@ -3929,6 +3931,16 @@ void UnFbx::FFbxImporter::InsertNewLODToBaseSkeletalMesh(USkeletalMesh* InSkelet
 			FName ExistMeshSectionSlotName = ExistingMeshSectionSlotNames[SectionIndex];
 			//Get the new skelmesh section slot import name
 			int32 NewRemapMaterialIndex = NewSkelMeshLodInfo.LODMaterialMap[SectionIndex];
+			if (!BaseSkeletalMesh->Materials.IsValidIndex(NewRemapMaterialIndex))
+			{
+				int32 RematchMaterialIndex = NewSkelMeshLodInfo.LODMaterialMap[SectionIndex];
+				if (RematchMaterialIndex == INDEX_NONE)
+				{
+					RematchMaterialIndex = NewLODModel.Sections[SectionIndex].MaterialIndex;
+				}
+				NewSkelMeshLodInfo.LODMaterialMap[SectionIndex] = FMath::Clamp<int32>(RematchMaterialIndex, 0, BaseSkeletalMesh->Materials.Num()-1);
+				continue;
+			}
 			FName NewMeshSectionSlotName = BaseSkeletalMesh->Materials[NewRemapMaterialIndex].ImportedMaterialSlotName;
 			//Get the Last imported skelmesh section slot import name
 			FName OriginalImportMeshSectionSlotName = OriginalImportMeshSectionSlotNames[SectionIndex];
@@ -4140,11 +4152,11 @@ bool UnFbx::FFbxImporter::ImportSkeletalMeshLOD(USkeletalMesh* InSkeletalMesh, U
 		TArray<FSkelMeshSection>& ImportedSections = ImportedResource->LODModels[0].Sections;
 		const FSkeletalMeshLODInfo& ImportedInfo = *(InSkeletalMesh->GetLODInfo(0));
 		
-		auto GetImportMaterialSlotName = [](const USkeletalMesh* SkelMesh, const FSkelMeshSection& Section, const FSkeletalMeshLODInfo& Info, int32& OutMaterialIndex)->FName
+		auto GetImportMaterialSlotName = [](const USkeletalMesh* SkelMesh, const FSkelMeshSection& Section, int32 SectionIndex, const FSkeletalMeshLODInfo& Info, int32& OutMaterialIndex)->FName
 		{
 			check(SkelMesh->Materials.Num() > 0);
 			OutMaterialIndex = Section.MaterialIndex;
-			if (Info.LODMaterialMap.IsValidIndex(Section.MaterialIndex))
+			if (Info.LODMaterialMap.IsValidIndex(SectionIndex) && Info.LODMaterialMap[SectionIndex] != INDEX_NONE)
 			{
 				OutMaterialIndex = Info.LODMaterialMap[Section.MaterialIndex];
 			}
@@ -4161,15 +4173,17 @@ bool UnFbx::FFbxImporter::ImportSkeletalMeshLOD(USkeletalMesh* InSkeletalMesh, U
 			return ImportedMaterialSlotName;
 		};
 
-		for (const FSkelMeshSection& ExistingSection : ExistingSections)
+		for(int32 ExistingSectionIndex = 0; ExistingSectionIndex < ExistingSections.Num(); ++ExistingSectionIndex)
 		{
+			const FSkelMeshSection& ExistingSection = ExistingSections[ExistingSectionIndex];
 			int32 ExistingMaterialIndex = 0;
-			FName ExistingImportedMaterialSlotName = GetImportMaterialSlotName(BaseSkeletalMesh, ExistingSection, ExistingInfo, ExistingMaterialIndex);
+			FName ExistingImportedMaterialSlotName = GetImportMaterialSlotName(BaseSkeletalMesh, ExistingSection, ExistingSectionIndex, ExistingInfo, ExistingMaterialIndex);
 
-			for (FSkelMeshSection& ImportedSection : ImportedSections)
+			for(int32 ImportedSectionIndex = 0; ImportedSectionIndex < ImportedSections.Num(); ++ImportedSectionIndex)
 			{
+				FSkelMeshSection& ImportedSection = ImportedSections[ImportedSectionIndex];
 				int32 ImportedMaterialIndex = 0;
-				FName ImportedImportedMaterialSlotName = GetImportMaterialSlotName(InSkeletalMesh, ImportedSection, ImportedInfo, ImportedMaterialIndex);
+				FName ImportedImportedMaterialSlotName = GetImportMaterialSlotName(InSkeletalMesh, ImportedSection, ImportedSectionIndex, ImportedInfo, ImportedMaterialIndex);
 				if (ExistingImportedMaterialSlotName != NAME_None)
 				{
 					if (ImportedImportedMaterialSlotName == ExistingImportedMaterialSlotName)
diff --git a/Engine/Source/Runtime/Engine/Classes/Components/SkinnedMeshComponent.h b/Engine/Source/Runtime/Engine/Classes/Components/SkinnedMeshComponent.h
--- a/Engine/Source/Runtime/Engine/Classes/Components/SkinnedMeshComponent.h
+++ b/Engine/Source/Runtime/Engine/Classes/Components/SkinnedMeshComponent.h
@@ -1428,7 +1428,7 @@ public:
 	 * @param LODIndex - Index of the LOD to modify material visibility within
 	 */
 	UFUNCTION(BlueprintCallable, Category = "Components|SkinnedMesh")
-	void ShowMaterialSection(int32 MaterialID, bool bShow, int32 LODIndex);
+	void ShowMaterialSection(int32 MaterialID, int32 SectionIndex, bool bShow, int32 LODIndex);
 
 	/** Clear any material visibility modifications made by ShowMaterialSection */
 	UFUNCTION(BlueprintCallable, Category = "Components|SkinnedMesh")
diff --git a/Engine/Source/Runtime/Engine/Private/Components/SkinnedMeshComponent.cpp b/Engine/Source/Runtime/Engine/Private/Components/SkinnedMeshComponent.cpp
--- a/Engine/Source/Runtime/Engine/Private/Components/SkinnedMeshComponent.cpp
+++ b/Engine/Source/Runtime/Engine/Private/Components/SkinnedMeshComponent.cpp
@@ -2343,7 +2343,7 @@ FName USkinnedMeshComponent::FindClosestBone_K2(FVector TestLocation, FVector& B
 	return FindClosestBone(TestLocation, &BoneLocation, IgnoreScale, bRequirePhysicsAsset);
 }
 
-void USkinnedMeshComponent::ShowMaterialSection(int32 MaterialID, bool bShow, int32 LODIndex)
+void USkinnedMeshComponent::ShowMaterialSection(int32 MaterialID, int32 SectionIndex, bool bShow, int32 LODIndex)
 {
 	if (!SkeletalMesh)
 	{
@@ -2365,15 +2365,12 @@ void USkinnedMeshComponent::ShowMaterialSection(int32 MaterialID, bool bShow, in
 			HiddenMaterials.Empty(SkeletalMesh->Materials.Num());
 			HiddenMaterials.AddZeroed(SkeletalMesh->Materials.Num());
 		}
-		// If we are at a dropped LOD, route material index through the LODMaterialMap in the LODInfo struct.
+		// If we have a valid LODInfo LODMaterialMap, route material index through it.
 		int32 UseMaterialIndex = MaterialID;			
-		if(LODIndex > 0)
+		if(SkelLODInfo.LODMaterialMap.IsValidIndex(SectionIndex) && SkelLODInfo.LODMaterialMap[SectionIndex] != INDEX_NONE)
 		{
-			if(SkelLODInfo.LODMaterialMap.IsValidIndex(MaterialID))
-			{
-				UseMaterialIndex = SkelLODInfo.LODMaterialMap[MaterialID];
-				UseMaterialIndex = FMath::Clamp( UseMaterialIndex, 0, HiddenMaterials.Num() );
-			}
+			UseMaterialIndex = SkelLODInfo.LODMaterialMap[SectionIndex];
+			UseMaterialIndex = FMath::Clamp( UseMaterialIndex, 0, HiddenMaterials.Num() );
 		}
 		// Mark the mapped section material entry as visible/hidden
 		if (HiddenMaterials.IsValidIndex(UseMaterialIndex))
diff --git a/Engine/Source/Runtime/Engine/Private/SkeletalMeshLODModel.cpp b/Engine/Source/Runtime/Engine/Private/SkeletalMeshLODModel.cpp
--- a/Engine/Source/Runtime/Engine/Private/SkeletalMeshLODModel.cpp
+++ b/Engine/Source/Runtime/Engine/Private/SkeletalMeshLODModel.cpp
@@ -1044,10 +1044,9 @@ void FSkeletalMeshLODModel::UpdateChunkedSectionInfo(const FString& SkeletalMesh
 	for (int32 LODModelSectionIndex = 0; LODModelSectionIndex < LODModelSectionNum; ++LODModelSectionIndex)
 	{
 		FSkelMeshSection& Section = Sections[LODModelSectionIndex];
-		int32 SectionMaterialIndex = Section.MaterialIndex;
 		
 		//If we have cloth on a chunked section we treat the chunked section has a parent section (this is to get the same result has before the refactor)
-		if (LastBoneCount >= MaxGPUSkinBones && SectionMaterialIndex == LastMaterialIndex && !Section.ClothingData.AssetGuid.IsValid())
+		if (LastBoneCount >= MaxGPUSkinBones && Section.MaterialIndex == LastMaterialIndex && !Section.ClothingData.AssetGuid.IsValid())
 		{
 			Section.ChunkedParentSectionIndex = CurrentParentChunkIndex;
 			Section.OriginalDataSectionIndex = Sections[CurrentParentChunkIndex].OriginalDataSectionIndex;
@@ -1078,17 +1077,7 @@ void FSkeletalMeshLODModel::UpdateChunkedSectionInfo(const FString& SkeletalMesh
 			Section.ChunkedParentSectionIndex = INDEX_NONE;
 		}
 
-		//Use the LODMaterialMap if the material is different from the original section index and set back the section materialindex to the originalDataSectionIndex
-		if (SectionMaterialIndex != Section.OriginalDataSectionIndex)
-		{
-			while (LODMaterialMap.Num() <= LODModelSectionIndex)
-			{
-				LODMaterialMap.Add(INDEX_NONE);
-			}
-			LODMaterialMap[LODModelSectionIndex] = Section.ChunkedParentSectionIndex != INDEX_NONE ? Sections[Section.ChunkedParentSectionIndex].MaterialIndex : Section.MaterialIndex;
-			Section.MaterialIndex = Section.OriginalDataSectionIndex;
-		}
-		LastMaterialIndex = SectionMaterialIndex;
+		LastMaterialIndex = Section.MaterialIndex;
 		//Set the last bone count
 		LastBoneCount = (uint32)Sections[LODModelSectionIndex].BoneMap.Num();
 		//its impossible to have more bone then the maximum allowed
