Fix the ClothingAsset property update when changing properties in the Clothing window.

Merging from Release-11.40 after Robomerge conflict.

diff --git a/Engine/Source/Editor/Persona/Private/PersonaMeshDetails.cpp b/Engine/Source/Editor/Persona/Private/PersonaMeshDetails.cpp
--- a/Engine/Source/Editor/Persona/Private/PersonaMeshDetails.cpp
+++ b/Engine/Source/Editor/Persona/Private/PersonaMeshDetails.cpp
@@ -5113,7 +5113,6 @@ void FPersonaMeshDetails::OnGenerateElementForClothingAsset( TSharedRef<IPropert
 			.ShowPublicViewControl(false)
 			.HideNameArea(true)
 			.IsPropertyEditingEnabledDelegate(FIsPropertyEditingEnabled::CreateSP(this, &FPersonaMeshDetails::IsClothingPanelEnabled))
-			.OnFinishedChangingProperties(FOnFinishedChangingProperties::FDelegate::CreateSP(this, &FPersonaMeshDetails::OnFinishedChangingClothingProperties, ElementIndex))
 		]
 	];
 
@@ -5541,24 +5540,6 @@ bool FPersonaMeshDetails::IsClothingPanelEnabled() const
 	return !GEditor->bIsSimulatingInEditor && !GEditor->PlayWorld;
 }
 
-void FPersonaMeshDetails::OnFinishedChangingClothingProperties(const FPropertyChangedEvent& Event, int32 InAssetIndex)
-{
-	USkeletalMesh* CurrentMesh = GetPersonaToolkit()->GetMesh();
-	if(CurrentMesh->MeshClothingAssets.IsValidIndex(InAssetIndex))
-	{
-		UClothingAssetBase* Asset = CurrentMesh->MeshClothingAssets[InAssetIndex];
-		if (Asset)
-		{
-			Asset->PostPropertyChangeCb(Event);
-		}
-	}
-	if(UDebugSkelMeshComponent* PreviewComponent = GetPersonaToolkit()->GetPreviewMeshComponent())
-	{
-		// Reregister our preview component to apply the change
-		FComponentReregisterContext Context(PreviewComponent);
-	}
-}
-
 bool FPersonaMeshDetails::CanDeleteMaterialElement(int32 LODIndex, int32 SectionIndex) const
 {
 	// Only allow deletion of extra elements
diff --git a/Engine/Source/Editor/Persona/Private/PersonaMeshDetails.h b/Engine/Source/Editor/Persona/Private/PersonaMeshDetails.h
--- a/Engine/Source/Editor/Persona/Private/PersonaMeshDetails.h
+++ b/Engine/Source/Editor/Persona/Private/PersonaMeshDetails.h
@@ -682,9 +682,6 @@ private:
 	// If the clothing details widget is editable
 	bool IsClothingPanelEnabled() const;
 
-	// Callback after the clothing details are changed
-	void OnFinishedChangingClothingProperties(const FPropertyChangedEvent& Event, int32 InAssetIndex);
-
 	/* Generate slate UI for Clothing category */
 	void CustomizeClothingProperties(class IDetailLayoutBuilder& DetailLayout, class IDetailCategoryBuilder& ClothingFilesCategory);
 
diff --git a/Engine/Source/Runtime/ClothingSystemRuntimeInterface/Public/ClothingAssetBase.h b/Engine/Source/Runtime/ClothingSystemRuntimeInterface/Public/ClothingAssetBase.h
--- a/Engine/Source/Runtime/ClothingSystemRuntimeInterface/Public/ClothingAssetBase.h
+++ b/Engine/Source/Runtime/ClothingSystemRuntimeInterface/Public/ClothingAssetBase.h
@@ -69,10 +69,6 @@ public:
 	virtual int32 AddNewLod()
 	PURE_VIRTUAL(UClothingAssetBase::AddNewLod(), return INDEX_NONE;);
 
-	/** Called from \c FPersonaMeshDetails::OnFinishedChangingClothingProperties(). */
-	virtual void PostPropertyChangeCb(const FPropertyChangedEvent& Event)
-	PURE_VIRTUAL(UClothingAssetBase::PostPropertyChangeCb, );
-
 	/** Builds self collision data */
 	virtual void BuildSelfCollisionData()
 	PURE_VIRTUAL(UClothingAssetBase::BuildSelfCollisionData(), );
diff --git a/Engine/Source/Runtime/ClothingSystemRuntimeNv/Private/Assets/ClothingAssetNv.cpp b/Engine/Source/Runtime/ClothingSystemRuntimeNv/Private/Assets/ClothingAssetNv.cpp
--- a/Engine/Source/Runtime/ClothingSystemRuntimeNv/Private/Assets/ClothingAssetNv.cpp
+++ b/Engine/Source/Runtime/ClothingSystemRuntimeNv/Private/Assets/ClothingAssetNv.cpp
@@ -300,19 +300,21 @@ int32 UClothingAssetNv::AddNewLod()
 	return Idx;
 }
 
-void UClothingAssetNv::PostPropertyChangeCb(const FPropertyChangedEvent& Event)
+void UClothingAssetNv::PostEditChangeChainProperty(FPropertyChangedChainEvent& ChainEvent)
 {
+	Super::PostEditChangeChainProperty(ChainEvent);
+
 	bool bReregisterComponents = false;
 
-	if(Event.ChangeType != EPropertyChangeType::Interactive)
+	if(ChainEvent.ChangeType != EPropertyChangeType::Interactive)
 	{
-		if(Event.Property->GetFName() == GET_MEMBER_NAME_CHECKED(UClothConfigNv, SelfCollisionRadius) ||
-		   Event.Property->GetFName() == GET_MEMBER_NAME_CHECKED(UClothConfigNv, SelfCollisionCullScale))
+		if(ChainEvent.Property->GetFName() == GET_MEMBER_NAME_CHECKED(UClothConfigNv, SelfCollisionRadius) ||
+		   ChainEvent.Property->GetFName() == GET_MEMBER_NAME_CHECKED(UClothConfigNv, SelfCollisionCullScale))
 		{
 			BuildSelfCollisionData();
 			bReregisterComponents = true;
 		}
-		else if(Event.Property->GetFName() == GET_MEMBER_NAME_CHECKED(UClothingAssetCommon, PhysicsAsset))
+		else if(ChainEvent.Property->GetFName() == GET_MEMBER_NAME_CHECKED(UClothingAssetCommon, PhysicsAsset))
 		{
 			bReregisterComponents = true;
 		}
diff --git a/Engine/Source/Runtime/ClothingSystemRuntimeNv/Public/Assets/ClothingAssetNv.h b/Engine/Source/Runtime/ClothingSystemRuntimeNv/Public/Assets/ClothingAssetNv.h
--- a/Engine/Source/Runtime/ClothingSystemRuntimeNv/Public/Assets/ClothingAssetNv.h
+++ b/Engine/Source/Runtime/ClothingSystemRuntimeNv/Public/Assets/ClothingAssetNv.h
@@ -25,7 +25,7 @@ public:
 #if WITH_EDITOR
 	virtual void InvalidateCachedData() override;
 	virtual int32 AddNewLod() override;
-	virtual void PostPropertyChangeCb(const FPropertyChangedEvent& Event) override;
+	virtual void PostEditChangeChainProperty(FPropertyChangedChainEvent& ChainEvent) override;
 #endif // WITH_EDITOR
 
 	/** 
