SkeletalMesh LOD import with new file. Set the base LOD if the reduce criterion is max vertices/triangles. If its other criterion (percent vertices/triangles) put the value to 100%.

diff --git a/Engine/Source/Editor/Persona/Private/PersonaMeshDetails.cpp b/Engine/Source/Editor/Persona/Private/PersonaMeshDetails.cpp
--- a/Engine/Source/Editor/Persona/Private/PersonaMeshDetails.cpp
+++ b/Engine/Source/Editor/Persona/Private/PersonaMeshDetails.cpp
@@ -2085,11 +2085,47 @@ FReply FPersonaMeshDetails::OnReimportLodClicked(EReimportButtonType InReimportT
 		}
 
 		FString SourceFilenameBackup("");
+		
+		//If we alter the reduction setting and the user cancel the import we must set them back
+		bool bRestoreReductionOnfail = false;
+		FSkeletalMeshOptimizationSettings ReductionSettingsBackup;
+		FSkeletalMeshLODInfo* LODInfo = SkelMesh->GetLODInfo(InLODIndex);
 		if(InReimportType == EReimportButtonType::ReimportWithNewFile)
 		{
 			// Back up current source filename and empty it so the importer asks for a new one.
-			SourceFilenameBackup = SkelMesh->GetLODInfo(InLODIndex)->SourceImportFilename;
-			SkelMesh->GetLODInfo(InLODIndex)->SourceImportFilename.Empty();
+			SourceFilenameBackup = LODInfo->SourceImportFilename;
+			LODInfo->SourceImportFilename.Empty();
+			
+			//Avoid changing the settings if the skeletal mesh is using a LODSettings asset valid for this LOD
+			bool bUseLODSettingAsset = SkelMesh->LODSettings != nullptr && SkelMesh->LODSettings->GetNumberOfSettings() > InLODIndex;
+			//Make the reduction settings change according to the context
+			if (!bUseLODSettingAsset && SkelMesh->IsReductionActive(InLODIndex) && LODInfo->bHasBeenSimplified && SkelMesh->GetImportedModel()->LODModels[InLODIndex].RawSkeletalMeshBulkData.IsEmpty())
+			{
+				FSkeletalMeshOptimizationSettings& ReductionSettings = LODInfo->ReductionSettings;
+				//Backup the reduction settings
+				ReductionSettingsBackup = ReductionSettings;
+				//In case we have a vert/tri percent we just put the percent to 100% and avoid reduction
+				//If we have a maximum criterion we change the BaseLOD to reduce the imported fbx instead of other LOD
+				switch (ReductionSettings.TerminationCriterion)
+				{
+					case SkeletalMeshTerminationCriterion::SMTC_NumOfTriangles:
+						ReductionSettings.NumOfTrianglesPercentage = 1.0f;
+						break;
+					case SkeletalMeshTerminationCriterion::SMTC_NumOfVerts:
+						ReductionSettings.NumOfVertPercentage = 1.0f;
+						break;
+					case SkeletalMeshTerminationCriterion::SMTC_TriangleOrVert:
+						ReductionSettings.NumOfTrianglesPercentage = 1.0f;
+						ReductionSettings.NumOfVertPercentage = 1.0f;
+						break;
+					case SkeletalMeshTerminationCriterion::SMTC_AbsNumOfTriangles:
+					case SkeletalMeshTerminationCriterion::SMTC_AbsNumOfVerts:
+					case SkeletalMeshTerminationCriterion::SMTC_AbsTriangleOrVert:
+						ReductionSettings.BaseLOD = InLODIndex;
+						break;
+				}
+				bRestoreReductionOnfail = true;
+			}
 		}
 
 		bool bImportSucceeded = FbxMeshUtils::ImportMeshLODDialog(SkelMesh, InLODIndex);
@@ -2097,7 +2133,16 @@ FReply FPersonaMeshDetails::OnReimportLodClicked(EReimportButtonType InReimportT
 		if(InReimportType == EReimportButtonType::ReimportWithNewFile && !bImportSucceeded)
 		{
 			// Copy old source file back, as this one failed
-			SkelMesh->GetLODInfo(InLODIndex)->SourceImportFilename = SourceFilenameBackup;
+			LODInfo->SourceImportFilename = SourceFilenameBackup;
+			if (bRestoreReductionOnfail)
+			{
+				LODInfo->ReductionSettings = ReductionSettingsBackup;
+			}
+		}
+		else if(InReimportType == EReimportButtonType::ReimportWithNewFile)
+		{
+			//Refresh the layout so the BaseLOD min max get recompute
+			MeshDetailLayout->ForceRefreshDetails();
 		}
 
 		return FReply::Handled();
