Fix canvas DPI scaling issues for static mesh socket names and skeletal mesh socket/bone names.

diff --git a/Engine/Source/Editor/Persona/Private/EditModes/SkeletonSelectionEditMode.cpp b/Engine/Source/Editor/Persona/Private/EditModes/SkeletonSelectionEditMode.cpp
--- a/Engine/Source/Editor/Persona/Private/EditModes/SkeletonSelectionEditMode.cpp
+++ b/Engine/Source/Editor/Persona/Private/EditModes/SkeletonSelectionEditMode.cpp
@@ -332,6 +332,14 @@ void FSkeletonSelectionEditMode::Render(const FSceneView* View, FViewport* Viewp
 	}
 }
 
+FIntPoint FSkeletonSelectionEditMode::GetDPIUnscaledSize(FViewport* Viewport, FViewportClient* Client)
+{
+	const FIntPoint Size = Viewport->GetSizeXY();
+	const float DPIScale = Client->GetDPIScale();
+	// (FIntPoint / float) implicitly casts the float to an int if you try to divide it directly
+	return FIntPoint(Size.X / DPIScale, Size.Y / DPIScale);
+}
+
 void FSkeletonSelectionEditMode::DrawHUD(FEditorViewportClient* ViewportClient, FViewport* Viewport, const FSceneView* View, FCanvas* Canvas)
 {
 	UDebugSkelMeshComponent* PreviewMeshComponent = GetAnimPreviewScene().GetPreviewMeshComponent();
@@ -339,7 +347,7 @@ void FSkeletonSelectionEditMode::DrawHUD(FEditorViewportClient* ViewportClient,
 	// Draw name of selected bone
 	if (IsSelectedBoneRequired())
 	{
-		const FIntPoint ViewPortSize = Viewport->GetSizeXY();
+		const FIntPoint ViewPortSize = GetDPIUnscaledSize(Viewport, ViewportClient);
 		const int32 HalfX = ViewPortSize.X / 2;
 		const int32 HalfY = ViewPortSize.Y / 2;
 
@@ -372,7 +380,7 @@ void FSkeletonSelectionEditMode::DrawHUD(FEditorViewportClient* ViewportClient,
 		const FPlane Proj = View->Project(SocketPos);
 		if (Proj.W > 0.f)
 		{
-			const FIntPoint ViewPortSize = Viewport->GetSizeXY();
+			const FIntPoint ViewPortSize = GetDPIUnscaledSize(Viewport, ViewportClient);
 			const int32 HalfX = ViewPortSize.X / 2;
 			const int32 HalfY = ViewPortSize.Y / 2;
 
diff --git a/Engine/Source/Editor/Persona/Private/EditModes/SkeletonSelectionEditMode.h b/Engine/Source/Editor/Persona/Private/EditModes/SkeletonSelectionEditMode.h
--- a/Engine/Source/Editor/Persona/Private/EditModes/SkeletonSelectionEditMode.h
+++ b/Engine/Source/Editor/Persona/Private/EditModes/SkeletonSelectionEditMode.h
@@ -47,6 +47,9 @@ private:
 	/** Check whether the currently selected bone is in the required bones list */
 	bool IsSelectedBoneRequired() const;
 
+	/** Unscale a viewport's size by its DPI factor */
+	static FIntPoint GetDPIUnscaledSize(FViewport* Viewport, FViewportClient* Client);
+
 private:
 	/** Whether we are currently in a manipulation  */
 	bool bManipulating;	
diff --git a/Engine/Source/Editor/StaticMeshEditor/Private/StaticMeshEditorViewportClient.cpp b/Engine/Source/Editor/StaticMeshEditor/Private/StaticMeshEditorViewportClient.cpp
--- a/Engine/Source/Editor/StaticMeshEditor/Private/StaticMeshEditorViewportClient.cpp
+++ b/Engine/Source/Editor/StaticMeshEditor/Private/StaticMeshEditorViewportClient.cpp
@@ -728,8 +728,8 @@ void FStaticMeshEditorViewportClient::DrawCanvas( FViewport& InViewport, FSceneV
 		return;
 	}
 
-	const int32 HalfX = Viewport->GetSizeXY().X/2;
-	const int32 HalfY = Viewport->GetSizeXY().Y/2;
+	const int32 HalfX = Viewport->GetSizeXY().X/2 / GetDPIScale();
+	const int32 HalfY = Viewport->GetSizeXY().Y/2 / GetDPIScale();
 
 	// Draw socket names if desired.
 	if( bShowSockets )
