Partial linux fix CIS build
#rnx

diff --git a/Engine/Source/Developer/MeshUtilities/MeshUtilities.Build.cs b/Engine/Source/Developer/MeshUtilities/MeshUtilities.Build.cs
--- a/Engine/Source/Developer/MeshUtilities/MeshUtilities.Build.cs
+++ b/Engine/Source/Developer/MeshUtilities/MeshUtilities.Build.cs
@@ -33,7 +33,6 @@ public class MeshUtilities : ModuleRules
 				"PropertyEditor",
 				"EditorStyle",
                 "GraphColor",
-                "MeshBuilder",
                 "MeshUtilitiesCommon",
                 "MeshDescription",
 				"StaticMeshDescription",
diff --git a/Engine/Source/Editor/Persona/Private/LODInfoUILayout.cpp b/Engine/Source/Editor/Persona/Private/LODInfoUILayout.cpp
--- a/Engine/Source/Editor/Persona/Private/LODInfoUILayout.cpp
+++ b/Engine/Source/Editor/Persona/Private/LODInfoUILayout.cpp
@@ -1,7 +1,5 @@
 // Copyright 1998-2019 Epic Games, Inc. All Rights Reserved.
 
-
-#pragma once
 #include "LODInfoUILayout.h"
 #include "CoreMinimal.h"
 #include "UObject/ObjectMacros.h"
diff --git a/Engine/Source/Editor/Persona/Private/PersonaMeshDetails.cpp b/Engine/Source/Editor/Persona/Private/PersonaMeshDetails.cpp
--- a/Engine/Source/Editor/Persona/Private/PersonaMeshDetails.cpp
+++ b/Engine/Source/Editor/Persona/Private/PersonaMeshDetails.cpp
@@ -892,7 +892,7 @@ TSharedRef<SWidget> FSkeletalMeshReductionSettingsLayout::FillReductionMethodMen
 {
 	if (EnumReductionMethod == nullptr)
 	{
-		SNullWidget::NullWidget;
+		return SNullWidget::NullWidget;
 	}
 	FMenuBuilder MenuBuilder(true, NULL);
 	int32 EnumCount = EnumReductionMethod->NumEnums();
@@ -922,7 +922,7 @@ TSharedRef<SWidget> FSkeletalMeshReductionSettingsLayout::FillReductionImportanc
 {
 	if (EnumImportance == nullptr)
 	{
-		SNullWidget::NullWidget;
+		return SNullWidget::NullWidget;
 	}
 	FMenuBuilder MenuBuilder(true, NULL);
 	int32 EnumCount = EnumImportance->NumEnums();
@@ -977,7 +977,7 @@ TSharedRef<class SWidget> FSkeletalMeshReductionSettingsLayout::FillReductionTer
 {
 	if (EnumTerminationCriterion == nullptr)
 	{
-		SNullWidget::NullWidget;
+		return SNullWidget::NullWidget;
 	}
 	FMenuBuilder MenuBuilder(true, NULL);
 	int32 EnumCount = EnumTerminationCriterion->NumEnums();
diff --git a/Engine/Source/Runtime/Engine/Private/SkeletalMeshRenderData.cpp b/Engine/Source/Runtime/Engine/Private/SkeletalMeshRenderData.cpp
--- a/Engine/Source/Runtime/Engine/Private/SkeletalMeshRenderData.cpp
+++ b/Engine/Source/Runtime/Engine/Private/SkeletalMeshRenderData.cpp
@@ -161,7 +161,6 @@ void FSkeletalMeshRenderData::Cache(USkeletalMesh* Owner)
 				for (int32 SectionIndex = 0; SectionIndex < LODModel->Sections.Num(); ++SectionIndex)
 				{
 					Ar << LODModel->Sections[SectionIndex];
-					LODModel->Sections[SectionIndex].OverlappingVertices;
 				}
 				LODModel->SyncronizeUserSectionsDataArray();
 			}
