Fix CIS staticanalysis

diff --git a/Engine/Source/Editor/Persona/Private/PersonaMeshDetails.cpp b/Engine/Source/Editor/Persona/Private/PersonaMeshDetails.cpp
--- a/Engine/Source/Editor/Persona/Private/PersonaMeshDetails.cpp
+++ b/Engine/Source/Editor/Persona/Private/PersonaMeshDetails.cpp
@@ -2902,8 +2902,9 @@ FReply FPersonaMeshDetails::ApplyLODChanges(int32 LODIndex)
 	FScopedSuspendAlternateSkinWeightPreview ScopedSuspendAlternateSkinnWeightPreview(SkelMesh);
 	{
 		FScopedSkeletalMeshPostEditChange ScopedPostEditChange(SkelMesh);
-		
+		check(SkelMesh);
 		FSkeletalMeshLODInfo* LODInfo = SkelMesh->GetLODInfo(LODIndex);
+		check(LODInfo);
 		int32 SourceLODIndex = LODIndex;
 		bool bHasBeenSimplified = LODInfo ? LODInfo->bHasBeenSimplified : false;
 		if(bHasBeenSimplified)
