Do not trash LODInfo LODMaterial data when reducing a LOD that was imported.
#rnx

diff --git a/Engine/Plugins/Experimental/SkeletalReduction/Source/Private/SkeletalMeshReductionPlugin.cpp b/Engine/Plugins/Experimental/SkeletalReduction/Source/Private/SkeletalMeshReductionPlugin.cpp
--- a/Engine/Plugins/Experimental/SkeletalReduction/Source/Private/SkeletalMeshReductionPlugin.cpp
+++ b/Engine/Plugins/Experimental/SkeletalReduction/Source/Private/SkeletalMeshReductionPlugin.cpp
@@ -1836,7 +1836,12 @@ void FQuadricSkeletalMeshReduction::ReduceSkeletalMesh(USkeletalMesh& SkeletalMe
 		SectionData.GenerateUpToLodIndex = BackupSectionLODModel.Sections[SectionIndex].GenerateUpToLodIndex;
 		SectionData.ChunkedParentSectionIndex = BackupSectionLODModel.Sections[SectionIndex].ChunkedParentSectionIndex;
 		SectionData.OriginalDataSectionIndex = BackupSectionLODModel.Sections[SectionIndex].OriginalDataSectionIndex;
-		SectionData.MaterialMap = bOldLodWasFromFile ? INDEX_NONE : LODInfo->LODMaterialMap.IsValidIndex(SectionIndex) ? LODInfo->LODMaterialMap[SectionIndex] : INDEX_NONE;
+		SectionData.MaterialMap = LODInfo->LODMaterialMap.IsValidIndex(SectionIndex) ? LODInfo->LODMaterialMap[SectionIndex] : INDEX_NONE;
+		if (SectionData.MaterialMap == SectionData.MaterialIndex)
+		{
+			//Remove any override if the value is the same
+			SectionData.MaterialMap = INDEX_NONE;
+		}
 	}
 
 	auto FillSectionMaterialSlot = [&SkeletalMeshResource, &LODIndex, bLODModelAdded](TArray<int32>& SectionMaterialSlot)
diff --git a/Engine/Source/Editor/Persona/Private/PersonaMeshDetails.cpp b/Engine/Source/Editor/Persona/Private/PersonaMeshDetails.cpp
--- a/Engine/Source/Editor/Persona/Private/PersonaMeshDetails.cpp
+++ b/Engine/Source/Editor/Persona/Private/PersonaMeshDetails.cpp
@@ -4923,27 +4923,36 @@ void FPersonaMeshDetails::OnSectionChanged(int32 LODIndex, int32 SectionIndex, i
 		int32 NumSections = ImportedResource->LODModels[LODIndex].Sections.Num();
 		FSkeletalMeshLODInfo& Info = *(Mesh->GetLODInfo(LODIndex));
 		int32 CurrentMaterialIndex = ImportedResource->LODModels[LODIndex].Sections[SectionIndex].MaterialIndex;
-		while (Info.LODMaterialMap.Num() <= CurrentMaterialIndex)
+		
+		auto SetLODMaterialMapValue = [&Info](int32 OriginalSectionMaterialIndex, int32 OverrideMaterialIndex)
 		{
-			int32 NewSectionIndex = Info.LODMaterialMap.Num();
-			Info.LODMaterialMap.Add(INDEX_NONE);
-		}
-		check(CurrentMaterialIndex < Info.LODMaterialMap.Num());
-		Info.LODMaterialMap[CurrentMaterialIndex] = NewSkeletalMaterialIndex;
+			if (OriginalSectionMaterialIndex == OverrideMaterialIndex)
+			{
+				if (Info.LODMaterialMap.IsValidIndex(OriginalSectionMaterialIndex))
+				{
+					Info.LODMaterialMap[OriginalSectionMaterialIndex] = INDEX_NONE;
+				}
+			}
+			else
+			{
+				while (Info.LODMaterialMap.Num() <= OriginalSectionMaterialIndex)
+				{
+					int32 NewSectionIndex = Info.LODMaterialMap.Num();
+					Info.LODMaterialMap.Add(INDEX_NONE);
+				}
+				check(OriginalSectionMaterialIndex < Info.LODMaterialMap.Num());
+				Info.LODMaterialMap[OriginalSectionMaterialIndex] = OverrideMaterialIndex;
+			}
+		};
 
+		SetLODMaterialMapValue(CurrentMaterialIndex, NewSkeletalMaterialIndex);
 		//Set the chunked section 
 		for (int32 SectionIdx = 0; SectionIdx < NumSections; SectionIdx++)
 		{
 			if (ImportedResource->LODModels[LODIndex].Sections[SectionIdx].ChunkedParentSectionIndex == SectionIndex)
 			{
 				int32 CurrentChunkMaterialIndex = ImportedResource->LODModels[LODIndex].Sections[SectionIdx].MaterialIndex;
-				while (Info.LODMaterialMap.Num() <= CurrentChunkMaterialIndex)
-				{
-					int32 NewSectionIndex = Info.LODMaterialMap.Num();
-					Info.LODMaterialMap.Add(INDEX_NONE);
-				}
-				check(CurrentChunkMaterialIndex < Info.LODMaterialMap.Num());
-				Info.LODMaterialMap[CurrentChunkMaterialIndex] = NewSkeletalMaterialIndex;
+				SetLODMaterialMapValue(CurrentChunkMaterialIndex, NewSkeletalMaterialIndex);
 			}
 		}
 
