Undo Works with Screen Rotate in Persona

#trivial
#jifa  UE-77787

diff --git a/Engine/Source/Editor/Persona/Private/EditModes/SkeletonSelectionEditMode.cpp b/Engine/Source/Editor/Persona/Private/EditModes/SkeletonSelectionEditMode.cpp
--- a/Engine/Source/Editor/Persona/Private/EditModes/SkeletonSelectionEditMode.cpp
+++ b/Engine/Source/Editor/Persona/Private/EditModes/SkeletonSelectionEditMode.cpp
@@ -99,7 +99,7 @@ bool FSkeletonSelectionEditMode::StartTracking(FEditorViewportClient* InViewport
 		bool bValidAxis = false;
 		FVector WorldAxisDir;
 
-		if ((CurrentAxis & EAxisList::XYZ) != 0)
+		if ( ((CurrentAxis & EAxisList::XYZ) | (CurrentAxis & EAxisList::Screen)) != 0)
 		{
 			UDebugSkelMeshComponent* PreviewMeshComponent = GetAnimPreviewScene().GetPreviewMeshComponent();
 			FSelectedSocketInfo SelectedSocketInfo = GetAnimPreviewScene().GetSelectedSocket();
