COPY from //Dev-Anim to //Dev-Main

diff --git a/Engine/Source/Editor/MovieSceneTools/Private/TrackEditors/SkeletalAnimationTrackEditor.cpp b/Engine/Source/Editor/MovieSceneTools/Private/TrackEditors/SkeletalAnimationTrackEditor.cpp
--- a/Engine/Source/Editor/MovieSceneTools/Private/TrackEditors/SkeletalAnimationTrackEditor.cpp
+++ b/Engine/Source/Editor/MovieSceneTools/Private/TrackEditors/SkeletalAnimationTrackEditor.cpp
@@ -54,6 +54,8 @@
 #include "DetailWidgetRow.h"
 #include "PropertyCustomizationHelpers.h"
 #include "IPropertyUtilities.h"
+#include "Factories/AnimSequenceFactory.h"
+#include "MovieSceneToolHelpers.h"
 
 
 namespace SkeletalAnimationEditorConstants
@@ -426,9 +428,9 @@ void FSkeletalAnimationSection::SlipSection(FFrameNumber SlipTime)
 	ISequencerSection::SlipSection(SlipTime);
 }
 
-bool FSkeletalAnimationSection::CreatePoseAsset(const TArray<UObject*> NewAssets, FGuid InObjectBinding)
+bool FSkeletalAnimationTrackEditor::CreatePoseAsset(const TArray<UObject*> NewAssets, FGuid InObjectBinding)
 {
-	USkeletalMeshComponent* SkeletalMeshComponent = AcquireSkeletalMeshFromObjectGuid(InObjectBinding, Sequencer.Pin());
+	USkeletalMeshComponent* SkeletalMeshComponent = AcquireSkeletalMeshFromObjectGuid(InObjectBinding, GetSequencer());
 
 	bool bResult = false;
 	if (NewAssets.Num() > 0)
@@ -480,33 +482,18 @@ bool FSkeletalAnimationSection::CreatePoseAsset(const TArray<UObject*> NewAssets
 }
 
 
-void FSkeletalAnimationSection::HandleCreatePoseAsset(FGuid InObjectBinding)
+void FSkeletalAnimationTrackEditor::HandleCreatePoseAsset(FGuid InObjectBinding)
 {
-	USkeleton* Skeleton = AcquireSkeletonFromObjectGuid(InObjectBinding, Sequencer.Pin());
+	USkeleton* Skeleton = AcquireSkeletonFromObjectGuid(InObjectBinding, GetSequencer());
 	if (Skeleton)
 	{
 		TArray<TWeakObjectPtr<UObject>> Skeletons;
 		Skeletons.Add(Skeleton);
-		AnimationEditorUtils::ExecuteNewAnimAsset<UPoseAssetFactory, UPoseAsset>(Skeletons, FString("_PoseAsset"), FAnimAssetCreated::CreateSP(this, &FSkeletalAnimationSection::CreatePoseAsset, InObjectBinding), false);
+		AnimationEditorUtils::ExecuteNewAnimAsset<UPoseAssetFactory, UPoseAsset>(Skeletons, FString("_PoseAsset"), FAnimAssetCreated::CreateSP(this, &FSkeletalAnimationTrackEditor::CreatePoseAsset, InObjectBinding), false);
 	}
 }
 
 
-void FSkeletalAnimationSection::BuildSectionContextMenu(FMenuBuilder& MenuBuilder, const FGuid& InObjectBinding)
-{
-	MenuBuilder.BeginSection(NAME_None, LOCTEXT("SkeletonMenuText", "Skeleton"));
-
-	MenuBuilder.AddMenuEntry(
-		LOCTEXT("CreatePoseAsset", "Create Pose Asset"), 
-		LOCTEXT("CreatePoseAsset_ToolTip", "Create Animation from current pose"),
-		FSlateIcon(), 
-		FUIAction(FExecuteAction::CreateRaw(this, &FSkeletalAnimationSection::HandleCreatePoseAsset, InObjectBinding)),
-		NAME_None, 
-		EUserInterfaceActionType::Button);
-
-	MenuBuilder.EndSection();
-}
-
 void FSkeletalAnimationSection::CustomizePropertiesDetailsView(TSharedRef<IDetailsView> DetailsView, const FSequencerSectionPropertyDetailsViewCustomizationParams& InParams) const
 {
 	DetailsView->RegisterInstancedCustomPropertyTypeLayout(
@@ -601,6 +588,118 @@ bool FSkeletalAnimationTrackEditor::HandleAssetAdded(UObject* Asset, const FGuid
 	return false;
 }
 
+void FSkeletalAnimationTrackEditor::BuildObjectBindingContextMenu(FMenuBuilder& MenuBuilder, const TArray<FGuid>& ObjectBindings, const UClass* ObjectClass)
+{
+	if (ObjectClass->IsChildOf(USkeletalMeshComponent::StaticClass()) || ObjectClass->IsChildOf(AActor::StaticClass()) || ObjectClass->IsChildOf(UChildActorComponent::StaticClass()))
+	{
+		ConstructObjectBindingTrackMenu(MenuBuilder, ObjectBindings);
+	}
+}
+
+bool FSkeletalAnimationTrackEditor::CreateAnimationSequence(const TArray<UObject*> NewAssets, USkeletalMeshComponent* SkelMeshComp)
+{
+	bool bResult = false;
+	if (NewAssets.Num() > 0)
+	{
+		for (UObject* NewAsset : NewAssets)
+		{
+			UAnimSequence* AnimSequence = Cast<UAnimSequence>(NewAsset);
+			if (AnimSequence)
+			{
+				const TSharedPtr<ISequencer> ParentSequencer = GetSequencer();
+				UMovieScene* MovieScene = ParentSequencer->GetFocusedMovieSceneSequence()->GetMovieScene();
+				FMovieSceneSequenceIDRef Template = ParentSequencer->GetFocusedTemplateID();
+				FMovieSceneSequenceTransform RootToLocalTransform;
+				bResult  = MovieSceneToolHelpers::ExportToAnimSequence(AnimSequence, MovieScene, ParentSequencer.Get(), SkelMeshComp, Template, RootToLocalTransform);
+			}
+		}
+
+		// if it contains error, warn them
+		if (bResult)
+		{
+			FText NotificationText;
+			if (NewAssets.Num() == 1)
+			{
+				NotificationText = FText::Format(LOCTEXT("NumAnimSequenceAssetsCreated", "{0} Anim Sequence  assets created."), NewAssets.Num());
+			}
+			else
+			{
+				NotificationText = FText::Format(LOCTEXT("AnimSequenceAssetsCreated", "Anim Sequence asset created: '{0}'."), FText::FromString(NewAssets[0]->GetName()));
+			}
+
+			FNotificationInfo Info(NotificationText);
+			Info.ExpireDuration = 8.0f;
+			Info.bUseLargeFont = false;
+			Info.Hyperlink = FSimpleDelegate::CreateLambda([NewAssets]()
+			{
+				GEditor->GetEditorSubsystem<UAssetEditorSubsystem>()->OpenEditorForAssets(NewAssets);
+			});
+			Info.HyperlinkText = FText::Format(LOCTEXT("OpenNewPoseAssetHyperlink", "Open {0}"), FText::FromString(NewAssets[0]->GetName()));
+
+			TSharedPtr<SNotificationItem> Notification = FSlateNotificationManager::Get().AddNotification(Info);
+			if (Notification.IsValid())
+			{
+				Notification->SetCompletionState(SNotificationItem::CS_Success);
+			}
+		}
+		else
+		{
+			FMessageDialog::Open(EAppMsgType::Ok, LOCTEXT("FailedToCreateAsset", "Failed to create asset"));
+		}
+	}
+	return bResult;
+}
+
+void FSkeletalAnimationTrackEditor::HandleCreateAnimationSequence(USkeletalMeshComponent* SkelMeshComp, USkeleton* Skeleton)
+{
+	if (SkelMeshComp)
+	{
+		TArray<TWeakObjectPtr<UObject>> Skels;
+		if (SkelMeshComp->SkeletalMesh)
+		{
+			Skels.Add(SkelMeshComp->SkeletalMesh);
+		}
+		else
+		{
+			Skels.Add(Skeleton);
+		}
+		AnimationEditorUtils::ExecuteNewAnimAsset<UAnimSequenceFactory, UAnimSequence>(Skels, FString("_Sequence"), FAnimAssetCreated::CreateSP(this, &FSkeletalAnimationTrackEditor::CreateAnimationSequence, SkelMeshComp), false);
+	}
+}
+
+void FSkeletalAnimationTrackEditor::ConstructObjectBindingTrackMenu(FMenuBuilder& MenuBuilder, TArray<FGuid> ObjectBindings)
+{
+	if(ObjectBindings.Num() > 0)
+	{
+		USkeletalMeshComponent* SkelMeshComp =  AcquireSkeletalMeshFromObjectGuid(ObjectBindings[0], GetSequencer());
+
+		if (SkelMeshComp)
+		{
+
+			MenuBuilder.BeginSection("Create Animation Assets", LOCTEXT("CreateAnimationAssetsName", "Create Animation Assets"));
+
+			USkeleton* Skeleton = GetSkeletonFromComponent(SkelMeshComp);
+
+			MenuBuilder.AddMenuEntry(
+				LOCTEXT("CreateAnimSequence", "Create Animation Sequence"),
+				LOCTEXT("PasteCreateAnimSequenceTooltip", "Create Animation Sequence for this Skeletal Mesh. Note it will create it based upon the Sequencer Display Range and Display Frame Rate"),
+				FSlateIcon(),
+				FUIAction(FExecuteAction::CreateRaw(this, &FSkeletalAnimationTrackEditor::HandleCreateAnimationSequence, SkelMeshComp,Skeleton)),
+				NAME_None,
+				EUserInterfaceActionType::Button);
+
+			MenuBuilder.AddMenuEntry(
+				LOCTEXT("CreatePoseAsset", "Create Pose Asset"),
+				LOCTEXT("CreatePoseAsset_ToolTip", "Create Animation from current Pose"),
+				FSlateIcon(),
+				FUIAction(FExecuteAction::CreateRaw(this, &FSkeletalAnimationTrackEditor::HandleCreatePoseAsset, ObjectBindings[0])),
+				NAME_None,
+				EUserInterfaceActionType::Button);
+
+			MenuBuilder.EndSection();
+		}
+	}
+}
 
 void FSkeletalAnimationTrackEditor::BuildObjectBindingTrackMenu(FMenuBuilder& MenuBuilder, const TArray<FGuid>& ObjectBindings, const UClass* ObjectClass)
 {
diff --git a/Engine/Source/Editor/MovieSceneTools/Private/TrackEditors/SkeletalAnimationTrackEditor.h b/Engine/Source/Editor/MovieSceneTools/Private/TrackEditors/SkeletalAnimationTrackEditor.h
--- a/Engine/Source/Editor/MovieSceneTools/Private/TrackEditors/SkeletalAnimationTrackEditor.h
+++ b/Engine/Source/Editor/MovieSceneTools/Private/TrackEditors/SkeletalAnimationTrackEditor.h
@@ -17,6 +17,7 @@ class FMenuBuilder;
 class FSequencerSectionPainter;
 class UMovieSceneSkeletalAnimationSection;
 class USkeleton;
+class USkeletalMeshComponent;
 
 /**
  * Tools for animation tracks
@@ -47,6 +48,7 @@ public:
 	// ISequencerTrackEditor interface
 
 	virtual void AddKey(const FGuid& ObjectGuid) override;
+	virtual void BuildObjectBindingContextMenu(FMenuBuilder& MenuBuilder, const TArray<FGuid>& ObjectBindings, const UClass* ObjectClass) override;
 	virtual void BuildObjectBindingTrackMenu(FMenuBuilder& MenuBuilder, const TArray<FGuid>& ObjectBindings, const UClass* ObjectClass) override;
 	virtual bool HandleAssetAdded(UObject* Asset, const FGuid& TargetObjectGuid) override;
 	virtual TSharedRef<ISequencerSection> MakeSectionInterface( UMovieSceneSection& SectionObject, UMovieSceneTrack& Track, FGuid ObjectBinding ) override;
@@ -73,8 +75,22 @@ private:
 
 	/** Delegate for AnimatablePropertyChanged in AddKey */
 	FKeyPropertyResult AddKeyInternal(FFrameNumber KeyTime, UObject* Object, class UAnimSequenceBase* AnimSequence, UMovieSceneTrack* Track, int32 RowIndex);
+	
+	/** Construct the binding menu*/
+	void ConstructObjectBindingTrackMenu(FMenuBuilder& MenuBuilder, TArray<FGuid> ObjectBindings);
+
+	/** Callback to Create the Animation Asset, pop open the dialolg*/
+	void HandleCreateAnimationSequence(USkeletalMeshComponent* SkelMeshComp, USkeleton* Skeleton);
+
+	/** Callback to Creae the Animation Asset after getting the name*/
+	bool CreateAnimationSequence(const TArray<UObject*> NewAssets,USkeletalMeshComponent* SkelMeshComp);
 
 	friend class FMovieSceneSkeletalAnimationParamsDetailCustomization;
+
+private:
+	/* Was part of the the section but should be at the track level since it takes the final blended result at the current time, not the section instance value*/
+	bool CreatePoseAsset(const TArray<UObject*> NewAssets, FGuid InObjectBinding);
+	void HandleCreatePoseAsset(FGuid InObjectBinding);
 };
 
 
@@ -104,13 +120,8 @@ public:
 	virtual void ResizeSection(ESequencerSectionResizeMode ResizeMode, FFrameNumber ResizeTime) override;
 	virtual void BeginSlipSection() override;
 	virtual void SlipSection(FFrameNumber SlipTime) override;
-	virtual void BuildSectionContextMenu(FMenuBuilder& MenuBuilder, const FGuid& InObjectBinding) override;
 	virtual void CustomizePropertiesDetailsView(TSharedRef<IDetailsView> DetailsView, const FSequencerSectionPropertyDetailsViewCustomizationParams& InParams) const override;
 
-private:
-
-	bool CreatePoseAsset(const TArray<UObject*> NewAssets, FGuid InObjectBinding);
-	void HandleCreatePoseAsset(FGuid InObjectBinding);
 
 private:
 
diff --git a/Engine/Source/Editor/MovieSceneTools/Private/TrackEditors/TransformTrackEditor.cpp b/Engine/Source/Editor/MovieSceneTools/Private/TrackEditors/TransformTrackEditor.cpp
--- a/Engine/Source/Editor/MovieSceneTools/Private/TrackEditors/TransformTrackEditor.cpp
+++ b/Engine/Source/Editor/MovieSceneTools/Private/TrackEditors/TransformTrackEditor.cpp
@@ -47,51 +47,6 @@ void GetActorAndSceneComponentFromObject( UObject* Object, AActor*& OutActor, US
 }
 
 
-class F3DTransformTrackCommands
-	: public TCommands<F3DTransformTrackCommands>
-{
-public:
-
-	F3DTransformTrackCommands()
-		: TCommands<F3DTransformTrackCommands>
-	(
-		"3DTransformTrack",
-		NSLOCTEXT("Contexts", "3DTransformTrack", "3DTransformTrack"),
-		NAME_None, // "MainFrame" // @todo Fix this crash
-		FEditorStyle::GetStyleSetName() // Icon Style Set
-	)
-		, BindingCount(0)
-	{ }
-		
-	/** Sets a transform key at the current time for the selected actor */
-	TSharedPtr< FUICommandInfo > AddTransformKey;
-
-	/** Sets a translation key at the current time for the selected actor */
-	TSharedPtr< FUICommandInfo > AddTranslationKey;
-
-	/** Sets a rotation key at the current time for the selected actor */
-	TSharedPtr< FUICommandInfo > AddRotationKey;
-
-	/** Sets a scale key at the current time for the selected actor */
-	TSharedPtr< FUICommandInfo > AddScaleKey;
-
-	/**
-	 * Initialize commands
-	 */
-	virtual void RegisterCommands() override;
-
-	mutable uint32 BindingCount;
-};
-
-
-void F3DTransformTrackCommands::RegisterCommands()
-{
-	UI_COMMAND( AddTransformKey, "Add Transform Key", "Add a transform key at the current time for the selected actor.", EUserInterfaceActionType::Button, FInputChord(EKeys::S) );
-	UI_COMMAND( AddTranslationKey, "Add Translation Key", "Add a translation key at the current time for the selected actor.", EUserInterfaceActionType::Button, FInputChord(EModifierKey::Shift, EKeys::W) );
-	UI_COMMAND( AddRotationKey, "Add Rotation Key", "Add a rotation key at the current time for the selected actor.", EUserInterfaceActionType::Button, FInputChord(EModifierKey::Shift, EKeys::E) );
-	UI_COMMAND( AddScaleKey, "Add Scale Key", "Add a scale key at the current time for the selected actor.", EUserInterfaceActionType::Button, FInputChord(EModifierKey::Shift, EKeys::R) );
-}
-
 FName F3DTransformTrackEditor::TransformPropertyName("Transform");
 
 F3DTransformTrackEditor::F3DTransformTrackEditor( TSharedRef<ISequencer> InSequencer )
@@ -100,8 +55,6 @@ F3DTransformTrackEditor::F3DTransformTrackEditor( TSharedRef<ISequencer> InSeque
 	// Listen for actor/component movement
 	FCoreUObjectDelegates::OnPreObjectPropertyChanged.AddRaw(this, &F3DTransformTrackEditor::OnPrePropertyChanged);
 	FCoreUObjectDelegates::OnObjectPropertyChanged.AddRaw(this, &F3DTransformTrackEditor::OnPostPropertyChanged);
-
-	F3DTransformTrackCommands::Register();
 }
 
 
@@ -114,14 +67,6 @@ void F3DTransformTrackEditor::OnRelease()
 	FCoreUObjectDelegates::OnPreObjectPropertyChanged.RemoveAll(this);
 	FCoreUObjectDelegates::OnObjectPropertyChanged.RemoveAll(this);
 
-	const F3DTransformTrackCommands& Commands = F3DTransformTrackCommands::Get();
-	Commands.BindingCount--;
-	
-	if (Commands.BindingCount < 1)
-	{
-		F3DTransformTrackCommands::Unregister();
-	}
-
 	for(FLevelEditorViewportClient* LevelVC : GEditor->GetLevelViewportClients())
 	{
 		if (LevelVC && LevelVC->IsPerspective() && LevelVC->GetViewMode() != VMI_Unknown)
@@ -452,40 +397,6 @@ void F3DTransformTrackEditor::OnAddTransformKeysForSelectedObjects( EMovieSceneT
 	}
 }
 
-
-void F3DTransformTrackEditor::BindCommands(TSharedRef<FUICommandList> SequencerCommandBindings)
-{
-	const F3DTransformTrackCommands& Commands = F3DTransformTrackCommands::Get();
-
-	CommandBindings = MakeShared<FUICommandList>();
-	CommandBindings->MapAction(
-		Commands.AddTransformKey,
-		FExecuteAction::CreateSP(this, &F3DTransformTrackEditor::OnAddTransformKeysForSelectedObjects, EMovieSceneTransformChannel::All),
-		FCanExecuteAction::CreateSP(this, &F3DTransformTrackEditor::CanAddTransformKeysForSelectedObjects));
-		
-
-	CommandBindings->MapAction(
-		Commands.AddTranslationKey,
-		FExecuteAction::CreateSP( this, &F3DTransformTrackEditor::OnAddTransformKeysForSelectedObjects, EMovieSceneTransformChannel::Translation ),
-		FCanExecuteAction::CreateSP(this, &F3DTransformTrackEditor::CanAddTransformKeysForSelectedObjects));
-	CommandBindings->MapAction(
-		Commands.AddRotationKey,
-		FExecuteAction::CreateSP( this, &F3DTransformTrackEditor::OnAddTransformKeysForSelectedObjects, EMovieSceneTransformChannel::Rotation ),
-		FCanExecuteAction::CreateSP(this, &F3DTransformTrackEditor::CanAddTransformKeysForSelectedObjects));
-	CommandBindings->MapAction(
-		Commands.AddScaleKey,
-		FExecuteAction::CreateSP( this, &F3DTransformTrackEditor::OnAddTransformKeysForSelectedObjects, EMovieSceneTransformChannel::Scale ),
-		FCanExecuteAction::CreateSP(this, &F3DTransformTrackEditor::CanAddTransformKeysForSelectedObjects));
-	Commands.BindingCount++;
-
-	// Add these bindings to Sequencer
-	SequencerCommandBindings->Append(CommandBindings.ToSharedRef());
-	
-	// Also add them to the Curve Editor 
-	GetSequencer()->GetCommandBindings(ESequencerCommandBindings::CurveEditor)->Append(CommandBindings.ToSharedRef());
-}
-
-
 void F3DTransformTrackEditor::BuildObjectBindingEditButtons(TSharedPtr<SHorizontalBox> EditBox, const FGuid& ObjectGuid, const UClass* ObjectClass)
 {
 	bool bHasCameraComponent = false;
diff --git a/Engine/Source/Editor/MovieSceneTools/Private/TrackEditors/TransformTrackEditor.h b/Engine/Source/Editor/MovieSceneTools/Private/TrackEditors/TransformTrackEditor.h
--- a/Engine/Source/Editor/MovieSceneTools/Private/TrackEditors/TransformTrackEditor.h
+++ b/Engine/Source/Editor/MovieSceneTools/Private/TrackEditors/TransformTrackEditor.h
@@ -49,14 +49,15 @@ public:
 
 	// ISequencerTrackEditor interface
 
-	virtual void BindCommands(TSharedRef<FUICommandList> SequencerCommandBindings) override;
 	virtual void BuildObjectBindingEditButtons(TSharedPtr<SHorizontalBox> EditBox, const FGuid& ObjectBinding, const UClass* ObjectClass) override;
 	virtual void BuildObjectBindingTrackMenu(FMenuBuilder& MenuBuilder, const TArray<FGuid>& ObjectBindinsg, const UClass* ObjectClass) override;
 	virtual TSharedRef<ISequencerSection> MakeSectionInterface( UMovieSceneSection& SectionObject, UMovieSceneTrack& Track, FGuid ObjectBinding ) override;
 	virtual void OnRelease() override;
 	virtual bool SupportsType( TSubclassOf<UMovieSceneTrack> Type ) const override;
 	virtual void BuildTrackContextMenu( FMenuBuilder& MenuBuilder, UMovieSceneTrack* Track ) override;
-
+	virtual bool HasTransformKeyBindings() const override { return true; }
+	virtual bool CanAddTransformKeysForSelectedObjects() const override;
+	virtual void OnAddTransformKeysForSelectedObjects(EMovieSceneTransformChannel Channel) override;
 private:
 
 	/** Returns whether or not a transform track can be added for an actor with a specific handle. */
@@ -114,21 +115,12 @@ private:
 	FText GetLockCameraToolTip(FGuid ObjectGuid) const; 
 
 
-	/** Whether or not we can add a transform key for a selected object 
-	* @return Returns true if we can.
-	**/
-	bool CanAddTransformKeysForSelectedObjects() const;
 
 	/** Generates transform keys based on the last transform, the current transform, and other options. 
 		One transform key is generated for each individual key to be added to the section. */
 	void GetTransformKeys( const TOptional<FTransformData>& LastTransform, const FTransformData& CurrentTransform, EMovieSceneTransformChannel ChannelsToKey, FGeneratedTrackKeys& OutGeneratedKeys );
 
-	/**
-	* Adds transform tracks and keys to the selected objects in the level.
-	*
-	* @param Channel The transform channel to add keys for.
-	*/
-	void OnAddTransformKeysForSelectedObjects( EMovieSceneTransformChannel Channel );
+
 
 	/** 
 	 * Adds transform keys to an object represented by a handle.
diff --git a/Engine/Source/Editor/MovieSceneTools/Public/Channels/FloatChannelCurveModel.h b/Engine/Source/Editor/MovieSceneTools/Public/Channels/FloatChannelCurveModel.h
new file mode 100644
--- /dev/null
+++ b/Engine/Source/Editor/MovieSceneTools/Public/Channels/FloatChannelCurveModel.h
@@ -0,0 +1,68 @@
+// Copyright Epic Games, Inc. All Rights Reserved.
+
+#pragma once
+
+#include "CoreTypes.h"
+#include "UObject/WeakObjectPtr.h"
+#include "UObject/WeakObjectPtrTemplates.h"
+#include "CurveModel.h"
+#include "IBufferedCurveModel.h"
+#include "Channels/MovieSceneFloatChannel.h"
+#include "Channels/MovieSceneChannelHandle.h"
+
+struct FMovieSceneFloatChannel;
+class UMovieSceneSection;
+class FCurveEditor;
+class ISequencer;
+
+class FFloatChannelCurveModel : public FCurveModel
+{
+public:
+	FFloatChannelCurveModel(TMovieSceneChannelHandle<FMovieSceneFloatChannel> InChannel, UMovieSceneSection* InOwningSection, TWeakPtr<ISequencer> InWeakSequencer);
+
+	virtual const void* GetCurve() const override;
+
+	virtual void Modify() override;
+
+	virtual void DrawCurve(const FCurveEditor& CurveEditor, const FCurveEditorScreenSpace& ScreenSpace, TArray<TTuple<double, double>>& InterpolatingPoints) const override;
+	virtual void GetKeys(const FCurveEditor& CurveEditor, double MinTime, double MaxTime, double MinValue, double MaxValue, TArray<FKeyHandle>& OutKeyHandles) const override;
+	virtual void GetKeyDrawInfo(ECurvePointType PointType, const FKeyHandle InKeyHandle, FKeyDrawInfo& OutDrawInfo) const override;
+
+	virtual void GetKeyPositions(TArrayView<const FKeyHandle> InKeys, TArrayView<FKeyPosition> OutKeyPositions) const override;
+	virtual void SetKeyPositions(TArrayView<const FKeyHandle> InKeys, TArrayView<const FKeyPosition> InKeyPositions, EPropertyChangeType::Type ChangeType = EPropertyChangeType::Unspecified) override;
+
+	virtual void GetKeyAttributes(TArrayView<const FKeyHandle> InKeys, TArrayView<FKeyAttributes> OutAttributes) const override;
+	virtual void SetKeyAttributes(TArrayView<const FKeyHandle> InKeys, TArrayView<const FKeyAttributes> InAttributes, EPropertyChangeType::Type ChangeType = EPropertyChangeType::Unspecified) override;
+
+	virtual void GetCurveAttributes(FCurveAttributes& OutCurveAttributes) const override;
+	virtual void SetCurveAttributes(const FCurveAttributes& InCurveAttributes) override;
+	virtual void GetTimeRange(double& MinTime, double& MaxTime) const override;
+	virtual void GetValueRange(double& MinValue, double& MaxValue) const override;
+	virtual int32 GetNumKeys() const override;
+	virtual void GetNeighboringKeys(const FKeyHandle InKeyHandle, TOptional<FKeyHandle>& OutPreviousKeyHandle, TOptional<FKeyHandle>& OutNextKeyHandle) const override;
+	virtual bool Evaluate(double ProspectiveTime, double& OutValue) const override;
+	virtual void AddKeys(TArrayView<const FKeyPosition> InKeyPositions, TArrayView<const FKeyAttributes> InAttributes, TArrayView<TOptional<FKeyHandle>>* OutKeyHandles) override;
+	virtual void RemoveKeys(TArrayView<const FKeyHandle> InKeys) override;
+
+	virtual void CreateKeyProxies(TArrayView<const FKeyHandle> InKeyHandles, TArrayView<UObject*> OutObjects) override;
+
+	virtual TUniquePtr<IBufferedCurveModel> CreateBufferedCurveCopy() const override;
+
+	virtual bool IsReadOnly() const override;
+	virtual UObject* GetOwningObject() const override
+	{
+		return WeakSection.Get();
+	}
+
+public:
+	const TMovieSceneChannelHandle<FMovieSceneFloatChannel>& GetChannelHandle()const  { return ChannelHandle; } 
+
+
+private:
+	void FeaturePointMethod(double StartTime, double EndTime, double StartValue, double Mu, int Depth, int MaxDepth, double& MaxV, double& MinVal) const;
+private:
+
+	TMovieSceneChannelHandle<FMovieSceneFloatChannel> ChannelHandle;
+	TWeakObjectPtr<UMovieSceneSection> WeakSection;
+	TWeakPtr<ISequencer> WeakSequencer;
+};
\ No newline at end of file
diff --git a/Engine/Source/Editor/MovieSceneTools/Public/MovieSceneToolHelpers.h b/Engine/Source/Editor/MovieSceneTools/Public/MovieSceneToolHelpers.h
--- a/Engine/Source/Editor/MovieSceneTools/Public/MovieSceneToolHelpers.h
+++ b/Engine/Source/Editor/MovieSceneTools/Public/MovieSceneToolHelpers.h
@@ -15,6 +15,7 @@
 #include "MovieSceneTranslator.h"
 #include "MovieSceneCaptureSettings.h"
 #include "SEnumCombobox.h"
+#include "Animation/AnimSequence.h"
 
 
 class ISequencer;
@@ -344,6 +345,14 @@ public:
 	 */
 	static UObject* ExportToCameraAnim(UMovieScene* InMovieScene, FGuid& InObjectBinding);
 
+	/*
+	
+	
+	*/
+	static bool ExportToAnimSequence(UAnimSequence* AnimSequence, UMovieScene* MovieScene, IMovieScenePlayer* Player,
+		USkeletalMeshComponent* SkelMesh, FMovieSceneSequenceIDRef& Template, FMovieSceneSequenceTransform& RootToLocalTransform);
+
+
 	/*
 	 * @return Whether this object class has hidden mobility and can't be animated
 	 */
diff --git a/Engine/Source/Editor/Persona/Persona.Build.cs b/Engine/Source/Editor/Persona/Persona.Build.cs
--- a/Engine/Source/Editor/Persona/Persona.Build.cs
+++ b/Engine/Source/Editor/Persona/Persona.Build.cs
@@ -7,10 +7,13 @@ public class Persona : ModuleRules
     public Persona(ReadOnlyTargetRules Target) : base(Target)
     {
         PrivateIncludePaths.Add("Editor/Persona/Private");  // For PCH includes (because they don't work with relative paths, yet)
+        PrivateIncludePaths.Add("Editor/Persona/Private/AnimTimeline");
 
         PublicIncludePathModuleNames.AddRange(
             new string[] {
                 "SkeletonEditor",
+                "AnimationEditor",
+                "AdvancedPreviewScene",
             }
         );
 
@@ -27,7 +30,6 @@ public class Persona : ModuleRules
                 "DesktopPlatform",
                 "ContentBrowser",
                 "AssetTools",
-                "AnimationEditor",
                 "MeshReductionInterface",
                 "SequenceRecorder",
                 "AnimationBlueprintEditor",
@@ -65,6 +67,10 @@ public class Persona : ModuleRules
                 "PinnedCommandList",
                 "RenderCore",
 				"SkeletalMeshUtilitiesCommon",
+                "CurveEditor",
+				"SequencerWidgets",
+				"TimeManagement",
+                "Sequencer",
             }
         );
 
diff --git a/Engine/Source/Editor/Persona/Private/AnimTimeline/AnimModel_AnimSequenceBase.cpp b/Engine/Source/Editor/Persona/Private/AnimTimeline/AnimModel_AnimSequenceBase.cpp
new file mode 100644
--- /dev/null
+++ b/Engine/Source/Editor/Persona/Private/AnimTimeline/AnimModel_AnimSequenceBase.cpp
@@ -0,0 +1,442 @@
+// Copyright Epic Games, Inc. All Rights Reserved.
+
+#include "AnimModel_AnimSequenceBase.h"
+#include "Animation/AnimSequence.h"
+#include "AnimTimelineTrack.h"
+#include "AnimTimelineTrack_Notifies.h"
+#include "AnimTimelineTrack_NotifiesPanel.h"
+#include "AnimTimelineTrack_Curves.h"
+#include "AnimTimelineTrack_Curve.h"
+#include "AnimTimelineTrack_FloatCurve.h"
+#include "AnimTimelineTrack_VectorCurve.h"
+#include "AnimTimelineTrack_TransformCurve.h"
+#include "AnimSequenceTimelineCommands.h"
+#include "Framework/Commands/UICommandList.h"
+#include "IAnimationEditor.h"
+#include "Preferences/PersonaOptions.h"
+#include "FrameNumberDisplayFormat.h"
+#include "Framework/Commands/GenericCommands.h"
+#include "IPersonaPreviewScene.h"
+#include "Animation/DebugSkelMeshComponent.h"
+#include "AnimPreviewInstance.h"
+#include "ScopedTransaction.h"
+
+#define LOCTEXT_NAMESPACE "FAnimModel_AnimSequence"
+
+FAnimModel_AnimSequenceBase::FAnimModel_AnimSequenceBase(const TSharedRef<IPersonaPreviewScene>& InPreviewScene, const TSharedRef<IEditableSkeleton>& InEditableSkeleton, const TSharedRef<FUICommandList>& InCommandList, UAnimSequenceBase* InAnimSequenceBase)
+	: FAnimModel(InPreviewScene, InEditableSkeleton, InCommandList)
+	, AnimSequenceBase(InAnimSequenceBase)
+{
+	UpdateRange();
+
+	// Clear display flags
+	for(bool& bElementNodeDisplayFlag : NotifiesTimingElementNodeDisplayFlags)
+	{
+		bElementNodeDisplayFlag = false;
+	}
+
+	if(GEditor)
+	{
+		GEditor->RegisterForUndo(this);
+	}
+}
+
+FAnimModel_AnimSequenceBase::~FAnimModel_AnimSequenceBase()
+{
+	if(GEditor)
+	{
+		GEditor->UnregisterForUndo(this);
+	}
+}
+
+void FAnimModel_AnimSequenceBase::Initialize()
+{
+	TSharedRef<FUICommandList> CommandList = WeakCommandList.Pin().ToSharedRef();
+
+	const FAnimSequenceTimelineCommands& Commands = FAnimSequenceTimelineCommands::Get();
+
+	CommandList->MapAction(
+		FGenericCommands::Get().Rename,
+		FExecuteAction::CreateLambda([this]
+		{
+			SelectedTracks.Array()[0]->RequestRename();
+		}),
+		FCanExecuteAction::CreateLambda([this]
+		{
+			return (SelectedTracks.Num() > 0) && (SelectedTracks.Array()[0]->CanRename());
+		})
+	);
+
+	CommandList->MapAction(
+		Commands.EditSelectedCurves,
+		FExecuteAction::CreateSP(this, &FAnimModel_AnimSequenceBase::EditSelectedCurves),
+		FCanExecuteAction::CreateSP(this, &FAnimModel_AnimSequenceBase::CanEditSelectedCurves));
+
+	CommandList->MapAction(
+		Commands.RemoveSelectedCurves,
+		FExecuteAction::CreateSP(this, &FAnimModel_AnimSequenceBase::RemoveSelectedCurves));
+
+	CommandList->MapAction(
+		Commands.DisplayFrames,
+		FExecuteAction::CreateSP(this, &FAnimModel_AnimSequenceBase::SetDisplayFormat, EFrameNumberDisplayFormats::Frames),
+		FCanExecuteAction(),
+		FIsActionChecked::CreateSP(this, &FAnimModel_AnimSequenceBase::IsDisplayFormatChecked, EFrameNumberDisplayFormats::Frames));
+
+	CommandList->MapAction(
+		Commands.DisplaySeconds,
+		FExecuteAction::CreateSP(this, &FAnimModel_AnimSequenceBase::SetDisplayFormat, EFrameNumberDisplayFormats::Seconds),
+		FCanExecuteAction(),
+		FIsActionChecked::CreateSP(this, &FAnimModel_AnimSequenceBase::IsDisplayFormatChecked, EFrameNumberDisplayFormats::Seconds));
+
+	CommandList->MapAction(
+		Commands.DisplayPercentage,
+		FExecuteAction::CreateSP(this, &FAnimModel_AnimSequenceBase::ToggleDisplayPercentage),
+		FCanExecuteAction(),
+		FIsActionChecked::CreateSP(this, &FAnimModel_AnimSequenceBase::IsDisplayPercentageChecked));
+
+	CommandList->MapAction(
+		Commands.DisplaySecondaryFormat,
+		FExecuteAction::CreateSP(this, &FAnimModel_AnimSequenceBase::ToggleDisplaySecondary),
+		FCanExecuteAction(),
+		FIsActionChecked::CreateSP(this, &FAnimModel_AnimSequenceBase::IsDisplaySecondaryChecked));
+}
+
+void FAnimModel_AnimSequenceBase::RefreshTracks()
+{
+	ClearTrackSelection();
+
+	// Clear all tracks
+	RootTracks.Empty();
+
+	// Add notifies
+	RefreshNotifyTracks();
+
+	// Add curves
+	RefreshCurveTracks();
+
+	// Tell the UI to refresh
+	OnTracksChangedDelegate.Broadcast();
+
+	UpdateRange();
+}
+
+UAnimSequenceBase* FAnimModel_AnimSequenceBase::GetAnimSequenceBase() const 
+{
+	return AnimSequenceBase;
+}
+
+void FAnimModel_AnimSequenceBase::RefreshNotifyTracks()
+{
+	AnimSequenceBase->InitializeNotifyTrack();
+
+	if(!NotifyRoot.IsValid())
+	{
+		// Add a root track for notifies & then the main 'panel' legacy track
+		NotifyRoot = MakeShared<FAnimTimelineTrack_Notifies>(SharedThis(this));
+	}
+
+	NotifyRoot->ClearChildren();
+	RootTracks.Add(NotifyRoot.ToSharedRef());
+
+	TSharedRef<FAnimTimelineTrack_NotifiesPanel> NotifyPanel = MakeShared<FAnimTimelineTrack_NotifiesPanel>(SharedThis(this));
+	NotifyRoot->AddChild(NotifyPanel);
+	NotifyRoot->SetAnimNotifyPanel(NotifyPanel);
+}
+
+void FAnimModel_AnimSequenceBase::RefreshCurveTracks()
+{
+	if(!CurveRoot.IsValid())
+	{
+		// Add a root track for curves
+		CurveRoot = MakeShared<FAnimTimelineTrack_Curves>(SharedThis(this));
+	}
+
+	CurveRoot->ClearChildren();
+	RootTracks.Add(CurveRoot.ToSharedRef());
+
+	// Next add a track for each float curve
+	for(FFloatCurve& FloatCurve : AnimSequenceBase->RawCurveData.FloatCurves)
+	{
+		CurveRoot->AddChild(MakeShared<FAnimTimelineTrack_FloatCurve>(FloatCurve, SharedThis(this)));
+	}
+
+	UAnimSequence* AnimSequence = Cast<UAnimSequence>(AnimSequenceBase);
+	if (AnimSequence)
+	{
+		if(!AdditiveRoot.IsValid())
+		{
+			// Add a root track for additive layers
+			AdditiveRoot = MakeShared<FAnimTimelineTrack>(LOCTEXT("AdditiveLayerTrackList_Title", "Additive Layer Tracks"), LOCTEXT("AdditiveLayerTrackList_Tooltip", "Additive modifications to bone transforms"), SharedThis(this), true);
+		}
+
+		AdditiveRoot->ClearChildren();
+		RootTracks.Add(AdditiveRoot.ToSharedRef());
+
+		// Next add a track for each transform curve
+		for(FTransformCurve& TransformCurve : AnimSequence->RawCurveData.TransformCurves)
+		{
+			TSharedRef<FAnimTimelineTrack_TransformCurve> TransformCurveTrack = MakeShared<FAnimTimelineTrack_TransformCurve>(TransformCurve, SharedThis(this));
+			TransformCurveTrack->SetExpanded(false);
+			AdditiveRoot->AddChild(TransformCurveTrack);
+
+			FText TransformName = FAnimTimelineTrack_TransformCurve::GetTransformCurveName(AsShared(), TransformCurve.Name);
+			FLinearColor TransformColor = TransformCurve.GetColor();
+			FLinearColor XColor = FLinearColor::Red;
+			FLinearColor YColor = FLinearColor::Green;
+			FLinearColor ZColor = FLinearColor::Blue;
+			FText XName = LOCTEXT("VectorXTrackName", "X");
+			FText YName = LOCTEXT("VectorYTrackName", "Y");
+			FText ZName = LOCTEXT("VectorZTrackName", "Z");
+			
+			FText VectorFormat = LOCTEXT("TransformVectorFormat", "{0}.{1}");
+			FText TranslationName = LOCTEXT("TransformTranslationTrackName", "Translation");
+			TSharedRef<FAnimTimelineTrack_VectorCurve> TranslationCurveTrack = MakeShared<FAnimTimelineTrack_VectorCurve>(TransformCurve.TranslationCurve, TransformCurve.Name, 0, ERawCurveTrackTypes::RCT_Transform, TranslationName, FText::Format(VectorFormat, TransformName, TranslationName), TransformColor, SharedThis(this));
+			TranslationCurveTrack->SetExpanded(false);
+			TransformCurveTrack->AddChild(TranslationCurveTrack);			
+
+			FText ComponentFormat = LOCTEXT("TransformComponentFormat", "{0}.{1}.{2}");
+			TranslationCurveTrack->AddChild(MakeShared<FAnimTimelineTrack_Curve>(TransformCurve.TranslationCurve.FloatCurves[0], TransformCurve.Name, 0, ERawCurveTrackTypes::RCT_Transform, XName, FText::Format(ComponentFormat, TransformName, TranslationName, XName), XColor, XColor, SharedThis(this)));
+			TranslationCurveTrack->AddChild(MakeShared<FAnimTimelineTrack_Curve>(TransformCurve.TranslationCurve.FloatCurves[1], TransformCurve.Name, 1, ERawCurveTrackTypes::RCT_Transform, YName, FText::Format(ComponentFormat, TransformName, TranslationName, YName), YColor, YColor, SharedThis(this)));
+			TranslationCurveTrack->AddChild(MakeShared<FAnimTimelineTrack_Curve>(TransformCurve.TranslationCurve.FloatCurves[2], TransformCurve.Name, 2, ERawCurveTrackTypes::RCT_Transform, ZName, FText::Format(ComponentFormat, TransformName, TranslationName, ZName), ZColor, ZColor, SharedThis(this)));
+
+			FText RollName = LOCTEXT("RotationRollTrackName", "Roll");
+			FText PitchName = LOCTEXT("RotationPitchTrackName", "Pitch");
+			FText YawName = LOCTEXT("RotationYawTrackName", "Yaw");
+			FText RotationName = LOCTEXT("TransformRotationTrackName", "Rotation");
+			TSharedRef<FAnimTimelineTrack_VectorCurve> RotationCurveTrack = MakeShared<FAnimTimelineTrack_VectorCurve>(TransformCurve.RotationCurve, TransformCurve.Name, 3, ERawCurveTrackTypes::RCT_Transform, RotationName, FText::Format(VectorFormat, TransformName, RotationName), TransformColor, SharedThis(this));
+			RotationCurveTrack->SetExpanded(false);
+			TransformCurveTrack->AddChild(RotationCurveTrack);
+			RotationCurveTrack->AddChild(MakeShared<FAnimTimelineTrack_Curve>(TransformCurve.RotationCurve.FloatCurves[0], TransformCurve.Name, 3, ERawCurveTrackTypes::RCT_Transform, RollName, FText::Format(ComponentFormat, TransformName, RotationName, RollName), XColor, XColor, SharedThis(this)));
+			RotationCurveTrack->AddChild(MakeShared<FAnimTimelineTrack_Curve>(TransformCurve.RotationCurve.FloatCurves[1], TransformCurve.Name, 4, ERawCurveTrackTypes::RCT_Transform, PitchName, FText::Format(ComponentFormat, TransformName, RotationName, PitchName), YColor, YColor, SharedThis(this)));
+			RotationCurveTrack->AddChild(MakeShared<FAnimTimelineTrack_Curve>(TransformCurve.RotationCurve.FloatCurves[2], TransformCurve.Name, 5, ERawCurveTrackTypes::RCT_Transform, YawName, FText::Format(ComponentFormat, TransformName, RotationName, YawName), ZColor, ZColor, SharedThis(this)));
+
+			FText ScaleName = LOCTEXT("TransformScaleTrackName", "Scale");
+			TSharedRef<FAnimTimelineTrack_VectorCurve> ScaleCurveTrack = MakeShared<FAnimTimelineTrack_VectorCurve>(TransformCurve.ScaleCurve, TransformCurve.Name, 6, ERawCurveTrackTypes::RCT_Transform, ScaleName, FText::Format(VectorFormat, TransformName, ScaleName), TransformColor, SharedThis(this));
+			ScaleCurveTrack->SetExpanded(false);
+			TransformCurveTrack->AddChild(ScaleCurveTrack);
+			ScaleCurveTrack->AddChild(MakeShared<FAnimTimelineTrack_Curve>(TransformCurve.ScaleCurve.FloatCurves[0], TransformCurve.Name, 6, ERawCurveTrackTypes::RCT_Transform, XName, FText::Format(ComponentFormat, TransformName, ScaleName, XName), XColor, XColor, SharedThis(this)));
+			ScaleCurveTrack->AddChild(MakeShared<FAnimTimelineTrack_Curve>(TransformCurve.ScaleCurve.FloatCurves[1], TransformCurve.Name, 7, ERawCurveTrackTypes::RCT_Transform, YName, FText::Format(ComponentFormat, TransformName, ScaleName, YName), YColor, YColor, SharedThis(this)));
+			ScaleCurveTrack->AddChild(MakeShared<FAnimTimelineTrack_Curve>(TransformCurve.ScaleCurve.FloatCurves[2], TransformCurve.Name, 8, ERawCurveTrackTypes::RCT_Transform, ZName, FText::Format(ComponentFormat, TransformName, ScaleName, ZName), ZColor, ZColor, SharedThis(this)));
+		}
+	}
+}
+
+void FAnimModel_AnimSequenceBase::EditSelectedCurves()
+{
+	TArray<IAnimationEditor::FCurveEditInfo> EditCurveInfo;
+	for(TSharedRef<FAnimTimelineTrack>& SelectedTrack : SelectedTracks)
+	{
+		if(SelectedTrack->IsA<FAnimTimelineTrack_Curve>())
+		{
+			TSharedRef<FAnimTimelineTrack_Curve> CurveTrack = StaticCastSharedRef<FAnimTimelineTrack_Curve>(SelectedTrack);
+			const TArray<FRichCurve*> Curves = CurveTrack->GetCurves();
+			int32 NumCurves = Curves.Num();
+			for(int32 CurveIndex = 0; CurveIndex < NumCurves; ++CurveIndex)
+			{
+				if(CurveTrack->CanEditCurve(CurveIndex))
+				{
+					FText FullName = CurveTrack->GetFullCurveName(CurveIndex);
+					FLinearColor Color = CurveTrack->GetCurveColor(CurveIndex);
+					FSmartName Name;
+					ERawCurveTrackTypes Type;
+					int32 EditCurveIndex;
+					CurveTrack->GetCurveEditInfo(CurveIndex, Name, Type, EditCurveIndex);
+					FSimpleDelegate OnCurveChanged = FSimpleDelegate::CreateSP(&CurveTrack.Get(), &FAnimTimelineTrack_Curve::HandleCurveChanged);
+					EditCurveInfo.AddUnique(IAnimationEditor::FCurveEditInfo(FullName, Color, Name, Type, EditCurveIndex, OnCurveChanged));
+				}
+			}
+		}
+	}
+
+	if(EditCurveInfo.Num() > 0)
+	{
+		OnEditCurves.ExecuteIfBound(AnimSequenceBase, EditCurveInfo, nullptr);
+	}
+}
+
+bool FAnimModel_AnimSequenceBase::CanEditSelectedCurves() const
+{
+	for(const TSharedRef<FAnimTimelineTrack>& SelectedTrack : SelectedTracks)
+	{
+		if(SelectedTrack->IsA<FAnimTimelineTrack_Curve>())
+		{
+			TSharedRef<FAnimTimelineTrack_Curve> CurveTrack = StaticCastSharedRef<FAnimTimelineTrack_Curve>(SelectedTrack);
+			const TArray<FRichCurve*>& Curves = CurveTrack->GetCurves();
+			for(int32 CurveIndex = 0; CurveIndex < Curves.Num(); ++CurveIndex)
+			{
+				if(CurveTrack->CanEditCurve(CurveIndex))
+				{
+					return true;
+				}
+			}
+		}
+	}
+
+	return false;
+}
+
+void FAnimModel_AnimSequenceBase::RemoveSelectedCurves()
+{
+	FScopedTransaction Transaction(LOCTEXT("CurvePanel_RemoveCurves", "Remove Curves"));
+
+	AnimSequenceBase->Modify(true);
+
+	bool bDeletedCurve = false;
+
+	for(TSharedRef<FAnimTimelineTrack>& SelectedTrack : SelectedTracks)
+	{
+		if(SelectedTrack->IsA<FAnimTimelineTrack_FloatCurve>())
+		{
+			TSharedRef<FAnimTimelineTrack_FloatCurve> FloatCurveTrack = StaticCastSharedRef<FAnimTimelineTrack_FloatCurve>(SelectedTrack);
+
+			FFloatCurve& FloatCurve = FloatCurveTrack->GetFloatCurve();
+			FSmartName CurveName = FloatCurveTrack->GetName();
+
+			if(AnimSequenceBase->RawCurveData.GetCurveData(CurveName.UID))
+			{
+				FSmartName TrackName;
+				if (AnimSequenceBase->GetSkeleton()->GetSmartNameByUID(USkeleton::AnimCurveMappingName, CurveName.UID, TrackName))
+				{
+					// Stop editing this curve in the external editor window
+					FSmartName Name;
+					ERawCurveTrackTypes Type;
+					int32 CurveEditIndex;
+					FloatCurveTrack->GetCurveEditInfo(0, Name, Type, CurveEditIndex);
+					IAnimationEditor::FCurveEditInfo EditInfo(Name, Type, CurveEditIndex);
+					OnStopEditingCurves.ExecuteIfBound(TArray<IAnimationEditor::FCurveEditInfo>({ EditInfo }));
+
+					AnimSequenceBase->RawCurveData.DeleteCurveData(TrackName);
+					bDeletedCurve = true;
+				}
+			}
+		}
+		else if(SelectedTrack->IsA<FAnimTimelineTrack_TransformCurve>())
+		{
+			TSharedRef<FAnimTimelineTrack_TransformCurve> TransformCurveTrack = StaticCastSharedRef<FAnimTimelineTrack_TransformCurve>(SelectedTrack);
+
+			FTransformCurve& TransformCurve = TransformCurveTrack->GetTransformCurve();
+			FSmartName CurveName = TransformCurveTrack->GetName();
+
+			if(AnimSequenceBase->RawCurveData.GetCurveData(CurveName.UID, ERawCurveTrackTypes::RCT_Transform))
+			{
+				FSmartName CurveToDelete;
+				if (AnimSequenceBase->GetSkeleton()->GetSmartNameByUID(USkeleton::AnimTrackCurveMappingName, CurveName.UID, CurveToDelete))
+				{
+					// Stop editing these curves in the external editor window
+					TArray<IAnimationEditor::FCurveEditInfo> CurveEditInfo;
+					for(int32 CurveIndex = 0; CurveIndex < TransformCurveTrack->GetCurves().Num(); ++CurveIndex)
+					{
+						FSmartName Name;
+						ERawCurveTrackTypes Type;
+						int32 CurveEditIndex;
+						TransformCurveTrack->GetCurveEditInfo(CurveIndex, Name, Type, CurveEditIndex);
+						IAnimationEditor::FCurveEditInfo EditInfo(Name, Type, CurveEditIndex);
+						CurveEditInfo.Add(EditInfo);
+					}
+
+					OnStopEditingCurves.ExecuteIfBound(CurveEditInfo);
+
+					AnimSequenceBase->RawCurveData.DeleteCurveData(CurveToDelete, ERawCurveTrackTypes::RCT_Transform);
+
+					if(UAnimSequence* AnimSequence = Cast<UAnimSequence>(AnimSequenceBase))
+					{
+						AnimSequence->bNeedsRebake = true;
+					}
+
+					bDeletedCurve = true;
+				}
+			}	
+		}
+	}
+
+	if(bDeletedCurve)
+	{
+		AnimSequenceBase->MarkRawDataAsModified();
+		AnimSequenceBase->PostEditChange();
+
+		if (GetPreviewScene()->GetPreviewMeshComponent()->PreviewInstance != nullptr)
+		{
+			GetPreviewScene()->GetPreviewMeshComponent()->PreviewInstance->RefreshCurveBoneControllers();
+		}
+	}
+
+	RefreshTracks();
+}
+
+void FAnimModel_AnimSequenceBase::SetDisplayFormat(EFrameNumberDisplayFormats InFormat)
+{
+	GetMutableDefault<UPersonaOptions>()->TimelineDisplayFormat = InFormat;
+}
+
+bool FAnimModel_AnimSequenceBase::IsDisplayFormatChecked(EFrameNumberDisplayFormats InFormat) const
+{
+	return GetDefault<UPersonaOptions>()->TimelineDisplayFormat == InFormat;
+}
+
+void FAnimModel_AnimSequenceBase::ToggleDisplayPercentage()
+{
+	GetMutableDefault<UPersonaOptions>()->bTimelineDisplayPercentage = !GetDefault<UPersonaOptions>()->bTimelineDisplayPercentage;
+}
+
+bool FAnimModel_AnimSequenceBase::IsDisplayPercentageChecked() const
+{
+	return GetDefault<UPersonaOptions>()->bTimelineDisplayPercentage;
+}
+
+void FAnimModel_AnimSequenceBase::ToggleDisplaySecondary()
+{
+	GetMutableDefault<UPersonaOptions>()->bTimelineDisplayFormatSecondary = !GetDefault<UPersonaOptions>()->bTimelineDisplayFormatSecondary;
+}
+
+bool FAnimModel_AnimSequenceBase::IsDisplaySecondaryChecked() const
+{
+	return GetDefault<UPersonaOptions>()->bTimelineDisplayFormatSecondary;
+}
+
+void FAnimModel_AnimSequenceBase::HandleUndoRedo()
+{
+	// Close any curves that are no longer editable
+	for(FFloatCurve& FloatCurve : AnimSequenceBase->RawCurveData.FloatCurves)
+	{
+		if(FloatCurve.GetCurveTypeFlag(AACF_Metadata))
+		{
+			IAnimationEditor::FCurveEditInfo CurveEditInfo(FloatCurve.Name, ERawCurveTrackTypes::RCT_Float, 0);
+			OnStopEditingCurves.ExecuteIfBound(TArray<IAnimationEditor::FCurveEditInfo>({ CurveEditInfo }));
+		}
+	}
+}
+
+void FAnimModel_AnimSequenceBase::UpdateRange()
+{
+	FAnimatedRange OldTotalRange = TotalRange;
+
+	// update range
+	TotalRange = FAnimatedRange(0.0, (double)AnimSequenceBase->GetPlayLength());
+
+	if (OldTotalRange != TotalRange)
+	{
+		// Update view range if total range changed
+		ViewRange = TotalRange;
+	}
+}
+
+bool FAnimModel_AnimSequenceBase::IsNotifiesTimingElementDisplayEnabled(ETimingElementType::Type ElementType) const
+{
+	return NotifiesTimingElementNodeDisplayFlags[ElementType];
+}
+
+void FAnimModel_AnimSequenceBase::ToggleNotifiesTimingElementDisplayEnabled(ETimingElementType::Type ElementType)
+{
+	NotifiesTimingElementNodeDisplayFlags[ElementType] = !NotifiesTimingElementNodeDisplayFlags[ElementType];
+}
+
+bool FAnimModel_AnimSequenceBase::ClampToEndTime(float NewEndTime)
+{
+	float SequenceLength = AnimSequenceBase->GetPlayLength();
+
+	//if we had a valid sequence length before and our new end time is shorter
+	//then we need to clamp.
+	return (SequenceLength > 0.f && NewEndTime < SequenceLength);
+}
+
+#undef LOCTEXT_NAMESPACE
\ No newline at end of file
diff --git a/Engine/Source/Editor/Persona/Private/AnimTimeline/AnimTimelineTrack_Curves.cpp b/Engine/Source/Editor/Persona/Private/AnimTimeline/AnimTimelineTrack_Curves.cpp
new file mode 100644
--- /dev/null
+++ b/Engine/Source/Editor/Persona/Private/AnimTimeline/AnimTimelineTrack_Curves.cpp
@@ -0,0 +1,370 @@
+// Copyright Epic Games, Inc. All Rights Reserved.
+
+#include "AnimTimelineTrack_Curves.h"
+#include "PersonaUtils.h"
+#include "Widgets/SBoxPanel.h"
+#include "AnimSequenceTimelineCommands.h"
+#include "Framework/MultiBox/MultiBoxBuilder.h"
+#include "Animation/AnimSequenceBase.h"
+#include "SAnimCurvePanel.h"
+#include "Widgets/Input/STextEntryPopup.h"
+#include "Framework/Application/SlateApplication.h"
+#include "ScopedTransaction.h"
+#include "Animation/AnimMontage.h"
+#include "SAnimOutlinerItem.h"
+#include "Preferences/PersonaOptions.h"
+
+#define LOCTEXT_NAMESPACE "FAnimTimelineTrack_Notifies"
+
+ANIMTIMELINE_IMPLEMENT_TRACK(FAnimTimelineTrack_Curves);
+
+FAnimTimelineTrack_Curves::FAnimTimelineTrack_Curves(const TSharedRef<FAnimModel>& InModel)
+	: FAnimTimelineTrack(LOCTEXT("CurvesRootTrackLabel", "Curves"), LOCTEXT("CurvesRootTrackToolTip", "Curve data contained in this asset"), InModel)
+{
+}
+
+TSharedRef<SWidget> FAnimTimelineTrack_Curves::GenerateContainerWidgetForOutliner(const TSharedRef<SAnimOutlinerItem>& InRow)
+{
+	TSharedPtr<SBorder> OuterBorder;
+	TSharedPtr<SHorizontalBox> InnerHorizontalBox;
+	OutlinerWidget = GenerateStandardOutlinerWidget(InRow, false, OuterBorder, InnerHorizontalBox);
+
+	OuterBorder->SetBorderBackgroundColor(FEditorStyle::GetColor("AnimTimeline.Outliner.HeaderColor"));
+
+	InnerHorizontalBox->AddSlot()
+		.VAlign(VAlign_Center)
+		.HAlign(HAlign_Left)
+		.Padding(2.0f, 1.0f)
+		.AutoWidth()
+		[
+			SNew(STextBlock)
+			.TextStyle(&FEditorStyle::Get().GetWidgetStyle<FTextBlockStyle>("AnimTimeline.Outliner.Label"))
+			.Text(this, &FAnimTimelineTrack_Curves::GetLabel)
+			.HighlightText(InRow->GetHighlightText())
+		];
+
+	InnerHorizontalBox->AddSlot()
+		.FillWidth(1.0f)
+		.HAlign(HAlign_Left)
+		.VAlign(VAlign_Center)
+		.Padding(5.0f, 1.0f)
+		[
+			SNew(STextBlock)
+			.TextStyle(&FEditorStyle::Get().GetWidgetStyle<FTextBlockStyle>("TinyText"))
+			.Text_Lambda([this]()
+			{ 
+				UAnimSequenceBase* AnimSequenceBase = GetModel()->GetAnimSequenceBase();
+				return FText::Format(LOCTEXT("CurveCountFormat", "({0})"), FText::AsNumber(AnimSequenceBase->RawCurveData.FloatCurves.Num())); 
+			})
+		];
+
+	UAnimMontage* AnimMontage = Cast<UAnimMontage>(GetModel()->GetAnimSequenceBase());
+	if(!(AnimMontage && AnimMontage->HasParentAsset()))
+	{
+		InnerHorizontalBox->AddSlot()
+			.AutoWidth()
+			.HAlign(HAlign_Right)
+			.VAlign(VAlign_Center)
+			.Padding(OutlinerRightPadding, 1.0f)
+			[
+				PersonaUtils::MakeTrackButton(LOCTEXT("EditCurvesButtonText", "Curves"), FOnGetContent::CreateSP(this, &FAnimTimelineTrack_Curves::BuildCurvesSubMenu), MakeAttributeSP(this, &FAnimTimelineTrack_Curves::IsHovered))
+			];
+	}
+
+	return OutlinerWidget.ToSharedRef();
+}
+
+void FAnimTimelineTrack_Curves::DeleteAllCurves()
+{
+	const FScopedTransaction Transaction( LOCTEXT("AnimCurve_RemoveAllCurves", "Remove All Curves") );
+
+	UAnimSequenceBase* AnimSequenceBase = GetModel()->GetAnimSequenceBase();
+	AnimSequenceBase->Modify(true);
+	AnimSequenceBase->RawCurveData.DeleteAllCurveData();
+	AnimSequenceBase->MarkRawDataAsModified();
+
+	GetModel()->RefreshTracks();
+}
+
+TSharedRef<SWidget> FAnimTimelineTrack_Curves::BuildCurvesSubMenu()
+{
+	FMenuBuilder MenuBuilder(true, GetModel()->GetCommandList());
+
+	MenuBuilder.BeginSection("Curves", LOCTEXT("CurvesMenuSection", "Curves"));
+	{
+		MenuBuilder.AddSubMenu(
+			FAnimSequenceTimelineCommands::Get().AddCurve->GetLabel(),
+			FAnimSequenceTimelineCommands::Get().AddCurve->GetDescription(),
+			FNewMenuDelegate::CreateSP(this, &FAnimTimelineTrack_Curves::FillVariableCurveMenu)
+		);
+
+		MenuBuilder.AddSubMenu(
+			FAnimSequenceTimelineCommands::Get().AddMetadata->GetLabel(),
+			FAnimSequenceTimelineCommands::Get().AddMetadata->GetDescription(),
+			FNewMenuDelegate::CreateSP(this, &FAnimTimelineTrack_Curves::FillMetadataEntryMenu)
+		);
+
+		UAnimSequenceBase* AnimSequenceBase = GetModel()->GetAnimSequenceBase();
+		if(AnimSequenceBase->RawCurveData.FloatCurves.Num() > 0)
+		{
+			MenuBuilder.AddMenuEntry(
+				FAnimSequenceTimelineCommands::Get().RemoveAllCurves->GetLabel(),
+				FAnimSequenceTimelineCommands::Get().RemoveAllCurves->GetDescription(),
+				FSlateIcon(),
+				FUIAction(FExecuteAction::CreateSP(this, &FAnimTimelineTrack_Curves::DeleteAllCurves))
+			);
+		}
+	}
+
+	MenuBuilder.EndSection();
+
+	MenuBuilder.BeginSection("Options", LOCTEXT("OptionsMenuSection", "Options"));
+	{
+		MenuBuilder.AddMenuEntry(
+			FAnimSequenceTimelineCommands::Get().ShowCurveKeys->GetLabel(),
+			FAnimSequenceTimelineCommands::Get().ShowCurveKeys->GetDescription(),
+			FAnimSequenceTimelineCommands::Get().ShowCurveKeys->GetIcon(),
+			FUIAction(
+				FExecuteAction::CreateSP(this, &FAnimTimelineTrack_Curves::HandleShowCurvePoints),
+				FCanExecuteAction(),
+				FIsActionChecked::CreateSP(this, &FAnimTimelineTrack_Curves::IsShowCurvePointsEnabled)
+			),
+			NAME_None,
+			EUserInterfaceActionType::ToggleButton
+		);
+	}
+	MenuBuilder.EndSection();
+
+	return MenuBuilder.MakeWidget();
+}
+
+void FAnimTimelineTrack_Curves::FillMetadataEntryMenu(FMenuBuilder& Builder)
+{
+	UAnimSequenceBase* AnimSequenceBase = GetModel()->GetAnimSequenceBase();
+	USkeleton* CurrentSkeleton = AnimSequenceBase->GetSkeleton();
+	check(CurrentSkeleton);
+
+	const FSmartNameMapping* Mapping = CurrentSkeleton->GetSmartNameContainer(USkeleton::AnimCurveMappingName);
+	TArray<USkeleton::AnimCurveUID> CurveUids;
+	Mapping->FillUidArray(CurveUids);
+
+	Builder.BeginSection(NAME_None, LOCTEXT("MetadataMenu_ListHeading", "Available Names"));
+	{
+		TArray<FSmartNameSortItem> SmartNameList;
+
+		for (USkeleton::AnimCurveUID Id : CurveUids)
+		{
+			if (!AnimSequenceBase->RawCurveData.GetCurveData(Id))
+			{
+				FName CurveName;
+				if (Mapping->GetName(Id, CurveName))
+				{
+					SmartNameList.Add(FSmartNameSortItem(CurveName, Id));
+				}
+			}
+		}
+
+		{
+			SmartNameList.Sort(FSmartNameSortItemSortOp());
+
+			for (FSmartNameSortItem SmartNameItem : SmartNameList)
+			{
+				const FText Description = LOCTEXT("NewMetadataSubMenu_ToolTip", "Add an existing metadata curve");
+				const FText Label = FText::FromName(SmartNameItem.SmartName);
+
+				FUIAction UIAction;
+				UIAction.ExecuteAction.BindRaw(
+					this, &FAnimTimelineTrack_Curves::AddMetadataEntry,
+					SmartNameItem.ID);
+
+				Builder.AddMenuEntry(Label, Description, FSlateIcon(), UIAction);
+			}
+		}
+	}
+	Builder.EndSection();
+
+	Builder.AddMenuSeparator();
+
+	const FText Description = LOCTEXT("NewMetadataCreateNew_ToolTip", "Create a new metadata entry");
+	const FText Label = LOCTEXT("NewMetadataCreateNew_Label","Create New");
+	FUIAction UIAction;
+	UIAction.ExecuteAction.BindRaw(this, &FAnimTimelineTrack_Curves::CreateNewMetadataEntryClicked);
+
+	Builder.AddMenuEntry(Label, Description, FSlateIcon(), UIAction);
+}
+
+void FAnimTimelineTrack_Curves::FillVariableCurveMenu(FMenuBuilder& Builder)
+{
+	FText Description = LOCTEXT("NewVariableCurveCreateNew_ToolTip", "Create a new variable curve");
+	FText Label = LOCTEXT("NewVariableCurveCreateNew_Label", "Create Curve");
+	FUIAction UIAction;
+	UIAction.ExecuteAction.BindRaw(this, &FAnimTimelineTrack_Curves::CreateNewCurveClicked);
+
+	Builder.AddMenuEntry(Label, Description, FSlateIcon(), UIAction);
+
+	UAnimSequenceBase* AnimSequenceBase = GetModel()->GetAnimSequenceBase();
+	USkeleton* CurrentSkeleton = AnimSequenceBase->GetSkeleton();
+	check(CurrentSkeleton);
+
+	const FSmartNameMapping* Mapping = CurrentSkeleton->GetSmartNameContainer(USkeleton::AnimCurveMappingName);
+	TArray<USkeleton::AnimCurveUID> CurveUids;
+	Mapping->FillUidArray(CurveUids);
+
+	Builder.BeginSection(NAME_None, LOCTEXT("VariableMenu_ListHeading", "Available Names"));
+	{
+		TArray<FSmartNameSortItem> SmartNameList;
+
+		for (USkeleton::AnimCurveUID Id : CurveUids)
+		{
+			if (!AnimSequenceBase->RawCurveData.GetCurveData(Id))
+			{
+				FName CurveName;
+				if (Mapping->GetName(Id, CurveName))
+				{
+					SmartNameList.Add(FSmartNameSortItem(CurveName, Id));
+				}
+			}
+		}
+
+		{
+			SmartNameList.Sort(FSmartNameSortItemSortOp());
+
+			for (FSmartNameSortItem SmartNameItem : SmartNameList)
+			{
+				Description = LOCTEXT("NewVariableSubMenu_ToolTip", "Add an existing variable curve");
+				Label = FText::FromName(SmartNameItem.SmartName);
+
+				UIAction.ExecuteAction.BindRaw(
+					this, &FAnimTimelineTrack_Curves::AddVariableCurve,
+					SmartNameItem.ID);
+
+				Builder.AddMenuEntry(Label, Description, FSlateIcon(), UIAction);
+			}
+		}
+	}
+	Builder.EndSection();
+}
+
+void FAnimTimelineTrack_Curves::AddMetadataEntry(USkeleton::AnimCurveUID Uid)
+{
+	FSmartName NewName;
+	UAnimSequenceBase* AnimSequenceBase = GetModel()->GetAnimSequenceBase();
+	ensureAlways(AnimSequenceBase->GetSkeleton()->GetSmartNameByUID(USkeleton::AnimCurveMappingName, Uid, NewName));
+
+	FScopedTransaction Transaction(LOCTEXT("AddCurveMetadata", "Add Curve Metadata"));
+
+	AnimSequenceBase->Modify(true);
+
+	if(AnimSequenceBase->RawCurveData.AddCurveData(NewName))
+	{
+		AnimSequenceBase->MarkRawDataAsModified();
+		FFloatCurve* Curve = static_cast<FFloatCurve *>(AnimSequenceBase->RawCurveData.GetCurveData(Uid, ERawCurveTrackTypes::RCT_Float));
+		Curve->FloatCurve.AddKey(0.0f, 1.0f);
+		Curve->SetCurveTypeFlag(AACF_Metadata, true);
+		
+		GetModel()->RefreshTracks();
+	}
+}
+
+void FAnimTimelineTrack_Curves::CreateNewMetadataEntryClicked()
+{
+	TSharedRef<STextEntryPopup> TextEntry =
+		SNew(STextEntryPopup)
+		.Label(LOCTEXT("NewMetadataCurveEntryLabal", "Metadata Name"))
+		.OnTextCommitted(this, &FAnimTimelineTrack_Curves::CreateNewMetadataEntry);
+
+	FSlateApplication& SlateApp = FSlateApplication::Get();
+	SlateApp.PushMenu(
+		OutlinerWidget.ToSharedRef(),
+		FWidgetPath(),
+		TextEntry,
+		SlateApp.GetCursorPos(),
+		FPopupTransitionEffect::TypeInPopup
+		);
+}
+
+void FAnimTimelineTrack_Curves::CreateNewMetadataEntry(const FText& CommittedText, ETextCommit::Type CommitType)
+{
+	FSlateApplication::Get().DismissAllMenus();
+	if(CommitType == ETextCommit::OnEnter)
+	{
+		// Add the name to the skeleton and then add the new curve to the sequence
+		UAnimSequenceBase* AnimSequenceBase = GetModel()->GetAnimSequenceBase();
+		USkeleton* Skeleton = AnimSequenceBase->GetSkeleton();
+		if(Skeleton && !CommittedText.IsEmpty())
+		{
+			FSmartName CurveName;
+
+			if(Skeleton->AddSmartNameAndModify(USkeleton::AnimCurveMappingName, FName(*CommittedText.ToString()), CurveName))
+			{
+				AddMetadataEntry(CurveName.UID);
+			}
+		}
+	}
+}
+
+void FAnimTimelineTrack_Curves::CreateNewCurveClicked()
+{
+	TSharedRef<STextEntryPopup> TextEntry =
+		SNew(STextEntryPopup)
+		.Label(LOCTEXT("NewCurveEntryLabal", "Curve Name"))
+		.OnTextCommitted(this, &FAnimTimelineTrack_Curves::CreateTrack);
+
+	FSlateApplication& SlateApp = FSlateApplication::Get();
+	SlateApp.PushMenu(
+		OutlinerWidget.ToSharedRef(),
+		FWidgetPath(),
+		TextEntry,
+		SlateApp.GetCursorPos(),
+		FPopupTransitionEffect::TypeInPopup
+		);
+}
+
+void FAnimTimelineTrack_Curves::CreateTrack(const FText& ComittedText, ETextCommit::Type CommitInfo)
+{
+	if ( CommitInfo == ETextCommit::OnEnter )
+	{
+		UAnimSequenceBase* AnimSequenceBase = GetModel()->GetAnimSequenceBase();
+		USkeleton* Skeleton = AnimSequenceBase->GetSkeleton();
+		if(Skeleton && !ComittedText.IsEmpty())
+		{
+			const FScopedTransaction Transaction(LOCTEXT("AnimCurve_AddTrack", "Add New Curve"));
+			FSmartName NewTrackName;
+
+			if(Skeleton->AddSmartNameAndModify(USkeleton::AnimCurveMappingName, FName(*ComittedText.ToString()), NewTrackName))
+			{
+				AddVariableCurve(NewTrackName.UID);
+			}
+		}
+
+		FSlateApplication::Get().DismissAllMenus();
+	}
+}
+
+void FAnimTimelineTrack_Curves::AddVariableCurve(USkeleton::AnimCurveUID CurveUid)
+{
+	FScopedTransaction Transaction(LOCTEXT("AddCurve", "Add Curve"));
+
+	UAnimSequenceBase* AnimSequenceBase = GetModel()->GetAnimSequenceBase();
+	AnimSequenceBase->Modify();
+	
+	USkeleton* Skeleton = AnimSequenceBase->GetSkeleton();
+	FSmartName NewName;
+	ensureAlways(Skeleton->GetSmartNameByUID(USkeleton::AnimCurveMappingName, CurveUid, NewName));
+	AnimSequenceBase->RawCurveData.AddCurveData(NewName);
+	AnimSequenceBase->MarkRawDataAsModified();
+
+	GetModel()->RefreshTracks();
+}
+
+void FAnimTimelineTrack_Curves::HandleShowCurvePoints()
+{
+	GetMutableDefault<UPersonaOptions>()->bTimelineDisplayCurveKeys = !GetDefault<UPersonaOptions>()->bTimelineDisplayCurveKeys;
+}
+
+bool FAnimTimelineTrack_Curves::IsShowCurvePointsEnabled() const
+{
+	return GetDefault<UPersonaOptions>()->bTimelineDisplayCurveKeys;
+}
+
+#undef LOCTEXT_NAMESPACE
\ No newline at end of file
diff --git a/Engine/Source/Editor/Persona/Private/AnimTimeline/AnimTimelineTrack_FloatCurve.cpp b/Engine/Source/Editor/Persona/Private/AnimTimeline/AnimTimelineTrack_FloatCurve.cpp
new file mode 100644
--- /dev/null
+++ b/Engine/Source/Editor/Persona/Private/AnimTimeline/AnimTimelineTrack_FloatCurve.cpp
@@ -0,0 +1,377 @@
+// Copyright Epic Games, Inc. All Rights Reserved.
+
+#include "AnimTimelineTrack_FloatCurve.h"
+#include "CurveEditor.h"
+#include "Animation/AnimSequenceBase.h"
+#include "Widgets/Layout/SBorder.h"
+#include "EditorStyleSet.h"
+#include "Framework/MultiBox/MultiBoxBuilder.h"
+#include "AnimSequenceTimelineCommands.h"
+#include "ScopedTransaction.h"
+#include "Widgets/Text/SInlineEditableTextBlock.h"
+#include "Animation/Skeleton.h"
+#include "Framework/Notifications/NotificationManager.h"
+#include "Widgets/Notifications/SNotificationList.h"
+#include "Animation/AnimMontage.h"
+#include "Widgets/Colors/SColorBlock.h"
+#include "PersonaUtils.h"
+#include "Widgets/Colors/SColorPicker.h"
+#include "Widgets/Input/SComboButton.h"
+#include "Widgets/Layout/SBox.h"
+#include "AnimModel_AnimSequenceBase.h"
+#include "SAnimOutlinerItem.h"
+
+#define LOCTEXT_NAMESPACE "FAnimTimelineTrack_FloatCurve"
+
+ANIMTIMELINE_IMPLEMENT_TRACK(FAnimTimelineTrack_FloatCurve);
+
+FAnimTimelineTrack_FloatCurve::FAnimTimelineTrack_FloatCurve(FFloatCurve& InCurve, const TSharedRef<FAnimModel>& InModel)
+	: FAnimTimelineTrack_Curve(InCurve.FloatCurve, InCurve.Name, 0, ERawCurveTrackTypes::RCT_Float, FText::FromName(InCurve.Name.DisplayName), FText::FromName(InCurve.Name.DisplayName), InCurve.GetColor(), InCurve.GetColor(), InModel)
+	, FloatCurve(InCurve)
+	, CurveName(InCurve.Name)
+{
+	SetHeight(32.0f);
+}
+
+TSharedRef<SWidget> FAnimTimelineTrack_FloatCurve::MakeTimelineWidgetContainer()
+{
+	auto ColorLambda = [this]()
+	{
+		if(GetModel()->IsTrackSelected(AsShared()))
+		{
+			return FEditorStyle::GetSlateColor("SelectionColor").GetSpecifiedColor().CopyWithNewOpacity(0.75f);
+		}
+		else
+		{
+			return FloatCurve.GetCurveTypeFlag(AACF_Metadata) ? FloatCurve.GetColor().Desaturate(0.25f) : FloatCurve.GetColor().Desaturate(0.75f); 
+		}
+	};
+
+	return
+		SAssignNew(TimelineWidgetContainer, SBorder)
+		.Padding(0.0f)
+		.BorderImage_Lambda([this](){ return FloatCurve.GetCurveTypeFlag(AACF_Metadata) ? FEditorStyle::GetBrush("Sequencer.Section.SelectedSectionOverlay") : FEditorStyle::GetBrush("AnimTimeline.Outliner.DefaultBorder"); })
+		.BorderBackgroundColor_Lambda(ColorLambda)
+		[
+			MakeCurveWidget()
+		];
+}
+
+TSharedRef<SWidget> FAnimTimelineTrack_FloatCurve::GenerateContainerWidgetForOutliner(const TSharedRef<SAnimOutlinerItem>& InRow)
+{
+	TSharedPtr<SBorder> OuterBorder;
+	TSharedPtr<SHorizontalBox> InnerHorizontalBox;
+	TSharedRef<SWidget> OutlinerWidget = GenerateStandardOutlinerWidget(InRow, false, OuterBorder, InnerHorizontalBox);
+
+	TWeakPtr<FAnimTimelineTrack_FloatCurve> WeakThisTrack = SharedThis(this);
+
+	UAnimMontage* AnimMontage = Cast<UAnimMontage>(GetModel()->GetAnimSequenceBase());
+	bool bChildAnimMontage = AnimMontage && AnimMontage->HasParentAsset();
+
+	InnerHorizontalBox->AddSlot()
+		.VAlign(VAlign_Center)
+		.HAlign(HAlign_Left)
+		.Padding(2.0f, 1.0f)
+		.FillWidth(1.0f)
+		[
+			SAssignNew(EditableTextLabel, SInlineEditableTextBlock)
+			.IsReadOnly(bChildAnimMontage)
+			.Text(this, &FAnimTimelineTrack_FloatCurve::GetLabel)
+			.IsSelected_Lambda([WeakThisTrack](){ return WeakThisTrack.Pin()->GetModel()->IsTrackSelected(WeakThisTrack.Pin().ToSharedRef()); })
+			.OnTextCommitted(this, &FAnimTimelineTrack_FloatCurve::OnCommitCurveName)
+			.HighlightText(InRow->GetHighlightText())
+		];
+
+	if(!bChildAnimMontage)
+	{
+		OuterBorder->SetOnMouseDoubleClick(FPointerEventHandler::CreateSP(this, &FAnimTimelineTrack_FloatCurve::HandleDoubleClicked));
+		AddCurveTrackButton(InnerHorizontalBox);
+	}
+
+	return OutlinerWidget;
+}
+
+TSharedRef<SWidget> FAnimTimelineTrack_FloatCurve::BuildCurveTrackMenu()
+{
+	FMenuBuilder MenuBuilder(true, GetModel()->GetCommandList());
+
+	bool bIsMetadata = FloatCurve.GetCurveTypeFlag(AACF_Metadata);
+
+	MenuBuilder.BeginSection("Curve", bIsMetadata ? LOCTEXT("CurveMetadataMenuSection", "Curve Metadata") : LOCTEXT("CurveMenuSection", "Curve"));
+	{
+		if(bIsMetadata)
+		{
+			MenuBuilder.AddMenuEntry(
+				FAnimSequenceTimelineCommands::Get().ConvertMetaDataToCurve->GetLabel(),
+				FAnimSequenceTimelineCommands::Get().ConvertMetaDataToCurve->GetDescription(),
+				FAnimSequenceTimelineCommands::Get().ConvertMetaDataToCurve->GetIcon(),
+				FUIAction(
+					FExecuteAction::CreateSP(this, &FAnimTimelineTrack_FloatCurve::ConvertMetaDataToCurve)
+				)
+			);
+		}
+		else
+		{
+			MenuBuilder.AddMenuEntry(
+				FAnimSequenceTimelineCommands::Get().EditCurve->GetLabel(),
+				FAnimSequenceTimelineCommands::Get().EditCurve->GetDescription(),
+				FAnimSequenceTimelineCommands::Get().EditCurve->GetIcon(),
+				FUIAction(FExecuteAction::CreateSP(this, &FAnimTimelineTrack_FloatCurve::HendleEditCurve)));
+
+			MenuBuilder.AddMenuEntry(
+				FAnimSequenceTimelineCommands::Get().ConvertCurveToMetaData->GetLabel(),
+				FAnimSequenceTimelineCommands::Get().ConvertCurveToMetaData->GetDescription(),
+				FAnimSequenceTimelineCommands::Get().ConvertCurveToMetaData->GetIcon(),
+				FUIAction(
+					FExecuteAction::CreateSP(this, &FAnimTimelineTrack_FloatCurve::ConvertCurveToMetaData)
+				)
+			);
+		}
+
+		MenuBuilder.AddMenuEntry(
+			FAnimSequenceTimelineCommands::Get().RemoveCurve->GetLabel(),
+			FAnimSequenceTimelineCommands::Get().RemoveCurve->GetDescription(),
+			FAnimSequenceTimelineCommands::Get().RemoveCurve->GetIcon(),
+			FUIAction(
+				FExecuteAction::CreateSP(this, &FAnimTimelineTrack_FloatCurve::RemoveCurve)
+			)
+		);
+	}
+	MenuBuilder.EndSection();
+
+	return MenuBuilder.MakeWidget();
+}
+
+void FAnimTimelineTrack_FloatCurve::ConvertCurveToMetaData()
+{
+	UAnimSequenceBase* AnimSequenceBase = GetModel()->GetAnimSequenceBase();
+
+	// Stop editing this curve in the external editor window
+	IAnimationEditor::FCurveEditInfo EditInfo(CurveName, ERawCurveTrackTypes::RCT_Float, 0);
+	StaticCastSharedRef<FAnimModel_AnimSequenceBase>(GetModel())->OnStopEditingCurves.ExecuteIfBound(TArray<IAnimationEditor::FCurveEditInfo>({ EditInfo }));
+
+	FScopedTransaction Transaction(LOCTEXT("CurvePanel_ConvertCurveToMetaData", "Convert curve to metadata"));
+	AnimSequenceBase->Modify(true);
+	FloatCurve.SetCurveTypeFlag(AACF_Metadata, true);
+
+	// We're moving to a metadata curve, we need to clear out the keys.
+	FloatCurve.FloatCurve.Reset();
+	FloatCurve.FloatCurve.AddKey(0.0f, 1.0f);
+}
+
+void FAnimTimelineTrack_FloatCurve::ConvertMetaDataToCurve()
+{
+	UAnimSequenceBase* AnimSequenceBase = GetModel()->GetAnimSequenceBase();
+
+	FScopedTransaction Transaction(LOCTEXT("CurvePanel_ConvertMetaDataToCurve", "Convert metadata to curve"));
+	AnimSequenceBase->Modify(true);
+	FloatCurve.SetCurveTypeFlag(AACF_Metadata, false);
+}
+
+void FAnimTimelineTrack_FloatCurve::RemoveCurve()
+{
+	UAnimSequenceBase* AnimSequenceBase = GetModel()->GetAnimSequenceBase();
+
+	FScopedTransaction Transaction(LOCTEXT("CurvePanel_RemoveCurve", "Remove Curve"));
+	
+	if(AnimSequenceBase->RawCurveData.GetCurveData(FloatCurve.Name.UID))
+	{
+		FSmartName TrackName;
+		if (AnimSequenceBase->GetSkeleton()->GetSmartNameByUID(USkeleton::AnimCurveMappingName, FloatCurve.Name.UID, TrackName))
+		{
+			// Stop editing this curve in the external editor window
+			IAnimationEditor::FCurveEditInfo EditInfo(CurveName, ERawCurveTrackTypes::RCT_Float, 0);
+			StaticCastSharedRef<FAnimModel_AnimSequenceBase>(GetModel())->OnStopEditingCurves.ExecuteIfBound(TArray<IAnimationEditor::FCurveEditInfo>({ EditInfo }));
+
+			AnimSequenceBase->Modify(true);
+			AnimSequenceBase->RawCurveData.DeleteCurveData(TrackName);
+			AnimSequenceBase->MarkRawDataAsModified();
+			AnimSequenceBase->PostEditChange();
+			
+			GetModel()->RefreshTracks();
+		}
+	}
+}
+
+void FAnimTimelineTrack_FloatCurve::OnCommitCurveName(const FText& InText, ETextCommit::Type CommitInfo)
+{
+	UAnimSequenceBase* AnimSequenceBase = GetModel()->GetAnimSequenceBase();
+
+	if (USkeleton* Skeleton = AnimSequenceBase->GetSkeleton())
+	{
+		// only do this if the name isn't same
+		FText CurrentCurveName = GetLabel();
+		if (!CurrentCurveName.EqualToCaseIgnored(InText))
+		{
+			// Stop editing this curve in the external editor window
+			IAnimationEditor::FCurveEditInfo EditInfo(CurveName, ERawCurveTrackTypes::RCT_Float, 0);
+			StaticCastSharedRef<FAnimModel_AnimSequenceBase>(GetModel())->OnStopEditingCurves.ExecuteIfBound(TArray<IAnimationEditor::FCurveEditInfo>({ EditInfo }));
+
+			// Check that the name doesn't already exist
+			const FName RequestedName = FName(*InText.ToString());
+
+			const FSmartNameMapping* NameMapping = Skeleton->GetSmartNameContainer(USkeleton::AnimCurveMappingName);
+
+			FScopedTransaction Transaction(LOCTEXT("CurveEditor_RenameCurve", "Rename Curve"));
+
+			FSmartName NewSmartName;
+			if (NameMapping->FindSmartName(RequestedName, NewSmartName))
+			{
+				// Already in use in this sequence, and if it's not my UID
+				if (NewSmartName.UID != FloatCurve.Name.UID && AnimSequenceBase->RawCurveData.GetCurveData(NewSmartName.UID) != nullptr)
+				{
+					Transaction.Cancel();
+
+					FFormatNamedArguments Args;
+					Args.Add(TEXT("InvalidName"), FText::FromName(RequestedName));
+					FNotificationInfo Info(FText::Format(LOCTEXT("AnimCurveRenamedInUse", "The name \"{InvalidName}\" is already used."), Args));
+
+					Info.bUseLargeFont = false;
+					Info.ExpireDuration = 5.0f;
+
+					TSharedPtr<SNotificationItem> Notification = FSlateNotificationManager::Get().AddNotification(Info);
+					if (Notification.IsValid())
+					{
+						Notification->SetCompletionState(SNotificationItem::CS_Fail);
+					}
+					return;
+				}
+			}
+			else
+			{
+				if(!Skeleton->AddSmartNameAndModify(USkeleton::AnimCurveMappingName, RequestedName, NewSmartName))
+				{
+					Transaction.Cancel();
+					FNotificationInfo Info(LOCTEXT("AnimCurveRenamedError", "Failed to rename curve smart name, check the log for errors."));
+
+					Info.bUseLargeFont = false;
+					Info.ExpireDuration = 5.0f;
+
+					TSharedPtr<SNotificationItem> Notification = FSlateNotificationManager::Get().AddNotification(Info);
+					if (Notification.IsValid())
+					{
+						Notification->SetCompletionState(SNotificationItem::CS_Fail);
+					}
+					return;
+				}
+			}
+
+			FloatCurve.Name.UID = NewSmartName.UID;
+			FloatCurve.Name.DisplayName = NewSmartName.DisplayName;
+		}
+	}
+}
+
+FText FAnimTimelineTrack_FloatCurve::GetLabel() const
+{
+	return FAnimTimelineTrack_FloatCurve::GetFloatCurveName(GetModel(), FloatCurve.Name);
+}
+
+FText FAnimTimelineTrack_FloatCurve::GetFloatCurveName(const TSharedRef<FAnimModel>& InModel, const FSmartName& InSmartName)
+{
+	const FSmartNameMapping* NameMapping = InModel->GetAnimSequenceBase()->GetSkeleton()->GetSmartNameContainer(USkeleton::AnimCurveMappingName);
+	if(NameMapping)
+	{
+		FName CurveName;
+		if(NameMapping->GetName(InSmartName.UID, CurveName))
+		{
+			return FText::FromName(CurveName);
+		}
+	}
+
+	return FText::FromName(InSmartName.DisplayName);
+}
+
+bool FAnimTimelineTrack_FloatCurve::CanEditCurve(int32 InCurveIndex) const
+{
+	return !FloatCurve.GetCurveTypeFlag(AACF_Metadata);
+}
+
+void FAnimTimelineTrack_FloatCurve::RequestRename()
+{
+	if(EditableTextLabel.IsValid())
+	{
+		EditableTextLabel->EnterEditingMode();
+	}
+}
+
+void FAnimTimelineTrack_FloatCurve::AddCurveTrackButton(TSharedPtr<SHorizontalBox> InnerHorizontalBox)
+{
+	InnerHorizontalBox->AddSlot()
+	.AutoWidth()
+	.HAlign(HAlign_Right)
+	.VAlign(VAlign_Center)
+	.Padding(0.0f, 1.0f)
+	[
+		PersonaUtils::MakeTrackButton(LOCTEXT("EditCurveButtonText", "Curve"), FOnGetContent::CreateSP(this, &FAnimTimelineTrack_FloatCurve::BuildCurveTrackMenu), MakeAttributeSP(this, &FAnimTimelineTrack_FloatCurve::IsHovered))
+	];
+
+	auto GetValue = [this]()
+	{
+		return FloatCurve.Color;
+	};
+
+	auto SetValue = [this](FLinearColor InNewColor)
+	{
+		UAnimSequenceBase* AnimSequenceBase = GetModel()->GetAnimSequenceBase();
+		FScopedTransaction Transaction(LOCTEXT("SetCurveColor", "Set Curve Color"));
+		AnimSequenceBase->Modify(true);
+		FloatCurve.Color = InNewColor;
+
+		Color = InNewColor;
+
+		// Set display curves too
+		for(const TPair<FCurveModelID, TUniquePtr<FCurveModel>>& CurvePair : CurveEditor->GetCurves())
+		{
+			CurvePair.Value->SetColor(InNewColor);
+		}
+	};
+
+	auto OnGetMenuContent = [GetValue, SetValue]()
+	{
+		// Open a color picker
+		return SNew(SColorPicker)
+			.TargetColorAttribute_Lambda(GetValue)
+			.UseAlpha(false)
+			.DisplayInlineVersion(true)
+			.OnColorCommitted_Lambda(SetValue);
+	};
+
+	InnerHorizontalBox->AddSlot()
+	.AutoWidth()
+	.HAlign(HAlign_Right)
+	.VAlign(VAlign_Fill)
+	.Padding(2.0f, 0.0f, 0.0f, 0.0f)
+	[
+		SNew(SComboButton)
+		.ToolTipText(LOCTEXT("EditCurveColor", "Edit Curve Color"))
+		.ContentPadding(0.0f)
+		.HasDownArrow(false)
+		.ButtonStyle(FEditorStyle::Get(), "Sequencer.AnimationOutliner.ColorStrip")
+		.OnGetMenuContent_Lambda(OnGetMenuContent)
+		.CollapseMenuOnParentFocus(true)
+		.VAlign(VAlign_Fill)
+		.ButtonContent()
+		[
+			SNew(SColorBlock)
+			.Color_Lambda(GetValue)
+			.ShowBackgroundForAlpha(false)
+			.IgnoreAlpha(true)
+			.Size(FVector2D(OutlinerRightPadding - 2.0f, OutlinerRightPadding))
+		]
+	];
+}
+
+FLinearColor FAnimTimelineTrack_FloatCurve::GetCurveColor(int32 InCurveIndex) const
+{ 
+	return FloatCurve.Color; 
+}
+
+void FAnimTimelineTrack_FloatCurve::GetCurveEditInfo(int32 InCurveIndex, FSmartName& OutName, ERawCurveTrackTypes& OutType, int32& OutCurveIndex) const
+{
+	OutName = CurveName;
+	OutType = ERawCurveTrackTypes::RCT_Float;
+	OutCurveIndex = InCurveIndex;
+}
+
+#undef LOCTEXT_NAMESPACE
\ No newline at end of file
diff --git a/Engine/Source/Editor/Persona/Private/AnimTimeline/AnimTimelineTrack_FloatCurve.h b/Engine/Source/Editor/Persona/Private/AnimTimeline/AnimTimelineTrack_FloatCurve.h
new file mode 100644
--- /dev/null
+++ b/Engine/Source/Editor/Persona/Private/AnimTimeline/AnimTimelineTrack_FloatCurve.h
@@ -0,0 +1,54 @@
+// Copyright Epic Games, Inc. All Rights Reserved.
+
+#pragma once
+
+#include "AnimTimelineTrack_Curve.h"
+
+struct FFloatCurve;
+class SBorder;
+class FCurveEditor;
+
+class FAnimTimelineTrack_FloatCurve : public FAnimTimelineTrack_Curve
+{
+	ANIMTIMELINE_DECLARE_TRACK(FAnimTimelineTrack_FloatCurve, FAnimTimelineTrack_Curve);
+
+public:
+	FAnimTimelineTrack_FloatCurve(FFloatCurve& InCurve, const TSharedRef<FAnimModel>& InModel);
+
+	/** FAnimTimelineTrack_Curve interface */
+	virtual TSharedRef<SWidget> MakeTimelineWidgetContainer() override;
+	virtual TSharedRef<SWidget> GenerateContainerWidgetForOutliner(const TSharedRef<SAnimOutlinerItem>& InRow) override;
+	virtual TSharedRef<SWidget> BuildCurveTrackMenu() override;
+	virtual FText GetLabel() const override;
+	virtual bool CanEditCurve(int32 InCurveIndex) const override;
+	virtual bool CanRename() const override { return true; }
+	virtual void RequestRename() override;
+	virtual void AddCurveTrackButton(TSharedPtr<SHorizontalBox> InnerHorizontalBox) override;
+	virtual FLinearColor GetCurveColor(int32 InCurveIndex) const override;
+	virtual void GetCurveEditInfo(int32 InCurveIndex, FSmartName& OutName, ERawCurveTrackTypes& OutType, int32& OutCurveIndex) const override;
+
+	/** Access the curve we are editing */
+	FFloatCurve& GetFloatCurve() { return FloatCurve; }
+
+	/** Helper function used to get a smart name for a curve */
+	static FText GetFloatCurveName(const TSharedRef<FAnimModel>& InModel, const FSmartName& InSmartName);
+
+	/** Get this curves name */
+	FSmartName GetName() const { return CurveName; }
+
+private:
+	void ConvertCurveToMetaData();
+	void ConvertMetaDataToCurve();
+	void RemoveCurve();
+	void OnCommitCurveName(const FText& InText, ETextCommit::Type CommitInfo);
+
+private:
+	/** The curve we are editing */
+	FFloatCurve& FloatCurve;
+
+	/** The curve name */
+	FSmartName CurveName;
+
+	/** Label we can edit */
+	TSharedPtr<SInlineEditableTextBlock> EditableTextLabel;
+};
\ No newline at end of file
diff --git a/Engine/Source/Editor/Persona/Private/AnimTimeline/AnimTimelineTrack_MontageSections.cpp b/Engine/Source/Editor/Persona/Private/AnimTimeline/AnimTimelineTrack_MontageSections.cpp
new file mode 100644
--- /dev/null
+++ b/Engine/Source/Editor/Persona/Private/AnimTimeline/AnimTimelineTrack_MontageSections.cpp
@@ -0,0 +1,23 @@
+// Copyright Epic Games, Inc. All Rights Reserved.
+
+#include "AnimTimelineTrack_MontageSections.h"
+#include "Styling/CoreStyle.h"
+#include "Fonts/FontMeasure.h"
+#include "Framework/Application/SlateApplication.h"
+#include "Widgets/SLeafWidget.h"
+#include "Animation/AnimMontage.h"
+#include "SCurveEditor.h"
+#include "EditorStyleSet.h"
+#include "AnimModel_AnimMontage.h"
+
+#define LOCTEXT_NAMESPACE "FAnimTimelineTrack_MontageSections"
+
+ANIMTIMELINE_IMPLEMENT_TRACK(FAnimTimelineTrack_MontageSections);
+
+FAnimTimelineTrack_MontageSections::FAnimTimelineTrack_MontageSections(const TSharedRef<FAnimModel>& InModel)
+	: FAnimTimelineTrack(LOCTEXT("TrackTitle", "Sections"), LOCTEXT("TrackTooltip", "Montage sections"), InModel)
+{
+	SetHeight(32.0f);
+}
+
+#undef LOCTEXT_NAMESPACE
\ No newline at end of file
diff --git a/Engine/Source/Editor/Persona/Private/AnimTimeline/AnimTimelineTrack_MontageSections.h b/Engine/Source/Editor/Persona/Private/AnimTimeline/AnimTimelineTrack_MontageSections.h
new file mode 100644
--- /dev/null
+++ b/Engine/Source/Editor/Persona/Private/AnimTimeline/AnimTimelineTrack_MontageSections.h
@@ -0,0 +1,20 @@
+// Copyright Epic Games, Inc. All Rights Reserved.
+
+#pragma once
+
+#include "AnimTimelineTrack.h"
+
+class SAnimMontagePanel;
+class SVerticalBox;
+
+/** A timeline track that displays anim montage section labels */
+class FAnimTimelineTrack_MontageSections : public FAnimTimelineTrack
+{
+	ANIMTIMELINE_DECLARE_TRACK(FAnimTimelineTrack_MontageSections, FAnimTimelineTrack);
+
+public:
+	FAnimTimelineTrack_MontageSections(const TSharedRef<FAnimModel>& InModel);
+
+	/** FAnimTimelineTrack interface */
+	virtual bool SupportsFiltering() const override { return false; }
+};
\ No newline at end of file
diff --git a/Engine/Source/Editor/Persona/Private/AnimTimeline/AnimTimelineTrack_TransformCurve.cpp b/Engine/Source/Editor/Persona/Private/AnimTimeline/AnimTimelineTrack_TransformCurve.cpp
new file mode 100644
--- /dev/null
+++ b/Engine/Source/Editor/Persona/Private/AnimTimeline/AnimTimelineTrack_TransformCurve.cpp
@@ -0,0 +1,202 @@
+// Copyright Epic Games, Inc. All Rights Reserved.
+
+#include "AnimTimelineTrack_TransformCurve.h"
+#include "Animation/AnimSequenceBase.h"
+#include "Widgets/Layout/SBorder.h"
+#include "EditorStyleSet.h"
+#include "Framework/MultiBox/MultiBoxBuilder.h"
+#include "AnimSequenceTimelineCommands.h"
+#include "ScopedTransaction.h"
+#include "Animation/AnimSequence.h"
+#include "IPersonaPreviewScene.h"
+#include "Animation/DebugSkelMeshComponent.h"
+#include "AnimPreviewInstance.h"
+#include "AnimModel_AnimSequenceBase.h"
+
+#define LOCTEXT_NAMESPACE "FAnimTimelineTrack_TransformCurve"
+
+ANIMTIMELINE_IMPLEMENT_TRACK(FAnimTimelineTrack_TransformCurve);
+
+FAnimTimelineTrack_TransformCurve::FAnimTimelineTrack_TransformCurve(FTransformCurve& InCurve, const TSharedRef<FAnimModel>& InModel)
+	: FAnimTimelineTrack_Curve(FAnimTimelineTrack_TransformCurve::GetTransformCurveName(InModel, InCurve.Name), FAnimTimelineTrack_TransformCurve::GetTransformCurveName(InModel, InCurve.Name), InCurve.GetColor(), InCurve.GetColor(), InModel)
+	, TransformCurve(InCurve)
+	, CurveName(InCurve.Name)
+{
+	Curves.Add(&InCurve.TranslationCurve.FloatCurves[0]);
+	Curves.Add(&InCurve.TranslationCurve.FloatCurves[1]);
+	Curves.Add(&InCurve.TranslationCurve.FloatCurves[2]);
+	Curves.Add(&InCurve.RotationCurve.FloatCurves[0]);
+	Curves.Add(&InCurve.RotationCurve.FloatCurves[1]);
+	Curves.Add(&InCurve.RotationCurve.FloatCurves[2]);
+	Curves.Add(&InCurve.ScaleCurve.FloatCurves[0]);
+	Curves.Add(&InCurve.ScaleCurve.FloatCurves[1]);
+	Curves.Add(&InCurve.ScaleCurve.FloatCurves[2]);
+}
+
+FLinearColor FAnimTimelineTrack_TransformCurve::GetCurveColor(int32 InCurveIndex) const
+{
+	static const FLinearColor Colors[3] =
+	{
+		FLinearColor::Red,
+		FLinearColor::Green,
+		FLinearColor::Blue,
+	};
+
+	return Colors[InCurveIndex % 3];
+}
+
+FText FAnimTimelineTrack_TransformCurve::GetFullCurveName(int32 InCurveIndex) const 
+{
+	check(InCurveIndex >= 0 && InCurveIndex < 9);
+
+	static const FText TrackNames[9] =
+	{
+		LOCTEXT("TranslationXTrackName", "Translation.X"),
+		LOCTEXT("TranslationYTrackName", "Translation.Y"),
+		LOCTEXT("TranslationZTrackName", "Translation.Z"),
+		LOCTEXT("RotationRollTrackName", "Rotation.Roll"),
+		LOCTEXT("RotationPitchTrackName", "Rotation.Pitch"),
+		LOCTEXT("RotationYawTrackName", "Rotation.Yaw"),
+		LOCTEXT("ScaleXTrackName", "Scale.X"),
+		LOCTEXT("ScaleYTrackName", "Scale.Y"),
+		LOCTEXT("ScaleZTrackName", "Scale.Z"),
+	};
+			
+	return FText::Format(LOCTEXT("TransformVectorFormat", "{0}.{1}"), FullCurveName, TrackNames[InCurveIndex]);
+}
+
+FText FAnimTimelineTrack_TransformCurve::GetTransformCurveName(const TSharedRef<FAnimModel>& InModel, const FSmartName& InSmartName)
+{
+	const FSmartNameMapping* NameMapping = InModel->GetAnimSequenceBase()->GetSkeleton()->GetSmartNameContainer(USkeleton::AnimTrackCurveMappingName);
+	if(NameMapping)
+	{
+		FName CurveName;
+		if(NameMapping->GetName(InSmartName.UID, CurveName))
+		{
+			return FText::FromName(CurveName);
+		}
+	}
+
+	return FText::FromName(InSmartName.DisplayName);
+}
+
+TSharedRef<SWidget> FAnimTimelineTrack_TransformCurve::BuildCurveTrackMenu()
+{
+	FMenuBuilder MenuBuilder(true, GetModel()->GetCommandList());
+
+	MenuBuilder.BeginSection("Curve", LOCTEXT("CurveMenuSection", "Curve"));
+	{
+			MenuBuilder.AddMenuEntry(
+				FAnimSequenceTimelineCommands::Get().EditCurve->GetLabel(),
+				FAnimSequenceTimelineCommands::Get().EditCurve->GetDescription(),
+				FAnimSequenceTimelineCommands::Get().EditCurve->GetIcon(),
+				FUIAction(FExecuteAction::CreateSP(this, &FAnimTimelineTrack_TransformCurve::HendleEditCurve)));
+
+		MenuBuilder.AddMenuEntry(
+			LOCTEXT("RemoveTrack", "Remove Track"),
+			LOCTEXT("RemoveTrackTooltip", "Remove this track"),
+			FSlateIcon(),
+			FUIAction(FExecuteAction::CreateSP(this, &FAnimTimelineTrack_TransformCurve::DeleteTrack)));
+
+		MenuBuilder.AddMenuEntry(
+			LOCTEXT("TrackEnabled", "Enabled"),
+			LOCTEXT("TrackEnabledTooltip", "Enable/disable this track"),
+			FSlateIcon(),
+			FUIAction(
+				FExecuteAction::CreateSP(this, &FAnimTimelineTrack_TransformCurve::ToggleEnabled),
+				FCanExecuteAction(), 
+				FIsActionChecked::CreateSP(this, &FAnimTimelineTrack_TransformCurve::IsEnabled)),
+			NAME_None,
+			EUserInterfaceActionType::ToggleButton
+			);
+	}
+	MenuBuilder.EndSection();
+
+	return MenuBuilder.MakeWidget();
+}
+
+void FAnimTimelineTrack_TransformCurve::DeleteTrack()
+{
+	UAnimSequenceBase* AnimSequenceBase = GetModel()->GetAnimSequenceBase();
+	TSharedRef<FAnimModel_AnimSequenceBase> BaseModel = StaticCastSharedRef<FAnimModel_AnimSequenceBase>(GetModel());
+
+	if(AnimSequenceBase->RawCurveData.GetCurveData(TransformCurve.Name.UID, ERawCurveTrackTypes::RCT_Transform))
+	{
+		const FScopedTransaction Transaction(LOCTEXT("AnimCurve_DeleteTrack", "Delete Curve"));
+		FSmartName CurveToDelete;
+		if (AnimSequenceBase->GetSkeleton()->GetSmartNameByUID(USkeleton::AnimTrackCurveMappingName, TransformCurve.Name.UID, CurveToDelete))
+		{
+			// Stop editing these curves in the external editor window
+			TArray<IAnimationEditor::FCurveEditInfo> CurveEditInfo;
+			for(int32 CurveIndex = 0; CurveIndex < Curves.Num(); ++CurveIndex)
+			{
+				FSmartName Name;
+				ERawCurveTrackTypes Type;
+				int32 CurveEditIndex;
+				GetCurveEditInfo(CurveIndex, Name, Type, CurveEditIndex);
+				IAnimationEditor::FCurveEditInfo EditInfo(Name, Type, CurveEditIndex);
+				CurveEditInfo.Add(EditInfo);
+			}
+
+			BaseModel->OnStopEditingCurves.ExecuteIfBound(CurveEditInfo);
+
+			AnimSequenceBase->Modify();
+			AnimSequenceBase->RawCurveData.DeleteCurveData(CurveToDelete, ERawCurveTrackTypes::RCT_Transform);
+
+			if(UAnimSequence* AnimSequence = Cast<UAnimSequence>(AnimSequenceBase))
+			{
+				AnimSequence->bNeedsRebake = true;
+			}
+
+			GetModel()->RefreshTracks();
+
+			if (GetModel()->GetPreviewScene()->GetPreviewMeshComponent()->PreviewInstance != nullptr)
+			{
+				GetModel()->GetPreviewScene()->GetPreviewMeshComponent()->PreviewInstance->RefreshCurveBoneControllers();
+			}
+		}
+	}
+}
+
+bool FAnimTimelineTrack_TransformCurve::IsEnabled() const
+{
+	UAnimSequenceBase* AnimSequenceBase = GetModel()->GetAnimSequenceBase();
+	FAnimCurveBase* Curve = AnimSequenceBase->RawCurveData.GetCurveData(TransformCurve.Name.UID, ERawCurveTrackTypes::RCT_Transform);
+	return Curve && !Curve->GetCurveTypeFlag(AACF_Disabled);
+}
+
+void FAnimTimelineTrack_TransformCurve::ToggleEnabled()
+{
+	UAnimSequenceBase* AnimSequenceBase = GetModel()->GetAnimSequenceBase();
+
+	FAnimCurveBase* Curve = AnimSequenceBase->RawCurveData.GetCurveData(TransformCurve.Name.UID, ERawCurveTrackTypes::RCT_Transform);
+	if (Curve)
+	{
+		bool bEnabled = !Curve->GetCurveTypeFlag(AACF_Disabled);
+
+		const FScopedTransaction Transaction(bEnabled ? LOCTEXT("AnimCurve_DisableTrack", "Disable track") : LOCTEXT("AnimCurve_EnableTrack", "Enable track"));
+		AnimSequenceBase->Modify();
+
+		Curve->SetCurveTypeFlag(AACF_Disabled, bEnabled);
+
+		if(UAnimSequence* AnimSequence = Cast<UAnimSequence>(AnimSequenceBase))
+		{
+			AnimSequence->bNeedsRebake = true;
+		}
+
+		// need to update curves, otherwise they're not disabled
+		if (GetModel()->GetPreviewScene()->GetPreviewMeshComponent()->PreviewInstance != nullptr)
+		{
+			GetModel()->GetPreviewScene()->GetPreviewMeshComponent()->PreviewInstance->RefreshCurveBoneControllers();
+		}
+	}
+}
+
+void FAnimTimelineTrack_TransformCurve::GetCurveEditInfo(int32 InCurveIndex, FSmartName& OutName, ERawCurveTrackTypes& OutType, int32& OutCurveIndex) const
+{
+	OutName = TransformCurve.Name;
+	OutType = ERawCurveTrackTypes::RCT_Transform;
+	OutCurveIndex = InCurveIndex;
+}
+
+#undef LOCTEXT_NAMESPACE
\ No newline at end of file
diff --git a/Engine/Source/Editor/Persona/Private/AnimTimeline/AnimTimelineTrack_TransformCurve.h b/Engine/Source/Editor/Persona/Private/AnimTimeline/AnimTimelineTrack_TransformCurve.h
new file mode 100644
--- /dev/null
+++ b/Engine/Source/Editor/Persona/Private/AnimTimeline/AnimTimelineTrack_TransformCurve.h
@@ -0,0 +1,50 @@
+// Copyright Epic Games, Inc. All Rights Reserved.
+
+#pragma once
+
+#include "AnimTimelineTrack_Curve.h"
+
+struct FTransformCurve;
+class SBorder;
+class FCurveEditor;
+
+class FAnimTimelineTrack_TransformCurve : public FAnimTimelineTrack_Curve
+{
+	ANIMTIMELINE_DECLARE_TRACK(FAnimTimelineTrack_TransformCurve, FAnimTimelineTrack_Curve);
+
+public:
+	FAnimTimelineTrack_TransformCurve(FTransformCurve& InCurve, const TSharedRef<FAnimModel>& InModel);
+
+	/** FAnimTimelineTrack_Curve interface */
+	virtual FLinearColor GetCurveColor(int32 InCurveIndex) const override;
+	virtual FText GetFullCurveName(int32 InCurveIndex) const override;
+	virtual TSharedRef<SWidget> BuildCurveTrackMenu() override;
+	virtual bool CanEditCurve(int32 InCurveIndex) const override { return true; }
+	virtual void GetCurveEditInfo(int32 InCurveIndex, FSmartName& OutName, ERawCurveTrackTypes& OutType, int32& OutCurveIndex) const override;
+
+	/** Access the curve we are editing */
+	FTransformCurve& GetTransformCurve() { return TransformCurve; }
+
+	/** Helper function used to get a smart name for a curve */
+	static FText GetTransformCurveName(const TSharedRef<FAnimModel>& InModel, const FSmartName& InSmartName);
+
+	/** Get this curves name */
+	FSmartName GetName() const { return CurveName; }
+
+private:
+	/** Delete this track via the track menu */
+	void DeleteTrack();
+
+	/** Show enabled state in the menu */
+	bool IsEnabled() const;
+
+	/** Toggle enabled state via the menu */
+	void ToggleEnabled();
+
+private:
+	/** The curve we are editing */
+	FTransformCurve& TransformCurve;
+
+	/** The curve name */
+	FSmartName CurveName;
+};
\ No newline at end of file
diff --git a/Engine/Source/Editor/Persona/Private/AnimTimeline/AnimTimelineTrack_VectorCurve.cpp b/Engine/Source/Editor/Persona/Private/AnimTimeline/AnimTimelineTrack_VectorCurve.cpp
new file mode 100644
--- /dev/null
+++ b/Engine/Source/Editor/Persona/Private/AnimTimeline/AnimTimelineTrack_VectorCurve.cpp
@@ -0,0 +1,58 @@
+// Copyright Epic Games, Inc. All Rights Reserved.
+
+#include "AnimTimelineTrack_VectorCurve.h"
+#include "Animation/AnimSequenceBase.h"
+#include "Widgets/Layout/SBorder.h"
+#include "EditorStyleSet.h"
+
+#define LOCTEXT_NAMESPACE "FAnimTimelineTrack_VectorCurve"
+
+ANIMTIMELINE_IMPLEMENT_TRACK(FAnimTimelineTrack_VectorCurve);
+
+FAnimTimelineTrack_VectorCurve::FAnimTimelineTrack_VectorCurve(FVectorCurve& InCurve, const FSmartName& InName, int32 InCurveIndex, ERawCurveTrackTypes InType, const FText& InCurveName, const FText& InFullCurveName, const FLinearColor& InColor, const TSharedRef<FAnimModel>& InModel)
+	: FAnimTimelineTrack_Curve(InCurveName, InFullCurveName, InColor, InColor, InModel)
+	, VectorCurve(InCurve)
+{
+	OuterCurveName = InName;
+	OuterCurveIndex = InCurveIndex;
+	OuterType = InType;
+
+	Curves.Add(&InCurve.FloatCurves[0]);
+	Curves.Add(&InCurve.FloatCurves[1]);
+	Curves.Add(&InCurve.FloatCurves[2]);
+}
+
+FLinearColor FAnimTimelineTrack_VectorCurve::GetCurveColor(int32 InCurveIndex) const
+{
+	static const FLinearColor Colors[3] =
+	{
+		FLinearColor::Red,
+		FLinearColor::Green,
+		FLinearColor::Blue,
+	};
+
+	return Colors[InCurveIndex % 3];
+}
+
+FText FAnimTimelineTrack_VectorCurve::GetFullCurveName(int32 InCurveIndex) const 
+{ 
+	check(InCurveIndex >= 0 && InCurveIndex < 3);
+
+	static const FText TrackNames[3] =
+	{
+		LOCTEXT("VectorXTrackName", "X"),
+		LOCTEXT("VectorYTrackName", "Y"),
+		LOCTEXT("VectorZTrackName", "Z")
+	};
+			
+	return FText::Format(LOCTEXT("TransformVectorFormat", "{0}.{1}"), FullCurveName, TrackNames[InCurveIndex]);
+}
+
+void FAnimTimelineTrack_VectorCurve::GetCurveEditInfo(int32 InCurveIndex, FSmartName& OutName, ERawCurveTrackTypes& OutType, int32& OutCurveIndex) const
+{
+	OutName = OuterCurveName;
+	OutType = OuterType;
+	OutCurveIndex = OuterCurveIndex + InCurveIndex;
+}
+
+#undef LOCTEXT_NAMESPACE
\ No newline at end of file
diff --git a/Engine/Source/Editor/Persona/Private/AnimTimeline/AnimTimelineTrack_VectorCurve.h b/Engine/Source/Editor/Persona/Private/AnimTimeline/AnimTimelineTrack_VectorCurve.h
new file mode 100644
--- /dev/null
+++ b/Engine/Source/Editor/Persona/Private/AnimTimeline/AnimTimelineTrack_VectorCurve.h
@@ -0,0 +1,30 @@
+// Copyright Epic Games, Inc. All Rights Reserved.
+
+#pragma once
+
+#include "AnimTimelineTrack_Curve.h"
+
+struct FVectorCurve;
+class SBorder;
+class FCurveEditor;
+
+class FAnimTimelineTrack_VectorCurve : public FAnimTimelineTrack_Curve
+{
+	ANIMTIMELINE_DECLARE_TRACK(FAnimTimelineTrack_VectorCurve, FAnimTimelineTrack_Curve);
+
+public:
+	FAnimTimelineTrack_VectorCurve(FVectorCurve& InCurve, const FSmartName& InName, int32 InCurveIndex, ERawCurveTrackTypes InType, const FText& InCurveName, const FText& InFullCurveName, const FLinearColor& InColor, const TSharedRef<FAnimModel>& InModel);
+
+	/** FAnimTimelineTrack_Curve interface */
+	virtual FLinearColor GetCurveColor(int32 InCurveIndex) const override;
+	virtual FText GetFullCurveName(int32 InCurveIndex) const override;
+	virtual bool CanEditCurve(int32 InCurveIndex) const override { return true; }
+	virtual void GetCurveEditInfo(int32 InCurveIndex, FSmartName& OutName, ERawCurveTrackTypes& OutType, int32& OutCurveIndex) const override;
+
+	/** Access the curve we are editing */
+	FVectorCurve& GetVectorCurve() { return VectorCurve; }
+
+private:
+	/** The curve we are editing */
+	FVectorCurve& VectorCurve;
+};
\ No newline at end of file
