COPY from //Dev-Anim to //Dev-Main

diff --git a/Engine/Source/Editor/Persona/Private/AnimTimeline/SAnimTimeline.cpp b/Engine/Source/Editor/Persona/Private/AnimTimeline/SAnimTimeline.cpp
new file mode 100644
--- /dev/null
+++ b/Engine/Source/Editor/Persona/Private/AnimTimeline/SAnimTimeline.cpp
@@ -0,0 +1,799 @@
+// Copyright Epic Games, Inc. All Rights Reserved.
+
+#include "SAnimTimeline.h"
+#include "Styling/ISlateStyle.h"
+#include "Widgets/SWidget.h"
+#include "AnimModel.h"
+#include "SAnimOutliner.h"
+#include "SAnimTrackArea.h"
+#include "Widgets/Layout/SSplitter.h"
+#include "SAnimTimelineOverlay.h"
+#include "SAnimTimelineSplitterOverlay.h"
+#include "ISequencerWidgetsModule.h"
+#include "FrameNumberNumericInterface.h"
+#include "Widgets/Input/SSearchBox.h"
+#include "Widgets/Layout/SScrollBorder.h"
+#include "EditorStyleSet.h"
+#include "Fonts/FontMeasure.h"
+#include "Widgets/Layout/SGridPanel.h"
+#include "Widgets/Layout/SSpacer.h"
+#include "Framework/Application/SlateApplication.h"
+#include "Widgets/Input/SButton.h"
+#include "Modules/ModuleManager.h"
+#include "Preferences/PersonaOptions.h"
+#include "IPersonaPreviewScene.h"
+#include "Animation/DebugSkelMeshComponent.h"
+#include "AnimPreviewInstance.h"
+#include "EditorWidgetsModule.h"
+#include "AnimationEditorPreviewScene.h"
+#include "Framework/MultiBox/MultiBoxBuilder.h"
+#include "ScopedTransaction.h"
+#include "Widgets/Input/STextEntryPopup.h"
+#include "AnimSequenceTimelineCommands.h"
+#include "Widgets/Input/SSpinBox.h"
+#include "SAnimTimelineTransportControls.h"
+#include "Animation/AnimSequence.h"
+
+#define LOCTEXT_NAMESPACE "SAnimTimeline"
+
+void SAnimTimeline::Construct(const FArguments& InArgs, const TSharedRef<FAnimModel>& InModel)
+{
+	TWeakPtr<FAnimModel> WeakModel = InModel;
+
+	Model = InModel;
+	OnReceivedFocus = InArgs._OnReceivedFocus;
+
+	int32 TickResolutionValue = InModel->GetTickResolution();
+	int32 SequenceFrameRate = FMath::RoundToInt(InModel->GetFrameRate());
+
+	if (InModel->GetPreviewScene()->GetPreviewMeshComponent()->PreviewInstance)
+	{
+		InModel->GetPreviewScene()->GetPreviewMeshComponent()->PreviewInstance->AddKeyCompleteDelegate(FSimpleDelegate::CreateSP(this, &SAnimTimeline::HandleKeyComplete));
+	}
+
+	ViewRange = MakeAttributeLambda([WeakModel](){ return WeakModel.IsValid() ? WeakModel.Pin()->GetViewRange() : FAnimatedRange(0.0, 0.0); });
+
+	TAttribute<EFrameNumberDisplayFormats> DisplayFormat = MakeAttributeLambda([]()
+	{
+		return GetDefault<UPersonaOptions>()->TimelineDisplayFormat;
+	});
+
+	TAttribute<EFrameNumberDisplayFormats> DisplayFormatSecondary = MakeAttributeLambda([]()
+	{
+		return GetDefault<UPersonaOptions>()->TimelineDisplayFormat == EFrameNumberDisplayFormats::Frames ? EFrameNumberDisplayFormats::Seconds : EFrameNumberDisplayFormats::Frames;
+	});
+
+	TAttribute<FFrameRate> TickResolution = MakeAttributeLambda([TickResolutionValue]()
+	{
+		return FFrameRate(TickResolutionValue, 1);
+	});
+
+	TAttribute<FFrameRate> DisplayRate = MakeAttributeLambda([SequenceFrameRate]()
+	{
+		return FFrameRate(SequenceFrameRate, 1);
+	});
+
+	// Create our numeric type interface so we can pass it to the time slider below.
+	NumericTypeInterface = MakeShareable(new FFrameNumberInterface(DisplayFormat, 0, TickResolution, DisplayRate));
+	SecondaryNumericTypeInterface = MakeShareable(new FFrameNumberInterface(DisplayFormatSecondary, 0, TickResolution, DisplayRate));
+
+	FTimeSliderArgs TimeSliderArgs;
+	{
+		TimeSliderArgs.ScrubPosition = MakeAttributeLambda([WeakModel](){ return WeakModel.IsValid() ? WeakModel.Pin()->GetScrubPosition() : FFrameTime(0); });
+		TimeSliderArgs.ViewRange = ViewRange;
+		TimeSliderArgs.PlaybackRange = MakeAttributeLambda([WeakModel](){ return WeakModel.IsValid() ? WeakModel.Pin()->GetPlaybackRange() : TRange<FFrameNumber>(0, 0); });
+		TimeSliderArgs.ClampRange = MakeAttributeLambda([WeakModel](){ return WeakModel.IsValid() ? WeakModel.Pin()->GetTotalRange() : FAnimatedRange(0.0, 0.0); });
+		TimeSliderArgs.DisplayRate = DisplayRate;
+		TimeSliderArgs.TickResolution = TickResolution;
+		TimeSliderArgs.OnViewRangeChanged = FOnViewRangeChanged::CreateSP(&InModel.Get(), &FAnimModel::HandleViewRangeChanged);
+		TimeSliderArgs.IsPlaybackRangeLocked = true;
+		TimeSliderArgs.PlaybackStatus = EMovieScenePlayerStatus::Stopped;
+		TimeSliderArgs.NumericTypeInterface = NumericTypeInterface;
+		TimeSliderArgs.OnScrubPositionChanged = FOnScrubPositionChanged::CreateSP(this, &SAnimTimeline::HandleScrubPositionChanged);
+	}
+
+	TimeSliderController = MakeShareable(new FAnimTimeSliderController(TimeSliderArgs, InModel, SharedThis(this), SecondaryNumericTypeInterface));
+	
+	TSharedRef<FAnimTimeSliderController> TimeSliderControllerRef = TimeSliderController.ToSharedRef();
+
+	// Create the top slider
+	const bool bMirrorLabels = false;
+	ISequencerWidgetsModule& SequencerWidgets = FModuleManager::Get().LoadModuleChecked<ISequencerWidgetsModule>("SequencerWidgets");
+	TopTimeSlider = SequencerWidgets.CreateTimeSlider(TimeSliderControllerRef, bMirrorLabels);
+
+	// Create bottom time range slider
+	TSharedRef<ITimeSlider> BottomTimeRange = SequencerWidgets.CreateTimeRange(
+		FTimeRangeArgs(
+			EShowRange::ViewRange | EShowRange::PlaybackRange,
+			EShowRange::ViewRange,
+			TimeSliderControllerRef,
+			EVisibility::Visible,
+			NumericTypeInterface.ToSharedRef()
+		),
+		SequencerWidgets.CreateTimeRangeSlider(TimeSliderControllerRef)
+	);
+
+	TSharedRef<SScrollBar> ScrollBar = SNew(SScrollBar)
+		.Thickness(FVector2D(5.0f, 5.0f));
+
+	InModel->RefreshTracks();
+
+	TrackArea = SNew(SAnimTrackArea, InModel, TimeSliderControllerRef);
+	Outliner = SNew(SAnimOutliner, InModel, TrackArea.ToSharedRef())
+		.ExternalScrollbar(ScrollBar)
+		.Clipping(EWidgetClipping::ClipToBounds)
+		.FilterText_Lambda([this](){ return FilterText; });
+
+	TrackArea->SetOutliner(Outliner);
+
+	ColumnFillCoefficients[0] = 0.2f;
+	ColumnFillCoefficients[1] = 0.8f;
+
+	TAttribute<float> FillCoefficient_0, FillCoefficient_1;
+	{
+		FillCoefficient_0.Bind(TAttribute<float>::FGetter::CreateSP(this, &SAnimTimeline::GetColumnFillCoefficient, 0));
+		FillCoefficient_1.Bind(TAttribute<float>::FGetter::CreateSP(this, &SAnimTimeline::GetColumnFillCoefficient, 1));
+	}
+
+	const int32 Column0 = 0, Column1 = 1;
+	const int32 Row0 = 0, Row1 = 1, Row2 = 2, Row3 = 3, Row4 = 4;
+
+	const float CommonPadding = 3.f;
+	const FMargin ResizeBarPadding(4.f, 0, 0, 0);
+
+	ChildSlot
+	[
+		SNew(SOverlay)
+		+SOverlay::Slot()
+		[
+			SNew(SVerticalBox)
+			+SVerticalBox::Slot()
+			[
+				SNew(SOverlay)
+				+SOverlay::Slot()
+				[
+					SNew(SGridPanel)
+					.FillRow(1, 1.0f)
+					.FillColumn(0, FillCoefficient_0)
+					.FillColumn(1, FillCoefficient_1)
+
+					// outliner search box
+					+SGridPanel::Slot(Column0, Row0, SGridPanel::Layer(10))
+					[
+						SNew(SHorizontalBox)
+						+SHorizontalBox::Slot()
+						.FillWidth(1.0f)
+						.VAlign(VAlign_Center)
+						[
+							SAssignNew(SearchBox, SSearchBox)
+							.HintText(LOCTEXT("FilterTracksHint", "Filter"))
+							.OnTextChanged(this, &SAnimTimeline::OnOutlinerSearchChanged)
+						]
+						+SHorizontalBox::Slot()
+						.VAlign(VAlign_Center)
+						.HAlign(HAlign_Center)
+						.AutoWidth()
+						.Padding(2.0f, 0.0f, 2.0f, 0.0f)
+						[
+							SNew(SBox)
+							.MinDesiredWidth(30.0f)
+							.VAlign(VAlign_Center)
+							.HAlign(HAlign_Center)
+							[
+								// Current Play Time 
+								SNew(SSpinBox<double>)
+								.Style(&FEditorStyle::GetWidgetStyle<FSpinBoxStyle>("Sequencer.PlayTimeSpinBox"))
+								.Value_Lambda([this]() -> double
+								{
+									return Model.Pin()->GetScrubPosition().Value;
+								})
+								.OnValueChanged(this, &SAnimTimeline::SetPlayTime)
+								.OnValueCommitted_Lambda([this](double InFrame, ETextCommit::Type)
+								{
+									SetPlayTime(InFrame);
+								})
+								.MinValue(TOptional<double>())
+								.MaxValue(TOptional<double>())
+								.TypeInterface(NumericTypeInterface)
+								.Delta(this, &SAnimTimeline::GetSpinboxDelta)
+								.LinearDeltaSensitivity(25)
+							]
+						]
+					]
+					// main timeline area
+					+SGridPanel::Slot(Column0, Row1, SGridPanel::Layer(10))
+					.ColumnSpan(2)
+					[
+						SNew(SHorizontalBox)
+						+SHorizontalBox::Slot()
+						[
+							SNew(SOverlay)
+							+SOverlay::Slot()
+							[
+								SNew(SScrollBorder, Outliner.ToSharedRef())
+								[
+									SNew(SHorizontalBox)
+
+									// outliner tree
+									+SHorizontalBox::Slot()
+									.FillWidth(FillCoefficient_0)
+									[
+										SNew(SBox)
+										[
+											Outliner.ToSharedRef()
+										]
+									]
+
+									// track area
+									+SHorizontalBox::Slot()
+									.FillWidth(FillCoefficient_1)
+									[
+										SNew(SBox)
+										.Padding(ResizeBarPadding)
+										.Clipping(EWidgetClipping::ClipToBounds)
+										[
+											TrackArea.ToSharedRef()
+										]
+									]
+								]
+							]
+
+							+SOverlay::Slot()
+							.HAlign(HAlign_Right)
+							[
+								ScrollBar
+							]
+						]
+					]
+
+					// Transport controls
+					+SGridPanel::Slot(Column0, Row3, SGridPanel::Layer(10))
+					.VAlign(VAlign_Center)
+					.HAlign(HAlign_Center)
+					[
+						SNew(SAnimTimelineTransportControls, InModel->GetPreviewScene(), InModel->GetAnimSequenceBase())
+					]
+
+					// Second column
+					+SGridPanel::Slot(Column1, Row0)
+					.Padding(ResizeBarPadding)
+					.RowSpan(2)
+					[
+						SNew(SBorder)
+						.BorderImage(FEditorStyle::GetBrush("ToolPanel.GroupBorder"))
+						[
+							SNew(SSpacer)
+						]
+					]
+
+					+SGridPanel::Slot(Column1, Row0, SGridPanel::Layer(10))
+					.Padding(ResizeBarPadding)
+					[
+						SNew( SBorder )
+						.BorderImage( FEditorStyle::GetBrush("ToolPanel.GroupBorder") )
+						.BorderBackgroundColor( FLinearColor(.50f, .50f, .50f, 1.0f ) )
+						.Padding(0)
+						.Clipping(EWidgetClipping::ClipToBounds)
+						[
+							TopTimeSlider.ToSharedRef()
+						]
+					]
+
+					// Overlay that draws the tick lines
+					+SGridPanel::Slot(Column1, Row1, SGridPanel::Layer(10))
+					.Padding(ResizeBarPadding)
+					[
+						SNew(SAnimTimelineOverlay, TimeSliderControllerRef)
+						.Visibility( EVisibility::HitTestInvisible )
+						.DisplayScrubPosition( false )
+						.DisplayTickLines( true )
+						.Clipping(EWidgetClipping::ClipToBounds)
+					]
+
+					// Overlay that draws the scrub position
+					+SGridPanel::Slot(Column1, Row1, SGridPanel::Layer(20))
+					.Padding(ResizeBarPadding)
+					[
+						SNew(SAnimTimelineOverlay, TimeSliderControllerRef)
+						.Visibility( EVisibility::HitTestInvisible )
+						.DisplayScrubPosition( true )
+						.DisplayTickLines( false )
+						.Clipping(EWidgetClipping::ClipToBounds)
+					]
+
+					// play range slider
+					+SGridPanel::Slot(Column1, Row3, SGridPanel::Layer(10))
+					.Padding(ResizeBarPadding)
+					[
+						SNew(SBorder)
+						.BorderImage( FEditorStyle::GetBrush("ToolPanel.GroupBorder") )
+						.BorderBackgroundColor( FLinearColor(0.5f, 0.5f, 0.5f, 1.0f ) )
+						.Clipping(EWidgetClipping::ClipToBounds)
+						.Padding(0)
+						[
+							BottomTimeRange
+						]
+					]
+				]
+				+SOverlay::Slot()
+				[
+					// track area virtual splitter overlay
+					SNew(SAnimTimelineSplitterOverlay)
+					.Style(FEditorStyle::Get(), "AnimTimeline.Outliner.Splitter")
+					.Visibility(EVisibility::SelfHitTestInvisible)
+
+					+ SSplitter::Slot()
+					.Value(FillCoefficient_0)
+					.OnSlotResized(SSplitter::FOnSlotResized::CreateSP(this, &SAnimTimeline::OnColumnFillCoefficientChanged, 0))
+					[
+						SNew(SSpacer)
+					]
+
+					+ SSplitter::Slot()
+					.Value(FillCoefficient_1)
+					.OnSlotResized(SSplitter::FOnSlotResized::CreateSP(this, &SAnimTimeline::OnColumnFillCoefficientChanged, 1))
+					[
+						SNew(SSpacer)
+					]
+				]
+			]
+		]
+	];
+}
+
+FReply SAnimTimeline::OnMouseButtonUp(const FGeometry& MyGeometry, const FPointerEvent& MouseEvent)
+{
+	if(MouseEvent.GetEffectingButton() == EKeys::RightMouseButton)
+	{
+		const bool bCloseAfterSelection = true;
+		FMenuBuilder MenuBuilder(bCloseAfterSelection, Model.Pin()->GetCommandList());
+
+		MenuBuilder.BeginSection("TimelineOptions", LOCTEXT("TimelineOptions", "Timeline Options") );
+		{
+			MenuBuilder.AddSubMenu(
+				LOCTEXT("TimeFormat", "Time Format"),
+				LOCTEXT("TimeFormatTooltip", "Choose the format of times we display in the timeline"),
+				FNewMenuDelegate::CreateLambda([](FMenuBuilder& InMenuBuilder)
+				{
+					InMenuBuilder.BeginSection("TimeFormat", LOCTEXT("TimeFormat", "Time Format") );
+					{
+						InMenuBuilder.AddMenuEntry(FAnimSequenceTimelineCommands::Get().DisplaySeconds);
+						InMenuBuilder.AddMenuEntry(FAnimSequenceTimelineCommands::Get().DisplayFrames);
+					}
+					InMenuBuilder.EndSection();
+
+					InMenuBuilder.BeginSection("TimelineAdditional", LOCTEXT("TimelineAdditional", "Additional Display") );
+					{
+						InMenuBuilder.AddMenuEntry(FAnimSequenceTimelineCommands::Get().DisplayPercentage);
+						InMenuBuilder.AddMenuEntry(FAnimSequenceTimelineCommands::Get().DisplaySecondaryFormat);
+					}
+					InMenuBuilder.EndSection();
+				})
+			);
+		}
+		MenuBuilder.EndSection();
+
+		UAnimSequence* AnimSequence = Cast<UAnimSequence>(Model.Pin()->GetAnimSequenceBase());
+		if( AnimSequence )
+		{
+			FFrameTime MouseTime = TimeSliderController->GetFrameTimeFromMouse(MyGeometry, MouseEvent.GetScreenSpacePosition());
+			float CurrentFrameTime = (float)((double)MouseTime.AsDecimal() / (double)Model.Pin()->GetTickResolution());
+			float SequenceLength = AnimSequence->GetPlayLength();
+			uint32 NumFrames = AnimSequence->GetNumberOfFrames();
+
+			MenuBuilder.BeginSection("SequenceEditingContext", LOCTEXT("SequenceEditing", "Sequence Editing") );
+			{
+				float CurrentFrameFraction = CurrentFrameTime / SequenceLength;
+				int32 CurrentFrameNumber = CurrentFrameFraction * NumFrames;
+
+				FUIAction Action;
+				FText Label;
+
+				//Menu - "Remove Before"
+				//Only show this option if the selected frame is greater than frame 1 (first frame)
+				if (CurrentFrameNumber > 0)
+				{
+					CurrentFrameFraction = (float)CurrentFrameNumber / (float)NumFrames;
+
+					//Corrected frame time based on selected frame number
+					float CorrectedFrameTime = CurrentFrameFraction * SequenceLength;
+
+					Action = FUIAction(FExecuteAction::CreateSP(this, &SAnimTimeline::OnCropAnimSequence, true, CorrectedFrameTime));
+					Label = FText::Format(LOCTEXT("RemoveTillFrame", "Remove frame 0 to frame {0}"), FText::AsNumber(CurrentFrameNumber));
+					MenuBuilder.AddMenuEntry(Label, LOCTEXT("RemoveBefore_ToolTip", "Remove sequence before current position"), FSlateIcon(), Action);
+				}
+
+				uint32 NextFrameNumber = CurrentFrameNumber + 1;
+
+				//Menu - "Remove After"
+				//Only show this option if next frame (CurrentFrameNumber + 1) is valid
+				if (NextFrameNumber < NumFrames)
+				{
+					float NextFrameFraction = (float)NextFrameNumber / (float)NumFrames;
+					float NextFrameTime = NextFrameFraction * SequenceLength;
+					Action = FUIAction(FExecuteAction::CreateSP(this, &SAnimTimeline::OnCropAnimSequence, false, NextFrameTime));
+					Label = FText::Format(LOCTEXT("RemoveFromFrame", "Remove from frame {0} to frame {1}"), FText::AsNumber(NextFrameNumber), FText::AsNumber(NumFrames));
+					MenuBuilder.AddMenuEntry(Label, LOCTEXT("RemoveAfter_ToolTip", "Remove sequence after current position"), FSlateIcon(), Action);
+				}
+
+				MenuBuilder.AddMenuSeparator();
+
+				//Corrected frame time based on selected frame number
+				float CorrectedFrameTime = CurrentFrameFraction * SequenceLength;
+
+				Action = FUIAction(FExecuteAction::CreateSP(this, &SAnimTimeline::OnInsertAnimSequence, true, CurrentFrameNumber));
+				Label = FText::Format(LOCTEXT("InsertBeforeCurrentFrame", "Insert frame before {0}"), FText::AsNumber(CurrentFrameNumber));
+				MenuBuilder.AddMenuEntry(Label, LOCTEXT("InsertBefore_ToolTip", "Insert a frame before current position"), FSlateIcon(), Action);
+
+				Action = FUIAction(FExecuteAction::CreateSP(this, &SAnimTimeline::OnInsertAnimSequence, false, CurrentFrameNumber));
+				Label = FText::Format(LOCTEXT("InsertAfterCurrentFrame", "Insert frame after {0}"), FText::AsNumber(CurrentFrameNumber));
+				MenuBuilder.AddMenuEntry(Label, LOCTEXT("InsertAfter_ToolTip", "Insert a frame after current position"), FSlateIcon(), Action);
+
+				MenuBuilder.AddMenuSeparator();
+
+				//Corrected frame time based on selected frame number
+				Action = FUIAction(FExecuteAction::CreateSP(this, &SAnimTimeline::OnShowPopupOfAppendAnimation, true));
+				MenuBuilder.AddMenuEntry(LOCTEXT("AppendBegin", "Append in the beginning"), LOCTEXT("AppendBegin_ToolTip", "Append in the beginning"), FSlateIcon(), Action);
+
+				Action = FUIAction(FExecuteAction::CreateSP(this, &SAnimTimeline::OnShowPopupOfAppendAnimation, false));
+				MenuBuilder.AddMenuEntry(LOCTEXT("AppendEnd", "Append at the end"), LOCTEXT("AppendEnd_ToolTip", "Append at the end"), FSlateIcon(), Action);
+
+				MenuBuilder.AddMenuSeparator();
+				//Menu - "ReZero"
+				Action = FUIAction(FExecuteAction::CreateSP(this, &SAnimTimeline::OnReZeroAnimSequence, CurrentFrameNumber));
+				Label = FText::Format(LOCTEXT("ReZeroAtFrame", "Re-zero at frame {0}"), FText::AsNumber(CurrentFrameNumber));
+				MenuBuilder.AddMenuEntry(Label, FText::Format(LOCTEXT("ReZeroAtFrame_ToolTip", "Resets the root track to (0, 0, 0) at frame {0} and apply the difference to all root transform of the sequence. It moves whole sequence to the amount of current root transform."), FText::AsNumber(CurrentFrameNumber)), FSlateIcon(), Action);
+
+				const int32 FrameNumberForCurrentTime = INDEX_NONE;
+				Action = FUIAction(FExecuteAction::CreateSP(this, &SAnimTimeline::OnReZeroAnimSequence, FrameNumberForCurrentTime));
+				Label = LOCTEXT("ReZeroAtCurrentTime", "Re-zero at current time");
+				MenuBuilder.AddMenuEntry(Label, LOCTEXT("ReZeroAtCurrentTime_ToolTip", "Resets the root track to (0, 0, 0) at the animation scrub time and apply the difference to all root transform of the sequence. It moves whole sequence to the amount of current root transform."), FSlateIcon(), Action);
+			}
+			MenuBuilder.EndSection();
+		}
+
+		FWidgetPath WidgetPath = MouseEvent.GetEventPath() != nullptr ? *MouseEvent.GetEventPath() : FWidgetPath();
+
+		FSlateApplication::Get().PushMenu(SharedThis(this), WidgetPath, MenuBuilder.MakeWidget(), FSlateApplication::Get().GetCursorPos(), FPopupTransitionEffect(FPopupTransitionEffect::ContextMenu));
+
+		return FReply::Handled();
+	}
+
+	return FReply::Unhandled();
+}
+
+void SAnimTimeline::OnCropAnimSequence( bool bFromStart, float CurrentTime )
+{
+	UAnimSingleNodeInstance* PreviewInstance = GetPreviewInstance();
+	if(PreviewInstance)
+	{
+		float Length = PreviewInstance->GetLength();
+		if (PreviewInstance->GetCurrentAsset())
+		{
+			UAnimSequence* AnimSequence = Cast<UAnimSequence>( PreviewInstance->GetCurrentAsset() );
+			if( AnimSequence )
+			{
+				const FScopedTransaction Transaction( LOCTEXT("CropAnimSequence", "Crop Animation Sequence") );
+
+				//Call modify to restore slider position
+				PreviewInstance->Modify();
+
+				//Call modify to restore anim sequence current state
+				AnimSequence->Modify();
+
+				// Crop the raw anim data.
+				AnimSequence->CropRawAnimData( CurrentTime, bFromStart );
+
+				//Resetting slider position to the first frame
+				PreviewInstance->SetPosition( 0.0f, false );
+
+				Model.Pin()->RefreshTracks();
+			}
+		}
+	}
+}
+
+void SAnimTimeline::OnAppendAnimSequence( bool bFromStart, int32 NumOfFrames )
+{
+	UAnimSingleNodeInstance* PreviewInstance = GetPreviewInstance();
+	if(PreviewInstance && PreviewInstance->GetCurrentAsset())
+	{
+		UAnimSequence* AnimSequence = Cast<UAnimSequence>(PreviewInstance->GetCurrentAsset());
+		if(AnimSequence)
+		{
+			const FScopedTransaction Transaction(LOCTEXT("InsertAnimSequence", "Insert Animation Sequence"));
+
+			//Call modify to restore slider position
+			PreviewInstance->Modify();
+
+			//Call modify to restore anim sequence current state
+			AnimSequence->Modify();
+
+			// Crop the raw anim data.
+			int32 StartFrame = (bFromStart)? 0 : AnimSequence->GetRawNumberOfFrames() - 1;
+			int32 EndFrame = StartFrame + NumOfFrames;
+			int32 CopyFrame = StartFrame;
+			AnimSequence->InsertFramesToRawAnimData(StartFrame, EndFrame, CopyFrame);
+
+			Model.Pin()->RefreshTracks();
+		}
+	}
+}
+
+void SAnimTimeline::OnInsertAnimSequence( bool bBefore, int32 CurrentFrame )
+{
+	UAnimSingleNodeInstance* PreviewInstance = GetPreviewInstance();
+	if(PreviewInstance && PreviewInstance->GetCurrentAsset())
+	{
+		UAnimSequence* AnimSequence = Cast<UAnimSequence>(PreviewInstance->GetCurrentAsset());
+		if(AnimSequence)
+		{
+			const FScopedTransaction Transaction(LOCTEXT("InsertAnimSequence", "Insert Animation Sequence"));
+
+			//Call modify to restore slider position
+			PreviewInstance->Modify();
+
+			//Call modify to restore anim sequence current state
+			AnimSequence->Modify();
+
+			// Crop the raw anim data.
+			int32 StartFrame = (bBefore)? CurrentFrame : CurrentFrame + 1;
+			int32 EndFrame = StartFrame + 1;
+			AnimSequence->InsertFramesToRawAnimData(StartFrame, EndFrame, CurrentFrame);
+
+			Model.Pin()->RefreshTracks();
+		}
+	}
+}
+
+void SAnimTimeline::OnReZeroAnimSequence(int32 FrameIndex)
+{
+	UAnimSingleNodeInstance* PreviewInstance = GetPreviewInstance();
+	if(PreviewInstance)
+	{
+		UDebugSkelMeshComponent* PreviewSkelComp =  Model.Pin()->GetPreviewScene()->GetPreviewMeshComponent();
+
+		if (PreviewInstance->GetCurrentAsset() && PreviewSkelComp )
+		{
+			UAnimSequence* AnimSequence = Cast<UAnimSequence>( PreviewInstance->GetCurrentAsset() );
+			if( AnimSequence )
+			{
+				const FScopedTransaction Transaction( LOCTEXT("ReZeroAnimation", "ReZero Animation Sequence") );
+
+				//Call modify to restore anim sequence current state
+				AnimSequence->Modify();
+
+				// As above, animations don't have any idea of hierarchy, so we don't know for sure if track 0 is the root bone's track.
+				FRawAnimSequenceTrack& RawTrack = AnimSequence->GetRawAnimationTrack(0);
+
+				// Find vector that would translate current root bone location onto origin.
+				FVector FrameTransform = FVector::ZeroVector;
+				if (FrameIndex == INDEX_NONE)
+				{
+					// Use current transform
+					FrameTransform = PreviewSkelComp->GetComponentSpaceTransforms()[0].GetLocation();
+				}
+				else if(RawTrack.PosKeys.IsValidIndex(FrameIndex))
+				{
+					// Use transform at frame
+					FrameTransform = RawTrack.PosKeys[FrameIndex];
+				}
+
+				FVector ApplyTranslation = -1.f * FrameTransform;
+
+				// Convert into world space
+				FVector WorldApplyTranslation = PreviewSkelComp->GetComponentTransform().TransformVector(ApplyTranslation);
+				ApplyTranslation = PreviewSkelComp->GetComponentTransform().InverseTransformVector(WorldApplyTranslation);
+
+				for(int32 i=0; i<RawTrack.PosKeys.Num(); i++)
+				{
+					RawTrack.PosKeys[i] += ApplyTranslation;
+				}
+
+				// Handle Raw Data changing
+				AnimSequence->MarkRawDataAsModified();
+				AnimSequence->OnRawDataChanged();
+
+				AnimSequence->MarkPackageDirty();
+
+				Model.Pin()->RefreshTracks();
+			}
+		}
+	}
+}
+
+void SAnimTimeline::OnShowPopupOfAppendAnimation(bool bBegin)
+{
+	TSharedRef<STextEntryPopup> TextEntry =
+		SNew(STextEntryPopup)
+		.Label(LOCTEXT("AppendAnim_AskNumFrames", "Number of Frames to Append"))
+		.OnTextCommitted(this, &SAnimTimeline::OnSequenceAppendedCalled, bBegin);
+
+	// Show dialog to enter new track name
+	FSlateApplication::Get().PushMenu(
+		SharedThis(this),
+		FWidgetPath(),
+		TextEntry,
+		FSlateApplication::Get().GetCursorPos(),
+		FPopupTransitionEffect(FPopupTransitionEffect::TypeInPopup)
+		);
+}
+
+void SAnimTimeline::OnSequenceAppendedCalled(const FText & InNewGroupText, ETextCommit::Type CommitInfo, bool bBegin)
+{
+	// just a concern
+	const static int32 MaxFrame = 1000;
+
+	// handle only onEnter. This is a big thing to apply when implicit focus change or any other event
+	if (CommitInfo == ETextCommit::OnEnter)
+	{
+		int32 NumFrames = FCString::Atoi(*InNewGroupText.ToString());
+		if (NumFrames > 0 && NumFrames < MaxFrame)
+		{
+			OnAppendAnimSequence(bBegin, NumFrames);
+			FSlateApplication::Get().DismissAllMenus();
+		}
+	}
+}
+
+TSharedRef<INumericTypeInterface<double>> SAnimTimeline::GetNumericTypeInterface() const
+{
+	return NumericTypeInterface.ToSharedRef();
+}
+
+// FFrameRate::ComputeGridSpacing doesnt deal well with prime numbers, so we have a custom impl here
+static bool ComputeGridSpacing(const FFrameRate& InFrameRate, float PixelsPerSecond, double& OutMajorInterval, int32& OutMinorDivisions, float MinTickPx, float DesiredMajorTickPx)
+{
+	const int32 RoundedFPS = FMath::RoundToInt(InFrameRate.AsDecimal());
+
+	// Showing frames
+	TArray<int32, TInlineAllocator<10>> CommonBases;
+
+	// Divide the rounded frame rate by 2s, 3s or 5s recursively
+	{
+		const int32 Denominators[] = { 2, 3, 5 };
+
+		int32 LowestBase = RoundedFPS;
+		for (;;)
+		{
+			CommonBases.Add(LowestBase);
+	
+			if (LowestBase % 2 == 0)      { LowestBase = LowestBase / 2; }
+			else if (LowestBase % 3 == 0) { LowestBase = LowestBase / 3; }
+			else if (LowestBase % 5 == 0) { LowestBase = LowestBase / 5; }
+			else
+			{ 
+				int32 LowestResult = LowestBase;
+				for(int32 Denominator : Denominators)
+				{
+					int32 Result = LowestBase / Denominator;
+					if(Result > 0 && Result < LowestResult)
+					{
+						LowestResult = Result;
+					}
+				}
+
+				if(LowestResult < LowestBase)
+				{
+					LowestBase = LowestResult;
+				}
+				else
+				{
+					break;
+				}
+			}
+		}
+	}
+
+	Algo::Reverse(CommonBases);
+
+	const int32 Scale     = FMath::CeilToInt(DesiredMajorTickPx / PixelsPerSecond * InFrameRate.AsDecimal());
+	const int32 BaseIndex = FMath::Min(Algo::LowerBound(CommonBases, Scale), CommonBases.Num()-1);
+	const int32 Base      = CommonBases[BaseIndex];
+
+	int32 MajorIntervalFrames = FMath::CeilToInt(Scale / float(Base)) * Base;
+	OutMajorInterval  = MajorIntervalFrames * InFrameRate.AsInterval();
+
+	// Find the lowest number of divisions we can show that's larger than the minimum tick size
+	OutMinorDivisions = MajorIntervalFrames;
+	for (int32 DivIndex = 0; DivIndex < BaseIndex; ++DivIndex)
+	{
+		if (Base % CommonBases[DivIndex] == 0)
+		{
+			int32 MinorDivisions = MajorIntervalFrames/CommonBases[DivIndex];
+			if (OutMajorInterval / MinorDivisions * PixelsPerSecond >= MinTickPx)
+			{
+				OutMinorDivisions = MinorDivisions;
+				break;
+			}
+		}
+	}
+
+	return OutMajorInterval != 0;
+}
+
+bool SAnimTimeline::GetGridMetrics(float PhysicalWidth, double& OutMajorInterval, int32& OutMinorDivisions) const
+{
+	FSlateFontInfo SmallLayoutFont = FCoreStyle::GetDefaultFontStyle("Regular", 8);
+	TSharedRef<FSlateFontMeasure> FontMeasureService = FSlateApplication::Get().GetRenderer()->GetFontMeasureService();
+
+	FFrameRate DisplayRate(FMath::RoundToInt(Model.Pin()->GetFrameRate()), 1);
+	double BiggestTime = ViewRange.Get().GetUpperBoundValue();
+	FString TickString = NumericTypeInterface->ToString((BiggestTime * DisplayRate).FrameNumber.Value);
+	FVector2D MaxTextSize = FontMeasureService->Measure(TickString, SmallLayoutFont);
+
+	static float MajorTickMultiplier = 2.f;
+
+	float MinTickPx = MaxTextSize.X + 5.f;
+	float DesiredMajorTickPx = MaxTextSize.X * MajorTickMultiplier;
+
+	if (PhysicalWidth > 0)
+	{
+		return ComputeGridSpacing(
+			DisplayRate,
+			PhysicalWidth / ViewRange.Get().Size<double>(),
+			OutMajorInterval,
+			OutMinorDivisions,
+			MinTickPx,
+			DesiredMajorTickPx);
+	}
+
+	return false;
+}
+
+TSharedPtr<ITimeSliderController> SAnimTimeline::GetTimeSliderController() const 
+{ 
+	return TimeSliderController; 
+}
+
+void SAnimTimeline::OnOutlinerSearchChanged( const FText& Filter )
+{
+	FilterText = Filter;
+
+	Outliner->RefreshFilter();
+}
+
+void SAnimTimeline::OnColumnFillCoefficientChanged(float FillCoefficient, int32 ColumnIndex)
+{
+	ColumnFillCoefficients[ColumnIndex] = FillCoefficient;
+}
+
+void SAnimTimeline::HandleKeyComplete()
+{
+	Model.Pin()->RefreshTracks();
+}
+
+class UAnimSingleNodeInstance* SAnimTimeline::GetPreviewInstance() const
+{
+	UDebugSkelMeshComponent* PreviewMeshComponent = Model.Pin()->GetPreviewScene()->GetPreviewMeshComponent();
+	return PreviewMeshComponent && PreviewMeshComponent->IsPreviewOn()? PreviewMeshComponent->PreviewInstance : nullptr;
+}
+
+void SAnimTimeline::HandleScrubPositionChanged(FFrameTime NewScrubPosition, bool bIsScrubbing)
+{
+	if (UAnimSingleNodeInstance* PreviewInstance = GetPreviewInstance())
+	{
+		if(PreviewInstance->IsPlaying())
+		{
+			PreviewInstance->SetPlaying(false);
+		}
+	}
+
+	Model.Pin()->SetScrubPosition(NewScrubPosition);
+}
+
+double SAnimTimeline::GetSpinboxDelta() const
+{
+	return FFrameRate(Model.Pin()->GetTickResolution(), 1).AsDecimal() * FFrameRate(FMath::RoundToInt(Model.Pin()->GetFrameRate()), 1).AsInterval();
+}
+
+void SAnimTimeline::SetPlayTime(double InFrameTime)
+{
+	if (UAnimSingleNodeInstance* PreviewInstance = GetPreviewInstance())
+	{
+		PreviewInstance->SetPlaying(false);
+		PreviewInstance->SetPosition(InFrameTime / (double)Model.Pin()->GetTickResolution());
+	}
+}
+
+#undef LOCTEXT_NAMESPACE
diff --git a/Engine/Source/Editor/Persona/Private/AnimationEditorPreviewScene.cpp b/Engine/Source/Editor/Persona/Private/AnimationEditorPreviewScene.cpp
--- a/Engine/Source/Editor/Persona/Private/AnimationEditorPreviewScene.cpp
+++ b/Engine/Source/Editor/Persona/Private/AnimationEditorPreviewScene.cpp
@@ -24,6 +24,7 @@
 #include "Particles/ParticleSystemComponent.h"
 #include "Factories/PreviewMeshCollectionFactory.h"
 #include "AnimPreviewAttacheInstance.h"
+#include "AnimCustomInstanceHelper.h"
 #include "Animation/PreviewCollectionInterface.h"
 #include "ScopedTransaction.h"
 #include "Preferences/PersonaOptions.h"
@@ -329,9 +330,9 @@ void FAnimationEditorPreviewScene::RefreshAdditionalMeshes(bool bAllowOverrideBa
 	for (USkeletalMeshComponent* Component : AdditionalMeshes)
 	{
 		const UAnimInstance* AnimInst = Component->GetAnimInstance();
-		if (AnimInst && AnimInst->IsA(UAnimCustomInstance::StaticClass()))
+		if (AnimInst && AnimInst->IsA(UAnimPreviewAttacheInstance::StaticClass()))
 		{
-			UAnimCustomInstance::UnbindFromSkeletalMeshComponent(Component);
+			FAnimCustomInstanceHelper::UnbindFromSkeletalMeshComponent<UAnimPreviewAttacheInstance>(Component);
 		}
 		
 		RemoveComponent(Component);
@@ -387,7 +388,7 @@ void FAnimationEditorPreviewScene::RefreshAdditionalMeshes(bool bAllowOverrideBa
 						else
 						{
 							bool bWasCreated = false;
-							UAnimCustomInstance::BindToSkeletalMeshComponent<UAnimPreviewAttacheInstance>(NewComp,bWasCreated);
+							FAnimCustomInstanceHelper::BindToSkeletalMeshComponent<UAnimPreviewAttacheInstance>(NewComp,bWasCreated);
 						}
 						AdditionalMeshes.Add(NewComp);
 					}
@@ -556,6 +557,12 @@ void FAnimationEditorPreviewScene::RemoveAttachedComponent( bool bRemovePreviewA
 				bRemove = false;
 			}
 
+			// you can use delegate to avoid being removed
+			if (OnRemoveAttachedComponentFilter.IsBound() && !OnRemoveAttachedComponentFilter.Execute(ChildComponent))
+			{
+				bRemove = false;
+			}
+
 			if(bRemove)
 			{
 				// PreviewComponet will be cleaned up by PreviewScene, 
diff --git a/Engine/Source/Editor/Persona/Private/AnimationEditorPreviewScene.h b/Engine/Source/Editor/Persona/Private/AnimationEditorPreviewScene.h
--- a/Engine/Source/Editor/Persona/Private/AnimationEditorPreviewScene.h
+++ b/Engine/Source/Editor/Persona/Private/AnimationEditorPreviewScene.h
@@ -341,6 +341,18 @@ private:
 
 	TSharedPtr<class IEditableSkeleton> GetEditableSkeleton() const { return EditableSkeletonPtr.Pin(); }
 
+	/* setter/getter for can remove attach component */
+	virtual void SetRemoveAttachedComponentFilter(const FOnRemoveAttachedComponentFilter& Delegate) override
+	{
+		OnRemoveAttachedComponentFilter = Delegate;
+	}
+	virtual void ClearRemoveAttachedComponentFilter() override
+	{
+		if (OnRemoveAttachedComponentFilter.IsBound())
+		{
+			OnRemoveAttachedComponentFilter.Unbind();
+		}
+	}
 private:
 	/** The one and only actor we have */
 	AActor* Actor;
@@ -448,4 +460,7 @@ private:
 
 	/** Allow additional meshes to be selectable */
 	bool bAdditionalMeshesSelectable;
+
+	/** Delegate Remove attach component */
+	FOnRemoveAttachedComponentFilter OnRemoveAttachedComponentFilter;
 };
diff --git a/Engine/Source/Editor/Persona/Private/AnimationEditorViewportClient.cpp b/Engine/Source/Editor/Persona/Private/AnimationEditorViewportClient.cpp
--- a/Engine/Source/Editor/Persona/Private/AnimationEditorViewportClient.cpp
+++ b/Engine/Source/Editor/Persona/Private/AnimationEditorViewportClient.cpp
@@ -1481,16 +1481,6 @@ void FAnimationViewportClient::DrawSockets(const UDebugSkelMeshComponent* InPrev
 
 FSphere FAnimationViewportClient::GetCameraTarget()
 {
-	const FSphere DefaultSphere(FVector(0,0,0), 100.0f);
-
-	UDebugSkelMeshComponent* PreviewMeshComponent = GetAnimPreviewScene()->GetPreviewMeshComponent();
-	if( !PreviewMeshComponent )
-	{
-		return DefaultSphere;
-	}
-
-	PreviewMeshComponent->CalcBounds(PreviewMeshComponent->GetComponentTransform());
-
 	// give the editor mode a chance to give us a camera target
 	if (GetPersonaModeManager())
 	{
@@ -1501,6 +1491,16 @@ FSphere FAnimationViewportClient::GetCameraTarget()
 		}
 	}
 
+	const FSphere DefaultSphere(FVector(0,0,0), 100.0f);
+
+	UDebugSkelMeshComponent* PreviewMeshComponent = GetAnimPreviewScene()->GetPreviewMeshComponent();
+	if( !PreviewMeshComponent )
+	{
+		return DefaultSphere;
+	}
+
+	PreviewMeshComponent->CalcBounds(PreviewMeshComponent->GetComponentTransform());
+
 	FBoxSphereBounds Bounds = PreviewMeshComponent->CalcBounds(FTransform::Identity);
 	return Bounds.GetSphere();
 }
@@ -1629,41 +1629,35 @@ void FAnimationViewportClient::FocusViewportOnPreviewMesh(bool bUseCustomCamera)
 		return;
 	}
 
-	UDebugSkelMeshComponent* const PreviewMeshComponent = GetAnimPreviewScene()->GetPreviewMeshComponent();
-	if ( !PreviewMeshComponent )
+	if (UDebugSkelMeshComponent* const PreviewMeshComponent = GetAnimPreviewScene()->GetPreviewMeshComponent())
 	{
-		return;
-	}
+		if (USkeletalMesh* const SkelMesh = PreviewMeshComponent->SkeletalMesh)
+		{
+			if (bUseCustomCamera && SkelMesh->bHasCustomDefaultEditorCamera)
+			{
+				FViewportCameraTransform& ViewTransform = GetViewTransform();
 
-	USkeletalMesh* const SkelMesh = PreviewMeshComponent->SkeletalMesh;
-	if (!SkelMesh)
-	{
-		return;
-	}
+				ViewTransform.SetLocation(SkelMesh->DefaultEditorCameraLocation);
+				ViewTransform.SetRotation(SkelMesh->DefaultEditorCameraRotation);
+				ViewTransform.SetLookAt(SkelMesh->DefaultEditorCameraLookAt);
+				ViewTransform.SetOrthoZoom(SkelMesh->DefaultEditorCameraOrthoZoom);
 
-	if (bUseCustomCamera && SkelMesh->bHasCustomDefaultEditorCamera)
-	{
-		FViewportCameraTransform& ViewTransform = GetViewTransform();
+				Invalidate();
+				return;
+			}
 
-		ViewTransform.SetLocation(SkelMesh->DefaultEditorCameraLocation);
-		ViewTransform.SetRotation(SkelMesh->DefaultEditorCameraRotation);
-		ViewTransform.SetLookAt(SkelMesh->DefaultEditorCameraLookAt);
-		ViewTransform.SetOrthoZoom(SkelMesh->DefaultEditorCameraOrthoZoom);
+			if (PreviewMeshComponent->GetSelectedEditorSection() != INDEX_NONE)
+			{
+				const FBox SelectedSectionBounds = ComputeBoundingBoxForSelectedEditorSection();
 
-		Invalidate();
-		return;
-	}
+				if (SelectedSectionBounds.IsValid)
+				{
+					FocusViewportOnBox(SelectedSectionBounds);
+				}
 
-	if (PreviewMeshComponent->GetSelectedEditorSection() != INDEX_NONE )
-	{
-		const FBox SelectedSectionBounds = ComputeBoundingBoxForSelectedEditorSection();
-		
-		if ( SelectedSectionBounds.IsValid )
-		{
-			FocusViewportOnBox(SelectedSectionBounds);
+				return;
+			}
 		}
-
-		return;
 	}
 
 	FSphere Sphere = GetCameraTarget();
diff --git a/Engine/Source/Editor/Persona/Private/AssetSearchBoxUtilPersona.cpp b/Engine/Source/Editor/Persona/Private/AssetSearchBoxUtilPersona.cpp
--- a/Engine/Source/Editor/Persona/Private/AssetSearchBoxUtilPersona.cpp
+++ b/Engine/Source/Editor/Persona/Private/AssetSearchBoxUtilPersona.cpp
@@ -5,6 +5,23 @@
 #include "Animation/Skeleton.h"
 #include "Engine/SkeletalMesh.h"
 #include "Engine/SkeletalMeshSocket.h"
+#include "Editor.h"
+
+SAssetSearchBoxForBones::SAssetSearchBoxForBones()
+{
+	if(GEditor)
+	{
+		GEditor->RegisterForUndo(this);
+	}
+}
+
+SAssetSearchBoxForBones::~SAssetSearchBoxForBones()
+{
+	if(GEditor)
+	{
+		GEditor->UnregisterForUndo(this);
+	}
+}
 
 void SAssetSearchBoxForBones::Construct( const FArguments& InArgs, const class UObject* Outer, TSharedPtr<class IPropertyHandle> BoneNameProperty )
 {
@@ -79,26 +96,35 @@ FText SAssetSearchBoxForBones::GetBoneName() const
 	return CurValue.IsNone() ? FText::GetEmpty() : FText::FromName(CurValue);
 }
 
-void SAssetSearchBoxForCurves::Construct(const FArguments& InArgs, const class USkeleton* InSkeleton, TSharedPtr<class IPropertyHandle> CurveNameProperty)
+SAssetSearchBoxForCurves::SAssetSearchBoxForCurves()
 {
-	check(InSkeleton);
+	if(GEditor)
+	{
+		GEditor->RegisterForUndo(this);
+	}
+}
 
-	// get the currently chosen curve, if any
-	const FText PropertyName = CurveNameProperty->GetPropertyDisplayName();
-	FString CurValue;
-	CurveNameProperty->GetValue(CurValue);
-	if (CurValue == FString("None"))
+SAssetSearchBoxForCurves::~SAssetSearchBoxForCurves()
+{
+	if(GEditor)
 	{
-		CurValue.Empty();
+		GEditor->UnregisterForUndo(this);
 	}
+}
+
+void SAssetSearchBoxForCurves::Construct(const FArguments& InArgs, const class USkeleton* InSkeleton, TSharedPtr<class IPropertyHandle> CurveNameProperty)
+{
+	check(InSkeleton);
+
+	CurveNamePropertyHandle = CurveNameProperty;
 
 	Skeleton = MakeWeakObjectPtr(const_cast<USkeleton*>(InSkeleton));
 	
 	// create the asset search box
 	ChildSlot
 		[
-			SNew(SAssetSearchBox)
-			.InitialText(FText::FromString(CurValue))
+			SAssignNew(SearchBox, SAssetSearchBox)
+			.InitialText(GetCurveName())
 			.HintText(InArgs._HintText)
 			.OnTextCommitted(InArgs._OnTextCommitted)
 			.PossibleSuggestions(this, &SAssetSearchBoxForCurves::GetCurveSearchSuggestions)
@@ -107,6 +133,24 @@ void SAssetSearchBoxForCurves::Construct(const FArguments& InArgs, const class U
 		];
 }
 
+FText SAssetSearchBoxForCurves::GetCurveName() const
+{
+	FName CurValue;
+	if (CurveNamePropertyHandle.IsValid())
+	{
+		CurveNamePropertyHandle->GetValue(CurValue);
+	}
+	return CurValue.IsNone() ? FText::GetEmpty() : FText::FromName(CurValue);
+}
+
+void SAssetSearchBoxForCurves::RefreshName()
+{
+	if (SearchBox.IsValid())
+	{
+		SearchBox->SetText(GetCurveName());
+	}
+}
+
 TArray<FAssetSearchBoxSuggestion> SAssetSearchBoxForCurves::GetCurveSearchSuggestions() const
 {
 	TArray<FAssetSearchBoxSuggestion> PossibleSuggestions;
diff --git a/Engine/Source/Editor/Persona/Private/AssetSearchBoxUtilPersona.h b/Engine/Source/Editor/Persona/Private/AssetSearchBoxUtilPersona.h
--- a/Engine/Source/Editor/Persona/Private/AssetSearchBoxUtilPersona.h
+++ b/Engine/Source/Editor/Persona/Private/AssetSearchBoxUtilPersona.h
@@ -8,10 +8,11 @@
 #include "Framework/SlateDelegates.h"
 #include "PropertyHandle.h"
 #include "SAssetSearchBox.h"
+#include "EditorUndoClient.h"
 
 class USkeleton;
 
-class SAssetSearchBoxForBones : public SCompoundWidget
+class SAssetSearchBoxForBones : public SCompoundWidget, public FEditorUndoClient
 {
 public:
 	SLATE_BEGIN_ARGS(SAssetSearchBoxForBones)
@@ -38,19 +39,28 @@ public:
 
 		SLATE_END_ARGS()
 
+		SAssetSearchBoxForBones();
+		~SAssetSearchBoxForBones();
+
 		/** Constructs this widget with InArgs */
 		void Construct( const FArguments& InArgs, const class UObject* Outer, TSharedPtr<class IPropertyHandle> ParentBoneProperty );
 
 		/** Refresh bone name */
 		void RefreshName();
+
+		virtual void PostUndo( bool bSuccess ) { RefreshName(); };
+		virtual void PostRedo( bool bSuccess ) { RefreshName(); };
+
 private:
 		TSharedPtr<IPropertyHandle>	BonePropertyHandle;
 		TSharedPtr<SAssetSearchBox>			SearchBox;
+
+		/** Get the bone name to display */
 		FText GetBoneName() const;
 };
 
 
-class SAssetSearchBoxForCurves : public SCompoundWidget
+class SAssetSearchBoxForCurves : public SCompoundWidget, public FEditorUndoClient
 {
 public:
 	SLATE_BEGIN_ARGS(SAssetSearchBoxForCurves)
@@ -77,12 +87,26 @@ public:
 
 	SLATE_END_ARGS()
 
+	SAssetSearchBoxForCurves();
+	~SAssetSearchBoxForCurves();
+
 	/** Constructs this widget with InArgs */
 	void Construct(const FArguments& InArgs, const class USkeleton* InSkeleton, TSharedPtr<class IPropertyHandle> ParentCurveProperty);
 
+	/** Refresh curve name */
+	void RefreshName();
+
+	virtual void PostUndo( bool bSuccess ) { RefreshName(); };
+	virtual void PostRedo( bool bSuccess ) { RefreshName(); };
+
 private:
+	TSharedPtr<IPropertyHandle>	CurveNamePropertyHandle;
 	TWeakObjectPtr<USkeleton> Skeleton;
+	TSharedPtr<SAssetSearchBox>			SearchBox;
 
 	/** Get the search suggestions */
 	TArray<FAssetSearchBoxSuggestion> GetCurveSearchSuggestions() const;
+
+	/** Get the curve name to display */
+	FText GetCurveName() const;
 };
diff --git a/Engine/Source/Editor/Persona/Private/PersonaModule.cpp b/Engine/Source/Editor/Persona/Private/PersonaModule.cpp
--- a/Engine/Source/Editor/Persona/Private/PersonaModule.cpp
+++ b/Engine/Source/Editor/Persona/Private/PersonaModule.cpp
@@ -47,7 +47,7 @@
 #include "IAssetTools.h"
 #include "AssetToolsModule.h"
 #include "Logging/MessageLog.h"
-#include "AnimationCompressionPanel.h"
+#include "AnimationEditorUtils.h"
 #include "DesktopPlatformModule.h"
 #include "FbxAnimUtils.h"
 #include "PersonaAssetFamilyManager.h"
@@ -70,6 +70,7 @@
 #include "Kismet2/KismetEditorUtilities.h"
 #include "Animation/AnimNotifies/AnimNotify.h"
 #include "Animation/AnimNotifies/AnimNotifyState.h"
+#include "Animation/AnimBoneCompressionSettings.h"
 #include "Kismet2/BlueprintEditorUtils.h"
 #include "Widgets/Notifications/SNotificationList.h"
 #include "Framework/Notifications/NotificationManager.h"
@@ -87,6 +88,9 @@
 #include "ISequenceRecorder.h"
 #include "SkinWeightProfileCustomization.h"
 #include "SAnimationBlendSpaceGridWidget.h"
+#include "SAnimSequenceCurveEditor.h"
+#include "AnimSequenceTimelineCommands.h"
+#include "SAnimMontageSectionsPanel.h"
 
 IMPLEMENT_MODULE( FPersonaModule, Persona );
 
@@ -156,6 +160,7 @@ void FPersonaModule::StartupModule()
 	FEditorModeRegistry::Get().RegisterMode<FSkeletonSelectionEditMode>(FPersonaEditModes::SkeletonSelection, LOCTEXT("SkeletonSelectionEditMode", "Skeleton Selection"), FSlateIcon(), false);
 
 	FPersonaCommonCommands::Register();
+	FAnimSequenceTimelineCommands::Register();
 
 	FKismetEditorUtilities::RegisterOnBlueprintCreatedCallback(this, UAnimNotify::StaticClass(), FKismetEditorUtilities::FOnBlueprintCreated::CreateRaw(this, &FPersonaModule::HandleNewAnimNotifyBlueprintCreated));
 	FKismetEditorUtilities::RegisterOnBlueprintCreatedCallback(this, UAnimNotifyState::StaticClass(), FKismetEditorUtilities::FOnBlueprintCreated::CreateRaw(this, &FPersonaModule::HandleNewAnimNotifyStateBlueprintCreated));
@@ -356,45 +361,71 @@ TSharedRef<SWidget> FPersonaModule::CreateBlendSpacePreviewWidget(TAttribute<con
 		.ShowSettingsButtons(false);
 }
 
-TSharedRef<SWidget> FPersonaModule::CreateEditorWidgetForAnimDocument(const TSharedRef<class FWorkflowCentricApplication>& InHostingApp, UObject* InAnimAsset, const FAnimDocumentArgs& InArgs, FString& OutDocumentLink)
+TSharedRef<class FWorkflowTabFactory> FPersonaModule::CreateAnimMontageSectionsTabFactory(const TSharedRef<class FWorkflowCentricApplication>& InHostingApp, const TSharedRef<IPersonaToolkit>& InPersonaToolkit, FSimpleMulticastDelegate& InOnSectionsChanged) const
+{
+	return MakeShareable(new FAnimMontageSectionsSummoner(InHostingApp, InPersonaToolkit, InOnSectionsChanged));
+}
+
+TSharedRef<SWidget> FPersonaModule::CreateEditorWidgetForAnimDocument(const TSharedRef<IAnimationEditor>& InHostingApp, UObject* InAnimAsset, const FAnimDocumentArgs& InArgs, FString& OutDocumentLink)
 {
 	TSharedPtr<SWidget> Result = SNullWidget::NullWidget;
 	if (InAnimAsset)
 	{
+		TWeakPtr<IAnimationEditor> WeakHostingApp = InHostingApp;
+		auto OnEditCurves = [WeakHostingApp](UAnimSequenceBase* InAnimSequence, const TArray<IAnimationEditor::FCurveEditInfo>& InCurveInfo, const TSharedPtr<ITimeSliderController>& InExternalTimeSliderController)
+		{ 
+			WeakHostingApp.Pin()->EditCurves(InAnimSequence, InCurveInfo, InExternalTimeSliderController);
+		};
+
+		auto OnStopEditingCurves = [WeakHostingApp](const TArray<IAnimationEditor::FCurveEditInfo>& InCurveInfo)
+		{ 
+			WeakHostingApp.Pin()->StopEditingCurves(InCurveInfo);
+		};
+
 		if (UAnimSequence* Sequence = Cast<UAnimSequence>(InAnimAsset))
 		{
-			Result = SNew(SSequenceEditor, InArgs.PreviewScene.Pin().ToSharedRef(), InArgs.EditableSkeleton.Pin().ToSharedRef())
+			Result = SNew(SSequenceEditor, InArgs.PreviewScene.Pin().ToSharedRef(), InArgs.EditableSkeleton.Pin().ToSharedRef(), InHostingApp->GetToolkitCommands())
 				.Sequence(Sequence)
 				.OnObjectsSelected(InArgs.OnDespatchObjectsSelected)
-				.OnInvokeTab(InArgs.OnDespatchInvokeTab);
+				.OnInvokeTab(InArgs.OnDespatchInvokeTab)
+				.OnEditCurves_Lambda(OnEditCurves)
+				.OnStopEditingCurves_Lambda(OnStopEditingCurves);
 
 			OutDocumentLink = TEXT("Engine/Animation/Sequences");
 		}
 		else if (UAnimComposite* Composite = Cast<UAnimComposite>(InAnimAsset))
 		{
-			Result = SNew(SAnimCompositeEditor, InArgs.PreviewScene.Pin().ToSharedRef(), InArgs.EditableSkeleton.Pin().ToSharedRef())
+			Result = SNew(SAnimCompositeEditor, InArgs.PreviewScene.Pin().ToSharedRef(), InArgs.EditableSkeleton.Pin().ToSharedRef(), InHostingApp->GetToolkitCommands())
 				.Composite(Composite)
 				.OnObjectsSelected(InArgs.OnDespatchObjectsSelected)
-				.OnInvokeTab(InArgs.OnDespatchInvokeTab);
+				.OnInvokeTab(InArgs.OnDespatchInvokeTab)
+				.OnEditCurves_Lambda(OnEditCurves)
+				.OnStopEditingCurves_Lambda(OnStopEditingCurves);
 
 			OutDocumentLink = TEXT("Engine/Animation/AnimationComposite");
 		}
 		else if (UAnimMontage* Montage = Cast<UAnimMontage>(InAnimAsset))
 		{
-			FMontageEditorRequiredArgs RequiredArgs(InArgs.PreviewScene.Pin().ToSharedRef(), InArgs.EditableSkeleton.Pin().ToSharedRef(), InArgs.OnSectionsChanged);
+			FMontageEditorRequiredArgs RequiredArgs(InArgs.PreviewScene.Pin().ToSharedRef(), InArgs.EditableSkeleton.Pin().ToSharedRef(), InArgs.OnSectionsChanged, InHostingApp->GetToolkitCommands());
 
 			Result = SNew(SMontageEditor, RequiredArgs)
 				.Montage(Montage)
 				.OnSectionsChanged(InArgs.OnDespatchSectionsChanged)
 				.OnInvokeTab(InArgs.OnDespatchInvokeTab)
-				.OnObjectsSelected(InArgs.OnDespatchObjectsSelected);
+				.OnObjectsSelected(InArgs.OnDespatchObjectsSelected)
+				.OnEditCurves_Lambda(OnEditCurves)
+				.OnStopEditingCurves_Lambda(OnStopEditingCurves);
 
 			OutDocumentLink = TEXT("Engine/Animation/AnimMontage");
 		}
 		else if (UAnimStreamable* StreamableAnim = Cast<UAnimStreamable>(InAnimAsset))
 		{
-			Result = SNew(SAnimStreamableEditor, InArgs.PreviewScene.Pin().ToSharedRef(), InArgs.EditableSkeleton.Pin().ToSharedRef())
-				.StreamableAnim(StreamableAnim);
+			Result = SNew(SAnimStreamableEditor, InArgs.PreviewScene.Pin().ToSharedRef(), InArgs.EditableSkeleton.Pin().ToSharedRef(), InHostingApp->GetToolkitCommands())
+				.StreamableAnim(StreamableAnim)
+				.OnObjectsSelected(InArgs.OnDespatchObjectsSelected)
+				.OnInvokeTab(InArgs.OnDespatchInvokeTab)
+				.OnEditCurves_Lambda(OnEditCurves)
+				.OnStopEditingCurves_Lambda(OnStopEditingCurves);
 
 			OutDocumentLink = TEXT("Engine/Animation/Sequences");
 		}
@@ -443,6 +474,13 @@ TSharedRef<SWidget> FPersonaModule::CreateEditorWidgetForAnimDocument(const TSha
 	return Result.ToSharedRef();
 }
 
+TSharedRef<IAnimSequenceCurveEditor> FPersonaModule::CreateCurveWidgetForAnimDocument(const TSharedRef<class FWorkflowCentricApplication>& InHostingApp, const TSharedRef<IPersonaPreviewScene>& InPreviewScene, UAnimSequenceBase* InAnimSequence, const TSharedPtr<ITimeSliderController>& InExternalTimeSliderController, const TSharedPtr<FTabManager>& InTabManager)
+{
+	return SNew(SAnimSequenceCurveEditor, InPreviewScene, InAnimSequence)
+		.ExternalTimeSliderController(InExternalTimeSliderController)
+		.TabManager(InTabManager);
+}
+
 void FPersonaModule::CustomizeMeshDetails(const TSharedRef<IDetailsView>& InDetailsView, const TSharedRef<IPersonaToolkit>& InPersonaToolkit)
 {
 	InDetailsView->SetGenericLayoutDetailsDelegate(FOnGetDetailCustomizationInstance::CreateStatic(&FPersonaMeshDetails::MakeInstance, TWeakPtr<IPersonaToolkit>(InPersonaToolkit)));
@@ -599,10 +637,51 @@ void FPersonaModule::TestSkeletonCurveNamesForUse(const TSharedRef<IEditableSkel
 	}
 }
 
-void FPersonaModule::ApplyCompression(TArray<TWeakObjectPtr<UAnimSequence>>& AnimSequences)
+void FPersonaModule::ApplyCompression(TArray<TWeakObjectPtr<UAnimSequence>>& AnimSequences, bool bPickBoneSettingsOverride)
 {
-	FDlgAnimCompression AnimCompressionDialog(AnimSequences);
-	AnimCompressionDialog.ShowModal();
+	UAnimBoneCompressionSettings* OverrideSettings = nullptr;
+	if (bPickBoneSettingsOverride)
+	{
+		FAnimationCompressionSelectionDialogConfig DialogConfig;
+
+		UAnimBoneCompressionSettings* CurrentSettings = nullptr;
+		if (AnimSequences.Num() != 0)
+		{
+			CurrentSettings = AnimSequences[0]->BoneCompressionSettings;
+			for (TWeakObjectPtr<UAnimSequence>& AnimSeq : AnimSequences)
+			{
+				if (CurrentSettings != AnimSeq->BoneCompressionSettings)
+				{
+					// One of the sequences in the list has a different settings asset, use the default behavior
+					CurrentSettings = nullptr;
+					break;
+				}
+			}
+		}
+
+		DialogConfig.DefaultSelectedAsset = CurrentSettings;
+
+		FAssetData AssetData = AnimationEditorUtils::CreateModalAnimationCompressionSelectionDialog(DialogConfig);
+		if (AssetData.IsValid())
+		{
+			OverrideSettings = Cast<UAnimBoneCompressionSettings>(AssetData.GetAsset());
+		}
+		else
+		{
+			// No asset selected but we need an override, do nothing
+			return;
+		}
+	}
+
+	TArray<UAnimSequence*> AnimSequencePtrs;
+	AnimSequencePtrs.Reserve(AnimSequences.Num());
+
+	for (TWeakObjectPtr<UAnimSequence>& AnimSeq : AnimSequences)
+	{
+		AnimSequencePtrs.Add(AnimSeq.Get());
+	}
+
+	AnimationEditorUtils::ApplyCompressionAlgorithm(AnimSequencePtrs, OverrideSettings);
 }
 
 bool FPersonaModule::ExportToFBX(TArray<TWeakObjectPtr<UAnimSequence>>& AnimSequences, USkeletalMesh* SkeletalMesh)
diff --git a/Engine/Source/Editor/Persona/Private/PersonaUtils.cpp b/Engine/Source/Editor/Persona/Private/PersonaUtils.cpp
--- a/Engine/Source/Editor/Persona/Private/PersonaUtils.cpp
+++ b/Engine/Source/Editor/Persona/Private/PersonaUtils.cpp
@@ -10,6 +10,12 @@
 #include "Kismet2/BlueprintEditorUtils.h"
 #include "Animation/AnimNodeBase.h"
 #include "Animation/AnimBlueprint.h"
+#include "Widgets/Text/STextBlock.h"
+#include "Styling/CoreStyle.h"
+#include "Widgets/Input/SComboButton.h"
+#include "EditorStyleSet.h"
+#include "Widgets/SBoxPanel.h"
+#include "Widgets/Images/SImage.h"
 
 namespace PersonaUtils
 {
@@ -123,4 +129,64 @@ void SetObjectBeingDebugged(UAnimBlueprint* InAnimBlueprint, UAnimInstance* InAn
 	}
 }
 
+TSharedRef<SWidget> MakeTrackButton(FText HoverText, FOnGetContent MenuContent, const TAttribute<bool>& HoverState)
+{
+	FSlateFontInfo SmallLayoutFont = FCoreStyle::GetDefaultFontStyle("Regular", 8);
+
+	TSharedRef<STextBlock> ComboButtonText = SNew(STextBlock)
+		.Text(HoverText)
+		.Font(SmallLayoutFont)
+		.ColorAndOpacity( FSlateColor::UseForeground() );
+
+	TSharedRef<SComboButton> ComboButton =
+
+		SNew(SComboButton)
+		.HasDownArrow(false)
+		.ButtonStyle(FEditorStyle::Get(), "HoverHintOnly")
+		.ForegroundColor( FSlateColor::UseForeground() )
+		.OnGetMenuContent(MenuContent)
+		.ContentPadding(FMargin(5, 2))
+		.HAlign(HAlign_Center)
+		.VAlign(VAlign_Center)
+		.ButtonContent()
+		[
+			SNew(SHorizontalBox)
+
+			+ SHorizontalBox::Slot()
+			.AutoWidth()
+			.VAlign(VAlign_Center)
+			.Padding(FMargin(0,0,2,0))
+			[
+				SNew(SImage)
+				.ColorAndOpacity( FSlateColor::UseForeground() )
+				.Image(FEditorStyle::GetBrush("ComboButton.Arrow"))
+			]
+
+			+ SHorizontalBox::Slot()
+			.VAlign(VAlign_Center)
+			.AutoWidth()
+			[
+				ComboButtonText
+			]
+		];
+
+	auto GetRolloverVisibility = [WeakComboButton = TWeakPtr<SComboButton>(ComboButton), HoverState]()
+	{
+		TSharedPtr<SComboButton> ComboButton = WeakComboButton.Pin();
+		if (HoverState.Get() || ComboButton->IsOpen())
+		{
+			return EVisibility::SelfHitTestInvisible;
+		}
+		else
+		{
+			return EVisibility::Collapsed;
+		}
+	};
+
+	TAttribute<EVisibility> Visibility = TAttribute<EVisibility>::Create(TAttribute<EVisibility>::FGetter::CreateLambda(GetRolloverVisibility));
+	ComboButtonText->SetVisibility(Visibility);
+
+	return ComboButton;
+}
+
 }
diff --git a/Engine/Source/Editor/Persona/Private/PreviewSceneCustomizations.cpp b/Engine/Source/Editor/Persona/Private/PreviewSceneCustomizations.cpp
--- a/Engine/Source/Editor/Persona/Private/PreviewSceneCustomizations.cpp
+++ b/Engine/Source/Editor/Persona/Private/PreviewSceneCustomizations.cpp
@@ -215,6 +215,21 @@ void FPreviewSceneDescriptionCustomization::CustomizeDetails(IDetailLayoutBuilde
 			.OnShouldFilterAsset(this, &FPreviewSceneDescriptionCustomization::HandleShouldFilterAsset, FName("Skeleton"), PersonaToolkit.Pin()->GetContext() == UPhysicsAsset::StaticClass()->GetFName())
 			.OnObjectChanged(this, &FPreviewSceneDescriptionCustomization::HandleMeshChanged)
 			.ThumbnailPool(DetailBuilder.GetThumbnailPool())
+			.CustomResetToDefault(FResetToDefaultOverride::Create(
+				FIsResetToDefaultVisible::CreateLambda([this](TSharedPtr<IPropertyHandle> PropertyHandle) -> bool {
+					if (PreviewScene.IsValid())
+					{
+						return PreviewScene.Pin()->GetPreviewMesh() != nullptr;
+					}
+					return false;
+				}),
+				FResetToDefaultHandler::CreateLambda([this](TSharedPtr<IPropertyHandle> PropertyHandle) {
+					if (PreviewScene.IsValid())
+					{
+						PreviewScene.Pin()->SetPreviewMesh(nullptr, false);
+					}
+				})
+			))
 		];
 
 		// Customize animation blueprint preview
diff --git a/Engine/Source/Editor/Persona/Private/SAnimCompositeEditor.cpp b/Engine/Source/Editor/Persona/Private/SAnimCompositeEditor.cpp
--- a/Engine/Source/Editor/Persona/Private/SAnimCompositeEditor.cpp
+++ b/Engine/Source/Editor/Persona/Private/SAnimCompositeEditor.cpp
@@ -7,6 +7,8 @@
 
 #include "SAnimNotifyPanel.h"
 #include "Editor.h"
+#include "AnimModel_AnimComposite.h"
+#include "SAnimTimeline.h"
 
 //////////////////////////////////////////////////////////////////////////
 // SAnimCompositeEditor
@@ -16,124 +18,40 @@ TSharedRef<SWidget> SAnimCompositeEditor::CreateDocumentAnchor()
 	return IDocumentation::Get()->CreateAnchor(TEXT("Engine/Animation/AnimationComposite"));
 }
 
-void SAnimCompositeEditor::Construct(const FArguments& InArgs, const TSharedRef<class IPersonaPreviewScene>& InPreviewScene, const TSharedRef<class IEditableSkeleton>& InEditableSkeleton)
+void SAnimCompositeEditor::Construct(const FArguments& InArgs, const TSharedRef<class IPersonaPreviewScene>& InPreviewScene, const TSharedRef<class IEditableSkeleton>& InEditableSkeleton, const TSharedRef<FUICommandList>& InCommandList)
 {
-	bIsActiveTimerRegistered = false;
 	CompositeObj = InArgs._Composite;
 	check(CompositeObj);
 
-	SAnimEditorBase::Construct( SAnimEditorBase::FArguments()
-		.OnObjectsSelected(InArgs._OnObjectsSelected), 
-		InPreviewScene );
+	AnimModel = MakeShared<FAnimModel_AnimComposite>(InPreviewScene, InEditableSkeleton, InCommandList, CompositeObj);
+	AnimModel->OnSelectObjects = FOnObjectsSelected::CreateSP(this, &SAnimEditorBase::OnSelectionChanged);
+	AnimModel->OnInvokeTab = InArgs._OnInvokeTab;
 
-	if(GEditor)
+	AnimModel->OnEditCurves = FOnEditCurves::CreateLambda([this, InOnEditCurves = InArgs._OnEditCurves](UAnimSequenceBase* InAnimSequence, const TArray<IAnimationEditor::FCurveEditInfo>& InCurveInfo, const TSharedPtr<ITimeSliderController>& InExternalTimeSliderController)
 	{
-		GEditor->RegisterForUndo(this);
-	}
+		InOnEditCurves.ExecuteIfBound(InAnimSequence, InCurveInfo, TimelineWidget->GetTimeSliderController());
+	});
 
-	EditorPanels->AddSlot()
-		.AutoHeight()
-		.Padding(0, 10)
-		[
-			SAssignNew( AnimCompositePanel, SAnimCompositePanel )
-			.Composite(CompositeObj)
-			.CompositeEditor(SharedThis(this))
-			.WidgetWidth(S2ColumnWidget::DEFAULT_RIGHT_COLUMN_WIDTH)
-			.ViewInputMin(this, &SAnimEditorBase::GetViewMinInput)
-			.ViewInputMax(this, &SAnimEditorBase::GetViewMaxInput)
-			.OnSetInputViewRange(this, &SAnimEditorBase::SetInputViewRange)
-		];
-
-	EditorPanels->AddSlot()
-		.AutoHeight()
-		.Padding(0, 10)
-		[
-			SAssignNew( AnimNotifyPanel, SAnimNotifyPanel, InEditableSkeleton )
-			.Sequence(CompositeObj)
-			.WidgetWidth(S2ColumnWidget::DEFAULT_RIGHT_COLUMN_WIDTH)
-			.InputMin(this, &SAnimEditorBase::GetMinInput)
-			.InputMax(this, &SAnimEditorBase::GetMaxInput)
-			.ViewInputMin(this, &SAnimEditorBase::GetViewMinInput)
-			.ViewInputMax(this, &SAnimEditorBase::GetViewMaxInput)
-			.OnSetInputViewRange(this, &SAnimEditorBase::SetInputViewRange)
-			.OnGetScrubValue(this, &SAnimEditorBase::GetScrubValue)
-			.OnSelectionChanged(this, &SAnimEditorBase::OnSelectionChanged)
-			.OnInvokeTab(InArgs._OnInvokeTab)
-		];
-
-	EditorPanels->AddSlot()
-		.AutoHeight()
-		.Padding(0, 10)
-		[
-			SAssignNew( AnimCurvePanel, SAnimCurvePanel, InEditableSkeleton)
-			.Sequence(CompositeObj)
-			.WidgetWidth(S2ColumnWidget::DEFAULT_RIGHT_COLUMN_WIDTH)
-			.ViewInputMin(this, &SAnimEditorBase::GetViewMinInput)
-			.ViewInputMax(this, &SAnimEditorBase::GetViewMaxInput)
-			.InputMin(this, &SAnimEditorBase::GetMinInput)
-			.InputMax(this, &SAnimEditorBase::GetMaxInput)
-			.OnSetInputViewRange(this, &SAnimEditorBase::SetInputViewRange)
-			.OnGetScrubValue(this, &SAnimEditorBase::GetScrubValue)
-		];
-
-	CollapseComposite();
-}
+	AnimModel->OnStopEditingCurves = InArgs._OnStopEditingCurves;
+	AnimModel->Initialize();
 
-SAnimCompositeEditor::~SAnimCompositeEditor()
-{
-	if(GEditor)
-	{
-		GEditor->UnregisterForUndo(this);
-	}
-}
-
-void SAnimCompositeEditor::PreAnimUpdate()
-{
-	CompositeObj->Modify();
-}
-
-void SAnimCompositeEditor::PostAnimUpdate()
-{
-	CompositeObj->MarkPackageDirty();
-	SortAndUpdateComposite();
-}
-
-void SAnimCompositeEditor::RebuildPanel()
-{
-	SortAndUpdateComposite();
-	AnimCompositePanel->Update();
-}
+	SAnimEditorBase::Construct( SAnimEditorBase::FArguments()
+		.OnObjectsSelected(InArgs._OnObjectsSelected)
+		.AnimModel(AnimModel), 
+		InPreviewScene);
 
-void SAnimCompositeEditor::OnCompositeChange(class UObject *EditorAnimBaseObj, bool bRebuild)
-{
-	if ( CompositeObj != nullptr )
+	if (GEditor)
 	{
-		if(bRebuild && !bIsActiveTimerRegistered)
-		{
-			// sometimes crashes because the timer delay but animation still renders, so invalidating here before calling timer
-			CompositeObj->InvalidateRecursiveAsset();
-			bIsActiveTimerRegistered = true;
-			RegisterActiveTimer(0.f, FWidgetActiveTimerDelegate::CreateSP(this, &SAnimCompositeEditor::TriggerRebuildPanel));
-		} 
-		else
-		{
-			CollapseComposite();
-		}
-
-		CompositeObj->MarkPackageDirty();
+		GEditor->RegisterForUndo(this);
 	}
 }
 
-void SAnimCompositeEditor::CollapseComposite()
+SAnimCompositeEditor::~SAnimCompositeEditor()
 {
-	if ( CompositeObj == nullptr )
+	if (GEditor)
 	{
-		return;
+		GEditor->UnregisterForUndo(this);
 	}
-
-	CompositeObj->AnimationTrack.CollapseAnimSegments();
-
-	RecalculateSequenceLength();
 }
 
 void SAnimCompositeEditor::PostUndo( bool bSuccess )
@@ -148,51 +66,7 @@ void SAnimCompositeEditor::PostRedo( bool bSuccess )
 
 void SAnimCompositeEditor::PostUndoRedo()
 {
-	if (!bIsActiveTimerRegistered)
-	{
-		bIsActiveTimerRegistered = true;
-		RegisterActiveTimer(0.f, FWidgetActiveTimerDelegate::CreateSP(this, &SAnimCompositeEditor::TriggerRebuildPanel));
-	}
-
-	// when undo or redo happens, we still have to recalculate length, so we can't rely on sequence length changes or not
-	if (CompositeObj->SequenceLength)
-	{
-		CompositeObj->SequenceLength = 0.f;
-	}
-}
-
-void SAnimCompositeEditor::InitDetailsViewEditorObject(UEditorAnimBaseObj* EdObj)
-{
-	EdObj->InitFromAnim(CompositeObj, FOnAnimObjectChange::CreateSP( SharedThis(this), &SAnimCompositeEditor::OnCompositeChange ));
-}
-
-EActiveTimerReturnType SAnimCompositeEditor::TriggerRebuildPanel(double InCurrentTime, float InDeltaTime)
-{
-	// we should not update any property related within PostEditChange, 
-	// so this is deferred to Tick, when it needs to rebuild, just mark it and this will update in the next tick
-	RebuildPanel();
-
-	bIsActiveTimerRegistered = false;
-	return EActiveTimerReturnType::Stop;
-}
-
-float SAnimCompositeEditor::CalculateSequenceLengthOfEditorObject() const
-{
-	return CompositeObj->AnimationTrack.GetLength();
-}
-
-void SAnimCompositeEditor::SortAndUpdateComposite()
-{
-	if (CompositeObj == nullptr)
-	{
-		return;
-	}
-
-	CompositeObj->AnimationTrack.SortAnimSegments();
-
-	RecalculateSequenceLength();
-
-	// Update view (this will recreate everything)
-	AnimCompositePanel->Update();
-}
+	GetPreviewScene()->SetPreviewAnimationAsset(CompositeObj);
 
+	AnimModel->RefreshTracks();
+}
\ No newline at end of file
diff --git a/Engine/Source/Editor/Persona/Private/SAnimCompositeEditor.h b/Engine/Source/Editor/Persona/Private/SAnimCompositeEditor.h
--- a/Engine/Source/Editor/Persona/Private/SAnimCompositeEditor.h
+++ b/Engine/Source/Editor/Persona/Private/SAnimCompositeEditor.h
@@ -12,6 +12,7 @@
 #include "SAnimCompositePanel.h"
 
 class SAnimNotifyPanel;
+class FAnimModel_AnimComposite;
 
 //////////////////////////////////////////////////////////////////////////
 // SAnimCompositeEditor
@@ -32,33 +33,19 @@ public:
 		SLATE_ARGUMENT( UAnimComposite*, Composite)
 		SLATE_EVENT(FOnObjectsSelected, OnObjectsSelected)
 		SLATE_EVENT(FOnInvokeTab, OnInvokeTab)
+		SLATE_EVENT(FOnEditCurves, OnEditCurves)
+		SLATE_EVENT(FOnStopEditingCurves, OnStopEditingCurves)
 
 	SLATE_END_ARGS()
 
 private:
-	/** Slate editor panels */
-	TSharedPtr<class SAnimCompositePanel> AnimCompositePanel;
-	TSharedPtr<class SAnimNotifyPanel> AnimNotifyPanel;
-	TSharedPtr<class SAnimCurvePanel>	AnimCurvePanel;
-
-	/** do I have to update the panel **/
-	bool bRebuildPanel;
-	void RebuildPanel();
-
-	/** Handler for when composite is edited in details view */
-	void OnCompositeChange(class UObject *EditorAnimBaseObj, bool bRebuild);
-
-	/** This will remove empty spaces in the composite's anim segments but not resort. e.g. - all cached indexes remains valid. UI IS NOT REBUILT after this */
-	void CollapseComposite();
-
-protected:
-	virtual void InitDetailsViewEditorObject(class UEditorAnimBaseObj* EdObj) override;
+	TSharedPtr<FAnimModel_AnimComposite> AnimModel;
 
 public:
-	void Construct(const FArguments& InArgs, const TSharedRef<class IPersonaPreviewScene>& InPreviewScene, const TSharedRef<class IEditableSkeleton>& InEditableSkeleton);
-
 	~SAnimCompositeEditor();
 
+	void Construct(const FArguments& InArgs, const TSharedRef<class IPersonaPreviewScene>& InPreviewScene, const TSharedRef<class IEditableSkeleton>& InEditableSkeleton, const TSharedRef<FUICommandList>& InCommandList);
+
 	/** Return the animation composite being edited */
 	UAnimComposite* GetCompositeObj() const { return CompositeObj; }
 	virtual UAnimationAsset* GetEditorObject() const override { return GetCompositeObj(); }
@@ -66,30 +53,16 @@ public:
 private:
 	/** Pointer to the animation composite being edited */
 	UAnimComposite* CompositeObj;
-	
-	/** FEditorUndoClient interface */
-	virtual void PostUndo( bool bSuccess ) override;
-	virtual void PostRedo( bool bSuccess ) override;
-	void PostUndoRedo();
-
-	virtual float CalculateSequenceLengthOfEditorObject() const override;
-
-	/** This will sort all components of the montage and update (recreate) the UI */
-	void SortAndUpdateComposite();
-
-	/** One-off active timer to trigger a panel rebuild */
-	EActiveTimerReturnType TriggerRebuildPanel(double InCurrentTime, float InDeltaTime);
-
-	/** Whether the active timer to trigger a panel rebuild is currently registered */
-	bool bIsActiveTimerRegistered;
 
 public:
-
-	/** delegate handlers for when the composite is being editor */
-	void			PreAnimUpdate();
-	void			PostAnimUpdate();
-
 	//~ Begin SAnimEditorBase Interface
 	virtual TSharedRef<SWidget> CreateDocumentAnchor() override;
 	//~ End SAnimEditorBase Interface
+
+	/** FEditorUndoClient interface */
+	virtual void PostUndo( bool bSuccess ) override;
+	virtual void PostRedo( bool bSuccess ) override;
+
+	/** Post undo **/
+	void PostUndoRedo();
 };
diff --git a/Engine/Source/Editor/Persona/Private/SAnimCompositePanel.cpp b/Engine/Source/Editor/Persona/Private/SAnimCompositePanel.cpp
--- a/Engine/Source/Editor/Persona/Private/SAnimCompositePanel.cpp
+++ b/Engine/Source/Editor/Persona/Private/SAnimCompositePanel.cpp
@@ -8,14 +8,24 @@
 #include "SAnimSegmentsPanel.h"
 #include "SAnimCompositeEditor.h"
 #include "Widgets/Layout/SExpandableArea.h"
+#include "AnimModel.h"
+#include "AssetToolsModule.h"
+#include "IAssetTypeActions.h"
+#include "Animation/AnimSequence.h"
+#include "Animation/AnimMontage.h"
 
 #define LOCTEXT_NAMESPACE "AnimCompositePanel"
 
 //////////////////////////////////////////////////////////////////////////
 // SAnimCompositePanel
 
-void SAnimCompositePanel::Construct(const FArguments& InArgs)
+void SAnimCompositePanel::Construct(const FArguments& InArgs, const TSharedRef<FAnimModel>& InModel)
 {
+	if(GEditor)
+	{
+		GEditor->RegisterForUndo(this);
+	}
+
 	SAnimTrackPanel::Construct( SAnimTrackPanel::FArguments()
 		.WidgetWidth(InArgs._WidgetWidth)
 		.ViewInputMin(InArgs._ViewInputMin)
@@ -24,51 +34,47 @@ void SAnimCompositePanel::Construct(const FArguments& InArgs)
 		.InputMax(InArgs._InputMax)
 		.OnSetInputViewRange(InArgs._OnSetInputViewRange));
 
+	WeakModel = InModel;
 	Composite = InArgs._Composite;
-	CompositeEditor = InArgs._CompositeEditor;
+	bIsActiveTimerRegistered = false;
+	bIsSelecting = false;
+
+	InModel->OnHandleObjectsSelected().AddSP(this, &SAnimCompositePanel::HandleObjectsSelected);
 
 	this->ChildSlot
 	[
-		SNew(SVerticalBox)
-		+SVerticalBox::Slot()
-		.FillHeight(1)
-		[
-			SNew( SExpandableArea )
-			.AreaTitle( LOCTEXT( "CompositeLabel", "Composite" ) )
-			.BodyContent()
-			[
-				SAssignNew( PanelArea, SBorder )
-				.BorderImage( FEditorStyle::GetBrush("NoBorder") )
-				.Padding( FMargin(2.0f, 2.0f) )
-				.ColorAndOpacity( FLinearColor::White )
-			]
-		]
+		SAssignNew( PanelArea, SBorder )
+		.BorderImage( FEditorStyle::GetBrush("NoBorder") )
+		.Padding(0.0f)
+		.ColorAndOpacity( FLinearColor::White )
 	];
 
 	Update();
+
+	CollapseComposite();
+}
+
+SAnimCompositePanel::~SAnimCompositePanel()
+{
+	if(GEditor)
+	{
+		GEditor->UnregisterForUndo(this);
+	}
 }
 
 void SAnimCompositePanel::Update()
 {
 	ClearSelected();
 	if ( Composite != NULL )
-	{
-		SAnimCompositeEditor* Editor = CompositeEditor.Pin().Get();
-
-		TSharedPtr<FTrackColorTracker> ColourTracker = MakeShareable(new FTrackColorTracker);
-		ColourTracker->AddColor(FLinearColor(0.9f, 0.9f, 0.9f, 0.9f));
-		ColourTracker->AddColor(FLinearColor(0.5f, 0.5f, 0.5f));
-
-		FLinearColor NodeColor = FLinearColor(0.f, 0.5f, 0.0f, 0.5f);
-		
+	{		
 		TSharedPtr<SVerticalBox> CompositeSlots;
 		PanelArea->SetContent(
 			SAssignNew( CompositeSlots, SVerticalBox )
 			);
 
-		TSharedRef<S2ColumnWidget> SectionTrack = Create2ColumnWidget(CompositeSlots.ToSharedRef());
+		UAnimMontage* AnimMontage = Cast<UAnimMontage>(Composite);
 
-		SectionTrack->LeftColumn->AddSlot()
+		CompositeSlots->AddSlot()
 			.AutoHeight()
 			.VAlign(VAlign_Center)
 			[
@@ -77,32 +83,160 @@ void SAnimCompositePanel::Update()
 				.NodeSelectionSet(&SelectionSet)
 				.ViewInputMin(ViewInputMin)
 				.ViewInputMax(ViewInputMax)
-				.ColorTracker(ColourTracker)
-				.NodeColor(NodeColor)
-				.ScrubPosition( Editor, &SAnimCompositeEditor::GetScrubValue )
+				.OnGetNodeColor(this,  &SAnimCompositePanel::HandleGetNodeColor)
 				.TrackMaxValue(Composite->SequenceLength)
 				.TrackNumDiscreteValues(Composite->GetNumberOfFrames())
-
+				.bChildAnimMontage(AnimMontage && AnimMontage->HasParentAsset())
 				.OnAnimSegmentNodeClicked( this, &SAnimCompositePanel::ShowSegmentInDetailsView )
-				.OnPreAnimUpdate( Editor, &SAnimCompositeEditor::PreAnimUpdate )
-				.OnPostAnimUpdate( Editor, &SAnimCompositeEditor::PostAnimUpdate )
+				.OnPreAnimUpdate( this, &SAnimCompositePanel::PreAnimUpdate )
+				.OnPostAnimUpdate( this, &SAnimCompositePanel::PostAnimUpdate )
 			];
 	}
 }
 
+void SAnimCompositePanel::OnCompositeChange(class UObject *EditorAnimBaseObj, bool bRebuild)
+{
+	if ( Composite != nullptr )
+	{
+		if(bRebuild && !bIsActiveTimerRegistered)
+		{
+			// sometimes crashes because the timer delay but animation still renders, so invalidating here before calling timer
+			Composite->InvalidateRecursiveAsset();
+			bIsActiveTimerRegistered = true;
+			RegisterActiveTimer(0.f, FWidgetActiveTimerDelegate::CreateSP(this, &SAnimCompositePanel::TriggerRebuildPanel));
+		} 
+		else
+		{
+			CollapseComposite();
+		}
+
+		Composite->MarkPackageDirty();
+	}
+}
+
+EActiveTimerReturnType SAnimCompositePanel::TriggerRebuildPanel(double InCurrentTime, float InDeltaTime)
+{
+	// we should not update any property related within PostEditChange, 
+	// so this is deferred to Tick, when it needs to rebuild, just mark it and this will update in the next tick
+	SortAndUpdateComposite();
+
+	bIsActiveTimerRegistered = false;
+	return EActiveTimerReturnType::Stop;
+}
+
 void SAnimCompositePanel::ShowSegmentInDetailsView(int32 SegmentIndex)
 {
-	UEditorAnimCompositeSegment *Obj = Cast<UEditorAnimCompositeSegment>(CompositeEditor.Pin()->ShowInDetailsView(UEditorAnimCompositeSegment::StaticClass()));
-	if(Obj != NULL)
+	if(!bIsSelecting)
 	{
-		Obj->InitAnimSegment(SegmentIndex);
+		TGuardValue<bool> GuardValue(bIsSelecting, true);
+
+		UEditorAnimCompositeSegment *Obj = Cast<UEditorAnimCompositeSegment>(WeakModel.Pin()->ShowInDetailsView(UEditorAnimCompositeSegment::StaticClass()));
+		if(Obj != nullptr)
+		{
+			Obj->InitFromAnim(Composite, FOnAnimObjectChange::CreateSP( this, &SAnimCompositePanel::OnCompositeChange ));
+			Obj->InitAnimSegment(SegmentIndex);
+		}
 	}
 }
 
 void SAnimCompositePanel::ClearSelected()
 {
-	SelectionSet.Empty();
-	CompositeEditor.Pin()->ClearDetailsView();
+	if(!bIsSelecting)
+	{
+		TGuardValue<bool> GuardValue(bIsSelecting, true);
+
+		SelectionSet.Empty();
+		WeakModel.Pin()->ClearDetailsView();
+	}
+}
+
+void SAnimCompositePanel::PreAnimUpdate()
+{
+	Composite->Modify();
+}
+
+void SAnimCompositePanel::PostAnimUpdate()
+{
+	Composite->MarkPackageDirty();
+	SortAndUpdateComposite();
+}
+
+void SAnimCompositePanel::SortAndUpdateComposite()
+{
+	if (Composite == nullptr)
+	{
+		return;
+	}
+
+	Composite->AnimationTrack.SortAnimSegments();
+
+	WeakModel.Pin()->RecalculateSequenceLength();
+
+	// Update view (this will recreate everything)
+	Update();
+
+	// Range may have changed
+	WeakModel.Pin()->UpdateRange();
+}
+
+void SAnimCompositePanel::CollapseComposite()
+{
+	if ( Composite == nullptr )
+	{
+		return;
+	}
+
+	Composite->AnimationTrack.CollapseAnimSegments();
+
+	WeakModel.Pin()->RecalculateSequenceLength();
+}
+
+void SAnimCompositePanel::PostUndo( bool bSuccess )
+{
+	PostUndoRedo();
+}
+
+void SAnimCompositePanel::PostRedo( bool bSuccess )
+{
+	PostUndoRedo();
+}
+
+void SAnimCompositePanel::PostUndoRedo()
+{
+	if (!bIsActiveTimerRegistered)
+	{
+		bIsActiveTimerRegistered = true;
+		RegisterActiveTimer(0.f, FWidgetActiveTimerDelegate::CreateSP(this, &SAnimCompositePanel::TriggerRebuildPanel));
+	}
+
+	// when undo or redo happens, we still have to recalculate length, so we can't rely on sequence length changes or not
+	if (Composite->SequenceLength)
+	{
+		Composite->SequenceLength = 0.f;
+	}
+}
+
+FLinearColor SAnimCompositePanel::HandleGetNodeColor(const FAnimSegment& InSegment) const
+{
+	static const FLinearColor DisabledColor(64, 64, 64);
+
+	if(InSegment.AnimReference != nullptr)
+	{
+		FAssetToolsModule& AssetToolsModule = FModuleManager::LoadModuleChecked<FAssetToolsModule>(TEXT("AssetTools"));
+		TWeakPtr<IAssetTypeActions> AssetTypeActions = AssetToolsModule.Get().GetAssetTypeActionsForClass(InSegment.AnimReference->GetClass());
+		check(AssetTypeActions.IsValid());
+		return AssetTypeActions.Pin()->GetTypeColor().ReinterpretAsLinear();
+	}
+
+	return DisabledColor;
+}
+
+void SAnimCompositePanel::HandleObjectsSelected(const TArray<UObject*>& InObjects)
+{
+	if(!bIsSelecting)
+	{
+		ClearSelected();
+	}
 }
 
 #undef LOCTEXT_NAMESPACE
diff --git a/Engine/Source/Editor/Persona/Private/SAnimCompositePanel.h b/Engine/Source/Editor/Persona/Private/SAnimCompositePanel.h
--- a/Engine/Source/Editor/Persona/Private/SAnimCompositePanel.h
+++ b/Engine/Source/Editor/Persona/Private/SAnimCompositePanel.h
@@ -9,6 +9,8 @@
 
 class SAnimCompositeEditor;
 class SBorder;
+class FAnimModel;
+struct FAnimSegment;
 
 //////////////////////////////////////////////////////////////////////////
 // SAnimCompositePanel
@@ -20,12 +22,11 @@ class SBorder;
 //
 //////////////////////////////////////////////////////////////////////////
 
-class SAnimCompositePanel : public SAnimTrackPanel
+class SAnimCompositePanel : public SAnimTrackPanel, public FEditorUndoClient
 {
 public:
 	SLATE_BEGIN_ARGS( SAnimCompositePanel )
 		: _Composite()
-		, _CompositeEditor()
 		, _ViewInputMin()
 		, _ViewInputMax()
 		, _InputMin()
@@ -34,7 +35,6 @@ public:
 	{}
 
 	SLATE_ARGUMENT( class UAnimComposite*, Composite)
-	SLATE_ARGUMENT( TWeakPtr<class SAnimCompositeEditor>, CompositeEditor)
 	SLATE_ARGUMENT( float, WidgetWidth )
 	SLATE_ATTRIBUTE( float, ViewInputMin )
 	SLATE_ATTRIBUTE( float, ViewInputMax )
@@ -43,7 +43,9 @@ public:
 	SLATE_EVENT( FOnSetInputViewRange, OnSetInputViewRange )
 	SLATE_END_ARGS()
 
-	void Construct(const FArguments& InArgs);
+	~SAnimCompositePanel();
+
+	void Construct(const FArguments& InArgs, const TSharedRef<FAnimModel>& InModel);
 
 	/** Recreates the editor panel to reflect changes to the composite */
 	void Update();
@@ -54,9 +56,34 @@ public:
 	/** Handlers for when the user clicks on a anim sequence node*/
 	void ShowSegmentInDetailsView(int32 SegmentIndex);
 
+	/** Delegate handlers for when the composite is being edited */
+	void PreAnimUpdate();
+	void PostAnimUpdate();
+
+	/** This will sort all components of the composite and update (recreate) the UI */
+	void SortAndUpdateComposite();
+
+	/** Handler for when composite is edited in details view */
+	void OnCompositeChange(class UObject *EditorAnimBaseObj, bool bRebuild);
+
+	/** One-off active timer to trigger a panel rebuild */
+	EActiveTimerReturnType TriggerRebuildPanel(double InCurrentTime, float InDeltaTime);
+
+	/** FEditorUndoClient interface */
+	virtual void PostUndo( bool bSuccess ) override;
+	virtual void PostRedo( bool bSuccess ) override;
+	void PostUndoRedo();
+
+	void CollapseComposite();
+
+	/** Get a color for an anim segment */
+	FLinearColor HandleGetNodeColor(const FAnimSegment& InSegment) const;
+
+	void HandleObjectsSelected(const TArray<UObject*>& InObjects);
+
 private:
-	/** Reference to our editor object */
-	TWeakPtr<SAnimCompositeEditor>	CompositeEditor;
+	/** Reference to our editor model */
+	TWeakPtr<FAnimModel> WeakModel;
 
 	/** The composite we are currently editing */
 	class UAnimComposite*		Composite;
@@ -66,4 +93,10 @@ private:
 
 	/** Passed to the anim segments panel to handle selection */
 	STrackNodeSelectionSet SelectionSet;
+
+	/** Whether the active timer to trigger a panel rebuild is currently registered */
+	bool bIsActiveTimerRegistered;
+
+	/** Recursion guard for selection */
+	bool bIsSelecting;
 };
diff --git a/Engine/Source/Editor/Persona/Private/SAnimEditorBase.cpp b/Engine/Source/Editor/Persona/Private/SAnimEditorBase.cpp
--- a/Engine/Source/Editor/Persona/Private/SAnimEditorBase.cpp
+++ b/Engine/Source/Editor/Persona/Private/SAnimEditorBase.cpp
@@ -12,6 +12,8 @@
 
 #include "AnimPreviewInstance.h"
 #include "Animation/BlendSpaceBase.h"
+#include "AnimModel.h"
+#include "SAnimTimeline.h"
 
 #define LOCTEXT_NAMESPACE "AnimEditorBase"
 
@@ -32,141 +34,50 @@ void SAnimEditorBase::Construct(const FArguments& InArgs, const TSharedRef<class
 
 	TSharedPtr<SVerticalBox> AnimVerticalBox;
 
-	this->ChildSlot
+	TSharedPtr<SWidget> TimelineToUse;
+	if (InArgs._DisplayAnimTimeline)
+	{
+		check(InArgs._AnimModel.IsValid());
+		TimelineToUse = SAssignNew(TimelineWidget, SAnimTimeline, InArgs._AnimModel.ToSharedRef());
+	}
+	else
+	{
+		TimelineToUse = SNullWidget::NullWidget;
+	}
+
+	ChildSlot
 	[
 		SAssignNew(AnimVerticalBox, SVerticalBox)
 		+SVerticalBox::Slot()
 		.FillHeight(1)
 		[
-			SNew(SBorder)
-			.BorderImage( FEditorStyle::GetBrush("ToolPanel.GroupBorder") )
+			SNew(SOverlay)
+			+SOverlay::Slot()
+			[
+				SAssignNew(NonScrollEditorPanels, SVerticalBox)
+			]
+			+SOverlay::Slot()
 			[
-				SNew(SOverlay)
-				+SOverlay::Slot()
-				[
-					SAssignNew(NonScrollEditorPanels, SVerticalBox)
-				]
-				+SOverlay::Slot()
-				[
-					SNew( SScrollBox )
-					.Visibility(EVisibility::SelfHitTestInvisible)
-					+SScrollBox::Slot() 
-					[
-						SAssignNew (EditorPanels, SVerticalBox)
-					]
-				]
+				TimelineToUse.ToSharedRef()
 			]
 		]
 	];
 
-	// If we want to create anim info bar, display that now
-	if (InArgs._DisplayAnimInfoBar)
+	if (InArgs._DisplayAnimScrubBar)
 	{
+		// If we are an anim sequence, add scrub panel as well
 		AnimVerticalBox->AddSlot()
 		.AutoHeight()
-		.VAlign(VAlign_Center)
+		.VAlign(VAlign_Bottom)
 		[
-			// this is *temporary* information to display stuff
-			SNew( SBorder )
-			.Padding(FMargin(4))
+			SNew(SHorizontalBox)
+			+SHorizontalBox::Slot()
+			.FillWidth(1)
 			[
-				SNew( SHorizontalBox )
-				+SHorizontalBox::Slot()
-				.FillWidth(1)
-				[
-					SNew( SHorizontalBox )
-					+SHorizontalBox::Slot()
-					.AutoWidth()
-					.Padding(FMargin(4,4,0,0))
-					[
-						SNew( STextBlock )
-						.Text(LOCTEXT("Animation", "Animation : "))
-					]
-
-					+SHorizontalBox::Slot()
-					.FillWidth(1)
-					.Padding(FMargin(4,4,0,0))
-					[
-						SNew( STextBlock )
-						.Text(this, &SAnimEditorBase::GetEditorObjectName)
-					]
-				]
-
-				+SHorizontalBox::Slot()
-				.AutoWidth()
-				[
-					SNew( SHorizontalBox )
-					+SHorizontalBox::Slot()
-					.AutoWidth()
-					.Padding(FMargin(4,4,0,0))
-					[
-						SNew( STextBlock )
-						.Text(LOCTEXT("Percentage", "Percentage: "))
-					]
-					+SHorizontalBox::Slot()
-					.FillWidth(1)
-					.Padding(FMargin(4,4,0,0))
-					[
-						SNew( STextBlock )
-						.Text(this, &SAnimEditorBase::GetCurrentPercentage)
-					]
-				]
-				+SHorizontalBox::Slot()
-				.AutoWidth()
-				[
-					SNew( SHorizontalBox )
-					+SHorizontalBox::Slot()
-					.AutoWidth()
-					.Padding(FMargin(4,4,0,0))
-					[
-						SNew( STextBlock )
-						.Text(LOCTEXT("CurrentTime", "CurrentTime: "))
-					]
-					+SHorizontalBox::Slot()
-					.FillWidth(1)
-					.Padding(FMargin(4,4,0,0))
-					[
-						SNew( STextBlock )
-						.Text(this, &SAnimEditorBase::GetCurrentSequenceTime)
-					]
-				]
-
-				+SHorizontalBox::Slot()
-				.AutoWidth()
-				[
-					SNew( SHorizontalBox )
-					+SHorizontalBox::Slot()
-					.AutoWidth()
-					.Padding(FMargin(4,4,0,0))
-					[
-						SNew( STextBlock )
-						.Text(LOCTEXT("CurrentFrame", "Current Frame: "))
-					]
-					+SHorizontalBox::Slot()
-					.FillWidth(1)
-					.Padding(FMargin(4,4,0,0))
-					[
-						SNew( STextBlock )
-						.Text(this, &SAnimEditorBase::GetCurrentFrame)
-					]
-				]
+				ConstructAnimScrubPanel()
 			]
 		];
-
 	}
-
-	// If we are an anim sequence, add scrub panel as well
-	AnimVerticalBox->AddSlot()
-	.AutoHeight()
-	.VAlign(VAlign_Bottom)
-	[
-		SNew(SHorizontalBox)
-		+SHorizontalBox::Slot()
-		.FillWidth(1)
-		[
-			ConstructAnimScrubPanel()
-		]
-	];
 }
 
 TSharedRef<class SAnimationScrubPanel> SAnimEditorBase::ConstructAnimScrubPanel()
diff --git a/Engine/Source/Editor/Persona/Private/SAnimEditorBase.h b/Engine/Source/Editor/Persona/Private/SAnimEditorBase.h
--- a/Engine/Source/Editor/Persona/Private/SAnimEditorBase.h
+++ b/Engine/Source/Editor/Persona/Private/SAnimEditorBase.h
@@ -14,6 +14,9 @@
 #include "EditorObjectsTracker.h"
 #include "SAnimCurvePanel.h"
 
+class FAnimModel;
+class SAnimTimeline;
+
 //////////////////////////////////////////////////////////////////////////
 // FTrackMarkerBar
 
@@ -33,12 +36,17 @@ class SAnimEditorBase : public SCompoundWidget, public FGCObject
 {
 public:
 	SLATE_BEGIN_ARGS(SAnimEditorBase)
-		: _DisplayAnimInfoBar(true)
+		: _DisplayAnimTimeline(true)
+		, _DisplayAnimScrubBar(false)
 	{}
 
 	SLATE_EVENT(FOnObjectsSelected, OnObjectsSelected)
 
-	SLATE_ARGUMENT( bool, DisplayAnimInfoBar )
+	SLATE_ARGUMENT( bool, DisplayAnimTimeline )
+
+	SLATE_ARGUMENT( bool, DisplayAnimScrubBar )
+
+	SLATE_ARGUMENT( TSharedPtr<FAnimModel>, AnimModel )
 
 	SLATE_END_ARGS()
 
@@ -139,6 +147,9 @@ protected:
 	/** Delegate called to select objects */
 	FOnObjectsSelected OnObjectsSelected;
 
+	/** The timeline widget */
+	TSharedPtr<SAnimTimeline> TimelineWidget;
+
 private:
 	float GetPercentageInternal() const;
 };
