COPY from //Dev-Anim to //Dev-Main

diff --git a/Engine/Source/Editor/Persona/Private/SAnimMontagePanel.cpp b/Engine/Source/Editor/Persona/Private/SAnimMontagePanel.cpp
--- a/Engine/Source/Editor/Persona/Private/SAnimMontagePanel.cpp
+++ b/Engine/Source/Editor/Persona/Private/SAnimMontagePanel.cpp
@@ -12,24 +12,30 @@
 #include "Widgets/Input/SButton.h"
 #include "Animation/EditorAnimSegment.h"
 #include "SAnimSegmentsPanel.h"
-#include "Widgets/Layout/SExpandableArea.h"
 #include "Widgets/Input/STextEntryPopup.h"
 #include "Widgets/Input/STextComboBox.h"
 #include "SAnimTimingPanel.h"
 #include "TabSpawners.h"
 #include "Widgets/Input/SNumericEntryBox.h"
 #include "Styling/CoreStyle.h"
-#include "ISkeletonEditorModule.h"
-#include "Modules/ModuleManager.h"
 #include "IEditableSkeleton.h"
 #include "Editor.h"
+#include "AnimModel.h"
+#include "AnimPreviewInstance.h"
+#include "ScopedTransaction.h"
+#include "Misc/MessageDialog.h"
+#include "Animation/EditorAnimCompositeSegment.h"
+#include "Factories/AnimMontageFactory.h"
+#include "Animation/EditorCompositeSection.h"
+#include "AnimModel_AnimMontage.h"
+#include "Subsystems/AssetEditorSubsystem.h"
 
 #define LOCTEXT_NAMESPACE "AnimMontagePanel"
 
 //////////////////////////////////////////////////////////////////////////
 // SAnimMontagePanel
 
-void SAnimMontagePanel::Construct(const FArguments& InArgs, FSimpleMulticastDelegate& OnSectionsChanged)
+void SAnimMontagePanel::Construct(const FArguments& InArgs, const TSharedRef<FAnimModel_AnimMontage>& InModel)
 {
 	SAnimTrackPanel::Construct( SAnimTrackPanel::FArguments()
 		.WidgetWidth(InArgs._WidgetWidth)
@@ -39,137 +45,570 @@ void SAnimMontagePanel::Construct(const FArguments& InArgs, FSimpleMulticastDele
 		.InputMax(InArgs._InputMax)
 		.OnSetInputViewRange(InArgs._OnSetInputViewRange));
 
+	WeakModel = InModel;
 	Montage = InArgs._Montage;
 	OnInvokeTab = InArgs._OnInvokeTab;
-	MontageEditor = InArgs._MontageEditor;
 	SectionTimingNodeVisibility = InArgs._SectionTimingNodeVisibility;
-	OnSetMontagePreviewSlot = InArgs._OnSetMontagePreviewSlot;
 	CurrentPreviewSlot = 0;
+	bIsSelecting = false;
+
+	InModel->OnHandleObjectsSelected().AddSP(this, &SAnimMontagePanel::HandleObjectsSelected);
 
 	bChildAnimMontage = InArgs._bChildAnimMontage;
 
 	//TrackStyle = TRACK_Double;
 	CurrentPosition = InArgs._CurrentPosition;
 
+	bIsActiveTimerRegistered = false;
+
+	EnsureStartingSection();
+	EnsureSlotNode();
+
+	if(GEditor)
+	{
+		GEditor->RegisterForUndo(this);
+	}
+
 	this->ChildSlot
 	[
-		SNew(SVerticalBox)
-		+SVerticalBox::Slot()
-		.FillHeight(1)
-		[
-			SNew( SExpandableArea )
-			.AreaTitle( LOCTEXT("Montage", "Montage") )
-			.BodyContent()
-			[
-				SAssignNew( PanelArea, SBorder )
-				.BorderImage( FEditorStyle::GetBrush("NoBorder") )
-				.Padding( FMargin(2.0f, 2.0f) )
-				.ColorAndOpacity( FLinearColor::White )
-			]
-		]
+		SAssignNew( PanelArea, SBorder )
+		.Padding(0.0f)
+		.BorderImage( FEditorStyle::GetBrush("NoBorder") )
+		.ColorAndOpacity( FLinearColor::White )
 	];
 
-	ISkeletonEditorModule& SkeletonEditorModule = FModuleManager::LoadModuleChecked<ISkeletonEditorModule>("SkeletonEditor");
-	TSharedPtr<IEditableSkeleton> EditableSkeleton = SkeletonEditorModule.CreateEditableSkeleton(Montage->GetSkeleton());
-	EditableSkeleton->RegisterOnNotifiesChanged(FSimpleDelegate::CreateSP(this, &SAnimMontagePanel::Update));
+	InModel->GetEditableSkeleton()->RegisterOnNotifiesChanged(FSimpleDelegate::CreateSP(this, &SAnimMontagePanel::Update));
+	InModel->OnTracksChanged().Add(FSimpleDelegate::CreateSP(this, &SAnimMontagePanel::Update));
 
-	OnSectionsChanged.Add(FSimpleDelegate::CreateSP(this, &SAnimMontagePanel::Update));
+	CollapseMontage();
 
 	Update();
 }
 
-/** This is the main function that creates the UI widgets for the montage tool.*/
+SAnimMontagePanel::~SAnimMontagePanel()
+{
+	if(GEditor)
+	{
+		GEditor->UnregisterForUndo(this);
+	}
+}
+
+UAnimPreviewInstance* SAnimMontagePanel::GetPreviewInstance() const
+{
+	UDebugSkelMeshComponent* PreviewMeshComponent = WeakModel.Pin()->GetPreviewScene()->GetPreviewMeshComponent();
+	return PreviewMeshComponent && PreviewMeshComponent->IsPreviewOn()? PreviewMeshComponent->PreviewInstance : nullptr;
+}
+
+FReply SAnimMontagePanel::OnFindParentClassInContentBrowserClicked()
+{
+	if (Montage != nullptr)
+	{
+		UObject* ParentClass = Montage->ParentAsset;
+		if (ParentClass != nullptr)
+		{
+			TArray< UObject* > ParentObjectList;
+			ParentObjectList.Add(ParentClass);
+			GEditor->SyncBrowserToObjects(ParentObjectList);
+		}
+	}
+
+	return FReply::Handled();
+}
+
+FReply SAnimMontagePanel::OnEditParentClassClicked()
+{
+	if (Montage != nullptr)
+	{
+		UObject* ParentClass = Montage->ParentAsset;
+		if (ParentClass != nullptr)
+		{
+			GEditor->GetEditorSubsystem<UAssetEditorSubsystem>()->OpenEditorForAsset(ParentClass);
+		}
+	}
+
+	return FReply::Handled();
+}
+
+void SAnimMontagePanel::OnSetMontagePreviewSlot(int32 SlotIndex)
+{
+	UAnimSingleNodeInstance * PreviewInstance = GetPreviewInstance();
+	if (PreviewInstance && Montage->SlotAnimTracks.IsValidIndex(SlotIndex))
+	{
+		const FName SlotName = Montage->SlotAnimTracks[SlotIndex].SlotName;
+		PreviewInstance->SetMontagePreviewSlot(SlotName);
+	}
+}
+
+bool SAnimMontagePanel::ValidIndexes(int32 AnimSlotIndex, int32 AnimSegmentIndex) const
+{
+	return (Montage != nullptr && Montage->SlotAnimTracks.IsValidIndex(AnimSlotIndex) && Montage->SlotAnimTracks[AnimSlotIndex].AnimTrack.AnimSegments.IsValidIndex(AnimSegmentIndex) );
+}
+
+bool SAnimMontagePanel::ValidSection(int32 SectionIndex) const
+{
+	return (Montage != nullptr && Montage->CompositeSections.IsValidIndex(SectionIndex));
+}
+
+void SAnimMontagePanel::HandleNotifiesChanged()
+{
+	WeakModel.Pin()->OnTracksChanged().Broadcast();
+}
+
+
+bool SAnimMontagePanel::GetSectionTime( int32 SectionIndex, float &OutTime ) const
+{
+	if (Montage != nullptr && Montage->CompositeSections.IsValidIndex(SectionIndex))
+	{
+		OutTime = Montage->CompositeSections[SectionIndex].GetTime();
+		return true;
+	}
+
+	return false;
+}
+
+TArray<FString>	SAnimMontagePanel::GetSectionNames() const
+{
+	TArray<FString> Names;
+	if (Montage != nullptr)
+	{
+		for( int32 I=0; I < Montage->CompositeSections.Num(); I++)
+		{
+			Names.Add(Montage->CompositeSections[I].SectionName.ToString());
+		}
+	}
+	return Names;
+}
+
+TArray<float> SAnimMontagePanel::GetSectionStartTimes() const
+{
+	TArray<float>	Times;
+	if (Montage != nullptr)
+	{
+		for( int32 I=0; I < Montage->CompositeSections.Num(); I++)
+		{
+			Times.Add(Montage->CompositeSections[I].GetTime());
+		}
+	}
+	return Times;
+}
+
+TArray<FTrackMarkerBar> SAnimMontagePanel::GetMarkerBarInformation() const
+{
+	TArray<FTrackMarkerBar> MarkerBars;
+	if (Montage != nullptr)
+	{
+		for( int32 I=0; I < Montage->CompositeSections.Num(); I++)
+		{
+			FTrackMarkerBar Bar;
+			Bar.Time = Montage->CompositeSections[I].GetTime();
+			Bar.DrawColour = FLinearColor(0.f,1.f,0.f);
+			MarkerBars.Add(Bar);
+		}
+	}
+	return MarkerBars;
+}
+
+TArray<float> SAnimMontagePanel::GetAnimSegmentStartTimes() const
+{
+	TArray<float>	Times;
+	if (Montage != nullptr)
+	{
+		for ( int32 i=0; i < Montage->SlotAnimTracks.Num(); i++)
+		{
+			for (int32 j=0; j < Montage->SlotAnimTracks[i].AnimTrack.AnimSegments.Num(); j++)
+			{
+				Times.Add( Montage->SlotAnimTracks[i].AnimTrack.AnimSegments[j].StartPos );
+			}
+		}
+	}
+	return Times;
+}
+
+void SAnimMontagePanel::PreAnimUpdate()
+{
+	Montage->Modify();
+}
+
+void SAnimMontagePanel::OnMontageModified()
+{
+	Montage->PostEditChange();
+	Montage->MarkPackageDirty();
+}
+
+void SAnimMontagePanel::PostAnimUpdate()
+{
+	SortAndUpdateMontage();
+	OnMontageModified();
+}
+
+bool SAnimMontagePanel::IsDiffererentFromParent(FName SlotName, int32 SegmentIdx, const FAnimSegment& Segment)
+{
+	// if it doesn't hare parent asset, no reason to come here
+	if (Montage && ensureAlways(Montage->ParentAsset))
+	{
+		// find correct source asset from parent
+		UAnimMontage* ParentMontage = Cast<UAnimMontage>(Montage->ParentAsset);
+		if (ParentMontage->IsValidSlot(SlotName))
+		{
+			const FAnimTrack* ParentTrack = ParentMontage->GetAnimationData(SlotName);
+
+			if (ParentTrack && ParentTrack->AnimSegments.IsValidIndex(SegmentIdx))
+			{
+				UAnimSequenceBase* SourceAsset = ParentTrack->AnimSegments[SegmentIdx].AnimReference;
+				return (SourceAsset != Segment.AnimReference);
+			}
+		}
+	}
+
+	// if something doesn't match, we assume they're different, so default feedback  is to return true
+	return true;
+}
+
+void SAnimMontagePanel::ReplaceAnimationMapping(FName SlotName, int32 SegmentIdx, UAnimSequenceBase* OldSequenceBase, UAnimSequenceBase* NewSequenceBase)
+{
+	// if it doesn't hare parent asset, no reason to come here
+	if (Montage && ensureAlways(Montage->ParentAsset))
+	{
+		// find correct source asset from parent
+		UAnimMontage* ParentMontage = Cast<UAnimMontage>(Montage->ParentAsset);
+		if (ParentMontage->IsValidSlot(SlotName))
+		{
+			const FAnimTrack* ParentTrack = ParentMontage->GetAnimationData(SlotName);
+
+			if (ParentTrack && ParentTrack->AnimSegments.IsValidIndex(SegmentIdx))
+			{
+				UAnimSequenceBase* SourceAsset = ParentTrack->AnimSegments[SegmentIdx].AnimReference;
+				if (Montage->RemapAsset(SourceAsset, NewSequenceBase))
+				{
+					// success
+					return;
+				}
+			}
+		}
+	}
+
+	// failed to do the process, check if the animation is correct or if the same type of animation
+	// print error
+	FMessageDialog::Open(EAppMsgType::Ok, LOCTEXT("FailedToRemap", "Make sure the target animation is valid. If source is additive, target animation has to be additive also."));
+}
+
+void SAnimMontagePanel::RebuildMontagePanel(bool bNotifyAsset /*= true*/)
+{
+	SortAndUpdateMontage();
+	WeakModel.Pin()->OnSectionsChanged.Execute();
+
+	if (bNotifyAsset)
+	{
+		OnMontageModified();
+	}
+}
+
+
+/** This will remove empty spaces in the montage's anim segment but not resort. e.g. - all cached indexes remains valid. UI IS NOT REBUILT after this */
+void SAnimMontagePanel::CollapseMontage()
+{
+	if (Montage==nullptr)
+	{
+		return;
+	}
+
+	for (int32 i=0; i < Montage->SlotAnimTracks.Num(); i++)
+	{
+		Montage->SlotAnimTracks[i].AnimTrack.CollapseAnimSegments();
+	}
+
+	Montage->UpdateLinkableElements();
+
+	WeakModel.Pin()->RecalculateSequenceLength();
+}
+
+/** This will sort all components of the montage and update (recreate) the UI */
+void SAnimMontagePanel::SortAndUpdateMontage()
+{
+	if (Montage==nullptr)
+	{
+		return;
+	}
+	
+	SortAnimSegments();
+
+	Montage->UpdateLinkableElements();
+
+	WeakModel.Pin()->RecalculateSequenceLength();
+
+	SortSections();
+
+	StaticCastSharedPtr<FAnimModel_AnimMontage>(WeakModel.Pin())->RefreshNotifyTriggerOffsets();
+
+	// Update view (this will recreate everything)
+// 	AnimMontagePanel->Update();
+// 	AnimMontageSectionsPanel->Update();
+// 	AnimTimingPanel->Update();
+
+	WeakModel.Pin()->RefreshTracks();
+
+	// Restart the preview instance of the montage
+	RestartPreview();
+}
+
+
+/** Make sure all Sections and Notifies are clamped to NewEndTime (called before NewEndTime is set to SequenceLength) */
+bool SAnimMontagePanel::ClampToEndTime(float NewEndTime)
+{
+	float SequenceLength = GetSequenceLength();
+
+	bool bClampingNeeded = (SequenceLength > 0.f && NewEndTime < SequenceLength);
+	if(bClampingNeeded)
+	{
+		float ratio = NewEndTime / Montage->SequenceLength;
+
+		for(int32 i=0; i < Montage->CompositeSections.Num(); i++)
+		{
+			if(Montage->CompositeSections[i].GetTime() > NewEndTime)
+			{
+				float CurrentTime = Montage->CompositeSections[i].GetTime();
+				Montage->CompositeSections[i].SetTime(CurrentTime * ratio);
+			}
+		}
+
+		for(int32 i=0; i < Montage->Notifies.Num(); i++)
+		{
+			FAnimNotifyEvent& Notify = Montage->Notifies[i];
+			float NotifyTime = Notify.GetTime();
+
+			if(NotifyTime >= NewEndTime)
+			{
+				Notify.SetTime(NotifyTime * ratio);
+				Notify.TriggerTimeOffset = GetTriggerTimeOffsetForType(Montage->CalculateOffsetForNotify(Notify.GetTime()));
+			}
+		}
+	}
+
+	return bClampingNeeded;
+}
+
+void SAnimMontagePanel::AddNewSection(float StartTime, FString SectionName)
+{
+	if ( Montage != nullptr )
+	{
+		const FScopedTransaction Transaction( LOCTEXT("AddNewSection", "Add New Section") );
+		Montage->Modify();
+
+		if (Montage->AddAnimCompositeSection(FName(*SectionName), StartTime) != INDEX_NONE)
+		{
+			RebuildMontagePanel();
+		}
+		OnMontageModified();
+		WeakModel.Pin()->OnSectionsChanged.Execute();
+	}
+}
+
+void SAnimMontagePanel::RemoveSection(int32 SectionIndex)
+{
+	if(ValidSection(SectionIndex))
+	{
+		const FScopedTransaction Transaction( LOCTEXT("DeleteSection", "Delete Section") );
+		Montage->Modify();
+
+		Montage->CompositeSections.RemoveAt(SectionIndex);
+		EnsureStartingSection();
+		OnMontageModified();
+		WeakModel.Pin()->RefreshTracks();
+		WeakModel.Pin()->OnSectionsChanged.Execute();
+		RestartPreview();
+	}
+}
+
+FString	SAnimMontagePanel::GetSectionName(int32 SectionIndex) const
+{
+	if (ValidSection(SectionIndex))
+	{
+		return Montage->GetSectionName(SectionIndex).ToString();
+	}
+	return FString();
+}
+
+void SAnimMontagePanel::RenameSlotNode(int32 SlotIndex, FString NewSlotName)
+{
+	if(Montage->SlotAnimTracks.IsValidIndex(SlotIndex))
+	{
+		FName NewName(*NewSlotName);
+		if(Montage->SlotAnimTracks[SlotIndex].SlotName != NewName)
+		{
+			const FScopedTransaction Transaction( LOCTEXT("RenameSlot", "Rename Slot") );
+			WeakModel.Pin()->GetEditableSkeleton()->RegisterSlotNode(NewName);
+
+			Montage->Modify();
+
+			Montage->SlotAnimTracks[SlotIndex].SlotName = NewName;
+			OnMontageModified();
+		}
+	}
+}
+
+void SAnimMontagePanel::AddNewMontageSlot( FName NewSlotName )
+{
+	if ( Montage != nullptr )
+	{
+		const FScopedTransaction Transaction( LOCTEXT("AddSlot", "Add Slot") );
+
+		WeakModel.Pin()->GetEditableSkeleton()->RegisterSlotNode(NewSlotName);
+
+		Montage->Modify();
+
+		Montage->AddSlot(NewSlotName);
+
+		OnMontageModified();
+
+		Update();
+
+		WeakModel.Pin()->RefreshTracks();
+	}
+}
+
+FText SAnimMontagePanel::GetMontageSlotName(int32 SlotIndex) const
+{
+	if(Montage->SlotAnimTracks.IsValidIndex(SlotIndex) && Montage->SlotAnimTracks[SlotIndex].SlotName != NAME_None)
+	{
+		return FText::FromName( Montage->SlotAnimTracks[SlotIndex].SlotName );
+	}	
+	return FText::GetEmpty();
+}
+
+void SAnimMontagePanel::RemoveMontageSlot(int32 AnimSlotIndex)
+{
+	if ( Montage != nullptr && Montage->SlotAnimTracks.IsValidIndex( AnimSlotIndex ) )
+	{
+		const FScopedTransaction Transaction( LOCTEXT("RemoveSlot", "Remove Slot") );
+		Montage->Modify();
+
+		Montage->SlotAnimTracks.RemoveAt(AnimSlotIndex);
+		OnMontageModified();
+		Update();
+
+		// Iterate the notifies and relink anything that is now invalid
+		for(FAnimNotifyEvent& Event : Montage->Notifies)
+		{
+			Event.ConditionalRelink();
+		}
+
+		// Do the same for sections
+		for(FCompositeSection& Section : Montage->CompositeSections)
+		{
+			Section.ConditionalRelink();
+		}
+
+		WeakModel.Pin()->RefreshTracks();
+	}
+}
+
+bool SAnimMontagePanel::CanRemoveMontageSlot(int32 AnimSlotIndex)
+{
+	return (Montage != nullptr) && (Montage->SlotAnimTracks.Num()) > 1;
+}
+
+void SAnimMontagePanel::DuplicateMontageSlot(int32 AnimSlotIndex)
+{
+	if (Montage != nullptr && Montage->SlotAnimTracks.IsValidIndex(AnimSlotIndex))
+	{
+		const FScopedTransaction Transaction(LOCTEXT("DuplicateSlot", "Duplicate Slot"));
+		Montage->Modify();
+
+		FSlotAnimationTrack& NewTrack = Montage->AddSlot(FAnimSlotGroup::DefaultSlotName); 
+		NewTrack.AnimTrack = Montage->SlotAnimTracks[AnimSlotIndex].AnimTrack;
+
+		OnMontageModified();
+
+		Update();
+
+		WeakModel.Pin()->RefreshTracks();
+	}
+}
+
+void SAnimMontagePanel::RestartPreview()
+{
+	if (UDebugSkelMeshComponent* MeshComponent = WeakModel.Pin()->GetPreviewScene()->GetPreviewMeshComponent())
+	{
+		if (UAnimPreviewInstance* Preview = MeshComponent->PreviewInstance)
+		{
+			Preview->MontagePreview_PreviewNormal(INDEX_NONE, Preview->IsPlaying());
+		}
+	}
+}
+
+void SAnimMontagePanel::RestartPreviewPlayAllSections()
+{
+	if (UDebugSkelMeshComponent* MeshComponent = WeakModel.Pin()->GetPreviewScene()->GetPreviewMeshComponent())
+	{
+		if (UAnimPreviewInstance* Preview = MeshComponent->PreviewInstance)
+		{
+			Preview->MontagePreview_PreviewAllSections(Preview->IsPlaying());
+		}
+	}
+}
+
+void SAnimMontagePanel::MakeDefaultSequentialSections()
+{
+	check( Montage != nullptr );
+	SortSections();
+	for(int32 SectionIdx=0; SectionIdx < Montage->CompositeSections.Num(); SectionIdx++)
+	{
+		Montage->CompositeSections[SectionIdx].NextSectionName = Montage->CompositeSections.IsValidIndex(SectionIdx+1) ? Montage->CompositeSections[SectionIdx+1].SectionName : NAME_None;
+	}
+	RestartPreview();
+}
+
+void SAnimMontagePanel::ClearSquenceOrdering()
+{
+	check( Montage != nullptr );
+	SortSections();
+	for(int32 SectionIdx=0; SectionIdx < Montage->CompositeSections.Num(); SectionIdx++)
+	{
+		Montage->CompositeSections[SectionIdx].NextSectionName = NAME_None;
+	}
+	RestartPreview();
+}
+
+void SAnimMontagePanel::PostUndo( bool bSuccess )
+{
+	PostRedoUndo();
+}
+
+void SAnimMontagePanel::PostRedo( bool bSuccess )
+{
+	PostRedoUndo();
+}
+
+void SAnimMontagePanel::PostRedoUndo()
+{
+	// when undo or redo happens, we still have to recalculate length, so we can't rely on sequence length changes or not
+	if (Montage->SequenceLength)
+	{
+		Montage->SequenceLength = 0.f;
+	}
+
+	RebuildMontagePanel(); //Rebuild here, undoing adds can cause slate to crash later on if we don't (using dummy args since they aren't used by the method
+}
+
 void SAnimMontagePanel::Update()
 {
-	if ( Montage != NULL )
+	if ( Montage != nullptr )
 	{
 		CurrentPreviewSlot = FMath::Clamp(CurrentPreviewSlot, 0, Montage->SlotAnimTracks.Num() - 1);
-		OnSetMontagePreviewSlot.ExecuteIfBound(CurrentPreviewSlot);
+		OnSetMontagePreviewSlot(CurrentPreviewSlot);
 
-		SMontageEditor * Editor = MontageEditor.Pin().Get();
 		int32 ColorIdx=0;
-		//FLinearColor Colors[] = { FLinearColor(0.9f, 0.9f, 0.9f, 0.9f), FLinearColor(0.5f, 0.5f, 0.5f) };
 		
-		TSharedPtr<FTrackColorTracker> ColourTracker = MakeShareable(new FTrackColorTracker);
-		ColourTracker->AddColor(FLinearColor(0.9f, 0.9f, 0.9f, 0.9f));
-		ColourTracker->AddColor(FLinearColor(0.5f, 0.5f, 0.5f));
-
 		FLinearColor NodeColor = FLinearColor(0.f, 0.5f, 0.0f, 0.5f);
-		//FLinearColor SelectedColor = FLinearColor(1.0f,0.65,0.0f);
 
 		TSharedPtr<SVerticalBox> MontageSlots;
 		PanelArea->SetContent(
 			SAssignNew( MontageSlots, SVerticalBox )
 			);
 
-		/************************************************************************/
-		/* Status Bar                                                                     */
-		/************************************************************************/
-		{
-			MontageSlots->AddSlot()
-				.AutoHeight()
-				[
-					// Header, shows name of timeline we are editing
-					SNew(SBorder)
-					.BorderImage(FEditorStyle::GetBrush(TEXT("Graph.TitleBackground")))
-					.HAlign(HAlign_Center)
-					[
-						SNew(SHorizontalBox)
-						+ SHorizontalBox::Slot()
-						.AutoWidth()
-						.FillWidth(3.f)
-						.HAlign(HAlign_Right)
-						.VAlign(VAlign_Center)
-						[
-							SAssignNew(StatusBarWarningImage, SImage)
-							.Image(FEditorStyle::GetBrush("AnimSlotManager.Warning"))
-							.Visibility(EVisibility::Hidden)
-						]
-
-						+ SHorizontalBox::Slot()
-						.AutoWidth()
-						.VAlign(VAlign_Center)
-						.HAlign(HAlign_Center)
-						[
-							SAssignNew(StatusBarTextBlock, STextBlock)
-							.Font(FCoreStyle::GetDefaultFontStyle("Regular", 12))
-							.ColorAndOpacity(FLinearColor(1, 1, 1, 0.5))
-						]
-					]
-				];
-		}
+		MontageSlots->ClearChildren();
 
-		// ===================================
-		// Section Name Track
-		// ===================================
-		{
-			TSharedRef<S2ColumnWidget> SectionTrack = Create2ColumnWidget(MontageSlots.ToSharedRef());
-
-			SectionTrack->LeftColumn->ClearChildren();
-			SectionTrack->LeftColumn->AddSlot()
-				.AutoHeight()
-				.VAlign(VAlign_Center)
-				.Padding( FMargin(0.5f, 0.5f) )
-				[
-					SAssignNew(SectionNameTrack, STrack)
-					.IsEnabled(!bChildAnimMontage)
-					.ViewInputMin(ViewInputMin)
-					.ViewInputMax(ViewInputMax)
-					.TrackColor( ColourTracker->GetNextColor() )
-					.OnBarDrag(Editor, &SMontageEditor::OnEditSectionTime)
-					.OnBarDrop(Editor, &SMontageEditor::OnEditSectionTimeFinish)
-					.OnBarClicked(SharedThis(this), &SAnimMontagePanel::ShowSectionInDetailsView)
-					.DraggableBars(Editor, &SMontageEditor::GetSectionStartTimes)
-					.DraggableBarSnapPositions(Editor, &SMontageEditor::GetAnimSegmentStartTimes)
-					.DraggableBarLabels(Editor, &SMontageEditor::GetSectionNames)
-					.TrackMaxValue(this, &SAnimMontagePanel::GetSequenceLength)
-					.TrackNumDiscreteValues(Montage->GetNumberOfFrames())
-					.OnTrackRightClickContextMenu( this, &SAnimMontagePanel::SummonTrackContextMenu, static_cast<int>(INDEX_NONE))
-					.ScrubPosition( Editor, &SMontageEditor::GetScrubValue )
-				];
-
-			RefreshTimingNodes();
-		}
+		RefreshTimingNodes();
 
 		// ===================================
 		// Anim Segment Tracks
@@ -189,88 +628,13 @@ void SAnimMontagePanel::Update()
 			check(SlotNameComboSelectedNames.Num() == NumAnimTracks);
 
 			for (int32 SlotAnimIdx = 0; SlotAnimIdx < NumAnimTracks; SlotAnimIdx++)
-			{
-				TSharedRef<S2ColumnWidget> SectionTrack = Create2ColumnWidget(MontageSlots.ToSharedRef());
-				
+			{			
 				int32 FoundIndex = SlotNameList.Find(SlotNameComboSelectedNames[SlotAnimIdx]);
 				TSharedPtr<FString> ComboItem = SlotNameComboListItems[FoundIndex];
 
-				// Right column
-				SectionTrack->RightColumn->AddSlot()
-				.VAlign(VAlign_Center)
-				.HAlign(HAlign_Fill)
-				.FillHeight(1)
-				[
-					SNew(SVerticalBox)
-
-					+ SVerticalBox::Slot()
-					.HAlign(HAlign_Fill)
-					[
-						SAssignNew(SlotNameComboBoxes[SlotAnimIdx], STextComboBox)
-						.OptionsSource(&SlotNameComboListItems)
-						.OnSelectionChanged(this, &SAnimMontagePanel::OnSlotNameChanged, SlotAnimIdx)
-						.OnComboBoxOpening(this, &SAnimMontagePanel::OnSlotListOpening, SlotAnimIdx)
-						.InitiallySelectedItem(ComboItem)
-						.ContentPadding(2)
-						.ToolTipText(FText::FromString(*ComboItem))
-					]
-
-					+ SVerticalBox::Slot()
-						.HAlign(HAlign_Left)
-						[
-							SNew(SHorizontalBox)
-							+ SHorizontalBox::Slot()
-							.AutoWidth()
-							.HAlign(HAlign_Left)
-							[
-								SNew(SButton)
-								.Text(LOCTEXT("AnimSlotNode_DetailPanelManageButtonLabel", "Anim Slot Manager"))
-								.ToolTipText(LOCTEXT("AnimSlotNode_DetailPanelManageButtonToolTipText", "Open Anim Slot Manager to edit Slots and Groups."))
-								.OnClicked(this, &SAnimMontagePanel::OnOpenAnimSlotManager)
-								.Content()
-								[
-									SNew(SImage)
-									.Image(FEditorStyle::GetBrush("MeshPaint.FindInCB"))
-								]
-							]
-							+ SHorizontalBox::Slot()
-							.AutoWidth()
-							.HAlign(HAlign_Left)
-							.VAlign(VAlign_Center)
-							[
-								SNew(SCheckBox)
-								.Style(FCoreStyle::Get(), "ToggleButtonCheckbox")
-								.ToolTipText(LOCTEXT("DetailPanelPreviewToolTipText", "Preview this slot in the editor"))
-								.IsChecked(this, &SAnimMontagePanel::IsSlotPreviewed, SlotAnimIdx)
-								.OnCheckStateChanged(this, &SAnimMontagePanel::OnSlotPreviewedChanged, SlotAnimIdx)
-								[
-									SNew(SBox)
-									.VAlign(VAlign_Center)
-									.HAlign(HAlign_Center)
-									.Padding(FMargin(4.0, 2.0))
-									[
-										SNew(STextBlock)
-										.Text(LOCTEXT("DetailPanelPreview", "Preview"))
-									]
-								]
-							]
-							+ SHorizontalBox::Slot()
-							.AutoWidth()
-							.FillWidth(2.f)
-							.HAlign(HAlign_Left)
-							.VAlign(VAlign_Center)
-							[
-								SAssignNew(SlotWarningImages[SlotAnimIdx], SImage)
-								.Image(FEditorStyle::GetBrush("AnimSlotManager.Warning"))
-								.Visibility(EVisibility::Hidden)
-							]
-						]
-				];
-
-				SectionTrack->RightColumn->SetEnabled(!bChildAnimMontage);
 				if (bChildAnimMontage)
 				{
-					SectionTrack->LeftColumn->AddSlot()
+					MontageSlots->AddSlot()
 					.AutoHeight()
 					.VAlign(VAlign_Center)
 					[
@@ -280,14 +644,12 @@ void SAnimMontagePanel::Update()
 						.NodeSelectionSet(&SelectionSet)
 						.ViewInputMin(ViewInputMin)
 						.ViewInputMax(ViewInputMax)
-						.ColorTracker(ColourTracker)
 						.bChildAnimMontage(bChildAnimMontage)
-						.NodeColor(NodeColor)
-						.OnPreAnimUpdate(Editor, &SMontageEditor::PreAnimUpdate)
-						.OnPostAnimUpdate(Editor, &SMontageEditor::PostAnimUpdate)
-						.OnAnimReplaceMapping(Editor, &SMontageEditor::ReplaceAnimationMapping)
-						.OnDiffFromParentAsset(Editor, &SMontageEditor::IsDiffererentFromParent)
-						.ScrubPosition(Editor, &SMontageEditor::GetScrubValue)
+						.OnGetNodeColor_Lambda([NodeColor](const FAnimSegment& InSegment){ return NodeColor; })
+						.OnPreAnimUpdate(this, &SAnimMontagePanel::PreAnimUpdate)
+						.OnPostAnimUpdate(this, &SAnimMontagePanel::PostAnimUpdate)
+						.OnAnimReplaceMapping(this, &SAnimMontagePanel::ReplaceAnimationMapping)
+						.OnDiffFromParentAsset(this, &SAnimMontagePanel::IsDiffererentFromParent)
 						.TrackMaxValue(this, &SAnimMontagePanel::GetSequenceLength)
 						.TrackNumDiscreteValues(Montage->GetNumberOfFrames())
 					];
@@ -295,7 +657,7 @@ void SAnimMontagePanel::Update()
 				}
 				else
 				{
-					SectionTrack->LeftColumn->AddSlot()
+					MontageSlots->AddSlot()
 					.AutoHeight()
 					.VAlign(VAlign_Center)
 					[
@@ -305,38 +667,20 @@ void SAnimMontagePanel::Update()
 						.NodeSelectionSet(&SelectionSet)
 						.ViewInputMin(ViewInputMin)
 						.ViewInputMax(ViewInputMax)
-						.ColorTracker(ColourTracker)
 						.bChildAnimMontage(bChildAnimMontage)
-						.NodeColor(NodeColor)
-						.DraggableBars(Editor, &SMontageEditor::GetSectionStartTimes)
-						.DraggableBarSnapPositions(Editor, &SMontageEditor::GetAnimSegmentStartTimes)
-						.ScrubPosition(Editor, &SMontageEditor::GetScrubValue)
+						.OnGetNodeColor_Lambda([NodeColor](const FAnimSegment& InSegment){ return NodeColor; })
 						.TrackMaxValue(this, &SAnimMontagePanel::GetSequenceLength)
 						.TrackNumDiscreteValues(Montage->GetNumberOfFrames())
 						.OnAnimSegmentNodeClicked(this, &SAnimMontagePanel::ShowSegmentInDetailsView, SlotAnimIdx)
-						.OnPreAnimUpdate(Editor, &SMontageEditor::PreAnimUpdate)
-						.OnPostAnimUpdate(Editor, &SMontageEditor::PostAnimUpdate)
+						.OnPreAnimUpdate(this, &SAnimMontagePanel::PreAnimUpdate)
+						.OnPostAnimUpdate(this, &SAnimMontagePanel::PostAnimUpdate)
 						.OnAnimSegmentRemoved(this, &SAnimMontagePanel::OnAnimSegmentRemoved, SlotAnimIdx)
-						.OnBarDrag(Editor, &SMontageEditor::OnEditSectionTime)
-						.OnBarDrop(Editor, &SMontageEditor::OnEditSectionTimeFinish)
-						.OnBarClicked(SharedThis(this), &SAnimMontagePanel::ShowSectionInDetailsView)
 						.OnTrackRightClickContextMenu(this, &SAnimMontagePanel::SummonTrackContextMenu, static_cast<int>(SlotAnimIdx))
 					];
 				}
 			}
 		}
 	}
-
-	UpdateSlotGroupWarningVisibility();
-}
-
-void SAnimMontagePanel::SetMontage(class UAnimMontage * InMontage)
-{
-	if (InMontage != Montage)
-	{
-		Montage = InMontage;
-		Update();
-	}
 }
 
 FReply SAnimMontagePanel::OnMouseButtonDown( const FGeometry& InMyGeometry, const FPointerEvent& InMouseEvent )
@@ -359,77 +703,7 @@ void SAnimMontagePanel::SummonTrackContextMenu( FMenuBuilder& MenuBuilder, float
 		UIAction.ExecuteAction.BindRaw(this, &SAnimMontagePanel::OnNewSectionClicked, static_cast<float>(DataPosX));
 		UIAction.CanExecuteAction.BindRaw(this, &SAnimMontagePanel::CanAddNewSection);
 		MenuBuilder.AddMenuEntry(LOCTEXT("NewMontageSection", "New Montage Section"), LOCTEXT("NewMontageSectionToolTip", "Adds a new Montage Section"), FSlateIcon(), UIAction);
-	
 		UIAction.CanExecuteAction.Unbind();
-
-		if (SectionIndex != INDEX_NONE && Montage->CompositeSections.Num() > 1) // Can't delete the last section!
-		{
-			UIAction.ExecuteAction.BindRaw(MontageEditor.Pin().Get(), &SMontageEditor::RemoveSection, SectionIndex);
-			MenuBuilder.AddMenuEntry(LOCTEXT("DeleteMontageSection", "Delete Montage Section"), LOCTEXT("DeleteMontageSectionToolTip", "Deletes Montage Section"), FSlateIcon(), UIAction);
-
-			FCompositeSection& Section = Montage->CompositeSections[SectionIndex];
-
-			// Add item to directly set section time
-			TSharedRef<SWidget> TimeWidget = 
-				SNew( SBox )
-				.HAlign( HAlign_Right )
-				.ToolTipText(LOCTEXT("SetSectionTimeToolTip", "Set the time of this section directly"))
-				[
-					SNew(SBox)
-					.Padding(FMargin(4.0f, 0.0f, 0.0f, 0.0f))
-					.WidthOverride(100.0f)
-					[
-						SNew(SNumericEntryBox<float>)
-						.Font(FEditorStyle::GetFontStyle(TEXT("MenuItem.Font")))
-						.MinValue(0.0f)
-						.MaxValue(Montage->SequenceLength)
-						.Value(Section.GetTime())
-						.AllowSpin(true)
-						.OnValueCommitted_Lambda([this, SectionIndex](float InValue, ETextCommit::Type InCommitType)
-						{
-							if (Montage->CompositeSections.IsValidIndex(SectionIndex))
-							{
-								MontageEditor.Pin()->SetSectionTime(SectionIndex, InValue);
-							}
-
-							FSlateApplication::Get().DismissAllMenus();
-						})
-					]
-				];
-
-			MenuBuilder.AddWidget(TimeWidget, LOCTEXT("SectionTimeMenuText", "Section Time"));
-
-			// Add item to directly set section frame
-			TSharedRef<SWidget> FrameWidget = 
-				SNew( SBox )
-				.HAlign( HAlign_Right )
-				.ToolTipText(LOCTEXT("SetFrameToolTip", "Set the frame of this section directly"))
-				[
-					SNew(SBox)
-					.Padding(FMargin(4.0f, 0.0f, 0.0f, 0.0f))
-					.WidthOverride(100.0f)
-					[
-						SNew(SNumericEntryBox<int32>)
-						.Font(FEditorStyle::GetFontStyle(TEXT("MenuItem.Font")))
-						.MinValue(0)
-						.MaxValue(Montage->GetNumberOfFrames())
-						.Value(Montage->GetFrameAtTime(Section.GetTime()))
-						.AllowSpin(true)						
-						.OnValueCommitted_Lambda([this, SectionIndex](int32 InValue, ETextCommit::Type InCommitType)
-						{
-							if (Montage->CompositeSections.IsValidIndex(SectionIndex))
-							{
-								float NewTime = FMath::Clamp(Montage->GetTimeAtFrame(InValue), 0.0f, Montage->SequenceLength);
-								MontageEditor.Pin()->SetSectionTime(SectionIndex, NewTime);
-							}
-
-							FSlateApplication::Get().DismissAllMenus();
-						})
-					]
-				];
-
-			MenuBuilder.AddWidget(FrameWidget, LOCTEXT("SectionFrameMenuText", "Section Frame"));
-		}
 	}
 	MenuBuilder.EndSection();
 
@@ -441,12 +715,12 @@ void SAnimMontagePanel::SummonTrackContextMenu( FMenuBuilder& MenuBuilder, float
 
 		if(AnimSlotIndex != INDEX_NONE)
 		{
-			UIAction.ExecuteAction.BindRaw(MontageEditor.Pin().Get(), &SMontageEditor::RemoveMontageSlot, AnimSlotIndex);
-			UIAction.CanExecuteAction.BindRaw(MontageEditor.Pin().Get(), &SMontageEditor::CanRemoveMontageSlot, AnimSlotIndex);
+			UIAction.ExecuteAction.BindRaw(this, &SAnimMontagePanel::RemoveMontageSlot, AnimSlotIndex);
+			UIAction.CanExecuteAction.BindRaw(this, &SAnimMontagePanel::CanRemoveMontageSlot, AnimSlotIndex);
 			MenuBuilder.AddMenuEntry(LOCTEXT("DeleteSlot", "Delete Slot"), LOCTEXT("DeleteSlotToolTip", "Deletes Slot"), FSlateIcon(), UIAction);
 			UIAction.CanExecuteAction.Unbind();
 
-			UIAction.ExecuteAction.BindRaw(MontageEditor.Pin().Get(), &SMontageEditor::DuplicateMontageSlot, AnimSlotIndex);
+			UIAction.ExecuteAction.BindRaw(this, &SAnimMontagePanel::DuplicateMontageSlot, AnimSlotIndex);
 			MenuBuilder.AddMenuEntry(LOCTEXT("DuplicateSlot", "Duplicate Slot"), LOCTEXT("DuplicateSlotToolTip", "Duplicates the slected slot"), FSlateIcon(), UIAction);
 		}
 	}
@@ -485,14 +759,14 @@ void SAnimMontagePanel::FillSlotSubMenu(FMenuBuilder& Menubuilder)
 /** Slots */
 void SAnimMontagePanel::OnNewSlotClicked()
 {
-	MontageEditor.Pin()->AddNewMontageSlot(FAnimSlotGroup::DefaultSlotName);
+	AddNewMontageSlot(FAnimSlotGroup::DefaultSlotName);
 }
 
 void SAnimMontagePanel::CreateNewSlot(const FText& NewSlotName, ETextCommit::Type CommitInfo)
 {
 	if (CommitInfo == ETextCommit::OnEnter)
 	{
-		MontageEditor.Pin()->AddNewMontageSlot(*NewSlotName.ToString());
+		AddNewMontageSlot(*NewSlotName.ToString());
 	}
 
 	FSlateApplication::Get().DismissAllMenus();
@@ -528,29 +802,70 @@ void SAnimMontagePanel::CreateNewSection(const FText& NewSectionName, ETextCommi
 {
 	if (CommitInfo == ETextCommit::OnEnter)
 	{
-		MontageEditor.Pin()->AddNewSection(StartTime,NewSectionName.ToString());
+		AddNewSection(StartTime,NewSectionName.ToString());
 	}
 	FSlateApplication::Get().DismissAllMenus();
 }
 
 void SAnimMontagePanel::ShowSegmentInDetailsView(int32 AnimSegmentIndex, int32 AnimSlotIndex)
 {
-	UEditorAnimSegment *Obj = Cast<UEditorAnimSegment>(MontageEditor.Pin()->ShowInDetailsView(UEditorAnimSegment::StaticClass()));
-	if(Obj != NULL)
+	if(!bIsSelecting)
 	{
-		Obj->InitAnimSegment(AnimSlotIndex,AnimSegmentIndex);
+		TGuardValue<bool> GuardValue(bIsSelecting, true);
+
+		UEditorAnimSegment *Obj = Cast<UEditorAnimSegment>(WeakModel.Pin()->ShowInDetailsView(UEditorAnimSegment::StaticClass()));
+		if(Obj != nullptr)
+		{
+			Obj->InitFromAnim(Montage, FOnAnimObjectChange::CreateSP(this, &SAnimMontagePanel::OnMontageChange));
+			Obj->InitAnimSegment(AnimSlotIndex,AnimSegmentIndex);
+		}
 	}
 }
 
-void SAnimMontagePanel::ShowSectionInDetailsView(int32 SectionIndex)
+EActiveTimerReturnType SAnimMontagePanel::TriggerRebuildMontagePanel(double InCurrentTime, float InDeltaTime)
 {
-	MontageEditor.Pin()->ShowSectionInDetailsView(SectionIndex);
+	RebuildMontagePanel();
+
+	bIsActiveTimerRegistered = false;
+	return EActiveTimerReturnType::Stop;
+}
+
+void SAnimMontagePanel::OnMontageChange(UObject* EditorAnimBaseObj, bool bRebuild)
+{
+	if ( Montage != nullptr )
+	{
+		float PreviousSeqLength = GetSequenceLength();
+
+		if(bRebuild && !bIsActiveTimerRegistered)
+		{
+			bIsActiveTimerRegistered = true;
+			RegisterActiveTimer(0.f, FWidgetActiveTimerDelegate::CreateSP(this, &SAnimMontagePanel::TriggerRebuildMontagePanel));
+		} 
+		else
+		{
+			CollapseMontage();
+		}
+
+		// if animation length changed, we might be out of range, let's restart
+		if (GetSequenceLength() != PreviousSeqLength)
+		{
+			// this might not be safe
+			RestartPreview();
+		}
+
+		OnMontageModified();
+	}
 }
 
 void SAnimMontagePanel::ClearSelected()
 {
-	SelectionSet.Empty();
-	MontageEditor.Pin()->ClearDetailsView();
+	if(!bIsSelecting)
+	{
+		TGuardValue<bool> GuardValue(bIsSelecting, true);
+
+		SelectionSet.Empty();
+		WeakModel.Pin()->ClearDetailsView();
+	}
 }
 
 void SAnimMontagePanel::OnSlotNameChanged(TSharedPtr<FString> NewSelection, ESelectInfo::Type SelectInfo, int32 AnimSlotIndex)
@@ -571,29 +886,24 @@ void SAnimMontagePanel::OnSlotNameChanged(TSharedPtr<FString> NewSelection, ESel
 
 			if (Montage->GetSkeleton()->ContainsSlotName(NewSlotName))
 			{
-				MontageEditor.Pin()->RenameSlotNode(AnimSlotIndex, NewSlotName.ToString());
+				RenameSlotNode(AnimSlotIndex, NewSlotName.ToString());
 			}
 			
-			UpdateSlotGroupWarningVisibility();
-
 			// Clear selection, so Details panel for AnimNotifies doesn't show outdated information.
 			ClearSelected();
 		}
 	}
 }
 
-ECheckBoxState SAnimMontagePanel::IsSlotPreviewed(int32 SlotIndex) const
+bool SAnimMontagePanel::IsSlotPreviewed(int32 SlotIndex) const
 {
-	return (SlotIndex == CurrentPreviewSlot) ? ECheckBoxState::Checked : ECheckBoxState::Unchecked;
+	return (SlotIndex == CurrentPreviewSlot);
 }
 
-void SAnimMontagePanel::OnSlotPreviewedChanged(ECheckBoxState NewState, int32 SlotIndex)
+void SAnimMontagePanel::OnSlotPreviewedChanged(int32 SlotIndex)
 {
-	if (NewState == ECheckBoxState::Checked)
-	{
-		CurrentPreviewSlot = SlotIndex;
-		OnSetMontagePreviewSlot.ExecuteIfBound(CurrentPreviewSlot);
-	}
+	CurrentPreviewSlot = SlotIndex;
+	OnSetMontagePreviewSlot(CurrentPreviewSlot);
 }
 
 void SAnimMontagePanel::OnSlotListOpening(int32 AnimSlotIndex)
@@ -602,11 +912,9 @@ void SAnimMontagePanel::OnSlotListOpening(int32 AnimSlotIndex)
 	RefreshComboLists(true);
 }
 
-FReply SAnimMontagePanel::OnOpenAnimSlotManager()
+void SAnimMontagePanel::OnOpenAnimSlotManager()
 {
 	OnInvokeTab.ExecuteIfBound(FPersonaTabs::SkeletonSlotNamesID);
-
-	return FReply::Handled();
 }
 
 void SAnimMontagePanel::RefreshComboLists(bool bOnlyRefreshIfDifferent /*= false*/)
@@ -616,7 +924,7 @@ void SAnimMontagePanel::RefreshComboLists(bool bOnlyRefreshIfDifferent /*= false
 	for (int32 TrackIndex = 0; TrackIndex < NumAnimTracks; TrackIndex++)
 	{
 		FName TrackSlotName = Montage->SlotAnimTracks[TrackIndex].SlotName;
-		Montage->GetSkeleton()->RegisterSlotNode(TrackSlotName);
+		WeakModel.Pin()->GetEditableSkeleton()->RegisterSlotNode(TrackSlotName);
 		SlotNameComboSelectedNames[TrackIndex] = TrackSlotName;
 	}
 
@@ -670,67 +978,6 @@ void SAnimMontagePanel::RefreshComboLists(bool bOnlyRefreshIfDifferent /*= false
 	}
 }
 
-void SAnimMontagePanel::UpdateSlotGroupWarningVisibility()
-{
-	bool bShowStatusBarWarning = false;
-	FName MontageGroupName = Montage->GetGroupName();
-
-	int32 NumAnimTracks = Montage->SlotAnimTracks.Num();
-	if (NumAnimTracks > 0)
-	{
-		TArray<FName> UniqueSlotNameList;
-		for (int32 TrackIndex = 0; TrackIndex < NumAnimTracks; TrackIndex++)
-		{
-			FName CurrentSlotName = SlotNameComboSelectedNames[TrackIndex];
-			FName CurrentSlotGroupName = Montage->GetSkeleton()->GetSlotGroupName(CurrentSlotName);
-
-			int32 SlotNameCount = 0;
-			for (int32 Index = 0; Index < NumAnimTracks; Index++)
-			{
-				if (CurrentSlotName == SlotNameComboSelectedNames[Index])
-				{
-					SlotNameCount++;
-				}
-			}
-			
-			// Verify that slot names are unique.
-			bool bSlotNameAlreadyInUse = UniqueSlotNameList.Contains(CurrentSlotName);
-			if (!bSlotNameAlreadyInUse)
-			{
-				UniqueSlotNameList.Add(CurrentSlotName);
-			}
-
-			bool bDifferentGroupName = (CurrentSlotGroupName != MontageGroupName);
-			bool bShowWarning = bDifferentGroupName || bSlotNameAlreadyInUse;
-			bShowStatusBarWarning = bShowStatusBarWarning || bShowWarning;
-
-			SlotWarningImages[TrackIndex]->SetVisibility(bShowWarning ? EVisibility::Visible : EVisibility::Hidden);
-			if (bDifferentGroupName)
-			{
-				FText WarningText = FText::Format(LOCTEXT("AnimMontagePanel_SlotGroupMismatchToolTipText", "Slot's group '{0}' is different than the Montage's group '{1}'. All slots must belong to the same group."), FText::FromName(CurrentSlotGroupName), FText::FromName(MontageGroupName));
-				SlotWarningImages[TrackIndex]->SetToolTipText(WarningText);
-				StatusBarTextBlock->SetText(WarningText);
-				StatusBarTextBlock->SetToolTipText(WarningText);
-			}
-			if (bSlotNameAlreadyInUse)
-			{
-				FText WarningText = FText::Format(LOCTEXT("AnimMontagePanel_SlotNameAlreadyInUseToolTipText", "Slot named '{0}' is already used in this Montage. All slots must be unique"), FText::FromName(CurrentSlotName));
-				SlotWarningImages[TrackIndex]->SetToolTipText(WarningText);
-				StatusBarTextBlock->SetText(WarningText);
-				StatusBarTextBlock->SetToolTipText(WarningText);
-			}
-		}
-	}
-
-	// Update Status bar
-	StatusBarWarningImage->SetVisibility(bShowStatusBarWarning ? EVisibility::Visible : EVisibility::Hidden);
-	if (!bShowStatusBarWarning)
-	{
-		StatusBarTextBlock->SetText(FText::Format(LOCTEXT("AnimMontagePanel_StatusBarText", "Montage Group: '{0}'"), FText::FromName(MontageGroupName)));
-		StatusBarTextBlock->SetToolTipText(LOCTEXT("AnimMontagePanel_StatusBarToolTipText", "The Montage Group is set by the first slot's group."));
-	}
-}
-
 void SAnimMontagePanel::OnSetElementsToLinkMode(EAnimLinkMethod::Type NewLinkMethod)
 {
 	TArray<FAnimLinkableElement*> Elements;
@@ -854,7 +1101,7 @@ void SAnimMontagePanel::RefreshTimingNodes()
 					.ViewInputMax(ViewInputMax)
 					.DataStartPos(Element->GetElementTime())
 					.Element(Element)
-					.bUseTooltip(false);
+					.bUseTooltip(true);
 
 				Node->SetVisibility(SectionTimingNodeVisibility);
 
@@ -867,4 +1114,54 @@ void SAnimMontagePanel::RefreshTimingNodes()
 	}
 }
 
+void SAnimMontagePanel::SortSections()
+{
+	StaticCastSharedPtr<FAnimModel_AnimMontage>(WeakModel.Pin())->SortSections();
+}
+
+void SAnimMontagePanel::EnsureStartingSection()
+{
+	StaticCastSharedPtr<FAnimModel_AnimMontage>(WeakModel.Pin())->EnsureStartingSection();
+}
+
+void SAnimMontagePanel::EnsureSlotNode()
+{
+	if (Montage && Montage->SlotAnimTracks.Num() == 0)
+	{
+		AddNewMontageSlot(FAnimSlotGroup::DefaultSlotName);
+		OnMontageModified();
+	}
+}
+
+void SAnimMontagePanel::SortAnimSegments()
+{
+	for (int32 I=0; I < Montage->SlotAnimTracks.Num(); I++)
+	{
+		Montage->SlotAnimTracks[I].AnimTrack.SortAnimSegments();
+	}
+}
+
+void SAnimMontagePanel::HandleObjectsSelected(const TArray<UObject*>& InObjects)
+{
+	if(!bIsSelecting)
+	{
+		ClearSelected();
+	}
+}
+
+void SAnimMontagePanel::SetSectionTime(int32 SectionIndex, float NewTime)
+{
+	if(Montage && Montage->CompositeSections.IsValidIndex(SectionIndex))
+	{
+		const FScopedTransaction Transaction(LOCTEXT("EditSection", "Edit Section Start Time"));
+		Montage->Modify();
+	
+		FCompositeSection& Section = Montage->CompositeSections[SectionIndex];
+		Section.SetTime(NewTime);
+		Section.LinkMontage(Montage, NewTime);
+
+		SortAndUpdateMontage();
+	}
+}
+
 #undef LOCTEXT_NAMESPACE
diff --git a/Engine/Source/Editor/Persona/Private/SAnimMontagePanel.h b/Engine/Source/Editor/Persona/Private/SAnimMontagePanel.h
--- a/Engine/Source/Editor/Persona/Private/SAnimMontagePanel.h
+++ b/Engine/Source/Editor/Persona/Private/SAnimMontagePanel.h
@@ -19,6 +19,8 @@ class FMenuBuilder;
 class SBorder;
 class SImage;
 class STextBlock;
+class FAnimModel_AnimMontage;
+class UAnimPreviewInstance;
 
 DECLARE_DELEGATE( FOnMontageLengthChange )
 DECLARE_DELEGATE( FOnMontagePropertyChanged )
@@ -34,12 +36,11 @@ DECLARE_DELEGATE_OneParam( FOnMontageSetPreviewSlot, int32)
 //
 //////////////////////////////////////////////////////////////////////////
 
-class SAnimMontagePanel : public SAnimTrackPanel
+class SAnimMontagePanel : public SAnimTrackPanel, public FEditorUndoClient
 {
 public:
 	SLATE_BEGIN_ARGS( SAnimMontagePanel )
 		: _Montage()
-		, _MontageEditor()
 		, _CurrentPosition()
 		, _ViewInputMin()
 		, _ViewInputMax()
@@ -53,7 +54,6 @@ public:
 	{}
 
 	SLATE_ARGUMENT( class UAnimMontage*, Montage)
-	SLATE_ARGUMENT( TWeakPtr<class SMontageEditor>, MontageEditor)
 	SLATE_ARGUMENT( float, WidgetWidth )
 	SLATE_ATTRIBUTE( float, CurrentPosition )
 	SLATE_ATTRIBUTE( float, ViewInputMin )
@@ -69,13 +69,21 @@ public:
 	SLATE_ARGUMENT(bool, bChildAnimMontage)
 	SLATE_END_ARGS()
 
-	void Construct(const FArguments& InArgs, FSimpleMulticastDelegate& OnSectionsChanged);
+	~SAnimMontagePanel();
 
-	void SetMontage(class UAnimMontage * InMontage);
+	void Construct(const FArguments& InArgs, const TSharedRef<FAnimModel_AnimMontage>& InModel);
+
+	/** Handler for when the preview slot is changed */
+	void OnSetMontagePreviewSlot(int32 SlotIndex);
 
 	// SWidget interface
 	virtual FReply OnMouseButtonDown( const FGeometry& MyGeometry, const FPointerEvent& MouseEvent ) override;
 
+	void RestartPreview();
+	void RestartPreviewFromSection(int32 FromSectionIdx = INDEX_NONE);
+	void RestartPreviewPlayAllSections();
+
+	/** This is the main function that creates the UI widgets for the montage tool.*/
 	void Update();
 
 	// Functions used as delegates that build the context menus for the montage tracks
@@ -103,8 +111,8 @@ public:
 	FText GetMontageSlotName(int32 SlotIndex) const;
 	
 	// Handlers for preview slot checkbox UI
-	ECheckBoxState IsSlotPreviewed(int32 SlotIndex) const;
-	void OnSlotPreviewedChanged(ECheckBoxState NewState, int32 SlotIndex);
+	bool IsSlotPreviewed(int32 SlotIndex) const;
+	void OnSlotPreviewedChanged(int32 SlotIndex);
 
 	// Context menu callback to set all elements in the montage to a given link method
 	void OnSetElementsToLinkMode(EAnimLinkMethod::Type NewLinkMethod);
@@ -118,8 +126,94 @@ public:
 	// Clears the name track of timing nodes and rebuilds them
 	void RefreshTimingNodes();
 
+	/** Sort Composite Sections by Start Time */
+	void SortSections();
+
+	/** Ensure there is at least one section in the montage and that the first section starts at T=0.f */
+	void EnsureStartingSection();
+
+	/** Ensure there is at least one slot node track */
+	void EnsureSlotNode();
+
+	/** Sort Segments by starting time */
+	void SortAnimSegments();
+
+	void SortAndUpdateMontage();
+	void CollapseMontage();
+
+	bool ClampToEndTime(float NewEndTime);
+
+	/** FEditorUndoClient interface */
+	virtual void PostUndo( bool bSuccess ) override;
+	virtual void PostRedo( bool bSuccess ) override;
+
+	void PostRedoUndo();
+	
+	bool GetSectionTime( int32 SectionIndex, float &OutTime ) const;
+
+	bool ValidIndexes(int32 AnimSlotIndex, int32 AnimSegmentIndex) const;
+	bool ValidSection(int32 SectionIndex) const;
+
+	/** Updates Notify trigger offsets to take into account current montage state */
+	void RefreshNotifyTriggerOffsets();
+
+	/** Handle notifies changing, so we rebuild timing */
+	void HandleNotifiesChanged();
+
+	/** One-off active timer to trigger a montage panel rebuild */
+	EActiveTimerReturnType TriggerRebuildMontagePanel(double InCurrentTime, float InDeltaTime);
+
+	/** Rebuilds the montage panel */
+	void RebuildMontagePanel(bool bNotifyAsset=true);
+
+	UAnimPreviewInstance* GetPreviewInstance() const;
+
+public:
+
+	/** These are meant to be callbacks used by the montage editing UI widgets */
+	void					OnMontageChange(class UObject *EditorAnimBaseObj, bool Rebuild);
+
+	void					SetSectionTime(int32 SectionIndex, float NewTime);
+
+	TArray<float>			GetSectionStartTimes() const;
+	TArray<FTrackMarkerBar>	GetMarkerBarInformation() const;
+	TArray<FString>			GetSectionNames() const;
+	TArray<float>			GetAnimSegmentStartTimes() const;
+
+	void					AddNewSection(float StartTime, FString SectionName);
+	void					RemoveSection(int32 SectionIndex);
+	FString					GetSectionName(int32 SectionIndex) const;
+
+	void					RenameSlotNode(int32 SlotIndex, FString NewSlotName);
+
+	// UI Slot Action handlers
+	void					AddNewMontageSlot(FName NewSlotName);
+	void					RemoveMontageSlot(int32 AnimSlotIndex);
+	bool					CanRemoveMontageSlot(int32 AnimSlotIndex);
+	void					DuplicateMontageSlot(int32 AnimSlotIndex);
+
+	void					MakeDefaultSequentialSections();
+	void					ClearSquenceOrdering();
+
+	/** Delegete handlers for when the editor UI is changing the montage */
+	void			PreAnimUpdate();
+	void			PostAnimUpdate();
+	void			OnMontageModified();
+	void			ReplaceAnimationMapping(FName SlotName, int32 SegmentIdx, UAnimSequenceBase* OldSequenceBase, UAnimSequenceBase* NewSequenceBase);
+	bool			IsDiffererentFromParent(FName SlotName, int32 SegmentIdx, const FAnimSegment& Segment);
+
+	/** Delegate fired when montage sections have changed */
+	FSimpleDelegate			OnSectionsChanged;
+
+	FReply	OnFindParentClassInContentBrowserClicked();
+	FReply	OnEditParentClassClicked();
+
+	void HandleObjectsSelected(const TArray<UObject*>& InObjects);
+
+	void OnOpenAnimSlotManager();
+
 private:
-	TWeakPtr<SMontageEditor>	MontageEditor;
+	TWeakPtr<FAnimModel_AnimMontage> WeakModel;
 	TSharedPtr<SBorder> PanelArea;
 	class UAnimMontage* Montage;
 	TAttribute<float> CurrentPosition;
@@ -143,14 +237,28 @@ private:
 
 	/** Member data to allow use to set preview slot from editor */
 	int32					 CurrentPreviewSlot;
-	FOnMontageSetPreviewSlot OnSetMontagePreviewSlot;
+
+	/** If previewing section, it is section used to restart previewing when play button is pushed */
+	int32 PreviewingStartSectionIdx;
+
+	/** Recursion guard for selection */
+	bool bIsSelecting;
+
+	/** If currently previewing all or selected section */
+	bool bPreviewingAllSections : 1;
+
+	/** If currently previewing tracks instead of sections */
+	bool bPreviewingTracks : 1;
+
+	/** If the active timer to trigger a montage panel rebuild is currently registered */
+	bool bIsActiveTimerRegistered : 1;
 
   	/* 
 	 * Child Anim Montage: Child Anim Montage only can replace name of animations, and no other meaningful edits 
 	 * as it will derive every data from Parent. There might be some other data that will allow to be replaced, but for now, it is
 	 * not. 
 	 */
-	bool bChildAnimMontage;
+	bool bChildAnimMontage : 1;
 
 	/************************************************************************/
 	/* Status Bar                                                           */
@@ -160,9 +268,8 @@ private:
 
 	void OnSlotNameChanged(TSharedPtr<FString> NewSelection, ESelectInfo::Type SelectInfo, int32 AnimSlotIndex);
 	void OnSlotListOpening(int32 AnimSlotIndex);
-	FReply OnOpenAnimSlotManager();
+
 	void RefreshComboLists(bool bOnlyRefreshIfDifferent = false);
-	void UpdateSlotGroupWarningVisibility();
 	
 	void CreateNewSlot(const FText& NewSlotName, ETextCommit::Type CommitInfo);
 	void CreateNewSection(const FText& NewSectionName, ETextCommit::Type CommitInfo, float StartTime);
diff --git a/Engine/Source/Editor/Persona/Private/SAnimMontageSectionsPanel.cpp b/Engine/Source/Editor/Persona/Private/SAnimMontageSectionsPanel.cpp
--- a/Engine/Source/Editor/Persona/Private/SAnimMontageSectionsPanel.cpp
+++ b/Engine/Source/Editor/Persona/Private/SAnimMontageSectionsPanel.cpp
@@ -9,69 +9,111 @@
 
 #include "ScopedTransaction.h"
 #include "Widgets/Layout/SExpandableArea.h"
+#include "IPersonaToolkit.h"
+#include "PersonaTabs.h"
+#include "IDocumentation.h"
+#include "Widgets/Layout/SWrapBox.h"
+#include "Widgets/Text/STextBlock.h"
+#include "EditorFontGlyphs.h"
+#include "Preferences/PersonaOptions.h"
+#include "Widgets/Input/SComboButton.h"
+#include "Framework/MultiBox/MultiBoxBuilder.h"
+#include "AnimPreviewInstance.h"
+#include "Animation/DebugSkelMeshComponent.h"
 
 #define LOCTEXT_NAMESPACE "AnimMontageSectionsPanel"
 
+FAnimMontageSectionsSummoner::FAnimMontageSectionsSummoner(TSharedPtr<class FAssetEditorToolkit> InHostingApp, const TSharedRef<IPersonaToolkit>& InPersonaToolkit, FSimpleMulticastDelegate& InOnSectionsChanged)
+	: FWorkflowTabFactory(FPersonaTabs::AnimMontageSectionsID, InHostingApp)
+	, PersonaToolkit(InPersonaToolkit)
+	, OnSectionsChanged(InOnSectionsChanged)
+{
+	TabLabel = LOCTEXT("MontageSectionsTabTitle", "Montage Sections");
+	TabIcon = FSlateIcon(FEditorStyle::GetStyleSetName(), "Persona.Tabs.AnimSlotManager");
+
+	bIsSingleton = true;
+
+	ViewMenuDescription = LOCTEXT("MontageSectionsMenu", "Montage Sections");
+	ViewMenuTooltip = LOCTEXT("MontageSectionsMenu_ToolTip", "Manage montage's sections.");
+}
+
+TSharedRef<SWidget> FAnimMontageSectionsSummoner::CreateTabBody(const FWorkflowTabSpawnInfo& Info) const
+{
+	return SNew(SAnimMontageSectionsPanel, PersonaToolkit.Pin().ToSharedRef(), OnSectionsChanged);
+}
 
+TSharedPtr<SToolTip> FAnimMontageSectionsSummoner::CreateTabToolTipWidget(const FWorkflowTabSpawnInfo& Info) const
+{
+	return IDocumentation::Get()->CreateToolTip(LOCTEXT("WindowTooltip", "This tab lets you modify a montage's sections"), NULL, TEXT("Shared/Editors/Persona"), TEXT("AnimMontageSections_Window"));
+}
 
 //////////////////////////////////////////////////////////////////////////
 // SAnimMontagePanel
 
-void SAnimMontageSectionsPanel::Construct(const FArguments& InArgs)
+void SAnimMontageSectionsPanel::Construct(const FArguments& InArgs, const TSharedRef<IPersonaToolkit>& InPersonaToolkit, FSimpleMulticastDelegate& InOnSectionsChanged)
 {
-	MontageEditor = InArgs._MontageEditor;
-	Montage = InArgs._Montage;
-	SelectedCompositeSection = INDEX_NONE;
-	bChildAnimMontage = InArgs._bChildAnimMontage;
+	WeakPersonaToolkit = InPersonaToolkit;
+
+	InOnSectionsChanged.Add(FSimpleDelegate::CreateSP(this, &SAnimMontageSectionsPanel::Update));
+
+	UAnimMontage* Montage = Cast<UAnimMontage>(WeakPersonaToolkit.Pin()->GetAnimationAsset());
+	if(Montage)
+	{
+		Montage->RegisterOnMontageChanged(UAnimMontage::FOnMontageChanged::CreateSP(this, &SAnimMontageSectionsPanel::Update));
+	}
+
+	if(GEditor)
+	{
+		GEditor->RegisterForUndo(this);
+	}
 
 	this->ChildSlot
 	[
-		SNew(SVerticalBox)
-		+SVerticalBox::Slot()
-		.FillHeight(1)
-		[
-			SNew( SExpandableArea )
-			.AreaTitle( LOCTEXT("Sections", "Sections") )
-			.BodyContent()
-			[
-				SNew( SBorder )
-				.Padding( FMargin(2.f, 2.f) )
-				[
-					SAssignNew( PanelArea, SBorder )
-					.BorderImage( FEditorStyle::GetBrush("NoBorder") )
-					.Padding( FMargin(2.0f, 2.0f) )
-					.ColorAndOpacity( FLinearColor::White )
-				]
-			]
-		]
+		SAssignNew( PanelArea, SBorder )
+		.BorderImage( FEditorStyle::GetBrush("NoBorder") )
+		.ColorAndOpacity( FLinearColor::White )
 	];
 
 	Update();
 }
 
+SAnimMontageSectionsPanel::~SAnimMontageSectionsPanel()
+{
+	if(TSharedPtr<IPersonaToolkit> PersonaToolkit = WeakPersonaToolkit.Pin())
+	{
+		UAnimMontage* Montage = Cast<UAnimMontage>(PersonaToolkit->GetAnimationAsset());
+		if(Montage)
+		{
+			Montage->UnregisterOnMontageChanged(this);
+		}
+	}
+
+	if(GEditor)
+	{
+		GEditor->RegisterForUndo(this);
+	}
+}
+
 void SAnimMontageSectionsPanel::Update()
 {
 	int32 ColorIdx=0;
 	FLinearColor Colors[] = { FLinearColor(0.9f, 0.9f, 0.9f, 0.9f), FLinearColor(0.5f, 0.5f, 0.5f) };
-	FLinearColor NodeColor = FLinearColor(0.f, 0.5f, 0.0f, 0.5f);
+	FLinearColor NodeColor = GetDefault<UPersonaOptions>()->SectionTimingNodeColor;
 	FLinearColor SelectedColor = FLinearColor(1.0f,0.65,0.0f);
 	FLinearColor LoopColor = FLinearColor(0.0f, 0.25f, 0.25f, 0.5f);
 
-	
-	if ( Montage != NULL )
+	UAnimMontage* Montage = Cast<UAnimMontage>(WeakPersonaToolkit.Pin()->GetAnimationAsset());
+	if ( Montage != nullptr )
 	{
+		bool bChildAnimMontage = Montage->HasParentAsset();
+
 		TSharedPtr<STrack> Track;
 		TSharedPtr<SVerticalBox> MontageSlots;
 		PanelArea->SetContent(
 			SAssignNew( MontageSlots, SVerticalBox )
 			);
 
-		SectionMap.Empty();
-		TopSelectionSet.Empty();
-		SelectionSet.Empty();
 		MontageSlots->ClearChildren();
-
-		SMontageEditor * Editor = MontageEditor.Pin().Get();
 		
 		TArray<bool>	Used;
 		Used.AddZeroed(Montage->CompositeSections.Num());
@@ -82,97 +124,116 @@ void SAnimMontageSectionsPanel::Update()
 		MontageSlots->AddSlot()
 			.AutoHeight()
 			.VAlign(VAlign_Center)
-			.Padding( FMargin(0.5f, 0.5f) )
 			[
-				SNew(SHorizontalBox)
-				+ SHorizontalBox::Slot()
-				.AutoWidth()
-				.VAlign(VAlign_Center)
-				[				
-					SNew(SButton)
-					.IsEnabled(!bChildAnimMontage)
-					.Visibility( EVisibility::Visible )
-					.Text( LOCTEXT("CreateDefault", "Create Default") )
-					.ToolTipText( LOCTEXT("CreateDefaultToolTip", "Reconstructs section ordering based on start time") )
-					.OnClicked(this, &SAnimMontageSectionsPanel::MakeDefaultSequence)
-					.HAlign(HAlign_Center)
+				SNew(SBorder)
+				.BorderImage(FCoreStyle::Get().GetBrush("ToolPanel.GroupBorder"))
+				[
+					SNew(SHorizontalBox)
+					+ SHorizontalBox::Slot()
+					.AutoWidth()
 					.VAlign(VAlign_Center)
-				]
+					.Padding(2.0f)
+					[				
+						SNew(SButton)
+						.ButtonStyle(FEditorStyle::Get(), "FlatButton.Success")
+						.IsEnabled(!bChildAnimMontage)
+						.ToolTipText( LOCTEXT("CreateDefaultToolTip", "Reconstructs section ordering based on start time") )
+						.OnClicked(this, &SAnimMontageSectionsPanel::MakeDefaultSequence)
+						[
+							SNew(SHorizontalBox)
+							+SHorizontalBox::Slot()
+							.VAlign(VAlign_Center)
+							.Padding(0.0f, 0.0f, 2.0f, 0.0f)
+							.AutoWidth()
+							[
+								SNew(STextBlock)
+								.Text(FEditorFontGlyphs::File)
+								.TextStyle(FEditorStyle::Get(), "FlatButton.DefaultTextStyle")
+								.Font(FEditorStyle::Get().GetFontStyle("FontAwesome.11"))
+							]
+							+SHorizontalBox::Slot()
+							.VAlign(VAlign_Center)
+							.AutoWidth()
+							[
+								SNew(STextBlock)
+								.Text( LOCTEXT("CreateDefault", "Create Default") )
+								.TextStyle(FEditorStyle::Get(), "FlatButton.DefaultTextStyle")
+							]
+						]
+					]
 
-				+ SHorizontalBox::Slot()
-				.AutoWidth()
-				.VAlign(VAlign_Center)
-				[				
-					SNew(SButton)
-					.IsEnabled(!bChildAnimMontage)
-					.Visibility( EVisibility::Visible )
-					.Text( LOCTEXT("Clear", "Clear") )
-					.ToolTipText( LOCTEXT("ClearToolTip", "Resets section orderings") )
-					.OnClicked(this, &SAnimMontageSectionsPanel::ClearSequence)
-					.HAlign(HAlign_Center)
+					+ SHorizontalBox::Slot()
+					.AutoWidth()
 					.VAlign(VAlign_Center)
-				]
-			];
-
-
-		/** Create top track of section nodes */
-		MontageSlots->AddSlot()
-			.AutoHeight()
-			.VAlign(VAlign_Center)
-			.Padding( FMargin(0.5f, 20.0f) )
-			[
-				SAssignNew(Track, STrack)
-				.IsEnabled(!bChildAnimMontage)
-				.ViewInputMin(0)
-				.ViewInputMax(100)
-				.TrackColor( FLinearColor(0.0f, 0.0f, 0.0f, 0.0f))
-				.TrackMaxValue(100)
-				
-			];
-
-		for(int32 SectionIdx=0; SectionIdx < Montage->CompositeSections.Num(); SectionIdx++)
-		{
-			const float NodeLength = 100.f / static_cast<float>(Montage->CompositeSections.Num()+1);
-			const float NodeSpacing = 100.f / static_cast<float>(Montage->CompositeSections.Num());
-
-			Track->AddTrackNode(
-				SNew(STrackNode)
-				.ViewInputMax(100)
-				.ViewInputMin(0)
-				.NodeColor(NodeColor)
-				.SelectedNodeColor(SelectedColor)
-				.DataLength(NodeLength)
-				.DataStartPos(NodeSpacing * SectionIdx)
-				.NodeName(Montage->CompositeSections[SectionIdx].SectionName.ToString())
-				.NodeSelectionSet(&TopSelectionSet)
-				.OnTrackNodeClicked( this, &SAnimMontageSectionsPanel::TopSectionClicked, SectionIdx)
-				.AllowDrag(false)
-				);
-		}
+					.Padding(2.0f)
+					[				
+						SNew(SButton)
+						.ButtonStyle(FEditorStyle::Get(), "FlatButton.Danger")
+						.TextStyle(FEditorStyle::Get(), "FlatButton.DefaultTextStyle")
+						.IsEnabled(!bChildAnimMontage)
+						.ToolTipText( LOCTEXT("ClearToolTip", "Resets section orderings") )
+						.OnClicked(this, &SAnimMontageSectionsPanel::ClearSequence)
+						[
+							SNew(SHorizontalBox)
+							+SHorizontalBox::Slot()
+							.VAlign(VAlign_Center)
+							.Padding(0.0f, 0.0f, 2.0f, 0.0f)
+							.AutoWidth()
+							[
+								SNew(STextBlock)
+								.Text(FEditorFontGlyphs::Times)
+								.TextStyle(FEditorStyle::Get(), "FlatButton.DefaultTextStyle")
+								.Font(FEditorStyle::Get().GetFontStyle("FontAwesome.11"))
+							]
+							+SHorizontalBox::Slot()
+							.VAlign(VAlign_Center)
+							.AutoWidth()
+							[
+								SNew(STextBlock)
+								.Text( LOCTEXT("Clear", "Clear") )
+								.TextStyle(FEditorStyle::Get(), "FlatButton.DefaultTextStyle")
+							]
+						]
+					]
 
-		MontageSlots->AddSlot()
-			.AutoHeight()
-			.VAlign(VAlign_Center)
-			.Padding( FMargin(0.5f, 0.0f) )
-			[
-				SNew(SHorizontalBox)
-				+ SHorizontalBox::Slot()
-				.AutoWidth()
-				.VAlign(VAlign_Center)
-				[				
-					SNew(SButton)
-					.Visibility( EVisibility::Visible )
-					.Text( LOCTEXT("PreviewAll", "Preview All Sections") )
-					.ToolTipText( LOCTEXT("PreviewAllToolTip", "Preview all sections in order they are") )
-					.OnClicked(this, &SAnimMontageSectionsPanel::PreviewAllSectionsClicked)
-					.HAlign(HAlign_Center)
+					+ SHorizontalBox::Slot()
+					.AutoWidth()
 					.VAlign(VAlign_Center)
+					.Padding(2.0f)
+					[				
+						SNew(SButton)
+						.ButtonStyle(FEditorStyle::Get(), "FlatButton")
+						.ToolTipText( LOCTEXT("PreviewAllToolTip", "Preview all sections in their specified order") )
+						.OnClicked(this, &SAnimMontageSectionsPanel::PreviewAllSectionsClicked)
+						[
+							SNew(SHorizontalBox)
+							+SHorizontalBox::Slot()
+							.VAlign(VAlign_Center)
+							.Padding(0.0f, 0.0f, 2.0f, 0.0f)
+							.AutoWidth()
+							[
+								SNew(STextBlock)
+								.Text(FEditorFontGlyphs::Play)
+								.TextStyle(FEditorStyle::Get(), "FlatButton.DefaultTextStyle")
+								.Font(FEditorStyle::Get().GetFontStyle("FontAwesome.11"))
+							]
+							+SHorizontalBox::Slot()
+							.VAlign(VAlign_Center)
+							.AutoWidth()
+							[
+								SNew(STextBlock)
+								.Text( LOCTEXT("PreviewAll", "Preview All") )
+								.TextStyle(FEditorStyle::Get(), "FlatButton.DefaultTextStyle")
+							]
+						]
+					]
 				]
 			];
 
 		/** Create as many tracks as necessary to show each section at least once
 		  * -Each track represents one chain of sections (section->next->next)
 		  */
+		TSharedPtr<SWrapBox> WrapBox;
 
 		while(true)
 		{
@@ -194,43 +255,56 @@ void SAnimMontageSectionsPanel::Update()
 			}
 
 			/** Create new track */
-			SectionMap.Add( TArray<int32>() );
 			MontageSlots->AddSlot()
 				.AutoHeight()
 				.VAlign(VAlign_Center)
-				.Padding( FMargin(0.5f, 0.5f) )
+				.Padding(0.0f, 2.0f, 0.0f, 0.0f)
 				[
-					SNew(SHorizontalBox)
-					+ SHorizontalBox::Slot()
-					.AutoWidth()
-					.VAlign(VAlign_Center)
-					[				
-						SNew(SButton)
-						.Visibility( EVisibility::Visible )
-						.Text( LOCTEXT("Preview", "Preview") )
-						.ToolTipText( LOCTEXT("PreviewToolTip", "Preview this track") )
-						.OnClicked(this, &SAnimMontageSectionsPanel::PreviewSectionClicked, SectionIdx)
-						.HAlign(HAlign_Center)
-						.VAlign(VAlign_Center)
-					]
-
-					+ SHorizontalBox::Slot()
-					.FillWidth(1.0f)
-					.VAlign(VAlign_Center)
+					SNew(SBorder)
+					.BorderImage(FCoreStyle::Get().GetBrush("ToolPanel.GroupBorder"))
 					[
-						SAssignNew(Track, STrack)
-						.ViewInputMin(0)
-						.ViewInputMax(100)
-						.TrackColor( Colors[ColorIdx++ & 1])
-						.TrackMaxValue(100)
+						SAssignNew(WrapBox, SWrapBox)
+						.UseAllottedWidth(true)
 					]
 				];
 			
+			WrapBox->AddSlot()
+			.Padding(2.0f)
+			[
+				SNew(SHorizontalBox)
+				+SHorizontalBox::Slot()
+				.VAlign(VAlign_Center)
+				[
+					SNew(SButton)
+					.ToolTipText( LOCTEXT("PreviewToolTip", "Preview this track") )
+					.ButtonStyle(FEditorStyle::Get(), "FlatButton")
+					.TextStyle(FEditorStyle::Get(), "FlatButton.DefaultTextStyle")
+					.OnClicked(this, &SAnimMontageSectionsPanel::PreviewSectionClicked, SectionIdx)
+					[
+						SNew(SHorizontalBox)
+						+SHorizontalBox::Slot()
+						.VAlign(VAlign_Center)
+						.Padding(0.0f, 0.0f, 2.0f, 0.0f)
+						.AutoWidth()
+						[
+							SNew(STextBlock)
+							.Text(FEditorFontGlyphs::Play)
+							.TextStyle(FEditorStyle::Get(), "FlatButton.DefaultTextStyle")
+							.Font(FEditorStyle::Get().GetFontStyle("FontAwesome.11"))
+						]
+						+SHorizontalBox::Slot()
+						.VAlign(VAlign_Center)
+						.AutoWidth()
+						[
+							SNew(STextBlock)
+							.Text( LOCTEXT("Preview", "Preview") )
+							.TextStyle(FEditorStyle::Get(), "FlatButton.DefaultTextStyle")
+						]
+					]
+				]
+			];
+
 			/** Add each track in this chain to the track we just created */
-			int count =0;
-			float TrackPos = 0;
-			const float BarLength = 8;
-			const float XLength = 1;
 			while(true)
 			{
 				/** Add section if it hasn't already been used in this row (if its used in another row, thats ok) */
@@ -238,53 +312,88 @@ void SAnimMontageSectionsPanel::Update()
 				{
 					UsedInThisRow[SectionIdx] = true;
 					Used[SectionIdx] = true;
-
-					SectionMap[RowIdx].Add(SectionIdx);
 				
-					Track->AddTrackNode(
-						SNew(STrackNode)
+					int32 NextSectionIdx = Montage->GetSectionIndex( Montage->CompositeSections[SectionIdx].NextSectionName );
+
+					WrapBox->AddSlot()
+					.Padding(2.0f)
+					[
+						SNew(SButton)
 						.IsEnabled(!bChildAnimMontage)
-						.ViewInputMax(100)
-						.ViewInputMin(0)
-						.NodeColor( IsLoop(SectionIdx) ? LoopColor : NodeColor)
-						.SelectedNodeColor(SelectedColor)
-						.DataLength(BarLength)
-						.DataStartPos(TrackPos)
-						.NodeName(Montage->CompositeSections[SectionIdx].SectionName.ToString())
-						.OnTrackNodeDragged( this, &SAnimMontageSectionsPanel::SetSectionPos, SectionIdx, RowIdx)
-						.OnTrackNodeDropped( this, &SAnimMontageSectionsPanel::OnSectionDrop)
-						.OnTrackNodeClicked( this, &SAnimMontageSectionsPanel::SectionClicked, SectionIdx)
-						.NodeSelectionSet(&SelectionSet)
-						.AllowDrag(false)
-					);
-					TrackPos += BarLength + 0.25f;
-
-					/** If this has a next section, create an X to delete that link */
-					if(Montage->CompositeSections[SectionIdx].NextSectionName != NAME_None)
+						.ButtonColorAndOpacity(IsLoop(SectionIdx) ? LoopColor : NodeColor)
+						.ButtonStyle(FEditorStyle::Get(), "FlatButton.Light")
+						[
+							SNew(SHorizontalBox)
+							+SHorizontalBox::Slot()
+							.VAlign(VAlign_Center)
+							.Padding(0.0f, 0.0f, 2.0f, 0.0f)
+							.AutoWidth()
+							[
+								SNew(STextBlock)
+								.Text(FText::FromName(Montage->CompositeSections[SectionIdx].SectionName))
+								.TextStyle(FEditorStyle::Get(), "FlatButton.DefaultTextStyle")
+							]
+						]
+					];
+
+					FText TextIcon;
+					FText ToolTipText;
+					if(NextSectionIdx != INDEX_NONE)
 					{
-						Track->AddTrackNode(
-							SNew(STrackNode)
-							.IsEnabled(!bChildAnimMontage)
-							.ViewInputMax(100)
-							.ViewInputMin(0)
-							.NodeColor(IsLoop(SectionIdx) ? LoopColor : NodeColor)
-							.SelectedNodeColor(SelectedColor)
-							.DataStartPos(TrackPos)
-							.NodeName(TEXT("x"))
-							.OnTrackNodeDropped( this, &SAnimMontageSectionsPanel::OnSectionDrop)
-							.OnTrackNodeClicked( this, &SAnimMontageSectionsPanel::RemoveLink, SectionIdx)
-							.NodeSelectionSet(&SelectionSet)
-							.AllowDrag(false)
-							);
-
-						TrackPos += XLength + 0.25;
+						if(NextSectionIdx == SectionIdx)
+						{
+							TextIcon = FEditorFontGlyphs::Undo;
+							ToolTipText = LOCTEXT("NextSectionLoopedTooltip", "Pick the next section (currently looped).");
+						}
+						else
+						{
+							if(NextSectionIdx > SectionIdx)
+							{
+								TextIcon = FEditorFontGlyphs::Arrow_Right;
+							}
+							else
+							{
+								TextIcon = FEditorFontGlyphs::Arrow_Left;
+							}
+							ToolTipText = FText::Format(LOCTEXT("NextSectionWithCurrentTooltip", "Pick the next section (currently {0})."), FText::FromName(Montage->CompositeSections[SectionIdx].NextSectionName));
+						}
+					}
+					else
+					{
+						TextIcon = FEditorFontGlyphs::Square;
+						ToolTipText = LOCTEXT("NextSectionTooltip", "Pick the next section.");
 					}
 
-					count++;
-					SectionIdx = Montage->GetSectionIndex( Montage->CompositeSections[SectionIdx].NextSectionName );
+					WrapBox->AddSlot()
+					.Padding(2.0f)
+					.VAlign(VAlign_Center)
+					[
+						SNew(SHorizontalBox)
+						+SHorizontalBox::Slot()
+						.VAlign(VAlign_Center)
+						[
+							SNew(SComboButton)
+							.ToolTipText(ToolTipText)
+							.IsEnabled(!bChildAnimMontage)
+							.ButtonStyle(FEditorStyle::Get(), "FlatButton")
+							.HasDownArrow(false)
+							.OnGetMenuContent(this, &SAnimMontageSectionsPanel::OnGetSectionMenuContent, SectionIdx)
+							.ContentPadding(FMargin(4.0, 2.0))
+							.ButtonContent()
+							[
+								SNew(STextBlock)
+								.Text(TextIcon)
+								.TextStyle(FEditorStyle::Get(), "FlatButton.DefaultTextStyle")
+								.Font(FEditorStyle::Get().GetFontStyle("FontAwesome.11"))
+							]
+						]
+					];
+
+					SectionIdx = NextSectionIdx;
 
 					continue;
-				} else
+				} 
+				else
 				{
 					break;
 				}
@@ -295,67 +404,173 @@ void SAnimMontageSectionsPanel::Update()
 		RowIdx++;
 
 	}
-	SelectedCompositeSection = INDEX_NONE;
+	else
+	{
+		PanelArea->ClearContent();
+	}
 }
 
-void SAnimMontageSectionsPanel::SetSectionPos(float NewPos, int32 SectionIdx, int32 RowIdx)
+TSharedRef<SWidget> SAnimMontageSectionsPanel::OnGetSectionMenuContent(int32 SectionIdx)
 {
-	// Drag drop of sections not implemented yet
+	UAnimMontage* Montage = Cast<UAnimMontage>(WeakPersonaToolkit.Pin()->GetAnimationAsset());
+
+	FMenuBuilder MenuBuilder(true, nullptr);
+
+	MenuBuilder.BeginSection("Section", LOCTEXT("SectionHeader", "Section"));
+	{
+		if(Montage->CompositeSections[SectionIdx].NextSectionName != NAME_None)
+		{
+			MenuBuilder.AddMenuEntry(
+				LOCTEXT("RemoveLink", "Remove Link"),
+				LOCTEXT("RemoveLinkTooltip", "Remove the link between this section and the next"),
+				FSlateIcon(),
+				FUIAction(FExecuteAction::CreateSP(this, &SAnimMontageSectionsPanel::RemoveLink, SectionIdx))
+			);
+		}
+	}
+	MenuBuilder.EndSection();
+
+	MenuBuilder.BeginSection("NextSection", LOCTEXT("NExtSectionHeader", "Next Section"));
+	{
+		for(int32 NextSectionIdx = 0; NextSectionIdx < Montage->CompositeSections.Num(); NextSectionIdx++)
+		{
+			FText SectionNameText = FText::FromName(Montage->CompositeSections[NextSectionIdx].SectionName);
+			MenuBuilder.AddMenuEntry(
+				SectionNameText,
+				FText::Format(LOCTEXT("RemoveLinkTooltip", "Set {0} to be the next section"), SectionNameText),
+				FSlateIcon(),
+				FUIAction(FExecuteAction::CreateSP(this, &SAnimMontageSectionsPanel::SetNextSectionIndex, SectionIdx, NextSectionIdx))
+			);
+		}
+	}
+	MenuBuilder.EndSection();
+
+	return MenuBuilder.MakeWidget();
 }
 
-void SAnimMontageSectionsPanel::OnSectionDrop()
+void SAnimMontageSectionsPanel::RestartPreview()
 {
-	Update();
+	if (UDebugSkelMeshComponent* MeshComponent = WeakPersonaToolkit.Pin()->GetPreviewScene()->GetPreviewMeshComponent())
+	{
+		if (UAnimPreviewInstance* Preview = MeshComponent->PreviewInstance)
+		{
+			Preview->MontagePreview_PreviewNormal(INDEX_NONE, Preview->IsPlaying());
+		}
+	}
 }
 
-void SAnimMontageSectionsPanel::TopSectionClicked(int32 SectionIndex)
+void SAnimMontageSectionsPanel::RestartPreviewPlayAllSections()
 {
-	// Top section was clicked, add it to next section if possible
-	if(Montage && Montage->IsValidSectionIndex(SelectedCompositeSection))
+	if (UDebugSkelMeshComponent* MeshComponent = WeakPersonaToolkit.Pin()->GetPreviewScene()->GetPreviewMeshComponent())
 	{
-		const FScopedTransaction Transaction( LOCTEXT("OnAddSectionToSectionChain", "Add Section to Composite Sections") );
-		Montage->Modify();
+		if (UAnimPreviewInstance* Preview = MeshComponent->PreviewInstance)
+		{
+			Preview->MontagePreview_PreviewAllSections(Preview->IsPlaying());
+		}
+	}
+}
+
+void SAnimMontageSectionsPanel::RestartPreviewFromSection(int32 FromSectionIdx)
+{
+	if (UDebugSkelMeshComponent* MeshComponent = WeakPersonaToolkit.Pin()->GetPreviewScene()->GetPreviewMeshComponent())
+	{
+		if (UAnimPreviewInstance* Preview = MeshComponent->PreviewInstance)
+		{
+			Preview->MontagePreview_PreviewNormal(FromSectionIdx, Preview->IsPlaying());
+		}
+	}
+}
 
-		FName PrevNext = Montage->CompositeSections[SelectedCompositeSection].NextSectionName;
+void SAnimMontageSectionsPanel::SortSections()
+{
+	UAnimMontage* Montage = Cast<UAnimMontage>(WeakPersonaToolkit.Pin()->GetAnimationAsset());
 
-		Montage->CompositeSections[SelectedCompositeSection].NextSectionName = Montage->CompositeSections[SectionIndex].SectionName;
-		MontageEditor.Pin()->RestartPreview();
+	struct FCompareSections
+	{
+		bool operator()( const FCompositeSection &A, const FCompositeSection &B ) const
+		{
+			return A.GetTime() < B.GetTime();
+		}
+	};
+	if (Montage != nullptr)
+	{
+		Montage->CompositeSections.Sort(FCompareSections());
+	}
+}
+
+void SAnimMontageSectionsPanel::MakeDefaultSequentialSections()
+{
+	UAnimMontage* Montage = Cast<UAnimMontage>(WeakPersonaToolkit.Pin()->GetAnimationAsset());
+	if( Montage != nullptr )
+	{
+		const FScopedTransaction Transaction( LOCTEXT("MakeDefaultSequentialSections", "Make default sequential sections") );
+		Montage->Modify();
+
+		SortSections();
+		for(int32 SectionIdx=0; SectionIdx < Montage->CompositeSections.Num(); SectionIdx++)
+		{
+			Montage->CompositeSections[SectionIdx].NextSectionName = Montage->CompositeSections.IsValidIndex(SectionIdx+1) ? Montage->CompositeSections[SectionIdx+1].SectionName : NAME_None;
+		}
+		RestartPreview();
 		Update();
 
 		Montage->PostEditChange();
 	}
-	MontageEditor.Pin()->ShowSectionInDetailsView(SectionIndex);
-	TopSelectionSet.Empty();
 }
 
-FReply SAnimMontageSectionsPanel::PreviewAllSectionsClicked()
+void SAnimMontageSectionsPanel::ClearSequenceOrdering()
 {
-	MontageEditor.Pin()->RestartPreviewPlayAllSections();
-	return FReply::Handled();
+	UAnimMontage* Montage = Cast<UAnimMontage>(WeakPersonaToolkit.Pin()->GetAnimationAsset());
+	if( Montage != nullptr )
+	{
+		const FScopedTransaction Transaction( LOCTEXT("ClearSectionOrdering", "Clear section ordering") );
+		Montage->Modify();
+
+		SortSections();
+		for(int32 SectionIdx=0; SectionIdx < Montage->CompositeSections.Num(); SectionIdx++)
+		{
+			Montage->CompositeSections[SectionIdx].NextSectionName = NAME_None;
+		}
+		RestartPreview();
+	}
 }
 
-FReply SAnimMontageSectionsPanel::PreviewSectionClicked(int32 SectionIndex)
+void SAnimMontageSectionsPanel::SetNextSectionIndex(int32 SectionIdx, int32 NextSectionIndex)
 {
-	MontageEditor.Pin()->RestartPreviewFromSection(SectionIndex);
+	UAnimMontage* Montage = Cast<UAnimMontage>(WeakPersonaToolkit.Pin()->GetAnimationAsset());
+	if(Montage && Montage->IsValidSectionIndex(SectionIdx) && Montage->IsValidSectionIndex(NextSectionIndex))
+	{
+		const FScopedTransaction Transaction( LOCTEXT("SetNextSection", "Set next section") );
+		Montage->Modify();
+
+		Montage->CompositeSections[SectionIdx].NextSectionName = Montage->CompositeSections[NextSectionIndex].SectionName;
+		Update();
+		RestartPreview();
+	}
+}
+
+FReply SAnimMontageSectionsPanel::PreviewAllSectionsClicked()
+{
+	RestartPreviewPlayAllSections(); 
 	return FReply::Handled();
 }
 
-void SAnimMontageSectionsPanel::SectionClicked(int32 SectionIndex)
+FReply SAnimMontageSectionsPanel::PreviewSectionClicked(int32 SectionIndex)
 {
-	// Other section was clicked
-	SelectedCompositeSection = SectionIndex;
-	MontageEditor.Pin()->ShowSectionInDetailsView(SectionIndex);
+	RestartPreviewFromSection(SectionIndex);
+	return FReply::Handled();
 }
 
 void SAnimMontageSectionsPanel::RemoveLink(int32 SectionIndex)
 {
+	UAnimMontage* Montage = Cast<UAnimMontage>(WeakPersonaToolkit.Pin()->GetAnimationAsset());
 	if(Montage && Montage->IsValidSectionIndex(SectionIndex))
 	{
-		const FScopedTransaction Transaction( LOCTEXT("RemoveNextSection", "Remove Next Section from Composite") );
+		const FScopedTransaction Transaction( LOCTEXT("RemoveNextSection", "Remove Next Section from Montage") );
 		Montage->Modify();
 		
 		Montage->CompositeSections[SectionIndex].NextSectionName = NAME_None;
-		MontageEditor.Pin()->RestartPreview();
+		RestartPreview();
 		Update();
 
 		Montage->PostEditChange();
@@ -364,26 +579,33 @@ void SAnimMontageSectionsPanel::RemoveLink(int32 SectionIndex)
 
 FReply SAnimMontageSectionsPanel::MakeDefaultSequence()
 {
-	if(MontageEditor.Pin().IsValid())
+	UAnimMontage* Montage = Cast<UAnimMontage>(WeakPersonaToolkit.Pin()->GetAnimationAsset());
+	if(Montage)
 	{
-		MontageEditor.Pin().Get()->MakeDefaultSequentialSections();
+		MakeDefaultSequentialSections();
 		Update();
+		Montage->PostEditChange();
 	}
+
 	return FReply::Handled();
 }
 
 FReply SAnimMontageSectionsPanel::ClearSequence()
 {
-	if(MontageEditor.Pin().IsValid())
+	UAnimMontage* Montage = Cast<UAnimMontage>(WeakPersonaToolkit.Pin()->GetAnimationAsset());
+	if(Montage)
 	{
-		MontageEditor.Pin().Get()->ClearSquenceOrdering();
+		ClearSequenceOrdering();
 		Update();
+		Montage->PostEditChange();
 	}
+
 	return FReply::Handled();
 }
 
 bool SAnimMontageSectionsPanel::IsLoop(int32 SectionIdx)
 {
+	UAnimMontage* Montage = Cast<UAnimMontage>(WeakPersonaToolkit.Pin()->GetAnimationAsset());
 	if(Montage && Montage->CompositeSections.IsValidIndex(SectionIdx))
 	{
 		TArray<bool>	Used;
diff --git a/Engine/Source/Editor/Persona/Private/SAnimMontageSectionsPanel.h b/Engine/Source/Editor/Persona/Private/SAnimMontageSectionsPanel.h
--- a/Engine/Source/Editor/Persona/Private/SAnimMontageSectionsPanel.h
+++ b/Engine/Source/Editor/Persona/Private/SAnimMontageSectionsPanel.h
@@ -9,6 +9,8 @@
 #include "Widgets/SCompoundWidget.h"
 #include "SNodePanel.h"
 #include "STrack.h"
+#include "WorkflowOrientedApp/WorkflowTabFactory.h"
+#include "EditorUndoClient.h"
 
 class SBorder;
 class SMontageEditor;
@@ -17,39 +19,45 @@ class UAnimMontage;
 // Forward declatations
 class UAnimMontage;
 class SMontageEditor;
+class IPersonaToolkit;
 
-class SAnimMontageSectionsPanel : public SCompoundWidget
+struct FAnimMontageSectionsSummoner : public FWorkflowTabFactory
+{
+public:
+	FAnimMontageSectionsSummoner(TSharedPtr<class FAssetEditorToolkit> InHostingApp, const TSharedRef<class IPersonaToolkit>& InPersonaToolkit, FSimpleMulticastDelegate& InOnSectionsChanged);
+
+	virtual TSharedRef<SWidget> CreateTabBody(const FWorkflowTabSpawnInfo& Info) const override;
+	virtual TSharedPtr<SToolTip> CreateTabToolTipWidget(const FWorkflowTabSpawnInfo& Info) const override;
+
+public:
+	TWeakPtr<class IPersonaToolkit> PersonaToolkit;
+	FSimpleMulticastDelegate& OnSectionsChanged;
+};
+
+class SAnimMontageSectionsPanel : public SCompoundWidget, public FEditorUndoClient
 {
 public:
 	SLATE_BEGIN_ARGS( SAnimMontageSectionsPanel )
-		: _Montage()
-		, _MontageEditor()
-		, _bChildAnimMontage(false)
 	{}
 
-		SLATE_ARGUMENT( UAnimMontage*, Montage)						// The montage asset being edited
-		SLATE_ARGUMENT( TWeakPtr<SMontageEditor>, MontageEditor )	// The montage editor that owns this widget
-		SLATE_ARGUMENT( bool, bChildAnimMontage)
-
 	SLATE_END_ARGS()
 
-	void Construct(const FArguments& InArgs);
+	~SAnimMontageSectionsPanel();
 
-	// Rebuild panel widgets
-	void Update();
+	// FEditorUndoClient interface
+	virtual void PostUndo(bool bSuccess) { Update(); };
+	virtual void PostRedo(bool bSuccess) { Update(); };
 
+	void Construct(const FArguments& InArgs, const TSharedRef<IPersonaToolkit>& InPersonaToolkit, FSimpleMulticastDelegate& InOnSectionsChanged);
 
-	void SetSectionPos(float NewPos, int32 SectionIdx, int32 RowIdx);
-	// Callback when mouse-drag ends for a section
-	void OnSectionDrop();
+	// Rebuild panel widgets
+	void Update();
 
 	// Callback when a section in the upper display is clicked
-	// @param SectionIndex - Section that was clicked
-	void TopSectionClicked(int32 SectionIndex);
+	// @param SectionIdx - Section that was clicked
+	// @param NextSecitonIdx - Section to set as the next section
+	void SetNextSectionIndex(int32 SectionIdx, int32 NextSecitonIdx);
 
-	// Callback when a section in the lower display is clicked
-	// @param SectionIndex - Section that was clicked
-	void SectionClicked(int32 SectionIndex);
 	// Unlinks the requested section
 	// @param SectionIndex - Section to unlink
 	void RemoveLink(int32 SectionIndex);
@@ -68,6 +76,24 @@ public:
 
 private:
 
+	/** Shows the context menu for a section */
+	TSharedRef<SWidget> OnGetSectionMenuContent(int32 SectionIdx);
+
+	/** Restarts the preview animation */
+	void RestartPreview();
+
+	/** Restarts the preview animation playing all sections */
+	void RestartPreviewPlayAllSections();
+
+	/** Restarts the preview animation from the specified sections */
+	void RestartPreviewFromSection(int32 FromSectionIdx);
+
+	void SortSections();
+
+	void MakeDefaultSequentialSections();
+
+	void ClearSequenceOrdering();
+
 	// Returns whether or not the provided section is part of a loop
 	// @param SectionIdx - Section to check
 	bool IsLoop(int32 SectionIdx);
@@ -75,24 +101,6 @@ private:
 	// Main panel area widget
 	TSharedPtr<SBorder>			PanelArea;
 
-	// The montage we are currently observing
-	UAnimMontage*				Montage;
-	// The montage editor panel we are a child of
-	TWeakPtr<SMontageEditor>	MontageEditor;
-	// Section to row idx mapping
-	TArray<TArray<int32>>		SectionMap;
-
-	// Selection sets 
-	STrackNodeSelectionSet		TopSelectionSet;
-	STrackNodeSelectionSet		SelectionSet;
-
-	// Currently selected section index
-	int32						SelectedCompositeSection;
-
-	/*
-	* Child Anim Montage: Child Anim Montage only can replace name of animations, and no other meaningful edits
-	* as it will derive every data from Parent. There might be some other data that will allow to be replaced, but for now, it is
-	* not.
-	*/
-	bool						bChildAnimMontage;
+	// Persona toolkit we are hosted in
+	TWeakPtr<IPersonaToolkit> WeakPersonaToolkit;
 };
