COPY from //Dev-Anim to //Dev-Main

diff --git a/Engine/Source/Editor/Persona/Private/SAnimNotifyPanel.cpp b/Engine/Source/Editor/Persona/Private/SAnimNotifyPanel.cpp
--- a/Engine/Source/Editor/Persona/Private/SAnimNotifyPanel.cpp
+++ b/Engine/Source/Editor/Persona/Private/SAnimNotifyPanel.cpp
@@ -46,15 +46,14 @@
 #include "Subsystems/AssetEditorSubsystem.h"
 #include "IAnimationEditor.h"
 #include "IAnimationSequenceBrowser.h"
+#include "AnimTimelineTrack_NotifiesPanel.h"
 
-// Track Panel drawing
-const float NotificationTrackHeight = 20.0f;
 
 // AnimNotify Drawing
 const float NotifyHeightOffset = 0.f;
-const float NotifyHeight = NotificationTrackHeight + 3.f;
-const FVector2D ScrubHandleSize(8.f, NotifyHeight);
-const FVector2D AlignmentMarkerSize(8.f, NotifyHeight);
+const float NotifyHeight = FAnimTimelineTrack_NotifiesPanel::NotificationTrackHeight;
+const FVector2D ScrubHandleSize(12.0f, 12.0f);
+const FVector2D AlignmentMarkerSize(10.f, 20.f);
 const FVector2D TextBorderSize(1.f, 1.f);
 
 #define LOCTEXT_NAMESPACE "AnimNotifyPanel"
@@ -149,8 +148,11 @@ struct INodeObjectInterface
 	virtual void CacheName() = 0;
 
 	virtual void Delete(UAnimSequenceBase* Seq) = 0;
+	virtual void MarkForDelete(UAnimSequenceBase* Seq) = 0;
 
 	virtual void ExportForCopy(UAnimSequenceBase* Seq, FString& StrValue) const = 0;
+
+	virtual FGuid GetGuid() const = 0;
 };
 
 struct FNotifyNodeInterface : public INodeObjectInterface
@@ -160,7 +162,10 @@ struct FNotifyNodeInterface : public INodeObjectInterface
 	// Cached notify name (can be generated by blueprints so want to cache this instead of hitting VM) 
 	FName CachedNotifyName;
 
-	FNotifyNodeInterface(FAnimNotifyEvent* InAnimNotifyEvent) : NotifyEvent(InAnimNotifyEvent) {}
+	// Stable Guid that allows us to refer to notify event
+	FGuid Guid;
+
+	FNotifyNodeInterface(FAnimNotifyEvent* InAnimNotifyEvent) : NotifyEvent(InAnimNotifyEvent), Guid(NotifyEvent->Guid) {}
 	virtual ENodeObjectTypes::Type GetType() const override { return ENodeObjectTypes::NOTIFY; }
 	virtual FAnimNotifyEvent* GetNotifyEvent() override { return NotifyEvent; }
 	virtual int GetTrackIndex() const override{ return NotifyEvent->TrackIndex; }
@@ -272,6 +277,19 @@ struct FNotifyNodeInterface : public INodeObjectInterface
 		}
 	}
 
+
+	virtual void MarkForDelete(UAnimSequenceBase* Seq) override
+	{
+		for (int32 I = 0; I < Seq->Notifies.Num(); ++I)
+		{
+			if (NotifyEvent == &(Seq->Notifies[I]))
+			{
+				Seq->Notifies[I].Guid = FGuid();
+				break;
+			}
+		}
+	}
+
 	virtual void ExportForCopy(UAnimSequenceBase* Seq, FString& StrValue) const override
 	{
 		int32 Index = INDEX_NONE;
@@ -293,13 +311,28 @@ struct FNotifyNodeInterface : public INodeObjectInterface
 			ArrayProperty->Inner->ExportTextItem(StrValue, PropertyData, PropertyData, Seq, PPF_Copy);
 		}
 	}
+
+	virtual FGuid GetGuid() const override
+	{
+		return Guid;
+	}
+
+	static void RemoveInvalidNotifies(UAnimSequenceBase* SeqBase)
+	{
+		SeqBase->Notifies.RemoveAll([](const FAnimNotifyEvent& InNotifyEvent){ return !InNotifyEvent.Guid.IsValid(); });
+		SeqBase->PostEditChange();
+		SeqBase->MarkPackageDirty();
+	}
 };
 
 struct FSyncMarkerNodeInterface : public INodeObjectInterface
 {
 	FAnimSyncMarker* SyncMarker;
 
-	FSyncMarkerNodeInterface(FAnimSyncMarker* InSyncMarker) : SyncMarker(InSyncMarker) {}
+	// Stable Guid that allows us to refer to sync marker event
+	FGuid Guid;
+
+	FSyncMarkerNodeInterface(FAnimSyncMarker* InSyncMarker) : SyncMarker(InSyncMarker), Guid(SyncMarker->Guid) {}
 	virtual ENodeObjectTypes::Type GetType() const override { return ENodeObjectTypes::SYNC_MARKER; }
 	virtual FAnimNotifyEvent* GetNotifyEvent() override { return NULL; }
 	virtual int GetTrackIndex() const override{ return SyncMarker->TrackIndex; }
@@ -350,6 +383,21 @@ struct FSyncMarkerNodeInterface : public INodeObjectInterface
 		}
 	}
 
+	virtual void MarkForDelete(UAnimSequenceBase* SeqBase) override
+	{
+		if(UAnimSequence* Seq = Cast<UAnimSequence>(SeqBase))
+		{
+			for (int32 I = 0; I < Seq->AuthoredSyncMarkers.Num(); ++I)
+			{
+				if (SyncMarker == &(Seq->AuthoredSyncMarkers[I]))
+				{
+					Seq->AuthoredSyncMarkers[I].Guid = FGuid();
+					break;
+				}
+			}
+		}
+	}
+
 	virtual void ExportForCopy(UAnimSequenceBase* SeqBase, FString& StrValue) const override
 	{
 		if (UAnimSequence* Seq = Cast<UAnimSequence>(SeqBase))
@@ -374,6 +422,21 @@ struct FSyncMarkerNodeInterface : public INodeObjectInterface
 			}
 		}
 	}
+
+	virtual FGuid GetGuid() const override
+	{
+		return Guid;
+	}
+
+	static void RemoveInvalidSyncMarkers(UAnimSequenceBase* SeqBase)
+	{
+		if(UAnimSequence* Seq = Cast<UAnimSequence>(SeqBase))
+		{
+			Seq->AuthoredSyncMarkers.RemoveAll([](const FAnimSyncMarker& InSyncMarker){ return !InSyncMarker.Guid.IsValid(); });
+			Seq->PostEditChange();
+			Seq->MarkPackageDirty();
+		}
+	}
 };
 
 // Struct that allows us to get the max value of 2 numbers at compile time
@@ -400,10 +463,9 @@ public:
 		, _OnNodeDragStarted()
 		, _OnUpdatePanel()
 		, _PanTrackRequest()
-		, _OnDeselectAllNotifies()
+		, _OnSelectionChanged()
 		, _ViewInputMin()
 		, _ViewInputMax()
-		, _MarkerBars()
 	{
 	}
 	SLATE_ARGUMENT( class UAnimSequenceBase*, Sequence )
@@ -412,11 +474,11 @@ public:
 	SLATE_EVENT( FOnNotifyNodeDragStarted, OnNodeDragStarted )
 	SLATE_EVENT( FOnUpdatePanel, OnUpdatePanel )
 	SLATE_EVENT( FPanTrackRequest, PanTrackRequest )
-	SLATE_EVENT( FDeselectAllNotifies, OnDeselectAllNotifies )
+	SLATE_EVENT( FOnTrackSelectionChanged, OnSelectionChanged )
 	SLATE_ATTRIBUTE( float, ViewInputMin )
 	SLATE_ATTRIBUTE( float, ViewInputMax )
-	SLATE_ATTRIBUTE( TArray<FTrackMarkerBar>, MarkerBars )
 	SLATE_ARGUMENT(TSharedPtr<SAnimTimingNode>, StateEndTimingNode)
+	SLATE_EVENT( FOnSnapPosition, OnSnapPosition )
 	SLATE_END_ARGS()
 
 	void Construct(const FArguments& Declaration);
@@ -449,7 +511,7 @@ public:
 	virtual int32 OnPaint(const FPaintArgs& Args, const FGeometry& AllottedGeometry, const FSlateRect& MyCullingRect, FSlateWindowElementList& OutDrawElements, int32 LayerId, const FWidgetStyle& InWidgetStyle, bool bParentEnabled) const /*override*/;
 
 	/** Helpers to draw scrub handles and snap offsets */
-	void DrawHandleOffset( const float& Offset, const float& HandleCentre, FSlateWindowElementList& OutDrawElements, int32 MarkerLayer, const FGeometry &AllottedGeometry, const FSlateRect& MyCullingRect ) const;
+	void DrawHandleOffset( const float& Offset, const float& HandleCentre, FSlateWindowElementList& OutDrawElements, int32 MarkerLayer, const FGeometry &AllottedGeometry, const FSlateRect& MyCullingRect, FLinearColor NodeColour ) const;
 	void DrawScrubHandle( float ScrubHandleCentre, FSlateWindowElementList& OutDrawElements, int32 ScrubHandleID, const FGeometry &AllottedGeometry, const FSlateRect& MyCullingRect, FLinearColor NodeColour ) const;
 
 	FLinearColor GetNotifyColor() const;
@@ -551,11 +613,11 @@ private:
 	/** Delegate to pan the track, needed if the markers are dragged out of the track */
 	FPanTrackRequest			PanTrackRequest;
 
-	/** Marker bars for snapping to when dragging the markers in a state notify node */
-	TAttribute<TArray<FTrackMarkerBar>> MarkerBars;
+	/** Delegate used to snap positions */
+	FOnSnapPosition				OnSnapPosition;
 
-	/** Delegate to deselect notifies and clear the details panel */
-	FDeselectAllNotifies OnDeselectAllNotifies;
+	/** Delegate to signal selection changing */
+	FOnTrackSelectionChanged	OnSelectionChanged;
 
 	/** Cached owning track geometry */
 	FGeometry CachedTrackGeometry;
@@ -662,7 +724,7 @@ public:
 		SLATE_ARGUMENT( TArray<FAnimSyncMarker *>, AnimSyncMarkers)
 		SLATE_ATTRIBUTE( float, ViewInputMin )
 		SLATE_ATTRIBUTE( float, ViewInputMax )
-		SLATE_ATTRIBUTE( TArray<FTrackMarkerBar>, MarkerBars)
+		SLATE_EVENT( FOnSnapPosition, OnSnapPosition )
 		SLATE_ARGUMENT( int32, TrackIndex )
 		SLATE_ARGUMENT( FLinearColor, TrackColor )
 		SLATE_ATTRIBUTE(EVisibility, QueuedNotifyTimingNodeVisibility)
@@ -697,7 +759,6 @@ public:
 	// SWidget interface
 	virtual void Tick( const FGeometry& AllottedGeometry, const double InCurrentTime, const float InDeltaTime ) override { UpdateCachedGeometry( AllottedGeometry ); }
 	virtual int32 OnPaint(const FPaintArgs& Args, const FGeometry& AllottedGeometry, const FSlateRect& MyCullingRect, FSlateWindowElementList& OutDrawElements, int32 LayerId, const FWidgetStyle& InWidgetStyle, bool bParentEnabled) const override;
-	virtual FReply OnMouseWheel(const FGeometry& MyGeometry, const FPointerEvent& MouseEvent) override;
 	virtual FCursorReply OnCursorQuery(const FGeometry& MyGeometry, const FPointerEvent& CursorEvent) const override;
 	// End of SWidget interface
 
@@ -735,7 +796,7 @@ public:
 	// Gets the currently selected SAnimNotifyNode instances
 	void AppendSelectedNodeWidgetsToArray(TArray<TSharedPtr<SAnimNotifyNode>>& NodeArray) const;
 	// Gets the indices of the selected notifies
-	TArray<int32> GetSelectedNotifyIndices() const {return SelectedNodeIndices;}
+	const TArray<int32>& GetSelectedNotifyIndices() const {return SelectedNodeIndices;}
 
 	INodeObjectInterface* GetNodeObjectInterface(int32 NodeIndex) { return NotifyNodes[NodeIndex]->NodeObjectInterface; }
 	/**
@@ -744,6 +805,15 @@ public:
 	*/
 	void DeselectAllNotifyNodes(bool bUpdateSelectionSet = true);
 
+	/** Select all nodes contained in the supplied Guid set. */
+	void SelectNodesByGuid(const TSet<FGuid>& InGuids, bool bUpdateSelectionSet);
+
+	/** Get the number of notify nodes we contain */
+	int32 GetNumNotifyNodes() const { return NotifyNodes.Num(); }
+
+	/** Check whether a node is selected */
+	bool IsNodeSelected(int32 NodeIndex) const { return NotifyNodes[NodeIndex]->bSelected; }
+
 	// get Property Data of one element (NotifyIndex) from Notifies property of Sequence
 	static uint8* FindNotifyPropertyData(UAnimSequenceBase* Sequence, int32 NotifyIndex, FArrayProperty*& ArrayProperty);
 
@@ -955,8 +1025,8 @@ protected:
 	/** Cached for drag drop handling code */
 	FGeometry CachedGeometry;
 
-	/** Attribute for accessing any marker positions we have to draw */
-	TAttribute<TArray<FTrackMarkerBar>>		MarkerBars;
+	/** Delegate used to snap when dragging */
+	FOnSnapPosition OnSnapPosition;
 
 	/** Nodes that are currently selected */
 	TArray<int32> SelectedNodeIndices;
@@ -978,9 +1048,6 @@ private:
 	/** Pointer to notify panel for drawing*/
 	TWeakPtr<SAnimNotifyPanel>			AnimPanelPtr;
 
-	/** Pointer to the track name UI */
-	TSharedPtr<SInlineEditableTextBlock> TrackText;
-
 public:
 	SLATE_BEGIN_ARGS( SNotifyEdTrack )
 		: _TrackIndex(INDEX_NONE)
@@ -1002,7 +1069,7 @@ public:
 	SLATE_ARGUMENT( float, WidgetWidth )
 	SLATE_ATTRIBUTE( float, ViewInputMin )
 	SLATE_ATTRIBUTE( float, ViewInputMax )
-	SLATE_ATTRIBUTE( TArray<FTrackMarkerBar>, MarkerBars)
+	SLATE_EVENT( FOnSnapPosition, OnSnapPosition )
 	SLATE_ATTRIBUTE( EVisibility, NotifyTimingNodeVisibility )
 	SLATE_ATTRIBUTE( EVisibility, BranchingPointTimingNodeVisibility )
 	SLATE_EVENT( FOnTrackSelectionChanged, OnSelectionChanged)
@@ -1041,12 +1108,6 @@ public:
 		/** Should never be possible but better than crashing the editor */
 		return LOCTEXT("TrackName_Invalid", "Invalid Track");
 	}
-
-	/** Activate the editable text box for the track name */
-	void TriggerRename()
-	{
-		TrackText->EnterEditingMode();
-	}
 };
 
 ////////////////////////////////////////////////////////////////////////////////////////////////////////////////////////////////////////////////////////////////////
@@ -1077,6 +1138,10 @@ public:
 		if ( bDropWasHandled == false )
 		{
 			int32 NumNodes = SelectedNodes.Num();
+
+			const FScopedTransaction Transaction(NumNodes > 0 ? LOCTEXT("MoveNotifiesEvent", "Move Anim Notifies") : LOCTEXT("MoveNotifyEvent", "Move Anim Notify"));
+			Sequence->Modify();
+			
 			for(int32 CurrentNode = 0 ; CurrentNode < NumNodes ; ++CurrentNode)
 			{
 				TSharedPtr<SAnimNotifyNode> Node = SelectedNodes[CurrentNode];
@@ -1105,7 +1170,7 @@ public:
 		FVector2D SelectionBeginPosition = NodeGroupPosition + SelectedNodes[0]->GetNotifyPositionOffset();
 		
 		FTrackClampInfo* SelectionPositionClampInfo = &GetTrackClampInfo(DragDropEvent.GetScreenSpacePosition());
-		if(SelectionPositionClampInfo->NotifyTrack->GetTrackIndex() < TrackSpan)
+		if((SelectionPositionClampInfo->NotifyTrack->GetTrackIndex() + TrackSpan) >= ClampInfos.Num())
 		{
 			// Our selection has moved off the bottom of the notify panel, adjust the clamping information to keep it on the panel
 			SelectionPositionClampInfo = &ClampInfos[ClampInfos.Num() - TrackSpan - 1];
@@ -1150,7 +1215,7 @@ public:
 			float SequenceEnd = TrackScaleInfo.InputToLocalX(Sequence->SequenceLength);
 
 			// Always clamp the Y to the current track
-			SelectionBeginPosition.Y = SelectionPositionClampInfo->TrackPos;
+			SelectionBeginPosition.Y = SelectionPositionClampInfo->TrackPos - 1.0f;
 
 			float SnapX = GetSnapPosition(NodeClamp, TrackNodePos.X);
 			if (FAnimNotifyEvent* CurrentEvent = CurrentNode->NodeObjectInterface->GetNotifyEvent())
@@ -1253,27 +1318,11 @@ public:
 		const float MaxSnapDist = 5.f;
 
 		float CurrentMinSnapDest = MaxSnapDist;
-		float SnapPosition = -1.f;
-
-		if(MarkerBars.IsBound())
-		{
-			const TArray<FTrackMarkerBar>& Bars = MarkerBars.Get();
-			if(Bars.Num() > 0)
-			{
-				for(int32 i = 0; i < Bars.Num(); ++i)
-				{
-					float WidgetSpaceSnapPosition = ScaleInfo.InputToLocalX(Bars[i].Time); // Do comparison in screen space so that zoom does not cause issues
-					float ThisMinSnapDist = FMath::Abs(WidgetSpaceSnapPosition - WidgetSpaceNotifyPosition);
-					if(ThisMinSnapDist < CurrentMinSnapDest)
-					{
-						CurrentMinSnapDest = ThisMinSnapDist;
-						SnapPosition = WidgetSpaceSnapPosition;
-					}
-				}
-			}
-		}
+		float SnapPosition = ScaleInfo.LocalXToInput(WidgetSpaceNotifyPosition);
+		bool bSnapped = OnSnapPosition.IsBound() && OnSnapPosition.Execute(SnapPosition, MaxSnapDist / ScaleInfo.PixelsPerInput);
+		SnapPosition = ScaleInfo.InputToLocalX(SnapPosition);
 
-		if(SnapPosition == -1.0f)
+		if(!bSnapped)
 		{
 			// Didn't snap to a bar, snap to the track bounds
 			float WidgetSpaceEndPosition = ScaleInfo.InputToLocalX(Sequence->SequenceLength);
@@ -1316,7 +1365,6 @@ public:
 	TArray<float>						NodeTimes;				// Times to drop each selected node at
 	float								SnapTime;				// The time that the snapped node was snapped to
 	TWeakPtr<SAnimNotifyNode>			SnappedNode;			// The node chosen for the snap
-	TAttribute<TArray<FTrackMarkerBar>>	MarkerBars;				// Branching point markers
 	TArray<TSharedPtr<SAnimNotifyNode>> SelectedNodes;			// The nodes that are in the current selection
 	TArray<float>						NodeTimeOffsets;		// Time offsets from the beginning of the selection to the nodes.
 	TArray<float>						NodeXOffsets;			// Offsets in X from the widget position to the scrub handle for each node.
@@ -1326,6 +1374,7 @@ public:
 	float								SelectionTimeLength;	// Length of time that the selection covers
 	int32								TrackSpan;				// Number of tracks that the selection spans
 	FOnUpdatePanel						OnUpdatePanel;			// Delegate to redraw the notify panel
+	FOnSnapPosition						OnSnapPosition;			// Delegate used to snap times
 
 	static TSharedRef<FNotifyDragDropOp> New(
 		TArray<TSharedPtr<SAnimNotifyNode>>			NotifyNodes, 
@@ -1337,7 +1386,7 @@ public:
 		const FVector2D&							SelectionSize, 
 		float&										CurrentDragXPosition, 
 		FPanTrackRequest&							RequestTrackPanDelegate, 
-		TAttribute<TArray<FTrackMarkerBar>>			MarkerBars,
+		FOnSnapPosition&							OnSnapPosition,
 		FOnUpdatePanel&								UpdatePanel
 		)
 	{
@@ -1349,10 +1398,10 @@ public:
 		Operation->NodeGroupPosition = SelectionScreenPosition;
 		Operation->NodeGroupSize = SelectionSize;
 		Operation->DragOffset = SelectionScreenPosition - CursorPosition;
-		Operation->MarkerBars = MarkerBars;
+		Operation->OnSnapPosition = OnSnapPosition;
 		Operation->Decorator = Decorator;
 		Operation->SelectedNodes = NotifyNodes;
-		Operation->TrackSpan = NotifyNodes[0]->NodeObjectInterface->GetTrackIndex() - NotifyNodes.Last()->NodeObjectInterface->GetTrackIndex();
+		Operation->TrackSpan = NotifyNodes.Last()->NodeObjectInterface->GetTrackIndex() - NotifyNodes[0]->NodeObjectInterface->GetTrackIndex();
 		
 		// Caclulate offsets for the selected nodes
 		float BeginTime = MAX_flt;
@@ -1477,12 +1526,11 @@ void SAnimNotifyNode::Construct(const FArguments& InArgs)
 
 	OnNodeDragStarted = InArgs._OnNodeDragStarted;
 	PanTrackRequest = InArgs._PanTrackRequest;
-	OnDeselectAllNotifies = InArgs._OnDeselectAllNotifies;
+	OnSelectionChanged = InArgs._OnSelectionChanged;
 
 	ViewInputMin = InArgs._ViewInputMin;
 	ViewInputMax = InArgs._ViewInputMax;
-
-	MarkerBars = InArgs._MarkerBars;
+	OnSnapPosition = InArgs._OnSnapPosition;
 
 	if(InArgs._StateEndTimingNode.IsValid())
 	{
@@ -1496,6 +1544,8 @@ void SAnimNotifyNode::Construct(const FArguments& InArgs)
 			];
 	}
 
+	SetClipping(EWidgetClipping::ClipToBounds);
+
 	SetToolTipText(TAttribute<FText>(this, &SAnimNotifyNode::GetNodeTooltip));
 }
 
@@ -1637,11 +1687,11 @@ void SAnimNotifyNode::UpdateSizeAndPosition(const FGeometry& AllottedGeometry)
 	FVector2D Size = GetSize();
 	float LeftEdgeToNotify = NotifyTimePositionX;
 	float RightEdgeToNotify = AllottedGeometry.Size.X - NotifyTimePositionX;
-	bDrawTooltipToRight = (RightEdgeToNotify > LabelWidth) || (RightEdgeToNotify > LeftEdgeToNotify);
+	bDrawTooltipToRight = NotifyDurationSizeX > 0.0f || ((RightEdgeToNotify > LabelWidth) || (RightEdgeToNotify > LeftEdgeToNotify));
 
 	// Calculate widget width/position based on where we are drawing the tool tip
 	WidgetX = bDrawTooltipToRight ? (NotifyTimePositionX - (NotifyHandleBoxWidth / 2.f)) : (NotifyTimePositionX - LabelWidth);
-	WidgetSize = bDrawTooltipToRight ? FVector2D(FMath::Max(LabelWidth, NotifyDurationSizeX), NotifyHeight) : FVector2D((LabelWidth + NotifyDurationSizeX), NotifyHeight);
+	WidgetSize = bDrawTooltipToRight ? FVector2D((NotifyDurationSizeX > 0.0f ? NotifyDurationSizeX : FMath::Max(LabelWidth, NotifyDurationSizeX)), NotifyHeight) : FVector2D((LabelWidth + NotifyDurationSizeX), NotifyHeight);
 	WidgetSize.X += NotifyHandleBoxWidth;
 	
 	if(EndMarkerNodeOverlay.IsValid())
@@ -1693,25 +1743,18 @@ int32 SAnimNotifyNode::OnPaint(const FPaintArgs& Args, const FGeometry& Allotted
 	}
 
 	const FSlateBrush* StyleInfo = FEditorStyle::GetBrush( TEXT("SpecialEditableTextImageNormal") );
-	FSlateDrawElement::MakeBox( 
-		OutDrawElements,
-		LayerId, 
-		AllottedGeometry.ToPaintGeometry(FVector2D(0,0), AllottedGeometry.Size), 
-		StyleInfo,
-		ESlateDrawEffect::None,
-		FLinearColor(1.0f, 1.0f, 1.0f,0.1f));
 
 	FText Text = GetNotifyText();
-	FLinearColor NodeColour = bSelected ? FLinearColor(1.0f, 0.5f, 0.0f) : FLinearColor::Red;
-	
+	FLinearColor NodeColor = SAnimNotifyNode::GetNotifyColor();
+	FLinearColor BoxColor = bSelected ? FEditorStyle::GetSlateColor("SelectionColor").GetSpecifiedColor() : SAnimNotifyNode::GetNotifyColor();
+
 	float HalfScrubHandleWidth = ScrubHandleSize.X / 2.0f;
 
 	// Show duration of AnimNotifyState
 	if( NotifyDurationSizeX > 0.f )
 	{
-		FLinearColor BoxColor = (NodeObjectInterface->GetTrackIndex() % 2) == 0 ? FLinearColor(0.0f, 1.0f, 0.5f, 0.5f) : FLinearColor(0.0f, 0.5f, 1.0f, 0.5f);
-		FVector2D DurationBoxSize = FVector2D(NotifyDurationSizeX, NotifyHeight);
-		FVector2D DurationBoxPosition = FVector2D(NotifyScrubHandleCentre, 0.f);
+		FVector2D DurationBoxSize = FVector2D(NotifyDurationSizeX, TextSize.Y + TextBorderSize.Y * 2.f);
+		FVector2D DurationBoxPosition = FVector2D(NotifyScrubHandleCentre, (NotifyHeight - TextSize.Y) * 0.5f);
 		FSlateDrawElement::MakeBox( 
 			OutDrawElements,
 			LayerId, 
@@ -1720,7 +1763,7 @@ int32 SAnimNotifyNode::OnPaint(const FPaintArgs& Args, const FGeometry& Allotted
 			ESlateDrawEffect::None,
 			BoxColor);
 
-		DrawScrubHandle(DurationBoxPosition.X + DurationBoxSize.X, OutDrawElements, ScrubHandleID, AllottedGeometry, MyCullingRect, NodeColour);
+		DrawScrubHandle(DurationBoxPosition.X + DurationBoxSize.X, OutDrawElements, ScrubHandleID, AllottedGeometry, MyCullingRect, NodeColor);
 		
 		// Render offsets if necessary
 		if(AnimNotifyEvent && AnimNotifyEvent->EndTriggerTimeOffset != 0.f) //Do we have an offset to render?
@@ -1729,8 +1772,8 @@ int32 SAnimNotifyNode::OnPaint(const FPaintArgs& Args, const FGeometry& Allotted
 			if(EndTime != Sequence->SequenceLength) //Don't render offset when we are at the end of the sequence, doesnt help the user
 			{
 				// ScrubHandle
-				float HandleCentre = NotifyDurationSizeX + ScrubHandleSize.X;
-				DrawHandleOffset(AnimNotifyEvent->EndTriggerTimeOffset, HandleCentre, OutDrawElements, MarkerLayer, AllottedGeometry, MyCullingRect);
+				float HandleCentre = NotifyDurationSizeX + (ScrubHandleSize.X - 2.0f);
+				DrawHandleOffset(AnimNotifyEvent->EndTriggerTimeOffset, HandleCentre, OutDrawElements, MarkerLayer, AllottedGeometry, MyCullingRect, NodeColor);
 			}
 		}
 	}
@@ -1742,43 +1785,17 @@ int32 SAnimNotifyNode::OnPaint(const FPaintArgs& Args, const FGeometry& Allotted
 	FVector2D LabelSize = TextSize + TextBorderSize * 2.f;
 	LabelSize.X += HalfScrubHandleWidth + (bDrawBranchingPoint ? (BranchingPointIconSize.X + TextBorderSize.X * 2.f) : 0.f);
 
-	float LabelX = bDrawTooltipToRight ? NotifyScrubHandleCentre : NotifyScrubHandleCentre - LabelSize.X;
-	float BoxHeight = (NotifyDurationSizeX > 0.f) ? (NotifyHeight - LabelSize.Y) : ((NotifyHeight - LabelSize.Y) / 2.f);
-
-	FVector2D LabelPosition(LabelX, BoxHeight);
-
-	FLinearColor NodeColor = SAnimNotifyNode::GetNotifyColor();
-
-	FSlateDrawElement::MakeBox( 
-		OutDrawElements,
-		LayerId, 
-		AllottedGeometry.ToPaintGeometry(LabelPosition, LabelSize), 
-		StyleInfo,
-		ESlateDrawEffect::None,
-		NodeColor);
+	FVector2D LabelPosition(bDrawTooltipToRight ? NotifyScrubHandleCentre : NotifyScrubHandleCentre - LabelSize.X, (NotifyHeight - TextSize.Y) * 0.5f);
 
-	// Frame
-	// Drawing lines is slow, reserved for single selected node
-	if( bSelected )
+	if( NotifyDurationSizeX == 0.f )
 	{
-		TArray<FVector2D> LinePoints;
-
-		LinePoints.Empty();
-		LinePoints.Add(LabelPosition);
-		LinePoints.Add(LabelPosition + FVector2D(LabelSize.X, 0.f));
-		LinePoints.Add(LabelPosition + FVector2D(LabelSize.X, LabelSize.Y));
-		LinePoints.Add(LabelPosition + FVector2D(0.f, LabelSize.Y));
-		LinePoints.Add(LabelPosition);
-
-		FSlateDrawElement::MakeLines( 
+		FSlateDrawElement::MakeBox( 
 			OutDrawElements,
-			LayerId,
-			AllottedGeometry.ToPaintGeometry(),
-			LinePoints,
+			LayerId, 
+			AllottedGeometry.ToPaintGeometry(LabelPosition, LabelSize), 
+			StyleInfo,
 			ESlateDrawEffect::None,
-			FLinearColor::Black, 
-			false
-			);
+			BoxColor);
 	}
 
 	// Text
@@ -1787,36 +1804,39 @@ int32 SAnimNotifyNode::OnPaint(const FPaintArgs& Args, const FGeometry& Allotted
 	{
 		TextPosition.X += HalfScrubHandleWidth;
 	}
-	TextPosition -= FVector2D(1.f, 1.f);
+
+	FVector2D DrawTextSize;
+	DrawTextSize.X = (NotifyDurationSizeX > 0.0f ? FMath::Min(NotifyDurationSizeX - (ScrubHandleSize.X + (bDrawBranchingPoint ? BranchingPointIconSize.X : 0)), TextSize.X) : TextSize.X);
+	DrawTextSize.Y = TextSize.Y;
+
+	if (bDrawBranchingPoint)
+	{
+		TextPosition.X += BranchingPointIconSize.X;
+	}
+
+	FPaintGeometry TextGeometry = AllottedGeometry.ToPaintGeometry(TextPosition, DrawTextSize);
+	OutDrawElements.PushClip(FSlateClippingZone(TextGeometry));
 
 	FSlateDrawElement::MakeText( 
 		OutDrawElements,
 		TextLayerID,
-		AllottedGeometry.ToPaintGeometry(TextPosition, TextSize),
+		TextGeometry,
 		Text,
 		Font,
 		ESlateDrawEffect::None,
 		FLinearColor::Black
 		);
 
-	DrawScrubHandle(NotifyScrubHandleCentre , OutDrawElements, ScrubHandleID, AllottedGeometry, MyCullingRect, NodeColour);
-
-	if(AnimNotifyEvent && AnimNotifyEvent->TriggerTimeOffset != 0.f) //Do we have an offset to render?
-	{
-		float NotifyTime = AnimNotifyEvent->GetTime();
-		if(NotifyTime != 0.f && NotifyTime != Sequence->SequenceLength) //Don't render offset when we are at the start/end of the sequence, doesn't help the user
-		{
-			float HandleCentre = NotifyScrubHandleCentre;
-			float &Offset = AnimNotifyEvent->TriggerTimeOffset;
-			
-			DrawHandleOffset(AnimNotifyEvent->TriggerTimeOffset, NotifyScrubHandleCentre, OutDrawElements, MarkerLayer, AllottedGeometry, MyCullingRect);
-		}
-	}
+	OutDrawElements.PopClip();
 
 	// Draw Branching Point
 	if (bDrawBranchingPoint)
 	{
-		FVector2D BranchPointIconPos = LabelPosition + LabelSize - BranchingPointIconSize - FVector2D(bDrawTooltipToRight ? TextBorderSize.X * 2.f : TextBorderSize.X * 4.f, 0.f);
+		FVector2D BranchPointIconPos = LabelPosition + TextBorderSize;
+		if(bDrawTooltipToRight)
+		{
+			BranchPointIconPos.X += HalfScrubHandleWidth;
+		}
 		FSlateDrawElement::MakeBox(
 			OutDrawElements,
 			BranchPointLayerID,
@@ -1827,6 +1847,20 @@ int32 SAnimNotifyNode::OnPaint(const FPaintArgs& Args, const FGeometry& Allotted
 			);
 	}	
 
+	DrawScrubHandle(NotifyScrubHandleCentre , OutDrawElements, ScrubHandleID, AllottedGeometry, MyCullingRect, NodeColor);
+
+	if(AnimNotifyEvent && AnimNotifyEvent->TriggerTimeOffset != 0.f) //Do we have an offset to render?
+	{
+		float NotifyTime = AnimNotifyEvent->GetTime();
+		if(NotifyTime != 0.f && NotifyTime != Sequence->SequenceLength) //Don't render offset when we are at the start/end of the sequence, doesn't help the user
+		{
+			float HandleCentre = NotifyScrubHandleCentre;
+			float &Offset = AnimNotifyEvent->TriggerTimeOffset;
+			
+			DrawHandleOffset(AnimNotifyEvent->TriggerTimeOffset, NotifyScrubHandleCentre, OutDrawElements, MarkerLayer, AllottedGeometry, MyCullingRect, NodeColor);
+		}
+	}
+
 	return TextLayerID;
 }
 
@@ -1844,7 +1878,7 @@ FReply SAnimNotifyNode::OnMouseMove( const FGeometry& MyGeometry, const FPointer
 	
 	FTrackScaleInfo ScaleInfo(ViewInputMin.Get(), ViewInputMax.Get(), 0, 0, CachedAllotedGeometrySize);
 	
-	float XPositionInTrack = MyGeometry.AbsolutePosition.X - CachedTrackGeometry.AbsolutePosition.X + ScrubHandleSize.X;
+	float XPositionInTrack = MyGeometry.AbsolutePosition.X - CachedTrackGeometry.AbsolutePosition.X;
 	float TrackScreenSpaceXPosition = MyGeometry.AbsolutePosition.X - XPositionInTrack;
 
 	if(CurrentDragHandle == ENotifyStateHandleHit::Start)
@@ -1962,7 +1996,9 @@ FReply SAnimNotifyNode::OnMouseButtonUp( const FGeometry& MyGeometry, const FPoi
 	{
 		// Clear the drag marker and give the mouse back
 		CurrentDragHandle = ENotifyStateHandleHit::None;
-		OnDeselectAllNotifies.ExecuteIfBound();
+
+		// Signal selection changing so details panels get updated
+		OnSelectionChanged.ExecuteIfBound();
 
 		// End drag transaction before handing mouse back
 		check(DragMarkerTransactionIdx != INDEX_NONE);
@@ -1985,30 +2021,14 @@ float SAnimNotifyNode::GetScrubHandleSnapPosition( float NotifyLocalX, ENotifySt
 	const float MaxSnapDist = 5.0f;
 
 	float CurrentMinSnapDistance = MaxSnapDist;
-	float SnapPosition = -1.0f;
-	float SnapTime = -1.0f;
+	float SnapPosition = ScaleInfo.LocalXToInput(NotifyLocalX);
 
-	if(MarkerBars.IsBound())
+	if(OnSnapPosition.IsBound())
 	{
-		const TArray<FTrackMarkerBar>& Bars = MarkerBars.Get();
-
-		if(Bars.Num() > 0)
-		{
-			for(auto Iter(Bars.CreateConstIterator()) ; Iter ; ++Iter)
-			{
-				const FTrackMarkerBar& Bar = *Iter;
-				float LocalSnapPosition = ScaleInfo.InputToLocalX(Bar.Time);
-				float ThisNodeMinSnapDistance = FMath::Abs(LocalSnapPosition - NotifyLocalX);
-				if(ThisNodeMinSnapDistance < CurrentMinSnapDistance)
-				{
-					CurrentMinSnapDistance = ThisNodeMinSnapDistance;
-					SnapPosition = LocalSnapPosition;
-				}
-			}
-		}
-	}
+		OnSnapPosition.Execute(SnapPosition, MaxSnapDist / ScaleInfo.PixelsPerInput);
+	}	
 
-	return SnapPosition;
+	return ScaleInfo.InputToLocalX(SnapPosition);
 }
 
 FReply SAnimNotifyNode::OnFocusReceived(const FGeometry& MyGeometry, const FFocusEvent& InFocusEvent)
@@ -2041,20 +2061,29 @@ void SAnimNotifyNode::DrawScrubHandle( float ScrubHandleCentre, FSlateWindowElem
 		OutDrawElements,
 		ScrubHandleID, 
 		AllottedGeometry.ToPaintGeometry(ScrubHandlePosition, ScrubHandleSize), 
-		FEditorStyle::GetBrush( TEXT( "Sequencer.Timeline.ScrubHandleWhole" ) ),
+		FEditorStyle::GetBrush( TEXT( "Sequencer.KeyDiamond" ) ),
 		ESlateDrawEffect::None,
 		NodeColour
 		);
+
+	FSlateDrawElement::MakeBox( 
+		OutDrawElements,
+		ScrubHandleID, 
+		AllottedGeometry.ToPaintGeometry(ScrubHandlePosition, ScrubHandleSize), 
+		FEditorStyle::GetBrush( TEXT( "Sequencer.KeyDiamondBorder" ) ),
+		ESlateDrawEffect::None,
+		bSelected ? FEditorStyle::GetSlateColor("SelectionColor").GetSpecifiedColor() : FLinearColor::Black
+		);
 }
 
-void SAnimNotifyNode::DrawHandleOffset( const float& Offset, const float& HandleCentre, FSlateWindowElementList& OutDrawElements, int32 MarkerLayer, const FGeometry &AllottedGeometry, const FSlateRect& MyCullingRect ) const
+void SAnimNotifyNode::DrawHandleOffset( const float& Offset, const float& HandleCentre, FSlateWindowElementList& OutDrawElements, int32 MarkerLayer, const FGeometry &AllottedGeometry, const FSlateRect& MyCullingRect, FLinearColor NodeColor ) const
 {
 	FVector2D MarkerPosition;
 	FVector2D MarkerSize = AlignmentMarkerSize;
 
 	if(Offset < 0.f)
 	{
-		MarkerPosition.Set( (HandleCentre) - AlignmentMarkerSize.X, (NotifyHeight - AlignmentMarkerSize.Y) / 2.f);
+		MarkerPosition.Set( HandleCentre - AlignmentMarkerSize.X, (NotifyHeight - AlignmentMarkerSize.Y) / 2.f);
 	}
 	else
 	{
@@ -2068,7 +2097,7 @@ void SAnimNotifyNode::DrawHandleOffset( const float& Offset, const float& Handle
 		AllottedGeometry.ToPaintGeometry(MarkerPosition, MarkerSize), 
 		FEditorStyle::GetBrush( TEXT( "Sequencer.Timeline.NotifyAlignmentMarker" ) ),
 		ESlateDrawEffect::None,
-		FLinearColor(0.f, 0.f, 1.f)
+		NodeColor
 		);
 }
 
@@ -2083,7 +2112,6 @@ void SAnimNotifyNode::OnFocusLost(const FFocusEvent& InFocusEvent)
 	{
 		// Lost focus while dragging a state node, clear the drag and end the current transaction
 		CurrentDragHandle = ENotifyStateHandleHit::None;
-		OnDeselectAllNotifies.ExecuteIfBound();
 		
 		check(DragMarkerTransactionIdx != INDEX_NONE);
 		GEditor->EndTransaction();
@@ -2142,7 +2170,7 @@ void SAnimNotifyTrack::Construct(const FArguments& InArgs)
 	OnGetDraggedNodePos = InArgs._OnGetDraggedNodePos;
 	OnNodeDragStarted = InArgs._OnNodeDragStarted;
 	TrackColor = InArgs._TrackColor;
-	MarkerBars = InArgs._MarkerBars;
+	OnSnapPosition = InArgs._OnSnapPosition;
 	OnRequestTrackPan = InArgs._OnRequestTrackPan;
 	OnRequestRefreshOffsets = InArgs._OnRequestOffsetRefresh;
 	OnDeleteNotify = InArgs._OnDeleteNotify;
@@ -2170,7 +2198,7 @@ FVector2D SAnimNotifyTrack::ComputeDesiredSize( float ) const
 {
 	FVector2D Size;
 	Size.X = 200;
-	Size.Y = NotificationTrackHeight;
+	Size.Y = FAnimTimelineTrack_NotifiesPanel::NotificationTrackHeight;
 	return Size;
 }
 
@@ -2180,87 +2208,41 @@ int32 SAnimNotifyTrack::OnPaint(const FPaintArgs& Args, const FGeometry& Allotte
 	FLinearColor Color = TrackColor.Get();
 
 	FPaintGeometry MyGeometry = AllottedGeometry.ToPaintGeometry();
-	FSlateDrawElement::MakeBox( 
-		OutDrawElements,
-		LayerId, 
-		MyGeometry, 
-		StyleInfo,
-		ESlateDrawEffect::None, 
-		Color
-		);
 
 	int32 CustomLayerId = LayerId + 1; 
-
-	// draw line for every 1/4 length
 	FTrackScaleInfo ScaleInfo(ViewInputMin.Get(), ViewInputMax.Get(), 0.f, 0.f, AllottedGeometry.Size);
-	const int32 NumOfFrames = Sequence->GetNumberOfFrames();
-	if (NumOfFrames > 0)
-	{
-		int32 Divider = SScrubWidget::GetDivider( ViewInputMin.Get(), ViewInputMax.Get(), AllottedGeometry.Size, Sequence->SequenceLength, Sequence->GetNumberOfFrames());
-		const FAnimKeyHelper Helper(Sequence->SequenceLength, NumOfFrames);
 
-		float TimePerKey = Helper.TimePerKey();
-		for (int32 I=1; I<Helper.GetNumKeys(); ++I)
-		{
-			if ( I % Divider == 0 )
-			{
-				float XPos = ScaleInfo.InputToLocalX(TimePerKey*I);
-
-				TArray<FVector2D> LinePoints;
-				LinePoints.Add(FVector2D(XPos, 0.f));
-				LinePoints.Add(FVector2D(XPos, AllottedGeometry.Size.Y));
-
-				FSlateDrawElement::MakeLines( 
-					OutDrawElements,
-					CustomLayerId,
-					MyGeometry,
-					LinePoints,
-					ESlateDrawEffect::None,
-					FLinearColor::Black
-					);
-			}
-		}
-	}
-
-	++CustomLayerId;
+	bool bAnyDraggedNodes = false;
 	for ( int32 I=0; I<NotifyNodes.Num(); ++I )
 	{
 		if ( NotifyNodes[I].Get()->bBeingDragged == false )
 		{
 			NotifyNodes[I].Get()->UpdateSizeAndPosition(AllottedGeometry);
 		}
+		else
+		{
+			bAnyDraggedNodes = true;
+		}
 	}
 
-	++CustomLayerId;
-
-	float Value = 0.f;
-
-	if ( OnGetScrubValue.IsBound() )
-	{
-		Value = OnGetScrubValue.Execute();
-	}
-
+	if(TrackIndex < Sequence->AnimNotifyTracks.Num() - 1)
 	{
-		float XPos = ScaleInfo.InputToLocalX(Value);
-
-		TArray<FVector2D> LinePoints;
-		LinePoints.Add(FVector2D(XPos, 0.f));
-		LinePoints.Add(FVector2D(XPos, AllottedGeometry.Size.Y));
-
-
-		FSlateDrawElement::MakeLines( 
+		// Draw track bottom border
+		FSlateDrawElement::MakeLines(
 			OutDrawElements,
 			CustomLayerId,
-			MyGeometry,
-			LinePoints,
+			AllottedGeometry.ToPaintGeometry(),
+			TArray<FVector2D>({ FVector2D(0.0f, AllottedGeometry.GetLocalSize().Y), FVector2D(AllottedGeometry.GetLocalSize().X, AllottedGeometry.GetLocalSize().Y) }),
 			ESlateDrawEffect::None,
-			FLinearColor::Red
-			);
+			FLinearColor(0.1f, 0.1f, 0.1f, 0.3f)
+		);
 	}
 
 	++CustomLayerId;
 
-	if ( OnGetDraggedNodePos.IsBound() )
+	float Value = 0.f;
+
+	if ( bAnyDraggedNodes && OnGetDraggedNodePos.IsBound() )
 	{
 		Value = OnGetDraggedNodePos.Execute();
 
@@ -2271,7 +2253,6 @@ int32 SAnimNotifyTrack::OnPaint(const FPaintArgs& Args, const FGeometry& Allotte
 			LinePoints.Add(FVector2D(XPos, 0.f));
 			LinePoints.Add(FVector2D(XPos, AllottedGeometry.Size.Y));
 
-
 			FSlateDrawElement::MakeLines( 
 				OutDrawElements,
 				CustomLayerId,
@@ -2283,53 +2264,9 @@ int32 SAnimNotifyTrack::OnPaint(const FPaintArgs& Args, const FGeometry& Allotte
 		}
 	}
 
-	++CustomLayerId;
-
-	// Draggable Bars
-	for ( int32 I=0; MarkerBars.IsBound() && I < MarkerBars.Get().Num(); I++ )
-	{
-		// Draw lines
-		FTrackMarkerBar Bar = MarkerBars.Get()[I];
-
-		float XPos = ScaleInfo.InputToLocalX(Bar.Time);
-
-		TArray<FVector2D> LinePoints;
-		LinePoints.Add(FVector2D(XPos, 0.f));
-		LinePoints.Add(FVector2D(XPos, AllottedGeometry.Size.Y));
-
-		FSlateDrawElement::MakeLines( 
-			OutDrawElements,
-			CustomLayerId,
-			MyGeometry,
-			LinePoints,
-			ESlateDrawEffect::None,
-			Bar.DrawColour
-			);
-	}
-
 	return SCompoundWidget::OnPaint(Args, AllottedGeometry, MyCullingRect, OutDrawElements, CustomLayerId, InWidgetStyle, bParentEnabled);
 }
 
-FReply SAnimNotifyTrack::OnMouseWheel(const FGeometry& MyGeometry, const FPointerEvent& MouseEvent)
-{
-	const float ZoomDelta = -0.1f * MouseEvent.GetWheelDelta();
-
-	const FVector2D WidgetSpace = MyGeometry.AbsoluteToLocal(MouseEvent.GetScreenSpacePosition());
-	const float ZoomRatio = FMath::Clamp( (WidgetSpace.X / MyGeometry.Size.X), 0.f, 1.f);
-
-	{
-		const float InputViewSize = ViewInputMax.Get() - ViewInputMin.Get();
-		const float InputChange = InputViewSize * ZoomDelta;
-
-		float ViewMinInput = ViewInputMin.Get() - (InputChange * ZoomRatio);
-		float ViewMaxInput = ViewInputMax.Get() + (InputChange * (1.f - ZoomRatio));
-
-		OnSetInputViewRange.ExecuteIfBound(ViewMinInput, ViewMaxInput);
-	}
-
-	return FReply::Handled();
-}
-
 FCursorReply SAnimNotifyTrack::OnCursorQuery(const FGeometry& MyGeometry, const FPointerEvent& CursorEvent) const
 {
 	if (ViewInputMin.Get() > 0.f || ViewInputMax.Get() < Sequence->SequenceLength)
@@ -2537,6 +2474,7 @@ FAnimNotifyEvent& SAnimNotifyTrack::CreateNewNotify(FString NewNotifyName, UClas
 	int32 NewNotifyIndex = Sequence->Notifies.Add(FAnimNotifyEvent());
 	FAnimNotifyEvent& NewEvent = Sequence->Notifies[NewNotifyIndex];
 	NewEvent.NotifyName = FName(*NewNotifyName);
+	NewEvent.Guid = FGuid::NewGuid();
 
 	NewEvent.Link(Sequence, StartTime);
 	NewEvent.TriggerTimeOffset = GetTriggerTimeOffsetForType(Sequence->CalculateOffsetForNotify(StartTime));
@@ -2796,6 +2734,26 @@ void SAnimNotifyTrack::DeselectAllNotifyNodes(bool bUpdateSelectionSet)
 	}
 }
 
+void SAnimNotifyTrack::SelectNodesByGuid(const TSet<FGuid>& InGuids, bool bUpdateSelectionSet)
+{
+	SelectedNodeIndices.Empty();
+
+	for(int32 NodeIndex = 0; NodeIndex < NotifyNodes.Num(); ++NodeIndex)
+	{
+		TSharedPtr<SAnimNotifyNode> Node = NotifyNodes[NodeIndex];
+		Node->bSelected = InGuids.Contains(Node->NodeObjectInterface->GetGuid());
+		if(Node->bSelected)
+		{
+			SelectedNodeIndices.Add(NodeIndex);
+		}
+	}
+
+	if(bUpdateSelectionSet)
+	{
+		SendSelectionChanged();
+	}
+}
+
 TSharedPtr<SWidget> SAnimNotifyTrack::SummonContextMenu(const FGeometry& MyGeometry, const FPointerEvent& MouseEvent)
 {
 	FVector2D CursorPos = MouseEvent.GetScreenSpacePosition();
@@ -2855,9 +2813,9 @@ TSharedPtr<SWidget> SAnimNotifyTrack::SummonContextMenu(const FGeometry& MyGeome
 										}
 									}
 									OnUpdatePanel.ExecuteIfBound();
-								}
 
-								FSlateApplication::Get().DismissAllMenus();
+									FSlateApplication::Get().DismissAllMenus();
+								}
 							})
 						]
 					];
@@ -2898,9 +2856,9 @@ TSharedPtr<SWidget> SAnimNotifyTrack::SummonContextMenu(const FGeometry& MyGeome
 										}
 									}
 									OnUpdatePanel.ExecuteIfBound();
-								}
 
-								FSlateApplication::Get().DismissAllMenus();
+									FSlateApplication::Get().DismissAllMenus();
+								}
 							})
 						]
 					];
@@ -2931,9 +2889,9 @@ TSharedPtr<SWidget> SAnimNotifyTrack::SummonContextMenu(const FGeometry& MyGeome
 									{
 										float NewWeight = FMath::Max(InValue, ZERO_ANIMWEIGHT_THRESH);
 										AnimNotifies[NotifyIndex]->TriggerWeightThreshold = NewWeight;
-									}
 
-									FSlateApplication::Get().DismissAllMenus();
+										FSlateApplication::Get().DismissAllMenus();
+									}
 								})
 							]
 						];
@@ -2972,9 +2930,9 @@ TSharedPtr<SWidget> SAnimNotifyTrack::SummonContextMenu(const FGeometry& MyGeome
 											// If we have a delegate bound to refresh the offsets, call it.
 											// This is used by the montage editor to keep the offsets up to date.
 											OnRequestRefreshOffsets.ExecuteIfBound();
-										}
 
-										FSlateApplication::Get().DismissAllMenus();
+											FSlateApplication::Get().DismissAllMenus();
+										}
 									})
 								]
 							];
@@ -3275,7 +3233,7 @@ void SAnimNotifyTrack::AddNewNotify(const FText& NewNotifyName, ETextCommit::Typ
 	FSlateApplication::Get().DismissAllMenus();
 }
 
-void SAnimNotifyTrack::AddNewSyncMarker(const FText& NewNotifyName, ETextCommit::Type CommitInfo)
+void SAnimNotifyTrack::AddNewSyncMarker(const FText& NewNotifyName, ETextCommit::Type CommitInfo) 
 {
 	USkeleton* SeqSkeleton = Sequence->GetSkeleton();
 	if ((CommitInfo == ETextCommit::OnEnter) && SeqSkeleton)
@@ -3350,7 +3308,7 @@ void SAnimNotifyTrack::Update()
 
 			SAssignNew(TimingNode, SAnimTimingNode)
 				.InElement(Element)
-				.bUseTooltip(false)
+				.bUseTooltip(true)
 				.Visibility(TimingNodeVisibility);
 
 			if(Event->NotifyStateClass)
@@ -3370,7 +3328,7 @@ void SAnimNotifyTrack::Update()
 					// Create an end timing node if we have a state
 					SAssignNew(EndTimingNode, SAnimTimingNode)
 						.InElement(*FoundStateEndElement)
-						.bUseTooltip(false)
+						.bUseTooltip(true)
 						.Visibility(TimingNodeVisibility);
 				}
 			}
@@ -3383,8 +3341,8 @@ void SAnimNotifyTrack::Update()
 				.PanTrackRequest(OnRequestTrackPan)
 				.ViewInputMin(ViewInputMin)
 				.ViewInputMax(ViewInputMax)
-				.MarkerBars(MarkerBars)
-				.OnDeselectAllNotifies(OnDeselectAllNotifies)
+				.OnSnapPosition(OnSnapPosition)
+				.OnSelectionChanged(OnSelectionChanged)
 				.StateEndTimingNode(EndTimingNode);
 
 			SAssignNew(NotifyPair, SAnimNotifyPair)
@@ -3420,8 +3378,8 @@ void SAnimNotifyTrack::Update()
 			.PanTrackRequest(OnRequestTrackPan)
 			.ViewInputMin(ViewInputMin)
 			.ViewInputMax(ViewInputMax)
-			.MarkerBars(MarkerBars)
-			.OnDeselectAllNotifies(OnDeselectAllNotifies)
+			.OnSnapPosition(OnSnapPosition)
+			.OnSelectionChanged(OnSelectionChanged)
 			.StateEndTimingNode(EndTimingNode);
 
 		NodeSlots->AddSlot()
@@ -3490,13 +3448,13 @@ FReply SAnimNotifyTrack::OnNotifyNodeDragStarted(TSharedRef<SAnimNotifyNode> Not
 
 FReply SAnimNotifyTrack::OnMouseButtonDown(const FGeometry& MyGeometry, const FPointerEvent& MouseEvent)
 {
-	if (MouseEvent.GetEffectingButton() == EKeys::LeftMouseButton)
-	{
-		FVector2D CursorPos = MouseEvent.GetScreenSpacePosition();
-		CursorPos = MyGeometry.AbsoluteToLocal(CursorPos);
-		int32 HitIndex = GetHitNotifyNode(MyGeometry, CursorPos);
+	FVector2D CursorPos = MouseEvent.GetScreenSpacePosition();
+	CursorPos = MyGeometry.AbsoluteToLocal(CursorPos);
+	int32 HitIndex = GetHitNotifyNode(MyGeometry, CursorPos);
 
-		if (HitIndex != INDEX_NONE)
+	if (HitIndex != INDEX_NONE)
+	{
+		if (MouseEvent.GetEffectingButton() == EKeys::LeftMouseButton)
 		{
 			// Hit a node, record the mouse position for use later so we can know when / where a
 			// drag happened on the node handles if necessary.
@@ -3504,6 +3462,11 @@ FReply SAnimNotifyTrack::OnMouseButtonDown(const FGeometry& MyGeometry, const FP
 
 			return FReply::Handled().DetectDrag(NotifyNodes[HitIndex].ToSharedRef(), EKeys::LeftMouseButton);
 		}
+		else if(MouseEvent.GetEffectingButton() == EKeys::RightMouseButton)
+		{
+			// Hit a node, return handled so we can pop a context menu on mouse up
+			return FReply::Handled();
+		}
 	}
 
 	return FReply::Unhandled();
@@ -3528,9 +3491,6 @@ void SAnimNotifyTrack::HandleNodeDrop(TSharedPtr<SAnimNotifyNode> Node, float Of
 {
 	ensure(Node.IsValid());
 
-	OnDeselectAllNotifies.ExecuteIfBound();
-	const FScopedTransaction Transaction(LOCTEXT("MoveNotifyEvent", "Move Anim Notify"));
-	Sequence->Modify();
 	if (Node->NodeObjectInterface->GetType() == ENodeObjectTypes::SYNC_MARKER)
 	{
 		UBlendSpaceBase::UpdateBlendSpacesUsingAnimSequence(Sequence);
@@ -3647,6 +3607,8 @@ void SAnimNotifyTrack::PasteSingleNotify(FString& NotifyString, float PasteTime)
 			NewNotify.EndLink.Link(Sequence, NewNotify.EndLink.GetTime());
 		}
 
+		NewNotify.Guid = FGuid::NewGuid();
+
 		if (!bValidNotify)
 		{
 			// Paste failed, remove the notify
@@ -3689,6 +3651,8 @@ void SAnimNotifyTrack::PasteSingleSyncMarker(FString& MarkerString, float PasteT
 			// Make sure the notify is within the track area
 			SyncMarker.Time = FMath::Clamp(SyncMarker.Time, 0.0f, Sequence->SequenceLength);
 			SyncMarker.TrackIndex = TrackIndex;
+
+			SyncMarker.Guid = FGuid::NewGuid();
 		}
 		else
 		{
@@ -3820,9 +3784,6 @@ void SNotifyEdTrack::Construct(const FArguments& InArgs)
 	//////////////////////////////
 	this->ChildSlot
 	[
-		SNew( SBorder )
-		.Padding( FMargin(2.0f, 2.0f) )
-		[
 			SNew(SHorizontalBox)
 			+SHorizontalBox::Slot()
 			.FillWidth(1)
@@ -3844,7 +3805,7 @@ void SNotifyEdTrack::Construct(const FArguments& InArgs)
 				.OnGetScrubValue(InArgs._OnGetScrubValue)
 				.OnGetDraggedNodePos(InArgs._OnGetDraggedNodePos)
 				.OnNodeDragStarted(InArgs._OnNodeDragStarted)
-				.MarkerBars(InArgs._MarkerBars)
+				.OnSnapPosition(InArgs._OnSnapPosition)
 				.TrackColor(Track.TrackColor)
 				.OnRequestTrackPan(FPanTrackRequest::CreateSP(PanelRef, &SAnimNotifyPanel::PanInputViewRange))
 				.OnRequestOffsetRefresh(InArgs._OnRequestRefreshOffsets)
@@ -3859,48 +3820,6 @@ void SNotifyEdTrack::Construct(const FArguments& InArgs)
 				.OnGetTimingNodeVisibility(InArgs._OnGetTimingNodeVisibility)
 				.OnInvokeTab(InArgs._OnInvokeTab)
 			]
-
-			+SHorizontalBox::Slot()
-			.AutoWidth()
-			[
-				SNew( SBox )
-				.WidthOverride(InArgs._WidgetWidth)
-				[
-					SNew(SHorizontalBox)
-					+SHorizontalBox::Slot()
-					.HAlign(HAlign_Center)
-					.FillWidth(1)
-					[
-						// Name of track
-						SAssignNew(TrackText, SInlineEditableTextBlock)
-						.Text(this, &SNotifyEdTrack::GetTrackName)
-						.ColorAndOpacity(Track.TrackColor)
-						.OnTextCommitted(PanelRef, &SAnimNotifyPanel::OnCommitTrackName, TrackIndex)
-					]
-
-					+SHorizontalBox::Slot()
-					.AutoWidth()
-					[
-						// Name of track
-						SNew(SButton)
-						.Text(LOCTEXT("AddTrackButtonLabel", "+"))
-						.ToolTipText( LOCTEXT("AddTrackTooltip", "Add track above here") )
-						.OnClicked( PanelRef, &SAnimNotifyPanel::InsertTrack, TrackIndex + 1 )
-					]
-
-					+SHorizontalBox::Slot()
-					.AutoWidth()
-					[
-						// Name of track
-						SNew(SButton)
-						.Text(LOCTEXT("RemoveTrackButtonLabel", "-"))
-						.ToolTipText( LOCTEXT("RemoveTrackTooltip", "Remove this track") )
-						.OnClicked( PanelRef, &SAnimNotifyPanel::DeleteTrack, TrackIndex )
-						.IsEnabled( SNotifyEdTrack::CanDeleteTrack() )
-					]
-				]
-			]
-		]
 	];
 }
 
@@ -3920,7 +3839,7 @@ void FAnimNotifyPanelCommands::RegisterCommands()
 //////////////////////////////////////////////////////////////////////////
 // SAnimNotifyPanel
 
-void SAnimNotifyPanel::Construct(const FArguments& InArgs, const TSharedRef<class IEditableSkeleton>& InEditableSkeleton)
+void SAnimNotifyPanel::Construct(const FArguments& InArgs, const TSharedRef<FAnimModel>& InModel)
 {
 	SAnimTrackPanel::Construct( SAnimTrackPanel::FArguments()
 		.WidgetWidth(InArgs._WidgetWidth)
@@ -3931,20 +3850,21 @@ void SAnimNotifyPanel::Construct(const FArguments& InArgs, const TSharedRef<clas
 		.OnSetInputViewRange(InArgs._OnSetInputViewRange));
 
 	Sequence = InArgs._Sequence;
-	MarkerBars = InArgs._MarkerBars;
 	OnInvokeTab = InArgs._OnInvokeTab;
 	OnNotifiesChanged = InArgs._OnNotifiesChanged;
+	OnSnapPosition = InArgs._OnSnapPosition;
+	bIsSelecting = false;
+	bIsUpdating = false;
+
+	InModel->OnHandleObjectsSelected().AddSP(this, &SAnimNotifyPanel::HandleObjectsSelected);
 
 	FAnimNotifyPanelCommands::Register();
 	BindCommands();
 
-	// @todo anim : this is kinda hack to make sure it has 1 track is alive
-	// we can do this whenever import or asset is created, but it's more places to handle than here
-	// the function name in that case will need to change
-	Sequence->InitializeNotifyTrack();
-	Sequence->RegisterOnNotifyChanged(UAnimSequenceBase::FOnNotifyChanged::CreateSP( this, &SAnimNotifyPanel::RefreshNotifyTracks )  );
+	Sequence->RegisterOnNotifyChanged(UAnimSequenceBase::FOnNotifyChanged::CreateSP(this, &SAnimNotifyPanel::RefreshNotifyTracks ));
 
-	InEditableSkeleton->RegisterOnNotifiesChanged(FSimpleDelegate::CreateSP(this, &SAnimNotifyPanel::RefreshNotifyTracks));
+	InModel->GetEditableSkeleton()->RegisterOnNotifiesChanged(FSimpleDelegate::CreateSP(this, &SAnimNotifyPanel::RefreshNotifyTracks));
+	InModel->OnTracksChanged().Add(FSimpleDelegate::CreateSP(this, &SAnimNotifyPanel::RefreshNotifyTracks));
 
 	if(GEditor)
 	{
@@ -3960,48 +3880,13 @@ void SAnimNotifyPanel::Construct(const FArguments& InArgs, const TSharedRef<clas
 
 	this->ChildSlot
 	[
-		SNew(SVerticalBox)
-		+SVerticalBox::Slot()
-		.FillHeight(1)
-		[
-			SNew( SExpandableArea )
-			.AreaTitle( LOCTEXT("Notifies", "Notifies") )
-			.AddMetaData<FTagMetaData>(TEXT("AnimNotify.Notify"))
-			.BodyContent()
-			[
-				SNew(SVerticalBox)
-				+SVerticalBox::Slot()
-				.FillHeight(1)
-				[
-					SAssignNew(PanelArea, SBorder)
-					.BorderImage(FEditorStyle::GetBrush("NoBorder"))
-					.Padding(FMargin(2.0f, 2.0f))
-					.ColorAndOpacity(FLinearColor::White)
-				]
-				+ SVerticalBox::Slot()
-				.AutoHeight()
-				[
-					SNew(SHorizontalBox)
-					+SHorizontalBox::Slot()
-					.FillWidth(1)
-					[
-						SAssignNew(NotifyTrackScrollBar, SScrollBar)
-						.Orientation(EOrientation::Orient_Horizontal)
-						.OnUserScrolled(this, &SAnimNotifyPanel::OnNotifyTrackScrolled)
-					]
-					+SHorizontalBox::Slot()
-					.AutoWidth()
-					[
-						SNew(SBox)
-						.WidthOverride(WidgetWidth)
-					]
-				]
-			]
-		]
+		SAssignNew(PanelArea, SBorder)
+		.AddMetaData<FTagMetaData>(TEXT("AnimNotify.Notify"))
+		.BorderImage(FEditorStyle::GetBrush("NoBorder"))
+		.Padding(0.0f)
+		.ColorAndOpacity(FLinearColor::White)
 	];
 
-	InputViewRangeChanged(ViewInputMin.Get(), ViewInputMax.Get());
-
 	OnPropertyChangedHandle = FCoreUObjectDelegates::FOnObjectPropertyChanged::FDelegate::CreateSP(this, &SAnimNotifyPanel::OnPropertyChanged);
 	OnPropertyChangedHandleDelegateHandle = FCoreUObjectDelegates::OnObjectPropertyChanged.Add(OnPropertyChangedHandle);
 
@@ -4080,16 +3965,21 @@ FReply SAnimNotifyPanel::InsertTrack(int32 TrackIndexToInsert)
 
 	Update();
 
-	//Now we have updated we can request rename on the track UI
-	int32 UITrackIndex = NotifyEditorTracks.Num() - 1 - TrackIndexToInsert;
-	RegisterActiveTimer(0.f, FWidgetActiveTimerDelegate::CreateSP(this, &SAnimNotifyPanel::TriggerRename, UITrackIndex));
 	return FReply::Handled();
 }
 
-EActiveTimerReturnType SAnimNotifyPanel::TriggerRename(double InCurrentTime, float InDeltaTime, int32 TrackIndex)
+FReply SAnimNotifyPanel::AddTrack()
 {
-	NotifyEditorTracks[TrackIndex]->TriggerRename();
-	return EActiveTimerReturnType::Stop;
+	FAnimNotifyTrack NewItem;
+	NewItem.TrackName = GetNewTrackName();
+	NewItem.TrackColor = FLinearColor::White;
+
+	Sequence->AnimNotifyTracks.Add(NewItem);
+	Sequence->MarkPackageDirty();
+
+	Update();
+
+	return FReply::Handled();
 }
 
 FReply SAnimNotifyPanel::DeleteTrack(int32 TrackIndexToDelete)
@@ -4150,74 +4040,117 @@ void SAnimNotifyPanel::OnCommitTrackName(const FText& InText, ETextCommit::Type
 
 void SAnimNotifyPanel::Update()
 {
-	if(Sequence != NULL)
+	if(!bIsUpdating)
 	{
-		Sequence->RefreshCacheData();
-	}
+		TGuardValue<bool> ScopeGuard(bIsUpdating, true);
 
-	RefreshNotifyTracks();
+		if(Sequence != NULL)
+		{
+			Sequence->RefreshCacheData();
+		}
 
-	OnNotifiesChanged.ExecuteIfBound();
+		RefreshNotifyTracks();
+
+		OnNotifiesChanged.ExecuteIfBound();
+	}
 }
 
-void SAnimNotifyPanel::RefreshNotifyTracks()
+// Helper to save/restore selection state when widgets are recreated
+struct FScopedSavedNotifySelection
 {
-	check (Sequence);
-
-	TSharedPtr<SVerticalBox> NotifySlots;
-	PanelArea->SetContent(
-		SAssignNew( NotifySlots, SVerticalBox )
-		);
+	FScopedSavedNotifySelection(SAnimNotifyPanel& InPanel)
+		: Panel(InPanel)
+	{
+		for (TSharedPtr<SAnimNotifyTrack> Track : InPanel.NotifyAnimTracks)
+		{
+			for(int32 NodeIndex = 0; NodeIndex < Track->GetNumNotifyNodes(); ++NodeIndex)
+			{
+				if(Track->IsNodeSelected(NodeIndex))
+				{
+					SelectedNodeGuids.Add(Track->GetNodeObjectInterface(NodeIndex)->GetGuid());	
+				}
+			}
+		}
+	}
 
-	// Clear node tool tips to stop slate referencing them and possibly
-	// causing a crash if the notify has gone away
-	for (TSharedPtr<SAnimNotifyTrack> Track : NotifyAnimTracks)
+	~FScopedSavedNotifySelection()
 	{
-		Track->ClearNodeTooltips();
+		// Re-apply selection state
+		for (TSharedPtr<SAnimNotifyTrack> Track : Panel.NotifyAnimTracks)
+		{
+			Track->SelectNodesByGuid(SelectedNodeGuids, false);
+		}
 	}
 
-	NotifyAnimTracks.Empty();
-	NotifyEditorTracks.Empty();
+	SAnimNotifyPanel& Panel;
+	TSet<FGuid> SelectedNodeGuids;
+};
+
+void SAnimNotifyPanel::RefreshNotifyTracks()
+{
+	check (Sequence);
 
-	for(int32 I=Sequence->AnimNotifyTracks.Num()-1; I>=0; --I)
 	{
-		FAnimNotifyTrack& Track = Sequence->AnimNotifyTracks[I];
-		TSharedPtr<SNotifyEdTrack> EdTrack;
+		FScopedSavedNotifySelection ScopedSelection(*this);
 
-		NotifySlots->AddSlot()
-		.AutoHeight()
-		.VAlign(VAlign_Center)
-		[
-			SAssignNew(EdTrack, SNotifyEdTrack)
-			.TrackIndex(I)
-			.Sequence(Sequence)
-			.AnimNotifyPanel(SharedThis(this))
-			.WidgetWidth(WidgetWidth)
-			.ViewInputMin(ViewInputMin)
-			.ViewInputMax(ViewInputMax)
-			.OnGetScrubValue(OnGetScrubValue)
-			.OnGetDraggedNodePos(this, &SAnimNotifyPanel::CalculateDraggedNodePos)
-			.OnUpdatePanel(this, &SAnimNotifyPanel::Update)
-			.OnGetNotifyBlueprintData(this, &SAnimNotifyPanel::OnGetNotifyBlueprintData, &NotifyClassNames)
-			.OnGetNotifyStateBlueprintData(this, &SAnimNotifyPanel::OnGetNotifyBlueprintData, &NotifyStateClassNames)
-			.OnGetNotifyNativeClasses(this, &SAnimNotifyPanel::OnGetNativeNotifyData, UAnimNotify::StaticClass(), &NotifyClassNames)
-			.OnGetNotifyStateNativeClasses(this, &SAnimNotifyPanel::OnGetNativeNotifyData, UAnimNotifyState::StaticClass(), &NotifyStateClassNames)
-			.OnSelectionChanged(this, &SAnimNotifyPanel::OnTrackSelectionChanged)
-			.OnNodeDragStarted(this, &SAnimNotifyPanel::OnNotifyNodeDragStarted)
-			.MarkerBars(MarkerBars)
-			.OnRequestRefreshOffsets(OnRequestRefreshOffsets)
-			.OnDeleteNotify(this, &SAnimNotifyPanel::DeleteSelectedNodeObjects)
-			.OnDeselectAllNotifies(this, &SAnimNotifyPanel::DeselectAllNotifies)
-			.OnCopyNodes(this, &SAnimNotifyPanel::CopySelectedNodesToClipboard)
-			.OnPasteNodes(this, &SAnimNotifyPanel::OnPasteNodes)
-			.OnSetInputViewRange(this, &SAnimNotifyPanel::InputViewRangeChanged)
-			.OnGetTimingNodeVisibility(OnGetTimingNodeVisibility)
-			.OnInvokeTab(OnInvokeTab)
-		];
+		TSharedPtr<SVerticalBox> NotifySlots;
+		PanelArea->SetContent(
+			SAssignNew( NotifySlots, SVerticalBox )
+			);
+
+		// Clear node tool tips to stop slate referencing them and possibly
+		// causing a crash if the notify has gone away
+		for (TSharedPtr<SAnimNotifyTrack> Track : NotifyAnimTracks)
+		{
+			Track->ClearNodeTooltips();
+		}
+
+		NotifyAnimTracks.Empty();
+		NotifyEditorTracks.Empty();
+
+		for(int32 TrackIndex = 0; TrackIndex < Sequence->AnimNotifyTracks.Num(); TrackIndex++)
+		{
+			FAnimNotifyTrack& Track = Sequence->AnimNotifyTracks[TrackIndex];
+			TSharedPtr<SNotifyEdTrack> EdTrack;
 
-		NotifyAnimTracks.Add(EdTrack->NotifyTrack);
-		NotifyEditorTracks.Add(EdTrack);
+			NotifySlots->AddSlot()
+			.AutoHeight()
+			.VAlign(VAlign_Center)
+			[
+				SAssignNew(EdTrack, SNotifyEdTrack)
+				.TrackIndex(TrackIndex)
+				.Sequence(Sequence)
+				.AnimNotifyPanel(SharedThis(this))
+				.WidgetWidth(WidgetWidth)
+				.ViewInputMin(ViewInputMin)
+				.ViewInputMax(ViewInputMax)
+				.OnGetScrubValue(OnGetScrubValue)
+				.OnGetDraggedNodePos(this, &SAnimNotifyPanel::CalculateDraggedNodePos)
+				.OnUpdatePanel(this, &SAnimNotifyPanel::Update)
+				.OnGetNotifyBlueprintData(this, &SAnimNotifyPanel::OnGetNotifyBlueprintData, &NotifyClassNames)
+				.OnGetNotifyStateBlueprintData(this, &SAnimNotifyPanel::OnGetNotifyBlueprintData, &NotifyStateClassNames)
+				.OnGetNotifyNativeClasses(this, &SAnimNotifyPanel::OnGetNativeNotifyData, UAnimNotify::StaticClass(), &NotifyClassNames)
+				.OnGetNotifyStateNativeClasses(this, &SAnimNotifyPanel::OnGetNativeNotifyData, UAnimNotifyState::StaticClass(), &NotifyStateClassNames)
+				.OnSelectionChanged(this, &SAnimNotifyPanel::OnTrackSelectionChanged)
+				.OnNodeDragStarted(this, &SAnimNotifyPanel::OnNotifyNodeDragStarted)
+				.OnSnapPosition(OnSnapPosition)
+				.OnRequestRefreshOffsets(OnRequestRefreshOffsets)
+				.OnDeleteNotify(this, &SAnimNotifyPanel::DeleteSelectedNodeObjects)
+				.OnDeselectAllNotifies(this, &SAnimNotifyPanel::DeselectAllNotifies)
+				.OnCopyNodes(this, &SAnimNotifyPanel::CopySelectedNodesToClipboard)
+				.OnPasteNodes(this, &SAnimNotifyPanel::OnPasteNodes)
+				.OnSetInputViewRange(this, &SAnimNotifyPanel::InputViewRangeChanged)
+				.OnGetTimingNodeVisibility(OnGetTimingNodeVisibility)
+				.OnInvokeTab(OnInvokeTab)
+			];
+
+			NotifyAnimTracks.Add(EdTrack->NotifyTrack);
+			NotifyEditorTracks.Add(EdTrack);
+		}
 	}
+
+	// Signal selection change to refresh details panel
+	OnTrackSelectionChanged();
 }
 
 float SAnimNotifyPanel::CalculateDraggedNodePos() const
@@ -4227,7 +4160,13 @@ float SAnimNotifyPanel::CalculateDraggedNodePos() const
 
 FReply SAnimNotifyPanel::OnNotifyNodeDragStarted(TArray<TSharedPtr<SAnimNotifyNode>> NotifyNodes, TSharedRef<SWidget> Decorator, const FVector2D& ScreenCursorPos, const FVector2D& ScreenNodePosition, const bool bDragOnMarker)
 {
-	TSharedRef<SOverlay> NodeDragDecorator = SNew(SOverlay);
+	TSharedRef<SOverlay> NodeDragDecoratorOverlay = SNew(SOverlay);
+	TSharedRef<SBorder> NodeDragDecorator = SNew(SBorder)
+	.BorderImage(FEditorStyle::GetBrush("ToolPanel.GroupBorder"))
+	[
+		NodeDragDecoratorOverlay
+	];
+
 	TArray<TSharedPtr<SAnimNotifyNode>> Nodes;
 
 	for(TSharedPtr<SAnimNotifyTrack> Track : NotifyAnimTracks)
@@ -4268,7 +4207,7 @@ FReply SAnimNotifyPanel::OnNotifyNodeDragStarted(TArray<TSharedPtr<SAnimNotifyNo
 	{
 		FVector2D OffsetFromFirst(Node->GetScreenPosition() - OverlayOrigin);
 
-		NodeDragDecorator->AddSlot()
+		NodeDragDecoratorOverlay->AddSlot()
 			.Padding(FMargin(OffsetFromFirst.X, OffsetFromFirst.Y, 0.0f, 0.0f))
 			[
 				Node->AsShared()
@@ -4277,7 +4216,7 @@ FReply SAnimNotifyPanel::OnNotifyNodeDragStarted(TArray<TSharedPtr<SAnimNotifyNo
 
 	FPanTrackRequest PanRequestDelegate = FPanTrackRequest::CreateSP(this, &SAnimNotifyPanel::PanInputViewRange);
 	FOnUpdatePanel UpdateDelegate = FOnUpdatePanel::CreateSP(this, &SAnimNotifyPanel::Update);
-	return FReply::Handled().BeginDragDrop(FNotifyDragDropOp::New(Nodes, NodeDragDecorator, NotifyAnimTracks, Sequence, ScreenCursorPos, OverlayOrigin, OverlayExtents, CurrentDragXPosition, PanRequestDelegate, MarkerBars, UpdateDelegate));
+	return FReply::Handled().BeginDragDrop(FNotifyDragDropOp::New(Nodes, NodeDragDecorator, NotifyAnimTracks, Sequence, ScreenCursorPos, OverlayOrigin, OverlayExtents, CurrentDragXPosition, PanRequestDelegate, OnSnapPosition, UpdateDelegate));
 }
 
 void SAnimNotifyPanel::PostUndo( bool bSuccess )
@@ -4321,14 +4260,16 @@ void SAnimNotifyPanel::DeleteSelectedNodeObjects()
 		FScopedTransaction Transaction(LOCTEXT("DeleteMarkers", "Delete Animation Markers"));
 		Sequence->Modify(true);
 
-		// Delete from highest index to lowest
-		SelectedNodes.Sort();
-		for (int NodeIndex = SelectedNodes.Num() - 1; NodeIndex >= 0; --NodeIndex)
+		// As we address node object's source data by pointer, we need to mark for delete then
+		// delete invalid entries to avoid concurrent modification of containers
+		for (INodeObjectInterface* NodeObject : SelectedNodes)
 		{
-			INodeObjectInterface* NodeObject = SelectedNodes[NodeIndex];
-			NodeObject->Delete(Sequence);
+			NodeObject->MarkForDelete(Sequence);
 		}
 
+		FNotifyNodeInterface::RemoveInvalidNotifies(Sequence);
+		FSyncMarkerNodeInterface::RemoveInvalidSyncMarkers(Sequence);
+
 		if (bContainsSyncMarkers)
 		{
 			UBlendSpaceBase::UpdateBlendSpacesUsingAnimSequence(Sequence);
@@ -4357,38 +4298,49 @@ void SAnimNotifyPanel::SetSequence(class UAnimSequenceBase *	InSequence)
 
 void SAnimNotifyPanel::OnTrackSelectionChanged()
 {
-	// Need to collect selection info from all tracks
-	TArray<UObject*> NotifyObjects;
-
-	for(int32 TrackIdx = 0 ; TrackIdx < NotifyAnimTracks.Num() ; ++TrackIdx)
+	if(!bIsSelecting)
 	{
-		TSharedPtr<SAnimNotifyTrack> Track = NotifyAnimTracks[TrackIdx];
-		TArray<int32> TrackIndices = Track->GetSelectedNotifyIndices();
-		for(int32 Idx : TrackIndices)
+		TGuardValue<bool> GuardValue(bIsSelecting, true);
+
+		// Need to collect selection info from all tracks
+		TArray<UObject*> NotifyObjects;
+
+		for(int32 TrackIdx = 0 ; TrackIdx < NotifyAnimTracks.Num() ; ++TrackIdx)
 		{
-			INodeObjectInterface* NodeObjectInterface = Track->GetNodeObjectInterface(Idx);
-			if (FAnimNotifyEvent* NotifyEvent = NodeObjectInterface->GetNotifyEvent())
+			TSharedPtr<SAnimNotifyTrack> Track = NotifyAnimTracks[TrackIdx];
+			const TArray<int32>& TrackIndices = Track->GetSelectedNotifyIndices();
+			for(int32 Idx : TrackIndices)
 			{
-				FString ObjName = MakeUniqueObjectName(GetTransientPackage(), UEditorNotifyObject::StaticClass()).ToString();
-				UEditorNotifyObject* NewNotifyObject = NewObject<UEditorNotifyObject>(GetTransientPackage(), FName(*ObjName), RF_Public | RF_Standalone | RF_Transient);
-				NewNotifyObject->InitFromAnim(Sequence, FOnAnimObjectChange::CreateSP(this, &SAnimNotifyPanel::OnNotifyObjectChanged));
-				NewNotifyObject->InitialiseNotify(NotifyAnimTracks.Num() - TrackIdx - 1, Idx);
-				NotifyObjects.AddUnique(NewNotifyObject);
+				INodeObjectInterface* NodeObjectInterface = Track->GetNodeObjectInterface(Idx);
+				if (FAnimNotifyEvent* NotifyEvent = NodeObjectInterface->GetNotifyEvent())
+				{
+					FString ObjName = MakeUniqueObjectName(GetTransientPackage(), UEditorNotifyObject::StaticClass()).ToString();
+					UEditorNotifyObject* NewNotifyObject = NewObject<UEditorNotifyObject>(GetTransientPackage(), FName(*ObjName), RF_Public | RF_Standalone | RF_Transient);
+					NewNotifyObject->InitFromAnim(Sequence, FOnAnimObjectChange::CreateSP(this, &SAnimNotifyPanel::OnNotifyObjectChanged));
+					NewNotifyObject->InitialiseNotify(*Sequence->AnimNotifyTracks[TrackIdx].Notifies[Idx]);
+					NotifyObjects.AddUnique(NewNotifyObject);
+				}
 			}
 		}
-	}
 
-	OnSelectionChanged.ExecuteIfBound(NotifyObjects);
+		OnSelectionChanged.ExecuteIfBound(NotifyObjects);
+	}
 }
 
 void SAnimNotifyPanel::DeselectAllNotifies()
 {
-	for(TSharedPtr<SAnimNotifyTrack> Track : NotifyAnimTracks)
+	if(!bIsSelecting)
 	{
-		Track->DeselectAllNotifyNodes(false);
-	}
+		TGuardValue<bool> GuardValue(bIsSelecting, true);
 
-	OnTrackSelectionChanged();
+		for(TSharedPtr<SAnimNotifyTrack> Track : NotifyAnimTracks)
+		{
+			Track->DeselectAllNotifyNodes(false);
+		}
+
+		TArray<UObject*> NotifyObjects;
+		OnSelectionChanged.ExecuteIfBound(NotifyObjects);
+	}
 }
 
 void SAnimNotifyPanel::CopySelectedNodesToClipboard() const
@@ -4409,7 +4361,7 @@ void SAnimNotifyPanel::CopySelectedNodesToClipboard() const
 		// Sort by track
 		SelectedNodes.Sort([](const INodeObjectInterface& A, const INodeObjectInterface& B)
 		{
-			return (A.GetTrackIndex() > B.GetTrackIndex()) || (A.GetTrackIndex() == B.GetTrackIndex() && A.GetTime() < B.GetTime());
+			return (A.GetTrackIndex() < B.GetTrackIndex()) || (A.GetTrackIndex() == B.GetTrackIndex() && A.GetTime() < B.GetTime());
 		});
 
 		// Need to find how many tracks this selection spans and the minimum time to use as the beginning of the selection
@@ -4504,7 +4456,7 @@ void SAnimNotifyPanel::OnReplaceSelectedWithNotify(FString NewNotifyName, UClass
 		{
 			float BeginTime = OldEvent->GetTime();
 			float Length = OldEvent->GetDuration();
-			int32 TargetTrackIndex = (NotifyAnimTracks.Num()-1) - OldEvent->TrackIndex;
+			int32 TargetTrackIndex = OldEvent->TrackIndex;
 			float TriggerTimeOffset = OldEvent->TriggerTimeOffset;
 			float EndTriggerTimeOffset = OldEvent->EndTriggerTimeOffset;
 			int32 SlotIndex = OldEvent->GetSlotIndex();
@@ -4593,16 +4545,15 @@ void SAnimNotifyPanel::OnPasteNodes(SAnimNotifyTrack* RequestTrack, float ClickT
 
 		// Expand the number of tracks if we don't have enough.
 		check(TrackSpan > 0);
-		if(PasteIdx - (TrackSpan - 1) < 0)
+		if(PasteIdx + TrackSpan > NumTracks)
 		{
-			int32 TracksToAdd = PasteIdx + TrackSpan - 1;
+			int32 TracksToAdd = (PasteIdx + TrackSpan) - NumTracks;
 			while(TracksToAdd)
 			{
-				InsertTrack(PasteIdx++);
+				AddTrack();
 				--TracksToAdd;
 			}
 			NumTracks = NotifyAnimTracks.Num();
-			RequestTrack = NotifyAnimTracks[NumTracks - 1 - PasteIdx].Get();
 		}
 
 		// Scaling for relative paste
@@ -4635,8 +4586,7 @@ void SAnimNotifyPanel::OnPasteNodes(SAnimNotifyTrack* RequestTrack, float ClickT
 				float TimeOffset = OrigTime - OrigBeginTime;
 				float TimeToPaste = ClickTime + TimeOffset * ScaleMultiplier;
 
-				// Have to invert the index here as tracks are stored in reverse
-				TSharedPtr<SAnimNotifyTrack> TrackToUse = NotifyAnimTracks[NotifyAnimTracks.Num() - 1 - (PasteIdx + TrackOffset)];
+				TSharedPtr<SAnimNotifyTrack> TrackToUse = NotifyAnimTracks[PasteIdx + TrackOffset];
 				if (NodeObjectType == ENodeObjectTypes::NOTIFY)
 				{
 					TrackToUse->PasteSingleNotify(NotifyExportString, TimeToPaste);
@@ -4922,11 +4872,17 @@ void SAnimNotifyPanel::OnNotifyObjectChanged(UObject* EditorBaseObj, bool bRebui
 {
 	if(UEditorNotifyObject* NotifyObject = Cast<UEditorNotifyObject>(EditorBaseObj))
 	{
-		// TODO: We should really un-invert these.
-		int32 WidgetTrackIdx = NotifyAnimTracks.Num() - NotifyObject->TrackIndex - 1;
-		if(NotifyAnimTracks.IsValidIndex(WidgetTrackIdx))
+		FScopedSavedNotifySelection ScopedSelection(*this);
+
+		for(FAnimNotifyEvent& Notify : Sequence->Notifies)
 		{
-			NotifyAnimTracks[WidgetTrackIdx]->Update();
+			if(Notify.Guid == NotifyObject->Event.Guid)
+			{
+				if(NotifyAnimTracks.IsValidIndex(Notify.TrackIndex))
+				{
+					NotifyAnimTracks[Notify.TrackIndex]->Update();
+				}
+			}
 		}
 	}
 }
@@ -4948,9 +4904,20 @@ void SAnimNotifyPanel::InputViewRangeChanged(float ViewMin, float ViewMax)
 {
 	float Ratio = (ViewMax - ViewMin) / Sequence->SequenceLength;
 	float OffsetFraction = ViewMin / Sequence->SequenceLength;
-	NotifyTrackScrollBar->SetState(OffsetFraction, Ratio);
+	if(NotifyTrackScrollBar.IsValid())
+	{
+		NotifyTrackScrollBar->SetState(OffsetFraction, Ratio);
+	}
 
 	SAnimTrackPanel::InputViewRangeChanged(ViewMin, ViewMax);
 }
 
+void SAnimNotifyPanel::HandleObjectsSelected(const TArray<UObject*>& InObjects)
+{
+	if(!bIsSelecting)
+	{
+		DeselectAllNotifies();
+	}
+}
+
 #undef LOCTEXT_NAMESPACE
diff --git a/Engine/Source/Editor/Persona/Private/SAnimNotifyPanel.h b/Engine/Source/Editor/Persona/Private/SAnimNotifyPanel.h
--- a/Engine/Source/Editor/Persona/Private/SAnimNotifyPanel.h
+++ b/Engine/Source/Editor/Persona/Private/SAnimNotifyPanel.h
@@ -17,6 +17,7 @@
 #include "Framework/Commands/Commands.h"
 #include "SAnimTimingPanel.h"
 #include "EditorUndoClient.h"
+#include "AnimModel.h"
 
 class FSlateWindowElementList;
 class SAnimNotifyNode;
@@ -37,6 +38,7 @@ DECLARE_DELEGATE_TwoParams( FReplaceWithBlueprintNotify, FString, FString )
 DECLARE_DELEGATE( FDeselectAllNotifies )
 DECLARE_DELEGATE_OneParam( FOnGetBlueprintNotifyData, TArray<FAssetData>& )
 DECLARE_DELEGATE_OneParam( FOnGetNativeNotifyClasses, TArray<UClass*>&)
+DECLARE_DELEGATE_RetVal_TwoParams(bool, FOnSnapPosition, float& /*InOutTimeToSnap*/, float /*InSnapMargin*/)
 
 class SAnimNotifyNode;
 class SAnimNotifyTrack;
@@ -177,7 +179,6 @@ public:
 	SLATE_ATTRIBUTE( float, ViewInputMax )
 	SLATE_ATTRIBUTE( float, InputMin )
 	SLATE_ATTRIBUTE( float, InputMax )
-	SLATE_ATTRIBUTE( TArray<FTrackMarkerBar>, MarkerBars )
 	SLATE_EVENT( FOnSetInputViewRange, OnSetInputViewRange )
 	SLATE_EVENT( FOnSelectionChanged, OnSelectionChanged )
 	SLATE_EVENT( FOnGetScrubValue, OnGetScrubValue )
@@ -185,10 +186,11 @@ public:
 	SLATE_EVENT( FOnGetTimingNodeVisibility, OnGetTimingNodeVisibility )
 	SLATE_EVENT( FOnInvokeTab, OnInvokeTab )
 	SLATE_EVENT( FSimpleDelegate, OnNotifiesChanged )
+	SLATE_EVENT( FOnSnapPosition, OnSnapPosition )
 
 	SLATE_END_ARGS()
 
-	void Construct(const FArguments& InArgs, const TSharedRef<class IEditableSkeleton>& InEditableSkeleton);
+	void Construct(const FArguments& InArgs, const TSharedRef<FAnimModel>& InModel);
 	virtual ~SAnimNotifyPanel();
 
 	void SetSequence(class UAnimSequenceBase *	InSequence);
@@ -196,12 +198,10 @@ public:
 	// Generate a new track name (smallest integer number that isn't currently used)
 	FName GetNewTrackName() const;
 
+	FReply AddTrack();
 	FReply InsertTrack(int32 TrackIndexToInsert);
 	FReply DeleteTrack(int32 TrackIndexToDelete);
 	bool CanDeleteTrack(int32 TrackIndexToDelete);
-
-	/** Widget timer function to trigger notify track rename (cannot do it directly from add track code) */
-	EActiveTimerReturnType TriggerRename(double InCurrentTime, float InDeltaTime, int32 TrackIndex);
 	
 	// Handler function for renaming a notify track
 	void OnCommitTrackName(const FText& InText, ETextCommit::Type CommitInfo, int32 TrackIndexToName);
@@ -246,7 +246,11 @@ public:
 	/** Handler for replacing with notify blueprint */
 	void OnReplaceSelectedWithNotifyBlueprint(FString NewBlueprintNotifyName, FString NewBlueprintNotifyClass);
 
+	void HandleObjectsSelected(const TArray<UObject*>& InObjects);
+
 private:
+	friend struct FScopedSavedNotifySelection;
+
 	TSharedPtr<SBorder> PanelArea;
 	TSharedPtr<SScrollBar> NotifyTrackScrollBar;
 	class UAnimSequenceBase* Sequence;
@@ -324,8 +328,8 @@ private:
 
 	virtual void InputViewRangeChanged(float ViewMin, float ViewMax) override;
 
-	/** Attribute for accessing any section/branching point positions we have to draw */
-	TAttribute<TArray<FTrackMarkerBar>>	MarkerBars;
+	/** Delegate used to snap when dragging */
+	FOnSnapPosition OnSnapPosition;
 
 	/** UI commands for this widget */
 	TSharedPtr<FUICommandList> UICommandList;
@@ -344,4 +348,10 @@ private:
 
 	/** Delegate used to inform others that notifies have changed (for timing) */
 	FSimpleDelegate OnNotifiesChanged;
+
+	/** Recursion guard for selection */
+	bool bIsSelecting;
+
+	/** Recursion guard for updating */
+	bool bIsUpdating;
 };
diff --git a/Engine/Source/Editor/Persona/Private/SAnimSegmentsPanel.cpp b/Engine/Source/Editor/Persona/Private/SAnimSegmentsPanel.cpp
--- a/Engine/Source/Editor/Persona/Private/SAnimSegmentsPanel.cpp
+++ b/Engine/Source/Editor/Persona/Private/SAnimSegmentsPanel.cpp
@@ -18,6 +18,7 @@
 #include "Framework/Application/SlateApplication.h"
 #include "Subsystems/AssetEditorSubsystem.h"
 #include "Editor.h"
+#include "Widgets/Layout/SBox.h"
 
 #define LOCTEXT_NAMESPACE "AnimSegmentPanel"
 
@@ -40,6 +41,7 @@ void SAnimSegmentsPanel::Construct(const FArguments& InArgs)
 	OnAnimSegmentRemovedDelegate		= InArgs._OnAnimSegmentRemoved;
 	OnAnimReplaceMapping				= InArgs._OnAnimReplaceMapping;
 	OnDiffFromParentAsset				= InArgs._OnDiffFromParentAsset;
+	OnGetNodeColor						= InArgs._OnGetNodeColor;
 
 	bChildAnimMontage = InArgs._bChildAnimMontage;
 
@@ -53,7 +55,7 @@ void SAnimSegmentsPanel::Construct(const FArguments& InArgs)
 	// Animation Segment tracks
 	TArray<TSharedPtr<STrackNode>> AnimNodes;
 
-	FLinearColor SelectedColor = FLinearColor(1.0f,0.65,0.0f);
+	FLinearColor SelectedColor = FEditorStyle::GetSlateColor("SelectionColor").GetSpecifiedColor();
 
 	TSharedPtr<SVerticalBox> AnimSegmentTracks;
 
@@ -62,8 +64,6 @@ void SAnimSegmentsPanel::Construct(const FArguments& InArgs)
 		SAssignNew(AnimSegmentTracks, SVerticalBox)
 	];
 
-	FLinearColor TrackColor = InArgs._ColorTracker->GetNextColor();
-
 	for (int32 TrackIdx=0; TrackIdx < NumTracks; TrackIdx++)
 	{
 		TSharedPtr<STrack> AnimSegmentTrack;
@@ -73,16 +73,14 @@ void SAnimSegmentsPanel::Construct(const FArguments& InArgs)
 			AnimSegmentTracks->AddSlot()
 				.AutoHeight()
 				.VAlign(VAlign_Center)
-				.Padding(FMargin(0.5f, 0.5f))
+				.Padding(0.0f)
 				[
 					SAssignNew(AnimSegmentTrack, STrack)
-					.TrackColor(TrackColor)
 					.ViewInputMin(ViewInputMin)
 					.ViewInputMax(ViewInputMax)
 					.TrackMaxValue(InArgs._TrackMaxValue)
 					.TrackNumDiscreteValues(InArgs._TrackNumDiscreteValues)
 					.OnTrackRightClickContextMenu(InArgs._OnTrackRightClickContextMenu)
-					.ScrubPosition(InArgs._ScrubPosition)
 					.OnTrackDragDrop(this, &SAnimSegmentsPanel::OnTrackDragDrop)
 				];
 		}
@@ -91,10 +89,9 @@ void SAnimSegmentsPanel::Construct(const FArguments& InArgs)
 			AnimSegmentTracks->AddSlot()
 				.AutoHeight()
 				.VAlign(VAlign_Center)
-				.Padding(FMargin(0.5f, 0.5f))
+				.Padding(0.0f)
 				[
 					SAssignNew(AnimSegmentTrack, STrack)
-					.TrackColor(TrackColor)
 					.ViewInputMin(ViewInputMin)
 					.ViewInputMax(ViewInputMax)
 					.TrackMaxValue(InArgs._TrackMaxValue)
@@ -106,7 +103,6 @@ void SAnimSegmentsPanel::Construct(const FArguments& InArgs)
 					.DraggableBarSnapPositions(InArgs._DraggableBarSnapPositions)
 					.TrackNumDiscreteValues(InArgs._TrackNumDiscreteValues)
 					.OnTrackRightClickContextMenu(InArgs._OnTrackRightClickContextMenu)
-					.ScrubPosition(InArgs._ScrubPosition)
 					.OnTrackDragDrop(this, &SAnimSegmentsPanel::OnTrackDragDrop)
 				];
 		}
@@ -114,8 +110,6 @@ void SAnimSegmentsPanel::Construct(const FArguments& InArgs)
 		TrackWidgets.Add(AnimSegmentTrack);
 	}
 
-	DefaultNodeColor = InArgs._NodeColor;
-
 	// Generate Nodes and map them to tracks
 	for ( int32 SegmentIdx=0; SegmentIdx < AnimTrack->AnimSegments.Num(); SegmentIdx++ )
 	{
@@ -177,9 +171,9 @@ FLinearColor	SAnimSegmentsPanel::GetNodeColor(int32 AnimSegmentIndex) const
 		{
 			return ChildMontageColor;
 		}
-		else
+		else if(OnGetNodeColor.IsBound())
 		{
-			return DefaultNodeColor.Get();
+			return OnGetNodeColor.Execute(AnimTrack->AnimSegments[AnimSegmentIndex]);
 		}
 	}
 
@@ -265,7 +259,7 @@ void SAnimSegmentsPanel::SetSegmentStartPos(float NewStartPos, int32 AnimSegment
 		if(!bDragging)
 		{
 			const FScopedTransaction Transaction( LOCTEXT("AnimSegmentPanel_SetSegmentStart", "Edit Segment Start Time") );
-			OnPreAnimUpdateDelegate.Execute();
+			OnPreAnimUpdateDelegate.ExecuteIfBound();
 			bDragging = true;
 		}
 
@@ -279,7 +273,7 @@ void SAnimSegmentsPanel::OnSegmentDropped(int32 AnimSegmentIndex)
 	if(bDragging)
 	{
 		bDragging = false;
-		OnPostAnimUpdateDelegate.Execute();
+		OnPostAnimUpdateDelegate.ExecuteIfBound();
 	}
 }
 
@@ -321,7 +315,7 @@ void SAnimSegmentsPanel::SummonSegmentNodeContextMenu(FMenuBuilder& MenuBuilder,
 void SAnimSegmentsPanel::AddAnimSegment( UAnimSequenceBase* NewSequenceBase, float NewStartPos )
 {
 	const FScopedTransaction Transaction( LOCTEXT("AnimSegmentPanel_AddSegment", "Add Segment") );
-	OnPreAnimUpdateDelegate.Execute();
+	OnPreAnimUpdateDelegate.ExecuteIfBound();
 
 	FAnimSegment NewSegment;
 	NewSegment.AnimReference = NewSequenceBase;
@@ -332,7 +326,7 @@ void SAnimSegmentsPanel::AddAnimSegment( UAnimSequenceBase* NewSequenceBase, flo
 	NewSegment.StartPos = NewStartPos;
 
 	AnimTrack->AnimSegments.Add(NewSegment);
-	OnPostAnimUpdateDelegate.Execute();
+	OnPostAnimUpdateDelegate.ExecuteIfBound();
 }
 
 void SAnimSegmentsPanel::ReplaceAnimSegment(int32 AnimSegmentIndex, UAnimSequenceBase* NewSequenceBase)
@@ -343,9 +337,9 @@ void SAnimSegmentsPanel::ReplaceAnimSegment(int32 AnimSegmentIndex, UAnimSequenc
 		UAnimSequenceBase* OldSequenceBase = AnimTrack->AnimSegments[AnimSegmentIndex].AnimReference;
 		if (OldSequenceBase != NewSequenceBase)
 		{
-			OnPreAnimUpdateDelegate.Execute();
+			OnPreAnimUpdateDelegate.ExecuteIfBound();
 			OnAnimReplaceMapping.ExecuteIfBound(SlotName, AnimSegmentIndex, OldSequenceBase, NewSequenceBase);
-			OnPostAnimUpdateDelegate.Execute();
+			OnPostAnimUpdateDelegate.ExecuteIfBound();
 		}
 	}
 
@@ -383,12 +377,12 @@ void SAnimSegmentsPanel::RemoveAnimSegment(int32 AnimSegmentIndex)
 	if(ValidIndex(AnimSegmentIndex))
 	{
 		const FScopedTransaction Transaction( LOCTEXT("AnimSegmentseEditor", "Remove Segment") );
-		OnPreAnimUpdateDelegate.Execute();
+		OnPreAnimUpdateDelegate.ExecuteIfBound();
 
 		AnimTrack->AnimSegments.RemoveAt(AnimSegmentIndex);
 
 		OnAnimSegmentRemovedDelegate.ExecuteIfBound(AnimSegmentIndex);
-		OnPostAnimUpdateDelegate.Execute();
+		OnPostAnimUpdateDelegate.ExecuteIfBound();
 	}
 }
 
@@ -444,8 +438,13 @@ void SAnimSegmentsPanel::FillSubMenu(FMenuBuilder& MenuBuilder, int32 AnimSegmen
 			FContentBrowserModule& ContentBrowserModule = FModuleManager::Get().LoadModuleChecked<FContentBrowserModule>(TEXT("ContentBrowser"));
 
 			MenuBuilder.AddWidget(
-				ContentBrowserModule.Get().CreateAssetPicker(AssetPickerConfig), 
-				LOCTEXT("ReplaceAnimation_Label", "Replace")
+				SNew(SBox)
+				.MaxDesiredHeight(600.0f)
+				[
+					ContentBrowserModule.Get().CreateAssetPicker(AssetPickerConfig)
+				], 
+				FText(),
+				true
 			);
 		}
 	}
diff --git a/Engine/Source/Editor/Persona/Private/SAnimSegmentsPanel.h b/Engine/Source/Editor/Persona/Private/SAnimSegmentsPanel.h
--- a/Engine/Source/Editor/Persona/Private/SAnimSegmentsPanel.h
+++ b/Engine/Source/Editor/Persona/Private/SAnimSegmentsPanel.h
@@ -25,6 +25,7 @@ DECLARE_DELEGATE_OneParam( FOnAnimSegmentNodeClicked, int32 )
 DECLARE_DELEGATE_OneParam( FOnAnimSegmentRemoved, int32 )
 DECLARE_DELEGATE_FourParams( FOnAnimReplaceMapping, FName, int32, UAnimSequenceBase*, UAnimSequenceBase*)
 DECLARE_DELEGATE_RetVal_ThreeParams(bool, FOnDiffFromParentAsset, FName, int32, const FAnimSegment& )
+DECLARE_DELEGATE_RetVal_OneParam(FLinearColor, FOnGetNodeColor, const FAnimSegment&)
 
 //////////////////////////////////////////////////////////////////////////
 // SAnimSegmentsPanel
@@ -63,7 +64,7 @@ public:
 	SLATE_ATTRIBUTE(float, ViewInputMin)
 	SLATE_ATTRIBUTE(float, ViewInputMax)
 
-	SLATE_ATTRIBUTE(FLinearColor, NodeColor)
+	SLATE_EVENT(FOnGetNodeColor, OnGetNodeColor)
 	SLATE_ATTRIBUTE(TArray<float>, DraggableBars)
 	SLATE_ATTRIBUTE(TArray<float>, DraggableBarSnapPositions)
 	SLATE_ATTRIBUTE(float, ScrubPosition)
@@ -102,6 +103,7 @@ private:
 	FOnAnimSegmentRemoved		OnAnimSegmentRemovedDelegate;
 	FOnAnimReplaceMapping		OnAnimReplaceMapping;
 	FOnDiffFromParentAsset		OnDiffFromParentAsset;
+	FOnGetNodeColor				OnGetNodeColor;
 
 	enum ETrackViewStyle
 	{
diff --git a/Engine/Source/Editor/Persona/Private/SAnimSequenceCurveEditor.cpp b/Engine/Source/Editor/Persona/Private/SAnimSequenceCurveEditor.cpp
new file mode 100644
--- /dev/null
+++ b/Engine/Source/Editor/Persona/Private/SAnimSequenceCurveEditor.cpp
@@ -0,0 +1,359 @@
+// Copyright Epic Games, Inc. All Rights Reserved.
+
+#include "SAnimSequenceCurveEditor.h"
+#include "CurveEditor.h"
+#include "RichCurveEditorModel.h"
+#include "Animation/AnimSequenceBase.h"
+#include "SCurveEditorPanel.h"
+#include "Tree/SCurveEditorTreeTextFilter.h"
+#include "Tree/SCurveEditorTreeFilterStatusBar.h"
+#include "Tree/SCurveEditorTree.h"
+#include "Tree/SCurveEditorTreePin.h"
+#include "Widgets/Layout/SScrollBorder.h"
+#include "Tree/ICurveEditorTreeItem.h"
+#include "AnimTimeline/SAnimTimelineTransportControls.h"
+#include "Tree/CurveEditorTreeFilter.h"
+#include "Editor.h"
+#include "Framework/MultiBox/MultiBoxBuilder.h"
+#include "EditorStyleSet.h"
+
+#define LOCTEXT_NAMESPACE "SAnimSequenceCurveEditor"
+
+// Model that references a named curve, rather than a raw pointer, so we avoid issues with 
+// reallocation of the curves arrays under the UI
+class FRichCurveEditorModelNamed : public FRichCurveEditorModel
+{
+public:
+	FRichCurveEditorModelNamed(const FSmartName& InName, ERawCurveTrackTypes InType, int32 InCurveIndex, UAnimSequenceBase* InAnimSequence, FCurveEditorTreeItemID InTreeId)
+		: FRichCurveEditorModel(InAnimSequence)
+		, Name(InName)
+		, AnimSequence(InAnimSequence)
+		, CurveIndex(InCurveIndex)
+		, Type(InType)
+		, TreeId(InTreeId)
+	{
+	}
+
+	virtual bool IsValid() const override
+	{
+		return AnimSequence->GetCurveData().GetCurveData(Name.UID, Type) != nullptr;
+	}
+
+	virtual FRichCurve& GetRichCurve() override
+	{
+		check(AnimSequence.Get() != nullptr);
+
+		FAnimCurveBase* CurveBase = AnimSequence->RawCurveData.GetCurveData(Name.UID, Type);
+		check(CurveBase);	// If this fails lifetime contracts have been violated - this curve should always be present if this model exists
+		
+		switch (Type)
+		{
+		case ERawCurveTrackTypes::RCT_Vector:
+		{
+			FVectorCurve& VectorCurve = *(static_cast<FVectorCurve*>(CurveBase));
+			check(CurveIndex < 3);
+			return VectorCurve.FloatCurves[CurveIndex];
+		}
+		case ERawCurveTrackTypes::RCT_Transform:
+		{
+			FTransformCurve& TransformCurve = *(static_cast<FTransformCurve*>(CurveBase));
+			check(CurveIndex < 9);
+			const int32 SubCurveIndex = CurveIndex / 3;
+			switch(SubCurveIndex)
+			{
+			default:
+				check(false);
+				// fall through
+			case 0:
+				return TransformCurve.TranslationCurve.FloatCurves[SubCurveIndex];
+			case 1:
+				return TransformCurve.RotationCurve.FloatCurves[SubCurveIndex];
+			case 2:
+				return TransformCurve.ScaleCurve.FloatCurves[SubCurveIndex];
+			}
+			
+		}
+		case ERawCurveTrackTypes::RCT_Float:
+		default:
+		{
+			FFloatCurve& FloatCurve = *(static_cast<FFloatCurve*>(CurveBase));
+			check(CurveIndex == 0);
+			return FloatCurve.FloatCurve;
+		}
+		}
+	}
+
+	virtual const FRichCurve& GetReadOnlyRichCurve() const override
+	{
+		return const_cast<FRichCurveEditorModelNamed*>(this)->GetRichCurve();
+	}
+
+	FSmartName Name;
+	TWeakObjectPtr<UAnimSequenceBase> AnimSequence;
+	int32 CurveIndex;
+	ERawCurveTrackTypes Type;
+	FCurveEditorTreeItemID TreeId;
+};
+
+class FAnimSequenceCurveEditorItem : public ICurveEditorTreeItem
+{
+public:
+	FAnimSequenceCurveEditorItem(const FSmartName& InName, ERawCurveTrackTypes InType, int32 InCurveIndex, UAnimSequenceBase* InAnimSequence, const FText& InCurveDisplayName, const FLinearColor& InCurveColor, FSimpleDelegate InOnCurveModified, FCurveEditorTreeItemID InTreeId)
+		: Name(InName)
+		, Type(InType)
+		, CurveIndex(InCurveIndex)
+		, AnimSequence(InAnimSequence)
+		, CurveDisplayName(InCurveDisplayName)
+		, CurveColor(InCurveColor)
+		, OnCurveModified(InOnCurveModified)
+		, TreeId(InTreeId)
+	{
+	}
+
+	virtual TSharedPtr<SWidget> GenerateCurveEditorTreeWidget(const FName& InColumnName, TWeakPtr<FCurveEditor> InCurveEditor, FCurveEditorTreeItemID InTreeItemID, const TSharedRef<ITableRow>& InTableRow) override
+	{
+		if (InColumnName == ColumnNames.Label)
+		{
+			return SNew(SHorizontalBox)
+				+SHorizontalBox::Slot()
+				.Padding(FMargin(4.f))
+				.VAlign(VAlign_Center)
+				.HAlign(HAlign_Right)
+				.AutoWidth()
+				[
+					SNew(STextBlock)
+					.Text(CurveDisplayName)
+					.ColorAndOpacity(FSlateColor(CurveColor))
+				];
+		}
+		else if (InColumnName == ColumnNames.PinHeader)
+		{
+			return SNew(SCurveEditorTreePin, InCurveEditor, InTreeItemID, InTableRow);
+		}
+
+		return nullptr;
+	}
+
+	virtual void CreateCurveModels(TArray<TUniquePtr<FCurveModel>>& OutCurveModels) override
+	{
+		TUniquePtr<FRichCurveEditorModelNamed> NewCurveModel = MakeUnique<FRichCurveEditorModelNamed>(Name, Type, CurveIndex, AnimSequence.Get(), TreeId);
+		NewCurveModel->SetShortDisplayName(CurveDisplayName);
+		NewCurveModel->SetLongDisplayName(CurveDisplayName);
+		NewCurveModel->SetColor(CurveColor);
+		NewCurveModel->SetClampInputRange(MakeAttributeLambda([this](){ return TRange<double>(0.0, AnimSequence->GetPlayLength()); }));
+		NewCurveModel->OnCurveModified().Add(OnCurveModified);
+
+		OutCurveModels.Add(MoveTemp(NewCurveModel));
+	}
+
+	virtual bool PassesFilter(const FCurveEditorTreeFilter* InFilter) const override
+	{
+		if (InFilter->GetType() == ECurveEditorTreeFilterType::Text)
+		{
+			const FCurveEditorTreeTextFilter* Filter = static_cast<const FCurveEditorTreeTextFilter*>(InFilter);
+			for (const FCurveEditorTreeTextFilterTerm& Term : Filter->GetTerms())
+			{
+				for(const FCurveEditorTreeTextFilterToken& Token : Term.ChildToParentTokens)
+				{
+					if(Token.Match(*CurveDisplayName.ToString()))
+					{
+						return true;
+					}
+				}
+			}
+
+			return false;
+		}
+
+		return false;
+	}
+
+	FSmartName Name;
+	ERawCurveTrackTypes Type;
+	int32 CurveIndex;
+	TWeakObjectPtr<UAnimSequenceBase> AnimSequence;
+	FText CurveDisplayName;
+	FLinearColor CurveColor;
+	FSimpleDelegate OnCurveModified;
+	FCurveEditorTreeItemID TreeId;
+};
+
+class FAnimSequenceCurveEditorBounds : public ICurveEditorBounds
+{
+public:
+	FAnimSequenceCurveEditorBounds(TSharedPtr<ITimeSliderController> InExternalTimeSliderController)
+		: ExternalTimeSliderController(InExternalTimeSliderController)
+	{}
+
+	virtual void GetInputBounds(double& OutMin, double& OutMax) const override
+	{
+		FAnimatedRange ViewRange = ExternalTimeSliderController.Pin()->GetViewRange();
+		OutMin = ViewRange.GetLowerBoundValue();
+		OutMax = ViewRange.GetUpperBoundValue();
+	}
+
+	virtual void SetInputBounds(double InMin, double InMax) override
+	{
+		ExternalTimeSliderController.Pin()->SetViewRange(InMin, InMax, EViewRangeInterpolation::Immediate);
+	}
+
+	TWeakPtr<ITimeSliderController> ExternalTimeSliderController;
+};
+
+SAnimSequenceCurveEditor::SAnimSequenceCurveEditor()
+{
+	if(GEditor)
+	{
+		GEditor->RegisterForUndo(this);
+	}
+}
+
+SAnimSequenceCurveEditor::~SAnimSequenceCurveEditor()
+{
+	if(GEditor)
+	{
+		GEditor->UnregisterForUndo(this);
+	}
+}
+
+void SAnimSequenceCurveEditor::Construct(const FArguments& InArgs, const TSharedRef<IPersonaPreviewScene>& InPreviewScene, UAnimSequenceBase* InAnimSequence)
+{
+	CurveEditor = MakeShared<FCurveEditor>();
+	CurveEditor->GridLineLabelFormatXAttribute = LOCTEXT("GridXLabelFormat", "{0}s");
+	CurveEditor->SetBounds(MakeUnique<FAnimSequenceCurveEditorBounds>(InArgs._ExternalTimeSliderController));
+
+	FCurveEditorInitParams CurveEditorInitParams;
+	CurveEditor->InitCurveEditor(CurveEditorInitParams);
+
+	AnimSequence = InAnimSequence;
+
+	CurveEditorTree = SNew(SCurveEditorTree, CurveEditor);
+
+	TSharedRef<SCurveEditorPanel> CurveEditorPanel = SNew(SCurveEditorPanel, CurveEditor.ToSharedRef())
+		.GridLineTint(FLinearColor(0.f, 0.f, 0.f, 0.3f))
+		.ExternalTimeSliderController(InArgs._ExternalTimeSliderController)
+		.TabManager(InArgs._TabManager)
+		.TreeSplitterWidth(0.2f)
+		.ContentSplitterWidth(0.8f)
+		.TreeContent()
+		[
+			SNew(SVerticalBox)
+			+SVerticalBox::Slot()
+			.AutoHeight()
+			[
+				SAssignNew(CurveEditorSearchBox, SCurveEditorTreeTextFilter, CurveEditor)
+			]
+			+SVerticalBox::Slot()
+			[
+				SNew(SScrollBorder, CurveEditorTree.ToSharedRef())
+				[
+					CurveEditorTree.ToSharedRef()
+				]
+			]
+			+SVerticalBox::Slot()
+			.AutoHeight()
+			[
+				SNew(SCurveEditorTreeFilterStatusBar, CurveEditor)
+			]
+			+SVerticalBox::Slot()
+			.AutoHeight()
+			.HAlign(HAlign_Center)
+			[
+				SNew(SAnimTimelineTransportControls, InPreviewScene, InAnimSequence)
+			]
+		];
+
+	ChildSlot
+	[
+		SNew(SVerticalBox)
+		+SVerticalBox::Slot()
+		.AutoHeight()
+		.Padding(0.0f, 0.0f, 0.0f, 3.0f)
+		[
+			MakeToolbar(CurveEditorPanel)
+		]
+		+SVerticalBox::Slot()
+		.FillHeight(1.0f)
+		[
+			CurveEditorPanel
+		]
+	];
+}
+
+TSharedRef<SWidget> SAnimSequenceCurveEditor::MakeToolbar(TSharedRef<SCurveEditorPanel> InEditorPanel)
+{
+	FToolBarBuilder ToolBarBuilder(InEditorPanel->GetCommands(), FMultiBoxCustomization::None, InEditorPanel->GetToolbarExtender(), EOrientation::Orient_Horizontal, true);
+	ToolBarBuilder.SetStyle(&FEditorStyle::Get(), "Sequencer.ToolBar");
+	ToolBarBuilder.BeginSection("Asset");
+	ToolBarBuilder.EndSection();
+	// We just use all of the extenders as our toolbar, we don't have a need to create a separate toolbar.
+	return ToolBarBuilder.MakeWidget();
+}
+
+void SAnimSequenceCurveEditor::ResetCurves()
+{
+	CurveEditor->RemoveAllTreeItems();
+	CurveEditor->RemoveAllCurves();
+}
+
+void SAnimSequenceCurveEditor::AddCurve(const FText& InCurveDisplayName, const FLinearColor& InCurveColor, const FSmartName& InName, ERawCurveTrackTypes InType, int32 InCurveIndex, FSimpleDelegate InOnCurveModified)
+{
+	FCurveEditorTreeItem* TreeItem = CurveEditor->AddTreeItem(FCurveEditorTreeItemID());
+	TreeItem->SetStrongItem(MakeShared<FAnimSequenceCurveEditorItem>(InName, InType, InCurveIndex, AnimSequence, InCurveDisplayName, InCurveColor, InOnCurveModified, TreeItem->GetID()));
+
+	// Update selection
+	const TMap<FCurveEditorTreeItemID, ECurveEditorTreeSelectionState>& Selection = CurveEditor->GetTreeSelection();
+	TArray<FCurveEditorTreeItemID> NewSelection;
+	NewSelection.Add(TreeItem->GetID());
+	for(const auto& SelectionPair : Selection)
+	{
+		if(SelectionPair.Value != ECurveEditorTreeSelectionState::None)
+		{
+			NewSelection.Add(SelectionPair.Key);
+		}
+	}
+	CurveEditor->SetTreeSelection(MoveTemp(NewSelection));
+}
+
+void SAnimSequenceCurveEditor::RemoveCurve(const FSmartName& InName, ERawCurveTrackTypes InType, int32 InCurveIndex)
+{
+	for(const auto& CurvePair : CurveEditor->GetCurves())
+	{
+		FRichCurveEditorModelNamed* Model = static_cast<FRichCurveEditorModelNamed*>(CurvePair.Value.Get());
+		if(Model->Name == InName && Model->Type == InType && Model->CurveIndex == InCurveIndex)
+		{
+			CurveEditor->RemoveCurve(CurvePair.Key);
+			CurveEditor->RemoveTreeItem(Model->TreeId);
+			break;
+		}
+	}
+}
+
+void SAnimSequenceCurveEditor::ZoomToFit()
+{
+	CurveEditor->ZoomToFit(EAxisList::Y);
+}
+
+void SAnimSequenceCurveEditor::PostUndoRedo()
+{
+	// Check if our curves are still valid & remove if not
+	bool bRemoved = false;
+	do
+	{
+		bRemoved = false;
+		for(const TPair<FCurveModelID, TUniquePtr<FCurveModel>>& Pair : CurveEditor->GetCurves())
+		{
+			FRichCurveEditorModelNamed* Model = static_cast<FRichCurveEditorModelNamed*>(Pair.Value.Get());
+			if(!Model->IsValid())
+			{
+				CurveEditor->RemoveCurve(Pair.Key);
+				CurveEditor->RemoveTreeItem(Model->TreeId);
+				bRemoved = true;
+				break;
+			}
+		}
+	}
+	while(bRemoved);
+}
+
+#undef LOCTEXT_NAMESPACE
\ No newline at end of file
diff --git a/Engine/Source/Editor/Persona/Private/SAnimSequenceCurveEditor.h b/Engine/Source/Editor/Persona/Private/SAnimSequenceCurveEditor.h
new file mode 100644
--- /dev/null
+++ b/Engine/Source/Editor/Persona/Private/SAnimSequenceCurveEditor.h
@@ -0,0 +1,60 @@
+// Copyright Epic Games, Inc. All Rights Reserved.
+
+#pragma once
+
+#include "IAnimSequenceCurveEditor.h"
+#include "Widgets/DeclarativeSyntaxSupport.h"
+#include "EditorUndoClient.h"
+
+class FCurveEditor;
+class UAnimSequenceBase;
+class ITimeSliderController;
+class SCurveEditorTree;
+class IPersonaPreviewScene;
+class SCurveEditorPanel;
+class FTabManager;
+
+class SAnimSequenceCurveEditor : public IAnimSequenceCurveEditor, public FEditorUndoClient
+{
+	SLATE_BEGIN_ARGS(SAnimSequenceCurveEditor) {}
+
+	SLATE_ARGUMENT(TSharedPtr<ITimeSliderController>, ExternalTimeSliderController)
+
+	SLATE_ARGUMENT(TSharedPtr<FTabManager>, TabManager)
+
+	SLATE_END_ARGS()
+
+	SAnimSequenceCurveEditor();
+	~SAnimSequenceCurveEditor();
+
+	void Construct(const FArguments& InArgs, const TSharedRef<IPersonaPreviewScene>& InPreviewScene, UAnimSequenceBase* InAnimSequence);
+
+	/** IAnimSequenceCurveEditor interface */
+	virtual void ResetCurves() override;
+	virtual void AddCurve(const FText& InCurveDisplayName, const FLinearColor& InCurveColor, const FSmartName& InName, ERawCurveTrackTypes InType, int32 InCurveIndex, FSimpleDelegate InOnCurveModified) override;
+	virtual void RemoveCurve(const FSmartName& InName, ERawCurveTrackTypes InType, int32 InCurveIndex) override;
+	virtual void ZoomToFit() override;
+	
+	virtual void PostUndo(bool bSuccess) override { PostUndoRedo(); }
+	virtual void PostRedo(bool bSuccess) override { PostUndoRedo(); }
+
+private:
+	// Build the toolbar for this curve editor
+	TSharedRef<SWidget> MakeToolbar(TSharedRef<SCurveEditorPanel> InEditorPanel);
+
+	// Handle undo/redo to check underlying curve data is still valid
+	void PostUndoRedo();
+
+private:
+	/** The actual curve editor */
+	TSharedPtr<FCurveEditor> CurveEditor;
+
+	/** The search widget for filtering curves in the Curve Editor tree. */
+	TSharedPtr<SWidget> CurveEditorSearchBox;
+
+	/** The anim sequence we are editing */
+	UAnimSequenceBase* AnimSequence;
+
+	/** The tree widget int he curve editor */
+	TSharedPtr<SCurveEditorTree> CurveEditorTree;
+};
\ No newline at end of file
diff --git a/Engine/Source/Editor/Persona/Private/SAnimStreamableEditor.cpp b/Engine/Source/Editor/Persona/Private/SAnimStreamableEditor.cpp
--- a/Engine/Source/Editor/Persona/Private/SAnimStreamableEditor.cpp
+++ b/Engine/Source/Editor/Persona/Private/SAnimStreamableEditor.cpp
@@ -2,156 +2,67 @@
 
 
 #include "SAnimStreamableEditor.h"
-#include "Animation/EditorAnimBaseObj.h"
 #include "IDocumentation.h"
-
-#include "SAnimNotifyPanel.h"
-#include "Editor.h"
+#include "AnimModel_AnimSequenceBase.h"
+#include "SAnimTimeline.h"
 
 //////////////////////////////////////////////////////////////////////////
-// SAnimCompositeEditor
+// SAnimStreamableEditor
 
 TSharedRef<SWidget> SAnimStreamableEditor::CreateDocumentAnchor()
 {
-	return IDocumentation::Get()->CreateAnchor(TEXT("Engine/Animation/AnimationComposite")); // UPDATE TO STREAMABLE
+	return IDocumentation::Get()->CreateAnchor(TEXT("Engine/Animation/AnimationStreamable"));
 }
 
-void SAnimStreamableEditor::Construct(const FArguments& InArgs, const TSharedRef<class IPersonaPreviewScene>& InPreviewScene, const TSharedRef<class IEditableSkeleton>& InEditableSkeleton)
+void SAnimStreamableEditor::Construct(const FArguments& InArgs, const TSharedRef<class IPersonaPreviewScene>& InPreviewScene, const TSharedRef<class IEditableSkeleton>& InEditableSkeleton, const TSharedRef<FUICommandList>& InCommandList)
 {
-	bIsActiveTimerRegistered = false;
 	StreamableAnim = InArgs._StreamableAnim;
 	check(StreamableAnim);
 
-	SAnimEditorBase::Construct( SAnimEditorBase::FArguments()
-		.OnObjectsSelected(InArgs._OnObjectsSelected), 
-		InPreviewScene );
+	AnimModel = MakeShared<FAnimModel_AnimSequenceBase>(InPreviewScene, InEditableSkeleton, InCommandList, StreamableAnim);
+	AnimModel->OnSelectObjects = FOnObjectsSelected::CreateSP(this, &SAnimEditorBase::OnSelectionChanged);
+	AnimModel->OnInvokeTab = InArgs._OnInvokeTab;
 
-	if(GEditor)
+	AnimModel->OnEditCurves = FOnEditCurves::CreateLambda([this, InOnEditCurves = InArgs._OnEditCurves](UAnimSequenceBase* InAnimSequence, const TArray<IAnimationEditor::FCurveEditInfo>& InCurveInfo, const TSharedPtr<ITimeSliderController>& InExternalTimeSliderController)
 	{
-		GEditor->RegisterForUndo(this);
-	}
-
-	/*EditorPanels->AddSlot()
-		.AutoHeight()
-		.Padding(0, 10)
-		[
-			SAssignNew( AnimCompositePanel, SAnimCompositePanel )
-			.Composite(CompositeObj)
-			.CompositeEditor(SharedThis(this))
-			.WidgetWidth(S2ColumnWidget::DEFAULT_RIGHT_COLUMN_WIDTH)
-			.ViewInputMin(this, &SAnimEditorBase::GetViewMinInput)
-			.ViewInputMax(this, &SAnimEditorBase::GetViewMaxInput)
-			.OnSetInputViewRange(this, &SAnimEditorBase::SetInputViewRange)
-		];*/
+		InOnEditCurves.ExecuteIfBound(InAnimSequence, InCurveInfo, TimelineWidget->GetTimeSliderController());
+	});
 
-	EditorPanels->AddSlot()
-		.AutoHeight()
-		.Padding(0, 10)
-		[
-			SAssignNew( AnimNotifyPanel, SAnimNotifyPanel, InEditableSkeleton )
-			.Sequence(StreamableAnim)
-			.WidgetWidth(S2ColumnWidget::DEFAULT_RIGHT_COLUMN_WIDTH)
-			.InputMin(this, &SAnimEditorBase::GetMinInput)
-			.InputMax(this, &SAnimEditorBase::GetMaxInput)
-			.ViewInputMin(this, &SAnimEditorBase::GetViewMinInput)
-			.ViewInputMax(this, &SAnimEditorBase::GetViewMaxInput)
-			.OnSetInputViewRange(this, &SAnimEditorBase::SetInputViewRange)
-			.OnGetScrubValue(this, &SAnimEditorBase::GetScrubValue)
-			.OnSelectionChanged(this, &SAnimEditorBase::OnSelectionChanged)
-			.OnInvokeTab(InArgs._OnInvokeTab)
-		];
+	AnimModel->OnStopEditingCurves = InArgs._OnStopEditingCurves;
+	AnimModel->Initialize();
 
-	EditorPanels->AddSlot()
-		.AutoHeight()
-		.Padding(0, 10)
-		[
-			SAssignNew( AnimCurvePanel, SAnimCurvePanel, InEditableSkeleton)
-			.Sequence(StreamableAnim)
-			.WidgetWidth(S2ColumnWidget::DEFAULT_RIGHT_COLUMN_WIDTH)
-			.ViewInputMin(this, &SAnimEditorBase::GetViewMinInput)
-			.ViewInputMax(this, &SAnimEditorBase::GetViewMaxInput)
-			.InputMin(this, &SAnimEditorBase::GetMinInput)
-			.InputMax(this, &SAnimEditorBase::GetMaxInput)
-			.OnSetInputViewRange(this, &SAnimEditorBase::SetInputViewRange)
-			.OnGetScrubValue(this, &SAnimEditorBase::GetScrubValue)
-		];
-}
+	SAnimEditorBase::Construct(SAnimEditorBase::FArguments()
+		.OnObjectsSelected(InArgs._OnObjectsSelected)
+		.AnimModel(AnimModel),
+		InPreviewScene);
 
-SAnimStreamableEditor::~SAnimStreamableEditor()
-{
-	if(GEditor)
+	if (GEditor)
 	{
-		GEditor->UnregisterForUndo(this);
+		GEditor->RegisterForUndo(this);
 	}
 }
 
-void SAnimStreamableEditor::PreAnimUpdate()
-{
-	StreamableAnim->Modify();
-}
-
-void SAnimStreamableEditor::PostAnimUpdate()
-{
-	StreamableAnim->MarkPackageDirty();
-	//SortAndUpdateComposite();
-}
-
-void SAnimStreamableEditor::RebuildPanel()
-{
-	//SortAndUpdateComposite();
-	//AnimCompositePanel->Update();
-}
-
-void SAnimStreamableEditor::OnStreamableChange(class UObject *EditorAnimBaseObj, bool bRebuild)
+SAnimStreamableEditor::~SAnimStreamableEditor()
 {
-	if (StreamableAnim != nullptr )
+	if (GEditor)
 	{
-		if(bRebuild && !bIsActiveTimerRegistered)
-		{
-			// sometimes crashes because the timer delay but animation still renders, so invalidating here before calling timer
-			bIsActiveTimerRegistered = true;
-			RegisterActiveTimer(0.f, FWidgetActiveTimerDelegate::CreateSP(this, &SAnimStreamableEditor::TriggerRebuildPanel));
-		} 
-
-		StreamableAnim->MarkPackageDirty();
+		GEditor->UnregisterForUndo(this);
 	}
 }
 
-void SAnimStreamableEditor::PostUndo( bool bSuccess )
+void SAnimStreamableEditor::PostUndo(bool bSuccess)
 {
 	PostUndoRedo();
 }
 
-void SAnimStreamableEditor::PostRedo( bool bSuccess )
+void SAnimStreamableEditor::PostRedo(bool bSuccess)
 {
 	PostUndoRedo();
 }
 
 void SAnimStreamableEditor::PostUndoRedo()
 {
-	if (!bIsActiveTimerRegistered)
-	{
-		bIsActiveTimerRegistered = true;
-		RegisterActiveTimer(0.f, FWidgetActiveTimerDelegate::CreateSP(this, &SAnimStreamableEditor::TriggerRebuildPanel));
-	}
-}
-
-void SAnimStreamableEditor::InitDetailsViewEditorObject(UEditorAnimBaseObj* EdObj)
-{
-	EdObj->InitFromAnim(StreamableAnim, FOnAnimObjectChange::CreateSP( SharedThis(this), &SAnimStreamableEditor::OnStreamableChange ));
-}
-
-EActiveTimerReturnType SAnimStreamableEditor::TriggerRebuildPanel(double InCurrentTime, float InDeltaTime)
-{
-	// we should not update any property related within PostEditChange, 
-	// so this is deferred to Tick, when it needs to rebuild, just mark it and this will update in the next tick
-	RebuildPanel();
-
-	bIsActiveTimerRegistered = false;
-	return EActiveTimerReturnType::Stop;
-}
+	GetPreviewScene()->SetPreviewAnimationAsset(StreamableAnim);
 
-float SAnimStreamableEditor::CalculateSequenceLengthOfEditorObject() const
-{
-	return StreamableAnim->SequenceLength;
+	AnimModel->RefreshTracks();
 }
