COPY from //Dev-Anim to //Dev-Main

diff --git a/Engine/Source/Editor/Persona/Private/SAnimStreamableEditor.h b/Engine/Source/Editor/Persona/Private/SAnimStreamableEditor.h
--- a/Engine/Source/Editor/Persona/Private/SAnimStreamableEditor.h
+++ b/Engine/Source/Editor/Persona/Private/SAnimStreamableEditor.h
@@ -11,77 +11,52 @@
 #include "Animation/AnimStreamable.h"
 
 class SAnimNotifyPanel;
+class FAnimModel_AnimSequenceBase;
 
 //////////////////////////////////////////////////////////////////////////
-// SAnimCompositeEditor
+// SAnimStreamableEditor
 
-/** Overall animation composite editing widget. This mostly contains functions for editing the UAnimComposite.
-
-	SAnimCompositeEditor will create the SAnimCompositePanel which is mostly responsible for setting up the UI 
-	portion of the composite tool and registering callbacks to the SAnimCompositeEditor to do the actual editing.
-	
-*/
+/** Overall streamable animation editing widget. This mostly contains functions for editing the UAnimStreamable. */
 class SAnimStreamableEditor : public SAnimEditorBase, public FEditorUndoClient
 {
 public:
-	SLATE_BEGIN_ARGS(SAnimStreamableEditor)
-		: _StreamableAnim(NULL)
-		{}
+	SLATE_BEGIN_ARGS( SAnimStreamableEditor )
+		: _StreamableAnim(nullptr)
+	{}
 
-		SLATE_ARGUMENT( UAnimStreamable*, StreamableAnim)
-		SLATE_EVENT(FOnObjectsSelected, OnObjectsSelected)
-		SLATE_EVENT(FOnInvokeTab, OnInvokeTab)
+	SLATE_ARGUMENT( UAnimStreamable*, StreamableAnim)
+	SLATE_EVENT(FOnObjectsSelected, OnObjectsSelected)
+	SLATE_EVENT(FOnInvokeTab, OnInvokeTab)
+	SLATE_EVENT(FOnEditCurves, OnEditCurves)
+	SLATE_EVENT(FOnStopEditingCurves, OnStopEditingCurves)
 
 	SLATE_END_ARGS()
 
 private:
-	/** Slate editor panels */
-	TSharedPtr<class SAnimNotifyPanel> AnimNotifyPanel;
-	TSharedPtr<class SAnimCurvePanel>	AnimCurvePanel;
-
-	/** do I have to update the panel **/
-	bool bRebuildPanel;
-	void RebuildPanel();
-
-	/** Handler for when composite is edited in details view */
-	void OnStreamableChange(class UObject *EditorAnimBaseObj, bool bRebuild);
-
-protected:
-	virtual void InitDetailsViewEditorObject(class UEditorAnimBaseObj* EdObj) override;
+	TSharedPtr<FAnimModel_AnimSequenceBase> AnimModel;
 
 public:
-	void Construct(const FArguments& InArgs, const TSharedRef<class IPersonaPreviewScene>& InPreviewScene, const TSharedRef<class IEditableSkeleton>& InEditableSkeleton);
-
 	~SAnimStreamableEditor();
 
-	/** Return the animation composite being edited */
-	UAnimStreamable* GetStreamableAnimObj() const { return StreamableAnim; }
-	virtual UAnimationAsset* GetEditorObject() const override { return GetStreamableAnimObj(); }
+	void Construct(const FArguments& InArgs, const TSharedRef<class IPersonaPreviewScene>& InPreviewScene, const TSharedRef<class IEditableSkeleton>& InEditableSkeleton, const TSharedRef<FUICommandList>& InCommandList);
+
+	/** Return the streamable animation being edited */
+	UAnimStreamable* GetStreamableAnim() const { return StreamableAnim; }
+	virtual UAnimationAsset* GetEditorObject() const override { return GetStreamableAnim(); }
 
 private:
-	/** Pointer to the animation composite being edited */
+	/** Pointer to the streamable animation being edited */
 	UAnimStreamable* StreamableAnim;
-	
-	/** FEditorUndoClient interface */
-	virtual void PostUndo( bool bSuccess ) override;
-	virtual void PostRedo( bool bSuccess ) override;
-	void PostUndoRedo();
-
-	virtual float CalculateSequenceLengthOfEditorObject() const override;
-
-	/** One-off active timer to trigger a panel rebuild */
-	EActiveTimerReturnType TriggerRebuildPanel(double InCurrentTime, float InDeltaTime);
-
-	/** Whether the active timer to trigger a panel rebuild is currently registered */
-	bool bIsActiveTimerRegistered;
 
 public:
-
-	/** delegate handlers for when the composite is being editor */
-	void			PreAnimUpdate();
-	void			PostAnimUpdate();
-
 	//~ Begin SAnimEditorBase Interface
 	virtual TSharedRef<SWidget> CreateDocumentAnchor() override;
 	//~ End SAnimEditorBase Interface
+
+	/** FEditorUndoClient interface */
+	virtual void PostUndo( bool bSuccess ) override;
+	virtual void PostRedo( bool bSuccess ) override;
+
+	/** Post undo **/
+	void PostUndoRedo();
 };
diff --git a/Engine/Source/Editor/Persona/Private/SAnimTimingPanel.cpp b/Engine/Source/Editor/Persona/Private/SAnimTimingPanel.cpp
--- a/Engine/Source/Editor/Persona/Private/SAnimTimingPanel.cpp
+++ b/Engine/Source/Editor/Persona/Private/SAnimTimingPanel.cpp
@@ -22,6 +22,10 @@
 #include "Modules/ModuleManager.h"
 #include "IEditableSkeleton.h"
 #include "Editor.h"
+#include "AnimModel.h"
+#include "AnimModel_AnimMontage.h"
+#include "Preferences/PersonaOptions.h"
+#include "FrameNumberDisplayFormat.h"
 
 #define LOCTEXT_NAMESPACE "AnimTimingPanel"
 
@@ -31,11 +35,101 @@ namespace AnimTimingConstants
 	static const int32 FontSize = 10;
 }
 
+class SAnimTimingNodeTooltip : public SToolTip
+{
+public:
+	void Construct(const FArguments& InArgs, const TSharedRef<FTimingRelevantElementBase>& InElement)
+	{
+		Element = InElement;
+		DescriptionBox = SNew(SVerticalBox);
+
+		FArguments Args = InArgs;
+		Args._TextMargin = FMargin(1.0f);
+		Args._BorderImage = FEditorStyle::GetBrush("ContentBrowser.TileViewToolTip.ToolTipBorder");
+		Args._Content.Widget =
+			SNew(SBorder)
+			.Padding(3)
+			.BorderImage(FEditorStyle::GetBrush("ContentBrowser.TileViewTooltip.NonContentBorder"))
+			[
+				SNew(SVerticalBox)
+				+ SVerticalBox::Slot()
+				.AutoHeight()
+				.Padding(0, 0, 0, 3)
+				[
+					SNew(SBorder)
+					.Padding(6)
+					.BorderImage(FEditorStyle::GetBrush("ContentBrowser.TileViewTooltip.ContentBorder"))
+					[
+						SNew(SBox)
+						.HAlign(HAlign_Left)
+						[
+							SNew(STextBlock)
+							.Text(FText::FromName(Element->GetTypeName()))
+							.Font(FEditorStyle::GetFontStyle("ContentBrowser.TileViewTooltip.NameFont"))
+						]
+					]
+				]
+
+				+ SVerticalBox::Slot()
+				[
+					SNew(SHorizontalBox)
+					+ SHorizontalBox::Slot()
+					.FillWidth(1.0f)
+					[
+						SNew(SBorder)
+						.Padding(3)
+						.BorderImage(FEditorStyle::GetBrush("ContentBrowser.TileViewTooltip.ContentBorder"))
+						[
+							DescriptionBox.ToSharedRef()
+						]
+					]
+				]
+			];
+
+		SToolTip::Construct(Args);
+	}
+
+	virtual void OnOpening() override
+	{
+		DescriptionBox->ClearChildren();
+
+		TMap<FString, FText> DescriptionItems;
+		Element->GetDescriptionItems(DescriptionItems);
+		for(TPair<FString, FText> ItemPair : DescriptionItems)
+		{
+			DescriptionBox->AddSlot()
+				.AutoHeight()
+				.Padding(0, 0, 3, 0)
+				[
+					SNew(SHorizontalBox)
+					+ SHorizontalBox::Slot()
+					.AutoWidth()
+					[
+						SNew(STextBlock)
+						.Text(FText::Format(LOCTEXT("Item", "{0}: "), FText::FromString(ItemPair.Key)))
+						.ColorAndOpacity(FSlateColor::UseSubduedForeground())
+					]
+
+					+SHorizontalBox::Slot()
+					.AutoWidth()
+					[
+						SNew(STextBlock)
+						.Text(ItemPair.Value)
+						.ColorAndOpacity(FSlateColor::UseForeground())
+					]
+				];
+		}
+	}
+
+	TSharedPtr<SVerticalBox> DescriptionBox;
+	TSharedPtr<FTimingRelevantElementBase> Element;
+};
+
 void SAnimTimingNode::Construct(const FArguments& InArgs)
 {
 	Element = InArgs._InElement;
 
-	const FSlateBrush* StyleInfo = FEditorStyle::GetBrush("ProgressBar.Background");
+	const FSlateBrush* StyleInfo = FEditorStyle::GetBrush(TEXT("SpecialEditableTextImageNormal"));
 	static FSlateFontInfo LabelFont = FCoreStyle::GetDefaultFontStyle("Regular", AnimTimingConstants::FontSize);
 
 	UPersonaOptions* EditorOptions = UPersonaOptions::StaticClass()->GetDefaultObject<UPersonaOptions>();
@@ -82,80 +176,7 @@ void SAnimTimingNode::Construct(const FArguments& InArgs)
 
 	if(InArgs._bUseTooltip)
 	{
-		// Add asset registry tags to a text list; except skeleton as that is implied in Persona
-		TSharedRef<SVerticalBox> DescriptionBox = SNew(SVerticalBox);
-		TMap<FString, FText> DescriptionItems;
-		Element->GetDescriptionItems(DescriptionItems);
-		for(TPair<FString, FText> ItemPair : DescriptionItems)
-		{
-			DescriptionBox->AddSlot()
-				.AutoHeight()
-				.Padding(0, 0, 3, 0)
-				[
-					SNew(SHorizontalBox)
-					+ SHorizontalBox::Slot()
-					.AutoWidth()
-					[
-						SNew(STextBlock)
-						.Text(FText::Format(LOCTEXT("Item", "{0} :"), FText::FromString(ItemPair.Key)))
-						.ColorAndOpacity(FSlateColor::UseSubduedForeground())
-					]
-
-					+ SHorizontalBox::Slot()
-					.AutoWidth()
-					[
-						SNew(STextBlock)
-						.Text(ItemPair.Value)
-						.ColorAndOpacity(FSlateColor::UseForeground())
-					]
-				];
-		}
-
-		// Tooltip
-		TSharedRef<SToolTip> NodeToolTip = SNew(SToolTip)
-			.TextMargin(1)
-			.BorderImage(FEditorStyle::GetBrush("ContentBrowser.TileViewToolTip.ToolTipBorder"))
-			[
-				SNew(SBorder)
-				.Padding(3)
-				.BorderImage(FEditorStyle::GetBrush("ContentBrowser.TileViewTooltip.NonContentBorder"))
-				[
-					SNew(SVerticalBox)
-					+ SVerticalBox::Slot()
-					.AutoHeight()
-					.Padding(0, 0, 0, 3)
-					[
-						SNew(SBorder)
-						.Padding(6)
-						.BorderImage(FEditorStyle::GetBrush("ContentBrowser.TileViewTooltip.ContentBorder"))
-						[
-							SNew(SBox)
-							.HAlign(HAlign_Center)
-							[
-								SNew(STextBlock)
-								.Text(FText::FromName(Element->GetTypeName()))
-								.Font(FEditorStyle::GetFontStyle("ContentBrowser.TileViewTooltip.NameFont"))
-							]
-						]
-					]
-
-					+ SVerticalBox::Slot()
-					[
-						SNew(SHorizontalBox)
-						+ SHorizontalBox::Slot()
-						.FillWidth(1.0f)
-						[
-							SNew(SBorder)
-							.Padding(3)
-							.BorderImage(FEditorStyle::GetBrush("ContentBrowser.TileViewTooltip.ContentBorder"))
-							[
-								DescriptionBox
-							]
-						]
-					]
-				]
-			];
-		SetToolTip(NodeToolTip);
+		SetToolTip(SNew(SAnimTimingNodeTooltip, Element.ToSharedRef()));
 	}
 }
 
@@ -192,8 +213,10 @@ void SAnimTimingTrackNode::Construct(const FArguments& InArgs)
 		);
 }
 
-void SAnimTimingPanel::Construct(const FArguments& InArgs, FSimpleMulticastDelegate& OnSectionsChanged)
+void SAnimTimingPanel::Construct(const FArguments& InArgs, const TSharedRef<FAnimModel_AnimMontage>& InModel)
 {
+	WeakModel = InModel;
+
 	SAnimTrackPanel::Construct(SAnimTrackPanel::FArguments()
 		.WidgetWidth(InArgs._WidgetWidth)
 		.ViewInputMin(InArgs._ViewInputMin)
@@ -207,22 +230,12 @@ void SAnimTimingPanel::Construct(const FArguments& InArgs, FSimpleMulticastDeleg
 	check(AnimSequence);
 
 	this->ChildSlot
-		[
-			SNew(SVerticalBox)
-			+ SVerticalBox::Slot()
-			.FillHeight(1)
-			[
-				SNew(SExpandableArea)
-				.AreaTitle(LOCTEXT("ExpandLabel", "Element Timing"))
-				.BodyContent()
-				[
-					SAssignNew(PanelArea, SBorder)
-					.BorderImage(FEditorStyle::GetBrush("NoBorder"))
-					.Padding(FMargin(2.0f, 2.0f))
-					.ColorAndOpacity(FLinearColor::White)
-				]
-			]
-		];
+	[
+		SAssignNew(PanelArea, SBorder)
+		.BorderImage(FEditorStyle::GetBrush("NoBorder"))
+		.Padding(0.0f)
+		.ColorAndOpacity(FLinearColor::White)
+	];
 
 	Update();
 
@@ -231,10 +244,7 @@ void SAnimTimingPanel::Construct(const FArguments& InArgs, FSimpleMulticastDeleg
 	TSharedPtr<IEditableSkeleton> EditableSkeleton = SkeletonEditorModule.CreateEditableSkeleton(AnimSequence->GetSkeleton());
 	EditableSkeleton->RegisterOnNotifiesChanged(FSimpleDelegate::CreateSP(this, &SAnimTimingPanel::RefreshTrackNodes));
 
-	OnSectionsChanged.Add(FSimpleDelegate::CreateSP(this, &SAnimTimingPanel::RefreshTrackNodes));
-
-	// Clear display flags
-	FMemory::Memset(bElementNodeDisplayFlags, false, sizeof(bElementNodeDisplayFlags));
+	InModel->OnTracksChanged().Add(FSimpleDelegate::CreateSP(this, &SAnimTimingPanel::RefreshTrackNodes));
 }
 
 void SAnimTimingPanel::Update()
@@ -245,12 +255,9 @@ void SAnimTimingPanel::Update()
 	
 	PanelArea->SetContent(SAssignNew(TimingSlots, SVerticalBox));
 
-	TSharedRef<S2ColumnWidget> TrackContainer = Create2ColumnWidget(TimingSlots.ToSharedRef());
-	TrackContainer->LeftColumn->ClearChildren();
-	TrackContainer->LeftColumn->AddSlot()
+	TimingSlots->AddSlot()
 		.AutoHeight()
 		.VAlign(VAlign_Center)
-		.Padding(FMargin(0.5f, 0.5f))
 		[
 			SAssignNew(Track, STimingTrack)
 			.ViewInputMin(ViewInputMin)
@@ -260,27 +267,6 @@ void SAnimTimingPanel::Update()
 			.TrackNumDiscreteValues(AnimSequence->GetNumberOfFrames())
 		];
 
-	TrackContainer->RightColumn->ClearChildren();
-	TrackContainer->RightColumn->AddSlot()
-		.AutoHeight()
-		.VAlign(VAlign_Center)
-		.Padding(FMargin(0.5f, 0.5f))
-		[
-			SNew(SBox)
-			.HAlign(HAlign_Center)
-			[
-				SNew(SButton)
-				.ToolTipText(LOCTEXT("TrackOptionsToolTip", "Display track options menu"))
-				.OnClicked(this, &SAnimTimingPanel::OnContextMenu)
-				.Content()
-				[
-					SNew(SImage)
-					.Image(FEditorStyle::GetBrush("ComboButton.Arrow"))
-					.ColorAndOpacity(FSlateColor::UseForeground())
-				]
-			]
-		];
-
 	RefreshTrackNodes();
 }
 
@@ -300,82 +286,40 @@ void SAnimTimingPanel::RefreshTrackNodes()
 	for(int32 ElementIdx = 0 ; ElementIdx < NumElements ; ++ElementIdx)
 	{
 		TSharedPtr<FTimingRelevantElementBase> Element = Elements[ElementIdx];
-
-		Track->AddTrackNode(
-			SNew(SAnimTimingTrackNode)
-			.ViewInputMin(ViewInputMin)
-			.ViewInputMax(ViewInputMax)
-			.DataStartPos(Element.ToSharedRef(), &FTimingRelevantElementBase::GetElementTime)
-			.NodeName(FString::FromInt(ElementIdx + 1))
-			.NodeColor(FLinearColor::Yellow)
-			.Element(Element)
-			);
-	}
-}
-
-FReply SAnimTimingPanel::OnContextMenu()
-{
-	FMenuBuilder Builder(true, nullptr);
-
-	Builder.BeginSection("TimingPanelOptions", LOCTEXT("TimingPanelOptionsHeader", "Options"));
-	{
-		Builder.AddWidget
-		(
-			SNew(SCheckBox)
-			.IsChecked(this, &SAnimTimingPanel::IsElementDisplayChecked, ETimingElementType::Section)
-			.OnCheckStateChanged(this, &SAnimTimingPanel::OnElementDisplayEnabledChanged, ETimingElementType::Section)
-			.ToolTipText(LOCTEXT("ShowSectionTimingNodes", "Show or hide the timing display for sections on the section name track"))
-			.HAlign(HAlign_Left)
-			[
-				SNew(STextBlock)
-				.Text(LOCTEXT("ToggleTimingNodes_Sections", "Show Section Timing Nodes"))
-			],
-			FText()
-		);
-
-		Builder.AddWidget
-		(
-			SNew(SCheckBox)
-			.IsChecked(this, &SAnimTimingPanel::IsElementDisplayChecked, ETimingElementType::QueuedNotify)
-			.OnCheckStateChanged(this, &SAnimTimingPanel::OnElementDisplayEnabledChanged, ETimingElementType::QueuedNotify)
-			.ToolTipText(LOCTEXT("ShowNotifyTimingNodes", "Show or hide the timing display for notifies in the notify panel"))
-			.HAlign(HAlign_Left)
-			[
-				SNew(STextBlock)
-				.Text(LOCTEXT("ToggleTimingNodes_Notifies", "Show Notify Timing Nodes"))
-			],
-			FText()
-		);
+		if(WeakModel.Pin()->IsTimingElementDisplayEnabled(Element->GetType()))
+		{
+			Track->AddTrackNode(
+				SNew(SAnimTimingTrackNode)
+				.ViewInputMin(ViewInputMin)
+				.ViewInputMax(ViewInputMax)
+				.DataStartPos(Element.ToSharedRef(), &FTimingRelevantElementBase::GetElementTime)
+				.NodeName(FString::FromInt(ElementIdx + 1))
+				.NodeColor(FLinearColor::Yellow)
+				.Element(Element)
+				);
+		}
 	}
-	Builder.EndSection();
-
-	FSlateApplication::Get().PushMenu(SharedThis(this),
-		FWidgetPath(),
-		Builder.MakeWidget(),
-		FSlateApplication::Get().GetCursorPos(),
-		FPopupTransitionEffect(FPopupTransitionEffect::ContextMenu));
-
-	return FReply::Handled();
 }
 
 bool SAnimTimingPanel::IsElementDisplayEnabled(ETimingElementType::Type ElementType) const
 {
-	return bElementNodeDisplayFlags[ElementType];
+	return WeakModel.Pin()->IsTimingElementDisplayEnabled(ElementType);
 }
 
-void SAnimTimingPanel::OnElementDisplayEnabledChanged(ECheckBoxState State, ETimingElementType::Type ElementType)
+void SAnimTimingPanel::OnElementDisplayEnabledChanged(ETimingElementType::Type ElementType)
 {
-	bElementNodeDisplayFlags[ElementType] = State == ECheckBoxState::Checked ? true : false;
+	WeakModel.Pin()->ToggleTimingElementDisplayEnabled(ElementType);
+	Update();
 }
 
 ECheckBoxState SAnimTimingPanel::IsElementDisplayChecked(ETimingElementType::Type ElementType) const
 {
-	return bElementNodeDisplayFlags[ElementType] ? ECheckBoxState::Checked : ECheckBoxState::Unchecked;
+	return WeakModel.Pin()->IsTimingElementDisplayEnabled(ElementType) ? ECheckBoxState::Checked : ECheckBoxState::Unchecked;
 }
 
 EVisibility SAnimTimingPanel::IsElementDisplayVisible(ETimingElementType::Type ElementType) const
 {
-	return bElementNodeDisplayFlags[ElementType] ? EVisibility::Visible : EVisibility::Hidden;
+	return WeakModel.Pin()->IsTimingElementDisplayEnabled(ElementType) ? EVisibility::Visible : EVisibility::Hidden;
 }
 
 void SAnimTimingPanel::GetTimingRelevantElements(UAnimSequenceBase* Sequence, TArray<TSharedPtr<FTimingRelevantElementBase>>& Elements)
@@ -459,11 +403,17 @@ void FTimingRelevantElement_Section::GetDescriptionItems(TMap<FString, FText>& I
 	check(Montage);
 	FCompositeSection& Section = Montage->CompositeSections[SectionIdx];
 
-	FNumberFormattingOptions NumberOptions;
-	NumberOptions.MinimumFractionalDigits = 3;
-
 	Items.Add(LOCTEXT("SectionName", "Name").ToString(), FText::FromName(Section.SectionName));
-	Items.Add(LOCTEXT("SectionTriggerTime", "Trigger Time").ToString(), FText::Format(LOCTEXT("SectionTriggerTimeValue", "{0}s"), FText::AsNumber(Section.GetTime(), &NumberOptions)));
+	if(GetDefault<UPersonaOptions>()->TimelineDisplayFormat == EFrameNumberDisplayFormats::Frames)
+	{
+		Items.Add(LOCTEXT("SectionTriggerFrame", "Trigger Frame").ToString(), FText::Format(LOCTEXT("SectionTriggerTimeValue", "{0}"), FText::AsNumber(Montage->GetFrameAtTime(Section.GetTime()))));
+	}
+	else
+	{
+		FNumberFormattingOptions NumberOptions;
+		NumberOptions.MinimumFractionalDigits = 3;
+		Items.Add(LOCTEXT("SectionTriggerTime", "Trigger Time").ToString(), FText::Format(LOCTEXT("SectionTriggerTimeValue", "{0}s"), FText::AsNumber(Section.GetTime(), &NumberOptions)));
+	}	
 }
 
 FName FTimingRelevantElement_Notify::GetTypeName()
@@ -525,7 +475,15 @@ void FTimingRelevantElement_Notify::GetDescriptionItems(TMap<FString, FText>& It
 	NumberOptions.MinimumFractionalDigits = 3;
 
 	Items.Add(LOCTEXT("NotifyName", "Name").ToString(), FText::FromName(Event.NotifyName));
-	Items.Add(LOCTEXT("NotifyTriggerTime", "Trigger Time").ToString(), FText::Format(LOCTEXT("NotifyTriggerTime_Val", "{0}s"), FText::AsNumber(Event.GetTime(), &NumberOptions)));
+
+	if(GetDefault<UPersonaOptions>()->TimelineDisplayFormat == EFrameNumberDisplayFormats::Frames)
+	{
+		Items.Add(LOCTEXT("NotifyTriggerFrame", "Trigger Frame").ToString(), FText::Format(LOCTEXT("NotifyTriggerFrame_Val", "{0}"), FText::AsNumber(Sequence->GetFrameAtTime(Event.GetTime()))));
+	}
+	else
+	{
+		Items.Add(LOCTEXT("NotifyTriggerTime", "Trigger Time").ToString(), FText::Format(LOCTEXT("NotifyTriggerTime_Val", "{0}s"), FText::AsNumber(Event.GetTime(), &NumberOptions)));
+	}	
 
 	// +1 as we start at 1 when showing tracks to the user
 	Items.Add(LOCTEXT("TrackIdx", "Track").ToString(), FText::AsNumber(Event.TrackIndex + 1));
diff --git a/Engine/Source/Editor/Persona/Private/SAnimTimingPanel.h b/Engine/Source/Editor/Persona/Private/SAnimTimingPanel.h
--- a/Engine/Source/Editor/Persona/Private/SAnimTimingPanel.h
+++ b/Engine/Source/Editor/Persona/Private/SAnimTimingPanel.h
@@ -13,6 +13,8 @@
 class SBorder;
 class UAnimMontage;
 class UAnimSequenceBase;
+class FAnimModel;
+class FAnimModel_AnimMontage;
 
 //////////////////////////////////////////////////////////////////////////
 // FTimingRelevantElement - data object holding timing data
@@ -122,6 +124,7 @@ class SAnimTimingNode : public SCompoundWidget
 public:
 	SLATE_BEGIN_ARGS(SAnimTimingNode)
 		: _InElement()
+		, _bUseTooltip(true)
 	{}
 
 	SLATE_ARGUMENT(TSharedPtr<FTimingRelevantElementBase>, InElement)
@@ -145,6 +148,7 @@ class SAnimTimingTrackNode : public STrackNode
 {
 public:
 	SLATE_BEGIN_ARGS(SAnimTimingTrackNode)
+		: _bUseTooltip(true)
 	{}
 
 	SLATE_ATTRIBUTE(float, ViewInputMin)
@@ -193,7 +197,7 @@ public:
 
 	// Construct the panel
 	// @param InArgs - Slate arguments
-	void Construct(const FArguments& InArgs, FSimpleMulticastDelegate& OnSectionsChanged);
+	void Construct(const FArguments& InArgs, const TSharedRef<FAnimModel_AnimMontage>& InModel);
 
 	// Updates panel widgets
 	void Update();
@@ -205,9 +209,8 @@ public:
 	EVisibility IsElementDisplayVisible(ETimingElementType::Type ElementType) const;
 
 	// Callback from slate when element display flag changes
-	// @param State - New state
 	// @param ElementType - Flag to set
-	void OnElementDisplayEnabledChanged(ECheckBoxState State, ETimingElementType::Type ElementType);
+	void OnElementDisplayEnabledChanged(ETimingElementType::Type ElementType);
 
 	// Inspects the provided sequence, collect and sorts the requested elements
 	// @param Sequence - Sequence to inspect
@@ -216,15 +219,14 @@ public:
 
 protected:
 
-	// Context summon callback
-	FReply OnContextMenu();
-
 	// Tick the panel state
 	virtual void Tick(const FGeometry& AllottedGeometry, const double InCurrentTime, const float InDeltaTime) override;
 
 	// Clears the timing track and rebuilds the nodes
 	void RefreshTrackNodes();
 
+	// The model for our data
+	TWeakPtr<FAnimModel_AnimMontage> WeakModel;
 	// Observed timing elements
 	TArray<TSharedPtr<FTimingRelevantElementBase>> Elements;
 	// Anim sequence that contains the timing elements we are observing
@@ -233,7 +235,4 @@ protected:
 	TSharedPtr<SBorder> PanelArea;
 	// The track to place timing nodes on
 	TSharedPtr<STrack> Track;
-
-	// Display flags for other panels
-	bool bElementNodeDisplayFlags[ETimingElementType::Max];
 };
diff --git a/Engine/Source/Editor/Persona/Private/SAnimTrackCurvePanel.cpp b/Engine/Source/Editor/Persona/Private/SAnimTrackCurvePanel.cpp
--- a/Engine/Source/Editor/Persona/Private/SAnimTrackCurvePanel.cpp
+++ b/Engine/Source/Editor/Persona/Private/SAnimTrackCurvePanel.cpp
@@ -544,7 +544,7 @@ void SAnimTrackCurvePanel::Construct(const FArguments& InArgs, const TSharedRef<
 
 	if (InPreviewScene->GetPreviewMeshComponent()->PreviewInstance)
 	{
-		InPreviewScene->GetPreviewMeshComponent()->PreviewInstance->SetKeyCompleteDelegate(FSimpleDelegate::CreateSP(this, &SAnimTrackCurvePanel::HandleKeyComplete));
+		InPreviewScene->GetPreviewMeshComponent()->PreviewInstance->AddKeyCompleteDelegate(FSimpleDelegate::CreateSP(this, &SAnimTrackCurvePanel::HandleKeyComplete));
 	}
 
 	Sequence->RegisterOnAnimTrackCurvesChanged(UAnimSequenceBase::FOnAnimTrackCurvesChanged::CreateSP(this, &SAnimTrackCurvePanel::UpdatePanel));
diff --git a/Engine/Source/Editor/Persona/Private/SAnimTrackPanel.cpp b/Engine/Source/Editor/Persona/Private/SAnimTrackPanel.cpp
--- a/Engine/Source/Editor/Persona/Private/SAnimTrackPanel.cpp
+++ b/Engine/Source/Editor/Persona/Private/SAnimTrackPanel.cpp
@@ -74,26 +74,6 @@ TSharedRef<class S2ColumnWidget> SAnimTrackPanel::Create2ColumnWidget( TSharedRe
 	return NewTrack.ToSharedRef();
 }
 
-FReply SAnimTrackPanel::OnMouseWheel( const FGeometry& MyGeometry, const FPointerEvent& MouseEvent )
-{
-	const float ZoomDelta = -0.1f * MouseEvent.GetWheelDelta();
-
-	const FVector2D MouseWidgetPos = MyGeometry.AbsoluteToLocal(MouseEvent.GetScreenSpacePosition());
-	const float ZoomRatio = FMath::Clamp((MouseWidgetPos.X / (MyGeometry.Size.X - WidgetWidth)), 0.f, 1.f);
-
-	{
-		const float InputViewSize = ViewInputMax.Get() - ViewInputMin.Get();
-		const float InputChange = InputViewSize * ZoomDelta;
-
-		float ViewMinInput = ViewInputMin.Get() - (InputChange * ZoomRatio);
-		float ViewMaxInput = ViewInputMax.Get() + (InputChange * (1.f - ZoomRatio));
-		
-		InputViewRangeChanged(ViewMinInput, ViewMaxInput);
-	}
-
-	return FReply::Handled();
-}
-
 FReply SAnimTrackPanel::OnMouseButtonDown( const FGeometry& InMyGeometry, const FPointerEvent& InMouseEvent )
 {
 	const bool bRightMouseButton = InMouseEvent.GetEffectingButton() == EKeys::RightMouseButton;
@@ -198,7 +178,7 @@ void SAnimTrackPanel::PanInputViewRange(int32 ScreenDelta, FVector2D ScreenViewS
 
 void SAnimTrackPanel::InputViewRangeChanged(float ViewMin, float ViewMax)
 {
-	OnSetInputViewRange.Execute(ViewMin, ViewMax);
+	OnSetInputViewRange.ExecuteIfBound(ViewMin, ViewMax);
 }
 
 #undef LOCTEXT_NAMESPACE
diff --git a/Engine/Source/Editor/Persona/Private/SAnimTrackPanel.h b/Engine/Source/Editor/Persona/Private/SAnimTrackPanel.h
--- a/Engine/Source/Editor/Persona/Private/SAnimTrackPanel.h
+++ b/Engine/Source/Editor/Persona/Private/SAnimTrackPanel.h
@@ -90,7 +90,6 @@ public:
 	void Construct(const FArguments& InArgs);
 
 	// SWidget interface
-	virtual FReply OnMouseWheel( const FGeometry& MyGeometry, const FPointerEvent& MouseEvent ) override;
 	virtual FReply OnMouseButtonDown( const FGeometry& MyGeometry, const FPointerEvent& MouseEvent ) override;
 	virtual FReply OnMouseButtonUp( const FGeometry& MyGeometry, const FPointerEvent& MouseEvent ) override;
 	virtual FReply OnMouseMove( const FGeometry& MyGeometry, const FPointerEvent& MouseEvent ) override;
diff --git a/Engine/Source/Editor/Persona/Private/SAnimationBlendSpaceBase.cpp b/Engine/Source/Editor/Persona/Private/SAnimationBlendSpaceBase.cpp
--- a/Engine/Source/Editor/Persona/Private/SAnimationBlendSpaceBase.cpp
+++ b/Engine/Source/Editor/Persona/Private/SAnimationBlendSpaceBase.cpp
@@ -35,7 +35,10 @@ void SBlendSpaceEditorBase::Construct(const FArguments& InArgs, const TSharedRef
 
 	bShouldSetPreviewValue = false;
 
-	SAnimEditorBase::Construct(SAnimEditorBase::FArguments().DisplayAnimInfoBar(false), InPreviewScene);
+	SAnimEditorBase::Construct(SAnimEditorBase::FArguments()
+		.DisplayAnimTimeline(false)
+		.DisplayAnimScrubBar(true),
+		InPreviewScene);
 
 	NonScrollEditorPanels->AddSlot()
 	[
diff --git a/Engine/Source/Editor/Persona/Private/SAnimationSequenceBrowser.cpp b/Engine/Source/Editor/Persona/Private/SAnimationSequenceBrowser.cpp
--- a/Engine/Source/Editor/Persona/Private/SAnimationSequenceBrowser.cpp
+++ b/Engine/Source/Editor/Persona/Private/SAnimationSequenceBrowser.cpp
@@ -683,7 +683,7 @@ void SAnimationSequenceBrowser::OnApplyCompression(TArray<FAssetData> SelectedAs
 		}
 
 		FPersonaModule& PersonaModule = FModuleManager::GetModuleChecked<FPersonaModule>("Persona");
-		PersonaModule.ApplyCompression(AnimSequences);
+		PersonaModule.ApplyCompression(AnimSequences, true);
 	}
 }
 
diff --git a/Engine/Source/Editor/Persona/Private/SMontageEditor.cpp b/Engine/Source/Editor/Persona/Private/SMontageEditor.cpp
--- a/Engine/Source/Editor/Persona/Private/SMontageEditor.cpp
+++ b/Engine/Source/Editor/Persona/Private/SMontageEditor.cpp
@@ -23,6 +23,8 @@
 #include "Factories/AnimMontageFactory.h"
 #include "Editor.h"
 #include "Subsystems/AssetEditorSubsystem.h"
+#include "AnimModel_AnimMontage.h"
+#include "SAnimTimeline.h"
 
 #define LOCTEXT_NAMESPACE "AnimSequenceEditor"
 
@@ -38,218 +40,26 @@ void SMontageEditor::Construct(const FArguments& InArgs, const FMontageEditorReq
 {
 	MontageObj = InArgs._Montage;
 	check(MontageObj);
-	OnSectionsChanged = InArgs._OnSectionsChanged;
-	MontageObj->RegisterOnMontageChanged(UAnimMontage::FOnMontageChanged::CreateSP(this, &SMontageEditor::RebuildMontagePanel, false));
-		
-	EnsureStartingSection();
-	EnsureSlotNode();
 
-	// set child montage if montage has parent
-	bChildAnimMontage = MontageObj->HasParentAsset();
+	AnimModel = MakeShared<FAnimModel_AnimMontage>(InRequiredArgs.PreviewScene, InRequiredArgs.EditableSkeleton, InRequiredArgs.CommandList, MontageObj);
+	AnimModel->OnSelectObjects = FOnObjectsSelected::CreateSP(this, &SAnimEditorBase::OnSelectionChanged);
+	AnimModel->OnSectionsChanged = InArgs._OnSectionsChanged;
+	AnimModel->OnInvokeTab = InArgs._OnInvokeTab;
 
-	bDragging = false;
-	bIsActiveTimerRegistered = false;
-
-	SAnimEditorBase::Construct( SAnimEditorBase::FArguments()
-		.OnObjectsSelected(InArgs._OnObjectsSelected), 
-		InRequiredArgs.PreviewScene );
-
-	if(GEditor)
-	{
-		GEditor->RegisterForUndo(this);
-	}
-
-	SAssignNew(AnimTimingPanel, SAnimTimingPanel, InRequiredArgs.OnSectionsChanged)
-		.IsEnabled(!bChildAnimMontage)
-		.InSequence(MontageObj)
-		.WidgetWidth(S2ColumnWidget::DEFAULT_RIGHT_COLUMN_WIDTH)
-		.ViewInputMin(this, &SAnimEditorBase::GetViewMinInput)
-		.ViewInputMax(this, &SAnimEditorBase::GetViewMaxInput)
-		.InputMin(this, &SAnimEditorBase::GetMinInput)
-		.InputMax(this, &SAnimEditorBase::GetMaxInput)
-		.OnSetInputViewRange(this, &SAnimEditorBase::SetInputViewRange);
-
-	TAttribute<EVisibility> SectionVisibility = TAttribute<EVisibility>::Create(TAttribute<EVisibility>::FGetter::CreateSP(AnimTimingPanel.ToSharedRef(), &SAnimTimingPanel::IsElementDisplayVisible, ETimingElementType::Section));
-	TAttribute<EVisibility> NotifyVisibility = TAttribute<EVisibility>::Create(TAttribute<EVisibility>::FGetter::CreateSP(AnimTimingPanel.ToSharedRef(), &SAnimTimingPanel::IsElementDisplayVisible, ETimingElementType::QueuedNotify));
-	FOnGetTimingNodeVisibility TimingNodeVisibilityDelegate = FOnGetTimingNodeVisibility::CreateSP(AnimTimingPanel.ToSharedRef(), &SAnimTimingPanel::IsElementDisplayVisible);
-	
-	if (bChildAnimMontage)
-	{
-		FFormatNamedArguments Args;
-		Args.Add(TEXT("ParentClassName"), FText::FromString(GetNameSafe(MontageObj->ParentAsset)) );
-
-		// add child montage warning section - and link to parent
-		EditorPanels->AddSlot()
-		.AutoHeight()
-		.Padding(0, 0)
-		[
-			SNew(SBorder)
-			.Padding(FMargin(10.f, 4.f))
-			.BorderImage(FEditorStyle::GetBrush(TEXT("Graph.InstructionBackground")))
-			.BorderBackgroundColor(FLinearColor(0.1f, 0.1f, 0.1f, 0.5f))
-			.HAlign(HAlign_Center)
-			.ColorAndOpacity(FLinearColor(1.f, 1.f, 1.f, 0.8f))
-			[
-				SNew(SVerticalBox)
-			
-				+SVerticalBox::Slot()
-				.AutoHeight()
-				[
-					SNew(SHorizontalBox)
-
-					+SHorizontalBox::Slot()
-					.AutoWidth()
-					[
-						SNew(STextBlock)
-						.Text(LOCTEXT("ParentAnimMontageLink", " This is a child anim montage. To edit the lay out, please go to the parent montage "))
-						.TextStyle(FEditorStyle::Get(), "Persona.MontageEditor.ChildMontageInstruction")
-					]
-					+ SHorizontalBox::Slot()
-					.AutoWidth()
-					[
-						SNew(SButton)
-						.VAlign(VAlign_Center)
-						.ButtonStyle(FEditorStyle::Get(), "HoverHintOnly")
-						.OnClicked(this, &SMontageEditor::OnFindParentClassInContentBrowserClicked)
-						.ToolTipText(LOCTEXT("FindParentInCBToolTip", "Find parent in Content Browser"))
-						.ForegroundColor(FSlateColor::UseForeground())
-						[
-							SNew(SImage)
-							.Image(FEditorStyle::GetBrush("PropertyWindow.Button_Browse"))
-						]
-					]
-					+ SHorizontalBox::Slot()
-					.AutoWidth()
-					[
-						SNew(SButton)
-						.VAlign(VAlign_Center)
-						.ButtonStyle(FEditorStyle::Get(), "HoverHintOnly")
-						.OnClicked(this, &SMontageEditor::OnEditParentClassClicked)
-						.ToolTipText(LOCTEXT("EditParentClassToolTip", "Open parent in editor"))
-						.ForegroundColor(FSlateColor::UseForeground())
-						[
-							SNew(SImage)
-							.Image(FEditorStyle::GetBrush("PropertyWindow.Button_Edit"))
-						]
-					]
-				]
-
-				+ SVerticalBox::Slot()
-				.AutoHeight()
-				[
-					SNew(STextBlock)
-					.Text(LOCTEXT("RemapHelpText", " To remap an asset to a different asset, use the context menu or drag and drop the animation on to the segment."))
-					.TextStyle(FEditorStyle::Get(), "Persona.MontageEditor.ChildMontageInstruction")
-				]
-			]
-		];
-	}
-
-	EditorPanels->AddSlot()
-	.AutoHeight()
-	.Padding(0, 5)
-	[
-		SAssignNew( AnimMontagePanel, SAnimMontagePanel, InRequiredArgs.OnSectionsChanged)
-		.Montage(MontageObj)
-		.bChildAnimMontage(bChildAnimMontage)
-		.MontageEditor(SharedThis(this))
-		.WidgetWidth(S2ColumnWidget::DEFAULT_RIGHT_COLUMN_WIDTH)
-		.ViewInputMin(this, &SAnimEditorBase::GetViewMinInput)
-		.ViewInputMax(this, &SAnimEditorBase::GetViewMaxInput)
-		.InputMin(this, &SAnimEditorBase::GetMinInput)
-		.InputMax(this, &SAnimEditorBase::GetMaxInput)
-		.OnSetInputViewRange(this, &SAnimEditorBase::SetInputViewRange)
-		.SectionTimingNodeVisibility(SectionVisibility)
-		.OnInvokeTab(InArgs._OnInvokeTab)
-		.OnSetMontagePreviewSlot(this, &SMontageEditor::OnSetMontagePreviewSlot)
-	];
-
-	EditorPanels->AddSlot()
-	.AutoHeight()
-	.Padding(0, 10)
-	[
-		SAssignNew( AnimMontageSectionsPanel, SAnimMontageSectionsPanel )
-		.bChildAnimMontage(bChildAnimMontage)
-		.Montage(MontageObj)
-		.MontageEditor(SharedThis(this))
-	];
-
-	EditorPanels->AddSlot()
-	.AutoHeight()
-	.Padding(0, 10)
-	[
-		AnimTimingPanel.ToSharedRef()
-	];
-
-	EditorPanels->AddSlot()
-	.AutoHeight()
-	.Padding(0, 10)
-	[
-		SAssignNew( AnimNotifyPanel, SAnimNotifyPanel, InRequiredArgs.EditableSkeleton )
-		.Sequence(MontageObj)
-		.IsEnabled(!bChildAnimMontage)
-		.WidgetWidth(S2ColumnWidget::DEFAULT_RIGHT_COLUMN_WIDTH)
-		.InputMin(this, &SAnimEditorBase::GetMinInput)
-		.InputMax(this, &SAnimEditorBase::GetMaxInput)
-		.ViewInputMin(this, &SAnimEditorBase::GetViewMinInput)
-		.ViewInputMax(this, &SAnimEditorBase::GetViewMaxInput)
-		.OnSetInputViewRange(this, &SAnimEditorBase::SetInputViewRange)
-		.OnGetScrubValue(this, &SAnimEditorBase::GetScrubValue)
-		.OnSelectionChanged(this, &SAnimEditorBase::OnSelectionChanged)
-		.MarkerBars(this, &SMontageEditor::GetMarkerBarInformation)
-		.OnRequestRefreshOffsets(this, &SMontageEditor::RefreshNotifyTriggerOffsets)
-		.OnGetTimingNodeVisibility(TimingNodeVisibilityDelegate)
-		.OnInvokeTab(InArgs._OnInvokeTab)
-		.OnNotifiesChanged(this, &SMontageEditor::HandleNotifiesChanged)
-	];
-
-	EditorPanels->AddSlot()
-	.AutoHeight()
-	.Padding(0, 10)
-	[
-		SAssignNew( AnimCurvePanel, SAnimCurvePanel, InRequiredArgs.EditableSkeleton )
-		.IsEnabled(!bChildAnimMontage)
-		.Sequence(MontageObj)
-		.WidgetWidth(S2ColumnWidget::DEFAULT_RIGHT_COLUMN_WIDTH)
-		.ViewInputMin(this, &SAnimEditorBase::GetViewMinInput)
-		.ViewInputMax(this, &SAnimEditorBase::GetViewMaxInput)
-		.InputMin(this, &SAnimEditorBase::GetMinInput)
-		.InputMax(this, &SAnimEditorBase::GetMaxInput)
-		.OnSetInputViewRange(this, &SAnimEditorBase::SetInputViewRange)
-		.OnGetScrubValue(this, &SAnimEditorBase::GetScrubValue)
-	];
-
-	CollapseMontage();
-}
-
-FReply SMontageEditor::OnFindParentClassInContentBrowserClicked()
-{
-	if (MontageObj != NULL)
+	AnimModel->OnEditCurves = FOnEditCurves::CreateLambda([this, InOnEditCurves = InArgs._OnEditCurves](UAnimSequenceBase* InAnimSequence, const TArray<IAnimationEditor::FCurveEditInfo>& InCurveInfo, const TSharedPtr<ITimeSliderController>& InExternalTimeSliderController)
 	{
-		UObject* ParentClass = MontageObj->ParentAsset;
-		if (ParentClass != NULL)
-		{
-			TArray< UObject* > ParentObjectList;
-			ParentObjectList.Add(ParentClass);
-			GEditor->SyncBrowserToObjects(ParentObjectList);
-		}
-	}
+		InOnEditCurves.ExecuteIfBound(InAnimSequence, InCurveInfo, TimelineWidget->GetTimeSliderController());
+	});
 
-	return FReply::Handled();
-}
+	AnimModel->OnStopEditingCurves = InArgs._OnStopEditingCurves;
+	AnimModel->Initialize();
 
-FReply SMontageEditor::OnEditParentClassClicked()
-{
-	if (MontageObj != NULL)
-	{
-		UObject* ParentClass = MontageObj->ParentAsset;
-		if (ParentClass != NULL)
-		{
-			GEditor->GetEditorSubsystem<UAssetEditorSubsystem>()->OpenEditorForAsset(ParentClass);
-		}
-	}
+	MontageObj->RegisterOnMontageChanged(UAnimMontage::FOnMontageChanged::CreateSP(AnimModel.Get(), &FAnimModel_AnimMontage::RefreshTracks));
 
-	return FReply::Handled();
+	SAnimEditorBase::Construct( SAnimEditorBase::FArguments()
+		.OnObjectsSelected(InArgs._OnObjectsSelected)
+		.AnimModel(AnimModel), 
+		InRequiredArgs.PreviewScene );
 }
 
 SMontageEditor::~SMontageEditor()
@@ -258,669 +68,7 @@ SMontageEditor::~SMontageEditor()
 	{
 		MontageObj->UnregisterOnMontageChanged(this);
 	}
-
-	if(GEditor)
-	{
-		GEditor->UnregisterForUndo(this);
-	}
-}
-
-TSharedRef<class SAnimationScrubPanel> SMontageEditor::ConstructAnimScrubPanel()
-{
-	return SAssignNew(AnimMontageScrubPanel, SAnimMontageScrubPanel, GetPreviewScene())
-		.LockedSequence(MontageObj)
-		.ViewInputMin(this, &SMontageEditor::GetViewMinInput)
-		.ViewInputMax(this, &SMontageEditor::GetViewMaxInput)
-		.OnSetInputViewRange(this, &SMontageEditor::SetInputViewRange)
-		.bAllowZoom(true)
-		.MontageEditor(SharedThis(this));
-}
-
-void SMontageEditor::SetMontageObj(UAnimMontage * NewMontage)
-{
-	MontageObj = NewMontage;
-
-	if (MontageObj)
-	{
-		SetInputViewRange(0, MontageObj->SequenceLength); // FIXME
-	}
-
-	AnimMontagePanel->SetMontage(NewMontage);
-	AnimNotifyPanel->SetSequence(NewMontage);
-	AnimCurvePanel->SetSequence(NewMontage);
-	// sequence editor locks the sequence, so it doesn't get replaced by clicking 
-	AnimMontageScrubPanel->ReplaceLockedSequence(NewMontage);
-}
-
-void SMontageEditor::OnSetMontagePreviewSlot(int32 SlotIndex)
-{
-	UAnimSingleNodeInstance * PreviewInstance = GetPreviewInstance();
-	if (PreviewInstance && MontageObj->SlotAnimTracks.IsValidIndex(SlotIndex))
-	{
-		const FName SlotName = MontageObj->SlotAnimTracks[SlotIndex].SlotName;
-		PreviewInstance->SetMontagePreviewSlot(SlotName);
-	}
-}
-
-bool SMontageEditor::ValidIndexes(int32 AnimSlotIndex, int32 AnimSegmentIndex) const
-{
-	return (MontageObj != NULL && MontageObj->SlotAnimTracks.IsValidIndex(AnimSlotIndex) && MontageObj->SlotAnimTracks[AnimSlotIndex].AnimTrack.AnimSegments.IsValidIndex(AnimSegmentIndex) );
-}
-
-bool SMontageEditor::ValidSection(int32 SectionIndex) const
-{
-	return (MontageObj != NULL && MontageObj->CompositeSections.IsValidIndex(SectionIndex));
-}
-
-void SMontageEditor::RefreshNotifyTriggerOffsets()
-{
-	for(auto Iter = MontageObj->Notifies.CreateIterator(); Iter; ++Iter)
-	{
-		FAnimNotifyEvent& Notify = (*Iter);
-
-		// Offset for the beginning of a notify
-		EAnimEventTriggerOffsets::Type PredictedOffset = MontageObj->CalculateOffsetForNotify(Notify.GetTime());
-		Notify.RefreshTriggerOffset(PredictedOffset);
-
-		// Offset for the end of a notify state if necessary
-		if(Notify.GetDuration() > 0.0f)
-		{
-			PredictedOffset = MontageObj->CalculateOffsetForNotify(Notify.GetTime() + Notify.GetDuration());
-			Notify.RefreshEndTriggerOffset(PredictedOffset);
-		}
-		else
-		{
-			Notify.EndTriggerTimeOffset = 0.0f;
-		}
-	}
-}
-
-void SMontageEditor::HandleNotifiesChanged()
-{
-	AnimTimingPanel->Update();
-}
-
-bool SMontageEditor::GetSectionTime( int32 SectionIndex, float &OutTime ) const
-{
-	if (MontageObj != NULL && MontageObj->CompositeSections.IsValidIndex(SectionIndex))
-	{
-		OutTime = MontageObj->CompositeSections[SectionIndex].GetTime();
-		return true;
-	}
-
-	return false;
-}
-
-TArray<FString>	SMontageEditor::GetSectionNames() const
-{
-	TArray<FString> Names;
-	if (MontageObj != NULL)
-	{
-		for( int32 I=0; I < MontageObj->CompositeSections.Num(); I++)
-		{
-			Names.Add(MontageObj->CompositeSections[I].SectionName.ToString());
-		}
-	}
-	return Names;
-}
-
-TArray<float> SMontageEditor::GetSectionStartTimes() const
-{
-	TArray<float>	Times;
-	if (MontageObj != NULL)
-	{
-		for( int32 I=0; I < MontageObj->CompositeSections.Num(); I++)
-		{
-			Times.Add(MontageObj->CompositeSections[I].GetTime());
-		}
-	}
-	return Times;
-}
-
-TArray<FTrackMarkerBar> SMontageEditor::GetMarkerBarInformation() const
-{
-	TArray<FTrackMarkerBar> MarkerBars;
-	if (MontageObj != NULL)
-	{
-		for( int32 I=0; I < MontageObj->CompositeSections.Num(); I++)
-		{
-			FTrackMarkerBar Bar;
-			Bar.Time = MontageObj->CompositeSections[I].GetTime();
-			Bar.DrawColour = FLinearColor(0.f,1.f,0.f);
-			MarkerBars.Add(Bar);
-		}
-	}
-	return MarkerBars;
-}
-
-TArray<float> SMontageEditor::GetAnimSegmentStartTimes() const
-{
-	TArray<float>	Times;
-	if (MontageObj != NULL)
-	{
-		for ( int32 i=0; i < MontageObj->SlotAnimTracks.Num(); i++)
-		{
-			for (int32 j=0; j < MontageObj->SlotAnimTracks[i].AnimTrack.AnimSegments.Num(); j++)
-			{
-				Times.Add( MontageObj->SlotAnimTracks[i].AnimTrack.AnimSegments[j].StartPos );
-			}
-		}
-	}
-	return Times;
-}
-
-void SMontageEditor::OnEditSectionTime( int32 SectionIndex, float NewTime)
-{
-	if (MontageObj != NULL && MontageObj->CompositeSections.IsValidIndex(SectionIndex))
-	{
-		if(!bDragging)
-		{
-			//If this is the first drag event
-			const FScopedTransaction Transaction( LOCTEXT("EditSection", "Edit Section Start Time") );
-			MontageObj->Modify();
-		}
-		bDragging = true;
-
-		MontageObj->CompositeSections[SectionIndex].SetTime(NewTime);
-		MontageObj->CompositeSections[SectionIndex].LinkMontage(MontageObj, NewTime);
-	}
-
-	AnimMontagePanel->RefreshTimingNodes();
-}
-void SMontageEditor::OnEditSectionTimeFinish( int32 SectionIndex )
-{
-	bDragging = false;
-
-	if(MontageObj!=NULL)
-	{
-		SortSections();
-		RefreshNotifyTriggerOffsets();
-		OnMontageModified();
-		AnimMontageSectionsPanel->Update();
-	}
-
-	OnSectionsChanged.ExecuteIfBound();
-}
-
-void SMontageEditor::SetSectionTime(int32 SectionIndex, float NewTime)
-{
-	if(MontageObj && MontageObj->CompositeSections.IsValidIndex(SectionIndex))
-	{
-		const FScopedTransaction Transaction(LOCTEXT("EditSection", "Edit Section Start Time"));
-		MontageObj->Modify();
-	
-		FCompositeSection& Section = MontageObj->CompositeSections[SectionIndex];
-		Section.SetTime(NewTime);
-		Section.LinkMontage(MontageObj, NewTime);
-
-		OnEditSectionTimeFinish(SectionIndex);
-	}
-}
-
-void SMontageEditor::PreAnimUpdate()
-{
-	MontageObj->Modify();
-}
-
-void SMontageEditor::OnMontageModified()
-{
-	MontageObj->PostEditChange();
-	MontageObj->MarkPackageDirty();
-}
-
-void SMontageEditor::PostAnimUpdate()
-{
-	SortAndUpdateMontage();
-	OnMontageModified();
-}
-
-bool SMontageEditor::IsDiffererentFromParent(FName SlotName, int32 SegmentIdx, const FAnimSegment& Segment)
-{
-	// if it doesn't hare parent asset, no reason to come here
-	if (MontageObj && ensureAlways(MontageObj->ParentAsset))
-	{
-		// find correct source asset from parent
-		UAnimMontage* ParentMontage = Cast<UAnimMontage>(MontageObj->ParentAsset);
-		if (ParentMontage->IsValidSlot(SlotName))
-		{
-			const FAnimTrack* ParentTrack = ParentMontage->GetAnimationData(SlotName);
-
-			if (ParentTrack && ParentTrack->AnimSegments.IsValidIndex(SegmentIdx))
-			{
-				UAnimSequenceBase* SourceAsset = ParentTrack->AnimSegments[SegmentIdx].AnimReference;
-				return (SourceAsset != Segment.AnimReference);
-			}
-		}
-	}
-
-	// if something doesn't match, we assume they're different, so default feedback  is to return true
-	return true;
-}
-
-void SMontageEditor::ReplaceAnimationMapping(FName SlotName, int32 SegmentIdx, UAnimSequenceBase* OldSequenceBase, UAnimSequenceBase* NewSequenceBase)
-{
-	// if it doesn't hare parent asset, no reason to come here
-	if (MontageObj && ensureAlways(MontageObj->ParentAsset))
-	{
-		// find correct source asset from parent
-		UAnimMontage* ParentMontage = Cast<UAnimMontage>(MontageObj->ParentAsset);
-		if (ParentMontage->IsValidSlot(SlotName))
-		{
-			const FAnimTrack* ParentTrack = ParentMontage->GetAnimationData(SlotName);
-
-			if (ParentTrack && ParentTrack->AnimSegments.IsValidIndex(SegmentIdx))
-			{
-				UAnimSequenceBase* SourceAsset = ParentTrack->AnimSegments[SegmentIdx].AnimReference;
-				if (MontageObj->RemapAsset(SourceAsset, NewSequenceBase))
-				{
-					// success
-					return;
-				}
-			}
-		}
-	}
-
-	// failed to do the process, check if the animation is correct or if the same type of animation
-	// print error
-	FMessageDialog::Open(EAppMsgType::Ok, LOCTEXT("FailedToRemap", "Make sure the target animation is valid. If source is additive, target animation has to be additive also."));
-}
-
-void SMontageEditor::RebuildMontagePanel(bool bNotifyAsset /*= true*/)
-{
-	SortAndUpdateMontage();
-	AnimMontageSectionsPanel->Update();
-
-	if (bNotifyAsset)
-	{
-		OnMontageModified();
-	}
-}
-
-EActiveTimerReturnType SMontageEditor::TriggerRebuildMontagePanel(double InCurrentTime, float InDeltaTime)
-{
-	RebuildMontagePanel();
-
-	bIsActiveTimerRegistered = false;
-	return EActiveTimerReturnType::Stop;
-}
-
-void SMontageEditor::OnMontageChange(class UObject *EditorAnimBaseObj, bool Rebuild)
-{
-	bDragging = false;
-
-	if ( MontageObj != nullptr )
-	{
-		float PreviouewSeqLength = GetSequenceLength();
-
-		if(Rebuild && !bIsActiveTimerRegistered)
-		{
-			bIsActiveTimerRegistered = true;
-			RegisterActiveTimer(0.f, FWidgetActiveTimerDelegate::CreateSP(this, &SMontageEditor::TriggerRebuildMontagePanel));
-		} 
-		else
-		{
-			CollapseMontage();
-		}
-
-		// if animation length changed, we might be out of range, let's restart
-		if (GetSequenceLength() != PreviouewSeqLength)
-		{
-			// this might not be safe
-			RestartPreview();
-		}
-
-		OnMontageModified();
-	}
-}
-
-void SMontageEditor::InitDetailsViewEditorObject(UEditorAnimBaseObj* EdObj)
-{
-	EdObj->InitFromAnim(MontageObj, FOnAnimObjectChange::CreateSP( SharedThis(this), &SMontageEditor::OnMontageChange ));
-}
-
-/** This will remove empty spaces in the montage's anim segment but not resort. e.g. - all cached indexes remains valid. UI IS NOT REBUILT after this */
-void SMontageEditor::CollapseMontage()
-{
-	if (MontageObj==NULL)
-	{
-		return;
-	}
-
-	for (int32 i=0; i < MontageObj->SlotAnimTracks.Num(); i++)
-	{
-		MontageObj->SlotAnimTracks[i].AnimTrack.CollapseAnimSegments();
-	}
-
-	MontageObj->UpdateLinkableElements();
-
-	RecalculateSequenceLength();
-}
-
-/** This will sort all components of the montage and update (recreate) the UI */
-void SMontageEditor::SortAndUpdateMontage()
-{
-	if (MontageObj==NULL)
-	{
-		return;
-	}
-	
-	SortAnimSegments();
-
-	MontageObj->UpdateLinkableElements();
-
-	RecalculateSequenceLength();
-
-	SortSections();
-
-	RefreshNotifyTriggerOffsets();
-
-	// Update view (this will recreate everything)
-	AnimMontagePanel->Update();
-	AnimMontageSectionsPanel->Update();
-	AnimTimingPanel->Update();
-	AnimNotifyPanel->Update();
-
-	// Restart the preview instance of the montage
-	RestartPreview();
-}
-
-float SMontageEditor::CalculateSequenceLengthOfEditorObject() const
-{
-	return MontageObj->CalculateSequenceLength();
-}
-
-/** Sort Segments by starting time */
-void SMontageEditor::SortAnimSegments()
-{
-	for (int32 I=0; I < MontageObj->SlotAnimTracks.Num(); I++)
-	{
-		MontageObj->SlotAnimTracks[I].AnimTrack.SortAnimSegments();
-	}
-}
-
-/** Sort Composite Sections by Start TIme */
-void SMontageEditor::SortSections()
-{
-	struct FCompareSections
-	{
-		bool operator()( const FCompositeSection &A, const FCompositeSection &B ) const
-		{
-			return A.GetTime() < B.GetTime();
-		}
-	};
-	if (MontageObj != NULL)
-	{
-		MontageObj->CompositeSections.Sort(FCompareSections());
-	}
-
-	EnsureStartingSection();
-}
-
-/** Ensure there is at least one section in the montage and that the first section starts at T=0.f */
-void SMontageEditor::EnsureStartingSection()
-{
-	if (UAnimMontageFactory::EnsureStartingSection(MontageObj))
-	{
-		OnMontageModified();
-	}
-}
-
-/** Esnure there is at least one slot node track */
-void SMontageEditor::EnsureSlotNode()
-{
-	if (MontageObj && MontageObj->SlotAnimTracks.Num()==0)
-	{
-		AddNewMontageSlot(FAnimSlotGroup::DefaultSlotName);
-		OnMontageModified();
-	}
-}
-
-/** Make sure all Sections and Notifies are clamped to NewEndTime (called before NewEndTime is set to SequenceLength) */
-bool SMontageEditor::ClampToEndTime(float NewEndTime)
-{
-	bool bClampingNeeded = SAnimEditorBase::ClampToEndTime(NewEndTime);
-	if(bClampingNeeded)
-	{
-		float ratio = NewEndTime / MontageObj->SequenceLength;
-
-		for(int32 i=0; i < MontageObj->CompositeSections.Num(); i++)
-		{
-			if(MontageObj->CompositeSections[i].GetTime() > NewEndTime)
-			{
-				float CurrentTime = MontageObj->CompositeSections[i].GetTime();
-				MontageObj->CompositeSections[i].SetTime(CurrentTime * ratio);
-			}
-		}
-
-		for(int32 i=0; i < MontageObj->Notifies.Num(); i++)
-		{
-			FAnimNotifyEvent& Notify = MontageObj->Notifies[i];
-			float NotifyTime = Notify.GetTime();
-
-			if(NotifyTime >= NewEndTime)
-			{
-				Notify.SetTime(NotifyTime * ratio);
-				Notify.TriggerTimeOffset = GetTriggerTimeOffsetForType(MontageObj->CalculateOffsetForNotify(Notify.GetTime()));
-			}
-		}
-	}
-
-	return bClampingNeeded;
-}
-
-void SMontageEditor::AddNewSection(float StartTime, FString SectionName)
-{
-	if ( MontageObj != nullptr )
-	{
-		const FScopedTransaction Transaction( LOCTEXT("AddNewSection", "Add New Section") );
-		MontageObj->Modify();
-
-		if (MontageObj->AddAnimCompositeSection(FName(*SectionName), StartTime) != INDEX_NONE)
-		{
-			RebuildMontagePanel();
-		}
-		OnMontageModified();
-	}
-}
-
-void SMontageEditor::RemoveSection(int32 SectionIndex)
-{
-	if(ValidSection(SectionIndex))
-	{
-		const FScopedTransaction Transaction( LOCTEXT("DeleteSection", "Delete Section") );
-		MontageObj->Modify();
-
-		MontageObj->CompositeSections.RemoveAt(SectionIndex);
-		EnsureStartingSection();
-		OnMontageModified();
-		AnimMontageSectionsPanel->Update();
-		AnimTimingPanel->Update();
-		RestartPreview();
-	}
-}
-
-FString	SMontageEditor::GetSectionName(int32 SectionIndex) const
-{
-	if (ValidSection(SectionIndex))
-	{
-		return MontageObj->GetSectionName(SectionIndex).ToString();
-	}
-	return FString();
-}
-
-void SMontageEditor::RenameSlotNode(int32 SlotIndex, FString NewSlotName)
-{
-	if(MontageObj->SlotAnimTracks.IsValidIndex(SlotIndex))
-	{
-		FName NewName(*NewSlotName);
-		if(MontageObj->SlotAnimTracks[SlotIndex].SlotName != NewName)
-		{
-			const FScopedTransaction Transaction( LOCTEXT("RenameSlot", "Rename Slot") );
-			MontageObj->Modify();
-
-			MontageObj->SlotAnimTracks[SlotIndex].SlotName = NewName;
-			OnMontageModified();
-		}
-	}
-}
-
-void SMontageEditor::AddNewMontageSlot( FName NewSlotName )
-{
-	if ( MontageObj != nullptr )
-	{
-		const FScopedTransaction Transaction( LOCTEXT("AddSlot", "Add Slot") );
-		MontageObj->Modify();
-
-		MontageObj->AddSlot(NewSlotName);
-
-		OnMontageModified();
-
-		if (AnimMontagePanel.IsValid())
-		{
-			AnimMontagePanel->Update();
-		}
-	}
-}
-
-FText SMontageEditor::GetMontageSlotName(int32 SlotIndex) const
-{
-	if(MontageObj->SlotAnimTracks.IsValidIndex(SlotIndex) && MontageObj->SlotAnimTracks[SlotIndex].SlotName != NAME_None)
-	{
-		return FText::FromName( MontageObj->SlotAnimTracks[SlotIndex].SlotName );
-	}	
-	return FText::GetEmpty();
-}
-
-void SMontageEditor::RemoveMontageSlot(int32 AnimSlotIndex)
-{
-	if ( MontageObj != nullptr && MontageObj->SlotAnimTracks.IsValidIndex( AnimSlotIndex ) )
-	{
-		const FScopedTransaction Transaction( LOCTEXT("RemoveSlot", "Remove Slot") );
-		MontageObj->Modify();
-
-		MontageObj->SlotAnimTracks.RemoveAt(AnimSlotIndex);
-		OnMontageModified();
-		AnimMontagePanel->Update();
-
-		// Iterate the notifies and relink anything that is now invalid
-		for(FAnimNotifyEvent& Event : MontageObj->Notifies)
-		{
-			Event.ConditionalRelink();
-		}
-
-		// Do the same for sections
-		for(FCompositeSection& Section : MontageObj->CompositeSections)
-		{
-			Section.ConditionalRelink();
-		}
-	}
-}
-
-bool SMontageEditor::CanRemoveMontageSlot(int32 AnimSlotIndex)
-{
-	return (MontageObj != nullptr) && (MontageObj->SlotAnimTracks.Num()) > 1;
-}
-
-void SMontageEditor::DuplicateMontageSlot(int32 AnimSlotIndex)
-{
-	if (MontageObj != nullptr && MontageObj->SlotAnimTracks.IsValidIndex(AnimSlotIndex))
-	{
-		const FScopedTransaction Transaction(LOCTEXT("DuplicateSlot", "Duplicate Slot"));
-		MontageObj->Modify();
-
-		FSlotAnimationTrack& NewTrack = MontageObj->AddSlot(FAnimSlotGroup::DefaultSlotName); 
-		NewTrack.AnimTrack = MontageObj->SlotAnimTracks[AnimSlotIndex].AnimTrack;
-
-		OnMontageModified();
-
-		AnimMontagePanel->Update();
-	}
-}
-
-void SMontageEditor::ShowSectionInDetailsView(int32 SectionIndex)
-{
-	UEditorCompositeSection *Obj = Cast<UEditorCompositeSection>(ShowInDetailsView(UEditorCompositeSection::StaticClass()));
-	if ( Obj != nullptr )
-	{
-		Obj->InitSection(SectionIndex);
-	}
-	RestartPreviewFromSection(SectionIndex);
-}
-
-void SMontageEditor::RestartPreview()
-{
-	if (UDebugSkelMeshComponent* MeshComponent = GetPreviewScene()->GetPreviewMeshComponent())
-	{
-		if (UAnimPreviewInstance* Preview = MeshComponent->PreviewInstance)
-		{
-			Preview->MontagePreview_PreviewNormal(INDEX_NONE, Preview->IsPlaying());
-		}
-	}
-}
-
-void SMontageEditor::RestartPreviewFromSection(int32 FromSectionIdx)
-{
-	if (UDebugSkelMeshComponent* MeshComponent = GetPreviewScene()->GetPreviewMeshComponent())
-	{
-		if (UAnimPreviewInstance* Preview = MeshComponent->PreviewInstance)
-		{
-			Preview->MontagePreview_PreviewNormal(FromSectionIdx, Preview->IsPlaying());
-		}
-	}
-}
-
-void SMontageEditor::RestartPreviewPlayAllSections()
-{
-	if (UDebugSkelMeshComponent* MeshComponent = GetPreviewScene()->GetPreviewMeshComponent())
-	{
-		if (UAnimPreviewInstance* Preview = MeshComponent->PreviewInstance)
-		{
-			Preview->MontagePreview_PreviewAllSections(Preview->IsPlaying());
-		}
-	}
-}
-
-void SMontageEditor::MakeDefaultSequentialSections()
-{
-	check( MontageObj != nullptr );
-	SortSections();
-	for(int32 SectionIdx=0; SectionIdx < MontageObj->CompositeSections.Num(); SectionIdx++)
-	{
-		MontageObj->CompositeSections[SectionIdx].NextSectionName = MontageObj->CompositeSections.IsValidIndex(SectionIdx+1) ? MontageObj->CompositeSections[SectionIdx+1].SectionName : NAME_None;
-	}
-	RestartPreview();
-}
-
-void SMontageEditor::ClearSquenceOrdering()
-{
-	check( MontageObj != nullptr );
-	SortSections();
-	for(int32 SectionIdx=0; SectionIdx < MontageObj->CompositeSections.Num(); SectionIdx++)
-	{
-		MontageObj->CompositeSections[SectionIdx].NextSectionName = NAME_None;
-	}
-	RestartPreview();
 }
 
-void SMontageEditor::PostUndo( bool bSuccess )
-{
-	PostRedoUndo();
-}
-
-void SMontageEditor::PostRedo( bool bSuccess )
-{
-	PostRedoUndo();
-}
-
-void SMontageEditor::PostRedoUndo()
-{
-	// when undo or redo happens, we still have to recalculate length, so we can't rely on sequence length changes or not
-	if (MontageObj->SequenceLength)
-	{
-		MontageObj->SequenceLength = 0.f;
-	}
-
-	RebuildMontagePanel(); //Rebuild here, undoing adds can cause slate to crash later on if we don't (using dummy args since they aren't used by the method
-}
 #undef LOCTEXT_NAMESPACE
 
diff --git a/Engine/Source/Editor/Persona/Private/SMontageEditor.h b/Engine/Source/Editor/Persona/Private/SMontageEditor.h
--- a/Engine/Source/Editor/Persona/Private/SMontageEditor.h
+++ b/Engine/Source/Editor/Persona/Private/SMontageEditor.h
@@ -15,18 +15,21 @@ class SAnimMontageScrubPanel;
 class SAnimMontageSectionsPanel;
 class SAnimNotifyPanel;
 class SAnimTimingPanel;
+class FAnimModel_AnimMontage;
 
 struct FMontageEditorRequiredArgs
 {
-	FMontageEditorRequiredArgs(const TSharedRef<class IPersonaPreviewScene>& InPreviewScene, const TSharedRef<class IEditableSkeleton>& InEditableSkeleton, FSimpleMulticastDelegate& InOnSectionsChanged)
+	FMontageEditorRequiredArgs(const TSharedRef<class IPersonaPreviewScene>& InPreviewScene, const TSharedRef<class IEditableSkeleton>& InEditableSkeleton, FSimpleMulticastDelegate& InOnSectionsChanged, const TSharedRef<FUICommandList>& InCommandList)
 		: PreviewScene(InPreviewScene)
 		, EditableSkeleton(InEditableSkeleton)
 		, OnSectionsChanged(InOnSectionsChanged)
+		, CommandList(InCommandList)
 	{}
 
 	TSharedRef<class IPersonaPreviewScene> PreviewScene;
 	TSharedRef<class IEditableSkeleton> EditableSkeleton;
 	FSimpleMulticastDelegate& OnSectionsChanged;
+	TSharedRef<FUICommandList> CommandList;
 };
 
 //////////////////////////////////////////////////////////////////////////
@@ -38,158 +41,44 @@ struct FMontageEditorRequiredArgs
 	portion of the Montage tool and registering callbacks to the SMontageEditor to do the actual editing.
 	
 */
-class SMontageEditor : public SAnimEditorBase, public FEditorUndoClient
+class SMontageEditor : public SAnimEditorBase
 {
 public:
 	SLATE_BEGIN_ARGS( SMontageEditor )
 		: _Montage(NULL)
 		{}
 
-		SLATE_EVENT(FOnInvokeTab, OnInvokeTab)
-		SLATE_EVENT(FSimpleDelegate, OnCurvesChanged)
-		SLATE_EVENT(FSimpleDelegate, OnSectionsChanged)
-		SLATE_ARGUMENT( UAnimMontage*, Montage )
-		SLATE_EVENT(FOnObjectsSelected, OnObjectsSelected)
-	SLATE_END_ARGS()
+	SLATE_EVENT(FOnInvokeTab, OnInvokeTab)
+	SLATE_EVENT(FSimpleDelegate, OnCurvesChanged)
+	SLATE_EVENT(FSimpleDelegate, OnSectionsChanged)
+	SLATE_ARGUMENT( UAnimMontage*, Montage )
+	SLATE_EVENT(FOnObjectsSelected, OnObjectsSelected)
+	SLATE_EVENT(FOnEditCurves, OnEditCurves)
+	SLATE_EVENT(FOnStopEditingCurves, OnStopEditingCurves)
 
-	~SMontageEditor();
+	SLATE_END_ARGS()
 
 private:
-	TSharedPtr<SAnimMontagePanel> AnimMontagePanel;
-	TSharedPtr<SAnimNotifyPanel> AnimNotifyPanel;
-	TSharedPtr<SAnimCurvePanel>	AnimCurvePanel;
-	TSharedPtr<SAnimMontageSectionsPanel> AnimMontageSectionsPanel;
-	TSharedPtr<SAnimMontageScrubPanel> AnimMontageScrubPanel;
-	TSharedPtr<SAnimTimingPanel> AnimTimingPanel;
-
-protected:
-	//~ Begin SAnimEditorBase Interface
-	virtual TSharedRef<class SAnimationScrubPanel> ConstructAnimScrubPanel() override;
-	//~ End SAnimEditorBase Interface
+	TSharedPtr<FAnimModel_AnimMontage> AnimModel;
 
 public:
+	~SMontageEditor();
+
 	void Construct(const FArguments& InArgs, const FMontageEditorRequiredArgs& InRequiredArgs);
 
 	void SetMontageObj(UAnimMontage * NewMontage);
 	UAnimMontage * GetMontageObj() const { return MontageObj; }
 
-	/** Handler for when the preview slot is changed */
-	void OnSetMontagePreviewSlot(int32 SlotIndex);
-
 	virtual UAnimationAsset* GetEditorObject() const override { return GetMontageObj(); }
 
-	void RestartPreview();
-	void RestartPreviewFromSection(int32 FromSectionIdx = INDEX_NONE);
-	void RestartPreviewPlayAllSections();
-
 private:
 	/** Pointer to the animation sequence being edited */
 	// @todo fix this
 	UAnimMontage* MontageObj;
-	/*
-	* Child Anim Montage: Child Anim Montage only can replace name of animations, and no other meaningful edits
-	* as it will derive every data from Parent. There might be some other data that will allow to be replaced, but for now, it is
-	* not.
-	*/
-	bool bChildAnimMontage;
-
-	/** If previewing section, it is section used to restart previewing when play button is pushed */
-	int32 PreviewingStartSectionIdx;
-
-	/** If currently previewing all or selected section */
-	bool bPreviewingAllSections : 1;
-
-	/** If currently previewing tracks instead of sections */
-	bool bPreviewingTracks : 1;
-
-	/** If user is currently dragging an item */
-	bool bDragging : 1;
-
-	/** If the active timer to trigger a montage panel rebuild is currently registered */
-	bool bIsActiveTimerRegistered : 1;
-
-	virtual float CalculateSequenceLengthOfEditorObject() const override;
-	void SortAndUpdateMontage();
-	void CollapseMontage();
-	void SortAnimSegments();
-	void SortSections();
-	void EnsureStartingSection();
-	void EnsureSlotNode();
-	virtual bool ClampToEndTime(float NewEndTime) override;
-
-	/** FEditorUndoClient interface */
-	virtual void PostUndo( bool bSuccess ) override;
-	virtual void PostRedo( bool bSuccess ) override;
-
-	void PostRedoUndo();
-	
-	bool GetSectionTime( int32 SectionIndex, float &OutTime ) const;
-
-	bool ValidIndexes(int32 AnimSlotIndex, int32 AnimSegmentIndex) const;
-	bool ValidSection(int32 SectionIndex) const;
-
-	/** Updates Notify trigger offsets to take into account current montage state */
-	void RefreshNotifyTriggerOffsets();
-
-	/** Handle notifies changing, so we rebuild timing */
-	void HandleNotifiesChanged();
 
-	/** One-off active timer to trigger a montage panel rebuild */
-	EActiveTimerReturnType TriggerRebuildMontagePanel(double InCurrentTime, float InDeltaTime);
-
-	/** Rebuilds the montage panel */
-	void RebuildMontagePanel(bool bNotifyAsset=true);
 
 protected:
-	virtual void InitDetailsViewEditorObject(class UEditorAnimBaseObj* EdObj) override;
-
-public:
-
-	/** These are meant to be callbacks used by the montage editing UI widgets */
-
-	void					OnMontageChange(class UObject *EditorAnimBaseObj, bool Rebuild);
-	void					ShowSectionInDetailsView(int32 SectionIdx);
-
-	TArray<float>			GetSectionStartTimes() const;
-	TArray<FTrackMarkerBar>	GetMarkerBarInformation() const;
-	TArray<FString>			GetSectionNames() const;
-	TArray<float>			GetAnimSegmentStartTimes() const;
-	void					OnEditSectionTime( int32 SectionIndex, float NewTime);
-	void					OnEditSectionTimeFinish( int32 SectionIndex);
-
-	void					AddNewSection(float StartTime, FString SectionName);
-	void					RemoveSection(int32 SectionIndex);
-	FString					GetSectionName(int32 SectionIndex) const;
-
-	// Set the time of the supplied section
-	void					SetSectionTime(int32 SectionIndex, float NewTime);
-
-	void					RenameSlotNode(int32 SlotIndex, FString NewSlotName);
-
-	// UI Slot Action handlers
-	void					AddNewMontageSlot(FName NewSlotName);
-	void					RemoveMontageSlot(int32 AnimSlotIndex);
-	bool					CanRemoveMontageSlot(int32 AnimSlotIndex);
-	void					DuplicateMontageSlot(int32 AnimSlotIndex);
-	FText					GetMontageSlotName(int32 SlotIndex) const;
-
-	void					MakeDefaultSequentialSections();
-	void					ClearSquenceOrdering();
-
-	/** Delegete handlers for when the editor UI is changing the montage */
-	void			PreAnimUpdate();
-	void			PostAnimUpdate();
-	void			OnMontageModified();
-	void			ReplaceAnimationMapping(FName SlotName, int32 SegmentIdx, UAnimSequenceBase* OldSequenceBase, UAnimSequenceBase* NewSequenceBase);
-	bool			IsDiffererentFromParent(FName SlotName, int32 SegmentIdx, const FAnimSegment& Segment);
-
 	//~ Begin SAnimEditorBase Interface
 	virtual TSharedRef<SWidget> CreateDocumentAnchor() override;
 	//~ End SAnimEditorBase Interface
-
-	/** Delegate fired when montage sections have changed */
-	FSimpleDelegate			OnSectionsChanged;
-
-	FReply	OnFindParentClassInContentBrowserClicked();
-	FReply	OnEditParentClassClicked();
 };
diff --git a/Engine/Source/Editor/Persona/Private/SPoseEditor.cpp b/Engine/Source/Editor/Persona/Private/SPoseEditor.cpp
--- a/Engine/Source/Editor/Persona/Private/SPoseEditor.cpp
+++ b/Engine/Source/Editor/Persona/Private/SPoseEditor.cpp
@@ -37,7 +37,7 @@ void SPoseEditor::Construct(const FArguments& InArgs, const TSharedRef<class IPe
 	check(PoseAssetObj);
 
 	SAnimEditorBase::Construct(SAnimEditorBase::FArguments()
-		.DisplayAnimInfoBar(false),
+		.DisplayAnimTimeline(false),
 		InPreviewScene);
 
 	NonScrollEditorPanels->AddSlot()
diff --git a/Engine/Source/Editor/Persona/Private/SSequenceEditor.cpp b/Engine/Source/Editor/Persona/Private/SSequenceEditor.cpp
--- a/Engine/Source/Editor/Persona/Private/SSequenceEditor.cpp
+++ b/Engine/Source/Editor/Persona/Private/SSequenceEditor.cpp
@@ -8,77 +8,41 @@
 #include "SAnimTrackCurvePanel.h"
 #include "AnimPreviewInstance.h"
 #include "Editor.h"
+#include "AnimModel_AnimSequenceBase.h"
+#include "SAnimTimeline.h"
 
 #define LOCTEXT_NAMESPACE "AnimSequenceEditor"
 
 //////////////////////////////////////////////////////////////////////////
 // SSequenceEditor
 
-void SSequenceEditor::Construct(const FArguments& InArgs, TSharedRef<class IPersonaPreviewScene> InPreviewScene, TSharedRef<class IEditableSkeleton> InEditableSkeleton)
+void SSequenceEditor::Construct(const FArguments& InArgs, TSharedRef<class IPersonaPreviewScene> InPreviewScene, TSharedRef<class IEditableSkeleton> InEditableSkeleton, const TSharedRef<FUICommandList>& InCommandList)
 {
 	SequenceObj = InArgs._Sequence;
 	check(SequenceObj);
 	PreviewScenePtr = InPreviewScene;
 
+	AnimModel = MakeShared<FAnimModel_AnimSequenceBase>(InPreviewScene, InEditableSkeleton, InCommandList, SequenceObj);
+
+	AnimModel->OnEditCurves = FOnEditCurves::CreateLambda([this, InOnEditCurves = InArgs._OnEditCurves](UAnimSequenceBase* InAnimSequence, const TArray<IAnimationEditor::FCurveEditInfo>& InCurveInfo, const TSharedPtr<ITimeSliderController>& InExternalTimeSliderController)
+	{
+		InOnEditCurves.ExecuteIfBound(InAnimSequence, InCurveInfo, TimelineWidget->GetTimeSliderController());
+	});
+
+	AnimModel->OnStopEditingCurves = InArgs._OnStopEditingCurves;
+	AnimModel->OnSelectObjects = FOnObjectsSelected::CreateSP(this, &SAnimEditorBase::OnSelectionChanged);
+	AnimModel->OnInvokeTab = InArgs._OnInvokeTab;
+	AnimModel->Initialize();
+
 	SAnimEditorBase::Construct( SAnimEditorBase::FArguments()
-		.OnObjectsSelected(InArgs._OnObjectsSelected), 
-		InPreviewScene );
+		.OnObjectsSelected(InArgs._OnObjectsSelected)
+		.AnimModel(AnimModel), 
+		InPreviewScene);
 
 	if(GEditor)
 	{
 		GEditor->RegisterForUndo(this);
 	}
-
-	EditorPanels->AddSlot()
-	.AutoHeight()
-	.Padding(0, 10)
-	[
-		SAssignNew( AnimNotifyPanel, SAnimNotifyPanel, InEditableSkeleton )
-		.Sequence(SequenceObj)
-		.WidgetWidth(S2ColumnWidget::DEFAULT_RIGHT_COLUMN_WIDTH)
-		.ViewInputMin(this, &SAnimEditorBase::GetViewMinInput)
-		.ViewInputMax(this, &SAnimEditorBase::GetViewMaxInput)
-		.InputMin(this, &SAnimEditorBase::GetMinInput)
-		.InputMax(this, &SAnimEditorBase::GetMaxInput)
-		.OnSetInputViewRange(this, &SAnimEditorBase::SetInputViewRange)
-		.OnGetScrubValue(this, &SAnimEditorBase::GetScrubValue)
-		.OnSelectionChanged(this, &SAnimEditorBase::OnSelectionChanged)
-		.OnInvokeTab(InArgs._OnInvokeTab)
-	];
-
-	EditorPanels->AddSlot()
-	.AutoHeight()
-	.Padding(0, 10)
-	[
-		SAssignNew( AnimCurvePanel, SAnimCurvePanel, InEditableSkeleton )
-		.Sequence(SequenceObj)
-		.WidgetWidth(S2ColumnWidget::DEFAULT_RIGHT_COLUMN_WIDTH)
-		.ViewInputMin(this, &SAnimEditorBase::GetViewMinInput)
-		.ViewInputMax(this, &SAnimEditorBase::GetViewMaxInput)
-		.InputMin(this, &SAnimEditorBase::GetMinInput)
-		.InputMax(this, &SAnimEditorBase::GetMaxInput)
-		.OnSetInputViewRange(this, &SAnimEditorBase::SetInputViewRange)
-		.OnGetScrubValue(this, &SAnimEditorBase::GetScrubValue)
-	];
-
-	UAnimSequence * AnimSeq = Cast<UAnimSequence>(SequenceObj);
-	if (AnimSeq)
-	{
-		EditorPanels->AddSlot()
-		.AutoHeight()
-		.Padding(0, 10)
-		[
-			SAssignNew(AnimTrackCurvePanel, SAnimTrackCurvePanel, InPreviewScene)
-			.Sequence(AnimSeq)
-			.WidgetWidth(S2ColumnWidget::DEFAULT_RIGHT_COLUMN_WIDTH)
-			.ViewInputMin(this, &SAnimEditorBase::GetViewMinInput)
-			.ViewInputMax(this, &SAnimEditorBase::GetViewMaxInput)
-			.InputMin(this, &SAnimEditorBase::GetMinInput)
-			.InputMax(this, &SAnimEditorBase::GetMaxInput)
-			.OnSetInputViewRange(this, &SAnimEditorBase::SetInputViewRange)
-			.OnGetScrubValue(this, &SAnimEditorBase::GetScrubValue)
-		];
-	}
 }
 
 SSequenceEditor::~SSequenceEditor()
@@ -103,17 +67,7 @@ void SSequenceEditor::PostUndoRedo()
 {
 	GetPreviewScene()->SetPreviewAnimationAsset(SequenceObj);
 
-	if( SequenceObj )
-	{
-		SetInputViewRange(0, SequenceObj->SequenceLength);
-
-		AnimNotifyPanel->Update();
-		AnimCurvePanel->UpdatePanel();
-		if (AnimTrackCurvePanel.IsValid())
-		{
-			AnimTrackCurvePanel->UpdatePanel();
-		}
-	}
+	AnimModel->RefreshTracks();
 }
 
 #undef LOCTEXT_NAMESPACE
diff --git a/Engine/Source/Editor/Persona/Private/SSequenceEditor.h b/Engine/Source/Editor/Persona/Private/SSequenceEditor.h
--- a/Engine/Source/Editor/Persona/Private/SSequenceEditor.h
+++ b/Engine/Source/Editor/Persona/Private/SSequenceEditor.h
@@ -10,9 +10,11 @@
 #include "SAnimationScrubPanel.h"
 #include "SAnimEditorBase.h"
 #include "EditorUndoClient.h"
+#include "Animation/AnimSequence.h"
 
 class SAnimNotifyPanel;
 class SAnimTrackCurvePanel;
+class FAnimModel_AnimSequenceBase;
 
 //////////////////////////////////////////////////////////////////////////
 // SSequenceEditor
@@ -28,17 +30,16 @@ public:
 		SLATE_ARGUMENT(UAnimSequenceBase*, Sequence)
 		SLATE_EVENT(FOnObjectsSelected, OnObjectsSelected)
 		SLATE_EVENT(FOnInvokeTab, OnInvokeTab)
+		SLATE_EVENT(FOnEditCurves, OnEditCurves)
+		SLATE_EVENT(FOnStopEditingCurves, OnStopEditingCurves)
 
 	SLATE_END_ARGS()
 
 private:
-	TSharedPtr<class SAnimNotifyPanel>	AnimNotifyPanel;
-	TSharedPtr<class SAnimCurvePanel>	AnimCurvePanel;
-	TSharedPtr<class SAnimTrackCurvePanel>	AnimTrackCurvePanel;
-	TSharedPtr<class SAnimationScrubPanel> AnimScrubPanel;
 	TWeakPtr<class IPersonaPreviewScene> PreviewScenePtr;
+	TSharedPtr<FAnimModel_AnimSequenceBase> AnimModel;
 public:
-	void Construct(const FArguments& InArgs, TSharedRef<class IPersonaPreviewScene> InPreviewScene, TSharedRef<class IEditableSkeleton> InEditableSkeleton);
+	void Construct(const FArguments& InArgs, TSharedRef<class IPersonaPreviewScene> InPreviewScene, TSharedRef<class IEditableSkeleton> InEditableSkeleton, const TSharedRef<FUICommandList>& InCommandList);
 
 	~SSequenceEditor();
 
diff --git a/Engine/Source/Editor/Persona/Private/SSkeletonSlotNames.cpp b/Engine/Source/Editor/Persona/Private/SSkeletonSlotNames.cpp
--- a/Engine/Source/Editor/Persona/Private/SSkeletonSlotNames.cpp
+++ b/Engine/Source/Editor/Persona/Private/SSkeletonSlotNames.cpp
@@ -131,6 +131,8 @@ void SSkeletonSlotNames::Construct(const FArguments& InArgs, const TSharedRef<IE
 	EditableSkeletonPtr = InEditableSkeleton;
 	OnObjectSelected = InArgs._OnObjectSelected;
 
+	InEditableSkeleton->RegisterOnSlotsChanged(FSimpleMulticastDelegate::FDelegate::CreateSP(this, &SSkeletonSlotNames::RefreshSlotNameListWithFilter));
+
 	InOnPostUndo.Add(FSimpleDelegate::CreateSP( this, &SSkeletonSlotNames::PostUndo ) );
 
 	// Toolbar
diff --git a/Engine/Source/Editor/Persona/Private/STrack.cpp b/Engine/Source/Editor/Persona/Private/STrack.cpp
--- a/Engine/Source/Editor/Persona/Private/STrack.cpp
+++ b/Engine/Source/Editor/Persona/Private/STrack.cpp
@@ -14,7 +14,7 @@
 #include "SCurveEditor.h"
 #include "SScrubWidget.h"
 
-const float STrackDefaultHeight = 20.0f;
+const float STrackDefaultHeight = 24.0f;
 const float DraggableBarSnapTolerance = 20.0f;
 const float NodeHandleWidth = 12.f;
 const float NodeHandleHeight = STrackDefaultHeight;
@@ -98,11 +98,11 @@ void STrackNode::Construct(const FArguments& InArgs)
 	bCenterOnPosition = InArgs._CenterOnPosition;
 	
 	NodeSelectionSet = InArgs._NodeSelectionSet;
-	AllowDrag = InArgs._AllowDrag;
+	bAllowDrag = InArgs._AllowDrag;
 
 	Font = FCoreStyle::GetDefaultFontStyle("Regular", 10);
 
-	const FSlateBrush* StyleInfo = FEditorStyle::GetBrush("ProgressBar.Background"); // FIXME: make slate argument for STrackNode
+	const FSlateBrush* StyleInfo = FEditorStyle::GetBrush( TEXT("SpecialEditableTextImageNormal") ); // FIXME: make slate argument for STrackNode
 
 	if(InArgs._OverrideContent.Widget != SNullWidget::NullWidget)
 	{
@@ -195,10 +195,18 @@ FVector2D STrackNode::GetDragDropScreenSpacePosition(const FGeometry& ParentAllo
 FReply STrackNode::OnMouseButtonUp( const FGeometry& MyGeometry, const FPointerEvent& MouseEvent )
 {
 	if (MouseEvent.GetEffectingButton() == EKeys::LeftMouseButton)
+	{
+		Select();
+		OnTrackNodeClicked.ExecuteIfBound();
+
+		return FReply::Handled().ReleaseMouseCapture();
+	}
+	else if(MouseEvent.GetEffectingButton() == EKeys::RightMouseButton)
 	{
 		Select();
 		OnTrackNodeClicked.ExecuteIfBound();
 	}
+
 	return FReply::Unhandled();
 }
 
@@ -219,7 +227,7 @@ void STrackNode::OnMouseLeave( const FPointerEvent& MouseEvent )
 
 FReply STrackNode::OnDragDetected( const FGeometry& MyGeometry, const FPointerEvent& MouseEvent )
 {
-	if (AllowDrag && MouseEvent.IsMouseButtonDown( EKeys::LeftMouseButton ) )
+	if (bAllowDrag && MouseEvent.IsMouseButtonDown( EKeys::LeftMouseButton ) )
 	{
 		bBeingDragged = true;
 
@@ -242,17 +250,23 @@ FReply STrackNode::BeginDrag( const FGeometry& MyGeometry, const FPointerEvent&
 	OnTrackNodeClicked.ExecuteIfBound();
 
 	//void FTrackNodeDragDropOp(TSharedRef<STrackNode> TrackNode, const FVector2D &CursorPosition, const FVector2D &ScreenPositionOfNode)
-	return FReply::Handled().BeginDragDrop(FTrackNodeDragDropOp::New(SharedThis(this), ScreenCursorPos, ScreenNodePosition));
+	return FReply::Handled().BeginDragDrop(FTrackNodeDragDropOp::New(SharedThis(this), ScreenCursorPos, ScreenNodePosition)).ReleaseMouseCapture();
 }
 
 FReply STrackNode::OnMouseButtonDown( const FGeometry& MyGeometry, const FPointerEvent& MouseEvent )
 {
-	if ( MouseEvent.GetEffectingButton() == EKeys::LeftMouseButton )
+	if(MouseEvent.GetEffectingButton() == EKeys::RightMouseButton)
+	{
+		return FReply::Handled();
+	}
+	else if ( MouseEvent.GetEffectingButton() == EKeys::LeftMouseButton )
 	{
-		if(AllowDrag)
+		if(bAllowDrag)
 		{
-			return FReply::Handled().DetectDrag( SharedThis(this), EKeys::LeftMouseButton );
+			return FReply::Handled().DetectDrag( SharedThis(this), EKeys::LeftMouseButton ).CaptureMouse(AsShared());
 		}
+
+		return FReply::Handled();
 	}
 
 	return FReply::Unhandled();
@@ -360,7 +374,6 @@ void STrack::Construct( const FArguments& InArgs )
 	TrackMinValue = InArgs._TrackMinValue;
 	TrackMaxValue = InArgs._TrackMaxValue;
 	TrackNumDiscreteValues = InArgs._TrackNumDiscreteValues;
-	ScrubPosition = InArgs._ScrubPosition;
 	ViewInputMin = InArgs._ViewInputMin;
 	ViewInputMax = InArgs._ViewInputMax;
 	OnSelectionChanged = InArgs._OnSelectionChanged;
@@ -421,35 +434,6 @@ int32 STrack::OnPaint( const FPaintArgs& Args, const FGeometry& AllottedGeometry
 	int32 CustomLayerId = LayerId + 1;
 	FPaintGeometry MyGeometry = AllottedGeometry.ToPaintGeometry();
 
-	// Background
-	FSlateDrawElement::MakeBox( 
-		OutDrawElements,
-		CustomLayerId++, 
-		MyGeometry, 
-		FEditorStyle::GetBrush(TEXT( "Persona.NotifyEditor.NotifyTrackBackground" )),
-		ESlateDrawEffect::None, 
-		TrackColor.Get()
-		);
-
-	// Draw Scrub Position
-	if ( ScrubPosition.Get() >= 0.f )
-	{
-		float XPos = DataToLocalX(ScrubPosition.Get(), AllottedGeometry);
-
-		TArray<FVector2D> LinePoints;
-		LinePoints.Add(FVector2D(XPos, 0.f));
-		LinePoints.Add(FVector2D(XPos, AllottedGeometry.Size.Y));
-
-		FSlateDrawElement::MakeLines( 
-			OutDrawElements,
-			CustomLayerId++,
-			MyGeometry,
-			LinePoints,
-			ESlateDrawEffect::None,
-			FLinearColor::Red
-			);
-	}
-
 	// Draggable Bars
 	for ( int32 I=0; DraggableBars.IsBound() && I < DraggableBars.Get().Num(); I++ )
 	{
@@ -505,36 +489,6 @@ int32 STrack::OnPaint( const FPaintArgs& Args, const FGeometry& AllottedGeometry
 			);
 	}
 	CustomLayerId++;
-
-	// draw line for every 1/4 length
-	if ( TrackMaxValue.Get() > 0.f && TrackNumDiscreteValues.Get() > 0 )
-	{
-		int32 Divider = SScrubWidget::GetDivider( ViewInputMin.Get(), ViewInputMax.Get(), AllottedGeometry.Size, TrackMaxValue.Get(), TrackNumDiscreteValues.Get());
-
-		const float TimePerValue = TrackMaxValue.Get()/static_cast<float>(TrackNumDiscreteValues.Get());
-		for (int I=1; I<TrackNumDiscreteValues.Get(); ++I)
-		{
-			if ( I % Divider == 0 )
-			{
-				float XPos = DataToLocalX(TimePerValue*I, AllottedGeometry);
-
-				TArray<FVector2D> LinePoints;
-				LinePoints.Add(FVector2D(XPos, 0.f));
-				LinePoints.Add(FVector2D(XPos, AllottedGeometry.Size.Y));
-
-				FSlateDrawElement::MakeLines( 
-					OutDrawElements,
-					CustomLayerId,
-					MyGeometry,
-					LinePoints,
-					ESlateDrawEffect::None,
-					FLinearColor::Black
-					);
-			}
-		}	
-		++CustomLayerId;
-	}
-
 	
 	return SPanel::OnPaint( Args, AllottedGeometry, MyCullingRect, OutDrawElements, CustomLayerId, InWidgetStyle, bParentEnabled );
 }
diff --git a/Engine/Source/Editor/Persona/Private/STrack.h b/Engine/Source/Editor/Persona/Private/STrack.h
--- a/Engine/Source/Editor/Persona/Private/STrack.h
+++ b/Engine/Source/Editor/Persona/Private/STrack.h
@@ -25,7 +25,6 @@ class STrack;
 class STrackNode;
 
 //////////////////////////////////////////////////////////////////////////
-DECLARE_DELEGATE_RetVal( float, FOnGetScrubValue )
 DECLARE_DELEGATE_OneParam( FOnSelectionChanged, const TArray<UObject*>& )
 DECLARE_DELEGATE( FOnNodeSelectionChanged )
 DECLARE_DELEGATE( FOnUpdatePanel )
@@ -206,9 +205,10 @@ protected:
 	mutable FVector2D	LastSize; //HACK Fixme: Need to update in OnPaint in case we are drag/dropped
 	FGeometry			CachedTrackGeometry; //Our parent tracks geometry, so we can calculate scale/position
 	
+	float				DistanceDragged;
 	bool				bBeingDragged;
 	bool				bCenterOnPosition;
-	bool				AllowDrag;
+	bool				bAllowDrag;
 	bool				bContentOverriden;
 
 	friend class STrack;
@@ -227,7 +227,6 @@ public:
 		, _TrackMaxValue()
 		, _TrackMinValue()
 		, _TrackNumDiscreteValues()
-		, _ScrubPosition()
 		, _TrackColor(FLinearColor::White)
 		, _OnSelectionChanged()
 		, _DraggableBars()
@@ -245,7 +244,6 @@ public:
 		SLATE_ATTRIBUTE( float, TrackMaxValue )
 		SLATE_ATTRIBUTE( float, TrackMinValue )
 		SLATE_ATTRIBUTE( int32, TrackNumDiscreteValues )
-		SLATE_ATTRIBUTE( float, ScrubPosition )
 		SLATE_ARGUMENT( FLinearColor, TrackColor )
 		SLATE_EVENT( FOnNodeSelectionChanged, OnSelectionChanged )
 		SLATE_ATTRIBUTE( TArray<float>, DraggableBars )
@@ -313,7 +311,6 @@ protected:
 	TAttribute<float>						TrackMaxValue;
 	TAttribute<float>						TrackMinValue;
 	TAttribute<int32>						TrackNumDiscreteValues; // Discrete values (such as "number of frames in animation" used for accurate grid lines)
-	TAttribute<float>						ScrubPosition;
 
 	TAttribute<float>						ViewInputMin;
 	TAttribute<float>						ViewInputMax;
diff --git a/Engine/Source/Editor/Persona/Private/TabSpawners.cpp b/Engine/Source/Editor/Persona/Private/TabSpawners.cpp
--- a/Engine/Source/Editor/Persona/Private/TabSpawners.cpp
+++ b/Engine/Source/Editor/Persona/Private/TabSpawners.cpp
@@ -61,6 +61,7 @@ const FName FPersonaTabs::SkeletonAnimNotifiesID("SkeletonAnimNotifies");
 const FName FPersonaTabs::SkeletonSlotNamesID("SkeletonSlotNames");
 const FName FPersonaTabs::SkeletonSlotGroupNamesID("SkeletonSlotGroupNames");
 const FName FPersonaTabs::BlendProfileManagerID("BlendProfileManager");
+const FName FPersonaTabs::AnimMontageSectionsID("AnimMontageSections");
 
 const FName FPersonaTabs::AdvancedPreviewSceneSettingsID("AdvancedPreviewTab");
 
diff --git a/Engine/Source/Editor/Persona/Public/IAnimSequenceCurveEditor.h b/Engine/Source/Editor/Persona/Public/IAnimSequenceCurveEditor.h
new file mode 100644
--- /dev/null
+++ b/Engine/Source/Editor/Persona/Public/IAnimSequenceCurveEditor.h
@@ -0,0 +1,27 @@
+// Copyright Epic Games, Inc. All Rights Reserved.
+
+#pragma once
+
+#include "CoreMinimal.h"
+#include "Widgets/SCompoundWidget.h"
+
+struct FRichCurve;
+enum class ERawCurveTrackTypes : uint8;
+struct FSmartName;
+
+/** Interface to the curve editor tab in the anim sequence editor */
+class IAnimSequenceCurveEditor : public SCompoundWidget
+{
+public:
+	/** Reset the curves that are edited by this widget */
+	virtual void ResetCurves() = 0;
+
+	/** Add a curve for editing */
+	virtual void AddCurve(const FText& InCurveDisplayName, const FLinearColor& InCurveColor, const FSmartName& InName, ERawCurveTrackTypes InType, int32 InCurveIndex, FSimpleDelegate InOnCurveModified) = 0;
+
+	/** Remove a curve for editing */
+	virtual void RemoveCurve(const FSmartName& InName, ERawCurveTrackTypes InType, int32 InCurveIndex) = 0;
+
+	/** Zoom to fit all the curves */
+	virtual void ZoomToFit() = 0;
+};
\ No newline at end of file
diff --git a/Engine/Source/Editor/Persona/Public/IPersonaPreviewScene.h b/Engine/Source/Editor/Persona/Public/IPersonaPreviewScene.h
--- a/Engine/Source/Editor/Persona/Public/IPersonaPreviewScene.h
+++ b/Engine/Source/Editor/Persona/Public/IPersonaPreviewScene.h
@@ -45,6 +45,9 @@ typedef FOnSelectedBoneChangedMulticaster::FDelegate FOnSelectedBoneChanged;
 DECLARE_MULTICAST_DELEGATE_OneParam(FOnSelectedSocketChangedMulticaster, const FSelectedSocketInfo& /*InSocketInfo*/);
 typedef FOnSelectedSocketChangedMulticaster::FDelegate FOnSelectedSocketChanged;
 
+//The delegate to check if the attach component can be removed
+DECLARE_DELEGATE_RetVal_OneParam(bool, FOnRemoveAttachedComponentFilter, const USceneComponent* /*InComponent*/);
+
 /** Modes that the preview scene defaults to (usually depending on asset editor context) */
 enum class EPreviewSceneDefaultAnimationMode : int32
 {
@@ -283,6 +286,10 @@ public:
 	/** Unregister a callback for just after the preview scene is ticked */
 	virtual void UnregisterOnPostTick(void* Thing) = 0;
 
+	/** setter/getter for can remove attach component */
+	virtual void SetRemoveAttachedComponentFilter(const FOnRemoveAttachedComponentFilter& Delegate) = 0;
+	virtual void ClearRemoveAttachedComponentFilter() = 0;
+
 	/** Let the preview scene know that it should tick (because it is visible) */
 	virtual void FlagTickable() = 0;
 	/** Handle syncing selection with the skeleton tree */
diff --git a/Engine/Source/Editor/Persona/Public/PersonaDelegates.h b/Engine/Source/Editor/Persona/Public/PersonaDelegates.h
--- a/Engine/Source/Editor/Persona/Public/PersonaDelegates.h
+++ b/Engine/Source/Editor/Persona/Public/PersonaDelegates.h
@@ -3,9 +3,12 @@
 #pragma once
 
 #include "CoreMinimal.h"
+#include "IAnimationEditor.h"
 
 class SAnimationEditorViewportTabBody;
 struct FTabId;
+class UAnimSequenceBase;
+class ITimeSliderController;
 
 // Called when a blend profile is selected in the blend profile tab
 DECLARE_MULTICAST_DELEGATE_OneParam(FOnSelectBlendProfile, class UBlendProfile*);
@@ -27,3 +30,9 @@ DECLARE_DELEGATE_OneParam(FOnOpenNewAsset, UObject* /* InAsset */);
 
 // Called to get an object (used by the asset details panel)
 DECLARE_DELEGATE_RetVal(UObject*, FOnGetAsset);
+
+// Called to invoke the curve editor
+DECLARE_DELEGATE_ThreeParams(FOnEditCurves, UAnimSequenceBase* /*InAnimSequence*/, const TArray<IAnimationEditor::FCurveEditInfo>& /*InCurveInfo*/, const TSharedPtr<ITimeSliderController>& /*InTimeSliderController*/);
+
+// Called to stop editing curves in the curve editor
+DECLARE_DELEGATE_OneParam(FOnStopEditingCurves, const TArray<IAnimationEditor::FCurveEditInfo>& /*InCurveInfo*/);
\ No newline at end of file
diff --git a/Engine/Source/Editor/Persona/Public/PersonaModule.h b/Engine/Source/Editor/Persona/Public/PersonaModule.h
--- a/Engine/Source/Editor/Persona/Public/PersonaModule.h
+++ b/Engine/Source/Editor/Persona/Public/PersonaModule.h
@@ -24,6 +24,8 @@ class FWorkflowAllowedTabSet;
 class IAssetFamily;
 class FWorkflowTabFactory;
 class UBlendSpaceBase;
+class IAnimSequenceCurveEditor;
+class IAnimationEditor;
 
 extern const FName PersonaAppName;
 
@@ -248,8 +250,14 @@ public:
 	/** Create a widget to preview a blendspace */
 	virtual TSharedRef<SWidget> CreateBlendSpacePreviewWidget(TAttribute<const UBlendSpaceBase*> InBlendSpace, TAttribute<FVector> InPosition) const;
 
+	/** Create a tab factory for editing montage sections */
+	virtual TSharedRef<FWorkflowTabFactory> CreateAnimMontageSectionsTabFactory(const TSharedRef<class FWorkflowCentricApplication>& InHostingApp, const TSharedRef<IPersonaToolkit>& InPersonaToolkit, FSimpleMulticastDelegate& InOnSectionsChanged) const;
+
 	/** Create a widget that acts as a document for an animation asset */
-	virtual TSharedRef<SWidget> CreateEditorWidgetForAnimDocument(const TSharedRef<class FWorkflowCentricApplication>& InHostingApp, UObject* InAnimAsset, const FAnimDocumentArgs& InArgs, FString& OutDocumentLink);
+	virtual TSharedRef<SWidget> CreateEditorWidgetForAnimDocument(const TSharedRef<IAnimationEditor>& InHostingApp, UObject* InAnimAsset, const FAnimDocumentArgs& InArgs, FString& OutDocumentLink);
+
+	/** Create a widget that acts as a curve document for an animation asset */
+	virtual TSharedRef<IAnimSequenceCurveEditor> CreateCurveWidgetForAnimDocument(const TSharedRef<class FWorkflowCentricApplication>& InHostingApp, const TSharedRef<IPersonaPreviewScene>& InPreviewScene, UAnimSequenceBase* InAnimSequence, const TSharedPtr<ITimeSliderController>& InExternalTimeSliderController, const TSharedPtr<FTabManager>& InTabManager);
 
 	/** Customize a skeletal mesh details panel */
 	virtual void CustomizeMeshDetails(const TSharedRef<IDetailsView>& InDetailsView, const TSharedRef<IPersonaToolkit>& InPersonaToolkit);
@@ -264,8 +272,8 @@ public:
 	/** Check all animations & skeletal meshes for curve usage */
 	virtual void TestSkeletonCurveNamesForUse(const TSharedRef<IEditableSkeleton>& InEditableSkeleton) const;
 
-	/** Apply Compression to list of animations */
-	virtual void ApplyCompression(TArray<TWeakObjectPtr<class UAnimSequence>>& AnimSequences);
+	/** Apply Compression to list of animations and optionally asks to pick an overrides to the bone compression settings */
+	virtual void ApplyCompression(TArray<TWeakObjectPtr<class UAnimSequence>>& AnimSequences, bool bPickBoneSettingsOverride);
 
 	/** Export to FBX files of the list of animations */
 	virtual bool ExportToFBX(TArray<TWeakObjectPtr<class UAnimSequence>>& AnimSequences, USkeletalMesh* SkeletalMesh);
@@ -283,19 +291,19 @@ public:
 	virtual class IPersonaEditorModeManager* CreatePersonaEditorModeManager();
 
 	/** Delegate used to query whether recording is active */
-	FIsRecordingActive& OnIsRecordingActive() { return IsRecordingActiveDelegate; }
+	virtual FIsRecordingActive& OnIsRecordingActive() { return IsRecordingActiveDelegate; }
 
 	/** Delegate used to start recording animation */
-	FRecord& OnRecord() { return RecordDelegate; }
+	virtual FRecord& OnRecord() { return RecordDelegate; }
 
 	/** Delegate used to stop recording animation */
-	FStopRecording& OnStopRecording() { return StopRecordingDelegate; }
+	virtual FStopRecording& OnStopRecording() { return StopRecordingDelegate; }
 
 	/** Delegate used to get the currently recording animation */
-	FGetCurrentRecording& OnGetCurrentRecording() { return GetCurrentRecordingDelegate; }
+	virtual FGetCurrentRecording& OnGetCurrentRecording() { return GetCurrentRecordingDelegate; }
 
 	/** Delegate used to get the currently recording animation time */
-	FGetCurrentRecordingTime& OnGetCurrentRecordingTime() { return GetCurrentRecordingTimeDelegate; }
+	virtual FGetCurrentRecordingTime& OnGetCurrentRecordingTime() { return GetCurrentRecordingTimeDelegate; }
 
 	/** Delegate broadcast when a preview scene is created */
 	virtual FOnPreviewSceneCreated& OnPreviewSceneCreated() { return OnPreviewSceneCreatedDelegate; }
diff --git a/Engine/Source/Editor/Persona/Public/PersonaTabs.h b/Engine/Source/Editor/Persona/Public/PersonaTabs.h
--- a/Engine/Source/Editor/Persona/Public/PersonaTabs.h
+++ b/Engine/Source/Editor/Persona/Public/PersonaTabs.h
@@ -38,6 +38,7 @@ struct PERSONA_API FPersonaTabs
 	static const FName SkeletonSlotGroupNamesID;
 	static const FName CurveNameManagerID;
 	static const FName BlendProfileManagerID;
+	static const FName AnimMontageSectionsID;
 
 	// Advanced Preview Scene
 	static const FName AdvancedPreviewSceneSettingsID;
diff --git a/Engine/Source/Editor/Persona/Public/PersonaUtils.h b/Engine/Source/Editor/Persona/Public/PersonaUtils.h
--- a/Engine/Source/Editor/Persona/Public/PersonaUtils.h
+++ b/Engine/Source/Editor/Persona/Public/PersonaUtils.h
@@ -3,10 +3,12 @@
 #pragma once
 
 #include "CoreMinimal.h"
+#include "Framework/SlateDelegates.h"
 
 class UAnimInstance;
 class USceneComponent;
 class UAnimBlueprint;
+class SWidget;
 
 namespace PersonaUtils
 {
@@ -57,5 +59,7 @@ PERSONA_API int32 CopyPropertiesToCDO(UAnimInstance* InAnimInstance, const FCopy
 
 /** Set the correct object to debug depending on whether we are running a sub-layer etc. */
 PERSONA_API void SetObjectBeingDebugged(UAnimBlueprint* InAnimBlueprint, UAnimInstance* InAnimInstance);
+/** Make an add button widget, as seen in the animation timeline */
+PERSONA_API TSharedRef<SWidget> MakeTrackButton(FText HoverText, FOnGetContent MenuContent, const TAttribute<bool>& HoverState);
 
 }
diff --git a/Engine/Source/Editor/PropertyEditor/Private/UserInterface/PropertyEditor/SPropertyMenuAssetPicker.cpp b/Engine/Source/Editor/PropertyEditor/Private/UserInterface/PropertyEditor/SPropertyMenuAssetPicker.cpp
--- a/Engine/Source/Editor/PropertyEditor/Private/UserInterface/PropertyEditor/SPropertyMenuAssetPicker.cpp
+++ b/Engine/Source/Editor/PropertyEditor/Private/UserInterface/PropertyEditor/SPropertyMenuAssetPicker.cpp
@@ -31,6 +31,8 @@ void SPropertyMenuAssetPicker::Construct( const FArguments& InArgs )
 	OnSet = InArgs._OnSet;
 	OnClose = InArgs._OnClose;
 
+	const bool bForceShowEngineContent = PropertyHandle ? PropertyHandle->HasMetaData(TEXT("ForceShowEngineContent")) : false;
+
 	const bool bInShouldCloseWindowAfterMenuSelection = true;
 	const bool bCloseSelfOnly = true;
 	
@@ -149,6 +151,8 @@ void SPropertyMenuAssetPicker::Construct( const FArguments& InArgs )
 		AssetPickerConfig.PropertyHandle = PropertyHandle;
 		// Populate the additional referencing assets with the Owner asset data
 		AssetPickerConfig.AdditionalReferencingAssets = OwnerAssetArray;
+		// Force show engine content if meta data says so
+		AssetPickerConfig.bForceShowEngineContent = bForceShowEngineContent;
 
 		MenuContent =
 			SNew(SBox)
diff --git a/Engine/Source/Editor/SequenceRecorder/Private/AnimationRecorder.cpp b/Engine/Source/Editor/SequenceRecorder/Private/AnimationRecorder.cpp
--- a/Engine/Source/Editor/SequenceRecorder/Private/AnimationRecorder.cpp
+++ b/Engine/Source/Editor/SequenceRecorder/Private/AnimationRecorder.cpp
@@ -70,16 +70,18 @@ void FAnimationRecorder::SetSampleRateAndLength(float SampleRateHz, float Length
 	}
 }
 
-bool FAnimationRecorder::SetAnimCompressionScheme(TSubclassOf<UAnimCompress> SchemeClass)
+bool FAnimationRecorder::SetAnimCompressionScheme(UAnimBoneCompressionSettings* Settings)
 {
 	if (AnimationObject)
 	{
-		UAnimCompress* const SchemeObject = NewObject<UAnimCompress>(GetTransientPackage(), SchemeClass);
-		if (SchemeObject)
+		if (Settings == nullptr)
 		{
-			AnimationObject->CompressionScheme = SchemeObject;
-			return true;
+			// The caller has not supplied a settings asset, use our default value
+			Settings = FAnimationUtils::GetDefaultAnimationRecorderBoneCompressionSettings();
 		}
+
+		AnimationObject->BoneCompressionSettings = Settings;
+		return true;
 	}
 
 	return false;
@@ -216,6 +218,7 @@ void FAnimationRecorder::StartRecord(USkeletalMeshComponent* Component, UAnimSeq
 	AnimationObject = InAnimationObject;
 
 	AnimationObject->RecycleAnimSequence();
+	AnimationObject->BoneCompressionSettings = FAnimationUtils::GetDefaultAnimationRecorderBoneCompressionSettings();
 
 	GetBoneTransforms(Component, PreviousSpacesBases);
 	PreviousAnimCurves = Component->GetAnimationCurves();
@@ -815,7 +818,6 @@ void FAnimRecorderInstance::InitInternal(USkeletalMeshComponent* InComponent, co
 	Recorder->bRecordLocalToWorld = Settings.bRecordInWorldSpace;
 	Recorder->InterpMode = Settings.InterpMode;
 	Recorder->TangentMode = Settings.TangentMode;
-	Recorder->SetAnimCompressionScheme(UAnimCompress_BitwiseCompressOnly::StaticClass());
 	Recorder->bAutoSaveAsset = Settings.bAutoSaveAsset;
 	Recorder->bRemoveRootTransform = Settings.bRemoveRootAnimation;
 	Recorder->bCheckDeltaTimeAtBeginning = Settings.bCheckDeltaTimeAtBeginning;
diff --git a/Engine/Source/Editor/SequenceRecorder/Public/AnimationRecorder.h b/Engine/Source/Editor/SequenceRecorder/Public/AnimationRecorder.h
--- a/Engine/Source/Editor/SequenceRecorder/Public/AnimationRecorder.h
+++ b/Engine/Source/Editor/SequenceRecorder/Public/AnimationRecorder.h
@@ -11,6 +11,7 @@
 #include "Animation/AnimNotifyQueue.h"
 #include "Serializers/MovieSceneAnimationSerialization.h"
 
+class UAnimBoneCompressionSettings;
 class UAnimNotify;
 class UAnimNotifyState;
 class UAnimSequence;
@@ -74,7 +75,7 @@ public:
 	/** Sets a new sample rate & max length for this recorder. Don't call while recording. */
 	void SetSampleRateAndLength(float SampleRateHz, float LengthInMinutes);
 
-	bool SetAnimCompressionScheme(TSubclassOf<class UAnimCompress> SchemeClass);
+	bool SetAnimCompressionScheme(UAnimBoneCompressionSettings* Settings);
 
 	const FTransform& GetInitialRootTransform() const { return InitialRootTransform; }
 
diff --git a/Engine/Source/Editor/SequenceRecorder/SequenceRecorder.Build.cs b/Engine/Source/Editor/SequenceRecorder/SequenceRecorder.Build.cs
--- a/Engine/Source/Editor/SequenceRecorder/SequenceRecorder.Build.cs
+++ b/Engine/Source/Editor/SequenceRecorder/SequenceRecorder.Build.cs
@@ -40,7 +40,6 @@ namespace UnrealBuildTool.Rules
 					"EditorStyle",
 					"Projects",
 					"LevelEditor",
-					"Persona",
 					"WorkspaceMenuStructure",
 					"PropertyEditor",
 					"MovieScene",
@@ -56,10 +55,18 @@ namespace UnrealBuildTool.Rules
                 }
                 );
 
-			DynamicallyLoadedModuleNames.AddRange(
+            PrivateIncludePathModuleNames.AddRange(
+                new string[]
+                {
+                    "Persona",
+                }
+                );
+
+            DynamicallyLoadedModuleNames.AddRange(
 				new string[]
 				{
-				}
+                    "Persona",
+                }
 				);
 
 			if (Target.Platform == UnrealTargetPlatform.Win64 || Target.Platform == UnrealTargetPlatform.Win32)
diff --git a/Engine/Source/Editor/Sequencer/Private/DisplayNodes/SequencerDisplayNode.cpp b/Engine/Source/Editor/Sequencer/Private/DisplayNodes/SequencerDisplayNode.cpp
--- a/Engine/Source/Editor/Sequencer/Private/DisplayNodes/SequencerDisplayNode.cpp
+++ b/Engine/Source/Editor/Sequencer/Private/DisplayNodes/SequencerDisplayNode.cpp
@@ -582,6 +582,22 @@ TSharedPtr<FSequencerObjectBindingNode> FSequencerDisplayNode::FindParentObjectB
 	return nullptr;
 }
 
+TSharedPtr<FSequencerTrackNode> FSequencerDisplayNode::FindParentTrackNode() const
+{
+	TSharedPtr<FSequencerDisplayNode> CurrentParentNode = GetParent();
+	while (CurrentParentNode.IsValid())
+	{
+		if (CurrentParentNode->GetType() == ESequencerNode::Track)
+		{
+			return StaticCastSharedPtr<FSequencerTrackNode>(CurrentParentNode);
+		}
+		CurrentParentNode = CurrentParentNode->GetParent();
+	}
+
+	return nullptr;
+}
+
+
 FGuid FSequencerDisplayNode::GetObjectGuid() const
 {
 	TSharedPtr<FSequencerObjectBindingNode> ObjectBindingNode = FindParentObjectBindingNode();
@@ -770,6 +786,12 @@ bool FSequencerDisplayNode::IsDimmed() const
 	return bDimLabel;
 }
 
+FSlateFontInfo FSequencerDisplayNode::GetDisplayNameFont() const
+{
+	FSlateFontInfo NodeFont = FEditorStyle::GetFontStyle("Sequencer.AnimationOutliner.RegularFont");
+	return NodeFont;
+}
+
 FLinearColor FSequencerDisplayNode::GetDisplayNameColor() const
 {
 	return IsDimmed() ? FLinearColor(0.6f, 0.6f, 0.6f, 0.6f) : FLinearColor(1.0f, 1.0f, 1.0f, 1.0f);
@@ -1341,6 +1363,7 @@ TSharedPtr<SWidget> FSequencerDisplayNode::GenerateCurveEditorTreeWidget(const F
 			[
 				SNew(STextBlock)
 				.Text(this, &FSequencerDisplayNode::GetDisplayName)
+				.Font(this, &FSequencerDisplayNode::GetDisplayNameFont)
 				.HighlightText_Static(SequencerNodeConstants::GetCurveEditorHighlightText, InCurveEditor)
 				.ToolTipText(this, &FSequencerDisplayNode::GetDisplayNameToolTipText)
 			];
diff --git a/Engine/Source/Editor/Sequencer/Private/DisplayNodes/SequencerDisplayNode.h b/Engine/Source/Editor/Sequencer/Private/DisplayNodes/SequencerDisplayNode.h
--- a/Engine/Source/Editor/Sequencer/Private/DisplayNodes/SequencerDisplayNode.h
+++ b/Engine/Source/Editor/Sequencer/Private/DisplayNodes/SequencerDisplayNode.h
@@ -15,6 +15,7 @@ class FSequencer;
 class FSequencerNodeTree;
 class FSequencerDisplayNodeDragDropOp;
 class FSequencerObjectBindingNode;
+class FSequencerTrackNode;
 class IKeyArea;
 class ISequencerTrackEditor;
 class SSequencerTreeViewRow;
@@ -108,6 +109,13 @@ public:
 	 */
 	TSharedPtr<FSequencerObjectBindingNode> FindParentObjectBindingNode() const;
 
+	/**
+	 * Finds any track node above this node in the hierarchy
+	 *
+	 * @return the parent track node, or nullptr if no track nodeis found
+	 */
+	TSharedPtr<FSequencerTrackNode> FindParentTrackNode() const;
+
 	/** 
 	 * Finds this display node's closest parent object binding GUID, or an empty FGuid if it there is none
 	 */
@@ -170,6 +178,11 @@ public:
 	*/
 	virtual FLinearColor GetDisplayNameColor() const;
 
+	/**
+	*@return The font used to draw the display name.
+	*/
+	virtual FSlateFontInfo GetDisplayNameFont() const;
+
 	/**
 	 * @return the text to display for the tool tip for the display name. 
 	 */
diff --git a/Engine/Source/Editor/Sequencer/Private/DisplayNodes/SequencerObjectBindingNode.cpp b/Engine/Source/Editor/Sequencer/Private/DisplayNodes/SequencerObjectBindingNode.cpp
--- a/Engine/Source/Editor/Sequencer/Private/DisplayNodes/SequencerObjectBindingNode.cpp
+++ b/Engine/Source/Editor/Sequencer/Private/DisplayNodes/SequencerObjectBindingNode.cpp
@@ -150,6 +150,7 @@ void FSequencerObjectBindingNode::BuildContextMenu(FMenuBuilder& MenuBuilder)
 	ISequencerModule& SequencerModule = FModuleManager::GetModuleChecked<ISequencerModule>("Sequencer");
 
 	UObject* BoundObject = GetSequencer().FindSpawnedObjectOrTemplate(ObjectBinding);
+	const UClass* ObjectClass = GetClassForObjectBinding();
 
 	TSharedRef<FUICommandList> CommandList(new FUICommandList);
 	TSharedPtr<FExtender> Extender = SequencerModule.GetObjectBindingContextMenuExtensibilityManager()->GetAllExtenders(CommandList, TArrayBuilder<UObject*>().Add(BoundObject));
@@ -237,7 +238,6 @@ void FSequencerObjectBindingNode::BuildContextMenu(FMenuBuilder& MenuBuilder)
 		}
 		else
 		{
-			const UClass* ObjectClass = GetClassForObjectBinding();
 			
 			if (ObjectClass->IsChildOf(AActor::StaticClass()))
 			{
@@ -297,6 +297,12 @@ void FSequencerObjectBindingNode::BuildContextMenu(FMenuBuilder& MenuBuilder)
 	MenuBuilder.EndSection();
 
 	GetSequencer().BuildCustomContextMenuForGuid(MenuBuilder, ObjectBinding);
+	TArray<FGuid> ObjectBindings;
+	ObjectBindings.Add(ObjectBinding);
+	for (const TSharedPtr<ISequencerTrackEditor>& TrackEditor : GetSequencer().GetTrackEditors())
+	{
+		TrackEditor->BuildObjectBindingContextMenu(MenuBuilder, ObjectBindings, ObjectClass);
+	}
 
 	FSequencerDisplayNode::BuildContextMenu(MenuBuilder);
 }
diff --git a/Engine/Source/Editor/Sequencer/Private/DisplayNodes/SequencerSectionCategoryNode.cpp b/Engine/Source/Editor/Sequencer/Private/DisplayNodes/SequencerSectionCategoryNode.cpp
--- a/Engine/Source/Editor/Sequencer/Private/DisplayNodes/SequencerSectionCategoryNode.cpp
+++ b/Engine/Source/Editor/Sequencer/Private/DisplayNodes/SequencerSectionCategoryNode.cpp
@@ -68,3 +68,33 @@ void FSequencerSectionCategoryNode::SetDisplayName(const FText& NewDisplayName)
 {
 	check(false);
 }
+
+FSlateFontInfo FSequencerSectionCategoryNode::GetDisplayNameFont() const
+{
+	bool bAllAnimated = false;
+	for (TSharedRef<FSequencerDisplayNode> ChildNode : GetChildNodes())
+	{
+		if (ChildNode->GetType() == ESequencerNode::KeyArea)
+		{
+			const FSequencerSectionKeyAreaNode KeyAreaNode = static_cast<const FSequencerSectionKeyAreaNode&>(ChildNode.Get());
+			for (const TSharedRef<IKeyArea>& KeyArea : KeyAreaNode.GetAllKeyAreas())
+			{
+				FMovieSceneChannel* Channel = KeyArea->ResolveChannel();
+				if (!Channel || Channel->GetNumKeys() == 0)
+				{
+					return FSequencerDisplayNode::GetDisplayNameFont();
+				}
+				else
+				{
+					bAllAnimated = true;
+				}
+			}
+		}
+	}
+	if (bAllAnimated == true)
+	{
+		return FEditorStyle::GetFontStyle("Sequencer.AnimationOutliner.ItalicFont");
+	}
+	return FSequencerDisplayNode::GetDisplayNameFont();
+}
+
diff --git a/Engine/Source/Editor/Sequencer/Private/DisplayNodes/SequencerSectionCategoryNode.h b/Engine/Source/Editor/Sequencer/Private/DisplayNodes/SequencerSectionCategoryNode.h
--- a/Engine/Source/Editor/Sequencer/Private/DisplayNodes/SequencerSectionCategoryNode.h
+++ b/Engine/Source/Editor/Sequencer/Private/DisplayNodes/SequencerSectionCategoryNode.h
@@ -38,5 +38,6 @@ public:
 	virtual FNodePadding GetNodePadding() const override;
 	virtual ESequencerNode::Type GetType() const override;
 	virtual void SetDisplayName(const FText& NewDisplayName) override;
+	virtual FSlateFontInfo GetDisplayNameFont() const override;
 
 };
