Fix LOCTEXT issues arising from Dev-Anim copy-up

diff --git a/Engine/Plugins/Experimental/ControlRig/Source/ControlRigEditor/Private/EditMode/ControlRigEditMode.cpp b/Engine/Plugins/Experimental/ControlRig/Source/ControlRigEditor/Private/EditMode/ControlRigEditMode.cpp
--- a/Engine/Plugins/Experimental/ControlRig/Source/ControlRigEditor/Private/EditMode/ControlRigEditMode.cpp
+++ b/Engine/Plugins/Experimental/ControlRig/Source/ControlRigEditor/Private/EditMode/ControlRigEditMode.cpp
@@ -1400,7 +1400,7 @@ void FControlRigEditMode::OnRigElementChanged(FRigHierarchyContainer* Container,
 
 void FControlRigEditMode::OnControlModified(IControlRigManipulatable* Subject, const FRigControl& Control, EControlRigSetKey InSetKey)
 {
-	FScopedTransaction ScopedTransaction(LOCTEXT("MoveControlTransaction", "Modify Control"),!GIsTransacting && InSetKey != EControlRigSetKey::Never);
+	FScopedTransaction ScopedTransaction(LOCTEXT("ModifyControlTransaction", "Modify Control"),!GIsTransacting && InSetKey != EControlRigSetKey::Never);
 	ControlProxy->Modify();
 	ControlProxy->ProxyChanged(Control.Name);
 	RecalcPivotTransform();
diff --git a/Engine/Source/Developer/AssetTools/Private/AssetTypeActions/AssetTypeActions_AnimBoneCompressionSettings.cpp b/Engine/Source/Developer/AssetTools/Private/AssetTypeActions/AssetTypeActions_AnimBoneCompressionSettings.cpp
--- a/Engine/Source/Developer/AssetTools/Private/AssetTypeActions/AssetTypeActions_AnimBoneCompressionSettings.cpp
+++ b/Engine/Source/Developer/AssetTools/Private/AssetTypeActions/AssetTypeActions_AnimBoneCompressionSettings.cpp
@@ -111,3 +111,5 @@ void FAssetTypeActions_AnimBoneCompressionSettings::ExecuteCompression(TWeakObje
 		AnimSeq->RequestSyncAnimRecompression(false);
 	}
 }
+
+#undef LOCTEXT_NAMESPACE
\ No newline at end of file
diff --git a/Engine/Source/Editor/Persona/Private/AnimTimeline/AnimTimelineTrack_MontagePanel.cpp b/Engine/Source/Editor/Persona/Private/AnimTimeline/AnimTimelineTrack_MontagePanel.cpp
--- a/Engine/Source/Editor/Persona/Private/AnimTimeline/AnimTimelineTrack_MontagePanel.cpp
+++ b/Engine/Source/Editor/Persona/Private/AnimTimeline/AnimTimelineTrack_MontagePanel.cpp
@@ -21,7 +21,7 @@ namespace
 }
 
 FAnimTimelineTrack_MontagePanel::FAnimTimelineTrack_MontagePanel(const TSharedRef<FAnimModel>& InModel)
-	: FAnimTimelineTrack(FText::GetEmpty(), LOCTEXT("SlotTooltip", "Montage slot"), InModel)
+	: FAnimTimelineTrack(FText::GetEmpty(), LOCTEXT("MontageSlotTooltip", "Montage slot"), InModel)
 {
 	UAnimMontage* AnimMontage = CastChecked<UAnimMontage>(GetModel()->GetAnimSequenceBase());
 
@@ -143,7 +143,7 @@ TSharedRef<SWidget> FAnimTimelineTrack_MontagePanel::BuildMontageSlotSubMenu(int
 		{
 			MenuBuilder.AddSubMenu(
 				LOCTEXT("SlotName", "Slot name"),
-				LOCTEXT("SlotTooltip", "Change the name of this slot"),
+				LOCTEXT("SlotNameTooltip", "Change the name of this slot"),
 				FNewMenuDelegate::CreateSP(this, &FAnimTimelineTrack_MontagePanel::BuildSlotNameMenu, InSlotIndex)
 			);
 
diff --git a/Engine/Source/Editor/Persona/Private/SAnimMontageSectionsPanel.cpp b/Engine/Source/Editor/Persona/Private/SAnimMontageSectionsPanel.cpp
--- a/Engine/Source/Editor/Persona/Private/SAnimMontageSectionsPanel.cpp
+++ b/Engine/Source/Editor/Persona/Private/SAnimMontageSectionsPanel.cpp
@@ -437,7 +437,7 @@ TSharedRef<SWidget> SAnimMontageSectionsPanel::OnGetSectionMenuContent(int32 Sec
 			FText SectionNameText = FText::FromName(Montage->CompositeSections[NextSectionIdx].SectionName);
 			MenuBuilder.AddMenuEntry(
 				SectionNameText,
-				FText::Format(LOCTEXT("RemoveLinkTooltip", "Set {0} to be the next section"), SectionNameText),
+				FText::Format(LOCTEXT("SetNextSectionTooltip", "Set {0} to be the next section"), SectionNameText),
 				FSlateIcon(),
 				FUIAction(FExecuteAction::CreateSP(this, &SAnimMontageSectionsPanel::SetNextSectionIndex, SectionIdx, NextSectionIdx))
 			);
diff --git a/Engine/Source/Editor/Persona/Private/SAnimTimingPanel.cpp b/Engine/Source/Editor/Persona/Private/SAnimTimingPanel.cpp
--- a/Engine/Source/Editor/Persona/Private/SAnimTimingPanel.cpp
+++ b/Engine/Source/Editor/Persona/Private/SAnimTimingPanel.cpp
@@ -406,7 +406,7 @@ void FTimingRelevantElement_Section::GetDescriptionItems(TMap<FString, FText>& I
 	Items.Add(LOCTEXT("SectionName", "Name").ToString(), FText::FromName(Section.SectionName));
 	if(GetDefault<UPersonaOptions>()->TimelineDisplayFormat == EFrameNumberDisplayFormats::Frames)
 	{
-		Items.Add(LOCTEXT("SectionTriggerFrame", "Trigger Frame").ToString(), FText::Format(LOCTEXT("SectionTriggerTimeValue", "{0}"), FText::AsNumber(Montage->GetFrameAtTime(Section.GetTime()))));
+		Items.Add(LOCTEXT("SectionTriggerFrame", "Trigger Frame").ToString(), FText::Format(LOCTEXT("SectionTriggerFrameValue", "{0}"), FText::AsNumber(Montage->GetFrameAtTime(Section.GetTime()))));
 	}
 	else
 	{
