Fix remove cloth data crash

diff --git a/Engine/Source/Editor/Persona/Private/PersonaMeshDetails.cpp b/Engine/Source/Editor/Persona/Private/PersonaMeshDetails.cpp
--- a/Engine/Source/Editor/Persona/Private/PersonaMeshDetails.cpp
+++ b/Engine/Source/Editor/Persona/Private/PersonaMeshDetails.cpp
@@ -5561,11 +5561,14 @@ FReply FPersonaMeshDetails::OnRemoveClothingAssetClicked(int32 AssetIndex, IDeta
 			{
 				for(FSkeletalMeshLODModel& LodModel : MeshResource->LODModels)
 				{
-					for(FSkelMeshSection& Section : LodModel.Sections)
+					for(int32 SectionIndex = 0; SectionIndex < LodModel.Sections.Num(); ++SectionIndex)
 					{
+						FSkelMeshSection& Section = LodModel.Sections[SectionIndex];
 						if(Section.CorrespondClothAssetIndex > AssetIndex)
 						{
 							--Section.CorrespondClothAssetIndex;
+							//Keep the user section data (build source data) in sync
+							SetSkelMeshSourceSectionUserData(LodModel, SectionIndex, Section.OriginalDataSectionIndex);
 						}
 					}
 				}
