Added blendspace debug widgets to blendspace nodes

diff --git a/Engine/Source/Editor/AnimationBlueprintEditor/Private/AnimationGraphFactory.cpp b/Engine/Source/Editor/AnimationBlueprintEditor/Private/AnimationGraphFactory.cpp
--- a/Engine/Source/Editor/AnimationBlueprintEditor/Private/AnimationGraphFactory.cpp
+++ b/Engine/Source/Editor/AnimationBlueprintEditor/Private/AnimationGraphFactory.cpp
@@ -8,6 +8,7 @@
 #include "AnimGraphNode_SequencePlayer.h"
 #include "AnimGraphNode_StateMachineBase.h"
 #include "AnimGraphNode_LayeredBoneBlend.h"
+#include "AnimGraphNode_BlendSpaceBase.h"
 #include "AnimStateNode.h"
 #include "AnimStateEntryNode.h"
 #include "AnimStateConduitNode.h"
@@ -25,6 +26,7 @@
 #include "AnimationNodes/SGraphNodeAnimationResult.h"
 #include "AnimationNodes/SGraphNodeStateMachineInstance.h"
 #include "AnimationNodes/SGraphNodeLayeredBoneBlend.h"
+#include "AnimationNodes/SGraphNodeBlendSpacePlayer.h"
 #include "AnimationPins/SGraphPinPose.h"
 
 #include "AnimGraphConnectionDrawingPolicy.h"
@@ -52,6 +54,10 @@ TSharedPtr<class SGraphNode> FAnimationGraphNodeFactory::CreateNode(class UEdGra
 		{
 			return SNew(SGraphNodeLayeredBoneBlend, LayeredBlend);
 		}
+		else if (UAnimGraphNode_BlendSpaceBase* BlendSpacePlayer = Cast<UAnimGraphNode_BlendSpaceBase>(InNode))
+		{
+			return SNew(SGraphNodeBlendSpacePlayer, BlendSpacePlayer);
+		}
 		else
 		{
 			return SNew(SAnimationGraphNode, BaseAnimNode);
diff --git a/Engine/Source/Editor/AnimationBlueprintEditor/Private/AnimationNodes/SGraphNodeBlendSpacePlayer.cpp b/Engine/Source/Editor/AnimationBlueprintEditor/Private/AnimationNodes/SGraphNodeBlendSpacePlayer.cpp
new file mode 100644
--- /dev/null
+++ b/Engine/Source/Editor/AnimationBlueprintEditor/Private/AnimationNodes/SGraphNodeBlendSpacePlayer.cpp
@@ -0,0 +1,116 @@
+// Copyright 1998-2019 Epic Games, Inc. All Rights Reserved.
+
+#include "AnimationNodes/SGraphNodeBlendSpacePlayer.h"
+#include "Animation/AnimBlueprintGeneratedClass.h"
+#include "Kismet2/KismetDebugUtilities.h"
+#include "Kismet2/BlueprintEditorUtils.h"
+#include "AnimNodes/AnimNode_BlendSpacePlayer.h"
+#include "AnimGraphNode_BlendSpaceBase.h"
+#include "PersonaModule.h"
+#include "Widgets/Layout/SBox.h"
+#include "Modules/ModuleManager.h"
+
+PRAGMA_DISABLE_OPTIMIZATION
+
+void SGraphNodeBlendSpacePlayer::Construct(const FArguments& InArgs, UAnimGraphNode_Base* InNode)
+{
+	this->GraphNode = InNode;
+
+	this->SetCursor(EMouseCursor::CardinalCross);
+
+	this->UpdateGraphNode();
+
+	SAnimationGraphNode::Construct(SAnimationGraphNode::FArguments(), InNode);
+
+	RegisterActiveTimer(1.0f / 60.0f, FWidgetActiveTimerDelegate::CreateLambda([this](double InCurrentTime, float InDeltaTime)
+	{
+		GetBlendSpaceInfo(CachedBlendSpace, CachedPosition);
+		return EActiveTimerReturnType::Continue;
+	}));
+}
+
+void SGraphNodeBlendSpacePlayer::CreateBelowWidgetControls(TSharedPtr<SVerticalBox> MainBox)
+{
+	FPersonaModule& PersonaModule = FModuleManager::LoadModuleChecked<FPersonaModule>("Persona");
+
+	MainBox->AddSlot()
+	.AutoHeight()
+	.VAlign(VAlign_Fill)
+	.Padding(0.0f)
+	[
+		SNew(SBox)
+		.HeightOverride_Lambda([WeakMainBox = TWeakPtr<SVerticalBox>(MainBox)]()
+		{
+			if(TSharedPtr<SVerticalBox> LocalMainBox = WeakMainBox.Pin())
+			{
+				float Size = LocalMainBox->GetDesiredSize().X;
+				return FMath::Min(Size, 100.0f);
+			}
+
+			return 0.0f;
+		})
+		.Visibility(this, &SGraphNodeBlendSpacePlayer::GetBlendSpaceVisibility)
+		[
+			PersonaModule.CreateBlendSpacePreviewWidget(
+			MakeAttributeLambda([this]()
+			{
+				return CachedBlendSpace.Get();
+			}),
+			MakeAttributeLambda([this]()
+			{
+				return CachedPosition;
+			}))
+		]
+	];
+}
+
+EVisibility SGraphNodeBlendSpacePlayer::GetBlendSpaceVisibility() const
+{
+	if (UBlueprint* Blueprint = FBlueprintEditorUtils::FindBlueprintForNode(GraphNode))
+	{
+		if (UProperty* Property = FKismetDebugUtilities::FindClassPropertyForNode(Blueprint, GraphNode))
+		{
+			if (UObject* ActiveObject = Blueprint->GetObjectBeingDebugged())
+			{
+				return EVisibility::Visible;
+			}
+		}
+	}
+
+	return EVisibility::Collapsed;
+}
+
+bool SGraphNodeBlendSpacePlayer::GetBlendSpaceInfo(TWeakObjectPtr<const UBlendSpaceBase>& OutBlendSpace, FVector& OutPosition) const
+{
+	if (UBlueprint* Blueprint = FBlueprintEditorUtils::FindBlueprintForNode(GraphNode))
+	{
+		if (UObject* ActiveObject = Blueprint->GetObjectBeingDebugged())
+		{
+			if (UAnimGraphNode_BlendSpaceBase* VisualBlendSpacePlayer = Cast<UAnimGraphNode_BlendSpaceBase>(GraphNode))
+			{
+				if (UAnimBlueprintGeneratedClass* Class = Cast<UAnimBlueprintGeneratedClass>((UObject*)ActiveObject->GetClass()))
+				{
+					if(int32* NodeIndexPtr = Class->GetAnimBlueprintDebugData().NodePropertyToIndexMap.Find(TWeakObjectPtr<UAnimGraphNode_Base>(Cast<UAnimGraphNode_Base>(GraphNode))))
+					{
+						int32 AnimNodeIndex = *NodeIndexPtr;
+						// reverse node index temporarily because of a bug in NodeGuidToIndexMap
+						AnimNodeIndex = Class->AnimNodeProperties.Num() - AnimNodeIndex - 1;
+
+						if (FAnimBlueprintDebugData::FBlendSpacePlayerRecord* DebugInfo = Class->GetAnimBlueprintDebugData().BlendSpacePlayerRecordsThisFrame.FindByPredicate([AnimNodeIndex](const FAnimBlueprintDebugData::FBlendSpacePlayerRecord& InRecord){ return InRecord.NodeID == AnimNodeIndex; }))
+						{
+							OutBlendSpace = DebugInfo->BlendSpace.Get();
+							OutPosition = FVector(DebugInfo->PositionX, DebugInfo->PositionY, DebugInfo->PositionZ);
+							return true;
+						}
+					}
+				}
+			}
+		}
+	}
+
+	OutBlendSpace = nullptr;
+	OutPosition = FVector::ZeroVector;
+	return false;
+}
+
+PRAGMA_ENABLE_OPTIMIZATION
\ No newline at end of file
diff --git a/Engine/Source/Editor/AnimationBlueprintEditor/Private/AnimationNodes/SGraphNodeBlendSpacePlayer.h b/Engine/Source/Editor/AnimationBlueprintEditor/Private/AnimationNodes/SGraphNodeBlendSpacePlayer.h
new file mode 100644
--- /dev/null
+++ b/Engine/Source/Editor/AnimationBlueprintEditor/Private/AnimationNodes/SGraphNodeBlendSpacePlayer.h
@@ -0,0 +1,33 @@
+// Copyright 1998-2019 Epic Games, Inc. All Rights Reserved.
+
+#pragma once
+
+#include "CoreMinimal.h"
+#include "Layout/Visibility.h"
+#include "Widgets/DeclarativeSyntaxSupport.h"
+#include "SNodePanel.h"
+#include "AnimationNodes/SAnimationGraphNode.h"
+#include "Animation/BlendSpaceBase.h"
+
+class SVerticalBox;
+class UAnimGraphNode_Base;
+
+class SGraphNodeBlendSpacePlayer : public SAnimationGraphNode
+{
+public:
+	SLATE_BEGIN_ARGS(SGraphNodeBlendSpacePlayer){}
+	SLATE_END_ARGS()
+
+	void Construct(const FArguments& InArgs, UAnimGraphNode_Base* InNode);
+
+	// SGraphNode interface
+	virtual void CreateBelowWidgetControls(TSharedPtr<SVerticalBox> MainBox) override;
+	// End of SGraphNode interface
+
+protected:
+	EVisibility GetBlendSpaceVisibility() const;
+	bool GetBlendSpaceInfo(TWeakObjectPtr<const UBlendSpaceBase>& OutBlendSpace, FVector& OutPosition) const;
+
+	TWeakObjectPtr<const UBlendSpaceBase> CachedBlendSpace;
+	FVector CachedPosition;
+};
diff --git a/Engine/Source/Editor/Persona/Private/PersonaModule.cpp b/Engine/Source/Editor/Persona/Private/PersonaModule.cpp
--- a/Engine/Source/Editor/Persona/Private/PersonaModule.cpp
+++ b/Engine/Source/Editor/Persona/Private/PersonaModule.cpp
@@ -86,6 +86,7 @@
 #include "AnimPreviewInstance.h"
 #include "ISequenceRecorder.h"
 #include "SkinWeightProfileCustomization.h"
+#include "SAnimationBlendSpaceGridWidget.h"
 
 IMPLEMENT_MODULE( FPersonaModule, Persona );
 
@@ -343,6 +344,18 @@ TSharedRef<class FWorkflowTabFactory> FPersonaModule::CreateSkeletonSlotNamesTab
 	return MakeShareable(new FSkeletonSlotNamesSummoner(InHostingApp, InEditableSkeleton, InOnPostUndo, InOnObjectSelected));
 }
 
+TSharedRef<SWidget> FPersonaModule::CreateBlendSpacePreviewWidget(TAttribute<const UBlendSpaceBase*> InBlendSpace, TAttribute<FVector> InPosition) const
+{
+	return
+		SNew(SBlendSpaceGridWidget)
+		.Cursor(EMouseCursor::Crosshairs)
+		.BlendSpaceBase(InBlendSpace)
+		.Position(InPosition)
+		.ReadOnly(true)
+		.ShowAxisLabels(false)
+		.ShowSettingsButtons(false);
+}
+
 TSharedRef<SWidget> FPersonaModule::CreateEditorWidgetForAnimDocument(const TSharedRef<class FWorkflowCentricApplication>& InHostingApp, UObject* InAnimAsset, const FAnimDocumentArgs& InArgs, FString& OutDocumentLink)
 {
 	TSharedPtr<SWidget> Result = SNullWidget::NullWidget;
diff --git a/Engine/Source/Editor/Persona/Private/SAnimationBlendSpaceGridWidget.cpp b/Engine/Source/Editor/Persona/Private/SAnimationBlendSpaceGridWidget.cpp
--- a/Engine/Source/Editor/Persona/Private/SAnimationBlendSpaceGridWidget.cpp
+++ b/Engine/Source/Editor/Persona/Private/SAnimationBlendSpaceGridWidget.cpp
@@ -38,14 +38,19 @@
 
 void SBlendSpaceGridWidget::Construct(const FArguments& InArgs)
 {
-	BlendSpace = InArgs._BlendSpaceBase;
+	BlendSpaceBase = InArgs._BlendSpaceBase;
+	PreviousBlendSpaceBase = BlendSpaceBase.Get();
+	Position = InArgs._Position;
 	NotifyHook = InArgs._NotifyHook;
 	OnSampleAdded = InArgs._OnSampleAdded;
 	OnSampleMoved = InArgs._OnSampleMoved;
 	OnSampleRemoved = InArgs._OnSampleRemoved;
 	OnSampleAnimationChanged = InArgs._OnSampleAnimationChanged;
+	bReadOnly = InArgs._ReadOnly;
+	bShowAxisLabels = InArgs._ShowAxisLabels;
+	bShowSettingsButtons = InArgs._ShowSettingsButtons;
 
-	GridType = BlendSpace->IsA<UBlendSpace1D>() ? EGridType::SingleAxis : EGridType::TwoAxis;
+	GridType = (BlendSpaceBase.Get() != nullptr && BlendSpaceBase.Get()->IsA<UBlendSpace1D>()) ? EGridType::SingleAxis : EGridType::TwoAxis;
 	BlendParametersToDraw = (GridType == EGridType::SingleAxis) ? 1 : 2;
 	
 	HighlightedSampleIndex = SelectedSampleIndex = DraggedSampleIndex = INDEX_NONE;
@@ -54,8 +59,8 @@ void SBlendSpaceGridWidget::Construct(const FArguments& InArgs)
 	bPreviewPositionSet = true;
 	bHighlightPreviewPin = false;
 	// Initialize preview value to center or the grid
-	LastPreviewingSampleValue.X = (BlendSpace->GetBlendParameter(0).GetRange() * .5f) + BlendSpace->GetBlendParameter(0).Min;
-	LastPreviewingSampleValue.Y = GridType == EGridType::TwoAxis ? (BlendSpace->GetBlendParameter(1).GetRange() * .5f) + BlendSpace->GetBlendParameter(1).Min : 0.0f;
+	LastPreviewingSampleValue.X = BlendSpaceBase.Get() != nullptr ? (BlendSpaceBase.Get()->GetBlendParameter(0).GetRange() * .5f) + BlendSpaceBase.Get()->GetBlendParameter(0).Min : 0.0f;
+	LastPreviewingSampleValue.Y = BlendSpaceBase.Get() != nullptr ? (GridType == EGridType::TwoAxis ? (BlendSpaceBase.Get()->GetBlendParameter(1).GetRange() * .5f) + BlendSpaceBase.Get()->GetBlendParameter(1).Min : 0.0f) : 0.0f;
 	LastPreviewingSampleValue.Z = 0.0f;
 
 	bShowTriangulation = false;
@@ -94,8 +99,9 @@ void SBlendSpaceGridWidget::Construct(const FArguments& InArgs)
 	KeySize = FVector2D(12.0f, 12.0f);
 	DragThreshold = 9.0f;
 	ClickAndHighlightThreshold = 12.0f;
-	TextMargin = 16.0f;
-	GridMargin = FMargin(MaxVerticalAxisTextWidth + (TextMargin * 2.0f), TextMargin, (HorizontalAxisMaxTextWidth * 0.5f) + TextMargin, MaxHorizontalAxisTextHeight + (TextMargin * 2.0f));
+	TextMargin = 8.0f;
+	GridMargin = bShowAxisLabels ?  FMargin(MaxVerticalAxisTextWidth + (TextMargin * 2.0f), TextMargin, (HorizontalAxisMaxTextWidth *.5f) + TextMargin, MaxHorizontalAxisTextHeight + (TextMargin * 2.0f)) : 
+									FMargin(TextMargin, TextMargin, TextMargin, TextMargin);
 
 	bPreviewToolTipHidden = false;
 
@@ -143,7 +149,7 @@ void SBlendSpaceGridWidget::Construct(const FArguments& InArgs)
 									[
 										SNew(SImage)
 										.Image(FEditorStyle::GetBrush("BlendSpaceEditor.ToggleTriangulation"))
-										.ColorAndOpacity(FSlateColor::UseForeground())										
+										.ColorAndOpacity(FSlateColor::UseForeground())
 									]
 								]
 							]
@@ -153,6 +159,7 @@ void SBlendSpaceGridWidget::Construct(const FArguments& InArgs)
 							[
 								SNew(SBorder)
 								.BorderImage(FEditorStyle::GetBrush("NoBorder"))
+								.Visibility(TAttribute<EVisibility>::Create(TAttribute<EVisibility>::FGetter::CreateSP(this, &SBlendSpaceGridWidget::GetAnimationNamesButtonVisibility)))
 								.VAlign(VAlign_Center)
 								[
 									SNew(SButton)
@@ -262,6 +269,11 @@ void SBlendSpaceGridWidget::Construct(const FArguments& InArgs)
 			.ColorAndOpacity(FLinearColor::Black)
 		]
 	];
+
+	if(bReadOnly)
+	{
+		StartPreviewing();
+	}
 }
 
 TSharedPtr<SWidget> SBlendSpaceGridWidget::CreateGridEntryBox(const int32 BoxIndex, const bool bShowLabel)
@@ -292,16 +304,24 @@ int32 SBlendSpaceGridWidget::OnPaint(const FPaintArgs& Args, const FGeometry& Al
 	SCompoundWidget::OnPaint(Args, AllottedGeometry, MyCullingRect, OutDrawElements, LayerId, InWidgetStyle, bParentEnabled && IsEnabled());
 	
 	PaintBackgroundAndGrid(AllottedGeometry, MyCullingRect, OutDrawElements, LayerId);
-	if (bShowTriangulation)
-	{
-		PaintTriangulation(AllottedGeometry, MyCullingRect, OutDrawElements, LayerId);
-	}	
-	PaintSampleKeys(AllottedGeometry, MyCullingRect, OutDrawElements, LayerId);
-	PaintAxisText(AllottedGeometry, MyCullingRect, OutDrawElements, LayerId);
 
-	if (bShowAnimationNames)
+	if(const UBlendSpaceBase* BlendSpace = BlendSpaceBase.Get())
 	{
-		PaintAnimationNames(AllottedGeometry, MyCullingRect, OutDrawElements, LayerId);
+		if (bShowTriangulation)
+		{
+			PaintTriangulation(AllottedGeometry, MyCullingRect, OutDrawElements, LayerId);
+		}	
+		PaintSampleKeys(AllottedGeometry, MyCullingRect, OutDrawElements, LayerId);
+
+		if(bShowAxisLabels)
+		{
+			PaintAxisText(AllottedGeometry, MyCullingRect, OutDrawElements, LayerId);
+		}
+
+		if (bShowAnimationNames)
+		{
+			PaintAnimationNames(AllottedGeometry, MyCullingRect, OutDrawElements, LayerId);
+		}
 	}
 
 	return LayerId;
@@ -310,101 +330,110 @@ int32 SBlendSpaceGridWidget::OnPaint(const FPaintArgs& Args, const FGeometry& Al
 void SBlendSpaceGridWidget::PaintBackgroundAndGrid(const FGeometry& AllottedGeometry, const FSlateRect& MyCullingRect, FSlateWindowElementList& OutDrawElements, int32& DrawLayerId) const
 {
 	// Fill the background
-	FSlateDrawElement::MakeBox( OutDrawElements, DrawLayerId, AllottedGeometry.ToPaintGeometry(), BackgroundImage );
+	FSlateDrawElement::MakeBox( OutDrawElements, DrawLayerId + 1, AllottedGeometry.ToPaintGeometry(), BackgroundImage );
+
+	if(const UBlendSpaceBase* BlendSpace = BlendSpaceBase.Get())
+	{
+		// Create the grid
+		const FVector2D GridSize = CachedGridRectangle.GetSize();
+		const FVector2D GridOffset = CachedGridRectangle.GetTopLeft();
+		TArray<FVector2D> LinePoints;
+
+		// Draw grid lines
+		LinePoints.SetNumZeroed(2);
+		const FVector2D StartVectors[2] = { FVector2D(1.0f, 0.0f), FVector2D(0.0f, 1.0f) };
+		const FVector2D OffsetVectors[2] = { FVector2D(0.0f, GridSize.Y), FVector2D(GridSize.X, 0.0f) };
+		for (uint32 ParameterIndex = 0; ParameterIndex < BlendParametersToDraw; ++ParameterIndex)
+		{
+			const FBlendParameter& BlendParameter = BlendSpace->GetBlendParameter(ParameterIndex);
+			const float Steps = GridSize[ParameterIndex] / ( BlendParameter.GridNum);
 
-	// Create the grid
-	const FVector2D GridSize = CachedGridRectangle.GetSize();
-	const FVector2D GridOffset = CachedGridRectangle.GetTopLeft();
-	TArray<FVector2D> LinePoints;
+			for (int32 Index = 1; Index < BlendParameter.GridNum; ++Index)
+			{			
+				// Calculate line points
+				LinePoints[0] = ((Index * Steps) * StartVectors[ParameterIndex]) + GridOffset;
+				LinePoints[1] = LinePoints[0] + OffsetVectors[ParameterIndex];
 
-	// Draw outer grid lines separately (this will avoid missing lines with 1D blend spaces)
-	LinePoints.SetNumZeroed(5);
+				FSlateDrawElement::MakeLines( OutDrawElements, DrawLayerId + 2, AllottedGeometry.ToPaintGeometry(), LinePoints, ESlateDrawEffect::None, GridLinesColor, true);
+			}
+		}
 
-	// Top line
-	LinePoints[0] = GridOffset;
+		// Draw outer grid lines separately (this will avoid missing lines with 1D blend spaces)
+		LinePoints.SetNumZeroed(5);
 
-	LinePoints[1] = GridOffset;
-	LinePoints[1].X += GridSize.X;
+		// Top line
+		LinePoints[0] = GridOffset;
 
-	LinePoints[2] = GridOffset;
-	LinePoints[2].X += GridSize.X;
-	LinePoints[2].Y += GridSize.Y;
+		LinePoints[1] = GridOffset;
+		LinePoints[1].X += GridSize.X;
 
-	LinePoints[3] = GridOffset;
-	LinePoints[3].Y += GridSize.Y;
+		LinePoints[2] = GridOffset;
+		LinePoints[2].X += GridSize.X;
+		LinePoints[2].Y += GridSize.Y;
 
-	LinePoints[4] = GridOffset;
+		LinePoints[3] = GridOffset;
+		LinePoints[3].Y += GridSize.Y;
 
-	FSlateDrawElement::MakeLines( OutDrawElements, DrawLayerId + 1, AllottedGeometry.ToPaintGeometry(),	LinePoints, ESlateDrawEffect::None, GridOutlineColor, true);
-	
-	// Draw grid lines
-	LinePoints.SetNumZeroed(2);
-	const FVector2D StartVectors[2] = { FVector2D(1.0f, 0.0f), FVector2D(0.0f, 1.0f) };
-	const FVector2D OffsetVectors[2] = { FVector2D(0.0f, GridSize.Y), FVector2D(GridSize.X, 0.0f) };
-	for (uint32 ParameterIndex = 0; ParameterIndex < BlendParametersToDraw; ++ParameterIndex)
-	{
-		const FBlendParameter& BlendParameter = BlendSpace->GetBlendParameter(ParameterIndex);
-		const float Steps = GridSize[ParameterIndex] / ( BlendParameter.GridNum);
+		LinePoints[4] = GridOffset;
 
-		for (int32 Index = 1; Index < BlendParameter.GridNum; ++Index)
-		{			
-			// Calculate line points
-			LinePoints[0] = ((Index * Steps) * StartVectors[ParameterIndex]) + GridOffset;
-			LinePoints[1] = LinePoints[0] + OffsetVectors[ParameterIndex];
+		FSlateDrawElement::MakeLines( OutDrawElements, DrawLayerId + 3, AllottedGeometry.ToPaintGeometry(),	LinePoints, ESlateDrawEffect::None, GridOutlineColor, true);
 
-			FSlateDrawElement::MakeLines( OutDrawElements, DrawLayerId, AllottedGeometry.ToPaintGeometry(), LinePoints, ESlateDrawEffect::None, GridLinesColor,	true);
-		}
-	}	
+	}
 	
-	DrawLayerId += 2;
+	DrawLayerId += 3;
 }
 
 void SBlendSpaceGridWidget::PaintSampleKeys(const FGeometry& AllottedGeometry, const FSlateRect& MyCullingRect, FSlateWindowElementList& OutDrawElements, int32& DrawLayerId) const
 {
-	// Draw keys
-	const TArray<FBlendSample>& Samples = BlendSpace->GetBlendSamples();
-	for (int32 SampleIndex = 0; SampleIndex < Samples.Num(); ++SampleIndex)
+	if(const UBlendSpaceBase* BlendSpace = BlendSpaceBase.Get())
 	{
-		const FBlendSample& Sample = Samples[SampleIndex];
-		
-		FLinearColor DrawColor = KeyColor.GetSpecifiedColor();
-		if (DraggedSampleIndex == SampleIndex)
-		{
-			DrawColor = (DragState == EDragState::PreDrag) ? PreDragKeyColor.GetSpecifiedColor() : DragKeyColor.GetSpecifiedColor();
-		}
-		else if (SelectedSampleIndex == SampleIndex)
-		{
-			DrawColor = SelectKeyColor.GetSpecifiedColor();
-		}
-		else if (HighlightedSampleIndex == SampleIndex)
+		// Draw keys
+		const TArray<FBlendSample>& Samples = BlendSpace->GetBlendSamples();
+		for (int32 SampleIndex = 0; SampleIndex < Samples.Num(); ++SampleIndex)
 		{
-			DrawColor = HighlightKeyColor.GetSpecifiedColor();
+			const FBlendSample& Sample = Samples[SampleIndex];
+		
+			FLinearColor DrawColor = KeyColor.GetSpecifiedColor();
+			if (DraggedSampleIndex == SampleIndex)
+			{
+				DrawColor = (DragState == EDragState::PreDrag) ? PreDragKeyColor.GetSpecifiedColor() : DragKeyColor.GetSpecifiedColor();
+			}
+			else if (SelectedSampleIndex == SampleIndex)
+			{
+				DrawColor = SelectKeyColor.GetSpecifiedColor();
+			}
+			else if (HighlightedSampleIndex == SampleIndex)
+			{
+				DrawColor = HighlightKeyColor.GetSpecifiedColor();
+			}
+			else if(!Sample.bIsValid)
+			{
+				DrawColor = InvalidColor.GetSpecifiedColor();
+			}
+			else
+			{
+				DrawColor = Sample.bSnapToGrid ? DrawColor : UnSnappedColor.GetSpecifiedColor();
+			}
+
+			const FVector2D GridPosition = SampleValueToGridPosition(Sample.SampleValue) - (KeySize * 0.5f);
+			FSlateDrawElement::MakeBox( OutDrawElements, DrawLayerId + 1, AllottedGeometry.ToPaintGeometry(GridPosition, KeySize), KeyBrush, ESlateDrawEffect::None, DrawColor );
 		}
-		else if(!Sample.bIsValid)
+
+		if (bPreviewPositionSet)
 		{
-			DrawColor = InvalidColor.GetSpecifiedColor();
+			const FVector2D MouseGridPosition = SampleValueToGridPosition(LastPreviewingSampleValue) - (KeySize * .5f);
+			FSlateDrawElement::MakeBox( OutDrawElements, DrawLayerId + 1, AllottedGeometry.ToPaintGeometry(MouseGridPosition, KeySize), KeyBrush, ESlateDrawEffect::None, PreviewKeyColor.GetSpecifiedColor() );
 		}
-		else
+
+		if (DragState == EDragState::DragDrop || DragState == EDragState::InvalidDragDrop)
 		{
-			DrawColor = Sample.bSnapToGrid ? DrawColor : UnSnappedColor.GetSpecifiedColor();
+			const FVector2D GridPoint = SnapToClosestGridPoint(LocalMousePosition) - (KeySize * .5f);
+			FSlateDrawElement::MakeBox( OutDrawElements, DrawLayerId + 1, AllottedGeometry.ToPaintGeometry(GridPoint, KeySize), KeyBrush, ESlateDrawEffect::None,
+				(DragState == EDragState::DragDrop) ? DropKeyColor.GetSpecifiedColor() : InvalidColor.GetSpecifiedColor() );
 		}
-
-		const FVector2D GridPosition = SampleValueToGridPosition(Sample.SampleValue) - (KeySize * 0.5f);
-		FSlateDrawElement::MakeBox( OutDrawElements, DrawLayerId + 1, AllottedGeometry.ToPaintGeometry(GridPosition, KeySize), KeyBrush, ESlateDrawEffect::None, DrawColor );
-	}
-
-	if (bPreviewPositionSet)
-	{
-		const FVector2D MouseGridPosition = SampleValueToGridPosition(LastPreviewingSampleValue) - (KeySize * .5f);
-		FSlateDrawElement::MakeBox( OutDrawElements, DrawLayerId + 1, AllottedGeometry.ToPaintGeometry(MouseGridPosition, KeySize), KeyBrush, ESlateDrawEffect::None, PreviewKeyColor.GetSpecifiedColor() );
 	}
 
-	if (DragState == EDragState::DragDrop || DragState == EDragState::InvalidDragDrop)
-	{
-		const FVector2D GridPoint = SnapToClosestGridPoint(LocalMousePosition) - (KeySize * .5f);
-		FSlateDrawElement::MakeBox( OutDrawElements, DrawLayerId + 1, AllottedGeometry.ToPaintGeometry(GridPoint, KeySize), KeyBrush, ESlateDrawEffect::None,
-			(DragState == EDragState::DragDrop) ? DropKeyColor.GetSpecifiedColor() : InvalidColor.GetSpecifiedColor() );
-	}
+	DrawLayerId += 1;
 }
 
 void SBlendSpaceGridWidget::PaintAxisText(const FGeometry& AllottedGeometry, const FSlateRect& MyCullingRect, FSlateWindowElementList& OutDrawElements, int32& DrawLayerId) const
@@ -475,173 +504,215 @@ void SBlendSpaceGridWidget::PaintAxisText(const FGeometry& AllottedGeometry, con
 		// Maximum value
 		FSlateDrawElement::MakeText(OutDrawElements, DrawLayerId + 1, AllottedGeometry.MakeChild(FVector2D(((GridMargin.Left - TextSize.X) * 0.5f - (KeySize.X * .25f) ) + GridRatioMargin.Left, ( GridMargin.Top + GridRatioMargin.Top) - (TextSize.Y * .5f)), FVector2D(1.0f, 1.0f)).ToPaintGeometry(), Text, FontInfo, ESlateDrawEffect::None, FLinearColor::White);
 	}
+
+	DrawLayerId += 1;
 }
 
 void SBlendSpaceGridWidget::PaintTriangulation(const FGeometry& AllottedGeometry, const FSlateRect& MyCullingRect, FSlateWindowElementList& OutDrawElements, int32& DrawLayerId) const
 {
-	const TArray<FBlendSample>& Samples = BlendSpace->GetBlendSamples();
-	const TArray<FEditorElement>& EditorElements = BlendSpace->GetGridSamples();
-		
-	for (const FEditorElement& Element : EditorElements)
+	if(const UBlendSpaceBase* BlendSpace = BlendSpaceBase.Get())
 	{
-		for (int32 SourceIndex = 0; SourceIndex < 3; ++SourceIndex)
+		const TArray<FBlendSample>& Samples = BlendSpace->GetBlendSamples();
+		const TArray<FEditorElement>& EditorElements = BlendSpace->GetGridSamples();
+		
+		for (const FEditorElement& Element : EditorElements)
 		{
-			if (Element.Indices[SourceIndex] != INDEX_NONE)
+			for (int32 SourceIndex = 0; SourceIndex < 3; ++SourceIndex)
 			{
-				const FBlendSample& SourceSample = Samples[Element.Indices[SourceIndex]];
-				for (int32 TargetIndex = 0; TargetIndex < 3; ++TargetIndex)
+				if (Element.Indices[SourceIndex] != INDEX_NONE)
 				{
-					if (Element.Indices[TargetIndex] != INDEX_NONE)
+					const FBlendSample& SourceSample = Samples[Element.Indices[SourceIndex]];
+					for (int32 TargetIndex = 0; TargetIndex < 3; ++TargetIndex)
 					{
-						if (TargetIndex != SourceIndex)
+						if (Element.Indices[TargetIndex] != INDEX_NONE)
 						{
-							const FBlendSample& TargetSample = Samples[Element.Indices[TargetIndex]];
-							TArray<FVector2D> Points;
+							if (TargetIndex != SourceIndex)
+							{
+								const FBlendSample& TargetSample = Samples[Element.Indices[TargetIndex]];
+								TArray<FVector2D> Points;
 
-							Points.Add(SampleValueToGridPosition(SourceSample.SampleValue));
-							Points.Add(SampleValueToGridPosition(TargetSample.SampleValue));
+								Points.Add(SampleValueToGridPosition(SourceSample.SampleValue));
+								Points.Add(SampleValueToGridPosition(TargetSample.SampleValue));
 
-							// Draw line from and to element
-							FSlateDrawElement::MakeLines(OutDrawElements, DrawLayerId, AllottedGeometry.ToPaintGeometry(), Points, ESlateDrawEffect::None, bShowAnimationNames ? GridLinesColor : FLinearColor::White, true, 0.1f);
+								// Draw line from and to element
+								FSlateDrawElement::MakeLines(OutDrawElements, DrawLayerId + 1, AllottedGeometry.ToPaintGeometry(), Points, ESlateDrawEffect::None, bShowAnimationNames ? GridLinesColor : FLinearColor::White, true, 0.1f);
+							}
 						}
 					}
 				}
 			}
 		}
 	}
+
+	DrawLayerId += 1;
 }
 
 void SBlendSpaceGridWidget::PaintAnimationNames(const FGeometry& AllottedGeometry, const FSlateRect& MyCullingRect, FSlateWindowElementList& OutDrawElements, int32& DrawLayerId) const
 {
-	const TSharedRef< FSlateFontMeasure > FontMeasure = FSlateApplication::Get().GetRenderer()->GetFontMeasureService();
-	const FVector2D GridCenter = CachedGridRectangle.GetCenter();
-	const TArray<FBlendSample>& Samples = BlendSpace->GetBlendSamples();
-	for (int32 SampleIndex = 0; SampleIndex < Samples.Num(); ++SampleIndex)
+	if(const UBlendSpaceBase* BlendSpace = BlendSpaceBase.Get())
 	{
-		const FBlendSample& Sample = Samples[SampleIndex];
-		if (Sample.Animation)
+		const TSharedRef< FSlateFontMeasure > FontMeasure = FSlateApplication::Get().GetRenderer()->GetFontMeasureService();
+		const FVector2D GridCenter = CachedGridRectangle.GetCenter();
+		const TArray<FBlendSample>& Samples = BlendSpace->GetBlendSamples();
+		for (int32 SampleIndex = 0; SampleIndex < Samples.Num(); ++SampleIndex)
 		{
-			const FString Name = Sample.Animation->GetName() + FString::Printf(TEXT(" (%i)"), SampleIndex);
-			const FVector2D TextSize = FontMeasure->Measure(Name, FontInfo);
-
-			FVector2D GridPosition = SampleValueToGridPosition(Sample.SampleValue);
-			// Check on which side of the sample the text should be positioned so that we don't run out of geometry space
-			if ((GridPosition + TextSize).X > AllottedGeometry.GetLocalSize().X)
-			{
-				GridPosition -= FVector2D(TextSize.X + KeySize.X, KeySize.X * .5f);
-			}
-			else
+			const FBlendSample& Sample = Samples[SampleIndex];
+			if (Sample.Animation)
 			{
-				GridPosition += FVector2D(KeySize.X, -KeySize.X * .5f);
-			}
+				const FString Name = Sample.Animation->GetName() + FString::Printf(TEXT(" (%i)"), SampleIndex);
+				const FVector2D TextSize = FontMeasure->Measure(Name, FontInfo);
+
+				FVector2D GridPosition = SampleValueToGridPosition(Sample.SampleValue);
+				// Check on which side of the sample the text should be positioned so that we don't run out of geometry space
+				if ((GridPosition + TextSize).X > AllottedGeometry.GetLocalSize().X)
+				{
+					GridPosition -= FVector2D(TextSize.X + KeySize.X, KeySize.X * .5f);
+				}
+				else
+				{
+					GridPosition += FVector2D(KeySize.X, -KeySize.X * .5f);
+				}
 
-			FSlateDrawElement::MakeText(OutDrawElements, DrawLayerId + 1, AllottedGeometry.MakeChild(FVector2D(GridPosition.X, GridPosition.Y), FVector2D(1.0f, 1.0f)).ToPaintGeometry(), Name, FontInfo, ESlateDrawEffect::None, FLinearColor::White);
+				FSlateDrawElement::MakeText(OutDrawElements, DrawLayerId + 1, AllottedGeometry.MakeChild(FVector2D(GridPosition.X, GridPosition.Y), FVector2D(1.0f, 1.0f)).ToPaintGeometry(), Name, FontInfo, ESlateDrawEffect::None, FLinearColor::White);
+			}
 		}
 	}
+
+	DrawLayerId += 1;
 }
 
 FReply SBlendSpaceGridWidget::OnDrop(const FGeometry& MyGeometry, const FDragDropEvent& DragDropEvent)
 {
-	// Check if we are in dropping state and if so snap to the grid and try to add the sample
-	if (DragState == EDragState::DragDrop || DragState == EDragState::InvalidDragDrop || DragState == EDragState::DragDropOverride)
+	if(!bReadOnly)
 	{
-		if (DragState == EDragState::DragDrop)
-		{
-			const FVector SampleValue = SnapToClosestSamplePoint(LocalMousePosition);
-			TSharedPtr<FAssetDragDropOp> DragDropOperation = DragDropEvent.GetOperationAs<FAssetDragDropOp>();
-			if (DragDropOperation.IsValid())
-			{
-				UAnimSequence* Animation = FAssetData::GetFirstAsset<UAnimSequence>(DragDropOperation->GetAssets());
-				OnSampleAdded.ExecuteIfBound(Animation, SampleValue);
-			}	
-		}
-		else if (DragState == EDragState::DragDropOverride)
+		if(const UBlendSpaceBase* BlendSpace = BlendSpaceBase.Get())
 		{
-			const FVector SampleValue = SnapToClosestSamplePoint(LocalMousePosition);
-			TSharedPtr<FAssetDragDropOp> DragDropOperation = DragDropEvent.GetOperationAs<FAssetDragDropOp>();
-			if (DragDropOperation.IsValid())
+			// Check if we are in dropping state and if so snap to the grid and try to add the sample
+			if (DragState == EDragState::DragDrop || DragState == EDragState::InvalidDragDrop || DragState == EDragState::DragDropOverride)
 			{
-				UAnimSequence* Animation = FAssetData::GetFirstAsset<UAnimSequence>(DragDropOperation->GetAssets());
-				OnSampleAnimationChanged.ExecuteIfBound(Animation, SampleValue);
+				if (DragState == EDragState::DragDrop)
+				{
+					const FVector SampleValue = SnapToClosestSamplePoint(LocalMousePosition);
+					TSharedPtr<FAssetDragDropOp> DragDropOperation = DragDropEvent.GetOperationAs<FAssetDragDropOp>();
+					if (DragDropOperation.IsValid())
+					{
+						UAnimSequence* Animation = FAssetData::GetFirstAsset<UAnimSequence>(DragDropOperation->GetAssets());
+						OnSampleAdded.ExecuteIfBound(Animation, SampleValue);
+					}	
+				}
+				else if (DragState == EDragState::DragDropOverride)
+				{
+					const FVector SampleValue = SnapToClosestSamplePoint(LocalMousePosition);
+					TSharedPtr<FAssetDragDropOp> DragDropOperation = DragDropEvent.GetOperationAs<FAssetDragDropOp>();
+					if (DragDropOperation.IsValid())
+					{
+						UAnimSequence* Animation = FAssetData::GetFirstAsset<UAnimSequence>(DragDropOperation->GetAssets());
+						OnSampleAnimationChanged.ExecuteIfBound(Animation, SampleValue);
+					}
+				}
+
+				DragState = EDragState::None;
 			}
-		}
 
-		DragState = EDragState::None;
+			DragDropAnimationSequence = nullptr;
+			DragDropAnimationName = FText::GetEmpty();
+			HoveredAnimationName = FText::GetEmpty();
+		}
 	}
 
-	DragDropAnimationSequence = nullptr;
-	DragDropAnimationName = FText::GetEmpty();
-	HoveredAnimationName = FText::GetEmpty();
-
 	return FReply::Unhandled();
 }
 
 void SBlendSpaceGridWidget::OnDragEnter(const FGeometry& MyGeometry, const FDragDropEvent& DragDropEvent)
 {
-	if (DragDropEvent.GetOperationAs<FAssetDragDropOp>().IsValid())
+	if(!bReadOnly)
 	{
-		DragState = IsValidDragDropOperation(DragDropEvent, InvalidDragDropText) ? EDragState::DragDrop : EDragState::InvalidDragDrop;
+		if(const UBlendSpaceBase* BlendSpace = BlendSpaceBase.Get())
+		{
+			if (DragDropEvent.GetOperationAs<FAssetDragDropOp>().IsValid())
+			{
+				DragState = IsValidDragDropOperation(DragDropEvent, InvalidDragDropText) ? EDragState::DragDrop : EDragState::InvalidDragDrop;
+			}
+		}
 	}
 }
 
 FReply SBlendSpaceGridWidget::OnDragOver(const FGeometry& MyGeometry, const FDragDropEvent& DragDropEvent)
 {
-	if (DragState == EDragState::DragDrop || DragState == EDragState::InvalidDragDrop || DragState == EDragState::DragDropOverride)
-	{		
-		LocalMousePosition = MyGeometry.AbsoluteToLocal(DragDropEvent.GetScreenSpacePosition());				
-		
-		// Always update the tool tip, in case it became invalid
-		TSharedPtr<FAssetDragDropOp> DragDropOperation = DragDropEvent.GetOperationAs<FAssetDragDropOp>();
-		if (DragDropOperation.IsValid())
+	if(!bReadOnly)
+	{
+		if(const UBlendSpaceBase* BlendSpace = BlendSpaceBase.Get())
 		{
-			DragDropOperation->SetToolTip(GetToolTipSampleValue(), DragDropOperation->GetIcon());
-		}		
+			if (DragState == EDragState::DragDrop || DragState == EDragState::InvalidDragDrop || DragState == EDragState::DragDropOverride)
+			{		
+				LocalMousePosition = MyGeometry.AbsoluteToLocal(DragDropEvent.GetScreenSpacePosition());				
+		
+				// Always update the tool tip, in case it became invalid
+				TSharedPtr<FAssetDragDropOp> DragDropOperation = DragDropEvent.GetOperationAs<FAssetDragDropOp>();
+				if (DragDropOperation.IsValid())
+				{
+					DragDropOperation->SetToolTip(GetToolTipSampleValue(), DragDropOperation->GetIcon());
+				}		
 
-		return FReply::Handled();
+				return FReply::Handled();
+			}
+		}
 	}
 	return FReply::Unhandled();
 }
 
 void SBlendSpaceGridWidget::OnDragLeave(const FDragDropEvent& DragDropEvent)
 {
-	if (DragState == EDragState::DragDrop || DragState == EDragState::InvalidDragDrop || DragState == EDragState::DragDropOverride)
+	if(!bReadOnly)
 	{
-		DragState = EDragState::None;
-		DragDropAnimationSequence = nullptr;
-		DragDropAnimationName = FText::GetEmpty();
-		HoveredAnimationName = FText::GetEmpty();
+		if(const UBlendSpaceBase* BlendSpace = BlendSpaceBase.Get())
+		{
+			if (DragState == EDragState::DragDrop || DragState == EDragState::InvalidDragDrop || DragState == EDragState::DragDropOverride)
+			{
+				DragState = EDragState::None;
+				DragDropAnimationSequence = nullptr;
+				DragDropAnimationName = FText::GetEmpty();
+				HoveredAnimationName = FText::GetEmpty();
+			}
+		}
 	}
 }
 
 FReply SBlendSpaceGridWidget::OnMouseButtonUp(const FGeometry& MyGeometry, const FPointerEvent& MouseEvent)
 {
-	if (this->HasMouseCapture())
+	if(!bReadOnly)
 	{
-		if (DragState == EDragState::None || DragState == EDragState::PreDrag)
-		{
-			ProcessClick(MyGeometry, MouseEvent);
-		}
-		else if (DragState == EDragState::DragSample)
-		{
-			// Process drag ending			
-			ResetToolTip();
-		}
-		else if (DragState == EDragState::DragPreview)
+		if(const UBlendSpaceBase* BlendSpace = BlendSpaceBase.Get())
 		{
-			// Process preview drag ending			
-			ResetToolTip();
-			StopPreviewing();
-		}
+			if (this->HasMouseCapture())
+			{
+				if (DragState == EDragState::None || DragState == EDragState::PreDrag)
+				{
+					ProcessClick(MyGeometry, MouseEvent);
+				}
+				else if (DragState == EDragState::DragSample)
+				{
+					// Process drag ending			
+					ResetToolTip();
+				}
+				else if (DragState == EDragState::DragPreview)
+				{
+					// Process preview drag ending			
+					ResetToolTip();
+					StopPreviewing();
+				}
 
-		// Reset drag state and index
-		DragState = EDragState::None;
-		DraggedSampleIndex = INDEX_NONE;
+				// Reset drag state and index
+				DragState = EDragState::None;
+				DraggedSampleIndex = INDEX_NONE;
 
-		return FReply::Handled().ReleaseMouseCapture();
-	}
-	else
-	{
-		return ProcessClick(MyGeometry, MouseEvent);
+				return FReply::Handled().ReleaseMouseCapture();
+			}
+			else
+			{
+				return ProcessClick(MyGeometry, MouseEvent);
+			}
+		}
 	}
 
 	return FReply::Unhandled();
@@ -649,147 +720,172 @@ FReply SBlendSpaceGridWidget::OnMouseButtonUp(const FGeometry& MyGeometry, const
 
 FReply SBlendSpaceGridWidget::OnMouseButtonDown(const FGeometry& MyGeometry, const FPointerEvent& MouseEvent)
 {
-	if (MouseEvent.GetEffectingButton() == EKeys::LeftMouseButton)
+	if(!bReadOnly)
 	{
-		// If we are over a sample, make it our currently (dragged) sample
-		if (bHighlightPreviewPin)
+		if(const UBlendSpaceBase* BlendSpace = BlendSpaceBase.Get())
 		{
-			ResetToolTip();
-			DragState = EDragState::DragPreview;
-			MouseDownPosition = LocalMousePosition;
-			StartPreviewing();
-			ShowToolTip();
-
-			// Start mouse capture
-			return FReply::Handled().CaptureMouse(SharedThis(this));
+			if (MouseEvent.GetEffectingButton() == EKeys::LeftMouseButton)
+			{
+				// If we are over a sample, make it our currently (dragged) sample
+				if (bHighlightPreviewPin)
+				{
+					ResetToolTip();
+					DragState = EDragState::DragPreview;
+					MouseDownPosition = LocalMousePosition;
+					StartPreviewing();
+					ShowToolTip();
+
+					// Start mouse capture
+					return FReply::Handled().CaptureMouse(SharedThis(this));
+				}
+				else if (HighlightedSampleIndex != INDEX_NONE)
+				{
+					DraggedSampleIndex = SelectedSampleIndex = HighlightedSampleIndex;
+					HighlightedSampleIndex = INDEX_NONE;
+					ResetToolTip();
+					DragState = EDragState::PreDrag;
+					MouseDownPosition = LocalMousePosition;
+
+					// Start mouse capture
+					return FReply::Handled().CaptureMouse(SharedThis(this));
+				}		
+			}
+
+			return FReply::Handled();
 		}
-		else if (HighlightedSampleIndex != INDEX_NONE)
-		{
-			DraggedSampleIndex = SelectedSampleIndex = HighlightedSampleIndex;
-			HighlightedSampleIndex = INDEX_NONE;
-			ResetToolTip();
-			DragState = EDragState::PreDrag;
-			MouseDownPosition = LocalMousePosition;
-
-			// Start mouse capture
-			return FReply::Handled().CaptureMouse(SharedThis(this));
-		}		
 	}
 
-	return FReply::Handled();
+	return FReply::Unhandled();
 }
 
 FReply SBlendSpaceGridWidget::OnMouseMove(const FGeometry& MyGeometry, const FPointerEvent& MouseEvent)
 {
-	// Cache the mouse position in local and screen space
-	LocalMousePosition = MyGeometry.AbsoluteToLocal(MouseEvent.GetScreenSpacePosition());
-	LastMousePosition = MouseEvent.GetScreenSpacePosition();
-
-	if (this->HasMouseCapture())
+	if(const UBlendSpaceBase* BlendSpace = BlendSpaceBase.Get())
 	{
-		if (DragState == EDragState::None)
+		// Cache the mouse position in local and screen space
+		LocalMousePosition = MyGeometry.AbsoluteToLocal(MouseEvent.GetScreenSpacePosition());
+		LastMousePosition = MouseEvent.GetScreenSpacePosition();
+
+		if (!bReadOnly && this->HasMouseCapture())
 		{
-			if (HighlightedSampleIndex != INDEX_NONE)
+			if (DragState == EDragState::None)
 			{
-				DragState = EDragState::DragSample;
-				DraggedSampleIndex = HighlightedSampleIndex;
-				HighlightedSampleIndex = INDEX_NONE;
+				if (HighlightedSampleIndex != INDEX_NONE)
+				{
+					DragState = EDragState::DragSample;
+					DraggedSampleIndex = HighlightedSampleIndex;
+					HighlightedSampleIndex = INDEX_NONE;
+				}
 			}
-		}
-		else if (DragState == EDragState::PreDrag)
-		{
-			// Actually start dragging
-			if ((LocalMousePosition - MouseDownPosition).SizeSquared() > DragThreshold)
+			else if (DragState == EDragState::PreDrag)
 			{
-				DragState = EDragState::DragSample;
-				HighlightedSampleIndex = INDEX_NONE;
-				ShowToolTip();
+				// Actually start dragging
+				if ((LocalMousePosition - MouseDownPosition).SizeSquared() > DragThreshold)
+				{
+					DragState = EDragState::DragSample;
+					HighlightedSampleIndex = INDEX_NONE;
+					ShowToolTip();
+				}
 			}
 		}
-	}
 
-	if (IsHovered() && !HasAnyUserFocus() && bMouseIsOverGeometry)
-	{
-		if (MouseEvent.IsLeftShiftDown() || MouseEvent.IsRightShiftDown())
+		if (IsHovered() && !HasAnyUserFocus() && bMouseIsOverGeometry)
 		{
-			StartPreviewing();
-			DragState = EDragState::Preview;
-			// Make tool tip visible (this will display the current preview sample value)
-			ShowToolTip();			
+			if (!bReadOnly && (MouseEvent.IsLeftShiftDown() || MouseEvent.IsRightShiftDown()))
+			{
+				StartPreviewing();
+				DragState = EDragState::Preview;
+				// Make tool tip visible (this will display the current preview sample value)
+				ShowToolTip();			
 
-			// Set flag for showing advanced preview info in tooltip
-			bAdvancedPreview = MouseEvent.IsLeftControlDown() || MouseEvent.IsRightControlDown();
-		}
-		else if (bSamplePreviewing)
-		{
-			StopPreviewing();
-			DragState = EDragState::None;
-			ResetToolTip();
+				// Set flag for showing advanced preview info in tooltip
+				bAdvancedPreview = MouseEvent.IsLeftControlDown() || MouseEvent.IsRightControlDown();
+			}
+			else if(bReadOnly)
+			{
+				StartPreviewing();
+				DragState = EDragState::None;
+				ShowToolTip();
+			}
+			else if (bSamplePreviewing)
+			{
+				StopPreviewing();
+				DragState = EDragState::None;
+				ResetToolTip();
+			}
 		}
+
+		return FReply::Handled();
 	}
 
-	return FReply::Handled();
+	return FReply::Unhandled();
 }
 
 FReply SBlendSpaceGridWidget::ProcessClick(const FGeometry& MyGeometry, const FPointerEvent& MouseEvent)
 {
-	if (MouseEvent.GetEffectingButton() == EKeys::LeftMouseButton)
+	if(!bReadOnly)
 	{
-		SelectedSampleIndex = INDEX_NONE;
-
-		if (HighlightedSampleIndex == INDEX_NONE)
+		if(const UBlendSpaceBase* BlendSpace = BlendSpaceBase.Get())
 		{
-			// If there isn't any sample currently being highlighted, retrieve all of them and see if we are over one
-			const TArray<FBlendSample>& Samples = BlendSpace->GetBlendSamples();
-			for (int32 SampleIndex = 0; SampleIndex < Samples.Num(); ++SampleIndex)
+			if (MouseEvent.GetEffectingButton() == EKeys::LeftMouseButton)
 			{
-				const FBlendSample& Sample = Samples[SampleIndex];
-				if (IsSampleValueWithinMouseRange(Sample.SampleValue))
+				SelectedSampleIndex = INDEX_NONE;
+
+				if (HighlightedSampleIndex == INDEX_NONE)
 				{
-					SelectedSampleIndex = SampleIndex;					
-					break;
+					// If there isn't any sample currently being highlighted, retrieve all of them and see if we are over one
+					const TArray<FBlendSample>& Samples = BlendSpace->GetBlendSamples();
+					for (int32 SampleIndex = 0; SampleIndex < Samples.Num(); ++SampleIndex)
+					{
+						const FBlendSample& Sample = Samples[SampleIndex];
+						if (IsSampleValueWithinMouseRange(Sample.SampleValue))
+						{
+							SelectedSampleIndex = SampleIndex;					
+							break;
+						}
+					}
+				}
+				else
+				{
+					// If we are over a sample, make it the selected sample index
+					SelectedSampleIndex = HighlightedSampleIndex;
+					HighlightedSampleIndex = INDEX_NONE;			
 				}
 			}
-		}
-		else
-		{
-			// If we are over a sample, make it the selected sample index
-			SelectedSampleIndex = HighlightedSampleIndex;
-			HighlightedSampleIndex = INDEX_NONE;			
-		}
-	}
-	else if (MouseEvent.GetEffectingButton() == EKeys::RightMouseButton)
-	{
-		// If we are over a sample open a context menu for editing its data
-		if (HighlightedSampleIndex != INDEX_NONE)
-		{	
-			SelectedSampleIndex = HighlightedSampleIndex;
+			else if (MouseEvent.GetEffectingButton() == EKeys::RightMouseButton)
+			{
+				// If we are over a sample open a context menu for editing its data
+				if (HighlightedSampleIndex != INDEX_NONE)
+				{	
+					SelectedSampleIndex = HighlightedSampleIndex;
 
-			// Create context menu
-			TSharedPtr<SWidget> MenuContent = CreateBlendSampleContextMenu();
+					// Create context menu
+					TSharedPtr<SWidget> MenuContent = CreateBlendSampleContextMenu();
 			
-			// Reset highlight sample index
-			HighlightedSampleIndex = INDEX_NONE;
+					// Reset highlight sample index
+					HighlightedSampleIndex = INDEX_NONE;
 
-			if (MenuContent.IsValid())
-			{
-				const FWidgetPath WidgetPath = MouseEvent.GetEventPath() != nullptr ? *MouseEvent.GetEventPath() : FWidgetPath();
-				const FVector2D MousePosition = MouseEvent.GetScreenSpacePosition();
-				// This is of a fixed size atm since MenuContent->GetDesiredSize() will not take the detail customization into account and return an incorrect (small) size
-				const FVector2D ExpectedSize(300, 100);				
-				const FVector2D MenuPosition = FSlateApplication::Get().CalculatePopupWindowPosition(FSlateRect(MousePosition.X, MousePosition.Y, MousePosition.X, MousePosition.Y), ExpectedSize, false);
-
-				FSlateApplication::Get().PushMenu(
-					AsShared(),
-					WidgetPath,
-					MenuContent.ToSharedRef(),
-					MenuPosition,
-					FPopupTransitionEffect(FPopupTransitionEffect::ContextMenu)
-					);
-
-				return FReply::Handled().SetUserFocus(MenuContent.ToSharedRef(), EFocusCause::SetDirectly).ReleaseMouseCapture();
+					if (MenuContent.IsValid())
+					{
+						const FWidgetPath WidgetPath = MouseEvent.GetEventPath() != nullptr ? *MouseEvent.GetEventPath() : FWidgetPath();
+						const FVector2D MousePosition = MouseEvent.GetScreenSpacePosition();
+						// This is of a fixed size atm since MenuContent->GetDesiredSize() will not take the detail customization into account and return an incorrect (small) size
+						const FVector2D ExpectedSize(300, 100);				
+						const FVector2D MenuPosition = FSlateApplication::Get().CalculatePopupWindowPosition(FSlateRect(MousePosition.X, MousePosition.Y, MousePosition.X, MousePosition.Y), ExpectedSize, false);
+
+						FSlateApplication::Get().PushMenu(
+							AsShared(),
+							WidgetPath,
+							MenuContent.ToSharedRef(),
+							MenuPosition,
+							FPopupTransitionEffect(FPopupTransitionEffect::ContextMenu)
+							);
+
+						return FReply::Handled().SetUserFocus(MenuContent.ToSharedRef(), EFocusCause::SetDirectly).ReleaseMouseCapture();
+					}
+				}		
 			}
-		}		
+		}
 	}
 
 	return FReply::Unhandled();
@@ -797,23 +893,29 @@ FReply SBlendSpaceGridWidget::ProcessClick(const FGeometry& MyGeometry, const FP
 
 FReply SBlendSpaceGridWidget::OnKeyDown(const FGeometry& MyGeometry, const FKeyEvent& InKeyEvent)
 {
-	// Start previewing when either one of the shift keys is pressed
-	if (bMouseIsOverGeometry)
+	if(!bReadOnly)
 	{
-		if (DragState != EDragState::DragPreview && ((InKeyEvent.GetKey() == EKeys::LeftShift) || (InKeyEvent.GetKey() == EKeys::RightShift)))
+		if(const UBlendSpaceBase* BlendSpace = BlendSpaceBase.Get())
 		{
-			StartPreviewing();
-			DragState = EDragState::Preview;
-			// Make tool tip visible (this will display the current preview sample value)
-			ShowToolTip();
-			return FReply::Handled();
-		}
+			// Start previewing when either one of the shift keys is pressed
+			if (bMouseIsOverGeometry)
+			{
+				if (DragState != EDragState::DragPreview && ((InKeyEvent.GetKey() == EKeys::LeftShift) || (InKeyEvent.GetKey() == EKeys::RightShift)))
+				{
+					StartPreviewing();
+					DragState = EDragState::Preview;
+					// Make tool tip visible (this will display the current preview sample value)
+					ShowToolTip();
+					return FReply::Handled();
+				}
 		
-		// Set flag for showing advanced preview info in tooltip
-		if ((InKeyEvent.GetKey() == EKeys::LeftControl) || (InKeyEvent.GetKey() == EKeys::RightControl))
-		{
-			bAdvancedPreview = true;
-			return FReply::Handled();
+				// Set flag for showing advanced preview info in tooltip
+				if ((InKeyEvent.GetKey() == EKeys::LeftControl) || (InKeyEvent.GetKey() == EKeys::RightControl))
+				{
+					bAdvancedPreview = true;
+					return FReply::Handled();
+				}
+			}
 		}
 	}
 
@@ -822,42 +924,48 @@ FReply SBlendSpaceGridWidget::OnKeyDown(const FGeometry& MyGeometry, const FKeyE
 
 FReply SBlendSpaceGridWidget::OnKeyUp(const FGeometry& MyGeometry, const FKeyEvent& InKeyEvent)
 {	
-	// Stop previewing when shift keys are released 
-	if (DragState != EDragState::DragPreview && ((InKeyEvent.GetKey() == EKeys::LeftShift) || (InKeyEvent.GetKey() == EKeys::RightShift)))
+	if(!bReadOnly)
 	{
-		StopPreviewing();
-		DragState = EDragState::None;
-		ResetToolTip();
-		return FReply::Handled();
-	}
-
-	if((InKeyEvent.GetKey() == EKeys::LeftControl) || (InKeyEvent.GetKey() == EKeys::RightControl))
-	{
-		bAdvancedPreview = false;
-		return FReply::Handled();
-	}
-
-	// If delete is pressed and we currently have a sample selected remove it from the blendspace
-	if (InKeyEvent.GetKey() == EKeys::Delete)
-	{
-		if (SelectedSampleIndex != INDEX_NONE)
+		if(const UBlendSpaceBase* BlendSpace = BlendSpaceBase.Get())
 		{
-			OnSampleRemoved.ExecuteIfBound(SelectedSampleIndex);
-			
-			if (SelectedSampleIndex == HighlightedSampleIndex)
+			// Stop previewing when shift keys are released 
+			if (DragState != EDragState::DragPreview && ((InKeyEvent.GetKey() == EKeys::LeftShift) || (InKeyEvent.GetKey() == EKeys::RightShift)))
 			{
-				HighlightedSampleIndex = INDEX_NONE;
+				StopPreviewing();
+				DragState = EDragState::None;
 				ResetToolTip();
+				return FReply::Handled();
 			}
 
-			SelectedSampleIndex = INDEX_NONE;
-		}
-	}
+			if((InKeyEvent.GetKey() == EKeys::LeftControl) || (InKeyEvent.GetKey() == EKeys::RightControl))
+			{
+				bAdvancedPreview = false;
+				return FReply::Handled();
+			}
 
-	// Pressing esc will remove the current key selection
-	if( InKeyEvent.GetKey() == EKeys::Escape)
-	{
-		SelectedSampleIndex = INDEX_NONE;
+			// If delete is pressed and we currently have a sample selected remove it from the blendspace
+			if (InKeyEvent.GetKey() == EKeys::Delete)
+			{
+				if (SelectedSampleIndex != INDEX_NONE)
+				{
+					OnSampleRemoved.ExecuteIfBound(SelectedSampleIndex);
+			
+					if (SelectedSampleIndex == HighlightedSampleIndex)
+					{
+						HighlightedSampleIndex = INDEX_NONE;
+						ResetToolTip();
+					}
+
+					SelectedSampleIndex = INDEX_NONE;
+				}
+			}
+
+			// Pressing esc will remove the current key selection
+			if( InKeyEvent.GetKey() == EKeys::Escape)
+			{
+				SelectedSampleIndex = INDEX_NONE;
+			}
+		}
 	}
 
 	return FReply::Unhandled();
@@ -893,12 +1001,15 @@ TSharedPtr<SWidget> SBlendSpaceGridWidget::CreateBlendSampleContextMenu()
 	StructureDetailsView = FModuleManager::GetModuleChecked<FPropertyEditorModule>("PropertyEditor")
 		.CreateStructureDetailView(DetailsViewArgs, StructureViewArgs, nullptr, LOCTEXT("SampleData", "Blend Sample"));
 	{
-		const FBlendSample& Sample = BlendSpace->GetBlendSample(HighlightedSampleIndex);		
-		StructureDetailsView->GetDetailsView()->SetGenericLayoutDetailsDelegate(FOnGetDetailCustomizationInstance::CreateStatic(&FBlendSampleDetails::MakeInstance, BlendSpace, this));
+		if(const UBlendSpaceBase* BlendSpace = BlendSpaceBase.Get())
+		{
+			const FBlendSample& Sample = BlendSpace->GetBlendSample(HighlightedSampleIndex);		
+			StructureDetailsView->GetDetailsView()->SetGenericLayoutDetailsDelegate(FOnGetDetailCustomizationInstance::CreateStatic(&FBlendSampleDetails::MakeInstance, BlendSpace, this));
 
-		FStructOnScope* Struct = new FStructOnScope(FBlendSample::StaticStruct(), (uint8*)&Sample);
-		Struct->SetPackage(BlendSpace->GetOutermost());
-		StructureDetailsView->SetStructureData(MakeShareable(Struct));
+			FStructOnScope* Struct = new FStructOnScope(FBlendSample::StaticStruct(), (uint8*)&Sample);
+			Struct->SetPackage(BlendSpace->GetOutermost());
+			StructureDetailsView->SetStructureData(MakeShareable(Struct));
+		}
 	}
 	
 	MenuBuilder.AddWidget(StructureDetailsView->GetWidget().ToSharedRef(), FText::GetEmpty(), true);
@@ -989,18 +1100,18 @@ const FVector2D SBlendSpaceGridWidget::SampleValueToGridPosition(const FVector&
 
 const FVector SBlendSpaceGridWidget::GridPositionToSampleValue(const FVector2D& GridPosition) const
 {
-	FVector2D Position = GridPosition;
+	FVector2D LocalGridPosition = GridPosition;
 	// Move to center of grid and convert to 0 - 1 form
-	Position -= CachedGridRectangle.GetCenter();
-	Position /= (CachedGridRectangle.GetSize() * 0.5f);
-	Position += FVector2D::UnitVector;
-	Position *= 0.5f;
+	LocalGridPosition -= CachedGridRectangle.GetCenter();
+	LocalGridPosition /= (CachedGridRectangle.GetSize() * 0.5f);
+	LocalGridPosition += FVector2D::UnitVector;
+	LocalGridPosition *= 0.5f;
 
 	// Calculate the sample value by mapping it to the blend parameter range
 	const FVector SampleValue
 	(		
-		FMath::Clamp((Position.X * SampleValueRange.X) + SampleValueMin.X, SampleValueMin.X, SampleValueMax.X),
-		FMath::Clamp(((GridType == EGridType::TwoAxis) ? SampleValueMax.Y - (Position.Y * SampleValueRange.Y) : 0.0f),
+		FMath::Clamp((LocalGridPosition.X * SampleValueRange.X) + SampleValueMin.X, SampleValueMin.X, SampleValueMax.X),
+		FMath::Clamp(((GridType == EGridType::TwoAxis) ? SampleValueMax.Y - (LocalGridPosition.Y * SampleValueRange.Y) : 0.0f),
 			SampleValueMin.Y, SampleValueMax.Y),
 		0.f	
 	);
@@ -1029,7 +1140,7 @@ void SBlendSpaceGridWidget::StartPreviewing()
 {
 	bSamplePreviewing = true;
 	LastPreviewingMousePosition = LocalMousePosition;
-	LastPreviewingSampleValue = GridPositionToSampleValue(LastPreviewingMousePosition);
+	LastPreviewingSampleValue = bReadOnly ? Position.Get() : GridPositionToSampleValue(LastPreviewingMousePosition);
 	bPreviewPositionSet = true;	
 	bPreviewToolTipHidden = true;
 }
@@ -1042,71 +1153,75 @@ void SBlendSpaceGridWidget::StopPreviewing()
 FText SBlendSpaceGridWidget::GetToolTipAnimationName() const
 {
 	FText ToolTipText = FText::GetEmpty();
-	const FText EmptyAnimationText = LOCTEXT("NoAnimationSetTooltipText","No Animation Set");
-	switch (DragState)
+	if(const UBlendSpaceBase* BlendSpace = BlendSpaceBase.Get())
 	{
-		// If we are not dragging, but over a valid blend sample return its animation asset name
-		case EDragState::None:
-		{		
-			if (bHighlightPreviewPin)
-			{
-				const FText PreviewPinText = LOCTEXT("HighlightPreviewPinTooltipText", "Preview Value (Click and Drag or Hold Shift to Move)");
-				ToolTipText = PreviewPinText;
-			}
-			else if (HighlightedSampleIndex != INDEX_NONE && BlendSpace->IsValidBlendSampleIndex(HighlightedSampleIndex))
-			{
-				const FBlendSample& BlendSample = BlendSpace->GetBlendSample(HighlightedSampleIndex);					
-				ToolTipText = (BlendSample.Animation != nullptr) ? FText::FromString(BlendSample.Animation->GetName()) : EmptyAnimationText;
-			}
-			break;
-		}
+		const FText EmptyAnimationText = LOCTEXT("NoAnimationSetTooltipText", "No Animation Set");
 
-		case EDragState::PreDrag:
+		switch (DragState)
 		{
-			break;
-		}
+			// If we are not dragging, but over a valid blend sample return its animation asset name
+			case EDragState::None:
+			{		
+				if (bHighlightPreviewPin && !bReadOnly)
+				{
+					const FText PreviewPinText = LOCTEXT("HighlightPreviewPinTooltipText", "Preview Value (Click and Drag or Hold Shift to Move)");
+					ToolTipText = PreviewPinText;
+				}
+				else if (HighlightedSampleIndex != INDEX_NONE && BlendSpace->IsValidBlendSampleIndex(HighlightedSampleIndex))
+				{
+					const FBlendSample& BlendSample = BlendSpace->GetBlendSample(HighlightedSampleIndex);					
+					ToolTipText = (BlendSample.Animation != nullptr) ? FText::FromString(BlendSample.Animation->GetName()) : EmptyAnimationText;
+				}
+				break;
+			}
 
-		// If we are dragging a sample return the dragged sample's animation asset name
-		case EDragState::DragSample:
-		{
-			if (BlendSpace->IsValidBlendSampleIndex(DraggedSampleIndex))
+			case EDragState::PreDrag:
 			{
-				const FBlendSample& BlendSample = BlendSpace->GetBlendSample(DraggedSampleIndex);
-				ToolTipText = (BlendSample.Animation != nullptr) ? FText::FromString(BlendSample.Animation->GetName()) : EmptyAnimationText;
+				break;
 			}
+
+			// If we are dragging a sample return the dragged sample's animation asset name
+			case EDragState::DragSample:
+			{
+				if (BlendSpace->IsValidBlendSampleIndex(DraggedSampleIndex))
+				{
+					const FBlendSample& BlendSample = BlendSpace->GetBlendSample(DraggedSampleIndex);
+					ToolTipText = (BlendSample.Animation != nullptr) ? FText::FromString(BlendSample.Animation->GetName()) : EmptyAnimationText;
+				}
 			
-			break;
-		}
+				break;
+			}
 
-		// If we are performing a drag/drop operation return the cached operation animation name
-		case EDragState::DragDrop:
-		{
-			ToolTipText = DragDropAnimationName;
-			break;
-		}
+			// If we are performing a drag/drop operation return the cached operation animation name
+			case EDragState::DragDrop:
+			{
+				ToolTipText = DragDropAnimationName;
+				break;
+			}
 
-		case EDragState::DragDropOverride:
-		{
-			ToolTipText = DragDropAnimationName;
-			break;
-		}
+			case EDragState::DragDropOverride:
+			{
+				ToolTipText = DragDropAnimationName;
+				break;
+			}
 
-		case EDragState::InvalidDragDrop:
-		{
-			break;
-		}
+			case EDragState::InvalidDragDrop:
+			{
+				break;
+			}
 		
-		// If we are previewing return a descriptive label
-		case EDragState::Preview:
-		case EDragState::DragPreview:
-		{
-			const FText AdvancedPreviewText = LOCTEXT("AdvancedPreviewValueTooltip", "Preview Value");
-			const FText BasicPreviewText = LOCTEXT("BasicPreviewValueTooltip", "Preview Value (Hold Ctrl for Weight Details)");	
-			ToolTipText = bAdvancedPreview ? AdvancedPreviewText : BasicPreviewText;
-			break;
+			// If we are previewing return a descriptive label
+			case EDragState::Preview:
+			case EDragState::DragPreview:
+			{
+				const FText AdvancedPreviewText = LOCTEXT("AdvancedPreviewValueTooltip", "Preview Value");
+				const FText BasicPreviewText = LOCTEXT("BasicPreviewValueTooltip", "Preview Value (Hold Ctrl for Weight Details)");	
+				ToolTipText = bAdvancedPreview ? AdvancedPreviewText : BasicPreviewText;
+				break;
+			}
+			default:
+				check(false);
 		}
-		default:
-			check(false);
 	}
 
 	return ToolTipText;
@@ -1116,101 +1231,111 @@ FText SBlendSpaceGridWidget::GetToolTipSampleValue() const
 {
 	FText ToolTipText = FText::GetEmpty();
 
-	const FTextFormat ValueFormattingText = (GridType == EGridType::TwoAxis) ? FTextFormat::FromString("{0}: {1} - {2}: {3}") : FTextFormat::FromString("{0}: {1}");
-
-	switch (DragState)
+	if(const UBlendSpaceBase* BlendSpace = BlendSpaceBase.Get())
 	{
-		// If we are over a sample return its sample value if valid and otherwise show an error message as to why the sample is invalid
-		case EDragState::None:
-		{		
-			if (bHighlightPreviewPin)
-			{
-				ToolTipText = FText::Format(ValueFormattingText, ParameterXName, FText::FromString(FString::SanitizeFloat(LastPreviewingSampleValue.X)), ParameterYName, FText::FromString(FString::SanitizeFloat(LastPreviewingSampleValue.Y)));
-			}
-			else if (HighlightedSampleIndex != INDEX_NONE && BlendSpace->IsValidBlendSampleIndex(HighlightedSampleIndex))
-			{
-				const FBlendSample& BlendSample = BlendSpace->GetBlendSample(HighlightedSampleIndex);
-
-				// Check if the sample is valid
-				if (BlendSample.bIsValid)
-				{
-					ToolTipText = FText::Format(ValueFormattingText, ParameterXName, FText::FromString(FString::SanitizeFloat(BlendSample.SampleValue.X)), ParameterYName, FText::FromString(FString::SanitizeFloat(BlendSample.SampleValue.Y)));
-				}
-				else
-				{
-					ToolTipText = GetSampleErrorMessage(BlendSample);
-				}
-			}
+		const FTextFormat ValueFormattingText = (GridType == EGridType::TwoAxis) ? FTextFormat::FromString("{0}: {1} - {2}: {3}") : FTextFormat::FromString("{0}: {1}");
 
-			break;
-		}
-
-		case EDragState::PreDrag:
+		if(bReadOnly)
 		{
-			break;
+			ToolTipText = FText::Format(ValueFormattingText, ParameterXName, FText::FromString(FString::SanitizeFloat(LastPreviewingSampleValue.X)), ParameterYName, FText::FromString(FString::SanitizeFloat(LastPreviewingSampleValue.Y)));
 		}
-
-		// If we are dragging a sample return the current sample value it is hovered at
-		case EDragState::DragSample:
+		else
 		{
-			if (DraggedSampleIndex != INDEX_NONE)
+			switch (DragState)
 			{
-				const FBlendSample& BlendSample = BlendSpace->GetBlendSample(DraggedSampleIndex);
-				ToolTipText = FText::Format(ValueFormattingText, ParameterXName, FText::FromString(FString::SanitizeFloat(BlendSample.SampleValue.X)), ParameterYName, FText::FromString(FString::SanitizeFloat(BlendSample.SampleValue.Y)));
-			}
-			break;
-		}
+				// If we are over a sample return its sample value if valid and otherwise show an error message as to why the sample is invalid
+				case EDragState::None:
+				{		
+					if (bHighlightPreviewPin)
+					{
+						ToolTipText = FText::Format(ValueFormattingText, ParameterXName, FText::FromString(FString::SanitizeFloat(LastPreviewingSampleValue.X)), ParameterYName, FText::FromString(FString::SanitizeFloat(LastPreviewingSampleValue.Y)));
+					}
+					else if (HighlightedSampleIndex != INDEX_NONE && BlendSpace->IsValidBlendSampleIndex(HighlightedSampleIndex))
+					{
+						const FBlendSample& BlendSample = BlendSpace->GetBlendSample(HighlightedSampleIndex);
 
-		// If we are performing a drag and drop operation return the current sample value it is hovered at
-		case EDragState::DragDrop:
-		{
-			const FVector SampleValue = SnapToClosestSamplePoint(LocalMousePosition);
+						// Check if the sample is valid
+						if (BlendSample.bIsValid)
+						{
+							ToolTipText = FText::Format(ValueFormattingText, ParameterXName, FText::FromString(FString::SanitizeFloat(BlendSample.SampleValue.X)), ParameterYName, FText::FromString(FString::SanitizeFloat(BlendSample.SampleValue.Y)));
+						}
+						else
+						{
+							ToolTipText = GetSampleErrorMessage(BlendSample);
+						}
+					}
 
-			ToolTipText = FText::Format(ValueFormattingText, ParameterXName, FText::FromString(FString::SanitizeFloat(SampleValue.X)), ParameterYName, FText::FromString(FString::SanitizeFloat(SampleValue.Y)));
+					break;
+				}
 
-			break;
-		}
+				case EDragState::PreDrag:
+				{
+					break;
+				}
 
-		case EDragState::DragDropOverride:
-		{
-			const FTextFormat OverrideAnimationFormat = FTextFormat::FromString("Changing Animation from {0} to {1}");
-			ToolTipText = FText::Format(OverrideAnimationFormat, HoveredAnimationName, DragDropAnimationName);
-			break;
-		}
-		// If the drag and drop operation is invalid return the cached error message as to why it is invalid
-		case EDragState::InvalidDragDrop:
-		{
-			ToolTipText = InvalidDragDropText;
-			break;
-		}
+				// If we are dragging a sample return the current sample value it is hovered at
+				case EDragState::DragSample:
+				{
+					if (DraggedSampleIndex != INDEX_NONE)
+					{
+						const FBlendSample& BlendSample = BlendSpace->GetBlendSample(DraggedSampleIndex);
+						ToolTipText = FText::Format(ValueFormattingText, ParameterXName, FText::FromString(FString::SanitizeFloat(BlendSample.SampleValue.X)), ParameterYName, FText::FromString(FString::SanitizeFloat(BlendSample.SampleValue.Y)));
+					}
+					break;
+				}
 
-		// If we are setting the preview value return the current preview sample value
-		case EDragState::DragPreview:
-		case EDragState::Preview:
-		{
-			FFormatOrderedArguments PreviewTextArguments;
-			PreviewTextArguments.Add(FText::Format(ValueFormattingText, ParameterXName, FText::FromString(FString::SanitizeFloat(LastPreviewingSampleValue.X)), ParameterYName, FText::FromString(FString::SanitizeFloat(LastPreviewingSampleValue.Y))));
+				// If we are performing a drag and drop operation return the current sample value it is hovered at
+				case EDragState::DragDrop:
+				{
+					const FVector SampleValue = SnapToClosestSamplePoint(LocalMousePosition);
+
+					ToolTipText = FText::Format(ValueFormattingText, ParameterXName, FText::FromString(FString::SanitizeFloat(SampleValue.X)), ParameterYName, FText::FromString(FString::SanitizeFloat(SampleValue.Y)));
 
-			FString PreviewSamplesString;
-			if (bAdvancedPreview)
-			{				
-				for (const FBlendSampleData& SampleData : PreviewedSamples)
+					break;
+				}
+
+				case EDragState::DragDropOverride:
 				{
-					PreviewSamplesString += "\n";
-					PreviewSamplesString += SampleData.Animation ? SampleData.Animation->GetName() : "No Animation Set";
-					PreviewSamplesString += ": ";
-					PreviewSamplesString += FString::SanitizeFloat(SampleData.TotalWeight);
-				}				
-			}
-			PreviewTextArguments.Add(FText::FromString(PreviewSamplesString));
+					const FTextFormat OverrideAnimationFormat = FTextFormat::FromString("Changing Animation from {0} to {1}");
+					ToolTipText = FText::Format(OverrideAnimationFormat, HoveredAnimationName, DragDropAnimationName);
+					break;
+				}
+				// If the drag and drop operation is invalid return the cached error message as to why it is invalid
+				case EDragState::InvalidDragDrop:
+				{
+					ToolTipText = InvalidDragDropText;
+					break;
+				}
+
+				// If we are setting the preview value return the current preview sample value
+				case EDragState::DragPreview:
+				case EDragState::Preview:
+				{
+					FFormatOrderedArguments PreviewTextArguments;
+					PreviewTextArguments.Add(FText::Format(ValueFormattingText, ParameterXName, FText::FromString(FString::SanitizeFloat(LastPreviewingSampleValue.X)), ParameterYName, FText::FromString(FString::SanitizeFloat(LastPreviewingSampleValue.Y))));
+
+					FString PreviewSamplesString;
+					if (bAdvancedPreview)
+					{				
+						for (const FBlendSampleData& SampleData : PreviewedSamples)
+						{
+							PreviewSamplesString += "\n";
+							PreviewSamplesString += SampleData.Animation ? SampleData.Animation->GetName() : "No Animation Set";
+							PreviewSamplesString += ": ";
+							PreviewSamplesString += FString::SanitizeFloat(SampleData.TotalWeight);
+						}				
+					}
+					PreviewTextArguments.Add(FText::FromString(PreviewSamplesString));
 			
-			const FTextFormat PreviewTextFormat = FTextFormat::FromString("{0}{1}");			
-			ToolTipText = FText::Format(PreviewTextFormat, PreviewTextArguments);
+					const FTextFormat PreviewTextFormat = FTextFormat::FromString("{0}{1}");			
+					ToolTipText = FText::Format(PreviewTextFormat, PreviewTextArguments);
 			
-			break;
+					break;
+				}
+				default:
+					check(false);
+			}
 		}
-		default:
-			check(false);
 	}
 
 	return ToolTipText;
@@ -1248,7 +1373,7 @@ void SBlendSpaceGridWidget::ResetToolTip()
 
 EVisibility SBlendSpaceGridWidget::GetInputBoxVisibility(const int32 ParameterIndex) const
 {
-	bool bVisible = true;
+	bool bVisible = !bReadOnly;
 	// Only show input boxes when a sample is selected (hide it when one is being dragged since we have the tooltip information as well)
 	bVisible &= (SelectedSampleIndex != INDEX_NONE && DraggedSampleIndex == INDEX_NONE);
 	if ( ParameterIndex == 1 )
@@ -1263,10 +1388,13 @@ TOptional<float> SBlendSpaceGridWidget::GetInputBoxValue(const int32 ParameterIn
 {
 	checkf(ParameterIndex < 3, TEXT("Invalid parameter index, suppose to be within FVector array range"));
 	float ReturnValue = 0.0f;
-	if (SelectedSampleIndex != INDEX_NONE && SelectedSampleIndex < BlendSpace->GetNumberOfBlendSamples())
+	if(const UBlendSpaceBase* BlendSpace = BlendSpaceBase.Get())
 	{
-		const FBlendSample& BlendSample = BlendSpace->GetBlendSample(SelectedSampleIndex);
-		ReturnValue = BlendSample.SampleValue[ParameterIndex];
+		if (SelectedSampleIndex != INDEX_NONE && SelectedSampleIndex < BlendSpace->GetNumberOfBlendSamples())
+		{
+			const FBlendSample& BlendSample = BlendSpace->GetBlendSample(SelectedSampleIndex);
+			ReturnValue = BlendSample.SampleValue[ParameterIndex];
+		}
 	}
 	return ReturnValue;
 }
@@ -1298,10 +1426,10 @@ void SBlendSpaceGridWidget::OnInputBoxValueChanged(const float NewValue, const i
 {
 	checkf(ParameterIndex < 3, TEXT("Invalid parameter index, suppose to be within FVector array range"));
 
-	if (SelectedSampleIndex != INDEX_NONE)
+	if (SelectedSampleIndex != INDEX_NONE && BlendSpaceBase.Get() != nullptr)
 	{
 		// Retrieve current sample value
-		const FBlendSample& Sample = BlendSpace->GetBlendSample(SelectedSampleIndex);
+		const FBlendSample& Sample = BlendSpaceBase.Get()->GetBlendSample(SelectedSampleIndex);
 		FVector SampleValue = Sample.SampleValue;
 
 		// Calculate snapped value
@@ -1328,18 +1456,23 @@ void SBlendSpaceGridWidget::OnInputBoxValueChanged(const float NewValue, const i
 EVisibility SBlendSpaceGridWidget::GetSampleToolTipVisibility() const
 {
 	// Show tool tip when the grid is empty
-	return BlendSpace->GetNumberOfBlendSamples() == 0 ? EVisibility::Visible : EVisibility::Collapsed;
+	return (!bReadOnly && BlendSpaceBase.Get() != nullptr && BlendSpaceBase.Get()->GetNumberOfBlendSamples() == 0) ? EVisibility::Visible : EVisibility::Collapsed;
 }
 
 EVisibility SBlendSpaceGridWidget::GetPreviewToolTipVisibility() const
 {
 	// Only show preview tooltip until the user discovers the functionality
-	return !bPreviewToolTipHidden ? EVisibility::Visible : EVisibility::Collapsed;
+	return (!bReadOnly && !bPreviewToolTipHidden) ? EVisibility::Visible : EVisibility::Collapsed;
 }
 
 EVisibility SBlendSpaceGridWidget::GetTriangulationButtonVisibility() const
 {
-	return (GridType == EGridType::TwoAxis) ? EVisibility::Visible : EVisibility::Collapsed;
+	return (bShowSettingsButtons && (GridType == EGridType::TwoAxis)) ? EVisibility::Visible : EVisibility::Collapsed;
+}
+
+EVisibility SBlendSpaceGridWidget::GetAnimationNamesButtonVisibility() const
+{
+	return bShowSettingsButtons ? EVisibility::Visible : EVisibility::Collapsed;
 }
 
 FReply SBlendSpaceGridWidget::ToggleFittingType()
@@ -1396,53 +1529,55 @@ FText SBlendSpaceGridWidget::GetFittingTypeButtonToolTipText() const
 
 EVisibility SBlendSpaceGridWidget::GetFittingButtonVisibility() const
 {
-	return (GridType == EGridType::TwoAxis) ? EVisibility::Visible : EVisibility::Collapsed;
+	return (bShowSettingsButtons && (GridType == EGridType::TwoAxis)) ? EVisibility::Visible : EVisibility::Collapsed;
 }
 
 void SBlendSpaceGridWidget::UpdateCachedBlendParameterData()
 {
-	checkf(BlendSpace != nullptr, TEXT("Invalid BlendSpace pointer"));
-	const FBlendParameter& BlendParameterX = BlendSpace->GetBlendParameter(0);
-	const FBlendParameter& BlendParameterY = BlendSpace->GetBlendParameter(1);
-	SampleValueRange.X = BlendParameterX.Max - BlendParameterX.Min;
-	SampleValueRange.Y = BlendParameterY.Max - BlendParameterY.Min;
+	if(const UBlendSpaceBase* BlendSpace = BlendSpaceBase.Get())
+	{
+		const FBlendParameter& BlendParameterX = BlendSpace->GetBlendParameter(0);
+		const FBlendParameter& BlendParameterY = BlendSpace->GetBlendParameter(1);
+		SampleValueRange.X = BlendParameterX.Max - BlendParameterX.Min;
+		SampleValueRange.Y = BlendParameterY.Max - BlendParameterY.Min;
 	
-	SampleValueMin.X = BlendParameterX.Min;
-	SampleValueMin.Y = BlendParameterY.Min;
+		SampleValueMin.X = BlendParameterX.Min;
+		SampleValueMin.Y = BlendParameterY.Min;
 
-	SampleValueMax.X = BlendParameterX.Max;
-	SampleValueMax.Y = BlendParameterY.Max;
+		SampleValueMax.X = BlendParameterX.Max;
+		SampleValueMax.Y = BlendParameterY.Max;
 
-	SampleGridDelta = SampleValueRange;
-	SampleGridDelta.X /= (BlendParameterX.GridNum);
-	SampleGridDelta.Y /= (BlendParameterY.GridNum);
+		SampleGridDelta = SampleValueRange;
+		SampleGridDelta.X /= (BlendParameterX.GridNum);
+		SampleGridDelta.Y /= (BlendParameterY.GridNum);
 
-	SampleGridDivisions.X = BlendParameterX.GridNum;
-	SampleGridDivisions.Y = BlendParameterY.GridNum;
+		SampleGridDivisions.X = BlendParameterX.GridNum;
+		SampleGridDivisions.Y = BlendParameterY.GridNum;
 
-	ParameterXName = FText::FromString(BlendParameterX.DisplayName);
-	ParameterYName = FText::FromString(BlendParameterY.DisplayName);
+		ParameterXName = FText::FromString(BlendParameterX.DisplayName);
+		ParameterYName = FText::FromString(BlendParameterY.DisplayName);
 	
-	const TSharedRef< FSlateFontMeasure > FontMeasure = FSlateApplication::Get().GetRenderer()->GetFontMeasureService();
-	MaxVerticalAxisTextWidth = HorizontalAxisMaxTextWidth = MaxHorizontalAxisTextHeight = 0.0f;
-	FVector2D TextSize = FontMeasure->Measure(ParameterYName, FontInfo);	
-	MaxVerticalAxisTextWidth = FMath::Max(MaxVerticalAxisTextWidth, TextSize.X);
+		const TSharedRef< FSlateFontMeasure > FontMeasure = FSlateApplication::Get().GetRenderer()->GetFontMeasureService();
+		MaxVerticalAxisTextWidth = HorizontalAxisMaxTextWidth = MaxHorizontalAxisTextHeight = 0.0f;
+		FVector2D TextSize = FontMeasure->Measure(ParameterYName, FontInfo);	
+		MaxVerticalAxisTextWidth = FMath::Max(MaxVerticalAxisTextWidth, TextSize.X);
 
-	TextSize = FontMeasure->Measure(FString::SanitizeFloat(SampleValueMin.Y), FontInfo);
-	MaxVerticalAxisTextWidth = FMath::Max(MaxVerticalAxisTextWidth, TextSize.X);
+		TextSize = FontMeasure->Measure(FString::SanitizeFloat(SampleValueMin.Y), FontInfo);
+		MaxVerticalAxisTextWidth = FMath::Max(MaxVerticalAxisTextWidth, TextSize.X);
 
-	TextSize = FontMeasure->Measure(FString::SanitizeFloat(SampleValueMax.Y), FontInfo);
-	MaxVerticalAxisTextWidth = FMath::Max(MaxVerticalAxisTextWidth, TextSize.X);
+		TextSize = FontMeasure->Measure(FString::SanitizeFloat(SampleValueMax.Y), FontInfo);
+		MaxVerticalAxisTextWidth = FMath::Max(MaxVerticalAxisTextWidth, TextSize.X);
 	
-	TextSize = FontMeasure->Measure(ParameterXName, FontInfo);
-	MaxHorizontalAxisTextHeight = FMath::Max(MaxHorizontalAxisTextHeight, TextSize.Y);
+		TextSize = FontMeasure->Measure(ParameterXName, FontInfo);
+		MaxHorizontalAxisTextHeight = FMath::Max(MaxHorizontalAxisTextHeight, TextSize.Y);
 
-	TextSize = FontMeasure->Measure(FString::SanitizeFloat(SampleValueMin.X), FontInfo);
-	MaxHorizontalAxisTextHeight = FMath::Max(MaxHorizontalAxisTextHeight, TextSize.Y);
+		TextSize = FontMeasure->Measure(FString::SanitizeFloat(SampleValueMin.X), FontInfo);
+		MaxHorizontalAxisTextHeight = FMath::Max(MaxHorizontalAxisTextHeight, TextSize.Y);
 
-	TextSize = FontMeasure->Measure(FString::SanitizeFloat(SampleValueMax.X), FontInfo);
-	MaxHorizontalAxisTextHeight = FMath::Max(MaxHorizontalAxisTextHeight, TextSize.Y);
-	HorizontalAxisMaxTextWidth = TextSize.X;
+		TextSize = FontMeasure->Measure(FString::SanitizeFloat(SampleValueMax.X), FontInfo);
+		MaxHorizontalAxisTextHeight = FMath::Max(MaxHorizontalAxisTextHeight, TextSize.Y);
+		HorizontalAxisMaxTextWidth = TextSize.X;
+	}
 }
 
 void SBlendSpaceGridWidget::OnMouseEnter(const FGeometry& MyGeometry, const FPointerEvent& MouseEvent)
@@ -1477,133 +1612,149 @@ void SBlendSpaceGridWidget::Tick(const FGeometry& AllottedGeometry, const double
 	HighlightedSampleIndex = INDEX_NONE;
 	const bool bPreviousHighlightPreviewPin = bHighlightPreviewPin;
 
-	if (DragState == EDragState::None)
+	if(const UBlendSpaceBase* BlendSpace = BlendSpaceBase.Get())
 	{
-		// Check if we are highlighting preview pin						
-		bHighlightPreviewPin = IsSampleValueWithinMouseRange(LastPreviewingSampleValue);
-		if (bHighlightPreviewPin)
-		{
-			if (bHighlightPreviewPin != bPreviousHighlightPreviewPin)
-			{
-				ShowToolTip();
-			}
-		}
-		else if (bPreviousHighlightPreviewPin != bHighlightPreviewPin)
+		if(PreviousBlendSpaceBase.Get() != BlendSpace)
 		{
-			ResetToolTip();
-		}
-		
-		// Determine highlighted sample
-		const TArray<FBlendSample>& Samples = BlendSpace->GetBlendSamples();
-		for (int32 SampleIndex = 0; SampleIndex < Samples.Num(); ++SampleIndex)
-		{
-			const FBlendSample& Sample = Samples[SampleIndex];
-			if (IsSampleValueWithinMouseRange(Sample.SampleValue))
-			{
-				HighlightedSampleIndex = SampleIndex;
-				break;
-			}
+			PreviousBlendSpaceBase = BlendSpace;
+			InvalidateCachedData();
 		}
 
-		if (!bHighlightPreviewPin)
+		GridType = BlendSpace->IsA<UBlendSpace1D>() ? EGridType::SingleAxis : EGridType::TwoAxis;
+		BlendParametersToDraw = (GridType == EGridType::SingleAxis) ? 1 : 2;
+
+		if(!bReadOnly)
 		{
-			// If we started selecting or selected a different sample make sure we show/hide the tooltip
-			if (PreviousSampleIndex != HighlightedSampleIndex)
+			if (DragState == EDragState::None)
 			{
-				if (HighlightedSampleIndex != INDEX_NONE)
+				// Check if we are highlighting preview pin						
+				bHighlightPreviewPin = IsSampleValueWithinMouseRange(LastPreviewingSampleValue);
+				if (bHighlightPreviewPin)
 				{
-					ShowToolTip();
+					if (bHighlightPreviewPin != bPreviousHighlightPreviewPin)
+					{
+						ShowToolTip();
+					}
 				}
-				else
+				else if (bPreviousHighlightPreviewPin != bHighlightPreviewPin)
 				{
 					ResetToolTip();
 				}
-			}
+		
+				// Determine highlighted sample
+				const TArray<FBlendSample>& Samples = BlendSpace->GetBlendSamples();
+				for (int32 SampleIndex = 0; SampleIndex < Samples.Num(); ++SampleIndex)
+				{
+					const FBlendSample& Sample = Samples[SampleIndex];
+					if (IsSampleValueWithinMouseRange(Sample.SampleValue))
+					{
+						HighlightedSampleIndex = SampleIndex;
+						break;
+					}
+				}
+
+				if (!bHighlightPreviewPin)
+				{
+					// If we started selecting or selected a different sample make sure we show/hide the tooltip
+					if (PreviousSampleIndex != HighlightedSampleIndex)
+					{
+						if (HighlightedSampleIndex != INDEX_NONE)
+						{
+							ShowToolTip();
+						}
+						else
+						{
+							ResetToolTip();
+						}
+					}
+				}
 		}
-	}
-	else if (DragState == EDragState::DragSample)
-	{
-		// If we are dragging a sample, find out whether or not it has actually moved to a different grid position since the last tick and update the blend space accordingly
-		const FBlendSample& BlendSample = BlendSpace->GetBlendSample(DraggedSampleIndex);
+			else if (DragState == EDragState::DragSample)
+			{
+				// If we are dragging a sample, find out whether or not it has actually moved to a different grid position since the last tick and update the blend space accordingly
+				const FBlendSample& BlendSample = BlendSpace->GetBlendSample(DraggedSampleIndex);
 
-		FVector SampleValue;
-		bool bSnap;
-		if(FSlateApplication::Get().GetModifierKeys().IsAltDown() || !BlendSample.bSnapToGrid)
-		{
-			const FVector2D GridPosition(FMath::Clamp(LocalMousePosition.X, CachedGridRectangle.Left, CachedGridRectangle.Right),
-										 FMath::Clamp(LocalMousePosition.Y, CachedGridRectangle.Top, CachedGridRectangle.Bottom));
+				FVector SampleValue;
+				bool bSnap;
+				if(FSlateApplication::Get().GetModifierKeys().IsAltDown() || !BlendSample.bSnapToGrid)
+				{
+					const FVector2D GridPosition(FMath::Clamp(LocalMousePosition.X, CachedGridRectangle.Left, CachedGridRectangle.Right),
+													FMath::Clamp(LocalMousePosition.Y, CachedGridRectangle.Top, CachedGridRectangle.Bottom));
  
-			SampleValue = GridPositionToSampleValue(GridPosition);
-			bSnap = false;
-		}
-		else
-		{
-			SampleValue = SnapToClosestSamplePoint(LocalMousePosition);
-			bSnap = true;
-		} 
+					SampleValue = GridPositionToSampleValue(GridPosition);
+					bSnap = false;
+				}
+				else
+				{
+					SampleValue = SnapToClosestSamplePoint(LocalMousePosition);
+					bSnap = true;
+				} 
 
-		if (SampleValue != LastDragPosition)
-		{
-			LastDragPosition = SampleValue;
-			OnSampleMoved.ExecuteIfBound(DraggedSampleIndex, SampleValue, false, bSnap);
-		}
-	}
-	else if (DragState == EDragState::DragDrop || DragState == EDragState::InvalidDragDrop || DragState == EDragState::DragDropOverride)
-	{
-		// Validate that the sample is not overlapping with a current sample when doing a drag/drop operation and that we are dropping a valid animation for the blend space (type)
-		const FVector DropSampleValue = SnapToClosestSamplePoint(LocalMousePosition);
-		const bool bValidPosition = BlendSpace->IsSampleWithinBounds(DropSampleValue);
-		const bool bExistingSample = BlendSpace->IsTooCloseToExistingSamplePoint(DropSampleValue, INDEX_NONE);
-		const bool bValidSequence = ValidateAnimationSequence(DragDropAnimationSequence, InvalidDragDropText);
-		
-		if (!bValidSequence)
-		{
-			DragState = EDragState::InvalidDragDrop;
-		}
-		else if (!bValidPosition)
-		{			
-			InvalidDragDropText = InvalidSamplePositionDragDropText;
-			DragState = EDragState::InvalidDragDrop;
-		}
-		else if (bExistingSample)
-		{	
-			const TArray<FBlendSample>& Samples = BlendSpace->GetBlendSamples();
-			for (int32 SampleIndex = 0; SampleIndex < Samples.Num(); ++SampleIndex)
-			{
-				const FBlendSample& Sample = Samples[SampleIndex];
-				if (Sample.SampleValue == DropSampleValue)
+				if (SampleValue != LastDragPosition)
 				{
-					HoveredAnimationName = Sample.Animation ? FText::FromString(Sample.Animation->GetName()) : FText::FromString("Invalid Animation Sequence");
-					break;
+					LastDragPosition = SampleValue;
+					OnSampleMoved.ExecuteIfBound(DraggedSampleIndex, SampleValue, false, bSnap);
 				}
 			}
+			else if (DragState == EDragState::DragDrop || DragState == EDragState::InvalidDragDrop || DragState == EDragState::DragDropOverride)
+			{
+				// Validate that the sample is not overlapping with a current sample when doing a drag/drop operation and that we are dropping a valid animation for the blend space (type)
+				const FVector DropSampleValue = SnapToClosestSamplePoint(LocalMousePosition);
+				const bool bValidPosition = BlendSpace->IsSampleWithinBounds(DropSampleValue);
+				const bool bExistingSample = BlendSpace->IsTooCloseToExistingSamplePoint(DropSampleValue, INDEX_NONE);
+				const bool bValidSequence = ValidateAnimationSequence(DragDropAnimationSequence, InvalidDragDropText);
+		
+				if (!bValidSequence)
+				{
+					DragState = EDragState::InvalidDragDrop;
+				}
+				else if (!bValidPosition)
+				{			
+					InvalidDragDropText = InvalidSamplePositionDragDropText;
+					DragState = EDragState::InvalidDragDrop;
+				}
+				else if (bExistingSample)
+				{	
+					const TArray<FBlendSample>& Samples = BlendSpace->GetBlendSamples();
+					for (int32 SampleIndex = 0; SampleIndex < Samples.Num(); ++SampleIndex)
+					{
+						const FBlendSample& Sample = Samples[SampleIndex];
+						if (Sample.SampleValue == DropSampleValue)
+						{
+							HoveredAnimationName = Sample.Animation ? FText::FromString(Sample.Animation->GetName()) : FText::FromString("Invalid Animation Sequence");
+							break;
+						}
+					}
 
-			DragState = EDragState::DragDropOverride;			
-		}
-		else if (bValidPosition && bValidSequence && !bExistingSample)
-		{
-			DragState = EDragState::DragDrop;
+					DragState = EDragState::DragDropOverride;			
+				}
+				else if (bValidPosition && bValidSequence && !bExistingSample)
+				{
+					DragState = EDragState::DragDrop;
+				}
+			}
 		}
-	}
 
-	// Check if we should update the preview sample value
-	if (bSamplePreviewing)
-	{
-		// Ensure the preview mouse position is clamped to the grid
-		LastPreviewingMousePosition.X = FMath::Clamp(LocalMousePosition.X, CachedGridRectangle.Left, CachedGridRectangle.Right);
-		LastPreviewingMousePosition.Y = FMath::Clamp(LocalMousePosition.Y, CachedGridRectangle.Top, CachedGridRectangle.Bottom);
-		LastPreviewingSampleValue = GridPositionToSampleValue(LastPreviewingMousePosition);
+		// Check if we should update the preview sample value
+		if (bSamplePreviewing)
+		{
+			// Ensure the preview mouse position is clamped to the grid
+			LastPreviewingMousePosition.X = FMath::Clamp(LocalMousePosition.X, CachedGridRectangle.Left, CachedGridRectangle.Right);
+			LastPreviewingMousePosition.Y = FMath::Clamp(LocalMousePosition.Y, CachedGridRectangle.Top, CachedGridRectangle.Bottom);
+			LastPreviewingSampleValue = bReadOnly ? Position.Get() : GridPositionToSampleValue(LastPreviewingMousePosition);
 						
-		// Retrieve and cache weighted samples
-		PreviewedSamples.Empty(4);
-		BlendSpace->GetSamplesFromBlendInput(GetBlendPreviewValue(), PreviewedSamples);
+			// Retrieve and cache weighted samples
+			PreviewedSamples.Empty(4);
+			BlendSpace->GetSamplesFromBlendInput(GetBlendPreviewValue(), PreviewedSamples);
+		}
 	}
 
 	// Refresh cache blendspace/grid data if needed
 	if (bRefreshCachedData)
 	{
 		UpdateCachedBlendParameterData();
-		GridMargin = FMargin(MaxVerticalAxisTextWidth + (TextMargin * 2.0f), TextMargin, (HorizontalAxisMaxTextWidth *.5f) + TextMargin, MaxHorizontalAxisTextHeight + (TextMargin * 2.0f));
+		GridMargin = bShowAxisLabels ?  FMargin(MaxVerticalAxisTextWidth + (TextMargin * 2.0f), TextMargin, (HorizontalAxisMaxTextWidth *.5f) + TextMargin, MaxHorizontalAxisTextHeight + (TextMargin * 2.0f)) : 
+										FMargin(TextMargin, TextMargin, TextMargin, TextMargin);
 		bRefreshCachedData = false;
 	}
 	
@@ -1634,7 +1785,7 @@ void SBlendSpaceGridWidget::InvalidateState()
 		DragState = EDragState::None;
 	}
 	
-	SelectedSampleIndex = BlendSpace->IsValidBlendSampleIndex(SelectedSampleIndex) ? SelectedSampleIndex : INDEX_NONE;
+	SelectedSampleIndex = (BlendSpaceBase.Get() != nullptr && BlendSpaceBase.Get()->IsValidBlendSampleIndex(SelectedSampleIndex)) ? SelectedSampleIndex : INDEX_NONE;
 	HighlightedSampleIndex = DraggedSampleIndex = INDEX_NONE;
 }
 
@@ -1676,36 +1827,39 @@ bool SBlendSpaceGridWidget::ValidateAnimationSequence(const UAnimSequence* Anima
 {	
 	if (AnimationSequence != nullptr)
 	{
-		// If there are any existing blend samples check whether or not the the animation should be additive and if so if the additive matches the existing samples
-		if (BlendSpace->GetNumberOfBlendSamples() > 0)
+		if(const UBlendSpaceBase* BlendSpace = BlendSpaceBase.Get())
 		{
-			const bool bIsAdditive = BlendSpace->ShouldAnimationBeAdditive();
-			if (AnimationSequence->IsValidAdditive() != bIsAdditive)
+			// If there are any existing blend samples check whether or not the the animation should be additive and if so if the additive matches the existing samples
+			if (BlendSpace->GetNumberOfBlendSamples() > 0)
 			{
-				InvalidOperationText = FText::FromString(bIsAdditive ? "Animation should be additive" : "Animation should be non-additive");
-				return false;
+				const bool bIsAdditive = BlendSpace->ShouldAnimationBeAdditive();
+				if (AnimationSequence->IsValidAdditive() != bIsAdditive)
+				{
+					InvalidOperationText = FText::FromString(bIsAdditive ? "Animation should be additive" : "Animation should be non-additive");
+					return false;
+				}
+
+				// If it is the supported additive type, but does not match existing samples
+				if (!BlendSpace->DoesAnimationMatchExistingSamples(AnimationSequence))
+				{
+					InvalidOperationText = FText::FromString("Additive Animation Type does not match existing Samples");
+					return false;
+				}
 			}
 
-			// If it is the supported additive type, but does not match existing samples
-			if (!BlendSpace->DoesAnimationMatchExistingSamples(AnimationSequence))
+			// Check if the supplied animation is of a different additive animation type 
+			if (!BlendSpace->IsAnimationCompatible(AnimationSequence))
 			{
-				InvalidOperationText = FText::FromString("Additive Animation Type does not match existing Samples");
+				InvalidOperationText = FText::FromString("Invalid Additive Animation Type");
 				return false;
 			}
-		}
 
-		// Check if the supplied animation is of a different additive animation type 
-		if (!BlendSpace->IsAnimationCompatible(AnimationSequence))
-		{
-			InvalidOperationText = FText::FromString("Invalid Additive Animation Type");
-			return false;
-		}
-
-		// Check if the supplied animation is compatible with the skeleton
-		if (!BlendSpace->IsAnimationCompatibleWithSkeleton(AnimationSequence))
-		{
-			InvalidOperationText = FText::FromString("Animation is incompatible with the skeleton");
-			return false;
+			// Check if the supplied animation is compatible with the skeleton
+			if (!BlendSpace->IsAnimationCompatibleWithSkeleton(AnimationSequence))
+			{
+				InvalidOperationText = FText::FromString("Animation is incompatible with the skeleton");
+				return false;
+			}
 		}
 	}
 
diff --git a/Engine/Source/Editor/Persona/Private/SAnimationBlendSpaceGridWidget.h b/Engine/Source/Editor/Persona/Private/SAnimationBlendSpaceGridWidget.h
--- a/Engine/Source/Editor/Persona/Private/SAnimationBlendSpaceGridWidget.h
+++ b/Engine/Source/Editor/Persona/Private/SAnimationBlendSpaceGridWidget.h
@@ -30,16 +30,22 @@ class SBlendSpaceGridWidget : public SCompoundWidget, public FNotifyHook
 public:
 	friend class FBlendSampleDetails;
 public:
-	SLATE_BEGIN_ARGS(SBlendSpaceGridWidget)
-		: _BlendSpaceBase(nullptr)
+	SLATE_BEGIN_ARGS(SBlendSpaceGridWidget) 
+		: _ReadOnly(false)
+		, _ShowAxisLabels(true)
+		, _ShowSettingsButtons(true)
 	{}
-	SLATE_ARGUMENT(const UBlendSpaceBase*, BlendSpaceBase)
+		SLATE_ATTRIBUTE(const UBlendSpaceBase*, BlendSpaceBase)
+		SLATE_ATTRIBUTE(FVector, Position)
 		SLATE_ARGUMENT(FNotifyHook*, NotifyHook)
+		SLATE_ARGUMENT(bool, ReadOnly)
+		SLATE_ARGUMENT(bool, ShowAxisLabels)
+		SLATE_ARGUMENT(bool, ShowSettingsButtons)
 		SLATE_EVENT(FOnSampleMoved, OnSampleMoved)
 		SLATE_EVENT(FOnSampleRemoved, OnSampleRemoved)
 		SLATE_EVENT(FOnSampleAdded, OnSampleAdded)
 		SLATE_EVENT(FOnSampleAnimationChanged, OnSampleAnimationChanged)
-		SLATE_END_ARGS()
+	SLATE_END_ARGS()
 
 protected:
 	enum class EGridType
@@ -133,11 +139,13 @@ protected:
 	FReply ProcessClick(const FGeometry& MyGeometry, const FPointerEvent& MouseEvent);
 	/** Toggle for drawing the triangulation (only available for 2D blendspace) */
 	FReply ToggleTriangulationVisibility();
-	/** Returns whether or not the triangulation button should be visibile, see above*/
+	/** Returns whether or not the triangulation button should be visible, see above*/
 	EVisibility GetTriangulationButtonVisibility() const;
-
+	/** Returns whether or not the animation names button should be visible */
+	EVisibility GetAnimationNamesButtonVisibility() const;
 	/** Toggle for how the grid area is determined (only available for 2D blendspace) */
 	FReply ToggleFittingType();
+	/** Returns whether or not the fitting button should be visible */
 	EVisibility GetFittingButtonVisibility() const;
 	/** Calcaultes a margin offset according to whether or not we should take into account the largest axis when creating the grid area*/
 	void UpdateGridRationMargin(const FVector2D& GeometrySize);
@@ -188,7 +196,11 @@ protected:
 	void UpdateCachedBlendParameterData();
 private:
 	/** Currently visualized blendspace (const to ensure changes to it are only made within SAnimationBlendSpace */
-	const UBlendSpaceBase* BlendSpace;
+	TAttribute<const UBlendSpaceBase*> BlendSpaceBase;
+	/** Previous blendspace, cached here to allow us to update cached data when it changes. */
+	TWeakObjectPtr<const UBlendSpaceBase> PreviousBlendSpaceBase;
+	/** Current position (driven from outside the widget) */
+	TAttribute<FVector> Position;
 	/** Notfiy hook (ptr to SAnimationBlendSpace instance), which is required for transacting FBlendSample object when edited using the context-menu/structure details panel */
 	FNotifyHook* NotifyHook;
 	/** Number of Blend Parameters to draw */
@@ -303,4 +315,10 @@ private:
 	FMargin GridRatioMargin;
 
 	bool bPreviewToolTipHidden;
+
+	bool bReadOnly;
+
+	bool bShowAxisLabels;
+
+	bool bShowSettingsButtons;
 };
diff --git a/Engine/Source/Editor/Persona/Public/PersonaModule.h b/Engine/Source/Editor/Persona/Public/PersonaModule.h
--- a/Engine/Source/Editor/Persona/Public/PersonaModule.h
+++ b/Engine/Source/Editor/Persona/Public/PersonaModule.h
@@ -23,6 +23,7 @@ class IPinnedCommandList;
 class FWorkflowAllowedTabSet;
 class IAssetFamily;
 class FWorkflowTabFactory;
+class UBlendSpaceBase;
 
 extern const FName PersonaAppName;
 
@@ -244,6 +245,9 @@ public:
 	/** Create a tab factory for editing slot names and groups */
 	virtual TSharedRef<FWorkflowTabFactory> CreateSkeletonSlotNamesTabFactory(const TSharedRef<class FWorkflowCentricApplication>& InHostingApp, const TSharedRef<class IEditableSkeleton>& InEditableSkeleton, FSimpleMulticastDelegate& InOnPostUndo, FOnObjectSelected InOnObjectSelected) const;
 
+	/** Create a widget to preview a blendspace */
+	virtual TSharedRef<SWidget> CreateBlendSpacePreviewWidget(TAttribute<const UBlendSpaceBase*> InBlendSpace, TAttribute<FVector> InPosition) const;
+
 	/** Create a widget that acts as a document for an animation asset */
 	virtual TSharedRef<SWidget> CreateEditorWidgetForAnimDocument(const TSharedRef<class FWorkflowCentricApplication>& InHostingApp, UObject* InAnimAsset, const FAnimDocumentArgs& InArgs, FString& OutDocumentLink);
 
