Fix not being able to add anim sequence frames from the curve editor view

Make sure that the popup's widget path is valid in all call circumstances by passing it through directly.

diff --git a/Engine/Source/Editor/Persona/Private/AnimTimeline/SAnimTimeline.cpp b/Engine/Source/Editor/Persona/Private/AnimTimeline/SAnimTimeline.cpp
--- a/Engine/Source/Editor/Persona/Private/AnimTimeline/SAnimTimeline.cpp
+++ b/Engine/Source/Editor/Persona/Private/AnimTimeline/SAnimTimeline.cpp
@@ -345,6 +345,8 @@ FReply SAnimTimeline::OnMouseButtonUp(const FGeometry& MyGeometry, const FPointe
 {
 	if(MouseEvent.GetEffectingButton() == EKeys::RightMouseButton)
 	{
+		FWidgetPath WidgetPath = MouseEvent.GetEventPath() != nullptr ? *MouseEvent.GetEventPath() : FWidgetPath();
+
 		const bool bCloseAfterSelection = true;
 		FMenuBuilder MenuBuilder(bCloseAfterSelection, Model.Pin()->GetCommandList());
 
@@ -432,10 +434,10 @@ FReply SAnimTimeline::OnMouseButtonUp(const FGeometry& MyGeometry, const FPointe
 				MenuBuilder.AddMenuSeparator();
 
 				//Corrected frame time based on selected frame number
-				Action = FUIAction(FExecuteAction::CreateSP(this, &SAnimTimeline::OnShowPopupOfAppendAnimation, true));
+				Action = FUIAction(FExecuteAction::CreateSP(this, &SAnimTimeline::OnShowPopupOfAppendAnimation, WidgetPath, true));
 				MenuBuilder.AddMenuEntry(LOCTEXT("AppendBegin", "Append in the beginning"), LOCTEXT("AppendBegin_ToolTip", "Append in the beginning"), FSlateIcon(), Action);
 
-				Action = FUIAction(FExecuteAction::CreateSP(this, &SAnimTimeline::OnShowPopupOfAppendAnimation, false));
+				Action = FUIAction(FExecuteAction::CreateSP(this, &SAnimTimeline::OnShowPopupOfAppendAnimation, WidgetPath, false));
 				MenuBuilder.AddMenuEntry(LOCTEXT("AppendEnd", "Append at the end"), LOCTEXT("AppendEnd_ToolTip", "Append at the end"), FSlateIcon(), Action);
 
 				MenuBuilder.AddMenuSeparator();
@@ -452,8 +454,6 @@ FReply SAnimTimeline::OnMouseButtonUp(const FGeometry& MyGeometry, const FPointe
 			MenuBuilder.EndSection();
 		}
 
-		FWidgetPath WidgetPath = MouseEvent.GetEventPath() != nullptr ? *MouseEvent.GetEventPath() : FWidgetPath();
-
 		FSlateApplication::Get().PushMenu(SharedThis(this), WidgetPath, MenuBuilder.MakeWidget(), FSlateApplication::Get().GetCursorPos(), FPopupTransitionEffect(FPopupTransitionEffect::ContextMenu));
 
 		return FReply::Handled();
@@ -602,7 +602,7 @@ void SAnimTimeline::OnReZeroAnimSequence(int32 FrameIndex)
 	}
 }
 
-void SAnimTimeline::OnShowPopupOfAppendAnimation(bool bBegin)
+void SAnimTimeline::OnShowPopupOfAppendAnimation(FWidgetPath WidgetPath, bool bBegin)
 {
 	TSharedRef<STextEntryPopup> TextEntry =
 		SNew(STextEntryPopup)
@@ -612,7 +612,7 @@ void SAnimTimeline::OnShowPopupOfAppendAnimation(bool bBegin)
 	// Show dialog to enter new track name
 	FSlateApplication::Get().PushMenu(
 		SharedThis(this),
-		FWidgetPath(),
+		WidgetPath,
 		TextEntry,
 		FSlateApplication::Get().GetCursorPos(),
 		FPopupTransitionEffect(FPopupTransitionEffect::TypeInPopup)
diff --git a/Engine/Source/Editor/Persona/Private/AnimTimeline/SAnimTimeline.h b/Engine/Source/Editor/Persona/Private/AnimTimeline/SAnimTimeline.h
--- a/Engine/Source/Editor/Persona/Private/AnimTimeline/SAnimTimeline.h
+++ b/Engine/Source/Editor/Persona/Private/AnimTimeline/SAnimTimeline.h
@@ -88,7 +88,7 @@ private:
 
 	void OnReZeroAnimSequence(int32 FrameIndex);
 
-	void OnShowPopupOfAppendAnimation(bool bBegin);
+	void OnShowPopupOfAppendAnimation(FWidgetPath WidgetPath, bool bBegin);
 
 	void OnSequenceAppendedCalled(const FText & InNewGroupText, ETextCommit::Type CommitInfo, bool bBegin);
 
