Fixed crash undo-ing a slot addition in montage editor

diff --git a/Engine/Source/Editor/Persona/Private/AnimTimeline/AnimTimelineTrack_MontagePanel.cpp b/Engine/Source/Editor/Persona/Private/AnimTimeline/AnimTimelineTrack_MontagePanel.cpp
--- a/Engine/Source/Editor/Persona/Private/AnimTimeline/AnimTimelineTrack_MontagePanel.cpp
+++ b/Engine/Source/Editor/Persona/Private/AnimTimeline/AnimTimelineTrack_MontagePanel.cpp
@@ -75,8 +75,14 @@ void FAnimTimelineTrack_MontagePanel::RefreshOutlinerWidget()
 
 		auto LabelLambda = [AnimMontage, Skeleton, SlotIndex]()
 		{ 
-			FText SlotItemText = FText::FromString(FString::Printf(TEXT("%s.%s"), *Skeleton->GetSlotGroupName(AnimMontage->SlotAnimTracks[SlotIndex].SlotName).ToString(), *AnimMontage->SlotAnimTracks[SlotIndex].SlotName.ToString()));
-			return AnimMontage->SlotAnimTracks.IsValidIndex(SlotIndex) ? SlotItemText : FText::GetEmpty(); 
+			if(AnimMontage->SlotAnimTracks.IsValidIndex(SlotIndex))
+			{
+				return FText::FromString(FString::Printf(TEXT("%s.%s"), *Skeleton->GetSlotGroupName(AnimMontage->SlotAnimTracks[SlotIndex].SlotName).ToString(), *AnimMontage->SlotAnimTracks[SlotIndex].SlotName.ToString()));
+			}
+			else
+			{
+				return FText::GetEmpty();
+			}
 		};
 
 		TSharedPtr<SHorizontalBox> HorizontalBox;
