Montage sections panel now scrolls when it gets too large

diff --git a/Engine/Source/Editor/Persona/Private/SAnimMontageSectionsPanel.cpp b/Engine/Source/Editor/Persona/Private/SAnimMontageSectionsPanel.cpp
--- a/Engine/Source/Editor/Persona/Private/SAnimMontageSectionsPanel.cpp
+++ b/Engine/Source/Editor/Persona/Private/SAnimMontageSectionsPanel.cpp
@@ -20,6 +20,7 @@
 #include "Framework/MultiBox/MultiBoxBuilder.h"
 #include "AnimPreviewInstance.h"
 #include "Animation/DebugSkelMeshComponent.h"
+#include "Widgets/Layout/SScrollBox.h"
 
 #define LOCTEXT_NAMESPACE "AnimMontageSectionsPanel"
 
@@ -108,12 +109,12 @@ void SAnimMontageSectionsPanel::Update()
 		bool bChildAnimMontage = Montage->HasParentAsset();
 
 		TSharedPtr<STrack> Track;
-		TSharedPtr<SVerticalBox> MontageSlots;
+		TSharedPtr<SVerticalBox> ContentArea;
 		PanelArea->SetContent(
-			SAssignNew( MontageSlots, SVerticalBox )
+			SAssignNew( ContentArea, SVerticalBox )
 			);
 
-		MontageSlots->ClearChildren();
+		ContentArea->ClearChildren();
 		
 		TArray<bool>	Used;
 		Used.AddZeroed(Montage->CompositeSections.Num());
@@ -121,7 +122,7 @@ void SAnimMontageSectionsPanel::Update()
 		int RowIdx=0;
 
 		/** Create Buttons for reseting/creating default section ordering */
-		MontageSlots->AddSlot()
+		ContentArea->AddSlot()
 			.AutoHeight()
 			.VAlign(VAlign_Center)
 			[
@@ -230,6 +231,16 @@ void SAnimMontageSectionsPanel::Update()
 				]
 			];
 
+		TSharedPtr<SVerticalBox> MontageSlots;
+		ContentArea->AddSlot()
+		[
+			SNew(SScrollBox)
+			+SScrollBox::Slot()
+			[
+				SAssignNew(MontageSlots, SVerticalBox)
+			]
+		];
+
 		/** Create as many tracks as necessary to show each section at least once
 		  * -Each track represents one chain of sections (section->next->next)
 		  */
