Fix crash when undoing curve edits

Switched curve viewer to use 'named curve' model the same as the curve editor. This is needed because the tree view doesnt refresh straight away, so tracks dont get removed until the tick after an undo is fired, while the undo event gets pushed to the track even when it isnt hooked up to the tree. Using the named curve approach means that the invalid track doesnt end up accessing a deleted curve.

diff --git a/Engine/Source/Editor/Persona/Private/AnimTimeline/AnimTimelineTrack_Curve.cpp b/Engine/Source/Editor/Persona/Private/AnimTimeline/AnimTimelineTrack_Curve.cpp
--- a/Engine/Source/Editor/Persona/Private/AnimTimeline/AnimTimelineTrack_Curve.cpp
+++ b/Engine/Source/Editor/Persona/Private/AnimTimeline/AnimTimelineTrack_Curve.cpp
@@ -20,6 +20,7 @@
 #include "IPersonaPreviewScene.h"
 #include "Animation/DebugSkelMeshComponent.h"
 #include "AnimPreviewInstance.h"
+#include "SAnimSequenceCurveEditor.h"
 
 #define LOCTEXT_NAMESPACE "FAnimTimelineTrack_Curve"
 
@@ -194,7 +195,12 @@ TSharedRef<SWidget> FAnimTimelineTrack_Curve::GenerateContainerWidgetForTimeline
 	{
 		FRichCurve* Curve = Curves[CurveIndex];
 
-		TUniquePtr<FRichCurveEditorModelRaw> NewCurveModel = MakeUnique<FRichCurveEditorModelRaw>(Curve, GetModel()->GetAnimSequenceBase());
+		FSmartName Name;
+		ERawCurveTrackTypes Type;
+		int32 EditIndex;
+		GetCurveEditInfo(CurveIndex, Name, Type, EditIndex);
+
+		TUniquePtr<FRichCurveEditorModelNamed> NewCurveModel = MakeUnique<FRichCurveEditorModelNamed>(Name, Type, EditIndex, GetModel()->GetAnimSequenceBase());
 		NewCurveModel->SetColor(GetCurveColor(CurveIndex));
 		NewCurveModel->SetIsKeyDrawEnabled(MakeAttributeLambda([](){ return GetDefault<UPersonaOptions>()->bTimelineDisplayCurveKeys; }));
 		CurveEditor->AddCurve(MoveTemp(NewCurveModel));
diff --git a/Engine/Source/Editor/Persona/Private/SAnimSequenceCurveEditor.cpp b/Engine/Source/Editor/Persona/Private/SAnimSequenceCurveEditor.cpp
--- a/Engine/Source/Editor/Persona/Private/SAnimSequenceCurveEditor.cpp
+++ b/Engine/Source/Editor/Persona/Private/SAnimSequenceCurveEditor.cpp
@@ -19,87 +19,65 @@
 
 #define LOCTEXT_NAMESPACE "SAnimSequenceCurveEditor"
 
-// Model that references a named curve, rather than a raw pointer, so we avoid issues with 
-// reallocation of the curves arrays under the UI
-class FRichCurveEditorModelNamed : public FRichCurveEditorModel
+bool FRichCurveEditorModelNamed::IsValid() const
 {
-public:
-	FRichCurveEditorModelNamed(const FSmartName& InName, ERawCurveTrackTypes InType, int32 InCurveIndex, UAnimSequenceBase* InAnimSequence, FCurveEditorTreeItemID InTreeId)
-		: FRichCurveEditorModel(InAnimSequence)
-		, Name(InName)
-		, AnimSequence(InAnimSequence)
-		, CurveIndex(InCurveIndex)
-		, Type(InType)
-		, TreeId(InTreeId)
-	{
-	}
+	return AnimSequence->GetCurveData().GetCurveData(Name.UID, Type) != nullptr;
+}
+
+FRichCurve& FRichCurveEditorModelNamed::GetRichCurve()
+{
+	check(AnimSequence.Get() != nullptr);
 
-	virtual bool IsValid() const override
+	FAnimCurveBase* CurveBase = AnimSequence->RawCurveData.GetCurveData(Name.UID, Type);
+	check(CurveBase);	// If this fails lifetime contracts have been violated - this curve should always be present if this model exists
+		
+	switch (Type)
+	{
+	case ERawCurveTrackTypes::RCT_Vector:
 	{
-		return AnimSequence->GetCurveData().GetCurveData(Name.UID, Type) != nullptr;
+		FVectorCurve& VectorCurve = *(static_cast<FVectorCurve*>(CurveBase));
+		check(CurveIndex < 3);
+		return VectorCurve.FloatCurves[CurveIndex];
 	}
-
-	virtual FRichCurve& GetRichCurve() override
+	case ERawCurveTrackTypes::RCT_Transform:
 	{
-		check(AnimSequence.Get() != nullptr);
-
-		FAnimCurveBase* CurveBase = AnimSequence->RawCurveData.GetCurveData(Name.UID, Type);
-		check(CurveBase);	// If this fails lifetime contracts have been violated - this curve should always be present if this model exists
-		
-		switch (Type)
-		{
-		case ERawCurveTrackTypes::RCT_Vector:
-		{
-			FVectorCurve& VectorCurve = *(static_cast<FVectorCurve*>(CurveBase));
-			check(CurveIndex < 3);
-			return VectorCurve.FloatCurves[CurveIndex];
-		}
-		case ERawCurveTrackTypes::RCT_Transform:
+		FTransformCurve& TransformCurve = *(static_cast<FTransformCurve*>(CurveBase));
+		check(CurveIndex < 9);
+		const int32 SubCurveIndex = CurveIndex % 3;
+		switch(CurveIndex)
 		{
-			FTransformCurve& TransformCurve = *(static_cast<FTransformCurve*>(CurveBase));
-			check(CurveIndex < 9);
-			const int32 SubCurveIndex = CurveIndex % 3;
-			switch(CurveIndex)
-			{
-			default:
-				check(false);
-				// fall through
-			case 0:
-			case 1:
-			case 2:
-				return TransformCurve.TranslationCurve.FloatCurves[SubCurveIndex];
-			case 3:
-			case 4:
-			case 5:
-				return TransformCurve.RotationCurve.FloatCurves[SubCurveIndex];
-			case 6:
-			case 7:
-			case 8:
-				return TransformCurve.ScaleCurve.FloatCurves[SubCurveIndex];
-			}
-			
-		}
-		case ERawCurveTrackTypes::RCT_Float:
 		default:
-		{
-			FFloatCurve& FloatCurve = *(static_cast<FFloatCurve*>(CurveBase));
-			check(CurveIndex == 0);
-			return FloatCurve.FloatCurve;
-		}
+			check(false);
+			// fall through
+		case 0:
+		case 1:
+		case 2:
+			return TransformCurve.TranslationCurve.FloatCurves[SubCurveIndex];
+		case 3:
+		case 4:
+		case 5:
+			return TransformCurve.RotationCurve.FloatCurves[SubCurveIndex];
+		case 6:
+		case 7:
+		case 8:
+			return TransformCurve.ScaleCurve.FloatCurves[SubCurveIndex];
 		}
+			
 	}
-
-	virtual const FRichCurve& GetReadOnlyRichCurve() const override
+	case ERawCurveTrackTypes::RCT_Float:
+	default:
 	{
-		return const_cast<FRichCurveEditorModelNamed*>(this)->GetRichCurve();
+		FFloatCurve& FloatCurve = *(static_cast<FFloatCurve*>(CurveBase));
+		check(CurveIndex == 0);
+		return FloatCurve.FloatCurve;
+	}
 	}
+}
 
-	FSmartName Name;
-	TWeakObjectPtr<UAnimSequenceBase> AnimSequence;
-	int32 CurveIndex;
-	ERawCurveTrackTypes Type;
-	FCurveEditorTreeItemID TreeId;
-};
+const FRichCurve& FRichCurveEditorModelNamed::GetReadOnlyRichCurve() const
+{
+	return const_cast<FRichCurveEditorModelNamed*>(this)->GetRichCurve();
+}
 
 class FAnimSequenceCurveEditorItem : public ICurveEditorTreeItem
 {
diff --git a/Engine/Source/Editor/Persona/Private/SAnimSequenceCurveEditor.h b/Engine/Source/Editor/Persona/Private/SAnimSequenceCurveEditor.h
--- a/Engine/Source/Editor/Persona/Private/SAnimSequenceCurveEditor.h
+++ b/Engine/Source/Editor/Persona/Private/SAnimSequenceCurveEditor.h
@@ -5,6 +5,10 @@
 #include "IAnimSequenceCurveEditor.h"
 #include "Widgets/DeclarativeSyntaxSupport.h"
 #include "EditorUndoClient.h"
+#include "RichCurveEditorModel.h"
+#include "Animation/SmartName.h"
+#include "Animation/AnimCurveTypes.h"
+#include "CurveEditorTypes.h"
 
 class FCurveEditor;
 class UAnimSequenceBase;
@@ -14,6 +18,32 @@ class IPersonaPreviewScene;
 class SCurveEditorPanel;
 class FTabManager;
 
+// Model that references a named curve, rather than a raw pointer, so we avoid issues with 
+// reallocation of the curves arrays under the UI
+class FRichCurveEditorModelNamed : public FRichCurveEditorModel
+{
+public:
+	FRichCurveEditorModelNamed(const FSmartName& InName, ERawCurveTrackTypes InType, int32 InCurveIndex, UAnimSequenceBase* InAnimSequence, FCurveEditorTreeItemID InTreeId = FCurveEditorTreeItemID())
+		: FRichCurveEditorModel(InAnimSequence)
+		, Name(InName)
+		, AnimSequence(InAnimSequence)
+		, CurveIndex(InCurveIndex)
+		, Type(InType)
+		, TreeId(InTreeId)
+	{
+	}
+
+	virtual bool IsValid() const override;
+	virtual FRichCurve& GetRichCurve() override;
+	virtual const FRichCurve& GetReadOnlyRichCurve() const override;
+
+	FSmartName Name;
+	TWeakObjectPtr<UAnimSequenceBase> AnimSequence;
+	int32 CurveIndex;
+	ERawCurveTrackTypes Type;
+	FCurveEditorTreeItemID TreeId;
+};
+
 class SAnimSequenceCurveEditor : public IAnimSequenceCurveEditor, public FEditorUndoClient
 {
 	SLATE_BEGIN_ARGS(SAnimSequenceCurveEditor) {}
