Fix non-unity build

diff --git a/Engine/Source/Editor/Persona/Private/SAnimSequenceCurveEditor.h b/Engine/Source/Editor/Persona/Private/SAnimSequenceCurveEditor.h
--- a/Engine/Source/Editor/Persona/Private/SAnimSequenceCurveEditor.h
+++ b/Engine/Source/Editor/Persona/Private/SAnimSequenceCurveEditor.h
@@ -9,9 +9,9 @@
 #include "Animation/SmartName.h"
 #include "Animation/AnimCurveTypes.h"
 #include "CurveEditorTypes.h"
+#include "Animation/AnimSequenceBase.h"
 
 class FCurveEditor;
-class UAnimSequenceBase;
 class ITimeSliderController;
 class SCurveEditorTree;
 class IPersonaPreviewScene;
