In the skeletalmesh reduction, add a MaxVerts/tris option to the percent Verts/tris termination criterion mode. This allow to put a percentage but still have a budget limit.

diff --git a/Engine/Plugins/Experimental/SkeletalReduction/Source/Private/SkeletalMeshReductionPlugin.cpp b/Engine/Plugins/Experimental/SkeletalReduction/Source/Private/SkeletalMeshReductionPlugin.cpp
--- a/Engine/Plugins/Experimental/SkeletalReduction/Source/Private/SkeletalMeshReductionPlugin.cpp
+++ b/Engine/Plugins/Experimental/SkeletalReduction/Source/Private/SkeletalMeshReductionPlugin.cpp
@@ -74,12 +74,12 @@ public:
 		{
 			case SkeletalMeshTerminationCriterion::SMTC_NumOfTriangles:
 			{
-				return ReductionSettings.NumOfTrianglesPercentage < Threshold_One;
+				return ReductionSettings.NumOfTrianglesPercentage < Threshold_One || ReductionSettings.MaxNumOfTrianglesPercentage < NumTriangles;
 			}
 			break;
 			case SkeletalMeshTerminationCriterion::SMTC_NumOfVerts:
 			{
-				return ReductionSettings.NumOfVertPercentage < Threshold_One;
+				return ReductionSettings.NumOfVertPercentage < Threshold_One || ReductionSettings.MaxNumOfVertsPercentage < NumVertices;
 			}
 			break;
 			case SkeletalMeshTerminationCriterion::SMTC_TriangleOrVert:
@@ -930,14 +930,23 @@ float FQuadricSkeletalMeshReduction::SimplifyMesh( const FSkeletalMeshOptimizati
 	const float MaxDist = FLT_MAX; // (Settings.ReductionMethod != SkeletalMeshOptimizationType::SMOT_NumOfTriangles) ? Settings.MaxDeviationPercentage * Bounds.SphereRadius : FLT_MAX;
 	const int32 SrcTriNum = Mesh.NumIndices() / 3;
 	const float TriangleRetainRatio = FMath::Clamp(Settings.NumOfTrianglesPercentage, 0.f, 1.f);
-	const int32 TargetTriNum = (bUseTrianglePercentCriterion) ? FMath::CeilToInt(TriangleRetainRatio * SrcTriNum) : Settings.MaxNumOfTriangles;
+	int32 TargetTriNum = (bUseTrianglePercentCriterion) ? FMath::CeilToInt(TriangleRetainRatio * SrcTriNum) : Settings.MaxNumOfTriangles;
+	if (Settings.TerminationCriterion == SkeletalMeshTerminationCriterion::SMTC_NumOfTriangles)
+	{
+		TargetTriNum = FMath::Min(TargetTriNum, int32(Settings.MaxNumOfTrianglesPercentage));
+	}
 
 	const int32 MinTriNumToRetain = (bUseTrianglePercentCriterion || bUseMaxTrisNumCriterion) ? FMath::Max(4, TargetTriNum) : 4;
 	const float MaxCollapseCost = FLT_MAX;
 
 	const int32 SrcVertNum = Mesh.NumVertices();
 	const float VertRetainRatio = FMath::Clamp(Settings.NumOfVertPercentage, 0.f, 1.f);
-	const int32 TargetVertNum = (bUseVertexPercentCriterion) ? FMath::CeilToInt(VertRetainRatio * SrcVertNum) : Settings.MaxNumOfVerts + 1;
+	int32 TargetVertNum = (bUseVertexPercentCriterion) ? FMath::CeilToInt(VertRetainRatio * SrcVertNum) : Settings.MaxNumOfVerts + 1;
+	if (Settings.TerminationCriterion == SkeletalMeshTerminationCriterion::SMTC_NumOfVerts)
+	{
+		TargetVertNum = FMath::Min(TargetVertNum, int32(Settings.MaxNumOfVertsPercentage + 1));
+	}
+
 	const int32 MinVerNumToRetain = (bUseVertexPercentCriterion || bUseMaxVertNumCriterion) ? FMath::Max(6, TargetVertNum) : 6;
 
 	const float VolumeImportance      = FMath::Clamp(Settings.VolumeImportance, 0.f, 2.f);
@@ -1661,7 +1670,7 @@ bool FQuadricSkeletalMeshReduction::ReduceSkeletalLODModel( const FSkeletalMeshL
 
 	int32 IterationNum = 0;
 	//We keep the original MaxNumVerts because if we iterate we want to still compare with the original request.
-	uint32 OriginalMaxNumVertsSetting = Settings.MaxNumOfVerts;
+	uint32 OriginalMaxNumVertsSetting = bUseVertexPercentCriterion ? Settings.MaxNumOfVertsPercentage : Settings.MaxNumOfVerts;
 	do 
 	{
 		if (bOptimizeMesh)
@@ -1713,7 +1722,14 @@ bool FQuadricSkeletalMeshReduction::ReduceSkeletalLODModel( const FSkeletalMeshL
 			{
 				// Some verts were created by chunking - we need simplify more.
 				int32 ExcessVerts = (int32)(OutSkeletalMeshLODModel.NumVertices - OriginalMaxNumVertsSetting);
-				Settings.MaxNumOfVerts = FMath::Max((int32)Settings.MaxNumOfVerts - ExcessVerts, 6);
+				if (bUseVertexPercentCriterion)
+				{
+					Settings.MaxNumOfVertsPercentage = FMath::Max((int32)Settings.MaxNumOfVertsPercentage - ExcessVerts, 6);
+				}
+				else
+				{
+					Settings.MaxNumOfVerts = FMath::Max((int32)Settings.MaxNumOfVerts - ExcessVerts, 6);
+				}
 
 				UE_LOG(LogSkeletalMeshReduction, Log, TEXT("Chunking to limit unique bones per section generated additional vertices - continuing simplification of LOD %d "), LODIndex);
 				ConvertToFSkinnedSkeletalMesh(SrcModel, BoneMatrices, LODIndex, SkinnedSkeletalMesh);
diff --git a/Engine/Source/Editor/DetailCustomizations/Private/SkeletalMeshReductionSettingsDetails.cpp b/Engine/Source/Editor/DetailCustomizations/Private/SkeletalMeshReductionSettingsDetails.cpp
--- a/Engine/Source/Editor/DetailCustomizations/Private/SkeletalMeshReductionSettingsDetails.cpp
+++ b/Engine/Source/Editor/DetailCustomizations/Private/SkeletalMeshReductionSettingsDetails.cpp
@@ -174,6 +174,8 @@ void FSkeletalMeshReductionSettingsDetails::CustomizeChildren(TSharedRef<IProper
 			GET_MEMBER_NAME_CHECKED(FSkeletalMeshOptimizationSettings, NumOfVertPercentage),
 			GET_MEMBER_NAME_CHECKED(FSkeletalMeshOptimizationSettings, MaxNumOfVerts),
 			GET_MEMBER_NAME_CHECKED(FSkeletalMeshOptimizationSettings, MaxNumOfTriangles),
+			GET_MEMBER_NAME_CHECKED(FSkeletalMeshOptimizationSettings, MaxNumOfVertsPercentage),
+			GET_MEMBER_NAME_CHECKED(FSkeletalMeshOptimizationSettings, MaxNumOfTrianglesPercentage),
 			GET_MEMBER_NAME_CHECKED(FSkeletalMeshOptimizationSettings, TerminationCriterion),
 			GET_MEMBER_NAME_CHECKED(FSkeletalMeshOptimizationSettings, bLockEdges),
 			GET_MEMBER_NAME_CHECKED(FSkeletalMeshOptimizationSettings, bEnforceBoneBoundaries),
@@ -255,6 +257,9 @@ void FSkeletalMeshReductionSettingsDetails::CustomizeChildren(TSharedRef<IProper
 		TSharedPtr<IPropertyHandle> MaxNumOfVertsPropertyHandle = PropertyHandles.FindChecked(GET_MEMBER_NAME_CHECKED(FSkeletalMeshOptimizationSettings, MaxNumOfVerts));
 		TSharedPtr<IPropertyHandle> MaxNumOfTrisPropertyHandle = PropertyHandles.FindChecked(GET_MEMBER_NAME_CHECKED(FSkeletalMeshOptimizationSettings, MaxNumOfTriangles));
 
+		TSharedPtr<IPropertyHandle> MaxNumOfVertsPercentagePropertyHandle = PropertyHandles.FindChecked(GET_MEMBER_NAME_CHECKED(FSkeletalMeshOptimizationSettings, MaxNumOfVertsPercentage));
+		TSharedPtr<IPropertyHandle> MaxNumOfTrisPercentagePropertyHandle = PropertyHandles.FindChecked(GET_MEMBER_NAME_CHECKED(FSkeletalMeshOptimizationSettings, MaxNumOfTrianglesPercentage));
+
 		for (auto Iter(PropertyHandles.CreateIterator()); Iter; ++Iter)
 		{
 			if (UnwantedPropertyHandles.Contains(Iter.Value())) 
@@ -301,6 +306,15 @@ void FSkeletalMeshReductionSettingsDetails::CustomizeChildren(TSharedRef<IProper
 					// Hide property if using triangle percentage
 					SettingsRow.Visibility(TAttribute<EVisibility>::Create(TAttribute<EVisibility>::FGetter::CreateSP(this, &FSkeletalMeshReductionSettingsDetails::ShowIfCurrentCriterionIs, VizList)));
 				}
+				else if (Iter.Value() == MaxNumOfVertsPercentagePropertyHandle)
+				{
+					const TArray< SkeletalMeshTerminationCriterion > VizList =
+					{
+						SMTC_NumOfVerts
+					};
+					// Hide property if using triangle percentage
+					SettingsRow.Visibility(TAttribute<EVisibility>::Create(TAttribute<EVisibility>::FGetter::CreateSP(this, &FSkeletalMeshReductionSettingsDetails::ShowIfCurrentCriterionIs, VizList)));
+				}
 				else if (Iter.Value() == MaxNumOfTrisPropertyHandle)
 				{
 					const TArray< SkeletalMeshTerminationCriterion > VizList = 
@@ -311,7 +325,15 @@ void FSkeletalMeshReductionSettingsDetails::CustomizeChildren(TSharedRef<IProper
 					// Hide property if using triangle percentage
 					SettingsRow.Visibility(TAttribute<EVisibility>::Create(TAttribute<EVisibility>::FGetter::CreateSP(this, &FSkeletalMeshReductionSettingsDetails::ShowIfCurrentCriterionIs, VizList)));
 				}
-
+				else if (Iter.Value() == MaxNumOfTrisPercentagePropertyHandle)
+				{
+					const TArray< SkeletalMeshTerminationCriterion > VizList =
+					{
+						SMTC_NumOfTriangles
+					};
+					// Hide property if using triangle percentage
+					SettingsRow.Visibility(TAttribute<EVisibility>::Create(TAttribute<EVisibility>::FGetter::CreateSP(this, &FSkeletalMeshReductionSettingsDetails::ShowIfCurrentCriterionIs, VizList)));
+				}
 			}
 		}
 
diff --git a/Engine/Source/Editor/Persona/Private/PersonaMeshDetails.cpp b/Engine/Source/Editor/Persona/Private/PersonaMeshDetails.cpp
--- a/Engine/Source/Editor/Persona/Private/PersonaMeshDetails.cpp
+++ b/Engine/Source/Editor/Persona/Private/PersonaMeshDetails.cpp
@@ -668,7 +668,14 @@ void FSkeletalMeshReductionSettingsLayout::GenerateChildContent(IDetailChildrenB
 				FGetIntegerDelegate::CreateRaw(this, &FSkeletalMeshReductionSettingsLayout::GetNumMaxTrianglesCount),
 				FSetIntegerDelegate::CreateRaw(this, &FSkeletalMeshReductionSettingsLayout::SetNumMaxTrianglesCount));
 
-			SetPercentAndAbsoluteVisibility(MaxTrianglesRow, SMTC_AbsNumOfTriangles, SMTC_AbsTriangleOrVert);
+			const TArray< SkeletalMeshTerminationCriterion > VizList =
+			{
+				SMTC_AbsNumOfTriangles,
+				SMTC_AbsTriangleOrVert
+			};
+
+			MaxTrianglesRow.Visibility(TAttribute<EVisibility>::Create(TAttribute<EVisibility>::FGetter::CreateSP(this, &FSkeletalMeshReductionSettingsLayout::ShowIfCurrentCriterionIs, VizList)));
+			
 		}
 
 		{
@@ -680,8 +687,51 @@ void FSkeletalMeshReductionSettingsLayout::GenerateChildContent(IDetailChildrenB
 				INT_MAX,
 				FGetIntegerDelegate::CreateRaw(this, &FSkeletalMeshReductionSettingsLayout::GetNumMaxVerticesCount),
 				FSetIntegerDelegate::CreateRaw(this, &FSkeletalMeshReductionSettingsLayout::SetNumMaxVerticesCount));
+				
+			const TArray< SkeletalMeshTerminationCriterion > VizList =
+			{
+				SMTC_AbsNumOfVerts,
+				SMTC_AbsTriangleOrVert
+			};
 
-			SetPercentAndAbsoluteVisibility(MaxVerticesRow, SMTC_AbsNumOfVerts, SMTC_AbsTriangleOrVert);
+			MaxVerticesRow.Visibility(TAttribute<EVisibility>::Create(TAttribute<EVisibility>::FGetter::CreateSP(this, &FSkeletalMeshReductionSettingsLayout::ShowIfCurrentCriterionIs, VizList)));
+		}
+
+		{
+			FDetailWidgetRow& MaxTrianglesRow = AddUnsignedIntegerRow(ChildrenBuilder,
+				LOCTEXT("MaxTriangles_Row", "Max Number of Triangles"),
+				LOCTEXT("MaxTriangles", "Max Triangles Count"),
+				LOCTEXT("MaxTriangles_ToolTip", "The maximum number of triangles to retain when using percentage criterion."),
+				0,
+				MAX_uint32,
+				FGetUnsignedIntegerDelegate::CreateRaw(this, &FSkeletalMeshReductionSettingsLayout::GetNumMaxTrianglesPercentageCount),
+				FSetUnsignedIntegerDelegate::CreateRaw(this, &FSkeletalMeshReductionSettingsLayout::SetNumMaxTrianglesPercentageCount));
+
+			const TArray< SkeletalMeshTerminationCriterion > VizList =
+			{
+				SMTC_NumOfTriangles
+			};
+
+			MaxTrianglesRow.Visibility(TAttribute<EVisibility>::Create(TAttribute<EVisibility>::FGetter::CreateSP(this, &FSkeletalMeshReductionSettingsLayout::ShowIfCurrentCriterionIs, VizList)));
+
+		}
+
+		{
+			FDetailWidgetRow& MaxVerticesRow = AddUnsignedIntegerRow(ChildrenBuilder,
+				LOCTEXT("MaxVertices_Row", "Max Number of Vertices"),
+				LOCTEXT("MaxVertices", "Max Vertex Count"),
+				LOCTEXT("MaxVertices_ToolTip", "The maximum number of vertices to retain when using percentage criterion."),
+				0,
+				MAX_uint32,
+				FGetUnsignedIntegerDelegate::CreateRaw(this, &FSkeletalMeshReductionSettingsLayout::GetNumMaxVerticesPercentageCount),
+				FSetUnsignedIntegerDelegate::CreateRaw(this, &FSkeletalMeshReductionSettingsLayout::SetNumMaxVerticesPercentageCount));
+
+			const TArray< SkeletalMeshTerminationCriterion > VizList =
+			{
+				SMTC_NumOfVerts
+			};
+
+			MaxVerticesRow.Visibility(TAttribute<EVisibility>::Create(TAttribute<EVisibility>::FGetter::CreateSP(this, &FSkeletalMeshReductionSettingsLayout::ShowIfCurrentCriterionIs, VizList)));
 		}
 
 		AddBoolRow(ChildrenBuilder,
@@ -893,6 +943,73 @@ FDetailWidgetRow& FSkeletalMeshReductionSettingsLayout::AddIntegerRow(IDetailChi
 	return Row;
 }
 
+FDetailWidgetRow& FSkeletalMeshReductionSettingsLayout::AddUnsignedIntegerRow(IDetailChildrenBuilder& ChildrenBuilder, const FText RowTitleText, const FText RowNameContentText, const FText RowNameContentTootlipText, const uint32 MinSliderValue, const uint32 MaxSliderValue, FGetUnsignedIntegerDelegate GetterDelegate, FSetUnsignedIntegerDelegate SetterDelegate)
+{
+	uint32 SliderDataIndex = SliderStateDataArray.Num();
+	FSliderStateData& SliderData = SliderStateDataArray.AddDefaulted_GetRef();
+	SliderData.bSliderActiveMode = false;
+
+	auto BeginSliderMovementHelperFunc = [GetterDelegate, SliderDataIndex, this]()
+	{
+		check(SliderStateDataArray.IsValidIndex(SliderDataIndex));
+		SliderStateDataArray[SliderDataIndex].bSliderActiveMode = true;
+		SliderStateDataArray[SliderDataIndex].MovementValueUnsignedInt = GetterDelegate.IsBound() ? GetterDelegate.Execute() : 0;
+	};
+
+	auto EndSliderMovementHelperFunc = [SetterDelegate, SliderDataIndex, this](int32 Value)
+	{
+		check(SliderStateDataArray.IsValidIndex(SliderDataIndex));
+		SliderStateDataArray[SliderDataIndex].bSliderActiveMode = false;
+		SliderStateDataArray[SliderDataIndex].MovementValueUnsignedInt = 0;
+		SetterDelegate.ExecuteIfBound(Value);
+	};
+
+	auto SetValueHelperFunc = [SetterDelegate, SliderDataIndex, this](uint32 Value)
+	{
+		check(SliderStateDataArray.IsValidIndex(SliderDataIndex));
+		if (SliderStateDataArray[SliderDataIndex].bSliderActiveMode)
+		{
+			SliderStateDataArray[SliderDataIndex].MovementValueUnsignedInt = Value;
+		}
+		else
+		{
+			SetterDelegate.ExecuteIfBound(Value);
+		}
+	};
+
+	auto GetValueHelperFunc = [GetterDelegate, SliderDataIndex, this]()
+	{
+		check(SliderStateDataArray.IsValidIndex(SliderDataIndex));
+		if (SliderStateDataArray[SliderDataIndex].bSliderActiveMode)
+		{
+			return SliderStateDataArray[SliderDataIndex].MovementValueUnsignedInt;
+		}
+		return GetterDelegate.IsBound() ? GetterDelegate.Execute() : 0;
+	};
+
+	FDetailWidgetRow& Row = ChildrenBuilder.AddCustomRow(RowTitleText)
+		.NameContent()
+		[
+			SNew(STextBlock)
+			.Font(IDetailLayoutBuilder::GetDetailFont())
+		.Text(RowNameContentText)
+		.ToolTipText(RowNameContentTootlipText)
+		]
+	.ValueContent()
+		[
+			SNew(SSpinBox<uint32>)
+			.Font(IDetailLayoutBuilder::GetDetailFont())
+		.MinValue(MinSliderValue)
+		.MaxValue(MaxSliderValue)
+		.Value_Lambda(GetValueHelperFunc)
+		.OnValueChanged_Lambda(SetValueHelperFunc)
+		.OnBeginSliderMovement_Lambda(BeginSliderMovementHelperFunc)
+		.OnEndSliderMovement_Lambda(EndSliderMovementHelperFunc)
+		.IsEnabled(this, &FSkeletalMeshReductionSettingsLayout::IsReductionEnabled)
+		];
+	return Row;
+}
+
 void FSkeletalMeshReductionSettingsLayout::AddBaseLODRow(IDetailChildrenBuilder& ChildrenBuilder)
 {
 	// Only able to do this for LOD2 and above, so only show the property if this is the case
@@ -1141,6 +1258,40 @@ void FSkeletalMeshReductionSettingsLayout::SetNumMaxVerticesCount(int32 Value)
 	}
 }
 
+uint32 FSkeletalMeshReductionSettingsLayout::GetNumMaxTrianglesPercentageCount() const
+{
+	return ReductionSettings.MaxNumOfTrianglesPercentage;
+}
+
+void FSkeletalMeshReductionSettingsLayout::SetNumMaxTrianglesPercentageCount(uint32 Value)
+{
+	if (ReductionSettings.MaxNumOfTrianglesPercentage != Value)
+	{
+		FText TransactionText = FText::Format(LOCTEXT("PersonaReductionChangedSetMaxTriangleCountLOD", "LOD{0} reduction settings: max triangles percentage count changed"), LODIndex);
+		FScopedTransaction Transaction(TransactionText);
+		ModifyMeshLODSettingsDelegate.ExecuteIfBound(LODIndex);
+
+		ReductionSettings.MaxNumOfTrianglesPercentage = Value;
+	}
+}
+
+uint32 FSkeletalMeshReductionSettingsLayout::GetNumMaxVerticesPercentageCount() const
+{
+	return ReductionSettings.MaxNumOfVertsPercentage;
+}
+
+void FSkeletalMeshReductionSettingsLayout::SetNumMaxVerticesPercentageCount(uint32 Value)
+{
+	if (ReductionSettings.MaxNumOfVertsPercentage != Value)
+	{
+		FText TransactionText = FText::Format(LOCTEXT("PersonaReductionChangedSetMaxVertexCountLOD", "LOD{0} reduction settings: max vertex percentage count changed"), LODIndex);
+		FScopedTransaction Transaction(TransactionText);
+		ModifyMeshLODSettingsDelegate.ExecuteIfBound(LODIndex);
+
+		ReductionSettings.MaxNumOfVertsPercentage = Value;
+	}
+}
+
 float FSkeletalMeshReductionSettingsLayout::GetAccuracyPercentage() const
 {
 	return ReductionSettings.MaxDeviationPercentage;
diff --git a/Engine/Source/Editor/Persona/Private/PersonaMeshDetails.h b/Engine/Source/Editor/Persona/Private/PersonaMeshDetails.h
--- a/Engine/Source/Editor/Persona/Private/PersonaMeshDetails.h
+++ b/Engine/Source/Editor/Persona/Private/PersonaMeshDetails.h
@@ -106,6 +106,9 @@ public:
 	DECLARE_DELEGATE_RetVal(int32, FGetIntegerDelegate);
 	DECLARE_DELEGATE_OneParam(FSetIntegerDelegate, int32);
 
+	DECLARE_DELEGATE_RetVal(uint32, FGetUnsignedIntegerDelegate);
+	DECLARE_DELEGATE_OneParam(FSetUnsignedIntegerDelegate, uint32);
+
 	DECLARE_DELEGATE_RetVal(ECheckBoxState, FGetCheckBoxStateDelegate);
 	DECLARE_DELEGATE_OneParam(FSetCheckBoxStateDelegate, ECheckBoxState);
 private:
@@ -124,6 +127,7 @@ private:
 	FDetailWidgetRow& AddFloatRow(IDetailChildrenBuilder& ChildrenBuilder, const FText RowTitleText, const FText RowNameContentText, const FText RowNameContentTootlipText, const float MinSliderValue, const float MaxSliderValue, FGetFloatDelegate GetterDelegate, FSetFloatDelegate SetterDelegate);
 	FDetailWidgetRow& AddBoolRow(IDetailChildrenBuilder& ChildrenBuilder, const FText RowTitleText, const FText RowNameContentText, const FText RowNameContentToolitipText, FGetCheckBoxStateDelegate GetterDelegate, FSetCheckBoxStateDelegate SetterDelegate);
 	FDetailWidgetRow& AddIntegerRow(IDetailChildrenBuilder& ChildrenBuilder, const FText RowTitleText, const FText RowNameContentText, const FText RowNameContentTootlipText, const int32 MinSliderValue, const int32 MaxSliderValue, FGetIntegerDelegate GetterDelegate, FSetIntegerDelegate SetterDelegate);
+	FDetailWidgetRow& AddUnsignedIntegerRow(IDetailChildrenBuilder& ChildrenBuilder, const FText RowTitleText, const FText RowNameContentText, const FText RowNameContentTootlipText, const uint32 MinSliderValue, const uint32 MaxSliderValue, FGetUnsignedIntegerDelegate GetterDelegate, FSetUnsignedIntegerDelegate SetterDelegate);
 	void AddBaseLODRow(IDetailChildrenBuilder& ChildrenBuilder);
 
 	void SetPercentAndAbsoluteVisibility(FDetailWidgetRow& Row, SkeletalMeshTerminationCriterion FirstCriterion, SkeletalMeshTerminationCriterion SecondCriterion);
@@ -157,6 +161,12 @@ private:
 	int32 GetNumMaxVerticesCount() const;
 	void SetNumMaxVerticesCount(int32 Value);
 
+	uint32 GetNumMaxTrianglesPercentageCount() const;
+	void SetNumMaxTrianglesPercentageCount(uint32 Value);
+
+	uint32 GetNumMaxVerticesPercentageCount() const;
+	void SetNumMaxVerticesPercentageCount(uint32 Value);
+
 	float GetAccuracyPercentage() const;
 	void SetAccuracyPercentage(float Value);
 
@@ -220,6 +230,7 @@ private:
 	{
 		float MovementValueFloat = 0.0f;
 		int32 MovementValueInt = 0;
+		uint32 MovementValueUnsignedInt = 0;
 		bool bSliderActiveMode = false;
 	};
 	TArray<FSliderStateData> SliderStateDataArray;
diff --git a/Engine/Source/Runtime/Engine/Private/SkeletalMesh.cpp b/Engine/Source/Runtime/Engine/Private/SkeletalMesh.cpp
--- a/Engine/Source/Runtime/Engine/Private/SkeletalMesh.cpp
+++ b/Engine/Source/Runtime/Engine/Private/SkeletalMesh.cpp
@@ -4081,6 +4081,8 @@ static void SerializeReductionSettingsForDDC(FArchive& Ar, FSkeletalMeshOptimiza
 	Ar << ReductionSettings.NumOfVertPercentage;
 	Ar << ReductionSettings.MaxNumOfTriangles;
 	Ar << ReductionSettings.MaxNumOfVerts;
+	Ar << ReductionSettings.MaxNumOfTrianglesPercentage;
+	Ar << ReductionSettings.MaxNumOfVertsPercentage;
 	Ar << ReductionSettings.MaxDeviationPercentage;
 	Ar << ReductionSettings.ReductionMethod;
 	Ar << ReductionSettings.SilhouetteImportance;
diff --git a/Engine/Source/Runtime/Engine/Public/SkeletalMeshReductionSettings.h b/Engine/Source/Runtime/Engine/Public/SkeletalMeshReductionSettings.h
--- a/Engine/Source/Runtime/Engine/Public/SkeletalMeshReductionSettings.h
+++ b/Engine/Source/Runtime/Engine/Public/SkeletalMeshReductionSettings.h
@@ -79,6 +79,14 @@ struct FSkeletalMeshOptimizationSettings
 	UPROPERTY(EditAnywhere, Category = ReductionMethod, meta = (DisplayName = "Max Vertex Count", ClampMin = 6))
 	uint32 MaxNumOfVerts;
 
+	/** The maximum number of triangles to retain when using percentage termination criterion. */
+	UPROPERTY(EditAnywhere, Category = ReductionMethod, meta = (DisplayName = "Max Triangle Count", ClampMin = 4, UIMin = "4"))
+	uint32 MaxNumOfTrianglesPercentage;
+
+	/** The maximum number of vertices to retain when using percentage termination criterion. */
+	UPROPERTY(EditAnywhere, Category = ReductionMethod, meta = (DisplayName = "Max Vertex Count", ClampMin = 6, UIMin = "6"))
+	uint32 MaxNumOfVertsPercentage;
+
 	/**If ReductionMethod equals MaxDeviation this value is the maximum deviation from the base mesh as a percentage of the bounding sphere. 
 	 * In code, it ranges from [0, 1]. In the editor UI, it ranges from [0, 100]
 	 */
@@ -163,6 +171,8 @@ struct FSkeletalMeshOptimizationSettings
 		, NumOfVertPercentage(0.5f)
 		, MaxNumOfTriangles(4)
 		, MaxNumOfVerts(6)
+		, MaxNumOfTrianglesPercentage(MAX_uint32)
+		, MaxNumOfVertsPercentage(MAX_uint32)
 		, MaxDeviationPercentage(0.5f)
 		, ReductionMethod(SMOT_NumOfTriangles)
 		, SilhouetteImportance(SMOI_Normal)
