Fix Localization build
#rnx

diff --git a/Engine/Source/Editor/Persona/Private/PersonaMeshDetails.cpp b/Engine/Source/Editor/Persona/Private/PersonaMeshDetails.cpp
--- a/Engine/Source/Editor/Persona/Private/PersonaMeshDetails.cpp
+++ b/Engine/Source/Editor/Persona/Private/PersonaMeshDetails.cpp
@@ -699,9 +699,9 @@ void FSkeletalMeshReductionSettingsLayout::GenerateChildContent(IDetailChildrenB
 
 		{
 			FDetailWidgetRow& MaxTrianglesRow = AddUnsignedIntegerRow(ChildrenBuilder,
-				LOCTEXT("MaxTriangles_Row", "Max Number of Triangles"),
-				LOCTEXT("MaxTriangles", "Max Triangles Count"),
-				LOCTEXT("MaxTriangles_ToolTip", "The maximum number of triangles to retain when using percentage criterion."),
+				LOCTEXT("MaxTrianglesPercentage_Row", "Max Number of Triangles"),
+				LOCTEXT("MaxTrianglesPercentage", "Max Triangles Count"),
+				LOCTEXT("MaxTrianglesPercentage_ToolTip", "The maximum number of triangles to retain when using percentage criterion."),
 				0,
 				MAX_uint32,
 				FGetUnsignedIntegerDelegate::CreateRaw(this, &FSkeletalMeshReductionSettingsLayout::GetNumMaxTrianglesPercentageCount),
@@ -718,9 +718,9 @@ void FSkeletalMeshReductionSettingsLayout::GenerateChildContent(IDetailChildrenB
 
 		{
 			FDetailWidgetRow& MaxVerticesRow = AddUnsignedIntegerRow(ChildrenBuilder,
-				LOCTEXT("MaxVertices_Row", "Max Number of Vertices"),
-				LOCTEXT("MaxVertices", "Max Vertex Count"),
-				LOCTEXT("MaxVertices_ToolTip", "The maximum number of vertices to retain when using percentage criterion."),
+				LOCTEXT("MaxVerticesPercentage_Row", "Max Number of Vertices"),
+				LOCTEXT("MaxVerticesPercentage", "Max Vertex Count"),
+				LOCTEXT("MaxVerticesPercentage_ToolTip", "The maximum number of vertices to retain when using percentage criterion."),
 				0,
 				MAX_uint32,
 				FGetUnsignedIntegerDelegate::CreateRaw(this, &FSkeletalMeshReductionSettingsLayout::GetNumMaxVerticesPercentageCount),
