Fix crash when undoing adding a montage segment

Notify tracks are not guaranteed to stick around correctly on undo (their array may be reallocated).

diff --git a/Engine/Source/Editor/Persona/Private/AnimTimeline/AnimTimelineTrack_NotifiesPanel.cpp b/Engine/Source/Editor/Persona/Private/AnimTimeline/AnimTimelineTrack_NotifiesPanel.cpp
--- a/Engine/Source/Editor/Persona/Private/AnimTimeline/AnimTimelineTrack_NotifiesPanel.cpp
+++ b/Engine/Source/Editor/Persona/Private/AnimTimeline/AnimTimelineTrack_NotifiesPanel.cpp
@@ -85,7 +85,7 @@ void FAnimTimelineTrack_NotifiesPanel::RefreshOutlinerWidget()
 				.Padding(30.0f, 0.0f, 0.0f, 0.0f)
 				[
 					SAssignNew(InlineEditableTextBlock, SInlineEditableTextBlock)
-					.Text_Lambda([&AnimNotifyTrack](){ return FText::FromName(AnimNotifyTrack.TrackName); })
+					.Text_Lambda([TrackIndex, AnimSequence](){ return AnimSequence->AnimNotifyTracks.IsValidIndex(TrackIndex) ? FText::FromName(AnimSequence->AnimNotifyTracks[TrackIndex].TrackName) : FText::GetEmpty(); }) 
 					.IsSelected(FIsSelected::CreateLambda([](){ return true; }))
 					.OnTextCommitted(this, &FAnimTimelineTrack_NotifiesPanel::OnCommitTrackName, TrackIndex)
 				]
