Fix localization build
#rnx

diff --git a/Engine/Source/Editor/Persona/Private/PersonaMeshDetails.cpp b/Engine/Source/Editor/Persona/Private/PersonaMeshDetails.cpp
--- a/Engine/Source/Editor/Persona/Private/PersonaMeshDetails.cpp
+++ b/Engine/Source/Editor/Persona/Private/PersonaMeshDetails.cpp
@@ -1267,7 +1267,7 @@ void FSkeletalMeshReductionSettingsLayout::SetNumMaxTrianglesPercentageCount(uin
 {
 	if (ReductionSettings.MaxNumOfTrianglesPercentage != Value)
 	{
-		FText TransactionText = FText::Format(LOCTEXT("PersonaReductionChangedSetMaxTriangleCountLOD", "LOD{0} reduction settings: max triangles percentage count changed"), LODIndex);
+		FText TransactionText = FText::Format(LOCTEXT("PersonaReductionChangedSetMaxTrianglePercentageCountLOD", "LOD{0} reduction settings: max triangles percentage count changed"), LODIndex);
 		FScopedTransaction Transaction(TransactionText);
 		ModifyMeshLODSettingsDelegate.ExecuteIfBound(LODIndex);
 
@@ -1284,7 +1284,7 @@ void FSkeletalMeshReductionSettingsLayout::SetNumMaxVerticesPercentageCount(uint
 {
 	if (ReductionSettings.MaxNumOfVertsPercentage != Value)
 	{
-		FText TransactionText = FText::Format(LOCTEXT("PersonaReductionChangedSetMaxVertexCountLOD", "LOD{0} reduction settings: max vertex percentage count changed"), LODIndex);
+		FText TransactionText = FText::Format(LOCTEXT("PersonaReductionChangedSetMaxVertexPercentageCountLOD", "LOD{0} reduction settings: max vertex percentage count changed"), LODIndex);
 		FScopedTransaction Transaction(TransactionText);
 		ModifyMeshLODSettingsDelegate.ExecuteIfBound(LODIndex);
 
