Anim timeline now allows access to times outside the sequence range

Only curve keys can be moved outside of the sequence range. Fixed notify tracks to prevent dragging/resizing outside the sequence range.
All curve editor features now usable, including retime tool.

diff --git a/Engine/Source/Editor/Persona/Private/AnimTimeline/AnimModel.cpp b/Engine/Source/Editor/Persona/Private/AnimTimeline/AnimModel.cpp
--- a/Engine/Source/Editor/Persona/Private/AnimTimeline/AnimModel.cpp
+++ b/Engine/Source/Editor/Persona/Private/AnimTimeline/AnimModel.cpp
@@ -31,9 +31,9 @@ FAnimatedRange FAnimModel::GetViewRange() const
 	return ViewRange;
 }
 
-FAnimatedRange FAnimModel::GetTotalRange() const
+FAnimatedRange FAnimModel::GetWorkingRange() const
 {
-	return TotalRange;
+	return WorkingRange;
 }
 
 double FAnimModel::GetFrameRate() const
@@ -56,7 +56,7 @@ int32 FAnimModel::GetTickResolution() const
 TRange<FFrameNumber> FAnimModel::GetPlaybackRange() const
 {
 	const int32 Resolution = GetTickResolution();
-	return TRange<FFrameNumber>(FFrameNumber(FMath::RoundToInt(TotalRange.GetLowerBoundValue() * (double)Resolution)), FFrameNumber(FMath::RoundToInt(TotalRange.GetUpperBoundValue() * (double)Resolution)));
+	return TRange<FFrameNumber>(FFrameNumber(FMath::RoundToInt(PlaybackRange.GetLowerBoundValue() * (double)Resolution)), FFrameNumber(FMath::RoundToInt(PlaybackRange.GetUpperBoundValue() * (double)Resolution)));
 }
 
 FFrameNumber FAnimModel::GetScrubPosition() const
@@ -87,12 +87,26 @@ void FAnimModel::SetScrubPosition(FFrameTime NewScrubPostion) const
 
 void FAnimModel::HandleViewRangeChanged(TRange<double> InRange, EViewRangeInterpolation InInterpolation)
 {
-	ViewRange = InRange;
+	SetViewRange(InRange);
 }
 
 void FAnimModel::SetViewRange(TRange<double> InRange)
 {
 	ViewRange = InRange;
+
+	if(WorkingRange.HasLowerBound() && WorkingRange.HasUpperBound())
+	{
+		WorkingRange = TRange<double>::Hull(WorkingRange, ViewRange);
+	}
+	else
+	{
+		WorkingRange = ViewRange;
+	}
+}
+
+void FAnimModel::HandleWorkingRangeChanged(TRange<double> InRange)
+{
+	WorkingRange = InRange;
 }
 
 bool FAnimModel::IsTrackSelected(const TSharedRef<FAnimTimelineTrack>& InTrack) const
diff --git a/Engine/Source/Editor/Persona/Private/AnimTimeline/AnimModel.h b/Engine/Source/Editor/Persona/Private/AnimTimeline/AnimModel.h
--- a/Engine/Source/Editor/Persona/Private/AnimTimeline/AnimModel.h
+++ b/Engine/Source/Editor/Persona/Private/AnimTimeline/AnimModel.h
@@ -37,8 +37,8 @@ public:
 	/** Set the current view range */
 	void SetViewRange(TRange<double> InRange);
 
-	/** Get the total range of the model's data */
-	FAnimatedRange GetTotalRange() const;
+	/** Get the working range of the model's data */
+	FAnimatedRange GetWorkingRange() const;
 
 	/** Get the playback range of the model's data */
 	TRange<FFrameNumber> GetPlaybackRange() const;
@@ -52,6 +52,9 @@ public:
 	/** Handle the view range being changed */
 	void HandleViewRangeChanged(TRange<double> InRange, EViewRangeInterpolation InInterpolation);
 
+	/** Handle the working range being changed */
+	void HandleWorkingRangeChanged(TRange<double> InRange);
+
 	/** Delegate fired when tracks have changed */
 	DECLARE_EVENT(FAnimModel, FOnTracksChanged)
 	FOnTracksChanged& OnTracksChanged() { return OnTracksChangedDelegate; }
@@ -162,8 +165,11 @@ protected:
 	/** The range we are currently viewing */
 	FAnimatedRange ViewRange;
 
-	/** The total range of this model for each timeframe */
-	FAnimatedRange TotalRange;
+	/** The working range of this model, encompassing the view range */
+	FAnimatedRange WorkingRange;
+
+	/** The playback range of this model for each timeframe */
+	FAnimatedRange PlaybackRange;
 
 	/** Delegate fired when tracks change */
 	FOnTracksChanged OnTracksChangedDelegate;
diff --git a/Engine/Source/Editor/Persona/Private/AnimTimeline/AnimModel_AnimSequenceBase.cpp b/Engine/Source/Editor/Persona/Private/AnimTimeline/AnimModel_AnimSequenceBase.cpp
--- a/Engine/Source/Editor/Persona/Private/AnimTimeline/AnimModel_AnimSequenceBase.cpp
+++ b/Engine/Source/Editor/Persona/Private/AnimTimeline/AnimModel_AnimSequenceBase.cpp
@@ -412,15 +412,15 @@ void FAnimModel_AnimSequenceBase::HandleUndoRedo()
 
 void FAnimModel_AnimSequenceBase::UpdateRange()
 {
-	FAnimatedRange OldTotalRange = TotalRange;
+	FAnimatedRange OldPlaybackRange = PlaybackRange;
 
-	// update range
-	TotalRange = FAnimatedRange(0.0, (double)AnimSequenceBase->GetPlayLength());
+	// update playback range
+	PlaybackRange = FAnimatedRange(0.0, (double)AnimSequenceBase->GetPlayLength());
 
-	if (OldTotalRange != TotalRange)
+	if (OldPlaybackRange != PlaybackRange)
 	{
-		// Update view range if total range changed
-		ViewRange = TotalRange;
+		// Update view/range if playback range changed
+		SetViewRange(PlaybackRange);
 	}
 }
 
diff --git a/Engine/Source/Editor/Persona/Private/AnimTimeline/AnimTimeSliderController.cpp b/Engine/Source/Editor/Persona/Private/AnimTimeline/AnimTimeSliderController.cpp
--- a/Engine/Source/Editor/Persona/Private/AnimTimeline/AnimTimeSliderController.cpp
+++ b/Engine/Source/Editor/Persona/Private/AnimTimeline/AnimTimeSliderController.cpp
@@ -302,11 +302,13 @@ int32 FAnimTimeSliderController::OnPaintTimeSlider( bool bMirrorLabels, const FG
 
 		// draw playback & selection range
 		FPaintPlaybackRangeArgs PlaybackRangeArgs(
-			bMirrorLabels ? FEditorStyle::GetBrush("Anim.Timeline.PlayRange_Bottom_L") : FEditorStyle::GetBrush("Anim.Timeline.PlayRange_Top_L"),
-			bMirrorLabels ? FEditorStyle::GetBrush("Anim.Timeline.PlayRange_Bottom_R") : FEditorStyle::GetBrush("Anim.Timeline.PlayRange_Top_R"),
+			bMirrorLabels ? FEditorStyle::GetBrush("Sequencer.Timeline.PlayRange_Bottom_L") : FEditorStyle::GetBrush("Sequencer.Timeline.PlayRange_Top_L"),
+			bMirrorLabels ? FEditorStyle::GetBrush("Sequencer.Timeline.PlayRange_Bottom_R") : FEditorStyle::GetBrush("Sequencer.Timeline.PlayRange_Top_R"),
 			6.f
 		);
 
+		LayerId = DrawPlaybackRange(AllottedGeometry, MyCullingRect, OutDrawElements, LayerId, RangeToScreen, PlaybackRangeArgs);
+
 		PlaybackRangeArgs.SolidFillOpacity = 0.05f;
 		LayerId = DrawSelectionRange(AllottedGeometry, MyCullingRect, OutDrawElements, LayerId, RangeToScreen, PlaybackRangeArgs);
 
@@ -454,6 +456,61 @@ int32 FAnimTimeSliderController::DrawSelectionRange(const FGeometry& AllottedGeo
 	return LayerId + 1;
 }
 
+
+int32 FAnimTimeSliderController::DrawPlaybackRange(const FGeometry& AllottedGeometry, const FSlateRect& MyCullingRect, FSlateWindowElementList& OutDrawElements, int32 LayerId, const FScrubRangeToScreen& RangeToScreen, const FPaintPlaybackRangeArgs& Args) const
+{
+	if (!TimeSliderArgs.PlaybackRange.IsSet())
+	{
+		return LayerId;
+	}
+
+	const uint8 OpacityBlend = TimeSliderArgs.SubSequenceRange.Get().IsSet() ? 128 : 255;
+
+	TRange<FFrameNumber> PlaybackRange = TimeSliderArgs.PlaybackRange.Get();
+	FFrameRate TickResolution = GetTickResolution();
+	const float PlaybackRangeL = RangeToScreen.InputToLocalX(PlaybackRange.GetLowerBoundValue() / TickResolution);
+	const float PlaybackRangeR = RangeToScreen.InputToLocalX(PlaybackRange.GetUpperBoundValue() / TickResolution) - 1;
+
+	FSlateDrawElement::MakeBox(
+		OutDrawElements,
+		LayerId+1,
+		AllottedGeometry.ToPaintGeometry(FVector2D(PlaybackRangeL, 0.f), FVector2D(Args.BrushWidth, AllottedGeometry.Size.Y)),
+		Args.StartBrush,
+		ESlateDrawEffect::None,
+		FColor(32, 128, 32, OpacityBlend)	// 120, 75, 50 (HSV)
+	);
+
+	FSlateDrawElement::MakeBox(
+		OutDrawElements,
+		LayerId+1,
+		AllottedGeometry.ToPaintGeometry(FVector2D(PlaybackRangeR - Args.BrushWidth, 0.f), FVector2D(Args.BrushWidth, AllottedGeometry.Size.Y)),
+		Args.EndBrush,
+		ESlateDrawEffect::None,
+		FColor(128, 32, 32, OpacityBlend)	// 0, 75, 50 (HSV)
+	);
+
+	// Black tint for excluded regions
+	FSlateDrawElement::MakeBox(
+		OutDrawElements,
+		LayerId+1,
+		AllottedGeometry.ToPaintGeometry(FVector2D(0.f, 0.f), FVector2D(PlaybackRangeL, AllottedGeometry.Size.Y)),
+		FEditorStyle::GetBrush("WhiteBrush"),
+		ESlateDrawEffect::None,
+		FLinearColor::Black.CopyWithNewOpacity(0.3f * OpacityBlend / 255.f)
+	);
+
+	FSlateDrawElement::MakeBox(
+		OutDrawElements,
+		LayerId+1,
+		AllottedGeometry.ToPaintGeometry(FVector2D(PlaybackRangeR, 0.f), FVector2D(AllottedGeometry.Size.X - PlaybackRangeR, AllottedGeometry.Size.Y)),
+		FEditorStyle::GetBrush("WhiteBrush"),
+		ESlateDrawEffect::None,
+		FLinearColor::Black.CopyWithNewOpacity(0.3f * OpacityBlend / 255.f)
+	);
+
+	return LayerId + 1;
+}
+
 int32 FAnimTimeSliderController::DrawEditableTimes(const FGeometry& AllottedGeometry, const FSlateRect& MyCullingRect, FSlateWindowElementList& OutDrawElements, int32 LayerId, const FScrubRangeToScreen& RangeToScreen) const
 {
 	const FLinearColor TimeColor = GetDefault<UPersonaOptions>()->SectionTimingNodeColor;
@@ -822,6 +879,7 @@ int32 FAnimTimeSliderController::OnPaintViewArea( const FGeometry& AllottedGeome
 	if (Args.PlaybackRangeArgs.IsSet())
 	{
 		FPaintPlaybackRangeArgs PaintArgs = Args.PlaybackRangeArgs.GetValue();
+		LayerId = DrawPlaybackRange(AllottedGeometry, MyCullingRect, OutDrawElements, LayerId, RangeToScreen, PaintArgs);
 		PaintArgs.SolidFillOpacity = 0.2f;
 		LayerId = DrawSelectionRange(AllottedGeometry, MyCullingRect, OutDrawElements, LayerId, RangeToScreen, PaintArgs);
 	}
@@ -985,7 +1043,6 @@ void FAnimTimeSliderController::SetViewRange( double NewRangeMin, double NewRang
 	// Clamp to a minimum size to avoid zero-sized or negative visible ranges
 	double MinVisibleTimeRange = FFrameNumber(1) / GetTickResolution();
 	TRange<double> ExistingViewRange  = TimeSliderArgs.ViewRange.Get();
-	TRange<double> ExistingClampRange = TimeSliderArgs.ClampRange.Get();
 
 	if (NewRangeMax == ExistingViewRange.GetUpperBoundValue())
 	{
@@ -1000,7 +1057,7 @@ void FAnimTimeSliderController::SetViewRange( double NewRangeMin, double NewRang
 	}
 
 	// Clamp to the clamp range
-	const TRange<double> NewRange = TRange<double>::Intersection(TRange<double>(NewRangeMin, NewRangeMax), ExistingClampRange);
+	const TRange<double> NewRange = TRange<double>(NewRangeMin, NewRangeMax);
 	TimeSliderArgs.OnViewRangeChanged.ExecuteIfBound( NewRange, Interpolation );
 
 	if( !TimeSliderArgs.ViewRange.IsBound() )
diff --git a/Engine/Source/Editor/Persona/Private/AnimTimeline/AnimTimeSliderController.h b/Engine/Source/Editor/Persona/Private/AnimTimeline/AnimTimeSliderController.h
--- a/Engine/Source/Editor/Persona/Private/AnimTimeline/AnimTimeSliderController.h
+++ b/Engine/Source/Editor/Persona/Private/AnimTimeline/AnimTimeSliderController.h
@@ -124,6 +124,13 @@ private:
 	 */
 	int32 DrawSelectionRange(const FGeometry& AllottedGeometry, const FSlateRect& MyCullingRect, FSlateWindowElementList& OutDrawElements, int32 LayerId, const FScrubRangeToScreen& RangeToScreen, const FPaintPlaybackRangeArgs& Args) const;
 
+	/**
+	 * Draw the playback range.
+	 *
+	 * @return the new layer ID
+	 */
+	int32 DrawPlaybackRange(const FGeometry& AllottedGeometry, const FSlateRect& MyCullingRect, FSlateWindowElementList& OutDrawElements, int32 LayerId, const FScrubRangeToScreen& RangeToScreen, const FPaintPlaybackRangeArgs& Args) const;
+
 private:
 
 	/**
diff --git a/Engine/Source/Editor/Persona/Private/AnimTimeline/SAnimTimeline.cpp b/Engine/Source/Editor/Persona/Private/AnimTimeline/SAnimTimeline.cpp
--- a/Engine/Source/Editor/Persona/Private/AnimTimeline/SAnimTimeline.cpp
+++ b/Engine/Source/Editor/Persona/Private/AnimTimeline/SAnimTimeline.cpp
@@ -82,10 +82,11 @@ void SAnimTimeline::Construct(const FArguments& InArgs, const TSharedRef<FAnimMo
 		TimeSliderArgs.ScrubPosition = MakeAttributeLambda([WeakModel](){ return WeakModel.IsValid() ? WeakModel.Pin()->GetScrubPosition() : FFrameTime(0); });
 		TimeSliderArgs.ViewRange = ViewRange;
 		TimeSliderArgs.PlaybackRange = MakeAttributeLambda([WeakModel](){ return WeakModel.IsValid() ? WeakModel.Pin()->GetPlaybackRange() : TRange<FFrameNumber>(0, 0); });
-		TimeSliderArgs.ClampRange = MakeAttributeLambda([WeakModel](){ return WeakModel.IsValid() ? WeakModel.Pin()->GetTotalRange() : FAnimatedRange(0.0, 0.0); });
+		TimeSliderArgs.ClampRange = MakeAttributeLambda([WeakModel](){ return WeakModel.IsValid() ? WeakModel.Pin()->GetWorkingRange() : FAnimatedRange(0.0, 0.0); });
 		TimeSliderArgs.DisplayRate = DisplayRate;
 		TimeSliderArgs.TickResolution = TickResolution;
 		TimeSliderArgs.OnViewRangeChanged = FOnViewRangeChanged::CreateSP(&InModel.Get(), &FAnimModel::HandleViewRangeChanged);
+		TimeSliderArgs.OnClampRangeChanged = FOnTimeRangeChanged::CreateSP(&InModel.Get(), &FAnimModel::HandleWorkingRangeChanged);
 		TimeSliderArgs.IsPlaybackRangeLocked = true;
 		TimeSliderArgs.PlaybackStatus = EMovieScenePlayerStatus::Stopped;
 		TimeSliderArgs.NumericTypeInterface = NumericTypeInterface;
@@ -104,8 +105,8 @@ void SAnimTimeline::Construct(const FArguments& InArgs, const TSharedRef<FAnimMo
 	// Create bottom time range slider
 	TSharedRef<ITimeSlider> BottomTimeRange = SequencerWidgets.CreateTimeRange(
 		FTimeRangeArgs(
-			EShowRange::ViewRange | EShowRange::PlaybackRange,
-			EShowRange::ViewRange,
+			EShowRange::ViewRange | EShowRange::WorkingRange | EShowRange::PlaybackRange,
+			EShowRange::ViewRange | EShowRange::WorkingRange,
 			TimeSliderControllerRef,
 			EVisibility::Visible,
 			NumericTypeInterface.ToSharedRef()
@@ -288,6 +289,7 @@ void SAnimTimeline::Construct(const FArguments& InArgs, const TSharedRef<FAnimMo
 						.DisplayScrubPosition( false )
 						.DisplayTickLines( true )
 						.Clipping(EWidgetClipping::ClipToBounds)
+						.PaintPlaybackRangeArgs(FPaintPlaybackRangeArgs(FEditorStyle::GetBrush("Sequencer.Timeline.PlayRange_L"), FEditorStyle::GetBrush("Sequencer.Timeline.PlayRange_R"), 6.f))
 					]
 
 					// Overlay that draws the scrub position
diff --git a/Engine/Source/Editor/Persona/Private/SAnimNotifyPanel.cpp b/Engine/Source/Editor/Persona/Private/SAnimNotifyPanel.cpp
--- a/Engine/Source/Editor/Persona/Private/SAnimNotifyPanel.cpp
+++ b/Engine/Source/Editor/Persona/Private/SAnimNotifyPanel.cpp
@@ -566,7 +566,7 @@ private:
 	FText GetNodeTooltip() const;
 
 	/** Detects any overflow on the anim notify track and requests a track pan */
-	float HandleOverflowPan( const FVector2D& ScreenCursorPos, float TrackScreenSpaceXPosition );
+	float HandleOverflowPan( const FVector2D& ScreenCursorPos, float TrackScreenSpaceXPosition, float TrackScreenSpaceMin, float TrackScreenSpaceMax);
 
 	/** Finds a snap position if possible for the provided scrub handle, if it is not possible, returns -1.0f */
 	float GetScrubHandleSnapPosition(float NotifyLocalX, ENotifyStateHandleHit::Type HandleToCheck);
@@ -619,6 +619,9 @@ private:
 	/** Delegate to signal selection changing */
 	FOnTrackSelectionChanged	OnSelectionChanged;
 
+	/** Delegate to redraw the notify panel */
+	FOnUpdatePanel				OnUpdatePanel;
+
 	/** Cached owning track geometry */
 	FGeometry CachedTrackGeometry;
 
@@ -984,6 +987,8 @@ protected:
 	TArray<FAnimSyncMarker*>				AnimSyncMarkers;
 	TAttribute<float>						ViewInputMin;
 	TAttribute<float>						ViewInputMax;
+	TAttribute<float>						InputMin;
+	TAttribute<float>						InputMax;
 	TAttribute<FLinearColor>				TrackColor;
 	int32									TrackIndex;
 	TAttribute<EVisibility>					NotifyTimingNodeVisibility;
@@ -1125,9 +1130,6 @@ public:
 	{
 		int32 TrackPos;
 		int32 TrackSnapTestPos;
-		float TrackMax;
-		float TrackMin;
-		float TrackWidth;
 		TSharedPtr<SAnimNotifyTrack> NotifyTrack;
 	};
 
@@ -1179,13 +1181,17 @@ public:
 		const FGeometry& TrackGeom = SelectionPositionClampInfo->NotifyTrack->GetCachedGeometry();
 		const FTrackScaleInfo& TrackScaleInfo = SelectionPositionClampInfo->NotifyTrack->GetCachedScaleInfo();
 		
+		float TrackMin = TrackGeom.LocalToAbsolute(FVector2D(TrackScaleInfo.InputToLocalX(0.0f), 0.0f)).X;
+		float TrackMax = TrackGeom.LocalToAbsolute(FVector2D(TrackScaleInfo.InputToLocalX(Sequence->GetPlayLength()), 0.0f)).X;
+		float TrackWidth = TrackMax - TrackMin;
+
 		// Tracks the movement amount to apply to the selection due to a snap.
 		float SnapMovement= 0.0f;
 		// Clamp the selection into the track
 		float SelectionBeginPositionX = TrackGeom.GetAbsolutePosition().X + TrackGeom.AbsoluteToLocal(SelectionBeginPosition).X;
 		const float SelectionLocalLength = TrackScaleInfo.PixelsPerInput * SelectionTimeLength;
-		const float ClampedEnd = FMath::Clamp(SelectionBeginPositionX + NodeGroupSize.X, SelectionPositionClampInfo->TrackMin, SelectionPositionClampInfo->TrackMax);
-		const float ClampedBegin = FMath::Clamp(SelectionBeginPositionX, SelectionPositionClampInfo->TrackMin, SelectionPositionClampInfo->TrackMax);
+		const float ClampedEnd = FMath::Clamp(SelectionBeginPositionX + NodeGroupSize.X, TrackMin, TrackMax);
+		const float ClampedBegin = FMath::Clamp(SelectionBeginPositionX, TrackMin, TrackMax);
 		if(ClampedBegin > SelectionBeginPositionX)
 		{
 			SelectionBeginPositionX = ClampedBegin;
@@ -1198,7 +1204,8 @@ public:
 		SelectionBeginPosition.X = TrackGeom.LocalToAbsolute(FVector2D(SelectionBeginPositionX - TrackGeom.GetAbsolutePosition().X, 0.0f)).X;
 
 		// Handle node snaps
-		for(int32 NodeIdx = 0 ; NodeIdx < SelectedNodes.Num() ; ++NodeIdx)
+		bool bSnapped = false;
+		for(int32 NodeIdx = 0 ; NodeIdx < SelectedNodes.Num() && !bSnapped; ++NodeIdx)
 		{
 			TSharedPtr<SAnimNotifyNode> CurrentNode = SelectedNodes[NodeIdx];
 
@@ -1217,10 +1224,10 @@ public:
 			// Always clamp the Y to the current track
 			SelectionBeginPosition.Y = SelectionPositionClampInfo->TrackPos - 1.0f;
 
-			float SnapX = GetSnapPosition(NodeClamp, TrackNodePos.X);
+			float SnapX = GetSnapPosition(NodeClamp, TrackNodePos.X, bSnapped);
 			if (FAnimNotifyEvent* CurrentEvent = CurrentNode->NodeObjectInterface->GetNotifyEvent())
 			{
-				if (SnapX >= 0.f)
+				if (bSnapped)
 				{
 					EAnimEventTriggerOffsets::Type Offset = EAnimEventTriggerOffsets::NoOffset;
 					if (SnapX == 0.0f || SnapX == SequenceEnd)
@@ -1252,10 +1259,10 @@ public:
 				if (CurrentNode.IsValid() && CurrentEvent->GetDuration() > 0)
 				{
 					// If we didn't snap the beginning of the node, attempt to snap the end
-					if (SnapX == -1.0f)
+					if (!bSnapped)
 					{
 						FVector2D TrackNodeEndPos = TrackNodePos + CurrentNode->GetDurationSize();
-						SnapX = GetSnapPosition(*SelectionPositionClampInfo, TrackNodeEndPos.X);
+						SnapX = GetSnapPosition(*SelectionPositionClampInfo, TrackNodeEndPos.X, bSnapped);
 
 						// Only attempt to snap if the node will fit on the track
 						if (SnapX >= CurrentNode->GetDurationSize())
@@ -1299,19 +1306,21 @@ public:
 
 		//scroll view
 		float MouseXPos = TrackGeom.GetAbsolutePosition().X + TrackGeom.AbsoluteToLocal(DragDropEvent.GetScreenSpacePosition()).X;
-		if(MouseXPos < SelectionPositionClampInfo->TrackMin)
+		float ViewportMin = TrackGeom.GetAbsolutePosition().X;
+		float ViewportMax = TrackGeom.GetAbsolutePosition().X + TrackGeom.GetAbsoluteSize().X;
+		if(MouseXPos < ViewportMin && ViewportMin > TrackMin - 10.0f)
 		{
-			float ScreenDelta = MouseXPos - SelectionPositionClampInfo->TrackMin;
-			RequestTrackPan.Execute(ScreenDelta, FVector2D(SelectionPositionClampInfo->TrackWidth, 1.f));
+			float ScreenDelta = FMath::Max(MouseXPos - ViewportMin, -10.0f);
+			RequestTrackPan.Execute(ScreenDelta, FVector2D(TrackWidth, 1.f));
 		}
-		else if(MouseXPos > SelectionPositionClampInfo->TrackMax)
+		else if(MouseXPos > ViewportMax && ViewportMax < TrackMax + 10.0f)
 		{
-			float ScreenDelta = MouseXPos - SelectionPositionClampInfo->TrackMax;
-			RequestTrackPan.Execute(ScreenDelta, FVector2D(SelectionPositionClampInfo->TrackWidth, 1.f));
+			float ScreenDelta =  FMath::Max(MouseXPos - ViewportMax, 10.0f);
+			RequestTrackPan.Execute(ScreenDelta, FVector2D(TrackWidth, 1.f));
 		}
 	}
 
-	float GetSnapPosition(const FTrackClampInfo& ClampInfo, float WidgetSpaceNotifyPosition)
+	float GetSnapPosition(const FTrackClampInfo& ClampInfo, float WidgetSpaceNotifyPosition, bool& bOutSnapped)
 	{
 		const FTrackScaleInfo& ScaleInfo = ClampInfo.NotifyTrack->GetCachedScaleInfo();
 
@@ -1319,22 +1328,26 @@ public:
 
 		float CurrentMinSnapDest = MaxSnapDist;
 		float SnapPosition = ScaleInfo.LocalXToInput(WidgetSpaceNotifyPosition);
-		bool bSnapped = OnSnapPosition.IsBound() && OnSnapPosition.Execute(SnapPosition, MaxSnapDist / ScaleInfo.PixelsPerInput);
+		bOutSnapped = OnSnapPosition.IsBound() && OnSnapPosition.Execute(SnapPosition, MaxSnapDist / ScaleInfo.PixelsPerInput);
 		SnapPosition = ScaleInfo.InputToLocalX(SnapPosition);
 
-		if(!bSnapped)
+		float WidgetSpaceStartPosition = ScaleInfo.InputToLocalX(0.0f);
+		float WidgetSpaceEndPosition = ScaleInfo.InputToLocalX(Sequence->SequenceLength);
+
+		if(!bOutSnapped)
 		{
 			// Didn't snap to a bar, snap to the track bounds
-			float WidgetSpaceEndPosition = ScaleInfo.InputToLocalX(Sequence->SequenceLength);
 			float SnapDistBegin = FMath::Abs(-WidgetSpaceNotifyPosition);
 			float SnapDistEnd = FMath::Abs(WidgetSpaceEndPosition - WidgetSpaceNotifyPosition);
 			if(SnapDistBegin < CurrentMinSnapDest)
 			{
 				SnapPosition = 0.0f;
+				bOutSnapped = true;
 			}
 			else if(SnapDistEnd < CurrentMinSnapDest)
 			{
 				SnapPosition = WidgetSpaceEndPosition;
+				bOutSnapped = true;
 			}
 		}
 
@@ -1439,9 +1452,6 @@ public:
 			Info.NotifyTrack = NotifyTracks[i];
 			const FGeometry& CachedGeometry = Info.NotifyTrack->GetCachedGeometry();
 			Info.TrackPos = CachedGeometry.AbsolutePosition.Y;
-			Info.TrackWidth = CachedGeometry.GetLocalSize().X;
-			Info.TrackMin = CachedGeometry.AbsolutePosition.X;
-			Info.TrackMax = Info.TrackMin + Info.TrackWidth;
 			Info.TrackSnapTestPos = Info.TrackPos + (CachedGeometry.Size.Y / 2);
 			Operation->ClampInfos.Add(Info);
 		}
@@ -1527,6 +1537,7 @@ void SAnimNotifyNode::Construct(const FArguments& InArgs)
 	OnNodeDragStarted = InArgs._OnNodeDragStarted;
 	PanTrackRequest = InArgs._PanTrackRequest;
 	OnSelectionChanged = InArgs._OnSelectionChanged;
+	OnUpdatePanel = InArgs._OnUpdatePanel;
 
 	ViewInputMin = InArgs._ViewInputMin;
 	ViewInputMax = InArgs._ViewInputMax;
@@ -1880,6 +1891,8 @@ FReply SAnimNotifyNode::OnMouseMove( const FGeometry& MyGeometry, const FPointer
 	
 	float XPositionInTrack = MyGeometry.AbsolutePosition.X - CachedTrackGeometry.AbsolutePosition.X;
 	float TrackScreenSpaceXPosition = MyGeometry.AbsolutePosition.X - XPositionInTrack;
+	float TrackScreenSpaceOrigin = CachedTrackGeometry.LocalToAbsolute(FVector2D(ScaleInfo.InputToLocalX(0.0f), 0.0f)).X;
+	float TrackScreenSpaceLimit = CachedTrackGeometry.LocalToAbsolute(FVector2D(ScaleInfo.InputToLocalX(Sequence->GetPlayLength()), 0.0f)).X;
 
 	if(CurrentDragHandle == ENotifyStateHandleHit::Start)
 	{
@@ -1897,18 +1910,19 @@ FReply SAnimNotifyNode::OnMouseMove( const FGeometry& MyGeometry, const FPointer
 				NewDisplayTime -= MinimumStateDuration - NewDuration;
 			}
 
-			NodeObjectInterface->SetTime(NewDisplayTime);
+			NodeObjectInterface->SetTime(FMath::Max(0.0f, NewDisplayTime));
 			NodeObjectInterface->SetDuration(NodeObjectInterface->GetDuration() + OldDisplayTime - NodeObjectInterface->GetTime());
 		}
 		else if(NodeObjectInterface->GetDuration() > MinimumStateDuration)
 		{
-			float Overflow = HandleOverflowPan(MouseEvent.GetScreenSpacePosition(), TrackScreenSpaceXPosition);
+			float Overflow = HandleOverflowPan(MouseEvent.GetScreenSpacePosition(), TrackScreenSpaceXPosition, TrackScreenSpaceOrigin, TrackScreenSpaceLimit);
 
 			// Update scale info to the new view inputs after panning
 			ScaleInfo.ViewMinInput = ViewInputMin.Get();
 			ScaleInfo.ViewMaxInput = ViewInputMax.Get();
 
-			NodeObjectInterface->SetTime(ScaleInfo.LocalXToInput(Overflow < 0.0f ? 0.0f : CachedAllotedGeometrySize.X));
+			float NewDisplayTime = ScaleInfo.LocalXToInput((MouseEvent.GetScreenSpacePosition() - MyGeometry.AbsolutePosition + XPositionInTrack).X);
+			NodeObjectInterface->SetTime(FMath::Max(0.0f, NewDisplayTime));
 			NodeObjectInterface->SetDuration(NodeObjectInterface->GetDuration() + OldDisplayTime - NodeObjectInterface->GetTime());
 
 			// Adjust incase we went under the minimum
@@ -1954,13 +1968,19 @@ FReply SAnimNotifyNode::OnMouseMove( const FGeometry& MyGeometry, const FPointer
 		}
 		else if(NodeObjectInterface->GetDuration() > MinimumStateDuration)
 		{
-			float Overflow = HandleOverflowPan(MouseEvent.GetScreenSpacePosition(), TrackScreenSpaceXPosition);
+			float Overflow = HandleOverflowPan(MouseEvent.GetScreenSpacePosition(), TrackScreenSpaceXPosition, TrackScreenSpaceOrigin, TrackScreenSpaceLimit);
 
 			// Update scale info to the new view inputs after panning
 			ScaleInfo.ViewMinInput = ViewInputMin.Get();
 			ScaleInfo.ViewMaxInput = ViewInputMax.Get();
 
-			NodeObjectInterface->SetDuration(FMath::Max(ScaleInfo.LocalXToInput(Overflow > 0.0f ? CachedAllotedGeometrySize.X : 0.0f) - NodeObjectInterface->GetTime(), MinimumStateDuration));
+			float NewDuration = ScaleInfo.LocalXToInput((MouseEvent.GetScreenSpacePosition() - MyGeometry.AbsolutePosition + XPositionInTrack).X) - NodeObjectInterface->GetTime();
+			NodeObjectInterface->SetDuration(FMath::Max(NewDuration, MinimumStateDuration));
+		}
+
+		if(NodeObjectInterface->GetTime() + NodeObjectInterface->GetDuration() > Sequence->GetPlayLength())
+		{
+			NodeObjectInterface->SetDuration(Sequence->GetPlayLength() - NodeObjectInterface->GetTime());
 		}
 
 		// Now we know where the scrub handle should be, look for possible snaps on montage marker bars
@@ -2008,6 +2028,8 @@ FReply SAnimNotifyNode::OnMouseButtonUp( const FGeometry& MyGeometry, const FPoi
 		Sequence->PostEditChange();
 		Sequence->MarkPackageDirty();
 
+		OnUpdatePanel.ExecuteIfBound();
+
 		return FReply::Handled().ReleaseMouseCapture();
 	}
 
@@ -2025,10 +2047,13 @@ float SAnimNotifyNode::GetScrubHandleSnapPosition( float NotifyLocalX, ENotifySt
 
 	if(OnSnapPosition.IsBound())
 	{
-		OnSnapPosition.Execute(SnapPosition, MaxSnapDist / ScaleInfo.PixelsPerInput);
+		if(OnSnapPosition.Execute(SnapPosition, MaxSnapDist / ScaleInfo.PixelsPerInput))
+		{
+			return ScaleInfo.InputToLocalX(SnapPosition);
+		}
 	}	
 
-	return ScaleInfo.InputToLocalX(SnapPosition);
+	return -1.0f;
 }
 
 FReply SAnimNotifyNode::OnFocusReceived(const FGeometry& MyGeometry, const FFocusEvent& InFocusEvent)
@@ -2036,19 +2061,21 @@ FReply SAnimNotifyNode::OnFocusReceived(const FGeometry& MyGeometry, const FFocu
 	return FReply::Handled().SetUserFocus(AsShared(), EFocusCause::SetDirectly, true);
 }
 
-float SAnimNotifyNode::HandleOverflowPan( const FVector2D &ScreenCursorPos, float TrackScreenSpaceXPosition )
+float SAnimNotifyNode::HandleOverflowPan( const FVector2D &ScreenCursorPos, float TrackScreenSpaceXPosition, float TrackScreenSpaceMin, float TrackScreenSpaceMax )
 {
 	float Overflow = 0.0f;
-	if(ScreenCursorPos.X < TrackScreenSpaceXPosition)
+
+	if(ScreenCursorPos.X < TrackScreenSpaceXPosition && TrackScreenSpaceXPosition > TrackScreenSpaceMin - 10.0f)
 	{
 		// Overflow left edge
-		Overflow = ScreenCursorPos.X - TrackScreenSpaceXPosition;
+		Overflow = FMath::Min(ScreenCursorPos.X - TrackScreenSpaceXPosition, -10.0f);
 	}
-	else
+	else if(ScreenCursorPos.X > CachedAllotedGeometrySize.X && (TrackScreenSpaceXPosition + CachedAllotedGeometrySize.X) < TrackScreenSpaceMax + 10.0f)
 	{
 		// Overflow right edge
-		Overflow = ScreenCursorPos.X - (TrackScreenSpaceXPosition + CachedAllotedGeometrySize.X);
+		Overflow = FMath::Max(ScreenCursorPos.X - (TrackScreenSpaceXPosition + CachedAllotedGeometrySize.X), 10.0f);
 	}
+
 	PanTrackRequest.ExecuteIfBound(Overflow, CachedAllotedGeometrySize);
 
 	return Overflow;
@@ -4175,34 +4202,18 @@ FReply SAnimNotifyPanel::OnNotifyNodeDragStarted(TArray<TSharedPtr<SAnimNotifyNo
 		Track->DisconnectSelectedNodesForDrag(Nodes);
 	}
 
-	FVector2D OverlayOrigin = Nodes[0]->GetScreenPosition();
-	FVector2D OverlayExtents = OverlayOrigin;
-	OverlayExtents.X += Nodes[0]->GetDurationSize();
+	FBox2D OverlayBounds(Nodes[0]->GetScreenPosition(), Nodes[0]->GetScreenPosition() + FVector2D(Nodes[0]->GetDurationSize(), 0.0f));
 	for(int32 Idx = 1 ; Idx < Nodes.Num() ; ++Idx)
 	{
 		TSharedPtr<SAnimNotifyNode> Node = Nodes[Idx];
 		FVector2D NodePosition = Node->GetScreenPosition();
 		float NodeDuration = Node->GetDurationSize();
 
-		if(NodePosition.X < OverlayOrigin.X)
-		{
-			OverlayOrigin.X = NodePosition.X;
-		}
-		else if(NodePosition.X + NodeDuration > OverlayExtents.X)
-		{
-			OverlayExtents.X = NodePosition.X + NodeDuration;
-		}
-
-		if(NodePosition.Y < OverlayOrigin.Y)
-		{
-			OverlayOrigin.Y = NodePosition.Y;
-		}
-		else if(NodePosition.Y + NotifyHeight > OverlayExtents.Y)
-		{
-			OverlayExtents.Y = NodePosition.Y + NotifyHeight;
-		}
+		OverlayBounds += FBox2D(NodePosition, NodePosition + FVector2D(NodeDuration, 0.0f));
 	}
-	OverlayExtents -= OverlayOrigin;
+
+	FVector2D OverlayOrigin = OverlayBounds.Min;
+	FVector2D OverlayExtents = OverlayBounds.GetSize();
 
 	for(TSharedPtr<SAnimNotifyNode> Node : Nodes)
 	{
diff --git a/Engine/Source/Editor/Persona/Private/SAnimSequenceCurveEditor.cpp b/Engine/Source/Editor/Persona/Private/SAnimSequenceCurveEditor.cpp
--- a/Engine/Source/Editor/Persona/Private/SAnimSequenceCurveEditor.cpp
+++ b/Engine/Source/Editor/Persona/Private/SAnimSequenceCurveEditor.cpp
@@ -124,7 +124,6 @@ public:
 		NewCurveModel->SetShortDisplayName(CurveDisplayName);
 		NewCurveModel->SetLongDisplayName(CurveDisplayName);
 		NewCurveModel->SetColor(CurveColor);
-		NewCurveModel->SetClampInputRange(MakeAttributeLambda([this](){ return TRange<double>(0.0, AnimSequence->GetPlayLength()); }));
 		NewCurveModel->OnCurveModified().Add(OnCurveModified);
 
 		OutCurveModels.Add(MoveTemp(NewCurveModel));
diff --git a/Engine/Source/Editor/Persona/Private/SAnimTrackPanel.cpp b/Engine/Source/Editor/Persona/Private/SAnimTrackPanel.cpp
--- a/Engine/Source/Editor/Persona/Private/SAnimTrackPanel.cpp
+++ b/Engine/Source/Editor/Persona/Private/SAnimTrackPanel.cpp
@@ -160,19 +160,6 @@ void SAnimTrackPanel::PanInputViewRange(int32 ScreenDelta, FVector2D ScreenViewS
 	float NewViewInputMin = ViewInputMin.Get() + InputDeltaX;
 	float NewViewInputMax = ViewInputMax.Get() + InputDeltaX;
 
-	// we'd like to keep  the range if outside when panning
-	float SequenceLength = GetSequenceLength();
-	if ( NewViewInputMin < 0.f )
-	{
-		NewViewInputMin = 0.f;
-		NewViewInputMax = ScaleInfo.ViewInputRange;
-	}
-	else if ( NewViewInputMax >SequenceLength )
-	{
-		NewViewInputMax = SequenceLength;
-		NewViewInputMin = NewViewInputMax - ScaleInfo.ViewInputRange;
-	}
-
 	InputViewRangeChanged(NewViewInputMin, NewViewInputMax);
 }
 
diff --git a/Engine/Source/Editor/Persona/Private/STimingTrack.cpp b/Engine/Source/Editor/Persona/Private/STimingTrack.cpp
--- a/Engine/Source/Editor/Persona/Private/STimingTrack.cpp
+++ b/Engine/Source/Editor/Persona/Private/STimingTrack.cpp
@@ -126,35 +126,8 @@ void STimingTrack::OnArrangeChildren(const FGeometry& AllottedGeometry, FArrange
 
 		for(NodeIsland& Island : Islands)
 		{
-			if(Island.Nodes.Num() == 1)
-			{
-				// Keep single nodes on the data track range but skip everything else
-				FBox2D& NodeBox = Island.Nodes[0]->ActualRect;
-				FBox2D& NodeQueryRect = Island.Nodes[0]->QueryRect;
-				float Offset = 0.0f;
-				float Begin = ScaleInfo.LocalXToInput(NodeBox.Min.X);
-				float End = ScaleInfo.LocalXToInput(NodeBox.Max.X);
-				if(Begin < 0)
-				{
-					Offset = ScaleInfo.InputToLocalX(-Begin);
-				}
-				else if(End > TrackMaxValue.Get())
-				{
-					Offset = ScaleInfo.InputToLocalX(TrackMaxValue.Get() - End);
-				}
-
-				if(Offset != 0.0f)
-				{
-					NodeBox.Min.X += Offset;
-					NodeBox.Max.X += Offset;
-					NodeQueryRect.Min.X = NodeBox.Min.X - NodeSeparation;
-					NodeQueryRect.Max.X = NodeBox.Max.X + NodeSeparation;
-				}
-
-				continue;
-			}
-
 			bResolved = true;
+
 			// Island resolution
 			int32 NumIslandNodes = Island.Nodes.Num();
 			float Width = FMath::Max<float>((NumIslandNodes - 1), 0.0f) * NodeSeparation;
@@ -167,13 +140,8 @@ void STimingTrack::OnArrangeChildren(const FGeometry& AllottedGeometry, FArrange
 			}
 			Centre /= NumIslandNodes;
 
-			// Make sure the group stays on the track
 			float Begin = Centre - Width / 2.0f;
-			float WidthAsInput = Width / ScaleInfo.PixelsPerInput;
-			float BeginAsInput = FMath::Clamp(ScaleInfo.LocalXToInput(Begin), 0.0f, TrackMaxValue.Get() - WidthAsInput);
-			Begin = ScaleInfo.InputToLocalX(BeginAsInput);
-			
-			for(int32 NodeIdx = 0 ; NodeIdx < NumIslandNodes ; ++NodeIdx)
+			for(int32 NodeIdx = 0 ; NodeIdx < NumIslandNodes; ++NodeIdx)
 			{
 				FBox2D& NodeBox = Island.Nodes[NodeIdx]->ActualRect;
 				float NodeWidth = NodeBox.GetSize().X;
