Fix the vertices count in the skeletalmesh viewport to avoid counting disabled section.
#rnx

diff --git a/Engine/Source/Editor/Persona/Private/AnimationEditorViewportClient.cpp b/Engine/Source/Editor/Persona/Private/AnimationEditorViewportClient.cpp
--- a/Engine/Source/Editor/Persona/Private/AnimationEditorViewportClient.cpp
+++ b/Engine/Source/Editor/Persona/Private/AnimationEditorViewportClient.cpp
@@ -839,10 +839,13 @@ FText FAnimationViewportClient::GetDisplayInfo(bool bDisplayAllInfo) const
 
 			// Calculate polys based on non clothing sections so we don't duplicate the counts.
 			uint32 NumTotalTriangles = 0;
-			int32 NumSections = LODData.NumNonClothingSections();
+			int32 NumSections = LODData.RenderSections.Num();
 			for(int32 SectionIndex = 0; SectionIndex < NumSections; SectionIndex++)
 			{
-				NumTotalTriangles += LODData.RenderSections[SectionIndex].NumTriangles;
+				if (!LODData.RenderSections[SectionIndex].bDisabled && !LODData.RenderSections[SectionIndex].HasClothingData())
+				{
+					NumTotalTriangles += LODData.RenderSections[SectionIndex].NumTriangles;
+				}
 			}
 
 			TextValue = ConcatenateLine(TextValue, FText::Format(LOCTEXT("MeshInfoFormat", "LOD: {0}, Bones: {1} (Mapped to Vertices: {2}), Polys: {3}"),
@@ -852,21 +855,27 @@ FText FAnimationViewportClient::GetDisplayInfo(bool bDisplayAllInfo) const
 				FText::AsNumber(NumTotalTriangles)));
 
 			TextValue = ConcatenateLine(TextValue, FText::Format(LOCTEXT("ScreenSizeFOVFormat", "Current Screen Size: {0}, FOV: {1}"), FText::AsNumber(CachedScreenSize), FText::AsNumber(ViewFOV)));
-
+			int32 NumVertices = 0;
 			for (int32 SectionIndex = 0; SectionIndex < LODData.RenderSections.Num(); SectionIndex++)
 			{
 				int32 SectionVerts = LODData.RenderSections[SectionIndex].GetNumVertices();
 
-				TextValue = ConcatenateLine(TextValue, FText::Format(LOCTEXT("SectionFormat", " [Section {0}] Verts: {1}, Bones: {2}, Max Influences: {3}"),
+				FText SectionDisabledText = LODData.RenderSections[SectionIndex].bDisabled ? LOCTEXT("SectionIsDisbable", " Disabled") : FText::GetEmpty();
+				TextValue = ConcatenateLine(TextValue, FText::Format(LOCTEXT("SectionFormat", " [Section {0}]{1} Verts: {2}, Bones: {3}, Max Influences: {4}"),
 					FText::AsNumber(SectionIndex),
+					SectionDisabledText,
 					FText::AsNumber(SectionVerts),
 					FText::AsNumber(LODData.RenderSections[SectionIndex].BoneMap.Num()),
 					FText::AsNumber(LODData.RenderSections[SectionIndex].MaxBoneInfluences)
 					));
+				if (!LODData.RenderSections[SectionIndex].bDisabled)
+				{
+					NumVertices += LODData.RenderSections[SectionIndex].GetNumVertices();
+				}
 			}
 
 			TextValue = ConcatenateLine(TextValue, FText::Format(LOCTEXT("TotalVerts", "TOTAL Verts: {0}"),
-				FText::AsNumber(LODData.GetNumVertices())));
+				FText::AsNumber(NumVertices)));
 
 			TextValue = ConcatenateLine(TextValue, FText::Format(LOCTEXT("Sections", "Sections: {0}"),
 				NumSectionsInUse
@@ -931,15 +940,21 @@ FText FAnimationViewportClient::GetDisplayInfo(bool bDisplayAllInfo) const
 
 			// Triangles
 			uint32 NumTotalTriangles = 0;
-			int32 NumSections = LODData.NumNonClothingSections();
+			uint32 NumTotalVertices = 0;
+			int32 NumSections = LODData.RenderSections.Num();
 			for (int32 SectionIndex = 0; SectionIndex < NumSections; SectionIndex++)
 			{
-				NumTotalTriangles += LODData.RenderSections[SectionIndex].NumTriangles;
+				if (!LODData.RenderSections[SectionIndex].bDisabled && !LODData.RenderSections[SectionIndex].HasClothingData())
+				{
+					NumTotalTriangles += LODData.RenderSections[SectionIndex].NumTriangles;
+					NumTotalVertices += LODData.RenderSections[SectionIndex].NumVertices;
+				}
 			}
+
 			TextValue = ConcatenateLine(TextValue, FText::Format(LOCTEXT("TrianglesFormat", "Triangles: {0}"), FText::AsNumber(NumTotalTriangles)));
 
 			// Vertices
-			TextValue = ConcatenateLine(TextValue, FText::Format(LOCTEXT("VerticesFormat", "Vertices: {0}"), FText::AsNumber(LODData.GetNumVertices())));
+			TextValue = ConcatenateLine(TextValue, FText::Format(LOCTEXT("VerticesFormat", "Vertices: {0}"), FText::AsNumber(NumTotalVertices)));
 
 			// UV Channels
 			TextValue = ConcatenateLine(TextValue, FText::Format(LOCTEXT("UVChannelsFormat", "UV Channels: {0}"), FText::AsNumber(LODData.GetNumTexCoords())));
