Fix the basic skeletalmesh info viewport display
#rnx

diff --git a/Engine/Source/Editor/Persona/Private/AnimationEditorViewportClient.cpp b/Engine/Source/Editor/Persona/Private/AnimationEditorViewportClient.cpp
--- a/Engine/Source/Editor/Persona/Private/AnimationEditorViewportClient.cpp
+++ b/Engine/Source/Editor/Persona/Private/AnimationEditorViewportClient.cpp
@@ -842,7 +842,7 @@ FText FAnimationViewportClient::GetDisplayInfo(bool bDisplayAllInfo) const
 			int32 NumSections = LODData.RenderSections.Num();
 			for(int32 SectionIndex = 0; SectionIndex < NumSections; SectionIndex++)
 			{
-				if (!LODData.RenderSections[SectionIndex].bDisabled && !LODData.RenderSections[SectionIndex].HasClothingData())
+				if (!LODData.RenderSections[SectionIndex].bDisabled)
 				{
 					NumTotalTriangles += LODData.RenderSections[SectionIndex].NumTriangles;
 				}
@@ -944,7 +944,7 @@ FText FAnimationViewportClient::GetDisplayInfo(bool bDisplayAllInfo) const
 			int32 NumSections = LODData.RenderSections.Num();
 			for (int32 SectionIndex = 0; SectionIndex < NumSections; SectionIndex++)
 			{
-				if (!LODData.RenderSections[SectionIndex].bDisabled && !LODData.RenderSections[SectionIndex].HasClothingData())
+				if (!LODData.RenderSections[SectionIndex].bDisabled)
 				{
 					NumTotalTriangles += LODData.RenderSections[SectionIndex].NumTriangles;
 					NumTotalVertices += LODData.RenderSections[SectionIndex].NumVertices;
