Fixed flicker when right-click dragging on a notify track to scroll

Also fixed DPI scaling issues when drag/droppiing notifies

diff --git a/Engine/Source/Editor/Persona/Private/SAnimNotifyPanel.cpp b/Engine/Source/Editor/Persona/Private/SAnimNotifyPanel.cpp
--- a/Engine/Source/Editor/Persona/Private/SAnimNotifyPanel.cpp
+++ b/Engine/Source/Editor/Persona/Private/SAnimNotifyPanel.cpp
@@ -1169,39 +1169,38 @@ public:
 
 		NodeGroupPosition = DragDropEvent.GetScreenSpacePosition() + DragOffset;
 
-		FVector2D SelectionBeginPosition = NodeGroupPosition + SelectedNodes[0]->GetNotifyPositionOffset();
-		
 		FTrackClampInfo* SelectionPositionClampInfo = &GetTrackClampInfo(DragDropEvent.GetScreenSpacePosition());
 		if((SelectionPositionClampInfo->NotifyTrack->GetTrackIndex() + TrackSpan) >= ClampInfos.Num())
 		{
 			// Our selection has moved off the bottom of the notify panel, adjust the clamping information to keep it on the panel
 			SelectionPositionClampInfo = &ClampInfos[ClampInfos.Num() - TrackSpan - 1];
 		}
-
+		
 		const FGeometry& TrackGeom = SelectionPositionClampInfo->NotifyTrack->GetCachedGeometry();
 		const FTrackScaleInfo& TrackScaleInfo = SelectionPositionClampInfo->NotifyTrack->GetCachedScaleInfo();
-		
-		float TrackMin = TrackGeom.LocalToAbsolute(FVector2D(TrackScaleInfo.InputToLocalX(0.0f), 0.0f)).X;
-		float TrackMax = TrackGeom.LocalToAbsolute(FVector2D(TrackScaleInfo.InputToLocalX(Sequence->GetPlayLength()), 0.0f)).X;
-		float TrackWidth = TrackMax - TrackMin;
+
+		FVector2D SelectionBeginPosition = TrackGeom.LocalToAbsolute(TrackGeom.AbsoluteToLocal(NodeGroupPosition) + SelectedNodes[0]->GetNotifyPositionOffset());
+	
+		float LocalTrackMin = TrackScaleInfo.InputToLocalX(0.0f);
+		float LocalTrackMax = TrackScaleInfo.InputToLocalX(Sequence->GetPlayLength());
+		float LocalTrackWidth = LocalTrackMax - LocalTrackMin;
 
 		// Tracks the movement amount to apply to the selection due to a snap.
-		float SnapMovement= 0.0f;
+		float SnapMovement = 0.0f;
 		// Clamp the selection into the track
-		float SelectionBeginPositionX = TrackGeom.GetAbsolutePosition().X + TrackGeom.AbsoluteToLocal(SelectionBeginPosition).X;
-		const float SelectionLocalLength = TrackScaleInfo.PixelsPerInput * SelectionTimeLength;
-		const float ClampedEnd = FMath::Clamp(SelectionBeginPositionX + NodeGroupSize.X, TrackMin, TrackMax);
-		const float ClampedBegin = FMath::Clamp(SelectionBeginPositionX, TrackMin, TrackMax);
-		if(ClampedBegin > SelectionBeginPositionX)
+		float SelectionBeginLocalPositionX = TrackGeom.AbsoluteToLocal(SelectionBeginPosition).X;
+		const float ClampedEnd = FMath::Clamp(SelectionBeginLocalPositionX + NodeGroupSize.X, LocalTrackMin, LocalTrackMax);
+		const float ClampedBegin = FMath::Clamp(SelectionBeginLocalPositionX, LocalTrackMin, LocalTrackMax);
+		if(ClampedBegin > SelectionBeginLocalPositionX)
 		{
-			SelectionBeginPositionX = ClampedBegin;
+			SelectionBeginLocalPositionX = ClampedBegin;
 		}
-		else if(ClampedEnd < SelectionBeginPositionX + SelectionLocalLength)
+		else if(ClampedEnd < SelectionBeginLocalPositionX + NodeGroupSize.X)
 		{
-			SelectionBeginPositionX = ClampedEnd - SelectionLocalLength;
+			SelectionBeginLocalPositionX = ClampedEnd - NodeGroupSize.X;
 		}
 
-		SelectionBeginPosition.X = TrackGeom.LocalToAbsolute(FVector2D(SelectionBeginPositionX - TrackGeom.GetAbsolutePosition().X, 0.0f)).X;
+		SelectionBeginPosition.X = TrackGeom.LocalToAbsolute(FVector2D(SelectionBeginLocalPositionX, 0.0f)).X;
 
 		// Handle node snaps
 		bool bSnapped = false;
@@ -1301,22 +1300,22 @@ public:
 
 		CurrentDragXPosition = TrackGeom.AbsoluteToLocal(FVector2D(SelectionBeginPosition.X,0.0f)).X;
 
-		CursorDecoratorWindow->MoveWindowTo(SelectionBeginPosition - SelectedNodes[0]->GetNotifyPositionOffset());
+		CursorDecoratorWindow->MoveWindowTo(TrackGeom.LocalToAbsolute(TrackGeom.AbsoluteToLocal(SelectionBeginPosition) - SelectedNodes[0]->GetNotifyPositionOffset()));
 		NodeGroupPosition = SelectionBeginPosition;
 
 		//scroll view
-		float MouseXPos = TrackGeom.GetAbsolutePosition().X + TrackGeom.AbsoluteToLocal(DragDropEvent.GetScreenSpacePosition()).X;
-		float ViewportMin = TrackGeom.GetAbsolutePosition().X;
-		float ViewportMax = TrackGeom.GetAbsolutePosition().X + TrackGeom.GetAbsoluteSize().X;
-		if(MouseXPos < ViewportMin && ViewportMin > TrackMin - 10.0f)
+		float LocalMouseXPos = TrackGeom.AbsoluteToLocal(DragDropEvent.GetScreenSpacePosition()).X;
+		float LocalViewportMin = 0.0f;
+		float LocalViewportMax = TrackGeom.GetLocalSize().X;
+		if(LocalMouseXPos < LocalViewportMin && LocalViewportMin > LocalTrackMin - 10.0f)
 		{
-			float ScreenDelta = FMath::Max(MouseXPos - ViewportMin, -10.0f);
-			RequestTrackPan.Execute(ScreenDelta, FVector2D(TrackWidth, 1.f));
+			float ScreenDelta = FMath::Max(LocalMouseXPos - LocalViewportMin, -10.0f);
+			RequestTrackPan.Execute(ScreenDelta, FVector2D(LocalTrackWidth, 1.f));
 		}
-		else if(MouseXPos > ViewportMax && ViewportMax < TrackMax + 10.0f)
+		else if(LocalMouseXPos > LocalViewportMax && LocalViewportMax < LocalTrackMax + 10.0f)
 		{
-			float ScreenDelta =  FMath::Max(MouseXPos - ViewportMax, 10.0f);
-			RequestTrackPan.Execute(ScreenDelta, FVector2D(TrackWidth, 1.f));
+			float ScreenDelta =  FMath::Max(LocalMouseXPos - LocalViewportMax, 10.0f);
+			RequestTrackPan.Execute(ScreenDelta, FVector2D(LocalTrackWidth, 1.f));
 		}
 	}
 
@@ -1337,11 +1336,11 @@ public:
 		if(!bOutSnapped)
 		{
 			// Didn't snap to a bar, snap to the track bounds
-			float SnapDistBegin = FMath::Abs(-WidgetSpaceNotifyPosition);
+			float SnapDistBegin = FMath::Abs(WidgetSpaceStartPosition - WidgetSpaceNotifyPosition);
 			float SnapDistEnd = FMath::Abs(WidgetSpaceEndPosition - WidgetSpaceNotifyPosition);
 			if(SnapDistBegin < CurrentMinSnapDest)
 			{
-				SnapPosition = 0.0f;
+				SnapPosition = WidgetSpaceStartPosition;
 				bOutSnapped = true;
 			}
 			else if(SnapDistEnd < CurrentMinSnapDest)
diff --git a/Engine/Source/Editor/Persona/Private/SAnimTrackPanel.cpp b/Engine/Source/Editor/Persona/Private/SAnimTrackPanel.cpp
--- a/Engine/Source/Editor/Persona/Private/SAnimTrackPanel.cpp
+++ b/Engine/Source/Editor/Persona/Private/SAnimTrackPanel.cpp
@@ -55,9 +55,6 @@ void SAnimTrackPanel::Construct(const FArguments& InArgs)
 	OnSetInputViewRange = InArgs._OnSetInputViewRange;
 
 	WidgetWidth = InArgs._WidgetWidth;
-
-	bPanning = false;
-	PanningDistance = 0.f;
 }
 
 TSharedRef<class S2ColumnWidget> SAnimTrackPanel::Create2ColumnWidget( TSharedRef<SVerticalBox> Parent )
@@ -74,83 +71,6 @@ TSharedRef<class S2ColumnWidget> SAnimTrackPanel::Create2ColumnWidget( TSharedRe
 	return NewTrack.ToSharedRef();
 }
 
-FReply SAnimTrackPanel::OnMouseButtonDown( const FGeometry& InMyGeometry, const FPointerEvent& InMouseEvent )
-{
-	const bool bRightMouseButton = InMouseEvent.GetEffectingButton() == EKeys::RightMouseButton;
-
-	if(bRightMouseButton)
-	{
-		UE_LOG(LogAnimation, Log, TEXT("MouseButtonDown %d, %0.5f"), bPanning, PanningDistance);
-		PanningDistance = 0.f;
-	}
-
-	return FReply::Unhandled();
-}
-
-FReply SAnimTrackPanel::OnMouseButtonUp( const FGeometry& InMyGeometry, const FPointerEvent& InMouseEvent )
-{
-	const bool bRightMouseButton = InMouseEvent.GetEffectingButton() == EKeys::RightMouseButton;
-
-	if(bRightMouseButton && this->HasMouseCapture())
-	{
-		bPanning = false;
-		PanningDistance = 0.f;
-		UE_LOG(LogAnimation, Log, TEXT("MouseButtonUp (Releasing Mouse) %d, %0.5f"), bPanning, PanningDistance);
-		return FReply::Handled().ReleaseMouseCapture();
-	}
-
-	return FReply::Unhandled();
-}
-
-FReply SAnimTrackPanel::OnMouseMove( const FGeometry& InMyGeometry, const FPointerEvent& InMouseEvent )
-{
-	const bool bRightMouseButtonDown = InMouseEvent.IsMouseButtonDown(EKeys::RightMouseButton);
-
-	// When mouse moves, if we are moving a key, update its 'input' position
-	if(bRightMouseButtonDown)
-	{
-		if( !bPanning )
-		{
-			PanningDistance += FMath::Abs(InMouseEvent.GetCursorDelta().X);
-			if ( PanningDistance > FSlateApplication::Get().GetDragTriggerDistance() )
-			{
-				bPanning = true;
-				UE_LOG(LogAnimation, Log, TEXT("MouseMove (Capturing Mouse) %d, %0.5f"), bPanning, PanningDistance);
-				return FReply::Handled().CaptureMouse(SharedThis(this));
-			}
-		}
-		else
-		{
-			FTrackScaleInfo ScaleInfo(ViewInputMin.Get(),  ViewInputMax.Get(), 0.f, 0.f, InMyGeometry.Size);
-			FVector2D ScreenDelta = InMouseEvent.GetCursorDelta();
-			FVector2D InputDelta;
-			InputDelta.X = ScreenDelta.X/ScaleInfo.PixelsPerInput;
-			InputDelta.Y = -ScreenDelta.Y/ScaleInfo.PixelsPerOutput;
-
-			float NewViewInputMin = ViewInputMin.Get() - InputDelta.X;
-			float NewViewInputMax = ViewInputMax.Get() - InputDelta.X;
-			// we'd like to keep  the range if outside when panning
-			if ( NewViewInputMin < InputMin.Get() )
-			{
-				NewViewInputMin = InputMin.Get();
-				NewViewInputMax = ScaleInfo.ViewInputRange;
-			}
-			else if ( NewViewInputMax > InputMax.Get() )
-			{
-				NewViewInputMax = InputMax.Get();
-				NewViewInputMin = NewViewInputMax - ScaleInfo.ViewInputRange;
-			}
-
-			InputViewRangeChanged(NewViewInputMin, NewViewInputMax);
-
-			UE_LOG(LogAnimation, Log, TEXT("MouseMove (Panning) %0.2f, %0.2f"), ViewInputMin.Get(), ViewInputMax.Get());
-			return FReply::Handled();
-		}
-	}
-
-	return FReply::Unhandled();
-}
-
 void SAnimTrackPanel::PanInputViewRange(int32 ScreenDelta, FVector2D ScreenViewSize)
 {
 	FTrackScaleInfo ScaleInfo(ViewInputMin.Get(),  ViewInputMax.Get(), 0.f, 0.f, ScreenViewSize);
diff --git a/Engine/Source/Editor/Persona/Private/SAnimTrackPanel.h b/Engine/Source/Editor/Persona/Private/SAnimTrackPanel.h
--- a/Engine/Source/Editor/Persona/Private/SAnimTrackPanel.h
+++ b/Engine/Source/Editor/Persona/Private/SAnimTrackPanel.h
@@ -89,11 +89,6 @@ public:
 
 	void Construct(const FArguments& InArgs);
 
-	// SWidget interface
-	virtual FReply OnMouseButtonDown( const FGeometry& MyGeometry, const FPointerEvent& MouseEvent ) override;
-	virtual FReply OnMouseButtonUp( const FGeometry& MyGeometry, const FPointerEvent& MouseEvent ) override;
-	virtual FReply OnMouseMove( const FGeometry& MyGeometry, const FPointerEvent& MouseEvent ) override;
-
 	void PanInputViewRange(int32 ScreenDelta, FVector2D ScreenViewSize);
 
 	virtual float GetSequenceLength() const {return 0.0f;}
@@ -115,8 +110,4 @@ protected:
 
 	/** Controls the width of the tracks column */
 	float WidgetWidth;
-
-	/** If we are currently panning the panel*/
-	bool bPanning;
-	float PanningDistance;
 };
