- Fix for SkeletalMesh not being notified about LOD properties being modified
- This would cause issues if uniform sampling was enabled as the data required would not be ready for things to like Niagara to sample from

#rnx

diff --git a/Engine/Source/Editor/Persona/Private/LODInfoUILayout.cpp b/Engine/Source/Editor/Persona/Private/LODInfoUILayout.cpp
--- a/Engine/Source/Editor/Persona/Private/LODInfoUILayout.cpp
+++ b/Engine/Source/Editor/Persona/Private/LODInfoUILayout.cpp
@@ -4,6 +4,7 @@
 #include "CoreMinimal.h"
 #include "UObject/ObjectMacros.h"
 #include "UObject/Object.h"
+#include "SkeletalMeshTypes.h"
 
 ULODInfoUILayout::ULODInfoUILayout(const FObjectInitializer& ObjectInitializer)
 	: Super(ObjectInitializer)
@@ -34,4 +35,5 @@ void ULODInfoUILayout::PostEditChangeProperty(struct FPropertyChangedEvent& Prop
 	SkeletalMesh->Modify();
 	//Copy the LODInfo into the real skeletal mesh LODInfo data
 	*SkeletalMeshLODInfo = LODInfo;
-}
\ No newline at end of file
+	FScopedSkeletalMeshPostEditChange ScopeSkeletalmeshPostEdit(SkeletalMesh);
+}
diff --git a/Engine/Source/Runtime/Engine/Private/SkeletalMesh.cpp b/Engine/Source/Runtime/Engine/Private/SkeletalMesh.cpp
--- a/Engine/Source/Runtime/Engine/Private/SkeletalMesh.cpp
+++ b/Engine/Source/Runtime/Engine/Private/SkeletalMesh.cpp
@@ -1364,6 +1364,10 @@ void USkeletalMesh::PostEditChangeProperty(FPropertyChangedEvent& PropertyChange
 		{
 			SamplingInfo.BuildWholeMesh(this);
 		}
+		else if (PropertyChangedEvent.MemberProperty->GetFName() == FName(TEXT("bSupportUniformlyDistributedSampling")))
+		{
+			SamplingInfo.BuildWholeMesh(this);
+		}
 	}
 	else
 	{
