Make sure the disable section are part of the vertices count
#rnx

diff --git a/Engine/Source/Editor/Persona/Private/AnimationEditorViewportClient.cpp b/Engine/Source/Editor/Persona/Private/AnimationEditorViewportClient.cpp
--- a/Engine/Source/Editor/Persona/Private/AnimationEditorViewportClient.cpp
+++ b/Engine/Source/Editor/Persona/Private/AnimationEditorViewportClient.cpp
@@ -842,10 +842,7 @@ FText FAnimationViewportClient::GetDisplayInfo(bool bDisplayAllInfo) const
 			int32 NumSections = LODData.RenderSections.Num();
 			for(int32 SectionIndex = 0; SectionIndex < NumSections; SectionIndex++)
 			{
-				if (!LODData.RenderSections[SectionIndex].bDisabled)
-				{
-					NumTotalTriangles += LODData.RenderSections[SectionIndex].NumTriangles;
-				}
+				NumTotalTriangles += LODData.RenderSections[SectionIndex].NumTriangles;
 			}
 
 			TextValue = ConcatenateLine(TextValue, FText::Format(LOCTEXT("MeshInfoFormat", "LOD: {0}, Bones: {1} (Mapped to Vertices: {2}), Polys: {3}"),
@@ -855,7 +852,6 @@ FText FAnimationViewportClient::GetDisplayInfo(bool bDisplayAllInfo) const
 				FText::AsNumber(NumTotalTriangles)));
 
 			TextValue = ConcatenateLine(TextValue, FText::Format(LOCTEXT("ScreenSizeFOVFormat", "Current Screen Size: {0}, FOV: {1}"), FText::AsNumber(CachedScreenSize), FText::AsNumber(ViewFOV)));
-			int32 NumVertices = 0;
 			for (int32 SectionIndex = 0; SectionIndex < LODData.RenderSections.Num(); SectionIndex++)
 			{
 				int32 SectionVerts = LODData.RenderSections[SectionIndex].GetNumVertices();
@@ -868,14 +864,10 @@ FText FAnimationViewportClient::GetDisplayInfo(bool bDisplayAllInfo) const
 					FText::AsNumber(LODData.RenderSections[SectionIndex].BoneMap.Num()),
 					FText::AsNumber(LODData.RenderSections[SectionIndex].MaxBoneInfluences)
 					));
-				if (!LODData.RenderSections[SectionIndex].bDisabled)
-				{
-					NumVertices += LODData.RenderSections[SectionIndex].GetNumVertices();
-				}
 			}
 
 			TextValue = ConcatenateLine(TextValue, FText::Format(LOCTEXT("TotalVerts", "TOTAL Verts: {0}"),
-				FText::AsNumber(NumVertices)));
+				FText::AsNumber(LODData.GetNumVertices())));
 
 			TextValue = ConcatenateLine(TextValue, FText::Format(LOCTEXT("Sections", "Sections: {0}"),
 				NumSectionsInUse
@@ -940,21 +932,16 @@ FText FAnimationViewportClient::GetDisplayInfo(bool bDisplayAllInfo) const
 
 			// Triangles
 			uint32 NumTotalTriangles = 0;
-			uint32 NumTotalVertices = 0;
 			int32 NumSections = LODData.RenderSections.Num();
 			for (int32 SectionIndex = 0; SectionIndex < NumSections; SectionIndex++)
 			{
-				if (!LODData.RenderSections[SectionIndex].bDisabled)
-				{
-					NumTotalTriangles += LODData.RenderSections[SectionIndex].NumTriangles;
-					NumTotalVertices += LODData.RenderSections[SectionIndex].NumVertices;
-				}
+				NumTotalTriangles += LODData.RenderSections[SectionIndex].NumTriangles;
 			}
 
 			TextValue = ConcatenateLine(TextValue, FText::Format(LOCTEXT("TrianglesFormat", "Triangles: {0}"), FText::AsNumber(NumTotalTriangles)));
 
 			// Vertices
-			TextValue = ConcatenateLine(TextValue, FText::Format(LOCTEXT("VerticesFormat", "Vertices: {0}"), FText::AsNumber(NumTotalVertices)));
+			TextValue = ConcatenateLine(TextValue, FText::Format(LOCTEXT("VerticesFormat", "Vertices: {0}"), FText::AsNumber(LODData.GetNumVertices())));
 
 			// UV Channels
 			TextValue = ConcatenateLine(TextValue, FText::Format(LOCTEXT("UVChannelsFormat", "UV Channels: {0}"), FText::AsNumber(LODData.GetNumTexCoords())));
diff --git a/Engine/Source/Runtime/Engine/Private/Rendering/SkeletalMeshLODImporterData.cpp b/Engine/Source/Runtime/Engine/Private/Rendering/SkeletalMeshLODImporterData.cpp
--- a/Engine/Source/Runtime/Engine/Private/Rendering/SkeletalMeshLODImporterData.cpp
+++ b/Engine/Source/Runtime/Engine/Private/Rendering/SkeletalMeshLODImporterData.cpp
@@ -549,13 +549,10 @@ void FReductionBaseSkeletalMeshBulkData::CacheGeometryInfo(const FSkeletalMeshLO
 	for (int32 SectionIndex = 0; SectionIndex < SourceLODModel.Sections.Num(); ++SectionIndex)
 	{
 		const FSkelMeshSection& Section = SourceLODModel.Sections[SectionIndex];
-
-		if (!Section.bDisabled)
-		{
-			//Make sure the count fit in a uint32
-			CacheLODVertexNumber += Section.NumVertices < 0 ? 0 : Section.NumVertices;
-			CacheLODTriNumber += Section.NumTriangles;
-		}
+		//We count disabled section, since the render buffer contain the disabled section data. This is crucial for memory budget
+		//Make sure the count fit in a uint32
+		CacheLODVertexNumber += Section.NumVertices < 0 ? 0 : Section.NumVertices;
+		CacheLODTriNumber += Section.NumTriangles;
 	}
 }
 
diff --git a/Engine/Source/Runtime/Engine/Private/SkeletalMesh.cpp b/Engine/Source/Runtime/Engine/Private/SkeletalMesh.cpp
--- a/Engine/Source/Runtime/Engine/Private/SkeletalMesh.cpp
+++ b/Engine/Source/Runtime/Engine/Private/SkeletalMesh.cpp
@@ -494,12 +494,9 @@ bool USkeletalMesh::IsReductionActive(int32 LODIndex) const
 				{
 					const FSkelMeshSection& Section = LODModel.Sections[SectionIndex];
 
-					if (!Section.bDisabled)
-					{
-						//Make sure the count fit in a uint32
-						LODVertexNumber += Section.NumVertices < 0 ? 0 : Section.NumVertices;
-						LODTriNumber += Section.NumTriangles;
-					}
+					//Make sure the count fit in a uint32
+					LODVertexNumber += Section.NumVertices < 0 ? 0 : Section.NumVertices;
+					LODTriNumber += Section.NumTriangles;
 				}
 			}
 			else if (GetImportedModel()->OriginalReductionSourceMeshData.IsValidIndex(LODIndex)
