Fixed right mouse button scrolling on notify tracks

diff --git a/Engine/Source/Editor/Persona/Private/SAnimNotifyPanel.cpp b/Engine/Source/Editor/Persona/Private/SAnimNotifyPanel.cpp
--- a/Engine/Source/Editor/Persona/Private/SAnimNotifyPanel.cpp
+++ b/Engine/Source/Editor/Persona/Private/SAnimNotifyPanel.cpp
@@ -2213,6 +2213,7 @@ void SAnimNotifyTrack::Construct(const FArguments& InArgs)
 	this->ChildSlot
 	[
 			SAssignNew( TrackArea, SBorder )
+			.Visibility(EVisibility::SelfHitTestInvisible)
 			.BorderImage( FEditorStyle::GetBrush("NoBorder") )
 			.Padding( FMargin(0.f, 0.f) )
 	];
@@ -3908,6 +3909,7 @@ void SAnimNotifyPanel::Construct(const FArguments& InArgs, const TSharedRef<FAni
 	this->ChildSlot
 	[
 		SAssignNew(PanelArea, SBorder)
+		.Visibility(EVisibility::SelfHitTestInvisible)
 		.AddMetaData<FTagMetaData>(TEXT("AnimNotify.Notify"))
 		.BorderImage(FEditorStyle::GetBrush("NoBorder"))
 		.Padding(0.0f)
@@ -4710,8 +4712,8 @@ FReply SAnimNotifyPanel::OnMouseMove(const FGeometry& MyGeometry, const FPointer
 		{
 			Marquee.Rect.UpdateEndPoint(MyGeometry.AbsoluteToLocal(MouseEvent.GetScreenSpacePosition()));
 			RefreshMarqueeSelectedNodes(MyGeometry);
+			return FReply::Handled();
 		}
-		return FReply::Handled();
 	}
 
 	return BaseReply;
