Refresh anim sequence model tracks whenever curves are changed outside of the Animation Editor(s)

diff --git a/Engine/Source/Editor/Persona/Private/AnimTimeline/AnimModel_AnimSequenceBase.cpp b/Engine/Source/Editor/Persona/Private/AnimTimeline/AnimModel_AnimSequenceBase.cpp
--- a/Engine/Source/Editor/Persona/Private/AnimTimeline/AnimModel_AnimSequenceBase.cpp
+++ b/Engine/Source/Editor/Persona/Private/AnimTimeline/AnimModel_AnimSequenceBase.cpp
@@ -35,6 +35,8 @@ FAnimModel_AnimSequenceBase::FAnimModel_AnimSequenceBase(const TSharedRef<IPerso
 		bElementNodeDisplayFlag = false;
 	}
 
+	AnimSequenceBase->RegisterOnAnimTrackCurvesChanged(UAnimSequenceBase::FOnAnimTrackCurvesChanged::CreateRaw(this, &FAnimModel_AnimSequenceBase::RefreshTracks));
+	
 	if(GEditor)
 	{
 		GEditor->RegisterForUndo(this);
@@ -47,6 +49,8 @@ FAnimModel_AnimSequenceBase::~FAnimModel_AnimSequenceBase()
 	{
 		GEditor->UnregisterForUndo(this);
 	}
+
+	AnimSequenceBase->UnregisterOnAnimTrackCurvesChanged(this);
 }
 
 void FAnimModel_AnimSequenceBase::Initialize()
