Revert previous change adressing UE-91227 which caused UE-91332, and fixed the crash rather than the UX issue instead

diff --git a/Engine/Source/Editor/AnimationEditor/Private/AnimationEditor.cpp b/Engine/Source/Editor/AnimationEditor/Private/AnimationEditor.cpp
--- a/Engine/Source/Editor/AnimationEditor/Private/AnimationEditor.cpp
+++ b/Engine/Source/Editor/AnimationEditor/Private/AnimationEditor.cpp
@@ -442,7 +442,6 @@ TSharedPtr<SDockTab> FAnimationEditor::OpenNewAnimationDocumentTab(UAnimationAss
 				];
 
 			OpenedTab->SetLeftContent(IDocumentation::Get()->CreateAnchor(DocumentLink));
-			OpenedTab->SetCanCloseTab(SDockTab::FCanCloseTab::CreateLambda([]() -> bool { return false; }));
 
 			TabManager->InsertNewDocumentTab(AnimationEditorTabs::DocumentTab, FTabManager::ESearchPreference::RequireClosedTab, OpenedTab.ToSharedRef());
 
diff --git a/Engine/Source/Editor/Persona/Private/AnimTimeline/AnimTimeSliderController.cpp b/Engine/Source/Editor/Persona/Private/AnimTimeline/AnimTimeSliderController.cpp
--- a/Engine/Source/Editor/Persona/Private/AnimTimeline/AnimTimeSliderController.cpp
+++ b/Engine/Source/Editor/Persona/Private/AnimTimeline/AnimTimeSliderController.cpp
@@ -930,11 +930,13 @@ int32 FAnimTimeSliderController::OnPaintViewArea( const FGeometry& AllottedGeome
 		);
 	}
 
+	TSharedPtr<FAnimModel> AnimModel = WeakModel.Pin();
+	if (AnimModel.IsValid())
 	{
 		const FLinearColor LineColor = GetDefault<UPersonaOptions>()->SectionTimingNodeColor;
 
 		// Draw all the times that we can drag in the timeline
-		for(double Time : WeakModel.Pin()->GetEditableTimes())
+		for(double Time : AnimModel->GetEditableTimes())
 		{
 			const float LinePos = RangeToScreen.InputToLocalX(Time);
 
