Fixed ui alignment of the "Socket Name" property

diff --git a/Engine/Source/Editor/Persona/Private/AnimNotifyDetails.cpp b/Engine/Source/Editor/Persona/Private/AnimNotifyDetails.cpp
--- a/Engine/Source/Editor/Persona/Private/AnimNotifyDetails.cpp
+++ b/Engine/Source/Editor/Persona/Private/AnimNotifyDetails.cpp
@@ -177,7 +177,7 @@ void FAnimNotifyDetails::AddBoneNameProperty(IDetailCategoryBuilder& CategoryBui
 					[
 						SNew(SHorizontalBox)
 						+ SHorizontalBox::Slot()
-						.Padding(FMargin(2, 1, 0, 1))
+						.Padding(FMargin(0, 1, 0, 1))
 						[
 							SNew(STextBlock)
 							.Text(Property->GetPropertyDisplayName())
