Deleting unused SAnimCurvePanel and SAnimTrackCurvePanel
#misc removed includes / dependencies from other anim editor widget

diff --git a/Engine/Source/Editor/Persona/Private/AnimTimeline/AnimTimelineTrack_Curves.cpp b/Engine/Source/Editor/Persona/Private/AnimTimeline/AnimTimelineTrack_Curves.cpp
--- a/Engine/Source/Editor/Persona/Private/AnimTimeline/AnimTimelineTrack_Curves.cpp
+++ b/Engine/Source/Editor/Persona/Private/AnimTimeline/AnimTimelineTrack_Curves.cpp
@@ -6,7 +6,6 @@
 #include "AnimSequenceTimelineCommands.h"
 #include "Framework/MultiBox/MultiBoxBuilder.h"
 #include "Animation/AnimSequenceBase.h"
-#include "SAnimCurvePanel.h"
 #include "Widgets/Input/STextEntryPopup.h"
 #include "Framework/Application/SlateApplication.h"
 #include "ScopedTransaction.h"
@@ -138,6 +137,26 @@ TSharedRef<SWidget> FAnimTimelineTrack_Curves::BuildCurvesSubMenu()
 	return MenuBuilder.MakeWidget();
 }
 
+
+struct FSmartNameSortItem
+{
+	FName SmartName;
+	USkeleton::AnimCurveUID ID;
+
+	FSmartNameSortItem(const FName& InSmartName, const USkeleton::AnimCurveUID& InID)
+		: SmartName(InSmartName)
+		, ID(InID)
+	{}
+};
+
+struct FSmartNameSortItemSortOp
+{
+	bool operator()(const FSmartNameSortItem& A, const FSmartNameSortItem& B) const
+	{
+		return (A.SmartName.Compare(B.SmartName) < 0);
+	}
+};
+
 void FAnimTimelineTrack_Curves::FillMetadataEntryMenu(FMenuBuilder& Builder)
 {
 	UAnimSequenceBase* AnimSequenceBase = GetModel()->GetAnimSequenceBase();
diff --git a/Engine/Source/Editor/Persona/Private/SAnimCompositeEditor.h b/Engine/Source/Editor/Persona/Private/SAnimCompositeEditor.h
--- a/Engine/Source/Editor/Persona/Private/SAnimCompositeEditor.h
+++ b/Engine/Source/Editor/Persona/Private/SAnimCompositeEditor.h
@@ -6,7 +6,6 @@
 #include "CoreMinimal.h"
 #include "Widgets/SWidget.h"
 #include "Widgets/DeclarativeSyntaxSupport.h"
-#include "SAnimCurvePanel.h"
 #include "SAnimEditorBase.h"
 #include "Animation/AnimComposite.h"
 #include "SAnimCompositePanel.h"
diff --git a/Engine/Source/Editor/Persona/Private/SAnimCurvePanel.cpp b/Engine/Source/Editor/Persona/Private/SAnimCurvePanel.cpp
deleted file mode 100644
diff --git a/Engine/Source/Editor/Persona/Private/SAnimCurvePanel.h b/Engine/Source/Editor/Persona/Private/SAnimCurvePanel.h
deleted file mode 100644
diff --git a/Engine/Source/Editor/Persona/Private/SAnimEditorBase.h b/Engine/Source/Editor/Persona/Private/SAnimEditorBase.h
--- a/Engine/Source/Editor/Persona/Private/SAnimEditorBase.h
+++ b/Engine/Source/Editor/Persona/Private/SAnimEditorBase.h
@@ -12,7 +12,6 @@
 #include "PersonaDelegates.h"
 #include "SAnimationScrubPanel.h"
 #include "EditorObjectsTracker.h"
-#include "SAnimCurvePanel.h"
 
 class FAnimModel;
 class SAnimTimeline;
diff --git a/Engine/Source/Editor/Persona/Private/SAnimMontagePanel.h b/Engine/Source/Editor/Persona/Private/SAnimMontagePanel.h
--- a/Engine/Source/Editor/Persona/Private/SAnimMontagePanel.h
+++ b/Engine/Source/Editor/Persona/Private/SAnimMontagePanel.h
@@ -14,6 +14,7 @@
 #include "Animation/EditorAnimBaseObj.h"
 #include "STrack.h"
 #include "EditorUndoClient.h"
+#include "SAnimCurveEd.h"
 
 class FMenuBuilder;
 class SBorder;
diff --git a/Engine/Source/Editor/Persona/Private/SAnimStreamableEditor.h b/Engine/Source/Editor/Persona/Private/SAnimStreamableEditor.h
--- a/Engine/Source/Editor/Persona/Private/SAnimStreamableEditor.h
+++ b/Engine/Source/Editor/Persona/Private/SAnimStreamableEditor.h
@@ -6,7 +6,6 @@
 #include "CoreMinimal.h"
 #include "Widgets/SWidget.h"
 #include "Widgets/DeclarativeSyntaxSupport.h"
-#include "SAnimCurvePanel.h"
 #include "SAnimEditorBase.h"
 #include "Animation/AnimStreamable.h"
 
diff --git a/Engine/Source/Editor/Persona/Private/SAnimTrackCurvePanel.cpp b/Engine/Source/Editor/Persona/Private/SAnimTrackCurvePanel.cpp
deleted file mode 100644
diff --git a/Engine/Source/Editor/Persona/Private/SAnimTrackCurvePanel.h b/Engine/Source/Editor/Persona/Private/SAnimTrackCurvePanel.h
deleted file mode 100644
diff --git a/Engine/Source/Editor/Persona/Private/SSequenceEditor.cpp b/Engine/Source/Editor/Persona/Private/SSequenceEditor.cpp
--- a/Engine/Source/Editor/Persona/Private/SSequenceEditor.cpp
+++ b/Engine/Source/Editor/Persona/Private/SSequenceEditor.cpp
@@ -5,7 +5,6 @@
 #include "Animation/AnimSequence.h"
 
 #include "SAnimNotifyPanel.h"
-#include "SAnimTrackCurvePanel.h"
 #include "AnimPreviewInstance.h"
 #include "Editor.h"
 #include "AnimModel_AnimSequenceBase.h"
diff --git a/Engine/Source/Editor/Persona/Private/SSequenceEditor.h b/Engine/Source/Editor/Persona/Private/SSequenceEditor.h
--- a/Engine/Source/Editor/Persona/Private/SSequenceEditor.h
+++ b/Engine/Source/Editor/Persona/Private/SSequenceEditor.h
@@ -6,14 +6,12 @@
 #include "CoreMinimal.h"
 #include "Widgets/DeclarativeSyntaxSupport.h"
 #include "IPersonaPreviewScene.h"
-#include "SAnimCurvePanel.h"
 #include "SAnimationScrubPanel.h"
 #include "SAnimEditorBase.h"
 #include "EditorUndoClient.h"
 #include "Animation/AnimSequence.h"
 
 class SAnimNotifyPanel;
-class SAnimTrackCurvePanel;
 class FAnimModel_AnimSequenceBase;
 
 //////////////////////////////////////////////////////////////////////////
