Fix for no unity compile error

diff --git a/Engine/Source/Editor/Persona/Private/SAnimNotifyPanel.cpp b/Engine/Source/Editor/Persona/Private/SAnimNotifyPanel.cpp
--- a/Engine/Source/Editor/Persona/Private/SAnimNotifyPanel.cpp
+++ b/Engine/Source/Editor/Persona/Private/SAnimNotifyPanel.cpp
@@ -4232,6 +4232,11 @@ FReply SAnimNotifyPanel::OnNotifyNodeDragStarted(TArray<TSharedPtr<SAnimNotifyNo
 	return FReply::Handled().BeginDragDrop(FNotifyDragDropOp::New(Nodes, NodeDragDecorator, NotifyAnimTracks, Sequence, ScreenCursorPos, OverlayOrigin, OverlayExtents, CurrentDragXPosition, PanRequestDelegate, OnSnapPosition, UpdateDelegate));
 }
 
+float SAnimNotifyPanel::GetSequenceLength() const
+{
+	return Sequence->SequenceLength;
+}
+
 void SAnimNotifyPanel::PostUndo( bool bSuccess )
 {
 	if(Sequence != NULL)
diff --git a/Engine/Source/Editor/Persona/Private/SAnimNotifyPanel.h b/Engine/Source/Editor/Persona/Private/SAnimNotifyPanel.h
--- a/Engine/Source/Editor/Persona/Private/SAnimNotifyPanel.h
+++ b/Engine/Source/Editor/Persona/Private/SAnimNotifyPanel.h
@@ -214,7 +214,7 @@ public:
 	/**Handler for when a notify node drag has been initiated */
 	FReply OnNotifyNodeDragStarted(TArray<TSharedPtr<SAnimNotifyNode>> NotifyNodes, TSharedRef<SWidget> Decorator, const FVector2D& ScreenCursorPos, const FVector2D& ScreenNodePosition, const bool bDragOnMarker);
 
-	virtual float GetSequenceLength() const override {return Sequence->SequenceLength;}
+	virtual float GetSequenceLength() const override;
 
 	void CopySelectedNodesToClipboard() const;
 	void OnPasteNodes(SAnimNotifyTrack* RequestTrack, float ClickTime, ENotifyPasteMode::Type PasteMode, ENotifyPasteMultipleMode::Type MultiplePasteType);
