Button text no longer truncated in Asset Details

diff --git a/Engine/Source/Editor/Persona/Private/PersonaMeshDetails.cpp b/Engine/Source/Editor/Persona/Private/PersonaMeshDetails.cpp
--- a/Engine/Source/Editor/Persona/Private/PersonaMeshDetails.cpp
+++ b/Engine/Source/Editor/Persona/Private/PersonaMeshDetails.cpp
@@ -71,7 +71,7 @@
 #include "IPersonaPreviewScene.h"
 #include "IDocumentation.h"
 #include "JsonObjectConverter.h"
-#include "Widgets/Layout/SWrapBox.h"
+#include "Widgets/Layout/SUniformWrapPanel.h"
 #include "Widgets/Input/SNumericDropDown.h"
 #include "ComponentReregisterContext.h"
 #include "ClothingSystemEditorInterfaceModule.h"
@@ -334,15 +334,15 @@ void SSkeletalLODActions::Construct(const FArguments& InArgs)
 	OnReimportNewFileClicked = InArgs._OnReimportNewFileClicked;
 	BuildAvailable = InArgs._BuildAvailable;
 
-	TSharedPtr<SWrapBox> WrapBox;
+	TSharedPtr<SUniformWrapPanel> WrapBox;
 	this->ChildSlot
 
 		[
 			SNew(SVerticalBox)
 			+ SVerticalBox::Slot()
 			[
-				SAssignNew(WrapBox, SWrapBox)
-				.UseAllottedSize(true)
+				SAssignNew(WrapBox, SUniformWrapPanel)
+				//.UseAllottedSize(true)
 			]
 		];
 
@@ -350,21 +350,16 @@ void SSkeletalLODActions::Construct(const FArguments& InArgs)
 	{
 		FText ButtonNameText = BuildAvailable ? LOCTEXT("ApplyLODChange", "Apply Changes") : LOCTEXT("RegenerateLOD", "Regenerate LOD");
 		WrapBox->AddSlot()
-		.Padding(FMargin(0, 0, 2, 4))
 		[
-			SNew(SBox)
-			.WidthOverride(120.f)
+			SNew(SButton)
+			.HAlign(HAlign_Center)
+			.VAlign(VAlign_Center)
+			.OnClicked(OnApplyLODChangeClicked)
+			.IsEnabled(this, &SSkeletalLODActions::IsNeedApplyLODChange)
 			[
-				SNew(SButton)
-				.HAlign(HAlign_Center)
-				.VAlign(VAlign_Center)
-				.OnClicked(OnApplyLODChangeClicked)
-				.IsEnabled(this, &SSkeletalLODActions::IsNeedApplyLODChange)
-				[
-					SNew(STextBlock)
-					.Font(IDetailLayoutBuilder::GetDetailFont())
-					.Text(ButtonNameText)
-				]
+				SNew(STextBlock)
+				.Font(IDetailLayoutBuilder::GetDetailFont())
+				.Text(ButtonNameText)
 			]
 		];
 	}
@@ -372,20 +367,15 @@ void SSkeletalLODActions::Construct(const FArguments& InArgs)
 	if (OnRemoveLODClicked.IsBound() && (ButtonFlags & EButtonFlags::BF_Remove))
 	{
 		WrapBox->AddSlot()
-		.Padding(FMargin(0, 0, 2, 4))
 		[
-			SNew(SBox)
-			.WidthOverride(120.f)
+			SNew(SButton)
+			.HAlign(HAlign_Center)
+			.VAlign(VAlign_Center)
+			.OnClicked(OnRemoveLODClicked)
 			[
-				SNew(SButton)
-				.HAlign(HAlign_Center)
-				.VAlign(VAlign_Center)
-				.OnClicked(OnRemoveLODClicked)
-				[
-					SNew(STextBlock)
-					.Font(IDetailLayoutBuilder::GetDetailFont())
-					.Text(LOCTEXT("RemoveLOD", "Remove this LOD"))
-				]
+				SNew(STextBlock)
+				.Font(IDetailLayoutBuilder::GetDetailFont())
+				.Text(LOCTEXT("RemoveLOD", "Remove this LOD"))
 			]
 		];
 	}
@@ -393,22 +383,17 @@ void SSkeletalLODActions::Construct(const FArguments& InArgs)
 	if (OnReimportClicked.IsBound() && (ButtonFlags & EButtonFlags::BF_Reimport))
 	{
 		WrapBox->AddSlot()
-		.Padding(FMargin(0, 0, 2, 4))
 		[
-			SNew(SBox)
-			.WidthOverride(120.f)
+			SNew(SButton)
+			.HAlign(HAlign_Center)
+			.VAlign(VAlign_Center)
+			.ToolTipText(this, &SSkeletalLODActions::GetReimportButtonToolTipText)
+			.IsEnabled(this, &SSkeletalLODActions::CanReimportFromSource)
+			.OnClicked(OnReimportClicked)
 			[
-				SNew(SButton)
-				.HAlign(HAlign_Center)
-				.VAlign(VAlign_Center)
-				.ToolTipText(this, &SSkeletalLODActions::GetReimportButtonToolTipText)
-				.IsEnabled(this, &SSkeletalLODActions::CanReimportFromSource)
-				.OnClicked(OnReimportClicked)
-				[
-					SNew(STextBlock)
-					.Font(IDetailLayoutBuilder::GetDetailFont())
-					.Text(LOCTEXT("ReimportLOD", "Reimport"))
-				]
+				SNew(STextBlock)
+				.Font(IDetailLayoutBuilder::GetDetailFont())
+				.Text(LOCTEXT("ReimportLOD", "Reimport"))
 			]
 		];
 	}
@@ -416,21 +401,16 @@ void SSkeletalLODActions::Construct(const FArguments& InArgs)
 	if (OnReimportNewFileClicked.IsBound() && (ButtonFlags & EButtonFlags::BF_ReimportNewFile))
 	{
 		WrapBox->AddSlot()
-		.Padding(FMargin(0, 0, 2, 4))
 		[
-			SNew(SBox)
-			.WidthOverride(120.f)
+			SNew(SButton)
+			.HAlign(HAlign_Center)
+			.VAlign(VAlign_Center)
+			.ToolTipText(this, &SSkeletalLODActions::GetReimportButtonNewFileToolTipText)
+			.OnClicked(OnReimportNewFileClicked)
 			[
-				SNew(SButton)
-				.HAlign(HAlign_Center)
-				.VAlign(VAlign_Center)
-				.ToolTipText(this, &SSkeletalLODActions::GetReimportButtonNewFileToolTipText)
-				.OnClicked(OnReimportNewFileClicked)
-				[
-					SNew(STextBlock)
-					.Font(IDetailLayoutBuilder::GetDetailFont())
-					.Text(LOCTEXT("ReimportLOD_NewFile", "Reimport (New File)"))
-				]
+				SNew(STextBlock)
+				.Font(IDetailLayoutBuilder::GetDetailFont())
+				.Text(LOCTEXT("ReimportLOD_NewFile", "Reimport (New File)"))
 			]
 		];
 	}
