Merging using //UE4/Main_to_//UE5/Release-50-M2
Fix crash with transactional temporary UI stub.

diff --git a/Engine/Source/Editor/Persona/Private/PersonaMeshDetails.cpp b/Engine/Source/Editor/Persona/Private/PersonaMeshDetails.cpp
--- a/Engine/Source/Editor/Persona/Private/PersonaMeshDetails.cpp
+++ b/Engine/Source/Editor/Persona/Private/PersonaMeshDetails.cpp
@@ -2626,7 +2626,7 @@ void FPersonaMeshDetails::AddLODLevelCategories(IDetailLayoutBuilder& DetailLayo
 		for (int32 LODIndex = 0; LODIndex < SkelMeshLODCount; ++LODIndex)
 		{
 			//Construct temporary LODInfo editor object
-			ULODInfoUILayout* LODInfoUILayout = NewObject<ULODInfoUILayout>(GetTransientPackage(), FName(*(FGuid::NewGuid().ToString())), RF_Standalone | RF_Transactional);
+			ULODInfoUILayout* LODInfoUILayout = NewObject<ULODInfoUILayout>(GetTransientPackage(), FName(*(FGuid::NewGuid().ToString())), RF_Standalone);
 			LODInfoUILayout->AddToRoot();
 			FSkeletalMeshLODInfo* LODInfoPtr = SkelMesh->GetLODInfo(LODIndex);
 			if (!LODInfoPtr || !SkelMesh->GetImportedModel()->LODModels.IsValidIndex(LODIndex))
