ControlRig: Added support for  'Additional Meshes' in Control Rig preview scene

diff --git a/Engine/Plugins/Experimental/ControlRig/Source/ControlRigEditor/Private/Editor/ControlRigEditor.cpp b/Engine/Plugins/Experimental/ControlRig/Source/ControlRigEditor/Private/Editor/ControlRigEditor.cpp
--- a/Engine/Plugins/Experimental/ControlRig/Source/ControlRigEditor/Private/Editor/ControlRigEditor.cpp
+++ b/Engine/Plugins/Experimental/ControlRig/Source/ControlRigEditor/Private/Editor/ControlRigEditor.cpp
@@ -171,7 +171,7 @@ void FControlRigEditor::InitControlRigEditor(const EToolkitMode::Type Mode, cons
 
 	FPersonaToolkitArgs PersonaToolkitArgs;
 	PersonaToolkitArgs.OnPreviewSceneCreated = FOnPreviewSceneCreated::FDelegate::CreateSP(this, &FControlRigEditor::HandlePreviewSceneCreated);
-	PersonaToolkit = PersonaModule.CreatePersonaToolkit(InControlRigBlueprint, PersonaToolkitArgs);
+	PersonaToolkit = PersonaModule.CreatePersonaToolkit(InControlRigBlueprint, PersonaToolkitArgs, InControlRigBlueprint->GetPreviewMesh()->Skeleton);
 
 	// set delegate prior to setting mesh
 	// otherwise, you don't get delegate
diff --git a/Engine/Source/Editor/Persona/Private/PersonaModule.cpp b/Engine/Source/Editor/Persona/Private/PersonaModule.cpp
--- a/Engine/Source/Editor/Persona/Private/PersonaModule.cpp
+++ b/Engine/Source/Editor/Persona/Private/PersonaModule.cpp
@@ -205,11 +205,11 @@ static void SetupPersonaToolkit(const TSharedRef<FPersonaToolkit>& Toolkit, cons
 	}
 }
 
-TSharedRef<IPersonaToolkit> FPersonaModule::CreatePersonaToolkit(UObject* InAsset, const FPersonaToolkitArgs& PersonaToolkitArgs) const
+TSharedRef<IPersonaToolkit> FPersonaModule::CreatePersonaToolkit(UObject* InAsset, const FPersonaToolkitArgs& PersonaToolkitArgs, USkeleton* InSkeleton) const
 {
 	TSharedRef<FPersonaToolkit> NewPersonaToolkit(new FPersonaToolkit());
 
-	NewPersonaToolkit->Initialize(InAsset);
+	NewPersonaToolkit->Initialize(InAsset, InSkeleton);
 
 	SetupPersonaToolkit(NewPersonaToolkit, PersonaToolkitArgs);
 
diff --git a/Engine/Source/Editor/Persona/Private/PersonaToolkit.cpp b/Engine/Source/Editor/Persona/Private/PersonaToolkit.cpp
--- a/Engine/Source/Editor/Persona/Private/PersonaToolkit.cpp
+++ b/Engine/Source/Editor/Persona/Private/PersonaToolkit.cpp
@@ -36,7 +36,7 @@ static void FindCounterpartAssets(const UObject* InAsset, TWeakObjectPtr<USkelet
 	OutMesh = const_cast<USkeletalMesh*>(CounterpartMesh);
 }
 
-void FPersonaToolkit::Initialize(UObject* InAsset)
+void FPersonaToolkit::Initialize(UObject* InAsset, USkeleton* InSkeleton)
 {
 	Asset = InAsset;
 	InitialAssetClass = Asset->GetClass();
@@ -45,6 +45,8 @@ void FPersonaToolkit::Initialize(UObject* InAsset)
 	{
 		Mesh = PreviewMeshInterface->GetPreviewMesh();
 	}
+
+	Skeleton = InSkeleton;
 }
 
 void FPersonaToolkit::Initialize(USkeleton* InSkeleton)
diff --git a/Engine/Source/Editor/Persona/Private/PersonaToolkit.h b/Engine/Source/Editor/Persona/Private/PersonaToolkit.h
--- a/Engine/Source/Editor/Persona/Private/PersonaToolkit.h
+++ b/Engine/Source/Editor/Persona/Private/PersonaToolkit.h
@@ -21,7 +21,7 @@ public:
 	virtual ~FPersonaToolkit() {}
 
 	/** Initialize from a various sources */
-	void Initialize(UObject* InAsset);
+	void Initialize(UObject* InAsset, USkeleton* InSkeleton = nullptr);
 	void Initialize(USkeleton* InSkeleton);
 	void Initialize(UAnimationAsset* InAnimationAsset);
 	void Initialize(USkeletalMesh* InSkeletalMesh);
diff --git a/Engine/Source/Editor/Persona/Public/PersonaModule.h b/Engine/Source/Editor/Persona/Public/PersonaModule.h
--- a/Engine/Source/Editor/Persona/Public/PersonaModule.h
+++ b/Engine/Source/Editor/Persona/Public/PersonaModule.h
@@ -195,7 +195,7 @@ public:
 	virtual void ShutdownModule();
 
 	/** Create a re-usable toolkit that multiple asset editors that are concerned with USkeleton-related data can use */
-	virtual TSharedRef<IPersonaToolkit> CreatePersonaToolkit(UObject* InAsset, const FPersonaToolkitArgs& PersonaToolkitArgs = FPersonaToolkitArgs()) const;
+	virtual TSharedRef<IPersonaToolkit> CreatePersonaToolkit(UObject* InAsset, const FPersonaToolkitArgs& PersonaToolkitArgs = FPersonaToolkitArgs(), USkeleton* InSkeleton = nullptr) const;
 	virtual TSharedRef<IPersonaToolkit> CreatePersonaToolkit(USkeleton* InSkeleton, const FPersonaToolkitArgs& PersonaToolkitArgs = FPersonaToolkitArgs()) const;
 	virtual TSharedRef<IPersonaToolkit> CreatePersonaToolkit(UAnimationAsset* InAnimationAsset, const FPersonaToolkitArgs& PersonaToolkitArgs = FPersonaToolkitArgs()) const;
 	virtual TSharedRef<IPersonaToolkit> CreatePersonaToolkit(USkeletalMesh* InSkeletalMesh, const FPersonaToolkitArgs& PersonaToolkitArgs = FPersonaToolkitArgs()) const;
