Autoscrub to the animnotify time when holding down shift while dragging a notify in the Animation Editor

diff --git a/Engine/Source/Editor/Persona/Persona.Build.cs b/Engine/Source/Editor/Persona/Persona.Build.cs
--- a/Engine/Source/Editor/Persona/Persona.Build.cs
+++ b/Engine/Source/Editor/Persona/Persona.Build.cs
@@ -73,6 +73,7 @@ public class Persona : ModuleRules
 				"SequencerWidgets",
 				"TimeManagement",
                 "Sequencer",
+				"StatusBar",
 			}
         );
 
diff --git a/Engine/Source/Editor/Persona/Private/AnimTimeline/AnimModel.cpp b/Engine/Source/Editor/Persona/Private/AnimTimeline/AnimModel.cpp
--- a/Engine/Source/Editor/Persona/Private/AnimTimeline/AnimModel.cpp
+++ b/Engine/Source/Editor/Persona/Private/AnimTimeline/AnimModel.cpp
@@ -110,6 +110,11 @@ void FAnimModel::SetScrubPosition(FFrameTime NewScrubPostion) const
 		UDebugSkelMeshComponent* PreviewMeshComponent = WeakPreviewScene.Pin()->GetPreviewMeshComponent();
 		if(PreviewMeshComponent && PreviewMeshComponent->IsPreviewOn())
 		{
+			if(PreviewMeshComponent->PreviewInstance->IsPlaying())
+			{
+				PreviewMeshComponent->PreviewInstance->SetPlaying(false);
+			}
+			
 			PreviewMeshComponent->PreviewInstance->SetPosition(NewScrubPostion.AsDecimal() / (double)GetTickResolution());
 		}
 	}
diff --git a/Engine/Source/Editor/Persona/Private/AnimTimeline/AnimTimelineTrack_NotifiesPanel.cpp b/Engine/Source/Editor/Persona/Private/AnimTimeline/AnimTimelineTrack_NotifiesPanel.cpp
--- a/Engine/Source/Editor/Persona/Private/AnimTimeline/AnimTimelineTrack_NotifiesPanel.cpp
+++ b/Engine/Source/Editor/Persona/Private/AnimTimeline/AnimTimelineTrack_NotifiesPanel.cpp
@@ -18,6 +18,7 @@
 #define LOCTEXT_NAMESPACE "FAnimTimelineTrack_NotifiesPanel"
 
 const float FAnimTimelineTrack_NotifiesPanel::NotificationTrackHeight = 24.0f;
+const FName FAnimTimelineTrack_NotifiesPanel::AnimationEditorStatusBarName = FName(TEXT("AssetEditor.AnimationEditor.MainMenu"));
 
 ANIMTIMELINE_IMPLEMENT_TRACK(FAnimTimelineTrack_NotifiesPanel);
 
@@ -295,7 +296,50 @@ TSharedRef<SAnimNotifyPanel> FAnimTimelineTrack_NotifiesPanel::GetAnimNotifyPane
 			.OnNotifiesChanged_Lambda([this]()
 			{ 
 				Update();
-				GetModel()->OnTracksChanged().Broadcast(); 
+				GetModel()->OnTracksChanged().Broadcast();
+				 
+				if (StatusBarMessageHandle.IsValid())
+				{
+					if(UStatusBarSubsystem* StatusBarSubsystem = GEditor->GetEditorSubsystem<UStatusBarSubsystem>())
+					{
+						StatusBarSubsystem->PopStatusBarMessage(AnimationEditorStatusBarName, StatusBarMessageHandle);
+						StatusBarMessageHandle.Reset();
+					}
+				}
+			})
+			.OnNotifyStateHandleBeingDragged_Lambda([this](TSharedPtr<SAnimNotifyNode> NotifyNode, const FPointerEvent& Event, ENotifyStateHandleHit::Type Handle, float Time)
+			{
+				if (Event.IsShiftDown())
+				{
+					const FFrameTime FrameTime = FFrameTime::FromDecimal(Time * (double)GetModel()->GetTickResolution());
+					GetModel()->SetScrubPosition(FrameTime);
+				}
+
+				if (!StatusBarMessageHandle.IsValid())
+				{
+					if (UStatusBarSubsystem* StatusBarSubsystem = GEditor->GetEditorSubsystem<UStatusBarSubsystem>())
+					{
+						StatusBarMessageHandle = StatusBarSubsystem->PushStatusBarMessage(AnimationEditorStatusBarName,
+							LOCTEXT("AutoscrubNotifyStateHandle", "Hold SHIFT while dragging a notify state Begin or End handle to auto scrub the timeline."));
+					}
+				}
+			})
+			.OnNotifyNodesBeingDragged_Lambda([this](const TArray<TSharedPtr<SAnimNotifyNode>>& NotifyNodes, const class FDragDropEvent& Event, float DragXPosition, float DragTime)
+			{
+				if (Event.IsShiftDown())
+				{
+					const FFrameTime FrameTime = FFrameTime::FromDecimal(DragTime * (double)GetModel()->GetTickResolution());
+					GetModel()->SetScrubPosition(FrameTime);
+				}
+
+				if (!StatusBarMessageHandle.IsValid())
+				{
+					if (UStatusBarSubsystem* StatusBarSubsystem = GEditor->GetEditorSubsystem<UStatusBarSubsystem>())
+					{
+						StatusBarMessageHandle = StatusBarSubsystem->PushStatusBarMessage(AnimationEditorStatusBarName,
+							LOCTEXT("AutoscrubNotify", "Hold SHIFT while dragging a notify to auto scrub the timeline."));
+					}
+				}
 			});
 
 		GetModel()->GetAnimSequenceBase()->RegisterOnNotifyChanged(UAnimSequenceBase::FOnNotifyChanged::CreateSP(this, &FAnimTimelineTrack_NotifiesPanel::HandleNotifyChanged));
diff --git a/Engine/Source/Editor/Persona/Private/AnimTimeline/AnimTimelineTrack_NotifiesPanel.h b/Engine/Source/Editor/Persona/Private/AnimTimeline/AnimTimelineTrack_NotifiesPanel.h
--- a/Engine/Source/Editor/Persona/Private/AnimTimeline/AnimTimelineTrack_NotifiesPanel.h
+++ b/Engine/Source/Editor/Persona/Private/AnimTimeline/AnimTimelineTrack_NotifiesPanel.h
@@ -4,6 +4,7 @@
 
 #include "AnimTimelineTrack.h"
 #include "SAnimTimingPanel.h"
+#include "StatusBarSubsystem.h"
 
 class SAnimNotifyPanel;
 class SVerticalBox;
@@ -16,6 +17,7 @@ class FAnimTimelineTrack_NotifiesPanel : public FAnimTimelineTrack
 
 public:
 	static const float NotificationTrackHeight;
+	static const FName AnimationEditorStatusBarName;
 
 	FAnimTimelineTrack_NotifiesPanel(const TSharedRef<FAnimModel>& InModel);
 
@@ -49,4 +51,7 @@ private:
 
 	/** Track index we want to trigger a rename for */
 	int32 PendingRenameTrackIndex;
+
+	/** Handle to status bar message */
+	FStatusBarMessageHandle StatusBarMessageHandle;
 };
\ No newline at end of file
diff --git a/Engine/Source/Editor/Persona/Private/SAnimNotifyPanel.cpp b/Engine/Source/Editor/Persona/Private/SAnimNotifyPanel.cpp
--- a/Engine/Source/Editor/Persona/Private/SAnimNotifyPanel.cpp
+++ b/Engine/Source/Editor/Persona/Private/SAnimNotifyPanel.cpp
@@ -461,6 +461,7 @@ public:
 		, _AnimNotify(nullptr)
 		, _AnimSyncMarker(nullptr)
 		, _OnNodeDragStarted()
+		, _OnNotifyStateHandleBeingDragged()
 		, _OnUpdatePanel()
 		, _PanTrackRequest()
 		, _OnSelectionChanged()
@@ -472,6 +473,7 @@ public:
 	SLATE_ARGUMENT( FAnimNotifyEvent *, AnimNotify )
 	SLATE_ARGUMENT( FAnimSyncMarker*, AnimSyncMarker)
 	SLATE_EVENT( FOnNotifyNodeDragStarted, OnNodeDragStarted )
+	SLATE_EVENT( FOnNotifyStateHandleBeingDragged, OnNotifyStateHandleBeingDragged)
 	SLATE_EVENT( FOnUpdatePanel, OnUpdatePanel )
 	SLATE_EVENT( FPanTrackRequest, PanTrackRequest )
 	SLATE_EVENT( FOnTrackSelectionChanged, OnSelectionChanged )
@@ -610,6 +612,9 @@ private:
 	/** Delegate that is called when the user initiates dragging */
 	FOnNotifyNodeDragStarted	OnNodeDragStarted;
 
+	/** Delegate that is called when a notify state handle is being dragged */
+	FOnNotifyStateHandleBeingDragged	OnNotifyStateHandleBeingDragged;
+
 	/** Delegate to pan the track, needed if the markers are dragged out of the track */
 	FPanTrackRequest			PanTrackRequest;
 
@@ -710,6 +715,7 @@ public:
 		, _OnGetScrubValue()
 		, _OnGetDraggedNodePos()
 		, _OnNodeDragStarted()
+		, _OnNotifyStateHandleBeingDragged()
 		, _OnRequestTrackPan()
 		, _OnRequestOffsetRefresh()
 		, _OnDeleteNotify()
@@ -741,6 +747,7 @@ public:
 		SLATE_EVENT( FOnGetScrubValue, OnGetScrubValue )
 		SLATE_EVENT( FOnGetDraggedNodePos, OnGetDraggedNodePos )
 		SLATE_EVENT( FOnNotifyNodesDragStarted, OnNodeDragStarted )
+		SLATE_EVENT( FOnNotifyStateHandleBeingDragged, OnNotifyStateHandleBeingDragged)
 		SLATE_EVENT( FPanTrackRequest, OnRequestTrackPan )
 		SLATE_EVENT( FRefreshOffsetsRequest, OnRequestOffsetRefresh )
 		SLATE_EVENT( FDeleteNotify, OnDeleteNotify )
@@ -1002,6 +1009,7 @@ protected:
 	FOnGetScrubValue						OnGetScrubValue;
 	FOnGetDraggedNodePos					OnGetDraggedNodePos;
 	FOnNotifyNodesDragStarted				OnNodeDragStarted;
+	FOnNotifyStateHandleBeingDragged		OnNotifyStateHandleBeingDragged;
 	FPanTrackRequest						OnRequestTrackPan;
 	FDeselectAllNotifies					OnDeselectAllNotifies;
 	FCopyNodes							OnCopyNodes;
@@ -1086,6 +1094,7 @@ public:
 	SLATE_EVENT( FOnGetNativeNotifyClasses, OnGetNotifyNativeClasses )
 	SLATE_EVENT( FOnGetNativeNotifyClasses, OnGetNotifyStateNativeClasses )
 	SLATE_EVENT( FOnNotifyNodesDragStarted, OnNodeDragStarted )
+	SLATE_EVENT( FOnNotifyStateHandleBeingDragged, OnNotifyStateHandleBeingDragged)
 	SLATE_EVENT( FRefreshOffsetsRequest, OnRequestRefreshOffsets )
 	SLATE_EVENT( FDeleteNotify, OnDeleteNotify )
 	SLATE_EVENT( FDeselectAllNotifies, OnDeselectAllNotifies)
@@ -1317,6 +1326,8 @@ public:
 			float ScreenDelta =  FMath::Max(LocalMouseXPos - LocalViewportMax, 10.0f);
 			RequestTrackPan.Execute(ScreenDelta, FVector2D(LocalTrackWidth, 1.f));
 		}
+
+		OnNodesBeingDragged.ExecuteIfBound(SelectedNodes, DragDropEvent, CurrentDragXPosition, TrackScaleInfo.LocalXToInput(CurrentDragXPosition));
 	}
 
 	float GetSnapPosition(const FTrackClampInfo& ClampInfo, float WidgetSpaceNotifyPosition, bool& bOutSnapped)
@@ -1387,6 +1398,7 @@ public:
 	int32								TrackSpan;				// Number of tracks that the selection spans
 	FOnUpdatePanel						OnUpdatePanel;			// Delegate to redraw the notify panel
 	FOnSnapPosition						OnSnapPosition;			// Delegate used to snap times
+	FOnNotifyNodesBeingDragged			OnNodesBeingDragged;	// Delegate to notify panel when the mouse was moved during the DragDropOp
 
 	static TSharedRef<FNotifyDragDropOp> New(
 		TArray<TSharedPtr<SAnimNotifyNode>>			NotifyNodes, 
@@ -1399,13 +1411,15 @@ public:
 		float&										CurrentDragXPosition, 
 		FPanTrackRequest&							RequestTrackPanDelegate, 
 		FOnSnapPosition&							OnSnapPosition,
-		FOnUpdatePanel&								UpdatePanel
+		FOnUpdatePanel&								UpdatePanel,
+		FOnNotifyNodesBeingDragged&					OnNodesBeingDragged
 		)
 	{
 		TSharedRef<FNotifyDragDropOp> Operation = MakeShareable(new FNotifyDragDropOp(CurrentDragXPosition));
 		Operation->Sequence = InSequence;
 		Operation->RequestTrackPan = RequestTrackPanDelegate;
 		Operation->OnUpdatePanel = UpdatePanel;
+		Operation->OnNodesBeingDragged = OnNodesBeingDragged;
 
 		Operation->NodeGroupPosition = SelectionScreenPosition;
 		Operation->NodeGroupSize = SelectionSize;
@@ -1511,6 +1525,7 @@ void SAnimNotifyNode::Construct(const FArguments& InArgs)
 	NodeObjectInterface->CacheName();
 
 	OnNodeDragStarted = InArgs._OnNodeDragStarted;
+	OnNotifyStateHandleBeingDragged = InArgs._OnNotifyStateHandleBeingDragged;
 	PanTrackRequest = InArgs._PanTrackRequest;
 	OnSelectionChanged = InArgs._OnSelectionChanged;
 	OnUpdatePanel = InArgs._OnUpdatePanel;
@@ -1932,6 +1947,8 @@ FReply SAnimNotifyNode::OnMouseMove( const FGeometry& MyGeometry, const FPointer
 				AnimNotifyEvent->TriggerTimeOffset = GetTriggerTimeOffsetForType(EAnimEventTriggerOffsets::NoOffset);
 			}
 		}
+
+		OnNotifyStateHandleBeingDragged.ExecuteIfBound(SharedThis(this), MouseEvent, CurrentDragHandle, NodeObjectInterface->GetTime());
 	}
 	else
 	{
@@ -1977,6 +1994,8 @@ FReply SAnimNotifyNode::OnMouseMove( const FGeometry& MyGeometry, const FPointer
 				AnimNotifyEvent->EndTriggerTimeOffset = GetTriggerTimeOffsetForType(EAnimEventTriggerOffsets::NoOffset);
 			}
 		}
+
+		OnNotifyStateHandleBeingDragged.ExecuteIfBound(SharedThis(this), MouseEvent, CurrentDragHandle, (NodeObjectInterface->GetTime() + NodeObjectInterface->GetDuration()));
 	}
 
 	return FReply::Handled();
@@ -2166,6 +2185,7 @@ void SAnimNotifyTrack::Construct(const FArguments& InArgs)
 	OnGetScrubValue = InArgs._OnGetScrubValue;
 	OnGetDraggedNodePos = InArgs._OnGetDraggedNodePos;
 	OnNodeDragStarted = InArgs._OnNodeDragStarted;
+	OnNotifyStateHandleBeingDragged = InArgs._OnNotifyStateHandleBeingDragged;
 	TrackColor = InArgs._TrackColor;
 	OnSnapPosition = InArgs._OnSnapPosition;
 	OnRequestTrackPan = InArgs._OnRequestTrackPan;
@@ -3342,6 +3362,7 @@ void SAnimNotifyTrack::Update()
 				.Sequence(Sequence)
 				.AnimNotify(Event)
 				.OnNodeDragStarted(this, &SAnimNotifyTrack::OnNotifyNodeDragStarted, NotifyIndex)
+				.OnNotifyStateHandleBeingDragged(OnNotifyStateHandleBeingDragged)
 				.OnUpdatePanel(OnUpdatePanel)
 				.PanTrackRequest(OnRequestTrackPan)
 				.ViewInputMin(ViewInputMin)
@@ -3808,6 +3829,7 @@ void SNotifyEdTrack::Construct(const FArguments& InArgs)
 				.OnGetScrubValue(InArgs._OnGetScrubValue)
 				.OnGetDraggedNodePos(InArgs._OnGetDraggedNodePos)
 				.OnNodeDragStarted(InArgs._OnNodeDragStarted)
+				.OnNotifyStateHandleBeingDragged(InArgs._OnNotifyStateHandleBeingDragged)
 				.OnSnapPosition(InArgs._OnSnapPosition)
 				.TrackColor(Track.TrackColor)
 				.OnRequestTrackPan(FPanTrackRequest::CreateSP(PanelRef, &SAnimNotifyPanel::PanInputViewRange))
@@ -3861,6 +3883,8 @@ void SAnimNotifyPanel::Construct(const FArguments& InArgs, const TSharedRef<FAni
 	OnInvokeTab = InArgs._OnInvokeTab;
 	OnNotifiesChanged = InArgs._OnNotifiesChanged;
 	OnSnapPosition = InArgs._OnSnapPosition;
+	OnNotifyStateHandleBeingDragged = InArgs._OnNotifyStateHandleBeingDragged;
+	OnNotifyNodesBeingDragged = InArgs._OnNotifyNodesBeingDragged;
 	bIsSelecting = false;
 	bIsUpdating = false;
 
@@ -4142,6 +4166,7 @@ void SAnimNotifyPanel::RefreshNotifyTracks()
 				.OnGetNotifyStateNativeClasses(this, &SAnimNotifyPanel::OnGetNativeNotifyData, UAnimNotifyState::StaticClass(), &NotifyStateClassNames)
 				.OnSelectionChanged(this, &SAnimNotifyPanel::OnTrackSelectionChanged)
 				.OnNodeDragStarted(this, &SAnimNotifyPanel::OnNotifyNodeDragStarted)
+				.OnNotifyStateHandleBeingDragged(OnNotifyStateHandleBeingDragged)
 				.OnSnapPosition(OnSnapPosition)
 				.OnRequestRefreshOffsets(OnRequestRefreshOffsets)
 				.OnDeleteNotify(this, &SAnimNotifyPanel::DeleteSelectedNodeObjects)
@@ -4209,7 +4234,7 @@ FReply SAnimNotifyPanel::OnNotifyNodeDragStarted(TArray<TSharedPtr<SAnimNotifyNo
 
 	FPanTrackRequest PanRequestDelegate = FPanTrackRequest::CreateSP(this, &SAnimNotifyPanel::PanInputViewRange);
 	FOnUpdatePanel UpdateDelegate = FOnUpdatePanel::CreateSP(this, &SAnimNotifyPanel::Update);
-	return FReply::Handled().BeginDragDrop(FNotifyDragDropOp::New(Nodes, NodeDragDecorator, NotifyAnimTracks, Sequence, ScreenCursorPos, OverlayOrigin, OverlayExtents, CurrentDragXPosition, PanRequestDelegate, OnSnapPosition, UpdateDelegate));
+	return FReply::Handled().BeginDragDrop(FNotifyDragDropOp::New(Nodes, NodeDragDecorator, NotifyAnimTracks, Sequence, ScreenCursorPos, OverlayOrigin, OverlayExtents, CurrentDragXPosition, PanRequestDelegate, OnSnapPosition, UpdateDelegate, OnNotifyNodesBeingDragged));
 }
 
 float SAnimNotifyPanel::GetSequenceLength() const
diff --git a/Engine/Source/Editor/Persona/Private/SAnimNotifyPanel.h b/Engine/Source/Editor/Persona/Private/SAnimNotifyPanel.h
--- a/Engine/Source/Editor/Persona/Private/SAnimNotifyPanel.h
+++ b/Engine/Source/Editor/Persona/Private/SAnimNotifyPanel.h
@@ -138,6 +138,9 @@ public:
 	TArray<TSharedPtr<SAnimNotifyNode>> OriginalSelection;
 };
 
+DECLARE_DELEGATE_FourParams(FOnNotifyStateHandleBeingDragged, TSharedPtr<SAnimNotifyNode> /*NotifyNode*/, const FPointerEvent& /*PointerEvent*/, ENotifyStateHandleHit::Type /*Handle*/, float /*Time*/)
+DECLARE_DELEGATE_FourParams(FOnNotifyNodesBeingDragged, const TArray<TSharedPtr<SAnimNotifyNode>>& /*NotifyNodes*/, const class FDragDropEvent& /*DragDropEvent*/, float /*XPosition*/, float /*Time*/)
+
 //////////////////////////////////////////////////////////////////////////
 // SAnimNotifyPanel
 
@@ -192,6 +195,8 @@ public:
 	SLATE_EVENT( FOnInvokeTab, OnInvokeTab )
 	SLATE_EVENT( FSimpleDelegate, OnNotifiesChanged )
 	SLATE_EVENT( FOnSnapPosition, OnSnapPosition )
+	SLATE_EVENT( FOnNotifyStateHandleBeingDragged, OnNotifyStateHandleBeingDragged)
+	SLATE_EVENT( FOnNotifyNodesBeingDragged, OnNotifyNodesBeingDragged)
 
 	SLATE_END_ARGS()
 
@@ -357,6 +362,12 @@ private:
 	/** Delegate used to inform others that notifies have changed (for timing) */
 	FSimpleDelegate OnNotifiesChanged;
 
+	/** Delegate used to inform others that a notify state handle is being dragged */
+	FOnNotifyStateHandleBeingDragged OnNotifyStateHandleBeingDragged;
+	
+	/** Delegate used to inform others that one or more notifies nodes are being dragged */
+	FOnNotifyNodesBeingDragged OnNotifyNodesBeingDragged;
+
 	/** Recursion guard for selection */
 	bool bIsSelecting;
 
