Tool menu API extension for Persona editor umbrella - exposing partial PersonalToolkit API to scripting through context uobject

diff --git a/Engine/Source/Editor/AnimationBlueprintEditor/Private/AnimationBlueprintEditor.cpp b/Engine/Source/Editor/AnimationBlueprintEditor/Private/AnimationBlueprintEditor.cpp
--- a/Engine/Source/Editor/AnimationBlueprintEditor/Private/AnimationBlueprintEditor.cpp
+++ b/Engine/Source/Editor/AnimationBlueprintEditor/Private/AnimationBlueprintEditor.cpp
@@ -69,6 +69,9 @@
 #include "AnimationBlueprintInterfaceEditorMode.h"
 #include "ToolMenus.h"
 
+#include "PersonaToolMenuContext.h"
+#include "ToolMenuContext.h"
+
 // Hide related nodes feature
 #include "Preferences/AnimationBlueprintEditorOptions.h"
 #include "Framework/MultiBox/MultiBoxBuilder.h"
@@ -1118,6 +1121,16 @@ FLinearColor FAnimationBlueprintEditor::GetWorldCentricTabColorScale() const
 	return FLinearColor( 0.5f, 0.25f, 0.35f, 0.5f );
 }
 
+void FAnimationBlueprintEditor::InitToolMenuContext(FToolMenuContext& MenuContext)
+{
+	FAssetEditorToolkit::InitToolMenuContext(MenuContext);
+
+	UPersonaToolMenuContext* Context = NewObject<UPersonaToolMenuContext>();
+	Context->SetToolkit(GetPersonaToolkit());
+
+	MenuContext.AddObject(Context);
+}
+
 IAnimationSequenceBrowser* FAnimationBlueprintEditor::GetAssetBrowser() const
 {
 	return SequenceBrowser.Pin().Get();
diff --git a/Engine/Source/Editor/AnimationBlueprintEditor/Private/AnimationBlueprintEditor.h b/Engine/Source/Editor/AnimationBlueprintEditor/Private/AnimationBlueprintEditor.h
--- a/Engine/Source/Editor/AnimationBlueprintEditor/Private/AnimationBlueprintEditor.h
+++ b/Engine/Source/Editor/AnimationBlueprintEditor/Private/AnimationBlueprintEditor.h
@@ -21,6 +21,8 @@ class UEdGraph;
 class USkeletalMesh;
 class ISkeletonTreeItem;
 
+struct FToolMenuContext;
+
 struct FAnimationBlueprintEditorModes
 {
 	// Mode constants
@@ -142,6 +144,7 @@ public:
 	virtual FText GetToolkitToolTipText() const override;
 	virtual FString GetWorldCentricTabPrefix() const override;
 	virtual FLinearColor GetWorldCentricTabColorScale() const override;	
+	virtual void InitToolMenuContext(FToolMenuContext& MenuContext) override;
 	//~ End IToolkit Interface
 
 	/** @return the documentation location for this editor */
diff --git a/Engine/Source/Editor/AnimationEditor/AnimationEditor.Build.cs b/Engine/Source/Editor/AnimationEditor/AnimationEditor.Build.cs
--- a/Engine/Source/Editor/AnimationEditor/AnimationEditor.Build.cs
+++ b/Engine/Source/Editor/AnimationEditor/AnimationEditor.Build.cs
@@ -21,6 +21,7 @@ public class AnimationEditor : ModuleRules
                 "SkeletonEditor",
                 "Kismet",
                 "AnimGraph",
+				"ToolMenus"
             }
 		);
 
diff --git a/Engine/Source/Editor/AnimationEditor/Private/AnimationEditor.cpp b/Engine/Source/Editor/AnimationEditor/Private/AnimationEditor.cpp
--- a/Engine/Source/Editor/AnimationEditor/Private/AnimationEditor.cpp
+++ b/Engine/Source/Editor/AnimationEditor/Private/AnimationEditor.cpp
@@ -52,6 +52,8 @@
 #include "IAnimSequenceCurveEditor.h"
 #include "EditorModeManager.h"
 #include "IPersonaEditorModeManager.h"
+#include "Toolkits/AssetEditorToolkit.h"
+#include "PersonaToolMenuContext.h"
 
 const FName AnimationEditorAppIdentifier = FName(TEXT("AnimationEditorApp"));
 
@@ -166,6 +168,16 @@ FLinearColor FAnimationEditor::GetWorldCentricTabColorScale() const
 	return FLinearColor(0.3f, 0.2f, 0.5f, 0.5f);
 }
 
+void FAnimationEditor::InitToolMenuContext(FToolMenuContext& MenuContext)
+{
+	FAssetEditorToolkit::InitToolMenuContext(MenuContext);
+
+	UPersonaToolMenuContext* Context = NewObject<UPersonaToolMenuContext>();
+	Context->SetToolkit(GetPersonaToolkit());
+
+	MenuContext.AddObject(Context);
+}
+
 void FAnimationEditor::Tick(float DeltaTime)
 {
 	GetPersonaToolkit()->GetPreviewScene()->InvalidateViews();
diff --git a/Engine/Source/Editor/AnimationEditor/Private/AnimationEditor.h b/Engine/Source/Editor/AnimationEditor/Private/AnimationEditor.h
--- a/Engine/Source/Editor/AnimationEditor/Private/AnimationEditor.h
+++ b/Engine/Source/Editor/AnimationEditor/Private/AnimationEditor.h
@@ -26,6 +26,7 @@ class UAnimSequenceBase;
 class ISkeletonTreeItem;
 class IAnimSequenceCurveEditor;
 struct FRichCurve;
+struct FToolMenuContext;
 
 namespace AnimationEditorModes
 {
@@ -75,6 +76,7 @@ public:
 	virtual FText GetBaseToolkitName() const override;
 	virtual FString GetWorldCentricTabPrefix() const override;
 	virtual FLinearColor GetWorldCentricTabColorScale() const override;
+	virtual void InitToolMenuContext(FToolMenuContext& MenuContext) override;
 
 	/** FTickableEditorObject Interface */
 	virtual void Tick(float DeltaTime) override;
diff --git a/Engine/Source/Editor/Persona/Private/PersonaToolMenuContext.cpp b/Engine/Source/Editor/Persona/Private/PersonaToolMenuContext.cpp
new file mode 100644
--- /dev/null
+++ b/Engine/Source/Editor/Persona/Private/PersonaToolMenuContext.cpp
@@ -0,0 +1,65 @@
+// Copyright Epic Games, Inc. All Rights Reserved.
+
+#include "PersonaToolMenuContext.h"
+#include "IPersonaToolkit.h"
+
+USkeleton* UPersonaToolMenuContext::GetSkeleton() const
+{
+	if (HasValidToolkit())
+	{
+		return WeakToolkit.Pin()->GetSkeleton();
+	}
+
+	return nullptr;
+}
+
+UDebugSkelMeshComponent* UPersonaToolMenuContext::GetPreviewMeshComponent() const
+{
+	if (HasValidToolkit())
+	{
+		return WeakToolkit.Pin()->GetPreviewMeshComponent();
+	}
+
+	return nullptr;
+}
+
+USkeletalMesh* UPersonaToolMenuContext::GetMesh() const
+{
+	if (HasValidToolkit())
+	{
+		return WeakToolkit.Pin()->GetMesh();
+	}
+
+	return nullptr;
+}
+
+UAnimBlueprint* UPersonaToolMenuContext::GetAnimBlueprint() const
+{
+	if (HasValidToolkit())
+	{
+		return WeakToolkit.Pin()->GetAnimBlueprint();
+	}
+
+	return nullptr;
+}
+
+UAnimationAsset* UPersonaToolMenuContext::GetAnimationAsset() const
+{
+	if (HasValidToolkit())
+	{
+		return WeakToolkit.Pin()->GetAnimationAsset();
+	}
+
+	return nullptr;
+}
+
+void UPersonaToolMenuContext::SetToolkit(TSharedRef<IPersonaToolkit> InToolkit)
+{
+	WeakToolkit = InToolkit;
+}
+
+bool UPersonaToolMenuContext::HasValidToolkit() const
+{
+	return WeakToolkit.IsValid();
+}
+
diff --git a/Engine/Source/Editor/Persona/Public/PersonaToolMenuContext.h b/Engine/Source/Editor/Persona/Public/PersonaToolMenuContext.h
new file mode 100644
--- /dev/null
+++ b/Engine/Source/Editor/Persona/Public/PersonaToolMenuContext.h
@@ -0,0 +1,49 @@
+// Copyright Epic Games, Inc. All Rights Reserved.
+
+#pragma once
+
+#include "CoreMinimal.h"
+#include "Templates/SharedPointer.h"
+
+#include "PersonaToolMenuContext.generated.h"
+
+class IPersonaToolkit;
+class USkeleton;
+class UDebugSkelMeshComponent;
+class USkeletalMesh;
+class UAnimBlueprint;
+class UAnimationAsset;
+
+UCLASS(BlueprintType)
+class PERSONA_API UPersonaToolMenuContext : public UObject
+{
+	GENERATED_BODY()
+public:
+	
+	/** Get the skeleton that we are editing */
+	UFUNCTION(BlueprintCallable, Category = PersonaEditorExtensions)
+	USkeleton* GetSkeleton() const;
+
+	/** Get the preview component that we are using */
+	UFUNCTION(BlueprintCallable, Category = PersonaEditorExtensions)
+	UDebugSkelMeshComponent* GetPreviewMeshComponent() const;
+
+	/** Get the skeletal mesh that we are editing */
+	UFUNCTION(BlueprintCallable, Category = PersonaEditorExtensions)
+	USkeletalMesh* GetMesh() const;
+
+	/** Get the anim blueprint that we are editing */
+	UFUNCTION(BlueprintCallable, Category = PersonaEditorExtensions)
+	UAnimBlueprint* GetAnimBlueprint() const;
+
+	/** Get the animation asset that we are editing */
+	UFUNCTION(BlueprintCallable, Category = PersonaEditorExtensions)
+	UAnimationAsset* GetAnimationAsset() const;
+
+	void SetToolkit(TSharedRef<IPersonaToolkit> InToolkit);
+
+protected:
+	bool HasValidToolkit() const;
+private:
+	TWeakPtr<IPersonaToolkit> WeakToolkit;
+};
diff --git a/Engine/Source/Editor/SkeletalMeshEditor/Private/SkeletalMeshEditor.cpp b/Engine/Source/Editor/SkeletalMeshEditor/Private/SkeletalMeshEditor.cpp
--- a/Engine/Source/Editor/SkeletalMeshEditor/Private/SkeletalMeshEditor.cpp
+++ b/Engine/Source/Editor/SkeletalMeshEditor/Private/SkeletalMeshEditor.cpp
@@ -46,6 +46,7 @@
 #include "Factories/FbxSkeletalMeshImportData.h"
 
 #include "Misc/MessageDialog.h"
+#include "PersonaToolMenuContext.h"
 
 const FName SkeletalMeshEditorAppIdentifier = FName(TEXT("SkeletalMeshEditorApp"));
 
@@ -588,6 +589,10 @@ void FSkeletalMeshEditor::InitToolMenuContext(FToolMenuContext& MenuContext)
 	USkeletalMeshToolMenuContext* Context = NewObject<USkeletalMeshToolMenuContext>();
 	Context->SkeletalMeshEditor = SharedThis(this);
 	MenuContext.AddObject(Context);
+	
+	UPersonaToolMenuContext* PersonaContext = NewObject<UPersonaToolMenuContext>();
+	PersonaContext->SetToolkit(GetPersonaToolkit());
+	MenuContext.AddObject(PersonaContext);
 }
 
 void FSkeletalMeshEditor::FillMeshClickMenu(FMenuBuilder& MenuBuilder, HActor* HitProxy, const FViewportClick& Click)
diff --git a/Engine/Source/Editor/SkeletonEditor/Private/SkeletonEditor.cpp b/Engine/Source/Editor/SkeletonEditor/Private/SkeletonEditor.cpp
--- a/Engine/Source/Editor/SkeletonEditor/Private/SkeletonEditor.cpp
+++ b/Engine/Source/Editor/SkeletonEditor/Private/SkeletonEditor.cpp
@@ -4,6 +4,7 @@
 #include "Modules/ModuleManager.h"
 #include "Framework/MultiBox/MultiBoxBuilder.h"
 #include "Animation/DebugSkelMeshComponent.h"
+#include "Toolkits/AssetEditorToolkit.h"
 #include "AssetData.h"
 #include "EdGraph/EdGraphSchema.h"
 #include "Editor/EditorEngine.h"
@@ -20,6 +21,7 @@
 #include "IEditableSkeleton.h"
 #include "ISkeletonTreeItem.h"
 #include "Algo/Transform.h"
+#include "PersonaToolMenuContext.h"
 
 const FName SkeletonEditorAppIdentifier = FName(TEXT("SkeletonEditorApp"));
 
@@ -125,6 +127,16 @@ FLinearColor FSkeletonEditor::GetWorldCentricTabColorScale() const
 	return FLinearColor(0.3f, 0.2f, 0.5f, 0.5f);
 }
 
+void FSkeletonEditor::InitToolMenuContext(FToolMenuContext& MenuContext)
+{
+	FAssetEditorToolkit::InitToolMenuContext(MenuContext);
+
+	UPersonaToolMenuContext* Context = NewObject<UPersonaToolMenuContext>();
+	Context->SetToolkit(GetPersonaToolkit());
+
+	MenuContext.AddObject(Context);
+}
+
 void FSkeletonEditor::BindCommands()
 {
 	FSkeletonEditorCommands::Register();
diff --git a/Engine/Source/Editor/SkeletonEditor/Private/SkeletonEditor.h b/Engine/Source/Editor/SkeletonEditor/Private/SkeletonEditor.h
--- a/Engine/Source/Editor/SkeletonEditor/Private/SkeletonEditor.h
+++ b/Engine/Source/Editor/SkeletonEditor/Private/SkeletonEditor.h
@@ -17,6 +17,8 @@ class ISkeletonTree;
 class USkeleton;
 class ISkeletonTreeItem;
 
+struct FToolMenuContext;
+
 namespace SkeletonEditorModes
 {
 	// Mode identifiers
@@ -56,6 +58,7 @@ public:
 	virtual FText GetBaseToolkitName() const override;
 	virtual FString GetWorldCentricTabPrefix() const override;
 	virtual FLinearColor GetWorldCentricTabColorScale() const override;
+	virtual void InitToolMenuContext(FToolMenuContext& MenuContext) override;
 
 	/** FEditorUndoClient interface */
 	virtual void PostUndo(bool bSuccess) override;
diff --git a/Engine/Source/Editor/SkeletonEditor/SkeletonEditor.Build.cs b/Engine/Source/Editor/SkeletonEditor/SkeletonEditor.Build.cs
--- a/Engine/Source/Editor/SkeletonEditor/SkeletonEditor.Build.cs
+++ b/Engine/Source/Editor/SkeletonEditor/SkeletonEditor.Build.cs
@@ -26,6 +26,7 @@ public class SkeletonEditor : ModuleRules
                 "BlueprintGraph",
                 "Kismet",
                 "PinnedCommandList",
+				"ToolMenus"
             }
 		);
 
