Add a default trigger weight threshold to notifies. Notify types can set this, just like they set the notify name. When a notify of that type is placed in an animation asset, it will use this value instead of the standard default of 0.0001.

diff --git a/Engine/Source/Editor/Persona/Private/SAnimNotifyPanel.cpp b/Engine/Source/Editor/Persona/Private/SAnimNotifyPanel.cpp
--- a/Engine/Source/Editor/Persona/Private/SAnimNotifyPanel.cpp
+++ b/Engine/Source/Editor/Persona/Private/SAnimNotifyPanel.cpp
@@ -2504,11 +2504,16 @@ FAnimNotifyEvent& SAnimNotifyTrack::CreateNewNotify(FString NewNotifyName, UClas
 		NewEvent.NotifyStateClass = Cast<UAnimNotifyState>(AnimNotifyClass);
 		NewEvent.Notify = Cast<UAnimNotify>(AnimNotifyClass);
 
-		// Set default duration to 1 frame for AnimNotifyState.
 		if( NewEvent.NotifyStateClass )
 		{
+			// Set default duration to 1 frame for AnimNotifyState.
 			NewEvent.SetDuration(1 / 30.f);
 			NewEvent.EndLink.Link(Sequence, NewEvent.EndLink.GetTime());
+			NewEvent.TriggerWeightThreshold = NewEvent.NotifyStateClass->GetDefaultTriggerWeightThreshold();
+		}
+		else if ( NewEvent.Notify )
+		{
+			NewEvent.TriggerWeightThreshold = NewEvent.Notify->GetDefaultTriggerWeightThreshold();
 		}
 	}
 	else
diff --git a/Engine/Source/Runtime/Engine/Classes/Animation/AnimNotifies/AnimNotify.h b/Engine/Source/Runtime/Engine/Classes/Animation/AnimNotifies/AnimNotify.h
--- a/Engine/Source/Runtime/Engine/Classes/Animation/AnimNotifies/AnimNotify.h
+++ b/Engine/Source/Runtime/Engine/Classes/Animation/AnimNotifies/AnimNotify.h
@@ -63,6 +63,10 @@ class ENGINE_API UAnimNotify : public UObject
 	virtual void OnAnimNotifyCreatedInEditor(FAnimNotifyEvent& ContainingAnimNotifyEvent) {};
 	virtual bool CanBePlaced(UAnimSequenceBase* Animation) const { return true; }
 	virtual void ValidateAssociatedAssets() {}
+	
+	/** TriggerWeightThreshold to use when creating notifies of this type */
+	UFUNCTION(BlueprintNativeEvent)
+	float GetDefaultTriggerWeightThreshold() const;
 #endif
 
 	virtual void Notify(USkeletalMeshComponent* MeshComp, UAnimSequenceBase* Animation);
diff --git a/Engine/Source/Runtime/Engine/Classes/Animation/AnimNotifies/AnimNotifyState.h b/Engine/Source/Runtime/Engine/Classes/Animation/AnimNotifies/AnimNotifyState.h
--- a/Engine/Source/Runtime/Engine/Classes/Animation/AnimNotifies/AnimNotifyState.h
+++ b/Engine/Source/Runtime/Engine/Classes/Animation/AnimNotifies/AnimNotifyState.h
@@ -44,6 +44,10 @@ class ENGINE_API UAnimNotifyState : public UObject
 #if WITH_EDITOR
 	virtual void OnAnimNotifyCreatedInEditor(FAnimNotifyEvent& ContainingAnimNotifyEvent) {};
 	virtual bool CanBePlaced(UAnimSequenceBase* Animation) const { return true; }
+
+	/** TriggerWeightThreshold to use when creating notifies of this type */
+	UFUNCTION(BlueprintNativeEvent)
+	float GetDefaultTriggerWeightThreshold() const;
 #endif
 
 	virtual void NotifyBegin(USkeletalMeshComponent * MeshComp, UAnimSequenceBase * Animation, float TotalDuration);
diff --git a/Engine/Source/Runtime/Engine/Private/Animation/AnimNotify.cpp b/Engine/Source/Runtime/Engine/Private/Animation/AnimNotify.cpp
--- a/Engine/Source/Runtime/Engine/Private/Animation/AnimNotify.cpp
+++ b/Engine/Source/Runtime/Engine/Private/Animation/AnimNotify.cpp
@@ -61,6 +61,13 @@ FString UAnimNotify::GetNotifyName_Implementation() const
 	return NotifyName;
 }
 
+#if WITH_EDITOR
+float UAnimNotify::GetDefaultTriggerWeightThreshold_Implementation() const
+{
+	return ZERO_ANIMWEIGHT_THRESH;
+}
+#endif
+
 /// @endcond
 
 void UAnimNotify::PostLoad()
diff --git a/Engine/Source/Runtime/Engine/Private/Animation/AnimNotifyState.cpp b/Engine/Source/Runtime/Engine/Private/Animation/AnimNotifyState.cpp
--- a/Engine/Source/Runtime/Engine/Private/Animation/AnimNotifyState.cpp
+++ b/Engine/Source/Runtime/Engine/Private/Animation/AnimNotifyState.cpp
@@ -74,6 +74,13 @@ FString UAnimNotifyState::GetNotifyName_Implementation() const
 
 /// @endcond
 
+#if WITH_EDITOR
+float UAnimNotifyState::GetDefaultTriggerWeightThreshold_Implementation() const
+{
+	return ZERO_ANIMWEIGHT_THRESH;
+}
+#endif
+
 void UAnimNotifyState::PostLoad()
 {
 	Super::PostLoad();
