To protect the skeletalmesh when using async operations like build or re-import we deprecate all public member variables of USkeletalMesh class and create the necessary getter and setter. Those member

diff --git a/Engine/Source/Editor/Persona/Private/PersonaMeshDetails.cpp b/Engine/Source/Editor/Persona/Private/PersonaMeshDetails.cpp
--- a/Engine/Source/Editor/Persona/Private/PersonaMeshDetails.cpp
+++ b/Engine/Source/Editor/Persona/Private/PersonaMeshDetails.cpp
@@ -280,16 +280,16 @@ private:
 				return false;
 			}
 			bool bValidLODSettings = false;
-			if (SkeletalMesh->LODSettings != nullptr)
+			if (SkeletalMesh->GetLODSettings() != nullptr)
 			{
-				const int32 NumSettings = FMath::Min(SkeletalMesh->LODSettings->GetNumberOfSettings(), SkeletalMesh->GetLODNum());
+				const int32 NumSettings = FMath::Min(SkeletalMesh->GetLODSettings()->GetNumberOfSettings(), SkeletalMesh->GetLODNum());
 				if (LODIndex < NumSettings)
 				{
 					bValidLODSettings = true;
 				}
 			}
 			
-			const FSkeletalMeshLODGroupSettings* SkeletalMeshLODGroupSettings = bValidLODSettings ? &SkeletalMesh->LODSettings->GetSettingsForLODLevel(LODIndex) : nullptr;
+			const FSkeletalMeshLODGroupSettings* SkeletalMeshLODGroupSettings = bValidLODSettings ? &SkeletalMesh->GetLODSettings()->GetSettingsForLODLevel(LODIndex) : nullptr;
 
 			FGuid BuildGUID = LODInfo->ComputeDeriveDataCacheKey(SkeletalMeshLODGroupSettings);
 			if (LODInfo->BuildGUID != BuildGUID)
@@ -2201,8 +2201,8 @@ void FPersonaMeshDetails::OnCopyMaterialList()
 
 	if (Mesh != nullptr)
 	{
-		FProperty* Property = USkeletalMesh::StaticClass()->FindPropertyByName(GET_MEMBER_NAME_STRING_CHECKED(USkeletalMesh, Materials));
-		auto JsonValue = FJsonObjectConverter::UPropertyToJsonValue(Property, &Mesh->Materials, 0, 0);
+		FProperty* Property = USkeletalMesh::StaticClass()->FindPropertyByName(USkeletalMesh::GetMaterialsMemberName());
+		auto JsonValue = FJsonObjectConverter::UPropertyToJsonValue(Property, &Mesh->GetMaterials(), 0, 0);
 
 		typedef TJsonWriter<TCHAR, TPrettyJsonPrintPolicy<TCHAR>> FStringWriter;
 		typedef TJsonWriterFactory<TCHAR, TPrettyJsonPrintPolicy<TCHAR>> FStringWriterFactory;
@@ -2224,7 +2224,7 @@ bool FPersonaMeshDetails::OnCanCopyMaterialList() const
 
 	if (Mesh != nullptr)
 	{
-		return Mesh->Materials.Num() > 0;
+		return Mesh->GetMaterials().Num() > 0;
 	}
 
 	return false;
@@ -2245,7 +2245,7 @@ void FPersonaMeshDetails::OnPasteMaterialList()
 
 		if (RootJsonValue.IsValid())
 		{
-			FProperty* Property = USkeletalMesh::StaticClass()->FindPropertyByName(GET_MEMBER_NAME_STRING_CHECKED(USkeletalMesh, Materials));
+			FProperty* Property = USkeletalMesh::StaticClass()->FindPropertyByName(USkeletalMesh::GetMaterialsMemberName());
 
 			Mesh->PreEditChange(Property);
 			FScopedTransaction Transaction(LOCTEXT("PersonaChangedPasteMaterialList", "Persona editor: Pasted material list"));
@@ -2253,11 +2253,12 @@ void FPersonaMeshDetails::OnPasteMaterialList()
 			TArray<FSkeletalMaterial> TempMaterials;
 			FJsonObjectConverter::JsonValueToUProperty(RootJsonValue, Property, &TempMaterials, 0, 0);
 			//Do not change the number of material in the array
+			TArray<FSkeletalMaterial>& MeshMaterials = Mesh->GetMaterials();
 			for (int32 MaterialIndex = 0; MaterialIndex < TempMaterials.Num(); ++MaterialIndex)
 			{
-				if (Mesh->Materials.IsValidIndex(MaterialIndex))
+				if (MeshMaterials.IsValidIndex(MaterialIndex))
 				{
-					Mesh->Materials[MaterialIndex].MaterialInterface = TempMaterials[MaterialIndex].MaterialInterface;
+					MeshMaterials[MaterialIndex].MaterialInterface = TempMaterials[MaterialIndex].MaterialInterface;
 				}
 			}
 			
@@ -2275,9 +2276,9 @@ void FPersonaMeshDetails::OnCopyMaterialItem(int32 CurrentSlot)
 	{
 		TSharedRef<FJsonObject> RootJsonObject = MakeShareable(new FJsonObject());
 
-		if (Mesh->Materials.IsValidIndex(CurrentSlot))
+		if (Mesh->GetMaterials().IsValidIndex(CurrentSlot))
 		{
-			const FSkeletalMaterial& Material = Mesh->Materials[CurrentSlot];
+			const FSkeletalMaterial& Material = Mesh->GetMaterials()[CurrentSlot];
 
 			FJsonObjectConverter::UStructToJsonObject(FSkeletalMaterial::StaticStruct(), &Material, RootJsonObject, 0, 0);
 		}
@@ -2302,7 +2303,7 @@ bool FPersonaMeshDetails::OnCanCopyMaterialItem(int32 CurrentSlot) const
 
 	if (Mesh != nullptr)
 	{
-		return Mesh->Materials.IsValidIndex(CurrentSlot);
+		return Mesh->GetMaterials().IsValidIndex(CurrentSlot);
 	}
 
 	return false;
@@ -2323,15 +2324,15 @@ void FPersonaMeshDetails::OnPasteMaterialItem(int32 CurrentSlot)
 
 		if (RootJsonObject.IsValid())
 		{
-			Mesh->PreEditChange(USkeletalMesh::StaticClass()->FindPropertyByName(GET_MEMBER_NAME_STRING_CHECKED(USkeletalMesh, Materials)));
+			Mesh->PreEditChange(USkeletalMesh::StaticClass()->FindPropertyByName(USkeletalMesh::GetMaterialsMemberName()));
 			FScopedTransaction Transaction(LOCTEXT("PersonaChangedPasteMaterialItem", "Persona editor: Pasted material item"));
 			Mesh->Modify();
 
-			if (Mesh->Materials.IsValidIndex(CurrentSlot))
+			if (Mesh->GetMaterials().IsValidIndex(CurrentSlot))
 			{
 				FSkeletalMaterial TmpSkeletalMaterial;
 				FJsonObjectConverter::JsonObjectToUStruct(RootJsonObject.ToSharedRef(), FSkeletalMaterial::StaticStruct(), &TmpSkeletalMaterial, 0, 0);
-				Mesh->Materials[CurrentSlot].MaterialInterface = TmpSkeletalMaterial.MaterialInterface;
+				Mesh->GetMaterials()[CurrentSlot].MaterialInterface = TmpSkeletalMaterial.MaterialInterface;
 			}
 
 			Mesh->PostEditChange();
@@ -2466,7 +2467,7 @@ void FPersonaMeshDetails::CustomizeLODInfoSetingsDetails(IDetailLayoutBuilder& D
 		SNew(SObjectPropertyEntryBox)
 		.PropertyHandle(BakePoseHandle)
 		.AllowedClass(UAnimSequence::StaticClass())
-		.OnShouldFilterAsset(this, &FPersonaMeshDetails::FilterOutBakePose, SkelMesh->Skeleton)
+		.OnShouldFilterAsset(this, &FPersonaMeshDetails::FilterOutBakePose, SkelMesh->GetSkeleton())
 	];
 
 	TSharedPtr<IPropertyHandle> BakePoseOverrideHandle = LODInfoChild->GetChildHandle(GET_MEMBER_NAME_CHECKED(FSkeletalMeshLODInfo, BakePoseOverride));
@@ -2484,7 +2485,7 @@ void FPersonaMeshDetails::CustomizeLODInfoSetingsDetails(IDetailLayoutBuilder& D
 		SNew(SObjectPropertyEntryBox)
 		.PropertyHandle(BakePoseOverrideHandle)
 		.AllowedClass(UAnimSequence::StaticClass())
-		.OnShouldFilterAsset(this, &FPersonaMeshDetails::FilterOutBakePose, SkelMesh->Skeleton)
+		.OnShouldFilterAsset(this, &FPersonaMeshDetails::FilterOutBakePose, SkelMesh->GetSkeleton())
 	];
 
 	TSharedPtr<IPropertyHandle> RemovedBonesHandle = LODInfoChild->GetChildHandle(GET_MEMBER_NAME_CHECKED(FSkeletalMeshLODInfo, BonesToRemove));
@@ -3128,7 +3129,7 @@ void FPersonaMeshDetails::CustomizeLODSettingsCategories(IDetailLayoutBuilder& D
 	];
 
 	// add lod setting assets
-	TSharedPtr<IPropertyHandle> LODSettingAssetPropertyHandle = DetailLayout.GetProperty(GET_MEMBER_NAME_CHECKED(USkeletalMesh, LODSettings), USkeletalMesh::StaticClass());
+	TSharedPtr<IPropertyHandle> LODSettingAssetPropertyHandle = DetailLayout.GetProperty(USkeletalMesh::GetLODSettingsMemberName(), USkeletalMesh::StaticClass());
 	DetailLayout.HideProperty(LODSettingAssetPropertyHandle);
 	LODSettingsCategory.AddCustomRow(LODSettingAssetPropertyHandle->GetPropertyDisplayName())
 	.NameContent()
@@ -3162,32 +3163,32 @@ void FPersonaMeshDetails::CustomizeLODSettingsCategories(IDetailLayoutBuilder& D
 		]
 	];
 
-	TSharedPtr<IPropertyHandle> MinLODPropertyHandle = DetailLayout.GetProperty(GET_MEMBER_NAME_CHECKED(USkeletalMesh, MinLod), USkeletalMesh::StaticClass());
+	TSharedPtr<IPropertyHandle> MinLODPropertyHandle = DetailLayout.GetProperty(USkeletalMesh::GetMinLodMemberName(), USkeletalMesh::StaticClass());
 	IDetailPropertyRow& MinLODRow = LODSettingsCategory.AddProperty(MinLODPropertyHandle);
 	MinLODRow.IsEnabled(TAttribute<bool>::Create(TAttribute<bool>::FGetter::CreateSP(this, &FPersonaMeshDetails::IsLODInfoEditingEnabled, -1)));
 	DetailLayout.HideProperty(MinLODPropertyHandle);
 
-	TSharedPtr<IPropertyHandle> DisableBelowMinLodStrippingPropertyHandle = DetailLayout.GetProperty(GET_MEMBER_NAME_CHECKED(USkeletalMesh, DisableBelowMinLodStripping), USkeletalMesh::StaticClass());
+	TSharedPtr<IPropertyHandle> DisableBelowMinLodStrippingPropertyHandle = DetailLayout.GetProperty(USkeletalMesh::GetDisableBelowMinLodStrippingMemberName(), USkeletalMesh::StaticClass());
 	IDetailPropertyRow& DisableBelowMinLodStrippingRow = LODSettingsCategory.AddProperty(DisableBelowMinLodStrippingPropertyHandle);
 	DisableBelowMinLodStrippingRow.IsEnabled(TAttribute<bool>::Create(TAttribute<bool>::FGetter::CreateSP(this, &FPersonaMeshDetails::IsLODInfoEditingEnabled, -1)));
 	DetailLayout.HideProperty(DisableBelowMinLodStrippingPropertyHandle);
 
-	TSharedPtr<IPropertyHandle> bOverrideLODStreamingSettingsHandle = DetailLayout.GetProperty(GET_MEMBER_NAME_CHECKED(USkeletalMesh, bOverrideLODStreamingSettings), USkeletalMesh::StaticClass());
+	TSharedPtr<IPropertyHandle> bOverrideLODStreamingSettingsHandle = DetailLayout.GetProperty(USkeletalMesh::GetOverrideLODStreamingSettingsMemberName(), USkeletalMesh::StaticClass());
 	IDetailPropertyRow& bOverrideLODStreamingSettingsRow = LODSettingsCategory.AddProperty(bOverrideLODStreamingSettingsHandle);
 	bOverrideLODStreamingSettingsRow.IsEnabled(TAttribute<bool>::Create(TAttribute<bool>::FGetter::CreateSP(this, &FPersonaMeshDetails::IsLODInfoEditingEnabled, -1)));
 	DetailLayout.HideProperty(bOverrideLODStreamingSettingsHandle);
 
-	TSharedPtr<IPropertyHandle> bSupportLODStreamingPropertyHandle = DetailLayout.GetProperty(GET_MEMBER_NAME_CHECKED(USkeletalMesh, bSupportLODStreaming), USkeletalMesh::StaticClass());
+	TSharedPtr<IPropertyHandle> bSupportLODStreamingPropertyHandle = DetailLayout.GetProperty(USkeletalMesh::GetSupportLODStreamingMemberName(), USkeletalMesh::StaticClass());
 	IDetailPropertyRow& bSupportLODStreamingRow = LODSettingsCategory.AddProperty(bSupportLODStreamingPropertyHandle);
 	bSupportLODStreamingRow.IsEnabled(TAttribute<bool>::Create(TAttribute<bool>::FGetter::CreateSP(this, &FPersonaMeshDetails::IsLODInfoEditingEnabled, -1)));
 	DetailLayout.HideProperty(bSupportLODStreamingPropertyHandle);
 
-	TSharedPtr<IPropertyHandle> MaxNumStreamedLODsPropertyHandle = DetailLayout.GetProperty(GET_MEMBER_NAME_CHECKED(USkeletalMesh, MaxNumStreamedLODs), USkeletalMesh::StaticClass());
+	TSharedPtr<IPropertyHandle> MaxNumStreamedLODsPropertyHandle = DetailLayout.GetProperty(USkeletalMesh::GetMaxNumStreamedLODsMemberName(), USkeletalMesh::StaticClass());
 	IDetailPropertyRow& MaxNumStreamedLODsRow = LODSettingsCategory.AddProperty(MaxNumStreamedLODsPropertyHandle);
 	MaxNumStreamedLODsRow.IsEnabled(TAttribute<bool>::Create(TAttribute<bool>::FGetter::CreateSP(this, &FPersonaMeshDetails::IsLODInfoEditingEnabled, -1)));
 	DetailLayout.HideProperty(MaxNumStreamedLODsPropertyHandle);
 
-	TSharedPtr<IPropertyHandle> MaxNumOptionalLODsPropertyHandle = DetailLayout.GetProperty(GET_MEMBER_NAME_CHECKED(USkeletalMesh, MaxNumOptionalLODs), USkeletalMesh::StaticClass());
+	TSharedPtr<IPropertyHandle> MaxNumOptionalLODsPropertyHandle = DetailLayout.GetProperty(USkeletalMesh::GetMaxNumOptionalLODsMemberName(), USkeletalMesh::StaticClass());
 	IDetailPropertyRow& MaxNumOptionalLODsRow = LODSettingsCategory.AddProperty(MaxNumOptionalLODsPropertyHandle);
 	MaxNumOptionalLODsRow.IsEnabled(TAttribute<bool>::Create(TAttribute<bool>::FGetter::CreateSP(this, &FPersonaMeshDetails::IsLODInfoEditingEnabled, -1)));
 	DetailLayout.HideProperty(MaxNumOptionalLODsPropertyHandle);
@@ -3235,7 +3236,7 @@ FReply FPersonaMeshDetails::OnSaveLODSettings()
 				UPackage::SavePackage(Package, NULL, RF_Standalone, *PackageFileName, GError, nullptr, false, true, SAVE_NoError);
 
 				// set the property back to SkelMesh;
-				SkelMesh->LODSettings = NewLODSettingAsset;
+				SkelMesh->SetLODSettings(NewLODSettingAsset);
 			}
 		}
 	}
@@ -3261,7 +3262,7 @@ bool FPersonaMeshDetails::IsLODInfoEditingEnabled(int32 LODIndex) const
 	USkeletalMesh* SkelMesh = GetPersonaToolkit()->GetMesh();
 	if (SkelMesh)
 	{
-		if (SkelMesh->LODSettings)
+		if (SkelMesh->GetLODSettings())
 		{
 			// if LODIndex == -1, we don't care about lod index
 			if (LODIndex == -1)
@@ -3269,7 +3270,7 @@ bool FPersonaMeshDetails::IsLODInfoEditingEnabled(int32 LODIndex) const
 				return false;
 			}
 
-			if (SkelMesh->LODSettings->GetNumberOfSettings() > LODIndex)
+			if (SkelMesh->GetLODSettings()->GetNumberOfSettings() > LODIndex)
 			{
 				return false;
 			}
@@ -3398,7 +3399,7 @@ FReply FPersonaMeshDetails::ApplyLODChanges(int32 LODIndex)
 			if (LODIndex == 0) //Base LOD must update the asset import data
 			{
 				//Update the Asset Import Data
-				UFbxSkeletalMeshImportData* SKImportData = Cast<UFbxSkeletalMeshImportData>(SkelMesh->AssetImportData);
+				UFbxSkeletalMeshImportData* SKImportData = Cast<UFbxSkeletalMeshImportData>(SkelMesh->GetAssetImportData());
 				if (SKImportData)
 				{
 					check(bSrcBuildDataAvailable);
@@ -3427,7 +3428,7 @@ FReply FPersonaMeshDetails::ApplyLODChanges(int32 LODIndex)
 	}
 
 	RefreshMeshDetailLayout();
-	if (SkelMesh->MorphTargets.Num() > 0)
+	if (SkelMesh->GetMorphTargets().Num() > 0)
 	{
 		GetPersonaToolkit()->GetPreviewScene()->BroadcastOnMorphTargetsChanged();
 	}
@@ -3565,15 +3566,15 @@ FReply FPersonaMeshDetails::RemoveOneLOD(int32 LODIndex)
 		FText ConfirmRemoveLODText = FText::Format( LOCTEXT("PersonaRemoveLOD_Confirmation", "Are you sure you want to remove LOD {0} from {1}?"), LODIndex, FText::FromString(SkelMesh->GetName()) );
 
 		// if we have lod settings, and then 
-		if (SkelMesh->LODSettings != nullptr)
+		if (SkelMesh->GetLODSettings() != nullptr)
 		{
 			// if I have more LODs, and if LODSettings will be copied back over, 
 			// all LODs have to be regenerated
 			// warn users about it
-			if (SkelMesh->IsValidLODIndex(LODIndex + 1) && SkelMesh->LODSettings->GetNumberOfSettings() > LODIndex)
+			if (SkelMesh->IsValidLODIndex(LODIndex + 1) && SkelMesh->GetLODSettings()->GetNumberOfSettings() > LODIndex)
 			{
 				// now the information will get copied over after removing this LOD
-				ConfirmRemoveLODText = FText::Format(LOCTEXT("PersonaRemoveLODOverriding_Confirmation", "You're currently using LOD Setting Asset \'{2}\' that will override the next LODs with current setting. This will require to regenerate the next LODs after removing this LOD. If you do not want this, clear the LOD Setting Asset before removing LODs. \n\n Are you sure you want to remove LOD {0} from {1}?"), LODIndex, FText::FromString(SkelMesh->GetName()), FText::FromString(SkelMesh->LODSettings->GetName()));
+				ConfirmRemoveLODText = FText::Format(LOCTEXT("PersonaRemoveLODOverriding_Confirmation", "You're currently using LOD Setting Asset \'{2}\' that will override the next LODs with current setting. This will require to regenerate the next LODs after removing this LOD. If you do not want this, clear the LOD Setting Asset before removing LODs. \n\n Are you sure you want to remove LOD {0} from {1}?"), LODIndex, FText::FromString(SkelMesh->GetName()), FText::FromString(SkelMesh->GetLODSettings()->GetName()));
 			}
 		}
 
@@ -3594,9 +3595,9 @@ FReply FPersonaMeshDetails::RemoveOneLOD(int32 LODIndex)
 
 				FLODUtilities::RemoveLOD(UpdateContext, LODIndex);
 
-				if (SkelMesh->LODSettings)
+				if (SkelMesh->GetLODSettings())
 				{
-					SkelMesh->LODSettings->SetLODSettingsToMesh(SkelMesh);
+					SkelMesh->GetLODSettings()->SetLODSettingsToMesh(SkelMesh);
 				}
 			}
 
@@ -3674,7 +3675,7 @@ void FPersonaMeshDetails::ApplyChanges()
 		//PostEditChange will be call when going out of scope
 	}
 	//Update the morph target list since we have rebuild the asset
-	if (SkelMesh->MorphTargets.Num() > 0)
+	if (SkelMesh->GetMorphTargets().Num() > 0)
 	{
 		GetPersonaToolkit()->GetPreviewScene()->BroadcastOnMorphTargetsChanged();
 	}
@@ -3737,9 +3738,9 @@ FText FPersonaMeshDetails::GetLODImportedText(int32 LODIndex) const
 FText FPersonaMeshDetails::GetMaterialSlotNameText(int32 MaterialIndex) const
 {
 	USkeletalMesh* Mesh = GetPersonaToolkit()->GetMesh();
-	if (Mesh && Mesh->Materials.IsValidIndex(MaterialIndex))
+	if (Mesh && Mesh->GetMaterials().IsValidIndex(MaterialIndex))
 	{
-		return FText::FromName(Mesh->Materials[MaterialIndex].MaterialSlotName);
+		return FText::FromName(Mesh->GetMaterials()[MaterialIndex].MaterialSlotName);
 	}
 
 	return LOCTEXT("SkeletalMeshMaterial_InvalidIndex", "Invalid Material Index");
@@ -3769,7 +3770,7 @@ void FPersonaMeshDetails::CustomizeDetails( IDetailLayoutBuilder& DetailLayout )
 
 	// Post process selector
 	IDetailCategoryBuilder& SkelMeshCategory = DetailLayout.EditCategory("SkeletalMesh");
-	TSharedRef<IPropertyHandle> PostProcessHandle = DetailLayout.GetProperty(GET_MEMBER_NAME_CHECKED(USkeletalMesh, PostProcessAnimBlueprint), USkeletalMesh::StaticClass());
+	TSharedRef<IPropertyHandle> PostProcessHandle = DetailLayout.GetProperty(USkeletalMesh::GetPostProcessAnimBlueprintMemberName(), USkeletalMesh::StaticClass());
 	PostProcessHandle->SetOnPropertyValueChanged(FSimpleDelegate::CreateSP(this, &FPersonaMeshDetails::OnPostProcessBlueprintChanged, &DetailLayout));
 	PostProcessHandle->MarkHiddenByCustomization();
 
@@ -3790,8 +3791,8 @@ void FPersonaMeshDetails::CustomizeDetails( IDetailLayoutBuilder& DetailLayout )
 	];
 
 	IDetailCategoryBuilder& ImportSettingsCategory = DetailLayout.EditCategory("ImportSettings");
-	TSharedRef<IPropertyHandle> AssetImportProperty = DetailLayout.GetProperty(GET_MEMBER_NAME_CHECKED(USkeletalMesh, AssetImportData), USkeletalMesh::StaticClass());
-	if (!SkeletalMeshPtr.IsValid() || !IsValid(SkeletalMeshPtr->AssetImportData) || !SkeletalMeshPtr->AssetImportData->IsA<UFbxSkeletalMeshImportData>())
+	TSharedRef<IPropertyHandle> AssetImportProperty = DetailLayout.GetProperty(USkeletalMesh::GetAssetImportDataMemberName(), USkeletalMesh::StaticClass());
+	if (!SkeletalMeshPtr.IsValid() || !IsValid(SkeletalMeshPtr->GetAssetImportData()) || !SkeletalMeshPtr->GetAssetImportData()->IsA<UFbxSkeletalMeshImportData>())
 	{
 		// Hide the ability to change the import settings object
 		IDetailPropertyRow& Row = ImportSettingsCategory.AddProperty(AssetImportProperty);
@@ -3869,7 +3870,7 @@ void FPersonaMeshDetails::HideUnnecessaryProperties(IDetailLayoutBuilder& Detail
 		}
 	}
 
-	TSharedRef<IPropertyHandle> MaterialsProperty = DetailLayout.GetProperty(FName("Materials"), USkeletalMesh::StaticClass());
+	TSharedRef<IPropertyHandle> MaterialsProperty = DetailLayout.GetProperty(USkeletalMesh::GetMaterialsMemberName(), USkeletalMesh::StaticClass());
 	DetailLayout.HideProperty(MaterialsProperty);
 
 	// hide all properties in Mirroring category to hide Mirroring category itself
@@ -3890,7 +3891,7 @@ void FPersonaMeshDetails::OnPostProcessBlueprintChanged(IDetailLayoutBuilder* De
 FString FPersonaMeshDetails::GetCurrentPostProcessBlueprintPath() const
 {
 	USkeletalMesh* SkelMesh = GetPersonaToolkit()->GetMesh();
-	if(UClass* PostProcessClass = *SkelMesh->PostProcessAnimBlueprint)
+	if(UClass* PostProcessClass = *SkelMesh->GetPostProcessAnimBlueprint())
 	{
 		return PostProcessClass->GetPathName();
 	}
@@ -3902,7 +3903,7 @@ bool FPersonaMeshDetails::OnShouldFilterPostProcessBlueprint(const FAssetData& A
 {
 	if(USkeletalMesh* SkelMesh = GetPersonaToolkit()->GetMesh())
 	{
-		const FString CurrentMeshSkeletonName = FString::Printf(TEXT("%s'%s'"), *SkelMesh->Skeleton->GetClass()->GetName(), *SkelMesh->Skeleton->GetPathName());
+		const FString CurrentMeshSkeletonName = FString::Printf(TEXT("%s'%s'"), *SkelMesh->GetSkeleton()->GetClass()->GetName(), *SkelMesh->GetSkeleton()->GetPathName());
 		const FString SkeletonName = AssetData.GetTagValueRef<FString>("TargetSkeleton");
 
 		return SkeletonName != CurrentMeshSkeletonName;
@@ -3947,7 +3948,7 @@ FReply FPersonaMeshDetails::OnReimportLodClicked(EReimportButtonType InReimportT
 			LODInfo->SourceImportFilename.Empty();
 			
 			//Avoid changing the settings if the skeletal mesh is using a LODSettings asset valid for this LOD
-			bool bUseLODSettingAsset = SkelMesh->LODSettings != nullptr && SkelMesh->LODSettings->GetNumberOfSettings() > InLODIndex;
+			bool bUseLODSettingAsset = SkelMesh->GetLODSettings() != nullptr && SkelMesh->GetLODSettings()->GetNumberOfSettings() > InLODIndex;
 			//Make the reduction settings change according to the context
 			if (!bUseLODSettingAsset && SkelMesh->IsReductionActive(InLODIndex) && LODInfo->bHasBeenSimplified && SkelMesh->IsLODImportedDataEmpty(InLODIndex))
 			{
@@ -4008,9 +4009,10 @@ void FPersonaMeshDetails::OnGetMaterialsForArray(class IMaterialListBuilder& Out
 	if (!SkelMesh)
 		return;
 
-	for (int32 MaterialIndex = 0; MaterialIndex < SkelMesh->Materials.Num(); ++MaterialIndex)
+	TArray<FSkeletalMaterial>& SkelMeshMaterials = SkelMesh->GetMaterials();
+	for (int32 MaterialIndex = 0; MaterialIndex < SkelMeshMaterials.Num(); ++MaterialIndex)
 	{
-		OutMaterials.AddMaterial(MaterialIndex, SkelMesh->Materials[MaterialIndex].MaterialInterface, true);
+		OutMaterials.AddMaterial(MaterialIndex, SkelMeshMaterials[MaterialIndex].MaterialInterface, true);
 	}
 }
 
@@ -4025,36 +4027,36 @@ void FPersonaMeshDetails::OnMaterialArrayChanged(UMaterialInterface* NewMaterial
 		FProperty* MaterialProperty = FindFProperty<FProperty>(USkeletalMesh::StaticClass(), "Materials");
 		check(MaterialProperty);
 		Mesh->PreEditChange(MaterialProperty);
-		check(Mesh->Materials.Num() > SlotIndex)
+		TArray<FSkeletalMaterial>& MeshMaterials = Mesh->GetMaterials();
+		check(MeshMaterials.Num() > SlotIndex);
 
 		if (NewMaterial != PrevMaterial)
 		{
 			GEditor->BeginTransaction(LOCTEXT("PersonaEditorMaterialChanged", "Persona editor: material changed"));
 			bMadeTransaction = true;
 			Mesh->Modify();
-			Mesh->Materials[SlotIndex].MaterialInterface = NewMaterial;
+			MeshMaterials[SlotIndex].MaterialInterface = NewMaterial;
 
 			//Add a default name to the material slot if this slot was manually add and there is no name yet
-			if (NewMaterial != nullptr && (Mesh->Materials[SlotIndex].ImportedMaterialSlotName == NAME_None || Mesh->Materials[SlotIndex].MaterialSlotName == NAME_None))
+			if (NewMaterial != nullptr && (MeshMaterials[SlotIndex].ImportedMaterialSlotName == NAME_None || MeshMaterials[SlotIndex].MaterialSlotName == NAME_None))
 			{
-				if (Mesh->Materials[SlotIndex].MaterialSlotName == NAME_None)
+				if (MeshMaterials[SlotIndex].MaterialSlotName == NAME_None)
 				{
-					
-					Mesh->Materials[SlotIndex].MaterialSlotName = NewMaterial->GetFName();
+					MeshMaterials[SlotIndex].MaterialSlotName = NewMaterial->GetFName();
 				}
 
 				//Ensure the imported material slot name is unique
-				if (Mesh->Materials[SlotIndex].ImportedMaterialSlotName == NAME_None)
+				if (MeshMaterials[SlotIndex].ImportedMaterialSlotName == NAME_None)
 				{
-					auto IsMaterialNameUnique = [&Mesh, SlotIndex](const FName TestName)
+					auto IsMaterialNameUnique = [&Mesh, &MeshMaterials, SlotIndex](const FName TestName)
 					{
-						for (int32 MaterialIndex = 0; MaterialIndex < Mesh->Materials.Num(); ++MaterialIndex)
+						for (int32 MaterialIndex = 0; MaterialIndex < MeshMaterials.Num(); ++MaterialIndex)
 						{
 							if (MaterialIndex == SlotIndex)
 							{
 								continue;
 							}
-							if (Mesh->Materials[MaterialIndex].ImportedMaterialSlotName == TestName)
+							if (MeshMaterials[MaterialIndex].ImportedMaterialSlotName == TestName)
 							{
 								return false;
 							}
@@ -4075,7 +4077,7 @@ void FPersonaMeshDetails::OnMaterialArrayChanged(UMaterialInterface* NewMaterial
 							MaterialSlotName = NewMaterial->GetName() + TEXT("_") + FString::FromInt(MatchNameCounter);
 						}
 					}
-					Mesh->Materials[SlotIndex].ImportedMaterialSlotName = FName(*MaterialSlotName);
+					MeshMaterials[SlotIndex].ImportedMaterialSlotName = FName(*MaterialSlotName);
 				}
 			}
 		}
@@ -4102,7 +4104,7 @@ FReply FPersonaMeshDetails::AddMaterialSlot()
 
 	FScopedTransaction Transaction(LOCTEXT("PersonaAddMaterialSlotTransaction", "Persona editor: Add material slot"));
 	SkeletalMeshPtr->Modify();
-	SkeletalMeshPtr->Materials.Add(FSkeletalMaterial());
+	SkeletalMeshPtr->GetMaterials().Add(FSkeletalMaterial());
 
 	SkeletalMeshPtr->PostEditChange();
 
@@ -4115,7 +4117,7 @@ FText FPersonaMeshDetails::GetMaterialArrayText() const
 	int32 SlotNumber = 0;
 	if (SkeletalMeshPtr.IsValid())
 	{
-		SlotNumber = SkeletalMeshPtr->Materials.Num();
+		SlotNumber = SkeletalMeshPtr->GetMaterials().Num();
 	}
 	MaterialArrayText = FString::FromInt(SlotNumber) + MaterialArrayText;
 	return FText::FromString(MaterialArrayText);
@@ -4132,23 +4134,24 @@ void FPersonaMeshDetails::OnGetSectionsForView(ISectionListBuilder& OutSections,
 		FSkeletalMeshLODModel& Model = ImportedResource->LODModels[LODIndex];
 
 		TArray<int32>& MaterialMap = SkelMesh->GetLODInfo(LODIndex)->LODMaterialMap;
+		const TArray<FSkeletalMaterial>& SkelMeshMaterials = SkelMesh->GetMaterials();
 		int32 NumSections = Model.Sections.Num();
 		for (int32 SectionIdx = 0; SectionIdx < NumSections; SectionIdx++)
 		{
 			int32 DefaultSectionMaterialIndex = GetDefaultMaterialIndex(SkelMesh, LODIndex, SectionIdx);
 			int32 MaterialIndex = Model.Sections[SectionIdx].MaterialIndex;;
-			if (MaterialMap.IsValidIndex(SectionIdx) && SkelMesh->Materials.IsValidIndex(MaterialMap[SectionIdx]))
+			if (MaterialMap.IsValidIndex(SectionIdx) && SkelMeshMaterials.IsValidIndex(MaterialMap[SectionIdx]))
 			{
 				MaterialIndex = MaterialMap[SectionIdx];
 			}
 
-			if (SkelMesh->Materials.IsValidIndex(MaterialIndex))
+			if (SkelMeshMaterials.IsValidIndex(MaterialIndex))
 			{
-				FName CurrentSectionMaterialSlotName = SkelMesh->Materials[MaterialIndex].MaterialSlotName;
-				FName CurrentSectionOriginalImportedMaterialName = SkelMesh->Materials[MaterialIndex].ImportedMaterialSlotName;
+				FName CurrentSectionMaterialSlotName = SkelMeshMaterials[MaterialIndex].MaterialSlotName;
+				FName CurrentSectionOriginalImportedMaterialName = SkelMeshMaterials[MaterialIndex].ImportedMaterialSlotName;
 				TMap<int32, FName> AvailableSectionName;
 				int32 CurrentIterMaterialIndex = 0;
-				for (const FSkeletalMaterial &SkeletalMaterial : SkelMesh->Materials)
+				for (const FSkeletalMaterial &SkeletalMaterial : SkelMeshMaterials)
 				{
 					if (MaterialIndex != CurrentIterMaterialIndex)
 					{
@@ -4166,7 +4169,7 @@ void FPersonaMeshDetails::OnGetSectionsForView(ISectionListBuilder& OutSections,
 				}
 				bool bClothSection = Model.Sections[SectionIdx].HasClothingData();
 				bool bIsChunkSection = Model.Sections[SectionIdx].ChunkedParentSectionIndex != INDEX_NONE;
-				OutSections.AddSection(LODIndex, SectionIdx, CurrentSectionMaterialSlotName, MaterialIndex, CurrentSectionOriginalImportedMaterialName, AvailableSectionName, SkelMesh->Materials[MaterialIndex].MaterialInterface, bClothSection, bIsChunkSection, DefaultSectionMaterialIndex);
+				OutSections.AddSection(LODIndex, SectionIdx, CurrentSectionMaterialSlotName, MaterialIndex, CurrentSectionOriginalImportedMaterialName, AvailableSectionName, SkelMeshMaterials[MaterialIndex].MaterialInterface, bClothSection, bIsChunkSection, DefaultSectionMaterialIndex);
 			}
 		}
 	}
@@ -4174,19 +4177,19 @@ void FPersonaMeshDetails::OnGetSectionsForView(ISectionListBuilder& OutSections,
 
 FText FPersonaMeshDetails::GetMaterialNameText(int32 MaterialIndex) const
 {
-	if (SkeletalMeshPtr.IsValid() && SkeletalMeshPtr->Materials.IsValidIndex(MaterialIndex))
+	if (SkeletalMeshPtr.IsValid() && SkeletalMeshPtr->GetMaterials().IsValidIndex(MaterialIndex))
 	{
-		return FText::FromName(SkeletalMeshPtr->Materials[MaterialIndex].MaterialSlotName);
+		return FText::FromName(SkeletalMeshPtr->GetMaterials()[MaterialIndex].MaterialSlotName);
 	}
 	return FText::FromName(NAME_None);
 }
 
 FText FPersonaMeshDetails::GetOriginalImportMaterialNameText(int32 MaterialIndex) const
 {
-	if (SkeletalMeshPtr.IsValid() && SkeletalMeshPtr->Materials.IsValidIndex(MaterialIndex))
+	if (SkeletalMeshPtr.IsValid() && SkeletalMeshPtr->GetMaterials().IsValidIndex(MaterialIndex))
 	{
 		FString OriginalImportMaterialName;
-		SkeletalMeshPtr->Materials[MaterialIndex].ImportedMaterialSlotName.ToString(OriginalImportMaterialName);
+		SkeletalMeshPtr->GetMaterials()[MaterialIndex].ImportedMaterialSlotName.ToString(OriginalImportMaterialName);
 		OriginalImportMaterialName = TEXT("Original Imported Material Name: ") + OriginalImportMaterialName;
 		return FText::FromString(OriginalImportMaterialName);
 			}
@@ -4196,7 +4199,7 @@ FText FPersonaMeshDetails::GetOriginalImportMaterialNameText(int32 MaterialIndex
 void FPersonaMeshDetails::OnMaterialNameCommitted(const FText& InValue, ETextCommit::Type CommitType, int32 MaterialIndex)
 {
 	FName InValueName = FName(*(InValue.ToString()));
-	if (SkeletalMeshPtr.IsValid() && SkeletalMeshPtr->Materials.IsValidIndex(MaterialIndex) && InValueName != SkeletalMeshPtr->Materials[MaterialIndex].MaterialSlotName)
+	if (SkeletalMeshPtr.IsValid() && SkeletalMeshPtr->GetMaterials().IsValidIndex(MaterialIndex) && InValueName != SkeletalMeshPtr->GetMaterials()[MaterialIndex].MaterialSlotName)
 	{
 		FScopedTransaction ScopeTransaction(LOCTEXT("PersonaMaterialSlotNameChanged", "Persona editor: Material slot name change"));
 
@@ -4204,7 +4207,7 @@ void FPersonaMeshDetails::OnMaterialNameCommitted(const FText& InValue, ETextCom
 		check(ChangedProperty);
 		SkeletalMeshPtr->PreEditChange(ChangedProperty);
 
-		SkeletalMeshPtr->Materials[MaterialIndex].MaterialSlotName = InValueName;
+		SkeletalMeshPtr->GetMaterials()[MaterialIndex].MaterialSlotName = InValueName;
 		
 		FPropertyChangedEvent PropertyUpdateStruct(ChangedProperty);
 		SkeletalMeshPtr->PostEditChangeProperty(PropertyUpdateStruct);
@@ -4299,14 +4302,14 @@ TSharedRef<SWidget> FPersonaMeshDetails::OnGetMaterialSlotUsedByMenuContent(int3
 }
 
 bool FPersonaMeshDetails::CanDeleteMaterialSlot(int32 MaterialIndex) const
-	{
+{
 	if (!SkeletalMeshPtr.IsValid())
 	{
 		return false;
 	}
 
-	return SkeletalMeshPtr->Materials.IsValidIndex(MaterialIndex);
-	}
+	return SkeletalMeshPtr->GetMaterials().IsValidIndex(MaterialIndex);
+}
 	
 void FPersonaMeshDetails::OnDeleteMaterialSlot(int32 MaterialIndex)
 {
@@ -4330,7 +4333,7 @@ void FPersonaMeshDetails::OnDeleteMaterialSlot(int32 MaterialIndex)
 	{
 		FScopedSkeletalMeshPostEditChange ScopedPostEditChange(SkeletalMeshPtr.Get());
 		//When we delete a material slot we must invalidate the DDC because material index is not part of the DDC key by design
-		SkeletalMeshPtr->Materials.RemoveAt(MaterialIndex);
+		SkeletalMeshPtr->GetMaterials().RemoveAt(MaterialIndex);
 		FSkeletalMeshModel* Model = SkeletalMeshPtr->GetImportedModel();
 
 		int32 NumLODInfos = SkeletalMeshPtr->GetLODNum();
@@ -4369,7 +4372,8 @@ bool FPersonaMeshDetails::OnMaterialListDirty()
 	TMap<int32, TArray<FSectionLocalizer>> TempMaterialUsedMap;
 	if (SkeletalMeshPtr.IsValid())
 	{
-		for (int32 MaterialIndex = 0; MaterialIndex < SkeletalMeshPtr->Materials.Num(); ++MaterialIndex)
+		const int32 MaterialCount = SkeletalMeshPtr->GetMaterials().Num();
+		for (int32 MaterialIndex = 0; MaterialIndex < MaterialCount; ++MaterialIndex)
 		{
 			TArray<FSectionLocalizer> SectionLocalizers;
 			FSkeletalMeshModel* ImportedResource = SkeletalMeshPtr->GetImportedModel();
@@ -5395,11 +5399,11 @@ EVisibility FPersonaMeshDetails::GetOverrideUVDensityVisibililty() const
 ECheckBoxState FPersonaMeshDetails::IsUVDensityOverridden(int32 MaterialIndex) const
 {
 	USkeletalMesh* Mesh = SkeletalMeshPtr.Get();
-	if (!Mesh || !Mesh->Materials.IsValidIndex(MaterialIndex))
+	if (!Mesh || !Mesh->GetMaterials().IsValidIndex(MaterialIndex))
 	{
 		return ECheckBoxState::Undetermined;
 	}
-	else if (Mesh->Materials[MaterialIndex].UVChannelData.bOverrideDensities)
+	else if (Mesh->GetMaterials()[MaterialIndex].UVChannelData.bOverrideDensities)
 	{
 		return ECheckBoxState::Checked;
 	}
@@ -5413,9 +5417,9 @@ ECheckBoxState FPersonaMeshDetails::IsUVDensityOverridden(int32 MaterialIndex) c
 void FPersonaMeshDetails::OnOverrideUVDensityChanged(ECheckBoxState NewState, int32 MaterialIndex)
 {
 	USkeletalMesh* Mesh = SkeletalMeshPtr.Get();
-	if (NewState != ECheckBoxState::Undetermined && Mesh && Mesh->Materials.IsValidIndex(MaterialIndex))
+	if (NewState != ECheckBoxState::Undetermined && Mesh && Mesh->GetMaterials().IsValidIndex(MaterialIndex))
 	{
-		Mesh->Materials[MaterialIndex].UVChannelData.bOverrideDensities = (NewState == ECheckBoxState::Checked);
+		Mesh->GetMaterials()[MaterialIndex].UVChannelData.bOverrideDensities = (NewState == ECheckBoxState::Checked);
 		Mesh->UpdateUVChannelData(true);
 	}
 }
@@ -5436,9 +5440,9 @@ EVisibility FPersonaMeshDetails::GetUVDensityVisibility(int32 MaterialIndex, int
 TOptional<float> FPersonaMeshDetails::GetUVDensityValue(int32 MaterialIndex, int32 UVChannelIndex) const
 {
 	USkeletalMesh* Mesh = SkeletalMeshPtr.Get();
-	if (Mesh && Mesh->Materials.IsValidIndex(MaterialIndex))
+	if (Mesh && Mesh->GetMaterials().IsValidIndex(MaterialIndex))
 	{
-		float Value = Mesh->Materials[MaterialIndex].UVChannelData.LocalUVDensities[UVChannelIndex];
+		float Value = Mesh->GetMaterials()[MaterialIndex].UVChannelData.LocalUVDensities[UVChannelIndex];
 		return FMath::RoundToFloat(Value * 4.f) * .25f;
 	}
 	return TOptional<float>();
@@ -5447,9 +5451,9 @@ TOptional<float> FPersonaMeshDetails::GetUVDensityValue(int32 MaterialIndex, int
 void FPersonaMeshDetails::SetUVDensityValue(float InDensity, ETextCommit::Type CommitType, int32 MaterialIndex, int32 UVChannelIndex)
 {
 	USkeletalMesh* Mesh = SkeletalMeshPtr.Get();
-	if (Mesh && Mesh->Materials.IsValidIndex(MaterialIndex))
+	if (Mesh && Mesh->GetMaterials().IsValidIndex(MaterialIndex))
 	{
-		Mesh->Materials[MaterialIndex].UVChannelData.LocalUVDensities[UVChannelIndex] = FMath::Max<float>(0, InDensity);
+		Mesh->GetMaterials()[MaterialIndex].UVChannelData.LocalUVDensities[UVChannelIndex] = FMath::Max<float>(0, InDensity);
 		Mesh->UpdateUVChannelData(true);
 	}
 }
@@ -5464,7 +5468,7 @@ int32 FPersonaMeshDetails::GetMaterialIndex(int32 LODIndex, int32 SectionIndex)
 	check(ImportedResource && ImportedResource->LODModels.IsValidIndex(LODIndex));
 	int32 MaterialIndex = ImportedResource->LODModels[LODIndex].Sections[SectionIndex].MaterialIndex;
 	FSkeletalMeshLODInfo& Info = *(SkelMesh->GetLODInfo(LODIndex));
-	if (Info.LODMaterialMap.IsValidIndex(SectionIndex) && SkelMesh->Materials.IsValidIndex(Info.LODMaterialMap[SectionIndex]))
+	if (Info.LODMaterialMap.IsValidIndex(SectionIndex) && SkelMesh->GetMaterials().IsValidIndex(Info.LODMaterialMap[SectionIndex]))
 	{
 		return Info.LODMaterialMap[SectionIndex];
 		
@@ -5489,12 +5493,13 @@ void FPersonaMeshDetails::OnSectionChanged(int32 LODIndex, int32 SectionIndex, i
 
 		int32 NewSkeletalMaterialIndex = INDEX_NONE;
 		FName NewImportedMaterialSlotName = NAME_None;
-		for (int SkeletalMaterialIndex = 0; SkeletalMaterialIndex < Mesh->Materials.Num(); ++SkeletalMaterialIndex)
+		const TArray<FSkeletalMaterial>& MeshMaterials = Mesh->GetMaterials();
+		for (int SkeletalMaterialIndex = 0; SkeletalMaterialIndex < MeshMaterials.Num(); ++SkeletalMaterialIndex)
 		{
-			if (NewMaterialSlotIndex == SkeletalMaterialIndex && Mesh->Materials[SkeletalMaterialIndex].MaterialSlotName == CleanNewMaterialSlotName)
+			if (NewMaterialSlotIndex == SkeletalMaterialIndex && MeshMaterials[SkeletalMaterialIndex].MaterialSlotName == CleanNewMaterialSlotName)
 			{
 				NewSkeletalMaterialIndex = SkeletalMaterialIndex;
-				NewImportedMaterialSlotName = Mesh->Materials[SkeletalMaterialIndex].ImportedMaterialSlotName;
+				NewImportedMaterialSlotName = MeshMaterials[SkeletalMaterialIndex].ImportedMaterialSlotName;
 				break;
 			}
 		}
@@ -5507,35 +5512,35 @@ void FPersonaMeshDetails::OnSectionChanged(int32 LODIndex, int32 SectionIndex, i
 		Mesh->Modify();
 		{
 			FScopedSkeletalMeshPostEditChange ScopedPostEditChange(Mesh);
-		int32 NumSections = ImportedResource->LODModels[LODIndex].Sections.Num();
-		FSkeletalMeshLODInfo& Info = *(Mesh->GetLODInfo(LODIndex));
+			int32 NumSections = ImportedResource->LODModels[LODIndex].Sections.Num();
+			FSkeletalMeshLODInfo& Info = *(Mesh->GetLODInfo(LODIndex));
 
 			auto SetLODMaterialMapValue = [&LODIndex, &Info, &ImportedResource](int32 InSectionIndex, int32 OverrideMaterialIndex)
-		{
-				if (ImportedResource->LODModels[LODIndex].Sections[InSectionIndex].MaterialIndex == OverrideMaterialIndex)
 			{
-					if (Info.LODMaterialMap.IsValidIndex(InSectionIndex))
+				if (ImportedResource->LODModels[LODIndex].Sections[InSectionIndex].MaterialIndex == OverrideMaterialIndex)
 				{
+					if (Info.LODMaterialMap.IsValidIndex(InSectionIndex))
+					{
 						Info.LODMaterialMap[InSectionIndex] = INDEX_NONE;
+					}
 				}
-			}
-			else
-			{
-					while (Info.LODMaterialMap.Num() <= InSectionIndex)
+				else
 				{
-					Info.LODMaterialMap.Add(INDEX_NONE);
-				}
+					while (Info.LODMaterialMap.Num() <= InSectionIndex)
+					{
+						Info.LODMaterialMap.Add(INDEX_NONE);
+					}
 					check(InSectionIndex < Info.LODMaterialMap.Num());
 					Info.LODMaterialMap[InSectionIndex] = OverrideMaterialIndex;
-			}
-		};
+				}
+			};
 
 			SetLODMaterialMapValue(SectionIndex, NewSkeletalMaterialIndex);
-		//Set the chunked section 
+			//Set the chunked section 
 			for (int32 SectionIdx = SectionIndex+1; SectionIdx < NumSections; SectionIdx++)
-		{
-			if (ImportedResource->LODModels[LODIndex].Sections[SectionIdx].ChunkedParentSectionIndex == SectionIndex)
 			{
+				if (ImportedResource->LODModels[LODIndex].Sections[SectionIdx].ChunkedParentSectionIndex == SectionIndex)
+				{
 					SetLODMaterialMapValue(SectionIdx, NewSkeletalMaterialIndex);
 				}
 				else
@@ -5555,7 +5560,7 @@ void FPersonaMeshDetails::OnSectionChanged(int32 LODIndex, int32 SectionIndex, i
 //
 void FPersonaMeshDetails::CustomizeClothingProperties(IDetailLayoutBuilder& DetailLayout, IDetailCategoryBuilder& ClothingFilesCategory)
 {
-	TSharedRef<IPropertyHandle> ClothingAssetsProperty = DetailLayout.GetProperty(GET_MEMBER_NAME_CHECKED(USkeletalMesh, MeshClothingAssets), USkeletalMesh::StaticClass());
+	TSharedRef<IPropertyHandle> ClothingAssetsProperty = DetailLayout.GetProperty(USkeletalMesh::GetMeshClothingAssetsMemberName(), USkeletalMesh::StaticClass());
 
 	if( ClothingAssetsProperty->IsValidHandle() )
 	{
@@ -5651,7 +5656,7 @@ void FPersonaMeshDetails::OnGenerateElementForClothingAsset( TSharedRef<IPropert
 	];	
 
 	USkeletalMesh* CurrentMesh = GetPersonaToolkit()->GetMesh();
-	UClothingAssetBase* CurrentAsset = CurrentMesh->MeshClothingAssets[ElementIndex];
+	UClothingAssetBase* CurrentAsset = CurrentMesh->GetMeshClothingAssets()[ElementIndex];
 
 	ChildrenBuilder.AddCustomRow(LOCTEXT("ClothingAsset_Search_Name", "Name"))
 	.NameContent()
@@ -5719,7 +5724,7 @@ TSharedRef<SUniformGridPanel> FPersonaMeshDetails::MakeClothingDetailsWidget(int
 	TSharedRef<SUniformGridPanel> Grid = SNew(SUniformGridPanel).SlotPadding(2.0f);
 
 	USkeletalMesh* SkelMesh = GetPersonaToolkit()->GetMesh();
-	UClothingAssetBase* ClothingAsset = SkelMesh->MeshClothingAssets[AssetIndex];
+	UClothingAssetBase* ClothingAsset = SkelMesh->GetMeshClothingAssets()[AssetIndex];
 	if (!ClothingAsset)
 	{
 		return Grid;
@@ -5844,9 +5849,9 @@ FReply FPersonaMeshDetails::OnReimportApexFileClicked(int32 AssetIndex, IDetailL
 {
 	USkeletalMesh* SkelMesh = GetPersonaToolkit()->GetMesh();
 
-	check(SkelMesh && SkelMesh->MeshClothingAssets.IsValidIndex(AssetIndex));
+	check(SkelMesh && SkelMesh->GetMeshClothingAssets().IsValidIndex(AssetIndex));
 
-	UClothingAssetBase* AssetToReimport = SkelMesh->MeshClothingAssets[AssetIndex];
+	UClothingAssetBase* AssetToReimport = SkelMesh->GetMeshClothingAssets()[AssetIndex];
 	check(AssetToReimport);
 
 	FString ReimportPath = AssetToReimport->ImportedFilePath;
@@ -5904,16 +5909,16 @@ FReply FPersonaMeshDetails::OnRemoveClothingAssetClicked(int32 AssetIndex, IDeta
 
 	FScopedSuspendAlternateSkinWeightPreview ScopedSuspendAlternateSkinnWeightPreview(SkelMesh);
 	// Now we can remove the asset.
-	if(SkelMesh->MeshClothingAssets.IsValidIndex(AssetIndex))
+	if(SkelMesh->GetMeshClothingAssets().IsValidIndex(AssetIndex))
 	{
 		// Need to unregister our components so they shut down their current clothing simulation
 		FScopedSkeletalMeshPostEditChange ScopedPostEditChange(SkelMesh);
 
-		if (UClothingAssetBase* AssetToRemove = SkelMesh->MeshClothingAssets[AssetIndex])
+		if (UClothingAssetBase* AssetToRemove = SkelMesh->GetMeshClothingAssets()[AssetIndex])
 		{
 			AssetToRemove->UnbindFromSkeletalMesh(SkelMesh);
 		}
-		SkelMesh->MeshClothingAssets.RemoveAt(AssetIndex);
+		SkelMesh->GetMeshClothingAssets().RemoveAt(AssetIndex);
 
 		// Need to fix up asset indices on sections.
 		if(FSkeletalMeshModel* MeshResource = SkelMesh->GetImportedModel())
@@ -5972,10 +5977,10 @@ void FPersonaMeshDetails::UpdateClothingEntries()
 
 	NewClothingAssetEntries.Add(ClothingNoneEntry);
 
-	const int32 NumClothingAssets = Mesh->MeshClothingAssets.Num();
+	const int32 NumClothingAssets = Mesh->GetMeshClothingAssets().Num();
 	for(int32 Idx = 0; Idx < NumClothingAssets; ++Idx)
 	{
-		if (UClothingAssetBase* ClothingAsset = Mesh->MeshClothingAssets[Idx])
+		if (UClothingAssetBase* ClothingAsset = Mesh->GetMeshClothingAssets()[Idx])
 		{
 			const int32 NumAssetLods = ClothingAsset->GetNumLods();
 			for (int32 AssetLodIndex = 0; AssetLodIndex < NumAssetLods; ++AssetLodIndex)
@@ -6093,7 +6098,7 @@ void FPersonaMeshDetails::OnClothingSelectionChanged(TSharedPtr<FClothingEntry>
 			{
 				//Successful bind so set the SectionUserData
 				int32 AssetIndex = INDEX_NONE;
-				check(Mesh->MeshClothingAssets.Find(ClothingAsset, AssetIndex));
+				check(Mesh->GetMeshClothingAssets().Find(ClothingAsset, AssetIndex));
 				OriginalSectionData.CorrespondClothAssetIndex = AssetIndex;
 				OriginalSectionData.ClothingData.AssetGuid = ClothingAsset->GetAssetGuid();
 				OriginalSectionData.ClothingData.AssetLodIndex = InNewEntry->AssetLodIndex;
diff --git a/Engine/Source/Editor/Persona/Private/PersonaModule.cpp b/Engine/Source/Editor/Persona/Private/PersonaModule.cpp
--- a/Engine/Source/Editor/Persona/Private/PersonaModule.cpp
+++ b/Engine/Source/Editor/Persona/Private/PersonaModule.cpp
@@ -529,9 +529,9 @@ void FPersonaModule::TestSkeletonCurveNamesForUse(const TSharedRef<IEditableSkel
 		const FString SkeletonString = FAssetData(&Skeleton).GetExportTextName();
 
 		TArray<FAssetData> SkeletalMeshes;
-		PopulateWithAssets(USkeletalMesh::StaticClass()->GetFName(), GET_MEMBER_NAME_CHECKED(USkeletalMesh, Skeleton), SkeletonString, SkeletalMeshes);
+		PopulateWithAssets(USkeletalMesh::StaticClass()->GetFName(), USkeletalMesh::GetSkeletonMemberName(), SkeletonString, SkeletalMeshes);
 		TArray<FAssetData> Animations;
-		PopulateWithAssets(UAnimSequence::StaticClass()->GetFName(), FName("Skeleton"), SkeletonString, Animations);
+		PopulateWithAssets(UAnimSequence::StaticClass()->GetFName(), USkeletalMesh::GetSkeletonMemberName(), SkeletonString, Animations);
 
 		FText TimeTakenMessage = FText::Format(LOCTEXT("TimeTakenWarning", "In order to verify curve usage all Skeletal Meshes and Animations that use this skeleton will be loaded, this may take some time.\n\nProceed?\n\nNumber of Meshes: {0}\nNumber of Animations: {1}"), FText::AsNumber(SkeletalMeshes.Num()), FText::AsNumber(Animations.Num()));
 
@@ -566,14 +566,14 @@ void FPersonaModule::TestSkeletonCurveNamesForUse(const TSharedRef<IEditableSkel
 					const USkeletalMesh* Mesh = Cast<USkeletalMesh>(SkeletalMeshes[MeshIdx].GetAsset());
 
 					// Filter morph targets from curves
-					const TArray<UMorphTarget*>& MorphTargets = Mesh->MorphTargets;
+					const TArray<UMorphTarget*>& MorphTargets = Mesh->GetMorphTargets();
 					for (int32 I = 0; I < MorphTargets.Num(); ++I)
 					{
 						const int32 CurveIndex = UnusedNames.RemoveSingleSwap(MorphTargets[I]->GetFName(), false);
 					}
 
 					// Filter material params from curves
-					for (const FSkeletalMaterial& Mat : Mesh->Materials)
+					for (const FSkeletalMaterial& Mat : Mesh->GetMaterials())
 					{
 						if (UnusedNames.Num() == 0)
 						{
diff --git a/Engine/Source/Editor/Persona/Private/PersonaPreviewSceneAnimationController.cpp b/Engine/Source/Editor/Persona/Private/PersonaPreviewSceneAnimationController.cpp
--- a/Engine/Source/Editor/Persona/Private/PersonaPreviewSceneAnimationController.cpp
+++ b/Engine/Source/Editor/Persona/Private/PersonaPreviewSceneAnimationController.cpp
@@ -22,7 +22,7 @@ IDetailPropertyRow* UPersonaPreviewSceneAnimationController::AddPreviewControlle
 {
 	TArray<UObject*> ListOfPreviewController{ this };
 
-	const USkeleton* Skeleton = PersonaToolkit->GetPreviewMeshComponent()->SkeletalMesh ? PersonaToolkit->GetPreviewMeshComponent()->SkeletalMesh->Skeleton : nullptr;
+	const USkeleton* Skeleton = PersonaToolkit->GetPreviewMeshComponent()->SkeletalMesh ? PersonaToolkit->GetPreviewMeshComponent()->SkeletalMesh->GetSkeleton() : nullptr;
 	if (Skeleton)
 	{
 		FString SkeletonName = FAssetData(Skeleton).GetExportTextName();
diff --git a/Engine/Source/Editor/Persona/Private/PersonaToolkit.cpp b/Engine/Source/Editor/Persona/Private/PersonaToolkit.cpp
--- a/Engine/Source/Editor/Persona/Private/PersonaToolkit.cpp
+++ b/Engine/Source/Editor/Persona/Private/PersonaToolkit.cpp
@@ -193,7 +193,7 @@ void FPersonaToolkit::SetMesh(class USkeletalMesh* InSkeletalMesh)
 {
 	if (InSkeletalMesh != nullptr && Skeleton.IsValid())
 	{
-		check(InSkeletalMesh->Skeleton == Skeleton);
+		check(InSkeletalMesh->GetSkeleton() == Skeleton);
 	}
 
 	Mesh = InSkeletalMesh;
@@ -266,7 +266,7 @@ void FPersonaToolkit::SetPreviewMesh(class USkeletalMesh* InSkeletalMesh, bool b
 	{
 		// If the skeleton itself is changing, then we need to re-open the asset editor
 		bool bReOpenEditor = false;
-		if(InSkeletalMesh != nullptr && EditableSkeleton.IsValid() && InSkeletalMesh->Skeleton != &EditableSkeleton->GetSkeleton())
+		if(InSkeletalMesh != nullptr && EditableSkeleton.IsValid() && InSkeletalMesh->GetSkeleton() != &EditableSkeleton->GetSkeleton())
 		{
 			bReOpenEditor = true;
 			bSetPreviewMeshInAsset = true;
diff --git a/Engine/Source/Editor/Persona/Private/PreviewSceneCustomizations.cpp b/Engine/Source/Editor/Persona/Private/PreviewSceneCustomizations.cpp
--- a/Engine/Source/Editor/Persona/Private/PreviewSceneCustomizations.cpp
+++ b/Engine/Source/Editor/Persona/Private/PreviewSceneCustomizations.cpp
@@ -229,7 +229,7 @@ void FPreviewSceneDescriptionCustomization::CustomizeDetails(IDetailLayoutBuilde
 			SNew(SObjectPropertyEntryBox)
 			.AllowedClass(USkeletalMesh::StaticClass())
 			.PropertyHandle(SkeletalMeshProperty)
-			.OnShouldFilterAsset(this, &FPreviewSceneDescriptionCustomization::HandleShouldFilterAsset, FName("Skeleton"), PersonaToolkit.Pin()->GetContext() == UPhysicsAsset::StaticClass()->GetFName())
+			.OnShouldFilterAsset(this, &FPreviewSceneDescriptionCustomization::HandleShouldFilterAsset, USkeletalMesh::GetSkeletonMemberName(), PersonaToolkit.Pin()->GetContext() == UPhysicsAsset::StaticClass()->GetFName())
 			.OnObjectChanged(this, &FPreviewSceneDescriptionCustomization::HandleMeshChanged)
 			.ThumbnailPool(DetailBuilder.GetThumbnailPool())
 		];
@@ -513,7 +513,7 @@ bool FPreviewSceneDescriptionCustomization::HandleShouldFilterAdditionalMesh(con
 		return true;
 	}
 
-	return HandleShouldFilterAsset(InAssetData, FName("Skeleton"), bCanUseDifferentSkeleton);
+	return HandleShouldFilterAsset(InAssetData, USkeletalMesh::GetSkeletonMemberName(), bCanUseDifferentSkeleton);
 }
 
 bool FPreviewSceneDescriptionCustomization::HandleShouldFilterAsset(const FAssetData& InAssetData, FName InTag, bool bCanUseDifferentSkeleton)
diff --git a/Engine/Source/Editor/Persona/Private/SAnimViewportToolBar.cpp b/Engine/Source/Editor/Persona/Private/SAnimViewportToolBar.cpp
--- a/Engine/Source/Editor/Persona/Private/SAnimViewportToolBar.cpp
+++ b/Engine/Source/Editor/Persona/Private/SAnimViewportToolBar.cpp
@@ -447,7 +447,7 @@ TSharedRef<SWidget> SAnimViewportToolBar::MakeFollowBoneWidget(TWeakPtr<SComboBu
 				USkeletalMesh* PreviewMesh = Viewport.Pin()->GetPreviewScene()->GetPreviewMesh();
 				if (PreviewMesh)
 				{
-					return PreviewMesh->RefSkeleton;
+					return PreviewMesh->GetRefSkeleton();
 				}
 
 				static FReferenceSkeleton EmptySkeleton;
diff --git a/Engine/Source/Editor/Persona/Private/SAnimationEditorViewport.cpp b/Engine/Source/Editor/Persona/Private/SAnimationEditorViewport.cpp
--- a/Engine/Source/Editor/Persona/Private/SAnimationEditorViewport.cpp
+++ b/Engine/Source/Editor/Persona/Private/SAnimationEditorViewport.cpp
@@ -2145,9 +2145,9 @@ void SAnimationEditorViewportTabBody::AddPostProcessNotification()
 	{
 		if (UDebugSkelMeshComponent* PreviewComponent = GetPreviewScene()->GetPreviewMeshComponent())
 		{
-			if(PreviewComponent->SkeletalMesh && PreviewComponent->SkeletalMesh->PostProcessAnimBlueprint && PreviewComponent->SkeletalMesh->PostProcessAnimBlueprint->ClassGeneratedBy)
+			if(PreviewComponent->SkeletalMesh && PreviewComponent->SkeletalMesh->GetPostProcessAnimBlueprint() && PreviewComponent->SkeletalMesh->GetPostProcessAnimBlueprint()->ClassGeneratedBy)
 			{
-				return FText::FromString(PreviewComponent->SkeletalMesh->PostProcessAnimBlueprint->ClassGeneratedBy->GetName());
+				return FText::FromString(PreviewComponent->SkeletalMesh->GetPostProcessAnimBlueprint()->ClassGeneratedBy->GetName());
 			}
 		}
 
@@ -2196,9 +2196,9 @@ void SAnimationEditorViewportTabBody::AddPostProcessNotification()
 	{
 		if (UDebugSkelMeshComponent* PreviewComponent = GetPreviewScene()->GetPreviewMeshComponent())
 		{
-			if(PreviewComponent->SkeletalMesh && PreviewComponent->SkeletalMesh->PostProcessAnimBlueprint)
+			if(PreviewComponent->SkeletalMesh && PreviewComponent->SkeletalMesh->GetPostProcessAnimBlueprint())
 			{
-				GEditor->GetEditorSubsystem<UAssetEditorSubsystem>()->OpenEditorForAssets(TArray<UObject*>({ PreviewComponent->SkeletalMesh->PostProcessAnimBlueprint->ClassGeneratedBy }));
+				GEditor->GetEditorSubsystem<UAssetEditorSubsystem>()->OpenEditorForAssets(TArray<UObject*>({ PreviewComponent->SkeletalMesh->GetPostProcessAnimBlueprint()->ClassGeneratedBy }));
 			}
 		}
 
@@ -2308,7 +2308,7 @@ void SAnimationEditorViewportTabBody::AddMinLODNotification()
 
 	auto GetMinLODNotificationVisibility = [this]()
 	{
-		if (GetPreviewScene()->GetPreviewMesh() && GetPreviewScene()->GetPreviewMesh()->MinLod.Default != 0)
+		if (GetPreviewScene()->GetPreviewMesh() && GetPreviewScene()->GetPreviewMesh()->GetMinLod().Default != 0)
 		{
 			return EVisibility::Visible;
 		}
diff --git a/Engine/Source/Editor/Persona/Private/SCreateClothingSettingsPanel.cpp b/Engine/Source/Editor/Persona/Private/SCreateClothingSettingsPanel.cpp
--- a/Engine/Source/Editor/Persona/Private/SCreateClothingSettingsPanel.cpp
+++ b/Engine/Source/Editor/Persona/Private/SCreateClothingSettingsPanel.cpp
@@ -62,7 +62,7 @@ void SCreateClothingSettingsPanel::Construct(const FArguments& InArgs)
 	StructureDetailsView->GetDetailsView()->SetGenericLayoutDetailsDelegate(FOnGetDetailCustomizationInstance::CreateStatic(&FClothCreateSettingsCustomization::MakeInstance, InArgs._Mesh, InArgs._bIsSubImport));
 
 	BuildParams.AssetName = InArgs._MeshName + TEXT("_Clothing");
-	BuildParams.PhysicsAsset = InArgs._Mesh->PhysicsAsset;
+	BuildParams.PhysicsAsset = InArgs._Mesh->GetPhysicsAsset();
 
 	FStructOnScope* Struct = new FStructOnScope(FSkeletalMeshClothBuildParams::StaticStruct(), (uint8*)&BuildParams);
 	StructureDetailsView->SetStructureData(MakeShareable(Struct));
@@ -149,10 +149,10 @@ TSharedRef<SWidget> FClothCreateSettingsCustomization::OnGetTargetAssetMenu()
 	{
 		if(USkeletalMesh* Mesh = MeshPtr.Get())
 		{
-			const int32 NumCloths = Mesh->MeshClothingAssets.Num();
+			const int32 NumCloths = Mesh->GetMeshClothingAssets().Num();
 			for(int32 ClothIndex = 0 ; ClothIndex < NumCloths ; ++ClothIndex)
 			{
-				if (UClothingAssetBase* ClothingAsset = Mesh->MeshClothingAssets[ClothIndex])
+				if (UClothingAssetBase* ClothingAsset = Mesh->GetMeshClothingAssets()[ClothIndex])
 				{
 					FUIAction Action;
 					Action.ExecuteAction = FExecuteAction::CreateSP(this, &FClothCreateSettingsCustomization::OnAssetSelected, ClothIndex);
@@ -186,12 +186,12 @@ void FClothCreateSettingsCustomization::OnAssetSelected(int32 InMeshClothingInde
 
 	if(USkeletalMesh* Mesh = MeshPtr.Get())
 	{
-		if(Mesh->MeshClothingAssets.IsValidIndex(InMeshClothingIndex) && 
-		   Mesh->MeshClothingAssets[InMeshClothingIndex])
+		if(Mesh->GetMeshClothingAssets().IsValidIndex(InMeshClothingIndex) &&
+		   Mesh->GetMeshClothingAssets()[InMeshClothingIndex])
 		{
-			if(ParamsStruct->TargetAsset != Mesh->MeshClothingAssets[InMeshClothingIndex])
+			if(ParamsStruct->TargetAsset != Mesh->GetMeshClothingAssets()[InMeshClothingIndex])
 			{
-				ParamsStruct->TargetAsset = Mesh->MeshClothingAssets[InMeshClothingIndex];
+				ParamsStruct->TargetAsset = Mesh->GetMeshClothingAssets()[InMeshClothingIndex];
 				ParamsStruct->TargetLod = INDEX_NONE;
 			}
 		}
diff --git a/Engine/Source/Editor/Persona/Private/SMorphTargetViewer.cpp b/Engine/Source/Editor/Persona/Private/SMorphTargetViewer.cpp
--- a/Engine/Source/Editor/Persona/Private/SMorphTargetViewer.cpp
+++ b/Engine/Source/Editor/Persona/Private/SMorphTargetViewer.cpp
@@ -465,7 +465,7 @@ void SMorphTargetViewer::CreateMorphTargetList( const FString& SearchText )
 	if ( SkeletalMesh )
 	{
 		UDebugSkelMeshComponent* MeshComponent = PreviewScenePtr.Pin()->GetPreviewMeshComponent();
-		TArray<UMorphTarget*>& MorphTargets = SkeletalMesh->MorphTargets;
+		TArray<UMorphTarget*>& MorphTargets = SkeletalMesh->GetMorphTargets();
 
 		bool bDoFiltering = !SearchText.IsEmpty();
 
diff --git a/Engine/Source/Editor/Persona/Private/SPoseEditor.cpp b/Engine/Source/Editor/Persona/Private/SPoseEditor.cpp
--- a/Engine/Source/Editor/Persona/Private/SPoseEditor.cpp
+++ b/Engine/Source/Editor/Persona/Private/SPoseEditor.cpp
@@ -462,7 +462,7 @@ void RestartAnimations(const USkeleton* CurrentSkeleton)
 	for (FObjectIterator Iter(USkeletalMeshComponent::StaticClass()); Iter; ++Iter)
 	{
 		USkeletalMeshComponent* SkeletalMeshComponent = Cast<USkeletalMeshComponent>(*Iter);
-		if (SkeletalMeshComponent->SkeletalMesh && SkeletalMeshComponent->SkeletalMesh->Skeleton == CurrentSkeleton)
+		if (SkeletalMeshComponent->SkeletalMesh && SkeletalMeshComponent->SkeletalMesh->GetSkeleton() == CurrentSkeleton)
 		{
 			SkeletalMeshComponent->InitAnim(true);
 		}
diff --git a/Engine/Source/Editor/Persona/Private/SRetargetManager.cpp b/Engine/Source/Editor/Persona/Private/SRetargetManager.cpp
--- a/Engine/Source/Editor/Persona/Private/SRetargetManager.cpp
+++ b/Engine/Source/Editor/Persona/Private/SRetargetManager.cpp
@@ -343,14 +343,14 @@ void SRetargetManager::ResetRetargetBasePose()
 	{
 		USkeletalMesh * PreviewMesh = PreviewMeshComp->SkeletalMesh;
 
-		check(PreviewMesh && &EditableSkeletonPtr.Pin()->GetSkeleton() == PreviewMesh->Skeleton);
+		check(PreviewMesh && &EditableSkeletonPtr.Pin()->GetSkeleton() == PreviewMesh->GetSkeleton());
 
 		if(PreviewMesh)
 		{
 			const FScopedTransaction Transaction(LOCTEXT("ResetRetargetBasePose_Action", "Reset Retarget Base Pose"));
 			PreviewMesh->Modify();
 			// reset to original ref pose
-			PreviewMesh->RetargetBasePose = PreviewMesh->RefSkeleton.GetRefBonePose();
+			PreviewMesh->SetRetargetBasePose(PreviewMesh->GetRefSkeleton().GetRefBonePose());
 			TurnOnPreviewRetargetBasePose();
 		}
 	}
@@ -365,7 +365,7 @@ void SRetargetManager::UseCurrentPose()
 	{
 		USkeletalMesh * PreviewMesh = PreviewMeshComp->SkeletalMesh;
 
-		check(PreviewMesh && &EditableSkeletonPtr.Pin()->GetSkeleton() == PreviewMesh->Skeleton);
+		check(PreviewMesh && &EditableSkeletonPtr.Pin()->GetSkeleton() == PreviewMesh->GetSkeleton());
 
 		if (PreviewMesh)
 		{
@@ -374,11 +374,11 @@ void SRetargetManager::UseCurrentPose()
 			// get space bases and calculate local
 			const TArray<FTransform> & SpaceBases = PreviewMeshComp->GetComponentSpaceTransforms();
 			// @todo check to see if skeleton vs preview mesh makes it different for missing bones
-			const FReferenceSkeleton& RefSkeleton = PreviewMesh->RefSkeleton;
-			TArray<FTransform> & NewRetargetBasePose = PreviewMesh->RetargetBasePose;
+			const FReferenceSkeleton& RefSkeleton = PreviewMesh->GetRefSkeleton();
+			TArray<FTransform> & NewRetargetBasePose = PreviewMesh->GetRetargetBasePose();
 			// if you're using master pose component in preview, this won't work
-			check(PreviewMesh->RefSkeleton.GetNum() == SpaceBases.Num());
-			int32 TotalNumBones = PreviewMesh->RefSkeleton.GetNum();
+			check(PreviewMesh->GetRefSkeleton().GetNum() == SpaceBases.Num());
+			int32 TotalNumBones = PreviewMesh->GetRefSkeleton().GetNum();
 			NewRetargetBasePose.Empty(TotalNumBones);
 			NewRetargetBasePose.AddUninitialized(TotalNumBones);
 
@@ -445,7 +445,7 @@ void SRetargetManager::ImportPose(const UPoseAsset* PoseAsset, const FName& Pose
 			{
 				USkeletalMesh * PreviewMesh = PreviewMeshComp->SkeletalMesh;
 
-				check(PreviewMesh && &EditableSkeletonPtr.Pin()->GetSkeleton() == PreviewMesh->Skeleton);
+				check(PreviewMesh && &EditableSkeletonPtr.Pin()->GetSkeleton() == PreviewMesh->GetSkeleton());
 
 				if (PreviewMesh)
 				{
@@ -453,13 +453,13 @@ void SRetargetManager::ImportPose(const UPoseAsset* PoseAsset, const FName& Pose
 					PreviewMesh->Modify();
 
 					// reset to original ref pose first
-					PreviewMesh->RetargetBasePose = PreviewMesh->RefSkeleton.GetRefBonePose();
+					PreviewMesh->SetRetargetBasePose(PreviewMesh->GetRefSkeleton().GetRefBonePose());
 
 					// now override imported pose
 					for (int32 TrackIndex = 0; TrackIndex < PoseTrackNames.Num(); ++TrackIndex)
 					{
-						int32 BoneIndex = PreviewMesh->RefSkeleton.FindBoneIndex(PoseTrackNames[TrackIndex]);
-						PreviewMesh->RetargetBasePose[BoneIndex] = PoseTransforms[TrackIndex];
+						int32 BoneIndex = PreviewMesh->GetRefSkeleton().FindBoneIndex(PoseTrackNames[TrackIndex]);
+						PreviewMesh->GetRetargetBasePose()[BoneIndex] = PoseTransforms[TrackIndex];
 					}
 
 					TurnOnPreviewRetargetBasePose();
diff --git a/Engine/Source/Editor/Persona/Private/SRetargetSourceWindow.cpp b/Engine/Source/Editor/Persona/Private/SRetargetSourceWindow.cpp
--- a/Engine/Source/Editor/Persona/Private/SRetargetSourceWindow.cpp
+++ b/Engine/Source/Editor/Persona/Private/SRetargetSourceWindow.cpp
@@ -455,7 +455,7 @@ void SRetargetSourceWindow::OnAddRetargetSource()
 	const USkeleton& Skeleton = EditableSkeletonPtr.Pin()->GetSkeleton();
 
 	FString SkeletonString = FAssetData(&Skeleton).GetExportTextName();
-	AssetPickerConfig.Filter.TagsAndValues.Add(GET_MEMBER_NAME_CHECKED(USkeletalMesh, Skeleton), SkeletonString);
+	AssetPickerConfig.Filter.TagsAndValues.Add(USkeletalMesh::GetSkeletonMemberName(), SkeletonString);
 
 	TSharedRef<SWidget> Widget = SNew(SBox)
 		.WidthOverride(384)
diff --git a/Engine/Source/Editor/Persona/Private/SRigWindow.cpp b/Engine/Source/Editor/Persona/Private/SRigWindow.cpp
--- a/Engine/Source/Editor/Persona/Private/SRigWindow.cpp
+++ b/Engine/Source/Editor/Persona/Private/SRigWindow.cpp
@@ -306,7 +306,7 @@ const struct FReferenceSkeleton& SRigWindow::GetReferenceSkeleton() const
 	UDebugSkelMeshComponent* PreviewMeshComp = PreviewScenePtr.Pin()->GetPreviewMeshComponent();
 	USkeletalMesh* PreviewMesh = (PreviewMeshComp) ? PreviewMeshComp->SkeletalMesh : nullptr;
 	// it's because retarget base pose leaves in mesh, so if you give ref skeleton of skeleton, you might have joint that your mesh doesn't have
-	return (PreviewMesh)? PreviewMesh->RefSkeleton : EditableSkeletonPtr.Pin()->GetSkeleton().GetReferenceSkeleton();
+	return (PreviewMesh)? PreviewMesh->GetRefSkeleton() : EditableSkeletonPtr.Pin()->GetSkeleton().GetReferenceSkeleton();
 }
 
 bool SRigWindow::OnTargetSkeletonSelected(USkeleton* SelectedSkeleton, URig*  Rig) const
diff --git a/Engine/Source/Editor/Persona/Private/SkeletalMeshSocketDetails.cpp b/Engine/Source/Editor/Persona/Private/SkeletalMeshSocketDetails.cpp
--- a/Engine/Source/Editor/Persona/Private/SkeletalMeshSocketDetails.cpp
+++ b/Engine/Source/Editor/Persona/Private/SkeletalMeshSocketDetails.cpp
@@ -119,7 +119,7 @@ void FSkeletalMeshSocketDetails::OnParentBoneNameCommitted(const FText& InSearch
 	{
 		if( SkeletalMesh )
 		{
-			Skeleton = SkeletalMesh->Skeleton;
+			Skeleton = SkeletalMesh->GetSkeleton();
 		}
 	}
 	if( Skeleton )
@@ -145,7 +145,7 @@ void FSkeletalMeshSocketDetails::OnSocketNameChanged(const FText& InSearchText)
 	{
 		if (SkeletalMesh)
 		{
-			Skeleton = SkeletalMesh->Skeleton;
+			Skeleton = SkeletalMesh->GetSkeleton();
 		}
 	}
 	if (Skeleton)
@@ -174,7 +174,7 @@ void FSkeletalMeshSocketDetails::OnSocketNameCommitted(const FText& InSearchText
 	{
 		if (SkeletalMesh)
 		{
-			Skeleton = SkeletalMesh->Skeleton;
+			Skeleton = SkeletalMesh->GetSkeleton();
 		}
 	}
 	if (Skeleton)
diff --git a/Engine/Source/Editor/PhysicsAssetEditor/Private/PhysicsAssetEditor.cpp b/Engine/Source/Editor/PhysicsAssetEditor/Private/PhysicsAssetEditor.cpp
--- a/Engine/Source/Editor/PhysicsAssetEditor/Private/PhysicsAssetEditor.cpp
+++ b/Engine/Source/Editor/PhysicsAssetEditor/Private/PhysicsAssetEditor.cpp
@@ -1463,7 +1463,7 @@ bool FPhysicsAssetEditor::ShouldFilterAssetBasedOnSkeleton( const FAssetData& As
 		USkeletalMesh* EditorSkelMesh = SharedData->PhysicsAsset->GetPreviewMesh();
 		if(EditorSkelMesh != nullptr)
 		{
-			USkeleton* Skeleton = EditorSkelMesh->Skeleton;
+			USkeleton* Skeleton = EditorSkelMesh->GetSkeleton();
 
 			if ( Skeleton && (*SkeletonName) == FString::Printf(TEXT("%s'%s'"), *Skeleton->GetClass()->GetName(), *Skeleton->GetPathName()) )
 			{
@@ -1838,7 +1838,7 @@ void FPhysicsAssetEditor::ResetBoneCollision()
 			SlowTask.EnterProgressFrame(1.0f, FText::Format(LOCTEXT("ResetCollsionStepInfo", "Generating collision for {0}"), FText::FromName(BodySetup->BoneName)));
 			BodySetup->Modify();
 
-			int32 BoneIndex = EditorSkelMesh->RefSkeleton.FindBoneIndex(BodySetup->BoneName);
+			int32 BoneIndex = EditorSkelMesh->GetRefSkeleton().FindBoneIndex(BodySetup->BoneName);
 			check(BoneIndex != INDEX_NONE);
 
 			const FBoneVertInfo& UseVertInfo = NewBodyData.VertWeight == EVW_DominantWeight ? SharedData->DominantWeightBoneInfos[BoneIndex] : SharedData->AnyWeightBoneInfos[BoneIndex];
diff --git a/Engine/Source/Editor/PhysicsAssetEditor/Private/PhysicsAssetEditorSharedData.cpp b/Engine/Source/Editor/PhysicsAssetEditor/Private/PhysicsAssetEditorSharedData.cpp
--- a/Engine/Source/Editor/PhysicsAssetEditor/Private/PhysicsAssetEditorSharedData.cpp
+++ b/Engine/Source/Editor/PhysicsAssetEditor/Private/PhysicsAssetEditorSharedData.cpp
@@ -149,7 +149,7 @@ void FPhysicsAssetEditorSharedData::Initialize(const TSharedRef<IPersonaPreviewS
 				continue;
 			}
 			FName BoneName = PhysicsAsset->SkeletalBodySetups[i]->BoneName;
-			int32 BoneIndex = EditorSkelMesh->RefSkeleton.FindBoneIndex(BoneName);
+			int32 BoneIndex = EditorSkelMesh->GetRefSkeleton().FindBoneIndex(BoneName);
 			if (BoneIndex == INDEX_NONE)
 			{
 				MissingBodyIndices.Add( i );
@@ -203,7 +203,7 @@ void FPhysicsAssetEditorSharedData::CachePreviewMesh()
 				LOCTEXT("Error_PhysicsAssetHasNoSkelMesh", "Warning: Physics Asset has no skeletal mesh assigned.\nFor now, a simple default skeletal mesh ({0}) will be used.\nYou can fix this by opening the asset and choosing another skeletal mesh from the toolbar."),
 				FText::FromString(PreviewMesh->GetFullName())));
 	}
-	else if(PreviewMesh->Skeleton == nullptr)
+	else if(PreviewMesh->GetSkeleton() == nullptr)
 	{
 		// Fall back in the case of a deleted skeleton
 		PreviewMesh = (USkeletalMesh*)StaticLoadObject(USkeletalMesh::StaticClass(), NULL, TEXT("/Engine/EngineMeshes/SkeletalCube.SkeletalCube"), NULL, LOAD_None, NULL);
@@ -281,7 +281,7 @@ void FPhysicsAssetEditorSharedData::Mirror()
 
 		for (FMirrorInfo & MirrorInfo : MirrorInfos)	//mirror all selected bodies/constraints
 		{
-			int32 BoneIndex = EditorSkelMesh->RefSkeleton.FindBoneIndex(MirrorInfo.BoneName);
+			int32 BoneIndex = EditorSkelMesh->GetRefSkeleton().FindBoneIndex(MirrorInfo.BoneName);
 
 			int32 MirrorBoneIndex = PhysicsAsset->FindMirroredBone(EditorSkelMesh, BoneIndex);
 			if (MirrorBoneIndex != INDEX_NONE)
@@ -1112,15 +1112,15 @@ bool FPhysicsAssetEditorSharedData::WeldSelectedBodies(bool bWeld /* = true */)
 	const FSelection& Body1 = SelectedBodies[BodyIndex1];
 
 	FName Bone0Name = PhysicsAsset->SkeletalBodySetups[Body0.Index]->BoneName;
-	int32 Bone0Index = EditorSkelMesh->RefSkeleton.FindBoneIndex(Bone0Name);
+	int32 Bone0Index = EditorSkelMesh->GetRefSkeleton().FindBoneIndex(Bone0Name);
 	check(Bone0Index != INDEX_NONE);
 
 	FName Bone1Name = PhysicsAsset->SkeletalBodySetups[Body1.Index]->BoneName;
-	int32 Bone1Index = EditorSkelMesh->RefSkeleton.FindBoneIndex(Bone1Name);
+	int32 Bone1Index = EditorSkelMesh->GetRefSkeleton().FindBoneIndex(Bone1Name);
 	check(Bone1Index != INDEX_NONE);
 
-	int32 Bone0ParentIndex = EditorSkelMesh->RefSkeleton.GetParentIndex(Bone0Index);
-	int32 Bone1ParentIndex = EditorSkelMesh->RefSkeleton.GetParentIndex(Bone1Index);
+	int32 Bone0ParentIndex = EditorSkelMesh->GetRefSkeleton().GetParentIndex(Bone0Index);
+	int32 Bone1ParentIndex = EditorSkelMesh->GetRefSkeleton().GetParentIndex(Bone1Index);
 
 	int ParentBodyIndex = INDEX_NONE;
 	int ChildBodyIndex = INDEX_NONE;
@@ -1238,7 +1238,7 @@ void FPhysicsAssetEditorSharedData::MakeNewBody(int32 NewBoneIndex, bool bAutoSe
 	}
 	PhysicsAsset->Modify();
 
-	FName NewBoneName = EditorSkelMesh->RefSkeleton.GetBoneName(NewBoneIndex);
+	FName NewBoneName = EditorSkelMesh->GetRefSkeleton().GetBoneName(NewBoneIndex);
 
 	// If this body is already physical, remove the current body
 	int32 NewBodyIndex = PhysicsAsset->FindBodyIndex(NewBoneName);
@@ -1277,11 +1277,11 @@ void FPhysicsAssetEditorSharedData::MakeNewBody(int32 NewBoneIndex, bool bAutoSe
 	}
 
 	// Check if the bone of the new body has any physical children bones
-	for (int32 i = 0; i < EditorSkelMesh->RefSkeleton.GetRawBoneNum(); ++i)
+	for (int32 i = 0; i < EditorSkelMesh->GetRefSkeleton().GetRawBoneNum(); ++i)
 	{
-		if (EditorSkelMesh->RefSkeleton.BoneIsChildOf(i, NewBoneIndex))
+		if (EditorSkelMesh->GetRefSkeleton().BoneIsChildOf(i, NewBoneIndex))
 		{
-			const int32 ChildBodyIndex = PhysicsAsset->FindBodyIndex(EditorSkelMesh->RefSkeleton.GetBoneName(i));
+			const int32 ChildBodyIndex = PhysicsAsset->FindBodyIndex(EditorSkelMesh->GetRefSkeleton().GetBoneName(i));
 			
 			// If the child bone is physical, it may require fixing up in regards to constraints
 			if (ChildBodyIndex != INDEX_NONE)
@@ -1311,11 +1311,11 @@ void FPhysicsAssetEditorSharedData::MakeNewBody(int32 NewBoneIndex, bool bAutoSe
 					UPhysicsConstraintTemplate* ExistingConstraintSetup = PhysicsAsset->ConstraintSetup[ ConstraintIndex ];
 					check(ExistingConstraintSetup);
 					
-					const int32 ExistingConstraintBoneIndex = EditorSkelMesh->RefSkeleton.FindBoneIndex(ExistingConstraintSetup->DefaultInstance.ConstraintBone2);
+					const int32 ExistingConstraintBoneIndex = EditorSkelMesh->GetRefSkeleton().FindBoneIndex(ExistingConstraintSetup->DefaultInstance.ConstraintBone2);
 					check(ExistingConstraintBoneIndex != INDEX_NONE);
 
 					// If the constraint exists between two child bones, then no fix up is required
-					if (EditorSkelMesh->RefSkeleton.BoneIsChildOf(ExistingConstraintBoneIndex, NewBoneIndex))
+					if (EditorSkelMesh->GetRefSkeleton().BoneIsChildOf(ExistingConstraintBoneIndex, NewBoneIndex))
 					{
 						continue;
 					}
@@ -1411,7 +1411,7 @@ void FPhysicsAssetEditorSharedData::SetConstraintRelTM(const FPhysicsAssetEditor
 	ConstraintSetup->Modify();
 
 	// Get child bone transform
-	int32 BoneIndex = EditorSkelMesh->RefSkeleton.FindBoneIndex(ConstraintSetup->DefaultInstance.ConstraintBone1);
+	int32 BoneIndex = EditorSkelMesh->GetRefSkeleton().FindBoneIndex(ConstraintSetup->DefaultInstance.ConstraintBone1);
 	if (BoneIndex != INDEX_NONE)
 	{
 		FTransform BoneTM = EditorSkelComp->GetBoneTransform(BoneIndex);
@@ -1436,8 +1436,8 @@ void FPhysicsAssetEditorSharedData::SnapConstraintToBone(FConstraintInstance& Co
 		return;
 	}
 
-	const int32 BoneIndex1 = EditorSkelMesh->RefSkeleton.FindBoneIndex(ConstraintInstance.ConstraintBone1);
-	const int32 BoneIndex2 = EditorSkelMesh->RefSkeleton.FindBoneIndex(ConstraintInstance.ConstraintBone2);
+	const int32 BoneIndex1 = EditorSkelMesh->GetRefSkeleton().FindBoneIndex(ConstraintInstance.ConstraintBone1);
+	const int32 BoneIndex2 = EditorSkelMesh->GetRefSkeleton().FindBoneIndex(ConstraintInstance.ConstraintBone2);
 
 	check(BoneIndex1 != INDEX_NONE);
 	check(BoneIndex2 != INDEX_NONE);
@@ -1619,7 +1619,7 @@ void FPhysicsAssetEditorSharedData::DeleteBody(int32 DelBodyIndex, bool bRefresh
 	TArray<int32> NearestBodiesBelow;
 	PhysicsAsset->GetNearestBodyIndicesBelow(NearestBodiesBelow, BodySetup->BoneName, EditorSkelMesh);
 	
-	int32 BoneIndex = EditorSkelMesh->RefSkeleton.FindBoneIndex(BodySetup->BoneName);
+	int32 BoneIndex = EditorSkelMesh->GetRefSkeleton().FindBoneIndex(BodySetup->BoneName);
 
 	if (BoneIndex != INDEX_NONE)	//it's possible to delete bodies that have no bones. In this case just ignore all of this fixup code
 	{
@@ -1771,11 +1771,11 @@ FTransform FPhysicsAssetEditorSharedData::GetConstraintBodyTM(const UPhysicsCons
 	int32 BoneIndex;
 	if (Frame == EConstraintFrame::Frame1)
 	{
-		BoneIndex = EditorSkelMesh->RefSkeleton.FindBoneIndex(ConstraintSetup->DefaultInstance.ConstraintBone1);
+		BoneIndex = EditorSkelMesh->GetRefSkeleton().FindBoneIndex(ConstraintSetup->DefaultInstance.ConstraintBone1);
 	}
 	else
 	{
-		BoneIndex = EditorSkelMesh->RefSkeleton.FindBoneIndex(ConstraintSetup->DefaultInstance.ConstraintBone2);
+		BoneIndex = EditorSkelMesh->GetRefSkeleton().FindBoneIndex(ConstraintSetup->DefaultInstance.ConstraintBone2);
 	}
 
 	// If we couldn't find the bone - fall back to identity.
@@ -1811,11 +1811,11 @@ FTransform FPhysicsAssetEditorSharedData::GetConstraintWorldTM(const UPhysicsCon
 	FTransform LFrame = ConstraintSetup->DefaultInstance.GetRefFrame(Frame);
 	if (Frame == EConstraintFrame::Frame1)
 	{
-		BoneIndex = EditorSkelMesh->RefSkeleton.FindBoneIndex(ConstraintSetup->DefaultInstance.ConstraintBone1);
+		BoneIndex = EditorSkelMesh->GetRefSkeleton().FindBoneIndex(ConstraintSetup->DefaultInstance.ConstraintBone1);
 	}
 	else
 	{
-		BoneIndex = EditorSkelMesh->RefSkeleton.FindBoneIndex(ConstraintSetup->DefaultInstance.ConstraintBone2);
+		BoneIndex = EditorSkelMesh->GetRefSkeleton().FindBoneIndex(ConstraintSetup->DefaultInstance.ConstraintBone2);
 	}
 
 	// If we couldn't find the bone - fall back to identity.
diff --git a/Engine/Source/Editor/PhysicsAssetEditor/Private/PhysicsAssetEditorSkeletonTreeBuilder.cpp b/Engine/Source/Editor/PhysicsAssetEditor/Private/PhysicsAssetEditorSkeletonTreeBuilder.cpp
--- a/Engine/Source/Editor/PhysicsAssetEditor/Private/PhysicsAssetEditorSkeletonTreeBuilder.cpp
+++ b/Engine/Source/Editor/PhysicsAssetEditor/Private/PhysicsAssetEditorSkeletonTreeBuilder.cpp
@@ -109,7 +109,7 @@ void FPhysicsAssetEditorSkeletonTreeBuilder::AddBodies(FSkeletonTreeBuilderOutpu
 		UDebugSkelMeshComponent* PreviewMeshComponent = PreviewScenePtr.Pin()->GetPreviewMeshComponent();
 		if (PreviewMeshComponent->SkeletalMesh)
 		{
-			FReferenceSkeleton& RefSkeleton = PreviewMeshComponent->SkeletalMesh->RefSkeleton;
+			FReferenceSkeleton& RefSkeleton = PreviewMeshComponent->SkeletalMesh->GetRefSkeleton();
 
 			for (int32 BoneIndex = 0; BoneIndex < RefSkeleton.GetRawBoneNum(); ++BoneIndex)
 			{
diff --git a/Engine/Source/Editor/SequenceRecorder/Private/AnimationRecorder.cpp b/Engine/Source/Editor/SequenceRecorder/Private/AnimationRecorder.cpp
--- a/Engine/Source/Editor/SequenceRecorder/Private/AnimationRecorder.cpp
+++ b/Engine/Source/Editor/SequenceRecorder/Private/AnimationRecorder.cpp
@@ -105,7 +105,7 @@ bool FAnimationRecorder::TriggerRecordAnimation(USkeletalMeshComponent* Componen
 	FString AssetPath;
 	FString AssetName;
 
-	if (!Component || !Component->SkeletalMesh || !Component->SkeletalMesh->Skeleton)
+	if (!Component || !Component->SkeletalMesh || !Component->SkeletalMesh->GetSkeleton())
 	{
 		return false;
 	}
@@ -121,7 +121,7 @@ bool FAnimationRecorder::TriggerRecordAnimation(USkeletalMeshComponent* Componen
 
 bool FAnimationRecorder::TriggerRecordAnimation(USkeletalMeshComponent* Component, const FString& InAssetPath, const FString& InAssetName)
 {
-	if (!Component || !Component->SkeletalMesh || !Component->SkeletalMesh->Skeleton)
+	if (!Component || !Component->SkeletalMesh || !Component->SkeletalMesh->GetSkeleton())
 	{
 		return false;
 	}
@@ -162,7 +162,7 @@ bool FAnimationRecorder::TriggerRecordAnimation(USkeletalMeshComponent* Componen
 	if (NewSeq)
 	{
 		// set skeleton
-		NewSeq->SetSkeleton(Component->SkeletalMesh->Skeleton);
+		NewSeq->SetSkeleton(Component->SkeletalMesh->GetSkeleton());
 		// Notify the asset registry
 		FAssetRegistryModule::AssetCreated(NewSeq);
 		StartRecord(Component, NewSeq);
@@ -649,7 +649,7 @@ bool FAnimationRecorder::Record(USkeletalMeshComponent* Component, FTransform co
 				if (BoneTreeIndex != INDEX_NONE)
 				{
 					int32 BoneIndex = AnimSkeleton->GetMeshBoneIndexFromSkeletonBoneIndex(SkeletalMesh, BoneTreeIndex);
-					int32 ParentIndex = SkeletalMesh->RefSkeleton.GetParentIndex(BoneIndex);
+					int32 ParentIndex = SkeletalMesh->GetRefSkeleton().GetParentIndex(BoneIndex);
 					FTransform LocalTransform = SpacesBases[BoneIndex];
 					if (ParentIndex == INDEX_NONE)
 					{
@@ -686,7 +686,7 @@ bool FAnimationRecorder::Record(USkeletalMeshComponent* Component, FTransform co
 			if (BoneTreeIndex != INDEX_NONE)
 			{
 				int32 BoneIndex = AnimSkeleton->GetMeshBoneIndexFromSkeletonBoneIndex(SkeletalMesh, BoneTreeIndex);
-				int32 ParentIndex = SkeletalMesh->RefSkeleton.GetParentIndex(BoneIndex);
+				int32 ParentIndex = SkeletalMesh->GetRefSkeleton().GetParentIndex(BoneIndex);
 				FTransform LocalTransform = SpacesBases[BoneIndex];
 				if ( ParentIndex != INDEX_NONE )
 				{
diff --git a/Engine/Source/Editor/SequenceRecorder/Private/Sections/MovieScene3DTransformSectionRecorder.cpp b/Engine/Source/Editor/SequenceRecorder/Private/Sections/MovieScene3DTransformSectionRecorder.cpp
--- a/Engine/Source/Editor/SequenceRecorder/Private/Sections/MovieScene3DTransformSectionRecorder.cpp
+++ b/Engine/Source/Editor/SequenceRecorder/Private/Sections/MovieScene3DTransformSectionRecorder.cpp
@@ -143,7 +143,7 @@ void FMovieScene3DTransformSectionRecorder::FinalizeSection(float CurrentTime)
 					if (BoneTreeIndex != INDEX_NONE)
 					{
 						int32 BoneIndex = AnimSkeleton->GetMeshBoneIndexFromSkeletonBoneIndex(SkeletalMesh, BoneTreeIndex);
-						int32 ParentIndex = SkeletalMesh->RefSkeleton.GetParentIndex(BoneIndex);
+						int32 ParentIndex = SkeletalMesh->GetRefSkeleton().GetParentIndex(BoneIndex);
 						if (ParentIndex == INDEX_NONE)
 						{
 							// found root
diff --git a/Engine/Source/Editor/SequenceRecorder/Private/Sections/MovieSceneAnimationSectionRecorder.cpp b/Engine/Source/Editor/SequenceRecorder/Private/Sections/MovieSceneAnimationSectionRecorder.cpp
--- a/Engine/Source/Editor/SequenceRecorder/Private/Sections/MovieSceneAnimationSectionRecorder.cpp
+++ b/Engine/Source/Editor/SequenceRecorder/Private/Sections/MovieSceneAnimationSectionRecorder.cpp
@@ -99,7 +99,7 @@ void FMovieSceneAnimationSectionRecorder::CreateSection(UObject* InObjectToRecor
 					FAssetRegistryModule::AssetCreated(AnimSequence.Get());
 
 					// set skeleton
-					AnimSequence->SetSkeleton(SkeletalMeshComponent->SkeletalMesh->Skeleton);
+					AnimSequence->SetSkeleton(SkeletalMeshComponent->SkeletalMesh->GetSkeleton());
 				}
 			}
 
diff --git a/Engine/Source/Editor/SkeletalMeshEditor/Private/MeshHierarchyCmd.cpp b/Engine/Source/Editor/SkeletalMeshEditor/Private/MeshHierarchyCmd.cpp
--- a/Engine/Source/Editor/SkeletalMeshEditor/Private/MeshHierarchyCmd.cpp
+++ b/Engine/Source/Editor/SkeletalMeshEditor/Private/MeshHierarchyCmd.cpp
@@ -34,10 +34,10 @@ bool FMeshHierarchyCmd::Exec(UWorld*, const TCHAR* Cmd, FOutputDevice& Ar)
 			GWarn->StatusUpdate(MeshIdx, SkeletalMeshes.Num(), StatusUpdate);
 
 			USkeletalMesh* Mesh = Cast<USkeletalMesh>(SkeletalMeshes[MeshIdx].GetAsset());
-			if (Mesh->Skeleton)
+			if (Mesh->GetSkeleton())
 			{
-				FName MeshRoot = Mesh->RefSkeleton.GetBoneName(0);
-				FName SkelRoot = Mesh->Skeleton->GetReferenceSkeleton().GetBoneName(0);
+				FName MeshRoot = Mesh->GetRefSkeleton().GetBoneName(0);
+				FName SkelRoot = Mesh->GetSkeleton()->GetReferenceSkeleton().GetBoneName(0);
 
 				if (MeshRoot != SkelRoot)
 				{
diff --git a/Engine/Source/Editor/SkeletalMeshEditor/Private/SControlRigMappingWindow.cpp b/Engine/Source/Editor/SkeletalMeshEditor/Private/SControlRigMappingWindow.cpp
--- a/Engine/Source/Editor/SkeletalMeshEditor/Private/SControlRigMappingWindow.cpp
+++ b/Engine/Source/Editor/SkeletalMeshEditor/Private/SControlRigMappingWindow.cpp
@@ -196,10 +196,10 @@ void SControlRigMappingWindow::AddNodeMapping(UBlueprint* NewSourceControlRig)
 		// this all can get very messy, so for now, we just support one for each
 		// create new mapper object
 		USkeletalMesh* SkeletalMesh = EditableSkeletalMeshPtr.Get();
-
-		for (int32 Index = 0; Index < SkeletalMesh->NodeMappingData.Num(); ++Index)
+		TArray<UNodeMappingContainer*>& NodeMappingData = SkeletalMesh->GetNodeMappingData();
+		for (int32 Index = 0; Index < NodeMappingData.Num(); ++Index)
 		{
-			if (NewSourceControlRig == SkeletalMesh->NodeMappingData[Index]->GetSourceAsset())
+			if (NewSourceControlRig == NodeMappingData[Index]->GetSourceAsset())
 			{
 				FMessageDialog::Open(EAppMsgType::Ok, LOCTEXT("ControlRigConfigAlreadyExists", "The same Control Rig configuration already exists in this mesh. Edit current existing setting."));
 				return;
@@ -216,7 +216,7 @@ void SControlRigMappingWindow::AddNodeMapping(UBlueprint* NewSourceControlRig)
 			NewMapperObject->SetTargetAsset(SkeletalMesh);
 			// add default mapping, this will map default settings
 			NewMapperObject->AddDefaultMapping();
-			CurrentlySelectedIndex = SkeletalMesh->NodeMappingData.Add(NewMapperObject);
+			CurrentlySelectedIndex = SkeletalMesh->GetNodeMappingData().Add(NewMapperObject);
 
 			RefreshList();
 		}
@@ -229,11 +229,12 @@ FReply SControlRigMappingWindow::OnDeleteNodeMappingButtonClicked()
 
 	// create new mapper object
 	USkeletalMesh* SkeletalMesh = EditableSkeletalMeshPtr.Get();
-	if (SkeletalMesh->NodeMappingData.IsValidIndex(CurrentlySelectedIndex))
+	TArray<UNodeMappingContainer*>& NodeMappingData = SkeletalMesh->GetNodeMappingData();
+	if (NodeMappingData.IsValidIndex(CurrentlySelectedIndex))
 	{
 		SkeletalMesh->Modify();
-		SkeletalMesh->NodeMappingData.RemoveAt(CurrentlySelectedIndex);
-		CurrentlySelectedIndex = (SkeletalMesh->NodeMappingData.Num() > 0) ? 0 : INDEX_NONE;
+		NodeMappingData.RemoveAt(CurrentlySelectedIndex);
+		CurrentlySelectedIndex = (NodeMappingData.Num() > 0) ? 0 : INDEX_NONE;
 		RefreshList();
 	}
 
@@ -244,7 +245,7 @@ FReply SControlRigMappingWindow::OnRefreshNodeMappingButtonClicked()
 {
 	// create new mapper object
 	USkeletalMesh* SkeletalMesh = EditableSkeletalMeshPtr.Get();
-	if (SkeletalMesh->NodeMappingData.IsValidIndex(CurrentlySelectedIndex))
+	if (SkeletalMesh->GetNodeMappingData().IsValidIndex(CurrentlySelectedIndex))
 	{
 		const FScopedTransaction Transaction(LOCTEXT("ControlRigMapping_Refresh", "Refresh Node Mapping"));
 
@@ -341,9 +342,10 @@ void SControlRigMappingWindow::RefreshList()
 	// @todo: have to make sure there is no duplicated name. If so, we'll have to create fake name
 	// for now, do path name
 	MappingOptionBoxList.Empty();
-	for (int32 Index = 0; Index < EditableSkeletalMeshPtr->NodeMappingData.Num(); ++Index)
+	TArray<UNodeMappingContainer*>& NodeMappingData = EditableSkeletalMeshPtr->GetNodeMappingData();
+	for (int32 Index = 0; Index < NodeMappingData.Num(); ++Index)
 	{
-		class UNodeMappingContainer* MappingData = EditableSkeletalMeshPtr->NodeMappingData[Index];
+		class UNodeMappingContainer* MappingData = NodeMappingData[Index];
 		if (MappingData)
 		{
 			MappingOptionBoxList.Add(MakeShareable(new UNodeMappingContainer*(MappingData)));
@@ -367,9 +369,10 @@ class UNodeMappingContainer* SControlRigMappingWindow::GetCurrentBoneMappingCont
 	if (EditableSkeletalMeshPtr.IsValid())
 	{
 		USkeletalMesh* Mesh = EditableSkeletalMeshPtr.Get();
-		if (Mesh->NodeMappingData.IsValidIndex(CurrentlySelectedIndex))
+		TArray<UNodeMappingContainer*>& NodeMappingData = Mesh->GetNodeMappingData();
+		if (NodeMappingData.IsValidIndex(CurrentlySelectedIndex))
 		{
-			return Mesh->NodeMappingData[CurrentlySelectedIndex];
+			return NodeMappingData[CurrentlySelectedIndex];
 		}
 	}
 
@@ -403,7 +406,7 @@ void SControlRigMappingWindow::CreateBoneMappingList(const FString& SearchText,
 {
 	BoneMappingList.Empty();
 
-	const FReferenceSkeleton RefSkeleton = EditableSkeletalMeshPtr.Get()->RefSkeleton;
+	const FReferenceSkeleton RefSkeleton = EditableSkeletalMeshPtr.Get()->GetRefSkeleton();
 	UNodeMappingContainer* Container = GetCurrentBoneMappingContainer();
 	if (Container)
 	{
@@ -438,7 +441,7 @@ void SControlRigMappingWindow::CreateBoneMappingList(const FString& SearchText,
 const struct FReferenceSkeleton& SControlRigMappingWindow::GetReferenceSkeleton() const
 {
 	static FReferenceSkeleton DummySkeleton;
-	return (EditableSkeletalMeshPtr.IsValid())? EditableSkeletalMeshPtr.Get()->RefSkeleton : DummySkeleton;
+	return (EditableSkeletalMeshPtr.IsValid())? EditableSkeletalMeshPtr.Get()->GetRefSkeleton() : DummySkeleton;
 }
 #undef LOCTEXT_NAMESPACE
 
diff --git a/Engine/Source/Editor/SkeletalMeshEditor/Private/SkeletalMeshEditor.cpp b/Engine/Source/Editor/SkeletalMeshEditor/Private/SkeletalMeshEditor.cpp
--- a/Engine/Source/Editor/SkeletalMeshEditor/Private/SkeletalMeshEditor.cpp
+++ b/Engine/Source/Editor/SkeletalMeshEditor/Private/SkeletalMeshEditor.cpp
@@ -136,15 +136,15 @@ bool FSkeletalMeshEditor::OnRequestClose()
 			}
 
 			bool bValidLODSettings = false;
-			if (SkeletalMesh->LODSettings != nullptr)
+			if (SkeletalMesh->GetLODSettings() != nullptr)
 			{
-				const int32 NumSettings = FMath::Min(SkeletalMesh->LODSettings->GetNumberOfSettings(), SkeletalMesh->GetLODNum());
+				const int32 NumSettings = FMath::Min(SkeletalMesh->GetLODSettings()->GetNumberOfSettings(), SkeletalMesh->GetLODNum());
 				if (LODIndex < NumSettings)
 				{
 					bValidLODSettings = true;
 				}
 			}
-			const FSkeletalMeshLODGroupSettings* SkeletalMeshLODGroupSettings = bValidLODSettings ? &SkeletalMesh->LODSettings->GetSettingsForLODLevel(LODIndex) : nullptr;
+			const FSkeletalMeshLODGroupSettings* SkeletalMeshLODGroupSettings = bValidLODSettings ? &SkeletalMesh->GetLODSettings()->GetSettingsForLODLevel(LODIndex) : nullptr;
 
 			FGuid BuildGUID = LODInfo->ComputeDeriveDataCacheKey(SkeletalMeshLODGroupSettings);
 			if (LODInfo->BuildGUID != BuildGUID)
@@ -309,7 +309,7 @@ void FSkeletalMeshEditor::RegisterReimportContextMenu(const FName InBaseMenuName
 		{
 			if (USkeletalMesh* FoundSkeletalMesh = SkeletalMeshEditor->SkeletalMesh)
 			{
-				UFbxSkeletalMeshImportData* SkeletalMeshImportData = Cast<UFbxSkeletalMeshImportData>(FoundSkeletalMesh->AssetImportData);
+				UFbxSkeletalMeshImportData* SkeletalMeshImportData = Cast<UFbxSkeletalMeshImportData>(FoundSkeletalMesh->GetAssetImportData());
 				if (SkeletalMeshImportData)
 				{
 					SkeletalMeshImportData->ImportContentType = SourceFileIndex == 0 ? EFBXImportContentType::FBXICT_All : SourceFileIndex == 1 ? EFBXImportContentType::FBXICT_Geometry : EFBXImportContentType::FBXICT_SkinningWeights;
@@ -351,7 +351,7 @@ void FSkeletalMeshEditor::RegisterReimportContextMenu(const FName InBaseMenuName
 		{
 			if (USkeletalMesh* FoundSkeletalMesh = SkeletalMeshEditor->SkeletalMesh)
 			{
-				UFbxSkeletalMeshImportData* SkeletalMeshImportData = FoundSkeletalMesh ? Cast<UFbxSkeletalMeshImportData>(FoundSkeletalMesh->AssetImportData) : nullptr;
+				UFbxSkeletalMeshImportData* SkeletalMeshImportData = FoundSkeletalMesh ? Cast<UFbxSkeletalMeshImportData>(FoundSkeletalMesh->GetAssetImportData()) : nullptr;
 				if (SkeletalMeshImportData)
 				{
 					SkeletalMeshImportData->ImportContentType = SourceFileIndex == 0 ? EFBXImportContentType::FBXICT_All : SourceFileIndex == 1 ? EFBXImportContentType::FBXICT_Geometry : EFBXImportContentType::FBXICT_SkinningWeights;
@@ -388,13 +388,13 @@ void FSkeletalMeshEditor::RegisterReimportContextMenu(const FName InBaseMenuName
 		if (SkeletalMeshEditor.IsValid())
 		{
 			USkeletalMesh* InSkeletalMesh = SkeletalMeshEditor->SkeletalMesh;
-			if (InSkeletalMesh && InSkeletalMesh->AssetImportData)
+			if (InSkeletalMesh && InSkeletalMesh->GetAssetImportData())
 			{
 				//Get the data
 				TArray<FString> SourceFilePaths;
-				InSkeletalMesh->AssetImportData->ExtractFilenames(SourceFilePaths);
+				InSkeletalMesh->GetAssetImportData()->ExtractFilenames(SourceFilePaths);
 				TArray<FString> SourceFileLabels;
-				InSkeletalMesh->AssetImportData->ExtractDisplayLabels(SourceFileLabels);
+				InSkeletalMesh->GetAssetImportData()->ExtractDisplayLabels(SourceFileLabels);
 
 				if (SourceFileLabels.Num() > 0 && SourceFileLabels.Num() == SourceFilePaths.Num())
 				{
@@ -431,7 +431,7 @@ void FSkeletalMeshEditor::RegisterReimportContextMenu(const FName InBaseMenuName
 			if (SkeletalMeshEditor.IsValid())
 			{
 				USkeletalMesh* InSkeletalMesh = SkeletalMeshEditor->SkeletalMesh;
-				bool bShowSubMenu = InSkeletalMesh != nullptr && InSkeletalMesh->AssetImportData != nullptr && InSkeletalMesh->AssetImportData->GetSourceFileCount() > 1;
+				bool bShowSubMenu = InSkeletalMesh != nullptr && InSkeletalMesh->GetAssetImportData() != nullptr && InSkeletalMesh->GetAssetImportData()->GetSourceFileCount() > 1;
 
 				FToolMenuSection& Section = InMenu->AddSection("Section");
 				if (!bShowSubMenu)
@@ -747,7 +747,7 @@ void FSkeletalMeshEditor::FillApplyClothingAssetMenu(FMenuBuilder& MenuBuilder,
 
 	MenuBuilder.BeginSection(TEXT("ApplyClothingMenu"), LOCTEXT("ApplyClothingMenuHeader", "Available Assets"));
 	{
-		for(UClothingAssetBase* BaseAsset : Mesh->MeshClothingAssets)
+		for(UClothingAssetBase* BaseAsset : Mesh->GetMeshClothingAssets())
 		{
 			UClothingAssetCommon* ClothAsset = CastChecked<UClothingAssetCommon>(BaseAsset);
 
@@ -900,7 +900,7 @@ bool FSkeletalMeshEditor::CanApplyClothing(int32 InLodIndex, int32 InSectionInde
 {
 	USkeletalMesh* Mesh = GetPersonaToolkit()->GetPreviewMesh();
 
-	if(Mesh->MeshClothingAssets.Num() > 0)
+	if(Mesh->GetMeshClothingAssets().Num() > 0)
 	{
 	FSkeletalMeshModel* MeshResource = Mesh->GetImportedModel();
 
@@ -962,7 +962,7 @@ bool FSkeletalMeshEditor::CanCreateClothingLod(int32 InLodIndex, int32 InSection
 {
 	USkeletalMesh* Mesh = GetPersonaToolkit()->GetPreviewMesh();
 
-	return Mesh && Mesh->MeshClothingAssets.Num() > 0 && CanApplyClothing(InLodIndex, InSectionIndex);
+	return Mesh && Mesh->GetMeshClothingAssets().Num() > 0 && CanApplyClothing(InLodIndex, InSectionIndex);
 }
 
 void FSkeletalMeshEditor::ApplyClothing(UClothingAssetBase* InAsset, int32 InLodIndex, int32 InSectionIndex, int32 InClothingLod)
@@ -1002,7 +1002,7 @@ void FSkeletalMeshEditor::ApplyClothing(UClothingAssetBase* InAsset, int32 InLod
 			{
 				//Successful bind so set the SectionUserData
 				int32 AssetIndex = INDEX_NONE;
-				check(Mesh->MeshClothingAssets.Find(ClothingAsset, AssetIndex));
+				check(Mesh->GetMeshClothingAssets().Find(ClothingAsset, AssetIndex));
 				OriginalSectionData.CorrespondClothAssetIndex = AssetIndex;
 				OriginalSectionData.ClothingData.AssetGuid = ClothingAsset->GetAssetGuid();
 				OriginalSectionData.ClothingData.AssetLodIndex = InClothingLod;
diff --git a/Engine/Source/Editor/SkeletalMeshEditor/Private/SkeletalMeshEditorSubsystem.cpp b/Engine/Source/Editor/SkeletalMeshEditor/Private/SkeletalMeshEditorSubsystem.cpp
--- a/Engine/Source/Editor/SkeletalMeshEditor/Private/SkeletalMeshEditorSubsystem.cpp
+++ b/Engine/Source/Editor/SkeletalMeshEditor/Private/SkeletalMeshEditorSubsystem.cpp
@@ -82,7 +82,7 @@ bool USkeletalMeshEditorSubsystem::RenameSocket(USkeletalMesh* SkeletalMesh, FNa
 		return false;
 	}
 
-	if (SkeletalMesh->Skeleton == nullptr)
+	if (SkeletalMesh->GetSkeleton() == nullptr)
 	{
 		UE_LOG(LogSkeletalMeshEditorSubsystem, Error, TEXT("RenameSocket: The SkeletalMesh's Skeleton is null."));
 		return false;
@@ -112,7 +112,7 @@ bool USkeletalMeshEditorSubsystem::RenameSocket(USkeletalMesh* SkeletalMesh, FNa
 		return false;
 	}
 
-	USkeletalMeshSocket* SkeletonSocket = SkeletalMesh->Skeleton->FindSocket(OldName);
+	USkeletalMeshSocket* SkeletonSocket = SkeletalMesh->GetSkeleton()->FindSocket(OldName);
 	if (SkeletonSocket == nullptr)
 	{
 		UE_LOG(LogSkeletalMeshEditorSubsystem, Error, TEXT("RenameSocket: The socket named '%s' does not exist on the Skeleton."), *OldName.ToString());
@@ -128,10 +128,11 @@ bool USkeletalMeshEditorSubsystem::RenameSocket(USkeletalMesh* SkeletalMesh, FNa
 	SkeletonSocket->Modify();
 	SkeletonSocket->SocketName = NewName;
 
+	FPreviewAssetAttachContainer& PreviewAssetAttachContainer = SkeletalMesh->GetPreviewAttachedAssetContainer();
 	bool bMeshModified = false;
-	for (int AttachedObjectIndex = 0; AttachedObjectIndex < SkeletalMesh->PreviewAttachedAssetContainer.Num(); ++AttachedObjectIndex)
+	for (int AttachedObjectIndex = 0; AttachedObjectIndex < PreviewAssetAttachContainer.Num(); ++AttachedObjectIndex)
 	{
-		FPreviewAttachedObjectPair& Pair = SkeletalMesh->PreviewAttachedAssetContainer[AttachedObjectIndex];
+		FPreviewAttachedObjectPair& Pair = PreviewAssetAttachContainer[AttachedObjectIndex];
 		if (Pair.AttachedTo == OldName)
 		{
 			// Only modify the mesh if we actually intend to change something. Avoids dirtying
@@ -146,15 +147,15 @@ bool USkeletalMeshEditorSubsystem::RenameSocket(USkeletalMesh* SkeletalMesh, FNa
 	}
 
 	bool bSkeletonModified = false;
-	for (int AttachedObjectIndex = 0; AttachedObjectIndex < SkeletalMesh->Skeleton->PreviewAttachedAssetContainer.Num(); ++AttachedObjectIndex)
+	for (int AttachedObjectIndex = 0; AttachedObjectIndex < SkeletalMesh->GetSkeleton()->PreviewAttachedAssetContainer.Num(); ++AttachedObjectIndex)
 	{
-		FPreviewAttachedObjectPair& Pair = SkeletalMesh->Skeleton->PreviewAttachedAssetContainer[AttachedObjectIndex];
+		FPreviewAttachedObjectPair& Pair = SkeletalMesh->GetSkeleton()->PreviewAttachedAssetContainer[AttachedObjectIndex];
 		if (Pair.AttachedTo == OldName)
 		{
 			// Only modify the skeleton if we actually intend to change something.
 			if (!bSkeletonModified)
 			{
-				SkeletalMesh->Skeleton->Modify();
+				SkeletalMesh->GetSkeleton()->Modify();
 				bSkeletonModified = true;
 			}
 			Pair.AttachedTo = NewName;
@@ -221,7 +222,7 @@ int32 USkeletalMeshEditorSubsystem::ImportLOD(USkeletalMesh* BaseMesh, const int
 		&& BaseMesh->IsReductionActive(LODIndex) //We remove reduction settings only if they are active
 		&& BaseMesh->GetLODInfo(LODIndex) //this test is redundant (IsReductionActive test this), but to avoid static analysis
 		&& BaseMesh->GetLODInfo(LODIndex)->ReductionSettings.BaseLOD < LODIndex //We do not remove the reduction if the reduction is base on this LOD imported data
-		&& (!BaseMesh->LODSettings || BaseMesh->LODSettings->GetNumberOfSettings() < LODIndex)) //We do not remove the reduction if the skeletal mesh is using a LODSettings for this LOD
+		&& (!BaseMesh->GetLODSettings() || BaseMesh->GetLODSettings()->GetNumberOfSettings() < LODIndex)) //We do not remove the reduction if the skeletal mesh is using a LODSettings for this LOD
 	{
 		//Remove the reduction settings
 		BaseMesh->GetLODInfo(LODIndex)->ReductionSettings.NumOfTrianglesPercentage = 1.0f;
diff --git a/Engine/Source/Editor/SkeletonEditor/Private/BoneProxy.cpp b/Engine/Source/Editor/SkeletonEditor/Private/BoneProxy.cpp
--- a/Engine/Source/Editor/SkeletonEditor/Private/BoneProxy.cpp
+++ b/Engine/Source/Editor/SkeletonEditor/Private/BoneProxy.cpp
@@ -54,7 +54,7 @@ void UBoneProxy::Tick(float DeltaTime)
 
 				Scale = LocalTransform.GetScale3D();
 
-				FTransform ReferenceTransform = Component->SkeletalMesh->RefSkeleton.GetRefBonePose()[BoneIndex];
+				FTransform ReferenceTransform = Component->SkeletalMesh->GetRefSkeleton().GetRefBonePose()[BoneIndex];
 				ReferenceLocation = ReferenceTransform.GetLocation();
 				ReferenceRotation = ReferenceTransform.GetRotation().Rotator();
 				ReferenceScale = ReferenceTransform.GetScale3D();
diff --git a/Engine/Source/Editor/SkeletonEditor/Private/EditableSkeleton.cpp b/Engine/Source/Editor/SkeletonEditor/Private/EditableSkeleton.cpp
--- a/Engine/Source/Editor/SkeletonEditor/Private/EditableSkeleton.cpp
+++ b/Engine/Source/Editor/SkeletonEditor/Private/EditableSkeleton.cpp
@@ -71,7 +71,7 @@ private:
 		if (USkeletalMesh* SocketMesh = Cast<USkeletalMesh>(NewSocket->GetOuter()))
 		{
 			SocketMesh->GetMeshOnlySocketList().Add(NewSocket);
-			RefSkel = &SocketMesh->RefSkeleton;
+			RefSkel = &SocketMesh->GetRefSkeleton();
 		}
 		else if (USkeleton* SocketSkeleton = Cast<USkeleton>(NewSocket->GetOuter()))
 		{
@@ -624,10 +624,11 @@ void FEditableSkeleton::RenameSocket(const FName OldSocketName, const FName NewS
 
 		if (InSkeletalMesh != nullptr)
 		{
+			FPreviewAssetAttachContainer& PreviewAttachedAssetContainer = InSkeletalMesh->GetPreviewAttachedAssetContainer();
 			bool bMeshModified = false;
-			for (int AttachedObjectIndex = 0; AttachedObjectIndex < InSkeletalMesh->PreviewAttachedAssetContainer.Num(); ++AttachedObjectIndex)
+			for (int AttachedObjectIndex = 0; AttachedObjectIndex < PreviewAttachedAssetContainer.Num(); ++AttachedObjectIndex)
 			{
-				FPreviewAttachedObjectPair& Pair = InSkeletalMesh->PreviewAttachedAssetContainer[AttachedObjectIndex];
+				FPreviewAttachedObjectPair& Pair = PreviewAttachedAssetContainer[AttachedObjectIndex];
 				if (Pair.AttachedTo == OldSocketName)
 				{
 					// Only modify the mesh if we actually intend to change something. Avoids dirtying
@@ -791,7 +792,7 @@ void FEditableSkeleton::HandleRemoveAllAssets(TSharedPtr<IPersonaPreviewScene> I
 	if (SkeletalMesh)
 	{
 		SkeletalMesh->Modify();
-		DeleteAttachedObjects(SkeletalMesh->PreviewAttachedAssetContainer, InPreviewScene);
+		DeleteAttachedObjects(SkeletalMesh->GetPreviewAttachedAssetContainer(), InPreviewScene);
 	}
 }
 
@@ -910,7 +911,7 @@ void FEditableSkeleton::HandleAttachAssets(const TArray<UObject*>& InObjects, co
 				{
 					InPreviewScene->AttachObjectToPreviewComponent(Object, InAttachToName);
 				}
-				SkeletalMesh->PreviewAttachedAssetContainer.AddAttachedObject(Object, InAttachToName);
+				SkeletalMesh->GetPreviewAttachedAssetContainer().AddAttachedObject(Object, InAttachToName);
 			}
 		}
 		else
@@ -960,7 +961,7 @@ void FEditableSkeleton::HandleDeleteAttachedAssets(const TArray<FPreviewAttached
 
 			if (SkeletalMesh != nullptr)
 			{
-				SkeletalMesh->PreviewAttachedAssetContainer.RemoveAttachedObject(AttachedObject.GetAttachedObject(), AttachedObject.AttachedTo);
+				SkeletalMesh->GetPreviewAttachedAssetContainer().RemoveAttachedObject(AttachedObject.GetAttachedObject(), AttachedObject.AttachedTo);
 
 				if (InPreviewScene.IsValid())
 				{
@@ -987,11 +988,11 @@ void FEditableSkeleton::HandleDeleteSockets(const TArray<FSelectedSocketInfo>& I
 			USkeletalMesh* SkeletalMesh = InPreviewScene->GetPreviewMeshComponent()->SkeletalMesh;
 			if (SkeletalMesh != nullptr)
 			{
-				UObject* Object = SkeletalMesh->PreviewAttachedAssetContainer.GetAttachedObjectByAttachName(SocketName);
+				UObject* Object = SkeletalMesh->GetPreviewAttachedAssetContainer().GetAttachedObjectByAttachName(SocketName);
 				if (Object)
 				{
 					SkeletalMesh->Modify();
-					SkeletalMesh->PreviewAttachedAssetContainer.RemoveAttachedObject(Object, SocketName);
+					SkeletalMesh->GetPreviewAttachedAssetContainer().RemoveAttachedObject(Object, SocketName);
 					if (InPreviewScene.IsValid())
 					{
 						InPreviewScene->RemoveAttachedObjectFromPreviewComponent(Object, SocketName);
@@ -1528,7 +1529,7 @@ void FEditableSkeleton::RemoveUnusedBones()
 	Filter.ClassNames.Add(USkeletalMesh::StaticClass()->GetFName());
 
 	FString SkeletonString = FAssetData(Skeleton).GetExportTextName();
-	Filter.TagsAndValues.Add(GET_MEMBER_NAME_CHECKED(USkeletalMesh, Skeleton), SkeletonString);
+	Filter.TagsAndValues.Add(USkeletalMesh::GetSkeletonMemberName(), SkeletonString);
 
 	TArray<FAssetData> SkeletalMeshes;
 	FAssetRegistryModule& AssetRegistryModule = FModuleManager::LoadModuleChecked<FAssetRegistryModule>(TEXT("AssetRegistry"));
@@ -1560,7 +1561,7 @@ void FEditableSkeleton::RemoveUnusedBones()
 				GWarn->StatusUpdate(MeshIdx, SkeletalMeshes.Num(), StatusUpdate);
 
 				USkeletalMesh* Mesh = Cast<USkeletalMesh>(SkeletalMeshes[MeshIdx].GetAsset());
-				const FReferenceSkeleton& MeshRefSkeleton = Mesh->RefSkeleton;
+				const FReferenceSkeleton& MeshRefSkeleton = Mesh->GetRefSkeleton();
 
 				for (int32 BoneIndex = 0; BoneIndex < MeshRefSkeleton.GetRawBoneNum(); ++BoneIndex)
 				{
diff --git a/Engine/Source/Editor/SkeletonEditor/Private/SSkeletonTree.cpp b/Engine/Source/Editor/SkeletonEditor/Private/SSkeletonTree.cpp
--- a/Engine/Source/Editor/SkeletonEditor/Private/SSkeletonTree.cpp
+++ b/Engine/Source/Editor/SkeletonEditor/Private/SSkeletonTree.cpp
@@ -1364,7 +1364,7 @@ bool SSkeletonTree::CanRemoveAllAssets() const
 	USkeletalMesh* SkeletalMesh = GetPreviewScene().IsValid() ? GetPreviewScene()->GetPreviewMeshComponent()->SkeletalMesh : nullptr;
 
 	const bool bHasPreviewAttachedObjects = GetEditableSkeletonInternal()->GetSkeleton().PreviewAttachedAssetContainer.Num() > 0;
-	const bool bHasMeshPreviewAttachedObjects = ( SkeletalMesh && SkeletalMesh->PreviewAttachedAssetContainer.Num() );
+	const bool bHasMeshPreviewAttachedObjects = ( SkeletalMesh && SkeletalMesh->GetPreviewAttachedAssetContainer().Num() );
 
 	return bHasPreviewAttachedObjects || bHasMeshPreviewAttachedObjects;
 }
diff --git a/Engine/Source/Editor/SkeletonEditor/Private/SkeletonTreeBuilder.cpp b/Engine/Source/Editor/SkeletonEditor/Private/SkeletonTreeBuilder.cpp
--- a/Engine/Source/Editor/SkeletonEditor/Private/SkeletonTreeBuilder.cpp
+++ b/Engine/Source/Editor/SkeletonEditor/Private/SkeletonTreeBuilder.cpp
@@ -251,7 +251,7 @@ void FSkeletonTreeBuilder::AddAttachedAssets(FSkeletonTreeBuilderOutput& Output)
 		UDebugSkelMeshComponent* PreviewMeshComponent = PreviewScenePtr.Pin()->GetPreviewMeshComponent();
 		if (USkeletalMesh* const SkeletalMesh = PreviewMeshComponent->SkeletalMesh)
 		{
-			AddAttachedAssetContainer(SkeletalMesh->PreviewAttachedAssetContainer, Output);
+			AddAttachedAssetContainer(SkeletalMesh->GetPreviewAttachedAssetContainer(), Output);
 		}
 	}
 
diff --git a/Engine/Source/Editor/UnrealEd/Classes/Animation/DebugSkelMeshComponent.h b/Engine/Source/Editor/UnrealEd/Classes/Animation/DebugSkelMeshComponent.h
--- a/Engine/Source/Editor/UnrealEd/Classes/Animation/DebugSkelMeshComponent.h
+++ b/Engine/Source/Editor/UnrealEd/Classes/Animation/DebugSkelMeshComponent.h
@@ -492,7 +492,7 @@ public:
 	{
 		if (SkeletalMesh)
 		{
-			return SkeletalMesh->RefSkeleton;
+			return SkeletalMesh->GetRefSkeleton();
 		}
 
 		static FReferenceSkeleton EmptySkeleton;
