Fix for preview scene settings UI. Additional Meshes array not refreshed when mesh collection asset reset to default.

diff --git a/Engine/Source/Editor/Persona/Private/PreviewSceneCustomizations.cpp b/Engine/Source/Editor/Persona/Private/PreviewSceneCustomizations.cpp
--- a/Engine/Source/Editor/Persona/Private/PreviewSceneCustomizations.cpp
+++ b/Engine/Source/Editor/Persona/Private/PreviewSceneCustomizations.cpp
@@ -76,6 +76,7 @@ void FPreviewSceneDescriptionCustomization::CustomizeDetails(IDetailLayoutBuilde
 	TSharedRef<IPropertyHandle> SkeletalMeshProperty = DetailBuilder.GetProperty(GET_MEMBER_NAME_CHECKED(UPersonaPreviewSceneDescription, PreviewMesh));
 
 	AdditionalMeshesProperty = DetailBuilder.GetProperty(GET_MEMBER_NAME_CHECKED(UPersonaPreviewSceneDescription, AdditionalMeshes));
+	AdditionalMeshesProperty->SetOnPropertyResetToDefault(FSimpleDelegate::CreateSP(this, &FPreviewSceneDescriptionCustomization::OnResetAdditionalMeshes));
 
 	TArray<UClass*> BuiltInPreviewControllers = { UPersonaPreviewSceneDefaultController::StaticClass(), UPersonaPreviewSceneRefPoseController::StaticClass(), UPersonaPreviewSceneAnimationController::StaticClass() };
 
@@ -352,9 +353,6 @@ void FPreviewSceneDescriptionCustomization::CustomizeDetails(IDetailLayoutBuilde
 		FactoryToUse->CurrentSkeleton = EditableSkeleton.IsValid() ? MakeWeakObjectPtr(const_cast<USkeleton*>(&EditableSkeleton.Pin()->GetSkeleton())) : nullptr;
 		TArray<UFactory*> FactoriesToUse({ FactoryToUse });
 
-		FAssetData AdditionalMeshesAsset;
-		AdditionalMeshesProperty->GetValue(AdditionalMeshesAsset);
-
 		// bAllowPreviewMeshCollectionsToSelectFromDifferentSkeletons option
 		DetailBuilder.EditCategory("Additional Meshes")
 		.AddCustomRow(LOCTEXT("AdditionalMeshOption", "Additional Mesh Selection Option"))
@@ -439,6 +437,8 @@ void FPreviewSceneDescriptionCustomization::CustomizeDetails(IDetailLayoutBuilde
 			]
 		];
 
+		FAssetData AdditionalMeshesAsset;
+		AdditionalMeshesProperty->GetValue(AdditionalMeshesAsset);
 		if (AdditionalMeshesAsset.IsValid())
 		{
 			TArray<UObject*> Objects;
@@ -641,6 +641,24 @@ void FPreviewSceneDescriptionCustomization::OnResetToBaseClicked(TSharedPtr<IPro
  	}
 }
 
+void FPreviewSceneDescriptionCustomization::OnResetAdditionalMeshes()
+{	
+	// this function resets the additional meshes property to null,
+	// in the future if we serialize the default setting, this will
+	// need to reset it to the default value, not just null.
+
+	// Only allow reset to base if the current material can be replaced
+	if (AdditionalMeshesProperty.IsValid())
+	{
+		FAssetData NullAsset;
+		AdditionalMeshesProperty->SetValue(NullAsset);
+
+		PreviewScene.Pin()->SetAdditionalMeshes(nullptr);
+	}
+
+	MyDetailLayout->ForceRefreshDetails();
+}
+
 ////////////////////////////////////////////////////////////////////////////////////////////////////////////////////////////////////////////
 // 
 // FPreviewMeshCollectionEntryCustomization
diff --git a/Engine/Source/Editor/Persona/Private/PreviewSceneCustomizations.h b/Engine/Source/Editor/Persona/Private/PreviewSceneCustomizations.h
--- a/Engine/Source/Editor/Persona/Private/PreviewSceneCustomizations.h
+++ b/Engine/Source/Editor/Persona/Private/PreviewSceneCustomizations.h
@@ -132,6 +132,7 @@ private:
 	* Called when reset to base is clicked
 	*/
 	void OnResetToBaseClicked(TSharedPtr<IPropertyHandle> PropertyHandle);
+	void OnResetAdditionalMeshes();
 
 	TSharedPtr<IPropertyHandle> AdditionalMeshesProperty;
 };
