Persona Toolbars

diff --git a/Engine/Source/Editor/AnimationBlueprintEditor/Private/AnimationBlueprintEditor.h b/Engine/Source/Editor/AnimationBlueprintEditor/Private/AnimationBlueprintEditor.h
--- a/Engine/Source/Editor/AnimationBlueprintEditor/Private/AnimationBlueprintEditor.h
+++ b/Engine/Source/Editor/AnimationBlueprintEditor/Private/AnimationBlueprintEditor.h
@@ -103,7 +103,6 @@ public:
 	virtual void OnActiveTabChanged(TSharedPtr<SDockTab> PreviouslyActive, TSharedPtr<SDockTab> NewlyActivated) override;
 	virtual void OnSelectedNodesChangedImpl(const TSet<class UObject*>& NewSelection) override;
 	virtual void HandleSetObjectBeingDebugged(UObject* InObject) override;
-	virtual bool UsesCustomToolbarPlacement() const override { return true; }
 
 	// Gets the Anim Blueprint being edited/viewed by this Persona instance
 	UAnimBlueprint* GetAnimBlueprint() const;
diff --git a/Engine/Source/Editor/AnimationBlueprintEditor/Private/AnimationBlueprintEditorMode.cpp b/Engine/Source/Editor/AnimationBlueprintEditor/Private/AnimationBlueprintEditorMode.cpp
--- a/Engine/Source/Editor/AnimationBlueprintEditor/Private/AnimationBlueprintEditorMode.cpp
+++ b/Engine/Source/Editor/AnimationBlueprintEditor/Private/AnimationBlueprintEditorMode.cpp
@@ -31,14 +31,6 @@ FAnimationBlueprintEditorMode::FAnimationBlueprintEditorMode(const TSharedRef<FA
 			FTabManager::NewPrimaryArea()
 			->SetOrientation(Orient_Vertical)
 			->Split
-			(
-				// Top toolbar
-				FTabManager::NewStack() 
-				->SetSizeCoefficient(0.186721f)
-				->SetHideTabWell(true)
-				->AddTab(InAnimationBlueprintEditor->GetToolbarTabId(), ETabState::OpenedTab )
-			)
-			->Split
 			(
 				// Main application area
 				FTabManager::NewSplitter()
diff --git a/Engine/Source/Editor/AnimationEditor/Private/AnimationEditorMode.cpp b/Engine/Source/Editor/AnimationEditor/Private/AnimationEditorMode.cpp
--- a/Engine/Source/Editor/AnimationEditor/Private/AnimationEditorMode.cpp
+++ b/Engine/Source/Editor/AnimationEditor/Private/AnimationEditorMode.cpp
@@ -43,13 +43,6 @@ FAnimationEditorMode::FAnimationEditorMode(TSharedRef<FWorkflowCentricApplicatio
 			FTabManager::NewPrimaryArea()
 			->SetOrientation(Orient_Vertical)
 			->Split
-			(
-				FTabManager::NewStack()
-				->SetSizeCoefficient(0.1f)
-				->SetHideTabWell(true)
-				->AddTab(InHostingApp->GetToolbarTabId(), ETabState::OpenedTab)
-			)
-			->Split
 			(
 				FTabManager::NewSplitter()
 				->SetSizeCoefficient(0.9f)
diff --git a/Engine/Source/Editor/EditorStyle/Private/StarshipStyle.cpp b/Engine/Source/Editor/EditorStyle/Private/StarshipStyle.cpp
--- a/Engine/Source/Editor/EditorStyle/Private/StarshipStyle.cpp
+++ b/Engine/Source/Editor/EditorStyle/Private/StarshipStyle.cpp
@@ -424,6 +424,8 @@ void FStarshipEditorStyle::FStyle::SetupGeneralStyles()
 	{
 		Set( "Icons.Refresh", new IMAGE_BRUSH( "Icons/icon_Refresh_16x", Icon16x16 ) );
 		Set( "Icons.Contact", new IMAGE_BRUSH( "Icons/icon_mail_16x", Icon16x16 ) );
+
+		Set( "Icons.DirtyBadge", new IMAGE_BRUSH_SVG("Starship/Common/DirtyBadge", Icon12x12));
 	}
 
 	Set( "WarningStripe", new IMAGE_BRUSH( "Common/WarningStripe", FVector2D(20,6), FLinearColor::White, ESlateBrushTileType::Horizontal ) );
@@ -4998,6 +5000,13 @@ void FStarshipEditorStyle::FStyle::SetupPersonaStyle()
 		Set("Persona.Tabs.SkeletonCurves", new IMAGE_BRUSH(TEXT("Icons/icon_Persona_Skeleton_Curves_16x"), Icon16x16));
 		Set("Persona.Tabs.AnimAssetDetails", new IMAGE_BRUSH(TEXT("Icons/icon_Persona_Anim_Asset_Details_16x"), Icon16x16));
 		Set("Persona.Tabs.ControlRigMappingWindow", new IMAGE_BRUSH(TEXT("Icons/icon_Persona_Skeleton_Tree_16x"), Icon16x16));
+
+		// AssetFamilyIcons
+		Set("Persona.AssetClass.Skeleton", new IMAGE_BRUSH_SVG("Starship/Persona/Skeleton", Icon20x20));
+		Set("Persona.AssetClass.Animation", new IMAGE_BRUSH_SVG("Starship/Persona/Animation", Icon20x20));
+		Set("Persona.AssetClass.SkeletalMesh", new IMAGE_BRUSH_SVG("Starship/Persona/SkeletalMesh", Icon20x20));
+		Set("Persona.AssetClass.Blueprint", new IMAGE_BRUSH_SVG( "Starship/MainToolbar/blueprints", Icon20x20) );
+		Set("Persona.AssetClass.Physics", new IMAGE_BRUSH_SVG("Starship/Persona/Physics", Icon20x20));
 	}
 
 	// Skeleton editor
diff --git a/Engine/Source/Editor/Persona/Private/PersonaAssetFamily.cpp b/Engine/Source/Editor/Persona/Private/PersonaAssetFamily.cpp
--- a/Engine/Source/Editor/Persona/Private/PersonaAssetFamily.cpp
+++ b/Engine/Source/Editor/Persona/Private/PersonaAssetFamily.cpp
@@ -219,6 +219,35 @@ FText FPersonaAssetFamily::GetAssetTypeDisplayName(UClass* InAssetClass) const
 	return FText();
 }
 
+const FSlateBrush* FPersonaAssetFamily::GetAssetTypeDisplayIcon(UClass* InAssetClass) const
+{
+	if (InAssetClass)
+	{
+		if (InAssetClass->IsChildOf<USkeleton>())
+		{
+			return FAppStyle::Get().GetBrush("Persona.AssetClass.Skeleton");
+		}
+		else if (InAssetClass->IsChildOf<UAnimationAsset>())
+		{
+			return FAppStyle::Get().GetBrush("Persona.AssetClass.Animation");
+		}
+		else if (InAssetClass->IsChildOf<USkeletalMesh>())
+		{
+			return FAppStyle::Get().GetBrush("Persona.AssetClass.SkeletalMesh");
+		}
+		else if (InAssetClass->IsChildOf<UAnimBlueprint>())
+		{
+			return FAppStyle::Get().GetBrush("Persona.AssetClass.Blueprint");
+		}
+		else if (InAssetClass->IsChildOf<UPhysicsAsset>())
+		{
+			return FAppStyle::Get().GetBrush("Persona.AssetClass.Physics");
+		}
+	}
+
+	return nullptr;
+}	
+
 bool FPersonaAssetFamily::IsAssetCompatible(const FAssetData& InAssetData) const
 {
 	UClass* Class = InAssetData.GetClass();
diff --git a/Engine/Source/Editor/Persona/Private/PersonaAssetFamily.h b/Engine/Source/Editor/Persona/Private/PersonaAssetFamily.h
--- a/Engine/Source/Editor/Persona/Private/PersonaAssetFamily.h
+++ b/Engine/Source/Editor/Persona/Private/PersonaAssetFamily.h
@@ -23,6 +23,7 @@ public:
 	virtual FAssetData FindAssetOfType(UClass* InAssetClass) const override;
 	virtual void FindAssetsOfType(UClass* InAssetClass, TArray<FAssetData>& OutAssets) const override;
 	virtual FText GetAssetTypeDisplayName(UClass* InAssetClass) const override;
+	virtual const FSlateBrush* GetAssetTypeDisplayIcon(UClass* InAssetClass) const override;
 	virtual bool IsAssetCompatible(const FAssetData& InAssetData) const override;
 	virtual UClass* GetAssetFamilyClass(UClass* InClass) const override;
 	virtual void RecordAssetOpened(const FAssetData& InAssetData) override;
diff --git a/Engine/Source/Editor/Persona/Private/SAssetFamilyShortcutBar.cpp b/Engine/Source/Editor/Persona/Private/SAssetFamilyShortcutBar.cpp
--- a/Engine/Source/Editor/Persona/Private/SAssetFamilyShortcutBar.cpp
+++ b/Engine/Source/Editor/Persona/Private/SAssetFamilyShortcutBar.cpp
@@ -12,6 +12,7 @@
 #include "Framework/MultiBox/MultiBoxDefs.h"
 #include "Widgets/Text/STextBlock.h"
 #include "Framework/MultiBox/MultiBoxBuilder.h"
+#include "Widgets/Layout/SSeparator.h"
 
 #include "EditorStyleSet.h"
 #include "IAssetFamily.h"
@@ -65,105 +66,103 @@ public:
 		TArray<FAssetData> Assets;
 		InAssetFamily->FindAssetsOfType(InAssetData.GetClass(), Assets);
 		bMultipleAssetsExist = Assets.Num() > 1;
-		AssetDirtyBrush = FEditorStyle::GetBrush("ContentBrowser.ContentDirty");
+		AssetDirtyBrush = FAppStyle::Get().GetBrush("Icons.DirtyBadge");
 
 		const FToolBarStyle& ToolBarStyle = FEditorStyle::Get().GetWidgetStyle<FToolBarStyle>("ToolBar");
 
 		ChildSlot
 		[
 			SNew(SHorizontalBox)
+
+			// This is the fat button for when there are not multiple options
 			+SHorizontalBox::Slot()
+			.VAlign(VAlign_Center)
 			[
+
 				SAssignNew(CheckBox, SCheckBox)
-				.Style(&ToolBarStyle.ToggleButton)
-				.ForegroundColor(FSlateColor::UseForeground())
-				.Padding(0.0f)
+				.Style(FAppStyle::Get(), "SegmentedCombo.ButtonOnly")
 				.OnCheckStateChanged(this, &SAssetShortcut::HandleOpenAssetShortcut)
 				.IsChecked(this, &SAssetShortcut::GetCheckState)
-				.Visibility(this, &SAssetShortcut::GetButtonVisibility)
+				.Visibility(this, &SAssetShortcut::GetSoloButtonVisibility)
 				.ToolTipText(this, &SAssetShortcut::GetButtonTooltip)
+				.Padding(0.0)
 				[
-					SNew(SHorizontalBox)
-					+ SHorizontalBox::Slot()
+					SNew(SOverlay)
+
+					+ SOverlay::Slot()
 					.VAlign(VAlign_Center)
-					.AutoWidth()
+					.HAlign(HAlign_Center)
+					.Padding(FMargin(28.f, 4.f))
 					[
-						//SNew(SBorder)
-						//.Padding(4.0f)
-						//.BorderImage(FEditorStyle::GetBrush("PropertyEditor.AssetThumbnailShadow"))
-						//[
-							SNew(SHorizontalBox)
-							+SHorizontalBox::Slot()
-							[
-								SAssignNew(ThumbnailBox, SBox)
-								.WidthOverride(AssetShortcutConstants::ThumbnailSize)
-								.HeightOverride(AssetShortcutConstants::ThumbnailSize)
-								.Visibility(this, &SAssetShortcut::GetThumbnailVisibility)
-								[
-									SNew(SOverlay)
-									+SOverlay::Slot()
-									[
-										AssetThumbnail->MakeThumbnailWidget()
-									]
-									+SOverlay::Slot()
-									.HAlign(HAlign_Left)
-									.VAlign(VAlign_Bottom)
-									[
-										SNew(SImage)
-										.Image(this, &SAssetShortcut::GetDirtyImage)
-									]
-								]
-							]
-							+SHorizontalBox::Slot()
-							[
-								SAssignNew(ThumbnailSmallBox, SBox)
-								.WidthOverride(AssetShortcutConstants::ThumbnailSizeSmall)
-								.HeightOverride(AssetShortcutConstants::ThumbnailSizeSmall)
-								.Visibility(this, &SAssetShortcut::GetSmallThumbnailVisibility)
-								[
-									SNew(SOverlay)
-									+SOverlay::Slot()
-									[
-										AssetThumbnailSmall->MakeThumbnailWidget()
-									]
-									+SOverlay::Slot()
-									.HAlign(HAlign_Left)
-									.VAlign(VAlign_Bottom)
-									[
-										SNew(SImage)
-										.Image(this, &SAssetShortcut::GetDirtyImage)
-									]
-								]
-							]
-						//]
+						SNew(SImage)
+						.ColorAndOpacity(FSlateColor::UseForeground())
+						.Image(this, &SAssetShortcut::GetAssetIcon)
 					]
-					+SHorizontalBox::Slot()
+
+					+ SOverlay::Slot()
+					.VAlign(VAlign_Bottom)
+					.HAlign(HAlign_Right)
+					.Padding(FMargin(2.f, 2.f))
+					[
+						SNew(SImage)
+						.ColorAndOpacity(FSlateColor::UseForeground())
+						.Image(this, &SAssetShortcut::GetDirtyImage)
+					]
+				]
+			]
+
+			// This is the left half of the button / combo pair for when there are multiple options
+			+SHorizontalBox::Slot()
+			.VAlign(VAlign_Center)
+			[
+				SAssignNew(CheckBox, SCheckBox)
+				.Style(FAppStyle::Get(), "SegmentedCombo.Left")
+				.OnCheckStateChanged(this, &SAssetShortcut::HandleOpenAssetShortcut)
+				.IsChecked(this, &SAssetShortcut::GetCheckState)
+				.Visibility(this, &SAssetShortcut::GetComboButtonVisibility)
+				.ToolTipText(this, &SAssetShortcut::GetButtonTooltip)
+				.Padding(0.0)
+				[
+					SNew(SOverlay)
+
+					+ SOverlay::Slot()
 					.VAlign(VAlign_Center)
-					.AutoWidth()
+					.HAlign(HAlign_Center)
+					.Padding(FMargin(16.f, 4.f))
+					[
+						SNew(SImage)
+						.ColorAndOpacity(FSlateColor::UseForeground())
+						.Image(this, &SAssetShortcut::GetAssetIcon)
+					]
+
+					+ SOverlay::Slot()
+					.VAlign(VAlign_Bottom)
+					.HAlign(HAlign_Right)
+					.Padding(FMargin(2.f, 2.f))
 					[
-						SNew(SHorizontalBox)
-						+SHorizontalBox::Slot()
-						.AutoWidth()
-						.VAlign(VAlign_Center)
-						.Padding(4.0f)
-						[
-							SNew(STextBlock)
-							.Text(this, &SAssetShortcut::GetAssetText)
-							.TextStyle(&ToolBarStyle.LabelStyle)
-						]
+						SNew(SImage)
+						.ColorAndOpacity(FSlateColor::UseForeground())
+						.Image(this, &SAssetShortcut::GetDirtyImage)
 					]
 				]
 			]
 			+SHorizontalBox::Slot()
-			.VAlign(VAlign_Fill)
+			.VAlign(VAlign_Center)
+			.AutoWidth()
+			[
+				SNew(SSeparator)
+				.Thickness(1.0)
+				.Orientation(EOrientation::Orient_Vertical)
+			]
+			+SHorizontalBox::Slot()
+			.VAlign(VAlign_Center)
 			.AutoWidth()
-			.Padding(2.0f, 0.0f, 0.0f, 0.0f)
 			[
 				SNew(SComboButton)
 				.Visibility(this, &SAssetShortcut::GetComboVisibility)
-				.ContentPadding(0)
+				.ContentPadding(FMargin(7.f, 0.f))
 				.ForegroundColor(FSlateColor::UseForeground())
-				.ButtonStyle(&ToolBarStyle.ButtonStyle)
+				.ComboButtonStyle(&FAppStyle::Get(), "SegmentedCombo.Right")
 				.OnGetMenuContent(this, &SAssetShortcut::HandleGetMenuContent)
 				.ToolTipText(LOCTEXT("AssetComboTooltip", "Find other assets of this type and perform asset operations./nShift-Click to open in new window."))
 			]
@@ -212,6 +211,11 @@ public:
 		return AssetFamily->GetAssetTypeDisplayName(AssetData.GetClass());
 	}
 
+	const FSlateBrush* GetAssetIcon() const 
+	{
+		return AssetFamily->GetAssetTypeDisplayIcon(AssetData.GetClass());	
+	}
+
 	ECheckBoxState GetCheckState() const
 	{
 		if(HostingApp.IsValid())
@@ -316,9 +320,14 @@ public:
 		return !AssetFamily->IsAssetCompatible(InAssetData);
 	}
 
-	EVisibility GetButtonVisibility() const
+	EVisibility GetSoloButtonVisibility() const
+	{
+		return AssetData.IsValid() && !bMultipleAssetsExist ? EVisibility::Visible : EVisibility::Collapsed;
+	}
+
+	EVisibility GetComboButtonVisibility() const
 	{
-		return AssetData.IsValid() || bMultipleAssetsExist ? EVisibility::Visible : EVisibility::Collapsed;
+		return AssetData.IsValid() && bMultipleAssetsExist ? EVisibility::Visible : EVisibility::Collapsed;
 	}
 
 	EVisibility GetComboVisibility() const
@@ -514,7 +523,7 @@ void SAssetFamilyShortcutBar::Construct(const FArguments& InArgs, const TSharedR
 		FAssetData AssetData = InAssetFamily->FindAssetOfType(Class);
 		HorizontalBox->AddSlot()
 		.AutoWidth()
-		.Padding(0.0f, 0.0f, AssetTypeIndex == AssetTypes.Num() - 1 ? 0.0f: 2.0f, 0.0f)
+		.Padding(0.0f, 4.0f, 16.0f, 4.0f)
 		[
 			SNew(SAssetShortcut, InHostingApp, InAssetFamily, AssetData, ThumbnailPool.ToSharedRef())
 		];
diff --git a/Engine/Source/Editor/Persona/Public/IAssetFamily.h b/Engine/Source/Editor/Persona/Public/IAssetFamily.h
--- a/Engine/Source/Editor/Persona/Public/IAssetFamily.h
+++ b/Engine/Source/Editor/Persona/Public/IAssetFamily.h
@@ -38,6 +38,9 @@ public:
 	/** Gets the name of an asset that will be displayed to a user */
 	virtual FText GetAssetTypeDisplayName(UClass* InAssetClass) const = 0;
 
+	/** Gets the slate brush that represents this asset family */
+	virtual const FSlateBrush* GetAssetTypeDisplayIcon(UClass* InAssetClass) const = 0;
+
 	/** Check whether an asset is compatible with this family */
 	virtual bool IsAssetCompatible(const FAssetData& InAssetData) const = 0;
 
diff --git a/Engine/Source/Editor/PhysicsAssetEditor/Private/PhysicsAssetEditorMode.cpp b/Engine/Source/Editor/PhysicsAssetEditor/Private/PhysicsAssetEditorMode.cpp
--- a/Engine/Source/Editor/PhysicsAssetEditor/Private/PhysicsAssetEditorMode.cpp
+++ b/Engine/Source/Editor/PhysicsAssetEditor/Private/PhysicsAssetEditorMode.cpp
@@ -215,13 +215,6 @@ FPhysicsAssetEditorMode::FPhysicsAssetEditorMode(TSharedRef<FWorkflowCentricAppl
 			FTabManager::NewPrimaryArea()
 			->SetOrientation(Orient_Vertical)
 			->Split
-			(
-				FTabManager::NewStack()
-				->SetSizeCoefficient(0.1f)
-				->SetHideTabWell(true)
-				->AddTab(InHostingApp->GetToolbarTabId(), ETabState::OpenedTab)
-			)
-			->Split
 			(
 				FTabManager::NewSplitter()
 				->SetSizeCoefficient(0.9f)
diff --git a/Engine/Source/Editor/SkeletalMeshEditor/Private/SkeletalMeshEditor.cpp b/Engine/Source/Editor/SkeletalMeshEditor/Private/SkeletalMeshEditor.cpp
--- a/Engine/Source/Editor/SkeletalMeshEditor/Private/SkeletalMeshEditor.cpp
+++ b/Engine/Source/Editor/SkeletalMeshEditor/Private/SkeletalMeshEditor.cpp
@@ -221,13 +221,12 @@ void FSkeletalMeshEditor::InitSkeletalMeshEditor(const EToolkitMode::Type Mode,
 	ISkeletonEditorModule& SkeletonEditorModule = FModuleManager::GetModuleChecked<ISkeletonEditorModule>("SkeletonEditor");
 	SkeletonTree = SkeletonEditorModule.CreateSkeletonTree(PersonaToolkit->GetSkeleton(), SkeletonTreeArgs);
 
-	BindCommands();
-	RegisterToolbar();
-
 	const bool bCreateDefaultStandaloneMenu = true;
 	const bool bCreateDefaultToolbar = true;
 	FAssetEditorToolkit::InitAssetEditor(Mode, InitToolkitHost, SkeletalMeshEditorAppIdentifier, FTabManager::FLayout::NullLayout, bCreateDefaultStandaloneMenu, bCreateDefaultToolbar, InSkeletalMesh);
 
+	BindCommands();
+
 	AddApplicationMode(
 		SkeletalMeshEditorModes::SkeletalMeshEditorMode,
 		MakeShareable(new FSkeletalMeshEditorMode(SharedThis(this), SkeletonTree.ToSharedRef())));
@@ -506,6 +505,41 @@ void FSkeletalMeshEditor::RegisterReimportContextMenu(const FName InBaseMenuName
 
 void FSkeletalMeshEditor::ExtendToolbar()
 {
+
+	// Add in Editor Specific functionality
+	FName ParentName;
+	static const FName MenuName = GetToolMenuToolbarName(ParentName);
+	RegisterReimportContextMenu(MenuName);
+
+	UToolMenu* ToolMenu = UToolMenus::Get()->ExtendMenu(MenuName);
+	const FToolMenuInsert SectionInsertLocation("Asset", EToolMenuInsertType::After);
+
+	{
+		ToolMenu->AddDynamicSection("Persona", FNewToolBarDelegateLegacy::CreateLambda([](FToolBarBuilder& ToolbarBuilder, UToolMenu* InMenu)
+		{
+			TSharedPtr<FSkeletalMeshEditor> SkeletalMeshEditor = GetSkeletalMeshEditor(InMenu->Context);
+			if (SkeletalMeshEditor.IsValid() && SkeletalMeshEditor->PersonaToolkit.IsValid())
+			{
+				FPersonaModule& PersonaModule = FModuleManager::LoadModuleChecked<FPersonaModule>("Persona");
+				FPersonaModule::FCommonToolbarExtensionArgs Args;
+				Args.bPreviewMesh = false;
+				PersonaModule.AddCommonToolbarExtensions(ToolbarBuilder, SkeletalMeshEditor->PersonaToolkit.ToSharedRef(), Args);
+			}
+		}), SectionInsertLocation);
+	}
+
+	{
+		FToolMenuSection& Section = ToolMenu->AddSection("Mesh", FText(), SectionInsertLocation);
+		Section.AddEntry(FToolMenuEntry::InitToolBarButton(FSkeletalMeshEditorCommands::Get().ReimportMesh));
+		Section.AddEntry(FToolMenuEntry::InitComboButton("ReimportContextMenu", FUIAction(), FNewToolMenuDelegate()));
+	}
+
+	{
+		FToolMenuSection& Section = ToolMenu->AddSection("SkeletalMesh", FText(), FToolMenuInsert("Mesh", EToolMenuInsertType::After));
+		Section.AddEntry(FToolMenuEntry::InitToolBarButton(FSkeletalMeshEditorCommands::Get().MeshSectionSelection));
+	}
+
+
 	// If the ToolbarExtender is valid, remove it before rebuilding it
 	if (ToolbarExtender.IsValid())
 	{
@@ -544,44 +578,6 @@ void FSkeletalMeshEditor::ExtendToolbar()
 	));
 }
 
-void FSkeletalMeshEditor::RegisterToolbar()
-{
-	static const FName MenuName = "AssetEditor.SkeletalMeshEditor.ToolBar";
-	RegisterReimportContextMenu(MenuName);
-
-	if (!UToolMenus::Get()->IsMenuRegistered(MenuName))
-	{
-		UToolMenu* ToolMenu = UToolMenus::Get()->RegisterMenu(MenuName, "AssetEditor.DefaultToolBar", EMultiBoxType::ToolBar);
-
-		const FToolMenuInsert SectionInsertLocation("Asset", EToolMenuInsertType::After);
-
-		{
-			ToolMenu->AddDynamicSection("Persona", FNewToolBarDelegateLegacy::CreateLambda([](FToolBarBuilder& ToolbarBuilder, UToolMenu* InMenu)
-			{
-				TSharedPtr<FSkeletalMeshEditor> SkeletalMeshEditor = GetSkeletalMeshEditor(InMenu->Context);
-				if (SkeletalMeshEditor.IsValid() && SkeletalMeshEditor->PersonaToolkit.IsValid())
-				{
-					FPersonaModule& PersonaModule = FModuleManager::LoadModuleChecked<FPersonaModule>("Persona");
-					FPersonaModule::FCommonToolbarExtensionArgs Args;
-					Args.bPreviewMesh = false;
-					PersonaModule.AddCommonToolbarExtensions(ToolbarBuilder, SkeletalMeshEditor->PersonaToolkit.ToSharedRef(), Args);
-				}
-			}), SectionInsertLocation);
-		}
-
-		{
-			FToolMenuSection& Section = ToolMenu->AddSection("Mesh", FText(), SectionInsertLocation);
-			Section.AddEntry(FToolMenuEntry::InitToolBarButton(FSkeletalMeshEditorCommands::Get().ReimportMesh));
-			Section.AddEntry(FToolMenuEntry::InitComboButton("ReimportContextMenu", FUIAction(), FNewToolMenuDelegate()));
-		}
-
-		{
-			FToolMenuSection& Section = ToolMenu->AddSection("SkeletalMesh", FText(), FToolMenuInsert("Mesh", EToolMenuInsertType::After));
-			Section.AddEntry(FToolMenuEntry::InitToolBarButton(FSkeletalMeshEditorCommands::Get().MeshSectionSelection));
-		}
-	}
-}
-
 void FSkeletalMeshEditor::InitToolMenuContext(FToolMenuContext& MenuContext)
 {
 	FAssetEditorToolkit::InitToolMenuContext(MenuContext);
diff --git a/Engine/Source/Editor/SkeletalMeshEditor/Private/SkeletalMeshEditor.h b/Engine/Source/Editor/SkeletalMeshEditor/Private/SkeletalMeshEditor.h
--- a/Engine/Source/Editor/SkeletalMeshEditor/Private/SkeletalMeshEditor.h
+++ b/Engine/Source/Editor/SkeletalMeshEditor/Private/SkeletalMeshEditor.h
@@ -147,7 +147,6 @@ private:
 private:
 	void ExtendMenu();
 
-	void RegisterToolbar();
 	void ExtendToolbar();
 
 	void BindCommands();
diff --git a/Engine/Source/Editor/SkeletalMeshEditor/Private/SkeletalMeshEditorMode.cpp b/Engine/Source/Editor/SkeletalMeshEditor/Private/SkeletalMeshEditorMode.cpp
--- a/Engine/Source/Editor/SkeletalMeshEditor/Private/SkeletalMeshEditorMode.cpp
+++ b/Engine/Source/Editor/SkeletalMeshEditor/Private/SkeletalMeshEditorMode.cpp
@@ -41,13 +41,6 @@ FSkeletalMeshEditorMode::FSkeletalMeshEditorMode(TSharedRef<FWorkflowCentricAppl
 			FTabManager::NewPrimaryArea()
 			->SetOrientation(Orient_Vertical)
 			->Split
-			(
-				FTabManager::NewStack()
-				->SetSizeCoefficient(0.1f)
-				->SetHideTabWell(true)
-				->AddTab(InHostingApp->GetToolbarTabId(), ETabState::OpenedTab)
-			)
-			->Split
 			(
 				FTabManager::NewSplitter()
 				->SetSizeCoefficient(0.9f)
diff --git a/Engine/Source/Editor/SkeletonEditor/Private/SkeletonEditorMode.cpp b/Engine/Source/Editor/SkeletonEditor/Private/SkeletonEditorMode.cpp
--- a/Engine/Source/Editor/SkeletonEditor/Private/SkeletonEditorMode.cpp
+++ b/Engine/Source/Editor/SkeletonEditor/Private/SkeletonEditorMode.cpp
@@ -41,13 +41,6 @@ FSkeletonEditorMode::FSkeletonEditorMode(TSharedRef<FWorkflowCentricApplication>
 			FTabManager::NewPrimaryArea()
 			->SetOrientation(Orient_Vertical)
 			->Split
-			(
-				FTabManager::NewStack()
-				->SetSizeCoefficient(0.1f)
-				->SetHideTabWell(true)
-				->AddTab(InHostingApp->GetToolbarTabId(), ETabState::OpenedTab)
-			)
-			->Split
 			(
 				FTabManager::NewSplitter()
 				->SetSizeCoefficient(0.9f)
diff --git a/Engine/Source/Editor/UnrealEd/Private/Toolkits/AssetEditorToolkit.cpp b/Engine/Source/Editor/UnrealEd/Private/Toolkits/AssetEditorToolkit.cpp
--- a/Engine/Source/Editor/UnrealEd/Private/Toolkits/AssetEditorToolkit.cpp
+++ b/Engine/Source/Editor/UnrealEd/Private/Toolkits/AssetEditorToolkit.cpp
@@ -1104,7 +1104,6 @@ void FAssetEditorToolkit::GenerateToolbar()
 			MiscWidgetsHBox->AddSlot()
 				.AutoWidth()
 				.VAlign(VAlign_Center)
-				.Padding(0.0f, 2.0f, 0.0f, 2.0f)
 				[
 					ToolbarWidgets[WidgetIdx]
 				];
@@ -1125,10 +1124,10 @@ void FAssetEditorToolkit::GenerateToolbar()
 		.VAlign(VAlign_Center)
 		.AutoWidth()
 		[
-			SNew(SVerticalBox)
-			+SVerticalBox::Slot()
-			.AutoHeight()
-			.VAlign(VAlign_Bottom)
+			SNew(SBorder)
+			.VAlign(VAlign_Center)
+			.BorderImage(FAppStyle::Get().GetBrush("Brushes.Background"))
+			.Padding(FMargin(0.0f))
 			[
 				MiscWidgets
 			]
@@ -1165,7 +1164,6 @@ void FAssetEditorToolkit::RegenerateMenusAndToolbars()
 }
 
 
-
 void FAssetEditorToolkit::RestoreFromLayout(const TSharedRef<FTabManager::FLayout>& NewLayout)
 {
 	TSharedPtr< class SStandaloneAssetEditorToolkitHost > HostWidget = StandaloneHost.Pin();
diff --git a/Engine/Source/Editor/UnrealEd/Public/WorkflowOrientedApp/WorkflowCentricApplication.h b/Engine/Source/Editor/UnrealEd/Public/WorkflowOrientedApp/WorkflowCentricApplication.h
--- a/Engine/Source/Editor/UnrealEd/Public/WorkflowOrientedApp/WorkflowCentricApplication.h
+++ b/Engine/Source/Editor/UnrealEd/Public/WorkflowOrientedApp/WorkflowCentricApplication.h
@@ -20,6 +20,7 @@ class UNREALED_API FWorkflowCentricApplication : public FAssetEditorToolkit
 {
 public:
 	// IToolkit interface
+	virtual bool UsesCustomToolbarPlacement() const { return false;  }
 	virtual void RegisterTabSpawners(const TSharedRef<class FTabManager>& TabManager) override;
 	virtual void UnregisterTabSpawners(const TSharedRef<class FTabManager>& TabManager) override;
 	// End of IToolkit interface
diff --git a/Engine/Source/Runtime/SlateCore/Private/Styling/StarshipCoreStyle.cpp b/Engine/Source/Runtime/SlateCore/Private/Styling/StarshipCoreStyle.cpp
--- a/Engine/Source/Runtime/SlateCore/Private/Styling/StarshipCoreStyle.cpp
+++ b/Engine/Source/Runtime/SlateCore/Private/Styling/StarshipCoreStyle.cpp
@@ -1362,6 +1362,45 @@ void FStarshipCoreStyle::SetupCheckboxStyles(TSharedRef<FStyle>& Style)
 		.SetLastControlStyle(SegmentedBoxRight)
 	);
 
+	Style->Set("SegmentedCombo.Left", SegmentedBoxLeft);
+
+	const FButtonStyle SegmentedComboRightButton = FButtonStyle()
+	.SetNormal(BOX_BRUSH("/Starship/CoreWidgets/SegmentedBox/right", FVector2D(16.f, 16.f), FMargin(4.0f / 16.0f), FStyleColors::Secondary))
+	.SetHovered(BOX_BRUSH("/Starship/CoreWidgets/SegmentedBox/right", FVector2D(16.f, 16.f), FMargin(4.0f / 16.0f), FStyleColors::Hover))
+	.SetPressed(BOX_BRUSH("/Starship/CoreWidgets/SegmentedBox/right", FVector2D(16.f, 16.f), FMargin(4.0f / 16.0f), FStyleColors::Secondary))
+	.SetNormalForeground(FStyleColors::Foreground)
+	.SetHoveredForeground(FStyleColors::ForegroundHover)
+	.SetPressedForeground(FStyleColors::Primary)
+	.SetDisabledForeground(FStyleColors::Foreground)
+	.SetNormalPadding(0)
+	.SetPressedPadding(0);
+
+	FComboButtonStyle SegmentedComboRight = FComboButtonStyle(Style->GetWidgetStyle<FComboButtonStyle>("ComboButton"))
+		.SetContentPadding(FMargin(6.f, 2.f))
+		.SetButtonStyle(SegmentedComboRightButton)
+		.SetDownArrowImage(IMAGE_BRUSH_SVG("Starship/Common/ellipsis-vertical-narrow", FVector2D(6, 24)));
+
+
+	Style->Set("SegmentedCombo.Right", SegmentedComboRight);
+
+	// Button that looks/behaves like a SegmentedCombo for cases where SegmentedCombos and 
+	// buttons are placed adjacent
+	const FCheckBoxStyle SegmentedBoxLeftOnly = FCheckBoxStyle()
+		.SetCheckBoxType(ESlateCheckBoxType::ToggleButton)
+		.SetUncheckedImage(       FSlateRoundedBoxBrush(FStyleColors::Secondary, 4.f) )
+		.SetUncheckedHoveredImage(FSlateRoundedBoxBrush(FStyleColors::Hover, 4.f))
+		.SetUncheckedPressedImage(FSlateRoundedBoxBrush(FStyleColors::Secondary, 4.f))
+		.SetCheckedImage(         FSlateRoundedBoxBrush(FStyleColors::Input, 4.f))
+		.SetCheckedHoveredImage(  FSlateRoundedBoxBrush(FStyleColors::Input, 4.f))
+		.SetCheckedPressedImage(  FSlateRoundedBoxBrush(FStyleColors::Input, 4.f))
+		.SetForegroundColor(FStyleColors::Foreground)
+		.SetHoveredForegroundColor(FStyleColors::ForegroundHover)
+		.SetPressedForegroundColor(FStyleColors::ForegroundHover)
+		.SetCheckedForegroundColor(FStyleColors::Primary)
+		.SetCheckedHoveredForegroundColor(FStyleColors::Primary)
+		.SetCheckedPressedForegroundColor(FStyleColors::Primary)
+		.SetPadding(DefaultMargins);
+	Style->Set("SegmentedCombo.ButtonOnly", SegmentedBoxLeftOnly);
 
 	/* Style for a toggleable button that mimics the coloring and look of a Table Row */
 	/*
