Fix for Crash opening animation sequence and blend spaces twice

diff --git a/Engine/Source/Editor/Persona/Private/SAssetFamilyShortcutBar.cpp b/Engine/Source/Editor/Persona/Private/SAssetFamilyShortcutBar.cpp
--- a/Engine/Source/Editor/Persona/Private/SAssetFamilyShortcutBar.cpp
+++ b/Engine/Source/Editor/Persona/Private/SAssetFamilyShortcutBar.cpp
@@ -439,9 +439,6 @@ public:
 		{
 			AssetThumbnail = MakeShareable(new FAssetThumbnail(AssetData, AssetShortcutConstants::ThumbnailSize, AssetShortcutConstants::ThumbnailSize, ThumbnailPoolPtr.Pin()));
 			AssetThumbnailSmall = MakeShareable(new FAssetThumbnail(AssetData, AssetShortcutConstants::ThumbnailSizeSmall, AssetShortcutConstants::ThumbnailSizeSmall, ThumbnailPoolPtr.Pin()));
-
-			ThumbnailBox->SetContent(AssetThumbnail->MakeThumbnailWidget());
-			ThumbnailSmallBox->SetContent(AssetThumbnailSmall->MakeThumbnailWidget());
 		}
 	}
 
@@ -485,10 +482,6 @@ private:
 	TSharedPtr<FAssetThumbnail> AssetThumbnail;
 	TSharedPtr<FAssetThumbnail> AssetThumbnailSmall;
 
-	/** Thumbnail widget containers */
-	TSharedPtr<SBox> ThumbnailBox;
-	TSharedPtr<SBox> ThumbnailSmallBox;
-
 	/** The asset editor we are embedded in */
 	TWeakPtr<class FWorkflowCentricApplication> HostingApp;
 
