Persona: Capture the skeletal mesh reduction in the undo stack.
This shelf has been sitting around for a bit - lina gave this the thumbs up before she left.

diff --git a/Engine/Source/Editor/Persona/Private/PersonaMeshDetails.cpp b/Engine/Source/Editor/Persona/Private/PersonaMeshDetails.cpp
--- a/Engine/Source/Editor/Persona/Private/PersonaMeshDetails.cpp
+++ b/Engine/Source/Editor/Persona/Private/PersonaMeshDetails.cpp
@@ -3376,6 +3376,12 @@ FReply FPersonaMeshDetails::ApplyLODChanges(int32 LODIndex)
 		return FReply::Handled();
 	}
 	
+	// inject a transaction and mark the mesh as modified so we can undo this state
+	FText TransactionText = FText::Format(LOCTEXT("PersonaReductionApplyLODChanges", "LOD{0} generation"), LODIndex);
+	FScopedTransaction Transaction(TransactionText);
+	SkelMesh->Modify();
+	
+
 	FScopedSuspendAlternateSkinWeightPreview ScopedSuspendAlternateSkinnWeightPreview(SkelMesh);
 	{
 		FScopedSkeletalMeshPostEditChange ScopedPostEditChange(SkelMesh);
