Add wrapping on the BlendSpace 1D and 2D. Also clamps/wraps the preview rendering of the BlendSpace inside the Anim Graph to avoid it straying outside the node.

diff --git a/Engine/Source/Editor/Persona/Private/SAnimationBlendSpaceGridWidget.cpp b/Engine/Source/Editor/Persona/Private/SAnimationBlendSpaceGridWidget.cpp
--- a/Engine/Source/Editor/Persona/Private/SAnimationBlendSpaceGridWidget.cpp
+++ b/Engine/Source/Editor/Persona/Private/SAnimationBlendSpaceGridWidget.cpp
@@ -1098,7 +1098,7 @@ const FVector2D SBlendSpaceGridWidget::SampleValueToGridPosition(const FVector&
 	return SamplePosition2D;	
 }
 
-const FVector SBlendSpaceGridWidget::GridPositionToSampleValue(const FVector2D& GridPosition) const
+const FVector SBlendSpaceGridWidget::GridPositionToSampleValue(const FVector2D& GridPosition, bool bClamp) const
 {
 	FVector2D LocalGridPosition = GridPosition;
 	// Move to center of grid and convert to 0 - 1 form
@@ -1108,13 +1108,17 @@ const FVector SBlendSpaceGridWidget::GridPositionToSampleValue(const FVector2D&
 	LocalGridPosition *= 0.5f;
 
 	// Calculate the sample value by mapping it to the blend parameter range
-	const FVector SampleValue
+	FVector SampleValue
 	(		
-		FMath::Clamp((LocalGridPosition.X * SampleValueRange.X) + SampleValueMin.X, SampleValueMin.X, SampleValueMax.X),
-		FMath::Clamp(((GridType == EGridType::TwoAxis) ? SampleValueMax.Y - (LocalGridPosition.Y * SampleValueRange.Y) : 0.0f),
-			SampleValueMin.Y, SampleValueMax.Y),
+		(LocalGridPosition.X * SampleValueRange.X) + SampleValueMin.X,
+		(GridType == EGridType::TwoAxis) ? SampleValueMax.Y - (LocalGridPosition.Y * SampleValueRange.Y) : 0.0f,
 		0.f	
 	);
+	if (bClamp)
+	{
+		SampleValue.X = FMath::Clamp(SampleValue.X, SampleValueMin.X, SampleValueMax.X);
+		SampleValue.Y = FMath::Clamp(SampleValue.Y, SampleValueMin.Y, SampleValueMax.Y);
+	}
 	return SampleValue;
 }
 
@@ -1123,13 +1127,13 @@ const FSlateRect SBlendSpaceGridWidget::GetGridRectangleFromGeometry(const FGeom
 	FSlateRect WindowRect = FSlateRect(0, 0, MyGeometry.GetLocalSize().X, MyGeometry.GetLocalSize().Y);
 	if (!bStretchToFit)
 	{
-		UpdateGridRationMargin(WindowRect.GetSize());
+		UpdateGridRatioMargin(WindowRect.GetSize());
 	}
 
 	return WindowRect.InsetBy(GridMargin + GridRatioMargin);
 }
 
-bool SBlendSpaceGridWidget::IsSampleValueWithinMouseRange(const FVector& SampleValue)
+bool SBlendSpaceGridWidget::IsSampleValueWithinMouseRange(const FVector& SampleValue) const
 {
 	const FVector2D GridPosition = SampleValueToGridPosition(SampleValue);
 	const float MouseDistance = FVector2D::Distance(LocalMousePosition, GridPosition);	
@@ -1140,7 +1144,7 @@ void SBlendSpaceGridWidget::StartPreviewing()
 {
 	bSamplePreviewing = true;
 	LastPreviewingMousePosition = LocalMousePosition;
-	LastPreviewingSampleValue = bReadOnly ? Position.Get() : GridPositionToSampleValue(LastPreviewingMousePosition);
+	LastPreviewingSampleValue = bReadOnly ? Position.Get() : GridPositionToSampleValue(LastPreviewingMousePosition, false);
 	bPreviewPositionSet = true;	
 	bPreviewToolTipHidden = true;
 }
@@ -1494,7 +1498,7 @@ FReply SBlendSpaceGridWidget::ToggleShowAnimationNames()
 	return FReply::Handled();
 }
 
-void SBlendSpaceGridWidget::UpdateGridRationMargin(const FVector2D& GeometrySize)
+void SBlendSpaceGridWidget::UpdateGridRatioMargin(const FVector2D& GeometrySize)
 {
 	if (GridType == EGridType::TwoAxis)
 	{
@@ -1681,7 +1685,7 @@ void SBlendSpaceGridWidget::Tick(const FGeometry& AllottedGeometry, const double
 					const FVector2D GridPosition(FMath::Clamp(LocalMousePosition.X, CachedGridRectangle.Left, CachedGridRectangle.Right),
 													FMath::Clamp(LocalMousePosition.Y, CachedGridRectangle.Top, CachedGridRectangle.Bottom));
  
-					SampleValue = GridPositionToSampleValue(GridPosition);
+					SampleValue = GridPositionToSampleValue(GridPosition, true);
 					bSnap = false;
 				}
 				else
@@ -1739,9 +1743,11 @@ void SBlendSpaceGridWidget::Tick(const FGeometry& AllottedGeometry, const double
 		if (bSamplePreviewing)
 		{
 			// Ensure the preview mouse position is clamped to the grid
-			LastPreviewingMousePosition.X = FMath::Clamp(LocalMousePosition.X, CachedGridRectangle.Left, CachedGridRectangle.Right);
-			LastPreviewingMousePosition.Y = FMath::Clamp(LocalMousePosition.Y, CachedGridRectangle.Top, CachedGridRectangle.Bottom);
-			LastPreviewingSampleValue = bReadOnly ? Position.Get() : GridPositionToSampleValue(LastPreviewingMousePosition);
+			LastPreviewingMousePosition.X = LocalMousePosition.X; //FMath::Clamp(LocalMousePosition.X, CachedGridRectangle.Left, CachedGridRectangle.Right);
+			LastPreviewingMousePosition.Y = LocalMousePosition.Y; //FMath::Clamp(LocalMousePosition.Y, CachedGridRectangle.Top, CachedGridRectangle.Bottom);
+			LastPreviewingSampleValue = bReadOnly ? Position.Get() : GridPositionToSampleValue(LastPreviewingMousePosition, false);
+
+			LastPreviewingSampleValue = BlendSpace->GetClampedAndWrappedBlendInput(LastPreviewingSampleValue);
 						
 			// Retrieve and cache weighted samples
 			PreviewedSamples.Empty(4);
diff --git a/Engine/Source/Editor/Persona/Private/SAnimationBlendSpaceGridWidget.h b/Engine/Source/Editor/Persona/Private/SAnimationBlendSpaceGridWidget.h
--- a/Engine/Source/Editor/Persona/Private/SAnimationBlendSpaceGridWidget.h
+++ b/Engine/Source/Editor/Persona/Private/SAnimationBlendSpaceGridWidget.h
@@ -147,8 +147,8 @@ protected:
 	FReply ToggleFittingType();
 	/** Returns whether or not the fitting button should be visible */
 	EVisibility GetFittingButtonVisibility() const;
-	/** Calcaultes a margin offset according to whether or not we should take into account the largest axis when creating the grid area*/
-	void UpdateGridRationMargin(const FVector2D& GeometrySize);
+	/** Calculates a margin offset according to whether or not we should take into account the largest axis when creating the grid area*/
+	void UpdateGridRatioMargin(const FVector2D& GeometrySize);
 	FText GetFittingTypeButtonToolTipText() const;
 	/** Toggles the animation labels being shown */
 	FReply ToggleShowAnimationNames();
@@ -163,12 +163,12 @@ protected:
 	int32 FindClosestGridPointIndex(const FVector2D& InPosition) const;
 	/** Converts the given sample value to a screen space position */
 	const FVector2D SampleValueToGridPosition(const FVector& SampleValue) const;
-	/** Converst a screen space (grid) position to a valid sample value */
-	const FVector GridPositionToSampleValue(const FVector2D& GridPosition) const;
-	/** Returns the (calculated) grid rectangle given the supplied geometry */
+	/** Converst a screen space (grid) position to a sample value, with optional clamping */
+	const FVector GridPositionToSampleValue(const FVector2D& GridPosition, bool bClamp) const;
+	/** Returns the (calculated) grid rectangle given the supplied geometry, and calls UpdateGridRatioMargin */
 	const FSlateRect GetGridRectangleFromGeometry(const FGeometry& MyGeometry);
 	/** Checks whether or not the blendspace sample value is within the range of the mouse position */
-	bool IsSampleValueWithinMouseRange(const FVector& SampleValue);
+	bool IsSampleValueWithinMouseRange(const FVector& SampleValue) const;
 
 	/** Sets the tooltip instance on the underlying widget instance */
 	void ShowToolTip();
diff --git a/Engine/Source/Runtime/Engine/Classes/Animation/BlendSpaceBase.h b/Engine/Source/Runtime/Engine/Classes/Animation/BlendSpaceBase.h
--- a/Engine/Source/Runtime/Engine/Classes/Animation/BlendSpaceBase.h
+++ b/Engine/Source/Runtime/Engine/Classes/Animation/BlendSpaceBase.h
@@ -58,11 +58,17 @@ struct FBlendParameter
 	UPROPERTY(EditAnywhere, DisplayName = "Number of Grid Divisions", Category=BlendParameter, meta=(UIMin="1", ClampMin="1"))
 	int32 GridNum;
 
+	/** Max value for this parameter. */
+	UPROPERTY(EditAnywhere, DisplayName = "Wrap Input", Category = BlendParameter, meta=(ToolTip=
+		"If false then input parameters are clamped to the min/max values on this axis. If true then the input can go outside the min/max range and the blend space is treated as being cyclic on this axis."))
+	bool WrapInput;
+
 	FBlendParameter()
 		: DisplayName(TEXT("None"))
 		, Min(0.f)
 		, Max(100.f)
 		, GridNum(4) // TODO when changing GridNum's default value, it breaks all grid samples ATM - provide way to rebuild grid samples during loading
+		, WrapInput(false)
 	{
 	}
 
@@ -295,7 +301,10 @@ class UBlendSpaceBase : public UAnimationAsset, public IInterpolationIndexProvid
 
 	/** Initialize BlendSpace for runtime. It needs certain data to be reinitialized per instsance **/
 	ENGINE_API void InitializeFilter(FBlendFilter* Filter) const;
-	
+
+	/** Returns the blend input after clamping and/or wrapping */
+	ENGINE_API FVector GetClampedAndWrappedBlendInput(const FVector& BlendInput) const;
+
 #if WITH_EDITOR	
 	/** Validates sample data for blendspaces using the given animation sequence */
 	ENGINE_API static void UpdateBlendSpacesUsingAnimSequence(UAnimSequenceBase* Sequence);
@@ -385,8 +394,8 @@ protected:
 	/** Utility function to calculate animation length from sample data list **/
 	float GetAnimationLengthFromSampleData(const TArray<FBlendSampleData> & SampleDataList) const;
 
-	/** Clamp blend input to valid point **/
-	FVector ClampBlendInput(const FVector& BlendInput) const;
+	/** Returns the blend input clamped to the valid range, unless that axis has been set to wrap in which case no clamping is done **/
+	FVector GetClampedBlendInput(const FVector& BlendInput) const;
 	
 	/** Translates BlendInput to grid space */
 	FVector GetNormalizedBlendInput(const FVector& BlendInput) const;
diff --git a/Engine/Source/Runtime/Engine/Private/Animation/BlendSpaceBase.cpp b/Engine/Source/Runtime/Engine/Private/Animation/BlendSpaceBase.cpp
--- a/Engine/Source/Runtime/Engine/Private/Animation/BlendSpaceBase.cpp
+++ b/Engine/Source/Runtime/Engine/Private/Animation/BlendSpaceBase.cpp
@@ -153,13 +153,14 @@ void UBlendSpaceBase::TickAssetPlayer(FAnimTickRecord& Instance, struct FAnimNot
 				}
 			}
 
-			// now find if clamped input is different
-			// if different, then apply scale to fit in
-			FVector ClampedInput = ClampBlendInput(BlendInput);
+			// Now find if clamped input is different. If different, then apply scale to fit in. This allows
+			// "extrapolation" of the blend space outside of the range by time scaling the animation, which is
+			// appropriate when the specified axis is speed (for example).
+			FVector ClampedInput = GetClampedBlendInput(BlendInput);
 			if ( !ClampedInput.Equals(BlendInput) ) 
 			{
 				// apply speed change if you want, 
-				if (AxisToScale == BSA_X)
+				if (AxisToScale == BSA_X && !BlendParameters[0].WrapInput)
 				{
 					if (ClampedInput.X != 0.f)
 					{
@@ -168,7 +169,7 @@ void UBlendSpaceBase::TickAssetPlayer(FAnimTickRecord& Instance, struct FAnimNot
 				}
 				else if (AxisToScale == BSA_Y)
 				{
-					if (ClampedInput.Y != 0.f)
+					if (ClampedInput.Y != 0.f && !BlendParameters[1].WrapInput)
 					{
 						FilterMultiplier *= BlendInput.Y / ClampedInput.Y;
 					}
@@ -1143,30 +1144,44 @@ float UBlendSpaceBase::GetAnimationLengthFromSampleData(const TArray<FBlendSampl
 	return BlendAnimLength;
 }
 
-FVector UBlendSpaceBase::ClampBlendInput(const FVector& BlendInput)  const
+FVector UBlendSpaceBase::GetClampedBlendInput(const FVector& BlendInput)  const
 {
-	FVector ClampedBlendInput;
+	FVector AdjustedInput = BlendInput;
+	for (int iAxis = 0; iAxis != 3; ++iAxis)
+	{
+		if (!BlendParameters[iAxis].WrapInput)
+		{
+			AdjustedInput[iAxis] = FMath::Clamp(AdjustedInput[iAxis], BlendParameters[iAxis].Min, BlendParameters[iAxis].Max);
+		}
+	}
+	return AdjustedInput;
+}
 
-	ClampedBlendInput.X = FMath::Clamp<float>(BlendInput.X, BlendParameters[0].Min, BlendParameters[0].Max);
-	ClampedBlendInput.Y = FMath::Clamp<float>(BlendInput.Y, BlendParameters[1].Min, BlendParameters[1].Max);
-	ClampedBlendInput.Z = FMath::Clamp<float>(BlendInput.Z, BlendParameters[2].Min, BlendParameters[2].Max);
-	return ClampedBlendInput;
+FVector UBlendSpaceBase::GetClampedAndWrappedBlendInput(const FVector& BlendInput) const
+{
+	FVector AdjustedInput = BlendInput;
+	for (int iAxis = 0; iAxis != 3; ++iAxis)
+	{
+		if (BlendParameters[iAxis].WrapInput)
+		{
+			AdjustedInput[iAxis] = FMath::Wrap(AdjustedInput[iAxis], BlendParameters[iAxis].Min, BlendParameters[iAxis].Max);
+		}
+		else
+		{
+			AdjustedInput[iAxis] = FMath::Clamp(AdjustedInput[iAxis], BlendParameters[iAxis].Min, BlendParameters[iAxis].Max);
+		}
+	}
+	return AdjustedInput;
 }
 
 FVector UBlendSpaceBase::GetNormalizedBlendInput(const FVector& BlendInput) const
 {
+	FVector AdjustedInput = GetClampedAndWrappedBlendInput(BlendInput);
+
 	const FVector MinBlendInput = FVector(BlendParameters[0].Min, BlendParameters[1].Min, BlendParameters[2].Min);
-	const FVector MaxBlendInput = FVector(BlendParameters[0].Max, BlendParameters[1].Max, BlendParameters[2].Max);
 	const FVector GridSize = FVector(BlendParameters[0].GetGridSize(), BlendParameters[1].GetGridSize(), BlendParameters[2].GetGridSize());
 
-	FVector NormalizedBlendInput;
-	NormalizedBlendInput.X = FMath::Clamp<float>(BlendInput.X, MinBlendInput.X, MaxBlendInput.X);
-	NormalizedBlendInput.Y = FMath::Clamp<float>(BlendInput.Y, MinBlendInput.Y, MaxBlendInput.Y);
-	NormalizedBlendInput.Z = FMath::Clamp<float>(BlendInput.Z, MinBlendInput.Z, MaxBlendInput.Z);
-
-	NormalizedBlendInput -= MinBlendInput; 
-	NormalizedBlendInput /= GridSize;
-
+	FVector NormalizedBlendInput = (AdjustedInput - MinBlendInput) / GridSize;
 	return NormalizedBlendInput;
 }
 
