Added logic to check whether or not a curve was removed during a bracket before trying to update the underlying FRichCurve* from the model. Which at that point does not exist anymore

diff --git a/Engine/Source/Editor/Persona/Private/SAnimSequenceCurveEditor.cpp b/Engine/Source/Editor/Persona/Private/SAnimSequenceCurveEditor.cpp
--- a/Engine/Source/Editor/Persona/Private/SAnimSequenceCurveEditor.cpp
+++ b/Engine/Source/Editor/Persona/Private/SAnimSequenceCurveEditor.cpp
@@ -27,6 +27,7 @@ FRichCurveEditorModelNamed::FRichCurveEditorModelNamed(const FSmartName& InName,
 , Type(InType)
 , TreeId(InTreeId)
 , CurveId(FAnimationCurveIdentifier(Name, Type))
+, bCurveRemoved(false)
 {
 	CurveModifiedDelegate.AddRaw(this, &FRichCurveEditorModelNamed::CurveHasChanged);
 
@@ -92,13 +93,21 @@ void FRichCurveEditorModelNamed::OnModelHasChanged(const EAnimDataModelNotifyTyp
 		case EAnimDataModelNotifyType::CurveFlagsChanged:
 		case EAnimDataModelNotifyType::CurveScaled:
 		{
-			if (NotifyCollector.IsNotWithinBracket())
+			const FCurvePayload& TypedPayload = Payload.GetPayload<FCurvePayload>();
+			if (TypedPayload.Identifier.InternalName == Name)
 			{
-				const FCurvePayload& TypedPayload = Payload.GetPayload<FCurvePayload>();
-				if (TypedPayload.Identifier.InternalName == Name)
+				if (NotifyCollector.IsNotWithinBracket())
 				{
 					UpdateCachedCurve();
 				}
+				else
+				{
+					// Curve was re-added after removal in same bracket
+					if (bCurveRemoved && NotifyType == EAnimDataModelNotifyType::CurveAdded)
+					{
+						bCurveRemoved = false;
+					}
+				}
 			}
 
 			break;
@@ -107,6 +116,11 @@ void FRichCurveEditorModelNamed::OnModelHasChanged(const EAnimDataModelNotifyTyp
 		case EAnimDataModelNotifyType::CurveRemoved:
 		{
 			// Curve was removed
+			const FCurveRemovedPayload& TypedPayload = Payload.GetPayload<FCurveRemovedPayload>();
+			if (TypedPayload.Identifier.InternalName == Name)
+			{
+				bCurveRemoved = true;
+			}
 			break;
 		}
 
@@ -131,7 +145,7 @@ void FRichCurveEditorModelNamed::OnModelHasChanged(const EAnimDataModelNotifyTyp
 		{
 			if (NotifyCollector.IsNotWithinBracket())
 			{
-				if (NotifyCollector.Contains({ EAnimDataModelNotifyType::CurveAdded, EAnimDataModelNotifyType::CurveChanged, EAnimDataModelNotifyType::CurveFlagsChanged, EAnimDataModelNotifyType::CurveScaled, EAnimDataModelNotifyType::CurveRenamed }))
+				if (!bCurveRemoved && NotifyCollector.Contains({ EAnimDataModelNotifyType::CurveAdded, EAnimDataModelNotifyType::CurveChanged, EAnimDataModelNotifyType::CurveFlagsChanged, EAnimDataModelNotifyType::CurveScaled, EAnimDataModelNotifyType::CurveRenamed }))
 				{
 					UpdateCachedCurve();
 				}
diff --git a/Engine/Source/Editor/Persona/Private/SAnimSequenceCurveEditor.h b/Engine/Source/Editor/Persona/Private/SAnimSequenceCurveEditor.h
--- a/Engine/Source/Editor/Persona/Private/SAnimSequenceCurveEditor.h
+++ b/Engine/Source/Editor/Persona/Private/SAnimSequenceCurveEditor.h
@@ -44,6 +44,7 @@ public:
 	FAnimationCurveIdentifier CurveId;
 	UE::Anim::FAnimDataModelNotifyCollector NotifyCollector;
 	FRichCurve CachedCurve;
+	bool bCurveRemoved;
 };
 
 class SAnimSequenceCurveEditor : public IAnimSequenceCurveEditor, public FEditorUndoClient
