Never zero out montage (play) length anymore during post-undo/redo. DataModelNotifies now take care of restoring the length.

diff --git a/Engine/Source/Editor/Persona/Private/SAnimCompositePanel.cpp b/Engine/Source/Editor/Persona/Private/SAnimCompositePanel.cpp
--- a/Engine/Source/Editor/Persona/Private/SAnimCompositePanel.cpp
+++ b/Engine/Source/Editor/Persona/Private/SAnimCompositePanel.cpp
@@ -208,12 +208,6 @@ void SAnimCompositePanel::PostUndoRedo()
 		bIsActiveTimerRegistered = true;
 		RegisterActiveTimer(0.f, FWidgetActiveTimerDelegate::CreateSP(this, &SAnimCompositePanel::TriggerRebuildPanel));
 	}
-
-	// when undo or redo happens, we still have to recalculate length, so we can't rely on sequence length changes or not
-	if (Composite->GetPlayLength())
-	{
-		Composite->SetCompositeLength(0.f);
-	}
 }
 
 FLinearColor SAnimCompositePanel::HandleGetNodeColor(const FAnimSegment& InSegment) const
diff --git a/Engine/Source/Editor/Persona/Private/SAnimMontagePanel.cpp b/Engine/Source/Editor/Persona/Private/SAnimMontagePanel.cpp
--- a/Engine/Source/Editor/Persona/Private/SAnimMontagePanel.cpp
+++ b/Engine/Source/Editor/Persona/Private/SAnimMontagePanel.cpp
@@ -581,12 +581,6 @@ void SAnimMontagePanel::PostRedo( bool bSuccess )
 
 void SAnimMontagePanel::PostRedoUndo()
 {
-	// when undo or redo happens, we still have to recalculate length, so we can't rely on sequence length changes or not
-	if (Montage->GetPlayLength())
-	{
-		Montage->SetCompositeLength(MINIMUM_ANIMATION_LENGTH);
-	}
-
 	RebuildMontagePanel(); //Rebuild here, undoing adds can cause slate to crash later on if we don't (using dummy args since they aren't used by the method
 }
 
