BlendSpace 2.0: Blendspace Graph Node

diff --git a/Engine/Source/Editor/Persona/Private/SAnimationBlendSpaceGridWidget.cpp b/Engine/Source/Editor/Persona/Private/SAnimationBlendSpaceGridWidget.cpp
--- a/Engine/Source/Editor/Persona/Private/SAnimationBlendSpaceGridWidget.cpp
+++ b/Engine/Source/Editor/Persona/Private/SAnimationBlendSpaceGridWidget.cpp
@@ -45,15 +45,19 @@ void SBlendSpaceGridWidget::Construct(const FArguments& InArgs)
 	OnSampleAdded = InArgs._OnSampleAdded;
 	OnSampleMoved = InArgs._OnSampleMoved;
 	OnSampleRemoved = InArgs._OnSampleRemoved;
+	OnSampleDoubleClicked = InArgs._OnSampleDoubleClicked;
 	OnSampleAnimationChanged = InArgs._OnSampleAnimationChanged;
+	OnGetBlendSpaceSampleName = InArgs._OnGetBlendSpaceSampleName;
+	OnExtendSampleTooltip = InArgs._OnExtendSampleTooltip;
 	bReadOnly = InArgs._ReadOnly;
 	bShowAxisLabels = InArgs._ShowAxisLabels;
 	bShowSettingsButtons = InArgs._ShowSettingsButtons;
+	StatusBarName = InArgs._StatusBarName;
 
 	GridType = (BlendSpaceBase.Get() != nullptr && BlendSpaceBase.Get()->IsA<UBlendSpace1D>()) ? EGridType::SingleAxis : EGridType::TwoAxis;
 	BlendParametersToDraw = (GridType == EGridType::SingleAxis) ? 1 : 2;
 	
-	HighlightedSampleIndex = SelectedSampleIndex = DraggedSampleIndex = INDEX_NONE;
+	HighlightedSampleIndex = SelectedSampleIndex = DraggedSampleIndex = ToolTipSampleIndex = INDEX_NONE;
 	DragState = EDragState::None;
 	// Initialize flags 
 	bPreviewPositionSet = true;
@@ -163,7 +167,7 @@ void SBlendSpaceGridWidget::Construct(const FArguments& InArgs)
 								.VAlign(VAlign_Center)
 								[
 									SNew(SButton)
-									.ToolTipText(LOCTEXT("ShowAnimationNames", "Show Animation Names"))
+									.ToolTipText(LOCTEXT("ShowAnimationNames", "Show Sample Names"))
 									.OnClicked(this, &SBlendSpaceGridWidget::ToggleShowAnimationNames)
 									.ButtonColorAndOpacity_Lambda([this]() -> FLinearColor { return bShowAnimationNames ? FEditorStyle::GetSlateColor("SelectionColor").GetSpecifiedColor() : FLinearColor::White; })
 									.ContentPadding(1)
@@ -247,17 +251,17 @@ void SBlendSpaceGridWidget::Construct(const FArguments& InArgs)
 			]
 		]
 	];
-		
+
 	SAssignNew(ToolTip, SToolTip)
-	.BorderImage(FCoreStyle::Get().GetBrush("ToolTip.BrightBackground"))
+	.BorderImage(FCoreStyle::Get().GetBrush("ToolTip.Background"))
 	[
 		SNew(SVerticalBox)
 		+ SVerticalBox::Slot()
+		.AutoHeight()
 		[
 			SNew(STextBlock)
 			.Text(this, &SBlendSpaceGridWidget::GetToolTipAnimationName)
 			.Font(FCoreStyle::Get().GetFontStyle("ToolTip.LargerFont"))
-			.ColorAndOpacity(FLinearColor::Black)
 		]
 
 		+ SVerticalBox::Slot()
@@ -266,16 +270,26 @@ void SBlendSpaceGridWidget::Construct(const FArguments& InArgs)
 			SNew(STextBlock)
 			.Text(this, &SBlendSpaceGridWidget::GetToolTipSampleValue)
 			.Font(FCoreStyle::Get().GetFontStyle("ToolTip.LargerFont"))
-			.ColorAndOpacity(FLinearColor::Black)
+		]
+
+		+ SVerticalBox::Slot()
+		.AutoHeight()
+		[
+			SAssignNew(ToolTipExtensionContainer, SBox)
 		]
 	];
 
-	if(bReadOnly)
+	if(Position.IsSet())
 	{
 		StartPreviewing();
 	}
 }
 
+SBlendSpaceGridWidget::~SBlendSpaceGridWidget()
+{
+	EnableStatusBarMessage(false);
+}
+
 TSharedPtr<SWidget> SBlendSpaceGridWidget::CreateGridEntryBox(const int32 BoxIndex, const bool bShowLabel)
 {
 	return SNew(SNumericEntryBox<float>)
@@ -547,6 +561,23 @@ void SBlendSpaceGridWidget::PaintTriangulation(const FGeometry& AllottedGeometry
 	DrawLayerId += 1;
 }
 
+FText SBlendSpaceGridWidget::GetSampleName(const FBlendSample& InBlendSample, int32 InSampleIndex) const
+{
+	if(OnGetBlendSpaceSampleName.IsBound())
+	{
+		return FText::FromName(OnGetBlendSpaceSampleName.Execute(InSampleIndex));
+	}
+	else
+	{
+		if(InBlendSample.Animation != nullptr)
+		{
+			return FText::FromString(InBlendSample.Animation->GetName());
+		}
+	}
+		
+	return LOCTEXT("NoAnimationSetTooltipText", "No Animation Set");
+}
+
 void SBlendSpaceGridWidget::PaintAnimationNames(const FGeometry& AllottedGeometry, const FSlateRect& MyCullingRect, FSlateWindowElementList& OutDrawElements, int32& DrawLayerId) const
 {
 	if(const UBlendSpaceBase* BlendSpace = BlendSpaceBase.Get())
@@ -557,24 +588,22 @@ void SBlendSpaceGridWidget::PaintAnimationNames(const FGeometry& AllottedGeometr
 		for (int32 SampleIndex = 0; SampleIndex < Samples.Num(); ++SampleIndex)
 		{
 			const FBlendSample& Sample = Samples[SampleIndex];
-			if (Sample.Animation)
-			{
-				const FString Name = Sample.Animation->GetName() + FString::Printf(TEXT(" (%i)"), SampleIndex);
-				const FVector2D TextSize = FontMeasure->Measure(Name, FontInfo);
 
-				FVector2D GridPosition = SampleValueToGridPosition(Sample.SampleValue);
-				// Check on which side of the sample the text should be positioned so that we don't run out of geometry space
-				if ((GridPosition + TextSize).X > AllottedGeometry.GetLocalSize().X)
-				{
-					GridPosition -= FVector2D(TextSize.X + KeySize.X, KeySize.X * .5f);
-				}
-				else
-				{
-					GridPosition += FVector2D(KeySize.X, -KeySize.X * .5f);
-				}
+			const FText Name = FText::Format(LOCTEXT("SampleNameFormat", "{0} ({1})"), GetSampleName(Sample, SampleIndex), FText::AsNumber(SampleIndex));
+			const FVector2D TextSize = FontMeasure->Measure(Name, FontInfo);
 
-				FSlateDrawElement::MakeText(OutDrawElements, DrawLayerId + 1, AllottedGeometry.MakeChild(FVector2D(GridPosition.X, GridPosition.Y), FVector2D(1.0f, 1.0f)).ToPaintGeometry(), Name, FontInfo, ESlateDrawEffect::None, FLinearColor::White);
+			FVector2D GridPosition = SampleValueToGridPosition(Sample.SampleValue);
+			// Check on which side of the sample the text should be positioned so that we don't run out of geometry space
+			if ((GridPosition + TextSize).X > AllottedGeometry.GetLocalSize().X)
+			{
+				GridPosition -= FVector2D(TextSize.X + KeySize.X, KeySize.X * .5f);
 			}
+			else
+			{
+				GridPosition += FVector2D(KeySize.X, -KeySize.X * .5f);
+			}
+
+			FSlateDrawElement::MakeText(OutDrawElements, DrawLayerId + 1, AllottedGeometry.MakeChild(FVector2D(GridPosition.X, GridPosition.Y), FVector2D(1.0f, 1.0f)).ToPaintGeometry(), Name, FontInfo, ESlateDrawEffect::None, FLinearColor::White);
 		}
 	}
 
@@ -758,10 +787,36 @@ FReply SBlendSpaceGridWidget::OnMouseButtonDown(const FGeometry& MyGeometry, con
 	return FReply::Unhandled();
 }
 
+FReply SBlendSpaceGridWidget::OnMouseButtonDoubleClick(const FGeometry& InMyGeometry, const FPointerEvent& InMouseEvent)
+{
+	if(!bReadOnly)
+	{
+		if(const UBlendSpaceBase* BlendSpace = BlendSpaceBase.Get())
+		{
+			if (InMouseEvent.GetEffectingButton() == EKeys::LeftMouseButton)
+			{
+				if (SelectedSampleIndex != INDEX_NONE)
+				{
+					OnSampleDoubleClicked.ExecuteIfBound(SelectedSampleIndex);
+				}
+
+				return FReply::Handled();
+			}
+		}
+	}
+
+	return FReply::Unhandled();
+}
+
 FReply SBlendSpaceGridWidget::OnMouseMove(const FGeometry& MyGeometry, const FPointerEvent& MouseEvent)
 {
 	if(const UBlendSpaceBase* BlendSpace = BlendSpaceBase.Get())
 	{
+		if(!bReadOnly)
+		{
+			EnableStatusBarMessage(true);
+		}
+
 		// Cache the mouse position in local and screen space
 		LocalMousePosition = MyGeometry.AbsoluteToLocal(MouseEvent.GetScreenSpacePosition());
 		LastMousePosition = MouseEvent.GetScreenSpacePosition();
@@ -788,10 +843,9 @@ FReply SBlendSpaceGridWidget::OnMouseMove(const FGeometry& MyGeometry, const FPo
 				}
 			}
 		}
-
-		if (IsHovered() && !HasAnyUserFocus() && bMouseIsOverGeometry)
+		else if (IsHovered() && bMouseIsOverGeometry && DragState != EDragState::DragPreview)
 		{
-			if (!bReadOnly && (MouseEvent.IsLeftShiftDown() || MouseEvent.IsRightShiftDown()))
+			if ((MouseEvent.IsLeftShiftDown() || MouseEvent.IsRightShiftDown()))
 			{
 				StartPreviewing();
 				DragState = EDragState::Preview;
@@ -801,11 +855,14 @@ FReply SBlendSpaceGridWidget::OnMouseMove(const FGeometry& MyGeometry, const FPo
 				// Set flag for showing advanced preview info in tooltip
 				bAdvancedPreview = MouseEvent.IsLeftControlDown() || MouseEvent.IsRightControlDown();
 			}
-			else if(bReadOnly)
+			else if(Position.IsSet())
 			{
 				StartPreviewing();
 				DragState = EDragState::None;
 				ShowToolTip();
+
+				// Set flag for showing advanced preview info in tooltip
+				bAdvancedPreview = MouseEvent.IsLeftControlDown() || MouseEvent.IsRightControlDown();
 			}
 			else if (bSamplePreviewing)
 			{
@@ -854,18 +911,9 @@ FReply SBlendSpaceGridWidget::ProcessClick(const FGeometry& MyGeometry, const FP
 			}
 			else if (MouseEvent.GetEffectingButton() == EKeys::RightMouseButton)
 			{
-				// If we are over a sample open a context menu for editing its data
-				if (HighlightedSampleIndex != INDEX_NONE)
-				{	
-					SelectedSampleIndex = HighlightedSampleIndex;
-
-					// Create context menu
-					TSharedPtr<SWidget> MenuContent = CreateBlendSampleContextMenu();
-			
-					// Reset highlight sample index
-					HighlightedSampleIndex = INDEX_NONE;
-
-					if (MenuContent.IsValid())
+				auto PushMenu = [this, &MouseEvent](TSharedPtr<SWidget> InMenuContent)
+				{
+					if (InMenuContent.IsValid())
 					{
 						const FWidgetPath WidgetPath = MouseEvent.GetEventPath() != nullptr ? *MouseEvent.GetEventPath() : FWidgetPath();
 						const FVector2D MousePosition = MouseEvent.GetScreenSpacePosition();
@@ -876,14 +924,36 @@ FReply SBlendSpaceGridWidget::ProcessClick(const FGeometry& MyGeometry, const FP
 						FSlateApplication::Get().PushMenu(
 							AsShared(),
 							WidgetPath,
-							MenuContent.ToSharedRef(),
+							InMenuContent.ToSharedRef(),
 							MenuPosition,
 							FPopupTransitionEffect(FPopupTransitionEffect::ContextMenu)
 							);
-
-						return FReply::Handled().SetUserFocus(MenuContent.ToSharedRef(), EFocusCause::SetDirectly).ReleaseMouseCapture();
 					}
-				}		
+				};
+
+				// If we are over a sample open a context menu for editing its data
+				if (HighlightedSampleIndex != INDEX_NONE)
+				{	
+					SelectedSampleIndex = HighlightedSampleIndex;
+
+					// Create context menu
+					TSharedPtr<SWidget> MenuContent = CreateBlendSampleContextMenu();
+			
+					// Reset highlight sample index
+					HighlightedSampleIndex = INDEX_NONE;
+
+					PushMenu(MenuContent);
+
+					return FReply::Handled().SetUserFocus(MenuContent.ToSharedRef(), EFocusCause::SetDirectly).ReleaseMouseCapture();
+				}
+				else
+				{
+					TSharedPtr<SWidget> MenuContent = CreateNewBlendSampleContextMenu(MyGeometry.AbsoluteToLocal(MouseEvent.GetScreenSpacePosition()));
+				
+					PushMenu(MenuContent);
+				
+					return FReply::Handled().SetUserFocus(MenuContent.ToSharedRef(), EFocusCause::SetDirectly).ReleaseMouseCapture();
+				}
 			}
 		}
 	}
@@ -971,6 +1041,52 @@ FReply SBlendSpaceGridWidget::OnKeyUp(const FGeometry& MyGeometry, const FKeyEve
 	return FReply::Unhandled();
 }
 
+void SBlendSpaceGridWidget::MakeViewContextMenuEntries(FMenuBuilder& InMenuBuilder)
+{
+	InMenuBuilder.BeginSection("ViewOptions", LOCTEXT("ViewOptionsMenuHeader", "View Options"));
+	{
+		InMenuBuilder.AddMenuEntry(
+			LOCTEXT("ShowTriangulation", "Show Triangulation"),
+			LOCTEXT("ShowTriangulationToolTip", "Show the Delaunay triangulation for all blend space samples"),
+			FSlateIcon("EditorStyle", "BlendSpaceEditor.ToggleTriangulation"),
+			FUIAction(
+				FExecuteAction::CreateLambda([this](){ bShowTriangulation = !bShowTriangulation; }),
+				FCanExecuteAction(),
+				FGetActionCheckState::CreateLambda([this](){ return bShowTriangulation ? ECheckBoxState::Checked : ECheckBoxState::Unchecked; })
+			),
+			NAME_None,
+			EUserInterfaceActionType::ToggleButton
+		);
+
+		InMenuBuilder.AddMenuEntry(
+			LOCTEXT("ShowAnimationNames", "Show Sample Names"),
+			LOCTEXT("ShowAnimationNamesToolTip", "Show the names of each of the samples"),
+			FSlateIcon("EditorStyle", "BlendSpaceEditor.ToggleLabels"),
+			FUIAction(
+				FExecuteAction::CreateLambda([this](){ bShowAnimationNames = !bShowAnimationNames; }),
+				FCanExecuteAction(),
+				FGetActionCheckState::CreateLambda([this](){ return bShowAnimationNames ? ECheckBoxState::Checked : ECheckBoxState::Unchecked; })
+			),
+			NAME_None,
+			EUserInterfaceActionType::ToggleButton
+		);
+
+		InMenuBuilder.AddMenuEntry(
+			LOCTEXT("StretchFittingText", "Stretch Grid to Fit"),
+			LOCTEXT("StretchFittingTextToolTip", "Whether to stretch the grid to fit or to fit the grid to the largest axis"),
+			FSlateIcon("EditorStyle", "WidgetDesigner.ZoomToFit"),
+			FUIAction(
+				FExecuteAction::CreateLambda([this](){ bStretchToFit = !bStretchToFit; }),
+				FCanExecuteAction(),
+				FGetActionCheckState::CreateLambda([this](){ return bStretchToFit ? ECheckBoxState::Checked : ECheckBoxState::Unchecked; })
+			),
+			NAME_None,
+			EUserInterfaceActionType::ToggleButton
+		);
+	}
+	InMenuBuilder.EndSection();
+}
+
 TSharedPtr<SWidget> SBlendSpaceGridWidget::CreateBlendSampleContextMenu()
 {
 	const bool bShouldCloseWindowAfterMenuSelection = true;
@@ -1004,7 +1120,7 @@ TSharedPtr<SWidget> SBlendSpaceGridWidget::CreateBlendSampleContextMenu()
 		if(const UBlendSpaceBase* BlendSpace = BlendSpaceBase.Get())
 		{
 			const FBlendSample& Sample = BlendSpace->GetBlendSample(HighlightedSampleIndex);		
-			StructureDetailsView->GetDetailsView()->SetGenericLayoutDetailsDelegate(FOnGetDetailCustomizationInstance::CreateStatic(&FBlendSampleDetails::MakeInstance, BlendSpace, this));
+			StructureDetailsView->GetDetailsView()->SetGenericLayoutDetailsDelegate(FOnGetDetailCustomizationInstance::CreateStatic(&FBlendSampleDetails::MakeInstance, BlendSpace, this, HighlightedSampleIndex));
 
 			FStructOnScope* Struct = new FStructOnScope(FBlendSample::StaticStruct(), (uint8*)&Sample);
 			Struct->SetPackage(BlendSpace->GetOutermost());
@@ -1012,7 +1128,54 @@ TSharedPtr<SWidget> SBlendSpaceGridWidget::CreateBlendSampleContextMenu()
 		}
 	}
 	
-	MenuBuilder.AddWidget(StructureDetailsView->GetWidget().ToSharedRef(), FText::GetEmpty(), true);
+	MenuBuilder.BeginSection("Sample", LOCTEXT("SampleMenuHeader", "Sample"));
+	{
+		MenuBuilder.AddWidget(StructureDetailsView->GetWidget().ToSharedRef(), FText::GetEmpty(), true);
+	}
+	MenuBuilder.EndSection();
+
+	MakeViewContextMenuEntries(MenuBuilder);
+
+	return MenuBuilder.MakeWidget();
+}
+
+TSharedPtr<SWidget> SBlendSpaceGridWidget::CreateNewBlendSampleContextMenu(const FVector2D& InMousePosition)
+{
+	const bool bShouldCloseWindowAfterMenuSelection = true;
+	FMenuBuilder MenuBuilder(bShouldCloseWindowAfterMenuSelection, nullptr);
+
+	FVector NewSampleValue;
+	if(FSlateApplication::Get().GetModifierKeys().IsAltDown())
+	{
+		const FVector2D GridPosition(FMath::Clamp(InMousePosition.X, CachedGridRectangle.Left, CachedGridRectangle.Right),
+									 FMath::Clamp(InMousePosition.Y, CachedGridRectangle.Top, CachedGridRectangle.Bottom));
+		NewSampleValue = GridPositionToSampleValue(GridPosition, false);
+	}
+	else
+	{
+		NewSampleValue = GridPositionToSampleValue(SnapToClosestGridPoint(InMousePosition), false);
+	}
+
+	if(const UBlendSpaceBase* BlendSpace = BlendSpaceBase.Get())
+	{
+		MenuBuilder.BeginSection("Sample", LOCTEXT("SampleMenuHeader", "Sample"));
+		{
+			if(!BlendSpace->IsAsset())
+			{
+				MenuBuilder.AddMenuEntry(
+					LOCTEXT("AddNewSample", "Add New Sample"),
+					LOCTEXT("AddNewSampleTooltip", "Add a new sample to this blendspace at this location"),
+						FSlateIcon("EditorStyle", "Plus"),
+						FUIAction(
+							FExecuteAction::CreateLambda([this, NewSampleValue](){ OnSampleAdded.ExecuteIfBound(nullptr, NewSampleValue); })
+						)
+				);
+			}
+		}
+		MenuBuilder.EndSection();
+	}
+
+	MakeViewContextMenuEntries(MenuBuilder);
 
 	return MenuBuilder.MakeWidget();
 }
@@ -1144,7 +1307,9 @@ void SBlendSpaceGridWidget::StartPreviewing()
 {
 	bSamplePreviewing = true;
 	LastPreviewingMousePosition = LocalMousePosition;
-	LastPreviewingSampleValue = bReadOnly ? Position.Get() : GridPositionToSampleValue(LastPreviewingMousePosition, false);
+	FModifierKeysState ModifierKeyState = FSlateApplication::Get().GetModifierKeys();
+	bool bIsManualPreviewing = (ModifierKeyState.IsLeftShiftDown() || ModifierKeyState.IsRightShiftDown()) || DragState == EDragState::DragPreview;
+	LastPreviewingSampleValue = Position.IsSet() && !bIsManualPreviewing ? Position.Get() : GridPositionToSampleValue(LastPreviewingMousePosition, false);
 	bPreviewPositionSet = true;	
 	bPreviewToolTipHidden = true;
 }
@@ -1159,72 +1324,80 @@ FText SBlendSpaceGridWidget::GetToolTipAnimationName() const
 	FText ToolTipText = FText::GetEmpty();
 	if(const UBlendSpaceBase* BlendSpace = BlendSpaceBase.Get())
 	{
-		const FText EmptyAnimationText = LOCTEXT("NoAnimationSetTooltipText", "No Animation Set");
+		const FText PreviewValue = LOCTEXT("PreviewValueTooltip", "Preview Value");
 
-		switch (DragState)
+		if(bReadOnly)
 		{
-			// If we are not dragging, but over a valid blend sample return its animation asset name
-			case EDragState::None:
-			{		
-				if (bHighlightPreviewPin && !bReadOnly)
-				{
-					const FText PreviewPinText = LOCTEXT("HighlightPreviewPinTooltipText", "Preview Value (Click and Drag or Hold Shift to Move)");
-					ToolTipText = PreviewPinText;
+			ToolTipText = PreviewValue;
+		}
+		else
+		{
+			switch (DragState)
+			{
+				// If we are not dragging, but over a valid blend sample return its animation asset name
+				case EDragState::None:
+				{		
+					if (bHighlightPreviewPin)
+					{
+						ToolTipText = PreviewValue;
+					}
+					else if (HighlightedSampleIndex != INDEX_NONE && BlendSpace->IsValidBlendSampleIndex(HighlightedSampleIndex))
+					{
+						const FBlendSample& BlendSample = BlendSpace->GetBlendSample(HighlightedSampleIndex);
+						ToolTipText = GetSampleName(BlendSample, HighlightedSampleIndex);
+					}
+					else if(Position.IsSet())
+					{
+						ToolTipText = PreviewValue;
+					}
+					break;
 				}
-				else if (HighlightedSampleIndex != INDEX_NONE && BlendSpace->IsValidBlendSampleIndex(HighlightedSampleIndex))
+
+				case EDragState::PreDrag:
 				{
-					const FBlendSample& BlendSample = BlendSpace->GetBlendSample(HighlightedSampleIndex);					
-					ToolTipText = (BlendSample.Animation != nullptr) ? FText::FromString(BlendSample.Animation->GetName()) : EmptyAnimationText;
+					break;
 				}
-				break;
-			}
 
-			case EDragState::PreDrag:
-			{
-				break;
-			}
-
-			// If we are dragging a sample return the dragged sample's animation asset name
-			case EDragState::DragSample:
-			{
-				if (BlendSpace->IsValidBlendSampleIndex(DraggedSampleIndex))
+				// If we are dragging a sample return the dragged sample's animation asset name
+				case EDragState::DragSample:
 				{
-					const FBlendSample& BlendSample = BlendSpace->GetBlendSample(DraggedSampleIndex);
-					ToolTipText = (BlendSample.Animation != nullptr) ? FText::FromString(BlendSample.Animation->GetName()) : EmptyAnimationText;
-				}
+					if (BlendSpace->IsValidBlendSampleIndex(DraggedSampleIndex))
+					{
+						const FBlendSample& BlendSample = BlendSpace->GetBlendSample(DraggedSampleIndex);
+						ToolTipText = GetSampleName(BlendSample, DraggedSampleIndex);
+					}
 			
-				break;
-			}
+					break;
+				}
 
-			// If we are performing a drag/drop operation return the cached operation animation name
-			case EDragState::DragDrop:
-			{
-				ToolTipText = DragDropAnimationName;
-				break;
-			}
+				// If we are performing a drag/drop operation return the cached operation animation name
+				case EDragState::DragDrop:
+				{
+					ToolTipText = DragDropAnimationName;
+					break;
+				}
 
-			case EDragState::DragDropOverride:
-			{
-				ToolTipText = DragDropAnimationName;
-				break;
-			}
+				case EDragState::DragDropOverride:
+				{
+					ToolTipText = DragDropAnimationName;
+					break;
+				}
 
-			case EDragState::InvalidDragDrop:
-			{
-				break;
-			}
+				case EDragState::InvalidDragDrop:
+				{
+					break;
+				}
 		
-			// If we are previewing return a descriptive label
-			case EDragState::Preview:
-			case EDragState::DragPreview:
-			{
-				const FText AdvancedPreviewText = LOCTEXT("AdvancedPreviewValueTooltip", "Preview Value");
-				const FText BasicPreviewText = LOCTEXT("BasicPreviewValueTooltip", "Preview Value (Hold Ctrl for Weight Details)");	
-				ToolTipText = bAdvancedPreview ? AdvancedPreviewText : BasicPreviewText;
-				break;
+				// If we are previewing return a descriptive label
+				case EDragState::Preview:
+				case EDragState::DragPreview:
+				{
+					ToolTipText = PreviewValue;
+					break;
+				}
+				default:
+					check(false);
 			}
-			default:
-				check(false);
 		}
 	}
 
@@ -1237,11 +1410,36 @@ FText SBlendSpaceGridWidget::GetToolTipSampleValue() const
 
 	if(const UBlendSpaceBase* BlendSpace = BlendSpaceBase.Get())
 	{
-		const FTextFormat ValueFormattingText = (GridType == EGridType::TwoAxis) ? FTextFormat::FromString("{0}: {1} - {2}: {3}") : FTextFormat::FromString("{0}: {1}");
+		static const FTextFormat OneAxisFormat = LOCTEXT("OneAxisFormat", "{0}: {1}");
+		static const FTextFormat TwoAxisFormat = LOCTEXT("TwoAxisFormat", "{0}: {1} - {2}: {3}");
+		const FTextFormat& ValueFormattingText = (GridType == EGridType::TwoAxis) ? TwoAxisFormat : OneAxisFormat;
+
+		auto AddAdvancedPreview = [this, &ToolTipText, BlendSpace]()
+		{
+			FTextBuilder TextBuilder;
+			TextBuilder.AppendLine(ToolTipText);
+
+			if (bAdvancedPreview)
+			{				
+				for (const FBlendSampleData& SampleData : PreviewedSamples)
+				{
+					if(BlendSpace->IsValidBlendSampleIndex(SampleData.SampleDataIndex))
+					{
+						const FBlendSample& BlendSample = BlendSpace->GetBlendSample(SampleData.SampleDataIndex);
+						static const FTextFormat SampleFormat = LOCTEXT("SampleFormat", "{0}: {1}");
+						TextBuilder.AppendLine(FText::Format(SampleFormat, GetSampleName(BlendSample, SampleData.SampleDataIndex), FText::AsNumber(SampleData.TotalWeight)));
+					}
+				}
+			}
+
+			ToolTipText = TextBuilder.ToText();
+		};
 
 		if(bReadOnly)
 		{
 			ToolTipText = FText::Format(ValueFormattingText, ParameterXName, FText::FromString(FString::SanitizeFloat(LastPreviewingSampleValue.X)), ParameterYName, FText::FromString(FString::SanitizeFloat(LastPreviewingSampleValue.Y)));
+
+			AddAdvancedPreview();
 		}
 		else
 		{
@@ -1253,6 +1451,8 @@ FText SBlendSpaceGridWidget::GetToolTipSampleValue() const
 					if (bHighlightPreviewPin)
 					{
 						ToolTipText = FText::Format(ValueFormattingText, ParameterXName, FText::FromString(FString::SanitizeFloat(LastPreviewingSampleValue.X)), ParameterYName, FText::FromString(FString::SanitizeFloat(LastPreviewingSampleValue.Y)));
+
+						AddAdvancedPreview();
 					}
 					else if (HighlightedSampleIndex != INDEX_NONE && BlendSpace->IsValidBlendSampleIndex(HighlightedSampleIndex))
 					{
@@ -1268,7 +1468,12 @@ FText SBlendSpaceGridWidget::GetToolTipSampleValue() const
 							ToolTipText = GetSampleErrorMessage(BlendSample);
 						}
 					}
+					else if(Position.IsSet())
+					{
+						ToolTipText = FText::Format(ValueFormattingText, ParameterXName, FText::FromString(FString::SanitizeFloat(LastPreviewingSampleValue.X)), ParameterYName, FText::FromString(FString::SanitizeFloat(LastPreviewingSampleValue.Y)));
 
+						AddAdvancedPreview();
+					}
 					break;
 				}
 
@@ -1300,7 +1505,7 @@ FText SBlendSpaceGridWidget::GetToolTipSampleValue() const
 
 				case EDragState::DragDropOverride:
 				{
-					const FTextFormat OverrideAnimationFormat = FTextFormat::FromString("Changing Animation from {0} to {1}");
+					static const FTextFormat OverrideAnimationFormat = FTextFormat::FromString("Changing Animation from {0} to {1}");
 					ToolTipText = FText::Format(OverrideAnimationFormat, HoveredAnimationName, DragDropAnimationName);
 					break;
 				}
@@ -1315,25 +1520,9 @@ FText SBlendSpaceGridWidget::GetToolTipSampleValue() const
 				case EDragState::DragPreview:
 				case EDragState::Preview:
 				{
-					FFormatOrderedArguments PreviewTextArguments;
-					PreviewTextArguments.Add(FText::Format(ValueFormattingText, ParameterXName, FText::FromString(FString::SanitizeFloat(LastPreviewingSampleValue.X)), ParameterYName, FText::FromString(FString::SanitizeFloat(LastPreviewingSampleValue.Y))));
+					ToolTipText = FText::Format(ValueFormattingText, ParameterXName, FText::FromString(FString::SanitizeFloat(LastPreviewingSampleValue.X)), ParameterYName, FText::FromString(FString::SanitizeFloat(LastPreviewingSampleValue.Y)));
 
-					FString PreviewSamplesString;
-					if (bAdvancedPreview)
-					{				
-						for (const FBlendSampleData& SampleData : PreviewedSamples)
-						{
-							PreviewSamplesString += "\n";
-							PreviewSamplesString += SampleData.Animation ? SampleData.Animation->GetName() : "No Animation Set";
-							PreviewSamplesString += ": ";
-							PreviewSamplesString += FString::SanitizeFloat(SampleData.TotalWeight);
-						}				
-					}
-					PreviewTextArguments.Add(FText::FromString(PreviewSamplesString));
-			
-					const FTextFormat PreviewTextFormat = FTextFormat::FromString("{0}{1}");			
-					ToolTipText = FText::Format(PreviewTextFormat, PreviewTextArguments);
-			
+					AddAdvancedPreview();
 					break;
 				}
 				default:
@@ -1345,6 +1534,37 @@ FText SBlendSpaceGridWidget::GetToolTipSampleValue() const
 	return ToolTipText;
 }
 
+void SBlendSpaceGridWidget::EnableStatusBarMessage(bool bEnable)
+{
+	if(!bReadOnly)
+	{
+		if(bEnable)
+		{
+			if (!StatusBarMessageHandle.IsValid())
+			{
+				if(UStatusBarSubsystem* StatusBarSubsystem = GEditor->GetEditorSubsystem<UStatusBarSubsystem>())
+				{
+					StatusBarMessageHandle = StatusBarSubsystem->PushStatusBarMessage(StatusBarName, MakeAttributeLambda([]()
+					{
+						return LOCTEXT("StatusBarMssage", "Hold Ctrl for weight details, click and drag or hold Shift to move preview value");
+					}));
+				}
+			}
+		}
+		else
+		{
+			if (StatusBarMessageHandle.IsValid())
+			{
+				if(UStatusBarSubsystem* StatusBarSubsystem = GEditor->GetEditorSubsystem<UStatusBarSubsystem>())
+				{
+					StatusBarSubsystem->PopStatusBarMessage(StatusBarName, StatusBarMessageHandle);
+					StatusBarMessageHandle.Reset();
+				}
+			}
+		}
+	}
+}
+
 FText SBlendSpaceGridWidget::GetSampleErrorMessage(const FBlendSample &BlendSample) const
 {
 	const FVector2D GridPosition = SampleValueToGridPosition(BlendSample.SampleValue);
@@ -1367,11 +1587,22 @@ FText SBlendSpaceGridWidget::GetSampleErrorMessage(const FBlendSample &BlendSamp
 
 void SBlendSpaceGridWidget::ShowToolTip()
 {
+	if(DragState != EDragState::DragPreview && HighlightedSampleIndex != INDEX_NONE && ToolTipSampleIndex != HighlightedSampleIndex)
+	{
+		ToolTipSampleIndex = HighlightedSampleIndex;
+		if(OnExtendSampleTooltip.IsBound())
+		{
+			ToolTipExtensionContainer->SetContent(OnExtendSampleTooltip.Execute(HighlightedSampleIndex));
+		}
+	}
+	
 	SetToolTip(ToolTip);
 }
 
 void SBlendSpaceGridWidget::ResetToolTip()
 {
+	ToolTipSampleIndex = INDEX_NONE;
+	ToolTipExtensionContainer->SetContent(SNullWidget::NullWidget);
 	SetToolTip(nullptr);
 }
 
@@ -1588,12 +1819,14 @@ void SBlendSpaceGridWidget::OnMouseEnter(const FGeometry& MyGeometry, const FPoi
 {	
 	SCompoundWidget::OnMouseEnter(MyGeometry, MouseEvent);
 	bMouseIsOverGeometry = true;
+	EnableStatusBarMessage(true);
 }
 
 void SBlendSpaceGridWidget::OnMouseLeave(const FPointerEvent& MouseEvent)
 {
 	SCompoundWidget::OnMouseLeave(MouseEvent);
 	bMouseIsOverGeometry = false;
+	EnableStatusBarMessage(false);
 }
 
 void SBlendSpaceGridWidget::OnFocusLost(const FFocusEvent& InFocusEvent)
@@ -1603,6 +1836,7 @@ void SBlendSpaceGridWidget::OnFocusLost(const FFocusEvent& InFocusEvent)
 	DragState = EDragState::None;
 	bSamplePreviewing = false;
 	ResetToolTip();
+	EnableStatusBarMessage(false);
 }
 
 bool SBlendSpaceGridWidget::SupportsKeyboardFocus() const
@@ -1745,10 +1979,12 @@ void SBlendSpaceGridWidget::Tick(const FGeometry& AllottedGeometry, const double
 			// Ensure the preview mouse position is clamped to the grid
 			LastPreviewingMousePosition.X = LocalMousePosition.X; //FMath::Clamp(LocalMousePosition.X, CachedGridRectangle.Left, CachedGridRectangle.Right);
 			LastPreviewingMousePosition.Y = LocalMousePosition.Y; //FMath::Clamp(LocalMousePosition.Y, CachedGridRectangle.Top, CachedGridRectangle.Bottom);
-			LastPreviewingSampleValue = bReadOnly ? Position.Get() : GridPositionToSampleValue(LastPreviewingMousePosition, false);
-
-			LastPreviewingSampleValue = BlendSpace->GetClampedAndWrappedBlendInput(LastPreviewingSampleValue);
+			FModifierKeysState ModifierKeyState = FSlateApplication::Get().GetModifierKeys();
+			bool bIsManualPreviewing = (ModifierKeyState.IsLeftShiftDown() || ModifierKeyState.IsRightShiftDown()) || DragState == EDragState::DragPreview;
+			LastPreviewingSampleValue = Position.IsSet() && !bIsManualPreviewing ? Position.Get() : GridPositionToSampleValue(LastPreviewingMousePosition, false);
 						
+			LastPreviewingSampleValue = BlendSpace->GetClampedAndWrappedBlendInput(LastPreviewingSampleValue);
+					
 			// Retrieve and cache weighted samples
 			PreviewedSamples.Empty(4);
 			BlendSpace->GetSamplesFromBlendInput(GetBlendPreviewValue(), PreviewedSamples);
@@ -1872,4 +2108,11 @@ bool SBlendSpaceGridWidget::ValidateAnimationSequence(const UAnimSequence* Anima
 	return AnimationSequence != nullptr;
 }
 
+const bool SBlendSpaceGridWidget::IsPreviewing() const 
+{ 
+	FModifierKeysState ModifierKeyState = FSlateApplication::Get().GetModifierKeys();
+	bool bIsManualPreviewing = (ModifierKeyState.IsLeftShiftDown() || ModifierKeyState.IsRightShiftDown()) || DragState == EDragState::DragPreview;
+	return (bSamplePreviewing && !Position.IsSet()) || (Position.IsSet() && bIsManualPreviewing);
+}
+
 #undef LOCTEXT_NAMESPACE // "SAnimationBlendSpaceGridWidget"
diff --git a/Engine/Source/Editor/Persona/Private/SAnimationBlendSpaceGridWidget.h b/Engine/Source/Editor/Persona/Private/SAnimationBlendSpaceGridWidget.h
--- a/Engine/Source/Editor/Persona/Private/SAnimationBlendSpaceGridWidget.h
+++ b/Engine/Source/Editor/Persona/Private/SAnimationBlendSpaceGridWidget.h
@@ -14,6 +14,8 @@
 #include "Widgets/SToolTip.h"
 #include "Misc/NotifyHook.h"
 #include "Animation/AnimationAsset.h"
+#include "PersonaDelegates.h"
+#include "StatusBarSubsystem.h"
 
 class FPaintArgs;
 class FSlateWindowElementList;
@@ -21,8 +23,6 @@ class UAnimSequence;
 class UBlendSpaceBase;
 
 DECLARE_DELEGATE_FourParams(FOnSampleMoved, const int32 /*SampleIndex*/, const FVector& /*SampleValue*/, bool /*bIsInteractive*/, bool /*bSnap*/);
-DECLARE_DELEGATE_OneParam(FOnSampleRemoved, const int32 );
-DECLARE_DELEGATE_TwoParams(FOnSampleAdded, UAnimSequence*, const FVector&);
 DECLARE_DELEGATE_TwoParams(FOnSampleAnimationChanged, UAnimSequence*, const FVector&);
 
 class SBlendSpaceGridWidget : public SCompoundWidget, public FNotifyHook
@@ -34,6 +34,7 @@ public:
 		: _ReadOnly(false)
 		, _ShowAxisLabels(true)
 		, _ShowSettingsButtons(true)
+		, _StatusBarName(TEXT("AssetEditor.AnimationEditor.MainMenu"))
 	{}
 		SLATE_ATTRIBUTE(const UBlendSpaceBase*, BlendSpaceBase)
 		SLATE_ATTRIBUTE(FVector, Position)
@@ -41,10 +42,14 @@ public:
 		SLATE_ARGUMENT(bool, ReadOnly)
 		SLATE_ARGUMENT(bool, ShowAxisLabels)
 		SLATE_ARGUMENT(bool, ShowSettingsButtons)
+		SLATE_EVENT(FOnExtendBlendSpaceSampleTooltip, OnExtendSampleTooltip)
 		SLATE_EVENT(FOnSampleMoved, OnSampleMoved)
-		SLATE_EVENT(FOnSampleRemoved, OnSampleRemoved)
-		SLATE_EVENT(FOnSampleAdded, OnSampleAdded)
+		SLATE_EVENT(FOnBlendSpaceSampleRemoved, OnSampleRemoved)
+		SLATE_EVENT(FOnBlendSpaceSampleAdded, OnSampleAdded)
 		SLATE_EVENT(FOnSampleAnimationChanged, OnSampleAnimationChanged)
+		SLATE_EVENT(FOnBlendSpaceSampleDoubleClicked, OnSampleDoubleClicked)
+		SLATE_EVENT(FOnGetBlendSpaceSampleName, OnGetBlendSpaceSampleName)
+		SLATE_ARGUMENT(FName, StatusBarName)
 	SLATE_END_ARGS()
 
 protected:
@@ -77,6 +82,8 @@ protected:
 
 public:
 
+	~SBlendSpaceGridWidget();
+
 	/**
 	* Construct this widget
 	*
@@ -92,6 +99,7 @@ public:
 	virtual FReply OnDrop(const FGeometry& MyGeometry, const FDragDropEvent& DragDropEvent) override;
 	virtual FReply OnMouseButtonUp(const FGeometry& MyGeometry, const FPointerEvent& MouseEvent) override;
 	virtual FReply OnMouseButtonDown(const FGeometry& MyGeometry, const FPointerEvent& MouseEvent) override;
+	virtual FReply OnMouseButtonDoubleClick(const FGeometry& InMyGeometry, const FPointerEvent& InMouseEvent) override;
 	virtual FReply OnMouseMove(const FGeometry& MyGeometry, const FPointerEvent& MouseEvent) override;
 	virtual FReply OnKeyDown(const FGeometry& MyGeometry, const FKeyEvent& InKeyEvent) override;
 	virtual FReply OnKeyUp(const FGeometry& MyGeometry, const FKeyEvent& InKeyEvent) override;
@@ -105,7 +113,7 @@ public:
 	/** Returns the sample value for previewing the blend space */
 	const FVector GetBlendPreviewValue();
 	/** Flag whether or not the user is actively previewing the blend space (moving the sample value) */
-	const bool IsPreviewing() const { return bSamplePreviewing; }
+	const bool IsPreviewing() const;
 
 	int32 GetSelectedSampleIndex() const { return SelectedSampleIndex; }
 
@@ -129,9 +137,15 @@ protected:
 	void StartPreviewing();
 	void StopPreviewing();
 
+	/** Common view options for the context menus */
+	void MakeViewContextMenuEntries(FMenuBuilder& InMenuBuilder);
+
 	/** Blend sample context menu creation */
 	TSharedPtr<SWidget> CreateBlendSampleContextMenu();
 
+	/** Create a new blend sample via the context menu */
+	TSharedPtr<SWidget> CreateNewBlendSampleContextMenu(const FVector2D& InMousePosition);
+
 	/** Construct the grid widget to change the grid position for the selected sample */
 	TSharedPtr<SWidget> CreateGridEntryBox(const int32 BoxIndex, const bool bShowLabel);
 
@@ -194,6 +208,12 @@ protected:
 
 	/** Updates the cached blend parameter data */
 	void UpdateCachedBlendParameterData();
+
+	/** Gets the name to display for a sample */
+	FText GetSampleName(const FBlendSample& InBlendSample, int32 InSampleIndex) const;
+
+	/** Enable/disable the status bar message */
+	void EnableStatusBarMessage(bool bEnable);
 private:
 	/** Currently visualized blendspace (const to ensure changes to it are only made within SAnimationBlendSpace */
 	TAttribute<const UBlendSpaceBase*> BlendSpaceBase;
@@ -217,6 +237,7 @@ private:
 	/** Selection and highlight sample index/state */
 	int32 SelectedSampleIndex;
 	int32 HighlightedSampleIndex;
+	int32 ToolTipSampleIndex;
 	bool bHighlightPreviewPin;
 
 	/** Drag state and data (not drag/drop) */
@@ -235,6 +256,9 @@ private:
 	/** Tooltip ptr which is shown when hovering/dropping/dragging a sample*/
 	TSharedPtr<SToolTip> ToolTip;
 
+	/** Container box used to extend a tooltip */
+	TSharedPtr<SBox> ToolTipExtensionContainer;
+
 	/** Drag and drop data */
 	FText InvalidDragDropText;
 	FText InvalidSamplePositionDragDropText;
@@ -268,10 +292,13 @@ private:
 	FVector2D YAxisTextSize;
 
 	/** Delegates populated from SAnimationBlendSpace and used as callbacks */
-	FOnSampleAdded OnSampleAdded;
+	FOnBlendSpaceSampleAdded OnSampleAdded;
 	FOnSampleMoved OnSampleMoved;
-	FOnSampleRemoved OnSampleRemoved;
+	FOnBlendSpaceSampleRemoved OnSampleRemoved;
 	FOnSampleAnimationChanged OnSampleAnimationChanged;
+	FOnBlendSpaceSampleDoubleClicked OnSampleDoubleClicked;
+	FOnGetBlendSpaceSampleName OnGetBlendSpaceSampleName;
+	FOnExtendBlendSpaceSampleTooltip OnExtendSampleTooltip;
 
 	/** Thresshold values for hovering, click and dragging samples */
 	float DragThreshold;
@@ -321,4 +348,8 @@ private:
 	bool bShowAxisLabels;
 
 	bool bShowSettingsButtons;
+
+	FStatusBarMessageHandle StatusBarMessageHandle;
+
+	FName StatusBarName;
 };
diff --git a/Engine/Source/Editor/Persona/Private/SSkeletonSlotNames.cpp b/Engine/Source/Editor/Persona/Private/SSkeletonSlotNames.cpp
--- a/Engine/Source/Editor/Persona/Private/SSkeletonSlotNames.cpp
+++ b/Engine/Source/Editor/Persona/Private/SSkeletonSlotNames.cpp
@@ -101,10 +101,9 @@ TSharedRef< SWidget > SSlotNameListRow::GenerateWidgetForColumn( const FName& Co
 /////////////////////////////////////////////////////
 // FSkeletonSlotNamesSummoner
 
-FSkeletonSlotNamesSummoner::FSkeletonSlotNamesSummoner(TSharedPtr<class FAssetEditorToolkit> InHostingApp, const TSharedRef<IEditableSkeleton>& InEditableSkeleton, FSimpleMulticastDelegate& InOnPostUndo, FOnObjectSelected InOnObjectSelected)
+FSkeletonSlotNamesSummoner::FSkeletonSlotNamesSummoner(TSharedPtr<class FAssetEditorToolkit> InHostingApp, const TSharedRef<IEditableSkeleton>& InEditableSkeleton, FOnObjectSelected InOnObjectSelected)
 	: FWorkflowTabFactory(FPersonaTabs::SkeletonSlotNamesID, InHostingApp)
 	, EditableSkeleton(InEditableSkeleton)
-	, OnPostUndo(InOnPostUndo)
 	, OnObjectSelected(InOnObjectSelected)
 {
 	TabLabel = LOCTEXT("AnimSlotManagerTabTitle", "Anim Slot Manager");
@@ -119,22 +118,20 @@ FSkeletonSlotNamesSummoner::FSkeletonSlotNamesSummoner(TSharedPtr<class FAssetEd
 
 TSharedRef<SWidget> FSkeletonSlotNamesSummoner::CreateTabBody(const FWorkflowTabSpawnInfo& Info) const
 {
-	return SNew(SSkeletonSlotNames, EditableSkeleton.Pin().ToSharedRef(), OnPostUndo)
+	return SNew(SSkeletonSlotNames, EditableSkeleton.Pin().ToSharedRef())
 		.OnObjectSelected(OnObjectSelected);
 }
 
 /////////////////////////////////////////////////////
 // SSkeletonSlotNames
 
-void SSkeletonSlotNames::Construct(const FArguments& InArgs, const TSharedRef<IEditableSkeleton>& InEditableSkeleton, FSimpleMulticastDelegate& InOnPostUndo)
+void SSkeletonSlotNames::Construct(const FArguments& InArgs, const TSharedRef<IEditableSkeleton>& InEditableSkeleton)
 {
 	EditableSkeletonPtr = InEditableSkeleton;
 	OnObjectSelected = InArgs._OnObjectSelected;
 
 	InEditableSkeleton->RegisterOnSlotsChanged(FSimpleMulticastDelegate::FDelegate::CreateSP(this, &SSkeletonSlotNames::RefreshSlotNameListWithFilter));
 
-	InOnPostUndo.Add(FSimpleDelegate::CreateSP( this, &SSkeletonSlotNames::PostUndo ) );
-
 	// Toolbar
 	FToolBarBuilder ToolbarBuilder(TSharedPtr< FUICommandList >(), FMultiBoxCustomization::None);
 
@@ -593,7 +590,7 @@ void SSkeletonSlotNames::ClearDetailsView()
 	OnObjectSelected.ExecuteIfBound(nullptr);
 }
 
-void SSkeletonSlotNames::PostUndo()
+void SSkeletonSlotNames::PostUndoRedo()
 {
 	RefreshSlotNameListWithFilter();
 }
diff --git a/Engine/Source/Editor/Persona/Private/SSkeletonSlotNames.h b/Engine/Source/Editor/Persona/Private/SSkeletonSlotNames.h
--- a/Engine/Source/Editor/Persona/Private/SSkeletonSlotNames.h
+++ b/Engine/Source/Editor/Persona/Private/SSkeletonSlotNames.h
@@ -32,7 +32,7 @@ struct FNotificationInfo;
 struct FSkeletonSlotNamesSummoner : public FWorkflowTabFactory
 {
 public:
-	FSkeletonSlotNamesSummoner(TSharedPtr<class FAssetEditorToolkit> InHostingApp, const TSharedRef<class IEditableSkeleton>& InEditableSkeleton, FSimpleMulticastDelegate& InOnPostUndo, FOnObjectSelected InOnObjectSelected);
+	FSkeletonSlotNamesSummoner(TSharedPtr<class FAssetEditorToolkit> InHostingApp, const TSharedRef<class IEditableSkeleton>& InEditableSkeleton, FOnObjectSelected InOnObjectSelected);
 
 	virtual TSharedRef<SWidget> CreateTabBody(const FWorkflowTabSpawnInfo& Info) const override;
 
@@ -43,7 +43,6 @@ public:
 	}
 public:
 	TWeakPtr<class IEditableSkeleton> EditableSkeleton;
-	FSimpleMulticastDelegate& OnPostUndo;
 	FOnObjectSelected OnObjectSelected;
 };
 
@@ -82,7 +81,7 @@ protected:
 /** Widgets list type */
 typedef STreeView< TSharedPtr<FDisplayedSlotNameInfo> > SSlotNameListType;
 
-class SSkeletonSlotNames : public SCompoundWidget, public FGCObject
+class SSkeletonSlotNames : public SCompoundWidget, public FGCObject, public FSelfRegisteringEditorUndoClient
 {
 public:
 	SLATE_BEGIN_ARGS( SSkeletonSlotNames )
@@ -92,7 +91,7 @@ public:
 
 	SLATE_END_ARGS()
 public:
-	void Construct(const FArguments& InArgs, const TSharedRef<IEditableSkeleton>& InEditableSkeleton, FSimpleMulticastDelegate& InOnPostUndo);
+	void Construct(const FArguments& InArgs, const TSharedRef<IEditableSkeleton>& InEditableSkeleton);
 
 	/**
 	* Accessor so our rows can grab the filter text for highlighting
@@ -106,8 +105,12 @@ public:
 	/** Clears the detail view of whatever we displayed last */
 	void ClearDetailsView();
 
+	// FSelfRegisteringEditorUndoClient interface
+	virtual void PostUndo(bool bSuccess) override { PostUndoRedo(); }
+	virtual void PostRedo(bool bSuccess) override { PostUndoRedo(); }
+
 	/** This triggers a UI repopulation after undo has been called */
-	void PostUndo();
+	void PostUndoRedo();
 
 	// FGCObject interface start
 	virtual void AddReferencedObjects( FReferenceCollector& Collector ) override;
diff --git a/Engine/Source/Editor/Persona/Private/TabSpawners.cpp b/Engine/Source/Editor/Persona/Private/TabSpawners.cpp
--- a/Engine/Source/Editor/Persona/Private/TabSpawners.cpp
+++ b/Engine/Source/Editor/Persona/Private/TabSpawners.cpp
@@ -155,11 +155,10 @@ TSharedRef<SWidget> FMorphTargetTabSummoner::CreateTabBody(const FWorkflowTabSpa
 /////////////////////////////////////////////////////
 // FAnimCurveViewerTabSummoner
 
-FAnimCurveViewerTabSummoner::FAnimCurveViewerTabSummoner(TSharedPtr<class FAssetEditorToolkit> InHostingApp, const TSharedRef<IEditableSkeleton>& InEditableSkeleton, const TSharedRef<IPersonaPreviewScene>& InPreviewScene, FSimpleMulticastDelegate& InOnPostUndo, FOnObjectsSelected InOnObjectsSelected)
+FAnimCurveViewerTabSummoner::FAnimCurveViewerTabSummoner(TSharedPtr<class FAssetEditorToolkit> InHostingApp, const TSharedRef<IEditableSkeleton>& InEditableSkeleton, const TSharedRef<IPersonaPreviewScene>& InPreviewScene, FOnObjectsSelected InOnObjectsSelected)
 	: FWorkflowTabFactory(FPersonaTabs::AnimCurveViewID, InHostingApp)
 	, EditableSkeleton(InEditableSkeleton)
 	, PreviewScene(InPreviewScene)
-	, OnPostUndo(InOnPostUndo)
 	, OnObjectsSelected(InOnObjectsSelected)
 {
 	TabLabel = LOCTEXT("AnimCurveViewTabTitle", "Anim Curves");
@@ -174,7 +173,7 @@ FAnimCurveViewerTabSummoner::FAnimCurveViewerTabSummoner(TSharedPtr<class FAsset
 
 TSharedRef<SWidget> FAnimCurveViewerTabSummoner::CreateTabBody(const FWorkflowTabSpawnInfo& Info) const
 {
-	return SNew(SAnimCurveViewer, EditableSkeleton.Pin().ToSharedRef(), PreviewScene.Pin().ToSharedRef(), OnPostUndo, OnObjectsSelected);
+	return SNew(SAnimCurveViewer, EditableSkeleton.Pin().ToSharedRef(), PreviewScene.Pin().ToSharedRef(), OnObjectsSelected);
 }
 
 /////////////////////////////////////////////////////
diff --git a/Engine/Source/Editor/Persona/Private/TabSpawners.h b/Engine/Source/Editor/Persona/Private/TabSpawners.h
--- a/Engine/Source/Editor/Persona/Private/TabSpawners.h
+++ b/Engine/Source/Editor/Persona/Private/TabSpawners.h
@@ -144,7 +144,7 @@ private:
 struct FAnimCurveViewerTabSummoner : public FWorkflowTabFactory
 {
 public:
-	FAnimCurveViewerTabSummoner(TSharedPtr<class FAssetEditorToolkit> InHostingApp, const TSharedRef<IEditableSkeleton>& InEditableSkeleton, const TSharedRef<IPersonaPreviewScene>& InPreviewScene, FSimpleMulticastDelegate& InOnPostUndo, FOnObjectsSelected InOnObjectsSelected);
+	FAnimCurveViewerTabSummoner(TSharedPtr<class FAssetEditorToolkit> InHostingApp, const TSharedRef<IEditableSkeleton>& InEditableSkeleton, const TSharedRef<IPersonaPreviewScene>& InPreviewScene, FOnObjectsSelected InOnObjectsSelected);
 
 	virtual TSharedRef<SWidget> CreateTabBody(const FWorkflowTabSpawnInfo& Info) const override;
 
@@ -157,7 +157,6 @@ public:
 private:
 	TWeakPtr<class IEditableSkeleton> EditableSkeleton;
 	TWeakPtr<class IPersonaPreviewScene> PreviewScene;
-	FSimpleMulticastDelegate& OnPostUndo;
 	FOnObjectsSelected OnObjectsSelected;
 };
 
diff --git a/Engine/Source/Editor/Persona/Public/PersonaDelegates.h b/Engine/Source/Editor/Persona/Public/PersonaDelegates.h
--- a/Engine/Source/Editor/Persona/Public/PersonaDelegates.h
+++ b/Engine/Source/Editor/Persona/Public/PersonaDelegates.h
@@ -35,4 +35,22 @@ DECLARE_DELEGATE_RetVal(UObject*, FOnGetAsset);
 DECLARE_DELEGATE_ThreeParams(FOnEditCurves, UAnimSequenceBase* /*InAnimSequence*/, const TArray<IAnimationEditor::FCurveEditInfo>& /*InCurveInfo*/, const TSharedPtr<ITimeSliderController>& /*InTimeSliderController*/);
 
 // Called to stop editing curves in the curve editor
-DECLARE_DELEGATE_OneParam(FOnStopEditingCurves, const TArray<IAnimationEditor::FCurveEditInfo>& /*InCurveInfo*/);
\ No newline at end of file
+DECLARE_DELEGATE_OneParam(FOnStopEditingCurves, const TArray<IAnimationEditor::FCurveEditInfo>& /*InCurveInfo*/);
+
+// Called when a blendspace sample point is double clicked
+DECLARE_DELEGATE_OneParam(FOnBlendSpaceSampleDoubleClicked, int32 /*SampleIndex*/);
+
+// Called when a blendspace sample point is removed
+DECLARE_DELEGATE_OneParam(FOnBlendSpaceSampleRemoved, const int32 /*SampleIndex*/);
+
+// Called when a blendspace sample point is added
+DECLARE_DELEGATE_TwoParams(FOnBlendSpaceSampleAdded, UAnimSequence* /*InSequence*/, const FVector& /*InSamplePoint*/);
+
+// Called to get the overridden name of a blend sample
+DECLARE_DELEGATE_RetVal_OneParam(FName, FOnGetBlendSpaceSampleName, int32 /*SampleIndex*/);
+
+// Called to extend a sample's tooltip
+DECLARE_DELEGATE_RetVal_OneParam(TSharedRef<SWidget>, FOnExtendBlendSpaceSampleTooltip, int32 /*SampleIndex*/);
+
+// Called to let the blendspace editor UI set the preview position of an external blendspace node
+DECLARE_DELEGATE_OneParam(FOnSetBlendSpacePreviewPosition, FVector /*InBlendSample*/);
\ No newline at end of file
diff --git a/Engine/Source/Editor/Persona/Public/PersonaModule.h b/Engine/Source/Editor/Persona/Public/PersonaModule.h
--- a/Engine/Source/Editor/Persona/Public/PersonaModule.h
+++ b/Engine/Source/Editor/Persona/Public/PersonaModule.h
@@ -81,11 +81,10 @@ struct FPersonaToolkitArgs
 
 struct FAnimDocumentArgs
 {
-	FAnimDocumentArgs(const TSharedRef<class IPersonaPreviewScene>& InPreviewScene, const TSharedRef<class IPersonaToolkit>& InPersonaToolkit, const TSharedRef<class IEditableSkeleton>& InEditableSkeleton, FSimpleMulticastDelegate& InOnPostUndo, FSimpleMulticastDelegate& InOnSectionsChanged)
+	FAnimDocumentArgs(const TSharedRef<class IPersonaPreviewScene>& InPreviewScene, const TSharedRef<class IPersonaToolkit>& InPersonaToolkit, const TSharedRef<class IEditableSkeleton>& InEditableSkeleton, FSimpleMulticastDelegate& InOnSectionsChanged)
 		: PreviewScene(InPreviewScene)
 		, PersonaToolkit(InPersonaToolkit)
 		, EditableSkeleton(InEditableSkeleton)
-		, OnPostUndo(InOnPostUndo)
 		, OnSectionsChanged(InOnSectionsChanged)
 	{}
 
@@ -93,7 +92,6 @@ struct FAnimDocumentArgs
 	TWeakPtr<class IPersonaPreviewScene> PreviewScene;
 	TWeakPtr<class IPersonaToolkit> PersonaToolkit;
 	TWeakPtr<class IEditableSkeleton> EditableSkeleton;
-	FSimpleMulticastDelegate& OnPostUndo;
 	FSimpleMulticastDelegate& OnSectionsChanged;
 
 	/** Optional args */
@@ -102,6 +100,44 @@ struct FAnimDocumentArgs
 	FSimpleDelegate OnDespatchSectionsChanged;
 };
 
+struct FBlendSpaceEditorArgs
+{
+	// Called when a blendspace sample point is removed
+	FOnBlendSpaceSampleRemoved OnBlendSpaceSampleRemoved;
+
+	// Called when a blendspace sample point is added
+	FOnBlendSpaceSampleAdded OnBlendSpaceSampleAdded;
+
+	// Called when a blendspace sample point is double clicked
+	FOnBlendSpaceSampleDoubleClicked OnBlendSpaceSampleDoubleClicked;
+
+	// Called to get the overridden name of a blend sample
+	FOnGetBlendSpaceSampleName OnGetBlendSpaceSampleName;
+
+	// Allows the preview position to be programmatically driven
+	TAttribute<FVector> PreviewPosition;
+
+	// Allows an external widget to be inserted into a sample's tooltip
+	FOnExtendBlendSpaceSampleTooltip OnExtendSampleTooltip;
+
+	// Allows preview position to drive external node
+	FOnSetBlendSpacePreviewPosition OnSetPreviewPosition;
+
+	// Status bar to display hint messages in
+	FName StatusBarName = TEXT("AssetEditor.AnimationEditor.MainMenu");
+};
+
+struct FBlendSpacePreviewArgs
+{
+	TAttribute<const UBlendSpaceBase*> PreviewBlendSpace;
+
+	// Allows the preview position to be programmatically driven
+	TAttribute<FVector> PreviewPosition;
+
+	// Called to get the overridden name of a blend sample
+	FOnGetBlendSpaceSampleName OnGetBlendSpaceSampleName;
+};
+
 /** Places that viewport text can be placed */
 enum class EViewportCorner : uint8
 {
@@ -220,9 +256,12 @@ public:
 	/** Create an anim notifies tab factory */
 	virtual TSharedRef<FWorkflowTabFactory> CreateAnimNotifiesTabFactory(const TSharedRef<class FWorkflowCentricApplication>& InHostingApp, const TSharedRef<class IEditableSkeleton>& InEditableSkeleton, FOnObjectsSelected InOnObjectsSelected) const;
 
-	/** Create a skeleton cuve viewer tab factory */
+	UE_DEPRECATED(5.0, "Please use the overload that does not take a post-undo delegate")
 	virtual TSharedRef<FWorkflowTabFactory> CreateCurveViewerTabFactory(const TSharedRef<class FWorkflowCentricApplication>& InHostingApp, const TSharedRef<class IEditableSkeleton>& InEditableSkeleton, const TSharedRef<class IPersonaPreviewScene>& InPreviewScene, FSimpleMulticastDelegate& InOnPostUndo, FOnObjectsSelected InOnObjectsSelected) const;
 
+	/** Create a skeleton curve viewer tab factory */
+	virtual TSharedRef<FWorkflowTabFactory> CreateCurveViewerTabFactory(const TSharedRef<class FWorkflowCentricApplication>& InHostingApp, const TSharedRef<class IEditableSkeleton>& InEditableSkeleton, const TSharedRef<class IPersonaPreviewScene>& InPreviewScene, FOnObjectsSelected InOnObjectsSelected) const;
+
 	/** Create a retarget manager tab factory */
 	virtual TSharedRef<FWorkflowTabFactory> CreateRetargetManagerTabFactory(const TSharedRef<class FWorkflowCentricApplication>& InHostingApp, const TSharedRef<class IEditableSkeleton>& InEditableSkeleton, const TSharedRef<IPersonaPreviewScene>& InPreviewScene, FSimpleMulticastDelegate& InOnPostUndo) const;
 
@@ -244,12 +283,22 @@ public:
 	/** Create a tab factory for editing anim blueprint parent overrides */
 	virtual TSharedRef<FWorkflowTabFactory> CreateAnimBlueprintAssetOverridesTabFactory(const TSharedRef<class FBlueprintEditor>& InBlueprintEditor, UAnimBlueprint* InAnimBlueprint, FSimpleMulticastDelegate& InOnPostUndo) const;
 
-	/** Create a tab factory for editing slot names and groups */
+	UE_DEPRECATED(5.0, "Please use the overload that does not take a post-undo delegate")
 	virtual TSharedRef<FWorkflowTabFactory> CreateSkeletonSlotNamesTabFactory(const TSharedRef<class FWorkflowCentricApplication>& InHostingApp, const TSharedRef<class IEditableSkeleton>& InEditableSkeleton, FSimpleMulticastDelegate& InOnPostUndo, FOnObjectSelected InOnObjectSelected) const;
 
-	/** Create a widget to preview a blendspace */
+	/** Create a tab factory for editing slot names and groups */
+	virtual TSharedRef<FWorkflowTabFactory> CreateSkeletonSlotNamesTabFactory(const TSharedRef<class FWorkflowCentricApplication>& InHostingApp, const TSharedRef<class IEditableSkeleton>& InEditableSkeleton, FOnObjectSelected InOnObjectSelected) const;
+
+	/** Deprecated */
+	UE_DEPRECATED(5.0, "Please use the overload that takes a FBlendSpacePreviewArgs struct")
 	virtual TSharedRef<SWidget> CreateBlendSpacePreviewWidget(TAttribute<const UBlendSpaceBase*> InBlendSpace, TAttribute<FVector> InPosition) const;
 
+	/** Create a widget to preview a blendspace */
+	virtual TSharedRef<SWidget> CreateBlendSpacePreviewWidget(const FBlendSpacePreviewArgs& InArgs) const;
+
+	/** Create a widget to edit a blendspace */
+	virtual TSharedRef<SWidget> CreateBlendSpaceEditWidget(UBlendSpaceBase* InBlendSpace, const FBlendSpaceEditorArgs& InArgs) const;
+
 	/** Create a tab factory for editing montage sections */
 	virtual TSharedRef<FWorkflowTabFactory> CreateAnimMontageSectionsTabFactory(const TSharedRef<class FWorkflowCentricApplication>& InHostingApp, const TSharedRef<IPersonaToolkit>& InPersonaToolkit, FSimpleMulticastDelegate& InOnSectionsChanged) const;
 
diff --git a/Engine/Source/Editor/SkeletonEditor/Private/SkeletonEditorMode.cpp b/Engine/Source/Editor/SkeletonEditor/Private/SkeletonEditorMode.cpp
--- a/Engine/Source/Editor/SkeletonEditor/Private/SkeletonEditorMode.cpp
+++ b/Engine/Source/Editor/SkeletonEditor/Private/SkeletonEditorMode.cpp
@@ -31,8 +31,8 @@ FSkeletonEditorMode::FSkeletonEditorMode(TSharedRef<FWorkflowCentricApplication>
 	TabFactories.RegisterFactory(PersonaModule.CreateAnimNotifiesTabFactory(InHostingApp, InSkeletonTree->GetEditableSkeleton(),  OnObjectsSelected));
 	TabFactories.RegisterFactory(PersonaModule.CreateAdvancedPreviewSceneTabFactory(InHostingApp, SkeletonEditor->GetPersonaToolkit()->GetPreviewScene()));
 	TabFactories.RegisterFactory(PersonaModule.CreateRetargetManagerTabFactory(InHostingApp, InSkeletonTree->GetEditableSkeleton(), SkeletonEditor->GetPersonaToolkit()->GetPreviewScene(), SkeletonEditor->OnPostUndo));
-	TabFactories.RegisterFactory(PersonaModule.CreateCurveViewerTabFactory(InHostingApp, InSkeletonTree->GetEditableSkeleton(), SkeletonEditor->GetPersonaToolkit()->GetPreviewScene(), SkeletonEditor->OnPostUndo, OnObjectsSelected));
-	TabFactories.RegisterFactory(PersonaModule.CreateSkeletonSlotNamesTabFactory(InHostingApp, InSkeletonTree->GetEditableSkeleton(), SkeletonEditor->OnPostUndo, OnObjectSelected));
+	TabFactories.RegisterFactory(PersonaModule.CreateCurveViewerTabFactory(InHostingApp, InSkeletonTree->GetEditableSkeleton(), SkeletonEditor->GetPersonaToolkit()->GetPreviewScene(), OnObjectsSelected));
+	TabFactories.RegisterFactory(PersonaModule.CreateSkeletonSlotNamesTabFactory(InHostingApp, InSkeletonTree->GetEditableSkeleton(), OnObjectSelected));
 	TabFactories.RegisterFactory(PersonaModule.CreateAssetDetailsTabFactory(InHostingApp, FOnGetAsset::CreateSP(&SkeletonEditor.Get(), &FSkeletonEditor::HandleGetAsset), FOnDetailsCreated()));
 
 	TabLayout = FTabManager::NewLayout("Standalone_SkeletonEditor_Layout_v1.3")
diff --git a/Engine/Source/Editor/UnrealEd/Private/Kismet2/BlueprintEditorUtils.cpp b/Engine/Source/Editor/UnrealEd/Private/Kismet2/BlueprintEditorUtils.cpp
--- a/Engine/Source/Editor/UnrealEd/Private/Kismet2/BlueprintEditorUtils.cpp
+++ b/Engine/Source/Editor/UnrealEd/Private/Kismet2/BlueprintEditorUtils.cpp
@@ -2266,34 +2266,40 @@ const UFunction* FBlueprintEditorUtils::GetMostUpToDateFunction(const UFunction*
 	return GetMostUpToDateFunction(const_cast<UFunction*>(Function));
 }
 
-bool FBlueprintEditorUtils::IsGraphNameUnique(UBlueprint* Blueprint, const FName& InName)
+bool FBlueprintEditorUtils::IsGraphNameUnique(UObject* InOuter, const FName& InName)
 {
 	// Check for any object directly created in the blueprint
-	if( !FindObject<UObject>(Blueprint, *InName.ToString()) )
+	if( !FindObject<UObject>(InOuter, *InName.ToString()) )
 	{
-		// Next, check for functions with that name in the blueprint's class scope
-		FFieldVariant ExistingField = FindUFieldOrFProperty(Blueprint->SkeletonGeneratedClass, InName);
-		if( !ExistingField )
+		if(UBlueprint* Blueprint = Cast<UBlueprint>(InOuter))
 		{
-			// Finally, check function entry points
-			TArray<UK2Node_Event*> AllEvents;
-			FBlueprintEditorUtils::GetAllNodesOfClass<UK2Node_Event>(Blueprint, AllEvents);
-
-			for(int32 i=0; i < AllEvents.Num(); i++)
+			// Next, check for functions with that name in the blueprint's class scope
+			FFieldVariant ExistingField = FindUFieldOrFProperty(Blueprint->SkeletonGeneratedClass, InName);
+			if( !ExistingField )
 			{
-				UK2Node_Event* EventNode = AllEvents[i];
-				check(EventNode);
+				// Finally, check function entry points
+				TArray<UK2Node_Event*> AllEvents;
+				FBlueprintEditorUtils::GetAllNodesOfClass<UK2Node_Event>(Blueprint, AllEvents);
 
-				if( EventNode->CustomFunctionName == InName
-					|| EventNode->EventReference.GetMemberName() == InName )
+				for(int32 i=0; i < AllEvents.Num(); i++)
 				{
-					return false;
+					UK2Node_Event* EventNode = AllEvents[i];
+					check(EventNode);
+
+					if( EventNode->CustomFunctionName == InName
+						|| EventNode->EventReference.GetMemberName() == InName )
+					{
+						return false;
+					}
 				}
-			}
 
-			// All good!
-			return true;
+				// All good!
+				return true;
+			}
 		}
+
+		// All good!
+		return true;
 	}
 
 	return false;
@@ -2632,11 +2638,7 @@ void FBlueprintEditorUtils::RemoveGraph(UBlueprint* Blueprint, class UEdGraph* G
 			OuterGraph->Modify();
 			OuterGraph->SubGraphs.Remove(GraphToRemove);
 		}
-		else if (! (Cast<UK2Node_Composite>(TestOuter)	|| 
-					Cast<UMaterialGraphNode_Composite>(TestOuter)	||
-					Cast<UAnimStateNodeBase>(TestOuter)	||
-					Cast<UAnimStateTransitionNode>(TestOuter)	||
-					Cast<UAnimGraphNode_StateMachineBase>(TestOuter)) )
+		else if (! (Cast<UEdGraphNode>(TestOuter) && Cast<UEdGraphNode>(TestOuter)->GetSubGraphs().Num() > 0) )
 		{
 			break;
 		}
@@ -2651,56 +2653,17 @@ void FBlueprintEditorUtils::RemoveGraph(UBlueprint* Blueprint, class UEdGraph* G
 	}
 
 	// Handle subgraphs held in graph
-	TArray<UK2Node_Composite*> AllCompositeNodes;
-	GraphToRemove->GetNodesOfClass<UK2Node_Composite>(AllCompositeNodes);
-
-	const bool bDontRecompile = true;
-	for (UK2Node_Composite* CompNode : AllCompositeNodes)
-	{
-		if (CompNode->BoundGraph && Local::IsASubGraph(CompNode->BoundGraph))
-		{
-			FBlueprintEditorUtils::RemoveGraph(Blueprint, CompNode->BoundGraph, EGraphRemoveFlags::None);
-		}
-	}
-
-	// Materials have their own form of composite, and we haven't made a generalized composite EdGraphNode yet.
-	TArray<UMaterialGraphNode_Composite*> AllMaterialCompositeNodes;
-	GraphToRemove->GetNodesOfClass<UMaterialGraphNode_Composite>(AllMaterialCompositeNodes);
-
-	for (UMaterialGraphNode_Composite* CompNode : AllMaterialCompositeNodes)
-	{
-		UEdGraph* BoundGraph = CompNode->BoundGraph;
-		if (BoundGraph && BoundGraph->GetOuter()->IsA(UMaterialGraphNode_Composite::StaticClass()))
-		{
-			FBlueprintEditorUtils::RemoveGraph(Blueprint, CompNode->BoundGraph, EGraphRemoveFlags::None);
-		}
-	}
-
-	// Animation nodes can contain subgraphs but are not composite nodes, handle their graphs
-	TArray<UAnimStateNodeBase*> AllAnimCompositeNodes;
-	GraphToRemove->GetNodesOfClassEx<UAnimStateNode>(AllAnimCompositeNodes);
-	GraphToRemove->GetNodesOfClassEx<UAnimStateConduitNode>(AllAnimCompositeNodes);
-	GraphToRemove->GetNodesOfClassEx<UAnimStateTransitionNode>(AllAnimCompositeNodes);
-
-	for(UAnimStateNodeBase* Node : AllAnimCompositeNodes)
-	{
-		UEdGraph* BoundGraph = Node->GetBoundGraph();
-		if(BoundGraph && BoundGraph->GetOuter()->IsA(UAnimStateNodeBase::StaticClass()))
-		{
-			FBlueprintEditorUtils::RemoveGraph(Blueprint, BoundGraph, EGraphRemoveFlags::None);
-		}
-	}
-
-	// Handle sub anim state machines
-	TArray<UAnimGraphNode_StateMachineBase*> AllStateMachines;
-	GraphToRemove->GetNodesOfClassEx<UAnimGraphNode_StateMachine>(AllStateMachines);
+	TArray<UEdGraphNode*> AllNodes;
+	GraphToRemove->GetNodesOfClass<UEdGraphNode>(AllNodes);
 
-	for(UAnimGraphNode_StateMachineBase* Node : AllStateMachines)
+	for (UEdGraphNode* GraphNode : AllNodes)
 	{
-		UEdGraph* BoundGraph = Node->EditorStateMachineGraph;
-		if(BoundGraph && BoundGraph->GetOuter()->IsA(UAnimGraphNode_StateMachineBase::StaticClass()))
+		for(UEdGraph* SubGraph : GraphNode->GetSubGraphs())
 		{
-			FBlueprintEditorUtils::RemoveGraph(Blueprint, BoundGraph, EGraphRemoveFlags::None);
+			if (SubGraph && SubGraph->GetOuter()->IsA(UEdGraphNode::StaticClass()))
+			{
+				FBlueprintEditorUtils::RemoveGraph(Blueprint, SubGraph, EGraphRemoveFlags::None);
+			}
 		}
 	}
 
@@ -7290,37 +7253,18 @@ void FBlueprintEditorUtils::UpdateOutOfDateCompositeWithOuter(UBlueprint* Bluepr
 	for (UEdGraphNode* Node : OuterGraph->Nodes)
 	{
 		//Is this node of a type that has a BoundGraph to update
-		UEdGraph* BoundGraph = nullptr;
-		if (UK2Node_Composite* Composite = Cast<UK2Node_Composite>(Node))
-		{
-			BoundGraph = Composite->BoundGraph;
-		}
-		else if (UAnimStateNode* StateNode = Cast<UAnimStateNode>(Node))
+		for(UEdGraph* BoundGraph : Node->GetSubGraphs())
 		{
-			BoundGraph = StateNode->BoundGraph;
-		}
-		else if (UAnimStateTransitionNode* TransitionNode = Cast<UAnimStateTransitionNode>(Node))
-		{
-			// Shared rule graphs are not necessarily outered to their node, so skip them here
-			if(!TransitionNode->bSharedRules)
+			if (BoundGraph)
 			{
-				BoundGraph = TransitionNode->BoundGraph;
-			}
-		}
-		else if (UAnimGraphNode_StateMachineBase* StateMachineNode = Cast<UAnimGraphNode_StateMachineBase>(Node))
-		{
-			BoundGraph = StateMachineNode->EditorStateMachineGraph;
-		}
-
-		if (BoundGraph)
-		{
-			// Check for out of date BoundGraph where outer is not the composite node
-			if (BoundGraph->GetOuter() != Node)
-			{
-				// change the outer of the BoundGraph to be the composite node instead of the OuterGraph
-				if (false == BoundGraph->Rename(*BoundGraph->GetName(), Node, ((BoundGraph->HasAnyFlags(RF_NeedLoad | RF_NeedPostLoad) ? REN_ForceNoResetLoaders : 0) | REN_DontCreateRedirectors)))
+				// Check for out of date BoundGraph where outer is not the composite node
+				if (BoundGraph->GetOuter() != Node)
 				{
-					UE_LOG(LogBlueprintDebug, Log, TEXT("CompositeNode: On Blueprint '%s' could not fix Outer() for BoundGraph of composite node '%s'"), *Blueprint->GetPathName(), *Node->GetName());
+					// change the outer of the BoundGraph to be the composite node instead of the OuterGraph
+					if (false == BoundGraph->Rename(*BoundGraph->GetName(), Node, ((BoundGraph->HasAnyFlags(RF_NeedLoad | RF_NeedPostLoad) ? REN_ForceNoResetLoaders : 0) | REN_DontCreateRedirectors)))
+					{
+						UE_LOG(LogBlueprintDebug, Log, TEXT("CompositeNode: On Blueprint '%s' could not fix Outer() for BoundGraph of composite node '%s'"), *Blueprint->GetPathName(), *Node->GetName());
+					}
 				}
 			}
 		}
@@ -8927,12 +8871,12 @@ bool FBlueprintEditorUtils::PropertyValueToString_Direct(const FProperty* Proper
 	return bSucceeded;
 }
 
-FName FBlueprintEditorUtils::GenerateUniqueGraphName(UBlueprint* const BlueprintOuter, FString const& ProposedName)
+FName FBlueprintEditorUtils::GenerateUniqueGraphName(UObject* const InOuter, FString const& ProposedName)
 {
 	FName UniqueGraphName(*ProposedName);
 
 	int32 CountPostfix = 1;
-	while (!FBlueprintEditorUtils::IsGraphNameUnique(BlueprintOuter, UniqueGraphName))
+	while (!FBlueprintEditorUtils::IsGraphNameUnique(InOuter, UniqueGraphName))
 	{
 		UniqueGraphName = FName(*FString::Printf(TEXT("%s%i"), *ProposedName, CountPostfix));
 		++CountPostfix;
diff --git a/Engine/Source/Editor/UnrealEd/Public/Kismet2/BlueprintEditorUtils.h b/Engine/Source/Editor/UnrealEd/Public/Kismet2/BlueprintEditorUtils.h
--- a/Engine/Source/Editor/UnrealEd/Public/Kismet2/BlueprintEditorUtils.h
+++ b/Engine/Source/Editor/UnrealEd/Public/Kismet2/BlueprintEditorUtils.h
@@ -318,7 +318,7 @@ public:
 	static void MarkBlueprintAsModified(UBlueprint* Blueprint, FPropertyChangedEvent PropertyChangedEvent = FPropertyChangedEvent(nullptr));
 
 	/** See whether or not the specified graph name / entry point name is unique */
-	static bool IsGraphNameUnique(UBlueprint* Blueprint, const FName& InName);
+	static bool IsGraphNameUnique(UObject* InOuter, const FName& InName);
 
 	/**
 	 * Creates a new empty graph.
@@ -1576,11 +1576,11 @@ public:
 	 * Generates a unique graph name for the supplied blueprint (guaranteed to not 
 	 * cause a naming conflict at the time of the call).
 	 * 
-	 * @param  BlueprintOuter	The blueprint you want a unique graph name for.
+	 * @param  InOuter		The blueprint/object you want a unique graph name for.
 	 * @param  ProposedName		The name you want to give the graph (the result will be some permutation of this string).
 	 * @return A unique name intended for a new graph.
 	 */
-	static FName GenerateUniqueGraphName(UBlueprint* const BlueprintOuter, FString const& ProposedName);
+	static FName GenerateUniqueGraphName(UObject* const InOuter, FString const& ProposedName);
 
 	/* Checks if the passed in selection set causes cycling on compile
 	 *
diff --git a/Engine/Source/Runtime/AnimGraphRuntime/Private/AnimGraphRuntimeTrace.cpp b/Engine/Source/Runtime/AnimGraphRuntime/Private/AnimGraphRuntimeTrace.cpp
--- a/Engine/Source/Runtime/AnimGraphRuntime/Private/AnimGraphRuntimeTrace.cpp
+++ b/Engine/Source/Runtime/AnimGraphRuntime/Private/AnimGraphRuntimeTrace.cpp
@@ -7,6 +7,7 @@
 #include "Trace/Trace.inl"
 #include "Animation/BlendSpaceBase.h"
 #include "AnimNodes/AnimNode_BlendSpacePlayer.h"
+#include "AnimNodes/AnimNode_BlendSpaceGraphBase.h"
 #include "Animation/AnimInstanceProxy.h"
 
 UE_TRACE_EVENT_BEGIN(Animation, BlendSpacePlayer)
@@ -42,4 +43,29 @@ void FAnimGraphRuntimeTrace::OutputBlendSpacePlayer(const FAnimationBaseContext&
 		<< BlendSpacePlayer.PositionZ(InNode.Z);
 }
 
+void FAnimGraphRuntimeTrace::OutputBlendSpace(const FAnimationBaseContext& InContext, const FAnimNode_BlendSpaceGraphBase& InNode)
+{
+	bool bEventEnabled = UE_TRACE_CHANNELEXPR_IS_ENABLED(AnimationChannel);
+	if (!bEventEnabled)
+	{
+		return;
+	}
+
+	check(InContext.AnimInstanceProxy);
+
+	TRACE_OBJECT(InContext.AnimInstanceProxy->GetAnimInstanceObject());
+	TRACE_OBJECT(InNode.GetBlendSpace());
+
+	FVector Coordinates = InNode.GetSampleCoordinates();
+
+	UE_TRACE_LOG(Animation, BlendSpacePlayer, AnimationChannel)
+		<< BlendSpacePlayer.Cycle(FPlatformTime::Cycles64())
+		<< BlendSpacePlayer.AnimInstanceId(FObjectTrace::GetObjectId(InContext.AnimInstanceProxy->GetAnimInstanceObject()))
+		<< BlendSpacePlayer.BlendSpaceId(FObjectTrace::GetObjectId(InNode.GetBlendSpace()))
+		<< BlendSpacePlayer.NodeId(InContext.GetCurrentNodeId())
+		<< BlendSpacePlayer.PositionX(Coordinates.X)
+		<< BlendSpacePlayer.PositionY(Coordinates.Y)
+		<< BlendSpacePlayer.PositionZ(Coordinates.Z);
+}
+
 #endif
\ No newline at end of file
diff --git a/Engine/Source/Runtime/AnimGraphRuntime/Public/AnimGraphRuntimeTrace.h b/Engine/Source/Runtime/AnimGraphRuntime/Public/AnimGraphRuntimeTrace.h
--- a/Engine/Source/Runtime/AnimGraphRuntime/Public/AnimGraphRuntimeTrace.h
+++ b/Engine/Source/Runtime/AnimGraphRuntime/Public/AnimGraphRuntimeTrace.h
@@ -9,18 +9,26 @@
 
 struct FAnimationBaseContext;
 struct FAnimNode_BlendSpacePlayer;
+struct FAnimNode_BlendSpaceGraphBase;
 
 struct FAnimGraphRuntimeTrace
 {
 	/** Helper function to output debug info for blendspace player nodes */
 	ANIMGRAPHRUNTIME_API static void OutputBlendSpacePlayer(const FAnimationBaseContext& InContext, const FAnimNode_BlendSpacePlayer& InNode);
+
+	/** Helper function to output debug info for blendspace nodes */
+	ANIMGRAPHRUNTIME_API static void OutputBlendSpace(const FAnimationBaseContext& InContext, const FAnimNode_BlendSpaceGraphBase& InNode);
 };
 
 #define TRACE_BLENDSPACE_PLAYER(Context, Node) \
 	FAnimGraphRuntimeTrace::OutputBlendSpacePlayer(Context, Node);
 
+#define TRACE_BLENDSPACE(Context, Node) \
+	FAnimGraphRuntimeTrace::OutputBlendSpace(Context, Node);
+
 #else
 
 #define TRACE_BLENDSPACE_PLAYER(Context, Node)
+#define TRACE_BLENDSPACE(Context, Node)
 
 #endif
\ No newline at end of file
diff --git a/Engine/Source/Runtime/AnimGraphRuntime/Public/AnimNodes/AnimNode_BlendSpaceGraphBase.h b/Engine/Source/Runtime/AnimGraphRuntime/Public/AnimNodes/AnimNode_BlendSpaceGraphBase.h
new file mode 100644
--- /dev/null
+++ b/Engine/Source/Runtime/AnimGraphRuntime/Public/AnimNodes/AnimNode_BlendSpaceGraphBase.h
@@ -0,0 +1,87 @@
+// Copyright Epic Games, Inc. All Rights Reserved.
+
+#pragma once
+
+#include "CoreMinimal.h"
+#include "UObject/ObjectMacros.h"
+#include "Animation/AnimationAsset.h"
+#include "Animation/AnimNodeBase.h"
+#include "AnimNode_BlendSpaceGraphBase.generated.h"
+
+class UBlendSpaceBase;
+
+// Allows multiple animations to be blended between based on input parameters
+USTRUCT(BlueprintInternalUseOnly)
+struct ANIMGRAPHRUNTIME_API FAnimNode_BlendSpaceGraphBase : public FAnimNode_Base
+{
+	GENERATED_BODY()
+
+	friend class UAnimGraphNode_BlendSpaceGraphBase;
+
+	// @return the blendspace that this node uses
+	const UBlendSpaceBase* GetBlendSpace() const { return BlendSpace; }
+
+	// @return the current sample coordinates that this node is using to sample the blendspace
+	FVector GetSampleCoordinates() const { return FVector(X, Y, Z); }
+
+#if WITH_EDITORONLY_DATA
+	// Set the node to preview a supplied sample value
+	void SetPreviewSampleValue(FVector InVector);
+#endif
+
+protected:
+	// The X coordinate to sample in the blendspace
+	UPROPERTY(EditAnywhere, BlueprintReadWrite, Category = Coordinates, meta = (PinShownByDefault, AllowPrivateAccess))
+	float X = 0.0f;
+
+	// The Y coordinate to sample in the blendspace
+	UPROPERTY(EditAnywhere, BlueprintReadWrite, Category = Coordinates, meta = (PinShownByDefault, AllowPrivateAccess))
+	float Y = 0.0f;
+
+	// The Z coordinate to sample in the blendspace
+	UPROPERTY(EditAnywhere, BlueprintReadWrite, Category = Coordinates, meta = (PinHiddenByDefault, AllowPrivateAccess))
+	float Z = 0.0f;
+
+	// The group name that we synchronize with. All nodes employing sync beyond this in the anim graph will implicitly use this sync group.
+	UPROPERTY(EditAnywhere, BlueprintReadWrite, Category = Sync, meta = (NeverAsPin, AllowPrivateAccess))
+	FName GroupName = NAME_None;
+
+	// The role this player can assume within the group
+	UPROPERTY(EditAnywhere, BlueprintReadWrite, Category = Sync, meta = (NeverAsPin, AllowPrivateAccess))
+	TEnumAsByte<EAnimGroupRole::Type> GroupRole = EAnimGroupRole::CanBeLeader;
+
+	// The internal blendspace asset to play
+	UPROPERTY()
+	const UBlendSpaceBase* BlendSpace = nullptr;
+
+	// Pose links for each sample in the blendspace
+	UPROPERTY()
+	TArray<FPoseLink> SamplePoseLinks;
+
+protected:
+	// FIR filter applied to inputs
+	FBlendFilter BlendFilter;
+
+	// Cache of sampled data, updated each frame
+	TArray<FBlendSampleData> BlendSampleDataCache;
+
+#if WITH_EDITORONLY_DATA
+	// Preview sample value - set in editor only
+	FVector PreviewSample = FVector::ZeroVector;
+
+	// Whether to use the preview sample value
+	bool bUsePreviewSampleValue = false;
+#endif
+
+	// Internal update handler, skipping evaluation of exposed inputs
+	void UpdateInternal(const FAnimationUpdateContext& Context);
+
+protected:
+	// FAnimNode_Base interface
+	virtual void Initialize_AnyThread(const FAnimationInitializeContext& Context) override;
+	virtual void CacheBones_AnyThread(const FAnimationCacheBonesContext& Context) override;
+	virtual void Update_AnyThread(const FAnimationUpdateContext& Context) override;
+	virtual void Evaluate_AnyThread(FPoseContext& Output) override;
+	virtual void GatherDebugData(FNodeDebugData& DebugData) override;
+	// End of FAnimNode_Base interface
+};
diff --git a/Engine/Source/Runtime/Engine/Classes/Animation/AnimBlueprintGeneratedClass.h b/Engine/Source/Runtime/Engine/Classes/Animation/AnimBlueprintGeneratedClass.h
--- a/Engine/Source/Runtime/Engine/Classes/Animation/AnimBlueprintGeneratedClass.h
+++ b/Engine/Source/Runtime/Engine/Classes/Animation/AnimBlueprintGeneratedClass.h
@@ -232,7 +232,7 @@ public:
 	// Record of a blend space player's state
 	struct FBlendSpacePlayerRecord
 	{
-		FBlendSpacePlayerRecord(int32 InNodeID, UBlendSpaceBase* InBlendSpace, float InPositionX, float InPositionY, float InPositionZ)
+		FBlendSpacePlayerRecord(int32 InNodeID, const UBlendSpaceBase* InBlendSpace, float InPositionX, float InPositionY, float InPositionZ)
 			: NodeID(InNodeID)
 			, BlendSpace(InBlendSpace)
 			, PositionX(InPositionX)
@@ -286,7 +286,7 @@ public:
 	void RecordStateData(int32 StateMachineIndex, int32 StateIndex, float Weight, float ElapsedTime);
 	void RecordNodeValue(int32 InNodeID, const FString& InText);
 	void RecordSequencePlayer(int32 InNodeID, float InPosition, float InLength, int32 InFrameCount);
-	void RecordBlendSpacePlayer(int32 InNodeID, UBlendSpaceBase* InBlendSpace, float InPositionX, float InPositionY, float InPositionZ);
+	void RecordBlendSpacePlayer(int32 InNodeID, const UBlendSpaceBase* InBlendSpace, float InPositionX, float InPositionY, float InPositionZ);
 
 	void AddPoseWatch(int32 NodeID, FColor Color);
 	void RemovePoseWatch(int32 NodeID);
@@ -364,6 +364,10 @@ private:
 	UPROPERTY()
 	FPropertyAccessLibrary PropertyAccessLibrary;
 
+	// Any internal blendspaces we host
+	UPROPERTY()
+	TArray<UBlendSpaceBase*> BlendSpaces;
+
 public:
 	// IAnimClassInterface interface
 	virtual const TArray<FBakedAnimationStateMachine>& GetBakedStateMachines() const override { return GetRootClass()->GetBakedStateMachines_Direct(); }
@@ -496,7 +500,7 @@ public:
 		FStructProperty* AnimationProperty = AnimNodeProperties[AnimNodeProperties.Num() - 1 - Index];
 		check(AnimationProperty);
 		check(AnimationProperty->Struct->IsChildOf(StructType::StaticStruct()));
-		return AnimationProperty->ContainerPtrToValuePtr<StructType>((void*)Object);
+		return *AnimationProperty->ContainerPtrToValuePtr<StructType>((void*)Object);
 	}
 
 	// Gets the property index from the original UAnimGraphNode's GUID. Does not remap to property order.
diff --git a/Engine/Source/Runtime/Engine/Classes/Animation/BlendSpaceBase.h b/Engine/Source/Runtime/Engine/Classes/Animation/BlendSpaceBase.h
--- a/Engine/Source/Runtime/Engine/Classes/Animation/BlendSpaceBase.h
+++ b/Engine/Source/Runtime/Engine/Classes/Animation/BlendSpaceBase.h
@@ -12,6 +12,8 @@
 #include "BoneContainer.h"
 #include "Animation/AnimationAsset.h"
 #include "AnimationRuntime.h"
+#include "AnimNodeBase.h"
+#include "Containers/ArrayView.h"
 #include "BlendSpaceBase.generated.h"
 
 /** Interpolation data types. */
@@ -89,6 +91,7 @@ struct FBlendSample
 {
 	GENERATED_USTRUCT_BODY()
 
+	// For linked animations
 	UPROPERTY(EditAnywhere, Category=BlendSample)
 	class UAnimSequence* Animation;
 
@@ -113,7 +116,7 @@ struct FBlendSample
 #endif // WITH_EDITORONLY_DATA
 
 	FBlendSample()
-		: Animation(NULL)
+		: Animation(nullptr)
 		, SampleValue(0.f)
 		, RateScale(1.0f)
 #if WITH_EDITORONLY_DATA
@@ -241,6 +244,7 @@ class UBlendSpaceBase : public UAnimationAsset, public IInterpolationIndexProvid
 
 	/** Required for accessing protected variable names */
 	friend class FBlendSpaceDetails;
+	friend class UAnimGraphNode_BlendSpaceGraphBase;
 
 	//~ Begin UObject Interface
 	virtual void PostLoad() override;
@@ -282,9 +286,11 @@ class UBlendSpaceBase : public UAnimationAsset, public IInterpolationIndexProvid
 	 * BlendSpace Get Animation Pose function
 	 */
 	UE_DEPRECATED(4.26, "Use GetAnimationPose with other signature")
-	ENGINE_API void GetAnimationPose(TArray<FBlendSampleData>& BlendSampleDataCache, /*out*/ FCompactPose& OutPose, /*out*/ FBlendedCurve& OutCurve);
+	ENGINE_API void GetAnimationPose(TArray<FBlendSampleData>& BlendSampleDataCache, /*out*/ FCompactPose& OutPose, /*out*/ FBlendedCurve& OutCurve) const;
 	
-	ENGINE_API void GetAnimationPose(TArray<FBlendSampleData>& BlendSampleDataCache, /*out*/ FAnimationPoseData& OutAnimationPoseData);
+	ENGINE_API void GetAnimationPose(TArray<FBlendSampleData>& BlendSampleDataCache, /*out*/ FAnimationPoseData& OutAnimationPoseData) const;
+
+	ENGINE_API void GetAnimationPose(TArray<FBlendSampleData>& BlendSampleDataCache, TArrayView<FPoseLink> InPoseLinks, /*out*/ FPoseContext& Output) const;
 
 	/** Accessor for blend parameter **/
 	ENGINE_API const FBlendParameter& GetBlendParameter(const int32 Index) const;
@@ -311,6 +317,18 @@ class UBlendSpaceBase : public UAnimationAsset, public IInterpolationIndexProvid
 	/** Returns the blend input after clamping and/or wrapping */
 	ENGINE_API FVector GetClampedAndWrappedBlendInput(const FVector& BlendInput) const;
 
+	/** 
+	 * Updates a cached set of blend samples according to internal parameters, blendspace position and a delta time. Used internally to GetAnimationPose().
+	 * Note that this function does not perform any filtering internally.
+ 	 * @param	InBlendSpacePosition	The current position parameter of the blendspace
+	 * @param	InOutSampleDataCache	The sample data cache. Previous frames samples are re-used in the case of target weight interpolation
+	 * @param	InDeltaTime				The tick time for this update
+	 */
+	ENGINE_API bool UpdateBlendSamples(const FVector& InBlendSpacePosition, float InDeltaTime, TArray<FBlendSampleData>& InOutSampleDataCache) const;
+
+	/** Interpolate BlendInput based on Filter data **/
+	ENGINE_API FVector FilterInput(FBlendFilter* Filter, const FVector& BlendInput, float DeltaTime) const;
+
 #if WITH_EDITOR	
 	/** Validates sample data for blendspaces using the given animation sequence */
 	ENGINE_API static void UpdateBlendSpacesUsingAnimSequence(UAnimSequenceBase* Sequence);
@@ -319,6 +337,7 @@ class UBlendSpaceBase : public UAnimationAsset, public IInterpolationIndexProvid
 	ENGINE_API void ValidateSampleData();
 
 	/** Add samples */
+	ENGINE_API bool AddSample(const FVector& SampleValue);
 	ENGINE_API bool	AddSample(UAnimSequence* AnimationSequence, const FVector& SampleValue);
 
 	/** edit samples */
@@ -412,9 +431,6 @@ protected:
 	/** Utility function to interpolate weight of samples from OldSampleDataList to NewSampleDataList and copy back the interpolated result to FinalSampleDataList **/
 	bool InterpolateWeightOfSampleData(float DeltaTime, const TArray<FBlendSampleData> & OldSampleDataList, const TArray<FBlendSampleData> & NewSampleDataList, TArray<FBlendSampleData> & FinalSampleDataList) const;
 
-	/** Interpolate BlendInput based on Filter data **/
-	FVector FilterInput(FBlendFilter* Filter, const FVector& BlendInput, float DeltaTime) const;
-
 	/** Returns whether or not all animation set on the blend space samples match the given additive type */
 	bool ContainsMatchingSamples(EAdditiveAnimationType AdditiveType) const;
 
@@ -430,6 +446,13 @@ protected:
 	virtual void RemapSamplesToNewAxisRange() PURE_VIRTUAL(UBlendSpaceBase::RemapSamplesToNewAxisRange, return;);
 #endif // WITH_EDITOR
 	
+private:
+	// Internal helper function for GetAnimationPose variants
+	void GetAnimationPose_Internal(TArray<FBlendSampleData>& BlendSampleDataCache, TArrayView<FPoseLink> InPoseLinks, FAnimInstanceProxy* InProxy, bool bInExpectsAdditivePose, /*out*/ FAnimationPoseData& OutAnimationPoseData) const;
+
+	// Internal helper function for UpdateBlendSamples and TickAssetPlayer
+	bool UpdateBlendSamples_Internal(const FVector& InBlendSpacePosition, float InDeltaTime, TArray<FBlendSampleData>& InOutOldSampleDataList, TArray<FBlendSampleData>& InOutSampleDataCache) const;
+
 public:
 	/**
 	* When you use blend per bone, allows rotation to blend in mesh space. This only works if this does not contain additive animation samples
@@ -499,7 +522,7 @@ protected:
 	struct FBlendParameter BlendParameters[3];
 
 	/** Reset to reference pose. It does apply different refpose based on additive or not*/
-	void ResetToRefPose(FCompactPose& OutPose);
+	void ResetToRefPose(FCompactPose& OutPose) const;
 
 #if WITH_EDITOR
 private:
diff --git a/Engine/Source/Runtime/Engine/Classes/Animation/InputScaleBias.h b/Engine/Source/Runtime/Engine/Classes/Animation/InputScaleBias.h
--- a/Engine/Source/Runtime/Engine/Classes/Animation/InputScaleBias.h
+++ b/Engine/Source/Runtime/Engine/Classes/Animation/InputScaleBias.h
@@ -5,8 +5,12 @@
 #include "CoreMinimal.h"
 #include "UObject/ObjectMacros.h"
 #include "AlphaBlend.h"
+#include "Animation/AnimNodeBase.h"
+#include "Animation/AnimInstanceProxy.h"
 #include "InputScaleBias.generated.h"
 
+struct FAnimationUpdateContext;
+
 // Input scaling struct
 USTRUCT(BlueprintType)
 struct ENGINE_API FInputScaleBias
@@ -171,4 +175,37 @@ struct ENGINE_API FInputAlphaBoolBlend
 
 	float ApplyTo(bool bEnabled, float InDeltaTime);
 	void Reinitialize() { bInitialized = false; }
+};
+
+// Alpha blending options helper functions for anim nodes
+// Assumes that the specified node contains the members:
+// AlphaInputType, ActualAlpha, AlphaScaleBias, Alpha, bAlphaBoolEnabled, AlphaCurveName
+struct FAnimNodeAlphaOptions
+{
+	// Per-tick update
+	template<typename AnimNodeType>
+	static bool Update(AnimNodeType& InAnimNode, const FAnimationUpdateContext& InContext)
+	{
+		// Determine Actual Alpha.
+		switch (InAnimNode.AlphaInputType)
+		{
+		case EAnimAlphaInputType::Float:
+			InAnimNode.ActualAlpha = InAnimNode.AlphaScaleBias.ApplyTo(InAnimNode.AlphaScaleBiasClamp.ApplyTo(InAnimNode.Alpha, InContext.GetDeltaTime()));
+			break;
+		case EAnimAlphaInputType::Bool:
+			InAnimNode.ActualAlpha = InAnimNode.AlphaBoolBlend.ApplyTo(InAnimNode.bAlphaBoolEnabled, InContext.GetDeltaTime());
+			break;
+		case EAnimAlphaInputType::Curve:
+			if (UAnimInstance* AnimInstance = Cast<UAnimInstance>(InContext.AnimInstanceProxy->GetAnimInstanceObject()))
+			{
+				InAnimNode.ActualAlpha = InAnimNode.AlphaScaleBiasClamp.ApplyTo(AnimInstance->GetCurveValue(InAnimNode.AlphaCurveName), InContext.GetDeltaTime());
+			}
+			break;
+		};
+
+		// Make sure Alpha is clamped between 0 and 1.
+		InAnimNode.ActualAlpha = FMath::Clamp<float>(InAnimNode.ActualAlpha, 0.f, 1.f);
+
+		return FAnimWeight::IsRelevant(InAnimNode.ActualAlpha);
+	}
 };
\ No newline at end of file
diff --git a/Engine/Source/Runtime/Engine/Classes/EdGraph/EdGraphNode.h b/Engine/Source/Runtime/Engine/Classes/EdGraph/EdGraphNode.h
--- a/Engine/Source/Runtime/Engine/Classes/EdGraph/EdGraphNode.h
+++ b/Engine/Source/Runtime/Engine/Classes/EdGraph/EdGraphNode.h
@@ -686,6 +686,9 @@ public:
 	/** Returns the graph that contains this node */
 	class UEdGraph* GetGraph() const;
 
+	/** @returns any sub graphs (graphs that have this node as an outer) that this node might contain (e.g. composite, animation state machine etc.).*/
+	virtual TArray<UEdGraph*> GetSubGraphs() const { return TArray<UEdGraph*>(); }
+
 	/**
 	 * Allocate default pins for a given node, based only the NodeType, which should already be filled in.
 	 *
diff --git a/Engine/Source/Runtime/Engine/Private/Animation/AnimBlueprintGeneratedClass.cpp b/Engine/Source/Runtime/Engine/Private/Animation/AnimBlueprintGeneratedClass.cpp
--- a/Engine/Source/Runtime/Engine/Private/Animation/AnimBlueprintGeneratedClass.cpp
+++ b/Engine/Source/Runtime/Engine/Private/Animation/AnimBlueprintGeneratedClass.cpp
@@ -224,7 +224,7 @@ void FAnimBlueprintDebugData::RecordSequencePlayer(int32 InNodeID, float InPosit
 	SequencePlayerRecordsThisFrame.Emplace(InNodeID, InPosition, InLength, InFrameCount);
 }
 
-void FAnimBlueprintDebugData::RecordBlendSpacePlayer(int32 InNodeID, UBlendSpaceBase* InBlendSpace, float InPositionX, float InPositionY, float InPositionZ)
+void FAnimBlueprintDebugData::RecordBlendSpacePlayer(int32 InNodeID, const UBlendSpaceBase* InBlendSpace, float InPositionX, float InPositionY, float InPositionZ)
 {
 	BlendSpacePlayerRecordsThisFrame.Emplace(InNodeID, InBlendSpace, InPositionX, InPositionY, InPositionZ);
 }
diff --git a/Engine/Source/Runtime/Engine/Private/Animation/AnimNode_AssetPlayerBase.cpp b/Engine/Source/Runtime/Engine/Private/Animation/AnimNode_AssetPlayerBase.cpp
--- a/Engine/Source/Runtime/Engine/Private/Animation/AnimNode_AssetPlayerBase.cpp
+++ b/Engine/Source/Runtime/Engine/Private/Animation/AnimNode_AssetPlayerBase.cpp
@@ -45,7 +45,7 @@ void FAnimNode_AssetPlayerBase::CreateTickRecordForNode(const FAnimationUpdateCo
 
 	FName GroupNameToUse = ((GroupRole < EAnimGroupRole::TransitionLeader) || bHasBeenFullWeight) ? GroupName : NAME_None;
 	EAnimSyncMethod MethodToUse = Method;
-	if(GroupNameToUse == NAME_None)
+	if(GroupNameToUse == NAME_None && MethodToUse == EAnimSyncMethod::SyncGroup)
 	{
 		MethodToUse = EAnimSyncMethod::DoNotSync;
 	}
