Misc BlendSpace 2.0 fixes

Fixed drag/drop of samples in blend space graphs.
Fixed issues around replacing samples and stale anim graph tabs when replacing/deleting samples.
Prevented sample point graph duplication from My Blueprint.
Preserved sync groups when converting from players to graphs.

diff --git a/Engine/Source/Editor/AnimGraph/Private/AnimGraphNode_BlendSpaceGraphBase.cpp b/Engine/Source/Editor/AnimGraph/Private/AnimGraphNode_BlendSpaceGraphBase.cpp
--- a/Engine/Source/Editor/AnimGraph/Private/AnimGraphNode_BlendSpaceGraphBase.cpp
+++ b/Engine/Source/Editor/AnimGraph/Private/AnimGraphNode_BlendSpaceGraphBase.cpp
@@ -24,6 +24,7 @@
 #include "AnimGraphNode_SequencePlayer.h"
 #include "Kismet2/Kismet2NameValidators.h"
 #include "AnimNodes/AnimNode_BlendSpaceGraphBase.h"
+#include "BlueprintEditor.h"
 
 #define LOCTEXT_NAMESPACE "UAnimGraphNode_BlendSpaceGraphBase"
 
@@ -161,20 +162,7 @@ void UAnimGraphNode_BlendSpaceGraphBase::SetupFromAsset(UBlendSpaceBase* InBlend
 
 			FName SampleName = Sample.Animation != nullptr ? Sample.Animation->GetFName() : FName("Sample", SampleIndex);
 
-			UAnimationBlendSpaceSampleGraph* SampleGraph = AddGraph(SampleName);
-
-			// Attach an asset player by default to default the sample
-			FGraphNodeCreator<UAnimGraphNode_SequencePlayer> SequencePlayerNodeCreator(*SampleGraph);
-			UAnimGraphNode_SequencePlayer* SequencePlayer = SequencePlayerNodeCreator.CreateNode();
-			SequencePlayer->SetAnimationAsset(Sample.Animation);
-			SequencePlayerNodeCreator.Finalize();
-
-			// Offset node in X
-			SequencePlayer->NodePosX = SampleGraph->ResultNode->NodePosX - 400;
-
-			UEdGraphPin* OutputPin = SequencePlayer->FindPinChecked(TEXT("Pose"), EGPD_Output);
-			UEdGraphPin* InputPin = SampleGraph->ResultNode->FindPinChecked(TEXT("Result"), EGPD_Input);
-			OutputPin->MakeLinkTo(InputPin);
+			UAnimationBlendSpaceSampleGraph* SampleGraph = AddGraph(SampleName, Sample.Animation);
 
 			// Clear the animation now we have created the point
 			Sample.Animation = nullptr;
@@ -301,31 +289,34 @@ void UAnimGraphNode_BlendSpaceGraphBase::PostPasteNode()
 {
 	Super::PostPasteNode();
 
-	// Add the new graph as a child of our parent graph
-	UEdGraph* ParentGraph = GetGraph();
-
-	if(ParentGraph->SubGraphs.Find(BlendSpaceGraph) == INDEX_NONE)
+	if(BlendSpaceGraph != nullptr)
 	{
-		ParentGraph->SubGraphs.Add(BlendSpaceGraph);
-	}
+		// Add the new graph as a child of our parent graph
+		UEdGraph* ParentGraph = GetGraph();
 
-	// Refresh sample graphs
-	for (UEdGraph* SampleGraph : Graphs)
-	{
-		for(UEdGraphNode* GraphNode : SampleGraph->Nodes)
+		if(ParentGraph->SubGraphs.Find(BlendSpaceGraph) == INDEX_NONE)
 		{
-			GraphNode->CreateNewGuid();
-			GraphNode->PostPasteNode();
-			GraphNode->ReconstructNode();
+			ParentGraph->SubGraphs.Add(BlendSpaceGraph);
 		}
-	}
 
-	// Find an interesting name
-	TSharedPtr<INameValidatorInterface> NameValidator = FNameValidatorFactory::MakeValidator(this);
-	FBlueprintEditorUtils::RenameGraphWithSuggestion(BlendSpaceGraph, NameValidator, BlendSpaceGraph->GetName());
+		// Refresh sample graphs
+		for (UEdGraph* SampleGraph : Graphs)
+		{
+			for(UEdGraphNode* GraphNode : SampleGraph->Nodes)
+			{
+				GraphNode->CreateNewGuid();
+				GraphNode->PostPasteNode();
+				GraphNode->ReconstructNode();
+			}
+		}
+
+		// Find an interesting name
+		TSharedPtr<INameValidatorInterface> NameValidator = FNameValidatorFactory::MakeValidator(this);
+		FBlueprintEditorUtils::RenameGraphWithSuggestion(BlendSpaceGraph, NameValidator, BlendSpaceGraph->GetName());
 
-	// Restore transactional flag that is lost during copy/paste process
-	BlendSpaceGraph->SetFlags(RF_Transactional);
+		// Restore transactional flag that is lost during copy/paste process
+		BlendSpaceGraph->SetFlags(RF_Transactional);
+	}
 }
 
 void UAnimGraphNode_BlendSpaceGraphBase::PostPlacedNewNode()
@@ -397,7 +388,16 @@ void UAnimGraphNode_BlendSpaceGraphBase::PostProcessPinName(const UEdGraphPin* P
 	Super::PostProcessPinName(Pin, DisplayName);
 }
 
-UAnimationBlendSpaceSampleGraph* UAnimGraphNode_BlendSpaceGraphBase::AddGraph(FName InSampleName)
+UAnimationBlendSpaceSampleGraph* UAnimGraphNode_BlendSpaceGraphBase::AddGraph(FName InSampleName, UAnimSequence* InSequence)
+{
+	UAnimationBlendSpaceSampleGraph* NewGraph = AddGraphInternal(InSampleName, InSequence);
+
+	Graphs.Add(NewGraph);
+
+	return NewGraph;
+}
+
+UAnimationBlendSpaceSampleGraph* UAnimGraphNode_BlendSpaceGraphBase::AddGraphInternal(FName InSampleName, UAnimSequence* InSequence)
 {
 	const UAnimationGraphSchema* AnimationGraphSchema = GetDefault<UAnimationGraphSchema>();
 
@@ -413,7 +413,23 @@ UAnimationBlendSpaceSampleGraph* UAnimGraphNode_BlendSpaceGraphBase::AddGraph(FN
 	SampleGraph->ResultNode = ResultSinkNode;
 	SampleGraph->bAllowDeletion = false;
 	SampleGraph->bAllowRenaming = true;
-	Graphs.Add(SampleGraph);
+
+	if(InSequence != nullptr)
+	{
+		// Attach an asset player if a valid animation is supplied
+		FGraphNodeCreator<UAnimGraphNode_SequencePlayer> SequencePlayerNodeCreator(*SampleGraph);
+		UAnimGraphNode_SequencePlayer* SequencePlayer = SequencePlayerNodeCreator.CreateNode();
+		SequencePlayer->SetAnimationAsset(InSequence);
+		SequencePlayer->SyncGroup.Method = EAnimSyncMethod::Graph;
+		SequencePlayerNodeCreator.Finalize();
+
+		// Offset node in X
+		SequencePlayer->NodePosX = SampleGraph->ResultNode->NodePosX - 400;
+
+		UEdGraphPin* OutputPin = SequencePlayer->FindPinChecked(TEXT("Pose"), EGPD_Output);
+		UEdGraphPin* InputPin = SampleGraph->ResultNode->FindPinChecked(TEXT("Result"), EGPD_Input);
+		OutputPin->MakeLinkTo(InputPin);
+	}
 
 	BlendSpaceGraph->Modify();
 	BlendSpaceGraph->SubGraphs.Add(SampleGraph);
@@ -427,12 +443,30 @@ void UAnimGraphNode_BlendSpaceGraphBase::RemoveGraph(int32 InSampleIndex)
 
 	check(Graphs.IsValidIndex(InSampleIndex));
 	UEdGraph* GraphToRemove = Graphs[InSampleIndex];
+	TSharedPtr<FBlueprintEditor> BlueprintEditor = StaticCastSharedPtr<FBlueprintEditor>(FKismetEditorUtilities::GetIBlueprintEditorForObject(GraphToRemove, false));
+	check(BlueprintEditor.IsValid());
 
 	// 'Asset' blendspace sample deletion uses a RemoveAtSwap, so we must mirror it here to maintain the same indices
 	Graphs.RemoveAtSwap(InSampleIndex);
 
-	BlendSpaceGraph->Modify();
-	BlendSpaceGraph->SubGraphs.Remove(GraphToRemove);
+	FBlueprintEditorUtils::RemoveGraph(GetAnimBlueprint(), GraphToRemove);
+	BlueprintEditor->CloseDocumentTab(GraphToRemove);
+}
+
+void UAnimGraphNode_BlendSpaceGraphBase::ReplaceGraph(int32 InSampleIndex, UAnimSequence* InSequence)
+{
+	Modify();
+
+	check(Graphs.IsValidIndex(InSampleIndex));
+	UEdGraph* GraphToRemove = Graphs[InSampleIndex];
+	FName GraphName = GraphToRemove->GetFName();
+	TSharedPtr<FBlueprintEditor> BlueprintEditor = StaticCastSharedPtr<FBlueprintEditor>(FKismetEditorUtilities::GetIBlueprintEditorForObject(GraphToRemove, false));
+	check(BlueprintEditor.IsValid());
+
+	FBlueprintEditorUtils::RemoveGraph(GetAnimBlueprint(), GraphToRemove);
+	BlueprintEditor->CloseDocumentTab(GraphToRemove);
+
+	Graphs[InSampleIndex] = AddGraphInternal(GraphName, InSequence);
 }
 
 FName UAnimGraphNode_BlendSpaceGraphBase::GetSyncGroupName() const
@@ -443,6 +477,13 @@ FName UAnimGraphNode_BlendSpaceGraphBase::GetSyncGroupName() const
 	return AnimNode->GroupName;
 }
 
+void UAnimGraphNode_BlendSpaceGraphBase::SetSyncGroupName(FName InName)
+{
+	FStructProperty* NodeProperty = GetFNodeProperty();
+	FAnimNode_BlendSpaceGraphBase* AnimNode = NodeProperty->ContainerPtrToValuePtr<FAnimNode_BlendSpaceGraphBase>(this);
+	AnimNode->GroupName = InName;
+}
+
 void UAnimGraphNode_BlendSpaceGraphBase::GetInputLinkAttributes(FNodeAttributeArray& OutAttributes) const
 {
 	if(GetSyncGroupName() != NAME_None)
diff --git a/Engine/Source/Editor/AnimGraph/Private/AnimationGraphSchema.cpp b/Engine/Source/Editor/AnimGraph/Private/AnimationGraphSchema.cpp
--- a/Engine/Source/Editor/AnimGraph/Private/AnimationGraphSchema.cpp
+++ b/Engine/Source/Editor/AnimGraph/Private/AnimationGraphSchema.cpp
@@ -37,6 +37,7 @@
 #include "AnimGraphNode_LinkedInputPose.h"
 #include "AnimGraphNode_LinkedAnimLayer.h"
 #include "AnimGraphNode_RigidBody.h"
+#include "AnimationBlendSpaceSampleGraph.h"
 
 #define LOCTEXT_NAMESPACE "AnimationGraphSchema"
 
@@ -618,6 +619,11 @@ FText UAnimationGraphSchema::GetPinDisplayName(const UEdGraphPin* Pin) const
 	return DisplayName;
 }
 
+bool UAnimationGraphSchema::CanDuplicateGraph(UEdGraph* InSourceGraph) const
+{
+	return InSourceGraph->GetFName() != UEdGraphSchema_K2::GN_AnimGraph && !InSourceGraph->IsA<UAnimationBlendSpaceSampleGraph>();
+}
+
 void UAnimationGraphSchema::GetGraphDisplayInformation(const UEdGraph& Graph, /*out*/ FGraphDisplayInfo& DisplayInfo) const
 {
 	if (GetGraphType(&Graph) == GT_Animation)
@@ -631,29 +637,31 @@ void UAnimationGraphSchema::GetGraphDisplayInformation(const UEdGraph& Graph, /*
 		if(!Graph.bAllowDeletion)
 		{
 			// Might be from an interface, so check
-			UBlueprint* Blueprint = FBlueprintEditorUtils::FindBlueprintForGraph(&Graph);
-			TSubclassOf<UInterface> Interface;
-
-			auto FindInterfaceForGraph = [&Blueprint, &Graph](TSubclassOf<UInterface>& OutInterface)
+			if(UBlueprint* Blueprint = FBlueprintEditorUtils::FindBlueprintForGraph(&Graph))
 			{
-				for(const FBPInterfaceDescription& InterfaceDesc : Blueprint->ImplementedInterfaces)
+				TSubclassOf<UInterface> Interface;
+
+				auto FindInterfaceForGraph = [&Blueprint, &Graph](TSubclassOf<UInterface>& OutInterface)
 				{
-					for(UEdGraph* InterfaceGraph : InterfaceDesc.Graphs)
+					for(const FBPInterfaceDescription& InterfaceDesc : Blueprint->ImplementedInterfaces)
 					{
-						if(InterfaceGraph == &Graph)
+						for(UEdGraph* InterfaceGraph : InterfaceDesc.Graphs)
 						{
-							OutInterface = InterfaceDesc.Interface;
-							return true;
+							if(InterfaceGraph == &Graph)
+							{
+								OutInterface = InterfaceDesc.Interface;
+								return true;
+							}
 						}
 					}
-				}
 
-				return false;
-			};
+					return false;
+				};
 
-			if(FindInterfaceForGraph(Interface))
-			{
-				DisplayInfo.Tooltip = FText::Format(LOCTEXT("GraphTooltip_AnimGraphInterface", "Layer inherited from interface '{0}'."), FText::FromString(Interface.Get()->GetName()));
+				if(FindInterfaceForGraph(Interface))
+				{
+					DisplayInfo.Tooltip = FText::Format(LOCTEXT("GraphTooltip_AnimGraphInterface", "Layer inherited from interface '{0}'."), FText::FromString(Interface.Get()->GetName()));
+				}
 			}
 		}
 			
diff --git a/Engine/Source/Editor/AnimGraph/Public/AnimGraphNode_BlendSpaceGraphBase.h b/Engine/Source/Editor/AnimGraph/Public/AnimGraphNode_BlendSpaceGraphBase.h
--- a/Engine/Source/Editor/AnimGraph/Public/AnimGraphNode_BlendSpaceGraphBase.h
+++ b/Engine/Source/Editor/AnimGraph/Public/AnimGraphNode_BlendSpaceGraphBase.h
@@ -30,11 +30,14 @@ public:
 	UBlendSpaceGraph* GetBlendSpaceGraph() const { return BlendSpaceGraph; }
 
 	// Adds a new graph to the internal array
-	UAnimationBlendSpaceSampleGraph* AddGraph(FName InSampleName);
+	UAnimationBlendSpaceSampleGraph* AddGraph(FName InSampleName, UAnimSequence* InSequence);
 
 	// Removes the graph at the specified index
 	void RemoveGraph(int32 InSampleIndex);
 
+	// Replaces the graph at the specified index
+	void ReplaceGraph(int32 InSampleIndex, UAnimSequence* InSequence);
+
 	// Setup this node from the specified asset
 	void SetupFromAsset(UBlendSpaceBase* InBlendSpace, bool bInIsTemplateNode);
 
@@ -44,6 +47,9 @@ public:
 	// @return the sync group name assigned to this node
 	FName GetSyncGroupName() const;
 
+	// Set the sync group name assigned to this node
+	void SetSyncGroupName(FName InName);
+
 protected:
 	// Get the name of the blendspace graph
 	FString GetBlendSpaceGraphName() const;
@@ -96,4 +102,7 @@ protected:
 
 	// Helper function for compilation
 	UAnimGraphNode_Base* ExpandGraphAndProcessNodes(UEdGraph* SourceGraph, UAnimGraphNode_Base* SourceRootNode, IAnimBlueprintCompilationContext& InCompilationContext, IAnimBlueprintGeneratedClassCompiledData& OutCompiledData);
+
+	// Helper function for AddGraph/ReplaceGraph - builds the new graph but doesnt add it to Graphs array.
+	UAnimationBlendSpaceSampleGraph* AddGraphInternal(FName InSampleName, UAnimSequence* InSequence);
 };
diff --git a/Engine/Source/Editor/AnimGraph/Public/AnimationGraphSchema.h b/Engine/Source/Editor/AnimGraph/Public/AnimationGraphSchema.h
--- a/Engine/Source/Editor/AnimGraph/Public/AnimationGraphSchema.h
+++ b/Engine/Source/Editor/AnimGraph/Public/AnimationGraphSchema.h
@@ -60,7 +60,7 @@ class UAnimationGraphSchema : public UEdGraphSchema_K2
 	virtual void GetAssetsGraphHoverMessage(const TArray<FAssetData>& Assets, const UEdGraph* HoverGraph, FString& OutTooltipText, bool& OutOkIcon) const override;
 	virtual void GetContextMenuActions(class UToolMenu* Menu, class UGraphNodeContextMenuContext* Context) const override;
 	virtual FText GetPinDisplayName(const UEdGraphPin* Pin) const override;
-	virtual bool CanDuplicateGraph(UEdGraph* InSourceGraph) const override { return InSourceGraph->GetFName() != UEdGraphSchema_K2::GN_AnimGraph; }
+	virtual bool CanDuplicateGraph(UEdGraph* InSourceGraph) const override;
 	virtual bool DoesSupportEventDispatcher() const	override { return false; }
 	virtual bool ShouldAlwaysPurgeOnModification() const override { return true; }
 	virtual bool TryCreateConnection(UEdGraphPin* A, UEdGraphPin* B) const override;
diff --git a/Engine/Source/Editor/AnimationBlueprintEditor/Private/AnimationBlueprintEditor.cpp b/Engine/Source/Editor/AnimationBlueprintEditor/Private/AnimationBlueprintEditor.cpp
--- a/Engine/Source/Editor/AnimationBlueprintEditor/Private/AnimationBlueprintEditor.cpp
+++ b/Engine/Source/Editor/AnimationBlueprintEditor/Private/AnimationBlueprintEditor.cpp
@@ -970,6 +970,10 @@ void FAnimationBlueprintEditor::OnConvertToBlendSpaceGraph()
 				FGraphNodeCreator<UAnimGraphNode_BlendSpaceGraph> NodeCreator(*TargetGraph);
 				UAnimGraphNode_BlendSpaceGraph* NewNode = NodeCreator.CreateNode();
 				NewNode->SetupFromAsset(OldNode->Node.BlendSpace, false);
+				if(OldNode->SyncGroup.GroupName != NAME_None && OldNode->SyncGroup.Method == EAnimSyncMethod::SyncGroup)
+				{
+					NewNode->SetSyncGroupName(OldNode->SyncGroup.GroupName);
+				}
 				NodeCreator.Finalize();
 
 				// get default data from old node to new node
@@ -1228,6 +1232,10 @@ void FAnimationBlueprintEditor::OnConvertToAimOffsetGraph()
 				FGraphNodeCreator<UAnimGraphNode_RotationOffsetBlendSpaceGraph> NodeCreator(*TargetGraph);
 				UAnimGraphNode_RotationOffsetBlendSpaceGraph* NewNode = NodeCreator.CreateNode();
 				NewNode->SetupFromAsset(OldNode->Node.BlendSpace, false);
+				if(OldNode->SyncGroup.GroupName != NAME_None && OldNode->SyncGroup.Method == EAnimSyncMethod::SyncGroup)
+				{
+					NewNode->SetSyncGroupName(OldNode->SyncGroup.GroupName);
+				}
 				NodeCreator.Finalize();
 
 				// get default data from old node to new node
diff --git a/Engine/Source/Editor/AnimationBlueprintEditor/Private/AnimationPins/SGraphPinPose.cpp b/Engine/Source/Editor/AnimationBlueprintEditor/Private/AnimationPins/SGraphPinPose.cpp
--- a/Engine/Source/Editor/AnimationBlueprintEditor/Private/AnimationPins/SGraphPinPose.cpp
+++ b/Engine/Source/Editor/AnimationBlueprintEditor/Private/AnimationPins/SGraphPinPose.cpp
@@ -161,120 +161,122 @@ void SGraphPinPose::ReconfigureWidgetForAttributes()
 	AttributeInfos.Empty();
 
 	UAnimGraphNode_Base* AnimGraphNode = CastChecked<UAnimGraphNode_Base>(GraphPinObj->GetOwningNode());
-	UAnimBlueprint* AnimBlueprint = Cast<UAnimBlueprint>(FBlueprintEditorUtils::FindBlueprintForNode(AnimGraphNode));
-
-	auto AddAttributes = [this, AnimBlueprint, AnimGraphNode](TArrayView<const FName> InUsedAttributes, TArrayView<const FName> InUnusedAttributes, TArrayView<const FName> InPassThroughAttributes)
+	if(UAnimBlueprint* AnimBlueprint = Cast<UAnimBlueprint>(FBlueprintEditorUtils::FindBlueprintForNode(AnimGraphNode)))
 	{
-		auto AddAttribute = [this, AnimBlueprint, AnimGraphNode](FName InAttribute, EAttributeUsage InUsage)
-		{
-			const FAnimGraphAttributeDesc* AttributeDesc = GetDefault<UAnimGraphAttributes>()->FindAttributeDesc(InAttribute);
 
-			// Early out if we don't want to display this attribute
-			if(AttributeDesc)
+		auto AddAttributes = [this, AnimBlueprint, AnimGraphNode](TArrayView<const FName> InUsedAttributes, TArrayView<const FName> InUnusedAttributes, TArrayView<const FName> InPassThroughAttributes)
+		{
+			auto AddAttribute = [this, AnimBlueprint, AnimGraphNode](FName InAttribute, EAttributeUsage InUsage)
 			{
-				if(AttributeDesc->DisplayMode == EAnimGraphAttributesDisplayMode::HideOnPins)
-				{
-					return;
-				}
+				const FAnimGraphAttributeDesc* AttributeDesc = GetDefault<UAnimGraphAttributes>()->FindAttributeDesc(InAttribute);
 
-				// Only store cached attributes for wires on used pins
-				if(InUsage != EAttributeUsage::Unused)
+				// Early out if we don't want to display this attribute
+				if(AttributeDesc)
 				{
-					AttributeInfos.Emplace(InAttribute, AttributeDesc->Color.GetSpecifiedColor(), AttributeDesc->Blend, AttributeDesc->SortOrder);
+					if(AttributeDesc->DisplayMode == EAnimGraphAttributesDisplayMode::HideOnPins)
+					{
+						return;
+					}
 
-					// Skip displaying passthrough attributes on pins
-					if(InUsage != EAttributeUsage::Passthrough && AnimGraphNode->ShouldShowAttributesOnPins())
+					// Only store cached attributes for wires on used pins
+					if(InUsage != EAttributeUsage::Unused)
 					{
-						const UAnimGraphNode_Base* ProxyGraphNode = AnimGraphNode->GetProxyNodeForAttributes();
+						AttributeInfos.Emplace(InAttribute, AttributeDesc->Color.GetSpecifiedColor(), AttributeDesc->Blend, AttributeDesc->SortOrder);
+
+						// Skip displaying passthrough attributes on pins
+						if(InUsage != EAttributeUsage::Passthrough && AnimGraphNode->ShouldShowAttributesOnPins())
+						{
+							const UAnimGraphNode_Base* ProxyGraphNode = AnimGraphNode->GetProxyNodeForAttributes();
 
-						GetLabelAndValue()->AddSlot()
-						.Padding(2.0f, 0.0f, 0.0f, 0.0f)
-						[
-							SNew(SAttributeIndicator, AnimBlueprint, AttributeDesc, GraphPinObj, InUsage, ProxyGraphNode)
-						];
+							GetLabelAndValue()->AddSlot()
+							.Padding(2.0f, 0.0f, 0.0f, 0.0f)
+							[
+								SNew(SAttributeIndicator, AnimBlueprint, AttributeDesc, GraphPinObj, InUsage, ProxyGraphNode)
+							];
+						}
 					}
 				}
+			};
+
+			for(const FName& Attribute : InUsedAttributes)
+			{
+				AddAttribute(Attribute, EAttributeUsage::Used);
+			}
+
+			for(const FName& Attribute : InUnusedAttributes)
+			{
+				AddAttribute(Attribute, EAttributeUsage::Unused);
+			}
+
+			for(const FName& Attribute : InPassThroughAttributes)
+			{
+				AddAttribute(Attribute, EAttributeUsage::Passthrough);
 			}
 		};
 
-		for(const FName& Attribute : InUsedAttributes)
-		{
-			AddAttribute(Attribute, EAttributeUsage::Used);
-		}
+		UAnimBlueprintGeneratedClass* AnimBlueprintClass = (UAnimBlueprintGeneratedClass*)(*(AnimBlueprint->GeneratedClass));
 
-		for(const FName& Attribute : InUnusedAttributes)
+		UAnimGraphNode_Base::FNodeAttributeArray PinAttributes;
+		switch(GraphPinObj->Direction)
 		{
-			AddAttribute(Attribute, EAttributeUsage::Unused);
+		case EGPD_Input:
+			AnimGraphNode->GetInputLinkAttributes(PinAttributes);
+			break;
+		case EGPD_Output:
+			AnimGraphNode->GetOutputLinkAttributes(PinAttributes);
+			break;
 		}
 
-		for(const FName& Attribute : InPassThroughAttributes)
+		UAnimGraphNode_Base::FNodeAttributeArray NodeAttributes;
+		AnimGraphNode->GetInputLinkAttributes(NodeAttributes);
+		AnimGraphNode->GetOutputLinkAttributes(NodeAttributes);
+
+		// Unlinked pins display attributes if they are inputs and the node takes them as inputs
+		if(GraphPinObj->LinkedTo.Num() == 0)
 		{
-			AddAttribute(Attribute, EAttributeUsage::Passthrough);
+			AddAttributes(TArrayView<const FName>(), PinAttributes, TArrayView<const FName>());
 		}
-	};
-
-	UAnimBlueprintGeneratedClass* AnimBlueprintClass = (UAnimBlueprintGeneratedClass*)(*(AnimBlueprint->GeneratedClass));
-
-	UAnimGraphNode_Base::FNodeAttributeArray PinAttributes;
-	switch(GraphPinObj->Direction)
-	{
-	case EGPD_Input:
-		AnimGraphNode->GetInputLinkAttributes(PinAttributes);
-		break;
-	case EGPD_Output:
-		AnimGraphNode->GetOutputLinkAttributes(PinAttributes);
-		break;
-	}
-
-	UAnimGraphNode_Base::FNodeAttributeArray NodeAttributes;
-	AnimGraphNode->GetInputLinkAttributes(NodeAttributes);
-	AnimGraphNode->GetOutputLinkAttributes(NodeAttributes);
-
-	// Unlinked pins display attributes if they are inputs and the node takes them as inputs
-	if(GraphPinObj->LinkedTo.Num() == 0)
-	{
-		AddAttributes(TArrayView<const FName>(), PinAttributes, TArrayView<const FName>());
-	}
-	else if(GraphPinObj->LinkedTo.Num() > 0)
-	{
-		TArrayView<const FName> CompilerGeneratedAttributes = AnimBlueprintClass->GetAnimBlueprintDebugData().GetNodeAttributes(AnimGraphNode);
-		UAnimGraphNode_Base::FNodeAttributeArray UsedAttributes;
-		UAnimGraphNode_Base::FNodeAttributeArray UnusedAttributes;
-		if(PinAttributes.Num() > 0)
+		else if(GraphPinObj->LinkedTo.Num() > 0)
 		{
-			for(const FName& Name : PinAttributes)
+			TArrayView<const FName> CompilerGeneratedAttributes = AnimBlueprintClass->GetAnimBlueprintDebugData().GetNodeAttributes(AnimGraphNode);
+			UAnimGraphNode_Base::FNodeAttributeArray UsedAttributes;
+			UAnimGraphNode_Base::FNodeAttributeArray UnusedAttributes;
+			if(PinAttributes.Num() > 0)
 			{
-				if(CompilerGeneratedAttributes.Contains(Name))
+				for(const FName& Name : PinAttributes)
 				{
-					UsedAttributes.Add(Name);
+					if(CompilerGeneratedAttributes.Contains(Name))
+					{
+						UsedAttributes.Add(Name);
+					}
+					else
+					{
+						UnusedAttributes.Add(Name);
+					}
 				}
-				else
+			}
+
+			UAnimGraphNode_Base::FNodeAttributeArray PassthroughAttributes;
+			for(const FName& Name : CompilerGeneratedAttributes)
+			{
+				if(!PinAttributes.Contains(Name))
 				{
-					UnusedAttributes.Add(Name);
+					PassthroughAttributes.Add(Name);
 				}
 			}
+
+			AddAttributes(UsedAttributes, UnusedAttributes, PassthroughAttributes);
 		}
 
-		UAnimGraphNode_Base::FNodeAttributeArray PassthroughAttributes;
-		for(const FName& Name : CompilerGeneratedAttributes)
+		// sort pin attributes
+		Algo::Sort(AttributeInfos, [](const FAttributeInfo& InValue0, const FAttributeInfo& InValue1)
 		{
-			if(!PinAttributes.Contains(Name))
-			{
-				PassthroughAttributes.Add(Name);
-			}
-		}
+			return InValue0.SortOrder < InValue1.SortOrder;
+		});
 
-		AddAttributes(UsedAttributes, UnusedAttributes, PassthroughAttributes);
+		// overrides attribute set in the base class Construct()
+		SetToolTipText(MakeAttributeSP(this, &SGraphPinPose::GetAttributeTooltipText));
 	}
-
-	// sort pin attributes
-	Algo::Sort(AttributeInfos, [](const FAttributeInfo& InValue0, const FAttributeInfo& InValue1)
-	{
-		return InValue0.SortOrder < InValue1.SortOrder;
-	});
-
-	// overrides attribute set in the base class Construct()
-	SetToolTipText(MakeAttributeSP(this, &SGraphPinPose::GetAttributeTooltipText));
 }
 
 TArrayView<const SGraphPinPose::FAttributeInfo> SGraphPinPose::GetAttributeInfo() const
diff --git a/Engine/Source/Editor/AnimationBlueprintEditor/Private/BlendSpaceDocumentTabFactory.cpp b/Engine/Source/Editor/AnimationBlueprintEditor/Private/BlendSpaceDocumentTabFactory.cpp
--- a/Engine/Source/Editor/AnimationBlueprintEditor/Private/BlendSpaceDocumentTabFactory.cpp
+++ b/Engine/Source/Editor/AnimationBlueprintEditor/Private/BlendSpaceDocumentTabFactory.cpp
@@ -53,7 +53,7 @@ TSharedRef<SWidget> FBlendSpaceDocumentTabFactory::CreateTabBody(const FWorkflow
 		if(WeakBlendSpaceNode.Get())
 		{
 			UAnimGraphNode_BlendSpaceGraphBase* BlendSpaceNode = WeakBlendSpaceNode.Get();
-			UAnimationBlendSpaceSampleGraph* NewGraph = BlendSpaceNode->AddGraph(TEXT("NewSample"));
+			UAnimationBlendSpaceSampleGraph* NewGraph = BlendSpaceNode->AddGraph(TEXT("NewSample"), InSequence);
 			BlueprintEditorPtr.Pin()->RefreshMyBlueprint();
 			BlueprintEditorPtr.Pin()->RenameNewlyAddedAction(NewGraph->GetFName());
 		}
@@ -69,6 +69,16 @@ TSharedRef<SWidget> FBlendSpaceDocumentTabFactory::CreateTabBody(const FWorkflow
 		}
 	});
 
+	Args.OnBlendSpaceSampleReplaced = FOnBlendSpaceSampleReplaced::CreateLambda([this, WeakBlendSpaceNode = TWeakObjectPtr<UAnimGraphNode_BlendSpaceGraphBase>(BlendSpaceNode)](int32 InSampleIndex, UAnimSequence* InSequence)
+	{
+		if(WeakBlendSpaceNode.Get())
+		{
+			UAnimGraphNode_BlendSpaceGraphBase* BlendSpaceNode = WeakBlendSpaceNode.Get();
+			BlendSpaceNode->ReplaceGraph(InSampleIndex, InSequence);
+			BlueprintEditorPtr.Pin()->RefreshMyBlueprint();
+		}
+	});
+
 	Args.OnGetBlendSpaceSampleName = FOnGetBlendSpaceSampleName::CreateLambda([this, WeakBlendSpaceNode = TWeakObjectPtr<UAnimGraphNode_BlendSpaceGraphBase>(BlendSpaceNode)](int32 InSampleIndex) -> FName
 	{
 		if(WeakBlendSpaceNode.Get())
diff --git a/Engine/Source/Editor/Persona/Private/PersonaModule.cpp b/Engine/Source/Editor/Persona/Private/PersonaModule.cpp
--- a/Engine/Source/Editor/Persona/Private/PersonaModule.cpp
+++ b/Engine/Source/Editor/Persona/Private/PersonaModule.cpp
@@ -1528,6 +1528,7 @@ TSharedRef<SWidget> FPersonaModule::CreateBlendSpaceEditWidget(UBlendSpaceBase*
 			.OnBlendSpaceSampleDoubleClicked(InArgs.OnBlendSpaceSampleDoubleClicked)
 			.OnBlendSpaceSampleAdded(InArgs.OnBlendSpaceSampleAdded)
 			.OnBlendSpaceSampleRemoved(InArgs.OnBlendSpaceSampleRemoved)
+			.OnBlendSpaceSampleReplaced(InArgs.OnBlendSpaceSampleReplaced)
 			.OnGetBlendSpaceSampleName(InArgs.OnGetBlendSpaceSampleName)
 			.OnExtendSampleTooltip(InArgs.OnExtendSampleTooltip)
 			.OnSetPreviewPosition(InArgs.OnSetPreviewPosition)
@@ -1542,6 +1543,7 @@ TSharedRef<SWidget> FPersonaModule::CreateBlendSpaceEditWidget(UBlendSpaceBase*
 			.OnBlendSpaceSampleDoubleClicked(InArgs.OnBlendSpaceSampleDoubleClicked)
 			.OnBlendSpaceSampleAdded(InArgs.OnBlendSpaceSampleAdded)
 			.OnBlendSpaceSampleRemoved(InArgs.OnBlendSpaceSampleRemoved)
+			.OnBlendSpaceSampleReplaced(InArgs.OnBlendSpaceSampleReplaced)
 			.OnGetBlendSpaceSampleName(InArgs.OnGetBlendSpaceSampleName)
 			.OnExtendSampleTooltip(InArgs.OnExtendSampleTooltip)
 			.OnSetPreviewPosition(InArgs.OnSetPreviewPosition)
diff --git a/Engine/Source/Editor/Persona/Private/SAnimationBlendSpace.cpp b/Engine/Source/Editor/Persona/Private/SAnimationBlendSpace.cpp
--- a/Engine/Source/Editor/Persona/Private/SAnimationBlendSpace.cpp
+++ b/Engine/Source/Editor/Persona/Private/SAnimationBlendSpace.cpp
@@ -16,6 +16,7 @@ void SBlendSpaceEditor::Construct(const FArguments& InArgs)
 		.OnBlendSpaceSampleDoubleClicked(InArgs._OnBlendSpaceSampleDoubleClicked)
 		.OnBlendSpaceSampleAdded(InArgs._OnBlendSpaceSampleAdded)
 		.OnBlendSpaceSampleRemoved(InArgs._OnBlendSpaceSampleRemoved)
+		.OnBlendSpaceSampleReplaced(InArgs._OnBlendSpaceSampleReplaced)
 		.OnGetBlendSpaceSampleName(InArgs._OnGetBlendSpaceSampleName)
 		.OnExtendSampleTooltip(InArgs._OnExtendSampleTooltip)
 		.OnSetPreviewPosition(InArgs._OnSetPreviewPosition)
@@ -31,6 +32,7 @@ void SBlendSpaceEditor::Construct(const FArguments& InArgs, const TSharedRef<cla
 		.OnBlendSpaceSampleDoubleClicked(InArgs._OnBlendSpaceSampleDoubleClicked)
 		.OnBlendSpaceSampleAdded(InArgs._OnBlendSpaceSampleAdded)
 		.OnBlendSpaceSampleRemoved(InArgs._OnBlendSpaceSampleRemoved)
+		.OnBlendSpaceSampleReplaced(InArgs._OnBlendSpaceSampleReplaced)
 		.OnGetBlendSpaceSampleName(InArgs._OnGetBlendSpaceSampleName)
 		.OnExtendSampleTooltip(InArgs._OnExtendSampleTooltip)
 		.OnSetPreviewPosition(InArgs._OnSetPreviewPosition)
diff --git a/Engine/Source/Editor/Persona/Private/SAnimationBlendSpace.h b/Engine/Source/Editor/Persona/Private/SAnimationBlendSpace.h
--- a/Engine/Source/Editor/Persona/Private/SAnimationBlendSpace.h
+++ b/Engine/Source/Editor/Persona/Private/SAnimationBlendSpace.h
@@ -26,6 +26,8 @@ public:
 
 	SLATE_EVENT(FOnBlendSpaceSampleRemoved, OnBlendSpaceSampleRemoved)
 
+	SLATE_EVENT(FOnBlendSpaceSampleReplaced, OnBlendSpaceSampleReplaced)
+
 	SLATE_EVENT(FOnGetBlendSpaceSampleName, OnGetBlendSpaceSampleName)
 
 	SLATE_EVENT(FOnExtendBlendSpaceSampleTooltip, OnExtendSampleTooltip)
diff --git a/Engine/Source/Editor/Persona/Private/SAnimationBlendSpace1D.cpp b/Engine/Source/Editor/Persona/Private/SAnimationBlendSpace1D.cpp
--- a/Engine/Source/Editor/Persona/Private/SAnimationBlendSpace1D.cpp
+++ b/Engine/Source/Editor/Persona/Private/SAnimationBlendSpace1D.cpp
@@ -17,6 +17,7 @@ void SBlendSpaceEditor1D::Construct(const FArguments& InArgs)
 		.OnBlendSpaceSampleDoubleClicked(InArgs._OnBlendSpaceSampleDoubleClicked)
 		.OnBlendSpaceSampleAdded(InArgs._OnBlendSpaceSampleAdded)
 		.OnBlendSpaceSampleRemoved(InArgs._OnBlendSpaceSampleRemoved)
+		.OnBlendSpaceSampleReplaced(InArgs._OnBlendSpaceSampleReplaced)
 		.OnGetBlendSpaceSampleName(InArgs._OnGetBlendSpaceSampleName)
 		.OnExtendSampleTooltip(InArgs._OnExtendSampleTooltip)
 		.OnSetPreviewPosition(InArgs._OnSetPreviewPosition)
@@ -32,6 +33,7 @@ void SBlendSpaceEditor1D::Construct(const FArguments& InArgs, const TSharedRef<I
 		.OnBlendSpaceSampleDoubleClicked(InArgs._OnBlendSpaceSampleDoubleClicked)
 		.OnBlendSpaceSampleAdded(InArgs._OnBlendSpaceSampleAdded)
 		.OnBlendSpaceSampleRemoved(InArgs._OnBlendSpaceSampleRemoved)
+		.OnBlendSpaceSampleReplaced(InArgs._OnBlendSpaceSampleReplaced)
 		.OnGetBlendSpaceSampleName(InArgs._OnGetBlendSpaceSampleName)
 		.OnExtendSampleTooltip(InArgs._OnExtendSampleTooltip)
 		.OnSetPreviewPosition(InArgs._OnSetPreviewPosition)
diff --git a/Engine/Source/Editor/Persona/Private/SAnimationBlendSpace1D.h b/Engine/Source/Editor/Persona/Private/SAnimationBlendSpace1D.h
--- a/Engine/Source/Editor/Persona/Private/SAnimationBlendSpace1D.h
+++ b/Engine/Source/Editor/Persona/Private/SAnimationBlendSpace1D.h
@@ -28,6 +28,8 @@ public:
 
 	SLATE_EVENT(FOnBlendSpaceSampleRemoved, OnBlendSpaceSampleRemoved)
 
+	SLATE_EVENT(FOnBlendSpaceSampleReplaced, OnBlendSpaceSampleReplaced)
+
 	SLATE_EVENT(FOnGetBlendSpaceSampleName, OnGetBlendSpaceSampleName)
 
 	SLATE_EVENT(FOnExtendBlendSpaceSampleTooltip, OnExtendSampleTooltip)
diff --git a/Engine/Source/Editor/Persona/Private/SAnimationBlendSpaceBase.cpp b/Engine/Source/Editor/Persona/Private/SAnimationBlendSpaceBase.cpp
--- a/Engine/Source/Editor/Persona/Private/SAnimationBlendSpaceBase.cpp
+++ b/Engine/Source/Editor/Persona/Private/SAnimationBlendSpaceBase.cpp
@@ -33,6 +33,7 @@ void SBlendSpaceEditorBase::Construct(const FArguments& InArgs)
 
 	OnBlendSpaceSampleAdded = InArgs._OnBlendSpaceSampleAdded;
 	OnBlendSpaceSampleRemoved = InArgs._OnBlendSpaceSampleRemoved;
+	OnBlendSpaceSampleReplaced = InArgs._OnBlendSpaceSampleReplaced;
 	OnSetPreviewPosition = InArgs._OnSetPreviewPosition;
 
 	bShouldSetPreviewValue = false;
@@ -74,7 +75,7 @@ void SBlendSpaceEditorBase::Construct(const FArguments& InArgs)
 							.OnSampleMoved(this, &SBlendSpaceEditorBase::OnSampleMoved)
 							.OnSampleRemoved(this, &SBlendSpaceEditorBase::OnSampleRemoved)
 							.OnSampleAdded(this, &SBlendSpaceEditorBase::OnSampleAdded)
-							.OnSampleAnimationChanged(this, &SBlendSpaceEditorBase::OnUpdateAnimation)
+							.OnSampleReplaced(this, &SBlendSpaceEditorBase::OnSampleReplaced)
 							.OnSampleDoubleClicked(InArgs._OnBlendSpaceSampleDoubleClicked)
 							.OnExtendSampleTooltip(InArgs._OnExtendSampleTooltip)
 							.OnGetBlendSpaceSampleName(InArgs._OnGetBlendSpaceSampleName)
@@ -156,16 +157,27 @@ void SBlendSpaceEditorBase::OnSampleAdded(UAnimSequence* Animation, const FVecto
 	BlendSpace->PostEditChange();
 }
 
-void SBlendSpaceEditorBase::OnUpdateAnimation(UAnimSequence* Animation, const FVector& Value)
+void SBlendSpaceEditorBase::OnSampleReplaced(int32 InSampleIndex, UAnimSequence* Animation)
 {
 	FScopedTransaction ScopedTransaction(LOCTEXT("UpdateAnimation", "Changing Animation Sequence"));
 	BlendSpace->Modify();
 
-	const bool bUpdateSuccesful = BlendSpace->UpdateSampleAnimation(Animation, Value);
+	bool bUpdateSuccesful = false;
+	if(BlendSpace->IsAsset())
+	{
+		bUpdateSuccesful = BlendSpace->ReplaceSampleAnimation(InSampleIndex, Animation);
+	}
+	else
+	{
+		bUpdateSuccesful = true;
+	}
+
 	if (bUpdateSuccesful)
 	{
 		ResampleData();
 		BlendSpace->ValidateSampleData();
+
+		OnBlendSpaceSampleReplaced.ExecuteIfBound(InSampleIndex, Animation);
 	}
 }
 
diff --git a/Engine/Source/Editor/Persona/Private/SAnimationBlendSpaceBase.h b/Engine/Source/Editor/Persona/Private/SAnimationBlendSpaceBase.h
--- a/Engine/Source/Editor/Persona/Private/SAnimationBlendSpaceBase.h
+++ b/Engine/Source/Editor/Persona/Private/SAnimationBlendSpaceBase.h
@@ -34,6 +34,8 @@ public:
 
 	SLATE_EVENT(FOnBlendSpaceSampleRemoved, OnBlendSpaceSampleRemoved)
 
+	SLATE_EVENT(FOnBlendSpaceSampleReplaced, OnBlendSpaceSampleReplaced)
+
 	SLATE_EVENT(FOnGetBlendSpaceSampleName, OnGetBlendSpaceSampleName) 
 
 	SLATE_EVENT(FOnExtendBlendSpaceSampleTooltip, OnExtendSampleTooltip)
@@ -65,7 +67,7 @@ protected:
 	void OnSampleMoved(const int32 SampleIndex, const FVector& NewValue, bool bIsInteractive, bool bSnap);
 	void OnSampleRemoved(const int32 SampleIndex);
 	void OnSampleAdded(UAnimSequence* Animation, const FVector& Value);
-	void OnUpdateAnimation(UAnimSequence* Animation, const FVector& Value);
+	void OnSampleReplaced(const int32 SampleIndex, UAnimSequence* Animation);
 	
 	// Begin SAnimEditorBase overrides
 	virtual UAnimationAsset* GetEditorObject() const override { return BlendSpace; }
@@ -110,6 +112,9 @@ protected:
 	/** Delegate called when a sample is removed */
 	FOnBlendSpaceSampleRemoved OnBlendSpaceSampleRemoved;
 
+	/** Delegate called when a sample is replaced */
+	FOnBlendSpaceSampleReplaced OnBlendSpaceSampleReplaced;
+
 	/** Delegate called to externally control the preview position */
 	FOnSetBlendSpacePreviewPosition OnSetPreviewPosition;
 
diff --git a/Engine/Source/Editor/Persona/Private/SAnimationBlendSpaceGridWidget.cpp b/Engine/Source/Editor/Persona/Private/SAnimationBlendSpaceGridWidget.cpp
--- a/Engine/Source/Editor/Persona/Private/SAnimationBlendSpaceGridWidget.cpp
+++ b/Engine/Source/Editor/Persona/Private/SAnimationBlendSpaceGridWidget.cpp
@@ -45,8 +45,8 @@ void SBlendSpaceGridWidget::Construct(const FArguments& InArgs)
 	OnSampleAdded = InArgs._OnSampleAdded;
 	OnSampleMoved = InArgs._OnSampleMoved;
 	OnSampleRemoved = InArgs._OnSampleRemoved;
+	OnSampleReplaced = InArgs._OnSampleReplaced;
 	OnSampleDoubleClicked = InArgs._OnSampleDoubleClicked;
-	OnSampleAnimationChanged = InArgs._OnSampleAnimationChanged;
 	OnGetBlendSpaceSampleName = InArgs._OnGetBlendSpaceSampleName;
 	OnExtendSampleTooltip = InArgs._OnExtendSampleTooltip;
 	bReadOnly = InArgs._ReadOnly;
@@ -631,12 +631,12 @@ FReply SBlendSpaceGridWidget::OnDrop(const FGeometry& MyGeometry, const FDragDro
 				}
 				else if (DragState == EDragState::DragDropOverride)
 				{
-					const FVector SampleValue = SnapToClosestSamplePoint(LocalMousePosition);
 					TSharedPtr<FAssetDragDropOp> DragDropOperation = DragDropEvent.GetOperationAs<FAssetDragDropOp>();
 					if (DragDropOperation.IsValid())
 					{
 						UAnimSequence* Animation = FAssetData::GetFirstAsset<UAnimSequence>(DragDropOperation->GetAssets());
-						OnSampleAnimationChanged.ExecuteIfBound(Animation, SampleValue);
+						int32 DroppedSampleIndex = GetClosestSamplePointIndexToMouse();
+						OnSampleReplaced.ExecuteIfBound(DroppedSampleIndex, Animation);
 					}
 				}
 
@@ -891,16 +891,7 @@ FReply SBlendSpaceGridWidget::ProcessClick(const FGeometry& MyGeometry, const FP
 				if (HighlightedSampleIndex == INDEX_NONE)
 				{
 					// If there isn't any sample currently being highlighted, retrieve all of them and see if we are over one
-					const TArray<FBlendSample>& Samples = BlendSpace->GetBlendSamples();
-					for (int32 SampleIndex = 0; SampleIndex < Samples.Num(); ++SampleIndex)
-					{
-						const FBlendSample& Sample = Samples[SampleIndex];
-						if (IsSampleValueWithinMouseRange(Sample.SampleValue))
-						{
-							SelectedSampleIndex = SampleIndex;					
-							break;
-						}
-					}
+					SelectedSampleIndex = GetClosestSamplePointIndexToMouse();
 				}
 				else
 				{
@@ -1296,11 +1287,37 @@ const FSlateRect SBlendSpaceGridWidget::GetGridRectangleFromGeometry(const FGeom
 	return WindowRect.InsetBy(GridMargin + GridRatioMargin);
 }
 
-bool SBlendSpaceGridWidget::IsSampleValueWithinMouseRange(const FVector& SampleValue) const
+bool SBlendSpaceGridWidget::IsSampleValueWithinMouseRange(const FVector& SampleValue, float& OutDistance) const
 {
 	const FVector2D GridPosition = SampleValueToGridPosition(SampleValue);
-	const float MouseDistance = FVector2D::Distance(LocalMousePosition, GridPosition);	
-	return (FMath::Abs(MouseDistance) < ClickAndHighlightThreshold);
+	OutDistance = FVector2D::Distance(LocalMousePosition, GridPosition);
+	return (OutDistance < ClickAndHighlightThreshold);
+}
+
+int32 SBlendSpaceGridWidget::GetClosestSamplePointIndexToMouse() const
+{
+	float BestDistance = FLT_MAX;
+	int32 BestIndex = INDEX_NONE;
+
+	if(const UBlendSpaceBase* BlendSpace = BlendSpaceBase.Get())
+	{
+		const TArray<FBlendSample>& Samples = BlendSpace->GetBlendSamples();
+		for (int32 SampleIndex = 0; SampleIndex < Samples.Num(); ++SampleIndex)
+		{
+			const FBlendSample& Sample = Samples[SampleIndex];
+			float Distance;
+			if (IsSampleValueWithinMouseRange(Sample.SampleValue, Distance))
+			{
+				if(Distance < BestDistance)
+				{
+					BestDistance = Distance;
+					BestIndex = SampleIndex;
+				}
+			}
+		}
+	}
+
+	return BestIndex;
 }
 
 void SBlendSpaceGridWidget::StartPreviewing()
@@ -1505,8 +1522,16 @@ FText SBlendSpaceGridWidget::GetToolTipSampleValue() const
 
 				case EDragState::DragDropOverride:
 				{
-					static const FTextFormat OverrideAnimationFormat = FTextFormat::FromString("Changing Animation from {0} to {1}");
-					ToolTipText = FText::Format(OverrideAnimationFormat, HoveredAnimationName, DragDropAnimationName);
+					if(HoveredAnimationName.IsEmpty())
+					{
+						static const FTextFormat OverrideAnimationFormat = LOCTEXT("InvalidSampleChangingFormat", "Changing sample to {0}");
+						ToolTipText = FText::Format(OverrideAnimationFormat, DragDropAnimationName);
+					}
+					else
+					{
+						static const FTextFormat OverrideAnimationFormat = LOCTEXT("ValidSampleChangingFormat", "Changing sample from {0} to {1}");
+						ToolTipText = FText::Format(OverrideAnimationFormat, HoveredAnimationName, DragDropAnimationName);
+					}
 					break;
 				}
 				// If the drag and drop operation is invalid return the cached error message as to why it is invalid
@@ -1865,8 +1890,9 @@ void SBlendSpaceGridWidget::Tick(const FGeometry& AllottedGeometry, const double
 		{
 			if (DragState == EDragState::None)
 			{
-				// Check if we are highlighting preview pin						
-				bHighlightPreviewPin = IsSampleValueWithinMouseRange(LastPreviewingSampleValue);
+				// Check if we are highlighting preview pin
+				float Distance;
+				bHighlightPreviewPin = IsSampleValueWithinMouseRange(LastPreviewingSampleValue, Distance);
 				if (bHighlightPreviewPin)
 				{
 					if (bHighlightPreviewPin != bPreviousHighlightPreviewPin)
@@ -1880,16 +1906,7 @@ void SBlendSpaceGridWidget::Tick(const FGeometry& AllottedGeometry, const double
 				}
 		
 				// Determine highlighted sample
-				const TArray<FBlendSample>& Samples = BlendSpace->GetBlendSamples();
-				for (int32 SampleIndex = 0; SampleIndex < Samples.Num(); ++SampleIndex)
-				{
-					const FBlendSample& Sample = Samples[SampleIndex];
-					if (IsSampleValueWithinMouseRange(Sample.SampleValue))
-					{
-						HighlightedSampleIndex = SampleIndex;
-						break;
-					}
-				}
+				HighlightedSampleIndex = GetClosestSamplePointIndexToMouse();
 
 				if (!bHighlightPreviewPin)
 				{
@@ -1959,7 +1976,7 @@ void SBlendSpaceGridWidget::Tick(const FGeometry& AllottedGeometry, const double
 						const FBlendSample& Sample = Samples[SampleIndex];
 						if (Sample.SampleValue == DropSampleValue)
 						{
-							HoveredAnimationName = Sample.Animation ? FText::FromString(Sample.Animation->GetName()) : FText::FromString("Invalid Animation Sequence");
+							HoveredAnimationName = Sample.Animation ? FText::FromString(Sample.Animation->GetName()) : FText::GetEmpty();
 							break;
 						}
 					}
@@ -2071,36 +2088,39 @@ bool SBlendSpaceGridWidget::ValidateAnimationSequence(const UAnimSequence* Anima
 	{
 		if(const UBlendSpaceBase* BlendSpace = BlendSpaceBase.Get())
 		{
-			// If there are any existing blend samples check whether or not the the animation should be additive and if so if the additive matches the existing samples
-			if (BlendSpace->GetNumberOfBlendSamples() > 0)
+			if(BlendSpace->IsAsset())
 			{
-				const bool bIsAdditive = BlendSpace->ShouldAnimationBeAdditive();
-				if (AnimationSequence->IsValidAdditive() != bIsAdditive)
+				// If there are any existing blend samples check whether or not the the animation should be additive and if so if the additive matches the existing samples
+				if ( BlendSpace->GetNumberOfBlendSamples() > 0)
 				{
-					InvalidOperationText = FText::FromString(bIsAdditive ? "Animation should be additive" : "Animation should be non-additive");
-					return false;
+					const bool bIsAdditive = BlendSpace->ShouldAnimationBeAdditive();
+					if (AnimationSequence->IsValidAdditive() != bIsAdditive)
+					{
+						InvalidOperationText = FText::FromString(bIsAdditive ? "Animation should be additive" : "Animation should be non-additive");
+						return false;
+					}
+
+					// If it is the supported additive type, but does not match existing samples
+					if (!BlendSpace->DoesAnimationMatchExistingSamples(AnimationSequence))
+					{
+						InvalidOperationText = FText::FromString("Additive Animation Type does not match existing Samples");
+						return false;
+					}
 				}
 
-				// If it is the supported additive type, but does not match existing samples
-				if (!BlendSpace->DoesAnimationMatchExistingSamples(AnimationSequence))
+				// Check if the supplied animation is of a different additive animation type 
+				if (!BlendSpace->IsAnimationCompatible(AnimationSequence))
 				{
-					InvalidOperationText = FText::FromString("Additive Animation Type does not match existing Samples");
+					InvalidOperationText = FText::FromString("Invalid Additive Animation Type");
 					return false;
 				}
-			}
-
-			// Check if the supplied animation is of a different additive animation type 
-			if (!BlendSpace->IsAnimationCompatible(AnimationSequence))
-			{
-				InvalidOperationText = FText::FromString("Invalid Additive Animation Type");
-				return false;
-			}
 
-			// Check if the supplied animation is compatible with the skeleton
-			if (!BlendSpace->IsAnimationCompatibleWithSkeleton(AnimationSequence))
-			{
-				InvalidOperationText = FText::FromString("Animation is incompatible with the skeleton");
-				return false;
+				// Check if the supplied animation is compatible with the skeleton
+				if (!BlendSpace->IsAnimationCompatibleWithSkeleton(AnimationSequence))
+				{
+					InvalidOperationText = FText::FromString("Animation is incompatible with the skeleton");
+					return false;
+				}
 			}
 		}
 	}
diff --git a/Engine/Source/Editor/Persona/Private/SAnimationBlendSpaceGridWidget.h b/Engine/Source/Editor/Persona/Private/SAnimationBlendSpaceGridWidget.h
--- a/Engine/Source/Editor/Persona/Private/SAnimationBlendSpaceGridWidget.h
+++ b/Engine/Source/Editor/Persona/Private/SAnimationBlendSpaceGridWidget.h
@@ -46,7 +46,7 @@ public:
 		SLATE_EVENT(FOnSampleMoved, OnSampleMoved)
 		SLATE_EVENT(FOnBlendSpaceSampleRemoved, OnSampleRemoved)
 		SLATE_EVENT(FOnBlendSpaceSampleAdded, OnSampleAdded)
-		SLATE_EVENT(FOnSampleAnimationChanged, OnSampleAnimationChanged)
+		SLATE_EVENT(FOnBlendSpaceSampleReplaced, OnSampleReplaced)
 		SLATE_EVENT(FOnBlendSpaceSampleDoubleClicked, OnSampleDoubleClicked)
 		SLATE_EVENT(FOnGetBlendSpaceSampleName, OnGetBlendSpaceSampleName)
 		SLATE_ARGUMENT(FName, StatusBarName)
@@ -182,7 +182,9 @@ protected:
 	/** Returns the (calculated) grid rectangle given the supplied geometry, and calls UpdateGridRatioMargin */
 	const FSlateRect GetGridRectangleFromGeometry(const FGeometry& MyGeometry);
 	/** Checks whether or not the blendspace sample value is within the range of the mouse position */
-	bool IsSampleValueWithinMouseRange(const FVector& SampleValue) const;
+	bool IsSampleValueWithinMouseRange(const FVector& SampleValue, float& OutDistance) const;
+	/** Get the sample index that most closely corresponds to the mouse position */
+	int32 GetClosestSamplePointIndexToMouse() const;
 
 	/** Sets the tooltip instance on the underlying widget instance */
 	void ShowToolTip();
@@ -295,7 +297,7 @@ private:
 	FOnBlendSpaceSampleAdded OnSampleAdded;
 	FOnSampleMoved OnSampleMoved;
 	FOnBlendSpaceSampleRemoved OnSampleRemoved;
-	FOnSampleAnimationChanged OnSampleAnimationChanged;
+	FOnBlendSpaceSampleReplaced OnSampleReplaced;
 	FOnBlendSpaceSampleDoubleClicked OnSampleDoubleClicked;
 	FOnGetBlendSpaceSampleName OnGetBlendSpaceSampleName;
 	FOnExtendBlendSpaceSampleTooltip OnExtendSampleTooltip;
diff --git a/Engine/Source/Editor/Persona/Public/PersonaDelegates.h b/Engine/Source/Editor/Persona/Public/PersonaDelegates.h
--- a/Engine/Source/Editor/Persona/Public/PersonaDelegates.h
+++ b/Engine/Source/Editor/Persona/Public/PersonaDelegates.h
@@ -46,6 +46,9 @@ DECLARE_DELEGATE_OneParam(FOnBlendSpaceSampleRemoved, const int32 /*SampleIndex*
 // Called when a blendspace sample point is added
 DECLARE_DELEGATE_TwoParams(FOnBlendSpaceSampleAdded, UAnimSequence* /*InSequence*/, const FVector& /*InSamplePoint*/);
 
+// Called when a blendspace sample point is replaced
+DECLARE_DELEGATE_TwoParams(FOnBlendSpaceSampleReplaced, const int32 /*SampleIndex*/, UAnimSequence* /*InSequence*/);
+
 // Called to get the overridden name of a blend sample
 DECLARE_DELEGATE_RetVal_OneParam(FName, FOnGetBlendSpaceSampleName, int32 /*SampleIndex*/);
 
diff --git a/Engine/Source/Editor/Persona/Public/PersonaModule.h b/Engine/Source/Editor/Persona/Public/PersonaModule.h
--- a/Engine/Source/Editor/Persona/Public/PersonaModule.h
+++ b/Engine/Source/Editor/Persona/Public/PersonaModule.h
@@ -107,6 +107,9 @@ struct FBlendSpaceEditorArgs
 
 	// Called when a blendspace sample point is added
 	FOnBlendSpaceSampleAdded OnBlendSpaceSampleAdded;
+	
+	// Called when a blendspace sample point is replaced
+	FOnBlendSpaceSampleReplaced OnBlendSpaceSampleReplaced;
 
 	// Called when a blendspace sample point is double clicked
 	FOnBlendSpaceSampleDoubleClicked OnBlendSpaceSampleDoubleClicked;
diff --git a/Engine/Source/Editor/UnrealEd/Private/Kismet2/BlueprintEditorUtils.cpp b/Engine/Source/Editor/UnrealEd/Private/Kismet2/BlueprintEditorUtils.cpp
--- a/Engine/Source/Editor/UnrealEd/Private/Kismet2/BlueprintEditorUtils.cpp
+++ b/Engine/Source/Editor/UnrealEd/Private/Kismet2/BlueprintEditorUtils.cpp
@@ -2581,15 +2581,6 @@ void FBlueprintEditorUtils::RemoveGraphs( UBlueprint* Blueprint, const TArray<cl
 // Removes the supplied graph from the Blueprint.
 void FBlueprintEditorUtils::RemoveGraph(UBlueprint* Blueprint, class UEdGraph* GraphToRemove, EGraphRemoveFlags::Type Flags /*= Transient | Recompile */)
 {
-	struct Local
-	{
-		static bool IsASubGraph(UEdGraph* Graph)
-		{
-			UObject* Outer = Graph->GetOuter();
-			return ( Outer && Outer->IsA( UK2Node_Composite::StaticClass() ) );
-		}
-	};
-
 	GraphToRemove->Modify();
 
 	for (UObject* TestOuter = GraphToRemove->GetOuter(); TestOuter; TestOuter = TestOuter->GetOuter())
diff --git a/Engine/Source/Runtime/Engine/Classes/Animation/BlendSpaceBase.h b/Engine/Source/Runtime/Engine/Classes/Animation/BlendSpaceBase.h
--- a/Engine/Source/Runtime/Engine/Classes/Animation/BlendSpaceBase.h
+++ b/Engine/Source/Runtime/Engine/Classes/Animation/BlendSpaceBase.h
@@ -343,9 +343,12 @@ class UBlendSpaceBase : public UAnimationAsset, public IInterpolationIndexProvid
 	/** edit samples */
 	ENGINE_API bool	EditSampleValue(const int32 BlendSampleIndex, const FVector& NewValue, bool bSnap = true);
 
-	/** update animation on grid sample */
+	UE_DEPRECATED(5.0, "Please use ReplaceSampleAnimation instead")
 	ENGINE_API bool	UpdateSampleAnimation(UAnimSequence* AnimationSequence, const FVector& SampleValue);
 
+	/** update animation on grid sample */
+	ENGINE_API bool	ReplaceSampleAnimation(const int32 BlendSampleIndex, UAnimSequence* AnimationSequence);
+
 	/** delete samples */
 	ENGINE_API bool	DeleteSample(const int32 BlendSampleIndex);
 	
diff --git a/Engine/Source/Runtime/Engine/Private/Animation/BlendSpaceBase.cpp b/Engine/Source/Runtime/Engine/Private/Animation/BlendSpaceBase.cpp
--- a/Engine/Source/Runtime/Engine/Private/Animation/BlendSpaceBase.cpp
+++ b/Engine/Source/Runtime/Engine/Private/Animation/BlendSpaceBase.cpp
@@ -1014,6 +1014,17 @@ bool UBlendSpaceBase::UpdateSampleAnimation(UAnimSequence* AnimationSequence, co
 	return UpdateSampleIndex != INDEX_NONE;
 }
 
+bool UBlendSpaceBase::ReplaceSampleAnimation(const int32 BlendSampleIndex, UAnimSequence* AnimationSequence)
+{
+	const bool bValidValue = SampleData.IsValidIndex(BlendSampleIndex);
+	if (bValidValue)
+	{
+		SampleData[BlendSampleIndex].Animation = AnimationSequence;
+	}
+
+	return bValidValue;
+}
+
 bool UBlendSpaceBase::DeleteSample(const int32 BlendSampleIndex)
 {
 	const bool bValidRemoval = SampleData.IsValidIndex(BlendSampleIndex);
