Show the current as well as target preview position in BlendSpaces

diff --git a/Engine/Plugins/Animation/GameplayInsights/Source/GameplayInsights/Private/AnimNodesTrack.cpp b/Engine/Plugins/Animation/GameplayInsights/Source/GameplayInsights/Private/AnimNodesTrack.cpp
--- a/Engine/Plugins/Animation/GameplayInsights/Source/GameplayInsights/Private/AnimNodesTrack.cpp
+++ b/Engine/Plugins/Animation/GameplayInsights/Source/GameplayInsights/Private/AnimNodesTrack.cpp
@@ -369,7 +369,7 @@ void FAnimNodesTrack::UpdateDebugData(const TraceServices::FFrame& InFrame)
 												BlendSpaceBase = TSoftObjectPtr<UBlendSpaceBase>(FSoftObjectPath(BlendSpaceInfo->PathName)).LoadSynchronous();
 											}
 
-											DebugData.RecordBlendSpacePlayer(InMessage.NodeId, BlendSpaceBase, InMessage.PositionX, InMessage.PositionY, InMessage.PositionZ);
+											DebugData.RecordBlendSpacePlayer(InMessage.NodeId, BlendSpaceBase, FVector(InMessage.PositionX, InMessage.PositionY, InMessage.PositionZ), FVector(InMessage.FilteredPositionX, InMessage.FilteredPositionY, InMessage.FilteredPositionZ));
 											return TraceServices::EEventEnumerate::Continue;
 										});
 									});
diff --git a/Engine/Plugins/Animation/GameplayInsights/Source/GameplayInsights/Private/AnimationAnalyzer.cpp b/Engine/Plugins/Animation/GameplayInsights/Source/GameplayInsights/Private/AnimationAnalyzer.cpp
--- a/Engine/Plugins/Animation/GameplayInsights/Source/GameplayInsights/Private/AnimationAnalyzer.cpp
+++ b/Engine/Plugins/Animation/GameplayInsights/Source/GameplayInsights/Private/AnimationAnalyzer.cpp
@@ -62,10 +62,12 @@ bool FAnimationAnalyzer::OnEvent(uint16 RouteId, EStyle Style, const FOnEventCon
 		float PlayRate = EventData.GetValue<float>("PlayRate");
 		float BlendSpacePositionX = EventData.GetValue<float>("BlendSpacePositionX");
 		float BlendSpacePositionY = EventData.GetValue<float>("BlendSpacePositionY");
+		float BlendSpaceFilteredPositionX = EventData.GetValue<float>("BlendSpaceFilteredPositionX");
+		float BlendSpaceFilteredPositionY = EventData.GetValue<float>("BlendSpaceFilteredPositionY");
 		uint16 FrameCounter = EventData.GetValue<uint16>("FrameCounter");
 		bool bLooping = EventData.GetValue<bool>("Looping");
 		bool bIsBlendSpace = EventData.GetValue<bool>("IsBlendSpace");
-		AnimationProvider.AppendTickRecord(AnimInstanceId, Context.EventTime.AsSeconds(Cycle), AssetId, NodeId, BlendWeight, PlaybackTime, RootMotionWeight, PlayRate, BlendSpacePositionX, BlendSpacePositionY, FrameCounter, bLooping, bIsBlendSpace);
+		AnimationProvider.AppendTickRecord(AnimInstanceId, Context.EventTime.AsSeconds(Cycle), AssetId, NodeId, BlendWeight, PlaybackTime, RootMotionWeight, PlayRate, BlendSpacePositionX, BlendSpacePositionY, BlendSpaceFilteredPositionX, BlendSpaceFilteredPositionY, FrameCounter, bLooping, bIsBlendSpace);
 		break;
 	}
 	case RouteId_SkeletalMesh:
@@ -279,10 +281,9 @@ bool FAnimationAnalyzer::OnEvent(uint16 RouteId, EStyle Style, const FOnEventCon
 		uint64 AnimInstanceId = EventData.GetValue<uint64>("AnimInstanceId");
 		uint64 BlendSpaceId = EventData.GetValue<uint64>("BlendSpaceId");
 		int32 NodeId = EventData.GetValue<int32>("NodeId");
-		float PositionX = EventData.GetValue<float>("PositionX");
-		float PositionY = EventData.GetValue<float>("PositionY");
-		float PositionZ = EventData.GetValue<float>("PositionZ");
-		AnimationProvider.AppendBlendSpacePlayer(AnimInstanceId, Context.EventTime.AsSeconds(Cycle), NodeId, BlendSpaceId, PositionX, PositionY, PositionZ);
+		FVector BlendPosition(EventData.GetValue<float>("BlendPositionX"), EventData.GetValue<float>("BlendPositionY"), EventData.GetValue<float>("BlendPositionZ"));
+		FVector FilteredBlendPosition(EventData.GetValue<float>("FilteredBlendPositionX"), EventData.GetValue<float>("FilteredBlendPositionY"), EventData.GetValue<float>("FilteredBlendPositionZ"));
+		AnimationProvider.AppendBlendSpacePlayer(AnimInstanceId, Context.EventTime.AsSeconds(Cycle), NodeId, BlendSpaceId, BlendPosition, FilteredBlendPosition);
 		break;
 	}
 	case RouteId_StateMachineState:
diff --git a/Engine/Plugins/Animation/GameplayInsights/Source/GameplayInsights/Private/AnimationProvider.cpp b/Engine/Plugins/Animation/GameplayInsights/Source/GameplayInsights/Private/AnimationProvider.cpp
--- a/Engine/Plugins/Animation/GameplayInsights/Source/GameplayInsights/Private/AnimationProvider.cpp
+++ b/Engine/Plugins/Animation/GameplayInsights/Source/GameplayInsights/Private/AnimationProvider.cpp
@@ -410,7 +410,7 @@ bool FAnimationProvider::HasAnyData() const
 	return bHasAnyData;
 }
 
-void FAnimationProvider::AppendTickRecord(uint64 InAnimInstanceId, double InTime, uint64 InAssetId, int32 InNodeId, float InBlendWeight, float InPlaybackTime, float InRootMotionWeight, float InPlayRate, float InBlendSpacePositionX, float InBlendSpacePositionY, uint16 InFrameCounter, bool bInLooping, bool bInIsBlendSpace)
+void FAnimationProvider::AppendTickRecord(uint64 InAnimInstanceId, double InTime, uint64 InAssetId, int32 InNodeId, float InBlendWeight, float InPlaybackTime, float InRootMotionWeight, float InPlayRate, float InBlendSpacePositionX, float InBlendSpacePositionY, float InBlendSpaceFilteredPositionX, float InBlendSpaceFilteredPositionY, uint16 InFrameCounter, bool bInLooping, bool bInIsBlendSpace)
 {
 	Session.WriteAccessCheck();
 
@@ -444,6 +444,8 @@ void FAnimationProvider::AppendTickRecord(uint64 InAnimInstanceId, double InTime
 	Message.PlayRate = InPlayRate;
 	Message.BlendSpacePositionX = InBlendSpacePositionX;
 	Message.BlendSpacePositionY = InBlendSpacePositionY;
+	Message.BlendSpaceFilteredPositionX = InBlendSpaceFilteredPositionX;
+	Message.BlendSpaceFilteredPositionY = InBlendSpaceFilteredPositionY;
 	Message.FrameCounter = InFrameCounter;
 	Message.bLooping = bInLooping;
 	Message.bIsBlendSpace = bInIsBlendSpace;
@@ -750,7 +752,7 @@ void FAnimationProvider::AppendAnimSequencePlayer(uint64 InAnimInstanceId, doubl
 	Session.UpdateDurationSeconds(InTime);
 }
 
-void FAnimationProvider::AppendBlendSpacePlayer(uint64 InAnimInstanceId, double InTime, int32 InNodeId, uint64 InBlendSpaceId, float InPositionX, float InPositionY, float InPositionZ)
+void FAnimationProvider::AppendBlendSpacePlayer(uint64 InAnimInstanceId, double InTime, int32 InNodeId, uint64 InBlendSpaceId, const FVector& InBlendPosition, const FVector& InFilteredBlendPosition)
 {
 	Session.WriteAccessCheck();
 
@@ -774,9 +776,12 @@ void FAnimationProvider::AppendBlendSpacePlayer(uint64 InAnimInstanceId, double
 	Message.AnimInstanceId = InAnimInstanceId;
 	Message.BlendSpaceId = InBlendSpaceId;
 	Message.NodeId = InNodeId;
-	Message.PositionX = InPositionX;
-	Message.PositionY = InPositionY;
-	Message.PositionZ = InPositionZ;
+	Message.PositionX = InBlendPosition.X;
+	Message.PositionY = InBlendPosition.Y;
+	Message.PositionZ = InBlendPosition.Z;
+	Message.FilteredPositionX = InFilteredBlendPosition.X;
+	Message.FilteredPositionY = InFilteredBlendPosition.Y;
+	Message.FilteredPositionZ = InFilteredBlendPosition.Z;
 
 	Timeline->AppendEvent(InTime, Message);
 
diff --git a/Engine/Plugins/Animation/GameplayInsights/Source/GameplayInsights/Private/AnimationProvider.h b/Engine/Plugins/Animation/GameplayInsights/Source/GameplayInsights/Private/AnimationProvider.h
--- a/Engine/Plugins/Animation/GameplayInsights/Source/GameplayInsights/Private/AnimationProvider.h
+++ b/Engine/Plugins/Animation/GameplayInsights/Source/GameplayInsights/Private/AnimationProvider.h
@@ -49,7 +49,7 @@ public:
 	bool HasAnyData() const;
 
 	/** Add a tick record */
-	void AppendTickRecord(uint64 InAnimInstanceId, double InTime, uint64 InAssetId, int32 InNodeId, float InBlendWeight, float InPlaybackTime, float InRootMotionWeight, float InPlayRate, float InBlendSpacePositionX, float InBlendSpacePositionY, uint16 InFrameCounter, bool bInLooping, bool bInIsBlendSpace);
+	void AppendTickRecord(uint64 InAnimInstanceId, double InTime, uint64 InAssetId, int32 InNodeId, float InBlendWeight, float InPlaybackTime, float InRootMotionWeight, float InPlayRate, float InBlendSpacePositionX, float InBlendSpacePositionY, float InBlendSpaceFilteredPositionX, float InBlendSpaceFilteredPositionY, uint16 InFrameCounter, bool bInLooping, bool bInIsBlendSpace);
 
 	/** Add a skeletal mesh */
 	void AppendSkeletalMesh(uint64 InObjectId, const TArrayView<const int32>& ParentIndices);
@@ -90,7 +90,7 @@ public:
 	void AppendAnimSequencePlayer(uint64 InAnimInstanceId, double InTime, int32 InNodeId, float InPosition, float InLength, uint16 InFrameCounter);
 
 	/** Add a blend space player record */
-	void AppendBlendSpacePlayer(uint64 InAnimInstanceId, double InTime, int32 InNodeId, uint64 InBlendSpaceId, float InPositionX, float InPositionY, float InPositionZ);
+	void AppendBlendSpacePlayer(uint64 InAnimInstanceId, double InTime, int32 InNodeId, uint64 InBlendSpaceId, const FVector& InBlendPosition, const FVector& InFilteredBlendPosition);
 
 	/** Add a state machine state */
 	void AppendStateMachineState(uint64 InAnimInstanceId, double InTime, int32 InNodeId, int32 InStateMachineIndex, int32 InStateIndex, float InStateWeight, float InElapsedTime);
diff --git a/Engine/Plugins/Animation/GameplayInsights/Source/GameplayInsights/Private/AnimationTickRecordsTrack.cpp b/Engine/Plugins/Animation/GameplayInsights/Source/GameplayInsights/Private/AnimationTickRecordsTrack.cpp
--- a/Engine/Plugins/Animation/GameplayInsights/Source/GameplayInsights/Private/AnimationTickRecordsTrack.cpp
+++ b/Engine/Plugins/Animation/GameplayInsights/Source/GameplayInsights/Private/AnimationTickRecordsTrack.cpp
@@ -40,6 +40,8 @@ FString FTickRecordSeries::FormatValue(double Value) const
 	case ESeriesType::PlayRate:
 	case ESeriesType::BlendSpacePositionX:
 	case ESeriesType::BlendSpacePositionY:
+	case ESeriesType::BlendSpaceFilteredPositionX:
+	case ESeriesType::BlendSpaceFilteredPositionY:
 		return FText::AsNumber(Value).ToString();
 	}
 
@@ -152,6 +154,18 @@ void FAnimationTickRecordsTrack::AddAllSeries()
 						FTickRecordSeries::ESeriesType::BlendSpacePositionY,
 						false
 					},
+					{
+						LOCTEXT("SeriesNameBlendSpaceFilteredPositionX", "BlendSpace Filtered Position X"),
+						LOCTEXT("SeriesDescBlendSpaceFilteredPositionX", "The X value after filtering used to sample this blend space"),
+						FTickRecordSeries::ESeriesType::BlendSpaceFilteredPositionX,
+						false
+					},
+					{
+						LOCTEXT("SeriesNameBlendSpaceFilteredPositionY", "BlendSpace Filtered Position Y"),
+						LOCTEXT("SeriesDescBlendSpaceFilteredPositionY", "The Y value after filtering used to sample this blend space"),
+						FTickRecordSeries::ESeriesType::BlendSpaceFilteredPositionY,
+						false
+					},
 				};
 
 				const FClassInfo& ClassInfo = GameplayProvider->GetClassInfoFromObject(InAssetId);
@@ -277,6 +291,10 @@ bool FAnimationTickRecordsTrack::UpdateSeriesBounds(FGameplayGraphSeries& InSeri
 		return UpdateSeriesBoundsHelper(TickRecordSeries, InViewport, &FTickRecordMessage::BlendSpacePositionX);
 	case FTickRecordSeries::ESeriesType::BlendSpacePositionY:
 		return UpdateSeriesBoundsHelper(TickRecordSeries, InViewport, &FTickRecordMessage::BlendSpacePositionY);
+	case FTickRecordSeries::ESeriesType::BlendSpaceFilteredPositionX:
+		return UpdateSeriesBoundsHelper(TickRecordSeries, InViewport, &FTickRecordMessage::BlendSpaceFilteredPositionX);
+	case FTickRecordSeries::ESeriesType::BlendSpaceFilteredPositionY:
+		return UpdateSeriesBoundsHelper(TickRecordSeries, InViewport, &FTickRecordMessage::BlendSpaceFilteredPositionY);
 	}
 
 	return false;
@@ -305,6 +323,12 @@ void FAnimationTickRecordsTrack::UpdateSeries(FGameplayGraphSeries& InSeries, co
 	case FTickRecordSeries::ESeriesType::BlendSpacePositionY:
 		UpdateSeriesHelper(TickRecordSeries, InViewport, &FTickRecordMessage::BlendSpacePositionY);
 		break;
+	case FTickRecordSeries::ESeriesType::BlendSpaceFilteredPositionX:
+		UpdateSeriesHelper(TickRecordSeries, InViewport, &FTickRecordMessage::BlendSpaceFilteredPositionX);
+		break;
+	case FTickRecordSeries::ESeriesType::BlendSpaceFilteredPositionY:
+		UpdateSeriesHelper(TickRecordSeries, InViewport, &FTickRecordMessage::BlendSpaceFilteredPositionY);
+		break;
 	}
 }
 
@@ -464,6 +488,8 @@ void FAnimationTickRecordsTrack::GetVariantsAtFrame(const TraceServices::FFrame&
 					{
 						Header->AddChild(FVariantTreeNode::MakeFloat(LOCTEXT("BlendSpacePositionX", "Blend Space Position X"), InMessage.BlendSpacePositionX));
 						Header->AddChild(FVariantTreeNode::MakeFloat(LOCTEXT("BlendSpacePositionY", "Blend Space Position Y"), InMessage.BlendSpacePositionY));
+						Header->AddChild(FVariantTreeNode::MakeFloat(LOCTEXT("BlendSpaceFilteredPositionX", "Blend Space Filtered Position X"), InMessage.BlendSpaceFilteredPositionX));
+						Header->AddChild(FVariantTreeNode::MakeFloat(LOCTEXT("BlendSpaceFilteredPositionY", "Blend Space Filtered Position Y"), InMessage.BlendSpaceFilteredPositionY));
 					}
 				}
 				return TraceServices::EEventEnumerate::Continue;
diff --git a/Engine/Plugins/Animation/GameplayInsights/Source/GameplayInsights/Private/AnimationTickRecordsTrack.h b/Engine/Plugins/Animation/GameplayInsights/Source/GameplayInsights/Private/AnimationTickRecordsTrack.h
--- a/Engine/Plugins/Animation/GameplayInsights/Source/GameplayInsights/Private/AnimationTickRecordsTrack.h
+++ b/Engine/Plugins/Animation/GameplayInsights/Source/GameplayInsights/Private/AnimationTickRecordsTrack.h
@@ -24,6 +24,8 @@ public:
 		PlayRate,
 		BlendSpacePositionX,
 		BlendSpacePositionY,
+		BlendSpaceFilteredPositionX,
+		BlendSpaceFilteredPositionY,
 
 		Count,
 	};
diff --git a/Engine/Plugins/Animation/GameplayInsights/Source/GameplayInsights/Public/IAnimationProvider.h b/Engine/Plugins/Animation/GameplayInsights/Source/GameplayInsights/Public/IAnimationProvider.h
--- a/Engine/Plugins/Animation/GameplayInsights/Source/GameplayInsights/Public/IAnimationProvider.h
+++ b/Engine/Plugins/Animation/GameplayInsights/Source/GameplayInsights/Public/IAnimationProvider.h
@@ -53,6 +53,8 @@ struct FTickRecordMessage
 	float PlayRate = 0.0f;
 	float BlendSpacePositionX = 0.0f;
 	float BlendSpacePositionY = 0.0f;
+	float BlendSpaceFilteredPositionX = 0.0f;
+	float BlendSpaceFilteredPositionY = 0.0f;
 	uint16 FrameCounter = 0;
 	bool bLooping = false;
 	bool bIsBlendSpace = false;
@@ -183,6 +185,9 @@ struct FBlendSpacePlayerMessage
 	float PositionX = 0.0f;
 	float PositionY = 0.0f;
 	float PositionZ = 0.0f;
+	float FilteredPositionX = 0.0f;
+	float FilteredPositionY = 0.0f;
+	float FilteredPositionZ = 0.0f;
 };
 
 struct FAnimStateMachineMessage
diff --git a/Engine/Source/Editor/AnimationBlueprintEditor/Private/AnimationNodes/SBlendSpacePreview.cpp b/Engine/Source/Editor/AnimationBlueprintEditor/Private/AnimationNodes/SBlendSpacePreview.cpp
--- a/Engine/Source/Editor/AnimationBlueprintEditor/Private/AnimationNodes/SBlendSpacePreview.cpp
+++ b/Engine/Source/Editor/AnimationBlueprintEditor/Private/AnimationNodes/SBlendSpacePreview.cpp
@@ -20,7 +20,8 @@ void SBlendSpacePreview::Construct(const FArguments& InArgs, UAnimGraphNode_Base
 	FBlendSpacePreviewArgs Args;
 
 	Args.PreviewBlendSpace = MakeAttributeLambda([this](){ return CachedBlendSpace.Get(); });
-	Args.PreviewPosition = MakeAttributeLambda([this](){ return CachedPosition;	});
+	Args.PreviewPosition = MakeAttributeLambda([this]() { return CachedPosition; });
+	Args.PreviewFilteredPosition = MakeAttributeLambda([this]() { return CachedFilteredPosition; });
 	Args.OnGetBlendSpaceSampleName = InArgs._OnGetBlendSpaceSampleName;
 
 	ChildSlot
@@ -38,7 +39,7 @@ void SBlendSpacePreview::Construct(const FArguments& InArgs, UAnimGraphNode_Base
 
 	RegisterActiveTimer(1.0f / 60.0f, FWidgetActiveTimerDelegate::CreateLambda([this](double InCurrentTime, float InDeltaTime)
 	{
-		GetBlendSpaceInfo(CachedBlendSpace, CachedPosition);
+		GetBlendSpaceInfo(CachedBlendSpace, CachedPosition, CachedFilteredPosition);
 		return EActiveTimerReturnType::Continue;
 	}));
 }
@@ -62,7 +63,7 @@ EVisibility SBlendSpacePreview::GetBlendSpaceVisibility() const
 	return EVisibility::Collapsed;
 }
 
-bool SBlendSpacePreview::GetBlendSpaceInfo(TWeakObjectPtr<const UBlendSpaceBase>& OutBlendSpace, FVector& OutPosition) const
+bool SBlendSpacePreview::GetBlendSpaceInfo(TWeakObjectPtr<const UBlendSpaceBase>& OutBlendSpace, FVector& OutPosition, FVector& OutFilteredPosition) const
 {
 	if (Node.Get() != nullptr)
 	{
@@ -78,10 +79,12 @@ bool SBlendSpacePreview::GetBlendSpaceInfo(TWeakObjectPtr<const UBlendSpaceBase>
 						// reverse node index temporarily because of a bug in NodeGuidToIndexMap
 						AnimNodeIndex = Class->GetAnimNodeProperties().Num() - AnimNodeIndex - 1;
 
-						if (FAnimBlueprintDebugData::FBlendSpacePlayerRecord* DebugInfo = Class->GetAnimBlueprintDebugData().BlendSpacePlayerRecordsThisFrame.FindByPredicate([AnimNodeIndex](const FAnimBlueprintDebugData::FBlendSpacePlayerRecord& InRecord){ return InRecord.NodeID == AnimNodeIndex; }))
+						if (FAnimBlueprintDebugData::FBlendSpacePlayerRecord* DebugInfo = Class->GetAnimBlueprintDebugData().BlendSpacePlayerRecordsThisFrame.FindByPredicate(
+							[AnimNodeIndex](const FAnimBlueprintDebugData::FBlendSpacePlayerRecord& InRecord){ return InRecord.NodeID == AnimNodeIndex; }))
 						{
 							OutBlendSpace = DebugInfo->BlendSpace.Get();
-							OutPosition = FVector(DebugInfo->PositionX, DebugInfo->PositionY, DebugInfo->PositionZ);
+							OutPosition = DebugInfo->Position;
+							OutFilteredPosition = DebugInfo->FilteredPosition;
 							return true;
 						}
 					}
@@ -92,5 +95,6 @@ bool SBlendSpacePreview::GetBlendSpaceInfo(TWeakObjectPtr<const UBlendSpaceBase>
 
 	OutBlendSpace = nullptr;
 	OutPosition = FVector::ZeroVector;
+	OutFilteredPosition = FVector::ZeroVector;
 	return false;
 }
diff --git a/Engine/Source/Editor/AnimationBlueprintEditor/Private/AnimationNodes/SBlendSpacePreview.h b/Engine/Source/Editor/AnimationBlueprintEditor/Private/AnimationNodes/SBlendSpacePreview.h
--- a/Engine/Source/Editor/AnimationBlueprintEditor/Private/AnimationNodes/SBlendSpacePreview.h
+++ b/Engine/Source/Editor/AnimationBlueprintEditor/Private/AnimationNodes/SBlendSpacePreview.h
@@ -23,9 +23,10 @@ public:
 
 protected:
 	EVisibility GetBlendSpaceVisibility() const;
-	bool GetBlendSpaceInfo(TWeakObjectPtr<const UBlendSpaceBase>& OutBlendSpace, FVector& OutPosition) const;
+	bool GetBlendSpaceInfo(TWeakObjectPtr<const UBlendSpaceBase>& OutBlendSpace, FVector& OutPosition, FVector& OutFilteredPosition) const;
 
 	TWeakObjectPtr<const UAnimGraphNode_Base> Node;
 	TWeakObjectPtr<const UBlendSpaceBase> CachedBlendSpace;
 	FVector CachedPosition;
+	FVector CachedFilteredPosition;
 };
diff --git a/Engine/Source/Editor/AnimationBlueprintEditor/Private/BlendSpaceDocumentTabFactory.cpp b/Engine/Source/Editor/AnimationBlueprintEditor/Private/BlendSpaceDocumentTabFactory.cpp
--- a/Engine/Source/Editor/AnimationBlueprintEditor/Private/BlendSpaceDocumentTabFactory.cpp
+++ b/Engine/Source/Editor/AnimationBlueprintEditor/Private/BlendSpaceDocumentTabFactory.cpp
@@ -122,9 +122,10 @@ TSharedRef<SWidget> FBlendSpaceDocumentTabFactory::CreateTabBody(const FWorkflow
 							// reverse node index temporarily because of a bug in NodeGuidToIndexMap
 							AnimNodeIndex = Class->GetAnimNodeProperties().Num() - AnimNodeIndex - 1;
 
-							if (FAnimBlueprintDebugData::FBlendSpacePlayerRecord* DebugInfo = Class->GetAnimBlueprintDebugData().BlendSpacePlayerRecordsThisFrame.FindByPredicate([AnimNodeIndex](const FAnimBlueprintDebugData::FBlendSpacePlayerRecord& InRecord){ return InRecord.NodeID == AnimNodeIndex; }))
+							if (FAnimBlueprintDebugData::FBlendSpacePlayerRecord* DebugInfo = Class->GetAnimBlueprintDebugData().BlendSpacePlayerRecordsThisFrame.FindByPredicate(
+								[AnimNodeIndex](const FAnimBlueprintDebugData::FBlendSpacePlayerRecord& InRecord){ return InRecord.NodeID == AnimNodeIndex; }))
 							{
-								return FVector(DebugInfo->PositionX, DebugInfo->PositionY, DebugInfo->PositionZ);
+								return DebugInfo->Position;
 							}
 						}
 					}
@@ -135,7 +136,37 @@ TSharedRef<SWidget> FBlendSpaceDocumentTabFactory::CreateTabBody(const FWorkflow
 		return FVector::ZeroVector;
 	});
 
-	Args.OnSetPreviewPosition = FOnSetBlendSpacePreviewPosition::CreateLambda([this, WeakBlendSpaceNode = TWeakObjectPtr<UAnimGraphNode_BlendSpaceGraphBase>(BlendSpaceNode)](FVector InSampleValue)
+	Args.PreviewFilteredPosition = MakeAttributeLambda([this, WeakBlendSpaceNode = TWeakObjectPtr<UAnimGraphNode_BlendSpaceGraphBase>(BlendSpaceNode)]()
+	{
+		if (WeakBlendSpaceNode.Get())
+		{
+			if (UBlueprint* Blueprint = FBlueprintEditorUtils::FindBlueprintForNode(WeakBlendSpaceNode.Get()))
+			{
+				if (UObject* ActiveObject = Blueprint->GetObjectBeingDebugged())
+				{
+					if (UAnimBlueprintGeneratedClass* Class = Cast<UAnimBlueprintGeneratedClass>(ActiveObject->GetClass()))
+					{
+						if (int32* NodeIndexPtr = Class->GetAnimBlueprintDebugData().NodePropertyToIndexMap.Find(WeakBlendSpaceNode))
+						{
+							int32 AnimNodeIndex = *NodeIndexPtr;
+							// reverse node index temporarily because of a bug in NodeGuidToIndexMap
+							AnimNodeIndex = Class->GetAnimNodeProperties().Num() - AnimNodeIndex - 1;
+
+							if (FAnimBlueprintDebugData::FBlendSpacePlayerRecord* DebugInfo = Class->GetAnimBlueprintDebugData().BlendSpacePlayerRecordsThisFrame.FindByPredicate(
+								[AnimNodeIndex](const FAnimBlueprintDebugData::FBlendSpacePlayerRecord& InRecord) { return InRecord.NodeID == AnimNodeIndex; }))
+							{
+								return DebugInfo->FilteredPosition;
+							}
+						}
+					}
+				}
+			}
+		}
+
+		return FVector::ZeroVector;
+	});
+
+	Args.OnSetPreviewPosition = FOnSetBlendSpacePreviewPosition::CreateLambda([this, WeakBlendSpaceNode = TWeakObjectPtr<UAnimGraphNode_BlendSpaceGraphBase>(BlendSpaceNode)](FVector InPreviewPosition)
 	{
 		if(WeakBlendSpaceNode.Get())
 		{
@@ -147,7 +178,7 @@ TSharedRef<SWidget> FBlendSpaceDocumentTabFactory::CreateTabBody(const FWorkflow
 					{
 						if(FAnimNode_BlendSpaceGraphBase* BlendSpaceGraphNode = Class->GetPropertyInstance<FAnimNode_BlendSpaceGraphBase>(ActiveObject, WeakBlendSpaceNode.Get()))
 						{
-							BlendSpaceGraphNode->SetPreviewSampleValue(InSampleValue);
+							BlendSpaceGraphNode->SetPreviewPosition(InPreviewPosition);
 						}
 					}
 				}
diff --git a/Engine/Source/Editor/Persona/Private/PersonaModule.cpp b/Engine/Source/Editor/Persona/Private/PersonaModule.cpp
--- a/Engine/Source/Editor/Persona/Private/PersonaModule.cpp
+++ b/Engine/Source/Editor/Persona/Private/PersonaModule.cpp
@@ -372,18 +372,20 @@ TSharedRef<SWidget> FPersonaModule::CreateBlendSpacePreviewWidget(const FBlendSp
 		.Cursor(EMouseCursor::Crosshairs)
 		.BlendSpaceBase(InArgs.PreviewBlendSpace)
 		.Position(InArgs.PreviewPosition)
+		.FilteredPosition(InArgs.PreviewFilteredPosition)
 		.ReadOnly(true)
 		.ShowAxisLabels(false)
 		.ShowSettingsButtons(false)
 		.OnGetBlendSpaceSampleName(InArgs.OnGetBlendSpaceSampleName);
 }
 
-TSharedRef<SWidget> FPersonaModule::CreateBlendSpacePreviewWidget(TAttribute<const UBlendSpaceBase*> InBlendSpace, TAttribute<FVector> InPosition) const
+TSharedRef<SWidget> FPersonaModule::CreateBlendSpacePreviewWidget(TAttribute<const UBlendSpaceBase*> InBlendSpace, TAttribute<FVector> InPosition, TAttribute<FVector> InFilteredPosition) const
 {
 PRAGMA_DISABLE_DEPRECATION_WARNINGS
 	FBlendSpacePreviewArgs Args;
 	Args.PreviewBlendSpace = InBlendSpace;
 	Args.PreviewPosition = InPosition;
+	Args.PreviewFilteredPosition = InFilteredPosition;
 
 	return CreateBlendSpacePreviewWidget(Args);
 PRAGMA_ENABLE_DEPRECATION_WARNINGS
@@ -1533,6 +1535,7 @@ TSharedRef<SWidget> FPersonaModule::CreateBlendSpaceEditWidget(UBlendSpaceBase*
 			.OnExtendSampleTooltip(InArgs.OnExtendSampleTooltip)
 			.OnSetPreviewPosition(InArgs.OnSetPreviewPosition)
 			.PreviewPosition(InArgs.PreviewPosition)
+			.PreviewFilteredPosition(InArgs.PreviewFilteredPosition)
 			.StatusBarName(InArgs.StatusBarName);
 	}
 	else if (UBlendSpace1D* BlendSpace1D = Cast<UBlendSpace1D>(InBlendSpace))
@@ -1548,6 +1551,7 @@ TSharedRef<SWidget> FPersonaModule::CreateBlendSpaceEditWidget(UBlendSpaceBase*
 			.OnExtendSampleTooltip(InArgs.OnExtendSampleTooltip)
 			.OnSetPreviewPosition(InArgs.OnSetPreviewPosition)
 			.PreviewPosition(InArgs.PreviewPosition)
+			.PreviewFilteredPosition(InArgs.PreviewFilteredPosition)
 			.StatusBarName(InArgs.StatusBarName);
 	}
 
diff --git a/Engine/Source/Editor/Persona/Private/SAnimationBlendSpace.cpp b/Engine/Source/Editor/Persona/Private/SAnimationBlendSpace.cpp
--- a/Engine/Source/Editor/Persona/Private/SAnimationBlendSpace.cpp
+++ b/Engine/Source/Editor/Persona/Private/SAnimationBlendSpace.cpp
@@ -21,6 +21,7 @@ void SBlendSpaceEditor::Construct(const FArguments& InArgs)
 		.OnExtendSampleTooltip(InArgs._OnExtendSampleTooltip)
 		.OnSetPreviewPosition(InArgs._OnSetPreviewPosition)
 		.PreviewPosition(InArgs._PreviewPosition)
+		.PreviewFilteredPosition(InArgs._PreviewFilteredPosition)
 		.StatusBarName(InArgs._StatusBarName));
 }
 
@@ -37,6 +38,7 @@ void SBlendSpaceEditor::Construct(const FArguments& InArgs, const TSharedRef<cla
 		.OnExtendSampleTooltip(InArgs._OnExtendSampleTooltip)
 		.OnSetPreviewPosition(InArgs._OnSetPreviewPosition)
 		.PreviewPosition(InArgs._PreviewPosition)
+		.PreviewFilteredPosition(InArgs._PreviewFilteredPosition)
 		.StatusBarName(InArgs._StatusBarName),
 		InPreviewScene);
 }
diff --git a/Engine/Source/Editor/Persona/Private/SAnimationBlendSpace.h b/Engine/Source/Editor/Persona/Private/SAnimationBlendSpace.h
--- a/Engine/Source/Editor/Persona/Private/SAnimationBlendSpace.h
+++ b/Engine/Source/Editor/Persona/Private/SAnimationBlendSpace.h
@@ -36,6 +36,8 @@ public:
 
 	SLATE_ATTRIBUTE(FVector, PreviewPosition)
 
+	SLATE_ATTRIBUTE(FVector, PreviewFilteredPosition)
+
 	SLATE_ARGUMENT(FName, StatusBarName)
 
 	SLATE_END_ARGS()
diff --git a/Engine/Source/Editor/Persona/Private/SAnimationBlendSpace1D.cpp b/Engine/Source/Editor/Persona/Private/SAnimationBlendSpace1D.cpp
--- a/Engine/Source/Editor/Persona/Private/SAnimationBlendSpace1D.cpp
+++ b/Engine/Source/Editor/Persona/Private/SAnimationBlendSpace1D.cpp
@@ -22,6 +22,7 @@ void SBlendSpaceEditor1D::Construct(const FArguments& InArgs)
 		.OnExtendSampleTooltip(InArgs._OnExtendSampleTooltip)
 		.OnSetPreviewPosition(InArgs._OnSetPreviewPosition)
 		.PreviewPosition(InArgs._PreviewPosition)
+		.PreviewFilteredPosition(InArgs._PreviewFilteredPosition)
 		.StatusBarName(InArgs._StatusBarName));
 }
 
@@ -38,6 +39,7 @@ void SBlendSpaceEditor1D::Construct(const FArguments& InArgs, const TSharedRef<I
 		.OnExtendSampleTooltip(InArgs._OnExtendSampleTooltip)
 		.OnSetPreviewPosition(InArgs._OnSetPreviewPosition)
 		.PreviewPosition(InArgs._PreviewPosition)
+		.PreviewFilteredPosition(InArgs._PreviewFilteredPosition)
 		.StatusBarName(InArgs._StatusBarName),
 		InPreviewScene);
 }
diff --git a/Engine/Source/Editor/Persona/Private/SAnimationBlendSpace1D.h b/Engine/Source/Editor/Persona/Private/SAnimationBlendSpace1D.h
--- a/Engine/Source/Editor/Persona/Private/SAnimationBlendSpace1D.h
+++ b/Engine/Source/Editor/Persona/Private/SAnimationBlendSpace1D.h
@@ -38,6 +38,8 @@ public:
 
 	SLATE_ATTRIBUTE(FVector, PreviewPosition)
 
+	SLATE_ATTRIBUTE(FVector, PreviewFilteredPosition)
+
 	SLATE_ARGUMENT(FName, StatusBarName)
 
 	SLATE_END_ARGS()
diff --git a/Engine/Source/Editor/Persona/Private/SAnimationBlendSpaceBase.cpp b/Engine/Source/Editor/Persona/Private/SAnimationBlendSpaceBase.cpp
--- a/Engine/Source/Editor/Persona/Private/SAnimationBlendSpaceBase.cpp
+++ b/Engine/Source/Editor/Persona/Private/SAnimationBlendSpaceBase.cpp
@@ -36,7 +36,7 @@ void SBlendSpaceEditorBase::Construct(const FArguments& InArgs)
 	OnBlendSpaceSampleReplaced = InArgs._OnBlendSpaceSampleReplaced;
 	OnSetPreviewPosition = InArgs._OnSetPreviewPosition;
 
-	bShouldSetPreviewValue = false;
+	bShouldSetPreviewPosition = false;
 
 	SAnimEditorBase::Construct(SAnimEditorBase::FArguments()
 		.DisplayAnimTimeline(false)
@@ -72,6 +72,7 @@ void SBlendSpaceEditorBase::Construct(const FArguments& InArgs)
 							.BlendSpaceBase(BlendSpace)
 							.NotifyHook(this)
 							.Position(InArgs._PreviewPosition)
+							.FilteredPosition(InArgs._PreviewFilteredPosition)
 							.OnSampleMoved(this, &SBlendSpaceEditorBase::OnSampleMoved)
 							.OnSampleRemoved(this, &SBlendSpaceEditorBase::OnSampleRemoved)
 							.OnSampleAdded(this, &SBlendSpaceEditorBase::OnSampleAdded)
@@ -194,7 +195,12 @@ void SBlendSpaceEditorBase::PostUndoRedo()
 	NewBlendSpaceGridWidget->InvalidateState();
 
 	// Set flag which will update the preview value in the next tick (this due the recreation of data after Undo)
-	bShouldSetPreviewValue = true;
+	bShouldSetPreviewPosition = true;
+}
+
+TSharedPtr<class IPersonaPreviewScene> SBlendSpaceEditorBase::GetPreviewScene() const
+{
+	return PreviewScenePtr.Pin();
 }
 
 void SBlendSpaceEditorBase::UpdatePreviewParameter() const
@@ -207,32 +213,48 @@ void SBlendSpaceEditorBase::UpdatePreviewParameter() const
 		{
 			if (Component->PreviewInstance->GetCurrentAsset() == BlendSpace)
 			{
-				const FVector BlendInput = NewBlendSpaceGridWidget->GetBlendPreviewValue();
-				Component->PreviewInstance->SetBlendSpaceInput(BlendInput);
+				const FVector PreviewPosition = NewBlendSpaceGridWidget->GetPreviewPosition();
+				Component->PreviewInstance->SetBlendSpacePosition(PreviewPosition);
 				GetPreviewScene()->InvalidateViews();			
 			}
 		}
 	}
 	else if(OnSetPreviewPosition.IsBound())
 	{
-		const FVector BlendInput = NewBlendSpaceGridWidget->GetBlendPreviewValue();
-		OnSetPreviewPosition.Execute(BlendInput);
+		const FVector PreviewPosition = NewBlendSpaceGridWidget->GetPreviewPosition();
+		OnSetPreviewPosition.Execute(PreviewPosition);
 	}
 }
 
-TSharedPtr<class IPersonaPreviewScene> SBlendSpaceEditorBase::GetPreviewScene() const
+void SBlendSpaceEditorBase::UpdateFromBlendSpaceState() const
 {
-	return PreviewScenePtr.Pin();
+	if (GetPreviewScene().IsValid())
+	{
+		class UDebugSkelMeshComponent* Component = GetPreviewScene()->GetPreviewMeshComponent();
+
+		if (Component != nullptr && Component->IsPreviewOn())
+		{
+			if (Component->PreviewInstance->GetCurrentAsset() == BlendSpace)
+			{
+				FVector FilteredPosition;
+				FVector Position;
+				Component->PreviewInstance->GetBlendSpaceState(Position, FilteredPosition);
+				NewBlendSpaceGridWidget->SetPreviewingState(Position, FilteredPosition);
+			}
+		}
+	}
 }
 
 void SBlendSpaceEditorBase::Tick(const FGeometry& AllottedGeometry, const double InCurrentTime, const float InDeltaTime)
 {
 	// Update the preview as long as its enabled
-	if (NewBlendSpaceGridWidget->IsPreviewing() || bShouldSetPreviewValue)
+	if (NewBlendSpaceGridWidget->IsPreviewing() || bShouldSetPreviewPosition)
 	{
 		UpdatePreviewParameter();
-		bShouldSetPreviewValue = false;
+		bShouldSetPreviewPosition = false;
 	}
+
+	UpdateFromBlendSpaceState();
 }
 
 void SBlendSpaceEditorBase::OnPropertyChanged(UObject* ObjectBeingModified, FPropertyChangedEvent& PropertyChangedEvent)
diff --git a/Engine/Source/Editor/Persona/Private/SAnimationBlendSpaceBase.h b/Engine/Source/Editor/Persona/Private/SAnimationBlendSpaceBase.h
--- a/Engine/Source/Editor/Persona/Private/SAnimationBlendSpaceBase.h
+++ b/Engine/Source/Editor/Persona/Private/SAnimationBlendSpaceBase.h
@@ -44,6 +44,8 @@ public:
 
 	SLATE_ATTRIBUTE(FVector, PreviewPosition)
 
+	SLATE_ATTRIBUTE(FVector, PreviewFilteredPosition)
+
 	SLATE_ARGUMENT(FName, StatusBarName)
 
 	SLATE_END_ARGS()
@@ -85,6 +87,9 @@ protected:
 	/** Updates Persona's preview window */
 	void UpdatePreviewParameter() const;
 	
+	/** Updates Persona's preview window with additional data that should always be shown. */
+	void UpdateFromBlendSpaceState() const;
+
 	/** Retrieves the preview scene shown by Persona */
 	TSharedPtr<class IPersonaPreviewScene> GetPreviewScene() const;
 
@@ -119,5 +124,5 @@ protected:
 	FOnSetBlendSpacePreviewPosition OnSetPreviewPosition;
 
 	/** Flag to check whether or not the preview value should be (re-)set on the next tick */
-	bool bShouldSetPreviewValue;
+	bool bShouldSetPreviewPosition;
 };
diff --git a/Engine/Source/Editor/Persona/Private/SAnimationBlendSpaceGridWidget.cpp b/Engine/Source/Editor/Persona/Private/SAnimationBlendSpaceGridWidget.cpp
--- a/Engine/Source/Editor/Persona/Private/SAnimationBlendSpaceGridWidget.cpp
+++ b/Engine/Source/Editor/Persona/Private/SAnimationBlendSpaceGridWidget.cpp
@@ -41,6 +41,7 @@ void SBlendSpaceGridWidget::Construct(const FArguments& InArgs)
 	BlendSpaceBase = InArgs._BlendSpaceBase;
 	PreviousBlendSpaceBase = BlendSpaceBase.Get();
 	Position = InArgs._Position;
+	FilteredPosition = InArgs._FilteredPosition;
 	NotifyHook = InArgs._NotifyHook;
 	OnSampleAdded = InArgs._OnSampleAdded;
 	OnSampleMoved = InArgs._OnSampleMoved;
@@ -63,9 +64,11 @@ void SBlendSpaceGridWidget::Construct(const FArguments& InArgs)
 	bPreviewPositionSet = true;
 	bHighlightPreviewPin = false;
 	// Initialize preview value to center or the grid
-	LastPreviewingSampleValue.X = BlendSpaceBase.Get() != nullptr ? (BlendSpaceBase.Get()->GetBlendParameter(0).GetRange() * .5f) + BlendSpaceBase.Get()->GetBlendParameter(0).Min : 0.0f;
-	LastPreviewingSampleValue.Y = BlendSpaceBase.Get() != nullptr ? (GridType == EGridType::TwoAxis ? (BlendSpaceBase.Get()->GetBlendParameter(1).GetRange() * .5f) + BlendSpaceBase.Get()->GetBlendParameter(1).Min : 0.0f) : 0.0f;
-	LastPreviewingSampleValue.Z = 0.0f;
+	PreviewPosition.X = BlendSpaceBase.Get() != nullptr ? (BlendSpaceBase.Get()->GetBlendParameter(0).GetRange() * .5f) + BlendSpaceBase.Get()->GetBlendParameter(0).Min : 0.0f;
+	PreviewPosition.Y = BlendSpaceBase.Get() != nullptr ? (GridType == EGridType::TwoAxis ? (BlendSpaceBase.Get()->GetBlendParameter(1).GetRange() * .5f) + BlendSpaceBase.Get()->GetBlendParameter(1).Min : 0.0f) : 0.0f;
+	PreviewPosition.Z = 0.0f;
+
+	PreviewFilteredPosition = PreviewPosition;
 
 	bShowTriangulation = false;
 	bMouseIsOverGeometry = false;
@@ -433,10 +436,16 @@ void SBlendSpaceGridWidget::PaintSampleKeys(const FGeometry& AllottedGeometry, c
 			FSlateDrawElement::MakeBox( OutDrawElements, DrawLayerId + 1, AllottedGeometry.ToPaintGeometry(GridPosition, KeySize), KeyBrush, ESlateDrawEffect::None, DrawColor );
 		}
 
+		// Always draw the filtered position which comes back from whatever is running
+		{
+			FVector2D GridPosition = SampleValueToGridPosition(PreviewFilteredPosition) - (KeySize * .5f);
+			FSlateDrawElement::MakeBox(OutDrawElements, DrawLayerId + 1, AllottedGeometry.ToPaintGeometry(GridPosition, KeySize), KeyBrush, ESlateDrawEffect::None, PreviewKeyColor.GetSpecifiedColor() * 0.5);
+		}
+
 		if (bPreviewPositionSet)
 		{
-			const FVector2D MouseGridPosition = SampleValueToGridPosition(LastPreviewingSampleValue) - (KeySize * .5f);
-			FSlateDrawElement::MakeBox( OutDrawElements, DrawLayerId + 1, AllottedGeometry.ToPaintGeometry(MouseGridPosition, KeySize), KeyBrush, ESlateDrawEffect::None, PreviewKeyColor.GetSpecifiedColor() );
+			FVector2D GridPosition = SampleValueToGridPosition(PreviewPosition) - (KeySize * .5f);
+			FSlateDrawElement::MakeBox(OutDrawElements, DrawLayerId + 1, AllottedGeometry.ToPaintGeometry(GridPosition, KeySize), KeyBrush, ESlateDrawEffect::None, PreviewKeyColor.GetSpecifiedColor());
 		}
 
 		if (DragState == EDragState::DragDrop || DragState == EDragState::InvalidDragDrop)
@@ -1326,7 +1335,8 @@ void SBlendSpaceGridWidget::StartPreviewing()
 	LastPreviewingMousePosition = LocalMousePosition;
 	FModifierKeysState ModifierKeyState = FSlateApplication::Get().GetModifierKeys();
 	bool bIsManualPreviewing = (ModifierKeyState.IsLeftShiftDown() || ModifierKeyState.IsRightShiftDown()) || DragState == EDragState::DragPreview;
-	LastPreviewingSampleValue = Position.IsSet() && !bIsManualPreviewing ? Position.Get() : GridPositionToSampleValue(LastPreviewingMousePosition, false);
+	PreviewPosition = Position.IsSet() && !bIsManualPreviewing ? Position.Get() : GridPositionToSampleValue(LastPreviewingMousePosition, false);
+	PreviewFilteredPosition = FilteredPosition.IsSet() ? FilteredPosition.Get() : PreviewPosition;
 	bPreviewPositionSet = true;	
 	bPreviewToolTipHidden = true;
 }
@@ -1454,7 +1464,7 @@ FText SBlendSpaceGridWidget::GetToolTipSampleValue() const
 
 		if(bReadOnly)
 		{
-			ToolTipText = FText::Format(ValueFormattingText, ParameterXName, FText::FromString(FString::SanitizeFloat(LastPreviewingSampleValue.X)), ParameterYName, FText::FromString(FString::SanitizeFloat(LastPreviewingSampleValue.Y)));
+			ToolTipText = FText::Format(ValueFormattingText, ParameterXName, FText::FromString(FString::SanitizeFloat(PreviewPosition.X)), ParameterYName, FText::FromString(FString::SanitizeFloat(PreviewPosition.Y)));
 
 			AddAdvancedPreview();
 		}
@@ -1467,7 +1477,7 @@ FText SBlendSpaceGridWidget::GetToolTipSampleValue() const
 				{		
 					if (bHighlightPreviewPin)
 					{
-						ToolTipText = FText::Format(ValueFormattingText, ParameterXName, FText::FromString(FString::SanitizeFloat(LastPreviewingSampleValue.X)), ParameterYName, FText::FromString(FString::SanitizeFloat(LastPreviewingSampleValue.Y)));
+						ToolTipText = FText::Format(ValueFormattingText, ParameterXName, FText::FromString(FString::SanitizeFloat(PreviewPosition.X)), ParameterYName, FText::FromString(FString::SanitizeFloat(PreviewPosition.Y)));
 
 						AddAdvancedPreview();
 					}
@@ -1487,7 +1497,7 @@ FText SBlendSpaceGridWidget::GetToolTipSampleValue() const
 					}
 					else if(Position.IsSet())
 					{
-						ToolTipText = FText::Format(ValueFormattingText, ParameterXName, FText::FromString(FString::SanitizeFloat(LastPreviewingSampleValue.X)), ParameterYName, FText::FromString(FString::SanitizeFloat(LastPreviewingSampleValue.Y)));
+						ToolTipText = FText::Format(ValueFormattingText, ParameterXName, FText::FromString(FString::SanitizeFloat(PreviewPosition.X)), ParameterYName, FText::FromString(FString::SanitizeFloat(PreviewPosition.Y)));
 
 						AddAdvancedPreview();
 					}
@@ -1545,7 +1555,7 @@ FText SBlendSpaceGridWidget::GetToolTipSampleValue() const
 				case EDragState::DragPreview:
 				case EDragState::Preview:
 				{
-					ToolTipText = FText::Format(ValueFormattingText, ParameterXName, FText::FromString(FString::SanitizeFloat(LastPreviewingSampleValue.X)), ParameterYName, FText::FromString(FString::SanitizeFloat(LastPreviewingSampleValue.Y)));
+					ToolTipText = FText::Format(ValueFormattingText, ParameterXName, FText::FromString(FString::SanitizeFloat(PreviewPosition.X)), ParameterYName, FText::FromString(FString::SanitizeFloat(PreviewPosition.Y)));
 
 					AddAdvancedPreview();
 					break;
@@ -1892,7 +1902,7 @@ void SBlendSpaceGridWidget::Tick(const FGeometry& AllottedGeometry, const double
 			{
 				// Check if we are highlighting preview pin
 				float Distance;
-				bHighlightPreviewPin = IsSampleValueWithinMouseRange(LastPreviewingSampleValue, Distance);
+				bHighlightPreviewPin = IsSampleValueWithinMouseRange(PreviewPosition, Distance);
 				if (bHighlightPreviewPin)
 				{
 					if (bHighlightPreviewPin != bPreviousHighlightPreviewPin)
@@ -1993,18 +2003,22 @@ void SBlendSpaceGridWidget::Tick(const FGeometry& AllottedGeometry, const double
 		// Check if we should update the preview sample value
 		if (bSamplePreviewing)
 		{
-			// Ensure the preview mouse position is clamped to the grid
-			LastPreviewingMousePosition.X = LocalMousePosition.X; //FMath::Clamp(LocalMousePosition.X, CachedGridRectangle.Left, CachedGridRectangle.Right);
-			LastPreviewingMousePosition.Y = LocalMousePosition.Y; //FMath::Clamp(LocalMousePosition.Y, CachedGridRectangle.Top, CachedGridRectangle.Bottom);
+			// Clamping happens later
+			LastPreviewingMousePosition.X = LocalMousePosition.X;
+			LastPreviewingMousePosition.Y = LocalMousePosition.Y;
 			FModifierKeysState ModifierKeyState = FSlateApplication::Get().GetModifierKeys();
 			bool bIsManualPreviewing = (ModifierKeyState.IsLeftShiftDown() || ModifierKeyState.IsRightShiftDown()) || DragState == EDragState::DragPreview;
-			LastPreviewingSampleValue = Position.IsSet() && !bIsManualPreviewing ? Position.Get() : GridPositionToSampleValue(LastPreviewingMousePosition, false);
-						
-			LastPreviewingSampleValue = BlendSpace->GetClampedAndWrappedBlendInput(LastPreviewingSampleValue);
-					
+			PreviewPosition = Position.IsSet() && !bIsManualPreviewing ? Position.Get() : GridPositionToSampleValue(LastPreviewingMousePosition, false);
+			PreviewPosition = BlendSpace->GetClampedAndWrappedBlendInput(PreviewPosition);
+
+			if (FilteredPosition.IsSet())
+			{
+				PreviewFilteredPosition = BlendSpace->GetClampedAndWrappedBlendInput(FilteredPosition.Get());
+			}
+
 			// Retrieve and cache weighted samples
 			PreviewedSamples.Empty(4);
-			BlendSpace->GetSamplesFromBlendInput(GetBlendPreviewValue(), PreviewedSamples);
+			BlendSpace->GetSamplesFromBlendInput(PreviewPosition, PreviewedSamples);
 		}
 	}
 
@@ -2022,9 +2036,23 @@ void SBlendSpaceGridWidget::Tick(const FGeometry& AllottedGeometry, const double
 	CalculateGridPoints();
 }
 
-const FVector SBlendSpaceGridWidget::GetBlendPreviewValue()
+const FVector SBlendSpaceGridWidget::GetPreviewPosition() const
 {	
-	return LastPreviewingSampleValue;
+	return PreviewPosition;
+}
+
+void SBlendSpaceGridWidget::SetPreviewingState(const FVector& InPosition, const FVector& InFilteredPosition)
+{
+	if (const UBlendSpaceBase* BlendSpace = BlendSpaceBase.Get())
+	{
+		PreviewFilteredPosition = BlendSpace->GetClampedAndWrappedBlendInput(InFilteredPosition);
+		PreviewPosition = BlendSpace->GetClampedAndWrappedBlendInput(InPosition);
+	}
+	else
+	{
+		PreviewFilteredPosition = InFilteredPosition;
+		PreviewPosition = InPosition;
+	}
 }
 
 void SBlendSpaceGridWidget::InvalidateCachedData()
diff --git a/Engine/Source/Editor/Persona/Private/SAnimationBlendSpaceGridWidget.h b/Engine/Source/Editor/Persona/Private/SAnimationBlendSpaceGridWidget.h
--- a/Engine/Source/Editor/Persona/Private/SAnimationBlendSpaceGridWidget.h
+++ b/Engine/Source/Editor/Persona/Private/SAnimationBlendSpaceGridWidget.h
@@ -38,6 +38,7 @@ public:
 	{}
 		SLATE_ATTRIBUTE(const UBlendSpaceBase*, BlendSpaceBase)
 		SLATE_ATTRIBUTE(FVector, Position)
+		SLATE_ATTRIBUTE(FVector, FilteredPosition)
 		SLATE_ARGUMENT(FNotifyHook*, NotifyHook)
 		SLATE_ARGUMENT(bool, ReadOnly)
 		SLATE_ARGUMENT(bool, ShowAxisLabels)
@@ -111,7 +112,11 @@ public:
 	// End of SWidget
 
 	/** Returns the sample value for previewing the blend space */
-	const FVector GetBlendPreviewValue();
+	const FVector GetPreviewPosition() const;
+
+	/** Records the state of the blend space being used for previewing so it can be displayed */
+	void SetPreviewingState(const FVector& InPosition, const FVector& InFilteredPosition);
+
 	/** Flag whether or not the user is actively previewing the blend space (moving the sample value) */
 	const bool IsPreviewing() const;
 
@@ -221,9 +226,11 @@ private:
 	TAttribute<const UBlendSpaceBase*> BlendSpaceBase;
 	/** Previous blendspace, cached here to allow us to update cached data when it changes. */
 	TWeakObjectPtr<const UBlendSpaceBase> PreviousBlendSpaceBase;
-	/** Current position (driven from outside the widget) */
+	/** Current BlendSpace params (driven from outside the widget) */
 	TAttribute<FVector> Position;
-	/** Notfiy hook (ptr to SAnimationBlendSpace instance), which is required for transacting FBlendSample object when edited using the context-menu/structure details panel */
+	/** Current BlendSpace state - filtered version of Position (driven from outside the widget) */
+	TAttribute<FVector> FilteredPosition;
+	/** Notify hook (ptr to SAnimationBlendSpace instance), which is required for transacting FBlendSample object when edited using the context-menu/structure details panel */
 	FNotifyHook* NotifyHook;
 	/** Number of Blend Parameters to draw */
 	uint32 BlendParametersToDraw;
@@ -250,11 +257,14 @@ private:
 	/** Currently set preview blend sample value and state data */
 	bool bSamplePreviewing;
 	FVector2D LastPreviewingMousePosition;
-	FVector LastPreviewingSampleValue;
+	FVector PreviewPosition;
 	bool bPreviewPositionSet;
 	bool bAdvancedPreview;
 	TArray<FBlendSampleData> PreviewedSamples;
 
+	/** Preview state that may have been returned */
+	FVector PreviewFilteredPosition;
+
 	/** Tooltip ptr which is shown when hovering/dropping/dragging a sample*/
 	TSharedPtr<SToolTip> ToolTip;
 
diff --git a/Engine/Source/Editor/Persona/Public/PersonaModule.h b/Engine/Source/Editor/Persona/Public/PersonaModule.h
--- a/Engine/Source/Editor/Persona/Public/PersonaModule.h
+++ b/Engine/Source/Editor/Persona/Public/PersonaModule.h
@@ -117,9 +117,12 @@ struct FBlendSpaceEditorArgs
 	// Called to get the overridden name of a blend sample
 	FOnGetBlendSpaceSampleName OnGetBlendSpaceSampleName;
 
-	// Allows the preview position to be programmatically driven
+	// Allows the target preview position to be programmatically driven
 	TAttribute<FVector> PreviewPosition;
 
+	// Allows the current position to be programmatically driven
+	TAttribute<FVector> PreviewFilteredPosition;
+
 	// Allows an external widget to be inserted into a sample's tooltip
 	FOnExtendBlendSpaceSampleTooltip OnExtendSampleTooltip;
 
@@ -134,9 +137,12 @@ struct FBlendSpacePreviewArgs
 {
 	TAttribute<const UBlendSpaceBase*> PreviewBlendSpace;
 
-	// Allows the preview position to be programmatically driven
+	// Allows the target preview position to be programatically driven
 	TAttribute<FVector> PreviewPosition;
 
+	// Allows the current preview position to be programatically driven
+	TAttribute<FVector> PreviewFilteredPosition;
+
 	// Called to get the overridden name of a blend sample
 	FOnGetBlendSpaceSampleName OnGetBlendSpaceSampleName;
 };
@@ -294,7 +300,7 @@ public:
 
 	/** Deprecated */
 	UE_DEPRECATED(5.0, "Please use the overload that takes a FBlendSpacePreviewArgs struct")
-	virtual TSharedRef<SWidget> CreateBlendSpacePreviewWidget(TAttribute<const UBlendSpaceBase*> InBlendSpace, TAttribute<FVector> InPosition) const;
+	virtual TSharedRef<SWidget> CreateBlendSpacePreviewWidget(TAttribute<const UBlendSpaceBase*> InBlendSpace, TAttribute<FVector> InBlendPosition, TAttribute<FVector> InFilteredBlendPosition) const;
 
 	/** Create a widget to preview a blendspace */
 	virtual TSharedRef<SWidget> CreateBlendSpacePreviewWidget(const FBlendSpacePreviewArgs& InArgs) const;
diff --git a/Engine/Source/Runtime/AnimGraphRuntime/Private/AnimGraphRuntimeTrace.cpp b/Engine/Source/Runtime/AnimGraphRuntime/Private/AnimGraphRuntimeTrace.cpp
--- a/Engine/Source/Runtime/AnimGraphRuntime/Private/AnimGraphRuntimeTrace.cpp
+++ b/Engine/Source/Runtime/AnimGraphRuntime/Private/AnimGraphRuntimeTrace.cpp
@@ -18,7 +18,10 @@ UE_TRACE_EVENT_BEGIN(Animation, BlendSpacePlayer)
 	UE_TRACE_EVENT_FIELD(float, PositionX)
 	UE_TRACE_EVENT_FIELD(float, PositionY)
 	UE_TRACE_EVENT_FIELD(float, PositionZ)
-UE_TRACE_EVENT_END()
+	UE_TRACE_EVENT_FIELD(float, FilteredPositionX)
+	UE_TRACE_EVENT_FIELD(float, FilteredPositionY)
+	UE_TRACE_EVENT_FIELD(float, FilteredPositionZ)
+	UE_TRACE_EVENT_END()
 
 void FAnimGraphRuntimeTrace::OutputBlendSpacePlayer(const FAnimationBaseContext& InContext, const FAnimNode_BlendSpacePlayer& InNode)
 {
@@ -40,7 +43,10 @@ void FAnimGraphRuntimeTrace::OutputBlendSpacePlayer(const FAnimationBaseContext&
 		<< BlendSpacePlayer.NodeId(InContext.GetCurrentNodeId())
 		<< BlendSpacePlayer.PositionX(InNode.X)
 		<< BlendSpacePlayer.PositionY(InNode.Y)
-		<< BlendSpacePlayer.PositionZ(InNode.Z);
+		<< BlendSpacePlayer.PositionZ(InNode.Z)
+		<< BlendSpacePlayer.FilteredPositionX(InNode.GetFilteredPosition().X)
+		<< BlendSpacePlayer.FilteredPositionY(InNode.GetFilteredPosition().Y)
+		<< BlendSpacePlayer.FilteredPositionZ(InNode.GetFilteredPosition().Z);
 }
 
 void FAnimGraphRuntimeTrace::OutputBlendSpace(const FAnimationBaseContext& InContext, const FAnimNode_BlendSpaceGraphBase& InNode)
@@ -56,7 +62,8 @@ void FAnimGraphRuntimeTrace::OutputBlendSpace(const FAnimationBaseContext& InCon
 	TRACE_OBJECT(InContext.AnimInstanceProxy->GetAnimInstanceObject());
 	TRACE_OBJECT(InNode.GetBlendSpace());
 
-	FVector Coordinates = InNode.GetSampleCoordinates();
+	FVector Coordinates = InNode.GetPosition();
+	FVector FilteredCoordinates = InNode.GetFilteredPosition();
 
 	UE_TRACE_LOG(Animation, BlendSpacePlayer, AnimationChannel)
 		<< BlendSpacePlayer.Cycle(FPlatformTime::Cycles64())
@@ -65,7 +72,10 @@ void FAnimGraphRuntimeTrace::OutputBlendSpace(const FAnimationBaseContext& InCon
 		<< BlendSpacePlayer.NodeId(InContext.GetCurrentNodeId())
 		<< BlendSpacePlayer.PositionX(Coordinates.X)
 		<< BlendSpacePlayer.PositionY(Coordinates.Y)
-		<< BlendSpacePlayer.PositionZ(Coordinates.Z);
+		<< BlendSpacePlayer.PositionZ(Coordinates.Z)
+		<< BlendSpacePlayer.FilteredPositionX(FilteredCoordinates.X)
+		<< BlendSpacePlayer.FilteredPositionY(FilteredCoordinates.Y)
+		<< BlendSpacePlayer.FilteredPositionZ(FilteredCoordinates.Z);
 }
 
 #endif
\ No newline at end of file
diff --git a/Engine/Source/Runtime/AnimGraphRuntime/Private/AnimNodes/AnimNode_BlendSpaceGraphBase.cpp b/Engine/Source/Runtime/AnimGraphRuntime/Private/AnimNodes/AnimNode_BlendSpaceGraphBase.cpp
--- a/Engine/Source/Runtime/AnimGraphRuntime/Private/AnimNodes/AnimNode_BlendSpaceGraphBase.cpp
+++ b/Engine/Source/Runtime/AnimGraphRuntime/Private/AnimNodes/AnimNode_BlendSpaceGraphBase.cpp
@@ -38,18 +38,18 @@ void FAnimNode_BlendSpaceGraphBase::UpdateInternal(const FAnimationUpdateContext
 	check(BlendSpace != nullptr);
 
 	// Filter input and update blend samples
-	FVector BlendInput(X, Y, Z);
+	FVector BlendParams(X, Y, Z);
 #if WITH_EDITORONLY_DATA
-	if(bUsePreviewSampleValue)
+	if(bUsePreviewPosition)
 	{
 		// Consume any preview sample we have set
-		BlendInput = PreviewSample;
-		bUsePreviewSampleValue = false;
+		BlendParams = PreviewPosition;
+		bUsePreviewPosition = false;
 	}
 #endif
 	const float DeltaTime = Context.GetDeltaTime();
-	const FVector FilteredInput = BlendSpace->FilterInput(&BlendFilter, BlendInput, DeltaTime);
-	BlendSpace->UpdateBlendSamples(FilteredInput, DeltaTime, BlendSampleDataCache);
+	const FVector FilteredBlendParams = BlendSpace->FilterInput(&BlendFilter, BlendParams, DeltaTime);
+	BlendSpace->UpdateBlendSamples(FilteredBlendParams, DeltaTime, BlendSampleDataCache);
 
 	for(int32 SampleIndex = 0; SampleIndex < BlendSampleDataCache.Num(); ++SampleIndex)
 	{
@@ -62,7 +62,7 @@ void FAnimNode_BlendSpaceGraphBase::UpdateInternal(const FAnimationUpdateContext
 #if WITH_EDITORONLY_DATA
 	if (FAnimBlueprintDebugData* DebugData = Context.AnimInstanceProxy->GetAnimBlueprintDebugData())
 	{
-		DebugData->RecordBlendSpacePlayer(Context.GetCurrentNodeId(), BlendSpace, BlendInput.X, BlendInput.Y, BlendInput.Z);
+		DebugData->RecordBlendSpacePlayer(Context.GetCurrentNodeId(), BlendSpace, BlendParams, FilteredBlendParams);
 	}
 #endif
 
@@ -103,9 +103,9 @@ void FAnimNode_BlendSpaceGraphBase::GatherDebugData(FNodeDebugData& DebugData)
 }
 
 #if WITH_EDITORONLY_DATA
-void FAnimNode_BlendSpaceGraphBase::SetPreviewSampleValue(FVector InVector)
+void FAnimNode_BlendSpaceGraphBase::SetPreviewPosition(FVector InVector)
 {
-	bUsePreviewSampleValue = true;
-	PreviewSample = InVector;
+	bUsePreviewPosition = true;
+	PreviewPosition = InVector;
 }
 #endif
\ No newline at end of file
diff --git a/Engine/Source/Runtime/AnimGraphRuntime/Private/AnimNodes/AnimNode_BlendSpacePlayer.cpp b/Engine/Source/Runtime/AnimGraphRuntime/Private/AnimNodes/AnimNode_BlendSpacePlayer.cpp
--- a/Engine/Source/Runtime/AnimGraphRuntime/Private/AnimNodes/AnimNode_BlendSpacePlayer.cpp
+++ b/Engine/Source/Runtime/AnimGraphRuntime/Private/AnimNodes/AnimNode_BlendSpacePlayer.cpp
@@ -90,12 +90,12 @@ void FAnimNode_BlendSpacePlayer::UpdateInternal(const FAnimationUpdateContext& C
 			Reinitialize(bResetPlayTimeWhenBlendSpaceChanges);
 		}
 
-		const FVector BlendInput(X, Y, Z);
+		const FVector Position(X, Y, Z);
 
 		// Create a tick record and push into the closest scope
 		UE::Anim::FAnimSyncGroupScope& SyncScope = Context.GetMessageChecked<UE::Anim::FAnimSyncGroupScope>();
 
-		FAnimTickRecord TickRecord(BlendSpace, BlendInput, BlendSampleDataCache, BlendFilter, bLoop, PlayRate, Context.GetFinalBlendWeight(), /*inout*/ InternalTimeAccumulator, MarkerTickRecord);
+		FAnimTickRecord TickRecord(BlendSpace, Position, BlendSampleDataCache, BlendFilter, bLoop, PlayRate, Context.GetFinalBlendWeight(), /*inout*/ InternalTimeAccumulator, MarkerTickRecord);
 		TickRecord.RootMotionWeightModifier = Context.GetRootMotionWeightModifier();
 		UE::Anim::FAnimSyncParams SyncParams(GroupName, GroupRole, Method);
 
@@ -106,7 +106,7 @@ void FAnimNode_BlendSpacePlayer::UpdateInternal(const FAnimationUpdateContext& C
 #if WITH_EDITORONLY_DATA
 		if (FAnimBlueprintDebugData* DebugData = Context.AnimInstanceProxy->GetAnimBlueprintDebugData())
 		{
-			DebugData->RecordBlendSpacePlayer(Context.GetCurrentNodeId(), BlendSpace, BlendInput.X, BlendInput.Y, BlendInput.Z);
+			DebugData->RecordBlendSpacePlayer(Context.GetCurrentNodeId(), BlendSpace, Position, BlendFilter.GetFilterLastOutput());
 		}
 #endif
 
diff --git a/Engine/Source/Runtime/AnimGraphRuntime/Public/AnimNodes/AnimNode_BlendSpaceGraphBase.h b/Engine/Source/Runtime/AnimGraphRuntime/Public/AnimNodes/AnimNode_BlendSpaceGraphBase.h
--- a/Engine/Source/Runtime/AnimGraphRuntime/Public/AnimNodes/AnimNode_BlendSpaceGraphBase.h
+++ b/Engine/Source/Runtime/AnimGraphRuntime/Public/AnimNodes/AnimNode_BlendSpaceGraphBase.h
@@ -22,11 +22,14 @@ struct ANIMGRAPHRUNTIME_API FAnimNode_BlendSpaceGraphBase : public FAnimNode_Bas
 	const UBlendSpaceBase* GetBlendSpace() const { return BlendSpace; }
 
 	// @return the current sample coordinates that this node is using to sample the blendspace
-	FVector GetSampleCoordinates() const { return FVector(X, Y, Z); }
+	FVector GetPosition() const { return FVector(X, Y, Z); }
+
+	// @return the current sample coordinates after going through the filtering
+	FVector GetFilteredPosition() const { return BlendFilter.GetFilterLastOutput(); }
 
 #if WITH_EDITORONLY_DATA
 	// Set the node to preview a supplied sample value
-	void SetPreviewSampleValue(FVector InVector);
+	void SetPreviewPosition(FVector InVector);
 #endif
 
 protected:
@@ -66,11 +69,11 @@ protected:
 	TArray<FBlendSampleData> BlendSampleDataCache;
 
 #if WITH_EDITORONLY_DATA
-	// Preview sample value - set in editor only
-	FVector PreviewSample = FVector::ZeroVector;
+	// Preview blend params - set in editor only
+	FVector PreviewPosition = FVector::ZeroVector;
 
-	// Whether to use the preview sample value
-	bool bUsePreviewSampleValue = false;
+	// Whether to use the preview blend params
+	bool bUsePreviewPosition = false;
 #endif
 
 	// Internal update handler, skipping evaluation of exposed inputs
diff --git a/Engine/Source/Runtime/AnimGraphRuntime/Public/AnimNodes/AnimNode_BlendSpacePlayer.h b/Engine/Source/Runtime/AnimGraphRuntime/Public/AnimNodes/AnimNode_BlendSpacePlayer.h
--- a/Engine/Source/Runtime/AnimGraphRuntime/Public/AnimNodes/AnimNode_BlendSpacePlayer.h
+++ b/Engine/Source/Runtime/AnimGraphRuntime/Public/AnimNodes/AnimNode_BlendSpacePlayer.h
@@ -16,6 +16,12 @@ struct ANIMGRAPHRUNTIME_API FAnimNode_BlendSpacePlayer : public FAnimNode_AssetP
 {
 	GENERATED_USTRUCT_BODY()
 
+	// @return the current sample coordinates that this node is using to sample the blendspace
+	FVector GetPosition() const { return FVector(X, Y, Z); }
+
+	// @return the current sample coordinates after going through the filtering
+	FVector GetFilteredPosition() const { return BlendFilter.GetFilterLastOutput(); }
+
 public:
 	// The X coordinate to sample in the blendspace
 	UPROPERTY(EditAnywhere, BlueprintReadWrite, Category=Coordinates, meta=(PinShownByDefault))
diff --git a/Engine/Source/Runtime/Engine/Classes/Animation/AnimBlueprintGeneratedClass.h b/Engine/Source/Runtime/Engine/Classes/Animation/AnimBlueprintGeneratedClass.h
--- a/Engine/Source/Runtime/Engine/Classes/Animation/AnimBlueprintGeneratedClass.h
+++ b/Engine/Source/Runtime/Engine/Classes/Animation/AnimBlueprintGeneratedClass.h
@@ -232,19 +232,17 @@ public:
 	// Record of a blend space player's state
 	struct FBlendSpacePlayerRecord
 	{
-		FBlendSpacePlayerRecord(int32 InNodeID, const UBlendSpaceBase* InBlendSpace, float InPositionX, float InPositionY, float InPositionZ)
+		FBlendSpacePlayerRecord(int32 InNodeID, const UBlendSpaceBase* InBlendSpace, const FVector& InPosition, const FVector& InFilteredPosition)
 			: NodeID(InNodeID)
 			, BlendSpace(InBlendSpace)
-			, PositionX(InPositionX)
-			, PositionY(InPositionY) 
-			, PositionZ(InPositionY)
+			, Position(InPosition)
+			, FilteredPosition(InFilteredPosition)
 		{}
 
 		int32 NodeID;
 		TWeakObjectPtr<const UBlendSpaceBase> BlendSpace;
-		float PositionX;
-		float PositionY;
-		float PositionZ;
+		FVector Position;
+		FVector FilteredPosition;
 	};
 
 	// All blend space player records this frame
@@ -286,7 +284,7 @@ public:
 	void RecordStateData(int32 StateMachineIndex, int32 StateIndex, float Weight, float ElapsedTime);
 	void RecordNodeValue(int32 InNodeID, const FString& InText);
 	void RecordSequencePlayer(int32 InNodeID, float InPosition, float InLength, int32 InFrameCount);
-	void RecordBlendSpacePlayer(int32 InNodeID, const UBlendSpaceBase* InBlendSpace, float InPositionX, float InPositionY, float InPositionZ);
+	void RecordBlendSpacePlayer(int32 InNodeID, const UBlendSpaceBase* InBlendSpace, const FVector& InPosition, const FVector& InFilteredPosition);
 
 	void AddPoseWatch(int32 NodeID, FColor Color);
 	void RemovePoseWatch(int32 NodeID);
diff --git a/Engine/Source/Runtime/Engine/Classes/Animation/AnimSingleNodeInstance.h b/Engine/Source/Runtime/Engine/Classes/Animation/AnimSingleNodeInstance.h
--- a/Engine/Source/Runtime/Engine/Classes/Animation/AnimSingleNodeInstance.h
+++ b/Engine/Source/Runtime/Engine/Classes/Animation/AnimSingleNodeInstance.h
@@ -55,7 +55,7 @@ public:
 	UFUNCTION(BlueprintCallable, Category="Animation")
 	void SetPositionWithPreviousTime(float InPosition, float InPreviousTime, bool bFireNotifies=true);
 	UFUNCTION(BlueprintCallable, Category="Animation")
-	void SetBlendSpaceInput(const FVector& InBlendInput);
+	void SetBlendSpacePosition(const FVector& InPosition);
 	UFUNCTION(BlueprintCallable, Category="Animation")
 	void SetPlaying(bool bIsPlaying);
 	UFUNCTION(BlueprintCallable, Category="Animation")
@@ -75,6 +75,9 @@ public:
  	UFUNCTION(BlueprintCallable, Category = "Animation")
  	void SetPreviewCurveOverride(const FName& PoseName, float Value, bool bRemoveIfZero);
 public:
+	/** Gets the current state of any BlendSpace */
+	void GetBlendSpaceState(FVector& OutPosition, FVector& OutFilteredPosition) const;
+
 	/** AnimSequence specific **/
 	void StepForward();
 	void StepBackward();
diff --git a/Engine/Source/Runtime/Engine/Classes/Animation/AnimationAsset.h b/Engine/Source/Runtime/Engine/Classes/Animation/AnimationAsset.h
--- a/Engine/Source/Runtime/Engine/Classes/Animation/AnimationAsset.h
+++ b/Engine/Source/Runtime/Engine/Classes/Animation/AnimationAsset.h
@@ -142,7 +142,7 @@ struct FBlendFilter
 	{
 	}
 
-	FVector GetFilterLastOutput()
+	FVector GetFilterLastOutput() const
 	{
 		return FVector (FilterPerAxis[0].LastOutput, FilterPerAxis[1].LastOutput, FilterPerAxis[2].LastOutput);
 	}
diff --git a/Engine/Source/Runtime/Engine/Private/Animation/AnimBlueprintGeneratedClass.cpp b/Engine/Source/Runtime/Engine/Private/Animation/AnimBlueprintGeneratedClass.cpp
--- a/Engine/Source/Runtime/Engine/Private/Animation/AnimBlueprintGeneratedClass.cpp
+++ b/Engine/Source/Runtime/Engine/Private/Animation/AnimBlueprintGeneratedClass.cpp
@@ -224,9 +224,9 @@ void FAnimBlueprintDebugData::RecordSequencePlayer(int32 InNodeID, float InPosit
 	SequencePlayerRecordsThisFrame.Emplace(InNodeID, InPosition, InLength, InFrameCount);
 }
 
-void FAnimBlueprintDebugData::RecordBlendSpacePlayer(int32 InNodeID, const UBlendSpaceBase* InBlendSpace, float InPositionX, float InPositionY, float InPositionZ)
+void FAnimBlueprintDebugData::RecordBlendSpacePlayer(int32 InNodeID, const UBlendSpaceBase* InBlendSpace, const FVector& InPosition, const FVector& InFilteredPosition)
 {
-	BlendSpacePlayerRecordsThisFrame.Emplace(InNodeID, InBlendSpace, InPositionX, InPositionY, InPositionZ);
+	BlendSpacePlayerRecordsThisFrame.Emplace(InNodeID, InBlendSpace, InPosition, InFilteredPosition);
 }
 
 void FAnimBlueprintDebugData::AddPoseWatch(int32 NodeID, FColor Color)
diff --git a/Engine/Source/Runtime/Engine/Private/Animation/AnimSingleNodeInstance.cpp b/Engine/Source/Runtime/Engine/Private/Animation/AnimSingleNodeInstance.cpp
--- a/Engine/Source/Runtime/Engine/Private/Animation/AnimSingleNodeInstance.cpp
+++ b/Engine/Source/Runtime/Engine/Private/Animation/AnimSingleNodeInstance.cpp
@@ -375,9 +375,14 @@ void UAnimSingleNodeInstance::SetPosition(float InPosition, bool bFireNotifies)
 	SetPositionWithPreviousTime(InPosition, PreviousTime, bFireNotifies);
 }
 
-void UAnimSingleNodeInstance::SetBlendSpaceInput(const FVector& InBlendInput)
+void UAnimSingleNodeInstance::SetBlendSpacePosition(const FVector& InPosition)
 {
-	GetProxyOnGameThread<FAnimSingleNodeInstanceProxy>().SetBlendSpaceInput(InBlendInput);
+	GetProxyOnGameThread<FAnimSingleNodeInstanceProxy>().SetBlendSpacePosition(InPosition);
+}
+
+void UAnimSingleNodeInstance::GetBlendSpaceState(FVector& OutPosition, FVector& OutFilteredPosition) const
+{
+	GetProxyOnGameThread<FAnimSingleNodeInstanceProxy>().GetBlendSpaceState(OutPosition, OutFilteredPosition);
 }
 
 float UAnimSingleNodeInstance::GetLength()
diff --git a/Engine/Source/Runtime/Engine/Private/Animation/AnimSingleNodeInstanceProxy.cpp b/Engine/Source/Runtime/Engine/Private/Animation/AnimSingleNodeInstanceProxy.cpp
--- a/Engine/Source/Runtime/Engine/Private/Animation/AnimSingleNodeInstanceProxy.cpp
+++ b/Engine/Source/Runtime/Engine/Private/Animation/AnimSingleNodeInstanceProxy.cpp
@@ -28,7 +28,7 @@ void FAnimSingleNodeInstanceProxy::Initialize(UAnimInstance* InAnimInstance)
 	UpdateCounter.Reset();
 
 	// it's already doing it when evaluate
-	BlendSpaceInput = FVector::ZeroVector;
+	BlendSpacePosition = FVector::ZeroVector;
 	CurrentTime = 0.f;
 
 	// initialize node manually 
@@ -198,10 +198,10 @@ void FAnimSingleNodeInstanceProxy::SetAnimationAsset(class UAnimationAsset* NewA
 	bLooping = bIsLooping;
 	PlayRate = InPlayRate;
 	CurrentTime = 0.f;
-	BlendSpaceInput = FVector::ZeroVector;
+	BlendSpacePosition = FVector::ZeroVector;
 	BlendSampleData.Reset();
 	MarkerTickRecord.Reset();
-	UpdateBlendspaceSamples(BlendSpaceInput);
+	UpdateBlendspaceSamples(BlendSpacePosition);
 
 #if WITH_EDITORONLY_DATA
 	PreviewPoseCurrentTime = 0.0f;
@@ -253,9 +253,15 @@ void FAnimSingleNodeInstanceProxy::SetReverse(bool bInReverse)
 	}*/
 }
 
-void FAnimSingleNodeInstanceProxy::SetBlendSpaceInput(const FVector& InBlendInput)
+void FAnimSingleNodeInstanceProxy::SetBlendSpacePosition(const FVector& InPosition)
 {
-	BlendSpaceInput = InBlendInput;
+	BlendSpacePosition = InPosition;
+}
+
+void FAnimSingleNodeInstanceProxy::GetBlendSpaceState(FVector& OutPosition, FVector& OutFilteredPosition) const
+{
+	OutFilteredPosition = BlendFilter.GetFilterLastOutput();
+	OutPosition = BlendSpacePosition;
 }
 
 void FAnimNode_SingleNode::Evaluate_AnyThread(FPoseContext& Output)
@@ -528,7 +534,7 @@ void FAnimNode_SingleNode::Update_AnyThread(const FAnimationUpdateContext& Conte
 
 		if (UBlendSpaceBase* BlendSpace = Cast<UBlendSpaceBase>(Proxy->CurrentAsset))
 		{
-			FAnimTickRecord TickRecord(BlendSpace, Proxy->BlendSpaceInput, Proxy->BlendSampleData, Proxy->BlendFilter, Proxy->bLooping, NewPlayRate, 1.f, /*inout*/ Proxy->CurrentTime, Proxy->MarkerTickRecord);
+			FAnimTickRecord TickRecord(BlendSpace, Proxy->BlendSpacePosition, Proxy->BlendSampleData, Proxy->BlendFilter, Proxy->bLooping, NewPlayRate, 1.f, /*inout*/ Proxy->CurrentTime, Proxy->MarkerTickRecord);
 			SyncScope.AddTickRecord(TickRecord);
 
 			TRACE_ANIM_TICK_RECORD(Context, TickRecord);
diff --git a/Engine/Source/Runtime/Engine/Private/Animation/AnimTrace.cpp b/Engine/Source/Runtime/Engine/Private/Animation/AnimTrace.cpp
--- a/Engine/Source/Runtime/Engine/Private/Animation/AnimTrace.cpp
+++ b/Engine/Source/Runtime/Engine/Private/Animation/AnimTrace.cpp
@@ -35,6 +35,8 @@ UE_TRACE_EVENT_BEGIN(Animation, TickRecord)
 	UE_TRACE_EVENT_FIELD(float, PlayRate)
 	UE_TRACE_EVENT_FIELD(float, BlendSpacePositionX)
 	UE_TRACE_EVENT_FIELD(float, BlendSpacePositionY)
+	UE_TRACE_EVENT_FIELD(float, BlendSpaceFilteredPositionX)
+	UE_TRACE_EVENT_FIELD(float, BlendSpaceFilteredPositionY)
 	UE_TRACE_EVENT_FIELD(uint16, FrameCounter)
 	UE_TRACE_EVENT_FIELD(bool, Looping)
 	UE_TRACE_EVENT_FIELD(bool, IsBlendSpace)
@@ -401,11 +403,15 @@ void FAnimTrace::OutputAnimTickRecord(const FAnimationBaseContext& InContext, co
 
 		float BlendSpacePositionX = 0.0f;
 		float BlendSpacePositionY = 0.0f;
+		float BlendSpaceFilteredPositionX = 0.0f;
+		float BlendSpaceFilteredPositionY = 0.0f;
 		const bool bIsBlendSpace = InTickRecord.SourceAsset->IsA<UBlendSpaceBase>();
 		if(bIsBlendSpace)
 		{
 			BlendSpacePositionX = InTickRecord.BlendSpace.BlendSpacePositionX;
 			BlendSpacePositionY = InTickRecord.BlendSpace.BlendSpacePositionY;
+			BlendSpaceFilteredPositionX = InTickRecord.BlendSpace.BlendFilter->GetFilterLastOutput().X;
+			BlendSpaceFilteredPositionY = InTickRecord.BlendSpace.BlendFilter->GetFilterLastOutput().Y;
 		}
 
 		UE_TRACE_LOG(Animation, TickRecord, AnimationChannel)
@@ -419,6 +425,8 @@ void FAnimTrace::OutputAnimTickRecord(const FAnimationBaseContext& InContext, co
 			<< TickRecord.PlayRate(InTickRecord.PlayRateMultiplier)
 			<< TickRecord.BlendSpacePositionX(BlendSpacePositionX)
 			<< TickRecord.BlendSpacePositionY(BlendSpacePositionY)
+			<< TickRecord.BlendSpaceFilteredPositionX(BlendSpaceFilteredPositionX)
+			<< TickRecord.BlendSpaceFilteredPositionY(BlendSpaceFilteredPositionY)
 			<< TickRecord.FrameCounter(FObjectTrace::GetObjectWorldTickCounter(AnimInstance))
 			<< TickRecord.Looping(InTickRecord.bLooping)
 			<< TickRecord.IsBlendSpace(bIsBlendSpace);
diff --git a/Engine/Source/Runtime/Engine/Public/Animation/AnimSingleNodeInstanceProxy.h b/Engine/Source/Runtime/Engine/Public/Animation/AnimSingleNodeInstanceProxy.h
--- a/Engine/Source/Runtime/Engine/Public/Animation/AnimSingleNodeInstanceProxy.h
+++ b/Engine/Source/Runtime/Engine/Public/Animation/AnimSingleNodeInstanceProxy.h
@@ -58,7 +58,7 @@ public:
 	FAnimSingleNodeInstanceProxy(UAnimInstance* InAnimInstance)
 		: FAnimInstanceProxy(InAnimInstance)
 		, CurrentAsset(nullptr)
-		, BlendSpaceInput(0.0f, 0.0f, 0.0f)
+		, BlendSpacePosition(0.0f, 0.0f, 0.0f)
 		, CurrentTime(0.0f)
 #if WITH_EDITORONLY_DATA
 		, PreviewPoseCurrentTime(0.0f)
@@ -146,7 +146,9 @@ public:
 
 	void SetReverse(bool bInReverse);
 
-	void SetBlendSpaceInput(const FVector& InBlendInput);
+	void SetBlendSpacePosition(const FVector& InPosition);
+
+	void GetBlendSpaceState(FVector& OutPosition, FVector& OutFilteredPosition) const;
 
 #if WITH_EDITOR
 	bool CanProcessAdditiveAnimations() const
@@ -188,7 +190,7 @@ protected:
 
 private:
 	/** Random cached values to play each asset **/
-	FVector BlendSpaceInput;
+	FVector BlendSpacePosition;
 
 	/** Random cached values to play each asset **/
 	TArray<FBlendSampleData> BlendSampleData;
