Fix divide by zero when the blend space grid widget is opened without data

diff --git a/Engine/Source/Editor/Persona/Private/SAnimationBlendSpaceGridWidget.cpp b/Engine/Source/Editor/Persona/Private/SAnimationBlendSpaceGridWidget.cpp
--- a/Engine/Source/Editor/Persona/Private/SAnimationBlendSpaceGridWidget.cpp
+++ b/Engine/Source/Editor/Persona/Private/SAnimationBlendSpaceGridWidget.cpp
@@ -1190,6 +1190,10 @@ void SBlendSpaceGridWidget::CalculateGridPoints()
 {
 	CachedGridPoints.Empty(SampleGridDivisions.X * SampleGridDivisions.Y);
 	CachedSamplePoints.Empty(SampleGridDivisions.X * SampleGridDivisions.Y);
+	if (SampleGridDivisions.X <= 0 || (GridType == EGridType::TwoAxis && SampleGridDivisions.Y <= 0))
+	{
+		return;
+	}
 	for (int32 GridY = 0; GridY < ((GridType == EGridType::TwoAxis) ? SampleGridDivisions.Y + 1 : 1); ++GridY)
 	{
 		for (int32 GridX = 0; GridX < SampleGridDivisions.X + 1; ++GridX)
