Disable loading of asset at Persona module startup because of major performance regression this code introduce

diff --git a/Engine/Source/Editor/Persona/Private/PersonaModule.cpp b/Engine/Source/Editor/Persona/Private/PersonaModule.cpp
--- a/Engine/Source/Editor/Persona/Private/PersonaModule.cpp
+++ b/Engine/Source/Editor/Persona/Private/PersonaModule.cpp
@@ -116,29 +116,35 @@ void FPersonaModule::StartupModule()
 	// Load all blueprint animnotifies from asset registry so they are available from drop downs in anim segment detail views
 	// TODO: Currently disabled when I/O store is enabled in editor builds because this triggers loading cooked SkeletalMeshes
 	// that currently crashes
-	FString Commandline = FCommandLine::Get();
-	const bool bIsCookCommandlet = Commandline.Contains(TEXT("cookcommandlet")) || Commandline.Contains(TEXT("run=cook"));
-	const bool bLoadAnimNotifiesBlueprints = !bIsCookCommandlet && !WITH_IOSTORE_IN_EDITOR;
-	if(bLoadAnimNotifiesBlueprints)
-	{
-		FAssetRegistryModule& AssetRegistryModule = FModuleManager::LoadModuleChecked<FAssetRegistryModule>(TEXT("AssetRegistry"));
-
-		// Collect a full list of assets with the specified class
-		TArray<FAssetData> AssetData;
-		AssetRegistryModule.Get().GetAssetsByClass(UBlueprint::StaticClass()->GetFName(), AssetData);
 
-		const FString BPAnimNotify( TEXT("Class'/Script/Engine.AnimNotify'" ));
-
-		for (int32 AssetIndex = 0; AssetIndex < AssetData.Num(); ++AssetIndex)
-		{
-			FString TagValue = AssetData[ AssetIndex ].GetTagValueRef<FString>(FBlueprintTags::ParentClassPath);
-			if (TagValue == BPAnimNotify)
-			{
-				FString BlueprintPath = AssetData[AssetIndex].ObjectPath.ToString();
-				LoadObject<UBlueprint>(NULL, *BlueprintPath, NULL, 0, NULL);
-			}
-		}
-	}
+	// This code didn't work before CL 15306274 because of timing issues
+	// This is now disabled because it tries to load half the world by pulling too many assets.
+	// We should find other ways to populate the dropdown without having to loading anything or doing a minimal amount of work only
+	// when the drop down is interacted with...
+
+	//FString Commandline = FCommandLine::Get();
+	//const bool bIsCookCommandlet = Commandline.Contains(TEXT("cookcommandlet")) || Commandline.Contains(TEXT("run=cook"));
+	//const bool bLoadAnimNotifiesBlueprints = !bIsCookCommandlet && !WITH_IOSTORE_IN_EDITOR;
+	//if(bLoadAnimNotifiesBlueprints)
+	//{
+	//	FAssetRegistryModule& AssetRegistryModule = FModuleManager::LoadModuleChecked<FAssetRegistryModule>(TEXT("AssetRegistry"));
+
+	//	// Collect a full list of assets with the specified class
+	//	TArray<FAssetData> AssetData;
+	//	AssetRegistryModule.Get().GetAssetsByClass(UBlueprint::StaticClass()->GetFName(), AssetData);
+
+	//	const FString BPAnimNotify( TEXT("Class'/Script/Engine.AnimNotify'" ));
+
+	//	for (int32 AssetIndex = 0; AssetIndex < AssetData.Num(); ++AssetIndex)
+	//	{
+	//		FString TagValue = AssetData[ AssetIndex ].GetTagValueRef<FString>(FBlueprintTags::ParentClassPath);
+	//		if (TagValue == BPAnimNotify)
+	//		{
+	//			FString BlueprintPath = AssetData[AssetIndex].ObjectPath.ToString();
+	//			LoadObject<UBlueprint>(NULL, *BlueprintPath, NULL, 0, NULL);
+	//		}
+	//	}
+	//}
 	{
 		FPropertyEditorModule& PropertyModule = FModuleManager::LoadModuleChecked<FPropertyEditorModule>("PropertyEditor");
 		PropertyModule.RegisterCustomClassLayout( "SkeletalMeshSocket", FOnGetDetailCustomizationInstance::CreateStatic( &FSkeletalMeshSocketDetails::MakeInstance ) );
