Fixed occasional incorrect binding being set for property access on anim nodes

The setting of the exposed value bool was causing a skeleton recompilation before the binding was made. Also prevented multiple skeleton recompiles (causing the same issue) with multi-selected nodes.

diff --git a/Engine/Source/Editor/Persona/Private/AnimGraphNodeDetails.cpp b/Engine/Source/Editor/Persona/Private/AnimGraphNodeDetails.cpp
--- a/Engine/Source/Editor/Persona/Private/AnimGraphNodeDetails.cpp
+++ b/Engine/Source/Editor/Persona/Private/AnimGraphNodeDetails.cpp
@@ -1323,9 +1323,6 @@ TSharedRef<SWidget> FAnimGraphNodeBindingExtension::GenerateExtensionWidget(cons
 			const UEdGraphSchema_K2* Schema = GetDefault<UEdGraphSchema_K2>();
 			IPropertyAccessEditor& PropertyAccessEditor = IModularFeatures::Get().GetModularFeature<IPropertyAccessEditor>("PropertyAccessEditor");
 
-			// Disconnect any connected pin since we will replace it with the bound property/function
-			ShowPinPropertyHandle->SetValue(false);
-
 			for(UObject* OuterObject : OuterObjects)
 			{
 				if(UAnimGraphNode_Base* AnimGraphNode = Cast<UAnimGraphNode_Base>(OuterObject))
@@ -1373,9 +1370,10 @@ TSharedRef<SWidget> FAnimGraphNodeBindingExtension::GenerateExtensionWidget(cons
 					}
 					AnimGraphNode->PropertyBindings.Add(InPropertyName, Binding);
 				}
-
-				FBlueprintEditorUtils::MarkBlueprintAsModified(Blueprint);
 			}
+
+			FBlueprintEditorUtils::MarkBlueprintAsModified(Blueprint);
+			
 			// Pins are exposed if we have a binding or not - and after running this we do.
 			ShowPinPropertyHandle->SetValue(true);
 		});
