Fix to non-unity build.

#rnx

diff --git a/Engine/Source/Editor/Persona/Private/Customization/BlendSampleDetails.cpp b/Engine/Source/Editor/Persona/Private/Customization/BlendSampleDetails.cpp
--- a/Engine/Source/Editor/Persona/Private/Customization/BlendSampleDetails.cpp
+++ b/Engine/Source/Editor/Persona/Private/Customization/BlendSampleDetails.cpp
@@ -11,8 +11,9 @@
 #include "DetailLayoutBuilder.h"
 #include "DetailCategoryBuilder.h"
 
-#include "Widgets/Text/STextBlock.h"
+#include "Widgets/Input/SButton.h"
 #include "Widgets/Input/SNumericEntryBox.h"
+#include "Widgets/Text/STextBlock.h"
 
 #include "Animation/AnimSequence.h"
 #include "Animation/BlendSpaceBase.h"
