Fix for BlendSpaceDetails Non-Unity compile error

diff --git a/Engine/Source/Editor/Persona/Private/Customization/BlendSpaceDetails.h b/Engine/Source/Editor/Persona/Private/Customization/BlendSpaceDetails.h
--- a/Engine/Source/Editor/Persona/Private/Customization/BlendSpaceDetails.h
+++ b/Engine/Source/Editor/Persona/Private/Customization/BlendSpaceDetails.h
@@ -4,6 +4,8 @@
 
 #include "CoreMinimal.h"
 #include "IDetailCustomization.h"
+#include "UObject/WeakObjectPtr.h"
+#include "UObject/WeakObjectPtrTemplates.h"
 
 class IDetailLayoutBuilder;
 class UBlendSpaceBase;
