Fix non unity CIS issue

diff --git a/Engine/Source/Editor/Persona/Private/SSkeletonSlotNames.h b/Engine/Source/Editor/Persona/Private/SSkeletonSlotNames.h
--- a/Engine/Source/Editor/Persona/Private/SSkeletonSlotNames.h
+++ b/Engine/Source/Editor/Persona/Private/SSkeletonSlotNames.h
@@ -17,6 +17,7 @@
 #include "Widgets/Views/STableRow.h"
 #include "Widgets/Views/STreeView.h"
 #include "EditorObjectsTracker.h"
+#include "EditorUndoClient.h"
 
 class FMenuBuilder;
 class IEditableSkeleton;
