Add search to the menu for adding curves in the animation sequencer

diff --git a/Engine/Source/Editor/Persona/Private/AnimTimeline/AnimTimelineTrack_Curves.cpp b/Engine/Source/Editor/Persona/Private/AnimTimeline/AnimTimelineTrack_Curves.cpp
--- a/Engine/Source/Editor/Persona/Private/AnimTimeline/AnimTimelineTrack_Curves.cpp
+++ b/Engine/Source/Editor/Persona/Private/AnimTimeline/AnimTimelineTrack_Curves.cpp
@@ -3,6 +3,8 @@
 #include "AnimTimelineTrack_Curves.h"
 #include "PersonaUtils.h"
 #include "Widgets/SBoxPanel.h"
+#include "Widgets/Layout/SScrollBox.h"
+#include "Widgets/Input/SSearchBox.h"
 #include "AnimSequenceTimelineCommands.h"
 #include "Framework/MultiBox/MultiBoxBuilder.h"
 #include "Animation/AnimSequenceBase.h"
@@ -13,9 +15,12 @@
 #include "SAnimOutlinerItem.h"
 #include "Preferences/PersonaOptions.h"
 #include "Animation/AnimData/AnimDataController.h"
+#include "SListViewSelectorDropdownMenu.h"
 
 #define LOCTEXT_NAMESPACE "FAnimTimelineTrack_Notifies"
 
+const float FAnimTimelineTrack_Curves::CurveListPadding = 8.0f;
+
 ANIMTIMELINE_IMPLEMENT_TRACK(FAnimTimelineTrack_Curves);
 
 FAnimTimelineTrack_Curves::FAnimTimelineTrack_Curves(const TSharedRef<FAnimModel>& InModel)
@@ -133,18 +138,6 @@ TSharedRef<SWidget> FAnimTimelineTrack_Curves::BuildCurvesSubMenu()
 	return MenuBuilder.MakeWidget();
 }
 
-
-struct FSmartNameSortItem
-{
-	FName SmartName;
-	USkeleton::AnimCurveUID ID;
-
-	FSmartNameSortItem(const FName& InSmartName, const USkeleton::AnimCurveUID& InID)
-		: SmartName(InSmartName)
-		, ID(InID)
-	{}
-};
-
 struct FSmartNameSortItemSortOp
 {
 	bool operator()(const FSmartNameSortItem& A, const FSmartNameSortItem& B) const
@@ -210,6 +203,115 @@ void FAnimTimelineTrack_Curves::FillMetadataEntryMenu(FMenuBuilder& Builder)
 	Builder.AddMenuEntry(Label, Description, FSlateIcon(), UIAction);
 }
 
+TSharedRef<ITableRow> FAnimTimelineTrack_Curves::GenerateCurveListRow(FCurveListItem InItem, const TSharedRef<STableViewBase>& OwnerList)
+{
+	return SNew(STableRow<FCurveListItem>, OwnerList)
+		.Padding(FMargin(CurveListPadding, 0.0))
+		[
+			SNew(STextBlock).Text(FText::FromString(InItem.Get()->SmartName.ToString())).HighlightText(SearchText)
+		];
+}
+
+void FAnimTimelineTrack_Curves::OnTypeSelectionChanged(FCurveListItem Selection, ESelectInfo::Type SelectInfo)
+{
+	// When the user is navigating, do not act upon the selection change
+	if (SelectInfo == ESelectInfo::OnNavigation)
+	{
+		return;
+	}
+
+	if (Selection.IsValid())
+	{
+		AddVariableCurve(Selection->ID);
+		FSlateApplication::Get().DismissAllMenus();
+	}
+}
+
+void FAnimTimelineTrack_Curves::OnMouseButtonClicked(FCurveListItem Selection)
+{
+	if (Selection.IsValid())
+	{
+		AddVariableCurve(Selection->ID);
+		FSlateApplication::Get().DismissAllMenus();
+	}
+}
+
+void FAnimTimelineTrack_Curves::OnCurveFilterTextChanged(const FText& NewText)
+{
+	SearchText = NewText;
+	FilteredCurveItems.Empty();
+
+	GetCurvesMatchingSearch(NewText, CurveItems, FilteredCurveItems);
+	CurveListView->RequestListRefresh();
+
+	auto SelectedItems = CurveListView->GetSelectedItems();
+	if (FilteredCurveItems.Num() > 0)
+	{
+		CurveListView->SetSelection(FilteredCurveItems[0], ESelectInfo::OnNavigation);
+	}
+}
+
+bool FAnimTimelineTrack_Curves::GetCurvesMatchingSearch(const FText& InSearchText, const TArray<FCurveListItem>& UnfilteredList, TArray<FCurveListItem>& OutFilteredList)
+{
+	// Trim and sanitized the filter text (so that it more likely matches the action descriptions)
+	FString TrimmedFilterString = FText::TrimPrecedingAndTrailing(InSearchText).ToString();
+
+	// Tokenize the search box text into a set of terms; all of them must be present to pass the filter
+	TArray<FString> FilterTerms;
+	TrimmedFilterString.ParseIntoArray(FilterTerms, TEXT(" "), true);
+
+	// Generate a list of sanitized versions of the strings
+	TArray<FString> SanitizedFilterTerms;
+	for (int32 iFilters = 0; iFilters < FilterTerms.Num(); iFilters++)
+	{
+		FString EachString = FName::NameToDisplayString(FilterTerms[iFilters], false);
+		EachString = EachString.Replace(TEXT(" "), TEXT(""));
+		SanitizedFilterTerms.Add(EachString);
+	}
+
+	ensure(SanitizedFilterTerms.Num() == FilterTerms.Num());
+
+	bool bReturnVal = false;
+
+	for (auto it = UnfilteredList.CreateConstIterator(); it; ++it)
+	{
+		FSmartNameSortItem CurItem = *(*it).Get();
+
+		const bool bIsEmptySearch = InSearchText.IsEmpty();
+		bool bFilterTextMatches = true;
+
+		if (!bIsEmptySearch)
+		{
+			const FString DescriptionString = CurItem.SmartName.ToString();
+			const FString MangledDescriptionString = DescriptionString.Replace(TEXT(" "), TEXT(""));
+
+			for (int32 FilterIndex = 0; FilterIndex < FilterTerms.Num() && bFilterTextMatches; ++FilterIndex)
+			{
+				const bool bMatchesTerm = MangledDescriptionString.Contains(FilterTerms[FilterIndex]) || MangledDescriptionString.Contains(SanitizedFilterTerms[FilterIndex]);
+				bFilterTextMatches = bFilterTextMatches && bMatchesTerm;
+			}
+		}
+		if (bIsEmptySearch || bFilterTextMatches)
+		{
+			OutFilteredList.Add(MakeShared<FSmartNameSortItem>(CurItem));
+			bReturnVal = true;
+		}
+	}
+	return bReturnVal;
+}
+
+void FAnimTimelineTrack_Curves::OnCurveFilterTextCommitted(const FText& NewText, ETextCommit::Type CommitInfo)
+{
+	if (CommitInfo == ETextCommit::OnEnter)
+	{
+		auto SelectedItems = CurveListView->GetSelectedItems();
+		if (SelectedItems.Num() > 0)
+		{
+			CurveListView->SetSelection(SelectedItems[0]);
+		}
+	}
+}
+
 void FAnimTimelineTrack_Curves::FillVariableCurveMenu(FMenuBuilder& Builder)
 {
 	FText Description = LOCTEXT("NewVariableCurveCreateNew_ToolTip", "Create a new variable curve");
@@ -226,39 +328,82 @@ void FAnimTimelineTrack_Curves::FillVariableCurveMenu(FMenuBuilder& Builder)
 	const FSmartNameMapping* Mapping = CurrentSkeleton->GetSmartNameContainer(USkeleton::AnimCurveMappingName);
 	TArray<USkeleton::AnimCurveUID> CurveUids;
 	Mapping->FillUidArray(CurveUids);
-
+	CurveItems.Empty();
 	Builder.BeginSection(NAME_None, LOCTEXT("VariableMenu_ListHeading", "Available Names"));
-	{
-		TArray<FSmartNameSortItem> SmartNameList;
 
-		for (USkeleton::AnimCurveUID Id : CurveUids)
+	for (USkeleton::AnimCurveUID Id : CurveUids)
+	{
+		if (!AnimSequenceBase->GetDataModel()->FindFloatCurve(FAnimationCurveIdentifier(Id, ERawCurveTrackTypes::RCT_Float)))
 		{
-			if (!AnimSequenceBase->GetDataModel()->FindFloatCurve(FAnimationCurveIdentifier(Id, ERawCurveTrackTypes::RCT_Float)))
+			FName CurveName;
+			if (Mapping->GetName(Id, CurveName))
 			{
-				FName CurveName;
-				if (Mapping->GetName(Id, CurveName))
-				{
-					SmartNameList.Add(FSmartNameSortItem(CurveName, Id));
-				}
+				CurveItems.Add(MakeShared<FSmartNameSortItem>(FSmartNameSortItem(CurveName, Id)));
 			}
 		}
+	}
 
-		{
-			SmartNameList.Sort(FSmartNameSortItemSortOp());
+	// Build a SearchBox followed by a list of all the available curves
+	FilteredCurveItems = CurveItems; 
+	TSharedPtr<SSearchBox>	CurveFilterTextBox;
+	TSharedPtr<SMenuOwner>	MenuContent;
 
-			for (FSmartNameSortItem SmartNameItem : SmartNameList)
-			{
-				Description = LOCTEXT("NewVariableSubMenu_ToolTip", "Add an existing variable curve");
-				Label = FText::FromName(SmartNameItem.SmartName);
+	SAssignNew(CurveListView, SCurveListView)
+		.ListItemsSource(&FilteredCurveItems)
+		.SelectionMode(ESelectionMode::Single)
+		.OnGenerateRow(this, &FAnimTimelineTrack_Curves::GenerateCurveListRow)
+		.OnMouseButtonClick(this, &FAnimTimelineTrack_Curves::OnMouseButtonClicked)
+		.OnSelectionChanged(this, &FAnimTimelineTrack_Curves::OnTypeSelectionChanged);
 
-				UIAction.ExecuteAction.BindRaw(
-					this, &FAnimTimelineTrack_Curves::AddVariableCurve,
-					SmartNameItem.ID);
+	SAssignNew(CurveFilterTextBox, SSearchBox)
+		.OnTextChanged(this, &FAnimTimelineTrack_Curves::OnCurveFilterTextChanged)
+		.OnTextCommitted(this, &FAnimTimelineTrack_Curves::OnCurveFilterTextCommitted);
+
+	SAssignNew(MenuContent, SMenuOwner)
+		[
+			SNew(SListViewSelectorDropdownMenu<FCurveListItem>, CurveFilterTextBox, CurveListView)
+			[
+				SNew(SVerticalBox)
+				+ SVerticalBox::Slot()
+				.AutoHeight()
+				.Padding(FMargin(CurveListPadding, 2.0))
+				[
+					CurveFilterTextBox.ToSharedRef()
+				]
+				+ SVerticalBox::Slot()
+				.FillHeight(1.f)
+				.VAlign(VAlign_Fill)
+				.Padding(CurveListPadding, 2.0)
+				[
+					SNew(SBox)
+					.WidthOverride(300.0f)
+					.HeightOverride(300.f)
+					[
+						SNew(SOverlay)
+						+ SOverlay::Slot()
+						[
+							SNew(SBorder)
+							.BorderImage(FEditorStyle::GetBrush("Graph.StateNode.Body"))
+							.BorderBackgroundColor(FAppStyle::Get().GetSlateColor("Colors.Input"))
+						]
+						+SOverlay::Slot()
+						[
+							SNew(SScrollBox)
+							.Orientation(EOrientation::Orient_Vertical)
+							+ SScrollBox::Slot()
+							.HAlign(HAlign_Fill)
+							.VAlign(VAlign_Fill)
+							[
+								CurveListView.ToSharedRef()
+							]
+						]
+					]
+				]
+			]
+		];
+
+	Builder.AddWidget(MenuContent.ToSharedRef(), FText::GetEmpty(),  false);
 
-				Builder.AddMenuEntry(Label, Description, FSlateIcon(), UIAction);
-			}
-		}
-	}
 	Builder.EndSection();
 }
 
diff --git a/Engine/Source/Editor/Persona/Private/AnimTimeline/AnimTimelineTrack_Curves.h b/Engine/Source/Editor/Persona/Private/AnimTimeline/AnimTimelineTrack_Curves.h
--- a/Engine/Source/Editor/Persona/Private/AnimTimeline/AnimTimelineTrack_Curves.h
+++ b/Engine/Source/Editor/Persona/Private/AnimTimeline/AnimTimelineTrack_Curves.h
@@ -4,6 +4,22 @@
 
 #include "AnimTimelineTrack.h"
 #include "Animation/Skeleton.h"
+#include "Widgets/Views/STableRow.h"
+#include "Widgets/Views/SListView.h"
+
+struct FSmartNameSortItem
+{
+	FName SmartName;
+	USkeleton::AnimCurveUID ID;
+
+	FSmartNameSortItem(const FName& InSmartName, const USkeleton::AnimCurveUID& InID)
+		: SmartName(InSmartName)
+		, ID(InID)
+	{}
+};
+
+typedef TSharedPtr<FSmartNameSortItem> FCurveListItem;
+typedef SListView<FCurveListItem> SCurveListView;
 
 class FAnimTimelineTrack_Curves : public FAnimTimelineTrack
 {
@@ -31,6 +47,22 @@ private:
 	void HandleShowCurvePoints();
 	bool IsShowCurvePointsEnabled() const;
 
+	virtual TSharedRef<ITableRow> GenerateCurveListRow(FCurveListItem InItem, const TSharedRef<STableViewBase>& OwnerList);
+	void OnTypeSelectionChanged(FCurveListItem Selection, ESelectInfo::Type SelectInfo);
+	void OnMouseButtonClicked(FCurveListItem Selection);
+
+	/** Curve searching support */
+	FText SearchText;
+	void OnCurveFilterTextChanged(const FText& NewText);
+	void OnCurveFilterTextCommitted(const FText& NewText, ETextCommit::Type CommitInfo);
+	bool GetCurvesMatchingSearch(const FText& InSearchText, const TArray<FCurveListItem>& UnfilteredList, TArray<FCurveListItem>& OutFilteredList);
 private:
-	TSharedPtr<SWidget> OutlinerWidget;
+	TSharedPtr<SWidget>	OutlinerWidget;
+
+	TSharedPtr<SCurveListView>	CurveListView;
+
+	TArray<FCurveListItem >	CurveItems;
+	TArray<FCurveListItem >	FilteredCurveItems;
+
+	static const float	CurveListPadding;
 };
\ No newline at end of file
