Fixed bone references inside a nested blendspace not displaying and causing an ensure

To determine the to pick from skeleton we now traverse the outer chain to we eventually pick up the node for the blendspace, as in this case the asset is nested within a bunch of objects.

diff --git a/Engine/Source/Editor/Persona/Private/AnimGraphNodeDetails.cpp b/Engine/Source/Editor/Persona/Private/AnimGraphNodeDetails.cpp
--- a/Engine/Source/Editor/Persona/Private/AnimGraphNodeDetails.cpp
+++ b/Engine/Source/Editor/Persona/Private/AnimGraphNodeDetails.cpp
@@ -511,59 +511,75 @@ void FBoneReferenceCustomization::SetEditableSkeleton(TSharedRef<IPropertyHandle
 
 	bEnsureOnInvalidSkeleton = true;
 
-	for (UObject* Outer : Objects)
+	auto FindSkeletonForObject = [this, &TargetSkeleton, &EditableSkeleton](UObject* InObject)
 	{
-		if (UAnimGraphNode_Base* AnimGraphNode = Cast<UAnimGraphNode_Base>(Outer))
+		for( ; InObject; InObject = InObject->GetOuter())
 		{
-			TargetSkeleton = AnimGraphNode->GetAnimBlueprint()->TargetSkeleton;
-			break;
-		}
+			if (UAnimGraphNode_Base* AnimGraphNode = Cast<UAnimGraphNode_Base>(InObject))
+			{
+				TargetSkeleton = AnimGraphNode->GetAnimBlueprint()->TargetSkeleton;
+				break;
+			}
 
-		if (USkeletalMesh* SkeletalMesh = Cast<USkeletalMesh>(Outer))
-		{
-			TargetSkeleton = SkeletalMesh->GetSkeleton();
-			break;
-		}
+			if (USkeletalMesh* SkeletalMesh = Cast<USkeletalMesh>(InObject))
+			{
+				TargetSkeleton = SkeletalMesh->GetSkeleton();
+				break;
+			}
 
-		if (ULODInfoUILayout* LODInfoUILayout = Cast<ULODInfoUILayout>(Outer))
-		{
-			USkeletalMesh* SkeletalMesh = LODInfoUILayout->GetPersonaToolkit()->GetPreviewMesh();
-			check(SkeletalMesh);
-			TargetSkeleton = SkeletalMesh->GetSkeleton();
-			break;
-		}
+			if (ULODInfoUILayout* LODInfoUILayout = Cast<ULODInfoUILayout>(InObject))
+			{
+				USkeletalMesh* SkeletalMesh = LODInfoUILayout->GetPersonaToolkit()->GetPreviewMesh();
+				check(SkeletalMesh);
+				TargetSkeleton = SkeletalMesh->GetSkeleton();
+				break;
+			}
 
-		if (UAnimationAsset* AnimationAsset = Cast<UAnimationAsset>(Outer))
-		{
-			TargetSkeleton = AnimationAsset->GetSkeleton();
-			break;
-		}
+			if (UAnimationAsset* AnimationAsset = Cast<UAnimationAsset>(InObject))
+			{
+				if(AnimationAsset->IsAsset())
+				{
+					TargetSkeleton = AnimationAsset->GetSkeleton();
+					break;
+				}
+			}
 
-		if (UAnimInstance* AnimInstance = Cast<UAnimInstance>(Outer))
-		{
-			if (AnimInstance->CurrentSkeleton)
+			if (UAnimInstance* AnimInstance = Cast<UAnimInstance>(InObject))
 			{
-				TargetSkeleton = AnimInstance->CurrentSkeleton;
-				break;
+				if (AnimInstance->CurrentSkeleton)
+				{
+					TargetSkeleton = AnimInstance->CurrentSkeleton;
+					break;
+				}
+				else if (UAnimBlueprintGeneratedClass* AnimBPClass = Cast<UAnimBlueprintGeneratedClass>(AnimInstance->GetClass()))
+				{
+					TargetSkeleton = AnimBPClass->TargetSkeleton;
+					break;
+				}
 			}
-			else if (UAnimBlueprintGeneratedClass* AnimBPClass = Cast<UAnimBlueprintGeneratedClass>(AnimInstance->GetClass()))
+
+			// editor animation curve bone links are responsible for linking joints to curve
+			// this is editor object that only exists for editor
+			if (UEditorAnimCurveBoneLinks* AnimCurveObj = Cast<UEditorAnimCurveBoneLinks>(InObject))
 			{
-				TargetSkeleton = AnimBPClass->TargetSkeleton;
+				EditableSkeleton = AnimCurveObj->EditableSkeleton.Pin();
 				break;
 			}
-		}
 
-		// editor animation curve bone links are responsible for linking joints to curve
-		// this is editor object that only exists for editor
-		if (UEditorAnimCurveBoneLinks* AnimCurveObj = Cast<UEditorAnimCurveBoneLinks>(Outer))
-		{
-			EditableSkeleton = AnimCurveObj->EditableSkeleton.Pin();
-			break;
+			if (IBoneReferenceSkeletonProvider* SkeletonProvider = Cast<IBoneReferenceSkeletonProvider>(InObject))
+			{
+				TargetSkeleton = SkeletonProvider->GetSkeleton(bEnsureOnInvalidSkeleton);
+				break;
+			}
 		}
 
-		if (IBoneReferenceSkeletonProvider* SkeletonProvider = Cast<IBoneReferenceSkeletonProvider>(Outer))
+		return TargetSkeleton != nullptr || EditableSkeleton != nullptr;
+	};
+	
+	for (UObject* Object : Objects)
+	{
+		if(FindSkeletonForObject(Object))
 		{
-			TargetSkeleton = SkeletonProvider->GetSkeleton(bEnsureOnInvalidSkeleton);
 			break;
 		}
 	}
