Control Rig: crash when attempting select a preview mesh on a rig that does not currently have one.

diff --git a/Engine/Source/Editor/Persona/Private/PreviewSceneCustomizations.cpp b/Engine/Source/Editor/Persona/Private/PreviewSceneCustomizations.cpp
--- a/Engine/Source/Editor/Persona/Private/PreviewSceneCustomizations.cpp
+++ b/Engine/Source/Editor/Persona/Private/PreviewSceneCustomizations.cpp
@@ -175,6 +175,10 @@ void FPreviewSceneDescriptionCustomization::CustomizeDetails(IDetailLayoutBuilde
 			PreviewMeshName = SkeletalMeshProperty->GetPropertyDisplayName();
 		}
 
+		const bool bCanUseDifferentSkeleton =
+			(PersonaToolkit.Pin()->GetContext() == UPhysicsAsset::StaticClass()->GetFName()) ||
+			(PersonaToolkit.Pin()->GetContext() == TEXT("ControlRigBlueprint"));
+
 		DetailBuilder.EditCategory("Mesh")
 		.AddProperty(SkeletalMeshProperty)
 		.CustomWidget()
@@ -231,7 +235,7 @@ void FPreviewSceneDescriptionCustomization::CustomizeDetails(IDetailLayoutBuilde
 			SNew(SObjectPropertyEntryBox)
 			.AllowedClass(USkeletalMesh::StaticClass())
 			.PropertyHandle(SkeletalMeshProperty)
-			.OnShouldFilterAsset(this, &FPreviewSceneDescriptionCustomization::HandleShouldFilterAsset, USkeletalMesh::GetSkeletonMemberName(), PersonaToolkit.Pin()->GetContext() == UPhysicsAsset::StaticClass()->GetFName())
+			.OnShouldFilterAsset(this, &FPreviewSceneDescriptionCustomization::HandleShouldFilterAsset, USkeletalMesh::GetSkeletonMemberName(), bCanUseDifferentSkeleton)
 			.OnObjectChanged(this, &FPreviewSceneDescriptionCustomization::HandleMeshChanged)
 			.ThumbnailPool(DetailBuilder.GetThumbnailPool())
 		];
@@ -525,6 +529,11 @@ bool FPreviewSceneDescriptionCustomization::HandleShouldFilterAsset(const FAsset
 	}
 
 	const UPersonaPreviewSceneDescription* PersonaPreviewSceneDescription = PreviewScene.Pin()->GetPreviewSceneDescription();
+	if(!PersonaPreviewSceneDescription->PreviewMesh.IsValid())
+	{
+		return false;
+	}
+
 	const USkeleton* Skeleton = PersonaPreviewSceneDescription->PreviewMesh->GetSkeleton();
 	const FString SkeletonTag = InAssetData.GetTagValueRef<FString>(InTag);
 	if (Skeleton && Skeleton->IsCompatibleSkeletonByAssetString(SkeletonTag))
