Changed Blendspace preview behavior & misc fixes

Preview point can no longer be dragged by just clicking on it. Now the only way to move it is to shift-drag. This is reflected in the various status and help messages.
Also fixed up some drawing issues:
- Preview point is now drawn underneath sample points
- Preview point is now a cross rather than the same icon as a sample point. This means it is visible when drawn co-incident with a sample
- Icons for arrows & points were being drawn at the wrong size.
- Settings icons (e.g. triagulation, labels) were of the wrong color
- Triangulation was over-bright
- Labels needed to dim triangulation when visible because they used the same color. Now they have a background so they are more readable and color is no longer dimmed.

diff --git a/Engine/Source/Editor/EditorStyle/Private/StarshipStyle.cpp b/Engine/Source/Editor/EditorStyle/Private/StarshipStyle.cpp
--- a/Engine/Source/Editor/EditorStyle/Private/StarshipStyle.cpp
+++ b/Engine/Source/Editor/EditorStyle/Private/StarshipStyle.cpp
@@ -5264,10 +5264,12 @@ void FStarshipEditorStyle::FStyle::SetupPersonaStyle()
 		// blend space
 		Set("BlendSpaceEditor.ToggleTriangulation", new IMAGE_BRUSH(TEXT("Persona/BlendSpace/triangulation_16"), Icon16x16));
 		Set("BlendSpaceEditor.ToggleLabels", new IMAGE_BRUSH(TEXT("Persona/BlendSpace/label_16"), Icon16x16));
-		Set("BlendSpaceEditor.ArrowDown", new IMAGE_BRUSH(TEXT("Persona/BlendSpace/arrow_down_12x"), Icon12x12));
-		Set("BlendSpaceEditor.ArrowUp", new IMAGE_BRUSH(TEXT("Persona/BlendSpace/arrow_up_12x"), Icon12x12));
-		Set("BlendSpaceEditor.ArrowRight", new IMAGE_BRUSH(TEXT("Persona/BlendSpace/arrow_right_12x"), Icon12x12));
-		Set("BlendSpaceEditor.ArrowLeft", new IMAGE_BRUSH(TEXT("Persona/BlendSpace/arrow_left_12x"), Icon12x12));
+		Set("BlendSpaceEditor.ArrowDown", new IMAGE_BRUSH(TEXT("Persona/BlendSpace/arrow_down_12x"), FVector2D(13.0f, 25.0f)));
+		Set("BlendSpaceEditor.ArrowUp", new IMAGE_BRUSH(TEXT("Persona/BlendSpace/arrow_up_12x"), FVector2D(13.0f, 25.0f)));
+		Set("BlendSpaceEditor.ArrowRight", new IMAGE_BRUSH(TEXT("Persona/BlendSpace/arrow_right_12x"), FVector2D(25.0f, 13.0f)));
+		Set("BlendSpaceEditor.ArrowLeft", new IMAGE_BRUSH(TEXT("Persona/BlendSpace/arrow_left_12x"), FVector2D(25.0f, 13.0f)));
+		Set("BlendSpaceEditor.PreviewIcon", new IMAGE_BRUSH(TEXT("Persona/BlendSpace/preview_21x"), FVector2D(21.0f, 21.0f)));
+		Set("BlendSpaceEditor.LabelBackground", new FSlateRoundedBoxBrush(FStyleColors::Background, FStyleColors::Foreground, 1.0f));
 
 		const FButtonStyle BlueprintContextTargetsButtonStyle = FButtonStyle()
 			.SetNormal(IMAGE_BRUSH("Common/TreeArrow_Collapsed_Hovered", Icon10x10, FLinearColor(0.2f, 0.2f, 0.2f, 1.f)))
diff --git a/Engine/Source/Editor/Persona/Private/SAnimationBlendSpaceGridWidget.cpp b/Engine/Source/Editor/Persona/Private/SAnimationBlendSpaceGridWidget.cpp
--- a/Engine/Source/Editor/Persona/Private/SAnimationBlendSpaceGridWidget.cpp
+++ b/Engine/Source/Editor/Persona/Private/SAnimationBlendSpaceGridWidget.cpp
@@ -20,7 +20,6 @@
 #include "UObject/StructOnScope.h"
 #include "EditorStyleSet.h"
 #include "PropertyEditorModule.h"
-#include "ISinglePropertyView.h"
 #include "IStructureDetailsView.h"
 
 #include "Customization/BlendSampleDetails.h"
@@ -28,11 +27,11 @@
 #include "DragAndDrop/AssetDragDropOp.h"
 #include "Classes/EditorStyleSettings.h"
 
-#include "Widgets/Input/SButton.h"
 #include "Fonts/FontMeasure.h"
 #include "Modules/ModuleManager.h"
 
 #include "Framework/MultiBox/MultiBoxBuilder.h"
+#include "Styling/StyleColors.h"
 
 #define LOCTEXT_NAMESPACE "SAnimationBlendSpaceGridWidget"
 
@@ -90,20 +89,24 @@ void SBlendSpaceGridWidget::Construct(const FArguments& InArgs)
 	UnSnappedColor = FEditorStyle::GetSlateColor("BlendSpaceKey.UnSnapped");
 	GridLinesColor = GetDefault<UEditorStyleSettings>()->RegularColor;
 	GridOutlineColor = GetDefault<UEditorStyleSettings>()->RuleColor;
-
+	TriangulationColor = FSlateColor(EStyleColor::Foreground);
+	
 	// Retrieve background and sample key brushes 
 	BackgroundImage = FEditorStyle::GetBrush(TEXT("Graph.Panel.SolidBackground"));
 	KeyBrush = FEditorStyle::GetBrush("CurveEd.CurveKey");
+	PreviewBrush = FEditorStyle::GetBrush("BlendSpaceEditor.PreviewIcon");
 	ArrowBrushes[(uint8)EArrowDirection::Up] = FEditorStyle::GetBrush("BlendSpaceEditor.ArrowUp");
 	ArrowBrushes[(uint8)EArrowDirection::Down] = FEditorStyle::GetBrush("BlendSpaceEditor.ArrowDown");
 	ArrowBrushes[(uint8)EArrowDirection::Right] = FEditorStyle::GetBrush("BlendSpaceEditor.ArrowRight");
 	ArrowBrushes[(uint8)EArrowDirection::Left] = FEditorStyle::GetBrush("BlendSpaceEditor.ArrowLeft");
-
+	LabelBrush = FEditorStyle::GetBrush(TEXT("BlendSpaceEditor.LabelBackground"));
+	
 	// Retrieve font data 
 	FontInfo = FEditorStyle::GetFontStyle("CurveEd.InfoFont");
 
 	// Initialize UI layout values
-	KeySize = FVector2D(12.0f, 12.0f);
+	KeySize = FVector2D(11.0f, 11.0f);
+	PreviewSize = FVector2D(21.0f, 21.0f);
 	DragThreshold = 9.0f;
 	ClickAndHighlightThreshold = 12.0f;
 	TextMargin = 8.0f;
@@ -244,7 +247,7 @@ void SBlendSpaceGridWidget::Construct(const FArguments& InArgs)
 						.Padding(FMargin(2.0f, 3.0f, 0.0f, 0.0f))
 						[
 							SNew(STextBlock)
-							.Text(LOCTEXT("BlendspacePreviewToolTip", "Click and Drag or Hold Shift to move the Preview (Green) Pin" ))
+							.Text(LOCTEXT("BlendspacePreviewToolTip", "Hold Shift to move the Preview Point (Green)" ))
 							.Font(FEditorStyle::GetFontStyle(TEXT("AnimViewport.MessageFont")))
 							.ColorAndOpacity(FLinearColor(1.0f, 1.0f, 1.0f, 0.7f))
 							.Visibility(TAttribute<EVisibility>::Create(TAttribute<EVisibility>::FGetter::CreateSP(this, &SBlendSpaceGridWidget::GetPreviewToolTipVisibility)))
@@ -393,7 +396,7 @@ void SBlendSpaceGridWidget::PaintBackgroundAndGrid(const FGeometry& AllottedGeom
 
 		LinePoints[4] = GridOffset;
 
-		FSlateDrawElement::MakeLines( OutDrawElements, DrawLayerId + 3, AllottedGeometry.ToPaintGeometry(),	LinePoints, ESlateDrawEffect::None, GridOutlineColor, true);
+		FSlateDrawElement::MakeLines( OutDrawElements, DrawLayerId + 3, AllottedGeometry.ToPaintGeometry(),	LinePoints, ESlateDrawEffect::None, GridOutlineColor, true, 2.0f);
 
 	}
 	
@@ -402,6 +405,10 @@ void SBlendSpaceGridWidget::PaintBackgroundAndGrid(const FGeometry& AllottedGeom
 
 void SBlendSpaceGridWidget::PaintSampleKeys(const FGeometry& AllottedGeometry, const FSlateRect& MyCullingRect, FSlateWindowElementList& OutDrawElements, int32& DrawLayerId) const
 {
+	const int32 FilteredPositionLayer = DrawLayerId + 1;
+	const int32 PreviewPositionLayer = DrawLayerId + 2;
+	const int32 SampleLayer = DrawLayerId + 3;
+	
 	if(const UBlendSpace* BlendSpace = BlendSpaceBase.Get())
 	{
 		// Draw keys
@@ -433,30 +440,30 @@ void SBlendSpaceGridWidget::PaintSampleKeys(const FGeometry& AllottedGeometry, c
 			}
 
 			const FVector2D GridPosition = SampleValueToGridPosition(Sample.SampleValue) - (KeySize * 0.5f);
-			FSlateDrawElement::MakeBox( OutDrawElements, DrawLayerId + 1, AllottedGeometry.ToPaintGeometry(GridPosition, KeySize), KeyBrush, ESlateDrawEffect::None, DrawColor );
+			FSlateDrawElement::MakeBox( OutDrawElements, SampleLayer, AllottedGeometry.ToPaintGeometry(GridPosition, KeySize), KeyBrush, ESlateDrawEffect::None, DrawColor );
 		}
 
 		// Always draw the filtered position which comes back from whatever is running
 		{
-			FVector2D GridPosition = SampleValueToGridPosition(PreviewFilteredPosition) - (KeySize * .5f);
-			FSlateDrawElement::MakeBox(OutDrawElements, DrawLayerId + 1, AllottedGeometry.ToPaintGeometry(GridPosition, KeySize), KeyBrush, ESlateDrawEffect::None, PreviewKeyColor.GetSpecifiedColor() * 0.5);
+			FVector2D GridPosition = SampleValueToGridPosition(PreviewFilteredPosition) - (PreviewSize * .5f);
+			FSlateDrawElement::MakeBox(OutDrawElements, FilteredPositionLayer, AllottedGeometry.ToPaintGeometry(GridPosition, PreviewSize), PreviewBrush, ESlateDrawEffect::None, PreviewKeyColor.GetSpecifiedColor() * 0.5);
 		}
 
 		if (bPreviewPositionSet)
 		{
-			FVector2D GridPosition = SampleValueToGridPosition(PreviewPosition) - (KeySize * .5f);
-			FSlateDrawElement::MakeBox(OutDrawElements, DrawLayerId + 1, AllottedGeometry.ToPaintGeometry(GridPosition, KeySize), KeyBrush, ESlateDrawEffect::None, PreviewKeyColor.GetSpecifiedColor());
+			FVector2D GridPosition = SampleValueToGridPosition(PreviewPosition) - (PreviewSize * .5f);
+			FSlateDrawElement::MakeBox(OutDrawElements, PreviewPositionLayer, AllottedGeometry.ToPaintGeometry(GridPosition, PreviewSize), PreviewBrush, ESlateDrawEffect::None, PreviewKeyColor.GetSpecifiedColor());
 		}
 
 		if (DragState == EDragState::DragDrop || DragState == EDragState::InvalidDragDrop)
 		{
 			const FVector2D GridPoint = SnapToClosestGridPoint(LocalMousePosition) - (KeySize * .5f);
-			FSlateDrawElement::MakeBox( OutDrawElements, DrawLayerId + 1, AllottedGeometry.ToPaintGeometry(GridPoint, KeySize), KeyBrush, ESlateDrawEffect::None,
+			FSlateDrawElement::MakeBox( OutDrawElements, SampleLayer, AllottedGeometry.ToPaintGeometry(GridPoint, KeySize), KeyBrush, ESlateDrawEffect::None,
 				(DragState == EDragState::DragDrop) ? DropKeyColor.GetSpecifiedColor() : InvalidColor.GetSpecifiedColor() );
 		}
 	}
 
-	DrawLayerId += 1;
+	DrawLayerId += 3;
 }
 
 void SBlendSpaceGridWidget::PaintAxisText(const FGeometry& AllottedGeometry, const FSlateRect& MyCullingRect, FSlateWindowElementList& OutDrawElements, int32& DrawLayerId) const
@@ -469,8 +476,9 @@ void SBlendSpaceGridWidget::PaintAxisText(const FGeometry& AllottedGeometry, con
 	FVector2D TextSize = FontMeasure->Measure(Text, FontInfo);
 
 	// arrow left
-	FVector2D TextPosition = FVector2D(GridCenter.X - (TextSize.X * .5f), CachedGridRectangle.Bottom + TextMargin + (KeySize.Y * .25f));
+	
 	FVector2D ArrowSize = ArrowBrushes[(uint8)EArrowDirection::Left]->GetImageSize();
+	FVector2D TextPosition = FVector2D(GridCenter.X - (TextSize.X * .5f), CachedGridRectangle.Bottom + TextMargin + (ArrowSize.Y * .25f));
 	FVector2D ArrowPosition = FVector2D(TextPosition.X - ArrowSize.X - 10.f/* give padding*/, TextPosition.Y);
 	FSlateDrawElement::MakeBox(OutDrawElements, DrawLayerId + 1, AllottedGeometry.ToPaintGeometry(ArrowPosition, ArrowSize), ArrowBrushes[(uint8)EArrowDirection::Left], ESlateDrawEffect::None, FLinearColor::White);
 
@@ -486,13 +494,13 @@ void SBlendSpaceGridWidget::PaintAxisText(const FGeometry& AllottedGeometry, con
 	TextSize = FontMeasure->Measure(Text, FontInfo);
 
 	// Minimum value
-	FSlateDrawElement::MakeText(OutDrawElements, DrawLayerId + 1, AllottedGeometry.MakeChild(FVector2D(CachedGridRectangle.Left - (TextSize.X * .5f), CachedGridRectangle.Bottom + TextMargin + (KeySize.Y * .25f)), FVector2D(1.0f, 1.0f)).ToPaintGeometry(), Text, FontInfo, ESlateDrawEffect::None, FLinearColor::White);
+	FSlateDrawElement::MakeText(OutDrawElements, DrawLayerId + 1, AllottedGeometry.MakeChild(FVector2D(CachedGridRectangle.Left - (TextSize.X * .5f), CachedGridRectangle.Bottom + TextMargin + (TextSize.Y * .25f)), FVector2D(1.0f, 1.0f)).ToPaintGeometry(), Text, FontInfo, ESlateDrawEffect::None, FLinearColor::White);
 
 	Text = FString::SanitizeFloat(SampleValueMax.X);
 	TextSize = FontMeasure->Measure(Text, FontInfo);
 
 	// Maximum value
-	FSlateDrawElement::MakeText(OutDrawElements, DrawLayerId + 1, AllottedGeometry.MakeChild(FVector2D(CachedGridRectangle.Right - (TextSize.X * .5f), CachedGridRectangle.Bottom + TextMargin + (KeySize.Y * .25f)), FVector2D(1.0f, 1.0f)).ToPaintGeometry(), Text, FontInfo, ESlateDrawEffect::None, FLinearColor::White);
+	FSlateDrawElement::MakeText(OutDrawElements, DrawLayerId + 1, AllottedGeometry.MakeChild(FVector2D(CachedGridRectangle.Right - (TextSize.X * .5f), CachedGridRectangle.Bottom + TextMargin + (TextSize.Y * .25f)), FVector2D(1.0f, 1.0f)).ToPaintGeometry(), Text, FontInfo, ESlateDrawEffect::None, FLinearColor::White);
 
 	// Only draw Y axis labels if this is a 2D grid
 	if (GridType == EGridType::TwoAxis)
@@ -502,9 +510,9 @@ void SBlendSpaceGridWidget::PaintAxisText(const FGeometry& AllottedGeometry, con
 		TextSize = FontMeasure->Measure(Text, FontInfo);
 
 		// arrow up
-		TextPosition = FVector2D(((GridMargin.Left - TextSize.X) * 0.5f - (KeySize.X * .25f)) + GridRatioMargin.Left, GridCenter.Y - (TextSize.Y * .5f));
 		ArrowSize = ArrowBrushes[(uint8)EArrowDirection::Up]->GetImageSize();
-		ArrowPosition = FVector2D(TextPosition.X + TextSize.X * 0.5f - ArrowSize.X * 0.5f, TextPosition.Y - ArrowSize.X - 10.f/* give padding*/);
+		TextPosition = FVector2D(((GridMargin.Left - TextSize.X) * 0.5f - (ArrowSize.X * .25f)) + GridRatioMargin.Left, GridCenter.Y - (TextSize.Y * .5f));
+		ArrowPosition = FVector2D(TextPosition.X + TextSize.X * 0.5f - ArrowSize.X * 0.5f, TextPosition.Y - ArrowSize.Y - 10.f/* give padding*/);
 		FSlateDrawElement::MakeBox(OutDrawElements, DrawLayerId + 1, AllottedGeometry.ToPaintGeometry(ArrowPosition, ArrowSize), ArrowBrushes[(uint8)EArrowDirection::Up], ESlateDrawEffect::None, FLinearColor::White);
 
 		// Label
@@ -519,13 +527,13 @@ void SBlendSpaceGridWidget::PaintAxisText(const FGeometry& AllottedGeometry, con
 		TextSize = FontMeasure->Measure(Text, FontInfo);
 
 		// Minimum value
-		FSlateDrawElement::MakeText(OutDrawElements, DrawLayerId + 1, AllottedGeometry.MakeChild(FVector2D(((GridMargin.Left - TextSize.X) * 0.5f - (KeySize.X * .25f)) + GridRatioMargin.Left, CachedGridRectangle.Bottom - (TextSize.Y * .5f)), FVector2D(1.0f, 1.0f)).ToPaintGeometry(), Text, FontInfo, ESlateDrawEffect::None, FLinearColor::White);
+		FSlateDrawElement::MakeText(OutDrawElements, DrawLayerId + 1, AllottedGeometry.MakeChild(FVector2D(((GridMargin.Left - TextSize.X) * 0.5f - (TextSize.X * .25f)) + GridRatioMargin.Left, CachedGridRectangle.Bottom - (TextSize.Y * .5f)), FVector2D(1.0f, 1.0f)).ToPaintGeometry(), Text, FontInfo, ESlateDrawEffect::None, FLinearColor::White);
 
 		Text = FString::SanitizeFloat(SampleValueMax.Y);
 		TextSize = FontMeasure->Measure(Text, FontInfo);
 
 		// Maximum value
-		FSlateDrawElement::MakeText(OutDrawElements, DrawLayerId + 1, AllottedGeometry.MakeChild(FVector2D(((GridMargin.Left - TextSize.X) * 0.5f - (KeySize.X * .25f) ) + GridRatioMargin.Left, ( GridMargin.Top + GridRatioMargin.Top) - (TextSize.Y * .5f)), FVector2D(1.0f, 1.0f)).ToPaintGeometry(), Text, FontInfo, ESlateDrawEffect::None, FLinearColor::White);
+		FSlateDrawElement::MakeText(OutDrawElements, DrawLayerId + 1, AllottedGeometry.MakeChild(FVector2D(((GridMargin.Left - TextSize.X) * 0.5f - (TextSize.X * .25f) ) + GridRatioMargin.Left, ( GridMargin.Top + GridRatioMargin.Top) - (TextSize.Y * .5f)), FVector2D(1.0f, 1.0f)).ToPaintGeometry(), Text, FontInfo, ESlateDrawEffect::None, FLinearColor::White);
 	}
 
 	DrawLayerId += 1;
@@ -558,7 +566,7 @@ void SBlendSpaceGridWidget::PaintTriangulation(const FGeometry& AllottedGeometry
 								Points.Add(SampleValueToGridPosition(TargetSample.SampleValue));
 
 								// Draw line from and to element
-								FSlateDrawElement::MakeLines(OutDrawElements, DrawLayerId + 1, AllottedGeometry.ToPaintGeometry(), Points, ESlateDrawEffect::None, bShowAnimationNames ? GridLinesColor : FLinearColor::White, true, 0.1f);
+								FSlateDrawElement::MakeLines(OutDrawElements, DrawLayerId + 1, AllottedGeometry.ToPaintGeometry(), Points, ESlateDrawEffect::None, TriangulationColor.GetSpecifiedColor(), true, 0.5f);
 							}
 						}
 					}
@@ -612,11 +620,12 @@ void SBlendSpaceGridWidget::PaintAnimationNames(const FGeometry& AllottedGeometr
 				GridPosition += FVector2D(KeySize.X, -KeySize.X * .5f);
 			}
 
-			FSlateDrawElement::MakeText(OutDrawElements, DrawLayerId + 1, AllottedGeometry.MakeChild(FVector2D(GridPosition.X, GridPosition.Y), FVector2D(1.0f, 1.0f)).ToPaintGeometry(), Name, FontInfo, ESlateDrawEffect::None, FLinearColor::White);
+			FSlateDrawElement::MakeBox(OutDrawElements, DrawLayerId + 1, AllottedGeometry.MakeChild(FVector2D(GridPosition.X - 6, GridPosition.Y - 2), TextSize + FVector2D(8.0f, 4.0f)).ToPaintGeometry(), LabelBrush, ESlateDrawEffect::None, FLinearColor::Black);
+			FSlateDrawElement::MakeText(OutDrawElements, DrawLayerId + 2, AllottedGeometry.MakeChild(FVector2D(GridPosition.X, GridPosition.Y), FVector2D(1.0f, 1.0f)).ToPaintGeometry(), Name, FontInfo, ESlateDrawEffect::None, FLinearColor::White);
 		}
 	}
 
-	DrawLayerId += 1;
+	DrawLayerId += 2;
 }
 
 FReply SBlendSpaceGridWidget::OnDrop(const FGeometry& MyGeometry, const FDragDropEvent& DragDropEvent)
@@ -733,12 +742,6 @@ FReply SBlendSpaceGridWidget::OnMouseButtonUp(const FGeometry& MyGeometry, const
 					// Process drag ending			
 					ResetToolTip();
 				}
-				else if (DragState == EDragState::DragPreview)
-				{
-					// Process preview drag ending			
-					ResetToolTip();
-					StopPreviewing();
-				}
 
 				// Reset drag state and index
 				DragState = EDragState::None;
@@ -765,18 +768,7 @@ FReply SBlendSpaceGridWidget::OnMouseButtonDown(const FGeometry& MyGeometry, con
 			if (MouseEvent.GetEffectingButton() == EKeys::LeftMouseButton)
 			{
 				// If we are over a sample, make it our currently (dragged) sample
-				if (bHighlightPreviewPin)
-				{
-					ResetToolTip();
-					DragState = EDragState::DragPreview;
-					MouseDownPosition = LocalMousePosition;
-					StartPreviewing();
-					ShowToolTip();
-
-					// Start mouse capture
-					return FReply::Handled().CaptureMouse(SharedThis(this));
-				}
-				else if (HighlightedSampleIndex != INDEX_NONE)
+				if (HighlightedSampleIndex != INDEX_NONE)
 				{
 					DraggedSampleIndex = SelectedSampleIndex = HighlightedSampleIndex;
 					HighlightedSampleIndex = INDEX_NONE;
@@ -830,58 +822,64 @@ FReply SBlendSpaceGridWidget::OnMouseMove(const FGeometry& MyGeometry, const FPo
 		LocalMousePosition = MyGeometry.AbsoluteToLocal(MouseEvent.GetScreenSpacePosition());
 		LastMousePosition = MouseEvent.GetScreenSpacePosition();
 
-		if (!bReadOnly && this->HasMouseCapture())
+		if(!bReadOnly)
 		{
-			if (DragState == EDragState::None)
+			if (this->HasMouseCapture())
 			{
-				if (HighlightedSampleIndex != INDEX_NONE)
+				if (DragState == EDragState::None)
 				{
-					DragState = EDragState::DragSample;
-					DraggedSampleIndex = HighlightedSampleIndex;
-					HighlightedSampleIndex = INDEX_NONE;
+					if (HighlightedSampleIndex != INDEX_NONE)
+					{
+						DragState = EDragState::DragSample;
+						DraggedSampleIndex = HighlightedSampleIndex;
+						HighlightedSampleIndex = INDEX_NONE;
+						return FReply::Handled();
+					}
 				}
-			}
-			else if (DragState == EDragState::PreDrag)
-			{
-				// Actually start dragging
-				if ((LocalMousePosition - MouseDownPosition).SizeSquared() > DragThreshold)
+				else if (DragState == EDragState::PreDrag)
 				{
-					DragState = EDragState::DragSample;
-					HighlightedSampleIndex = INDEX_NONE;
-					ShowToolTip();
+					// Actually start dragging
+					if ((LocalMousePosition - MouseDownPosition).SizeSquared() > DragThreshold)
+					{
+						DragState = EDragState::DragSample;
+						HighlightedSampleIndex = INDEX_NONE;
+						ShowToolTip();
+						return FReply::Handled();
+					}
 				}
 			}
-		}
-		else if (IsHovered() && bMouseIsOverGeometry && DragState != EDragState::DragPreview)
-		{
-			if ((MouseEvent.IsLeftShiftDown() || MouseEvent.IsRightShiftDown()))
+			else if (IsHovered() && bMouseIsOverGeometry)
 			{
-				StartPreviewing();
-				DragState = EDragState::Preview;
-				// Make tool tip visible (this will display the current preview sample value)
-				ShowToolTip();			
+				if (MouseEvent.IsLeftShiftDown() || MouseEvent.IsRightShiftDown())
+				{
+					StartPreviewing();
+					DragState = EDragState::Preview;
+					// Make tool tip visible (this will display the current preview sample value)
+					ShowToolTip();			
 
-				// Set flag for showing advanced preview info in tooltip
-				bAdvancedPreview = MouseEvent.IsLeftControlDown() || MouseEvent.IsRightControlDown();
-			}
-			else if(Position.IsSet())
-			{
-				StartPreviewing();
-				DragState = EDragState::None;
-				ShowToolTip();
+					// Set flag for showing advanced preview info in tooltip
+					bAdvancedPreview = MouseEvent.IsLeftControlDown() || MouseEvent.IsRightControlDown();
+					return FReply::Handled();
+				}
+				else if(Position.IsSet())
+				{
+					StartPreviewing();
+					DragState = EDragState::None;
+					ShowToolTip();
 
-				// Set flag for showing advanced preview info in tooltip
-				bAdvancedPreview = MouseEvent.IsLeftControlDown() || MouseEvent.IsRightControlDown();
-			}
-			else if (bSamplePreviewing)
-			{
-				StopPreviewing();
-				DragState = EDragState::None;
-				ResetToolTip();
+					// Set flag for showing advanced preview info in tooltip
+					bAdvancedPreview = MouseEvent.IsLeftControlDown() || MouseEvent.IsRightControlDown();
+					return FReply::Handled();
+				}
+				else if (bSamplePreviewing)
+				{
+					StopPreviewing();
+					DragState = EDragState::None;
+					ResetToolTip();
+					return FReply::Handled();
+				}
 			}
 		}
-
-		return FReply::Handled();
 	}
 
 	return FReply::Unhandled();
@@ -968,9 +966,9 @@ FReply SBlendSpaceGridWidget::OnKeyDown(const FGeometry& MyGeometry, const FKeyE
 		if(const UBlendSpace* BlendSpace = BlendSpaceBase.Get())
 		{
 			// Start previewing when either one of the shift keys is pressed
-			if (bMouseIsOverGeometry)
+			if (IsHovered() && bMouseIsOverGeometry)
 			{
-				if (DragState != EDragState::DragPreview && ((InKeyEvent.GetKey() == EKeys::LeftShift) || (InKeyEvent.GetKey() == EKeys::RightShift)))
+				if ((InKeyEvent.GetKey() == EKeys::LeftShift) || (InKeyEvent.GetKey() == EKeys::RightShift))
 				{
 					StartPreviewing();
 					DragState = EDragState::Preview;
@@ -999,7 +997,7 @@ FReply SBlendSpaceGridWidget::OnKeyUp(const FGeometry& MyGeometry, const FKeyEve
 		if(const UBlendSpace* BlendSpace = BlendSpaceBase.Get())
 		{
 			// Stop previewing when shift keys are released 
-			if (DragState != EDragState::DragPreview && ((InKeyEvent.GetKey() == EKeys::LeftShift) || (InKeyEvent.GetKey() == EKeys::RightShift)))
+			if ((InKeyEvent.GetKey() == EKeys::LeftShift) || (InKeyEvent.GetKey() == EKeys::RightShift))
 			{
 				StopPreviewing();
 				DragState = EDragState::None;
@@ -1338,7 +1336,7 @@ void SBlendSpaceGridWidget::StartPreviewing()
 	bSamplePreviewing = true;
 	LastPreviewingMousePosition = LocalMousePosition;
 	FModifierKeysState ModifierKeyState = FSlateApplication::Get().GetModifierKeys();
-	bool bIsManualPreviewing = (ModifierKeyState.IsLeftShiftDown() || ModifierKeyState.IsRightShiftDown()) || DragState == EDragState::DragPreview;
+	bool bIsManualPreviewing = !bReadOnly && IsHovered() && bMouseIsOverGeometry && (ModifierKeyState.IsLeftShiftDown() || ModifierKeyState.IsRightShiftDown());
 	PreviewPosition = Position.IsSet() && !bIsManualPreviewing ? Position.Get() : GridPositionToSampleValue(LastPreviewingMousePosition, false);
 	PreviewFilteredPosition = FilteredPosition.IsSet() ? FilteredPosition.Get() : PreviewPosition;
 	bPreviewPositionSet = true;	
@@ -1421,7 +1419,6 @@ FText SBlendSpaceGridWidget::GetToolTipAnimationName() const
 		
 				// If we are previewing return a descriptive label
 				case EDragState::Preview:
-				case EDragState::DragPreview:
 				{
 					ToolTipText = PreviewValue;
 					break;
@@ -1556,7 +1553,6 @@ FText SBlendSpaceGridWidget::GetToolTipSampleValue() const
 				}
 
 				// If we are setting the preview value return the current preview sample value
-				case EDragState::DragPreview:
 				case EDragState::Preview:
 				{
 					ToolTipText = FText::Format(ValueFormattingText, ParameterXName, FText::FromString(FString::SanitizeFloat(PreviewPosition.X)), ParameterYName, FText::FromString(FString::SanitizeFloat(PreviewPosition.Y)));
@@ -1585,7 +1581,7 @@ void SBlendSpaceGridWidget::EnableStatusBarMessage(bool bEnable)
 				{
 					StatusBarMessageHandle = StatusBarSubsystem->PushStatusBarMessage(StatusBarName, MakeAttributeLambda([]()
 					{
-						return LOCTEXT("StatusBarMssage", "Hold Ctrl for weight details, click and drag or hold Shift to move preview value");
+						return LOCTEXT("StatusBarMssage", "Hold Ctrl for weight details, hold Shift to move preview value");
 					}));
 				}
 			}
@@ -1626,7 +1622,7 @@ FText SBlendSpaceGridWidget::GetSampleErrorMessage(const FBlendSample &BlendSamp
 
 void SBlendSpaceGridWidget::ShowToolTip()
 {
-	if(DragState != EDragState::DragPreview && HighlightedSampleIndex != INDEX_NONE && ToolTipSampleIndex != HighlightedSampleIndex)
+	if(HighlightedSampleIndex != INDEX_NONE && ToolTipSampleIndex != HighlightedSampleIndex)
 	{
 		ToolTipSampleIndex = HighlightedSampleIndex;
 		if(OnExtendSampleTooltip.IsBound())
@@ -2011,7 +2007,7 @@ void SBlendSpaceGridWidget::Tick(const FGeometry& AllottedGeometry, const double
 			LastPreviewingMousePosition.X = LocalMousePosition.X;
 			LastPreviewingMousePosition.Y = LocalMousePosition.Y;
 			FModifierKeysState ModifierKeyState = FSlateApplication::Get().GetModifierKeys();
-			bool bIsManualPreviewing = (ModifierKeyState.IsLeftShiftDown() || ModifierKeyState.IsRightShiftDown()) || DragState == EDragState::DragPreview;
+			bool bIsManualPreviewing = !bReadOnly && IsHovered() && bMouseIsOverGeometry && (ModifierKeyState.IsLeftShiftDown() || ModifierKeyState.IsRightShiftDown());
 			PreviewPosition = Position.IsSet() && !bIsManualPreviewing ? Position.Get() : GridPositionToSampleValue(LastPreviewingMousePosition, false);
 			PreviewPosition = BlendSpace->GetClampedAndWrappedBlendInput(PreviewPosition);
 
@@ -2163,7 +2159,7 @@ bool SBlendSpaceGridWidget::ValidateAnimationSequence(const UAnimSequence* Anima
 const bool SBlendSpaceGridWidget::IsPreviewing() const 
 { 
 	FModifierKeysState ModifierKeyState = FSlateApplication::Get().GetModifierKeys();
-	bool bIsManualPreviewing = (ModifierKeyState.IsLeftShiftDown() || ModifierKeyState.IsRightShiftDown()) || DragState == EDragState::DragPreview;
+	bool bIsManualPreviewing = !bReadOnly && IsHovered() && bMouseIsOverGeometry && (ModifierKeyState.IsLeftShiftDown() || ModifierKeyState.IsRightShiftDown());
 	return (bSamplePreviewing && !Position.IsSet()) || (Position.IsSet() && bIsManualPreviewing);
 }
 
diff --git a/Engine/Source/Editor/Persona/Private/SAnimationBlendSpaceGridWidget.h b/Engine/Source/Editor/Persona/Private/SAnimationBlendSpaceGridWidget.h
--- a/Engine/Source/Editor/Persona/Private/SAnimationBlendSpaceGridWidget.h
+++ b/Engine/Source/Editor/Persona/Private/SAnimationBlendSpaceGridWidget.h
@@ -67,8 +67,6 @@ protected:
 		PreDrag,
 		/** The user is dragging the selected sample. */
 		DragSample,
-		/** The user is dragging the preview pin */
-		DragPreview,
 		/** The user is setting the preview value */
 		Preview,
 		/** The user is dropping a new sample onto the grid */
@@ -318,8 +316,11 @@ private:
 
 	/** Sample drawing data */
 	FVector2D KeySize;
+	FVector2D PreviewSize;
 	const FSlateBrush* KeyBrush;
+	const FSlateBrush* PreviewBrush;
 	const FSlateBrush* BackgroundImage;
+	const FSlateBrush* LabelBrush;
 	enum class EArrowDirection : uint8 
 	{
 		Left, 
@@ -345,6 +346,7 @@ private:
 	FMargin GridMargin;
 	FLinearColor GridLinesColor;
 	FLinearColor GridOutlineColor;
+	FSlateColor TriangulationColor;
 	FSlateFontInfo FontInfo;
 	float TextMargin;
 	bool bShowTriangulation;
