Fixed a crash when using Set Preview Mesh on some duplicated skeleton.

diff --git a/Engine/Source/Editor/Persona/Private/PersonaToolkit.cpp b/Engine/Source/Editor/Persona/Private/PersonaToolkit.cpp
--- a/Engine/Source/Editor/Persona/Private/PersonaToolkit.cpp
+++ b/Engine/Source/Editor/Persona/Private/PersonaToolkit.cpp
@@ -320,6 +320,11 @@ void FPersonaToolkit::SetPreviewMesh(class USkeletalMesh* InSkeletalMesh, bool b
 			{
 				AssetToReopen = PhysicsAsset;
 			}
+			else if (InitialAssetClass == USkeleton::StaticClass())
+			{
+				AssetToReopen = Skeleton.Get();
+			}
+
 			check(AssetToReopen);
 
 			GEditor->GetEditorSubsystem<UAssetEditorSubsystem>()->CloseAllEditorsForAsset(AssetToReopen);
