Re-added missing include to fix non-unity build

diff --git a/Engine/Source/Editor/Persona/Private/SAnimationBlendSpaceGridWidget.cpp b/Engine/Source/Editor/Persona/Private/SAnimationBlendSpaceGridWidget.cpp
--- a/Engine/Source/Editor/Persona/Private/SAnimationBlendSpaceGridWidget.cpp
+++ b/Engine/Source/Editor/Persona/Private/SAnimationBlendSpaceGridWidget.cpp
@@ -27,6 +27,7 @@
 #include "DragAndDrop/AssetDragDropOp.h"
 #include "Classes/EditorStyleSettings.h"
 
+#include "Widgets/Input/SButton.h"
 #include "Fonts/FontMeasure.h"
 #include "Modules/ModuleManager.h"
 
