Animation editor tracks for custom attributes

diff --git a/Engine/Source/Editor/Persona/Private/AnimTimeline/AnimModel_AnimSequenceBase.cpp b/Engine/Source/Editor/Persona/Private/AnimTimeline/AnimModel_AnimSequenceBase.cpp
--- a/Engine/Source/Editor/Persona/Private/AnimTimeline/AnimModel_AnimSequenceBase.cpp
+++ b/Engine/Source/Editor/Persona/Private/AnimTimeline/AnimModel_AnimSequenceBase.cpp
@@ -10,6 +10,9 @@
 #include "AnimTimelineTrack_FloatCurve.h"
 #include "AnimTimelineTrack_VectorCurve.h"
 #include "AnimTimelineTrack_TransformCurve.h"
+#include "AnimTimelineTrack_Attributes.h"
+#include "AnimTimelineTrack_PerBoneAttributes.h"
+#include "AnimTimelineTrack_Attribute.h"
 #include "AnimSequenceTimelineCommands.h"
 #include "Framework/Commands/UICommandList.h"
 #include "IAnimationEditor.h"
@@ -155,6 +158,9 @@ void FAnimModel_AnimSequenceBase::RefreshTracks()
 	// Add curves
 	RefreshCurveTracks();
 
+	// Add attributes
+	RefreshAttributeTracks();
+
 	// Snaps
 	RefreshSnapTimes();
 
@@ -264,7 +270,50 @@ void FAnimModel_AnimSequenceBase::RefreshCurveTracks()
 			ScaleCurveTrack->AddChild(MakeShared<FAnimTimelineTrack_Curve>(&TransformCurve.ScaleCurve.FloatCurves[0], TransformCurve.Name, 6, ERawCurveTrackTypes::RCT_Transform, XName, FText::Format(ComponentFormat, TransformName, ScaleName, XName), XColor, XColor, SharedThis(this)));
 			ScaleCurveTrack->AddChild(MakeShared<FAnimTimelineTrack_Curve>(&TransformCurve.ScaleCurve.FloatCurves[1], TransformCurve.Name, 7, ERawCurveTrackTypes::RCT_Transform, YName, FText::Format(ComponentFormat, TransformName, ScaleName, YName), YColor, YColor, SharedThis(this)));
 			ScaleCurveTrack->AddChild(MakeShared<FAnimTimelineTrack_Curve>(&TransformCurve.ScaleCurve.FloatCurves[2], TransformCurve.Name, 8, ERawCurveTrackTypes::RCT_Transform, ZName, FText::Format(ComponentFormat, TransformName, ScaleName, ZName), ZColor, ZColor, SharedThis(this)));
+		}		
+	}
+}
+
+void FAnimModel_AnimSequenceBase::RefreshAttributeTracks()
+{
+	if (UAnimSequence* AnimSequence = Cast<UAnimSequence>(AnimSequenceBase))
+	{
+		if (!AttributesRoot.IsValid())
+		{
+			// Add a root track for attributes
+			AttributesRoot = MakeShared<FAnimTimelineTrack_Attributes>(SharedThis(this));
 		}
+
+		AttributesRoot->ClearChildren();
+		RootTracks.Add(AttributesRoot.ToSharedRef());
+			   
+		TMap<FName, TSharedPtr<FAnimTimelineTrack_PerBoneAttributes>> BoneTracks;
+
+
+		// Next add a track for each attribute curve
+		for (const FAnimatedBoneAttribute& BoneAttribute : AnimSequence->GetDataModel()->GetAttributes())
+		{
+			TSharedPtr<FAnimTimelineTrack_PerBoneAttributes> BoneTrack = nullptr;
+
+			if (BoneTracks.Contains(BoneAttribute.Identifier.GetBoneName()))
+			{
+				BoneTrack = BoneTracks.FindChecked(BoneAttribute.Identifier.GetBoneName());
+			}
+			else
+			{
+				TSharedRef<FAnimTimelineTrack_PerBoneAttributes> BoneTrackRef = MakeShared<FAnimTimelineTrack_PerBoneAttributes>(BoneAttribute.Identifier.GetBoneName(), SharedThis(this));
+				AttributesRoot->AddChild(BoneTrackRef);
+
+				BoneTracks.Add(BoneAttribute.Identifier.GetBoneName(), BoneTrackRef);
+
+				BoneTrack = BoneTrackRef;
+
+				BoneTrackRef->SetExpanded(false);
+			}
+
+			TSharedRef<FAnimTimelineTrack_Attribute> AttributeTrack = MakeShared<FAnimTimelineTrack_Attribute>(BoneAttribute, SharedThis(this));
+			BoneTrack->AddChild(AttributeTrack);
+		}		
 	}
 }
 
@@ -279,6 +328,9 @@ void FAnimModel_AnimSequenceBase::OnDataModelChanged(const EAnimDataModelNotifyT
 		case EAnimDataModelNotifyType::TrackAdded:
 		case EAnimDataModelNotifyType::TrackChanged:
 		case EAnimDataModelNotifyType::TrackRemoved:
+		case EAnimDataModelNotifyType::AttributeAdded:
+		case EAnimDataModelNotifyType::AttributeChanged:
+		case EAnimDataModelNotifyType::AttributeRemoved:
 		case EAnimDataModelNotifyType::SequenceLengthChanged:
 		case EAnimDataModelNotifyType::FrameRateChanged:
 		{
diff --git a/Engine/Source/Editor/Persona/Private/AnimTimeline/AnimModel_AnimSequenceBase.h b/Engine/Source/Editor/Persona/Private/AnimTimeline/AnimModel_AnimSequenceBase.h
--- a/Engine/Source/Editor/Persona/Private/AnimTimeline/AnimModel_AnimSequenceBase.h
+++ b/Engine/Source/Editor/Persona/Private/AnimTimeline/AnimModel_AnimSequenceBase.h
@@ -14,6 +14,7 @@ class FAnimTimelineTrack_Notifies;
 class FAnimTimelineTrack_Curves;
 class FAnimTimelineTrack;
 class FAnimTimelineTrack_NotifiesPanel;
+class FAnimTimelineTrack_Attributes;
 enum class EFrameNumberDisplayFormats : uint8;
 
 /** Anim model for an anim sequence base */
@@ -62,6 +63,9 @@ protected:
 	/** Refresh curve tracks */
 	void RefreshCurveTracks();
 
+	/** Refresh attribute tracks */
+	void RefreshAttributeTracks();
+
 	/** Callback for any change made to the UAnimDataModel embedded in the AnimSequenceBase instance this represents */
 	void OnDataModelChanged(const EAnimDataModelNotifyType& NotifyType, UAnimDataModel* Model, const FAnimDataModelNotifPayload& PayLoad);
 
@@ -94,6 +98,9 @@ private:
 	/** Root track for additive layers */
 	TSharedPtr<FAnimTimelineTrack> AdditiveRoot;
 
+	/** Root track for custom attributes */
+	TSharedPtr<FAnimTimelineTrack_Attributes> AttributesRoot;
+
 	/** Display flags for notifies track */
 	bool NotifiesTimingElementNodeDisplayFlags[ETimingElementType::Max];
 
diff --git a/Engine/Source/Editor/Persona/Private/AnimTimeline/AnimTimelineTrack_Attribute.cpp b/Engine/Source/Editor/Persona/Private/AnimTimeline/AnimTimelineTrack_Attribute.cpp
new file mode 100644
--- /dev/null
+++ b/Engine/Source/Editor/Persona/Private/AnimTimeline/AnimTimelineTrack_Attribute.cpp
@@ -0,0 +1,45 @@
+// Copyright Epic Games, Inc. All Rights Reserved.
+
+#include "AnimTimelineTrack_Attribute.h"
+#include "AnimTimeline/SAnimOutlinerItem.h"
+
+#define LOCTEXT_NAMESPACE "FAnimTimelineTrack_CustomBoneAttributes"
+
+ANIMTIMELINE_IMPLEMENT_TRACK(FAnimTimelineTrack_Attribute);
+
+FAnimTimelineTrack_Attribute::FAnimTimelineTrack_Attribute(const FAnimatedBoneAttribute& InAttribute, const TSharedRef<FAnimModel>& InModel)
+	: FAnimTimelineTrack(LOCTEXT("AttributeTrackLabel", "Attribute"), LOCTEXT("AttributeToolTip", "Attribute contained in this asset"), InModel), Attribute(InAttribute)
+{
+}
+
+TSharedRef<SWidget> FAnimTimelineTrack_Attribute::GenerateContainerWidgetForOutliner(const TSharedRef<SAnimOutlinerItem>& InRow)
+{
+	TSharedPtr<SBorder> OuterBorder;
+	TSharedPtr<SHorizontalBox> InnerHorizontalBox;
+	OutlinerWidget = GenerateStandardOutlinerWidget(InRow, false, OuterBorder, InnerHorizontalBox);
+	
+	InnerHorizontalBox->AddSlot()
+		.VAlign(VAlign_Center)
+		.HAlign(HAlign_Left)
+		.Padding(2.0f, 1.0f)
+		.AutoWidth()
+		[
+			SNew(STextBlock)
+			.Text(this, &FAnimTimelineTrack_Attribute::GetLabel)
+			.HighlightText(InRow->GetHighlightText())
+		];
+
+	return OutlinerWidget.ToSharedRef();
+}
+
+FText FAnimTimelineTrack_Attribute::GetLabel() const
+{
+	return FText::FromName(Attribute.Identifier.GetName());
+}
+
+FText FAnimTimelineTrack_Attribute::GetToolTipText() const
+{
+	return FText::Format(LOCTEXT("AttributeTooltipFormat", "Attribute: {0}\nType: {1}\nNumber of Keys: {2}"), FText::FromName(Attribute.Identifier.GetName()), FText::FromName(Attribute.Identifier.GetType()->GetFName()), FText::AsNumber(Attribute.Curve.GetNumKeys()));
+}
+
+#undef LOCTEXT_NAMESPACE
\ No newline at end of file
diff --git a/Engine/Source/Editor/Persona/Private/AnimTimeline/AnimTimelineTrack_Attribute.h b/Engine/Source/Editor/Persona/Private/AnimTimeline/AnimTimelineTrack_Attribute.h
new file mode 100644
--- /dev/null
+++ b/Engine/Source/Editor/Persona/Private/AnimTimeline/AnimTimelineTrack_Attribute.h
@@ -0,0 +1,23 @@
+// Copyright Epic Games, Inc. All Rights Reserved.
+
+#pragma once
+
+#include "AnimTimelineTrack.h"
+
+/** Animation timeline track inserted for each animated attribute (inserted as child of FAnimTimelineTrack_PerBoneAttributes) */
+class FAnimTimelineTrack_Attribute : public FAnimTimelineTrack
+{
+	ANIMTIMELINE_DECLARE_TRACK(FAnimTimelineTrack_Attribute, FAnimTimelineTrack);
+
+public:
+	FAnimTimelineTrack_Attribute(const FAnimatedBoneAttribute& InAttribute, const TSharedRef<FAnimModel>& InModel);
+
+	/** FAnimTimelineTrack interface */
+	virtual TSharedRef<SWidget> GenerateContainerWidgetForOutliner(const TSharedRef<SAnimOutlinerItem>& InRow) override;
+	virtual FText GetLabel() const override;
+	virtual FText GetToolTipText() const override;
+	   	   
+private:
+	const FAnimatedBoneAttribute& Attribute;
+	TSharedPtr<SWidget> OutlinerWidget;
+};
\ No newline at end of file
diff --git a/Engine/Source/Editor/Persona/Private/AnimTimeline/AnimTimelineTrack_Attributes.cpp b/Engine/Source/Editor/Persona/Private/AnimTimeline/AnimTimelineTrack_Attributes.cpp
new file mode 100644
--- /dev/null
+++ b/Engine/Source/Editor/Persona/Private/AnimTimeline/AnimTimelineTrack_Attributes.cpp
@@ -0,0 +1,46 @@
+// Copyright Epic Games, Inc. All Rights Reserved.
+
+#include "AnimTimelineTrack_Attributes.h"
+#include "PersonaUtils.h"
+#include "Widgets/SBoxPanel.h"
+#include "AnimSequenceTimelineCommands.h"
+#include "Framework/MultiBox/MultiBoxBuilder.h"
+#include "Animation/AnimSequence.h"
+#include "Widgets/Input/STextEntryPopup.h"
+#include "Framework/Application/SlateApplication.h"
+#include "ScopedTransaction.h"
+#include "SAnimOutlinerItem.h"
+
+#define LOCTEXT_NAMESPACE "FAnimTimelineTrack_Attributes"
+
+ANIMTIMELINE_IMPLEMENT_TRACK(FAnimTimelineTrack_Attributes);
+
+FAnimTimelineTrack_Attributes::FAnimTimelineTrack_Attributes(const TSharedRef<FAnimModel>& InModel)
+	: FAnimTimelineTrack(LOCTEXT("AttributesRootTrackLabel", "Attributes"), LOCTEXT("AttributesRootToolTip", "Animated (per bone) Attribute data contained in this asset"), InModel)
+{
+}
+
+TSharedRef<SWidget> FAnimTimelineTrack_Attributes::GenerateContainerWidgetForOutliner(const TSharedRef<SAnimOutlinerItem>& InRow)
+{
+	TSharedPtr<SBorder> OuterBorder;
+	TSharedPtr<SHorizontalBox> InnerHorizontalBox;
+	OutlinerWidget = GenerateStandardOutlinerWidget(InRow, false, OuterBorder, InnerHorizontalBox);
+
+	OuterBorder->SetBorderBackgroundColor(FEditorStyle::GetColor("AnimTimeline.Outliner.HeaderColor"));
+
+	InnerHorizontalBox->AddSlot()
+		.VAlign(VAlign_Center)
+		.HAlign(HAlign_Left)
+		.Padding(2.0f, 1.0f)
+		.AutoWidth()
+		[
+			SNew(STextBlock)
+			.TextStyle(&FEditorStyle::Get().GetWidgetStyle<FTextBlockStyle>("AnimTimeline.Outliner.Label"))
+			.Text(this, &FAnimTimelineTrack_Attributes::GetLabel)
+			.HighlightText(InRow->GetHighlightText())
+		];
+
+	return OutlinerWidget.ToSharedRef();
+}
+
+#undef LOCTEXT_NAMESPACE
\ No newline at end of file
diff --git a/Engine/Source/Editor/Persona/Private/AnimTimeline/AnimTimelineTrack_Attributes.h b/Engine/Source/Editor/Persona/Private/AnimTimeline/AnimTimelineTrack_Attributes.h
new file mode 100644
--- /dev/null
+++ b/Engine/Source/Editor/Persona/Private/AnimTimeline/AnimTimelineTrack_Attributes.h
@@ -0,0 +1,20 @@
+// Copyright Epic Games, Inc. All Rights Reserved.
+
+#pragma once
+
+#include "AnimTimelineTrack.h"
+
+/** Root-level animation timeline track under which per-bone animated attributes are inserted */
+class FAnimTimelineTrack_Attributes : public FAnimTimelineTrack
+{
+	ANIMTIMELINE_DECLARE_TRACK(FAnimTimelineTrack_Attributes, FAnimTimelineTrack);
+
+public:
+	FAnimTimelineTrack_Attributes(const TSharedRef<FAnimModel>& InModel);
+
+	/** FAnimTimelineTrack interface */
+	virtual TSharedRef<SWidget> GenerateContainerWidgetForOutliner(const TSharedRef<SAnimOutlinerItem>& InRow) override;
+
+private:
+	TSharedPtr<SWidget> OutlinerWidget;
+};
\ No newline at end of file
diff --git a/Engine/Source/Editor/Persona/Private/AnimTimeline/AnimTimelineTrack_PerBoneAttributes.cpp b/Engine/Source/Editor/Persona/Private/AnimTimeline/AnimTimelineTrack_PerBoneAttributes.cpp
new file mode 100644
--- /dev/null
+++ b/Engine/Source/Editor/Persona/Private/AnimTimeline/AnimTimelineTrack_PerBoneAttributes.cpp
@@ -0,0 +1,49 @@
+// Copyright Epic Games, Inc. All Rights Reserved.
+
+#include "AnimTimelineTrack_PerBoneAttributes.h"
+#include "Framework/Application/SlateApplication.h"
+#include "AnimTimeline/SAnimOutlinerItem.h"
+
+#define LOCTEXT_NAMESPACE "FAnimTimelineTrack_PerBoneAttributes"
+
+ANIMTIMELINE_IMPLEMENT_TRACK(FAnimTimelineTrack_PerBoneAttributes);
+
+FAnimTimelineTrack_PerBoneAttributes::FAnimTimelineTrack_PerBoneAttributes(const FName& InBoneName, const TSharedRef<FAnimModel>& InModel)
+	: FAnimTimelineTrack(LOCTEXT("BoneAttributesTrackLabel", "Bone Attributes"), LOCTEXT("BoneAttributesToolTip", "Contained Attributes for specific Bone"), InModel), BoneName(InBoneName)
+{
+}
+
+TSharedRef<SWidget> FAnimTimelineTrack_PerBoneAttributes::GenerateContainerWidgetForOutliner(const TSharedRef<SAnimOutlinerItem>& InRow)
+{
+	TSharedPtr<SBorder> OuterBorder;
+	TSharedPtr<SHorizontalBox> InnerHorizontalBox;
+	OutlinerWidget = GenerateStandardOutlinerWidget(InRow, false, OuterBorder, InnerHorizontalBox);
+
+	OuterBorder->SetBorderBackgroundColor(FEditorStyle::GetColor("AnimTimeline.Outliner.HeaderColor"));
+
+	InnerHorizontalBox->AddSlot()
+		.VAlign(VAlign_Center)
+		.HAlign(HAlign_Left)
+		.Padding(2.0f, 1.0f)
+		.AutoWidth()
+		[
+			SNew(STextBlock)
+			.TextStyle(&FEditorStyle::Get().GetWidgetStyle<FTextBlockStyle>("AnimTimeline.Outliner.Label"))
+			.Text(this, &FAnimTimelineTrack_PerBoneAttributes::GetLabel)
+			.HighlightText(InRow->GetHighlightText())
+		];
+
+	return OutlinerWidget.ToSharedRef();
+}
+
+FText FAnimTimelineTrack_PerBoneAttributes::GetLabel() const
+{
+	return FText::FromName(BoneName);
+}
+
+FText FAnimTimelineTrack_PerBoneAttributes::GetToolTipText() const
+{
+	return FText::Format(LOCTEXT("BoneAttributesTooltipFormat", "Attributes for Bone: {0}\nNumber of Attributes: {1}"), FText::FromName(BoneName), FText::AsNumber(Children.Num()));
+}
+
+#undef LOCTEXT_NAMESPACE
\ No newline at end of file
diff --git a/Engine/Source/Editor/Persona/Private/AnimTimeline/AnimTimelineTrack_PerBoneAttributes.h b/Engine/Source/Editor/Persona/Private/AnimTimeline/AnimTimelineTrack_PerBoneAttributes.h
new file mode 100644
--- /dev/null
+++ b/Engine/Source/Editor/Persona/Private/AnimTimeline/AnimTimelineTrack_PerBoneAttributes.h
@@ -0,0 +1,23 @@
+// Copyright Epic Games, Inc. All Rights Reserved.
+
+#pragma once
+
+#include "AnimTimelineTrack.h"
+
+/** Animation timeline track inserted for each unique bone containing animated attributes (inserted as child of FAnimTimelineTrack_Attributes) */
+class FAnimTimelineTrack_PerBoneAttributes : public FAnimTimelineTrack
+{
+	ANIMTIMELINE_DECLARE_TRACK(FAnimTimelineTrack_PerBoneAttributes, FAnimTimelineTrack);
+
+public:
+	FAnimTimelineTrack_PerBoneAttributes(const FName& InBoneName, const TSharedRef<FAnimModel>& InModel);
+
+	/** FAnimTimelineTrack interface */
+	virtual TSharedRef<SWidget> GenerateContainerWidgetForOutliner(const TSharedRef<SAnimOutlinerItem>& InRow) override;
+	virtual FText GetLabel() const override;
+	virtual FText GetToolTipText() const override;
+	   	   
+private:
+	FName BoneName;
+	TSharedPtr<SWidget> OutlinerWidget;
+};
\ No newline at end of file
