Make the asset family buttons aware of compatible skeletons.
This will make persona find compatible assets as well when using skeleton remapping.

#rnx

diff --git a/Engine/Source/Editor/Persona/Private/PersonaAssetFamily.cpp b/Engine/Source/Editor/Persona/Private/PersonaAssetFamily.cpp
--- a/Engine/Source/Editor/Persona/Private/PersonaAssetFamily.cpp
+++ b/Engine/Source/Editor/Persona/Private/PersonaAssetFamily.cpp
@@ -60,6 +60,14 @@ static void FindAssets(const USkeleton* InSkeleton, TArray<FAssetData>& OutAsset
 	Filter.bRecursiveClasses = true;
 	Filter.ClassNames.Add(AssetType::StaticClass()->GetFName());
 	Filter.TagsAndValues.Add(SkeletonTag, FAssetData(InSkeleton).GetExportTextName());
+	
+	// Also include all compatible assets.
+	FString CompatibleTagValue;
+	for (const auto& CompatibleSkeleton : InSkeleton->GetCompatibleSkeletons())
+	{
+		CompatibleTagValue = FString::Format(TEXT("{0}'{1}'"), { *USkeleton::StaticClass()->GetName(), *CompatibleSkeleton.ToString() });
+		Filter.TagsAndValues.Add(SkeletonTag, CompatibleTagValue);
+	}
 
 	AssetRegistryModule.Get().GetAssets(Filter, OutAssetData);
 }
@@ -255,7 +263,10 @@ bool FPersonaAssetFamily::IsAssetCompatible(const FAssetData& InAssetData) const
 	{
 		if (Class->IsChildOf<USkeleton>())
 		{
-			return Skeleton.Get()->IsCompatibleSkeletonByAssetData(InAssetData);
+			if (Skeleton.Get())
+			{
+				return Skeleton.Get()->IsCompatibleSkeletonByAssetData(InAssetData);
+			}
 		}
 		else if (Class->IsChildOf<UAnimationAsset>() || Class->IsChildOf<USkeletalMesh>())
 		{
@@ -265,7 +276,7 @@ bool FPersonaAssetFamily::IsAssetCompatible(const FAssetData& InAssetData) const
 			{
 				if (Skeleton.Get())
 				{
-					return Skeleton.Get()->IsCompatibleSkeletonByAssetString(Result.GetValue());
+					return Skeleton.Get()->IsCompatibleSkeletonByAssetData(InAssetData);
 				}
 			}
 		}
@@ -275,7 +286,10 @@ bool FPersonaAssetFamily::IsAssetCompatible(const FAssetData& InAssetData) const
 
 			if (Result.IsSet())
 			{
-				return Skeleton.Get()->IsCompatibleSkeletonByAssetString(Result.GetValue());
+				if (Skeleton.Get())
+				{
+					return Skeleton.Get()->IsCompatibleSkeletonByAssetString(Result.GetValue());
+				}
 			}
 		}
 		else if (Class->IsChildOf<UPhysicsAsset>())
