Non-unity fix
#fix added pre-declare and include

diff --git a/Engine/Source/Editor/Persona/Private/AnimTimeline/AnimTimelineTrack_Attribute.cpp b/Engine/Source/Editor/Persona/Private/AnimTimeline/AnimTimelineTrack_Attribute.cpp
--- a/Engine/Source/Editor/Persona/Private/AnimTimeline/AnimTimelineTrack_Attribute.cpp
+++ b/Engine/Source/Editor/Persona/Private/AnimTimeline/AnimTimelineTrack_Attribute.cpp
@@ -2,6 +2,7 @@
 
 #include "AnimTimelineTrack_Attribute.h"
 #include "AnimTimeline/SAnimOutlinerItem.h"
+#include "Animation/AnimData/AnimDataModel.h"
 
 #define LOCTEXT_NAMESPACE "FAnimTimelineTrack_CustomBoneAttributes"
 
diff --git a/Engine/Source/Editor/Persona/Private/AnimTimeline/AnimTimelineTrack_Attribute.h b/Engine/Source/Editor/Persona/Private/AnimTimeline/AnimTimelineTrack_Attribute.h
--- a/Engine/Source/Editor/Persona/Private/AnimTimeline/AnimTimelineTrack_Attribute.h
+++ b/Engine/Source/Editor/Persona/Private/AnimTimeline/AnimTimelineTrack_Attribute.h
@@ -4,6 +4,8 @@
 
 #include "AnimTimelineTrack.h"
 
+struct FAnimatedBoneAttribute;
+
 /** Animation timeline track inserted for each animated attribute (inserted as child of FAnimTimelineTrack_PerBoneAttributes) */
 class FAnimTimelineTrack_Attribute : public FAnimTimelineTrack
 {
