Fix static analysis warning

diff --git a/Engine/Source/Editor/Persona/Private/SAnimationBlendSpaceGridWidget.cpp b/Engine/Source/Editor/Persona/Private/SAnimationBlendSpaceGridWidget.cpp
--- a/Engine/Source/Editor/Persona/Private/SAnimationBlendSpaceGridWidget.cpp
+++ b/Engine/Source/Editor/Persona/Private/SAnimationBlendSpaceGridWidget.cpp
@@ -1972,7 +1972,7 @@ void SBlendSpaceGridWidget::OnInputBoxValueCommited(const float NewValue, ETextC
 
 void SBlendSpaceGridWidget::OnInputBoxValueChanged(const float NewValue, const int32 ParameterIndex, bool bIsInteractive)
 {
-	checkf(ParameterIndex < 3, TEXT("Invalid parameter index, suppose to be within FVector array range"));
+	checkf(ParameterIndex < 2, TEXT("Invalid parameter index, suppose to be within FVector array range"));
 
 	if (SelectedSampleIndex != INDEX_NONE && BlendSpaceBase.Get() != nullptr)
 	{
