Fix for BlendSpace editor loses its scrub bar when it gets re-used

diff --git a/Engine/Source/Editor/Persona/Private/SAnimationBlendSpace.h b/Engine/Source/Editor/Persona/Private/SAnimationBlendSpace.h
--- a/Engine/Source/Editor/Persona/Private/SAnimationBlendSpace.h
+++ b/Engine/Source/Editor/Persona/Private/SAnimationBlendSpace.h
@@ -20,7 +20,10 @@ class SBlendSpaceEditor : public SAnimEditorBase, public FNotifyHook, public FSe
 {
 public:
 	SLATE_BEGIN_ARGS(SBlendSpaceEditor)
-		: _BlendSpace(NULL)
+		: _BlendSpace(nullptr)
+		, _DisplayScrubBar(true)
+		, _PreviewPosition(0.f)
+		, _PreviewFilteredPosition(0.f)
 		, _StatusBarName(TEXT("AssetEditor.AnimationEditor.MainMenu"))
 		{}
 		
diff --git a/Engine/Source/Editor/Persona/Private/SAnimationBlendSpaceGridWidget.h b/Engine/Source/Editor/Persona/Private/SAnimationBlendSpaceGridWidget.h
--- a/Engine/Source/Editor/Persona/Private/SAnimationBlendSpaceGridWidget.h
+++ b/Engine/Source/Editor/Persona/Private/SAnimationBlendSpaceGridWidget.h
@@ -31,7 +31,11 @@ public:
 	friend class FBlendSampleDetails;
 public:
 	SLATE_BEGIN_ARGS(SBlendSpaceGridWidget) 
-		: _ReadOnly(false)
+		: _BlendSpaceBase(nullptr)
+		, _Position(0.f)
+		, _FilteredPosition(0.f)
+		, _NotifyHook(nullptr)
+		, _ReadOnly(false)
 		, _ShowAxisLabels(true)
 		, _ShowSettingsButtons(true)
 		, _StatusBarName(TEXT("AssetEditor.AnimationEditor.MainMenu"))
