Fix for overzealous initialisation of attributes

diff --git a/Engine/Source/Editor/Persona/Private/SAnimationBlendSpace.h b/Engine/Source/Editor/Persona/Private/SAnimationBlendSpace.h
--- a/Engine/Source/Editor/Persona/Private/SAnimationBlendSpace.h
+++ b/Engine/Source/Editor/Persona/Private/SAnimationBlendSpace.h
@@ -22,8 +22,6 @@ public:
 	SLATE_BEGIN_ARGS(SBlendSpaceEditor)
 		: _BlendSpace(nullptr)
 		, _DisplayScrubBar(true)
-		, _PreviewPosition(0.f)
-		, _PreviewFilteredPosition(0.f)
 		, _StatusBarName(TEXT("AssetEditor.AnimationEditor.MainMenu"))
 		{}
 		
diff --git a/Engine/Source/Editor/Persona/Private/SAnimationBlendSpaceGridWidget.h b/Engine/Source/Editor/Persona/Private/SAnimationBlendSpaceGridWidget.h
--- a/Engine/Source/Editor/Persona/Private/SAnimationBlendSpaceGridWidget.h
+++ b/Engine/Source/Editor/Persona/Private/SAnimationBlendSpaceGridWidget.h
@@ -32,8 +32,6 @@ public:
 public:
 	SLATE_BEGIN_ARGS(SBlendSpaceGridWidget) 
 		: _BlendSpaceBase(nullptr)
-		, _Position(0.f)
-		, _FilteredPosition(0.f)
 		, _NotifyHook(nullptr)
 		, _ReadOnly(false)
 		, _ShowAxisLabels(true)
