Restrict adding montage Slot tracks to valid slots

diff --git a/Engine/Source/Editor/Persona/Private/AnimTimeline/AnimTimelineTrack_Montage.cpp b/Engine/Source/Editor/Persona/Private/AnimTimeline/AnimTimelineTrack_Montage.cpp
--- a/Engine/Source/Editor/Persona/Private/AnimTimeline/AnimTimelineTrack_Montage.cpp
+++ b/Engine/Source/Editor/Persona/Private/AnimTimeline/AnimTimelineTrack_Montage.cpp
@@ -518,11 +518,10 @@ TSharedRef<SWidget> FAnimTimelineTrack_Montage::BuildMontageSubMenu()
 		// Slots
 		MenuBuilder.BeginSection("AnimMontageSlots", LOCTEXT("Slots", "Slots") );
 		{
-			MenuBuilder.AddMenuEntry(
+			MenuBuilder.AddSubMenu(
 				LOCTEXT("NewSlot", "New Slot"),
-				LOCTEXT("NewSlotToolTip", "Adds a new Slot"),
-				FSlateIcon(), 
-				FUIAction(FExecuteAction::CreateSP(&MontageModel->GetMontagePanel()->GetAnimMontagePanel().Get(), &SAnimMontagePanel::OnNewSlotClicked))
+				LOCTEXT("NewSlotToolTip", "Adds a new Slot"), 
+				FNewMenuDelegate::CreateSP(&MontageModel->GetMontagePanel()->GetAnimMontagePanel().Get(), &SAnimMontagePanel::BuildNewSlotMenu)
 			);
 
 			MenuBuilder.AddMenuEntry(
diff --git a/Engine/Source/Editor/Persona/Private/AnimTimeline/AnimTimelineTrack_MontagePanel.cpp b/Engine/Source/Editor/Persona/Private/AnimTimeline/AnimTimelineTrack_MontagePanel.cpp
--- a/Engine/Source/Editor/Persona/Private/AnimTimeline/AnimTimelineTrack_MontagePanel.cpp
+++ b/Engine/Source/Editor/Persona/Private/AnimTimeline/AnimTimelineTrack_MontagePanel.cpp
@@ -153,11 +153,10 @@ TSharedRef<SWidget> FAnimTimelineTrack_MontagePanel::BuildMontageSlotSubMenu(int
 				FNewMenuDelegate::CreateSP(this, &FAnimTimelineTrack_MontagePanel::BuildSlotNameMenu, InSlotIndex)
 			);
 
-			MenuBuilder.AddMenuEntry(
-				LOCTEXT("NewSlot", "New Slot"), 
-				LOCTEXT("NewSlotToolTip", "Adds a new Slot"), 
-				FSlateIcon(), 
-				FUIAction(FExecuteAction::CreateSP(AnimMontagePanel.Get(), &SAnimMontagePanel::OnNewSlotClicked))
+			MenuBuilder.AddSubMenu(
+				LOCTEXT("NewSlot", "New Slot"),
+			 	LOCTEXT("NewSlotToolTip", "Adds a new Slot"), 
+				FNewMenuDelegate::CreateSP(AnimMontagePanel.Get(), &SAnimMontagePanel::BuildNewSlotMenu)
 			);
 
 			if(InSlotIndex != INDEX_NONE)
diff --git a/Engine/Source/Editor/Persona/Private/SAnimMontagePanel.cpp b/Engine/Source/Editor/Persona/Private/SAnimMontagePanel.cpp
--- a/Engine/Source/Editor/Persona/Private/SAnimMontagePanel.cpp
+++ b/Engine/Source/Editor/Persona/Private/SAnimMontagePanel.cpp
@@ -704,8 +704,11 @@ void SAnimMontagePanel::SummonTrackContextMenu( FMenuBuilder& MenuBuilder, float
 	// Slots
 	MenuBuilder.BeginSection("AnimMontageSlots", LOCTEXT("Slots", "Slots") );
 	{
-		UIAction.ExecuteAction.BindRaw(this, &SAnimMontagePanel::OnNewSlotClicked);
-		MenuBuilder.AddMenuEntry(LOCTEXT("NewSlot", "New Slot"), LOCTEXT("NewSlotToolTip", "Adds a new Slot"), FSlateIcon(), UIAction);
+		MenuBuilder.AddSubMenu(
+			LOCTEXT("NewSlot", "New Slot"),
+			LOCTEXT("NewSlotToolTip", "Adds a new Slot"), 
+			FNewMenuDelegate::CreateSP(this, &SAnimMontagePanel::BuildNewSlotMenu)
+		);
 
 		if(AnimSlotIndex != INDEX_NONE)
 		{
@@ -751,9 +754,53 @@ void SAnimMontagePanel::FillSlotSubMenu(FMenuBuilder& Menubuilder)
 }
 
 /** Slots */
-void SAnimMontagePanel::OnNewSlotClicked()
+
+void SAnimMontagePanel::BuildNewSlotMenu(FMenuBuilder& InMenuBuilder)
+{
+	USkeleton* Skeleton = Montage->GetSkeleton();
+	FName CurrentSlotGroupName = FAnimSlotGroup::DefaultGroupName;
+	if (Montage->SlotAnimTracks.Num() > 0)
+	{
+		FName CurrentSlotName = Montage->SlotAnimTracks[0].SlotName;
+		CurrentSlotGroupName = Skeleton->GetSlotGroupName(CurrentSlotName);
+	}
+		
+	if (FAnimSlotGroup* SlotGroup = Skeleton->FindAnimSlotGroup(CurrentSlotGroupName))
+	{
+		InMenuBuilder.BeginSection("AnimMontageAvailableAddSlots", FText::FromString(SlotGroup->GroupName.ToString()));
+		{
+			for (const FName& SlotName : SlotGroup->SlotNames)
+			{
+				FText SlotItemText = FText::FromString(*SlotName.ToString());
+
+				FText Tooltip = CanCreateNewSlot(SlotName) ? FText::Format(LOCTEXT("SlotTooltipFormat", "Add new Slot '{0}'"), SlotItemText) :
+				                                             FText::Format(LOCTEXT("SlotUnavailableTooltipFormat", "Slot '{0}' already has a track in this Montage"), SlotItemText);
+
+				InMenuBuilder.AddMenuEntry(
+					SlotItemText,
+					Tooltip,
+					FSlateIcon(),
+					FUIAction(
+						FExecuteAction::CreateSP(this, &SAnimMontagePanel::AddNewMontageSlot, SlotName),
+						FCanExecuteAction::CreateSP(this, &SAnimMontagePanel::CanCreateNewSlot, SlotName)
+					));
+			}
+		}
+		InMenuBuilder.EndSection();
+	}
+}
+
+bool SAnimMontagePanel::CanCreateNewSlot(FName InName) const
 {
-	AddNewMontageSlot(FAnimSlotGroup::DefaultSlotName);
+	for (auto &Track : Montage->SlotAnimTracks)
+	{
+		if (Track.SlotName == InName)
+		{
+			return false;
+		}
+	}
+
+	return true;
 }
 
 void SAnimMontagePanel::CreateNewSlot(const FText& NewSlotName, ETextCommit::Type CommitInfo)
diff --git a/Engine/Source/Editor/Persona/Private/SAnimMontagePanel.h b/Engine/Source/Editor/Persona/Private/SAnimMontagePanel.h
--- a/Engine/Source/Editor/Persona/Private/SAnimMontagePanel.h
+++ b/Engine/Source/Editor/Persona/Private/SAnimMontagePanel.h
@@ -93,7 +93,10 @@ public:
 	void FillElementSubMenuForTimes( FMenuBuilder& MenuBuilder );
 	void FillSlotSubMenu( FMenuBuilder& Menubuilder );
 
-	void OnNewSlotClicked();
+	void BuildNewSlotMenu(FMenuBuilder& InMenuBuilder);
+	void CreateNewSlot(FName InName);
+	bool CanCreateNewSlot(FName InName) const;
+
 	void OnNewSectionClicked(float DataPosX);
 	bool CanAddNewSection();
 
