Improve the way the preview position etc are shown and used. Some code tidyup too, removing unnecessary checks on BlendSpace.Get()

diff --git a/Engine/Source/Editor/Persona/Private/SAnimationBlendSpaceGridWidget.cpp b/Engine/Source/Editor/Persona/Private/SAnimationBlendSpaceGridWidget.cpp
--- a/Engine/Source/Editor/Persona/Private/SAnimationBlendSpaceGridWidget.cpp
+++ b/Engine/Source/Editor/Persona/Private/SAnimationBlendSpaceGridWidget.cpp
@@ -65,7 +65,6 @@ void SBlendSpaceGridWidget::Construct(const FArguments& InArgs)
 	HighlightedSampleIndex = SelectedSampleIndex = DraggedSampleIndex = ToolTipSampleIndex = INDEX_NONE;
 	DragState = EDragState::None;
 	// Initialize flags 
-	bPreviewPositionSet = true;
 	bHighlightPreviewPin = false;
 	// Initialize preview value to center or the grid
 	PreviewPosition.X = BlendSpaceBase.Get() != nullptr ? (BlendSpaceBase.Get()->GetBlendParameter(0).GetRange() * .5f) + BlendSpaceBase.Get()->GetBlendParameter(0).Min : 0.0f;
@@ -118,8 +117,6 @@ void SBlendSpaceGridWidget::Construct(const FArguments& InArgs)
 	GridMargin = bShowAxisLabels ?  FMargin(MaxVerticalAxisTextWidth + (TextMargin * 2.0f), TextMargin, (HorizontalAxisMaxTextWidth *.5f) + TextMargin, MaxHorizontalAxisTextHeight + (TextMargin * 2.0f)) : 
 									FMargin(TextMargin, TextMargin, TextMargin, TextMargin);
 
-	bPreviewToolTipHidden = false;
-
 	const bool bShowInputBoxLabel = true;
 	// Widget construction
 	this->ChildSlot
@@ -330,23 +327,20 @@ int32 SBlendSpaceGridWidget::OnPaint(const FPaintArgs& Args, const FGeometry& Al
 	
 	PaintBackgroundAndGrid(AllottedGeometry, MyCullingRect, OutDrawElements, LayerId);
 
-	if(const UBlendSpace* BlendSpace = BlendSpaceBase.Get())
+	if (bShowTriangulation)
 	{
-		if (bShowTriangulation)
-		{
-			PaintTriangulation(AllottedGeometry, MyCullingRect, OutDrawElements, LayerId);
-		}	
-		PaintSampleKeys(AllottedGeometry, MyCullingRect, OutDrawElements, LayerId);
+		PaintTriangulation(AllottedGeometry, MyCullingRect, OutDrawElements, LayerId);
+	}	
+	PaintSampleKeys(AllottedGeometry, MyCullingRect, OutDrawElements, LayerId);
 
-		if(bShowAxisLabels)
-		{
-			PaintAxisText(AllottedGeometry, MyCullingRect, OutDrawElements, LayerId);
-		}
+	if(bShowAxisLabels)
+	{
+		PaintAxisText(AllottedGeometry, MyCullingRect, OutDrawElements, LayerId);
+	}
 
-		if (bShowAnimationNames)
-		{
-			PaintAnimationNames(AllottedGeometry, MyCullingRect, OutDrawElements, LayerId);
-		}
+	if (bShowAnimationNames)
+	{
+		PaintAnimationNames(AllottedGeometry, MyCullingRect, OutDrawElements, LayerId);
 	}
 
 	return LayerId;
@@ -456,7 +450,7 @@ void SBlendSpaceGridWidget::PaintSampleKeys(
 				PreviewBrush, ESlateDrawEffect::None, PreviewKeyColor.GetSpecifiedColor() * 0.5);
 		}
 
-		if (bPreviewPositionSet)
+		// Always draw the preview position
 		{
 			FVector2D GridPosition = SampleValueToScreenPosition(PreviewPosition) - (PreviewSize * .5f);
 			FSlateDrawElement::MakeBox(
@@ -855,7 +849,6 @@ void SBlendSpaceGridWidget::PaintAnimationNames(const FGeometry& AllottedGeometr
 	if(const UBlendSpace* BlendSpace = BlendSpaceBase.Get())
 	{
 		const TSharedRef< FSlateFontMeasure > FontMeasure = FSlateApplication::Get().GetRenderer()->GetFontMeasureService();
-		const FVector2D GridCenter = CachedGridRectangle.GetCenter();
 		const TArray<FBlendSample>& Samples = BlendSpace->GetBlendSamples();
 		for (int32 SampleIndex = 0; SampleIndex < Samples.Num(); ++SampleIndex)
 		{
@@ -884,41 +877,38 @@ void SBlendSpaceGridWidget::PaintAnimationNames(const FGeometry& AllottedGeometr
 
 FReply SBlendSpaceGridWidget::OnDrop(const FGeometry& MyGeometry, const FDragDropEvent& DragDropEvent)
 {
-	if(!bReadOnly)
+	if(!bReadOnly && BlendSpaceBase.IsSet())
 	{
-		if(const UBlendSpace* BlendSpace = BlendSpaceBase.Get())
+		// Check if we are in dropping state and if so snap to the grid and try to add the sample
+		if (DragState == EDragState::DragDrop || DragState == EDragState::InvalidDragDrop || DragState == EDragState::DragDropOverride)
 		{
-			// Check if we are in dropping state and if so snap to the grid and try to add the sample
-			if (DragState == EDragState::DragDrop || DragState == EDragState::InvalidDragDrop || DragState == EDragState::DragDropOverride)
+			if (DragState == EDragState::DragDrop)
 			{
-				if (DragState == EDragState::DragDrop)
+				const FVector SampleValue = ScreenPositionToSampleValueWithSnapping(LocalMousePosition, FSlateApplication::Get().GetModifierKeys().IsShiftDown());
+				TSharedPtr<FAssetDragDropOp> DragDropOperation = DragDropEvent.GetOperationAs<FAssetDragDropOp>();
+				if (DragDropOperation.IsValid())
 				{
-					const FVector SampleValue = ScreenPositionToSampleValueWithSnapping(LocalMousePosition, FSlateApplication::Get().GetModifierKeys().IsShiftDown());
-					TSharedPtr<FAssetDragDropOp> DragDropOperation = DragDropEvent.GetOperationAs<FAssetDragDropOp>();
-					if (DragDropOperation.IsValid())
-					{
-						UAnimSequence* Animation = FAssetData::GetFirstAsset<UAnimSequence>(DragDropOperation->GetAssets());
-						OnSampleAdded.ExecuteIfBound(Animation, SampleValue);
-					}	
-				}
-				else if (DragState == EDragState::DragDropOverride)
+					UAnimSequence* Animation = FAssetData::GetFirstAsset<UAnimSequence>(DragDropOperation->GetAssets());
+					OnSampleAdded.ExecuteIfBound(Animation, SampleValue);
+				}	
+			}
+			else if (DragState == EDragState::DragDropOverride)
+			{
+				TSharedPtr<FAssetDragDropOp> DragDropOperation = DragDropEvent.GetOperationAs<FAssetDragDropOp>();
+				if (DragDropOperation.IsValid())
 				{
-					TSharedPtr<FAssetDragDropOp> DragDropOperation = DragDropEvent.GetOperationAs<FAssetDragDropOp>();
-					if (DragDropOperation.IsValid())
-					{
-						UAnimSequence* Animation = FAssetData::GetFirstAsset<UAnimSequence>(DragDropOperation->GetAssets());
-						int32 DroppedSampleIndex = GetClosestSamplePointIndexToMouse();
-						OnSampleReplaced.ExecuteIfBound(DroppedSampleIndex, Animation);
-					}
+					UAnimSequence* Animation = FAssetData::GetFirstAsset<UAnimSequence>(DragDropOperation->GetAssets());
+					int32 DroppedSampleIndex = GetClosestSamplePointIndexToMouse();
+					OnSampleReplaced.ExecuteIfBound(DroppedSampleIndex, Animation);
 				}
-
-				DragState = EDragState::None;
 			}
 
-			DragDropAnimationSequence = nullptr;
-			DragDropAnimationName = FText::GetEmpty();
-			HoveredAnimationName = FText::GetEmpty();
+			DragState = EDragState::None;
 		}
+
+		DragDropAnimationSequence = nullptr;
+		DragDropAnimationName = FText::GetEmpty();
+		HoveredAnimationName = FText::GetEmpty();
 	}
 
 	return FReply::Unhandled();
@@ -926,37 +916,31 @@ FReply SBlendSpaceGridWidget::OnDrop(const FGeometry& MyGeometry, const FDragDro
 
 void SBlendSpaceGridWidget::OnDragEnter(const FGeometry& MyGeometry, const FDragDropEvent& DragDropEvent)
 {
-	if(!bReadOnly)
+	if(!bReadOnly && BlendSpaceBase.IsSet())
 	{
-		if(const UBlendSpace* BlendSpace = BlendSpaceBase.Get())
+		if (DragDropEvent.GetOperationAs<FAssetDragDropOp>().IsValid())
 		{
-			if (DragDropEvent.GetOperationAs<FAssetDragDropOp>().IsValid())
-			{
-				DragState = IsValidDragDropOperation(DragDropEvent, InvalidDragDropText) ? EDragState::DragDrop : EDragState::InvalidDragDrop;
-			}
+			DragState = IsValidDragDropOperation(DragDropEvent, InvalidDragDropText) ? EDragState::DragDrop : EDragState::InvalidDragDrop;
 		}
 	}
 }
 
 FReply SBlendSpaceGridWidget::OnDragOver(const FGeometry& MyGeometry, const FDragDropEvent& DragDropEvent)
 {
-	if(!bReadOnly)
+	if(!bReadOnly && BlendSpaceBase.IsSet())
 	{
-		if(const UBlendSpace* BlendSpace = BlendSpaceBase.Get())
-		{
-			if (DragState == EDragState::DragDrop || DragState == EDragState::InvalidDragDrop || DragState == EDragState::DragDropOverride)
-			{		
-				LocalMousePosition = MyGeometry.AbsoluteToLocal(DragDropEvent.GetScreenSpacePosition());				
-		
-				// Always update the tool tip, in case it became invalid
-				TSharedPtr<FAssetDragDropOp> DragDropOperation = DragDropEvent.GetOperationAs<FAssetDragDropOp>();
-				if (DragDropOperation.IsValid())
-				{
-					DragDropOperation->SetToolTip(GetToolTipSampleValue(), DragDropOperation->GetIcon());
-				}		
+		if (DragState == EDragState::DragDrop || DragState == EDragState::InvalidDragDrop || DragState == EDragState::DragDropOverride)
+		{		
+			LocalMousePosition = MyGeometry.AbsoluteToLocal(DragDropEvent.GetScreenSpacePosition());				
+	
+			// Always update the tool tip, in case it became invalid
+			TSharedPtr<FAssetDragDropOp> DragDropOperation = DragDropEvent.GetOperationAs<FAssetDragDropOp>();
+			if (DragDropOperation.IsValid())
+			{
+				DragDropOperation->SetToolTip(GetToolTipSampleValue(), DragDropOperation->GetIcon());
+			}		
 
-				return FReply::Handled();
-			}
+			return FReply::Handled();
 		}
 	}
 	return FReply::Unhandled();
@@ -964,49 +948,43 @@ FReply SBlendSpaceGridWidget::OnDragOver(const FGeometry& MyGeometry, const FDra
 
 void SBlendSpaceGridWidget::OnDragLeave(const FDragDropEvent& DragDropEvent)
 {
-	if(!bReadOnly)
+	if(!bReadOnly && BlendSpaceBase.IsSet())
 	{
-		if(const UBlendSpace* BlendSpace = BlendSpaceBase.Get())
+		if (DragState == EDragState::DragDrop || DragState == EDragState::InvalidDragDrop || DragState == EDragState::DragDropOverride)
 		{
-			if (DragState == EDragState::DragDrop || DragState == EDragState::InvalidDragDrop || DragState == EDragState::DragDropOverride)
-			{
-				DragState = EDragState::None;
-				DragDropAnimationSequence = nullptr;
-				DragDropAnimationName = FText::GetEmpty();
-				HoveredAnimationName = FText::GetEmpty();
-			}
+			DragState = EDragState::None;
+			DragDropAnimationSequence = nullptr;
+			DragDropAnimationName = FText::GetEmpty();
+			HoveredAnimationName = FText::GetEmpty();
 		}
 	}
 }
 
 FReply SBlendSpaceGridWidget::OnMouseButtonUp(const FGeometry& MyGeometry, const FPointerEvent& MouseEvent)
 {
-	if(!bReadOnly)
+	if(!bReadOnly && BlendSpaceBase.IsSet())
 	{
-		if(const UBlendSpace* BlendSpace = BlendSpaceBase.Get())
+		if (this->HasMouseCapture())
 		{
-			if (this->HasMouseCapture())
+			if (DragState == EDragState::None || DragState == EDragState::PreDrag)
 			{
-				if (DragState == EDragState::None || DragState == EDragState::PreDrag)
-				{
-					ProcessClick(MyGeometry, MouseEvent);
-				}
-				else if (DragState == EDragState::DragSample)
-				{
-					// Process drag ending			
-					ResetToolTip();
-				}
-
-				// Reset drag state and index
-				DragState = EDragState::None;
-				DraggedSampleIndex = INDEX_NONE;
-
-				return FReply::Handled().ReleaseMouseCapture();
+				ProcessClick(MyGeometry, MouseEvent);
 			}
-			else
+			else if (DragState == EDragState::DragSample)
 			{
-				return ProcessClick(MyGeometry, MouseEvent);
+				// Process drag ending			
+				ResetToolTip();
 			}
+
+			// Reset drag state and index
+			DragState = EDragState::None;
+			DraggedSampleIndex = INDEX_NONE;
+
+			return FReply::Handled().ReleaseMouseCapture();
+		}
+		else
+		{
+			return ProcessClick(MyGeometry, MouseEvent);
 		}
 	}
 
@@ -1015,28 +993,25 @@ FReply SBlendSpaceGridWidget::OnMouseButtonUp(const FGeometry& MyGeometry, const
 
 FReply SBlendSpaceGridWidget::OnMouseButtonDown(const FGeometry& MyGeometry, const FPointerEvent& MouseEvent)
 {
-	if(!bReadOnly)
+	if(!bReadOnly && BlendSpaceBase.IsSet())
 	{
-		if(const UBlendSpace* BlendSpace = BlendSpaceBase.Get())
+		if (MouseEvent.GetEffectingButton() == EKeys::LeftMouseButton)
 		{
-			if (MouseEvent.GetEffectingButton() == EKeys::LeftMouseButton)
+			// If we are over a sample, make it our currently (dragged) sample
+			if (HighlightedSampleIndex != INDEX_NONE)
 			{
-				// If we are over a sample, make it our currently (dragged) sample
-				if (HighlightedSampleIndex != INDEX_NONE)
-				{
-					DraggedSampleIndex = SelectedSampleIndex = HighlightedSampleIndex;
-					HighlightedSampleIndex = INDEX_NONE;
-					ResetToolTip();
-					DragState = EDragState::PreDrag;
-					MouseDownPosition = LocalMousePosition;
-
-					// Start mouse capture
-					return FReply::Handled().CaptureMouse(SharedThis(this));
-				}		
-			}
+				DraggedSampleIndex = SelectedSampleIndex = HighlightedSampleIndex;
+				HighlightedSampleIndex = INDEX_NONE;
+				ResetToolTip();
+				DragState = EDragState::PreDrag;
+				MouseDownPosition = LocalMousePosition;
 
-			return FReply::Handled();
+				// Start mouse capture
+				return FReply::Handled().CaptureMouse(SharedThis(this));
+			}		
 		}
+
+		return FReply::Handled();
 	}
 
 	return FReply::Unhandled();
@@ -1044,19 +1019,16 @@ FReply SBlendSpaceGridWidget::OnMouseButtonDown(const FGeometry& MyGeometry, con
 
 FReply SBlendSpaceGridWidget::OnMouseButtonDoubleClick(const FGeometry& InMyGeometry, const FPointerEvent& InMouseEvent)
 {
-	if(!bReadOnly)
+	if(!bReadOnly && BlendSpaceBase.IsSet())
 	{
-		if(const UBlendSpace* BlendSpace = BlendSpaceBase.Get())
+		if (InMouseEvent.GetEffectingButton() == EKeys::LeftMouseButton)
 		{
-			if (InMouseEvent.GetEffectingButton() == EKeys::LeftMouseButton)
+			if (SelectedSampleIndex != INDEX_NONE)
 			{
-				if (SelectedSampleIndex != INDEX_NONE)
-				{
-					OnSampleDoubleClicked.ExecuteIfBound(SelectedSampleIndex);
-				}
-
-				return FReply::Handled();
+				OnSampleDoubleClicked.ExecuteIfBound(SelectedSampleIndex);
 			}
+
+			return FReply::Handled();
 		}
 	}
 
@@ -1065,75 +1037,72 @@ FReply SBlendSpaceGridWidget::OnMouseButtonDoubleClick(const FGeometry& InMyGeom
 
 FReply SBlendSpaceGridWidget::OnMouseMove(const FGeometry& MyGeometry, const FPointerEvent& MouseEvent)
 {
-	if(const UBlendSpace* BlendSpace = BlendSpaceBase.Get())
+	if(!bReadOnly && BlendSpaceBase.IsSet())
 	{
-		if(!bReadOnly)
-		{
-			EnableStatusBarMessage(true);
-		}
+		EnableStatusBarMessage(true);
+	}
 
-		// Cache the mouse position in local and screen space
-		LocalMousePosition = MyGeometry.AbsoluteToLocal(MouseEvent.GetScreenSpacePosition());
-		LastMousePosition = MouseEvent.GetScreenSpacePosition();
+	// Cache the mouse position in local and screen space
+	LocalMousePosition = MyGeometry.AbsoluteToLocal(MouseEvent.GetScreenSpacePosition());
+	LastMousePosition = MouseEvent.GetScreenSpacePosition();
 
-		if(!bReadOnly)
+	if(!bReadOnly)
+	{
+		if (this->HasMouseCapture())
 		{
-			if (this->HasMouseCapture())
+			if (DragState == EDragState::None)
 			{
-				if (DragState == EDragState::None)
-				{
-					if (HighlightedSampleIndex != INDEX_NONE)
-					{
-						DragState = EDragState::DragSample;
-						DraggedSampleIndex = HighlightedSampleIndex;
-						HighlightedSampleIndex = INDEX_NONE;
-						return FReply::Handled();
-					}
-				}
-				else if (DragState == EDragState::PreDrag)
+				if (HighlightedSampleIndex != INDEX_NONE)
 				{
-					// Actually start dragging
-					if ((LocalMousePosition - MouseDownPosition).SizeSquared() > DragThreshold)
-					{
-						DragState = EDragState::DragSample;
-						HighlightedSampleIndex = INDEX_NONE;
-						ShowToolTip();
-						return FReply::Handled();
-					}
+					DragState = EDragState::DragSample;
+					DraggedSampleIndex = HighlightedSampleIndex;
+					HighlightedSampleIndex = INDEX_NONE;
+					return FReply::Handled();
 				}
 			}
-			else if (IsHovered() && bMouseIsOverGeometry)
+			else if (DragState == EDragState::PreDrag)
 			{
-				if (MouseEvent.IsControlDown())
-				{
-					StartPreviewing();
-					DragState = EDragState::Preview;
-					// Make tool tip visible (this will display the current preview sample value)
-					ShowToolTip();			
-
-					// Set flag for showing advanced preview info in tooltip
-					bAdvancedPreview = MouseEvent.IsAltDown();
-					return FReply::Handled();
-				}
-				else if(Position.IsSet())
+				// Actually start dragging
+				if ((LocalMousePosition - MouseDownPosition).SizeSquared() > DragThreshold)
 				{
-					StartPreviewing();
-					DragState = EDragState::None;
+					DragState = EDragState::DragSample;
+					HighlightedSampleIndex = INDEX_NONE;
 					ShowToolTip();
-
-					// Set flag for showing advanced preview info in tooltip
-					bAdvancedPreview = MouseEvent.IsAltDown();
-					return FReply::Handled();
-				}
-				else if (bSamplePreviewing)
-				{
-					StopPreviewing();
-					DragState = EDragState::None;
-					ResetToolTip();
 					return FReply::Handled();
 				}
 			}
 		}
+		else if (IsHovered() && bMouseIsOverGeometry)
+		{
+			if (MouseEvent.IsControlDown())
+			{
+				StartPreviewing();
+				DragState = EDragState::Preview;
+				// Make tool tip visible (this will display the current preview sample value)
+				ShowToolTip();			
+
+				// Set flag for showing advanced preview info in tooltip
+				bAdvancedPreview = MouseEvent.IsAltDown();
+				return FReply::Handled();
+			}
+			else if(Position.IsSet())
+			{
+				StartPreviewing();
+				DragState = EDragState::None;
+				ShowToolTip();
+
+				// Set flag for showing advanced preview info in tooltip
+				bAdvancedPreview = MouseEvent.IsAltDown();
+				return FReply::Handled();
+			}
+			else if (bSamplePreviewing)
+			{
+				StopPreviewing();
+				DragState = EDragState::None;
+				ResetToolTip();
+				return FReply::Handled();
+			}
+		}
 	}
 
 	return FReply::Unhandled();
@@ -1141,71 +1110,68 @@ FReply SBlendSpaceGridWidget::OnMouseMove(const FGeometry& MyGeometry, const FPo
 
 FReply SBlendSpaceGridWidget::ProcessClick(const FGeometry& MyGeometry, const FPointerEvent& MouseEvent)
 {
-	if(!bReadOnly)
+	if(!bReadOnly && BlendSpaceBase.IsSet())
 	{
-		if(const UBlendSpace* BlendSpace = BlendSpaceBase.Get())
+		if (MouseEvent.GetEffectingButton() == EKeys::LeftMouseButton)
 		{
-			if (MouseEvent.GetEffectingButton() == EKeys::LeftMouseButton)
-			{
-				SelectedSampleIndex = INDEX_NONE;
+			SelectedSampleIndex = INDEX_NONE;
 
-				if (HighlightedSampleIndex == INDEX_NONE)
-				{
-					// If there isn't any sample currently being highlighted, retrieve all of them and see if we are over one
-					SelectedSampleIndex = GetClosestSamplePointIndexToMouse();
-				}
-				else
-				{
-					// If we are over a sample, make it the selected sample index
-					SelectedSampleIndex = HighlightedSampleIndex;
-					HighlightedSampleIndex = INDEX_NONE;			
-				}
+			if (HighlightedSampleIndex == INDEX_NONE)
+			{
+				// If there isn't any sample currently being highlighted, retrieve all of them and see if we are over one
+				SelectedSampleIndex = GetClosestSamplePointIndexToMouse();
 			}
-			else if (MouseEvent.GetEffectingButton() == EKeys::RightMouseButton)
+			else
+			{
+				// If we are over a sample, make it the selected sample index
+				SelectedSampleIndex = HighlightedSampleIndex;
+				HighlightedSampleIndex = INDEX_NONE;			
+			}
+		}
+		else if (MouseEvent.GetEffectingButton() == EKeys::RightMouseButton)
+		{
+			auto PushMenu = [this, &MouseEvent](TSharedPtr<SWidget> InMenuContent)
 			{
-				auto PushMenu = [this, &MouseEvent](TSharedPtr<SWidget> InMenuContent)
+				if (InMenuContent.IsValid())
 				{
-					if (InMenuContent.IsValid())
-					{
-						const FWidgetPath WidgetPath = MouseEvent.GetEventPath() != nullptr ? *MouseEvent.GetEventPath() : FWidgetPath();
-						const FVector2D MousePosition = MouseEvent.GetScreenSpacePosition();
-						// This is of a fixed size atm since MenuContent->GetDesiredSize() will not take the detail customization into account and return an incorrect (small) size
-						const FVector2D ExpectedSize(300, 100);				
-						const FVector2D MenuPosition = FSlateApplication::Get().CalculatePopupWindowPosition(FSlateRect(MousePosition.X, MousePosition.Y, MousePosition.X, MousePosition.Y), ExpectedSize, false);
-
-						FSlateApplication::Get().PushMenu(
-							AsShared(),
-							WidgetPath,
-							InMenuContent.ToSharedRef(),
-							MenuPosition,
-							FPopupTransitionEffect(FPopupTransitionEffect::ContextMenu)
-							);
-					}
-				};
+					const FWidgetPath WidgetPath = MouseEvent.GetEventPath() != nullptr ? *MouseEvent.GetEventPath() : FWidgetPath();
+					const FVector2D MousePosition = MouseEvent.GetScreenSpacePosition();
+					// This is of a fixed size atm since MenuContent->GetDesiredSize() will not take the detail customization into account and return an incorrect (small) size
+					const FVector2D ExpectedSize(300, 100);				
+					const FVector2D MenuPosition = FSlateApplication::Get().CalculatePopupWindowPosition(FSlateRect(MousePosition.X, MousePosition.Y, MousePosition.X, MousePosition.Y), ExpectedSize, false);
+
+					FSlateApplication::Get().PushMenu(
+						AsShared(),
+						WidgetPath,
+						InMenuContent.ToSharedRef(),
+						MenuPosition,
+						FPopupTransitionEffect(FPopupTransitionEffect::ContextMenu)
+						);
+				}
+			};
 
-				// If we are over a sample open a context menu for editing its data
-				if (HighlightedSampleIndex != INDEX_NONE)
-				{	
-					SelectedSampleIndex = HighlightedSampleIndex;
+			// If we are over a sample open a context menu for editing its data
+			if (HighlightedSampleIndex != INDEX_NONE)
+			{	
+				SelectedSampleIndex = HighlightedSampleIndex;
 
-					// Create context menu
-					TSharedPtr<SWidget> MenuContent = CreateBlendSampleContextMenu();
-			
-					// Reset highlight sample index
-					HighlightedSampleIndex = INDEX_NONE;
+				// Create context menu
+				TSharedPtr<SWidget> MenuContent = CreateBlendSampleContextMenu();
+		
+				// Reset highlight sample index
+				HighlightedSampleIndex = INDEX_NONE;
 
-					PushMenu(MenuContent);
+				PushMenu(MenuContent);
 
-					return FReply::Handled().SetUserFocus(MenuContent.ToSharedRef(), EFocusCause::SetDirectly).ReleaseMouseCapture();
-				}
-				else
-				{
-					TSharedPtr<SWidget> MenuContent = CreateNewBlendSampleContextMenu(MyGeometry.AbsoluteToLocal(MouseEvent.GetScreenSpacePosition()));
-				
-					PushMenu(MenuContent);
-				
-					return FReply::Handled().SetUserFocus(MenuContent.ToSharedRef(), EFocusCause::SetDirectly).ReleaseMouseCapture();
-				}
+				return FReply::Handled().SetUserFocus(MenuContent.ToSharedRef(), EFocusCause::SetDirectly).ReleaseMouseCapture();
+			}
+			else
+			{
+				TSharedPtr<SWidget> MenuContent = CreateNewBlendSampleContextMenu(MyGeometry.AbsoluteToLocal(MouseEvent.GetScreenSpacePosition()));
+			
+				PushMenu(MenuContent);
+			
+				return FReply::Handled().SetUserFocus(MenuContent.ToSharedRef(), EFocusCause::SetDirectly).ReleaseMouseCapture();
 			}
 		}
 	}
@@ -1215,9 +1181,9 @@ FReply SBlendSpaceGridWidget::ProcessClick(const FGeometry& MyGeometry, const FP
 
 FReply SBlendSpaceGridWidget::OnKeyDown(const FGeometry& MyGeometry, const FKeyEvent& InKeyEvent)
 {
-	if(!bReadOnly)
+	if(!bReadOnly && BlendSpaceBase.IsSet())
 	{
-		if(const UBlendSpace* BlendSpace = BlendSpaceBase.Get())
+		if (BlendSpaceBase.IsSet())
 		{
 			// Start previewing when either one of the shift keys is pressed
 			if (IsHovered() && bMouseIsOverGeometry)
@@ -1246,48 +1212,45 @@ FReply SBlendSpaceGridWidget::OnKeyDown(const FGeometry& MyGeometry, const FKeyE
 
 FReply SBlendSpaceGridWidget::OnKeyUp(const FGeometry& MyGeometry, const FKeyEvent& InKeyEvent)
 {	
-	if(!bReadOnly)
+	if(!bReadOnly && BlendSpaceBase.IsSet())
 	{
-		if(const UBlendSpace* BlendSpace = BlendSpaceBase.Get())
+		// Stop previewing when shift keys are released 
+		if ((InKeyEvent.GetKey() == EKeys::LeftControl) || (InKeyEvent.GetKey() == EKeys::RightControl))
 		{
-			// Stop previewing when shift keys are released 
-			if ((InKeyEvent.GetKey() == EKeys::LeftControl) || (InKeyEvent.GetKey() == EKeys::RightControl))
-			{
-				StopPreviewing();
-				DragState = EDragState::None;
-				ResetToolTip();
-				return FReply::Handled();
-			}
+			StopPreviewing();
+			DragState = EDragState::None;
+			ResetToolTip();
+			return FReply::Handled();
+		}
 
-			if((InKeyEvent.GetKey() == EKeys::LeftAlt) || (InKeyEvent.GetKey() == EKeys::RightAlt))
-			{
-				bAdvancedPreview = false;
-				return FReply::Handled();
-			}
+		if((InKeyEvent.GetKey() == EKeys::LeftAlt) || (InKeyEvent.GetKey() == EKeys::RightAlt))
+		{
+			bAdvancedPreview = false;
+			return FReply::Handled();
+		}
 
-			// If delete is pressed and we currently have a sample selected remove it from the blendspace
-			if (InKeyEvent.GetKey() == EKeys::Delete)
+		// If delete is pressed and we currently have a sample selected remove it from the blendspace
+		if (InKeyEvent.GetKey() == EKeys::Delete)
+		{
+			if (SelectedSampleIndex != INDEX_NONE)
 			{
-				if (SelectedSampleIndex != INDEX_NONE)
+				OnSampleRemoved.ExecuteIfBound(SelectedSampleIndex);
+		
+				if (SelectedSampleIndex == HighlightedSampleIndex)
 				{
-					OnSampleRemoved.ExecuteIfBound(SelectedSampleIndex);
-			
-					if (SelectedSampleIndex == HighlightedSampleIndex)
-					{
-						HighlightedSampleIndex = INDEX_NONE;
-						ResetToolTip();
-					}
-
-					SelectedSampleIndex = INDEX_NONE;
+					HighlightedSampleIndex = INDEX_NONE;
+					ResetToolTip();
 				}
-			}
 
-			// Pressing esc will remove the current key selection
-			if( InKeyEvent.GetKey() == EKeys::Escape)
-			{
 				SelectedSampleIndex = INDEX_NONE;
 			}
 		}
+
+		// Pressing esc will remove the current key selection
+		if( InKeyEvent.GetKey() == EKeys::Escape)
+		{
+			SelectedSampleIndex = INDEX_NONE;
+		}
 	}
 
 	return FReply::Unhandled();
@@ -1527,7 +1490,6 @@ int32 SBlendSpaceGridWidget::FindClosestGridPointIndexFromScreenPosition(const F
 const FVector2D SBlendSpaceGridWidget::SampleValueToScreenPosition(const FVector& SampleValue) const
 {
 	const FVector2D GridSize = CachedGridRectangle.GetSize();	
-	const FVector2D GridCenter = GridSize * 0.5f;
 		
 	FVector2D SamplePosition2D;
 	// Convert the sample value to -1 to 1 form
@@ -1613,12 +1575,6 @@ void SBlendSpaceGridWidget::StartPreviewing()
 {
 	bSamplePreviewing = true;
 	LastPreviewingMousePosition = LocalMousePosition;
-	FModifierKeysState ModifierKeyState = FSlateApplication::Get().GetModifierKeys();
-	bool bIsManualPreviewing = !bReadOnly && IsHovered() && bMouseIsOverGeometry && ModifierKeyState.IsControlDown();
-	PreviewPosition = Position.IsSet() && !bIsManualPreviewing ? Position.Get() : ScreenPositionToSampleValue(LastPreviewingMousePosition, false);
-	PreviewFilteredPosition = FilteredPosition.IsSet() ? FilteredPosition.Get() : PreviewPosition;
-	bPreviewPositionSet = true;	
-	bPreviewToolTipHidden = true;
 }
 
 void SBlendSpaceGridWidget::StopPreviewing()
@@ -2010,7 +1966,7 @@ EVisibility SBlendSpaceGridWidget::GetSampleToolTipVisibility() const
 EVisibility SBlendSpaceGridWidget::GetPreviewToolTipVisibility() const
 {
 	// Only show preview tooltip until the user discovers the functionality
-	return (!bReadOnly && !bPreviewToolTipHidden) ? EVisibility::Visible : EVisibility::Collapsed;
+	return bReadOnly ? EVisibility::Collapsed : EVisibility::Visible;
 }
 
 EVisibility SBlendSpaceGridWidget::GetTriangulationButtonVisibility() const
@@ -2212,10 +2168,7 @@ void SBlendSpaceGridWidget::Tick(const FGeometry& AllottedGeometry, const double
 			{
 				// If we are dragging a sample, find out whether or not it has actually moved to a
 				// different grid position since the last tick and update the blend space accordingly
-				const FBlendSample& BlendSample = BlendSpace->GetBlendSample(DraggedSampleIndex);
-
 				FVector SampleValue = ScreenPositionToSampleValueWithSnapping(LocalMousePosition, FSlateApplication::Get().GetModifierKeys().IsShiftDown());
-
 				if (SampleValue != LastDragPosition)
 				{
 					LastDragPosition = SampleValue;
@@ -2271,8 +2224,20 @@ void SBlendSpaceGridWidget::Tick(const FGeometry& AllottedGeometry, const double
 			LastPreviewingMousePosition.Y = LocalMousePosition.Y;
 			FModifierKeysState ModifierKeyState = FSlateApplication::Get().GetModifierKeys();
 			bool bIsManualPreviewing = !bReadOnly && IsHovered() && bMouseIsOverGeometry && ModifierKeyState.IsControlDown();
-			PreviewPosition = Position.IsSet() && !bIsManualPreviewing ? Position.Get() : ScreenPositionToSampleValue(LastPreviewingMousePosition, false);
-			PreviewPosition = BlendSpace->GetClampedAndWrappedBlendInput(PreviewPosition);
+			if (Position.IsSet() && !bIsManualPreviewing)
+			{
+				PreviewPosition = Position.Get();
+				if (bReadOnly)
+				{
+					// Happens when we are showing in the graph - don't want to render outside the valid region
+					PreviewPosition = BlendSpace->GetClampedAndWrappedBlendInput(PreviewPosition);
+				}
+			}
+			else
+			{
+				PreviewPosition = ScreenPositionToSampleValue(LastPreviewingMousePosition, false);
+			}
+			
 
 			if (FilteredPosition.IsSet())
 			{
@@ -2309,13 +2274,12 @@ void SBlendSpaceGridWidget::SetPreviewingState(const FVector& InPosition, const
 	if (const UBlendSpace* BlendSpace = BlendSpaceBase.Get())
 	{
 		PreviewFilteredPosition = BlendSpace->GetClampedAndWrappedBlendInput(InFilteredPosition);
-		PreviewPosition = BlendSpace->GetClampedAndWrappedBlendInput(InPosition);
 	}
 	else
 	{
 		PreviewFilteredPosition = InFilteredPosition;
-		PreviewPosition = InPosition;
 	}
+	PreviewPosition = InPosition;
 }
 
 void SBlendSpaceGridWidget::InvalidateCachedData()
diff --git a/Engine/Source/Editor/Persona/Private/SAnimationBlendSpaceGridWidget.h b/Engine/Source/Editor/Persona/Private/SAnimationBlendSpaceGridWidget.h
--- a/Engine/Source/Editor/Persona/Private/SAnimationBlendSpaceGridWidget.h
+++ b/Engine/Source/Editor/Persona/Private/SAnimationBlendSpaceGridWidget.h
@@ -258,7 +258,6 @@ private:
 	bool bSamplePreviewing;
 	FVector2D LastPreviewingMousePosition;
 	FVector PreviewPosition;
-	bool bPreviewPositionSet;
 	bool bAdvancedPreview;
 	TArray<FBlendSampleData> PreviewedSamples;
 	/** Previous position in the triangulation/segmentation */
@@ -356,16 +355,10 @@ private:
 	float TextMargin;
 	bool bShowTriangulation;
 	bool bShowAnimationNames;
-
 	bool bStretchToFit;
 	FMargin GridRatioMargin;
-
-	bool bPreviewToolTipHidden;
-
 	bool bReadOnly;
-
 	bool bShowAxisLabels;
-
 	bool bShowSettingsButtons;
 
 	FStatusBarMessageHandle StatusBarMessageHandle;
