Non-unity fixes

diff --git a/Engine/Source/Editor/AnimGraph/Private/AnimStateAliasNode.cpp b/Engine/Source/Editor/AnimGraph/Private/AnimStateAliasNode.cpp
--- a/Engine/Source/Editor/AnimGraph/Private/AnimStateAliasNode.cpp
+++ b/Engine/Source/Editor/AnimGraph/Private/AnimStateAliasNode.cpp
@@ -6,6 +6,7 @@
 #include "Kismet2/Kismet2NameValidators.h"
 #include "AnimationStateMachineGraph.h"
 #include "EdGraph/EdGraph.h"
+#include "EdGraph/EdGraphSchema.h"
 #include "AnimStateNode.h"
 
 #define LOCTEXT_NAMESPACE "AnimStateAliasNode"
diff --git a/Engine/Source/Editor/Persona/Private/AnimTimeline/AnimTimelineTrack_Curves.cpp b/Engine/Source/Editor/Persona/Private/AnimTimeline/AnimTimelineTrack_Curves.cpp
--- a/Engine/Source/Editor/Persona/Private/AnimTimeline/AnimTimelineTrack_Curves.cpp
+++ b/Engine/Source/Editor/Persona/Private/AnimTimeline/AnimTimelineTrack_Curves.cpp
@@ -15,6 +15,7 @@
 #include "SAnimOutlinerItem.h"
 #include "Preferences/PersonaOptions.h"
 #include "SListViewSelectorDropdownMenu.h"
+#include "Animation/AnimData/AnimDataModel.h"
 
 #define LOCTEXT_NAMESPACE "FAnimTimelineTrack_Notifies"
 
diff --git a/Engine/Source/Editor/Persona/Private/AnimTimeline/AnimTimelineTrack_FloatCurve.cpp b/Engine/Source/Editor/Persona/Private/AnimTimeline/AnimTimelineTrack_FloatCurve.cpp
--- a/Engine/Source/Editor/Persona/Private/AnimTimeline/AnimTimelineTrack_FloatCurve.cpp
+++ b/Engine/Source/Editor/Persona/Private/AnimTimeline/AnimTimelineTrack_FloatCurve.cpp
@@ -19,6 +19,7 @@
 #include "Widgets/Input/SComboButton.h"
 #include "Widgets/Layout/SBox.h"
 #include "AnimModel_AnimSequenceBase.h"
+#include "Animation/AnimData/AnimDataModel.h"
 #include "SAnimOutlinerItem.h"
 
 #define LOCTEXT_NAMESPACE "FAnimTimelineTrack_FloatCurve"
diff --git a/Engine/Source/Editor/Persona/Private/SAnimSequenceCurveEditor.cpp b/Engine/Source/Editor/Persona/Private/SAnimSequenceCurveEditor.cpp
--- a/Engine/Source/Editor/Persona/Private/SAnimSequenceCurveEditor.cpp
+++ b/Engine/Source/Editor/Persona/Private/SAnimSequenceCurveEditor.cpp
@@ -17,6 +17,7 @@
 #include "Editor.h"
 #include "Framework/MultiBox/MultiBoxBuilder.h"
 #include "EditorStyleSet.h"
+#include "Animation/AnimData/AnimDataModel.h"
 
 #define LOCTEXT_NAMESPACE "SAnimSequenceCurveEditor"
 
diff --git a/Engine/Source/Runtime/Engine/Classes/Animation/AnimSequenceHelpers.h b/Engine/Source/Runtime/Engine/Classes/Animation/AnimSequenceHelpers.h
--- a/Engine/Source/Runtime/Engine/Classes/Animation/AnimSequenceHelpers.h
+++ b/Engine/Source/Runtime/Engine/Classes/Animation/AnimSequenceHelpers.h
@@ -16,6 +16,7 @@ struct FRawCurveTracks;
 struct FRawAnimSequenceTrack;
 struct FBoneAnimationTrack;
 struct FRawAnimSequenceTrack;
+class IAnimationDataController;
 
 enum class EAnimInterpolationType : uint8;
 
diff --git a/Engine/Source/Runtime/Engine/Private/Animation/AnimSequenceBase.cpp b/Engine/Source/Runtime/Engine/Private/Animation/AnimSequenceBase.cpp
--- a/Engine/Source/Runtime/Engine/Private/Animation/AnimSequenceBase.cpp
+++ b/Engine/Source/Runtime/Engine/Private/Animation/AnimSequenceBase.cpp
@@ -13,6 +13,7 @@
 #include "Animation/AnimationPoseData.h"
 #include "Animation/AnimSequenceHelpers.h"
 #include "Animation/AnimData/AnimDataModel.h"
+#include "Modules/ModuleManager.h"
 
 #if WITH_EDITOR
 #include "IAnimationDataControllerModule.h"
diff --git a/Engine/Source/Runtime/Engine/Public/FXSystem.h b/Engine/Source/Runtime/Engine/Public/FXSystem.h
--- a/Engine/Source/Runtime/Engine/Public/FXSystem.h
+++ b/Engine/Source/Runtime/Engine/Public/FXSystem.h
@@ -19,6 +19,7 @@ class FGPUSpriteResources;
 struct FParticleEmitterInstance;
 class FRHIUniformBuffer;
 class UVectorFieldComponent;
+class FGlobalDistanceFieldParameterData;
 
 /*-----------------------------------------------------------------------------
 	Forward declarations.
diff --git a/Engine/Source/Runtime/Renderer/Private/VT/VirtualTextureAllocator.cpp b/Engine/Source/Runtime/Renderer/Private/VT/VirtualTextureAllocator.cpp
--- a/Engine/Source/Runtime/Renderer/Private/VT/VirtualTextureAllocator.cpp
+++ b/Engine/Source/Runtime/Renderer/Private/VT/VirtualTextureAllocator.cpp
@@ -3,6 +3,9 @@
 #include "VirtualTextureAllocator.h"
 #include "AllocatedVirtualTexture.h"
 #include "VirtualTexturing.h"
+#include "Modules/ModuleManager.h"
+#include "IImageWrapperModule.h"
+#include "Misc/Paths.h"
 
 #if WITH_EDITOR
 #include "IImageWrapper.h"
