Open in new window for animation asset browse, with shift+double click or via context menu

diff --git a/Engine/Source/Editor/Persona/Private/SAnimationSequenceBrowser.cpp b/Engine/Source/Editor/Persona/Private/SAnimationSequenceBrowser.cpp
--- a/Engine/Source/Editor/Persona/Private/SAnimationSequenceBrowser.cpp
+++ b/Engine/Source/Editor/Persona/Private/SAnimationSequenceBrowser.cpp
@@ -473,16 +473,23 @@ SAnimationSequenceBrowser::~SAnimationSequenceBrowser()
 
 void SAnimationSequenceBrowser::OnRequestOpenAsset(const FAssetData& AssetData, bool bFromHistory)
 {
-		if (UObject* RawAsset = AssetData.GetAsset())
-		{
+	if (UObject* RawAsset = AssetData.GetAsset())
+	{
 		if (UAnimationAsset* AnimationAsset = Cast<UAnimationAsset>(RawAsset))
+		{
+			if (FSlateApplication::Get().GetModifierKeys().IsShiftDown())
 			{
-			if (!bFromHistory)
+				IAnimationEditorModule& AnimationEditorModule = FModuleManager::LoadModuleChecked<IAnimationEditorModule>("AnimationEditor");
+				AnimationEditorModule.CreateAnimationEditor(EToolkitMode::Standalone, nullptr, AnimationAsset);
+			}
+			else
 			{
-				AddAssetToHistory(AssetData);
+				if (!bFromHistory)
+				{
+					AddAssetToHistory(AssetData);
+				}
+				OnOpenNewAsset.ExecuteIfBound(AnimationAsset);
 			}
-
-			OnOpenNewAsset.ExecuteIfBound(AnimationAsset);
 		}
 		else if (USoundWave* SoundWave = Cast<USoundWave>(RawAsset))
 		{
@@ -616,6 +623,16 @@ TSharedPtr<SWidget> SAnimationSequenceBrowser::OnGetAssetContextMenu(const TArra
 				)
 			);
 
+		MenuBuilder.AddMenuEntry(
+			LOCTEXT("AnimSequenceBase_OpenInNewWindow", "Open In New Window"),
+			LOCTEXT("AnimSequenceBase_OpenInNewWindowTooltip", "Will always open asset in a new window, and not re-use existing window. (Shift+Double-Click)"),
+			FSlateIcon(FEditorStyle::GetStyleSetName(), "ContentBrowser.AssetActions.OpenInExternalEditor"),
+			FUIAction(
+				FExecuteAction::CreateSP(this, &SAnimationSequenceBrowser::OpenInNewWindow, SelectedAssets),
+				FCanExecuteAction()
+			)
+		);
+			
 		MenuBuilder.AddMenuEntry(FGlobalEditorCommonCommands::Get().FindInContentBrowser);
 	}
 	MenuBuilder.EndSection();
@@ -640,6 +657,21 @@ TSharedPtr<SWidget> SAnimationSequenceBrowser::OnGetAssetContextMenu(const TArra
 	return MenuBuilder.MakeWidget();
 }
 
+void SAnimationSequenceBrowser::OpenInNewWindow(TArray<FAssetData> AnimationAssets)
+{
+	for(auto Iter = AnimationAssets.CreateConstIterator(); Iter; ++Iter)
+	{
+		UObject* Asset =  Iter->GetAsset();
+		UAnimationAsset* AnimationAsset = Cast<UAnimationAsset>(Asset);
+		if (AnimationAsset)
+		{
+			IAnimationEditorModule& AnimationEditorModule = FModuleManager::LoadModuleChecked<IAnimationEditorModule>("AnimationEditor");
+			AnimationEditorModule.CreateAnimationEditor(EToolkitMode::Standalone, nullptr, AnimationAsset);
+		}
+	}
+}
+
+
 void SAnimationSequenceBrowser::FindInContentBrowser()
 {
 	TArray<FAssetData> CurrentSelection = GetCurrentSelectionDelegate.Execute();
diff --git a/Engine/Source/Editor/Persona/Private/SAnimationSequenceBrowser.h b/Engine/Source/Editor/Persona/Private/SAnimationSequenceBrowser.h
--- a/Engine/Source/Editor/Persona/Private/SAnimationSequenceBrowser.h
+++ b/Engine/Source/Editor/Persona/Private/SAnimationSequenceBrowser.h
@@ -74,6 +74,9 @@ public:
 	/** Delegate that handles creation of context menu */
 	TSharedPtr<SWidget> OnGetAssetContextMenu(const TArray<FAssetData>& SelectedAssets);
 
+	/** Delegate to open animation assets each in a new window */
+	void OpenInNewWindow(TArray<FAssetData> AnimationAssets);
+
 	/** Delegate to handle "Find in Content Browser" context menu option */
 	void FindInContentBrowser();
 
diff --git a/Engine/Source/Editor/Persona/Private/SAssetFamilyShortcutBar.cpp b/Engine/Source/Editor/Persona/Private/SAssetFamilyShortcutBar.cpp
--- a/Engine/Source/Editor/Persona/Private/SAssetFamilyShortcutBar.cpp
+++ b/Engine/Source/Editor/Persona/Private/SAssetFamilyShortcutBar.cpp
@@ -165,7 +165,7 @@ public:
 				.ForegroundColor(FSlateColor::UseForeground())
 				.ComboButtonStyle(&FAppStyle::Get(), "SegmentedCombo.Right")
 				.OnGetMenuContent(this, &SAssetShortcut::HandleGetMenuContent)
-				.ToolTipText(LOCTEXT("AssetComboTooltip", "Find other assets of this type and perform asset operations./nShift-Click to open in new window."))
+				.ToolTipText(LOCTEXT("AssetComboTooltip", "Find other assets of this type and perform asset operations.\nShift-Click to open in new window."))
 			]
 		];
 
