Default the time snapping of the anim sequence curve editor to the input anim framerate

diff --git a/Engine/Source/Editor/Persona/Private/SAnimSequenceCurveEditor.cpp b/Engine/Source/Editor/Persona/Private/SAnimSequenceCurveEditor.cpp
--- a/Engine/Source/Editor/Persona/Private/SAnimSequenceCurveEditor.cpp
+++ b/Engine/Source/Editor/Persona/Private/SAnimSequenceCurveEditor.cpp
@@ -348,6 +348,7 @@ void SAnimSequenceCurveEditor::Construct(const FArguments& InArgs, const TShared
 
 	FCurveEditorInitParams CurveEditorInitParams;
 	CurveEditor->InitCurveEditor(CurveEditorInitParams);
+	CurveEditor->InputSnapRateAttribute = InAnimSequence->GetSamplingFrameRate();
 
 	AnimSequence = InAnimSequence;
 
