Change recompute tangents vertex color blending from a global setting to per mesh section setting
*) remove r.SkinCache.BlendUsingVertexColorForRecomputeTangents
*) tick box & combo RGB choices are replaced by a single combo box choices of: None/All/R/G/B
*) intermediate tangent buffers are only created when >=1 sections on a mesh require vertex color blending

diff --git a/Engine/Source/Editor/Persona/Private/PersonaMeshDetails.cpp b/Engine/Source/Editor/Persona/Private/PersonaMeshDetails.cpp
--- a/Engine/Source/Editor/Persona/Private/PersonaMeshDetails.cpp
+++ b/Engine/Source/Editor/Persona/Private/PersonaMeshDetails.cpp
@@ -4578,56 +4578,63 @@ TSharedRef<SWidget> FPersonaMeshDetails::OnGenerateCustomSectionWidgetsForSectio
 		]
 	];
 #endif// #if WITH_APEX_CLOTHING || WITH_CHAOS_CLOTHING
+
 	SectionWidget->AddSlot()
 	.AutoHeight()
 	.Padding(0, 2, 0, 0)
+	.HAlign(HAlign_Fill)
 	[
 		SNew(SHorizontalBox)
-		+SHorizontalBox::Slot()
+		+ SHorizontalBox::Slot()
 		.AutoWidth()
-		.Padding(2, 0, 2, 0)
+		.VAlign(VAlign_Center)
 		[
-			SNew(SCheckBox)
-			.IsChecked(this, &FPersonaMeshDetails::IsSectionShadowCastingEnabled, LODIndex, SectionIndex)
-			.OnCheckStateChanged(this, &FPersonaMeshDetails::OnSectionShadowCastingChanged, LODIndex, SectionIndex)
+			SNew(SBox)
+			.HAlign(HAlign_Right)
+			.MinDesiredWidth(65.0f)
 			[
 				SNew(STextBlock)
 				.Font(IDetailLayoutBuilder::GetDetailFont())
-				.Text(LOCTEXT("Cast Shadows", "Cast Shadows"))
+			.Text(LOCTEXT("RecomputeTangents", "Recompute Tangents"))
 			]
 		]
+		+ SHorizontalBox::Slot()
+		.FillWidth(1.0f)
+		.Padding(5, 2, 0, 0)
+		[
+			SNew(SComboButton)
+			.IsEnabled(IsGPUSkinCacheAvailable(GMaxRHIShaderPlatform))
+			.OnGetMenuContent(this, &FPersonaMeshDetails::OnGenerateRecomputeTangentsSetting, LODIndex, SectionIndex)
+			.VAlign(VAlign_Center)
+			.ContentPadding(2)
+			.ButtonContent()
+			[
+				SNew(STextBlock)
+				.Font(IDetailLayoutBuilder::GetDetailFont())
+				.Text(this, &FPersonaMeshDetails::GetCurrentRecomputeTangentsSetting, LODIndex, SectionIndex)
+				.ToolTipText(LOCTEXT("RecomputeTangentsToolTip", "Require (Support Skincache Shaders) enabled in Project Settings and compute shaders. None:Do_NOT_Recompute, All:Do_Recompute, R/G/B:Do_Recompute_And_Use_Vertex_Color_Channel_As_Blending_Mask"))
+			]
+		]
+	];
+
+	SectionWidget->AddSlot()
+	.AutoHeight()
+	.Padding(0, 2, 0, 0)
+	[
+		SNew(SHorizontalBox)
 		+SHorizontalBox::Slot()
 		.AutoWidth()
 		.Padding(2, 0, 2, 0)
 		[
 			SNew(SCheckBox)
-			.IsEnabled(IsGPUSkinCacheAvailable(GMaxRHIShaderPlatform))
-			.IsChecked(this, &FPersonaMeshDetails::IsSectionRecomputeTangentEnabled, LODIndex, SectionIndex)
-			.OnCheckStateChanged(this, &FPersonaMeshDetails::OnSectionRecomputeTangentChanged, LODIndex, SectionIndex)
+			.IsChecked(this, &FPersonaMeshDetails::IsSectionShadowCastingEnabled, LODIndex, SectionIndex)
+			.OnCheckStateChanged(this, &FPersonaMeshDetails::OnSectionShadowCastingChanged, LODIndex, SectionIndex)
 			[
 				SNew(STextBlock)
 				.Font(IDetailLayoutBuilder::GetDetailFont())
-				.Text(LOCTEXT("RecomputeTangent_Title", "Recompute Tangent"))
-				.ToolTipText(LOCTEXT("RecomputeTangent_Tooltip", "This feature only works if you enable (Support Skincache Shaders) in the Project Settings. Please note that skin cache is an experimental feature and only works if you have compute shaders."))
+				.Text(LOCTEXT("Cast Shadows", "Cast Shadows"))
 			]
 		]
-		+ SHorizontalBox::Slot()
-			.FillWidth(1.0f)
-			.Padding(5, 2, 0, 0)
-			[
-				SNew(SComboButton)
-				.IsEnabled(this, &FPersonaMeshDetails::IsGenerateRecomputeTangentsVertexChannelMaskPicker, LODIndex, SectionIndex)
-			    .OnGetMenuContent(this, &FPersonaMeshDetails::OnGenerateRecomputeTangentsVertexChannelMaskPicker, LODIndex, SectionIndex)
-			    .VAlign(VAlign_Center)
-			    .ContentPadding(2)
-			    .ButtonContent()
-			    [
-			    	SNew(STextBlock)
-			    	.Font(IDetailLayoutBuilder::GetDetailFont())
-			        .Text(this, &FPersonaMeshDetails::GetCurrentRecomputeTangentsVertexChannelMaskName, LODIndex, SectionIndex)
-			        .ToolTipText(LOCTEXT("RecomputeTangentsVertexChannelMaskTip", "Which Vertex Color Channel to use"))
-			    ]
-			]
 	];
 	return SectionWidget;
 }
@@ -5176,93 +5183,16 @@ void FPersonaMeshDetails::OnSectionShadowCastingChanged(ECheckBoxState NewState,
 	}
 }
 
-ECheckBoxState FPersonaMeshDetails::IsSectionRecomputeTangentEnabled(int32 LODIndex, int32 SectionIndex) const
-{
-	ECheckBoxState State = ECheckBoxState::Unchecked;
-	const USkeletalMesh* Mesh = GetPersonaToolkit()->GetMesh();
-	if (Mesh == nullptr)
-		return State;
-
-	check(Mesh->GetImportedModel());
-
-	if (!Mesh->GetImportedModel()->LODModels.IsValidIndex(LODIndex))
-		return State;
-
-	const FSkeletalMeshLODModel& LODModel = Mesh->GetImportedModel()->LODModels[LODIndex];
-
-	if (!LODModel.Sections.IsValidIndex(SectionIndex))
-		return State;
-
-	const FSkelMeshSection& Section = LODModel.Sections[SectionIndex];
-
-	State = Section.bRecomputeTangent ? ECheckBoxState::Checked : ECheckBoxState::Unchecked;
-	return State;
-}
-
-void FPersonaMeshDetails::OnSectionRecomputeTangentChanged(ECheckBoxState NewState, int32 LODIndex, int32 SectionIndex)
-{
-	ECheckBoxState State = ECheckBoxState::Unchecked;
-	USkeletalMesh* Mesh = GetPersonaToolkit()->GetMesh();
-	if (Mesh == nullptr)
-		return;
-
-	check(Mesh->GetImportedModel());
-
-	if (!Mesh->GetImportedModel()->LODModels.IsValidIndex(LODIndex))
-		return;
-
-	FSkeletalMeshLODModel& LODModel = Mesh->GetImportedModel()->LODModels[LODIndex];
-
-	if (!LODModel.Sections.IsValidIndex(SectionIndex))
-		return;
-
-	FSkelMeshSection& Section = LODModel.Sections[SectionIndex];
-
-	//Update Original PolygonGroup
-	auto UpdatePolygonGroupRecomputeTangent = [&Mesh, &LODModel, &Section, &SectionIndex](bool bRecomputeTangent)
-	{
-		FScopedSuspendAlternateSkinWeightPreview ScopedSuspendAlternateSkinnWeightPreview(Mesh);
-		{
-			FScopedSkeletalMeshPostEditChange ScopedPostEditChange(Mesh);
-			Section.bRecomputeTangent = bRecomputeTangent;
-			for (int32 AfterSectionIndex = SectionIndex + 1; AfterSectionIndex < LODModel.Sections.Num(); ++AfterSectionIndex)
-			{
-				if (LODModel.Sections[AfterSectionIndex].ChunkedParentSectionIndex == SectionIndex)
-				{
-					LODModel.Sections[AfterSectionIndex].bRecomputeTangent = bRecomputeTangent;
-				}
-				else
-				{
-					break;
-				}
-			}
-			//We display only the parent chunk
-			check(Section.ChunkedParentSectionIndex == INDEX_NONE);
-			SetSkelMeshSourceSectionUserData(LODModel, SectionIndex, Section.OriginalDataSectionIndex);
-		}
-	};
-
-	if (NewState == ECheckBoxState::Checked)
-	{
-		const FScopedTransaction Transaction(LOCTEXT("PersonaSetSectionRecomputeTangentFlag", "Persona editor: Set Recompute Tangent For Section"));
-		Mesh->Modify();
-		UpdatePolygonGroupRecomputeTangent(true);
-	}
-	else if (NewState == ECheckBoxState::Unchecked)
-	{
-		const FScopedTransaction Transaction(LOCTEXT("PersonaClearSectionRecomputeTangentFlag", "Persona editor: Clear Recompute Tangent For Section"));
-		Mesh->Modify();
-		UpdatePolygonGroupRecomputeTangent(false);
-	}
-}
+const FString RecomputeTangentsChannelNames[] = {
+			TEXT("None"),
+			TEXT("All"),
+			TEXT("Red"),
+			TEXT("Green"),
+			TEXT("Blue")
+};
 
-FText FPersonaMeshDetails::GetCurrentRecomputeTangentsVertexChannelMaskName(int32 LODIndex, int32 SectionIndex) const
+FText FPersonaMeshDetails::GetCurrentRecomputeTangentsSetting(int32 LODIndex, int32 SectionIndex) const
 {
-	const FString ChannelNames[] = {
-				TEXT("Red"),
-				TEXT("Green"),
-				TEXT("Blue")
-	};
 	const USkeletalMesh* Mesh = GetPersonaToolkit()->GetMesh();
 	if (Mesh == nullptr)
 		return FText::GetEmpty();
@@ -5279,51 +5209,36 @@ FText FPersonaMeshDetails::GetCurrentRecomputeTangentsVertexChannelMaskName(int3
 
 	FSkelMeshSection& Section = LODModel.Sections[SectionIndex];
 
-	int32 ChannelIndex = FMath::Clamp<int32>(static_cast<int32>(Section.RecomputeTangentsVertexMaskChannel), 0, 2); 
-	
-	return FText::FromString(ChannelNames[ChannelIndex]);
-}
-
-bool FPersonaMeshDetails::IsGenerateRecomputeTangentsVertexChannelMaskPicker(int32 LODIndex, int32 SectionIndex) const
-{
+	// Index_0 None:  bRecomputeTangent == false
+	// Index_1 All:   bRecomputeTangent == true & RecomputeTangentsVertexMaskChannel == ESkinVertexColorChannel::None
+	// Index_2 Red:   bRecomputeTangent == true & RecomputeTangentsVertexMaskChannel == ESkinVertexColorChannel::Red
+	// Index_3 Green: bRecomputeTangent == true & RecomputeTangentsVertexMaskChannel == ESkinVertexColorChannel::Green
+	// Index_4 Blue:  bRecomputeTangent == true & RecomputeTangentsVertexMaskChannel == ESkinVertexColorChannel::Blue
+	int32 ChannelIndex = 0;
+	if (Section.bRecomputeTangent)
+	{
+		if (Section.RecomputeTangentsVertexMaskChannel == ESkinVertexColorChannel::None)
+			ChannelIndex = 1;
+		else
+			ChannelIndex = FMath::Clamp<int32>(static_cast<int32>(Section.RecomputeTangentsVertexMaskChannel) + 2, 0, 4);
+	}
 	
-	const USkeletalMesh* Mesh = GetPersonaToolkit()->GetMesh();
-	if (Mesh == nullptr)
-		return false;
-
-	check(Mesh->GetImportedModel());
-
-	if (!Mesh->GetImportedModel()->LODModels.IsValidIndex(LODIndex))
-		return false;
-
-	const FSkeletalMeshLODModel& LODModel = Mesh->GetImportedModel()->LODModels[LODIndex];
-
-	if (!LODModel.Sections.IsValidIndex(SectionIndex))
-		return false;
-
-	const FSkelMeshSection& Section = LODModel.Sections[SectionIndex];
-	return Section.bRecomputeTangent;
+	return FText::FromString(RecomputeTangentsChannelNames[ChannelIndex]);
 }
 
-
-TSharedRef<SWidget> FPersonaMeshDetails::OnGenerateRecomputeTangentsVertexChannelMaskPicker(int32 LODIndex, int32 SectionIndex)
+TSharedRef<SWidget> FPersonaMeshDetails::OnGenerateRecomputeTangentsSetting(int32 LODIndex, int32 SectionIndex)
 {
-	const FString ChannelNames[] = {
-			TEXT("Red"),
-			TEXT("Green"),
-			TEXT("Blue")
-	};
 	FMenuBuilder MenuBuilder(true, NULL);
-	for (int32 ChannelIndex = 0; ChannelIndex < 3; ChannelIndex++) { 
-		FUIAction Action(FExecuteAction::CreateSP(this, &FPersonaMeshDetails::SetCurrentRecomputeTangentsVertexChannel, LODIndex, SectionIndex, ChannelIndex));
+	for (int32 ChannelIndex = 0; ChannelIndex < 5; ChannelIndex++) { 
+		FUIAction Action(FExecuteAction::CreateSP(this, &FPersonaMeshDetails::SetCurrentRecomputeTangentsSetting, LODIndex, SectionIndex, ChannelIndex));
 
-		MenuBuilder.AddMenuEntry(FText::FromString(ChannelNames[ChannelIndex]), FText::GetEmpty(), FSlateIcon(), Action);
+		MenuBuilder.AddMenuEntry(FText::FromString(RecomputeTangentsChannelNames[ChannelIndex]), FText::GetEmpty(), FSlateIcon(), Action);
 	}
 	return MenuBuilder.MakeWidget();
 }
 
 
-void FPersonaMeshDetails::SetCurrentRecomputeTangentsVertexChannel(int32 LODIndex, int32 SectionIndex, int32 Index)
+void FPersonaMeshDetails::SetCurrentRecomputeTangentsSetting(int32 LODIndex, int32 SectionIndex, int32 Index)
 {
 	 USkeletalMesh* Mesh = GetPersonaToolkit()->GetMesh();
 	if (Mesh == nullptr)
@@ -5341,21 +5256,27 @@ void FPersonaMeshDetails::SetCurrentRecomputeTangentsVertexChannel(int32 LODInde
 
 	FSkelMeshSection& Section = LODModel.Sections[SectionIndex];
 
-	Section.RecomputeTangentsVertexMaskChannel = static_cast<ESkinVertexColorChannel>(Index);
-
-
 	auto UpdatePolygonGroupRecomputeTangentVertexColor = [&Mesh, &LODModel, &Section, &SectionIndex, &Index]()
 	{
 		FScopedSuspendAlternateSkinWeightPreview ScopedSuspendAlternateSkinnWeightPreview(Mesh);
 		{
 			FScopedSkeletalMeshPostEditChange ScopedPostEditChange(Mesh);
-			Section.RecomputeTangentsVertexMaskChannel = static_cast<ESkinVertexColorChannel>(Index);
+
+			// Index_0 None:  bRecomputeTangent == false
+			// Index_1 All:   bRecomputeTangent == true & RecomputeTangentsVertexMaskChannel == ESkinVertexColorChannel::None
+			// Index_2 Red:   bRecomputeTangent == true & RecomputeTangentsVertexMaskChannel == ESkinVertexColorChannel::Red
+			// Index_3 Green: bRecomputeTangent == true & RecomputeTangentsVertexMaskChannel == ESkinVertexColorChannel::Green
+			// Index_4 Blue:  bRecomputeTangent == true & RecomputeTangentsVertexMaskChannel == ESkinVertexColorChannel::Blue
+			Section.bRecomputeTangent = Index > 0;
+			Section.RecomputeTangentsVertexMaskChannel = Index > 1 ? static_cast<ESkinVertexColorChannel>(Index - 2) : ESkinVertexColorChannel::None;
 			
 			for (int32 AfterSectionIndex = SectionIndex+1; AfterSectionIndex < LODModel.Sections.Num(); ++AfterSectionIndex)
 			{
 				if (LODModel.Sections[AfterSectionIndex].ChunkedParentSectionIndex == SectionIndex)
 				{
-					LODModel.Sections[AfterSectionIndex].RecomputeTangentsVertexMaskChannel = static_cast<ESkinVertexColorChannel>(Index);
+					// Index: None-0, All-1, Red-2, Green-3, Blue-4
+					LODModel.Sections[AfterSectionIndex].bRecomputeTangent = Index > 0;
+					LODModel.Sections[AfterSectionIndex].RecomputeTangentsVertexMaskChannel = Index > 1 ? static_cast<ESkinVertexColorChannel>(Index-2) : ESkinVertexColorChannel::None;
 				}
 				else
 				{
@@ -6137,4 +6058,4 @@ bool FPersonaMeshDetails::FilterOutBakePose(const FAssetData& AssetData, USkelet
 	return !(Skeleton && Skeleton->IsCompatibleSkeletonByAssetData(AssetData));
 }
 
-#undef LOCTEXT_NAMESPACE
+#undef LOCTEXT_NAMESPACE
\ No newline at end of file
diff --git a/Engine/Source/Editor/Persona/Private/PersonaMeshDetails.h b/Engine/Source/Editor/Persona/Private/PersonaMeshDetails.h
--- a/Engine/Source/Editor/Persona/Private/PersonaMeshDetails.h
+++ b/Engine/Source/Editor/Persona/Private/PersonaMeshDetails.h
@@ -524,32 +524,15 @@ private:
 	*/
 	void OnSectionShadowCastingChanged(ECheckBoxState NewState, int32 LODIndex, int32 SectionIndex);
 
-	/**
-	* Handler for check box display based on whether this section does recalculate normal or not
-	*
-	* @param LODIndex	The LODIndex we want to change
-	* @param SectionIndex	The SectionIndex we change the RecomputeTangent
-	*/
-	ECheckBoxState IsSectionRecomputeTangentEnabled(int32 LODIndex, int32 SectionIndex) const;
-
-	/**
-	* Handler for changing recalulate normal status on a section
-	*
-	* @param LODIndex	The LODIndex we want to change
-	* @param SectionIndex	The SectionIndex we change the RecomputeTangent
-	*/
-	void OnSectionRecomputeTangentChanged(ECheckBoxState NewState, int32 LODIndex, int32 SectionIndex);
-
 	/**
 	* Handler for selecting which vertex color to mask the blending of recomputing tangents
 	*
 	* @param LODIndex	The LODIndex we want to change
 	* @param SectionIndex	The SectionIndex we change the RecomputeTangent
 	*/
-	TSharedRef<class SWidget> OnGenerateRecomputeTangentsVertexChannelMaskPicker(int32 LODIndex, int32 SectionIndex);
-	bool IsGenerateRecomputeTangentsVertexChannelMaskPicker(int32 LODIndex, int32 SectionIndex) const;
-	FText GetCurrentRecomputeTangentsVertexChannelMaskName(int32 LODIndex, int32 SectionIndex) const;
-	void SetCurrentRecomputeTangentsVertexChannel(int32 LODIndex, int32 SectionIndex, int32 Index);
+	TSharedRef<class SWidget> OnGenerateRecomputeTangentsSetting(int32 LODIndex, int32 SectionIndex);
+	FText GetCurrentRecomputeTangentsSetting(int32 LODIndex, int32 SectionIndex) const;
+	void SetCurrentRecomputeTangentsSetting(int32 LODIndex, int32 SectionIndex, int32 Index);
 
 	/**
 	 * Handler for enabling delete button on materials
diff --git a/Engine/Source/Runtime/Engine/Private/GPUSkinCache.cpp b/Engine/Source/Runtime/Engine/Private/GPUSkinCache.cpp
--- a/Engine/Source/Runtime/Engine/Private/GPUSkinCache.cpp
+++ b/Engine/Source/Runtime/Engine/Private/GPUSkinCache.cpp
@@ -117,15 +117,6 @@ FAutoConsoleVariableRef CVarGPUSkinCacheAllowDupedVertesForRecomputeTangents(
 	ECVF_RenderThreadSafe
 );
 
-static int32 GBlendUsingVertexColorForRecomputeTangents = 0;
-FAutoConsoleVariableRef CVarGPUSkinCacheBlendUsingVertexColorForRecomputeTangents(
-	TEXT("r.SkinCache.BlendUsingVertexColorForRecomputeTangents"),
-	GBlendUsingVertexColorForRecomputeTangents,
-	TEXT("0: off (default)\n")
-	TEXT("1: Linear interpolation between source and recompute tangents.\n"),
-	ECVF_RenderThreadSafe
-);
-
 int32 GRecomputeTangentsParallelDispatch = 0;
 FAutoConsoleVariableRef CVarRecomputeTangentsParallelDispatch(
 	TEXT("r.SkinCache.RecomputeTangentsParallelDispatch"),
@@ -345,16 +336,8 @@ public:
 
 		FGPUSkinCache::FSkinCacheRWBuffer* GetActiveTangentRWBuffer()
 		{
-			bool bUseIntermediateTangentBuffer = IndexBuffer && GBlendUsingVertexColorForRecomputeTangents > 0;
-
-			if (bUseIntermediateTangentBuffer)
-			{
-				return IntermediateTangentBuffer;
-			}
-			else
-			{
-				return TangentBuffer;
-			}
+			// This is the buffer containing tangent results from the skinning CS pass
+			return (IndexBuffer && IntermediateTangentBuffer) ? IntermediateTangentBuffer : TangentBuffer;
 		}
 
 		inline FGPUSkinCache::FSkinCacheRWBuffer* GetIntermediateAccumulatedTangentBuffer()
@@ -538,7 +521,6 @@ protected:
 	FRHIShaderResourceView* MorphBuffer;
 	FShaderResourceViewRHIRef ClothBuffer;
 	int32 LOD;
-
 	bool bMultipleClothSkinInfluences;
 
 	friend class FGPUSkinCache;
@@ -965,7 +947,7 @@ public:
 		SetShaderValue(RHICmdList, ShaderRHI, NumTriangles, DispatchData.NumTriangles);
 
 		SetSRVParameter(RHICmdList, ShaderRHI, GPUPositionCacheBuffer, DispatchData.GetPositionRWBuffer()->Buffer.SRV);
-		SetSRVParameter(RHICmdList, ShaderRHI, GPUTangentCacheBuffer, GBlendUsingVertexColorForRecomputeTangents > 0 ? DispatchData.IntermediateTangentBuffer->Buffer.SRV : DispatchData.GetTangentRWBuffer()->Buffer.SRV);
+		SetSRVParameter(RHICmdList, ShaderRHI, GPUTangentCacheBuffer, DispatchData.GetActiveTangentRWBuffer()->Buffer.SRV);
 		SetSRVParameter(RHICmdList, ShaderRHI, UVsInputBuffer, DispatchData.UVsBufferSRV);
 
 		SetShaderValue(RHICmdList, ShaderRHI, SkinCacheStart, DispatchData.OutputStreamStart);
@@ -1251,7 +1233,7 @@ void FGPUSkinCache::DispatchUpdateSkinTangents(FRHICommandListImmediate& RHICmdL
 		TShaderMapRef<FRecomputeTangentsPerVertexPassCS<0>> ComputeShader0(GlobalShaderMap);
 		TShaderMapRef<FRecomputeTangentsPerVertexPassCS<1>> ComputeShader1(GlobalShaderMap);
 		TShaderRef<FBaseRecomputeTangentsPerVertexShader> ComputeShader;
-		if (GBlendUsingVertexColorForRecomputeTangents > 0)
+		if (DispatchData.Section->RecomputeTangentsVertexMaskChannel < ESkinVertexColorChannel::None)
 			ComputeShader = ComputeShader1;
 		else
 			ComputeShader = ComputeShader0;
@@ -1276,10 +1258,10 @@ void FGPUSkinCache::DispatchUpdateSkinTangents(FRHICommandListImmediate& RHICmdL
 	}
 }
 
-FGPUSkinCache::FRWBuffersAllocation* FGPUSkinCache::TryAllocBuffer(uint32 NumVertices, bool WithTangnents, uint32 NumTriangles, FRHICommandListImmediate& RHICmdList)
+FGPUSkinCache::FRWBuffersAllocation* FGPUSkinCache::TryAllocBuffer(uint32 NumVertices, bool WithTangnents, bool UseIntermediateTangents, uint32 NumTriangles, FRHICommandListImmediate& RHICmdList)
 {
 	uint64 MaxSizeInBytes = (uint64)(GSkinCacheSceneMemoryLimitInMB * 1024.0f * 1024.0f);
-	uint64 RequiredMemInBytes = FRWBuffersAllocation::CalculateRequiredMemory(NumVertices, WithTangnents, NumTriangles);
+	uint64 RequiredMemInBytes = FRWBuffersAllocation::CalculateRequiredMemory(NumVertices, WithTangnents, UseIntermediateTangents, NumTriangles);
 	if (bRequiresMemoryLimit && UsedMemoryInBytes + RequiredMemInBytes >= MaxSizeInBytes)
 	{
 		ExtraRequiredMemory += RequiredMemInBytes;
@@ -1288,7 +1270,7 @@ FGPUSkinCache::FRWBuffersAllocation* FGPUSkinCache::TryAllocBuffer(uint32 NumVer
 		return nullptr;
 	}
 
-	FRWBuffersAllocation* NewAllocation = new FRWBuffersAllocation(NumVertices, WithTangnents, NumTriangles, RHICmdList);
+	FRWBuffersAllocation* NewAllocation = new FRWBuffersAllocation(NumVertices, WithTangnents, UseIntermediateTangents, NumTriangles, RHICmdList);
 	Allocations.Add(NewAllocation);
 
 	UsedMemoryInBytes += RequiredMemInBytes;
@@ -1531,13 +1513,11 @@ void FGPUSkinCache::ProcessEntry(
 	{
 		bool WithTangents = RecomputeTangentsMode > 0;
 		int32 TotalNumVertices = VertexFactory->GetNumVertices();
-
-		// IntermediateAccumulatedTangents buffer is needed if mesh has at least one section needing recomputing tangents.
-		// TotalNumTriangles > 0 signals creation of IntermediateAccumulatedTangents buffer.
-		uint32 TotalNumTriangles = 0;
-		if (GRecomputeTangentsParallelDispatch && RecomputeTangentsMode > 0)
+		
+		bool bShouldRecomputeTangent = false;
+		if (RecomputeTangentsMode > 0)
 		{
-			bool bShouldRecomputeTangent = (RecomputeTangentsMode == 1);
+			bShouldRecomputeTangent = (RecomputeTangentsMode == 1);
 			if (!bShouldRecomputeTangent)
 			{
 				for (const FSkelMeshRenderSection& RenderSection : LodData.RenderSections)
@@ -1549,18 +1529,32 @@ void FGPUSkinCache::ProcessEntry(
 					}
 				}
 			}
-			if (bShouldRecomputeTangent)
+		}
+
+		// IntermediateTangents buffer is needed if mesh has at least one section using vertex color as recompute tangents blending mask
+		bool bEntryUseIntermediateTangents = false;
+		if (bShouldRecomputeTangent)
+		{
+			for (const FSkelMeshRenderSection& RenderSection : LodData.RenderSections)
 			{
-				TotalNumTriangles = LodData.GetTotalFaces();
+				if (RenderSection.RecomputeTangentsVertexMaskChannel < ESkinVertexColorChannel::None)
+				{
+					bEntryUseIntermediateTangents = true;
+					break;
+				}
 			}
 		}
 
-		FRWBuffersAllocation* NewPositionAllocation = TryAllocBuffer(TotalNumVertices, WithTangents, TotalNumTriangles, RHICmdList);
+		// IntermediateAccumulatedTangents buffer is needed if mesh has at least one section needing recomputing tangents.
+		// TotalNumTriangles > 0 signals creation of IntermediateAccumulatedTangents buffer.
+		const uint32 TotalNumTriangles = (GRecomputeTangentsParallelDispatch && bShouldRecomputeTangent) ? LodData.GetTotalFaces() : 0;
+
+		FRWBuffersAllocation* NewPositionAllocation = TryAllocBuffer(TotalNumVertices, WithTangents, bEntryUseIntermediateTangents, TotalNumTriangles, RHICmdList);
 		if (!NewPositionAllocation)
 		{
 			if (GSkinCachePrintMemorySummary > 0)
 			{
-				uint64 RequiredMemInBytes = FRWBuffersAllocation::CalculateRequiredMemory(TotalNumVertices, WithTangents, TotalNumTriangles);
+				uint64 RequiredMemInBytes = FRWBuffersAllocation::CalculateRequiredMemory(TotalNumVertices, WithTangents, bEntryUseIntermediateTangents, TotalNumTriangles);
 				UE_LOG(LogSkinCache, Warning, TEXT("FGPUSkinCache::ProcessEntry failed to allocate %.3fMB for mesh %s LOD%d, extra required memory increased to %.3fMB"),
 					RequiredMemInBytes / MBSize, *GetSkeletalMeshObjectName(Skin), LODIndex, ExtraRequiredMemory / MBSize);
 			}
@@ -2195,12 +2189,6 @@ bool FGPUSkinCache::IsEntryValid(FGPUSkinCacheEntry* SkinCacheEntry, int32 Secti
 	return SkinCacheEntry->IsSectionValid(Section);
 }
 
-bool FGPUSkinCache::UseIntermediateTangents()
-{
-	int32 RecomputeTangentsMode = GForceRecomputeTangents > 0 ? 1 : GSkinCacheRecomputeTangents;
-	return (RecomputeTangentsMode > 0) && (GBlendUsingVertexColorForRecomputeTangents > 0);
-}
-
 FGPUSkinBatchElementUserData* FGPUSkinCache::InternalGetFactoryUserData(FGPUSkinCacheEntry* Entry, int32 Section)
 {
 	return &Entry->BatchElementsUserData[Section];
diff --git a/Engine/Source/Runtime/Engine/Private/SkeletalMeshLODModel.cpp b/Engine/Source/Runtime/Engine/Private/SkeletalMeshLODModel.cpp
--- a/Engine/Source/Runtime/Engine/Private/SkeletalMeshLODModel.cpp
+++ b/Engine/Source/Runtime/Engine/Private/SkeletalMeshLODModel.cpp
@@ -373,8 +373,8 @@ FArchive& operator<<(FArchive& Ar, FSkelMeshSection& S)
 	}
 	else
 	{
-		// Our default is to use the green vertex color channel 
-		S.RecomputeTangentsVertexMaskChannel = ESkinVertexColorChannel::Green;
+		// Our default is not to use vertex color as mask
+		S.RecomputeTangentsVertexMaskChannel = ESkinVertexColorChannel::None;
 	}
 
 	if (Ar.CustomVer(FEditorObjectVersion::GUID) >= FEditorObjectVersion::RefactorMeshEditorMaterials)
@@ -551,8 +551,8 @@ FArchive& operator<<(FArchive& Ar, FSkelMeshSourceSectionUserData& S)
 	}
  	else
 	{
-		// Our default is to use the green vertex color channel 
-		S.RecomputeTangentsVertexMaskChannel = ESkinVertexColorChannel::Green;
+		// Our default is not to use vertex color as mask
+		S.RecomputeTangentsVertexMaskChannel = ESkinVertexColorChannel::None;
 	}
 
 	Ar << S.bCastShadow;
diff --git a/Engine/Source/Runtime/Engine/Private/SkeletalMeshLODRenderData.cpp b/Engine/Source/Runtime/Engine/Private/SkeletalMeshLODRenderData.cpp
--- a/Engine/Source/Runtime/Engine/Private/SkeletalMeshLODRenderData.cpp
+++ b/Engine/Source/Runtime/Engine/Private/SkeletalMeshLODRenderData.cpp
@@ -95,8 +95,8 @@ FArchive& operator<<(FArchive& Ar, FSkelMeshRenderSection& S)
 	}
 	else
 	{
-		// Our default is to use the green vertex color channel 
-		S.RecomputeTangentsVertexMaskChannel = ESkinVertexColorChannel::Green;
+		// Our default is not to use vertex color as mask
+		S.RecomputeTangentsVertexMaskChannel = ESkinVertexColorChannel::None;
 	}
 	Ar << S.bCastShadow;
 	Ar << S.BaseVertexIndex;
diff --git a/Engine/Source/Runtime/Engine/Public/GPUSkinCache.h b/Engine/Source/Runtime/Engine/Public/GPUSkinCache.h
--- a/Engine/Source/Runtime/Engine/Public/GPUSkinCache.h
+++ b/Engine/Source/Runtime/Engine/Public/GPUSkinCache.h
@@ -203,8 +203,6 @@ public:
 
 	static bool IsEntryValid(FGPUSkinCacheEntry* SkinCacheEntry, int32 Section);
 
-	static bool UseIntermediateTangents();
-
 	ENGINE_API uint64 GetExtraRequiredMemoryAndReset();
 
 	enum
@@ -236,8 +234,8 @@ public:
 	{
 		friend struct FRWBufferTracker;
 
-		FRWBuffersAllocation(uint32 InNumVertices, bool InWithTangents, uint32 InNumTriangles, FRHICommandListImmediate& RHICmdList)
-			: NumVertices(InNumVertices), WithTangents(InWithTangents), NumTriangles(InNumTriangles)
+		FRWBuffersAllocation(uint32 InNumVertices, bool InWithTangents, bool InUseIntermediateTangents, uint32 InNumTriangles, FRHICommandListImmediate& RHICmdList)
+			: NumVertices(InNumVertices), WithTangents(InWithTangents), UseIntermediateTangents(InUseIntermediateTangents), NumTriangles(InNumTriangles)
 		{
 			for (int32 Index = 0; Index < NUM_BUFFERS; ++Index)
 			{
@@ -248,7 +246,7 @@ public:
 			{
 				Tangents.Buffer.Initialize(TEXT("SkinCacheTangents"), TangentBufferBytesPerElement, NumVertices * 2, PF_R16G16B16A16_SNORM, BUF_Static);
 				Tangents.AccessState = ERHIAccess::Unknown;
-				if (FGPUSkinCache::UseIntermediateTangents())
+				if (UseIntermediateTangents)
 				{
 					IntermediateTangents.Buffer.Initialize(TEXT("SkinCacheIntermediateTangents"), TangentBufferBytesPerElement, NumVertices * 2, PF_R16G16B16A16_SNORM, BUF_Static);
 					IntermediateTangents.AccessState = ERHIAccess::Unknown;
@@ -280,12 +278,12 @@ public:
 			}
 		}
 
-		static uint64 CalculateRequiredMemory(uint32 InNumVertices, bool InWithTangents, uint32 InNumTriangles)
+		static uint64 CalculateRequiredMemory(uint32 InNumVertices, bool InWithTangents, bool InUseIntermediateTangents, uint32 InNumTriangles)
 		{
 			uint64 PositionBufferSize = PosBufferBytesPerElement * InNumVertices * 3 * NUM_BUFFERS;
 			uint64 TangentBufferSize = InWithTangents ? TangentBufferBytesPerElement * InNumVertices * 2 : 0;
 			uint64 IntermediateTangentBufferSize = 0;
-			if (FGPUSkinCache::UseIntermediateTangents())
+			if (InUseIntermediateTangents)
 			{
 				IntermediateTangentBufferSize = InWithTangents ? TangentBufferBytesPerElement * InNumVertices * 2 : 0;
 			}
@@ -295,7 +293,7 @@ public:
 
 		uint64 GetNumBytes() const
 		{
-			return CalculateRequiredMemory(NumVertices, WithTangents, NumTriangles);
+			return CalculateRequiredMemory(NumVertices, WithTangents, UseIntermediateTangents, NumTriangles);
 		}
 
 		FSkinCacheRWBuffer* GetTangentBuffer()
@@ -305,7 +303,7 @@ public:
 
 		FSkinCacheRWBuffer* GetIntermediateTangentBuffer()
 		{
-			return WithTangents ? &IntermediateTangents : nullptr;
+			return (WithTangents && UseIntermediateTangents) ? &IntermediateTangents : nullptr;
 		}
 
 		FSkinCacheRWBuffer* GetIntermediateAccumulatedTangentBuffer()
@@ -325,6 +323,7 @@ public:
 
 		const uint32 NumVertices;
 		const bool WithTangents;
+		const bool UseIntermediateTangents;
 		const uint32 NumTriangles;
 
 		static const uint32 PosBufferBytesPerElement = 4;
@@ -455,7 +454,7 @@ protected:
 	TArray<FGPUSkinCacheEntry*> Entries;
 	TArray<FDispatchEntry> BatchDispatches;
 
-	FRWBuffersAllocation* TryAllocBuffer(uint32 NumVertices, bool WithTangnents, uint32 NumTriangles, FRHICommandListImmediate& RHICmdList);
+	FRWBuffersAllocation* TryAllocBuffer(uint32 NumVertices, bool WithTangnents, bool UseIntermediateTangents, uint32 NumTriangles, FRHICommandListImmediate& RHICmdList);
 	void DoDispatch(FRHICommandListImmediate& RHICmdList);
 	void DoDispatch(FRHICommandListImmediate& RHICmdList, FGPUSkinCacheEntry* SkinCacheEntry, int32 Section, int32 RevisionNumber);
 	void DispatchUpdateSkinTangents(FRHICommandListImmediate& RHICmdList, FGPUSkinCacheEntry* Entry, int32 SectionIndex, FSkinCacheRWBuffer*& StagingBuffer, bool bTrianglePass);
diff --git a/Engine/Source/Runtime/Engine/Public/Rendering/SkeletalMeshLODModel.h b/Engine/Source/Runtime/Engine/Public/Rendering/SkeletalMeshLODModel.h
--- a/Engine/Source/Runtime/Engine/Public/Rendering/SkeletalMeshLODModel.h
+++ b/Engine/Source/Runtime/Engine/Public/Rendering/SkeletalMeshLODModel.h
@@ -77,7 +77,7 @@ struct FSkelMeshSection
 	/** This section will recompute tangent in runtime */
 	bool bRecomputeTangent;
 
-	/** Vertex color channel to mask recompute tangents. R=0,G=1 (default),B=2 */
+	/** Vertex color channel to mask recompute tangents. R=0,G=1,B=2,A=None=3 */
 	ESkinVertexColorChannel RecomputeTangentsVertexMaskChannel;
 
 	/** This section will cast shadow */
@@ -155,7 +155,7 @@ struct FSkelMeshSection
 		, NumTriangles(0)
 		, bSelected(false)
 		, bRecomputeTangent(false)
-		, RecomputeTangentsVertexMaskChannel(ESkinVertexColorChannel::Green)
+		, RecomputeTangentsVertexMaskChannel(ESkinVertexColorChannel::None)
 		, bCastShadow(true)
 		, bLegacyClothingSection_DEPRECATED(false)
 		, CorrespondClothSectionIndex_DEPRECATED(-1)
@@ -264,7 +264,7 @@ struct FSkelMeshSourceSectionUserData
 
 	FSkelMeshSourceSectionUserData()
 		: bRecomputeTangent(false)
-		, RecomputeTangentsVertexMaskChannel(ESkinVertexColorChannel::Green)
+		, RecomputeTangentsVertexMaskChannel(ESkinVertexColorChannel::None)
 		, bCastShadow(true)
 		, CorrespondClothAssetIndex(INDEX_NONE)
 		, bDisabled(false)
diff --git a/Engine/Source/Runtime/Engine/Public/Rendering/SkeletalMeshLODRenderData.h b/Engine/Source/Runtime/Engine/Public/Rendering/SkeletalMeshLODRenderData.h
--- a/Engine/Source/Runtime/Engine/Public/Rendering/SkeletalMeshLODRenderData.h
+++ b/Engine/Source/Runtime/Engine/Public/Rendering/SkeletalMeshLODRenderData.h
@@ -74,7 +74,7 @@ struct FSkelMeshRenderSection
 		, NumTriangles(0)
 		, bRecomputeTangent(false)
 		, bCastShadow(true)
-		, RecomputeTangentsVertexMaskChannel(ESkinVertexColorChannel::Green)
+		, RecomputeTangentsVertexMaskChannel(ESkinVertexColorChannel::None)
 		, BaseVertexIndex(0)
 		, NumVertices(0)
 		, MaxBoneInfluences(4)
diff --git a/Engine/Source/Runtime/Engine/Public/SkeletalMeshTypes.h b/Engine/Source/Runtime/Engine/Public/SkeletalMeshTypes.h
--- a/Engine/Source/Runtime/Engine/Public/SkeletalMeshTypes.h
+++ b/Engine/Source/Runtime/Engine/Public/SkeletalMeshTypes.h
@@ -32,17 +32,18 @@ struct ESkeletalMeshVertexFlags
 	};
 };
 
-/** Name of vertex color channels */
+/** Name of vertex color channels, used by recompute tangents */
 enum class ESkinVertexColorChannel : uint8
 {
-	// 
+	// Use red channel as recompute tangents blending mask
 	Red = 0,
-	// 
+	// Use green channel as recompute tangents blending mask
 	Green = 1,
-	// 
+	// Use blue channel as recompute tangents blending mask
 	Blue = 2,
-	//
-	Alpha = 3
+	// Alpha channel not used by recompute tangents
+	Alpha = 3,
+	None = Alpha
 };
 
 
