Added an option to show mirrored animation in the animation preview window

diff --git a/Engine/Source/Editor/Persona/Private/AnimationEditorViewportClient.cpp b/Engine/Source/Editor/Persona/Private/AnimationEditorViewportClient.cpp
--- a/Engine/Source/Editor/Persona/Private/AnimationEditorViewportClient.cpp
+++ b/Engine/Source/Editor/Persona/Private/AnimationEditorViewportClient.cpp
@@ -13,6 +13,7 @@
 #include "Preferences/PersonaOptions.h"
 #include "Engine/CollisionProfile.h"
 #include "Animation/AnimBlueprintGeneratedClass.h"
+#include "Animation/MirrorDataTable.h"
 #include "GameFramework/WorldSettings.h"
 #include "PersonaModule.h"
 #include "Rendering/SkeletalMeshRenderData.h"
@@ -976,6 +977,12 @@ FText FAnimationViewportClient::GetDisplayInfo(bool bDisplayAllInfo) const
 			const FRuntimeSkinWeightProfileData* OverrideData = LODData.SkinWeightProfilesData.GetOverrideData(ProfileName);
 			TextValue = ConcatenateLine(TextValue, FText::Format(LOCTEXT("NumSkinWeightOverrides", "Skin Weight Profile Weights: {0}"),	(OverrideData && OverrideData->NumWeightsPerVertex > 0) ? FText::AsNumber(OverrideData->BoneWeights.Num() / OverrideData->NumWeightsPerVertex) : LOCTEXT("NoSkinWeightsOverridesForLOD", "no data for LOD")));
 		}
+		
+		const bool bMirroring = PreviewMeshComponent->PreviewInstance && PreviewMeshComponent->PreviewInstance->GetMirrorDataTable();
+        if (bMirroring)
+        {
+        	TextValue = ConcatenateLine(TextValue,FText::Format(LOCTEXT("Preview_mirrored", "Mirrored with {0} "), FText::FromString(PreviewMeshComponent->PreviewInstance->GetMirrorDataTable()->GetName())));
+        }
 	}
 
 	if (const IClothingSimulation* const ClothingSimulation = PreviewMeshComponent->GetClothingSimulation())
diff --git a/Engine/Source/Editor/Persona/Private/SAnimViewportToolBar.cpp b/Engine/Source/Editor/Persona/Private/SAnimViewportToolBar.cpp
--- a/Engine/Source/Editor/Persona/Private/SAnimViewportToolBar.cpp
+++ b/Engine/Source/Editor/Persona/Private/SAnimViewportToolBar.cpp
@@ -2,6 +2,8 @@
 
 
 #include "SAnimViewportToolBar.h"
+
+#include "AnimPreviewInstance.h"
 #include "Widgets/Layout/SBorder.h"
 #include "Widgets/Text/STextBlock.h"
 #include "Framework/MultiBox/MultiBoxBuilder.h"
@@ -46,6 +48,9 @@
 #include "UICommandList_Pinnable.h"
 #include "BoneSelectionWidget.h"
 #include "Widgets/Text/SRichTextBlock.h"
+#include "ContentBrowserModule.h"
+#include "IContentBrowserSingleton.h"
+#include "Animation/MirrorDataTable.h"
 
 #define LOCTEXT_NAMESPACE "AnimViewportToolBar"
 
@@ -628,7 +633,7 @@ TSharedRef<SWidget> SAnimViewportToolBar::GenerateCharacterMenu() const
 			InMenuBuilder.AddSubMenu(
 				LOCTEXT("CharacterMenu_AnimationSubMenu", "Animation"),
 				LOCTEXT("CharacterMenu_AnimationSubMenuToolTip", "Animation-related options"),
-				FNewMenuDelegate::CreateLambda([WeakSharedViewport = Viewport](FMenuBuilder& SubMenuBuilder)
+				FNewMenuDelegate::CreateLambda([this, WeakSharedViewport = Viewport](FMenuBuilder& SubMenuBuilder)
 				{
 					SubMenuBuilder.BeginSection("AnimViewportRootMotion", LOCTEXT("CharacterMenu_RootMotionLabel", "Root Motion"));
 					{
@@ -642,6 +647,31 @@ TSharedRef<SWidget> SAnimViewportToolBar::GenerateCharacterMenu() const
 						SubMenuBuilder.AddMenuEntry(FAnimViewportShowCommands::Get().ShowNonRetargetedAnimation);
 						SubMenuBuilder.AddMenuEntry(FAnimViewportShowCommands::Get().ShowAdditiveBaseBones);
 						SubMenuBuilder.AddMenuEntry(FAnimViewportShowCommands::Get().ShowSourceRawAnimation);
+
+						if (WeakSharedViewport.IsValid())
+						{
+							if ( UDebugSkelMeshComponent* PreviewComponent = WeakSharedViewport.Pin()->GetPreviewScene()->GetPreviewMeshComponent())
+							{
+								FUIAction DisableUnlessPreviewInstance(
+									FExecuteAction::CreateLambda([](){}),
+									FCanExecuteAction::CreateLambda([PreviewComponent]()
+									{
+										return (PreviewComponent->PreviewInstance && (PreviewComponent->PreviewInstance == PreviewComponent->GetAnimInstance() ) );
+									})
+								);
+
+								SubMenuBuilder.AddSubMenu(
+									LOCTEXT("CharacterMenu_AnimationSubMenu_MirrorSubMenu", "Mirror"),
+									LOCTEXT("CharacterMenu_AnimationSubMenu_MirrorSubMenuToolTip", "Mirror the animation using the selected mirror data table"),
+									FNewMenuDelegate::CreateRaw(const_cast<SAnimViewportToolBar*>(this), &SAnimViewportToolBar::FillCharacterMirrorMenu),
+									DisableUnlessPreviewInstance,
+									NAME_None,
+									EUserInterfaceActionType::Button,
+									false,
+									FSlateIcon(),
+									false);
+							}
+						}
 						SubMenuBuilder.AddMenuEntry(FAnimViewportShowCommands::Get().ShowBakedAnimation);
 						SubMenuBuilder.AddMenuEntry(FAnimViewportShowCommands::Get().DisablePostProcessBlueprint);
 					}
@@ -756,6 +786,50 @@ void SAnimViewportToolBar::FillCharacterAdvancedMenu(FMenuBuilder& MenuBuilder)
 	MenuBuilder.EndSection();
 }
 
+void SAnimViewportToolBar::FillCharacterMirrorMenu(FMenuBuilder& MenuBuilder) const
+{
+	FAssetPickerConfig AssetPickerConfig;
+	UDebugSkelMeshComponent* PreviewComp = Viewport.Pin()->GetPreviewScene()->GetPreviewMeshComponent();
+	USkeletalMesh* Mesh = PreviewComp->SkeletalMesh;
+	UAnimPreviewInstance* PreviewInstance = PreviewComp->PreviewInstance; 
+	if (Mesh && PreviewInstance)
+	{
+		USkeleton* Skeleton = Mesh->GetSkeleton();
+	
+		AssetPickerConfig.Filter.ClassNames.Add(*UMirrorDataTable::StaticClass()->GetName());
+		AssetPickerConfig.Filter.bRecursiveClasses = false;
+		AssetPickerConfig.bAllowNullSelection = true;
+		AssetPickerConfig.Filter.TagsAndValues.Add(TEXT("Skeleton"), FAssetData(Skeleton).GetExportTextName());
+		AssetPickerConfig.InitialAssetSelection = FAssetData(PreviewInstance->GetMirrorDataTable());
+		AssetPickerConfig.OnAssetSelected = FOnAssetSelected::CreateRaw(const_cast<SAnimViewportToolBar*>(this), &SAnimViewportToolBar::OnMirrorDataTableSelected);
+		AssetPickerConfig.InitialAssetViewType = EAssetViewType::List;
+		AssetPickerConfig.ThumbnailScale = 0.1f;
+		AssetPickerConfig.bAddFilterUI = false;
+		
+		FContentBrowserModule& ContentBrowserModule = FModuleManager::Get().LoadModuleChecked<FContentBrowserModule>(TEXT("ContentBrowser"));
+
+		MenuBuilder.AddWidget(
+			ContentBrowserModule.Get().CreateAssetPicker(AssetPickerConfig),
+			FText::GetEmpty()
+		);
+	}
+}
+
+void SAnimViewportToolBar::OnMirrorDataTableSelected(const FAssetData& SelectedMirrorTableData)
+{
+	UMirrorDataTable* MirrorDataTable = Cast<UMirrorDataTable>(SelectedMirrorTableData.GetAsset());
+	if (Viewport.Pin().IsValid())
+	{
+		UDebugSkelMeshComponent* PreviewComp = Viewport.Pin()->GetPreviewScene()->GetPreviewMeshComponent();
+		USkeletalMesh* Mesh = PreviewComp->SkeletalMesh;
+		UAnimPreviewInstance* PreviewInstance = PreviewComp->PreviewInstance; 
+		if (Mesh && PreviewInstance)
+		{
+			PreviewInstance->SetMirrorDataTable(MirrorDataTable);
+		}
+	}
+}
+
 void SAnimViewportToolBar::FillCharacterClothingMenu(FMenuBuilder& MenuBuilder)
 {
 #if WITH_APEX_CLOTHING || WITH_CHAOS_CLOTHING
diff --git a/Engine/Source/Editor/Persona/Private/SAnimViewportToolBar.h b/Engine/Source/Editor/Persona/Private/SAnimViewportToolBar.h
--- a/Engine/Source/Editor/Persona/Private/SAnimViewportToolBar.h
+++ b/Engine/Source/Editor/Persona/Private/SAnimViewportToolBar.h
@@ -84,6 +84,11 @@ private:
 	 */
 	void FillCharacterAdvancedMenu(FMenuBuilder& MenuBuilder) const;
 
+	/**
+	* Generates the Show -> Mirror sub menu content
+	*/
+	void FillCharacterMirrorMenu(FMenuBuilder& MenuBuilder) const;
+
 	/**
 	* Generates the Show -> Clothing sub menu content
 	*/
@@ -157,6 +162,9 @@ private:
 	/** Called when the floor offset slider is adjusted in the perspective viewport */
 	void OnFloorOffsetChanged( float NewValue );
 
+	/** Called when users select the mirror table in the menu */
+	void OnMirrorDataTableSelected(const FAssetData& SelectedMirrorTableData);
+	
 	// Called to determine if the gizmos can be used in the current preview
 	EVisibility GetTransformToolbarVisibility() const;
 
diff --git a/Engine/Source/Runtime/Engine/Classes/Animation/AnimSingleNodeInstance.h b/Engine/Source/Runtime/Engine/Classes/Animation/AnimSingleNodeInstance.h
--- a/Engine/Source/Runtime/Engine/Classes/Animation/AnimSingleNodeInstance.h
+++ b/Engine/Source/Runtime/Engine/Classes/Animation/AnimSingleNodeInstance.h
@@ -43,7 +43,10 @@ protected:
 	virtual void Montage_Advance(float DeltaTime) override;
 	//~ End UAnimInstance Interface
 public:
-
+	UFUNCTION(BlueprintCallable, Category="Animation")
+    void SetMirrorDataTable(const UMirrorDataTable* MirrorDataTable);
+	UFUNCTION(BlueprintCallable, Category="Animation")
+	const UMirrorDataTable* GetMirrorDataTable();
 	UFUNCTION(BlueprintCallable, Category="Animation")
 	void SetLooping(bool bIsLooping);
 	UFUNCTION(BlueprintCallable, Category="Animation")
diff --git a/Engine/Source/Runtime/Engine/Private/Animation/AnimSingleNodeInstance.cpp b/Engine/Source/Runtime/Engine/Private/Animation/AnimSingleNodeInstance.cpp
--- a/Engine/Source/Runtime/Engine/Private/Animation/AnimSingleNodeInstance.cpp
+++ b/Engine/Source/Runtime/Engine/Private/Animation/AnimSingleNodeInstance.cpp
@@ -232,6 +232,19 @@ void UAnimSingleNodeInstance::Montage_Advance(float DeltaTime)
 	}
 }
 
+void UAnimSingleNodeInstance::SetMirrorDataTable(const UMirrorDataTable* MirrorDataTable)
+{
+	FAnimSingleNodeInstanceProxy& Proxy = GetProxyOnGameThread<FAnimSingleNodeInstanceProxy>();
+	Proxy.SetMirrorDataTable(MirrorDataTable);
+}
+
+const UMirrorDataTable* UAnimSingleNodeInstance::GetMirrorDataTable()
+{
+	FAnimSingleNodeInstanceProxy& Proxy = GetProxyOnGameThread<FAnimSingleNodeInstanceProxy>();
+	return Proxy.GetMirrorDataTable();
+}
+
+
 void UAnimSingleNodeInstance::PlayAnim(bool bIsLooping, float InPlayRate, float InStartPosition)
 {
 	SetPlaying(true);
diff --git a/Engine/Source/Runtime/Engine/Private/Animation/AnimSingleNodeInstanceProxy.cpp b/Engine/Source/Runtime/Engine/Private/Animation/AnimSingleNodeInstanceProxy.cpp
--- a/Engine/Source/Runtime/Engine/Private/Animation/AnimSingleNodeInstanceProxy.cpp
+++ b/Engine/Source/Runtime/Engine/Private/Animation/AnimSingleNodeInstanceProxy.cpp
@@ -11,6 +11,7 @@
 #include "Animation/AnimTrace.h"
 #include "Animation/AnimationPoseData.h"
 #include "Animation/AnimSyncScope.h"
+#include "Animation/MirrorDataTable.h"
 
 FAnimSingleNodeInstanceProxy::~FAnimSingleNodeInstanceProxy()
 {
@@ -45,7 +46,18 @@ bool FAnimSingleNodeInstanceProxy::Evaluate(FPoseContext& Output)
 	return true;
 }
 
+void FAnimSingleNodeInstanceProxy::SetMirrorDataTable(const UMirrorDataTable* InMirrorDataTable)
+{
+	MirrorDataTable = InMirrorDataTable;
+}
+
+const UMirrorDataTable* FAnimSingleNodeInstanceProxy::GetMirrorDataTable()
+{
+	return MirrorDataTable;
+}
+
 #if WITH_EDITORONLY_DATA
+
 void FAnimSingleNodeInstanceProxy::PropagatePreviewCurve(FPoseContext& Output) 
 {
 	USkeleton* MySkeleton = GetSkeleton();
@@ -529,6 +541,11 @@ void FAnimNode_SingleNode::Evaluate_AnyThread(FPoseContext& Output)
 				}
 			}
 		}
+
+		if (Proxy->MirrorDataTable)
+		{
+			FAnimationRuntime::MirrorPose(Output.Pose, *(Proxy->MirrorDataTable));
+		}
 #endif // WITH_EDITORONLY_DATA
 	}
 	else
diff --git a/Engine/Source/Runtime/Engine/Public/Animation/AnimSingleNodeInstanceProxy.h b/Engine/Source/Runtime/Engine/Public/Animation/AnimSingleNodeInstanceProxy.h
--- a/Engine/Source/Runtime/Engine/Public/Animation/AnimSingleNodeInstanceProxy.h
+++ b/Engine/Source/Runtime/Engine/Public/Animation/AnimSingleNodeInstanceProxy.h
@@ -11,6 +11,7 @@
 #include "AnimSingleNodeInstanceProxy.generated.h"
 
 struct FAnimSingleNodeInstanceProxy;
+class UMirrorDataTable;
 
 /** 
  * Local anim node for extensible processing. 
@@ -58,6 +59,7 @@ public:
 	FAnimSingleNodeInstanceProxy(UAnimInstance* InAnimInstance)
 		: FAnimInstanceProxy(InAnimInstance)
 		, CurrentAsset(nullptr)
+		, MirrorDataTable(nullptr)
 		, BlendSpacePosition(0.0f, 0.0f, 0.0f)
 		, CurrentTime(0.0f)
 		, DeltaTimeRecord()
@@ -174,6 +176,10 @@ public:
 	}
 #endif
 
+	void SetMirrorDataTable(const UMirrorDataTable* InMirrorDataTable);
+
+	const UMirrorDataTable* GetMirrorDataTable();
+
 #if WITH_EDITORONLY_DATA
 	void PropagatePreviewCurve(FPoseContext& Output);
 #endif // WITH_EDITORONLY_DATA
@@ -202,6 +208,9 @@ protected:
 	/** Current Asset being played. Note that this will be nullptr outside of pre/post update **/
 	UAnimationAsset* CurrentAsset;
 
+	/** If set the result will be mirrored using the table */ 
+	const UMirrorDataTable* MirrorDataTable;
+	
 	/** The internal anim node that does our processing */
 	FAnimNode_SingleNode SingleNode;
 
