Fix crash when deleting a blend space graph

When deleting the node, the graph is renamed to make the package its outer, which gets picked up by the details customization for a frame afterwards.

diff --git a/Engine/Source/Editor/Persona/Private/Customization/BlendSpaceDetails.cpp b/Engine/Source/Editor/Persona/Private/Customization/BlendSpaceDetails.cpp
--- a/Engine/Source/Editor/Persona/Private/Customization/BlendSpaceDetails.cpp
+++ b/Engine/Source/Editor/Persona/Private/Customization/BlendSpaceDetails.cpp
@@ -56,7 +56,7 @@ void FBlendSpaceDetails::CustomizeDetails(class IDetailLayoutBuilder& DetailBuil
 		if(UBlendSpaceGraph* BlendSpaceGraph = Cast<UBlendSpaceGraph>(BlendSpaceBase->GetOuter()))
 		{
 			check(BlendSpaceBase == BlendSpaceGraph->BlendSpace);
-			BlendSpaceNode = CastChecked<UAnimGraphNode_BlendSpaceGraphBase>(BlendSpaceGraph->GetOuter());
+			BlendSpaceNode = Cast<UAnimGraphNode_BlendSpaceGraphBase>(BlendSpaceGraph->GetOuter());
 		}
 		const bool b1DBlendSpace = BlendSpaceBase->IsA<UBlendSpace1D>();
 
