Analysis for BlendSpace. Also includes:
Tweaks to the rendering of blendspaces, including separating the graph from the buttons
Context menu changes for the sample details
Changes to AnimationBlueprintLibrary and AnimPose from Jurre to help get the poses during analysis

diff --git a/Engine/Source/Editor/AnimGraph/Private/AnimGraphNode_BlendSpaceGraphBase.cpp b/Engine/Source/Editor/AnimGraph/Private/AnimGraphNode_BlendSpaceGraphBase.cpp
--- a/Engine/Source/Editor/AnimGraph/Private/AnimGraphNode_BlendSpaceGraphBase.cpp
+++ b/Engine/Source/Editor/AnimGraph/Private/AnimGraphNode_BlendSpaceGraphBase.cpp
@@ -392,6 +392,18 @@ void UAnimGraphNode_BlendSpaceGraphBase::PostProcessPinName(const UEdGraphPin* P
 	Super::PostProcessPinName(Pin, DisplayName);
 }
 
+int32 UAnimGraphNode_BlendSpaceGraphBase::GetSampleIndex(const UEdGraph* Graph) const
+{
+	for (int32 Index = 0 ; Index != Graphs.Num() ; ++Index)
+	{
+		if (Graphs[Index] == Graph)
+		{
+			return Index;
+		}
+	}
+	return INDEX_NONE;
+}
+
 UAnimationBlendSpaceSampleGraph* UAnimGraphNode_BlendSpaceGraphBase::AddGraph(FName InSampleName, UAnimSequence* InSequence)
 {
 	UAnimationBlendSpaceSampleGraph* NewGraph = AddGraphInternal(InSampleName, InSequence);
diff --git a/Engine/Source/Editor/AnimGraph/Public/AnimGraphNode_BlendSpaceGraphBase.h b/Engine/Source/Editor/AnimGraph/Public/AnimGraphNode_BlendSpaceGraphBase.h
--- a/Engine/Source/Editor/AnimGraph/Public/AnimGraphNode_BlendSpaceGraphBase.h
+++ b/Engine/Source/Editor/AnimGraph/Public/AnimGraphNode_BlendSpaceGraphBase.h
@@ -32,6 +32,9 @@ public:
 	// Adds a new graph to the internal array
 	UAnimationBlendSpaceSampleGraph* AddGraph(FName InSampleName, UAnimSequence* InSequence);
 
+	/** Returns the sample index associated with the graph, or -1 if not found */
+	int32 GetSampleIndex(const UEdGraph* Graph) const;
+
 	// Removes the graph at the specified index
 	void RemoveGraph(int32 InSampleIndex);
 
@@ -104,6 +107,6 @@ protected:
 	// Helper function for compilation
 	UAnimGraphNode_Base* ExpandGraphAndProcessNodes(UEdGraph* SourceGraph, UAnimGraphNode_Base* SourceRootNode, IAnimBlueprintCompilationContext& InCompilationContext, IAnimBlueprintGeneratedClassCompiledData& OutCompiledData);
 
-	// Helper function for AddGraph/ReplaceGraph - builds the new graph but doesnt add it to Graphs array.
+	// Helper function for AddGraph/ReplaceGraph - builds the new graph but doesn't add it to Graphs array.
 	UAnimationBlendSpaceSampleGraph* AddGraphInternal(FName InSampleName, UAnimSequence* InSequence);
 };
diff --git a/Engine/Source/Editor/AnimationBlueprintEditor/Private/BlendSpaceDocumentTabFactory.cpp b/Engine/Source/Editor/AnimationBlueprintEditor/Private/BlendSpaceDocumentTabFactory.cpp
--- a/Engine/Source/Editor/AnimationBlueprintEditor/Private/BlendSpaceDocumentTabFactory.cpp
+++ b/Engine/Source/Editor/AnimationBlueprintEditor/Private/BlendSpaceDocumentTabFactory.cpp
@@ -71,15 +71,18 @@ TSharedRef<SWidget> FBlendSpaceDocumentTabFactory::CreateTabBody(const FWorkflow
 		}
 	});
 
-	Args.OnBlendSpaceSampleAdded = FOnBlendSpaceSampleAdded::CreateLambda([this, WeakBlendSpaceNode = TWeakObjectPtr<UAnimGraphNode_BlendSpaceGraphBase>(BlendSpaceNode)](UAnimSequence* InSequence, const FVector& InSamplePoint)
+	Args.OnBlendSpaceSampleAdded = FOnBlendSpaceSampleAdded::CreateLambda([this, WeakBlendSpaceNode = TWeakObjectPtr<UAnimGraphNode_BlendSpaceGraphBase>(BlendSpaceNode)](UAnimSequence* InSequence, const FVector& InSamplePoint) -> int32
 	{
+		int32 Index = INDEX_NONE;
 		if(WeakBlendSpaceNode.Get())
 		{
 			UAnimGraphNode_BlendSpaceGraphBase* BlendSpaceNode = WeakBlendSpaceNode.Get();
 			UAnimationBlendSpaceSampleGraph* NewGraph = BlendSpaceNode->AddGraph(TEXT("NewSample"), InSequence);
+			Index = BlendSpaceNode->GetSampleIndex(NewGraph);
 			BlueprintEditorPtr.Pin()->RefreshMyBlueprint();
 			BlueprintEditorPtr.Pin()->RenameNewlyAddedAction(NewGraph->GetFName());
 		}
+		return Index;
 	});
 
 	Args.OnBlendSpaceSampleRemoved = FOnBlendSpaceSampleRemoved::CreateLambda([this, WeakBlendSpaceNode = TWeakObjectPtr<UAnimGraphNode_BlendSpaceGraphBase>(BlendSpaceNode)](int32 InSampleIndex)
diff --git a/Engine/Source/Editor/AnimationBlueprintLibrary/Public/AnimPose.h b/Engine/Source/Editor/AnimationBlueprintLibrary/Public/AnimPose.h
--- a/Engine/Source/Editor/AnimationBlueprintLibrary/Public/AnimPose.h
+++ b/Engine/Source/Editor/AnimationBlueprintLibrary/Public/AnimPose.h
@@ -115,6 +115,7 @@ class ANIMATIONBLUEPRINTLIBRARY_API UAnimPoseExtensions : public UBlueprintFunct
 {
 	GENERATED_BODY()
 	
+public:
 	/**
 	*  Returns whether the Anim Pose contains valid data
 	*
diff --git a/Engine/Source/Editor/AnimationModifiers/AnimationModifiers.Build.cs b/Engine/Source/Editor/AnimationModifiers/AnimationModifiers.Build.cs
--- a/Engine/Source/Editor/AnimationModifiers/AnimationModifiers.Build.cs
+++ b/Engine/Source/Editor/AnimationModifiers/AnimationModifiers.Build.cs
@@ -23,8 +23,6 @@ public class AnimationModifiers : ModuleRules
                 "Kismet",
                 "AssetTools",
                 "ClassViewer",
-                "AnimationEditor",
-                "SkeletonEditor",
                 "AssetRegistry",
                 "AnimationBlueprintLibrary",
             }
diff --git a/Engine/Source/Editor/Persona/Persona.Build.cs b/Engine/Source/Editor/Persona/Persona.Build.cs
--- a/Engine/Source/Editor/Persona/Persona.Build.cs
+++ b/Engine/Source/Editor/Persona/Persona.Build.cs
@@ -34,11 +34,13 @@ public class Persona : ModuleRules
                 "SequenceRecorder",
                 "AnimationBlueprintEditor",
 			}
-        );
+		);
 
         PrivateDependencyModuleNames.AddRange(
             new string[] {
                 "AppFramework",
+                "AnimationModifiers",
+                "AnimationBlueprintLibrary",
                 "Core", 
                 "CoreUObject", 
 				"ApplicationCore",
diff --git a/Engine/Source/Editor/Persona/Private/Customization/BlendSampleDetails.cpp b/Engine/Source/Editor/Persona/Private/Customization/BlendSampleDetails.cpp
--- a/Engine/Source/Editor/Persona/Private/Customization/BlendSampleDetails.cpp
+++ b/Engine/Source/Editor/Persona/Private/Customization/BlendSampleDetails.cpp
@@ -13,6 +13,8 @@
 
 #include "Widgets/Input/SButton.h"
 #include "Widgets/Input/SNumericEntryBox.h"
+#include "Widgets/Input/SVectorInputBox.h"
+#include "Widgets/Layout/SUniformGridPanel.h"
 #include "Widgets/Text/STextBlock.h"
 
 #include "Animation/AnimSequence.h"
@@ -30,9 +32,88 @@
 #include "BlueprintEditorModule.h"
 #include "AnimationGraph.h"
 #include "BlendSpaceGraph.h"
+#include "PersonaBlendSpaceAnalysis.h"
+#include "SAnimationBlendSpaceGridWidget.h"
 
 #define LOCTEXT_NAMESPACE "BlendSampleDetails"
 
+FReply FBlendSampleDetails::HandleAnalyzeAndDuplicateSample()
+{
+	if (BlendSpace->IsAsset())
+	{
+		// Dismiss menus so that operations which might affect samples (and cause reallocation etc) won't invalidate
+		// data being used in the UI.
+		FSlateApplication::Get().DismissAllMenus();
+		bool bAnalyzed[3] = { false, false, false };
+		FVector OrigValue = BlendSpace->GetBlendSample(SampleIndex).SampleValue;
+		FVector NewValue = FBlendSpaceAnalysis::CalculateSampleValue(
+			*BlendSpace, *BlendSpace->GetBlendSample(SampleIndex).Animation, 
+			BlendSpace->GetBlendSample(SampleIndex).RateScale, OrigValue, bAnalyzed);
+		NewValue.Z = OrigValue.Z;
+		if (NewValue != OrigValue)
+		{
+			GridWidget->OnSampleDuplicated.ExecuteIfBound(SampleIndex, NewValue);
+		}
+	}
+	return FReply::Handled();
+}
+
+FReply FBlendSampleDetails::HandleAnalyzeSample()
+{
+	if (BlendSpace->IsAsset())
+	{
+		bool bAnalyzed[3] = { false, false, false };
+		FVector OrigValue = BlendSpace->GetBlendSample(SampleIndex).SampleValue;
+		FVector NewValue = FBlendSpaceAnalysis::CalculateSampleValue(
+			*BlendSpace, *BlendSpace->GetBlendSample(SampleIndex).Animation, 
+			BlendSpace->GetBlendSample(SampleIndex).RateScale, OrigValue, bAnalyzed);
+		NewValue.Z = OrigValue.Z;
+		if (NewValue != OrigValue)
+		{
+			GridWidget->OnSampleMoved.ExecuteIfBound(SampleIndex, NewValue, false);
+		}
+	}
+	return FReply::Handled();
+}
+
+FReply FBlendSampleDetails::HandleAnalyzeSampleX()
+{
+	if (BlendSpace->IsAsset())
+	{
+		bool bAnalyzed[3] = { false, false, false };
+		FVector OrigValue = BlendSpace->GetBlendSample(SampleIndex).SampleValue;
+		FVector NewValue = FBlendSpaceAnalysis::CalculateSampleValue(
+			*BlendSpace, *BlendSpace->GetBlendSample(SampleIndex).Animation, 
+			BlendSpace->GetBlendSample(SampleIndex).RateScale, OrigValue, bAnalyzed);
+		NewValue.Y = OrigValue.Y;
+		NewValue.Z = OrigValue.Z;
+		if (NewValue != OrigValue)
+		{
+			GridWidget->OnSampleMoved.ExecuteIfBound(SampleIndex, NewValue, false);
+		}
+	}
+	return FReply::Handled();
+}
+
+FReply FBlendSampleDetails::HandleAnalyzeSampleY()
+{
+	if (BlendSpace->IsAsset())
+	{
+		bool bAnalyzed[3] = { false, false, false };
+		FVector OrigValue = BlendSpace->GetBlendSample(SampleIndex).SampleValue;
+		FVector NewValue = FBlendSpaceAnalysis::CalculateSampleValue(
+			*BlendSpace, *BlendSpace->GetBlendSample(SampleIndex).Animation, 
+			BlendSpace->GetBlendSample(SampleIndex).RateScale, OrigValue, bAnalyzed);
+		NewValue.X = OrigValue.X;
+		NewValue.Z = OrigValue.Z;
+		if (NewValue != OrigValue)
+		{
+			GridWidget->OnSampleMoved.ExecuteIfBound(SampleIndex, NewValue, false);
+		}
+	}
+	return FReply::Handled();
+}
+
 FBlendSampleDetails::FBlendSampleDetails(const UBlendSpace* InBlendSpace, class SBlendSpaceGridWidget* InGridWidget, int32 InSampleIndex)
 	: BlendSpace(InBlendSpace)
 	, GridWidget(InGridWidget)
@@ -62,7 +143,7 @@ void FBlendSampleDetails::CustomizeDetails(class IDetailLayoutBuilder& DetailBui
 	// Hide all default properties
 	for (TSharedRef<IPropertyHandle> Property : DefaultProperties)
 	{
-		Property->MarkHiddenByCustomization();		
+		Property->MarkHiddenByCustomization();
 	}
 	
 	TArray < TSharedPtr<FStructOnScope>> Structs;
@@ -77,7 +158,7 @@ void FBlendSampleDetails::CustomizeDetails(class IDetailLayoutBuilder& DetailBui
 	TArray<UObject*> Objects;
 	UPackageTools::GetObjectsInPackages(&Packages, Objects);
 
-	UBlendSpace* BlendSpaceBase = nullptr;
+	const UBlendSpace* BlendSpaceBase = nullptr;
 	// Find blendspace in found objects
 	for ( UObject* Object : Objects )
 	{
@@ -96,17 +177,207 @@ void FBlendSampleDetails::CustomizeDetails(class IDetailLayoutBuilder& DetailBui
 		BlendSpaceNode = CastChecked<UAnimGraphNode_BlendSpaceGraphBase>(BlendSpaceGraph->GetOuter());
 	}
 
-	FBlendSampleDetails::GenerateBlendSampleWidget([&CategoryBuilder]() -> FDetailWidgetRow& { return CategoryBuilder.AddCustomRow(FText::FromString(TEXT("SampleValue"))); }, GridWidget->OnSampleMoved, (BlendSpaceBase != nullptr) ? BlendSpaceBase : BlendSpace, GridWidget->GetSelectedSampleIndex(), false );
+	// Sample value
+	FBlendSampleDetails::GenerateBlendSampleWidget(
+		[&CategoryBuilder]() -> FDetailWidgetRow&
+		{
+			return CategoryBuilder.AddCustomRow(FText::FromString(TEXT("SampleValue")));
+		}, GridWidget->OnSampleMoved, (BlendSpaceBase != nullptr) ? BlendSpaceBase : BlendSpace,
+		GridWidget->GetSelectedSampleIndex(), false);
 
-	TSharedPtr<IPropertyHandle> AnimationProperty = DetailBuilder.GetProperty(GET_MEMBER_NAME_CHECKED(FBlendSample, Animation), FBlendSample::StaticStruct());
+	// Locked
+	const bool b1DBlendSpace = BlendSpace->IsA<UBlendSpace1D>();
+	{
+		IDetailPropertyRow& LockPositionXPropertyRow = CategoryBuilder.AddProperty(DetailBuilder.GetProperty(
+            GET_MEMBER_NAME_CHECKED(FBlendSample, bLockX), FBlendSample::StaticStruct()));
+		LockPositionXPropertyRow.DisplayName(FText::Format(
+            LOCTEXT("BlendSpaceSampleLockLabel", "{0} {1}"), FText::FromString(TEXT("Lock ")),
+            FText::FromString(BlendSpace->GetBlendParameter(0).DisplayName)));
+	}
+	if (!b1DBlendSpace)
+	{
+		IDetailPropertyRow& LockPositionYPropertyRow = CategoryBuilder.AddProperty(DetailBuilder.GetProperty(
+			GET_MEMBER_NAME_CHECKED(FBlendSample, bLockY), FBlendSample::StaticStruct()));
+		LockPositionYPropertyRow.DisplayName(FText::Format(
+			LOCTEXT("BlendSpaceSampleLockLabel", "{0} {1}"), FText::FromString(TEXT("Lock ")),
+			FText::FromString(BlendSpace->GetBlendParameter(1).DisplayName)));
+	}
 
+	// Animation and rate
 	if(BlendspaceToUse->IsAsset())
 	{
+		TSharedPtr<IPropertyHandle> AnimationProperty = DetailBuilder.GetProperty(
+			GET_MEMBER_NAME_CHECKED(FBlendSample, Animation), FBlendSample::StaticStruct());
+
 		FDetailWidgetRow& AnimationRow = CategoryBuilder.AddCustomRow(FText::FromString(TEXT("Animation")));
 		FBlendSampleDetails::GenerateAnimationWidget(AnimationRow, BlendspaceToUse, AnimationProperty);
 
-		TSharedPtr<IPropertyHandle> RateScaleProperty = DetailBuilder.GetProperty(GET_MEMBER_NAME_CHECKED(FBlendSample, RateScale), FBlendSample::StaticStruct());
-		CategoryBuilder.AddProperty(RateScaleProperty);	
+		TSharedPtr<IPropertyHandle> RateScaleProperty = DetailBuilder.GetProperty(
+			GET_MEMBER_NAME_CHECKED(FBlendSample, RateScale), FBlendSample::StaticStruct());
+		CategoryBuilder.AddProperty(RateScaleProperty);
+
+		bool bShowAnalysis = false;
+		bool bAnalyzed[3] = { false, false, false };
+		FVector NewValue = FBlendSpaceAnalysis::CalculateSampleValue(
+			*BlendSpace, *BlendSpace->GetBlendSample(SampleIndex).Animation,
+			BlendSpace->GetBlendSample(SampleIndex).RateScale, 
+			BlendSpace->GetBlendSample(SampleIndex).SampleValue, bAnalyzed);
+		FText AnalysisTexts[2];
+		FText ValueTexts[2];
+		FText ToolTipTexts[2];
+
+		for (int32 Index = 0 ; Index != 2 ; ++Index)
+		{
+			if (bAnalyzed[Index])
+			{
+				AnalysisTexts[Index] = FText::Format(FTextFormat(
+					LOCTEXT("AnalysisSampleValue", "Set {0}")), 
+					FText::FromString(BlendSpace->GetBlendParameter(Index).DisplayName));
+				ToolTipTexts[Index] = FText::Format(FTextFormat(
+					LOCTEXT("AnalysisSampleToolTip", "Set {0} to the analysed value {1}")), 
+					FText::FromString(BlendSpace->GetBlendParameter(Index).DisplayName),
+					NewValue[Index]);
+				ValueTexts[Index] = FText::Format(FTextFormat(
+					LOCTEXT("AnalysisValue", "{0}")), NewValue[Index]);
+				bShowAnalysis = true;
+			}
+			else
+			{
+				AnalysisTexts[Index] = LOCTEXT("Unanalyzed", "Not analyzed");
+				ToolTipTexts[Index] = FText::Format(FTextFormat(
+					LOCTEXT("AnalysisSampleToolTipUnanalyzed", "Analysis has not been set for {0}")), 
+					FText::FromString(BlendSpace->GetBlendParameter(Index).DisplayName));
+			}
+		}
+
+		if (bShowAnalysis && NewValue != BlendSpace->GetBlendSample(SampleIndex).SampleValue)
+		{
+			// Button to run analysis
+			CategoryBuilder
+				.AddGroup(FName("AnalysisGroup"), FText::GetEmpty())
+				.HeaderRow()
+				.NameContent()
+				[
+					SNew(SHorizontalBox)
+					+SHorizontalBox::Slot()
+					.AutoWidth()
+					[
+						SNew(SButton)
+						.ButtonStyle(FAppStyle::Get(), "SimpleButton")
+						.ToolTipText(LOCTEXT("AnalyzeText", "Move this sample to the analyzed position"))
+						.OnClicked(this, &FBlendSampleDetails::HandleAnalyzeSample)
+						.ContentPadding(1)
+						[
+							SNew(SHorizontalBox)
+							+SHorizontalBox::Slot()
+							.AutoWidth()
+							[
+								SNew(SImage)
+								.Image(FEditorStyle::GetBrush("Icons.ArrowRight"))
+								.ColorAndOpacity(FSlateColor::UseForeground())
+							]
+							+SHorizontalBox::Slot()
+							.Padding(6.0f, 0.0f)
+							.VAlign(VAlign_Center)
+							[
+								SNew(STextBlock)
+								.Font(DetailBuilder.GetDetailFont())
+								.Text(LOCTEXT("AnalyzeLabel", "Move"))
+							]
+						]
+					]
+				]
+				.ValueContent()
+				.HAlign(HAlign_Fill)
+				.VAlign(VAlign_Center)
+				[
+					SNew(SUniformGridPanel)
+					+SUniformGridPanel::Slot(0,0)
+					.HAlign(HAlign_Center)
+					[   // Analyze X button
+						SNew(SButton)
+						.ButtonStyle(FAppStyle::Get(), "RoundButton")
+						.ToolTipText(ToolTipTexts[0])
+						.IsEnabled(bAnalyzed[0])
+						.OnClicked(this, &FBlendSampleDetails::HandleAnalyzeSampleX)
+						[
+							SNew(STextBlock)
+							.Justification(ETextJustify::Center)
+							.Text(AnalysisTexts[0])
+						]
+					]
+					+SUniformGridPanel::Slot(1,0)
+					.HAlign(HAlign_Center)
+					[   // Analyze Y button
+						SNew(SButton)
+						.ButtonStyle(FAppStyle::Get(), "RoundButton")
+						.ToolTipText(ToolTipTexts[1])
+						.Visibility(b1DBlendSpace ? EVisibility::Hidden : EVisibility::Visible)
+						.IsEnabled(bAnalyzed[1])
+						.OnClicked(this, &FBlendSampleDetails::HandleAnalyzeSampleY)
+						[
+							SNew(STextBlock)
+							.Justification(ETextJustify::Center)
+							.Text(AnalysisTexts[1])
+						]
+					]
+				];
+
+			// Button to duplicate analysis
+			CategoryBuilder
+				.AddGroup(FName("DuplicateGroup"), FText::GetEmpty())
+				.HeaderRow()
+				.NameContent()
+				[
+					SNew(SHorizontalBox)
+					+SHorizontalBox::Slot()
+					.AutoWidth()
+					[
+						SNew(SButton)
+						.ButtonStyle(FAppStyle::Get(), "SimpleButton")
+						.ToolTipText(LOCTEXT("DuplicateText", "Duplicate this sample and place it at the analyzed position"))
+						.OnClicked(this, &FBlendSampleDetails::HandleAnalyzeAndDuplicateSample)
+						.ContentPadding(1)
+						[
+							SNew(SHorizontalBox)
+							+SHorizontalBox::Slot()
+							.AutoWidth()
+							[
+								SNew(SImage)
+								.Image(FEditorStyle::GetBrush("Icons.Duplicate"))
+								.ColorAndOpacity(FSlateColor::UseForeground())
+							]
+							+SHorizontalBox::Slot()
+							.Padding(6.0f, 0.0f)
+							.VAlign(VAlign_Center)
+							[
+								SNew(STextBlock)
+								.Font(DetailBuilder.GetDetailFont())
+								.Text(LOCTEXT("DuplicateLabel", "Duplicate to"))
+							]
+						]
+					]
+				]
+				.ValueContent()
+				.HAlign(HAlign_Fill)
+				.VAlign(VAlign_Center)
+				[
+					SNew(SUniformGridPanel)
+					+SUniformGridPanel::Slot(0,0)
+					.HAlign(HAlign_Center)
+					[
+						SNew(STextBlock)
+						.Text(ValueTexts[0])
+					]
+					+SUniformGridPanel::Slot(1,0)
+					.HAlign(HAlign_Center)
+					[
+						SNew(STextBlock)
+						.Text(ValueTexts[1])
+					]
+				];
+
+		}
 	}
 	else
 	{
diff --git a/Engine/Source/Editor/Persona/Private/Customization/BlendSampleDetails.h b/Engine/Source/Editor/Persona/Private/Customization/BlendSampleDetails.h
--- a/Engine/Source/Editor/Persona/Private/Customization/BlendSampleDetails.h
+++ b/Engine/Source/Editor/Persona/Private/Customization/BlendSampleDetails.h
@@ -39,6 +39,11 @@ public:
 protected:
 	/** Checks whether or not the specified asset should not be shown in the mini content browser when changing the animation */
 	bool ShouldFilterAsset(const FAssetData& AssetData) const;
+
+	FReply HandleAnalyzeAndDuplicateSample();
+	FReply HandleAnalyzeSample();
+	FReply HandleAnalyzeSampleX();
+	FReply HandleAnalyzeSampleY();
 private:
 	/** Pointer to the current parent blend space for the customized blend sample*/
 	const class UBlendSpace* BlendSpace;
diff --git a/Engine/Source/Editor/Persona/Private/Customization/BlendSpaceDetails.cpp b/Engine/Source/Editor/Persona/Private/Customization/BlendSpaceDetails.cpp
--- a/Engine/Source/Editor/Persona/Private/Customization/BlendSpaceDetails.cpp
+++ b/Engine/Source/Editor/Persona/Private/Customization/BlendSpaceDetails.cpp
@@ -14,17 +14,20 @@
 #include "BlendSampleDetails.h"
 
 #include "Widgets/Input/SNumericEntryBox.h"
+#include "Widgets/Input/SButton.h"
+#include "Widgets/Input/STextComboBox.h"
 #include "PropertyCustomizationHelpers.h"
 #include "ScopedTransaction.h"
 #include "BlendSpaceGraph.h"
 #include "AnimGraphNode_BlendSpaceGraph.h"
+#include "PersonaBlendSpaceAnalysis.h"
 
 #define LOCTEXT_NAMESPACE "BlendSpaceDetails"
 
 FBlendSpaceDetails::FBlendSpaceDetails()
 {
 	Builder = nullptr;
-	BlendSpaceBase = nullptr;
+	BlendSpace = nullptr;
 	BlendSpaceNode = nullptr;
 }
 
@@ -32,18 +35,132 @@ FBlendSpaceDetails::~FBlendSpaceDetails()
 {
 }
 
+//======================================================================================================================
+FReply FBlendSpaceDetails::HandleClearSamples()
+{
+	FSlateApplication::Get().DismissAllMenus();
+	BlendSpace->Modify();
+	BlendSpace->SampleData.Empty();
+
+	FPropertyChangedEvent ChangedEvent(nullptr, EPropertyChangeType::ArrayClear);
+	BlendSpace->PostEditChangeProperty(ChangedEvent);
+
+	return FReply::Handled();
+}
+
+//======================================================================================================================
+static bool GetLockAfterAnalysis(const TObjectPtr<UAnalysisProperties>& AnalysisProperties)
+{
+	return AnalysisProperties ? AnalysisProperties->bLockAfterAnalysis : false;
+}
+
+//======================================================================================================================
+FReply FBlendSpaceDetails::HandleAnalyzeSamples()
+{
+	if(BlendSpace->IsAsset())
+	{
+		FSlateApplication::Get().DismissAllMenus();
+		bool bChangedOne = false;
+
+		bool bLockX = GetLockAfterAnalysis(BlendSpace->AnalysisProperties[0]);
+		bool bLockY = GetLockAfterAnalysis(BlendSpace->AnalysisProperties[1]);
+		bool bLockZ = GetLockAfterAnalysis(BlendSpace->AnalysisProperties[2]);
+		
+		const int32 NumSamples = BlendSpace->SampleData.Num();
+		for (int32 SampleIndex = 0 ; SampleIndex != NumSamples ; ++SampleIndex)
+		{
+			bool bAnalyzed[3] = { false, false, false };
+			FVector NewValue = FBlendSpaceAnalysis::CalculateSampleValue(
+				*BlendSpace, *BlendSpace->SampleData[SampleIndex].Animation, 
+				BlendSpace->SampleData[SampleIndex].RateScale, 
+				BlendSpace->SampleData[SampleIndex].SampleValue, bAnalyzed);
+			if (NewValue != BlendSpace->SampleData[SampleIndex].SampleValue)
+			{
+				BlendSpace->EditSampleValue(SampleIndex, NewValue);
+				// Note that the sample might not move if the destination position is in use
+				if (NewValue == BlendSpace->SampleData[SampleIndex].SampleValue)
+				{
+					bChangedOne = true;
+					BlendSpace->LockSample(
+						SampleIndex,
+						bAnalyzed[0] ? GetLockAfterAnalysis(BlendSpace->AnalysisProperties[0]) : false,
+						bAnalyzed[1] ? GetLockAfterAnalysis(BlendSpace->AnalysisProperties[1]) : false,
+						bAnalyzed[2] ? GetLockAfterAnalysis(BlendSpace->AnalysisProperties[2]) : false);
+				}
+			}
+		}
+
+		if (bChangedOne)
+		{
+			BlendSpace->Modify();
+			FPropertyChangedEvent ChangedEvent(nullptr, EPropertyChangeType::ArrayClear);
+			BlendSpace->PostEditChangeProperty(ChangedEvent);
+		}
+	}
+	return FReply::Handled();
+}
+
+//======================================================================================================================
+// Returns a bit mask of EAnalysisProperty
+static EVisibility GetAnalyzeButtonVisibility(
+	TSharedPtr<IPropertyHandle> AnalysisPropertiesHandle,
+	int32                       HideIndex)
+{
+	for (int32 AxisIndex = 0; AxisIndex < 3; ++AxisIndex)
+	{
+		if (AxisIndex < HideIndex)
+		{
+			TSharedPtr<IPropertyHandle> AnalysisPropertyHandle = AnalysisPropertiesHandle->GetChildHandle(AxisIndex);
+			UObject* Object;
+			if (AnalysisPropertyHandle->GetValue(Object))
+			{
+				UAnalysisProperties* AnalysisProperties = Cast<UAnalysisProperties>(Object);
+				if (AnalysisProperties && !AnalysisProperties->Function.IsEmpty() && AnalysisProperties->Function != TEXT("None"))
+				{
+					return EVisibility::Visible;
+				}
+			}
+		}
+	}
+	return EVisibility::Collapsed;
+}
+
+//======================================================================================================================
+// When the existing one gets replaced it will get garbage collected automatically as nothing will reference it. Note 
+// that if "None" was selected then the new value will be null
+void FBlendSpaceDetails::HandleAnalysisFunctionChanged(int32 AxisIndex, TSharedPtr<FString> NewFunctionName)
+{
+	BlendSpace->Modify();
+	UAnalysisProperties* NewAnalysisProperties = FBlendSpaceAnalysis::MakeAnalysisProperties(BlendSpace, *NewFunctionName);
+	// Preserve values where possible
+	if (BlendSpace->AnalysisProperties[AxisIndex])
+	{
+		BlendSpace->AnalysisProperties[AxisIndex]->MakeCache(BlendSpace->CachedAnalysisProperties[AxisIndex]);
+	}
+	if (NewAnalysisProperties)
+	{
+		NewAnalysisProperties->InitializeFromCache(BlendSpace->CachedAnalysisProperties[AxisIndex]);
+	}
+	BlendSpace->AnalysisProperties[AxisIndex] = NewAnalysisProperties;
+	Builder->ForceRefreshDetails();
+}
+
+//======================================================================================================================
 void FBlendSpaceDetails::CustomizeDetails(class IDetailLayoutBuilder& DetailBuilder)
 {
 	TArray< TWeakObjectPtr<UObject> > Objects;
 	DetailBuilder.GetObjectsBeingCustomized(Objects);
 
 	Builder = &DetailBuilder;
-	TWeakObjectPtr<UObject>* WeakPtr = Objects.FindByPredicate([](const TWeakObjectPtr<UObject>& ObjectPtr) { return ObjectPtr->IsA<UBlendSpace>(); });
-	if (WeakPtr)
+
+	FSimpleDelegate RefreshDelegate = FSimpleDelegate::CreateLambda([this]() { Builder->ForceRefreshDetails(); });
+
+	TWeakObjectPtr<UObject>* BlendSpacePtr = Objects.FindByPredicate([](const TWeakObjectPtr<UObject>& ObjectPtr) { return ObjectPtr->IsA<UBlendSpace>(); });
+	if (BlendSpacePtr)
 	{
-		BlendSpaceBase = Cast<UBlendSpace>(WeakPtr->Get());
+		BlendSpace = Cast<UBlendSpace>(BlendSpacePtr->Get());
 
-		if(!BlendSpaceBase->IsAsset())
+		if(!BlendSpace->IsAsset())
 		{
 			// Hide various properties when we are 'internal'
 			DetailBuilder.HideCategory("MetaData");
@@ -51,14 +168,17 @@ void FBlendSpaceDetails::CustomizeDetails(class IDetailLayoutBuilder& DetailBuil
 			DetailBuilder.HideCategory("Thumbnail");
 			DetailBuilder.HideCategory("Animation");
 			DetailBuilder.HideCategory("AdditiveSettings");
+			DetailBuilder.HideCategory("Analysis");
+			DetailBuilder.HideCategory("AnalysisProperties");
+			DetailBuilder.HideCategory("Graph");
 		}
 
-		if(UBlendSpaceGraph* BlendSpaceGraph = Cast<UBlendSpaceGraph>(BlendSpaceBase->GetOuter()))
+		if(UBlendSpaceGraph* BlendSpaceGraph = Cast<UBlendSpaceGraph>(BlendSpace->GetOuter()))
 		{
-			check(BlendSpaceBase == BlendSpaceGraph->BlendSpace);
+			check(BlendSpace == BlendSpaceGraph->BlendSpace);
 			BlendSpaceNode = Cast<UAnimGraphNode_BlendSpaceGraphBase>(BlendSpaceGraph->GetOuter());
 		}
-		const bool b1DBlendSpace = BlendSpaceBase->IsA<UBlendSpace1D>();
+		const bool b1DBlendSpace = BlendSpace->IsA<UBlendSpace1D>();
 
 		if (b1DBlendSpace)
 		{
@@ -66,155 +186,357 @@ void FBlendSpaceDetails::CustomizeDetails(class IDetailLayoutBuilder& DetailBuil
 			DetailBuilder.HideProperty(DetailBuilder.GetProperty(GET_MEMBER_NAME_CHECKED(UBlendSpace, PreferredTriangulationDirection), UBlendSpace::StaticClass()));
 		}
 
-		IDetailCategoryBuilder& CategoryBuilder = DetailBuilder.EditCategory(FName("Axis Settings"));
-		IDetailGroup* Groups[2] =
-		{
-			&CategoryBuilder.AddGroup(FName("Horizontal Axis"), LOCTEXT("HorizontalAxisName", "Horizontal Axis")),
-			b1DBlendSpace ? nullptr : &CategoryBuilder.AddGroup(FName("Vertical Axis"), LOCTEXT("VerticalAxisName", "Vertical Axis"))
-		};
-
-		// Hide the default blend and interpolation parameters
-		TSharedPtr<IPropertyHandle> BlendParameters = DetailBuilder.GetProperty(GET_MEMBER_NAME_CHECKED(UBlendSpace, BlendParameters), UBlendSpace::StaticClass());
-		TSharedPtr<IPropertyHandle> InterpolationParameters = DetailBuilder.GetProperty(GET_MEMBER_NAME_CHECKED(UBlendSpace, InterpolationParam), UBlendSpace::StaticClass());
-		DetailBuilder.HideProperty(BlendParameters);
-		DetailBuilder.HideProperty(InterpolationParameters);
-
-		// Add the properties to the corresponding groups created above (third axis will always be hidden since it isn't used)
-		int32 HideIndex = b1DBlendSpace ? 1 : 2;
-		for (int32 AxisIndex = 0; AxisIndex < 3; ++AxisIndex)
+		// How many samples are there?
+		TSharedPtr<IPropertyHandleArray> BlendSamplesArrayProperty = DetailBuilder.GetProperty(GET_MEMBER_NAME_CHECKED(UBlendSpace, SampleData), UBlendSpace::StaticClass())->AsArray();
+		uint32 NumBlendSampleEntries = 0;
+		BlendSamplesArrayProperty->GetNumElements(NumBlendSampleEntries);
+
+		//==============================================================================================================
+		// Axis Settings section
+		//==============================================================================================================
 		{
-			TSharedPtr<IPropertyHandle> BlendParameter = BlendParameters->GetChildHandle(AxisIndex);
-			TSharedPtr<IPropertyHandle> InterpolationParameter = InterpolationParameters->GetChildHandle(AxisIndex);
+			IDetailCategoryBuilder& DetailCategoryBuilder = DetailBuilder.EditCategory(FName("Axis Settings"));
+			IDetailGroup* Groups[2] =
+			{
+				&DetailCategoryBuilder.AddGroup(FName("Horizontal Axis"), LOCTEXT("HorizontalAxisName", "Horizontal Axis")),
+                b1DBlendSpace ? nullptr : &DetailCategoryBuilder.AddGroup(FName("Vertical Axis"), LOCTEXT("VerticalAxisName", "Vertical Axis"))
+            };
+
+			// Hide the default blend and interpolation parameters
+			TSharedPtr<IPropertyHandle> BlendParameters = DetailBuilder.GetProperty(GET_MEMBER_NAME_CHECKED(UBlendSpace, BlendParameters), UBlendSpace::StaticClass());
+			TSharedPtr<IPropertyHandle> InterpolationParameters = DetailBuilder.GetProperty(GET_MEMBER_NAME_CHECKED(UBlendSpace, InterpolationParam), UBlendSpace::StaticClass());
+			DetailBuilder.HideProperty(BlendParameters);
+			DetailBuilder.HideProperty(InterpolationParameters);
 
-			if (AxisIndex < HideIndex)
+			// Add the properties to the corresponding groups created above (third axis will always be hidden since it isn't used)
+			int32 HideIndex = b1DBlendSpace ? 1 : 2;
+			for (int32 AxisIndex = 0; AxisIndex < 3; ++AxisIndex)
 			{
-				Groups[AxisIndex]->AddPropertyRow(BlendParameter.ToSharedRef());
-				// Don't add InterpolationParameter in the same way as BlendParameter, because it would add the
-				// elements as customizations that we can't subsequently customize. We will add them individually
-				// below.
-
-				TSharedPtr<IPropertyHandle> InterpolationTime = InterpolationParameter->GetChildHandle(
-                    GET_MEMBER_NAME_CHECKED(FInterpolationParameter, InterpolationTime));
-				TSharedPtr<IPropertyHandle> DampingRatio = InterpolationParameter->GetChildHandle(
-                    GET_MEMBER_NAME_CHECKED(FInterpolationParameter, DampingRatio));
-				TSharedPtr<IPropertyHandle> MaxSpeed = InterpolationParameter->GetChildHandle(
-                    GET_MEMBER_NAME_CHECKED(FInterpolationParameter, MaxSpeed));
-				TSharedPtr<IPropertyHandle> InterpolationType = InterpolationParameter->GetChildHandle(
-                    GET_MEMBER_NAME_CHECKED(FInterpolationParameter, InterpolationType));
-
-				// Custom edit condition for MaxSpeed
-				TAttribute<bool> MaxSpeedEditCondition = TAttribute<bool>::Create(
-                    [this, InterpolationTime, InterpolationType]()
-                    {
-                    uint8 IntType;
-                    InterpolationType->GetValue(IntType);
-                    EFilterInterpolationType Type = (EFilterInterpolationType) IntType;
-                    float Time;
-                    InterpolationTime->GetValue(Time);
-                    if (Time > 0.0f && (Type == EFilterInterpolationType::BSIT_SpringDamper ||
-                        Type == EFilterInterpolationType::BSIT_ExponentialDecay))
-                    {
-                        return true;
-                    }
-                    return false;
-                });
-
-				Groups[AxisIndex]->AddPropertyRow(InterpolationTime.ToSharedRef());
-				Groups[AxisIndex]->AddPropertyRow(InterpolationType.ToSharedRef());
-				Groups[AxisIndex]->AddPropertyRow(DampingRatio.ToSharedRef());
-				IDetailPropertyRow& MaxSpeedProperty = Groups[AxisIndex]->AddPropertyRow(MaxSpeed.ToSharedRef());
-				MaxSpeedProperty.EditCondition(MaxSpeedEditCondition, nullptr);
+				TSharedPtr<IPropertyHandle> BlendParameter = BlendParameters->GetChildHandle(AxisIndex);
+				TSharedPtr<IPropertyHandle> InterpolationParameter = InterpolationParameters->GetChildHandle(AxisIndex);
+
+				if (AxisIndex < HideIndex)
+				{
+					Groups[AxisIndex]->AddPropertyRow(BlendParameter.ToSharedRef());
+					// Don't add InterpolationParameter in the same way as BlendParameter, because it would add the
+					// elements as customizations that we can't subsequently customize. We will add them individually
+					// below.
+
+					TSharedPtr<IPropertyHandle> InterpolationTime = InterpolationParameter->GetChildHandle(
+						GET_MEMBER_NAME_CHECKED(FInterpolationParameter, InterpolationTime));
+					TSharedPtr<IPropertyHandle> DampingRatio = InterpolationParameter->GetChildHandle(
+						GET_MEMBER_NAME_CHECKED(FInterpolationParameter, DampingRatio));
+					TSharedPtr<IPropertyHandle> MaxSpeed = InterpolationParameter->GetChildHandle(
+						GET_MEMBER_NAME_CHECKED(FInterpolationParameter, MaxSpeed));
+					TSharedPtr<IPropertyHandle> InterpolationType = InterpolationParameter->GetChildHandle(
+						GET_MEMBER_NAME_CHECKED(FInterpolationParameter, InterpolationType));
+
+					// Custom edit condition for MaxSpeed
+					TAttribute<bool> MaxSpeedEditCondition = TAttribute<bool>::Create(
+						[this, InterpolationTime, InterpolationType]()
+						{
+							uint8 IntType;
+							InterpolationType->GetValue(IntType);
+							EFilterInterpolationType Type = (EFilterInterpolationType)IntType;
+							float Time;
+							InterpolationTime->GetValue(Time);
+							if (Time > 0.0f && (Type == EFilterInterpolationType::BSIT_SpringDamper ||
+												Type == EFilterInterpolationType::BSIT_ExponentialDecay))
+							{
+								return true;
+							}
+							return false;
+						});
+
+					Groups[AxisIndex]->AddPropertyRow(InterpolationTime.ToSharedRef());
+					Groups[AxisIndex]->AddPropertyRow(InterpolationType.ToSharedRef());
+					Groups[AxisIndex]->AddPropertyRow(DampingRatio.ToSharedRef());
+					IDetailPropertyRow& MaxSpeedProperty = Groups[AxisIndex]->AddPropertyRow(MaxSpeed.ToSharedRef());
+					MaxSpeedProperty.EditCondition(MaxSpeedEditCondition, nullptr);
+				}
+				else
+				{
+					DetailBuilder.HideProperty(BlendParameter);
+					DetailBuilder.HideProperty(InterpolationParameter);
+				}
 			}
-			else
+		}
+		
+		//==============================================================================================================
+		// Analysis section
+		//==============================================================================================================
+		if (BlendSpace->IsAsset())
+		{
+			IDetailCategoryBuilder& DetailCategoryBuilder = DetailBuilder.EditCategory(FName("Analysis"));
+			TSharedPtr<IPropertyHandle> AnalysisPropertiesHandle = DetailBuilder.GetProperty(GET_MEMBER_NAME_CHECKED(UBlendSpace, AnalysisProperties), UBlendSpace::StaticClass());
+
+			int32 HideIndex = b1DBlendSpace ? 1 : 2;
+
+			// Re-analyse button
+			if (NumBlendSampleEntries)
 			{
-				DetailBuilder.HideProperty(BlendParameter);
-				DetailBuilder.HideProperty(InterpolationParameter);
+				DetailCategoryBuilder
+					.AddGroup(FName("GroupName"), FText::GetEmpty())
+					.HeaderRow()
+					.Visibility(TAttribute<EVisibility>::Create([AnalysisPropertiesHandle, HideIndex]() {
+						return GetAnalyzeButtonVisibility(AnalysisPropertiesHandle, HideIndex);}))
+					.NameContent()
+					[
+						SNew(STextBlock)
+						.Font(DetailBuilder.GetDetailFont())
+						.Text(FText::Format(LOCTEXT("SamplesLabel", "Analyze {0} Samples"), NumBlendSampleEntries))
+					]
+					.ValueContent()
+					[
+						SNew(SButton)
+						.ButtonStyle(FAppStyle::Get(), "SimpleButton")
+						.HAlign(EHorizontalAlignment::HAlign_Left)
+						.ToolTipText(LOCTEXT("Analyze", "Analyze all samples"))
+						.OnClicked(this, &FBlendSpaceDetails::HandleAnalyzeSamples)
+						.ContentPadding(1)
+						[
+							SNew(SImage)
+							.Image(FEditorStyle::GetBrush("Icons.Refresh"))
+							.ColorAndOpacity(FSlateColor::UseForeground())
+						]
+					];
 			}
-		}
 
-		IDetailCategoryBuilder& SampleCategoryBuilder = DetailBuilder.EditCategory(FName("BlendSamples"));
-		TArray<TSharedRef<IPropertyHandle>> DefaultProperties;
-		SampleCategoryBuilder.GetDefaultProperties(DefaultProperties);
-		for (TSharedRef<IPropertyHandle> DefaultProperty : DefaultProperties)
-		{
-			DefaultProperty->MarkHiddenByCustomization();
-		}
+			FText AxisTexts[2] = 
+			{ 
+				LOCTEXT("HorizontalAxisName", "Horizontal Axis Function"), 
+				LOCTEXT("VerticalAxisName", "Vertical Axis Function") 
+			};
 
-		FSimpleDelegate RefreshDelegate = FSimpleDelegate::CreateLambda([this]() { Builder->ForceRefreshDetails(); });
+			// Hide the default parameters
+			DetailBuilder.HideProperty(AnalysisPropertiesHandle);
 
-		// Retrieve blend samples array
-		TSharedPtr<IPropertyHandleArray> BlendSamplesArrayProperty = DetailBuilder.GetProperty(GET_MEMBER_NAME_CHECKED(UBlendSpace, SampleData), UBlendSpace::StaticClass())->AsArray();
-		BlendSamplesArrayProperty->SetOnNumElementsChanged(RefreshDelegate);
-		
-		uint32 NumBlendSampleEntries = 0;
-		BlendSamplesArrayProperty->GetNumElements(NumBlendSampleEntries);
-		for (uint32 SampleIndex = 0; SampleIndex < NumBlendSampleEntries; ++SampleIndex)
-		{
-			TSharedPtr<IPropertyHandle> BlendSampleProperty = BlendSamplesArrayProperty->GetElement(SampleIndex);
-			BlendSampleProperty->SetOnChildPropertyValueChanged(RefreshDelegate);
-			TSharedPtr<IPropertyHandle> AnimationProperty = BlendSampleProperty->GetChildHandle(GET_MEMBER_NAME_CHECKED(FBlendSample, Animation));
-			TSharedPtr<IPropertyHandle> SampleValueProperty = BlendSampleProperty->GetChildHandle(GET_MEMBER_NAME_CHECKED(FBlendSample, SampleValue));
-			TSharedPtr<IPropertyHandle> RateScaleProperty = BlendSampleProperty->GetChildHandle(GET_MEMBER_NAME_CHECKED(FBlendSample, RateScale));
-
-			IDetailGroup& Group = SampleCategoryBuilder.AddGroup(FName("GroupName"), FText::GetEmpty());
-			Group.HeaderRow()
-			.NameContent()
-			[
-				SNew(SHorizontalBox)
-				+ SHorizontalBox::Slot()
-				.Padding(FMargin(0,2,2,2))
-				.FillWidth(1.0f)
-				.HAlign(HAlign_Right)
-				[
-					SNew(STextBlock)
-					.Font(DetailBuilder.GetDetailFont())
-					.Text_Lambda([this, AnimationProperty, SampleIndex]() -> FText
+			// Add the properties to the corresponding groups created above (third axis will always be hidden since it isn't used)
+			for (int32 AxisIndex = 0; AxisIndex < 3; ++AxisIndex)
+			{
+				TSharedPtr<IPropertyHandle> AnalysisProperty = AnalysisPropertiesHandle->GetChildHandle(AxisIndex);
+				
+				if (AxisIndex < HideIndex)
+				{
+					UObject* AnalysisPropertiesObject;
+					FPropertyAccess::Result Result = AnalysisProperty->GetValue(AnalysisPropertiesObject);
+					UAnalysisProperties* AnalysisProperties = Cast<UAnalysisProperties>(AnalysisPropertiesObject);
+
+					// Note that we need to make sure the function drop-down is shown, even if there is no current
+					// analysis property.
+
+					// Prepare the drop-down. Note that it uses pointers to strings, and comparison is between
+					// pointers, not string contents!
+					AnalysisFunctionNames[AxisIndex].Empty();
+					const TArray<FString> AnalysisFunctions = FBlendSpaceAnalysis::GetAnalysisFunctions();
+					TSharedPtr<FString> CurrentlySelectedFunction;
+					for (const FString& AnalysisFunction : AnalysisFunctions)
 					{
-						FAssetData AssetData;
-						AnimationProperty->GetValue(AssetData);
-						if(AssetData.IsValid())
+						AnalysisFunctionNames[AxisIndex].Add(MakeShared<FString>(AnalysisFunction));
+						if (AnalysisProperties && AnalysisProperties->Function == *AnalysisFunctionNames[AxisIndex].Last())
 						{
-							return FText::Format(LOCTEXT("BlendSpaceAnimationNameLabel", "{0} ({1})"), FText::FromString(AssetData.GetAsset()->GetName()), FText::FromString(FString::FromInt(SampleIndex)));
+							CurrentlySelectedFunction = AnalysisFunctionNames[AxisIndex].Last();
 						}
-						else if(BlendSpaceNode.Get() && BlendSpaceNode->GetGraphs().IsValidIndex(SampleIndex))
+						// In the event that we don't find anything (possible if the list of available functions
+						// changes), it's nice if we start off by at least selecting "None".
+						if (!CurrentlySelectedFunction && *AnalysisFunctionNames[AxisIndex].Last() == "None")
 						{
-							return FText::Format(LOCTEXT("BlendSpaceAnimationNameLabel", "{0} ({1})"), FText::FromName(BlendSpaceNode->GetGraphs()[SampleIndex]->GetFName()), FText::FromString(FString::FromInt(SampleIndex)));
+							CurrentlySelectedFunction = AnalysisFunctionNames[AxisIndex].Last();
 						}
-						return LOCTEXT("NoAnimation", "No Animation");
-					})
-				]
-			];
+					}
 
-			FBlendSampleDetails::GenerateBlendSampleWidget([&Group]() -> FDetailWidgetRow& { return Group.AddWidgetRow(); }, FOnSampleMoved::CreateLambda([this](const uint32 Index, const FVector& SampleValue, bool bIsInteractive) 
-			{
-				if (BlendSpaceBase->IsValidBlendSampleIndex(Index) && BlendSpaceBase->GetBlendSample(Index).SampleValue != SampleValue && !BlendSpaceBase->IsTooCloseToExistingSamplePoint(SampleValue, Index))
-				{
-					BlendSpaceBase->Modify();
+					// Add the analysis properties. We want the function drop-down even if they don't exist
+					IDetailPropertyRow* AnalysisPropertiesRow = nullptr;
+					if (AnalysisProperties && 
+						!AnalysisProperties->Function.IsEmpty() && 
+						AnalysisProperties->Function != TEXT("None") && 
+						AnalysisFunctions.Contains(AnalysisProperties->Function))
+					{
+						AnalysisPropertiesRow = DetailCategoryBuilder.AddExternalObjects( {AnalysisProperties} );
+					}
 
-					bool bMoveSuccesful = BlendSpaceBase->EditSampleValue(Index, SampleValue);
-					if (bMoveSuccesful)
+					FDetailWidgetRow *FunctionWidgetRow = nullptr;
+					// Insert the function drop-down
+					if (AnalysisPropertiesRow)
+					{
+						// Note that there is an extra/unwanted level of indirection here in the UI that is generated,
+						// with the block being underneath an unnecessary "Analysis Properties" section.
+						FunctionWidgetRow = &AnalysisPropertiesRow->CustomWidget();
+					}
+					else
 					{
-						BlendSpaceBase->ValidateSampleData();
-						FPropertyChangedEvent ChangedEvent(nullptr, bIsInteractive ? EPropertyChangeType::Interactive : EPropertyChangeType::ValueSet);
-						BlendSpaceBase->PostEditChangeProperty(ChangedEvent);
+						FunctionWidgetRow = &DetailCategoryBuilder.AddCustomRow(LOCTEXT("AnalysisProperties", "Analysis Properties"));
 					}
+
+					FunctionWidgetRow->NameContent()
+					[
+						SNew(STextBlock)
+						.Text(AxisTexts[AxisIndex])
+						.Font(FEditorStyle::GetFontStyle(TEXT("PropertyWindow.NormalFont")))
+					]
+					.ValueContent()
+					[
+						SNew(STextComboBox)
+						.Font(FEditorStyle::GetFontStyle("PropertyWindow.NormalFont"))
+						.OptionsSource(&AnalysisFunctionNames[AxisIndex])
+						.InitiallySelectedItem(CurrentlySelectedFunction)
+						.OnSelectionChanged_Lambda(
+							[this, AxisIndex, CurrentlySelectedFunction]
+							(TSharedPtr<FString> NewFunctionName, ESelectInfo::Type SelectInfo) 
+							{
+								if (NewFunctionName && NewFunctionName != CurrentlySelectedFunction)
+								{
+									this->HandleAnalysisFunctionChanged(AxisIndex, NewFunctionName);
+								}
+							})
+					];
+				}
+				else
+				{
+					DetailBuilder.HideProperty(AnalysisProperty);
 				}
-			}), BlendSpaceBase, SampleIndex, false);
-			
-			if(BlendSpaceBase->IsAsset())
+			}
+		}
+		
+		//==============================================================================================================
+		// Blend Samples section
+		//==============================================================================================================
+		{
+			IDetailCategoryBuilder& DetailCategoryBuilder = DetailBuilder.EditCategory(FName("BlendSamples"));
+			TArray<TSharedRef<IPropertyHandle>> DefaultProperties;
+			DetailCategoryBuilder.GetDefaultProperties(DefaultProperties);
+			for (TSharedRef<IPropertyHandle> DefaultProperty : DefaultProperties)
 			{
-				FDetailWidgetRow& AnimationRow = Group.AddWidgetRow();
-				FBlendSampleDetails::GenerateAnimationWidget(AnimationRow, BlendSpaceBase, AnimationProperty);
-				Group.AddPropertyRow(RateScaleProperty.ToSharedRef());
+				 DefaultProperty->MarkHiddenByCustomization();
 			}
-			else if(BlendSpaceNode.Get())
+
+			BlendSamplesArrayProperty->SetOnNumElementsChanged(RefreshDelegate);
+
+			// Add a "Remove all" button if there are some samples
+			if (NumBlendSampleEntries)
 			{
-				FDetailWidgetRow& GraphRow = Group.AddWidgetRow();
-				FBlendSampleDetails::GenerateSampleGraphWidget(GraphRow, BlendSpaceNode.Get(), SampleIndex);
+				DetailCategoryBuilder
+					.AddGroup(FName("GroupName"), FText::GetEmpty())
+					.HeaderRow()
+					.NameContent()
+					[
+						SNew(STextBlock)
+						.Font(DetailBuilder.GetDetailFont())
+						.Text(FText::Format(LOCTEXT("SamplesLabel", "{0} Samples"), NumBlendSampleEntries))
+					]
+					.ValueContent()
+					[
+						SNew(SButton)
+						.ButtonStyle(FAppStyle::Get(), "SimpleButton")
+						.HAlign(EHorizontalAlignment::HAlign_Left)
+						.ToolTipText(LOCTEXT("Empty", "Remove all samples"))
+						.OnClicked(this, &FBlendSpaceDetails::HandleClearSamples)
+						.ContentPadding(1)
+						[
+							SNew(SImage)
+							.Image(FEditorStyle::GetBrush("Icons.Delete"))
+							.ColorAndOpacity(FSlateColor::UseForeground())
+						]
+					];
+			}
+
+			for (uint32 SampleIndex = 0; SampleIndex < NumBlendSampleEntries; ++SampleIndex)
+			{
+				TSharedPtr<IPropertyHandle> BlendSampleProperty = BlendSamplesArrayProperty->GetElement(SampleIndex);
+				BlendSampleProperty->SetOnChildPropertyValueChanged(RefreshDelegate);
+				TSharedPtr<IPropertyHandle> AnimationProperty = BlendSampleProperty->GetChildHandle(GET_MEMBER_NAME_CHECKED(FBlendSample, Animation));
+				TSharedPtr<IPropertyHandle> SampleValueProperty = BlendSampleProperty->GetChildHandle(GET_MEMBER_NAME_CHECKED(FBlendSample, SampleValue));
+				TSharedPtr<IPropertyHandle> RateScaleProperty = BlendSampleProperty->GetChildHandle(GET_MEMBER_NAME_CHECKED(FBlendSample, RateScale));
+
+				IDetailGroup& Group = DetailCategoryBuilder.AddGroup(FName("GroupName"), FText::GetEmpty());
+				Group.HeaderRow()
+				.NameContent()
+				[
+					SNew(SHorizontalBox)
+					+ SHorizontalBox::Slot()
+					.Padding(FMargin(0,2,2,2))
+					.FillWidth(1.0f)
+					.HAlign(HAlign_Right)
+					[
+						SNew(STextBlock)
+						.Font(DetailBuilder.GetDetailFont())
+						.Text_Lambda([this, AnimationProperty, SampleIndex]() -> FText
+						{
+							FAssetData AssetData;
+							AnimationProperty->GetValue(AssetData);
+							if(AssetData.IsValid())
+							{
+								return FText::Format(
+									LOCTEXT("BlendSpaceAnimationNameLabel", "{0} ({1})"),
+									FText::FromString(AssetData.GetAsset()->GetName()),
+									FText::FromString(FString::FromInt(SampleIndex)));
+							}
+							else if(BlendSpaceNode.Get() && BlendSpaceNode->GetGraphs().IsValidIndex(SampleIndex))
+							{
+								return FText::Format(
+									LOCTEXT("BlendSpaceAnimationNameLabel", "{0} ({1})"),
+									FText::FromName(BlendSpaceNode->GetGraphs()[SampleIndex]->GetFName()),
+									FText::FromString(FString::FromInt(SampleIndex)));
+							}
+							return LOCTEXT("NoAnimation", "No Animation");
+						})
+					]
+				];
+
+				FBlendSampleDetails::GenerateBlendSampleWidget(
+					[&Group]() -> FDetailWidgetRow& { return Group.AddWidgetRow(); }, FOnSampleMoved::CreateLambda(
+						[this](const uint32 Index, const FVector& SampleValue, bool bIsInteractive) 
+				{
+					if (BlendSpace->IsValidBlendSampleIndex(Index) && BlendSpace->GetBlendSample(Index).SampleValue !=
+						SampleValue && !BlendSpace->IsTooCloseToExistingSamplePoint(SampleValue, Index))
+					{
+						BlendSpace->Modify();
+						bool bMoveSuccesful = BlendSpace->EditSampleValue(Index, SampleValue);
+						if (bMoveSuccesful)
+						{
+							BlendSpace->ValidateSampleData();
+							FPropertyChangedEvent ChangedEvent(nullptr, bIsInteractive ? EPropertyChangeType::Interactive : EPropertyChangeType::ValueSet);
+							BlendSpace->PostEditChangeProperty(ChangedEvent);
+						}
+					}
+				}), BlendSpace, SampleIndex, false);
+
+
+				// Locked
+				{
+					TSharedPtr<IPropertyHandle> LockPositionProperty = BlendSampleProperty->GetChildHandle(
+						GET_MEMBER_NAME_CHECKED(FBlendSample, bLockX));
+					IDetailPropertyRow& LockPositionPropertyRow = Group.AddPropertyRow(
+						LockPositionProperty.ToSharedRef());
+					LockPositionPropertyRow.DisplayName(FText::Format(
+						LOCTEXT("BlendSpaceSampleLockLabel", "{0} {1}"), FText::FromString(TEXT("Lock ")),
+						FText::FromString(BlendSpace->GetBlendParameter(0).DisplayName)));
+				}
+				if (!b1DBlendSpace)
+				{
+					TSharedPtr<IPropertyHandle> LockPositionProperty = BlendSampleProperty->GetChildHandle(
+						GET_MEMBER_NAME_CHECKED(FBlendSample, bLockY));
+					IDetailPropertyRow& LockPositionPropertyRow = Group.AddPropertyRow(
+						LockPositionProperty.ToSharedRef());
+					LockPositionPropertyRow.DisplayName(FText::Format(
+						LOCTEXT("BlendSpaceSampleLockLabel", "{0} {1}"), FText::FromString(TEXT("Lock ")),
+						FText::FromString(BlendSpace->GetBlendParameter(1).DisplayName)));
+				}
+				
+				if(BlendSpace->IsAsset())
+				{
+					FDetailWidgetRow& AnimationRow = Group.AddWidgetRow();
+					FBlendSampleDetails::GenerateAnimationWidget(AnimationRow, BlendSpace, AnimationProperty);
+					Group.AddPropertyRow(RateScaleProperty.ToSharedRef());
+				}
+				else if(BlendSpaceNode.Get())
+				{
+					FDetailWidgetRow& GraphRow = Group.AddWidgetRow();
+					FBlendSampleDetails::GenerateSampleGraphWidget(GraphRow, BlendSpaceNode.Get(), SampleIndex);
+				}
 			}
 		}
-	}
-	
+	}	
 }
 
 #undef LOCTEXT_NAMESPACE
diff --git a/Engine/Source/Editor/Persona/Private/Customization/BlendSpaceDetails.h b/Engine/Source/Editor/Persona/Private/Customization/BlendSpaceDetails.h
--- a/Engine/Source/Editor/Persona/Private/Customization/BlendSpaceDetails.h
+++ b/Engine/Source/Editor/Persona/Private/Customization/BlendSpaceDetails.h
@@ -4,6 +4,7 @@
 
 #include "CoreMinimal.h"
 #include "IDetailCustomization.h"
+#include "Input/Reply.h"
 #include "UObject/WeakObjectPtr.h"
 #include "UObject/WeakObjectPtrTemplates.h"
 
@@ -25,7 +26,12 @@ public:
 	// IDetailCustomization interface
 	virtual void CustomizeDetails(class IDetailLayoutBuilder& DetailBuilder) override;
 private:
+	FReply HandleClearSamples();
+	FReply HandleAnalyzeSamples();
+	void HandleAnalysisFunctionChanged(int32 AxisIndex, TSharedPtr<FString> NewItem);
+
 	IDetailLayoutBuilder* Builder;
-	UBlendSpace* BlendSpaceBase;
+	UBlendSpace* BlendSpace;
 	TWeakObjectPtr<UAnimGraphNode_BlendSpaceGraphBase> BlendSpaceNode;
+	TArray<TSharedPtr<FString>> AnalysisFunctionNames[3];
 };
diff --git a/Engine/Source/Editor/Persona/Private/PersonaBlendSpaceAnalysis.cpp b/Engine/Source/Editor/Persona/Private/PersonaBlendSpaceAnalysis.cpp
new file mode 100644
--- /dev/null
+++ b/Engine/Source/Editor/Persona/Private/PersonaBlendSpaceAnalysis.cpp
@@ -0,0 +1,1445 @@
+// Copyright Epic Games, Inc. All Rights Reserved.
+#include "PersonaBlendSpaceAnalysis.h"
+
+#include "AnimPose.h"
+#include "Engine/SkeletalMeshSocket.h"
+#include "Animation/BlendSpace.h"
+#include "Animation/AnimSequence.h"
+#include "Animation/AnimData/AnimDataModel.h"
+#include "Animation/BoneSocketReference.h"
+
+#define LOCTEXT_NAMESPACE "BlendSpaceAnalysis"
+#define ANALYSIS_VERBOSE_LOG
+
+
+//======================================================================================================================
+static FVector GetAxisFromTM(const FTransform& TM, EAnalysisLinearAxis Axis)
+{
+	switch (Axis)
+	{
+	case EAnalysisLinearAxis::X: return TM.TransformVectorNoScale(FVector(1, 0, 0));
+	case EAnalysisLinearAxis::Y: return TM.TransformVectorNoScale(FVector(0, 1, 0));
+	case EAnalysisLinearAxis::Z: return TM.TransformVectorNoScale(FVector(0, 0, 1));
+	}
+	return FVector(0, 0, 0);
+}
+
+//======================================================================================================================
+// Retrieves the bone index and transform offset given the BoneSocketTarget. Returns true if found
+static bool GetBoneInfo(
+	const UAnimSequence& Animation, const FBoneSocketTarget& BoneSocket, FTransform& BoneOffset, FName& BoneName)
+{
+	const UAnimDataModel* AnimDataModel = Animation.GetDataModel();
+
+	if (BoneSocket.bUseSocket)
+	{
+		const USkeleton* Skeleton = Animation.GetSkeleton();
+		const USkeletalMeshSocket* Socket = Skeleton->FindSocket(BoneSocket.SocketReference.SocketName);
+		if (Socket)
+		{
+			BoneOffset = Socket->GetSocketLocalTransform();
+			BoneName = Socket->BoneName;
+			return !BoneName.IsNone();
+		}
+	}
+	else
+	{
+		BoneName = BoneSocket.BoneReference.BoneName;
+		return !BoneName.IsNone();
+	}
+	return false;
+}
+
+//======================================================================================================================
+template<typename T>
+void CalculateFrameTM(
+	bool& bNeedToUpdateFrameTM, FTransform& FrameTM, 
+	const int32 SampleKey, const T& AnalysisProperties, const UAnimSequence& Animation)
+{
+	if (bNeedToUpdateFrameTM)
+	{
+		FrameTM.SetIdentity();
+		if (AnalysisProperties->Space != EAnalysisSpace::World)
+		{
+			FTransform SpaceBoneOffset;
+			FName SpaceBoneName;
+			if (GetBoneInfo(Animation, AnalysisProperties->SpaceBoneSocket, SpaceBoneOffset, SpaceBoneName))
+			{
+				FAnimPose AnimPose;
+				UAnimPoseExtensions::GetAnimPoseAtFrame(&Animation, SampleKey, FAnimPoseEvaluationOptions(), AnimPose);
+				FTransform SpaceBoneTM = UAnimPoseExtensions::GetBonePose(AnimPose, SpaceBoneName, EAnimPoseSpaces::World);
+				FrameTM = SpaceBoneOffset * SpaceBoneTM;
+			}
+		}
+
+		bNeedToUpdateFrameTM = (
+			AnalysisProperties->Space == EAnalysisSpace::Changing || 
+			AnalysisProperties->Space == EAnalysisSpace::Moving);
+	}
+}
+
+//======================================================================================================================
+template<typename T>
+void GetFrameDirs(
+	FVector& FrameFacingDir, FVector& FrameUpDir, FVector& FrameRightDir, 
+	const FTransform& FrameTM, const T& AnalysisProperties)
+{
+	FrameFacingDir = GetAxisFromTM(FrameTM, AnalysisProperties->CharacterFacingAxis);
+	FrameUpDir = GetAxisFromTM(FrameTM, AnalysisProperties->CharacterUpAxis);
+	FrameRightDir = FVector::CrossProduct(FrameUpDir, FrameFacingDir);
+}
+
+//======================================================================================================================
+template <typename T>
+static bool CalculatePosition(
+	FVector&                         Result,
+	const UBlendSpace&               BlendSpace,
+	const T*                         AnalysisProperties,
+	const UAnimSequence&             Animation,
+	const float                      RateScale)
+{
+	const int32 NumSampledKeys = Animation.GetNumberOfSampledKeys() - 1;
+	if (!AnalysisProperties || NumSampledKeys <= 0)
+	{
+		return false;
+	}
+
+	FTransform BoneOffset;
+	FName BoneName;
+	if (!GetBoneInfo(Animation, AnalysisProperties->BoneSocket, BoneOffset, BoneName))
+	{
+		return false;
+	}
+
+	int32 FirstKey = FMath::Clamp(
+		(int32) (NumSampledKeys * AnalysisProperties->StartTimeFraction), 0, NumSampledKeys);
+	int32 LastKey = FMath::Clamp(
+		(int32) (NumSampledKeys * AnalysisProperties->EndTimeFraction), FirstKey, NumSampledKeys);
+
+	FTransform FrameTM;
+	bool bNeedToUpdateFrameTM = true;
+	CalculateFrameTM(bNeedToUpdateFrameTM, FrameTM, FirstKey, AnalysisProperties, Animation);
+
+	FAnimPose AnimPose;
+	Result.Set(0, 0, 0);
+	for (int32 Key = FirstKey; Key != LastKey + 1; ++Key)
+	{
+		CalculateFrameTM(bNeedToUpdateFrameTM, FrameTM, Key, AnalysisProperties, Animation);
+
+		UAnimPoseExtensions::GetAnimPoseAtFrame(&Animation, Key, FAnimPoseEvaluationOptions(), AnimPose);
+		FTransform BoneTM = UAnimPoseExtensions::GetBonePose(AnimPose, BoneName, EAnimPoseSpaces::World);
+		FTransform TM = BoneOffset * BoneTM;
+		FVector RelativePos = FrameTM.InverseTransformPosition(TM.GetTranslation());
+		Result += RelativePos;
+	}
+	Result /= (1 + LastKey - FirstKey);
+	return true;
+}
+
+//======================================================================================================================
+template <typename T>
+static bool CalculateDeltaPosition(
+	FVector&                         Result,
+	const UBlendSpace&               BlendSpace,
+	const T*                         AnalysisProperties,
+	const UAnimSequence&             Animation,
+	const float                      RateScale)
+{
+	const int32 NumSampledKeys = Animation.GetNumberOfSampledKeys() - 1;
+	if (!AnalysisProperties || NumSampledKeys <= 0)
+	{
+		return false;
+	}
+
+	FTransform BoneOffset;
+	FName BoneName;
+	if (!GetBoneInfo(Animation, AnalysisProperties->BoneSocket, BoneOffset, BoneName))
+	{
+		return false;
+	}
+
+	int32 FirstKey = FMath::Clamp(
+		(int32) (NumSampledKeys * AnalysisProperties->StartTimeFraction), 0, NumSampledKeys);
+	int32 LastKey = FMath::Clamp(
+		(int32) (NumSampledKeys * AnalysisProperties->EndTimeFraction), FirstKey, NumSampledKeys);
+
+	FTransform FrameTM;
+	bool bNeedToUpdateFrameTM = true;
+	CalculateFrameTM(bNeedToUpdateFrameTM, FrameTM, FirstKey, AnalysisProperties, Animation);
+
+	FAnimPose AnimPose;
+
+	UAnimPoseExtensions::GetAnimPoseAtFrame(&Animation, FirstKey, FAnimPoseEvaluationOptions(), AnimPose);
+	FTransform BoneTM1 = UAnimPoseExtensions::GetBonePose(AnimPose, BoneName, EAnimPoseSpaces::World);
+	FTransform TM1 = BoneOffset * BoneTM1;
+	FVector RelativePos1 = FrameTM.InverseTransformPosition(TM1.GetTranslation());
+
+	CalculateFrameTM(bNeedToUpdateFrameTM, FrameTM, FirstKey, AnalysisProperties, Animation);
+
+	UAnimPoseExtensions::GetAnimPoseAtFrame(&Animation, LastKey, FAnimPoseEvaluationOptions(), AnimPose);
+	FTransform BoneTM2 = UAnimPoseExtensions::GetBonePose(AnimPose, BoneName, EAnimPoseSpaces::World);
+	FTransform TM2 = BoneOffset * BoneTM2;
+	FVector RelativePos2 = FrameTM.InverseTransformPosition(TM2.GetTranslation());
+
+	Result = RelativePos2 - RelativePos1;
+	return true;
+}
+
+//======================================================================================================================
+template <typename T>
+static bool CalculateVelocity(
+	FVector&             Result,
+	const UBlendSpace&   BlendSpace,
+	const T*             AnalysisProperties,
+	const UAnimSequence& Animation,
+	const float          RateScale)
+{
+	const int32 NumSampledKeys = Animation.GetNumberOfSampledKeys() - 1;
+	if (!AnalysisProperties || NumSampledKeys <= 0)
+	{
+		return false;
+	}
+
+	if (NumSampledKeys == 1)
+	{
+		Result.Set(0, 0, 0);
+		return true;
+	}
+
+	FTransform BoneOffset;
+	FName BoneName;
+	if (!GetBoneInfo(Animation, AnalysisProperties->BoneSocket, BoneOffset, BoneName))
+	{
+		return false;
+	}
+
+	float DeltaTime = Animation.GetPlayLength() / NumSampledKeys;
+
+	int32 FirstKey = FMath::Clamp(
+		(int32) (NumSampledKeys * AnalysisProperties->StartTimeFraction), 0, NumSampledKeys);
+	int32 LastKey = FMath::Clamp(
+		(int32) (NumSampledKeys * AnalysisProperties->EndTimeFraction), FirstKey, NumSampledKeys-1);
+
+	if (FirstKey >= LastKey)
+	{
+		Result.Set(0, 0, 0);
+		return true;
+	}
+
+	FTransform FrameTM;
+	bool bNeedToUpdateFrameTM = true;
+	CalculateFrameTM(bNeedToUpdateFrameTM, FrameTM, FirstKey, AnalysisProperties, Animation);
+
+	FAnimPose AnimPose;
+	Result.Set(0, 0, 0);
+	for (int32 Key = FirstKey; Key != LastKey + 1; ++Key)
+	{
+		CalculateFrameTM(bNeedToUpdateFrameTM, FrameTM, Key, AnalysisProperties, Animation);
+
+		UAnimPoseExtensions::GetAnimPoseAtFrame(&Animation, Key, FAnimPoseEvaluationOptions(), AnimPose);
+		FTransform BoneTM1 = UAnimPoseExtensions::GetBonePose(AnimPose, BoneName, EAnimPoseSpaces::World);
+		FTransform TM1 = BoneOffset * BoneTM1;
+		FVector RelativePos1 = FrameTM.InverseTransformPosition(TM1.GetTranslation());
+
+		int32 NextKey = (Key + 1) % (NumSampledKeys + 1);
+		if (AnalysisProperties->Space == EAnalysisSpace::Moving)
+		{
+			CalculateFrameTM(bNeedToUpdateFrameTM, FrameTM, NextKey, AnalysisProperties, Animation);
+		}
+
+		UAnimPoseExtensions::GetAnimPoseAtFrame(&Animation, NextKey, FAnimPoseEvaluationOptions(), AnimPose);
+		FTransform BoneTM2 = UAnimPoseExtensions::GetBonePose(AnimPose, BoneName, EAnimPoseSpaces::World);
+		FTransform TM2 = BoneOffset * BoneTM2;
+		FVector RelativePos2 = FrameTM.InverseTransformPosition(TM2.GetTranslation());
+		FVector Velocity = (RelativePos2 - RelativePos1) / DeltaTime;
+
+#ifdef ANALYSIS_VERBOSE_LOG
+		UE_LOG(LogAnimation, Log, TEXT("%d Velocity = %f %f %f Height = %f"), 
+			   Key, Velocity.X, Velocity.Y, Velocity.Z, 0.5f * (RelativePos1 + RelativePos2).Z);
+#endif
+		Result += Velocity;
+	}
+	Result /= (1 + LastKey - FirstKey);
+	Result *= Animation.RateScale * RateScale;
+	UE_LOG(LogAnimation, Log, TEXT("%s vel = %f %f %f"), *BoneName.ToString(), Result.X, Result.Y, Result.Z);
+	return true;
+}
+
+//======================================================================================================================
+// Calculates the movement speed (magnitude) 
+template <typename T>
+static bool CalculateMovementSpeed(
+	float&                             Result,
+	const UBlendSpace&                 BlendSpace,
+	const T*                           AnalysisProperties,
+	const UAnimSequence&               Animation,
+	const float                        RateScale)
+{
+	if (!AnalysisProperties)
+	{
+		return false;
+	}
+
+	FVector Velocity;
+	if (CalculateVelocity(Velocity, BlendSpace, AnalysisProperties, Animation, RateScale))
+	{
+		FTransform FrameTM;
+		FVector FrameFacingDir = GetAxisFromTM(FrameTM, AnalysisProperties->CharacterFacingAxis);
+		FVector FrameUpDir = GetAxisFromTM(FrameTM, AnalysisProperties->CharacterUpAxis);
+		FVector FrameRightDir = FVector::CrossProduct(FrameUpDir, FrameFacingDir);
+		float Fwd = Velocity | FrameFacingDir;
+		float Right = Velocity | FrameRightDir;
+		Result = FMath::Sqrt(FMath::Square(Fwd) + FMath::Square(Right));
+		return true;
+	}
+	return false;
+}
+
+//======================================================================================================================
+// Calculates the movement speed (magnitude) 
+template <typename T>
+static bool CalculateMovementDirection(
+	float&                             Result,
+	const UBlendSpace&                 BlendSpace,
+	const T*                           AnalysisProperties,
+	const UAnimSequence&               Animation,
+	const float                        RateScale)
+{
+	if (!AnalysisProperties)
+	{
+		return false;
+	}
+
+	FVector Velocity;
+	if (CalculateVelocity(Velocity, BlendSpace, AnalysisProperties, Animation, RateScale))
+	{
+		FTransform FrameTM;
+		FVector FrameFacingDir = GetAxisFromTM(FrameTM, AnalysisProperties->CharacterFacingAxis);
+		FVector FrameUpDir = GetAxisFromTM(FrameTM, AnalysisProperties->CharacterUpAxis);
+		FVector FrameRightDir = FVector::CrossProduct(FrameUpDir, FrameFacingDir);
+
+		float Fwd = Velocity | FrameFacingDir;
+		float Right = Velocity | FrameRightDir;
+		Result = FMath::RadiansToDegrees(FMath::Atan2(Right, Fwd));
+		return true;
+	}
+	return false;
+}
+
+//======================================================================================================================
+template <typename T>
+void CalculateBoneOrientation(
+	FVector&                   RollPitchYaw,
+	const UAnimSequence&       Animation, 
+	const int32                Key, 
+	const FName                BoneName, 
+	const FTransform&          BoneOffset, 
+	const T&                   AnalysisProperties, 
+	const FVector&             FrameFacingDir, 
+	const FVector&             FrameRightDir, 
+	const FVector&             FrameUpDir)
+{
+	FAnimPose AnimPose;
+
+	UAnimPoseExtensions::GetAnimPoseAtFrame(&Animation, Key, FAnimPoseEvaluationOptions(), AnimPose);
+	FTransform BoneTM = UAnimPoseExtensions::GetBonePose(AnimPose, BoneName, EAnimPoseSpaces::World);
+
+	FTransform TM = BoneOffset * BoneTM;
+	FQuat AimQuat = TM.GetRotation();
+	FVector AimFwdDir = GetAxisFromTM(TM, AnalysisProperties->BoneFacingAxis);
+	FVector AimRightDir = GetAxisFromTM(TM, AnalysisProperties->BoneRightAxis);
+
+	// Note that Yaw is best taken from the AimRightDir - this is to avoid problems when the gun is pointing
+	// up or down - especially if it goes beyond 90 degrees in pitch.
+	float Yaw = FMath::RadiansToDegrees(FMath::Atan2(
+		FVector::DotProduct(AimRightDir, -FrameFacingDir), FVector::DotProduct(AimRightDir, FrameRightDir)));
+
+	// Undo the yaw to get pitch
+	const FQuat YawQuat(FrameUpDir, FMath::DegreesToRadians(Yaw));
+	FVector UnYawedAimFwdDir = YawQuat.UnrotateVector(AimFwdDir);
+	float Up = UnYawedAimFwdDir | FrameUpDir;
+	float Fwd = UnYawedAimFwdDir | FrameFacingDir;
+	float Pitch = FMath::RadiansToDegrees(FMath::Atan2(Up, Fwd));
+
+	// Undo the pitch to get roll
+	FVector UnYawedAimRightDir = YawQuat.UnrotateVector(AimRightDir);
+	const FQuat PitchQuat(FrameRightDir, -FMath::DegreesToRadians(Pitch));
+
+	FVector UnYawedUnPitchedAimRightDir = PitchQuat.UnrotateVector(UnYawedAimRightDir);
+
+	float Roll = FMath::RadiansToDegrees(FMath::Atan2(
+		FVector::DotProduct(UnYawedUnPitchedAimRightDir, -FrameUpDir), 
+		FVector::DotProduct(UnYawedUnPitchedAimRightDir, FrameRightDir)));
+
+	RollPitchYaw.Set(Roll, Pitch, Yaw);
+}
+
+//======================================================================================================================
+// Note that if a looping animation has 56 keys, then its first key is 0 and last is 55, but these will be identical poses.
+// Thus it has one fewer intervals/unique keys
+template <typename T>
+static bool CalculateOrientation(
+	FVector&                        Result,
+	const UBlendSpace&              BlendSpace,
+	const T*                        AnalysisProperties,
+	const UAnimSequence&            Animation,
+	const float                     RateScale)
+{
+	const int32 NumSampledKeys = Animation.GetNumberOfSampledKeys() - 1;
+	if (!AnalysisProperties || NumSampledKeys <= 0)
+	{
+		return false;
+	}
+
+	FTransform BoneOffset;
+	FName BoneName;
+	if (!GetBoneInfo(Animation, AnalysisProperties->BoneSocket, BoneOffset, BoneName))
+	{
+		return false;
+	}
+
+	int32 FirstKey = FMath::Clamp(
+		(int32) (NumSampledKeys * AnalysisProperties->StartTimeFraction), 0, NumSampledKeys);
+	int32 LastKey = FMath::Clamp(
+		(int32) (NumSampledKeys * AnalysisProperties->EndTimeFraction), FirstKey, NumSampledKeys);
+
+	FTransform FrameTM;
+	bool bNeedToUpdateFrameTM = true;
+	FVector FrameFacingDir, FrameUpDir, FrameRightDir;
+	CalculateFrameTM(bNeedToUpdateFrameTM, FrameTM, FirstKey, AnalysisProperties, Animation);
+	GetFrameDirs(FrameFacingDir, FrameUpDir, FrameRightDir, FrameTM, AnalysisProperties);
+
+	FAnimPose AnimPose;
+	Result.Set(0, 0, 0);
+	for (int32 Key = FirstKey; Key != LastKey + 1; ++Key)
+	{
+		CalculateFrameTM(bNeedToUpdateFrameTM, FrameTM, Key, AnalysisProperties, Animation);
+		GetFrameDirs(FrameFacingDir, FrameUpDir, FrameRightDir, FrameTM, AnalysisProperties);
+
+		FVector RollPitchYaw;
+		CalculateBoneOrientation(
+			RollPitchYaw, Animation, Key, BoneName, BoneOffset, 
+			AnalysisProperties, FrameFacingDir, FrameRightDir, FrameUpDir);
+#ifdef ANALYSIS_VERBOSE_LOG
+		UE_LOG(LogAnimation, Log, TEXT("Roll/pitch/yaw = %f %f %f"), RollPitchYaw.X, RollPitchYaw.Y, RollPitchYaw.Z);
+#endif
+		Result += RollPitchYaw;
+	}
+	Result /= (1 + LastKey - FirstKey);
+	UE_LOG(LogAnimation, Log, TEXT("%s Orientation = %f %f %f"), *BoneName.ToString(), Result.X, Result.Y, Result.Z);
+	return true;
+}
+
+//======================================================================================================================
+// Note that if a looping animation has 56 keys, then its first key is 0 and last is 55, but these will be identical poses.
+// Thus it has one fewer intervals/unique keys
+template <typename T>
+static bool CalculateDeltaOrientation(
+	FVector&                        Result,
+	const UBlendSpace&              BlendSpace,
+	const T*                        AnalysisProperties,
+	const UAnimSequence&            Animation,
+	const float                     RateScale)
+{
+	const int32 NumSampledKeys = Animation.GetNumberOfSampledKeys() - 1;
+	if (!AnalysisProperties || NumSampledKeys <= 0)
+	{
+		return false;
+	}
+
+	FTransform BoneOffset;
+	FName BoneName;
+	if (!GetBoneInfo(Animation, AnalysisProperties->BoneSocket, BoneOffset, BoneName))
+	{
+		return false;
+	}
+
+	int32 FirstKey = FMath::Clamp(
+		(int32) (NumSampledKeys * AnalysisProperties->StartTimeFraction), 0, NumSampledKeys);
+	int32 LastKey = FMath::Clamp(
+		(int32) (NumSampledKeys * AnalysisProperties->EndTimeFraction), FirstKey, NumSampledKeys);
+
+	FTransform FrameTM;
+	bool bNeedToUpdateFrameTM = true;
+	FVector FrameFacingDir, FrameUpDir, FrameRightDir;
+	CalculateFrameTM(bNeedToUpdateFrameTM, FrameTM, FirstKey, AnalysisProperties, Animation);
+	GetFrameDirs(FrameFacingDir, FrameUpDir, FrameRightDir, FrameTM, AnalysisProperties);
+
+	FVector RollPitchYaw1;
+	CalculateBoneOrientation(
+		RollPitchYaw1, Animation, FirstKey, BoneName, BoneOffset,
+		AnalysisProperties, FrameFacingDir, FrameRightDir, FrameUpDir);
+
+	CalculateFrameTM(bNeedToUpdateFrameTM, FrameTM, LastKey, AnalysisProperties, Animation);
+	GetFrameDirs(FrameFacingDir, FrameUpDir, FrameRightDir, FrameTM, AnalysisProperties);
+
+	FVector RollPitchYaw2;
+	CalculateBoneOrientation(
+		RollPitchYaw2, Animation, LastKey, BoneName, BoneOffset,
+		AnalysisProperties, FrameFacingDir, FrameRightDir, FrameUpDir);
+
+	Result = RollPitchYaw2 - RollPitchYaw1;
+	return true;
+}
+
+//======================================================================================================================
+template <typename T>
+static bool CalculateAngularVelocity(
+	FVector&                         Result,
+	const UBlendSpace&               BlendSpace,
+	const T*                         AnalysisProperties,
+	const UAnimSequence&             Animation,
+	const float                      RateScale)
+{
+	const int32 NumSampledKeys = Animation.GetNumberOfSampledKeys() - 1;
+	if (!AnalysisProperties || NumSampledKeys <= 0)
+	{
+		return false;
+	}
+
+	if (NumSampledKeys == 1)
+	{
+		Result.Set(0, 0, 0);
+		return true;
+	}
+
+	FTransform BoneOffset;
+	FName BoneName;
+	if (!GetBoneInfo(Animation, AnalysisProperties->BoneSocket, BoneOffset, BoneName))
+	{
+		return false;
+	}
+
+	float DeltaTime = Animation.GetPlayLength() / NumSampledKeys;
+
+	int32 FirstKey = FMath::Clamp(
+		(int32) (NumSampledKeys * AnalysisProperties->StartTimeFraction), 0, NumSampledKeys);
+	int32 LastKey = FMath::Clamp(
+		(int32) (NumSampledKeys * AnalysisProperties->EndTimeFraction), FirstKey, NumSampledKeys-1);
+
+	if (FirstKey >= LastKey)
+	{
+		Result.Set(0, 0, 0);
+		return true;
+	}
+
+	FTransform FrameTM;
+	bool bNeedToUpdateFrameTM = true;
+	CalculateFrameTM(bNeedToUpdateFrameTM, FrameTM, FirstKey, AnalysisProperties, Animation);
+
+	FAnimPose AnimPose;
+	Result.Set(0, 0, 0);
+	for (int32 Key = FirstKey; Key != LastKey + 1; ++Key)
+	{
+		CalculateFrameTM(bNeedToUpdateFrameTM, FrameTM, Key, AnalysisProperties, Animation);
+
+		UAnimPoseExtensions::GetAnimPoseAtFrame(&Animation, Key, FAnimPoseEvaluationOptions(), AnimPose);
+		FTransform BoneTM1 = UAnimPoseExtensions::GetBonePose(AnimPose, BoneName, EAnimPoseSpaces::World);
+		FTransform TM1 = BoneOffset * BoneTM1;
+		FQuat RelativeQuat1 = FrameTM.InverseTransformRotation(TM1.GetRotation());
+
+		int32 NextKey = (Key + 1) % (NumSampledKeys + 1);
+		if (AnalysisProperties->Space == EAnalysisSpace::Moving)
+		{
+			CalculateFrameTM(bNeedToUpdateFrameTM, FrameTM, NextKey, AnalysisProperties, Animation);
+		}
+
+		UAnimPoseExtensions::GetAnimPoseAtFrame(&Animation, NextKey, FAnimPoseEvaluationOptions(), AnimPose);
+		FTransform BoneTM2 = UAnimPoseExtensions::GetBonePose(AnimPose, BoneName, EAnimPoseSpaces::World);
+		FTransform TM2 = BoneOffset * BoneTM2;
+		FQuat RelativeQuat2 = FrameTM.InverseTransformRotation(TM2.GetRotation());
+
+		FQuat Rotation = RelativeQuat2 * RelativeQuat1.Inverse();
+		FVector Axis;
+		float Angle;
+		Rotation.ToAxisAndAngle(Axis, Angle);
+		FVector AngularVelocity = FMath::RadiansToDegrees(Axis * (Angle / DeltaTime));
+#ifdef ANALYSIS_VERBOSE_LOG
+		UE_LOG(LogAnimation, Log, TEXT("Angular Velocity = %f %f %f"), AngularVelocity.X, AngularVelocity.Y, AngularVelocity.Z);
+#endif
+		Result += AngularVelocity;
+	}
+	Result /= (1 + LastKey - FirstKey);
+	Result *= Animation.RateScale * RateScale;
+	UE_LOG(LogAnimation, Log, TEXT("%s angular velocity = %f %f %f"), *BoneName.ToString(), Result.X, Result.Y, Result.Z);
+	return true;
+}
+
+//======================================================================================================================
+template <typename T>
+static bool CalculateOrientationRate(
+	FVector&                        Result,
+	const UBlendSpace&              BlendSpace,
+	const T*                        AnalysisProperties,
+	const UAnimSequence&            Animation,
+	const float                     RateScale)
+{
+	const int32 NumSampledKeys = Animation.GetNumberOfSampledKeys() - 1;
+	if (!AnalysisProperties || NumSampledKeys <= 0)
+	{
+		return false;
+	}
+
+	if (NumSampledKeys == 1)
+	{
+		Result.Set(0, 0, 0);
+		return true;
+	}
+
+	FTransform BoneOffset;
+	FName BoneName;
+	if (!GetBoneInfo(Animation, AnalysisProperties->BoneSocket, BoneOffset, BoneName))
+	{
+		return false;
+	}
+
+	float DeltaTime = Animation.GetPlayLength() / NumSampledKeys;
+
+	int32 FirstKey = FMath::Clamp(
+		(int32) (NumSampledKeys * AnalysisProperties->StartTimeFraction), 0, NumSampledKeys);
+	int32 LastKey = FMath::Clamp(
+		(int32) (NumSampledKeys * AnalysisProperties->EndTimeFraction), FirstKey, NumSampledKeys-1);
+
+	if (FirstKey >= LastKey)
+	{
+		Result.Set(0, 0, 0);
+		return true;
+	}
+
+	FTransform FrameTM;
+	bool bNeedToUpdateFrameTM = true;
+	FVector FrameFacingDir, FrameUpDir, FrameRightDir;
+	CalculateFrameTM(bNeedToUpdateFrameTM, FrameTM, FirstKey, AnalysisProperties, Animation);
+	GetFrameDirs(FrameFacingDir, FrameUpDir, FrameRightDir, FrameTM, AnalysisProperties);
+
+	FAnimPose AnimPose;
+	Result.Set(0, 0, 0);
+	for (int32 Key = FirstKey; Key != LastKey + 1; ++Key)
+	{
+		CalculateFrameTM(bNeedToUpdateFrameTM, FrameTM, Key, AnalysisProperties, Animation);
+		GetFrameDirs(FrameFacingDir, FrameUpDir, FrameRightDir, FrameTM, AnalysisProperties);
+
+		FVector RollPitchYaw1;
+		CalculateBoneOrientation(
+			RollPitchYaw1, Animation, Key, BoneName, BoneOffset, 
+			AnalysisProperties, FrameFacingDir, FrameRightDir, FrameUpDir);
+
+		int32 NextKey = (Key + 1) % (NumSampledKeys + 1);
+		if (AnalysisProperties->Space == EAnalysisSpace::Moving)
+		{
+			CalculateFrameTM(bNeedToUpdateFrameTM, FrameTM, NextKey, AnalysisProperties, Animation);
+			GetFrameDirs(FrameFacingDir, FrameUpDir, FrameRightDir, FrameTM, AnalysisProperties);
+		}
+
+		FVector RollPitchYaw2;
+		CalculateBoneOrientation(
+			RollPitchYaw2, Animation, NextKey, BoneName, BoneOffset, 
+			AnalysisProperties, FrameFacingDir, FrameRightDir, FrameUpDir);
+
+		const FVector OrientationRate = (RollPitchYaw2 - RollPitchYaw1) / DeltaTime;
+#ifdef ANALYSIS_VERBOSE_LOG
+		UE_LOG(LogAnimation, Log, TEXT("Orientation rate = %f %f %f"), OrientationRate.X, OrientationRate.Y, OrientationRate.Z);
+#endif
+		Result += OrientationRate;
+	}
+	Result /= (1 + LastKey - FirstKey);
+	Result *= Animation.RateScale * RateScale;
+	UE_LOG(LogAnimation, Log, TEXT("%s Orientation rate = %f %f %f"), *BoneName.ToString(), Result.X, Result.Y, Result.Z);
+	return true;
+}
+
+//======================================================================================================================
+// Note that if a looping animation has 56 keys, then its first key is 0 and last is 55, but these will be identical poses.
+// Thus it has one fewer intervals/unique keys
+static bool CalculateLocomotionVelocity(
+	FVector&                             Result,
+	const UBlendSpace&                   BlendSpace,
+	const ULocomotionAnalysisProperties* AnalysisProperties,
+	const FBoneSocketTarget&             BoneSocket,
+	const UAnimSequence&                 Animation,
+	const float                          RateScale)
+{
+	const int32 NumSampledKeys = Animation.GetNumberOfSampledKeys();
+	if (!AnalysisProperties || NumSampledKeys <= 0)
+	{
+		return false;
+	}
+	if (NumSampledKeys == 1)
+	{
+		Result.Set(0, 0, 0);
+		return true;
+	}
+
+	FTransform BoneOffset;
+	FName BoneName;
+	if (!GetBoneInfo(Animation, BoneSocket, BoneOffset, BoneName))
+	{
+		return false;
+	}
+
+	// Note that for locomotion we don't support the frame changing
+	FTransform FrameTM = FTransform::Identity;
+	FVector FrameFacingDir, FrameUpDir, FrameRightDir;
+	GetFrameDirs(FrameFacingDir, FrameUpDir, FrameRightDir, FrameTM, AnalysisProperties);
+
+	FAnimPose AnimPose;
+
+	// The frame time delta
+	float DeltaTime = Animation.GetPlayLength() / NumSampledKeys;
+
+	// First step is to figure out the approximate direction. Note that the average velocity will be zero (assuming a
+	// complete cycle) - but if we apply a weight that is based on the height, then we can bias it towards the foot that
+	// is on the ground.
+	float MinHeight = FLT_MAX;
+	float MaxHeight = -FLT_MAX;
+	FVector AveragePos(0.0f);
+	TArray<FVector> Positions;
+	Positions.SetNum(NumSampledKeys);
+	for (int32 Key = 0; Key != NumSampledKeys; ++Key)
+	{
+		UAnimPoseExtensions::GetAnimPoseAtFrame(&Animation, Key, FAnimPoseEvaluationOptions(), AnimPose);
+		FTransform BoneTM = UAnimPoseExtensions::GetBonePose(AnimPose, BoneName, EAnimPoseSpaces::World);
+		FTransform TM = BoneOffset * BoneTM;
+		FVector Pos = TM.GetTranslation();
+		Positions[Key] = Pos;
+		float Height = Pos | FrameUpDir;
+		MinHeight = FMath::Min(MinHeight, Height);
+		MaxHeight = FMath::Max(MaxHeight, Height);
+		AveragePos += Pos;
+#ifdef ANALYSIS_VERBOSE_LOG
+		UE_LOG(LogAnimation, Log, TEXT("Pos %f %f %f"), Pos.X, Pos.Y, Pos.Z);
+#endif
+	}
+	AveragePos /= NumSampledKeys;
+
+	// Calculate velocities. 
+	TArray<FVector> Velocities;
+	Velocities.SetNum(NumSampledKeys);
+	for (int32 Key = 0; Key != NumSampledKeys; ++Key)
+	{
+		int32 PrevKey = (Key + NumSampledKeys - 1) % NumSampledKeys;
+		int32 NextKey = (Key +  1) % NumSampledKeys;
+		Velocities[Key] = (Positions[NextKey] - Positions[PrevKey]) / (2.0f * DeltaTime);
+	}
+
+	FVector BiasedFootVel(0);
+	float TotalWeight = 0.0f;
+	for (int32 Key = 0 ; Key != NumSampledKeys ; ++Key)
+	{
+		float Height = Positions[Key] | FrameUpDir;
+		float Weight = 1.0f - (Height - MinHeight) / (MaxHeight - MinHeight);
+		BiasedFootVel += Velocities[Key] * Weight;
+		TotalWeight += Weight;
+	}
+	BiasedFootVel /= TotalWeight; 
+
+	if (BiasedFootVel.IsNearlyZero())
+	{
+		Result.Set(0, 0, 0);
+		return true;
+	}
+
+	FVector ApproxLocoDir = -BiasedFootVel.GetSafeNormal();
+
+	// Now we can form a mask, where 0 means traveling in the wrong direction (so clearly off the ground), and positive
+	// numbers will indicate how far into a valid segment we are. We will assume that the animation is looping
+	TArray<int32> Mask;
+	Mask.SetNum(NumSampledKeys);
+	for (int32 Key = 0 ; Key != NumSampledKeys ; ++Key)
+	{
+		Mask[Key] = (Velocities[Key] | ApproxLocoDir) >= 0.0f ? 0 : 1;
+	}
+
+	int32 StartKey = -1;
+	TArray<int32> Mask1 = Mask;
+	bool bRepeat = true;
+	int32 MaxMask = 0;
+	int32 PrevNumFound = NumSampledKeys + 1;
+	while (bRepeat)
+	{
+		bRepeat = false;
+		int32 NumFound = 0;
+		for (int32 Key = 0 ; Key != NumSampledKeys ; ++Key)
+		{
+			if (Mask[Key] > 0)
+			{
+				int32 PrevKey = (Key + NumSampledKeys - 1) % NumSampledKeys;
+				int32 NextKey = (Key + 1) % NumSampledKeys;
+				if (Mask[PrevKey] == Mask[Key] && Mask[NextKey] == Mask[Key])
+				{
+					++Mask1[Key];
+					MaxMask = FMath::Max(MaxMask, Mask1[Key]);
+					bRepeat = true;
+					++NumFound;
+				}
+			}
+		}
+		Mask = Mask1;
+		// Avoid a perpetual loop (e.g. can happen if initially all the mask values are 1... though that shouldn't
+		// really happen).
+		if (NumFound >= PrevNumFound)
+		{
+			bRepeat = false;
+		}
+		else
+		{
+			PrevNumFound = NumFound;
+		}
+	}
+
+	// When searching we will want to start outside of a "good" region.
+	int32 AZeroKey = 0;
+	for (int32 Key = 0 ; Key != NumSampledKeys ; ++Key)
+	{
+		if (Mask[Key] == 0)
+		{
+			AZeroKey = Key;
+		}
+	}
+
+	// We use the mask (with a somewhat arbitrary threshold) to get rid of velocities that are near to the foot
+	// plant/take-off time (and might be when the foot is in the time). Then We look for the highest velocity. Not that
+	// if we're being called with the foot (ankle) joint, it will tend to underestimate velocities since it is nearer
+	// the hip than the ground contact point.
+	// 
+	int32 Threshold = FMath::Max(MaxMask / 2, 1);
+	int32 Num = 0;
+	FVector AverageFootVel(0);
+	float BestSpeed = 0.0f;
+	int32 BestSpeedKey = 0;
+	for (int32 K = AZeroKey ; K != AZeroKey + NumSampledKeys ; ++K)
+	{
+		int32 Key = K % NumSampledKeys;
+		if (Mask[Key] >= Threshold)
+		{
+			float Speed = Velocities[Key].Size();
+#ifdef ANALYSIS_VERBOSE_LOG
+			UE_LOG(LogAnimation, Log, TEXT("Candidate %d Mask %d vel = %f %f %f, speed = %f"), 
+				   Key, Mask[Key], Velocities[Key].X, Velocities[Key].Y, Velocities[Key].Z, Speed);
+#endif
+			if (Speed > BestSpeed)
+			{
+				BestSpeed = Speed;
+				BestSpeedKey = Key;
+			}
+		}
+		else
+		{
+			if (BestSpeed > 0.0f)
+			{
+#ifdef ANALYSIS_VERBOSE_LOG
+				UE_LOG(LogAnimation, Log, TEXT("Picked Candidate %d vel = %f %f %f, speed = %f"), 
+					   BestSpeedKey, Velocities[BestSpeedKey].X, Velocities[BestSpeedKey].Y, Velocities[BestSpeedKey].Z, BestSpeed);
+#endif
+				AverageFootVel += Velocities[BestSpeedKey];
+				++Num;
+				BestSpeed = 0.0f;
+			}
+		}
+	}
+	// Make sure we didn't miss the last data point
+	if (BestSpeed > 0.0f)
+	{
+#ifdef ANALYSIS_VERBOSE_LOG
+		UE_LOG(LogAnimation, Log, TEXT("Picked Candidate %d vel = %f %f %f, speed = %f"), 
+			   BestSpeedKey, Velocities[BestSpeedKey].X, Velocities[BestSpeedKey].Y, Velocities[BestSpeedKey].Z, BestSpeed);
+#endif
+		AverageFootVel += Velocities[BestSpeedKey];
+		++Num;
+		BestSpeed = 0.0f;
+	}
+
+	AverageFootVel /= Num;
+	float FacingVel = -AverageFootVel | FrameFacingDir;
+	float RightVel = -AverageFootVel | FrameRightDir;
+	float UpVel = -AverageFootVel | FrameUpDir;
+
+	Result.Set(FacingVel, RightVel, UpVel);
+	Result *= Animation.RateScale * RateScale;
+	UE_LOG(LogAnimation, Log, TEXT("%s Locomotion vel = %f %f %f"), *BoneName.ToString(), Result.X, Result.Y, Result.Z);
+	return true;
+}
+
+//======================================================================================================================
+static bool CalculateLocomotionVelocity(
+	FVector&                             Result,
+	const UBlendSpace&                   BlendSpace,
+	const ULocomotionAnalysisProperties* AnalysisProperties,
+	const UAnimSequence&                 Animation,
+	const float                          RateScale)
+{
+	if (!AnalysisProperties)
+	{
+		return false;
+	}
+
+	FVector Result1, Result2;
+	int32 Num = 0;
+	Result.Set(0, 0, 0);
+	if (CalculateLocomotionVelocity(Result1, BlendSpace, AnalysisProperties, 
+									AnalysisProperties->PrimaryBoneSocket, Animation, RateScale))
+	{
+#ifdef ANALYSIS_VERBOSE_LOG
+		UE_LOG(LogAnimation, Log, TEXT("Loco vel from primary = %f %f %f"), Result1.X, Result1.Y, Result1.Z);
+#endif
+		Result += Result1;
+		++Num;
+	}
+	if (CalculateLocomotionVelocity(Result2, BlendSpace, AnalysisProperties, 
+									AnalysisProperties->SecondaryBoneSocket, Animation, RateScale))
+	{
+#ifdef ANALYSIS_VERBOSE_LOG
+		UE_LOG(LogAnimation, Log, TEXT("Loco vel from secondary = %f %f %f"), Result2.X, Result2.Y, Result2.Z);
+#endif
+		Result += Result2;
+		++Num;
+	}
+	if (Num)
+	{
+		Result /= Num;
+		UE_LOG(LogAnimation, Log, TEXT("Loco vel = %f %f %f"), Result.X, Result.Y, Result.Z);
+		return true;
+	}
+	return false;
+}
+
+//======================================================================================================================
+// Calculates the locomotion speed in the character's facing direction
+static bool CalculateLocomotionFwdSpeed(
+	float&                               Result,
+	const UBlendSpace&                   BlendSpace,
+	const ULocomotionAnalysisProperties* AnalysisProperties,
+	const UAnimSequence&                 Animation,
+	const float                          RateScale)
+{
+	FVector Movement;
+	if (CalculateLocomotionVelocity(Movement, BlendSpace, AnalysisProperties, Animation, RateScale))
+	{
+		Result = Movement.X;
+		return true;
+	}
+	return false;
+}
+
+//======================================================================================================================
+// Calculates the locomotion speed in the character's upwards direction
+static bool CalculateLocomotionUpSpeed(
+	float&                               Result,
+	const UBlendSpace&                   BlendSpace,
+	const ULocomotionAnalysisProperties* AnalysisProperties,
+	const UAnimSequence&                 Animation,
+	const float                          RateScale)
+{
+	FVector Movement;
+	if (CalculateLocomotionVelocity(Movement, BlendSpace, AnalysisProperties, Animation, RateScale))
+	{
+		Result = Movement.Z;
+		return true;
+	}
+	return false;
+}
+
+//======================================================================================================================
+// Calculates the locomotion speed in the character's right direction
+static bool CalculateLocomotionRightSpeed(
+	float&                               Result,
+	const UBlendSpace&                   BlendSpace,
+	const ULocomotionAnalysisProperties* AnalysisProperties,
+	const UAnimSequence&                 Animation,
+	const float                          RateScale)
+{
+	FVector Movement;
+	if (CalculateLocomotionVelocity(Movement, BlendSpace, AnalysisProperties, Animation, RateScale))
+	{
+		Result = Movement.Y;
+		return true;
+	}
+	return false;
+}
+
+//======================================================================================================================
+// Calculates the locomotion speed (magnitude)
+static bool CalculateLocomotionSpeed(
+	float&                               Result,
+	const UBlendSpace&                   BlendSpace,
+	const ULocomotionAnalysisProperties* AnalysisProperties,
+	const UAnimSequence&                 Animation,
+	const float                          RateScale)
+{
+	FVector Movement;
+	if (CalculateLocomotionVelocity(Movement, BlendSpace, AnalysisProperties, Animation, RateScale))
+	{
+		Result = Movement.Size();
+		return true;
+	}
+	return false;
+}
+
+//======================================================================================================================
+// Calculates the locomotion direction (degrees)
+static bool CalculateLocomotionDirection(
+	float&                               Result,
+	const UBlendSpace&                   BlendSpace,
+	const ULocomotionAnalysisProperties* AnalysisProperties,
+	const UAnimSequence&                 Animation,
+	const float                          RateScale)
+{
+	FVector Movement;
+	if (CalculateLocomotionVelocity(Movement, BlendSpace, AnalysisProperties, Animation, RateScale))
+	{
+		Result = FMath::RadiansToDegrees(FMath::Atan2(Movement.Y, Movement.X));
+		return true;
+	}
+	return false;
+}
+
+//======================================================================================================================
+// Calculates the locomotion slope angle (degrees) going in the facing direction
+static bool CalculateLocomotionFwdSlope(
+	float&                               Result,
+	const UBlendSpace&                   BlendSpace,
+	const ULocomotionAnalysisProperties* AnalysisProperties,
+	const UAnimSequence&                 Animation,
+	const float                          RateScale)
+{
+	FVector Movement;
+	if (CalculateLocomotionVelocity(Movement, BlendSpace, AnalysisProperties, Animation, RateScale))
+	{
+		if (Movement.X >= 0.0f)
+		{
+			Result = FMath::RadiansToDegrees(FMath::Atan2(Movement.Z, Movement.X));
+		}
+		else
+		{
+			Result = FMath::RadiansToDegrees(FMath::Atan2(-Movement.Z, -Movement.X));
+		}
+
+		return true;
+	}
+	return false;
+}
+
+//======================================================================================================================
+// Calculates the locomotion slope angle (degrees) going in the rightwards direction
+static bool CalculateLocomotionRightSlope(
+	float&                               Result,
+	const UBlendSpace&                   BlendSpace,
+	const ULocomotionAnalysisProperties* AnalysisProperties,
+	const UAnimSequence&                 Animation,
+	const float                          RateScale)
+{
+	FVector Movement;
+	if (CalculateLocomotionVelocity(Movement, BlendSpace, AnalysisProperties, Animation, RateScale))
+	{
+		if (Movement.Y > 0.0f)
+		{
+			Result = FMath::RadiansToDegrees(FMath::Atan2(Movement.Z, Movement.Y));
+		}
+		else
+		{
+			Result = FMath::RadiansToDegrees(FMath::Atan2(-Movement.Z, -Movement.Y));
+		}
+		return true;
+	}
+	return false;
+}
+
+//======================================================================================================================
+/**
+ * Helper to extract the component from the FVector functions
+ */
+template<typename FunctionType, typename T>
+static bool CalculateComponentSampleValue(
+	float&                     Result,
+	const FunctionType&        Fn,
+	const UBlendSpace&         BlendSpace, 
+	const T*                   AnalysisProperties, 
+	const UAnimSequence&       Animation,
+	const float                RateScale)
+{
+	FVector Value;
+	int32 ComponentIndex = (int32) AnalysisProperties->FunctionAxis;
+	if (Fn(Value, BlendSpace, AnalysisProperties, Animation, RateScale))
+	{
+		Result = Value[ComponentIndex];
+		return true;
+	}
+	return false;
+}
+
+//======================================================================================================================
+void ULinearAnalysisProperties::InitializeFromCache(TSharedPtr<FCachedAnalysisProperties> Cache)
+{
+	if (Cache)
+	{
+		Super::InitializeFromCache(Cache);
+		FunctionAxis = Cache->LinearFunctionAxis;
+		Space = Cache->Space;
+		SpaceBoneSocket = Cache->SpaceBoneSocket;
+		BoneSocket = Cache->BoneSocket1;
+		StartTimeFraction = Cache->StartTimeFraction;
+		EndTimeFraction = Cache->EndTimeFraction;
+	}
+}
+
+//======================================================================================================================
+void ULinearAnalysisProperties::MakeCache(TSharedPtr<FCachedAnalysisProperties>& Cache) const
+{
+	Super::MakeCache(Cache);
+	Cache->LinearFunctionAxis = FunctionAxis;
+	Cache->Space = Space;
+	Cache->SpaceBoneSocket = SpaceBoneSocket;
+	Cache->BoneSocket1 = BoneSocket;
+	Cache->StartTimeFraction = StartTimeFraction;
+	Cache->EndTimeFraction = EndTimeFraction;
+}
+
+//======================================================================================================================
+void UEulerAnalysisProperties::InitializeFromCache(TSharedPtr<FCachedAnalysisProperties> Cache)
+{
+	if (Cache)
+	{
+		Super::InitializeFromCache(Cache);
+		FunctionAxis = Cache->EulerFunctionAxis;
+		Space = Cache->Space;
+		SpaceBoneSocket = Cache->SpaceBoneSocket;
+		CharacterFacingAxis = Cache->CharacterFacingAxis;
+		CharacterUpAxis = Cache->CharacterUpAxis;
+		StartTimeFraction = Cache->StartTimeFraction;
+		EndTimeFraction = Cache->EndTimeFraction;
+		BoneSocket = Cache->BoneSocket1;
+		BoneFacingAxis = Cache->BoneFacingAxis;
+		BoneRightAxis = Cache->BoneRightAxis;
+	}
+}
+
+//======================================================================================================================
+void UEulerAnalysisProperties::MakeCache(TSharedPtr<FCachedAnalysisProperties>& Cache) const
+{
+	Super::MakeCache(Cache);
+	Cache->EulerFunctionAxis = FunctionAxis;
+	Cache->Space = Space;
+	Cache->SpaceBoneSocket = SpaceBoneSocket;
+	Cache->CharacterFacingAxis = CharacterFacingAxis;
+	Cache->CharacterUpAxis = CharacterUpAxis;
+	Cache->StartTimeFraction = StartTimeFraction;
+	Cache->EndTimeFraction = EndTimeFraction;
+	Cache->BoneSocket1 = BoneSocket;
+	Cache->BoneFacingAxis = BoneFacingAxis;
+	Cache->BoneRightAxis = BoneRightAxis;
+}
+
+//======================================================================================================================
+void UMovementAnalysisProperties::InitializeFromCache(TSharedPtr<FCachedAnalysisProperties> Cache)
+{
+	if (Cache)
+	{
+		Super::InitializeFromCache(Cache);
+		Space = Cache->Space;
+		SpaceBoneSocket = Cache->SpaceBoneSocket;
+		CharacterFacingAxis = Cache->CharacterFacingAxis;
+		CharacterUpAxis = Cache->CharacterUpAxis;
+		StartTimeFraction = Cache->StartTimeFraction;
+		EndTimeFraction = Cache->EndTimeFraction;
+		BoneSocket = Cache->BoneSocket1;
+	}
+}
+
+//======================================================================================================================
+void UMovementAnalysisProperties::MakeCache(TSharedPtr<FCachedAnalysisProperties>& Cache) const
+{
+	Super::MakeCache(Cache);
+	Cache->Space = Space;
+	Cache->SpaceBoneSocket = SpaceBoneSocket;
+	Cache->CharacterFacingAxis = CharacterFacingAxis;
+	Cache->CharacterUpAxis = CharacterUpAxis;
+	Cache->StartTimeFraction = StartTimeFraction;
+	Cache->EndTimeFraction = EndTimeFraction;
+	Cache->BoneSocket1 = BoneSocket;
+}
+
+//======================================================================================================================
+void ULocomotionAnalysisProperties::InitializeFromCache(TSharedPtr<FCachedAnalysisProperties> Cache)
+{
+	if (Cache)
+	{
+		Super::InitializeFromCache(Cache);
+		CharacterFacingAxis = Cache->CharacterFacingAxis;
+		CharacterUpAxis = Cache->CharacterUpAxis;
+		PrimaryBoneSocket = Cache->BoneSocket1;
+		SecondaryBoneSocket = Cache->BoneSocket2;
+	}
+}
+
+//======================================================================================================================
+void ULocomotionAnalysisProperties::MakeCache(TSharedPtr<FCachedAnalysisProperties>& Cache) const
+{
+	Super::MakeCache(Cache);
+	Cache->CharacterFacingAxis = CharacterFacingAxis;
+	Cache->CharacterUpAxis = CharacterUpAxis;
+	Cache->BoneSocket1 = PrimaryBoneSocket;
+	Cache->BoneSocket2 = SecondaryBoneSocket;
+}
+
+
+//======================================================================================================================
+class FCoreBlendSpaceAnalysisFeature : public IBlendSpaceAnalysisFeature
+{
+public:
+	// This should process the animation according to the analysis properties, or return false if that is not possible.
+	bool CalculateSampleValue(float&                     Result,
+							  const UBlendSpace&         BlendSpace,
+							  const UAnalysisProperties* AnalysisProperties,
+							  const UAnimSequence&       Animation,
+							  const float                RateScale) const override;
+
+	// This should return an instance derived from UAnalysisProperties that is suitable for the FunctionName
+	UAnalysisProperties* MakeAnalysisProperties(UObject* Outer, const FString& FunctionName) const override;
+
+	// This should return the names of the functions handled
+	TArray<FString> GetAnalysisFunctions() const override;
+};
+
+static FCoreBlendSpaceAnalysisFeature CoreBlendSpaceAnalysisFeature;
+
+//======================================================================================================================
+TArray<FString> FCoreBlendSpaceAnalysisFeature::GetAnalysisFunctions() const
+{
+	TArray<FString> Functions = 
+	{
+		TEXT("None"),
+		TEXT("Position"),
+		TEXT("Velocity"),
+		TEXT("DeltaPosition"),
+		TEXT("Orientation"),
+		TEXT("OrientationRate"),
+		TEXT("DeltaOrientation"),
+		TEXT("AngularVelocity"),
+		TEXT("MovementSpeed"),
+		TEXT("MovementDirection"),
+		TEXT("LocomotionRightSpeed"),
+		TEXT("LocomotionForwardSpeed"),
+		TEXT("LocomotionUpSpeed"),
+		TEXT("LocomotionSpeed"),
+		TEXT("LocomotionDirection"),
+		TEXT("LocomotionForwardSlope"),
+		TEXT("LocomotionRightSlope")
+	};
+	return Functions;
+}
+
+//======================================================================================================================
+UAnalysisProperties* FCoreBlendSpaceAnalysisFeature::MakeAnalysisProperties(UObject* Outer, const FString& FunctionName) const
+{
+	UAnalysisProperties* Result = nullptr;
+	if (FunctionName.Equals(TEXT("Position")) ||
+		FunctionName.Equals(TEXT("Velocity")) ||
+		FunctionName.Equals(TEXT("DeltaPosition")) ||
+		FunctionName.Equals(TEXT("AngularVelocity")))
+	{
+		Result = NewObject<ULinearAnalysisProperties>(Outer);
+	}
+	else if (FunctionName.Equals(TEXT("Orientation")) ||
+			 FunctionName.Equals(TEXT("OrientationRate")) ||
+			 FunctionName.Equals(TEXT("DeltaOrientation")))
+	{
+		Result = NewObject<UEulerAnalysisProperties>(Outer);
+	}
+	else if (FunctionName.Equals(TEXT("MovementSpeed")) ||
+			 FunctionName.Equals(TEXT("MovementDirection")))
+	{
+		Result = NewObject<UMovementAnalysisProperties>(Outer);
+	}
+	else if (FunctionName.Equals(TEXT("LocomotionForwardSpeed")) ||
+			 FunctionName.Equals(TEXT("LocomotionRightSpeed")) ||
+			 FunctionName.Equals(TEXT("LocomotionUpSpeed")) ||
+			 FunctionName.Equals(TEXT("LocomotionSpeed")) ||
+			 FunctionName.Equals(TEXT("LocomotionDirection")) ||
+			 FunctionName.Equals(TEXT("LocomotionForwardSlope")) ||
+			 FunctionName.Equals(TEXT("LocomotionRightSlope")))
+	{
+		Result = NewObject<ULocomotionAnalysisProperties>(Outer);
+	}
+
+	if (Result)
+	{
+		Result->Function = FunctionName;
+	}
+	return Result;
+}
+
+//======================================================================================================================
+bool FCoreBlendSpaceAnalysisFeature::CalculateSampleValue(float&                     Result,
+														  const UBlendSpace&         BlendSpace,
+														  const UAnalysisProperties* AnalysisProperties,
+														  const UAnimSequence&       Animation,
+														  const float                RateScale) const
+{
+	if (!AnalysisProperties)
+	{
+		return false;
+	}
+	const FString& FunctionName = AnalysisProperties->Function;
+	if (FunctionName.Equals(TEXT("Position")))
+	{
+		return CalculateComponentSampleValue(
+			Result, CalculatePosition<ULinearAnalysisProperties>, BlendSpace, 
+			Cast<ULinearAnalysisProperties>(AnalysisProperties), Animation, RateScale);
+	}
+	else if (FunctionName.Equals(TEXT("Velocity")))
+	{
+		return CalculateComponentSampleValue(
+			Result, CalculateVelocity<ULinearAnalysisProperties>, BlendSpace, 
+			Cast<ULinearAnalysisProperties>(AnalysisProperties), Animation, RateScale);
+	}
+	else if (FunctionName.Equals(TEXT("DeltaPosition")))
+	{
+		return CalculateComponentSampleValue(
+			Result, CalculateDeltaPosition<ULinearAnalysisProperties>, BlendSpace, 
+			Cast<ULinearAnalysisProperties>(AnalysisProperties), Animation, RateScale);
+	}
+	else if (FunctionName.Equals(TEXT("AngularVelocity")))
+	{
+		return CalculateComponentSampleValue(
+			Result, CalculateAngularVelocity<ULinearAnalysisProperties>, BlendSpace, 
+			Cast<ULinearAnalysisProperties>(AnalysisProperties), Animation, RateScale);
+	}
+	else if (FunctionName.Equals(TEXT("Orientation")))
+	{
+		return CalculateComponentSampleValue(
+			Result, CalculateOrientation<UEulerAnalysisProperties>, BlendSpace, 
+			Cast<UEulerAnalysisProperties>(AnalysisProperties), Animation, RateScale);
+	}
+	else if (FunctionName.Equals(TEXT("OrientationRate")))
+	{
+		return CalculateComponentSampleValue(
+			Result, CalculateOrientationRate<UEulerAnalysisProperties>, BlendSpace, 
+			Cast<UEulerAnalysisProperties>(AnalysisProperties), Animation, RateScale);
+	}
+	else if (FunctionName.Equals(TEXT("DeltaOrientation")))
+	{
+		return CalculateComponentSampleValue(
+			Result, CalculateDeltaOrientation<UEulerAnalysisProperties>, BlendSpace, 
+			Cast<UEulerAnalysisProperties>(AnalysisProperties), Animation, RateScale);
+	}
+	else if (FunctionName.Equals(TEXT("MovementSpeed")))
+	{
+		return CalculateMovementSpeed(
+			Result, BlendSpace, Cast<UMovementAnalysisProperties>(AnalysisProperties), Animation, RateScale);
+	}
+	else if (FunctionName.Equals(TEXT("MovementDirection")))
+	{
+		return CalculateMovementDirection(
+			Result, BlendSpace, Cast<UMovementAnalysisProperties>(AnalysisProperties), Animation, RateScale);
+	}
+	else if (FunctionName.Equals(TEXT("LocomotionForwardSpeed")))
+	{
+		return CalculateLocomotionFwdSpeed(
+			Result, BlendSpace, Cast<ULocomotionAnalysisProperties>(AnalysisProperties), Animation, RateScale);
+	}
+	else if (FunctionName.Equals(TEXT("LocomotionRightSpeed")))
+	{
+		return CalculateLocomotionRightSpeed(
+			Result, BlendSpace, Cast<ULocomotionAnalysisProperties>(AnalysisProperties), Animation, RateScale);
+	}
+	else if (FunctionName.Equals(TEXT("LocomotionUpSpeed")))
+	{
+		return CalculateLocomotionUpSpeed(
+			Result, BlendSpace, Cast<ULocomotionAnalysisProperties>(AnalysisProperties), Animation, RateScale);
+	}
+	else if (FunctionName.Equals(TEXT("LocomotionSpeed")))
+	{
+		return CalculateLocomotionSpeed(
+			Result, BlendSpace, Cast<ULocomotionAnalysisProperties>(AnalysisProperties), Animation, RateScale);
+	}
+	else if (FunctionName.Equals(TEXT("LocomotionDirection")))
+	{
+		return CalculateLocomotionDirection(
+			Result, BlendSpace, Cast<ULocomotionAnalysisProperties>(AnalysisProperties), Animation, RateScale);
+	}
+	else if (FunctionName.Equals(TEXT("LocomotionForwardSlope")))
+	{
+		return CalculateLocomotionFwdSlope(
+			Result, BlendSpace, Cast<ULocomotionAnalysisProperties>(AnalysisProperties), Animation, RateScale);
+	}
+	else if (FunctionName.Equals(TEXT("LocomotionRightSlope")))
+	{
+		return CalculateLocomotionRightSlope(
+			Result, BlendSpace, Cast<ULocomotionAnalysisProperties>(AnalysisProperties), Animation, RateScale);
+	}
+	return false;
+}
+
+//======================================================================================================================
+static TArray<IBlendSpaceAnalysisFeature*> GetAnalysisFeatures()
+{
+	TArray<IBlendSpaceAnalysisFeature*> ModularFeatures = 
+		IModularFeatures::Get().GetModularFeatureImplementations<IBlendSpaceAnalysisFeature>(
+			IBlendSpaceAnalysisFeature::GetModuleFeatureName());
+
+	// Put the core one on last so that user-defined ones can override the default behaviour
+	ModularFeatures.Push(&CoreBlendSpaceAnalysisFeature);
+	return ModularFeatures;
+}
+
+//======================================================================================================================
+FVector FBlendSpaceAnalysis::CalculateSampleValue(const UBlendSpace& BlendSpace, const UAnimSequence& Animation,
+												  const float RateScale, const FVector& OriginalPosition, bool bAnalyzed[3])
+{
+
+	FVector AdjustedPosition = OriginalPosition;
+	TArray<IBlendSpaceAnalysisFeature*> ModularFeatures = GetAnalysisFeatures();
+	for (int32 Index = 0; Index != 2; ++Index)
+	{
+		bAnalyzed[Index] = false;
+		const UAnalysisProperties* AnalysisProperties = BlendSpace.AnalysisProperties[Index].Get();
+		for (const IBlendSpaceAnalysisFeature* Feature : ModularFeatures)
+		{
+			bAnalyzed[Index] = Feature->CalculateSampleValue(
+				AdjustedPosition[Index], BlendSpace, AnalysisProperties, Animation, RateScale);
+
+			if (bAnalyzed[Index])
+			{
+				break;
+			}
+		}
+	}
+	return AdjustedPosition;
+}
+
+//======================================================================================================================
+UAnalysisProperties* FBlendSpaceAnalysis::MakeAnalysisProperties(UObject* Outer, const FString& FunctionName)
+{
+	UAnalysisProperties* Result = nullptr;
+
+	TArray<IBlendSpaceAnalysisFeature*> ModularFeatures = GetAnalysisFeatures();
+	for (const IBlendSpaceAnalysisFeature* Feature : ModularFeatures)
+	{
+		Result = Feature->MakeAnalysisProperties(Outer, FunctionName);
+		if (Result)
+		{
+			return Result;
+		}
+	}
+	return Result;
+}
+
+//======================================================================================================================
+TArray<FString> FBlendSpaceAnalysis::GetAnalysisFunctions()
+{
+	TArray<FString> FunctionNames;
+	TArray<IBlendSpaceAnalysisFeature*> ModularFeatures = GetAnalysisFeatures();
+	for (const IBlendSpaceAnalysisFeature* Feature : ModularFeatures)
+	{
+		TArray<FString> FeatureFunctionNames = Feature->GetAnalysisFunctions();
+		for (const FString& FeatureFunctionName : FeatureFunctionNames)
+		{
+			FunctionNames.AddUnique(FeatureFunctionName);
+		}
+	}
+	return FunctionNames;
+}
+
+
+
+#undef LOCTEXT_NAMESPACE
diff --git a/Engine/Source/Editor/Persona/Private/PersonaBlendSpaceAnalysis.h b/Engine/Source/Editor/Persona/Private/PersonaBlendSpaceAnalysis.h
new file mode 100644
--- /dev/null
+++ b/Engine/Source/Editor/Persona/Private/PersonaBlendSpaceAnalysis.h
@@ -0,0 +1,22 @@
+// Copyright Epic Games, Inc. All Rights Reserved.
+#pragma once
+
+#include "BlendSpaceAnalysis.h"
+
+class FBlendSpaceAnalysis
+{
+public:
+	static FVector CalculateSampleValue(const UBlendSpace&   BlendSpace, 
+										const UAnimSequence& Animation, 
+										const float          RateScale, 
+										const FVector&       OriginalPosition, 
+										bool                 bAnalyzed[3]);
+
+
+	// This will return an instance derived from UAnalysisProperties that is suitable for the Function. The caller will
+	// pass in a suitable owning object, outer, that the implementation should assign as owner of the newly created object.
+	static UAnalysisProperties* MakeAnalysisProperties(UObject* Outer, const FString& FunctionName);
+
+	// This will return the names of the functions handled
+	static TArray<FString> GetAnalysisFunctions();
+};
diff --git a/Engine/Source/Editor/Persona/Private/SAnimationBlendSpace.cpp b/Engine/Source/Editor/Persona/Private/SAnimationBlendSpace.cpp
--- a/Engine/Source/Editor/Persona/Private/SAnimationBlendSpace.cpp
+++ b/Engine/Source/Editor/Persona/Private/SAnimationBlendSpace.cpp
@@ -1,6 +1,7 @@
 // Copyright Epic Games, Inc. All Rights Reserved.
 
 #include "SAnimationBlendSpace.h"
+
 #include "Widgets/Notifications/SNotificationList.h"
 #include "Editor.h"
 
@@ -19,6 +20,10 @@
 #include "Animation/AimOffsetBlendSpace.h"
 #include "Animation/AimOffsetBlendSpace1D.h"
 #include "SAnimationBlendSpaceGridWidget.h"
+#include "Animation/AnimSequenceHelpers.h"
+#include "Animation/AnimSequence.h"
+#include "Engine/SkeletalMeshSocket.h"
+#include "PersonaBlendSpaceAnalysis.h"
 
 #define LOCTEXT_NAMESPACE "BlendSpaceEditor"
 
@@ -32,6 +37,7 @@ void SBlendSpaceEditor::Construct(const FArguments& InArgs)
 	BlendSpace = InArgs._BlendSpace;
 
 	OnBlendSpaceSampleAdded = InArgs._OnBlendSpaceSampleAdded;
+	OnBlendSpaceSampleDuplicated = InArgs._OnBlendSpaceSampleDuplicated;
 	OnBlendSpaceSampleRemoved = InArgs._OnBlendSpaceSampleRemoved;
 	OnBlendSpaceSampleReplaced = InArgs._OnBlendSpaceSampleReplaced;
 	OnSetPreviewPosition = InArgs._OnSetPreviewPosition;
@@ -68,7 +74,7 @@ void SBlendSpaceEditor::Construct(const FArguments& InArgs)
 						+SVerticalBox::Slot()
 						.FillHeight(1)
 						[
-							SAssignNew(NewBlendSpaceGridWidget, SBlendSpaceGridWidget)
+							SAssignNew(BlendSpaceGridWidget, SBlendSpaceGridWidget)
 							.Cursor(EMouseCursor::Crosshairs)
 							.BlendSpaceBase(BlendSpace)
 							.NotifyHook(this)
@@ -77,6 +83,7 @@ void SBlendSpaceEditor::Construct(const FArguments& InArgs)
 							.OnSampleMoved(this, &SBlendSpaceEditor::OnSampleMoved)
 							.OnSampleRemoved(this, &SBlendSpaceEditor::OnSampleRemoved)
 							.OnSampleAdded(this, &SBlendSpaceEditor::OnSampleAdded)
+							.OnSampleDuplicated(this, &SBlendSpaceEditor::OnSampleDuplicated)
 							.OnSampleReplaced(this, &SBlendSpaceEditor::OnSampleReplaced)
 							.OnSampleDoubleClicked(InArgs._OnBlendSpaceSampleDoubleClicked)
 							.OnExtendSampleTooltip(InArgs._OnExtendSampleTooltip)
@@ -107,16 +114,17 @@ void SBlendSpaceEditor::Construct(const FArguments& InArgs, const TSharedRef<cla
 
 void SBlendSpaceEditor::OnSampleMoved(const int32 SampleIndex, const FVector& NewValue, bool bIsInteractive)
 {
-	bool bMoveSuccesful = true;
+	bool bMoveSuccessful = true;
 	if (BlendSpace->IsValidBlendSampleIndex(SampleIndex) && BlendSpace->GetBlendSample(SampleIndex).SampleValue != NewValue && !BlendSpace->IsTooCloseToExistingSamplePoint(NewValue, SampleIndex))
 	{
 		FScopedTransaction ScopedTransaction(LOCTEXT("MoveSample", "Moving Blend Grid Sample"));
-		BlendSpace->Modify();
 
-		bMoveSuccesful = BlendSpace->EditSampleValue(SampleIndex, NewValue);
-		if (bMoveSuccesful)
+		bMoveSuccessful = BlendSpace->EditSampleValue(SampleIndex, NewValue);
+		if (bMoveSuccessful)
 		{
+			BlendSpace->Modify();
 			BlendSpace->ValidateSampleData();
+			BlendSpace->PostEditChange();
 			ResampleData();
 		}
 	}
@@ -125,65 +133,105 @@ void SBlendSpaceEditor::OnSampleMoved(const int32 SampleIndex, const FVector& Ne
 void SBlendSpaceEditor::OnSampleRemoved(const int32 SampleIndex)
 {
 	FScopedTransaction ScopedTransaction(LOCTEXT("RemoveSample", "Removing Blend Grid Sample"));
-	BlendSpace->Modify();
 
-	const bool bRemoveSuccesful = BlendSpace->DeleteSample(SampleIndex);
-	if (bRemoveSuccesful)
+	const bool bRemoveSuccessful = BlendSpace->DeleteSample(SampleIndex);
+	if (bRemoveSuccessful)
 	{
+		BlendSpace->Modify();
 		ResampleData();
 		BlendSpace->ValidateSampleData();
 
+		BlendSpaceGridWidget->InvalidateCachedData();
+		BlendSpaceGridWidget->InvalidateState();
+
 		OnBlendSpaceSampleRemoved.ExecuteIfBound(SampleIndex);
+		BlendSpace->PostEditChange();
 	}
-	BlendSpace->PostEditChange();
 }
 
-void SBlendSpaceEditor::OnSampleAdded(UAnimSequence* Animation, const FVector& Value)
+//======================================================================================================================
+static bool GetLockAfterAnalysis(const TObjectPtr<UAnalysisProperties>& AnalysisProperties)
+{
+	return AnalysisProperties ? AnalysisProperties->bLockAfterAnalysis : false;
+}
+
+//======================================================================================================================
+int32 SBlendSpaceEditor::OnSampleAdded(UAnimSequence* Animation, const FVector& Value)
 {
 	FScopedTransaction ScopedTransaction(LOCTEXT("AddSample", "Adding Blend Grid Sample"));
-	BlendSpace->Modify();
 
-	bool bAddSuccesful = false;
+	FVector AdjustedValue = Value;
+	bool bAnalyzed[3] = { false, false, false };
+
+	if(BlendSpace->IsAsset())
+	{
+		AdjustedValue = FBlendSpaceAnalysis::CalculateSampleValue(*BlendSpace, *Animation, 1.0f, Value, bAnalyzed);
+	}
+		
+	int32 NewSampleIndex = -1;
 
 	if(BlendSpace->IsAsset())
 	{
-		bAddSuccesful = BlendSpace->AddSample(Animation, Value);
+		NewSampleIndex = BlendSpace->AddSample(Animation, AdjustedValue);
 	}
 	else
 	{
-		bAddSuccesful = BlendSpace->AddSample(Value);
+		NewSampleIndex = BlendSpace->AddSample(AdjustedValue);
 	}
 
-	if (bAddSuccesful)
+	if (NewSampleIndex >= 0)
 	{
+		BlendSpace->Modify();
+		BlendSpace->LockSample(NewSampleIndex,
+							   bAnalyzed[0] ? GetLockAfterAnalysis(BlendSpace->AnalysisProperties[0]) : false,
+							   bAnalyzed[1] ? GetLockAfterAnalysis(BlendSpace->AnalysisProperties[1]) : false,
+							   bAnalyzed[2] ? GetLockAfterAnalysis(BlendSpace->AnalysisProperties[2]) : false);
+
 		ResampleData();
 		BlendSpace->ValidateSampleData();
 
-		OnBlendSpaceSampleAdded.ExecuteIfBound(Animation, Value);
+		BlendSpaceGridWidget->InvalidateCachedData();
+		BlendSpaceGridWidget->InvalidateState();
+
+		if (OnBlendSpaceSampleAdded.IsBound())
+		{
+			OnBlendSpaceSampleAdded.Execute(Animation, AdjustedValue);
+		}
+		BlendSpace->PostEditChange();
+	}
+
+	return NewSampleIndex;
+}
+
+void SBlendSpaceEditor::OnSampleDuplicated(const int32 SampleIndex, const FVector& NewValue)
+{
+	const FBlendSample& OrigSample = BlendSpace->GetBlendSample(SampleIndex);
+	int32 NewSampleIndex = OnSampleAdded(OrigSample.Animation, NewValue);
+	if (NewSampleIndex >= 0)
+	{
+		BlendSpace->LockSample(NewSampleIndex, OrigSample.bLockX, OrigSample.bLockY, OrigSample.bLockZ);
 	}
-	BlendSpace->PostEditChange();
 }
 
 void SBlendSpaceEditor::OnSampleReplaced(int32 InSampleIndex, UAnimSequence* Animation)
 {
 	FScopedTransaction ScopedTransaction(LOCTEXT("UpdateAnimation", "Changing Animation Sequence"));
-	BlendSpace->Modify();
 
-	bool bUpdateSuccesful = false;
+	bool bUpdateSuccessful = false;
 	if(BlendSpace->IsAsset())
 	{
-		bUpdateSuccesful = BlendSpace->ReplaceSampleAnimation(InSampleIndex, Animation);
+		bUpdateSuccessful = BlendSpace->ReplaceSampleAnimation(InSampleIndex, Animation);
 	}
 	else
 	{
-		bUpdateSuccesful = true;
+		bUpdateSuccessful = true;
 	}
 
-	if (bUpdateSuccesful)
+	if (bUpdateSuccessful)
 	{
+		BlendSpace->Modify();
 		ResampleData();
 		BlendSpace->ValidateSampleData();
-
 		OnBlendSpaceSampleReplaced.ExecuteIfBound(InSampleIndex, Animation);
 	}
 }
@@ -195,10 +243,10 @@ void SBlendSpaceEditor::PostUndoRedo()
 	ResampleData();
 
 	// Invalidate widget data
-	NewBlendSpaceGridWidget->InvalidateCachedData();
+	BlendSpaceGridWidget->InvalidateCachedData();
 
 	// Invalidate sample indices used for UI info
-	NewBlendSpaceGridWidget->InvalidateState();
+	BlendSpaceGridWidget->InvalidateState();
 
 	// Set flag which will update the preview value in the next tick (this due the recreation of data after Undo)
 	bShouldSetPreviewPosition = true;
@@ -219,7 +267,7 @@ void SBlendSpaceEditor::UpdatePreviewParameter() const
 		{
 			if (Component->PreviewInstance->GetCurrentAsset() == BlendSpace)
 			{
-				const FVector PreviewPosition = NewBlendSpaceGridWidget->GetPreviewPosition();
+				const FVector PreviewPosition = BlendSpaceGridWidget->GetPreviewPosition();
 				Component->PreviewInstance->SetBlendSpacePosition(PreviewPosition);
 				GetPreviewScene()->InvalidateViews();			
 			}
@@ -227,7 +275,7 @@ void SBlendSpaceEditor::UpdatePreviewParameter() const
 	}
 	else if(OnSetPreviewPosition.IsBound())
 	{
-		const FVector PreviewPosition = NewBlendSpaceGridWidget->GetPreviewPosition();
+		const FVector PreviewPosition = BlendSpaceGridWidget->GetPreviewPosition();
 		OnSetPreviewPosition.Execute(PreviewPosition);
 	}
 }
@@ -245,7 +293,7 @@ void SBlendSpaceEditor::UpdateFromBlendSpaceState() const
 				FVector FilteredPosition;
 				FVector Position;
 				Component->PreviewInstance->GetBlendSpaceState(Position, FilteredPosition);
-				NewBlendSpaceGridWidget->SetPreviewingState(Position, FilteredPosition);
+				BlendSpaceGridWidget->SetPreviewingState(Position, FilteredPosition);
 			}
 		}
 	}
@@ -254,7 +302,7 @@ void SBlendSpaceEditor::UpdateFromBlendSpaceState() const
 void SBlendSpaceEditor::Tick(const FGeometry& AllottedGeometry, const double InCurrentTime, const float InDeltaTime)
 {
 	// Update the preview as long as its enabled
-	if (NewBlendSpaceGridWidget->IsPreviewing() || bShouldSetPreviewPosition)
+	if (BlendSpaceGridWidget->IsPreviewing() || bShouldSetPreviewPosition)
 	{
 		UpdatePreviewParameter();
 		bShouldSetPreviewPosition = false;
@@ -269,7 +317,7 @@ void SBlendSpaceEditor::OnPropertyChanged(UObject* ObjectBeingModified, FPropert
 	{
 		BlendSpace->ValidateSampleData();
 		ResampleData();
-		NewBlendSpaceGridWidget->InvalidateCachedData();
+		BlendSpaceGridWidget->InvalidateCachedData();
 	}
 }
 
diff --git a/Engine/Source/Editor/Persona/Private/SAnimationBlendSpace.h b/Engine/Source/Editor/Persona/Private/SAnimationBlendSpace.h
--- a/Engine/Source/Editor/Persona/Private/SAnimationBlendSpace.h
+++ b/Engine/Source/Editor/Persona/Private/SAnimationBlendSpace.h
@@ -24,31 +24,19 @@ public:
 		, _DisplayScrubBar(true)
 		, _StatusBarName(TEXT("AssetEditor.AnimationEditor.MainMenu"))
 		{}
-		
 	SLATE_ARGUMENT(UBlendSpace*, BlendSpace)
-
 	SLATE_ARGUMENT(bool, DisplayScrubBar)
-
 	SLATE_EVENT(FOnBlendSpaceSampleDoubleClicked, OnBlendSpaceSampleDoubleClicked)
-
 	SLATE_EVENT(FOnBlendSpaceSampleAdded, OnBlendSpaceSampleAdded)
-
+	SLATE_EVENT(FOnBlendSpaceSampleDuplicated, OnBlendSpaceSampleDuplicated)
 	SLATE_EVENT(FOnBlendSpaceSampleRemoved, OnBlendSpaceSampleRemoved)
-
 	SLATE_EVENT(FOnBlendSpaceSampleReplaced, OnBlendSpaceSampleReplaced)
-
 	SLATE_EVENT(FOnGetBlendSpaceSampleName, OnGetBlendSpaceSampleName) 
-
 	SLATE_EVENT(FOnExtendBlendSpaceSampleTooltip, OnExtendSampleTooltip)
-
 	SLATE_EVENT(FOnSetBlendSpacePreviewPosition, OnSetPreviewPosition)
-
 	SLATE_ATTRIBUTE(FVector, PreviewPosition)
-
 	SLATE_ATTRIBUTE(FVector, PreviewFilteredPosition)
-
 	SLATE_ARGUMENT(FName, StatusBarName)
-
 	SLATE_END_ARGS()
 
 	~SBlendSpaceEditor();
@@ -69,9 +57,10 @@ public:
 protected:
 	void OnSampleMoved(const int32 SampleIndex, const FVector& NewValue, bool bIsInteractive);
 	void OnSampleRemoved(const int32 SampleIndex);
-	void OnSampleAdded(UAnimSequence* Animation, const FVector& Value);
+	int32 OnSampleAdded(UAnimSequence* Animation, const FVector& Value);
+	void OnSampleDuplicated(const int32 SampleIndex, const FVector& NewValue);
 	void OnSampleReplaced(const int32 SampleIndex, UAnimSequence* Animation);
-	
+
 	// Begin SAnimEditorBase overrides
 	virtual UAnimationAsset* GetEditorObject() const override { return BlendSpace; }
 	// End SAnimEditorBase overrides
@@ -105,7 +94,7 @@ protected:
 	TWeakPtr<class IPersonaPreviewScene> PreviewScenePtr;
 	
 	/** Pointer to the grid widget which displays the blendspace visualization */
-	TSharedPtr<class SBlendSpaceGridWidget> NewBlendSpaceGridWidget;	
+	TSharedPtr<class SBlendSpaceGridWidget> BlendSpaceGridWidget;	
 
 	// Property changed delegate
 	FCoreUObjectDelegates::FOnObjectPropertyChanged::FDelegate OnPropertyChangedHandle;	
@@ -116,6 +105,9 @@ protected:
 	/** Delegate called when a sample is added */
 	FOnBlendSpaceSampleAdded OnBlendSpaceSampleAdded;
 
+	/** Delegate called when a sample is duplicated */
+	FOnBlendSpaceSampleDuplicated OnBlendSpaceSampleDuplicated;
+
 	/** Delegate called when a sample is removed */
 	FOnBlendSpaceSampleRemoved OnBlendSpaceSampleRemoved;
 
@@ -126,5 +118,5 @@ protected:
 	FOnSetBlendSpacePreviewPosition OnSetPreviewPosition;
 
 	/** Flag to check whether or not the preview value should be (re-)set on the next tick */
-	bool bShouldSetPreviewPosition;
+	bool bShouldSetPreviewPosition = true;
 };
diff --git a/Engine/Source/Editor/Persona/Private/SAnimationBlendSpaceGridWidget.cpp b/Engine/Source/Editor/Persona/Private/SAnimationBlendSpaceGridWidget.cpp
--- a/Engine/Source/Editor/Persona/Private/SAnimationBlendSpaceGridWidget.cpp
+++ b/Engine/Source/Editor/Persona/Private/SAnimationBlendSpaceGridWidget.cpp
@@ -44,10 +44,11 @@ void SBlendSpaceGridWidget::Construct(const FArguments& InArgs)
 {
 	BlendSpaceBase = InArgs._BlendSpaceBase;
 	PreviousBlendSpaceBase = BlendSpaceBase.Get();
-	Position = InArgs._Position;
+	TargetPosition = InArgs._Position;
 	FilteredPosition = InArgs._FilteredPosition;
 	NotifyHook = InArgs._NotifyHook;
 	OnSampleAdded = InArgs._OnSampleAdded;
+	OnSampleDuplicated = InArgs._OnSampleDuplicated;
 	OnSampleMoved = InArgs._OnSampleMoved;
 	OnSampleRemoved = InArgs._OnSampleRemoved;
 	OnSampleReplaced = InArgs._OnSampleReplaced;
@@ -138,14 +139,14 @@ void SBlendSpaceGridWidget::Construct(const FArguments& InArgs)
 					.HAlign(HAlign_Left)
 					.BorderImage(FEditorStyle::GetBrush("NoBorder"))
 					.DesiredSizeScale(FVector2D(1.0f, 1.0f))
-					.Padding_Lambda([&]() { return FMargin(GridMargin.Left + 6, GridMargin.Top + 6, 0, 0) + GridRatioMargin; })
+					.Padding_Lambda([&]() { return FMargin(GridMargin.Left + 6, 0, 0, 0) + GridRatioMargin; })
 					[
 						SNew(SVerticalBox)
 						+ SVerticalBox::Slot()
 						.AutoHeight()
 						[
 							SNew(SHorizontalBox)
-							+ SHorizontalBox::Slot()
+							+ SHorizontalBox::Slot() // Button to show triangulation
 							.AutoWidth()
 							[
 								SNew(SBorder)
@@ -166,7 +167,7 @@ void SBlendSpaceGridWidget::Construct(const FArguments& InArgs)
 								]
 							]
 	
-							+ SHorizontalBox::Slot()
+							+ SHorizontalBox::Slot() // Button to toggle labels
 							.AutoWidth()
 							[
 								SNew(SBorder)
@@ -187,7 +188,7 @@ void SBlendSpaceGridWidget::Construct(const FArguments& InArgs)
 								]
 							]
 
-							+ SHorizontalBox::Slot()
+							+ SHorizontalBox::Slot() // Button to fit or stretch the graph
 							.AutoWidth()
 							[
 								SNew(SBorder)
@@ -208,7 +209,7 @@ void SBlendSpaceGridWidget::Construct(const FArguments& InArgs)
 								]
 							]
 
-							+ SHorizontalBox::Slot()
+							+ SHorizontalBox::Slot() // Sample X value input
 							.AutoWidth()
 							[
 								SNew(SBorder)
@@ -220,7 +221,7 @@ void SBlendSpaceGridWidget::Construct(const FArguments& InArgs)
 								]
 							]
 	
-							+ SHorizontalBox::Slot()
+							+ SHorizontalBox::Slot() // Sample Y value input
 							.AutoWidth()
 							[
 								SNew(SBorder)
@@ -233,7 +234,7 @@ void SBlendSpaceGridWidget::Construct(const FArguments& InArgs)
 							]
 						]
 						
-						+ SVerticalBox::Slot()
+						+ SVerticalBox::Slot() // Tip for dragging in, when there are no samples
 						.AutoHeight()
 						.Padding(FMargin(2.0f, 3.0f, 0.0f, 0.0f ))
 						[
@@ -244,7 +245,7 @@ void SBlendSpaceGridWidget::Construct(const FArguments& InArgs)
 							.Visibility(TAttribute<EVisibility>::Create(TAttribute<EVisibility>::FGetter::CreateSP(this, &SBlendSpaceGridWidget::GetSampleToolTipVisibility)))
 						]
 
-						+ SVerticalBox::Slot()
+						+ SVerticalBox::Slot() // Tip for adjusting the preview point
 						.AutoHeight()
 						.Padding(FMargin(2.0f, 3.0f, 0.0f, 0.0f))
 						[
@@ -287,7 +288,7 @@ void SBlendSpaceGridWidget::Construct(const FArguments& InArgs)
 		]
 	];
 
-	if(Position.IsSet())
+	if(TargetPosition.IsSet())
 	{
 		StartPreviewing();
 	}
@@ -440,6 +441,42 @@ void SBlendSpaceGridWidget::PaintSampleKeys(
 			FSlateDrawElement::MakeBox(
 				OutDrawElements, SampleLayer, AllottedGeometry.ToPaintGeometry(GridPosition, KeySize), 
 				KeyBrush, ESlateDrawEffect::None, DrawColor );
+
+			// Draw lines/boxes to indicated which axes are locked
+			if (!bReadOnly)
+			{
+				TArray<FVector2D> LinePoints;
+				LinePoints.SetNumUninitialized(2);
+				const float LockSizeScale = 0.5f; // Have the markers just poke out
+				const FVector2D SamplePosition = SampleValueToScreenPosition(Sample.SampleValue) + FVector2D(0.0, 0.0);
+				if (Sample.bLockX)
+				{
+					LinePoints[0] = FVector2D(SamplePosition - FVector2D(KeySize.X, -KeySize.Y) * LockSizeScale);
+					LinePoints[1] = FVector2D(SamplePosition - FVector2D(KeySize.X, +KeySize.Y) * LockSizeScale);
+					FSlateDrawElement::MakeLines(
+						OutDrawElements, SampleLayer - 1, AllottedGeometry.ToPaintGeometry(),
+						LinePoints, ESlateDrawEffect::None, InvalidColor.GetSpecifiedColor(), true, 1.0f);
+					LinePoints[0] = FVector2D(SamplePosition + FVector2D(KeySize.X, -KeySize.Y) * LockSizeScale);
+					LinePoints[1] = FVector2D(SamplePosition + FVector2D(KeySize.X, +KeySize.Y) * LockSizeScale);
+					FSlateDrawElement::MakeLines(
+						OutDrawElements, SampleLayer - 1, AllottedGeometry.ToPaintGeometry(),
+						LinePoints, ESlateDrawEffect::None, InvalidColor.GetSpecifiedColor(), true, 1.0f);
+				}
+				const bool b1DBlendSpace = BlendSpace->IsA<UBlendSpace1D>();
+				if (Sample.bLockY && !BlendSpace->IsA<UBlendSpace1D>())
+				{
+					LinePoints[0] = FVector2D(SamplePosition - FVector2D(-KeySize.X, KeySize.X) * LockSizeScale);
+					LinePoints[1] = FVector2D(SamplePosition - FVector2D(+KeySize.X, KeySize.X) * LockSizeScale);
+					FSlateDrawElement::MakeLines(
+						OutDrawElements, SampleLayer - 1, AllottedGeometry.ToPaintGeometry(),
+						LinePoints, ESlateDrawEffect::None, InvalidColor.GetSpecifiedColor(), true, 1.0f);
+					LinePoints[0] = FVector2D(SamplePosition + FVector2D(-KeySize.X, KeySize.X) * LockSizeScale);
+					LinePoints[1] = FVector2D(SamplePosition + FVector2D(+KeySize.X, KeySize.X) * LockSizeScale);
+					FSlateDrawElement::MakeLines(
+						OutDrawElements, SampleLayer - 1, AllottedGeometry.ToPaintGeometry(),
+						LinePoints, ESlateDrawEffect::None, InvalidColor.GetSpecifiedColor(), true, 1.0f);
+				}
+			}
 		}
 
 		// Always draw the filtered position which comes back from whatever is running
@@ -514,27 +551,37 @@ void SBlendSpaceGridWidget::PaintAxisText(const FGeometry& AllottedGeometry, con
 	FVector2D ArrowSize = ArrowBrushes[(uint8)EArrowDirection::Left]->GetImageSize();
 	FVector2D TextPosition = FVector2D(GridCenter.X - (TextSize.X * .5f), CachedGridRectangle.Bottom + TextMargin + (ArrowSize.Y * .25f));
 	FVector2D ArrowPosition = FVector2D(TextPosition.X - ArrowSize.X - 10.f/* give padding*/, TextPosition.Y);
-	FSlateDrawElement::MakeBox(OutDrawElements, DrawLayerId + 1, AllottedGeometry.ToPaintGeometry(ArrowPosition, ArrowSize), ArrowBrushes[(uint8)EArrowDirection::Left], ESlateDrawEffect::None, FLinearColor::White);
+	FSlateDrawElement::MakeBox(
+		OutDrawElements, DrawLayerId + 1, AllottedGeometry.ToPaintGeometry(ArrowPosition, ArrowSize), 
+		ArrowBrushes[(uint8)EArrowDirection::Left], ESlateDrawEffect::None, FLinearColor::White);
 
 	// Label
-	FSlateDrawElement::MakeText(OutDrawElements, DrawLayerId + 1, AllottedGeometry.MakeChild(TextPosition, FVector2D(1.0f, 1.0f)).ToPaintGeometry(), Text, FontInfo, ESlateDrawEffect::None, FLinearColor::White);
+	FSlateDrawElement::MakeText(
+		OutDrawElements, DrawLayerId + 1, AllottedGeometry.MakeChild(TextPosition, FVector2D(1.0f, 1.0f)).ToPaintGeometry(), 
+		Text, FontInfo, ESlateDrawEffect::None, FLinearColor::White);
 
 	// arrow right
 	ArrowSize = ArrowBrushes[(uint8)EArrowDirection::Right]->GetImageSize();
 	ArrowPosition = FVector2D(TextPosition.X + TextSize.X + 10.f/* give padding*/, TextPosition.Y);
-	FSlateDrawElement::MakeBox(OutDrawElements, DrawLayerId + 1, AllottedGeometry.ToPaintGeometry(ArrowPosition, ArrowSize), ArrowBrushes[(uint8)EArrowDirection::Right], ESlateDrawEffect::None, FLinearColor::White);
+	FSlateDrawElement::MakeBox(
+		OutDrawElements, DrawLayerId + 1, AllottedGeometry.ToPaintGeometry(ArrowPosition, ArrowSize), 
+		ArrowBrushes[(uint8)EArrowDirection::Right], ESlateDrawEffect::None, FLinearColor::White);
 
 	Text = FString::SanitizeFloat(SampleValueMin.X);
 	TextSize = FontMeasure->Measure(Text, FontInfo);
 
 	// Minimum value
-	FSlateDrawElement::MakeText(OutDrawElements, DrawLayerId + 1, AllottedGeometry.MakeChild(FVector2D(CachedGridRectangle.Left - (TextSize.X * .5f), CachedGridRectangle.Bottom + TextMargin + (TextSize.Y * .25f)), FVector2D(1.0f, 1.0f)).ToPaintGeometry(), Text, FontInfo, ESlateDrawEffect::None, FLinearColor::White);
+	FSlateDrawElement::MakeText(OutDrawElements, DrawLayerId + 1, AllottedGeometry.MakeChild(FVector2D(
+		CachedGridRectangle.Left - (TextSize.X * .5f), CachedGridRectangle.Bottom + TextMargin + (TextSize.Y * .25f)), 
+		FVector2D(1.0f, 1.0f)).ToPaintGeometry(), Text, FontInfo, ESlateDrawEffect::None, FLinearColor::White);
 
 	Text = FString::SanitizeFloat(SampleValueMax.X);
 	TextSize = FontMeasure->Measure(Text, FontInfo);
 
 	// Maximum value
-	FSlateDrawElement::MakeText(OutDrawElements, DrawLayerId + 1, AllottedGeometry.MakeChild(FVector2D(CachedGridRectangle.Right - (TextSize.X * .5f), CachedGridRectangle.Bottom + TextMargin + (TextSize.Y * .25f)), FVector2D(1.0f, 1.0f)).ToPaintGeometry(), Text, FontInfo, ESlateDrawEffect::None, FLinearColor::White);
+	FSlateDrawElement::MakeText(OutDrawElements, DrawLayerId + 1, AllottedGeometry.MakeChild(FVector2D(
+		CachedGridRectangle.Right - (TextSize.X * .5f), CachedGridRectangle.Bottom + TextMargin + (TextSize.Y * .25f)), 
+		FVector2D(1.0f, 1.0f)).ToPaintGeometry(), Text, FontInfo, ESlateDrawEffect::None, FLinearColor::White);
 
 	// Only draw Y axis labels if this is a 2D grid
 	if (GridType == EGridType::TwoAxis)
@@ -547,27 +594,39 @@ void SBlendSpaceGridWidget::PaintAxisText(const FGeometry& AllottedGeometry, con
 		ArrowSize = ArrowBrushes[(uint8)EArrowDirection::Up]->GetImageSize();
 		TextPosition = FVector2D(((GridMargin.Left - TextSize.X) * 0.5f - (ArrowSize.X * .25f)) + GridRatioMargin.Left, GridCenter.Y - (TextSize.Y * .5f));
 		ArrowPosition = FVector2D(TextPosition.X + TextSize.X * 0.5f - ArrowSize.X * 0.5f, TextPosition.Y - ArrowSize.Y - 10.f/* give padding*/);
-		FSlateDrawElement::MakeBox(OutDrawElements, DrawLayerId + 1, AllottedGeometry.ToPaintGeometry(ArrowPosition, ArrowSize), ArrowBrushes[(uint8)EArrowDirection::Up], ESlateDrawEffect::None, FLinearColor::White);
+		FSlateDrawElement::MakeBox(
+			OutDrawElements, DrawLayerId + 1, AllottedGeometry.ToPaintGeometry(ArrowPosition, ArrowSize), 
+			ArrowBrushes[(uint8)EArrowDirection::Up], ESlateDrawEffect::None, FLinearColor::White);
 
 		// Label
-		FSlateDrawElement::MakeText(OutDrawElements, DrawLayerId + 1, AllottedGeometry.MakeChild(TextPosition, FVector2D(1.0f, 1.0f)).ToPaintGeometry(), Text,	FontInfo, ESlateDrawEffect::None, FLinearColor::White);
+		FSlateDrawElement::MakeText(
+			OutDrawElements, DrawLayerId + 1, AllottedGeometry.MakeChild(TextPosition, FVector2D(1.0f, 1.0f)).ToPaintGeometry(),
+			Text,	FontInfo, ESlateDrawEffect::None, FLinearColor::White);
 
 		// arrow down
 		ArrowSize = ArrowBrushes[(uint8)EArrowDirection::Down]->GetImageSize();
 		ArrowPosition = FVector2D(TextPosition.X + TextSize.X * 0.5f - ArrowSize.X * 0.5f, TextPosition.Y + TextSize.Y + 10.f/* give padding*/);
-		FSlateDrawElement::MakeBox(OutDrawElements, DrawLayerId + 1, AllottedGeometry.ToPaintGeometry(ArrowPosition, ArrowSize), ArrowBrushes[(uint8)EArrowDirection::Down], ESlateDrawEffect::None, FLinearColor::White);
+		FSlateDrawElement::MakeBox(
+			OutDrawElements, DrawLayerId + 1, AllottedGeometry.ToPaintGeometry(ArrowPosition, ArrowSize), 
+			ArrowBrushes[(uint8)EArrowDirection::Down], ESlateDrawEffect::None, FLinearColor::White);
 
 		Text = FString::SanitizeFloat(SampleValueMin.Y);
 		TextSize = FontMeasure->Measure(Text, FontInfo);
 
 		// Minimum value
-		FSlateDrawElement::MakeText(OutDrawElements, DrawLayerId + 1, AllottedGeometry.MakeChild(FVector2D(((GridMargin.Left - TextSize.X) * 0.5f - (TextSize.X * .25f)) + GridRatioMargin.Left, CachedGridRectangle.Bottom - (TextSize.Y * .5f)), FVector2D(1.0f, 1.0f)).ToPaintGeometry(), Text, FontInfo, ESlateDrawEffect::None, FLinearColor::White);
+		FSlateDrawElement::MakeText(OutDrawElements, DrawLayerId + 1, AllottedGeometry.MakeChild(FVector2D(
+			((GridMargin.Left - TextSize.X) * 0.5f - (TextSize.X * .25f)) + GridRatioMargin.Left, 
+			CachedGridRectangle.Bottom - (TextSize.Y * .5f)), FVector2D(1.0f, 1.0f)).ToPaintGeometry(), 
+			Text, FontInfo, ESlateDrawEffect::None, FLinearColor::White);
 
 		Text = FString::SanitizeFloat(SampleValueMax.Y);
 		TextSize = FontMeasure->Measure(Text, FontInfo);
 
 		// Maximum value
-		FSlateDrawElement::MakeText(OutDrawElements, DrawLayerId + 1, AllottedGeometry.MakeChild(FVector2D(((GridMargin.Left - TextSize.X) * 0.5f - (TextSize.X * .25f) ) + GridRatioMargin.Left, ( GridMargin.Top + GridRatioMargin.Top) - (TextSize.Y * .5f)), FVector2D(1.0f, 1.0f)).ToPaintGeometry(), Text, FontInfo, ESlateDrawEffect::None, FLinearColor::White);
+		FSlateDrawElement::MakeText(OutDrawElements, DrawLayerId + 1, AllottedGeometry.MakeChild(FVector2D(
+			((GridMargin.Left - TextSize.X) * 0.5f - (TextSize.X * .25f) ) + GridRatioMargin.Left, 
+			CachedGridRectangle.Top - (TextSize.Y * .5f)), 
+			FVector2D(1.0f, 1.0f)).ToPaintGeometry(), Text, FontInfo, ESlateDrawEffect::None, FLinearColor::White);
 	}
 
 	DrawLayerId += 1;
@@ -691,16 +750,18 @@ void SBlendSpaceGridWidget::PaintTriangulation(const FGeometry& AllottedGeometry
 			{
 				for (int32 SourceIndex = 0; SourceIndex < 3; ++SourceIndex)
 				{
-					if (Element.Indices[SourceIndex] != INDEX_NONE)
+					const int32 SourceSampleIndex = Element.Indices[SourceIndex]; 
+					if (Samples.IsValidIndex(SourceSampleIndex))
 					{
-						const FBlendSample& SourceSample = Samples[Element.Indices[SourceIndex]];
+						const FBlendSample& SourceSample = Samples[SourceSampleIndex];
 						for (int32 TargetIndex = 0; TargetIndex < 3; ++TargetIndex)
 						{
-							if (Element.Indices[TargetIndex] != INDEX_NONE)
+							const int32 TargetSampleIndex = Element.Indices[TargetIndex]; 
+							if (Samples.IsValidIndex(TargetSampleIndex))
 							{
 								if (TargetIndex != SourceIndex)
 								{
-									const FBlendSample& TargetSample = Samples[Element.Indices[TargetIndex]];
+									const FBlendSample& TargetSample = Samples[TargetSampleIndex];
 									TArray<FVector2D> Points;
 
 									Points.Add(SampleValueToScreenPosition(SourceSample.SampleValue));
@@ -724,64 +785,101 @@ void SBlendSpaceGridWidget::PaintTriangulation(const FGeometry& AllottedGeometry
 			for (const FBlendSpaceSegment& Segment : BlendSpaceData.Segments)
 			{
 				int32 SampleIndex = Segment.SampleIndices[0];
-				int32 SampleIndexNext = Segment.SampleIndices[1];
+				int32 SampleIndex1 = Segment.SampleIndices[1];
 				TArray<FVector2D> Points;
 				Points.Add(SampleValueToScreenPosition(Samples[SampleIndex].SampleValue));
-				Points.Add(SampleValueToScreenPosition(Samples[SampleIndexNext].SampleValue));
+				Points.Add(SampleValueToScreenPosition(Samples[SampleIndex1].SampleValue));
 				FSlateDrawElement::MakeLines(
 					OutDrawElements, DrawLayerId + 1, AllottedGeometry.ToPaintGeometry(), Points,
 					ESlateDrawEffect::None, TriangulationColor.GetSpecifiedColor(), true, 0.5f);
 			}
 
+			// Flag any triangle that has an interior angle smaller than this
+			const float CriticalAngle = 4.0f;
+			// Flag any triangle that has a smaller area than this (normalized units)
+			const float CriticalArea = 5e-4f;
+			const float CriticalDot = FMath::Cos(FMath::DegreesToRadians(CriticalAngle));
 			for (const FBlendSpaceTriangle& Triangle : BlendSpaceData.Triangles)
 			{
+				FLinearColor TriangleFillColor = TriangulationCurrentColor.GetSpecifiedColor();
+				FLinearColor TriangleLineColor = TriangulationColor.GetSpecifiedColor();
+				TriangleFillColor.A = 0.03f; // Alpha for tinting the triangulation background
+				int32 TriangleLayer = DrawLayerId + 1;
+
+				FVector2D ScreenPositions[3] = {
+					SampleValueToScreenPosition(Samples[Triangle.SampleIndices[0]].SampleValue),
+					SampleValueToScreenPosition(Samples[Triangle.SampleIndices[1]].SampleValue),
+					SampleValueToScreenPosition(Samples[Triangle.SampleIndices[2]].SampleValue),
+				};
+
+				// If we just have one triangle it's OK for it to be degenerate
+				if (BlendSpaceData.Triangles.Num() > 1)
+				{
+					FVector2D NormalizedPositions[3] = {
+						SampleValueToNormalizedPosition(Samples[Triangle.SampleIndices[0]].SampleValue),
+						SampleValueToNormalizedPosition(Samples[Triangle.SampleIndices[1]].SampleValue),
+						SampleValueToNormalizedPosition(Samples[Triangle.SampleIndices[2]].SampleValue),
+					};
+
+					for (int32 Index = 0; Index != 3; ++Index)
+					{
+						FVector2D A = NormalizedPositions[(Index + 2) % 3] - NormalizedPositions[(Index + 1) % 3];
+						FVector2D B = NormalizedPositions[Index] - NormalizedPositions[(Index + 1) % 3];
+						float Dot = A.GetSafeNormal() | B.GetSafeNormal();
+						float Area = 0.5f * FMath::Abs(B ^ A);
+						if (Dot > CriticalDot || Area < CriticalArea)
+						{
+							TriangleFillColor = FLinearColor::Red;
+							TriangleFillColor.A = 0.5f;
+							TriangleLineColor = FLinearColor::Red;
+							TriangleLayer = DrawLayerId + 10; // just bump it up so it definitely shows
+						}
+					}
+				}
+
 				FVector2D MidPoint(0, 0);
 				PolygonPoints.Empty(3);
-				for (int32 Index = 0; Index != FBlendSpaceTriangle::NUM_VERTICES; ++Index)
+				for (int32 Index0 = 0; Index0 != FBlendSpaceTriangle::NUM_VERTICES; ++Index0)
 				{
-					int32 IndexNext = (Index + 1) % FBlendSpaceTriangle::NUM_VERTICES;
-					int32 SampleIndex = Triangle.SampleIndices[Index];
-					int32 SampleIndexNext = Triangle.SampleIndices[IndexNext];
-					FVector2D P1 = SampleValueToScreenPosition(Samples[SampleIndex].SampleValue);
-					FVector2D P2 = SampleValueToScreenPosition(Samples[SampleIndexNext].SampleValue);
-					TArray<FVector2D> Points = { P1, P2 };
-					MidPoint += SampleValueToScreenPosition(Samples[SampleIndex].SampleValue) / 3.0f;
+					int32 Index1 = (Index0 + 1) % FBlendSpaceTriangle::NUM_VERTICES;
+					int32 Index2 = (Index0 + 2) % FBlendSpaceTriangle::NUM_VERTICES;
+					int32 SampleIndex0 = Triangle.SampleIndices[Index0];
+					int32 SampleIndex1 = Triangle.SampleIndices[Index1];
+					int32 SampleIndex2 = Triangle.SampleIndices[Index2];
+
+					TArray<FVector2D> Points = { ScreenPositions[Index0], ScreenPositions[Index1] };
+					MidPoint += SampleValueToScreenPosition(Samples[SampleIndex0].SampleValue) / 3.0f;
 
 					FSlateDrawElement::MakeLines(
-						OutDrawElements, DrawLayerId + 1, AllottedGeometry.ToPaintGeometry(), Points,
-						ESlateDrawEffect::None, TriangulationColor.GetSpecifiedColor(), true, 0.5f);
-					PolygonPoints.Push(P1);
+						OutDrawElements, TriangleLayer, AllottedGeometry.ToPaintGeometry(), Points,
+						ESlateDrawEffect::None, TriangleLineColor, true, 0.5f);
+					PolygonPoints.Push(ScreenPositions[Index0]);
 				}
 
-				FLinearColor Color = TriangulationCurrentColor.GetSpecifiedColor();
-				Color.A = 0.03f; // Alpha for tinting the triangulation background
-				PaintTriangle(PolygonPoints[0], PolygonPoints[1], PolygonPoints[2], AllottedGeometry, Color, LabelBrush, OutDrawElements, DrawLayerId);
+				PaintTriangle(PolygonPoints[0], PolygonPoints[1], PolygonPoints[2], AllottedGeometry, 
+							  TriangleFillColor, LabelBrush, OutDrawElements, DrawLayerId);
 
 #ifdef DEBUG_BLENDSPACE_TRIANGULATION
 				// Draw the adjacent triangle indices around the perimeter
 				if (!bSamplePreviewing)
 				{
-					for (int32 Index = 0; Index != FBlendSpaceTriangle::NUM_VERTICES; ++Index)
+					for (int32 Index0 = 0; Index0 != FBlendSpaceTriangle::NUM_VERTICES; ++Index0)
 					{
-						if (Triangle.EdgeInfo[Index].NeighbourTriangleIndex < 0)
+						if (Triangle.EdgeInfo[Index0].NeighbourTriangleIndex < 0)
 						{
-							int32 IndexNext = (Index + 1) % FBlendSpaceTriangle::NUM_VERTICES;
-							int32 SampleIndex = Triangle.SampleIndices[Index];
-							int32 SampleIndexNext = Triangle.SampleIndices[IndexNext];
-							FVector2D Point0 = SampleValueToScreenPosition(Samples[SampleIndex].SampleValue);
-							FVector2D Point1 = SampleValueToScreenPosition(Samples[SampleIndexNext].SampleValue);
-							FVector2D MidEdge = (Point0 + Point1) * 0.5;
+							int32 Index1 = (Index0 + 1) % FBlendSpaceTriangle::NUM_VERTICES;
+							FVector2D MidEdge = ( ScreenPositions[Index0] + ScreenPositions[Index1]) * 0.5;
 							float PullInAmount = 0.2;
 							FSlateDrawElement::MakeText(
 								OutDrawElements, DrawLayerId + 1, AllottedGeometry.MakeChild(
-									FMath::Lerp(Point0, MidEdge, PullInAmount), FVector2D(1.0f, 1.0f)).ToPaintGeometry(),
-								FText::AsNumber(Triangle.EdgeInfo[Index].AdjacentPerimeterTriangleIndices[0]),
+									FMath::Lerp(ScreenPositions[Index0], MidEdge, PullInAmount), FVector2D(1.0f, 1.0f)).ToPaintGeometry(),
+								FText::AsNumber(Triangle.EdgeInfo[Index0].AdjacentPerimeterTriangleIndices[0]),
 								FontInfo, ESlateDrawEffect::None, FLinearColor::Red);
 
 							FSlateDrawElement::MakeText(
 								OutDrawElements, DrawLayerId + 1, AllottedGeometry.MakeChild(
-									FMath::Lerp(Point1, MidEdge, PullInAmount), FVector2D(1.0f, 1.0f)).ToPaintGeometry(),
-								FText::AsNumber(Triangle.EdgeInfo[Index].AdjacentPerimeterTriangleIndices[1]),
+									FMath::Lerp(ScreenPositions[Index1], MidEdge, PullInAmount), FVector2D(1.0f, 1.0f)).ToPaintGeometry(),
+								FText::AsNumber(Triangle.EdgeInfo[Index0].AdjacentPerimeterTriangleIndices[1]),
 								FontInfo, ESlateDrawEffect::None, FLinearColor::Red);
 						}
 					}
@@ -884,12 +982,23 @@ FReply SBlendSpaceGridWidget::OnDrop(const FGeometry& MyGeometry, const FDragDro
 		{
 			if (DragState == EDragState::DragDrop)
 			{
-				const FVector SampleValue = ScreenPositionToSampleValueWithSnapping(LocalMousePosition, FSlateApplication::Get().GetModifierKeys().IsShiftDown());
 				TSharedPtr<FAssetDragDropOp> DragDropOperation = DragDropEvent.GetOperationAs<FAssetDragDropOp>();
 				if (DragDropOperation.IsValid())
 				{
-					UAnimSequence* Animation = FAssetData::GetFirstAsset<UAnimSequence>(DragDropOperation->GetAssets());
-					OnSampleAdded.ExecuteIfBound(Animation, SampleValue);
+					const FVector SampleValue = ScreenPositionToSampleValueWithSnapping(
+						LocalMousePosition, FSlateApplication::Get().GetModifierKeys().IsShiftDown());
+
+					const TArray<FAssetData>& Assets = DragDropOperation->GetAssets();
+					for (int Index = 0 ; Index != Assets.Num() ; ++Index)
+					{
+						const FAssetData& AssetData = Assets[Index];
+						UObject* Asset = AssetData.GetAsset();
+						UAnimSequence* Animation = (UAnimSequence*) Asset;
+						if (OnSampleAdded.IsBound())
+						{
+							OnSampleAdded.Execute(Animation, SampleValue);
+						}
+					}
 				}	
 			}
 			else if (DragState == EDragState::DragDropOverride)
@@ -1085,7 +1194,7 @@ FReply SBlendSpaceGridWidget::OnMouseMove(const FGeometry& MyGeometry, const FPo
 				bAdvancedPreview = MouseEvent.IsAltDown();
 				return FReply::Handled();
 			}
-			else if(Position.IsSet())
+			else if(TargetPosition.IsSet())
 			{
 				StartPreviewing();
 				DragState = EDragState::None;
@@ -1136,9 +1245,11 @@ FReply SBlendSpaceGridWidget::ProcessClick(const FGeometry& MyGeometry, const FP
 				{
 					const FWidgetPath WidgetPath = MouseEvent.GetEventPath() != nullptr ? *MouseEvent.GetEventPath() : FWidgetPath();
 					const FVector2D MousePosition = MouseEvent.GetScreenSpacePosition();
-					// This is of a fixed size atm since MenuContent->GetDesiredSize() will not take the detail customization into account and return an incorrect (small) size
-					const FVector2D ExpectedSize(300, 100);				
-					const FVector2D MenuPosition = FSlateApplication::Get().CalculatePopupWindowPosition(FSlateRect(MousePosition.X, MousePosition.Y, MousePosition.X, MousePosition.Y), ExpectedSize, false);
+					// This is of a fixed size atm since MenuContent->GetDesiredSize() will not take the detail 
+					// customization into account and return an incorrect (small) size
+					const FVector2D ExpectedSize(300, 300);
+					const FVector2D MenuPosition = FSlateApplication::Get().CalculatePopupWindowPosition(
+						FSlateRect(MousePosition.X, MousePosition.Y, MousePosition.X, MousePosition.Y), ExpectedSize, false);
 
 					FSlateApplication::Get().PushMenu(
 						AsShared(),
@@ -1335,7 +1446,9 @@ TSharedPtr<SWidget> SBlendSpaceGridWidget::CreateBlendSampleContextMenu()
 		if(const UBlendSpace* BlendSpace = BlendSpaceBase.Get())
 		{
 			const FBlendSample& Sample = BlendSpace->GetBlendSample(HighlightedSampleIndex);		
-			StructureDetailsView->GetDetailsView()->SetGenericLayoutDetailsDelegate(FOnGetDetailCustomizationInstance::CreateStatic(&FBlendSampleDetails::MakeInstance, BlendSpace, this, HighlightedSampleIndex));
+			StructureDetailsView->GetDetailsView()->SetGenericLayoutDetailsDelegate(
+				FOnGetDetailCustomizationInstance::CreateStatic(
+					&FBlendSampleDetails::MakeInstance, BlendSpace, this, HighlightedSampleIndex));
 
 			FStructOnScope* Struct = new FStructOnScope(FBlendSample::StaticStruct(), (uint8*)&Sample);
 			Struct->SetPackage(BlendSpace->GetOutermost());
@@ -1382,7 +1495,14 @@ TSharedPtr<SWidget> SBlendSpaceGridWidget::CreateNewBlendSampleContextMenu(const
 					LOCTEXT("AddNewSampleTooltip", "Add a new sample to this blendspace at this location"),
 						FSlateIcon("EditorStyle", "Plus"),
 						FUIAction(
-							FExecuteAction::CreateLambda([this, NewSampleValue](){ OnSampleAdded.ExecuteIfBound(nullptr, NewSampleValue); })
+							FExecuteAction::CreateLambda(
+								[this, NewSampleValue]()
+								{
+									if (OnSampleAdded.IsBound())
+									{
+										OnSampleAdded.Execute(nullptr, NewSampleValue);
+									}
+								})
 						)
 				);
 			}
@@ -1487,20 +1607,23 @@ int32 SBlendSpaceGridWidget::FindClosestGridPointIndexFromScreenPosition(const F
 
 	return GridPointIndex;
 }
-const FVector2D SBlendSpaceGridWidget::SampleValueToScreenPosition(const FVector& SampleValue) const
-{
-	const FVector2D GridSize = CachedGridRectangle.GetSize();	
-		
-	FVector2D SamplePosition2D;
-	// Convert the sample value to -1 to 1 form
-	SamplePosition2D.X = (((SampleValue.X - SampleValueMin.X) / SampleValueRange.X) * 2.0f) - 1.0f;
-	SamplePosition2D.Y = (GridType == EGridType::TwoAxis) ? (((SampleValueMax.Y - SampleValue.Y) / SampleValueRange.Y) * 2.0f) - 1.0f : 0.0f;
 
-	// Multiply by half of the grid size and offset is using the grid center position
-	SamplePosition2D *= CachedGridRectangle.GetSize()* 0.5f;
-	SamplePosition2D += CachedGridRectangle.GetCenter();
+const FVector2D SBlendSpaceGridWidget::SampleValueToNormalizedPosition(const FVector& SampleValue) const
+{
+	// Convert the sample value to 0 to 1 form
+	FVector2D NormalizedPosition(
+		((SampleValue.X - SampleValueMin.X) / SampleValueRange.X),
+		GridType == EGridType::TwoAxis ? ((SampleValueMax.Y - SampleValue.Y) / SampleValueRange.Y) : 0.5f);
+	return NormalizedPosition;	
+}
 
-	return SamplePosition2D;	
+const FVector2D SBlendSpaceGridWidget::SampleValueToScreenPosition(const FVector& SampleValue) const
+{
+	FVector2D NormalizedPosition = SampleValueToNormalizedPosition(SampleValue);
+	const FVector2D GridSize = CachedGridRectangle.GetSize();
+	const FVector2D GridCorner = CachedGridRectangle.GetCenter() - 0.5f * GridSize;
+	FVector2D ScreenPosition = GridCorner + NormalizedPosition * CachedGridRectangle.GetSize();
+	return ScreenPosition;	
 }
 
 const FVector SBlendSpaceGridWidget::ScreenPositionToSampleValue(const FVector2D& ScreenPosition, bool bClamp) const
@@ -1529,7 +1652,8 @@ const FVector SBlendSpaceGridWidget::ScreenPositionToSampleValue(const FVector2D
 
 const FSlateRect SBlendSpaceGridWidget::GetGridRectangleFromGeometry(const FGeometry& MyGeometry)
 {
-	FSlateRect WindowRect = FSlateRect(0, 0, MyGeometry.GetLocalSize().X, MyGeometry.GetLocalSize().Y);
+	const float TopOffset = 20.0f; // Ideally we'd get the size of the buttons (showing the label/triangulation etc)
+	FSlateRect WindowRect = FSlateRect(0, TopOffset, MyGeometry.GetLocalSize().X, MyGeometry.GetLocalSize().Y);
 	if (!bStretchToFit)
 	{
 		UpdateGridRatioMargin(WindowRect.GetSize());
@@ -1609,7 +1733,7 @@ FText SBlendSpaceGridWidget::GetToolTipAnimationName() const
 						const FBlendSample& BlendSample = BlendSpace->GetBlendSample(HighlightedSampleIndex);
 						ToolTipText = GetSampleName(BlendSample, HighlightedSampleIndex);
 					}
-					else if(Position.IsSet())
+					else if(TargetPosition.IsSet())
 					{
 						ToolTipText = PreviewValue;
 					}
@@ -1730,7 +1854,7 @@ FText SBlendSpaceGridWidget::GetToolTipSampleValue() const
 							ToolTipText = GetSampleErrorMessage(BlendSample);
 						}
 					}
-					else if(Position.IsSet())
+					else if(TargetPosition.IsSet())
 					{
 						ToolTipText = FText::Format(ValueFormattingText, ParameterXName, FText::FromString(FString::SanitizeFloat(PreviewPosition.X)), ParameterYName, FText::FromString(FString::SanitizeFloat(PreviewPosition.Y)));
 
@@ -2169,6 +2293,11 @@ void SBlendSpaceGridWidget::Tick(const FGeometry& AllottedGeometry, const double
 				// If we are dragging a sample, find out whether or not it has actually moved to a
 				// different grid position since the last tick and update the blend space accordingly
 				FVector SampleValue = ScreenPositionToSampleValueWithSnapping(LocalMousePosition, FSlateApplication::Get().GetModifierKeys().IsShiftDown());
+				// Only allow dragging on each axis if not locked
+				const FBlendSample& BlendSample = BlendSpace->GetBlendSample(DraggedSampleIndex);
+				SampleValue.X = BlendSample.bLockX ? BlendSample.SampleValue.X : SampleValue.X;
+				SampleValue.Y = BlendSample.bLockY ? BlendSample.SampleValue.Y : SampleValue.Y;
+				SampleValue.Z = BlendSample.bLockZ ? BlendSample.SampleValue.Z : SampleValue.Z;
 				if (SampleValue != LastDragPosition)
 				{
 					LastDragPosition = SampleValue;
@@ -2224,9 +2353,9 @@ void SBlendSpaceGridWidget::Tick(const FGeometry& AllottedGeometry, const double
 			LastPreviewingMousePosition.Y = LocalMousePosition.Y;
 			FModifierKeysState ModifierKeyState = FSlateApplication::Get().GetModifierKeys();
 			bool bIsManualPreviewing = !bReadOnly && IsHovered() && bMouseIsOverGeometry && ModifierKeyState.IsControlDown();
-			if (Position.IsSet() && !bIsManualPreviewing)
+			if (TargetPosition.IsSet() && !bIsManualPreviewing)
 			{
-				PreviewPosition = Position.Get();
+				PreviewPosition = TargetPosition.Get();
 				if (bReadOnly)
 				{
 					// Happens when we are showing in the graph - don't want to render outside the valid region
@@ -2387,7 +2516,7 @@ const bool SBlendSpaceGridWidget::IsPreviewing() const
 { 
 	FModifierKeysState ModifierKeyState = FSlateApplication::Get().GetModifierKeys();
 	bool bIsManualPreviewing = !bReadOnly && IsHovered() && bMouseIsOverGeometry && ModifierKeyState.IsControlDown();
-	return (bSamplePreviewing && !Position.IsSet()) || (Position.IsSet() && bIsManualPreviewing);
+	return (bSamplePreviewing && !TargetPosition.IsSet()) || (TargetPosition.IsSet() && bIsManualPreviewing);
 }
 
 #undef LOCTEXT_NAMESPACE // "SAnimationBlendSpaceGridWidget"
diff --git a/Engine/Source/Editor/Persona/Private/SAnimationBlendSpaceGridWidget.h b/Engine/Source/Editor/Persona/Private/SAnimationBlendSpaceGridWidget.h
--- a/Engine/Source/Editor/Persona/Private/SAnimationBlendSpaceGridWidget.h
+++ b/Engine/Source/Editor/Persona/Private/SAnimationBlendSpaceGridWidget.h
@@ -49,6 +49,7 @@ public:
 		SLATE_EVENT(FOnSampleMoved, OnSampleMoved)
 		SLATE_EVENT(FOnBlendSpaceSampleRemoved, OnSampleRemoved)
 		SLATE_EVENT(FOnBlendSpaceSampleAdded, OnSampleAdded)
+		SLATE_EVENT(FOnBlendSpaceSampleDuplicated, OnSampleDuplicated)
 		SLATE_EVENT(FOnBlendSpaceSampleReplaced, OnSampleReplaced)
 		SLATE_EVENT(FOnBlendSpaceSampleDoubleClicked, OnSampleDoubleClicked)
 		SLATE_EVENT(FOnGetBlendSpaceSampleName, OnGetBlendSpaceSampleName)
@@ -182,7 +183,9 @@ protected:
 	int32 FindClosestGridPointIndexFromScreenPosition(const FVector2D& InPosition) const;
 	/** Converts the given sample value to a screen space position */
 	const FVector2D SampleValueToScreenPosition(const FVector& SampleValue) const;
-	/** Converst a screen space (grid) position to a sample value, with optional clamping */
+	/** Converts the given sample value to a normalized position (between 0 and 1) */
+	const FVector2D SampleValueToNormalizedPosition(const FVector& SampleValue) const;
+	/** Converts a screen space (grid) position to a sample value, with optional clamping */
 	const FVector ScreenPositionToSampleValue(const FVector2D& GridPosition, bool bClamp) const;
 	/** Returns the (calculated) grid rectangle given the supplied geometry, and calls UpdateGridRatioMargin */
 	const FSlateRect GetGridRectangleFromGeometry(const FGeometry& MyGeometry);
@@ -227,7 +230,7 @@ private:
 	/** Previous blendspace, cached here to allow us to update cached data when it changes. */
 	TWeakObjectPtr<const UBlendSpace> PreviousBlendSpaceBase;
 	/** Current BlendSpace params (driven from outside the widget) */
-	TAttribute<FVector> Position;
+	TAttribute<FVector> TargetPosition;
 	/** Current BlendSpace state - filtered version of Position (driven from outside the widget) */
 	TAttribute<FVector> FilteredPosition;
 	/** Notify hook (ptr to SAnimationBlendSpace instance), which is required for transacting FBlendSample object when edited using the context-menu/structure details panel */
@@ -307,6 +310,7 @@ private:
 
 	/** Delegates populated from SAnimationBlendSpace and used as callbacks */
 	FOnBlendSpaceSampleAdded OnSampleAdded;
+	FOnBlendSpaceSampleDuplicated OnSampleDuplicated;
 	FOnSampleMoved OnSampleMoved;
 	FOnBlendSpaceSampleRemoved OnSampleRemoved;
 	FOnBlendSpaceSampleReplaced OnSampleReplaced;
@@ -314,7 +318,7 @@ private:
 	FOnGetBlendSpaceSampleName OnGetBlendSpaceSampleName;
 	FOnExtendBlendSpaceSampleTooltip OnExtendSampleTooltip;
 
-	/** Thresshold values for hovering, click and dragging samples */
+	/** Threshold values for hovering, click and dragging samples */
 	float DragThreshold;
 	float ClickAndHighlightThreshold;
 
diff --git a/Engine/Source/Editor/Persona/Public/BlendSpaceAnalysis.h b/Engine/Source/Editor/Persona/Public/BlendSpaceAnalysis.h
new file mode 100644
--- /dev/null
+++ b/Engine/Source/Editor/Persona/Public/BlendSpaceAnalysis.h
@@ -0,0 +1,248 @@
+// Copyright Epic Games, Inc. All Rights Reserved.
+
+#pragma once
+
+#include "CoreMinimal.h"
+#include "Animation/BoneSocketReference.h"
+#include "Features/IModularFeature.h"
+#include "Animation/BlendSpace.h"
+
+#include "BlendSpaceAnalysis.generated.h"
+
+class UAnimSequence;
+class UBlendSpace;
+class UAnalysisProperties;
+
+/**
+* Users wishing to add their own analysis functions and structures should inherit from this, implement the virtual
+* functions, and register an instance with IModularFeatures. It may help to look at the implementation of 
+* FCoreBlendSpaceAnalysisFeature when doing this.
+*/
+class IBlendSpaceAnalysisFeature : public IModularFeature
+{
+public:
+	static FName GetModuleFeatureName() { return "BlendSpaceAnalysis"; }
+
+	// This should process the animation according to the analysis properties, or return false if that is not possible.
+	virtual bool CalculateSampleValue(float&                     Result,
+									  const UBlendSpace&         BlendSpace,
+									  const UAnalysisProperties* AnalysisProperties,
+									  const UAnimSequence&       Animation,
+									  const float                RateScale) const = 0;
+
+	// This should return an instance derived from UAnalysisProperties that is suitable for the Function. The caller
+	// will pass in a suitable owning object, outer, that the implementation should assign as owner of the newly created
+	// object. 
+	virtual UAnalysisProperties* MakeAnalysisProperties(UObject* Outer, const FString& FunctionName) const = 0;
+
+	// This should return the names of the functions handled
+	virtual TArray<FString> GetAnalysisFunctions() const = 0;
+};
+
+UENUM()
+enum class EAnalysisSpace : uint8
+{
+	World    UMETA(DisplayName = "World"),
+	Fixed    UMETA(DisplayName = "Fixed"),
+	Changing UMETA(DisplayName = "Changing"),
+	Moving   UMETA(DisplayName = "Moving"),
+};
+
+UENUM()
+enum class EAnalysisLinearAxis : uint8
+{
+	X  UMETA(DisplayName = "X"),
+	Y  UMETA(DisplayName = "Y"),
+	Z  UMETA(DisplayName = "Z"),
+};
+
+UENUM()
+enum class EAnalysisEulerAxis : uint8
+{
+	Roll   UMETA(DisplayName = "Roll"),
+	Pitch  UMETA(DisplayName = "Pitch"),
+	Yaw    UMETA(DisplayName = "Yaw"),
+};
+
+/**
+ * This will be used to preserve values as far as possible when switching between analysis functions, so it contains all
+ * the parameters used by the engine functions. User defined can inherit from this and add their own - then the
+ * user-defined MakeCache function should replace any base class cache that is passed in with their own.
+*/
+class FCachedAnalysisProperties
+{
+public:
+	EAnalysisLinearAxis  LinearFunctionAxis = EAnalysisLinearAxis::X;
+	EAnalysisEulerAxis   EulerFunctionAxis = EAnalysisEulerAxis::Pitch;
+	EAnalysisSpace       Space = EAnalysisSpace::World;
+	FBoneSocketTarget    SpaceBoneSocket;
+	EAnalysisLinearAxis  CharacterFacingAxis = EAnalysisLinearAxis::Y;
+	EAnalysisLinearAxis  CharacterUpAxis = EAnalysisLinearAxis::Z;
+	float                StartTimeFraction = 0.0f;
+	float                EndTimeFraction = 1.0f;
+	FBoneSocketTarget    BoneSocket1;
+	FBoneSocketTarget    BoneSocket2;
+	EAnalysisLinearAxis  BoneFacingAxis = EAnalysisLinearAxis::X;
+	EAnalysisLinearAxis  BoneRightAxis = EAnalysisLinearAxis::Y;
+	bool                 bLockAfterAnalysis = false;
+};
+
+UCLASS()
+class ULinearAnalysisProperties : public UAnalysisProperties
+{
+	GENERATED_BODY()
+public:
+	void InitializeFromCache(TSharedPtr<FCachedAnalysisProperties> Cache) override;
+	void MakeCache(TSharedPtr<FCachedAnalysisProperties>& Cache) const override;
+
+	/** Axis for the analysis function */
+	UPROPERTY(EditAnywhere, DisplayName = "Axis", Category = AnalysisProperties)
+	EAnalysisLinearAxis FunctionAxis = EAnalysisLinearAxis::X;
+
+	/**
+	* The space in which to perform the analysis. Fixed will use the analysis bone/socket at the first frame
+	* of the analysis time range. Changing will use the analysis bone/socket at the relevant frame during the
+	* analysis, but calculate velocities assuming that frame isn't moving. Moving will do the same but velocities
+	* as well as positions/rotations will be relative to this moving frame.
+	*/
+	UPROPERTY(EditAnywhere, Category = AnalysisProperties)
+	EAnalysisSpace Space = EAnalysisSpace::World;
+
+	/** Bone or socket that defines the analysis space (when it isn't World) */
+	UPROPERTY(EditAnywhere, DisplayName = "Analysis Space Bone/Socket", Category = AnalysisProperties, meta = (EditCondition = "Space != EAnalysisSpace::World"))
+	FBoneSocketTarget SpaceBoneSocket;
+
+	/** The bone or socket used for analysis */
+	UPROPERTY(EditAnywhere, DisplayName = "Bone/Socket", Category = AnalysisProperties)
+	FBoneSocketTarget BoneSocket;
+
+	/** Fraction through each animation at which analysis starts */
+	UPROPERTY(EditAnywhere, DisplayName = "Start time fraction", Category = AnalysisProperties, meta = (ClampMin = "0", ClampMax = "1"))
+	float StartTimeFraction = 0.0f;
+
+	/** Fraction through each animation at which analysis ends */
+	UPROPERTY(EditAnywhere, DisplayName = "End time fraction", Category = AnalysisProperties, meta = (ClampMin = "0", ClampMax = "1"))
+	float EndTimeFraction = 1.0f;
+};
+
+UCLASS()
+class UEulerAnalysisProperties : public UAnalysisProperties
+{
+	GENERATED_BODY()
+public:
+	void InitializeFromCache(TSharedPtr<FCachedAnalysisProperties> Cache) override;
+	void MakeCache(TSharedPtr<FCachedAnalysisProperties>& Cache) const override;
+
+	/** Axis for the analysis function */
+	UPROPERTY(EditAnywhere, DisplayName = "Axis", Category = AnalysisProperties)
+	EAnalysisEulerAxis FunctionAxis = EAnalysisEulerAxis::Pitch;
+
+	/**
+	* The space in which to perform the analysis. Fixed will use the analysis bone/socket at the first frame
+	* of the analysis time range. Changing will use the analysis bone/socket at the relevant frame during the
+	* analysis, but calculate velocities assuming that frame isn't moving. Moving will do the same but velocities
+	* as well as positions/rotations will be relative to this moving frame.
+	*/
+	UPROPERTY(EditAnywhere, Category = AnalysisProperties)
+	EAnalysisSpace Space = EAnalysisSpace::World;
+
+	/** Bone or socket that defines the analysis space (when it isn't World) */
+	UPROPERTY(EditAnywhere, DisplayName = "Analysis Space Bone/Socket", Category = AnalysisProperties, meta = (EditCondition = "Space != 0"))
+	FBoneSocketTarget SpaceBoneSocket;
+
+	/** World or bone/socket axis that specifies the character's facing direction */
+	UPROPERTY(EditAnywhere, Category = AnalysisProperties)
+	EAnalysisLinearAxis CharacterFacingAxis = EAnalysisLinearAxis::Y;
+
+	/** World or bone/socket axis that specifies the character's up direction */
+	UPROPERTY(EditAnywhere, Category = AnalysisProperties)
+	EAnalysisLinearAxis CharacterUpAxis = EAnalysisLinearAxis::Z;
+
+	/** Fraction through each animation at which analysis starts */
+	UPROPERTY(EditAnywhere, DisplayName = "Start time fraction", Category = AnalysisProperties, meta = (ClampMin = "0", ClampMax = "1"))
+	float StartTimeFraction = 0.0f;
+	/** Fraction through each animation at which analysis ends */
+
+	UPROPERTY(EditAnywhere, DisplayName = "End time fraction", Category = AnalysisProperties, meta = (ClampMin = "0", ClampMax = "1"))
+	float EndTimeFraction = 1.0f;
+
+	/** The bone or socket used for analysis */
+	UPROPERTY(EditAnywhere, DisplayName = "Bone/Socket", Category = AnalysisProperties)
+	FBoneSocketTarget BoneSocket;
+
+	/** Used for some analysis functions - specifies the bone/socket axis that points in the facing/forwards direction */
+	UPROPERTY(EditAnywhere, Category = AnalysisProperties)
+	EAnalysisLinearAxis BoneFacingAxis = EAnalysisLinearAxis::X;
+
+	/** Used for some analysis functions - specifies the bone/socket axis that points to the "right" */
+	UPROPERTY(EditAnywhere, Category = AnalysisProperties)
+	EAnalysisLinearAxis BoneRightAxis = EAnalysisLinearAxis::Y;
+};
+
+UCLASS()
+class UMovementAnalysisProperties : public UAnalysisProperties
+{
+	GENERATED_BODY()
+public:
+	void InitializeFromCache(TSharedPtr<FCachedAnalysisProperties> Cache) override;
+	void MakeCache(TSharedPtr<FCachedAnalysisProperties>& Cache) const override;
+
+	/**
+	* The space in which to perform the analysis. Fixed will use the analysis bone/socket at the first frame
+	* of the analysis time range. Changing will use the analysis bone/socket at the relevant frame during the
+	* analysis, but calculate velocities assuming that frame isn't moving. Moving will do the same but velocities
+	* as well as positions/rotations will be relative to this moving frame.
+	*/
+	UPROPERTY(EditAnywhere, Category = AnalysisProperties)
+	EAnalysisSpace Space = EAnalysisSpace::World;
+
+	/** Bone or socket that defines the analysis space (when it isn't World) */
+	UPROPERTY(EditAnywhere, DisplayName = "Analysis Space Bone/Socket", Category = AnalysisProperties, meta = (EditCondition = "Space != EAnalysisSpace::World"))
+	FBoneSocketTarget SpaceBoneSocket;
+
+	/** World or bone/socket axis that specifies the character's facing direction */
+	UPROPERTY(EditAnywhere, Category = AnalysisProperties)
+	EAnalysisLinearAxis CharacterFacingAxis = EAnalysisLinearAxis::Y;
+
+	/** World or bone/socket axis that specifies the character's up direction */
+	UPROPERTY(EditAnywhere, Category = AnalysisProperties)
+	EAnalysisLinearAxis CharacterUpAxis = EAnalysisLinearAxis::Z;
+
+	/** Fraction through each animation at which analysis starts */
+	UPROPERTY(EditAnywhere, DisplayName = "Start time fraction", Category = AnalysisProperties, meta = (ClampMin = "0", ClampMax = "1"))
+	float StartTimeFraction = 0.0f;
+
+	/** Fraction through each animation at which analysis ends */
+	UPROPERTY(EditAnywhere, DisplayName = "End time fraction", Category = AnalysisProperties, meta = (ClampMin = "0", ClampMax = "1"))
+	float EndTimeFraction = 1.0f;
+
+	/** The bone or socket used for analysis */
+	UPROPERTY(EditAnywhere, DisplayName = "Bone/Socket", Category = AnalysisProperties)
+	FBoneSocketTarget BoneSocket;
+};
+
+UCLASS()
+class ULocomotionAnalysisProperties : public UAnalysisProperties
+{
+	GENERATED_BODY()
+public:
+	void InitializeFromCache(TSharedPtr<FCachedAnalysisProperties> Cache) override;
+	void MakeCache(TSharedPtr<FCachedAnalysisProperties>& Cache) const override;
+
+	/** World or bone/socket axis that specifies the character's facing direction */
+	UPROPERTY(EditAnywhere, Category = AnalysisProperties)
+	EAnalysisLinearAxis CharacterFacingAxis = EAnalysisLinearAxis::Y;
+
+	/** World or bone/socket axis that specifies the character's up direction */
+	UPROPERTY(EditAnywhere, Category = AnalysisProperties)
+	EAnalysisLinearAxis CharacterUpAxis = EAnalysisLinearAxis::Z;
+
+	/** The primary bone or socket used for analysis */
+	UPROPERTY(EditAnywhere, DisplayName = "Bone/Socket 1", Category = AnalysisProperties)
+	FBoneSocketTarget PrimaryBoneSocket;
+
+	/** The secondary bone or socket used for analysis */
+	UPROPERTY(EditAnywhere, DisplayName = "Bone/Socket 2", Category = AnalysisProperties)
+	FBoneSocketTarget SecondaryBoneSocket;
+};
+
diff --git a/Engine/Source/Editor/Persona/Public/PersonaDelegates.h b/Engine/Source/Editor/Persona/Public/PersonaDelegates.h
--- a/Engine/Source/Editor/Persona/Public/PersonaDelegates.h
+++ b/Engine/Source/Editor/Persona/Public/PersonaDelegates.h
@@ -43,8 +43,11 @@ DECLARE_DELEGATE_OneParam(FOnBlendSpaceSampleDoubleClicked, int32 /*SampleIndex*
 // Called when a blendspace sample point is removed
 DECLARE_DELEGATE_OneParam(FOnBlendSpaceSampleRemoved, const int32 /*SampleIndex*/);
 
+// Called when a blendspace sample point is added. Returns the new SampleIndex
+DECLARE_DELEGATE_RetVal_TwoParams(int32, FOnBlendSpaceSampleAdded, UAnimSequence* /*InSequence*/, const FVector& /*InSamplePoint*/);
+
 // Called when a blendspace sample point is added
-DECLARE_DELEGATE_TwoParams(FOnBlendSpaceSampleAdded, UAnimSequence* /*InSequence*/, const FVector& /*InSamplePoint*/);
+DECLARE_DELEGATE_TwoParams(FOnBlendSpaceSampleDuplicated, int32 /*SampleIndex*/, const FVector& /*InSamplePoint*/);
 
 // Called when a blendspace sample point is replaced
 DECLARE_DELEGATE_TwoParams(FOnBlendSpaceSampleReplaced, const int32 /*SampleIndex*/, UAnimSequence* /*InSequence*/);
diff --git a/Engine/Source/Runtime/AnimGraphRuntime/Private/BoneControllers/AnimNode_SkeletalControlBase.cpp b/Engine/Source/Runtime/AnimGraphRuntime/Private/BoneControllers/AnimNode_SkeletalControlBase.cpp
--- a/Engine/Source/Runtime/AnimGraphRuntime/Private/BoneControllers/AnimNode_SkeletalControlBase.cpp
+++ b/Engine/Source/Runtime/AnimGraphRuntime/Private/BoneControllers/AnimNode_SkeletalControlBase.cpp
@@ -4,47 +4,6 @@
 #include "Animation/AnimInstanceProxy.h"
 #include "Engine/SkeletalMeshSocket.h"
 
-////////////////////////////////////////////////////////////////////////////////////////////////////////////////////////
-// Socket Reference 
-////////////////////////////////////////////////////////////////////////////////////////////////////////////////////////
-
-void FSocketReference::InitializeSocketInfo(const FAnimInstanceProxy* InAnimInstanceProxy)
-{
-	CachedSocketMeshBoneIndex = INDEX_NONE;
-	CachedSocketCompactBoneIndex = FCompactPoseBoneIndex(INDEX_NONE);
-
-	if (SocketName != NAME_None)
-	{
-		const USkeletalMeshComponent* OwnerMeshComponent = InAnimInstanceProxy->GetSkelMeshComponent();
-		if (OwnerMeshComponent && OwnerMeshComponent->DoesSocketExist(SocketName))
-		{
-			USkeletalMeshSocket const* const Socket = OwnerMeshComponent->GetSocketByName(SocketName);
-			if (Socket)
-			{
-				CachedSocketLocalTransform = Socket->GetSocketLocalTransform();
-				// cache mesh bone index, so that we know this is valid information to follow
-				CachedSocketMeshBoneIndex = OwnerMeshComponent->GetBoneIndex(Socket->BoneName);
-
-				ensureMsgf(CachedSocketMeshBoneIndex != INDEX_NONE, TEXT("%s : socket has invalid bone."), *SocketName.ToString());
-			}
-		}
-		else
-		{
-			// @todo : move to graph node warning
-			UE_LOG(LogAnimation, Warning, TEXT("%s: socket doesn't exist"), *SocketName.ToString());
-		}
-	}
-}
-
-void FSocketReference::InitialzeCompactBoneIndex(const FBoneContainer& RequiredBones)
-{
-	if (CachedSocketMeshBoneIndex != INDEX_NONE)
-	{
-		const int32 SocketBoneSkeletonIndex = RequiredBones.GetPoseToSkeletonBoneIndexArray()[CachedSocketMeshBoneIndex];
-		CachedSocketCompactBoneIndex = RequiredBones.GetCompactPoseIndexFromSkeletonIndex(SocketBoneSkeletonIndex);
-	}
-}
-
 /////////////////////////////////////////////////////
 // FAnimNode_SkeletalControlBase
 
diff --git a/Engine/Source/Runtime/AnimGraphRuntime/Public/BoneControllers/AnimNode_SkeletalControlBase.h b/Engine/Source/Runtime/AnimGraphRuntime/Public/BoneControllers/AnimNode_SkeletalControlBase.h
--- a/Engine/Source/Runtime/AnimGraphRuntime/Public/BoneControllers/AnimNode_SkeletalControlBase.h
+++ b/Engine/Source/Runtime/AnimGraphRuntime/Public/BoneControllers/AnimNode_SkeletalControlBase.h
@@ -10,289 +10,13 @@
 #include "CoreMinimal.h"
 #include "UObject/ObjectMacros.h"
 #include "BonePose.h"
+#include "Animation/BoneSocketReference.h"
 #include "Animation/AnimNodeBase.h"
 #include "Animation/InputScaleBias.h"
 #include "AnimNode_SkeletalControlBase.generated.h"
 
 class USkeletalMeshComponent;
 
-USTRUCT()
-struct ANIMGRAPHRUNTIME_API FSocketReference
-{
-	GENERATED_USTRUCT_BODY()
-
-private:
-	FTransform CachedSocketLocalTransform;
-
-public:
-	/** Target socket to look at. Used if LookAtBone is empty. - You can use  LookAtLocation if you need offset from this point. That location will be used in their local space. **/
-	UPROPERTY(EditAnywhere, Category = FSocketReference)
-	FName SocketName;
-
-private:
-	int32 CachedSocketMeshBoneIndex;
-	FCompactPoseBoneIndex CachedSocketCompactBoneIndex;
-	
-public:
-	FSocketReference()
-		: CachedSocketMeshBoneIndex(INDEX_NONE)
-		, CachedSocketCompactBoneIndex(INDEX_NONE)
-	{
-	}
-
-	FSocketReference(const FName& InSocketName)
-		: SocketName(InSocketName)
-		, CachedSocketMeshBoneIndex(INDEX_NONE)
-		, CachedSocketCompactBoneIndex(INDEX_NONE)
-	{
-	}
-
-	void InitializeSocketInfo(const FAnimInstanceProxy* InAnimInstanceProxy);
-	void InitialzeCompactBoneIndex(const FBoneContainer& RequiredBones);
-	/* There are subtle difference between this two IsValid function
-	 * First one says the configuration had a valid socket as mesh index is valid
-	 * Second one says the current bonecontainer doesn't contain it, meaning the current LOD is missing the joint that is required to evaluate 
-	 * Although the expected behavior is ambiguous, I'll still split these two, and use it accordingly */
-	bool HasValidSetup() const
-	{
-		return (CachedSocketMeshBoneIndex != INDEX_NONE);
-	}
-
-	bool IsValidToEvaluate() const
-	{
-		return (CachedSocketCompactBoneIndex != INDEX_NONE);
-	}
-
-	FCompactPoseBoneIndex GetCachedSocketCompactBoneIndex() const
-	{
-		return CachedSocketCompactBoneIndex;
-	}
-
-	template<typename poseType>
-	FTransform GetAnimatedSocketTransform(struct FCSPose<poseType>& InPose) const
-	{
-		// current LOD has valid index (FCompactPoseBoneIndex is valid if current LOD supports)
-		if (CachedSocketCompactBoneIndex != INDEX_NONE)
-		{
-			FTransform BoneTransform = InPose.GetComponentSpaceTransform(CachedSocketCompactBoneIndex);
-			return CachedSocketLocalTransform * BoneTransform;
-		}
-
-		return FTransform::Identity;
-	}
-
-	void InvalidateCachedBoneIndex()
-	{
-		CachedSocketMeshBoneIndex = INDEX_NONE;
-		CachedSocketCompactBoneIndex = FCompactPoseBoneIndex(INDEX_NONE);
-	}
-};
-
-USTRUCT()
-struct ANIMGRAPHRUNTIME_API  FBoneSocketTarget
-{
-	GENERATED_USTRUCT_BODY()
-
-	UPROPERTY(EditAnywhere, Category = FBoneSocketTarget)
-	bool bUseSocket;
-
-	UPROPERTY(EditAnywhere, Category = FBoneSocketTarget, meta=(EditCondition = "!bUseSocket"))
-	FBoneReference BoneReference;
-
-	UPROPERTY(EditAnywhere, Category = FBoneSocketTarget, meta = (EditCondition = "bUseSocket"))
-	FSocketReference SocketReference;
-
-	FBoneSocketTarget(FName InName = NAME_None, bool bInUseSocket = false)
-	{
-		bUseSocket = bInUseSocket;
-
-		if (bUseSocket)
-		{
-			SocketReference.SocketName = InName;
-		}
-		else
-		{
-			BoneReference.BoneName = InName;
-		}
-	}
-
-	void InitializeBoneReferences(const FBoneContainer& RequiredBones)
-	{
-		if (bUseSocket)
-		{
-			SocketReference.InitialzeCompactBoneIndex(RequiredBones);
-			BoneReference.InvalidateCachedBoneIndex();
-		}
-		else
-		{
-			BoneReference.Initialize(RequiredBones);
-			SocketReference.InvalidateCachedBoneIndex();
-		}
-	}
-
-	/** Initialize Bone Reference, return TRUE if success, otherwise, return false **/
-	void Initialize(const FAnimInstanceProxy* InAnimInstanceProxy)
-	{
-		if (bUseSocket)
-		{
-			SocketReference.InitializeSocketInfo(InAnimInstanceProxy);
-		}
-	}
-
-
-	/** return true if valid. Otherwise return false **/
-	bool HasValidSetup() const
-	{
-		if (bUseSocket)
-		{
-			return SocketReference.HasValidSetup();
-		}
-
-		return BoneReference.BoneIndex != INDEX_NONE;
-	}
-
-	bool HasTargetSetup() const
-	{
-		if (bUseSocket)
-		{
-			return (SocketReference.SocketName != NAME_None);
-		}
-
-		return (BoneReference.BoneName != NAME_None);
-	}
-
-	FName GetTargetSetup() const
-	{
-		if (bUseSocket)
-		{
-			return (SocketReference.SocketName);
-		}
-
-		return (BoneReference.BoneName);
-	}
-
-	/** return true if valid. Otherwise return false **/
-	bool IsValidToEvaluate(const FBoneContainer& RequiredBones) const
-	{
-		if (bUseSocket)
-		{
-			return SocketReference.IsValidToEvaluate();
-		}
-
-		return BoneReference.IsValidToEvaluate(RequiredBones);
-	}
-
-	// this will return the compact pose bone index that matters
-	// if you're using sockeet, it will return socket's related joint's compact pose index
-	FCompactPoseBoneIndex GetCompactPoseBoneIndex() const 
-	{
-		if (bUseSocket)
-		{
-			return SocketReference.GetCachedSocketCompactBoneIndex();
-		}
-
-		return BoneReference.CachedCompactPoseIndex;
-	}
-
-	/** Get Target Transform from current incoming pose */
-	template<typename poseType>
-	FTransform GetTargetTransform(const FVector& TargetOffset, FCSPose<poseType>& InPose, const FTransform& InComponentToWorld) const
-	{
-		FTransform OutTargetTransform;
-
-		auto SetComponentSpaceOffset = [](const FVector& InTargetOffset, const FTransform& LocalInComponentToWorld, FTransform& LocalOutTargetTransform)
-		{
-			LocalOutTargetTransform.SetIdentity();
-			FVector CSTargetOffset = LocalInComponentToWorld.InverseTransformPosition(InTargetOffset);
-			LocalOutTargetTransform.SetLocation(CSTargetOffset);
-		};
-
-		if (bUseSocket)
-		{
-			// this has to be done outside
-			if (SocketReference.IsValidToEvaluate())
-			{
-				FTransform SocketTransformInCS = SocketReference.GetAnimatedSocketTransform(InPose);
-
-				FVector CSTargetOffset = SocketTransformInCS.TransformPosition(TargetOffset);
-				OutTargetTransform = SocketTransformInCS;
-				OutTargetTransform.SetLocation(CSTargetOffset);
-			}
-			else
-			{
-				// if none is found, we consider this offset is world offset
-				SetComponentSpaceOffset(TargetOffset, InComponentToWorld, OutTargetTransform);
-			}
-		}
-		// if valid data is available
-		else if (BoneReference.HasValidSetup())
-		{
-			if (BoneReference.IsValidToEvaluate() && 
-				ensureMsgf(InPose.GetPose().IsValidIndex(BoneReference.CachedCompactPoseIndex), TEXT("Invalid Cached Pose : Name %s(Bone Index (%d), Cached (%d))"), *BoneReference.BoneName.ToString(), BoneReference.BoneIndex, BoneReference.CachedCompactPoseIndex.GetInt()))
-			{
-				OutTargetTransform = InPose.GetComponentSpaceTransform(BoneReference.CachedCompactPoseIndex);
-				FVector CSTargetOffset = OutTargetTransform.TransformPosition(TargetOffset);
-				OutTargetTransform.SetLocation(CSTargetOffset);
-			}
-			else
-			{
-				// if none is found, we consider this offset is world offset
-				SetComponentSpaceOffset(TargetOffset, InComponentToWorld, OutTargetTransform);
-			}
-		}
-		else
-		{
-			// if none is found, we consider this offset is world offset
-			SetComponentSpaceOffset(TargetOffset, InComponentToWorld, OutTargetTransform);
-		}
-
-		return OutTargetTransform;
-	}
-
-	template<typename poseType>
-	FTransform GetTargetTransform(const FTransform& TargetOffset, FCSPose<poseType>& InPose, const FTransform& InComponentToWorld) const
-	{
-		FTransform OutTargetTransform;
-
-		auto SetComponentSpaceOffset = [](const FTransform& InTargetOffset, const FTransform& LocalInComponentToWorld, FTransform& LocalOutTargetTransform)
-		{
-			LocalOutTargetTransform = InTargetOffset.GetRelativeTransform(LocalInComponentToWorld);
-		};
-
-		if (bUseSocket)
-		{
-			// this has to be done outside
-			if (SocketReference.IsValidToEvaluate())
-			{
-				OutTargetTransform = TargetOffset * SocketReference.GetAnimatedSocketTransform(InPose);
-			}
-			else
-			{
-				SetComponentSpaceOffset(TargetOffset, InComponentToWorld, OutTargetTransform);
-			}
-		}
-		// if valid data is available
-		else if (BoneReference.HasValidSetup())
-		{
-			if (BoneReference.IsValidToEvaluate() && 
-				ensureMsgf(InPose.GetPose().IsValidIndex(BoneReference.CachedCompactPoseIndex), TEXT("Invalid Cached Pose : Name %s(Bone Index (%d), Cached (%d))"), *BoneReference.BoneName.ToString(), BoneReference.BoneIndex, BoneReference.CachedCompactPoseIndex.GetInt()))
-			{
-				OutTargetTransform = TargetOffset * InPose.GetComponentSpaceTransform(BoneReference.CachedCompactPoseIndex);
-			}
-			else
-			{
-				SetComponentSpaceOffset(TargetOffset, InComponentToWorld, OutTargetTransform);
-			}
-		}
-		else
-		{
-			SetComponentSpaceOffset(TargetOffset, InComponentToWorld, OutTargetTransform);
-		}
-
-		return OutTargetTransform;
-	}
-};
-
 USTRUCT(BlueprintInternalUseOnly)
 struct ANIMGRAPHRUNTIME_API FAnimNode_SkeletalControlBase : public FAnimNode_Base
 {
diff --git a/Engine/Source/Runtime/Engine/Classes/Animation/BlendSpace.h b/Engine/Source/Runtime/Engine/Classes/Animation/BlendSpace.h
--- a/Engine/Source/Runtime/Engine/Classes/Animation/BlendSpace.h
+++ b/Engine/Source/Runtime/Engine/Classes/Animation/BlendSpace.h
@@ -14,8 +14,34 @@
 #include "AnimationRuntime.h"
 #include "AnimNodeBase.h"
 #include "Containers/ArrayView.h"
+#include "Animation/BoneSocketReference.h"
 #include "BlendSpace.generated.h"
 
+class FCachedAnalysisProperties;
+
+/**
+* The base class for properties to be used in analysis. Engine will inherit from this to define structures used for
+* the functions it supports. User-defined functions will likely need their own analysis structures inheriting from
+* this too.
+*/
+UCLASS(config=Engine, MinimalAPI)
+class UAnalysisProperties : public UObject
+{
+	GENERATED_BODY()
+
+public:
+	ENGINE_API virtual void InitializeFromCache(TSharedPtr<FCachedAnalysisProperties> Cache);
+	ENGINE_API virtual void MakeCache(TSharedPtr<FCachedAnalysisProperties>& Cache) const;
+
+	/** Analysis function for this axis */
+	UPROPERTY()
+	FString Function = TEXT("None");
+
+	/** If set then the sample values will be locked after analysis, preventing them from being moved by accident */
+	UPROPERTY(EditAnywhere, Category = AnalysisProperties)
+	bool bLockAfterAnalysis = false;
+};
+
 /** Interpolation data types. */
 UENUM()
 enum EBlendSpaceAxis
@@ -33,6 +59,7 @@ enum class EPreferredTriangulationDirection : uint8
 	Radial UMETA(DisplayName = "Radial")
 };
 
+
 USTRUCT()
 struct FInterpolationParameter
 {
@@ -133,37 +160,49 @@ struct FBlendSample
 	FVector SampleValue;
 	
 	UPROPERTY(EditAnywhere, Category = BlendSample, meta=(UIMin="0.01", UIMax="2.0", ClampMin="0.01", ClampMax="64.0"))
-	float RateScale;
+	float RateScale = 1.0f;
 
 #if WITH_EDITORONLY_DATA
 	UPROPERTY(transient)
 	uint8 bIsValid : 1;
 
+	UPROPERTY(EditAnywhere, Category=BlendSample)
+	uint8 bLockX : 1;
+
+	UPROPERTY(EditAnywhere, Category=BlendSample)
+	uint8 bLockY : 1;
+
+	UPROPERTY(EditAnywhere, Category=BlendSample)
+	uint8 bLockZ : 1;
+
 	// Cache the samples marker data counter so that we can track if it changes and revalidate the blendspace
 	int32 CachedMarkerDataUpdateCounter;
 
 #endif // WITH_EDITORONLY_DATA
 
 	FBlendSample()
-		: Animation(nullptr)
-		, SampleValue(0.f)
-		, RateScale(1.0f)
+        : Animation(nullptr)
+        , SampleValue(0.f)
+        , RateScale(1.0f)
 #if WITH_EDITORONLY_DATA
-		, bIsValid(false)
+        , bIsValid(false)
+		, bLockX(false)
+		, bLockY(false)
+		, bLockZ(false)
 		, CachedMarkerDataUpdateCounter(INDEX_NONE)
 #endif // WITH_EDITORONLY_DATA
-	{		
-	}
+{		
+}
 	
 	FBlendSample(class UAnimSequence* InAnim, FVector InValue, bool bInIsSnapped, bool bInIsValid) 
-		: Animation(InAnim)
-		, SampleValue(InValue)
-		, RateScale(1.0f)
+        : Animation(InAnim)
+        , SampleValue(InValue)
+        , RateScale(1.0f)
 #if WITH_EDITORONLY_DATA
-		, bIsValid(bInIsValid)
+        , bIsValid(bInIsValid)
 		, CachedMarkerDataUpdateCounter(INDEX_NONE)
 #endif // WITH_EDITORONLY_DATA
-	{		
+	{		 
 	}
 	
 	bool operator==( const FBlendSample& Other ) const 
@@ -412,6 +451,7 @@ public:
 
 	/** Required for accessing protected variable names */
 	friend class FBlendSpaceDetails;
+	friend class FBlendSampleDetails;
 	friend class UAnimGraphNode_BlendSpaceGraphBase;
 
 	//~ Begin UObject Interface
@@ -516,9 +556,14 @@ public:
 	ENGINE_API void ValidateSampleData();
 
 	/** Add samples */
-	ENGINE_API bool AddSample(const FVector& SampleValue);
-	ENGINE_API bool	AddSample(UAnimSequence* AnimationSequence, const FVector& SampleValue);
+	ENGINE_API int32 AddSample(const FVector& SampleValue);
+	ENGINE_API int32 AddSample(UAnimSequence* AnimationSequence, const FVector& SampleValue);
+
+	/** if requested, locks the position of the sample (e.g. after analysis) */
+	ENGINE_API void LockSample(int32 SampleIndex, bool bLockX, bool bLockY, bool bLockZ);
 
+	ENGINE_API void ExpandRangeForSample(const FVector& SampleValue);
+	
 	/** edit samples */
 	ENGINE_API bool	EditSampleValue(const int32 BlendSampleIndex, const FVector& NewValue);
 
@@ -653,6 +698,15 @@ public:
 	UPROPERTY(EditAnywhere, Category = InputInterpolation)
 	FInterpolationParameter	InterpolationParam[3];
 
+#if WITH_EDITORONLY_DATA
+	/** Analysis properties for each axis. Note that these can be null. */
+	UPROPERTY(EditAnywhere, Category = AnalysisProperties)
+	TObjectPtr<UAnalysisProperties> AnalysisProperties[3];
+
+	/** Cached properties used to initialize properties when newly created. */
+	TSharedPtr<FCachedAnalysisProperties> CachedAnalysisProperties[3];
+#endif
+
 	/**
 	* If greater than zero, this is the speed at which the sample weights are allowed to change.
 	* 
@@ -677,15 +731,16 @@ public:
 	 */
 	UPROPERTY(EditAnywhere, Category = SampleSmoothing, meta = (DisplayName = "Smoothing"))
 	bool bTargetWeightInterpolationEaseInOut = true;
-	
+
+
 #if WITH_EDITORONLY_DATA
 	/** Preview Base pose for additive BlendSpace **/
 	UPROPERTY(EditAnywhere, Category = AdditiveSettings)
 	TObjectPtr<UAnimSequence> PreviewBasePose;
 #endif // WITH_EDITORONLY_DATA
 
-	/** This animation length changes based on current input (resulting in different blend time)**/
-	UPROPERTY(transient)
+	/** This is the maximum length of any sample in the blendspace. **/
+	UPROPERTY()
 	float AnimLength;
 
 	/** The current mode used by the BlendSpace to decide which animation notifies to fire. Valid options are:
@@ -736,7 +791,7 @@ protected:
 	UPROPERTY()
 	TArray<struct FEditorElement> GridSamples;
 
-	/** COntainer for the runtime data, which could be line segments, triangulation or tetrahedrons */
+	/** Container for the runtime data, which could be line segments, triangulation or tetrahedrons */
 	UPROPERTY()
 	FBlendSpaceData BlendSpaceData;
 	
diff --git a/Engine/Source/Runtime/Engine/Private/Animation/AnimSequenceHelpers.cpp b/Engine/Source/Runtime/Engine/Private/Animation/AnimSequenceHelpers.cpp
--- a/Engine/Source/Runtime/Engine/Private/Animation/AnimSequenceHelpers.cpp
+++ b/Engine/Source/Runtime/Engine/Private/Animation/AnimSequenceHelpers.cpp
@@ -381,7 +381,7 @@ void Retargeting::RetargetPose(FCompactPose& InOutPose, const FName& RetargetSou
 	const FBoneContainer& RequiredBones = InOutPose.GetBoneContainer();
 	const bool bDisableRetargeting = RequiredBones.GetDisableRetargeting();
 
-	if (!bDisableRetargeting)
+	if (!bDisableRetargeting && RetargetTransforms.Num())
 	{
 		const TArray<UE::Anim::FRetargetTracking>& RetargetTracking = UE::Anim::FBuildRawPoseScratchArea::Get().RetargetTracking;
 
diff --git a/Engine/Source/Runtime/Engine/Private/Animation/AnimationRuntime.cpp b/Engine/Source/Runtime/Engine/Private/Animation/AnimationRuntime.cpp
--- a/Engine/Source/Runtime/Engine/Private/Animation/AnimationRuntime.cpp
+++ b/Engine/Source/Runtime/Engine/Private/Animation/AnimationRuntime.cpp
@@ -2510,6 +2510,7 @@ void FAnimationRuntime::RetargetBoneTransform(const USkeleton* SourceSkeleton, c
 
 void FAnimationRuntime::RetargetBoneTransform(const USkeleton* SourceSkeleton, const FName& SourceName, const TArray<FTransform>& RetargetTransforms, FTransform& BoneTransform, const int32 SkeletonBoneIndex, const FCompactPoseBoneIndex& BoneIndex, const FBoneContainer& RequiredBones, const bool bIsBakedAdditive)
 {
+	check(!RetargetTransforms.IsEmpty());
 	if (SourceSkeleton)
 	{
 		const USkeleton* TargetSkeleton = RequiredBones.GetSkeletonAsset();
diff --git a/Engine/Source/Runtime/Engine/Private/Animation/BlendSpace.cpp b/Engine/Source/Runtime/Engine/Private/Animation/BlendSpace.cpp
--- a/Engine/Source/Runtime/Engine/Private/Animation/BlendSpace.cpp
+++ b/Engine/Source/Runtime/Engine/Private/Animation/BlendSpace.cpp
@@ -17,6 +17,9 @@
 #include "UObject/UObjectIterator.h"
 #include "Logging/TokenizedMessage.h"
 #include "Logging/MessageLog.h"
+#if WITH_EDITOR
+#include "BlendSpaceAnalysis.h"
+#endif
 
 #define LOCTEXT_NAMESPACE "BlendSpace"
 
@@ -33,6 +36,30 @@ struct FBlendSpaceScratchData : public TThreadSingleton<FBlendSpaceScratchData>
 	TArray<FGridBlendSample, TInlineAllocator<4> > RawGridSamples;
 };
 
+//======================================================================================================================
+void UAnalysisProperties::InitializeFromCache(TSharedPtr<FCachedAnalysisProperties> Cache)
+{
+#if WITH_EDITOR
+	if (Cache)
+	{
+		bLockAfterAnalysis = Cache->bLockAfterAnalysis;
+	}
+#endif
+}
+
+//======================================================================================================================
+void UAnalysisProperties::MakeCache(TSharedPtr<FCachedAnalysisProperties>& Cache) const
+{
+#if WITH_EDITOR
+	if (!Cache)
+	{
+		Cache = MakeShared<FCachedAnalysisProperties>();
+	}
+	Cache->bLockAfterAnalysis = bLockAfterAnalysis;
+#endif
+}
+
+//======================================================================================================================
 UBlendSpace::UBlendSpace(const FObjectInitializer& ObjectInitializer)
 	: Super(ObjectInitializer)
 {
@@ -1154,11 +1181,42 @@ void UBlendSpace::ValidateSampleData()
 	}
 }
 
-bool UBlendSpace::AddSample(const FVector& SampleValue)
+void UBlendSpace::ExpandRangeForSample(const FVector& SampleValue)
+{
+	for (int32 Index = 0 ; Index != 3 ; ++Index)
+	{
+		const float Delta = (BlendParameters[Index].Max - BlendParameters[Index].Min) / BlendParameters[Index].GridNum;
+		while (SampleValue[Index] > BlendParameters[Index].Max)
+		{
+			BlendParameters[Index].Max += Delta;
+			BlendParameters[Index].GridNum += 1;
+		}
+		while (SampleValue[Index] < BlendParameters[Index].Min)
+		{
+			BlendParameters[Index].Min -= Delta;
+			BlendParameters[Index].GridNum += 1;
+		}
+	}
+}
+
+void UBlendSpace::LockSample(int32 SampleIndex, bool bLockX, bool bLockY, bool bLockZ)
+{
+	if (SampleData.IsValidIndex(SampleIndex))
+	{
+		SampleData[SampleIndex].bLockX = bLockX; 
+		SampleData[SampleIndex].bLockY = bLockY; 
+		SampleData[SampleIndex].bLockZ = bLockZ; 
+	}
+}
+
+int32 UBlendSpace::AddSample(const FVector& SampleValue)
 {
 	// We should only be adding samples without a source animation if we are not a standalone asset
 	check(!IsAsset());
 
+	// Expand the range if necessary
+	ExpandRangeForSample(SampleValue);
+	
 	const bool bValidSampleData = ValidateSampleValue(SampleValue);
 
 	if (bValidSampleData)
@@ -1167,11 +1225,14 @@ bool UBlendSpace::AddSample(const FVector& SampleValue)
 		UpdatePreviewBasePose();
 	}
 
-	return bValidSampleData;
+	return bValidSampleData ? SampleData.Num() - 1 : -1;
 }
 
-bool UBlendSpace::AddSample(UAnimSequence* AnimationSequence, const FVector& SampleValue)
+int32 UBlendSpace::AddSample(UAnimSequence* AnimationSequence, const FVector& SampleValue)
 {
+	// Expand the range if necessary
+	ExpandRangeForSample(SampleValue);
+
 	const bool bValidSampleData = ValidateSampleValue(SampleValue) && ValidateAnimationSequence(AnimationSequence);
 
 	if (bValidSampleData)
@@ -1180,11 +1241,14 @@ bool UBlendSpace::AddSample(UAnimSequence* AnimationSequence, const FVector& Sam
 		UpdatePreviewBasePose();
 	}
 
-	return bValidSampleData;
+	return bValidSampleData ? SampleData.Num() - 1 : -1;
 }
 
 bool UBlendSpace::EditSampleValue(const int32 BlendSampleIndex, const FVector& NewValue)
 {
+	// Expand the range if necessary
+	ExpandRangeForSample(NewValue);
+
 	const bool bValidValue = SampleData.IsValidIndex(BlendSampleIndex) && ValidateSampleValue(NewValue, BlendSampleIndex);
 
 	if (bValidValue)
diff --git a/Engine/Source/Runtime/Engine/Private/Animation/BlendSpaceHelpers.cpp b/Engine/Source/Runtime/Engine/Private/Animation/BlendSpaceHelpers.cpp
--- a/Engine/Source/Runtime/Engine/Private/Animation/BlendSpaceHelpers.cpp
+++ b/Engine/Source/Runtime/Engine/Private/Animation/BlendSpaceHelpers.cpp
@@ -261,7 +261,7 @@ FDelaunayTriangleGenerator::ECircumCircleState FDelaunayTriangleGenerator::GetCi
 	else
 	{
 		// On top of the triangle edge
-		if (FMath::IsNearlyZero(Det, KINDA_SMALL_NUMBER))
+		if (FMath::IsNearlyZero(Det, SMALL_NUMBER))
 		{
 			return ECCS_On;
 		}
diff --git a/Engine/Source/Runtime/Engine/Private/Animation/BoneSocketReference.cpp b/Engine/Source/Runtime/Engine/Private/Animation/BoneSocketReference.cpp
new file mode 100644
--- /dev/null
+++ b/Engine/Source/Runtime/Engine/Private/Animation/BoneSocketReference.cpp
@@ -0,0 +1,47 @@
+// Copyright Epic Games, Inc. All Rights Reserved.
+
+#include "Animation/BoneSocketReference.h"
+#include "Animation/AnimInstanceProxy.h"
+#include "Engine/SkeletalMeshSocket.h"
+
+////////////////////////////////////////////////////////////////////////////////////////////////////////////////////////
+// Socket Reference 
+////////////////////////////////////////////////////////////////////////////////////////////////////////////////////////
+
+void FSocketReference::InitializeSocketInfo(const FAnimInstanceProxy* InAnimInstanceProxy)
+{
+	CachedSocketMeshBoneIndex = INDEX_NONE;
+	CachedSocketCompactBoneIndex = FCompactPoseBoneIndex(INDEX_NONE);
+
+	if (SocketName != NAME_None)
+	{
+		const USkeletalMeshComponent* OwnerMeshComponent = InAnimInstanceProxy->GetSkelMeshComponent();
+		if (OwnerMeshComponent && OwnerMeshComponent->DoesSocketExist(SocketName))
+		{
+			USkeletalMeshSocket const* const Socket = OwnerMeshComponent->GetSocketByName(SocketName);
+			if (Socket)
+			{
+				CachedSocketLocalTransform = Socket->GetSocketLocalTransform();
+				// cache mesh bone index, so that we know this is valid information to follow
+				CachedSocketMeshBoneIndex = OwnerMeshComponent->GetBoneIndex(Socket->BoneName);
+
+				ensureMsgf(CachedSocketMeshBoneIndex != INDEX_NONE, TEXT("%s : socket has invalid bone."), *SocketName.ToString());
+			}
+		}
+		else
+		{
+			// @todo : move to graph node warning
+			UE_LOG(LogAnimation, Warning, TEXT("%s: socket doesn't exist"), *SocketName.ToString());
+		}
+	}
+}
+
+void FSocketReference::InitialzeCompactBoneIndex(const FBoneContainer& RequiredBones)
+{
+	if (CachedSocketMeshBoneIndex != INDEX_NONE)
+	{
+		const int32 SocketBoneSkeletonIndex = RequiredBones.GetPoseToSkeletonBoneIndexArray()[CachedSocketMeshBoneIndex];
+		CachedSocketCompactBoneIndex = RequiredBones.GetCompactPoseIndexFromSkeletonIndex(SocketBoneSkeletonIndex);
+	}
+}
+
diff --git a/Engine/Source/Runtime/Engine/Public/Animation/BoneSocketReference.h b/Engine/Source/Runtime/Engine/Public/Animation/BoneSocketReference.h
new file mode 100644
--- /dev/null
+++ b/Engine/Source/Runtime/Engine/Public/Animation/BoneSocketReference.h
@@ -0,0 +1,296 @@
+// Copyright Epic Games, Inc. All Rights Reserved.
+
+/**
+ *	Abstract base class for a skeletal controller.
+ *	A SkelControl is a module that can modify the position or orientation of a set of bones in a skeletal mesh in some programmatic way.
+ */
+
+#pragma once
+
+#include "CoreMinimal.h"
+#include "UObject/ObjectMacros.h"
+#include "BoneIndices.h"
+#include "BonePose.h"
+#include "BoneContainer.h"
+#include "BoneSocketReference.generated.h"
+
+class USkeletalMeshComponent;
+struct FAnimInstanceProxy;
+
+USTRUCT()
+struct ENGINE_API FSocketReference
+{
+	GENERATED_USTRUCT_BODY()
+
+private:
+	FTransform CachedSocketLocalTransform;
+
+public:
+	/** Target socket to look at. Used if LookAtBone is empty. - You can use  LookAtLocation if you need offset from this point. That location will be used in their local space. **/
+	UPROPERTY(EditAnywhere, Category = FSocketReference)
+	FName SocketName;
+
+private:
+	int32 CachedSocketMeshBoneIndex;
+	FCompactPoseBoneIndex CachedSocketCompactBoneIndex;
+	
+public:
+	FSocketReference()
+		: CachedSocketMeshBoneIndex(INDEX_NONE)
+		, CachedSocketCompactBoneIndex(INDEX_NONE)
+	{
+	}
+
+	FSocketReference(const FName& InSocketName)
+		: SocketName(InSocketName)
+		, CachedSocketMeshBoneIndex(INDEX_NONE)
+		, CachedSocketCompactBoneIndex(INDEX_NONE)
+	{
+	}
+
+	void InitializeSocketInfo(const FAnimInstanceProxy* InAnimInstanceProxy);
+	void InitialzeCompactBoneIndex(const FBoneContainer& RequiredBones);
+	/* There are subtle difference between this two IsValid function
+	 * First one says the configuration had a valid socket as mesh index is valid
+	 * Second one says the current bonecontainer doesn't contain it, meaning the current LOD is missing the joint that is required to evaluate 
+	 * Although the expected behavior is ambiguous, I'll still split these two, and use it accordingly */
+	bool HasValidSetup() const
+	{
+		return (CachedSocketMeshBoneIndex != INDEX_NONE);
+	}
+
+	bool IsValidToEvaluate() const
+	{
+		return (CachedSocketCompactBoneIndex != INDEX_NONE);
+	}
+
+	FCompactPoseBoneIndex GetCachedSocketCompactBoneIndex() const
+	{
+		return CachedSocketCompactBoneIndex;
+	}
+
+	template<typename poseType>
+	FTransform GetAnimatedSocketTransform(struct FCSPose<poseType>& InPose) const
+	{
+		// current LOD has valid index (FCompactPoseBoneIndex is valid if current LOD supports)
+		if (CachedSocketCompactBoneIndex != INDEX_NONE)
+		{
+			FTransform BoneTransform = InPose.GetComponentSpaceTransform(CachedSocketCompactBoneIndex);
+			return CachedSocketLocalTransform * BoneTransform;
+		}
+
+		return FTransform::Identity;
+	}
+
+	void InvalidateCachedBoneIndex()
+	{
+		CachedSocketMeshBoneIndex = INDEX_NONE;
+		CachedSocketCompactBoneIndex = FCompactPoseBoneIndex(INDEX_NONE);
+	}
+};
+
+USTRUCT()
+struct ENGINE_API  FBoneSocketTarget
+{
+	GENERATED_USTRUCT_BODY()
+
+	UPROPERTY(EditAnywhere, Category = FBoneSocketTarget)
+	bool bUseSocket;
+
+	UPROPERTY(EditAnywhere, Category = FBoneSocketTarget, meta=(EditCondition = "!bUseSocket"))
+	FBoneReference BoneReference;
+
+	UPROPERTY(EditAnywhere, Category = FBoneSocketTarget, meta = (EditCondition = "bUseSocket"))
+	FSocketReference SocketReference;
+
+	FBoneSocketTarget(FName InName = NAME_None, bool bInUseSocket = false)
+	{
+		bUseSocket = bInUseSocket;
+
+		if (bUseSocket)
+		{
+			SocketReference.SocketName = InName;
+		}
+		else
+		{
+			BoneReference.BoneName = InName;
+		}
+	}
+
+	void InitializeBoneReferences(const FBoneContainer& RequiredBones)
+	{
+		if (bUseSocket)
+		{
+			SocketReference.InitialzeCompactBoneIndex(RequiredBones);
+			BoneReference.InvalidateCachedBoneIndex();
+		}
+		else
+		{
+			BoneReference.Initialize(RequiredBones);
+			SocketReference.InvalidateCachedBoneIndex();
+		}
+	}
+
+	/** Initialize Bone Reference, return TRUE if success, otherwise, return false **/
+	void Initialize(const FAnimInstanceProxy* InAnimInstanceProxy)
+	{
+		if (bUseSocket)
+		{
+			SocketReference.InitializeSocketInfo(InAnimInstanceProxy);
+		}
+	}
+
+
+	/** return true if valid. Otherwise return false **/
+	bool HasValidSetup() const
+	{
+		if (bUseSocket)
+		{
+			return SocketReference.HasValidSetup();
+		}
+
+		return BoneReference.BoneIndex != INDEX_NONE;
+	}
+
+	bool HasTargetSetup() const
+	{
+		if (bUseSocket)
+		{
+			return (SocketReference.SocketName != NAME_None);
+		}
+
+		return (BoneReference.BoneName != NAME_None);
+	}
+
+	FName GetTargetSetup() const
+	{
+		if (bUseSocket)
+		{
+			return (SocketReference.SocketName);
+		}
+
+		return (BoneReference.BoneName);
+	}
+
+	/** return true if valid. Otherwise return false **/
+	bool IsValidToEvaluate(const FBoneContainer& RequiredBones) const
+	{
+		if (bUseSocket)
+		{
+			return SocketReference.IsValidToEvaluate();
+		}
+
+		return BoneReference.IsValidToEvaluate(RequiredBones);
+	}
+
+	// this will return the compact pose bone index that matters
+	// if you're using socket, it will return socket's related joint's compact pose index
+	FCompactPoseBoneIndex GetCompactPoseBoneIndex() const 
+	{
+		if (bUseSocket)
+		{
+			return SocketReference.GetCachedSocketCompactBoneIndex();
+		}
+
+		return BoneReference.CachedCompactPoseIndex;
+	}
+
+	/** Get Target Transform from current incoming pose */
+	template<typename poseType>
+	FTransform GetTargetTransform(const FVector& TargetOffset, FCSPose<poseType>& InPose, const FTransform& InComponentToWorld) const
+	{
+		FTransform OutTargetTransform;
+
+		auto SetComponentSpaceOffset = [](const FVector& InTargetOffset, const FTransform& LocalInComponentToWorld, FTransform& LocalOutTargetTransform)
+		{
+			LocalOutTargetTransform.SetIdentity();
+			FVector CSTargetOffset = LocalInComponentToWorld.InverseTransformPosition(InTargetOffset);
+			LocalOutTargetTransform.SetLocation(CSTargetOffset);
+		};
+
+		if (bUseSocket)
+		{
+			// this has to be done outside
+			if (SocketReference.IsValidToEvaluate())
+			{
+				FTransform SocketTransformInCS = SocketReference.GetAnimatedSocketTransform(InPose);
+
+				FVector CSTargetOffset = SocketTransformInCS.TransformPosition(TargetOffset);
+				OutTargetTransform = SocketTransformInCS;
+				OutTargetTransform.SetLocation(CSTargetOffset);
+			}
+			else
+			{
+				// if none is found, we consider this offset is world offset
+				SetComponentSpaceOffset(TargetOffset, InComponentToWorld, OutTargetTransform);
+			}
+		}
+		// if valid data is available
+		else if (BoneReference.HasValidSetup())
+		{
+			if (BoneReference.IsValidToEvaluate() && 
+				ensureMsgf(InPose.GetPose().IsValidIndex(BoneReference.CachedCompactPoseIndex), TEXT("Invalid Cached Pose : Name %s(Bone Index (%d), Cached (%d))"), *BoneReference.BoneName.ToString(), BoneReference.BoneIndex, BoneReference.CachedCompactPoseIndex.GetInt()))
+			{
+				OutTargetTransform = InPose.GetComponentSpaceTransform(BoneReference.CachedCompactPoseIndex);
+				FVector CSTargetOffset = OutTargetTransform.TransformPosition(TargetOffset);
+				OutTargetTransform.SetLocation(CSTargetOffset);
+			}
+			else
+			{
+				// if none is found, we consider this offset is world offset
+				SetComponentSpaceOffset(TargetOffset, InComponentToWorld, OutTargetTransform);
+			}
+		}
+		else
+		{
+			// if none is found, we consider this offset is world offset
+			SetComponentSpaceOffset(TargetOffset, InComponentToWorld, OutTargetTransform);
+		}
+
+		return OutTargetTransform;
+	}
+
+	template<typename poseType>
+	FTransform GetTargetTransform(const FTransform& TargetOffset, FCSPose<poseType>& InPose, const FTransform& InComponentToWorld) const
+	{
+		FTransform OutTargetTransform;
+
+		auto SetComponentSpaceOffset = [](const FTransform& InTargetOffset, const FTransform& LocalInComponentToWorld, FTransform& LocalOutTargetTransform)
+		{
+			LocalOutTargetTransform = InTargetOffset.GetRelativeTransform(LocalInComponentToWorld);
+		};
+
+		if (bUseSocket)
+		{
+			// this has to be done outside
+			if (SocketReference.IsValidToEvaluate())
+			{
+				OutTargetTransform = TargetOffset * SocketReference.GetAnimatedSocketTransform(InPose);
+			}
+			else
+			{
+				SetComponentSpaceOffset(TargetOffset, InComponentToWorld, OutTargetTransform);
+			}
+		}
+		// if valid data is available
+		else if (BoneReference.HasValidSetup())
+		{
+			if (BoneReference.IsValidToEvaluate() && 
+				ensureMsgf(InPose.GetPose().IsValidIndex(BoneReference.CachedCompactPoseIndex), TEXT("Invalid Cached Pose : Name %s(Bone Index (%d), Cached (%d))"), *BoneReference.BoneName.ToString(), BoneReference.BoneIndex, BoneReference.CachedCompactPoseIndex.GetInt()))
+			{
+				OutTargetTransform = TargetOffset * InPose.GetComponentSpaceTransform(BoneReference.CachedCompactPoseIndex);
+			}
+			else
+			{
+				SetComponentSpaceOffset(TargetOffset, InComponentToWorld, OutTargetTransform);
+			}
+		}
+		else
+		{
+			SetComponentSpaceOffset(TargetOffset, InComponentToWorld, OutTargetTransform);
+		}
+
+		return OutTargetTransform;
+	}
+};
+
