Fix unity build errors following CL 16356830

diff --git a/Engine/Source/Editor/Persona/Private/Customization/BlendSampleDetails.cpp b/Engine/Source/Editor/Persona/Private/Customization/BlendSampleDetails.cpp
--- a/Engine/Source/Editor/Persona/Private/Customization/BlendSampleDetails.cpp
+++ b/Engine/Source/Editor/Persona/Private/Customization/BlendSampleDetails.cpp
@@ -46,7 +46,7 @@ FReply FBlendSampleDetails::HandleAnalyzeAndDuplicateSample()
 		FSlateApplication::Get().DismissAllMenus();
 		bool bAnalyzed[3] = { false, false, false };
 		FVector OrigValue = BlendSpace->GetBlendSample(SampleIndex).SampleValue;
-		FVector NewValue = FBlendSpaceAnalysis::CalculateSampleValue(
+		FVector NewValue = BlendSpaceAnalysis::CalculateSampleValue(
 			*BlendSpace, *BlendSpace->GetBlendSample(SampleIndex).Animation, 
 			BlendSpace->GetBlendSample(SampleIndex).RateScale, OrigValue, bAnalyzed);
 		NewValue.Z = OrigValue.Z;
@@ -64,7 +64,7 @@ FReply FBlendSampleDetails::HandleAnalyzeSample()
 	{
 		bool bAnalyzed[3] = { false, false, false };
 		FVector OrigValue = BlendSpace->GetBlendSample(SampleIndex).SampleValue;
-		FVector NewValue = FBlendSpaceAnalysis::CalculateSampleValue(
+		FVector NewValue = BlendSpaceAnalysis::CalculateSampleValue(
 			*BlendSpace, *BlendSpace->GetBlendSample(SampleIndex).Animation, 
 			BlendSpace->GetBlendSample(SampleIndex).RateScale, OrigValue, bAnalyzed);
 		NewValue.Z = OrigValue.Z;
@@ -82,7 +82,7 @@ FReply FBlendSampleDetails::HandleAnalyzeSampleX()
 	{
 		bool bAnalyzed[3] = { false, false, false };
 		FVector OrigValue = BlendSpace->GetBlendSample(SampleIndex).SampleValue;
-		FVector NewValue = FBlendSpaceAnalysis::CalculateSampleValue(
+		FVector NewValue = BlendSpaceAnalysis::CalculateSampleValue(
 			*BlendSpace, *BlendSpace->GetBlendSample(SampleIndex).Animation, 
 			BlendSpace->GetBlendSample(SampleIndex).RateScale, OrigValue, bAnalyzed);
 		NewValue.Y = OrigValue.Y;
@@ -101,7 +101,7 @@ FReply FBlendSampleDetails::HandleAnalyzeSampleY()
 	{
 		bool bAnalyzed[3] = { false, false, false };
 		FVector OrigValue = BlendSpace->GetBlendSample(SampleIndex).SampleValue;
-		FVector NewValue = FBlendSpaceAnalysis::CalculateSampleValue(
+		FVector NewValue = BlendSpaceAnalysis::CalculateSampleValue(
 			*BlendSpace, *BlendSpace->GetBlendSample(SampleIndex).Animation, 
 			BlendSpace->GetBlendSample(SampleIndex).RateScale, OrigValue, bAnalyzed);
 		NewValue.X = OrigValue.X;
@@ -218,7 +218,7 @@ void FBlendSampleDetails::CustomizeDetails(class IDetailLayoutBuilder& DetailBui
 
 		bool bShowAnalysis = false;
 		bool bAnalyzed[3] = { false, false, false };
-		FVector NewValue = FBlendSpaceAnalysis::CalculateSampleValue(
+		FVector NewValue = BlendSpaceAnalysis::CalculateSampleValue(
 			*BlendSpace, *BlendSpace->GetBlendSample(SampleIndex).Animation,
 			BlendSpace->GetBlendSample(SampleIndex).RateScale, 
 			BlendSpace->GetBlendSample(SampleIndex).SampleValue, bAnalyzed);
diff --git a/Engine/Source/Editor/Persona/Private/Customization/BlendSpaceDetails.cpp b/Engine/Source/Editor/Persona/Private/Customization/BlendSpaceDetails.cpp
--- a/Engine/Source/Editor/Persona/Private/Customization/BlendSpaceDetails.cpp
+++ b/Engine/Source/Editor/Persona/Private/Customization/BlendSpaceDetails.cpp
@@ -48,12 +48,6 @@ FReply FBlendSpaceDetails::HandleClearSamples()
 	return FReply::Handled();
 }
 
-//======================================================================================================================
-static bool GetLockAfterAnalysis(const TObjectPtr<UAnalysisProperties>& AnalysisProperties)
-{
-	return AnalysisProperties ? AnalysisProperties->bLockAfterAnalysis : false;
-}
-
 //======================================================================================================================
 FReply FBlendSpaceDetails::HandleAnalyzeSamples()
 {
@@ -62,15 +56,15 @@ FReply FBlendSpaceDetails::HandleAnalyzeSamples()
 		FSlateApplication::Get().DismissAllMenus();
 		bool bChangedOne = false;
 
-		bool bLockX = GetLockAfterAnalysis(BlendSpace->AnalysisProperties[0]);
-		bool bLockY = GetLockAfterAnalysis(BlendSpace->AnalysisProperties[1]);
-		bool bLockZ = GetLockAfterAnalysis(BlendSpace->AnalysisProperties[2]);
+		bool bLockX = BlendSpaceAnalysis::GetLockAfterAnalysis(BlendSpace->AnalysisProperties[0]);
+		bool bLockY = BlendSpaceAnalysis::GetLockAfterAnalysis(BlendSpace->AnalysisProperties[1]);
+		bool bLockZ = BlendSpaceAnalysis::GetLockAfterAnalysis(BlendSpace->AnalysisProperties[2]);
 		
 		const int32 NumSamples = BlendSpace->SampleData.Num();
 		for (int32 SampleIndex = 0 ; SampleIndex != NumSamples ; ++SampleIndex)
 		{
 			bool bAnalyzed[3] = { false, false, false };
-			FVector NewValue = FBlendSpaceAnalysis::CalculateSampleValue(
+			FVector NewValue = BlendSpaceAnalysis::CalculateSampleValue(
 				*BlendSpace, *BlendSpace->SampleData[SampleIndex].Animation, 
 				BlendSpace->SampleData[SampleIndex].RateScale, 
 				BlendSpace->SampleData[SampleIndex].SampleValue, bAnalyzed);
@@ -83,9 +77,9 @@ FReply FBlendSpaceDetails::HandleAnalyzeSamples()
 					bChangedOne = true;
 					BlendSpace->LockSample(
 						SampleIndex,
-						bAnalyzed[0] ? GetLockAfterAnalysis(BlendSpace->AnalysisProperties[0]) : false,
-						bAnalyzed[1] ? GetLockAfterAnalysis(BlendSpace->AnalysisProperties[1]) : false,
-						bAnalyzed[2] ? GetLockAfterAnalysis(BlendSpace->AnalysisProperties[2]) : false);
+						bAnalyzed[0] ? BlendSpaceAnalysis::GetLockAfterAnalysis(BlendSpace->AnalysisProperties[0]) : false,
+						bAnalyzed[1] ? BlendSpaceAnalysis::GetLockAfterAnalysis(BlendSpace->AnalysisProperties[1]) : false,
+						bAnalyzed[2] ? BlendSpaceAnalysis::GetLockAfterAnalysis(BlendSpace->AnalysisProperties[2]) : false);
 				}
 			}
 		}
@@ -131,7 +125,7 @@ static EVisibility GetAnalyzeButtonVisibility(
 void FBlendSpaceDetails::HandleAnalysisFunctionChanged(int32 AxisIndex, TSharedPtr<FString> NewFunctionName)
 {
 	BlendSpace->Modify();
-	UAnalysisProperties* NewAnalysisProperties = FBlendSpaceAnalysis::MakeAnalysisProperties(BlendSpace, *NewFunctionName);
+	UAnalysisProperties* NewAnalysisProperties = BlendSpaceAnalysis::MakeAnalysisProperties(BlendSpace, *NewFunctionName);
 	// Preserve values where possible
 	if (BlendSpace->AnalysisProperties[AxisIndex])
 	{
@@ -328,7 +322,7 @@ void FBlendSpaceDetails::CustomizeDetails(class IDetailLayoutBuilder& DetailBuil
 					// Prepare the drop-down. Note that it uses pointers to strings, and comparison is between
 					// pointers, not string contents!
 					AnalysisFunctionNames[AxisIndex].Empty();
-					const TArray<FString> AnalysisFunctions = FBlendSpaceAnalysis::GetAnalysisFunctions();
+					const TArray<FString> AnalysisFunctions = BlendSpaceAnalysis::GetAnalysisFunctions();
 					TSharedPtr<FString> CurrentlySelectedFunction;
 					for (const FString& AnalysisFunction : AnalysisFunctions)
 					{
diff --git a/Engine/Source/Editor/Persona/Private/PersonaBlendSpaceAnalysis.cpp b/Engine/Source/Editor/Persona/Private/PersonaBlendSpaceAnalysis.cpp
--- a/Engine/Source/Editor/Persona/Private/PersonaBlendSpaceAnalysis.cpp
+++ b/Engine/Source/Editor/Persona/Private/PersonaBlendSpaceAnalysis.cpp
@@ -1383,7 +1383,7 @@ static TArray<IBlendSpaceAnalysisFeature*> GetAnalysisFeatures()
 }
 
 //======================================================================================================================
-FVector FBlendSpaceAnalysis::CalculateSampleValue(const UBlendSpace& BlendSpace, const UAnimSequence& Animation,
+FVector BlendSpaceAnalysis::CalculateSampleValue(const UBlendSpace& BlendSpace, const UAnimSequence& Animation,
 												  const float RateScale, const FVector& OriginalPosition, bool bAnalyzed[3])
 {
 
@@ -1408,7 +1408,7 @@ FVector FBlendSpaceAnalysis::CalculateSampleValue(const UBlendSpace& BlendSpace,
 }
 
 //======================================================================================================================
-UAnalysisProperties* FBlendSpaceAnalysis::MakeAnalysisProperties(UObject* Outer, const FString& FunctionName)
+UAnalysisProperties* BlendSpaceAnalysis::MakeAnalysisProperties(UObject* Outer, const FString& FunctionName)
 {
 	UAnalysisProperties* Result = nullptr;
 
@@ -1425,7 +1425,7 @@ UAnalysisProperties* FBlendSpaceAnalysis::MakeAnalysisProperties(UObject* Outer,
 }
 
 //======================================================================================================================
-TArray<FString> FBlendSpaceAnalysis::GetAnalysisFunctions()
+TArray<FString> BlendSpaceAnalysis::GetAnalysisFunctions()
 {
 	TArray<FString> FunctionNames;
 	TArray<IBlendSpaceAnalysisFeature*> ModularFeatures = GetAnalysisFeatures();
@@ -1440,6 +1440,11 @@ TArray<FString> FBlendSpaceAnalysis::GetAnalysisFunctions()
 	return FunctionNames;
 }
 
+//======================================================================================================================
+bool BlendSpaceAnalysis::GetLockAfterAnalysis(const TObjectPtr<UAnalysisProperties>& AnalysisProperties)
+{
+	return AnalysisProperties ? AnalysisProperties->bLockAfterAnalysis : false;
+}
 
 
 #undef LOCTEXT_NAMESPACE
diff --git a/Engine/Source/Editor/Persona/Private/PersonaBlendSpaceAnalysis.h b/Engine/Source/Editor/Persona/Private/PersonaBlendSpaceAnalysis.h
--- a/Engine/Source/Editor/Persona/Private/PersonaBlendSpaceAnalysis.h
+++ b/Engine/Source/Editor/Persona/Private/PersonaBlendSpaceAnalysis.h
@@ -3,20 +3,33 @@
 
 #include "BlendSpaceAnalysis.h"
 
-class FBlendSpaceAnalysis
+namespace BlendSpaceAnalysis
 {
-public:
-	static FVector CalculateSampleValue(const UBlendSpace&   BlendSpace, 
-										const UAnimSequence& Animation, 
-										const float          RateScale, 
-										const FVector&       OriginalPosition, 
-										bool                 bAnalyzed[3]);
+	/**
+	 * Calculates the sample value/position and indicates if it needed to be analyzed. Note that the original position
+	 * is passed in, and its elements will be used in the output where analysis doesn't need to be done.
+	 */
+	FVector CalculateSampleValue(const UBlendSpace&   BlendSpace, 
+								 const UAnimSequence& Animation, 
+								 const float          RateScale, 
+								 const FVector&       OriginalPosition, 
+								 bool                 bAnalyzed[3]);
 
 
-	// This will return an instance derived from UAnalysisProperties that is suitable for the Function. The caller will
-	// pass in a suitable owning object, outer, that the implementation should assign as owner of the newly created object.
-	static UAnalysisProperties* MakeAnalysisProperties(UObject* Outer, const FString& FunctionName);
+	/**
+	 * This will return an instance derived from UAnalysisProperties that is suitable for the Function. The caller will
+	 * pass in a suitable owning object, outer, that the implementation should assign as owner of the newly created object.
+	 */
+	UAnalysisProperties* MakeAnalysisProperties(UObject* Outer, const FString& FunctionName);
 
-	// This will return the names of the functions handled
-	static TArray<FString> GetAnalysisFunctions();
+	/**
+	 * This will return the names of the functions handled
+	 */
+	TArray<FString> GetAnalysisFunctions();
+
+	/**
+	 * Simple helper that avoids the need to check pointer validity when determining if a sample should be locked 
+	 * after analysis
+	 */
+	bool GetLockAfterAnalysis(const TObjectPtr<UAnalysisProperties>& AnalysisProperties);
 };
diff --git a/Engine/Source/Editor/Persona/Private/SAnimationBlendSpace.cpp b/Engine/Source/Editor/Persona/Private/SAnimationBlendSpace.cpp
--- a/Engine/Source/Editor/Persona/Private/SAnimationBlendSpace.cpp
+++ b/Engine/Source/Editor/Persona/Private/SAnimationBlendSpace.cpp
@@ -149,12 +149,6 @@ void SBlendSpaceEditor::OnSampleRemoved(const int32 SampleIndex)
 	}
 }
 
-//======================================================================================================================
-static bool GetLockAfterAnalysis(const TObjectPtr<UAnalysisProperties>& AnalysisProperties)
-{
-	return AnalysisProperties ? AnalysisProperties->bLockAfterAnalysis : false;
-}
-
 //======================================================================================================================
 int32 SBlendSpaceEditor::OnSampleAdded(UAnimSequence* Animation, const FVector& Value)
 {
@@ -165,7 +159,7 @@ int32 SBlendSpaceEditor::OnSampleAdded(UAnimSequence* Animation, const FVector&
 
 	if(BlendSpace->IsAsset())
 	{
-		AdjustedValue = FBlendSpaceAnalysis::CalculateSampleValue(*BlendSpace, *Animation, 1.0f, Value, bAnalyzed);
+		AdjustedValue = BlendSpaceAnalysis::CalculateSampleValue(*BlendSpace, *Animation, 1.0f, Value, bAnalyzed);
 	}
 		
 	int32 NewSampleIndex = -1;
@@ -182,10 +176,11 @@ int32 SBlendSpaceEditor::OnSampleAdded(UAnimSequence* Animation, const FVector&
 	if (NewSampleIndex >= 0)
 	{
 		BlendSpace->Modify();
-		BlendSpace->LockSample(NewSampleIndex,
-							   bAnalyzed[0] ? GetLockAfterAnalysis(BlendSpace->AnalysisProperties[0]) : false,
-							   bAnalyzed[1] ? GetLockAfterAnalysis(BlendSpace->AnalysisProperties[1]) : false,
-							   bAnalyzed[2] ? GetLockAfterAnalysis(BlendSpace->AnalysisProperties[2]) : false);
+		BlendSpace->LockSample(
+			NewSampleIndex,
+			bAnalyzed[0] ? BlendSpaceAnalysis::GetLockAfterAnalysis(BlendSpace->AnalysisProperties[0]) : false,
+			bAnalyzed[1] ? BlendSpaceAnalysis::GetLockAfterAnalysis(BlendSpace->AnalysisProperties[1]) : false,
+			bAnalyzed[2] ? BlendSpaceAnalysis::GetLockAfterAnalysis(BlendSpace->AnalysisProperties[2]) : false);
 
 		ResampleData();
 		BlendSpace->ValidateSampleData();
