Fix for LWC

diff --git a/Engine/Source/Editor/Persona/Private/PersonaBlendSpaceAnalysis.cpp b/Engine/Source/Editor/Persona/Private/PersonaBlendSpaceAnalysis.cpp
--- a/Engine/Source/Editor/Persona/Private/PersonaBlendSpaceAnalysis.cpp
+++ b/Engine/Source/Editor/Persona/Private/PersonaBlendSpaceAnalysis.cpp
@@ -1395,11 +1395,12 @@ FVector BlendSpaceAnalysis::CalculateSampleValue(const UBlendSpace& BlendSpace,
 		const UAnalysisProperties* AnalysisProperties = BlendSpace.AnalysisProperties[Index].Get();
 		for (const IBlendSpaceAnalysisFeature* Feature : ModularFeatures)
 		{
+			float FloatValue = (float) AdjustedPosition[Index];
 			bAnalyzed[Index] = Feature->CalculateSampleValue(
-				AdjustedPosition[Index], BlendSpace, AnalysisProperties, Animation, RateScale);
-
+				FloatValue, BlendSpace, AnalysisProperties, Animation, RateScale);
 			if (bAnalyzed[Index])
 			{
+				AdjustedPosition[Index] = (FVector::FReal) FloatValue;
 				break;
 			}
 		}
