Persona: fix BP out of date notification no longer shows up in viewport

diff --git a/Engine/Plugins/Experimental/ControlRig/Source/ControlRigEditor/Private/Editor/ControlRigEditor.cpp b/Engine/Plugins/Experimental/ControlRig/Source/ControlRigEditor/Private/Editor/ControlRigEditor.cpp
--- a/Engine/Plugins/Experimental/ControlRig/Source/ControlRigEditor/Private/Editor/ControlRigEditor.cpp
+++ b/Engine/Plugins/Experimental/ControlRig/Source/ControlRigEditor/Private/Editor/ControlRigEditor.cpp
@@ -2745,7 +2745,8 @@ void FControlRigEditor::HandleViewportCreated(const TSharedRef<class IPersonaVie
 					.Text(LOCTEXT("ControlRigBPViewportCompileButtonLabel", "Compile"))
 				]
 			]
-		]
+		],
+		FPersonaViewportNotificationOptions(TAttribute<EVisibility>::Create(GetCompilationStateVisibility))
 	);
 
 	InViewport->AddToolbarExtender(TEXT("AnimViewportDefaultCamera"), FMenuExtensionDelegate::CreateLambda(
diff --git a/Engine/Source/Editor/AnimationBlueprintEditor/Private/AnimationBlueprintEditor.cpp b/Engine/Source/Editor/AnimationBlueprintEditor/Private/AnimationBlueprintEditor.cpp
--- a/Engine/Source/Editor/AnimationBlueprintEditor/Private/AnimationBlueprintEditor.cpp
+++ b/Engine/Source/Editor/AnimationBlueprintEditor/Private/AnimationBlueprintEditor.cpp
@@ -2065,7 +2065,8 @@ void FAnimationBlueprintEditor::HandleViewportCreated(const TSharedRef<IPersonaV
 					.Text(LOCTEXT("AnimBPViewportCompileButtonLabel", "Compile"))
 				]
 			]
-		]
+		],
+		FPersonaViewportNotificationOptions(TAttribute<EVisibility>::Create(GetCompilationStateVisibility))
 	);
 }
 
diff --git a/Engine/Source/Editor/Persona/Private/SAnimationEditorViewport.cpp b/Engine/Source/Editor/Persona/Private/SAnimationEditorViewport.cpp
--- a/Engine/Source/Editor/Persona/Private/SAnimationEditorViewport.cpp
+++ b/Engine/Source/Editor/Persona/Private/SAnimationEditorViewport.cpp
@@ -318,7 +318,7 @@ TSharedRef<IPinnedCommandList> SAnimationEditorViewportTabBody::GetPinnedCommand
 	return ViewportWidget->GetViewportToolbar()->GetPinnedCommandList().ToSharedRef();
 }
 
-TWeakPtr<SWidget> SAnimationEditorViewportTabBody::AddNotification(TAttribute<EMessageSeverity::Type> InSeverity, TAttribute<bool> InCanBeDismissed, const TSharedRef<SWidget>& InNotificationWidget)
+TWeakPtr<SWidget> SAnimationEditorViewportTabBody::AddNotification(TAttribute<EMessageSeverity::Type> InSeverity, TAttribute<bool> InCanBeDismissed, const TSharedRef<SWidget>& InNotificationWidget, FPersonaViewportNotificationOptions InOptions)
 {
 	TSharedPtr<SBorder> ContainingWidget = nullptr;
 	TWeakPtr<SWidget> WeakNotificationWidget = InNotificationWidget;
@@ -333,17 +333,14 @@ TWeakPtr<SWidget> SAnimationEditorViewportTabBody::AddNotification(TAttribute<EM
 		return FMargin(0.0f);
 	};
 
-	auto GetVisibility = [WeakNotificationWidget]()
+	TAttribute<EVisibility> GetVisibility(EVisibility::Visible);
+	
+	if (InOptions.OnGetVisibility.IsSet())
 	{
-		if(WeakNotificationWidget.IsValid())
-		{
-			return WeakNotificationWidget.Pin()->GetVisibility();
-		}
-
-		return EVisibility::Collapsed;
-	};
-
-	auto GetBrushForSeverity = [InSeverity]()
+		GetVisibility = InOptions.OnGetVisibility;
+	}
+	
+	TAttribute<const FSlateBrush*> GetBrushForSeverity = TAttribute<const FSlateBrush*>::Create([InSeverity]()
 	{
 		switch(InSeverity.Get())
 		{
@@ -357,7 +354,12 @@ TWeakPtr<SWidget> SAnimationEditorViewportTabBody::AddNotification(TAttribute<EM
 		case EMessageSeverity::Info:
 			return FEditorStyle::GetBrush("AnimViewport.Notification.Message");
 		}
-	};
+	});
+
+	if (InOptions.OnGetBrushOverride.IsSet())
+	{
+		GetBrushForSeverity = InOptions.OnGetBrushOverride;
+	}
 
 	TSharedPtr<SHorizontalBox> BodyBox = nullptr;
 
@@ -367,8 +369,8 @@ TWeakPtr<SWidget> SAnimationEditorViewportTabBody::AddNotification(TAttribute<EM
 	.Padding(MakeAttributeLambda(GetPadding))
 	[
 		SAssignNew(ContainingWidget, SBorder)
-		.Visibility_Lambda(GetVisibility)
-		.BorderImage_Lambda(GetBrushForSeverity)
+		.Visibility(GetVisibility)
+		.BorderImage(GetBrushForSeverity)
 		[
 			SAssignNew(BodyBox, SHorizontalBox)
 			+SHorizontalBox::Slot()
@@ -2139,7 +2141,8 @@ void SAnimationEditorViewportTabBody::AddRecordingNotification()
 					.Text(LOCTEXT("AnimViewportStopRecordingButtonLabel", "Stop"))
 				]
 			]
-		]
+		],
+		FPersonaViewportNotificationOptions(TAttribute<EVisibility>::Create(GetRecordingStateStateVisibility))
 	);
 }
 
@@ -2309,7 +2312,8 @@ void SAnimationEditorViewportTabBody::AddPostProcessNotification()
 					.Text(LOCTEXT("EditPostProcessAnimBPButtonText", "Edit"))
 				]
 			]
-		]
+		],
+		FPersonaViewportNotificationOptions(TAttribute<EVisibility>::Create(GetVisibility))
 	);
 }
 
@@ -2358,7 +2362,8 @@ void SAnimationEditorViewportTabBody::AddMinLODNotification()
 				.Text(LOCTEXT("MinLODNotification", "Min LOD applied"))
 				.TextStyle(FEditorStyle::Get(), "AnimViewport.MessageText")
 			]
-		]
+		],
+		FPersonaViewportNotificationOptions(TAttribute<EVisibility>::Create(GetMinLODNotificationVisibility))
 	);
 }
 
@@ -2418,7 +2423,8 @@ void SAnimationEditorViewportTabBody::AddSkinWeightProfileNotification()
 				.Text_Lambda(GetSkinWeightProfileNotificationText)
 				.TextStyle(FEditorStyle::Get(), "AnimViewport.MessageText")
 			]
-		]
+		],
+		FPersonaViewportNotificationOptions(TAttribute<EVisibility>::Create(GetSkinWeightProfileNotificationVisibility))
 	);
 }
 
diff --git a/Engine/Source/Editor/Persona/Private/SAnimationEditorViewport.h b/Engine/Source/Editor/Persona/Private/SAnimationEditorViewport.h
--- a/Engine/Source/Editor/Persona/Private/SAnimationEditorViewport.h
+++ b/Engine/Source/Editor/Persona/Private/SAnimationEditorViewport.h
@@ -219,7 +219,7 @@ public:
 	virtual void RestoreState(TSharedRef<IPersonaViewportState> InState) override;
 	virtual FEditorViewportClient& GetViewportClient() const override;
 	virtual TSharedRef<IPinnedCommandList> GetPinnedCommandList() const override;
-	virtual TWeakPtr<SWidget> AddNotification(TAttribute<EMessageSeverity::Type> InSeverity, TAttribute<bool> InCanBeDismissed, const TSharedRef<SWidget>& InNotificationWidget) override;
+	virtual TWeakPtr<SWidget> AddNotification(TAttribute<EMessageSeverity::Type> InSeverity, TAttribute<bool> InCanBeDismissed, const TSharedRef<SWidget>& InNotificationWidget, FPersonaViewportNotificationOptions InOptions) override;
 	virtual void RemoveNotification(const TWeakPtr<SWidget>& InContainingWidget) override;
 	virtual void AddToolbarExtender(FName MenuToExtend, FMenuExtensionDelegate MenuBuilderDelegate) override;
 	virtual FPersonaViewportKeyDownDelegate& GetKeyDownDelegate() override { return OnKeyDownDelegate; }
diff --git a/Engine/Source/Editor/Persona/Public/IPersonaViewport.h b/Engine/Source/Editor/Persona/Public/IPersonaViewport.h
--- a/Engine/Source/Editor/Persona/Public/IPersonaViewport.h
+++ b/Engine/Source/Editor/Persona/Public/IPersonaViewport.h
@@ -17,6 +17,18 @@ struct IPersonaViewportState
 {
 };
 
+struct FPersonaViewportNotificationOptions
+{
+	FPersonaViewportNotificationOptions() = default;
+	
+	// in general on GetVisibility should always be provided
+	explicit FPersonaViewportNotificationOptions(const TAttribute<EVisibility> InOnGetVisibility)
+		:OnGetVisibility(InOnGetVisibility) {};
+
+	TAttribute<EVisibility> OnGetVisibility;
+	TAttribute<const FSlateBrush*> OnGetBrushOverride;
+};
+
 /** Abstract viewport that can save and restore state */
 class IPersonaViewport : public SCompoundWidget
 {
@@ -38,9 +50,10 @@ public:
 	 * @param	InSeverity				The severity of the message
 	 * @param	InCanBeDismissed		Whether the message can be manually dismissed
 	 * @param	InNotificationWidget	The widget showing the notification
+	 * @param	InOptions				Specify GetVisibility callback and other optional callback overrides
 	 * @return the widget containing the notification
 	 */
-	virtual TWeakPtr<SWidget> AddNotification(TAttribute<EMessageSeverity::Type> InSeverity, TAttribute<bool> InCanBeDismissed, const TSharedRef<SWidget>& InNotificationWidget) = 0;
+	virtual TWeakPtr<SWidget> AddNotification(TAttribute<EMessageSeverity::Type> InSeverity, TAttribute<bool> InCanBeDismissed, const TSharedRef<SWidget>& InNotificationWidget, FPersonaViewportNotificationOptions InOptions) = 0;
 	
 	/** 
 	* Remove a notification widget 
