Makes the analysis functions more easily available for use by custom analysis. Implements Locomotion and RootMotion functions as a plugin. Some improvements to the blendspace sample menu. Adds duplicate functionality to that menu too. Other miscellaneous fixes.

diff --git a/Engine/Plugins/Animation/BlendSpaceMotionAnalysis/BlendSpaceMotionAnalysis.uplugin b/Engine/Plugins/Animation/BlendSpaceMotionAnalysis/BlendSpaceMotionAnalysis.uplugin
new file mode 100644
--- /dev/null
+++ b/Engine/Plugins/Animation/BlendSpaceMotionAnalysis/BlendSpaceMotionAnalysis.uplugin
@@ -0,0 +1,25 @@
+{
+	"FileVersion": 3,
+	"Version": 1,
+	"VersionName": "1.0",
+	"FriendlyName": "Blendspace Motion Analysis",
+	"Description": "Allows analysis of locomotion/root motion properties in blend spaces",
+	"Category": "BlendSpace",
+	"CreatedBy": "Epic Games, Inc.",
+	"CreatedByURL": "http://epicgames.com",
+	"DocsURL": "",
+	"MarketplaceURL": "",
+	"SupportURL": "",
+	"EnabledByDefault": true,
+	"CanContainContent": false,
+	"IsBetaVersion": false,
+	"Installed": false,
+	"Modules": 
+	[
+		{
+			"Name": "BlendSpaceMotionAnalysis",
+			"Type": "Editor",
+			"LoadingPhase": "Default"
+		}
+	]
+}
diff --git a/Engine/Plugins/Animation/BlendSpaceMotionAnalysis/Source/BlendSpaceMotionAnalysis/BlendSpaceMotionAnalysis.Build.cs b/Engine/Plugins/Animation/BlendSpaceMotionAnalysis/Source/BlendSpaceMotionAnalysis/BlendSpaceMotionAnalysis.Build.cs
new file mode 100644
--- /dev/null
+++ b/Engine/Plugins/Animation/BlendSpaceMotionAnalysis/Source/BlendSpaceMotionAnalysis/BlendSpaceMotionAnalysis.Build.cs
@@ -0,0 +1,35 @@
+// Copyright Epic Games, Inc. All Rights Reserved.
+
+using System.IO;
+
+namespace UnrealBuildTool.Rules
+{
+	public class BlendSpaceMotionAnalysis : ModuleRules
+	{
+		public BlendSpaceMotionAnalysis(ReadOnlyTargetRules Target) : base(Target)
+		{
+			PublicDependencyModuleNames.AddRange(
+				new string[]
+				{
+				}
+			);
+
+			PrivateDependencyModuleNames.AddRange(
+				new string[]
+				{
+					"Core",
+					"CoreUObject",
+					"Engine",
+					"Persona",
+				}
+			);
+
+			PrivateIncludePathModuleNames.AddRange(
+				new string[]
+				{
+				}
+			);
+
+		}
+	}
+}
diff --git a/Engine/Plugins/Animation/BlendSpaceMotionAnalysis/Source/BlendSpaceMotionAnalysis/Private/BlendSpaceMotionAnalysis.cpp b/Engine/Plugins/Animation/BlendSpaceMotionAnalysis/Source/BlendSpaceMotionAnalysis/Private/BlendSpaceMotionAnalysis.cpp
new file mode 100644
--- /dev/null
+++ b/Engine/Plugins/Animation/BlendSpaceMotionAnalysis/Source/BlendSpaceMotionAnalysis/Private/BlendSpaceMotionAnalysis.cpp
@@ -0,0 +1,104 @@
+// Copyright Epic Games, Inc. All Rights Reserved.
+
+#include "BlendSpaceMotionAnalysis.h"
+#include "Features/IModularFeatures.h"
+#include "Engine/Engine.h"
+#include "Modules/ModuleManager.h"
+#include "Logging/LogMacros.h"
+
+IMPLEMENT_MODULE(FBlendSpaceMotionAnalysis, BlendSpaceMotionAnalysis)
+
+DEFINE_LOG_CATEGORY_STATIC(LogBlendSpaceMotionAnalysis, Log, All);
+
+#define LOCTEXT_NAMESPACE "BlendSpaceMotionAnalysis"
+
+//======================================================================================================================
+class FBlendSpaceMotionAnalysisFeature : public IBlendSpaceAnalysisFeature 
+{
+public:
+	// This should process the animation according to the analysis properties, or return false if that is not possible.
+	bool CalculateSampleValue(float&                     Result,
+							  const UBlendSpace&         BlendSpace,
+							  const UAnalysisProperties* AnalysisProperties,
+							  const UAnimSequence&       Animation,
+							  const float                RateScale) const override;
+
+	// This should return an instance derived from UAnalysisProperties that is suitable for the Function. The caller
+	// will pass in a suitable owning object, outer, that the implementation should assign as owner of the newly created
+	// object. 
+	UAnalysisProperties* MakeAnalysisProperties(UObject* Outer, const FString& FunctionName) const override;
+
+	// This should return the names of the functions handled
+	TArray<FString> GetAnalysisFunctions() const override;
+};
+
+static FBlendSpaceMotionAnalysisFeature BlendSpaceMotionAnalysisFeature;
+
+//======================================================================================================================
+TArray<FString> FBlendSpaceMotionAnalysisFeature::GetAnalysisFunctions() const
+{
+	TArray<FString> Functions = 
+	{
+		TEXT("Locomotion"),
+		TEXT("RootMotion")
+	};
+	return Functions;
+}
+
+//======================================================================================================================
+UAnalysisProperties* FBlendSpaceMotionAnalysisFeature::MakeAnalysisProperties(UObject* Outer, const FString& FunctionName) const
+{
+	UAnalysisProperties* Result = nullptr;
+	if (FunctionName.Equals(TEXT("RootMotion")))
+	{
+		Result = NewObject<URootMotionAnalysisProperties>(Outer);
+	}
+	else if (FunctionName.Equals(TEXT("Locomotion")))
+	{
+		Result = NewObject<ULocomotionAnalysisProperties>(Outer);
+	}
+
+	if (Result)
+	{
+		Result->Function = FunctionName;
+	}
+	return Result;
+}
+
+//======================================================================================================================
+bool FBlendSpaceMotionAnalysisFeature::CalculateSampleValue(float&                     Result,
+															const UBlendSpace&         BlendSpace,
+															const UAnalysisProperties* AnalysisProperties,
+															const UAnimSequence&       Animation,
+															const float                RateScale) const
+{
+	if (!AnalysisProperties)
+	{
+		return false;
+	}
+	const FString& FunctionName = AnalysisProperties->Function;
+	if (FunctionName.Equals(TEXT("RootMotion")))
+	{
+		return CalculateRootMotion(
+			Result, BlendSpace, Cast<URootMotionAnalysisProperties>(AnalysisProperties), Animation, RateScale);
+	}
+	else if (FunctionName.Equals(TEXT("Locomotion")))
+	{
+		return CalculateLocomotion(
+			Result, BlendSpace, Cast<ULocomotionAnalysisProperties>(AnalysisProperties), Animation, RateScale);
+	}
+	return false;
+}
+
+//======================================================================================================================
+void FBlendSpaceMotionAnalysis::StartupModule()
+{
+	IModularFeatures::Get().RegisterModularFeature(IBlendSpaceAnalysisFeature::GetModuleFeatureName(), &BlendSpaceMotionAnalysisFeature);
+}
+
+void FBlendSpaceMotionAnalysis::ShutdownModule()
+{
+	IModularFeatures::Get().UnregisterModularFeature(IBlendSpaceAnalysisFeature::GetModuleFeatureName(), &BlendSpaceMotionAnalysisFeature);
+}
+
+#undef LOCTEXT_NAMESPACE
\ No newline at end of file
diff --git a/Engine/Plugins/Animation/BlendSpaceMotionAnalysis/Source/BlendSpaceMotionAnalysis/Private/LocomotionAnalysis.cpp b/Engine/Plugins/Animation/BlendSpaceMotionAnalysis/Source/BlendSpaceMotionAnalysis/Private/LocomotionAnalysis.cpp
new file mode 100644
--- /dev/null
+++ b/Engine/Plugins/Animation/BlendSpaceMotionAnalysis/Source/BlendSpaceMotionAnalysis/Private/LocomotionAnalysis.cpp
@@ -0,0 +1,484 @@
+// Copyright Epic Games, Inc. All Rights Reserved.
+
+#include "LocomotionAnalysis.h"
+#include "BlendSpaceMotionAnalysis.h"
+#include "Engine/Engine.h"
+#include "Logging/LogMacros.h"
+
+DEFINE_LOG_CATEGORY_STATIC(LogLocomotionAnalysis, Log, All);
+
+#define LOCTEXT_NAMESPACE "LocomotionAnalysis"
+
+//======================================================================================================================
+void ULocomotionAnalysisProperties::InitializeFromCache(TObjectPtr<UCachedAnalysisProperties> Cache)
+{
+	if (Cache)
+	{
+		UCachedAnalysisProperties* CachePtr = Cache.Get();
+		UCachedMotionAnalysisProperties* CachedMotionAnalysisPropertiesPtr = Cast<UCachedMotionAnalysisProperties>(CachePtr);
+		if (!CachedMotionAnalysisPropertiesPtr)
+		{
+			// Replace any existing cache if it is not our own type
+			CachedMotionAnalysisPropertiesPtr = NewObject<UCachedMotionAnalysisProperties>(CachePtr->GetOuter());
+			CachedMotionAnalysisPropertiesPtr->CopyFrom(*CachePtr);
+			Cache = CachedMotionAnalysisPropertiesPtr;
+		}
+		FunctionAxis = CachedMotionAnalysisPropertiesPtr->LocomotionFunctionAxis;
+		PrimaryBoneSocket = CachedMotionAnalysisPropertiesPtr->BoneSocket1;
+		SecondaryBoneSocket = CachedMotionAnalysisPropertiesPtr->BoneSocket2;
+		CharacterFacingAxis = CachedMotionAnalysisPropertiesPtr->CharacterFacingAxis;
+		CharacterUpAxis = CachedMotionAnalysisPropertiesPtr->CharacterUpAxis;
+		bLockAfterAnalysis = CachedMotionAnalysisPropertiesPtr->bLockAfterAnalysis;
+	}
+}
+
+//======================================================================================================================
+void ULocomotionAnalysisProperties::MakeCache(TObjectPtr<UCachedAnalysisProperties>& Cache, UBlendSpace* BlendSpace)
+{
+	UCachedMotionAnalysisProperties* CachedMotionAnalysisPropertiesPtr = Cast<UCachedMotionAnalysisProperties>(Cache.Get());
+	if (!CachedMotionAnalysisPropertiesPtr)
+	{
+		// Replace any existing cache if it is not our own type
+		CachedMotionAnalysisPropertiesPtr = NewObject<UCachedMotionAnalysisProperties>(BlendSpace);
+		if (Cache.Get())
+		{
+			CachedMotionAnalysisPropertiesPtr->CopyFrom(*Cache.Get());
+		}
+		Cache = CachedMotionAnalysisPropertiesPtr;
+	}
+	CachedMotionAnalysisPropertiesPtr->LocomotionFunctionAxis = FunctionAxis;
+	CachedMotionAnalysisPropertiesPtr->BoneSocket1 = PrimaryBoneSocket;
+	CachedMotionAnalysisPropertiesPtr->BoneSocket2 = SecondaryBoneSocket;
+	CachedMotionAnalysisPropertiesPtr->CharacterFacingAxis = CharacterFacingAxis;
+	CachedMotionAnalysisPropertiesPtr->CharacterUpAxis = CharacterUpAxis;
+	CachedMotionAnalysisPropertiesPtr->bLockAfterAnalysis = bLockAfterAnalysis;
+}
+
+//======================================================================================================================
+// Note that if a looping animation has 56 keys, then its first key is 0 and last is 55, but these will be identical poses.
+// Thus it has one fewer intervals/unique keys
+static bool CalculateLocomotionVelocity(
+	FVector&                             Result,
+	const UBlendSpace&                   BlendSpace,
+	const ULocomotionAnalysisProperties* AnalysisProperties,
+	const FBoneSocketTarget&             BoneSocket,
+	const UAnimSequence&                 Animation,
+	const float                          RateScale)
+{
+	const int32 NumSampledKeys = Animation.GetNumberOfSampledKeys();
+	if (!AnalysisProperties || NumSampledKeys <= 0)
+	{
+		return false;
+	}
+	if (NumSampledKeys == 1)
+	{
+		Result.Set(0, 0, 0);
+		return true;
+	}
+
+	FTransform BoneOffset;
+	FName BoneName;
+	if (!BlendSpaceAnalysis::GetBoneInfo(Animation, BoneSocket, BoneOffset, BoneName))
+	{
+		return false;
+	}
+
+	// Note that for locomotion we don't support the frame changing
+	FTransform FrameTM = FTransform::Identity;
+	FVector FrameFacingDir, FrameUpDir, FrameRightDir;
+	BlendSpaceAnalysis::GetFrameDirs(FrameFacingDir, FrameUpDir, FrameRightDir, FrameTM, AnalysisProperties);
+
+	// The frame time delta
+	float DeltaTime = Animation.GetPlayLength() / NumSampledKeys;
+
+	// First step is to figure out the approximate direction. Note that the average velocity will be zero (assuming a
+	// complete cycle) - but if we apply a weight that is based on the height, then we can bias it towards the foot that
+	// is on the ground.
+	float MinHeight = FLT_MAX;
+	float MaxHeight = -FLT_MAX;
+	FVector AveragePos(0.0f);
+	TArray<FVector> Positions;
+	Positions.SetNum(NumSampledKeys);
+	for (int32 Key = 0; Key != NumSampledKeys; ++Key)
+	{
+		FTransform BoneTM = BlendSpaceAnalysis::GetBoneTransform(Animation, Key, BoneName);
+		FTransform TM = BoneOffset * BoneTM;
+		FVector Pos = TM.GetTranslation();
+		Positions[Key] = Pos;
+		float Height = Pos | FrameUpDir;
+		MinHeight = FMath::Min(MinHeight, Height);
+		MaxHeight = FMath::Max(MaxHeight, Height);
+		AveragePos += Pos;
+#ifdef ANALYSIS_VERBOSE_LOG
+		UE_LOG(LogAnimation, Log, TEXT("Pos %f %f %f"), Pos.X, Pos.Y, Pos.Z);
+#endif
+	}
+	AveragePos /= NumSampledKeys;
+
+	// Calculate velocities. 
+	TArray<FVector> Velocities;
+	Velocities.SetNum(NumSampledKeys);
+	for (int32 Key = 0; Key != NumSampledKeys; ++Key)
+	{
+		int32 PrevKey = (Key + NumSampledKeys - 1) % NumSampledKeys;
+		int32 NextKey = (Key +  1) % NumSampledKeys;
+		Velocities[Key] = (Positions[NextKey] - Positions[PrevKey]) / (2.0f * DeltaTime);
+	}
+
+	FVector BiasedFootVel(0);
+	float TotalWeight = 0.0f;
+	for (int32 Key = 0 ; Key != NumSampledKeys ; ++Key)
+	{
+		float Height = Positions[Key] | FrameUpDir;
+		float Weight = 1.0f - (Height - MinHeight) / (MaxHeight - MinHeight);
+		BiasedFootVel += Velocities[Key] * Weight;
+		TotalWeight += Weight;
+	}
+	BiasedFootVel /= TotalWeight; 
+
+	if (BiasedFootVel.IsNearlyZero())
+	{
+		Result.Set(0, 0, 0);
+		return true;
+	}
+
+	FVector ApproxLocoDir = -BiasedFootVel.GetSafeNormal();
+
+	// Now we can form a mask, where 0 means traveling in the wrong direction (so clearly off the ground), and positive
+	// numbers will indicate how far into a valid segment we are. We will assume that the animation is looping
+	TArray<int32> Mask;
+	Mask.SetNum(NumSampledKeys);
+	for (int32 Key = 0 ; Key != NumSampledKeys ; ++Key)
+	{
+		Mask[Key] = (Velocities[Key] | ApproxLocoDir) >= 0.0f ? 0 : 1;
+	}
+
+	int32 StartKey = -1;
+	TArray<int32> Mask1 = Mask;
+	bool bRepeat = true;
+	int32 MaxMask = 0;
+	int32 PrevNumFound = NumSampledKeys + 1;
+	while (bRepeat)
+	{
+		bRepeat = false;
+		int32 NumFound = 0;
+		for (int32 Key = 0 ; Key != NumSampledKeys ; ++Key)
+		{
+			if (Mask[Key] > 0)
+			{
+				int32 PrevKey = (Key + NumSampledKeys - 1) % NumSampledKeys;
+				int32 NextKey = (Key + 1) % NumSampledKeys;
+				if (Mask[PrevKey] == Mask[Key] && Mask[NextKey] == Mask[Key])
+				{
+					++Mask1[Key];
+					MaxMask = FMath::Max(MaxMask, Mask1[Key]);
+					bRepeat = true;
+					++NumFound;
+				}
+			}
+		}
+		Mask = Mask1;
+		// Avoid a perpetual loop (e.g. can happen if initially all the mask values are 1... though that shouldn't
+		// really happen).
+		if (NumFound >= PrevNumFound)
+		{
+			bRepeat = false;
+		}
+		else
+		{
+			PrevNumFound = NumFound;
+		}
+	}
+
+	// When searching we will want to start outside of a "good" region.
+	int32 AZeroKey = 0;
+	for (int32 Key = 0 ; Key != NumSampledKeys ; ++Key)
+	{
+		if (Mask[Key] == 0)
+		{
+			AZeroKey = Key;
+		}
+	}
+
+	// We use the mask (with a somewhat arbitrary threshold) to get rid of velocities that are near to the foot
+	// plant/take-off time (and might be when the foot is in the time). Then We look for the highest velocity. Note that
+	// if we're being called with the foot (ankle) joint, it will tend to underestimate velocities since it is nearer
+	// the hip than the ground contact point.
+	// 
+	int32 Threshold = FMath::Max(MaxMask / 2, 1);
+	int32 Num = 0;
+	FVector AverageFootVel(0);
+	float BestSpeed = 0.0f;
+	int32 BestSpeedKey = 0;
+	for (int32 K = AZeroKey ; K != AZeroKey + NumSampledKeys ; ++K)
+	{
+		int32 Key = K % NumSampledKeys;
+		if (Mask[Key] >= Threshold)
+		{
+			float Speed = Velocities[Key].Size();
+#ifdef ANALYSIS_VERBOSE_LOG
+			UE_LOG(LogAnimation, Log, TEXT("Candidate %d Mask %d vel = %f %f %f, speed = %f"), 
+				   Key, Mask[Key], Velocities[Key].X, Velocities[Key].Y, Velocities[Key].Z, Speed);
+#endif
+			if (Speed > BestSpeed)
+			{
+				BestSpeed = Speed;
+				BestSpeedKey = Key;
+			}
+		}
+		else
+		{
+			if (BestSpeed > 0.0f)
+			{
+#ifdef ANALYSIS_VERBOSE_LOG
+				UE_LOG(LogAnimation, Log, TEXT("Picked Candidate %d vel = %f %f %f, speed = %f"), 
+					   BestSpeedKey, Velocities[BestSpeedKey].X, Velocities[BestSpeedKey].Y, Velocities[BestSpeedKey].Z, BestSpeed);
+#endif
+				AverageFootVel += Velocities[BestSpeedKey];
+				++Num;
+				BestSpeed = 0.0f;
+			}
+		}
+	}
+	// Make sure we didn't miss the last data point
+	if (BestSpeed > 0.0f)
+	{
+#ifdef ANALYSIS_VERBOSE_LOG
+		UE_LOG(LogAnimation, Log, TEXT("Picked Candidate %d vel = %f %f %f, speed = %f"), 
+			   BestSpeedKey, Velocities[BestSpeedKey].X, Velocities[BestSpeedKey].Y, Velocities[BestSpeedKey].Z, BestSpeed);
+#endif
+		AverageFootVel += Velocities[BestSpeedKey];
+		++Num;
+		BestSpeed = 0.0f;
+	}
+
+	AverageFootVel /= Num;
+	float FacingVel = -AverageFootVel | FrameFacingDir;
+	float RightVel = -AverageFootVel | FrameRightDir;
+	float UpVel = -AverageFootVel | FrameUpDir;
+
+	Result.Set(FacingVel, RightVel, UpVel);
+	Result *= Animation.RateScale * RateScale;
+	UE_LOG(LogAnimation, Log, TEXT("%s Locomotion vel = %f %f %f"), *BoneName.ToString(), Result.X, Result.Y, Result.Z);
+	return true;
+}
+
+//======================================================================================================================
+static bool CalculateLocomotionVelocity(
+	FVector&                             Result,
+	const UBlendSpace&                   BlendSpace,
+	const ULocomotionAnalysisProperties* AnalysisProperties,
+	const UAnimSequence&                 Animation,
+	const float                          RateScale)
+{
+	if (!AnalysisProperties)
+	{
+		return false;
+	}
+
+	FVector Result1, Result2;
+	int32 Num = 0;
+	Result.Set(0, 0, 0);
+	if (CalculateLocomotionVelocity(Result1, BlendSpace, AnalysisProperties, 
+									AnalysisProperties->PrimaryBoneSocket, Animation, RateScale))
+	{
+#ifdef ANALYSIS_VERBOSE_LOG
+		UE_LOG(LogAnimation, Log, TEXT("Loco vel from primary = %f %f %f"), Result1.X, Result1.Y, Result1.Z);
+#endif
+		Result += Result1;
+		++Num;
+	}
+	if (CalculateLocomotionVelocity(Result2, BlendSpace, AnalysisProperties, 
+									AnalysisProperties->SecondaryBoneSocket, Animation, RateScale))
+	{
+#ifdef ANALYSIS_VERBOSE_LOG
+		UE_LOG(LogAnimation, Log, TEXT("Loco vel from secondary = %f %f %f"), Result2.X, Result2.Y, Result2.Z);
+#endif
+		Result += Result2;
+		++Num;
+	}
+	if (Num)
+	{
+		Result /= Num;
+		UE_LOG(LogAnimation, Log, TEXT("Loco vel = %f %f %f"), Result.X, Result.Y, Result.Z);
+		return true;
+	}
+	return false;
+}
+
+//======================================================================================================================
+// Calculates the locomotion speed (magnitude)
+static bool CalculateLocomotionSpeed(
+	float&                               Result,
+	const UBlendSpace&                   BlendSpace,
+	const ULocomotionAnalysisProperties* AnalysisProperties,
+	const UAnimSequence&                 Animation,
+	const float                          RateScale)
+{
+	FVector Movement;
+	if (CalculateLocomotionVelocity(Movement, BlendSpace, AnalysisProperties, Animation, RateScale))
+	{
+		Result = Movement.Size();
+		return true;
+	}
+	return false;
+}
+
+//======================================================================================================================
+// Calculates the locomotion direction (degrees)
+static bool CalculateLocomotionDirection(
+	float&                               Result,
+	const UBlendSpace&                   BlendSpace,
+	const ULocomotionAnalysisProperties* AnalysisProperties,
+	const UAnimSequence&                 Animation,
+	const float                          RateScale)
+{
+	FVector Movement;
+	if (CalculateLocomotionVelocity(Movement, BlendSpace, AnalysisProperties, Animation, RateScale))
+	{
+		Result = FMath::RadiansToDegrees(FMath::Atan2(Movement.Y, Movement.X));
+		return true;
+	}
+	return false;
+}
+
+//======================================================================================================================
+// Calculates the locomotion speed in the character's facing direction
+static bool CalculateLocomotionForwardSpeed(
+	float&                               Result,
+	const UBlendSpace&                   BlendSpace,
+	const ULocomotionAnalysisProperties* AnalysisProperties,
+	const UAnimSequence&                 Animation,
+	const float                          RateScale)
+{
+	FVector Movement;
+	if (CalculateLocomotionVelocity(Movement, BlendSpace, AnalysisProperties, Animation, RateScale))
+	{
+		Result = Movement.X;
+		return true;
+	}
+	return false;
+}
+
+//======================================================================================================================
+// Calculates the locomotion speed in the character's upwards direction
+static bool CalculateLocomotionUpwardSpeed(
+	float&                               Result,
+	const UBlendSpace&                   BlendSpace,
+	const ULocomotionAnalysisProperties* AnalysisProperties,
+	const UAnimSequence&                 Animation,
+	const float                          RateScale)
+{
+	FVector Movement;
+	if (CalculateLocomotionVelocity(Movement, BlendSpace, AnalysisProperties, Animation, RateScale))
+	{
+		Result = Movement.Z;
+		return true;
+	}
+	return false;
+}
+
+//======================================================================================================================
+// Calculates the locomotion speed in the character's right direction
+static bool CalculateLocomotionRightwardSpeed(
+	float&                               Result,
+	const UBlendSpace&                   BlendSpace,
+	const ULocomotionAnalysisProperties* AnalysisProperties,
+	const UAnimSequence&                 Animation,
+	const float                          RateScale)
+{
+	FVector Movement;
+	if (CalculateLocomotionVelocity(Movement, BlendSpace, AnalysisProperties, Animation, RateScale))
+	{
+		Result = Movement.Y;
+		return true;
+	}
+	return false;
+}
+
+//======================================================================================================================
+// Calculates the locomotion slope angle (degrees) going in the facing direction
+static bool CalculateLocomotionForwardSlope(
+	float&                               Result,
+	const UBlendSpace&                   BlendSpace,
+	const ULocomotionAnalysisProperties* AnalysisProperties,
+	const UAnimSequence&                 Animation,
+	const float                          RateScale)
+{
+	FVector Movement;
+	if (CalculateLocomotionVelocity(Movement, BlendSpace, AnalysisProperties, Animation, RateScale))
+	{
+		if (Movement.X >= 0.0f)
+		{
+			Result = FMath::RadiansToDegrees(FMath::Atan2(Movement.Z, Movement.X));
+		}
+		else
+		{
+			Result = FMath::RadiansToDegrees(FMath::Atan2(-Movement.Z, -Movement.X));
+		}
+
+		return true;
+	}
+	return false;
+}
+
+//======================================================================================================================
+// Calculates the locomotion slope angle (degrees) going in the rightwards direction
+static bool CalculateLocomotionRightwardSlope(
+	float&                               Result,
+	const UBlendSpace&                   BlendSpace,
+	const ULocomotionAnalysisProperties* AnalysisProperties,
+	const UAnimSequence&                 Animation,
+	const float                          RateScale)
+{
+	FVector Movement;
+	if (CalculateLocomotionVelocity(Movement, BlendSpace, AnalysisProperties, Animation, RateScale))
+	{
+		if (Movement.Y > 0.0f)
+		{
+			Result = FMath::RadiansToDegrees(FMath::Atan2(Movement.Z, Movement.Y));
+		}
+		else
+		{
+			Result = FMath::RadiansToDegrees(FMath::Atan2(-Movement.Z, -Movement.Y));
+		}
+		return true;
+	}
+	return false;
+}
+
+//======================================================================================================================
+bool CalculateLocomotion(
+	float&                               Result,
+	const UBlendSpace&                   BlendSpace,
+	const ULocomotionAnalysisProperties* AnalysisProperties,
+	const UAnimSequence&                 Animation,
+	const float                          RateScale)
+{
+	if (!AnalysisProperties)
+	{
+		return false;
+
+	}
+	switch (AnalysisProperties->FunctionAxis)
+	{
+	case EAnalysisLocomotionAxis::Speed:
+		return CalculateLocomotionSpeed(Result, BlendSpace, AnalysisProperties, Animation, RateScale);
+	case EAnalysisLocomotionAxis::Direction:
+		return CalculateLocomotionDirection(Result, BlendSpace, AnalysisProperties, Animation, RateScale);
+	case EAnalysisLocomotionAxis::ForwardSpeed:
+		return CalculateLocomotionForwardSpeed(Result, BlendSpace, AnalysisProperties, Animation, RateScale);
+	case EAnalysisLocomotionAxis::RightwardSpeed:
+		return CalculateLocomotionRightwardSpeed(Result, BlendSpace, AnalysisProperties, Animation, RateScale);
+	case EAnalysisLocomotionAxis::UpwardSpeed:
+		return CalculateLocomotionUpwardSpeed(Result, BlendSpace, AnalysisProperties, Animation, RateScale);
+	case EAnalysisLocomotionAxis::ForwardSlope:
+		return CalculateLocomotionForwardSlope(Result, BlendSpace, AnalysisProperties, Animation, RateScale);
+	case EAnalysisLocomotionAxis::RightwardSlope:
+		return CalculateLocomotionRightwardSlope(Result, BlendSpace, AnalysisProperties, Animation, RateScale);
+	default:
+		return false;
+	}
+}
+
+#undef LOCTEXT_NAMESPACE
\ No newline at end of file
diff --git a/Engine/Plugins/Animation/BlendSpaceMotionAnalysis/Source/BlendSpaceMotionAnalysis/Private/RootMotionAnalysis.cpp b/Engine/Plugins/Animation/BlendSpaceMotionAnalysis/Source/BlendSpaceMotionAnalysis/Private/RootMotionAnalysis.cpp
new file mode 100644
--- /dev/null
+++ b/Engine/Plugins/Animation/BlendSpaceMotionAnalysis/Source/BlendSpaceMotionAnalysis/Private/RootMotionAnalysis.cpp
@@ -0,0 +1,273 @@
+// Copyright Epic Games, Inc. All Rights Reserved.
+
+#include "RootMotionAnalysis.h"
+#include "BlendSpaceMotionAnalysis.h"
+#include "Engine/Engine.h"
+#include "Logging/LogMacros.h"
+
+DEFINE_LOG_CATEGORY_STATIC(LogRootMotionAnalysis, Log, All);
+
+#define LOCTEXT_NAMESPACE "RootMotionAnalysis"
+
+//======================================================================================================================
+void URootMotionAnalysisProperties::InitializeFromCache(TObjectPtr<UCachedAnalysisProperties> Cache)
+{
+	if (Cache)
+	{
+		UCachedAnalysisProperties* CachePtr = Cache.Get();
+		UCachedMotionAnalysisProperties* CachedMotionAnalysisPropertiesPtr = Cast<UCachedMotionAnalysisProperties>(CachePtr);
+		if (!CachedMotionAnalysisPropertiesPtr)
+		{
+			// Replace any existing cache if it is not our own type
+			CachedMotionAnalysisPropertiesPtr = NewObject<UCachedMotionAnalysisProperties>(CachePtr->GetOuter());
+			CachedMotionAnalysisPropertiesPtr->CopyFrom(*CachePtr);
+			Cache = CachedMotionAnalysisPropertiesPtr;
+		}
+		FunctionAxis = CachedMotionAnalysisPropertiesPtr->RootMotionFunctionAxis;
+		BoneSocket = CachedMotionAnalysisPropertiesPtr->BoneSocket1;
+		Space = CachedMotionAnalysisPropertiesPtr->Space;
+		SpaceBoneSocket = CachedMotionAnalysisPropertiesPtr->SpaceBoneSocket;
+		CharacterFacingAxis = CachedMotionAnalysisPropertiesPtr->CharacterFacingAxis;
+		CharacterUpAxis = CachedMotionAnalysisPropertiesPtr->CharacterUpAxis;
+		StartTimeFraction = CachedMotionAnalysisPropertiesPtr->StartTimeFraction;
+		EndTimeFraction = CachedMotionAnalysisPropertiesPtr->EndTimeFraction;
+		bLockAfterAnalysis = CachedMotionAnalysisPropertiesPtr->bLockAfterAnalysis;
+	}
+}
+
+//======================================================================================================================
+void URootMotionAnalysisProperties::MakeCache(TObjectPtr<UCachedAnalysisProperties>& Cache, UBlendSpace* BlendSpace)
+{
+	UCachedMotionAnalysisProperties* CachedMotionAnalysisPropertiesPtr = Cast<UCachedMotionAnalysisProperties>(Cache.Get());
+	if (!CachedMotionAnalysisPropertiesPtr)
+	{
+		// Replace any existing cache if it is not our own type
+		CachedMotionAnalysisPropertiesPtr = NewObject<UCachedMotionAnalysisProperties>(BlendSpace);
+		if (Cache.Get())
+		{
+			CachedMotionAnalysisPropertiesPtr->CopyFrom(*Cache.Get());
+		}
+		Cache = CachedMotionAnalysisPropertiesPtr;
+	}
+	CachedMotionAnalysisPropertiesPtr->RootMotionFunctionAxis = FunctionAxis;
+	CachedMotionAnalysisPropertiesPtr->BoneSocket1 = BoneSocket;
+	CachedMotionAnalysisPropertiesPtr->Space = Space;
+	CachedMotionAnalysisPropertiesPtr->SpaceBoneSocket = SpaceBoneSocket;
+	CachedMotionAnalysisPropertiesPtr->CharacterFacingAxis = CharacterFacingAxis;
+	CachedMotionAnalysisPropertiesPtr->CharacterUpAxis = CharacterUpAxis;
+	CachedMotionAnalysisPropertiesPtr->StartTimeFraction = StartTimeFraction;
+	CachedMotionAnalysisPropertiesPtr->EndTimeFraction = EndTimeFraction;
+	CachedMotionAnalysisPropertiesPtr->bLockAfterAnalysis = bLockAfterAnalysis;
+}
+
+//======================================================================================================================
+// Calculates the movement speed (magnitude) 
+static bool CalculateRootMotionVelocity(
+	FVector&                             Result,
+	const UBlendSpace&                   BlendSpace,
+	const URootMotionAnalysisProperties* AnalysisProperties,
+	const UAnimSequence&                 Animation,
+	const float                          RateScale)
+{
+	FVector Velocity;
+	if (BlendSpaceAnalysis::CalculateVelocity(Velocity, BlendSpace, AnalysisProperties, Animation, RateScale))
+	{
+		FTransform FrameTM;
+		FVector FrameFacingDir = BlendSpaceAnalysis::GetAxisFromTM(FrameTM, AnalysisProperties->CharacterFacingAxis);
+		FVector FrameUpDir = BlendSpaceAnalysis::GetAxisFromTM(FrameTM, AnalysisProperties->CharacterUpAxis);
+		FVector FrameRightDir = FVector::CrossProduct(FrameUpDir, FrameFacingDir);
+		float ForwardVel = Velocity | FrameFacingDir;
+		float RightVel = Velocity | FrameRightDir;
+		float UpVel = Velocity | FrameUpDir;
+		Result.Set(ForwardVel, RightVel, UpVel);
+		return true;
+	}
+	return false;
+}
+
+//======================================================================================================================
+// Calculates the movement speed (magnitude) 
+static bool CalculateRootMotionSpeed(
+	float&                               Result,
+	const UBlendSpace&                   BlendSpace,
+	const URootMotionAnalysisProperties* AnalysisProperties,
+	const UAnimSequence&                 Animation,
+	const float                          RateScale)
+{
+	FVector Movement;
+	if (CalculateRootMotionVelocity(Movement, BlendSpace, AnalysisProperties, Animation, RateScale))
+	{
+		Result = Movement.Size();
+		return true;
+	}
+	return false;
+}
+
+//======================================================================================================================
+// Calculates the movement speed (magnitude) 
+static bool CalculateRootMotionDirection(
+	float&                               Result,
+	const UBlendSpace&                   BlendSpace,
+	const URootMotionAnalysisProperties* AnalysisProperties,
+	const UAnimSequence&                 Animation,
+	const float                          RateScale)
+{
+	FVector Movement;
+	if (CalculateRootMotionVelocity(Movement, BlendSpace, AnalysisProperties, Animation, RateScale))
+	{
+		Result = FMath::RadiansToDegrees(FMath::Atan2(Movement.Y, Movement.X));
+		return true;
+	}
+	return false;
+}
+
+//======================================================================================================================
+// Calculates the locomotion speed in the character's facing direction
+static bool CalculateRootMotionForwardSpeed(
+	float&                               Result,
+	const UBlendSpace&                   BlendSpace,
+	const URootMotionAnalysisProperties* AnalysisProperties,
+	const UAnimSequence&                 Animation,
+	const float                          RateScale)
+{
+	FVector Movement;
+	if (CalculateRootMotionVelocity(Movement, BlendSpace, AnalysisProperties, Animation, RateScale))
+	{
+		Result = Movement.X;
+		return true;
+	}
+	return false;
+}
+
+//======================================================================================================================
+// Calculates the locomotion speed in the character's upwards direction
+static bool CalculateRootMotionUpwardSpeed(
+	float&                               Result,
+	const UBlendSpace&                   BlendSpace,
+	const URootMotionAnalysisProperties* AnalysisProperties,
+	const UAnimSequence&                 Animation,
+	const float                          RateScale)
+{
+	FVector Movement;
+	if (CalculateRootMotionVelocity(Movement, BlendSpace, AnalysisProperties, Animation, RateScale))
+	{
+		Result = Movement.Z;
+		return true;
+	}
+	return false;
+}
+
+//======================================================================================================================
+// Calculates the locomotion speed in the character's right direction
+static bool CalculateRootMotionRightwardSpeed(
+	float&                               Result,
+	const UBlendSpace&                   BlendSpace,
+	const URootMotionAnalysisProperties* AnalysisProperties,
+	const UAnimSequence&                 Animation,
+	const float                          RateScale)
+{
+	FVector Movement;
+	if (CalculateRootMotionVelocity(Movement, BlendSpace, AnalysisProperties, Animation, RateScale))
+	{
+		Result = Movement.Y;
+		return true;
+	}
+	return false;
+}
+
+//======================================================================================================================
+// Calculates the locomotion slope angle (degrees) going in the facing direction
+static bool CalculateRootMotionForwardSlope(
+	float&                               Result,
+	const UBlendSpace&                   BlendSpace,
+	const URootMotionAnalysisProperties* AnalysisProperties,
+	const UAnimSequence&                 Animation,
+	const float                          RateScale)
+{
+	FVector Movement;
+	if (CalculateRootMotionVelocity(Movement, BlendSpace, AnalysisProperties, Animation, RateScale))
+	{
+		if (Movement.X >= 0.0f)
+		{
+			Result = FMath::RadiansToDegrees(FMath::Atan2(Movement.Z, Movement.X));
+		}
+		else
+		{
+			Result = FMath::RadiansToDegrees(FMath::Atan2(-Movement.Z, -Movement.X));
+		}
+
+		return true;
+	}
+	return false;
+}
+
+//======================================================================================================================
+// Calculates the locomotion slope angle (degrees) going in the rightwards direction
+static bool CalculateRootMotionRightwardSlope(
+	float&                               Result,
+	const UBlendSpace&                   BlendSpace,
+	const URootMotionAnalysisProperties* AnalysisProperties,
+	const UAnimSequence&                 Animation,
+	const float                          RateScale)
+{
+	FVector Movement;
+	if (CalculateRootMotionVelocity(Movement, BlendSpace, AnalysisProperties, Animation, RateScale))
+	{
+		if (Movement.Y > 0.0f)
+		{
+			Result = FMath::RadiansToDegrees(FMath::Atan2(Movement.Z, Movement.Y));
+		}
+		else
+		{
+			Result = FMath::RadiansToDegrees(FMath::Atan2(-Movement.Z, -Movement.Y));
+		}
+		return true;
+	}
+	return false;
+}
+
+
+//======================================================================================================================
+// Calculates the movement speed (magnitude) 
+bool CalculateRootMotion(
+	float&                               Result,
+	const UBlendSpace&                   BlendSpace,
+	const URootMotionAnalysisProperties* AnalysisProperties,
+	const UAnimSequence&                 Animation,
+	const float                          RateScale)
+{
+	if (!AnalysisProperties)
+	{
+		return false;
+	}
+
+	FVector Velocity;
+	switch (AnalysisProperties->FunctionAxis)
+	{
+	case EAnalysisRootMotionAxis::Speed:
+		return CalculateRootMotionSpeed(Result, BlendSpace, AnalysisProperties, Animation, RateScale);
+		break;
+	case EAnalysisRootMotionAxis::Direction:
+		return CalculateRootMotionDirection(Result, BlendSpace, AnalysisProperties, Animation, RateScale);
+		break;
+	case EAnalysisRootMotionAxis::ForwardSpeed:
+		return CalculateRootMotionForwardSpeed(Result, BlendSpace, AnalysisProperties, Animation, RateScale);
+		break;
+	case EAnalysisRootMotionAxis::RightwardSpeed:
+		return CalculateRootMotionRightwardSpeed(Result, BlendSpace, AnalysisProperties, Animation, RateScale);
+		break;
+	case EAnalysisRootMotionAxis::UpwardSpeed:
+		return CalculateRootMotionUpwardSpeed(Result, BlendSpace, AnalysisProperties, Animation, RateScale);
+		break;
+	case EAnalysisRootMotionAxis::ForwardSlope:
+		return CalculateRootMotionForwardSlope(Result, BlendSpace, AnalysisProperties, Animation, RateScale);
+		break;
+	case EAnalysisRootMotionAxis::RightwardSlope:
+		return CalculateRootMotionRightwardSlope(Result, BlendSpace, AnalysisProperties, Animation, RateScale);
+		break;
+	default:
+		return false;
+	}
+}
+
+#undef LOCTEXT_NAMESPACE
\ No newline at end of file
diff --git a/Engine/Plugins/Animation/BlendSpaceMotionAnalysis/Source/BlendSpaceMotionAnalysis/Public/BlendSpaceMotionAnalysis.h b/Engine/Plugins/Animation/BlendSpaceMotionAnalysis/Source/BlendSpaceMotionAnalysis/Public/BlendSpaceMotionAnalysis.h
new file mode 100644
--- /dev/null
+++ b/Engine/Plugins/Animation/BlendSpaceMotionAnalysis/Source/BlendSpaceMotionAnalysis/Public/BlendSpaceMotionAnalysis.h
@@ -0,0 +1,28 @@
+// Copyright Epic Games, Inc. All Rights Reserved.
+
+#pragma once
+
+#include "CoreMinimal.h"
+#include "Modules/ModuleInterface.h"
+#include "BlendSpaceAnalysis.h"
+#include "RootMotionAnalysis.h"
+#include "LocomotionAnalysis.h"
+#include "BlendSpaceMotionAnalysis.generated.h"
+
+//======================================================================================================================
+UCLASS()
+class UCachedMotionAnalysisProperties : public UCachedAnalysisProperties
+{
+	GENERATED_BODY()
+public:
+	EAnalysisRootMotionAxis RootMotionFunctionAxis = EAnalysisRootMotionAxis::Speed;
+	EAnalysisLocomotionAxis LocomotionFunctionAxis = EAnalysisLocomotionAxis::Speed;
+};
+
+//======================================================================================================================
+class BLENDSPACEMOTIONANALYSIS_API FBlendSpaceMotionAnalysis : public IModuleInterface 
+{
+public:
+	void StartupModule() override;
+	void ShutdownModule() override;
+};
diff --git a/Engine/Plugins/Animation/BlendSpaceMotionAnalysis/Source/BlendSpaceMotionAnalysis/Public/LocomotionAnalysis.h b/Engine/Plugins/Animation/BlendSpaceMotionAnalysis/Source/BlendSpaceMotionAnalysis/Public/LocomotionAnalysis.h
new file mode 100644
--- /dev/null
+++ b/Engine/Plugins/Animation/BlendSpaceMotionAnalysis/Source/BlendSpaceMotionAnalysis/Public/LocomotionAnalysis.h
@@ -0,0 +1,60 @@
+// Copyright Epic Games, Inc. All Rights Reserved.
+
+#pragma once
+
+#include "CoreMinimal.h"
+#include "Modules/ModuleInterface.h"
+#include "BlendSpaceAnalysis.h"
+#include "LocomotionAnalysis.generated.h"
+
+//======================================================================================================================
+UENUM()
+enum class EAnalysisLocomotionAxis : uint8
+{
+	Speed,
+	Direction,
+	ForwardSpeed,
+	RightwardSpeed,
+	UpwardSpeed,
+	ForwardSlope,
+	RightwardSlope,
+};
+
+//======================================================================================================================
+UCLASS()
+class ULocomotionAnalysisProperties : public UAnalysisProperties
+{
+	GENERATED_BODY()
+public:
+	void InitializeFromCache(TObjectPtr<UCachedAnalysisProperties> Cache) override;
+	void MakeCache(TObjectPtr<UCachedAnalysisProperties>& Cache, UBlendSpace* BlendSpace) override;
+
+	/** Axis for the analysis function */
+	UPROPERTY(EditAnywhere, DisplayName = "Axis", Category = AnalysisProperties)
+		EAnalysisLocomotionAxis FunctionAxis = EAnalysisLocomotionAxis::Speed;
+
+	/** The primary bone or socket used for analysis */
+	UPROPERTY(EditAnywhere, DisplayName = "Bone/Socket 1", Category = AnalysisProperties)
+		FBoneSocketTarget PrimaryBoneSocket;
+
+	/** The secondary bone or socket used for analysis */
+	UPROPERTY(EditAnywhere, DisplayName = "Bone/Socket 2", Category = AnalysisProperties)
+		FBoneSocketTarget SecondaryBoneSocket;
+
+	/** World or bone/socket axis that specifies the character's facing direction */
+	UPROPERTY(EditAnywhere, Category = AnalysisProperties)
+		EAnalysisLinearAxis CharacterFacingAxis = EAnalysisLinearAxis::Y;
+
+	/** World or bone/socket axis that specifies the character's up direction */
+	UPROPERTY(EditAnywhere, Category = AnalysisProperties)
+		EAnalysisLinearAxis CharacterUpAxis = EAnalysisLinearAxis::Z;
+};
+
+//======================================================================================================================
+bool CalculateLocomotion(
+	float&                               Result,
+	const UBlendSpace&                   BlendSpace,
+	const ULocomotionAnalysisProperties* AnalysisProperties,
+	const UAnimSequence&                 Animation,
+	const float                          RateScale);
+
diff --git a/Engine/Plugins/Animation/BlendSpaceMotionAnalysis/Source/BlendSpaceMotionAnalysis/Public/RootMotionAnalysis.h b/Engine/Plugins/Animation/BlendSpaceMotionAnalysis/Source/BlendSpaceMotionAnalysis/Public/RootMotionAnalysis.h
new file mode 100644
--- /dev/null
+++ b/Engine/Plugins/Animation/BlendSpaceMotionAnalysis/Source/BlendSpaceMotionAnalysis/Public/RootMotionAnalysis.h
@@ -0,0 +1,76 @@
+// Copyright Epic Games, Inc. All Rights Reserved.
+
+#pragma once
+
+#include "CoreMinimal.h"
+#include "Modules/ModuleInterface.h"
+#include "BlendSpaceAnalysis.h"
+#include "RootMotionAnalysis.generated.h"
+
+//======================================================================================================================
+UENUM()
+enum class EAnalysisRootMotionAxis : uint8
+{
+	Speed,
+	Direction,
+	ForwardSpeed,
+	RightwardSpeed,
+	UpwardSpeed,
+	ForwardSlope,
+	RightwardSlope,
+};
+
+//======================================================================================================================
+UCLASS()
+class URootMotionAnalysisProperties : public UAnalysisProperties
+{
+	GENERATED_BODY()
+public:
+	void InitializeFromCache(TObjectPtr<UCachedAnalysisProperties> Cache) override;
+	void MakeCache(TObjectPtr<UCachedAnalysisProperties>& Cache, UBlendSpace* BlendSpace) override;
+
+	/** Axis for the analysis function */
+	UPROPERTY(EditAnywhere, DisplayName = "Axis", Category = AnalysisProperties)
+		EAnalysisRootMotionAxis FunctionAxis = EAnalysisRootMotionAxis::Speed;
+
+	/** The bone or socket used for analysis */
+	UPROPERTY(EditAnywhere, DisplayName = "Bone/Socket", Category = AnalysisProperties)
+		FBoneSocketTarget BoneSocket;
+
+	/**
+	* The space in which to perform the analysis. Fixed will use the analysis bone/socket at the first frame
+	* of the analysis time range. Changing will use the analysis bone/socket at the relevant frame during the
+	* analysis, but calculate velocities assuming that frame isn't moving. Moving will do the same but velocities
+	* as well as positions/rotations will be relative to this moving frame.
+	*/
+	UPROPERTY(EditAnywhere, Category = AnalysisProperties)
+		EAnalysisSpace Space = EAnalysisSpace::World;
+
+	/** Bone or socket that defines the analysis space (when it isn't World) */
+	UPROPERTY(EditAnywhere, DisplayName = "Analysis Space Bone/Socket", Category = AnalysisProperties, meta = (EditCondition = "Space != EAnalysisSpace::World"))
+		FBoneSocketTarget SpaceBoneSocket;
+
+	/** World or bone/socket axis that specifies the character's facing direction */
+	UPROPERTY(EditAnywhere, Category = AnalysisProperties)
+		EAnalysisLinearAxis CharacterFacingAxis = EAnalysisLinearAxis::Y;
+
+	/** World or bone/socket axis that specifies the character's up direction */
+	UPROPERTY(EditAnywhere, Category = AnalysisProperties)
+		EAnalysisLinearAxis CharacterUpAxis = EAnalysisLinearAxis::Z;
+
+	/** Fraction through each animation at which analysis starts */
+	UPROPERTY(EditAnywhere, DisplayName = "Start time fraction", Category = AnalysisProperties, meta = (ClampMin = "0", ClampMax = "1"))
+		float StartTimeFraction = 0.0f;
+
+	/** Fraction through each animation at which analysis ends */
+	UPROPERTY(EditAnywhere, DisplayName = "End time fraction", Category = AnalysisProperties, meta = (ClampMin = "0", ClampMax = "1"))
+		float EndTimeFraction = 1.0f;
+};
+
+//======================================================================================================================
+bool CalculateRootMotion(
+	float&                               Result,
+	const UBlendSpace&                   BlendSpace,
+	const URootMotionAnalysisProperties* AnalysisProperties,
+	const UAnimSequence&                 Animation,
+	const float                          RateScale);
diff --git a/Engine/Source/Editor/AnimationBlueprintEditor/Private/BlendSpaceDocumentTabFactory.cpp b/Engine/Source/Editor/AnimationBlueprintEditor/Private/BlendSpaceDocumentTabFactory.cpp
--- a/Engine/Source/Editor/AnimationBlueprintEditor/Private/BlendSpaceDocumentTabFactory.cpp
+++ b/Engine/Source/Editor/AnimationBlueprintEditor/Private/BlendSpaceDocumentTabFactory.cpp
@@ -71,7 +71,7 @@ TSharedRef<SWidget> FBlendSpaceDocumentTabFactory::CreateTabBody(const FWorkflow
 		}
 	});
 
-	Args.OnBlendSpaceSampleAdded = FOnBlendSpaceSampleAdded::CreateLambda([this, WeakBlendSpaceNode = TWeakObjectPtr<UAnimGraphNode_BlendSpaceGraphBase>(BlendSpaceNode)](UAnimSequence* InSequence, const FVector& InSamplePoint) -> int32
+	Args.OnBlendSpaceSampleAdded = FOnBlendSpaceSampleAdded::CreateLambda([this, WeakBlendSpaceNode = TWeakObjectPtr<UAnimGraphNode_BlendSpaceGraphBase>(BlendSpaceNode)](UAnimSequence* InSequence, const FVector& InSamplePoint, bool bRunAnalysis) -> int32
 	{
 		int32 Index = INDEX_NONE;
 		if(WeakBlendSpaceNode.Get())
diff --git a/Engine/Source/Editor/Persona/Private/Customization/BlendSampleDetails.cpp b/Engine/Source/Editor/Persona/Private/Customization/BlendSampleDetails.cpp
--- a/Engine/Source/Editor/Persona/Private/Customization/BlendSampleDetails.cpp
+++ b/Engine/Source/Editor/Persona/Private/Customization/BlendSampleDetails.cpp
@@ -44,21 +44,13 @@ FReply FBlendSampleDetails::HandleAnalyzeAndDuplicateSample()
 		// Dismiss menus so that operations which might affect samples (and cause reallocation etc) won't invalidate
 		// data being used in the UI.
 		FSlateApplication::Get().DismissAllMenus();
-		bool bAnalyzed[3] = { false, false, false };
 		FVector OrigValue = BlendSpace->GetBlendSample(SampleIndex).SampleValue;
-		FVector NewValue = BlendSpaceAnalysis::CalculateSampleValue(
-			*BlendSpace, *BlendSpace->GetBlendSample(SampleIndex).Animation, 
-			BlendSpace->GetBlendSample(SampleIndex).RateScale, OrigValue, bAnalyzed);
-		NewValue.Z = OrigValue.Z;
-		if (NewValue != OrigValue)
-		{
-			GridWidget->OnSampleDuplicated.ExecuteIfBound(SampleIndex, NewValue);
-		}
+		GridWidget->OnSampleDuplicated.ExecuteIfBound(SampleIndex, OrigValue, true);
 	}
 	return FReply::Handled();
 }
 
-FReply FBlendSampleDetails::HandleAnalyzeSample()
+FReply FBlendSampleDetails::HandleAnalyzeAndMoveSample()
 {
 	if (BlendSpace->IsAsset())
 	{
@@ -76,7 +68,7 @@ FReply FBlendSampleDetails::HandleAnalyzeSample()
 	return FReply::Handled();
 }
 
-FReply FBlendSampleDetails::HandleAnalyzeSampleX()
+FReply FBlendSampleDetails::HandleAnalyzeAndMoveSampleX()
 {
 	if (BlendSpace->IsAsset())
 	{
@@ -95,7 +87,7 @@ FReply FBlendSampleDetails::HandleAnalyzeSampleX()
 	return FReply::Handled();
 }
 
-FReply FBlendSampleDetails::HandleAnalyzeSampleY()
+FReply FBlendSampleDetails::HandleAnalyzeAndMoveSampleY()
 {
 	if (BlendSpace->IsAsset())
 	{
@@ -218,10 +210,11 @@ void FBlendSampleDetails::CustomizeDetails(class IDetailLayoutBuilder& DetailBui
 
 		bool bShowAnalysis = false;
 		bool bAnalyzed[3] = { false, false, false };
+		// Note that the analyzed position won't change whilst this menu is open, but the original value might.
+		FVector OrigValue = BlendSpace->GetBlendSample(SampleIndex).SampleValue;
 		FVector NewValue = BlendSpaceAnalysis::CalculateSampleValue(
-			*BlendSpace, *BlendSpace->GetBlendSample(SampleIndex).Animation,
-			BlendSpace->GetBlendSample(SampleIndex).RateScale, 
-			BlendSpace->GetBlendSample(SampleIndex).SampleValue, bAnalyzed);
+			*BlendSpace, *BlendSpace->GetBlendSample(SampleIndex).Animation, 
+			BlendSpace->GetBlendSample(SampleIndex).RateScale, OrigValue, bAnalyzed);
 		FText AnalysisTexts[2];
 		FText ValueTexts[2];
 		FText ToolTipTexts[2];
@@ -250,26 +243,46 @@ void FBlendSampleDetails::CustomizeDetails(class IDetailLayoutBuilder& DetailBui
 			}
 		}
 
-		if (bShowAnalysis && NewValue != BlendSpace->GetBlendSample(SampleIndex).SampleValue)
+		TAttribute<bool> WouldMoveConditionX = TAttribute<bool>::Create(
+			[this, bAnalyzed, NewValue]()
+			{
+				FVector OrigValue = BlendSpace->GetBlendSample(SampleIndex).SampleValue;
+				return bAnalyzed[0] && OrigValue.X != NewValue.X;
+			});
+		TAttribute<bool> WouldMoveConditionY = TAttribute<bool>::Create(
+			[this, bAnalyzed, NewValue]()
+			{
+				FVector OrigValue = BlendSpace->GetBlendSample(SampleIndex).SampleValue;
+				return bAnalyzed[1] && OrigValue.Y != NewValue.Y;
+			});
+		TAttribute<bool> WouldMoveCondition = TAttribute<bool>::Create(
+			[this, NewValue]()
+			{
+				FVector OrigValue = BlendSpace->GetBlendSample(SampleIndex).SampleValue;
+				return OrigValue != NewValue;
+			});
+
+
+		if (bShowAnalysis)
 		{
-			// Button to run analysis
+			// Move buttons
 			CategoryBuilder
-				.AddGroup(FName("AnalysisGroup"), FText::GetEmpty())
+				.AddGroup(FName("MoveGroup"), FText::GetEmpty())
 				.HeaderRow()
 				.NameContent()
 				[
 					SNew(SHorizontalBox)
 					+SHorizontalBox::Slot()
-					.AutoWidth()
 					[
 						SNew(SButton)
-						.ButtonStyle(FAppStyle::Get(), "SimpleButton")
-						.ToolTipText(LOCTEXT("AnalyzeText", "Move this sample to the analyzed position"))
-						.OnClicked(this, &FBlendSampleDetails::HandleAnalyzeSample)
-						.ContentPadding(1)
+						.IsEnabled(WouldMoveCondition)
+						.ButtonStyle(FAppStyle::Get(), "HoverHintOnly")
+						.ToolTipText(LOCTEXT("MoveText", "Move this sample to the analyzed position"))
+						.OnClicked(this, &FBlendSampleDetails::HandleAnalyzeAndMoveSample)
 						[
 							SNew(SHorizontalBox)
 							+SHorizontalBox::Slot()
+							.Padding(3.0f, 0.0f)
 							.AutoWidth()
 							[
 								SNew(SImage)
@@ -277,12 +290,12 @@ void FBlendSampleDetails::CustomizeDetails(class IDetailLayoutBuilder& DetailBui
 								.ColorAndOpacity(FSlateColor::UseForeground())
 							]
 							+SHorizontalBox::Slot()
-							.Padding(6.0f, 0.0f)
+							.Padding(3.0f, 0.0f)
 							.VAlign(VAlign_Center)
 							[
 								SNew(STextBlock)
 								.Font(DetailBuilder.GetDetailFont())
-								.Text(LOCTEXT("AnalyzeLabel", "Move"))
+								.Text(LOCTEXT("MoveLabel", "Move"))
 							]
 						]
 					]
@@ -293,13 +306,13 @@ void FBlendSampleDetails::CustomizeDetails(class IDetailLayoutBuilder& DetailBui
 				[
 					SNew(SUniformGridPanel)
 					+SUniformGridPanel::Slot(0,0)
-					.HAlign(HAlign_Center)
+					.HAlign(HAlign_Fill)
 					[   // Analyze X button
 						SNew(SButton)
-						.ButtonStyle(FAppStyle::Get(), "RoundButton")
+						.IsEnabled(WouldMoveConditionX)
+						.ButtonStyle(FAppStyle::Get(), "HoverHintOnly")
 						.ToolTipText(ToolTipTexts[0])
-						.IsEnabled(bAnalyzed[0])
-						.OnClicked(this, &FBlendSampleDetails::HandleAnalyzeSampleX)
+						.OnClicked(this, &FBlendSampleDetails::HandleAnalyzeAndMoveSampleX)
 						[
 							SNew(STextBlock)
 							.Justification(ETextJustify::Center)
@@ -307,14 +320,14 @@ void FBlendSampleDetails::CustomizeDetails(class IDetailLayoutBuilder& DetailBui
 						]
 					]
 					+SUniformGridPanel::Slot(1,0)
-					.HAlign(HAlign_Center)
+					.HAlign(HAlign_Fill)
 					[   // Analyze Y button
 						SNew(SButton)
-						.ButtonStyle(FAppStyle::Get(), "RoundButton")
+						.ButtonStyle(FAppStyle::Get(), "HoverHintOnly")
 						.ToolTipText(ToolTipTexts[1])
 						.Visibility(b1DBlendSpace ? EVisibility::Hidden : EVisibility::Visible)
-						.IsEnabled(bAnalyzed[1])
-						.OnClicked(this, &FBlendSampleDetails::HandleAnalyzeSampleY)
+						.IsEnabled(WouldMoveConditionY) // Needs to be after visibility
+						.OnClicked(this, &FBlendSampleDetails::HandleAnalyzeAndMoveSampleY)
 						[
 							SNew(STextBlock)
 							.Justification(ETextJustify::Center)
@@ -331,16 +344,16 @@ void FBlendSampleDetails::CustomizeDetails(class IDetailLayoutBuilder& DetailBui
 				[
 					SNew(SHorizontalBox)
 					+SHorizontalBox::Slot()
-					.AutoWidth()
 					[
 						SNew(SButton)
-						.ButtonStyle(FAppStyle::Get(), "SimpleButton")
+						.IsEnabled(WouldMoveCondition)
+						.ButtonStyle(FAppStyle::Get(), "HoverHintOnly")
 						.ToolTipText(LOCTEXT("DuplicateText", "Duplicate this sample and place it at the analyzed position"))
 						.OnClicked(this, &FBlendSampleDetails::HandleAnalyzeAndDuplicateSample)
-						.ContentPadding(1)
 						[
 							SNew(SHorizontalBox)
 							+SHorizontalBox::Slot()
+							.Padding(3.0f, 0.0f)
 							.AutoWidth()
 							[
 								SNew(SImage)
@@ -348,12 +361,12 @@ void FBlendSampleDetails::CustomizeDetails(class IDetailLayoutBuilder& DetailBui
 								.ColorAndOpacity(FSlateColor::UseForeground())
 							]
 							+SHorizontalBox::Slot()
-							.Padding(6.0f, 0.0f)
+							.Padding(3.0f, 0.0f)
 							.VAlign(VAlign_Center)
 							[
 								SNew(STextBlock)
 								.Font(DetailBuilder.GetDetailFont())
-								.Text(LOCTEXT("DuplicateLabel", "Duplicate to"))
+								.Text(LOCTEXT("DuplicateLabel", "Duplicate"))
 							]
 						]
 					]
diff --git a/Engine/Source/Editor/Persona/Private/Customization/BlendSampleDetails.h b/Engine/Source/Editor/Persona/Private/Customization/BlendSampleDetails.h
--- a/Engine/Source/Editor/Persona/Private/Customization/BlendSampleDetails.h
+++ b/Engine/Source/Editor/Persona/Private/Customization/BlendSampleDetails.h
@@ -41,9 +41,9 @@ protected:
 	bool ShouldFilterAsset(const FAssetData& AssetData) const;
 
 	FReply HandleAnalyzeAndDuplicateSample();
-	FReply HandleAnalyzeSample();
-	FReply HandleAnalyzeSampleX();
-	FReply HandleAnalyzeSampleY();
+	FReply HandleAnalyzeAndMoveSample();
+	FReply HandleAnalyzeAndMoveSampleX();
+	FReply HandleAnalyzeAndMoveSampleY();
 private:
 	/** Pointer to the current parent blend space for the customized blend sample*/
 	const class UBlendSpace* BlendSpace;
diff --git a/Engine/Source/Editor/Persona/Private/Customization/BlendSpaceDetails.cpp b/Engine/Source/Editor/Persona/Private/Customization/BlendSpaceDetails.cpp
--- a/Engine/Source/Editor/Persona/Private/Customization/BlendSpaceDetails.cpp
+++ b/Engine/Source/Editor/Persona/Private/Customization/BlendSpaceDetails.cpp
@@ -129,7 +129,7 @@ void FBlendSpaceDetails::HandleAnalysisFunctionChanged(int32 AxisIndex, TSharedP
 	// Preserve values where possible
 	if (BlendSpace->AnalysisProperties[AxisIndex])
 	{
-		BlendSpace->AnalysisProperties[AxisIndex]->MakeCache(BlendSpace->CachedAnalysisProperties[AxisIndex]);
+		BlendSpace->AnalysisProperties[AxisIndex]->MakeCache(BlendSpace->CachedAnalysisProperties[AxisIndex], BlendSpace);
 	}
 	if (NewAnalysisProperties)
 	{
@@ -192,8 +192,8 @@ void FBlendSpaceDetails::CustomizeDetails(class IDetailLayoutBuilder& DetailBuil
 			IDetailCategoryBuilder& DetailCategoryBuilder = DetailBuilder.EditCategory(FName("Axis Settings"));
 			IDetailGroup* Groups[2] =
 			{
-				&DetailCategoryBuilder.AddGroup(FName("Horizontal Axis"), LOCTEXT("HorizontalAxisName", "Horizontal Axis")),
-                b1DBlendSpace ? nullptr : &DetailCategoryBuilder.AddGroup(FName("Vertical Axis"), LOCTEXT("VerticalAxisName", "Vertical Axis"))
+				&DetailCategoryBuilder.AddGroup(FName("Horizontal Axis"), LOCTEXT("HorizontalAxis", "Horizontal Axis")),
+                b1DBlendSpace ? nullptr : &DetailCategoryBuilder.AddGroup(FName("Vertical Axis"), LOCTEXT("VerticalAxis", "Vertical Axis"))
             };
 
 			// Hide the default blend and interpolation parameters
@@ -278,14 +278,14 @@ void FBlendSpaceDetails::CustomizeDetails(class IDetailLayoutBuilder& DetailBuil
 					[
 						SNew(STextBlock)
 						.Font(DetailBuilder.GetDetailFont())
-						.Text(FText::Format(LOCTEXT("SamplesLabel", "Analyze {0} Samples"), NumBlendSampleEntries))
+						.Text(FText::Format(LOCTEXT("AnalyzeSamples", "Analyze {0} Samples"), NumBlendSampleEntries))
 					]
 					.ValueContent()
 					[
 						SNew(SButton)
 						.ButtonStyle(FAppStyle::Get(), "SimpleButton")
 						.HAlign(EHorizontalAlignment::HAlign_Left)
-						.ToolTipText(LOCTEXT("Analyze", "Analyze all samples"))
+						.ToolTipText(LOCTEXT("AnalyzeAllSamples", "Analyze all samples"))
 						.OnClicked(this, &FBlendSpaceDetails::HandleAnalyzeSamples)
 						.ContentPadding(1)
 						[
@@ -298,8 +298,8 @@ void FBlendSpaceDetails::CustomizeDetails(class IDetailLayoutBuilder& DetailBuil
 
 			FText AxisTexts[2] = 
 			{ 
-				LOCTEXT("HorizontalAxisName", "Horizontal Axis Function"), 
-				LOCTEXT("VerticalAxisName", "Vertical Axis Function") 
+				LOCTEXT("HorizontalAxisFunction", "Horizontal Axis Function"), 
+				LOCTEXT("VerticalAxisFunction", "Vertical Axis Function") 
 			};
 
 			// Hide the default parameters
@@ -406,8 +406,8 @@ void FBlendSpaceDetails::CustomizeDetails(class IDetailLayoutBuilder& DetailBuil
 
 			BlendSamplesArrayProperty->SetOnNumElementsChanged(RefreshDelegate);
 
-			// Add a "Remove all" button if there are some samples
-			if (NumBlendSampleEntries)
+			// Add a "Remove all" button if there are some samples. Only in the asset blendspace for now.
+			if (NumBlendSampleEntries && BlendSpace->IsAsset())
 			{
 				DetailCategoryBuilder
 					.AddGroup(FName("GroupName"), FText::GetEmpty())
@@ -423,7 +423,7 @@ void FBlendSpaceDetails::CustomizeDetails(class IDetailLayoutBuilder& DetailBuil
 						SNew(SButton)
 						.ButtonStyle(FAppStyle::Get(), "SimpleButton")
 						.HAlign(EHorizontalAlignment::HAlign_Left)
-						.ToolTipText(LOCTEXT("Empty", "Remove all samples"))
+						.ToolTipText(LOCTEXT("RemoveAllSamples", "Remove all samples"))
 						.OnClicked(this, &FBlendSpaceDetails::HandleClearSamples)
 						.ContentPadding(1)
 						[
diff --git a/Engine/Source/Editor/Persona/Private/PersonaBlendSpaceAnalysis.cpp b/Engine/Source/Editor/Persona/Private/PersonaBlendSpaceAnalysis.cpp
--- a/Engine/Source/Editor/Persona/Private/PersonaBlendSpaceAnalysis.cpp
+++ b/Engine/Source/Editor/Persona/Private/PersonaBlendSpaceAnalysis.cpp
@@ -9,11 +9,27 @@
 #include "Animation/BoneSocketReference.h"
 
 #define LOCTEXT_NAMESPACE "BlendSpaceAnalysis"
-#define ANALYSIS_VERBOSE_LOG
 
+//======================================================================================================================
+void UCachedAnalysisProperties::CopyFrom(const UCachedAnalysisProperties& Other)
+{
+	LinearFunctionAxis = Other.LinearFunctionAxis;
+	EulerFunctionAxis = Other.EulerFunctionAxis;
+	BoneSocket1 = Other.BoneSocket1;
+	BoneSocket2 = Other.BoneSocket2;
+	BoneFacingAxis = Other.BoneFacingAxis;
+	BoneRightAxis = Other.BoneRightAxis;
+	Space = Other.Space;
+	SpaceBoneSocket = Other.SpaceBoneSocket;
+	CharacterFacingAxis = Other.CharacterFacingAxis;
+	CharacterUpAxis = Other.CharacterUpAxis;
+	StartTimeFraction = Other.StartTimeFraction;
+	EndTimeFraction = Other.EndTimeFraction;
+	bLockAfterAnalysis = Other.bLockAfterAnalysis;
+}
 
 //======================================================================================================================
-static FVector GetAxisFromTM(const FTransform& TM, EAnalysisLinearAxis Axis)
+FVector BlendSpaceAnalysis::GetAxisFromTM(const FTransform& TM, EAnalysisLinearAxis Axis)
 {
 	switch (Axis)
 	{
@@ -25,9 +41,10 @@ static FVector GetAxisFromTM(const FTransform& TM, EAnalysisLinearAxis Axis)
 }
 
 //======================================================================================================================
-// Retrieves the bone index and transform offset given the BoneSocketTarget. Returns true if found
-static bool GetBoneInfo(
-	const UAnimSequence& Animation, const FBoneSocketTarget& BoneSocket, FTransform& BoneOffset, FName& BoneName)
+bool BlendSpaceAnalysis::GetBoneInfo(const UAnimSequence&     Animation, 
+									 const FBoneSocketTarget& BoneSocket, 
+									 FTransform&              BoneOffset, 
+									 FName&                   BoneName)
 {
 	const UAnimDataModel* AnimDataModel = Animation.GetDataModel();
 
@@ -51,1137 +68,91 @@ static bool GetBoneInfo(
 }
 
 //======================================================================================================================
-template<typename T>
-void CalculateFrameTM(
-	bool& bNeedToUpdateFrameTM, FTransform& FrameTM, 
-	const int32 SampleKey, const T& AnalysisProperties, const UAnimSequence& Animation)
-{
-	if (bNeedToUpdateFrameTM)
-	{
-		FrameTM.SetIdentity();
-		if (AnalysisProperties->Space != EAnalysisSpace::World)
-		{
-			FTransform SpaceBoneOffset;
-			FName SpaceBoneName;
-			if (GetBoneInfo(Animation, AnalysisProperties->SpaceBoneSocket, SpaceBoneOffset, SpaceBoneName))
-			{
-				FAnimPose AnimPose;
-				UAnimPoseExtensions::GetAnimPoseAtFrame(&Animation, SampleKey, FAnimPoseEvaluationOptions(), AnimPose);
-				FTransform SpaceBoneTM = UAnimPoseExtensions::GetBonePose(AnimPose, SpaceBoneName, EAnimPoseSpaces::World);
-				FrameTM = SpaceBoneOffset * SpaceBoneTM;
-			}
-		}
-
-		bNeedToUpdateFrameTM = (
-			AnalysisProperties->Space == EAnalysisSpace::Changing || 
-			AnalysisProperties->Space == EAnalysisSpace::Moving);
-	}
-}
-
-//======================================================================================================================
-template<typename T>
-void GetFrameDirs(
-	FVector& FrameFacingDir, FVector& FrameUpDir, FVector& FrameRightDir, 
-	const FTransform& FrameTM, const T& AnalysisProperties)
-{
-	FrameFacingDir = GetAxisFromTM(FrameTM, AnalysisProperties->CharacterFacingAxis);
-	FrameUpDir = GetAxisFromTM(FrameTM, AnalysisProperties->CharacterUpAxis);
-	FrameRightDir = FVector::CrossProduct(FrameUpDir, FrameFacingDir);
-}
-
-//======================================================================================================================
-template <typename T>
-static bool CalculatePosition(
-	FVector&                         Result,
-	const UBlendSpace&               BlendSpace,
-	const T*                         AnalysisProperties,
-	const UAnimSequence&             Animation,
-	const float                      RateScale)
-{
-	const int32 NumSampledKeys = Animation.GetNumberOfSampledKeys() - 1;
-	if (!AnalysisProperties || NumSampledKeys <= 0)
-	{
-		return false;
-	}
-
-	FTransform BoneOffset;
-	FName BoneName;
-	if (!GetBoneInfo(Animation, AnalysisProperties->BoneSocket, BoneOffset, BoneName))
-	{
-		return false;
-	}
-
-	int32 FirstKey = FMath::Clamp(
-		(int32) (NumSampledKeys * AnalysisProperties->StartTimeFraction), 0, NumSampledKeys);
-	int32 LastKey = FMath::Clamp(
-		(int32) (NumSampledKeys * AnalysisProperties->EndTimeFraction), FirstKey, NumSampledKeys);
-
-	FTransform FrameTM;
-	bool bNeedToUpdateFrameTM = true;
-	CalculateFrameTM(bNeedToUpdateFrameTM, FrameTM, FirstKey, AnalysisProperties, Animation);
-
-	FAnimPose AnimPose;
-	Result.Set(0, 0, 0);
-	for (int32 Key = FirstKey; Key != LastKey + 1; ++Key)
-	{
-		CalculateFrameTM(bNeedToUpdateFrameTM, FrameTM, Key, AnalysisProperties, Animation);
-
-		UAnimPoseExtensions::GetAnimPoseAtFrame(&Animation, Key, FAnimPoseEvaluationOptions(), AnimPose);
-		FTransform BoneTM = UAnimPoseExtensions::GetBonePose(AnimPose, BoneName, EAnimPoseSpaces::World);
-		FTransform TM = BoneOffset * BoneTM;
-		FVector RelativePos = FrameTM.InverseTransformPosition(TM.GetTranslation());
-		Result += RelativePos;
-	}
-	Result /= (1 + LastKey - FirstKey);
-	return true;
-}
-
-//======================================================================================================================
-template <typename T>
-static bool CalculateDeltaPosition(
-	FVector&                         Result,
-	const UBlendSpace&               BlendSpace,
-	const T*                         AnalysisProperties,
-	const UAnimSequence&             Animation,
-	const float                      RateScale)
-{
-	const int32 NumSampledKeys = Animation.GetNumberOfSampledKeys() - 1;
-	if (!AnalysisProperties || NumSampledKeys <= 0)
-	{
-		return false;
-	}
-
-	FTransform BoneOffset;
-	FName BoneName;
-	if (!GetBoneInfo(Animation, AnalysisProperties->BoneSocket, BoneOffset, BoneName))
-	{
-		return false;
-	}
-
-	int32 FirstKey = FMath::Clamp(
-		(int32) (NumSampledKeys * AnalysisProperties->StartTimeFraction), 0, NumSampledKeys);
-	int32 LastKey = FMath::Clamp(
-		(int32) (NumSampledKeys * AnalysisProperties->EndTimeFraction), FirstKey, NumSampledKeys);
-
-	FTransform FrameTM;
-	bool bNeedToUpdateFrameTM = true;
-	CalculateFrameTM(bNeedToUpdateFrameTM, FrameTM, FirstKey, AnalysisProperties, Animation);
-
-	FAnimPose AnimPose;
-
-	UAnimPoseExtensions::GetAnimPoseAtFrame(&Animation, FirstKey, FAnimPoseEvaluationOptions(), AnimPose);
-	FTransform BoneTM1 = UAnimPoseExtensions::GetBonePose(AnimPose, BoneName, EAnimPoseSpaces::World);
-	FTransform TM1 = BoneOffset * BoneTM1;
-	FVector RelativePos1 = FrameTM.InverseTransformPosition(TM1.GetTranslation());
-
-	CalculateFrameTM(bNeedToUpdateFrameTM, FrameTM, FirstKey, AnalysisProperties, Animation);
-
-	UAnimPoseExtensions::GetAnimPoseAtFrame(&Animation, LastKey, FAnimPoseEvaluationOptions(), AnimPose);
-	FTransform BoneTM2 = UAnimPoseExtensions::GetBonePose(AnimPose, BoneName, EAnimPoseSpaces::World);
-	FTransform TM2 = BoneOffset * BoneTM2;
-	FVector RelativePos2 = FrameTM.InverseTransformPosition(TM2.GetTranslation());
-
-	Result = RelativePos2 - RelativePos1;
-	return true;
-}
-
-//======================================================================================================================
-template <typename T>
-static bool CalculateVelocity(
-	FVector&             Result,
-	const UBlendSpace&   BlendSpace,
-	const T*             AnalysisProperties,
-	const UAnimSequence& Animation,
-	const float          RateScale)
+FTransform BlendSpaceAnalysis::GetBoneTransform(const UAnimSequence& Animation, int32 Key, const FName& BoneName)
 {
-	const int32 NumSampledKeys = Animation.GetNumberOfSampledKeys() - 1;
-	if (!AnalysisProperties || NumSampledKeys <= 0)
-	{
-		return false;
-	}
-
-	if (NumSampledKeys == 1)
-	{
-		Result.Set(0, 0, 0);
-		return true;
-	}
-
-	FTransform BoneOffset;
-	FName BoneName;
-	if (!GetBoneInfo(Animation, AnalysisProperties->BoneSocket, BoneOffset, BoneName))
-	{
-		return false;
-	}
-
-	float DeltaTime = Animation.GetPlayLength() / NumSampledKeys;
-
-	int32 FirstKey = FMath::Clamp(
-		(int32) (NumSampledKeys * AnalysisProperties->StartTimeFraction), 0, NumSampledKeys);
-	int32 LastKey = FMath::Clamp(
-		(int32) (NumSampledKeys * AnalysisProperties->EndTimeFraction), FirstKey, NumSampledKeys-1);
-
-	if (FirstKey >= LastKey)
-	{
-		Result.Set(0, 0, 0);
-		return true;
-	}
-
-	FTransform FrameTM;
-	bool bNeedToUpdateFrameTM = true;
-	CalculateFrameTM(bNeedToUpdateFrameTM, FrameTM, FirstKey, AnalysisProperties, Animation);
-
 	FAnimPose AnimPose;
-	Result.Set(0, 0, 0);
-	for (int32 Key = FirstKey; Key != LastKey + 1; ++Key)
-	{
-		CalculateFrameTM(bNeedToUpdateFrameTM, FrameTM, Key, AnalysisProperties, Animation);
-
-		UAnimPoseExtensions::GetAnimPoseAtFrame(&Animation, Key, FAnimPoseEvaluationOptions(), AnimPose);
-		FTransform BoneTM1 = UAnimPoseExtensions::GetBonePose(AnimPose, BoneName, EAnimPoseSpaces::World);
-		FTransform TM1 = BoneOffset * BoneTM1;
-		FVector RelativePos1 = FrameTM.InverseTransformPosition(TM1.GetTranslation());
-
-		int32 NextKey = (Key + 1) % (NumSampledKeys + 1);
-		if (AnalysisProperties->Space == EAnalysisSpace::Moving)
-		{
-			CalculateFrameTM(bNeedToUpdateFrameTM, FrameTM, NextKey, AnalysisProperties, Animation);
-		}
-
-		UAnimPoseExtensions::GetAnimPoseAtFrame(&Animation, NextKey, FAnimPoseEvaluationOptions(), AnimPose);
-		FTransform BoneTM2 = UAnimPoseExtensions::GetBonePose(AnimPose, BoneName, EAnimPoseSpaces::World);
-		FTransform TM2 = BoneOffset * BoneTM2;
-		FVector RelativePos2 = FrameTM.InverseTransformPosition(TM2.GetTranslation());
-		FVector Velocity = (RelativePos2 - RelativePos1) / DeltaTime;
-
-#ifdef ANALYSIS_VERBOSE_LOG
-		UE_LOG(LogAnimation, Log, TEXT("%d Velocity = %f %f %f Height = %f"), 
-			   Key, Velocity.X, Velocity.Y, Velocity.Z, 0.5f * (RelativePos1 + RelativePos2).Z);
-#endif
-		Result += Velocity;
-	}
-	Result /= (1 + LastKey - FirstKey);
-	Result *= Animation.RateScale * RateScale;
-	UE_LOG(LogAnimation, Log, TEXT("%s vel = %f %f %f"), *BoneName.ToString(), Result.X, Result.Y, Result.Z);
-	return true;
-}
-
-//======================================================================================================================
-// Calculates the movement speed (magnitude) 
-template <typename T>
-static bool CalculateMovementSpeed(
-	float&                             Result,
-	const UBlendSpace&                 BlendSpace,
-	const T*                           AnalysisProperties,
-	const UAnimSequence&               Animation,
-	const float                        RateScale)
-{
-	if (!AnalysisProperties)
-	{
-		return false;
-	}
-
-	FVector Velocity;
-	if (CalculateVelocity(Velocity, BlendSpace, AnalysisProperties, Animation, RateScale))
-	{
-		FTransform FrameTM;
-		FVector FrameFacingDir = GetAxisFromTM(FrameTM, AnalysisProperties->CharacterFacingAxis);
-		FVector FrameUpDir = GetAxisFromTM(FrameTM, AnalysisProperties->CharacterUpAxis);
-		FVector FrameRightDir = FVector::CrossProduct(FrameUpDir, FrameFacingDir);
-		float Fwd = Velocity | FrameFacingDir;
-		float Right = Velocity | FrameRightDir;
-		Result = FMath::Sqrt(FMath::Square(Fwd) + FMath::Square(Right));
-		return true;
-	}
-	return false;
-}
-
-//======================================================================================================================
-// Calculates the movement speed (magnitude) 
-template <typename T>
-static bool CalculateMovementDirection(
-	float&                             Result,
-	const UBlendSpace&                 BlendSpace,
-	const T*                           AnalysisProperties,
-	const UAnimSequence&               Animation,
-	const float                        RateScale)
-{
-	if (!AnalysisProperties)
-	{
-		return false;
-	}
-
-	FVector Velocity;
-	if (CalculateVelocity(Velocity, BlendSpace, AnalysisProperties, Animation, RateScale))
-	{
-		FTransform FrameTM;
-		FVector FrameFacingDir = GetAxisFromTM(FrameTM, AnalysisProperties->CharacterFacingAxis);
-		FVector FrameUpDir = GetAxisFromTM(FrameTM, AnalysisProperties->CharacterUpAxis);
-		FVector FrameRightDir = FVector::CrossProduct(FrameUpDir, FrameFacingDir);
-
-		float Fwd = Velocity | FrameFacingDir;
-		float Right = Velocity | FrameRightDir;
-		Result = FMath::RadiansToDegrees(FMath::Atan2(Right, Fwd));
-		return true;
-	}
-	return false;
-}
-
-//======================================================================================================================
-template <typename T>
-void CalculateBoneOrientation(
-	FVector&                   RollPitchYaw,
-	const UAnimSequence&       Animation, 
-	const int32                Key, 
-	const FName                BoneName, 
-	const FTransform&          BoneOffset, 
-	const T&                   AnalysisProperties, 
-	const FVector&             FrameFacingDir, 
-	const FVector&             FrameRightDir, 
-	const FVector&             FrameUpDir)
-{
-	FAnimPose AnimPose;
-
 	UAnimPoseExtensions::GetAnimPoseAtFrame(&Animation, Key, FAnimPoseEvaluationOptions(), AnimPose);
 	FTransform BoneTM = UAnimPoseExtensions::GetBonePose(AnimPose, BoneName, EAnimPoseSpaces::World);
-
-	FTransform TM = BoneOffset * BoneTM;
-	FQuat AimQuat = TM.GetRotation();
-	FVector AimFwdDir = GetAxisFromTM(TM, AnalysisProperties->BoneFacingAxis);
-	FVector AimRightDir = GetAxisFromTM(TM, AnalysisProperties->BoneRightAxis);
-
-	// Note that Yaw is best taken from the AimRightDir - this is to avoid problems when the gun is pointing
-	// up or down - especially if it goes beyond 90 degrees in pitch.
-	float Yaw = FMath::RadiansToDegrees(FMath::Atan2(
-		FVector::DotProduct(AimRightDir, -FrameFacingDir), FVector::DotProduct(AimRightDir, FrameRightDir)));
-
-	// Undo the yaw to get pitch
-	const FQuat YawQuat(FrameUpDir, FMath::DegreesToRadians(Yaw));
-	FVector UnYawedAimFwdDir = YawQuat.UnrotateVector(AimFwdDir);
-	float Up = UnYawedAimFwdDir | FrameUpDir;
-	float Fwd = UnYawedAimFwdDir | FrameFacingDir;
-	float Pitch = FMath::RadiansToDegrees(FMath::Atan2(Up, Fwd));
-
-	// Undo the pitch to get roll
-	FVector UnYawedAimRightDir = YawQuat.UnrotateVector(AimRightDir);
-	const FQuat PitchQuat(FrameRightDir, -FMath::DegreesToRadians(Pitch));
-
-	FVector UnYawedUnPitchedAimRightDir = PitchQuat.UnrotateVector(UnYawedAimRightDir);
-
-	float Roll = FMath::RadiansToDegrees(FMath::Atan2(
-		FVector::DotProduct(UnYawedUnPitchedAimRightDir, -FrameUpDir), 
-		FVector::DotProduct(UnYawedUnPitchedAimRightDir, FrameRightDir)));
-
-	RollPitchYaw.Set(Roll, Pitch, Yaw);
+	return BoneTM;
 }
 
-//======================================================================================================================
-// Note that if a looping animation has 56 keys, then its first key is 0 and last is 55, but these will be identical poses.
-// Thus it has one fewer intervals/unique keys
-template <typename T>
-static bool CalculateOrientation(
-	FVector&                        Result,
-	const UBlendSpace&              BlendSpace,
-	const T*                        AnalysisProperties,
-	const UAnimSequence&            Animation,
-	const float                     RateScale)
-{
-	const int32 NumSampledKeys = Animation.GetNumberOfSampledKeys() - 1;
-	if (!AnalysisProperties || NumSampledKeys <= 0)
-	{
-		return false;
-	}
-
-	FTransform BoneOffset;
-	FName BoneName;
-	if (!GetBoneInfo(Animation, AnalysisProperties->BoneSocket, BoneOffset, BoneName))
-	{
-		return false;
-	}
-
-	int32 FirstKey = FMath::Clamp(
-		(int32) (NumSampledKeys * AnalysisProperties->StartTimeFraction), 0, NumSampledKeys);
-	int32 LastKey = FMath::Clamp(
-		(int32) (NumSampledKeys * AnalysisProperties->EndTimeFraction), FirstKey, NumSampledKeys);
-
-	FTransform FrameTM;
-	bool bNeedToUpdateFrameTM = true;
-	FVector FrameFacingDir, FrameUpDir, FrameRightDir;
-	CalculateFrameTM(bNeedToUpdateFrameTM, FrameTM, FirstKey, AnalysisProperties, Animation);
-	GetFrameDirs(FrameFacingDir, FrameUpDir, FrameRightDir, FrameTM, AnalysisProperties);
-
-	FAnimPose AnimPose;
-	Result.Set(0, 0, 0);
-	for (int32 Key = FirstKey; Key != LastKey + 1; ++Key)
-	{
-		CalculateFrameTM(bNeedToUpdateFrameTM, FrameTM, Key, AnalysisProperties, Animation);
-		GetFrameDirs(FrameFacingDir, FrameUpDir, FrameRightDir, FrameTM, AnalysisProperties);
-
-		FVector RollPitchYaw;
-		CalculateBoneOrientation(
-			RollPitchYaw, Animation, Key, BoneName, BoneOffset, 
-			AnalysisProperties, FrameFacingDir, FrameRightDir, FrameUpDir);
-#ifdef ANALYSIS_VERBOSE_LOG
-		UE_LOG(LogAnimation, Log, TEXT("Roll/pitch/yaw = %f %f %f"), RollPitchYaw.X, RollPitchYaw.Y, RollPitchYaw.Z);
-#endif
-		Result += RollPitchYaw;
-	}
-	Result /= (1 + LastKey - FirstKey);
-	UE_LOG(LogAnimation, Log, TEXT("%s Orientation = %f %f %f"), *BoneName.ToString(), Result.X, Result.Y, Result.Z);
-	return true;
-}
 
 //======================================================================================================================
-// Note that if a looping animation has 56 keys, then its first key is 0 and last is 55, but these will be identical poses.
-// Thus it has one fewer intervals/unique keys
-template <typename T>
-static bool CalculateDeltaOrientation(
-	FVector&                        Result,
-	const UBlendSpace&              BlendSpace,
-	const T*                        AnalysisProperties,
-	const UAnimSequence&            Animation,
-	const float                     RateScale)
-{
-	const int32 NumSampledKeys = Animation.GetNumberOfSampledKeys() - 1;
-	if (!AnalysisProperties || NumSampledKeys <= 0)
-	{
-		return false;
-	}
-
-	FTransform BoneOffset;
-	FName BoneName;
-	if (!GetBoneInfo(Animation, AnalysisProperties->BoneSocket, BoneOffset, BoneName))
-	{
-		return false;
-	}
-
-	int32 FirstKey = FMath::Clamp(
-		(int32) (NumSampledKeys * AnalysisProperties->StartTimeFraction), 0, NumSampledKeys);
-	int32 LastKey = FMath::Clamp(
-		(int32) (NumSampledKeys * AnalysisProperties->EndTimeFraction), FirstKey, NumSampledKeys);
-
-	FTransform FrameTM;
-	bool bNeedToUpdateFrameTM = true;
-	FVector FrameFacingDir, FrameUpDir, FrameRightDir;
-	CalculateFrameTM(bNeedToUpdateFrameTM, FrameTM, FirstKey, AnalysisProperties, Animation);
-	GetFrameDirs(FrameFacingDir, FrameUpDir, FrameRightDir, FrameTM, AnalysisProperties);
-
-	FVector RollPitchYaw1;
-	CalculateBoneOrientation(
-		RollPitchYaw1, Animation, FirstKey, BoneName, BoneOffset,
-		AnalysisProperties, FrameFacingDir, FrameRightDir, FrameUpDir);
-
-	CalculateFrameTM(bNeedToUpdateFrameTM, FrameTM, LastKey, AnalysisProperties, Animation);
-	GetFrameDirs(FrameFacingDir, FrameUpDir, FrameRightDir, FrameTM, AnalysisProperties);
-
-	FVector RollPitchYaw2;
-	CalculateBoneOrientation(
-		RollPitchYaw2, Animation, LastKey, BoneName, BoneOffset,
-		AnalysisProperties, FrameFacingDir, FrameRightDir, FrameUpDir);
-
-	Result = RollPitchYaw2 - RollPitchYaw1;
-	return true;
-}
-
-//======================================================================================================================
-template <typename T>
-static bool CalculateAngularVelocity(
-	FVector&                         Result,
-	const UBlendSpace&               BlendSpace,
-	const T*                         AnalysisProperties,
-	const UAnimSequence&             Animation,
-	const float                      RateScale)
-{
-	const int32 NumSampledKeys = Animation.GetNumberOfSampledKeys() - 1;
-	if (!AnalysisProperties || NumSampledKeys <= 0)
-	{
-		return false;
-	}
-
-	if (NumSampledKeys == 1)
-	{
-		Result.Set(0, 0, 0);
-		return true;
-	}
-
-	FTransform BoneOffset;
-	FName BoneName;
-	if (!GetBoneInfo(Animation, AnalysisProperties->BoneSocket, BoneOffset, BoneName))
-	{
-		return false;
-	}
-
-	float DeltaTime = Animation.GetPlayLength() / NumSampledKeys;
-
-	int32 FirstKey = FMath::Clamp(
-		(int32) (NumSampledKeys * AnalysisProperties->StartTimeFraction), 0, NumSampledKeys);
-	int32 LastKey = FMath::Clamp(
-		(int32) (NumSampledKeys * AnalysisProperties->EndTimeFraction), FirstKey, NumSampledKeys-1);
-
-	if (FirstKey >= LastKey)
-	{
-		Result.Set(0, 0, 0);
-		return true;
-	}
-
-	FTransform FrameTM;
-	bool bNeedToUpdateFrameTM = true;
-	CalculateFrameTM(bNeedToUpdateFrameTM, FrameTM, FirstKey, AnalysisProperties, Animation);
-
-	FAnimPose AnimPose;
-	Result.Set(0, 0, 0);
-	for (int32 Key = FirstKey; Key != LastKey + 1; ++Key)
-	{
-		CalculateFrameTM(bNeedToUpdateFrameTM, FrameTM, Key, AnalysisProperties, Animation);
-
-		UAnimPoseExtensions::GetAnimPoseAtFrame(&Animation, Key, FAnimPoseEvaluationOptions(), AnimPose);
-		FTransform BoneTM1 = UAnimPoseExtensions::GetBonePose(AnimPose, BoneName, EAnimPoseSpaces::World);
-		FTransform TM1 = BoneOffset * BoneTM1;
-		FQuat RelativeQuat1 = FrameTM.InverseTransformRotation(TM1.GetRotation());
-
-		int32 NextKey = (Key + 1) % (NumSampledKeys + 1);
-		if (AnalysisProperties->Space == EAnalysisSpace::Moving)
-		{
-			CalculateFrameTM(bNeedToUpdateFrameTM, FrameTM, NextKey, AnalysisProperties, Animation);
-		}
-
-		UAnimPoseExtensions::GetAnimPoseAtFrame(&Animation, NextKey, FAnimPoseEvaluationOptions(), AnimPose);
-		FTransform BoneTM2 = UAnimPoseExtensions::GetBonePose(AnimPose, BoneName, EAnimPoseSpaces::World);
-		FTransform TM2 = BoneOffset * BoneTM2;
-		FQuat RelativeQuat2 = FrameTM.InverseTransformRotation(TM2.GetRotation());
-
-		FQuat Rotation = RelativeQuat2 * RelativeQuat1.Inverse();
-		FVector Axis;
-		float Angle;
-		Rotation.ToAxisAndAngle(Axis, Angle);
-		FVector AngularVelocity = FMath::RadiansToDegrees(Axis * (Angle / DeltaTime));
-#ifdef ANALYSIS_VERBOSE_LOG
-		UE_LOG(LogAnimation, Log, TEXT("Angular Velocity = %f %f %f"), AngularVelocity.X, AngularVelocity.Y, AngularVelocity.Z);
-#endif
-		Result += AngularVelocity;
-	}
-	Result /= (1 + LastKey - FirstKey);
-	Result *= Animation.RateScale * RateScale;
-	UE_LOG(LogAnimation, Log, TEXT("%s angular velocity = %f %f %f"), *BoneName.ToString(), Result.X, Result.Y, Result.Z);
-	return true;
-}
-
-//======================================================================================================================
-template <typename T>
-static bool CalculateOrientationRate(
-	FVector&                        Result,
-	const UBlendSpace&              BlendSpace,
-	const T*                        AnalysisProperties,
-	const UAnimSequence&            Animation,
-	const float                     RateScale)
-{
-	const int32 NumSampledKeys = Animation.GetNumberOfSampledKeys() - 1;
-	if (!AnalysisProperties || NumSampledKeys <= 0)
-	{
-		return false;
-	}
-
-	if (NumSampledKeys == 1)
-	{
-		Result.Set(0, 0, 0);
-		return true;
-	}
-
-	FTransform BoneOffset;
-	FName BoneName;
-	if (!GetBoneInfo(Animation, AnalysisProperties->BoneSocket, BoneOffset, BoneName))
-	{
-		return false;
-	}
-
-	float DeltaTime = Animation.GetPlayLength() / NumSampledKeys;
-
-	int32 FirstKey = FMath::Clamp(
-		(int32) (NumSampledKeys * AnalysisProperties->StartTimeFraction), 0, NumSampledKeys);
-	int32 LastKey = FMath::Clamp(
-		(int32) (NumSampledKeys * AnalysisProperties->EndTimeFraction), FirstKey, NumSampledKeys-1);
-
-	if (FirstKey >= LastKey)
-	{
-		Result.Set(0, 0, 0);
-		return true;
-	}
-
-	FTransform FrameTM;
-	bool bNeedToUpdateFrameTM = true;
-	FVector FrameFacingDir, FrameUpDir, FrameRightDir;
-	CalculateFrameTM(bNeedToUpdateFrameTM, FrameTM, FirstKey, AnalysisProperties, Animation);
-	GetFrameDirs(FrameFacingDir, FrameUpDir, FrameRightDir, FrameTM, AnalysisProperties);
-
-	FAnimPose AnimPose;
-	Result.Set(0, 0, 0);
-	for (int32 Key = FirstKey; Key != LastKey + 1; ++Key)
-	{
-		CalculateFrameTM(bNeedToUpdateFrameTM, FrameTM, Key, AnalysisProperties, Animation);
-		GetFrameDirs(FrameFacingDir, FrameUpDir, FrameRightDir, FrameTM, AnalysisProperties);
-
-		FVector RollPitchYaw1;
-		CalculateBoneOrientation(
-			RollPitchYaw1, Animation, Key, BoneName, BoneOffset, 
-			AnalysisProperties, FrameFacingDir, FrameRightDir, FrameUpDir);
-
-		int32 NextKey = (Key + 1) % (NumSampledKeys + 1);
-		if (AnalysisProperties->Space == EAnalysisSpace::Moving)
-		{
-			CalculateFrameTM(bNeedToUpdateFrameTM, FrameTM, NextKey, AnalysisProperties, Animation);
-			GetFrameDirs(FrameFacingDir, FrameUpDir, FrameRightDir, FrameTM, AnalysisProperties);
-		}
-
-		FVector RollPitchYaw2;
-		CalculateBoneOrientation(
-			RollPitchYaw2, Animation, NextKey, BoneName, BoneOffset, 
-			AnalysisProperties, FrameFacingDir, FrameRightDir, FrameUpDir);
-
-		const FVector OrientationRate = (RollPitchYaw2 - RollPitchYaw1) / DeltaTime;
-#ifdef ANALYSIS_VERBOSE_LOG
-		UE_LOG(LogAnimation, Log, TEXT("Orientation rate = %f %f %f"), OrientationRate.X, OrientationRate.Y, OrientationRate.Z);
-#endif
-		Result += OrientationRate;
-	}
-	Result /= (1 + LastKey - FirstKey);
-	Result *= Animation.RateScale * RateScale;
-	UE_LOG(LogAnimation, Log, TEXT("%s Orientation rate = %f %f %f"), *BoneName.ToString(), Result.X, Result.Y, Result.Z);
-	return true;
-}
-
-//======================================================================================================================
-// Note that if a looping animation has 56 keys, then its first key is 0 and last is 55, but these will be identical poses.
-// Thus it has one fewer intervals/unique keys
-static bool CalculateLocomotionVelocity(
-	FVector&                             Result,
-	const UBlendSpace&                   BlendSpace,
-	const ULocomotionAnalysisProperties* AnalysisProperties,
-	const FBoneSocketTarget&             BoneSocket,
-	const UAnimSequence&                 Animation,
-	const float                          RateScale)
-{
-	const int32 NumSampledKeys = Animation.GetNumberOfSampledKeys();
-	if (!AnalysisProperties || NumSampledKeys <= 0)
-	{
-		return false;
-	}
-	if (NumSampledKeys == 1)
-	{
-		Result.Set(0, 0, 0);
-		return true;
-	}
-
-	FTransform BoneOffset;
-	FName BoneName;
-	if (!GetBoneInfo(Animation, BoneSocket, BoneOffset, BoneName))
-	{
-		return false;
-	}
-
-	// Note that for locomotion we don't support the frame changing
-	FTransform FrameTM = FTransform::Identity;
-	FVector FrameFacingDir, FrameUpDir, FrameRightDir;
-	GetFrameDirs(FrameFacingDir, FrameUpDir, FrameRightDir, FrameTM, AnalysisProperties);
-
-	FAnimPose AnimPose;
-
-	// The frame time delta
-	float DeltaTime = Animation.GetPlayLength() / NumSampledKeys;
-
-	// First step is to figure out the approximate direction. Note that the average velocity will be zero (assuming a
-	// complete cycle) - but if we apply a weight that is based on the height, then we can bias it towards the foot that
-	// is on the ground.
-	float MinHeight = FLT_MAX;
-	float MaxHeight = -FLT_MAX;
-	FVector AveragePos(0.0f);
-	TArray<FVector> Positions;
-	Positions.SetNum(NumSampledKeys);
-	for (int32 Key = 0; Key != NumSampledKeys; ++Key)
-	{
-		UAnimPoseExtensions::GetAnimPoseAtFrame(&Animation, Key, FAnimPoseEvaluationOptions(), AnimPose);
-		FTransform BoneTM = UAnimPoseExtensions::GetBonePose(AnimPose, BoneName, EAnimPoseSpaces::World);
-		FTransform TM = BoneOffset * BoneTM;
-		FVector Pos = TM.GetTranslation();
-		Positions[Key] = Pos;
-		float Height = Pos | FrameUpDir;
-		MinHeight = FMath::Min(MinHeight, Height);
-		MaxHeight = FMath::Max(MaxHeight, Height);
-		AveragePos += Pos;
-#ifdef ANALYSIS_VERBOSE_LOG
-		UE_LOG(LogAnimation, Log, TEXT("Pos %f %f %f"), Pos.X, Pos.Y, Pos.Z);
-#endif
-	}
-	AveragePos /= NumSampledKeys;
-
-	// Calculate velocities. 
-	TArray<FVector> Velocities;
-	Velocities.SetNum(NumSampledKeys);
-	for (int32 Key = 0; Key != NumSampledKeys; ++Key)
-	{
-		int32 PrevKey = (Key + NumSampledKeys - 1) % NumSampledKeys;
-		int32 NextKey = (Key +  1) % NumSampledKeys;
-		Velocities[Key] = (Positions[NextKey] - Positions[PrevKey]) / (2.0f * DeltaTime);
-	}
-
-	FVector BiasedFootVel(0);
-	float TotalWeight = 0.0f;
-	for (int32 Key = 0 ; Key != NumSampledKeys ; ++Key)
-	{
-		float Height = Positions[Key] | FrameUpDir;
-		float Weight = 1.0f - (Height - MinHeight) / (MaxHeight - MinHeight);
-		BiasedFootVel += Velocities[Key] * Weight;
-		TotalWeight += Weight;
-	}
-	BiasedFootVel /= TotalWeight; 
-
-	if (BiasedFootVel.IsNearlyZero())
-	{
-		Result.Set(0, 0, 0);
-		return true;
-	}
-
-	FVector ApproxLocoDir = -BiasedFootVel.GetSafeNormal();
-
-	// Now we can form a mask, where 0 means traveling in the wrong direction (so clearly off the ground), and positive
-	// numbers will indicate how far into a valid segment we are. We will assume that the animation is looping
-	TArray<int32> Mask;
-	Mask.SetNum(NumSampledKeys);
-	for (int32 Key = 0 ; Key != NumSampledKeys ; ++Key)
-	{
-		Mask[Key] = (Velocities[Key] | ApproxLocoDir) >= 0.0f ? 0 : 1;
-	}
-
-	int32 StartKey = -1;
-	TArray<int32> Mask1 = Mask;
-	bool bRepeat = true;
-	int32 MaxMask = 0;
-	int32 PrevNumFound = NumSampledKeys + 1;
-	while (bRepeat)
-	{
-		bRepeat = false;
-		int32 NumFound = 0;
-		for (int32 Key = 0 ; Key != NumSampledKeys ; ++Key)
-		{
-			if (Mask[Key] > 0)
-			{
-				int32 PrevKey = (Key + NumSampledKeys - 1) % NumSampledKeys;
-				int32 NextKey = (Key + 1) % NumSampledKeys;
-				if (Mask[PrevKey] == Mask[Key] && Mask[NextKey] == Mask[Key])
-				{
-					++Mask1[Key];
-					MaxMask = FMath::Max(MaxMask, Mask1[Key]);
-					bRepeat = true;
-					++NumFound;
-				}
-			}
-		}
-		Mask = Mask1;
-		// Avoid a perpetual loop (e.g. can happen if initially all the mask values are 1... though that shouldn't
-		// really happen).
-		if (NumFound >= PrevNumFound)
-		{
-			bRepeat = false;
-		}
-		else
-		{
-			PrevNumFound = NumFound;
-		}
-	}
-
-	// When searching we will want to start outside of a "good" region.
-	int32 AZeroKey = 0;
-	for (int32 Key = 0 ; Key != NumSampledKeys ; ++Key)
-	{
-		if (Mask[Key] == 0)
-		{
-			AZeroKey = Key;
-		}
-	}
-
-	// We use the mask (with a somewhat arbitrary threshold) to get rid of velocities that are near to the foot
-	// plant/take-off time (and might be when the foot is in the time). Then We look for the highest velocity. Not that
-	// if we're being called with the foot (ankle) joint, it will tend to underestimate velocities since it is nearer
-	// the hip than the ground contact point.
-	// 
-	int32 Threshold = FMath::Max(MaxMask / 2, 1);
-	int32 Num = 0;
-	FVector AverageFootVel(0);
-	float BestSpeed = 0.0f;
-	int32 BestSpeedKey = 0;
-	for (int32 K = AZeroKey ; K != AZeroKey + NumSampledKeys ; ++K)
-	{
-		int32 Key = K % NumSampledKeys;
-		if (Mask[Key] >= Threshold)
-		{
-			float Speed = Velocities[Key].Size();
-#ifdef ANALYSIS_VERBOSE_LOG
-			UE_LOG(LogAnimation, Log, TEXT("Candidate %d Mask %d vel = %f %f %f, speed = %f"), 
-				   Key, Mask[Key], Velocities[Key].X, Velocities[Key].Y, Velocities[Key].Z, Speed);
-#endif
-			if (Speed > BestSpeed)
-			{
-				BestSpeed = Speed;
-				BestSpeedKey = Key;
-			}
-		}
-		else
-		{
-			if (BestSpeed > 0.0f)
-			{
-#ifdef ANALYSIS_VERBOSE_LOG
-				UE_LOG(LogAnimation, Log, TEXT("Picked Candidate %d vel = %f %f %f, speed = %f"), 
-					   BestSpeedKey, Velocities[BestSpeedKey].X, Velocities[BestSpeedKey].Y, Velocities[BestSpeedKey].Z, BestSpeed);
-#endif
-				AverageFootVel += Velocities[BestSpeedKey];
-				++Num;
-				BestSpeed = 0.0f;
-			}
-		}
-	}
-	// Make sure we didn't miss the last data point
-	if (BestSpeed > 0.0f)
-	{
-#ifdef ANALYSIS_VERBOSE_LOG
-		UE_LOG(LogAnimation, Log, TEXT("Picked Candidate %d vel = %f %f %f, speed = %f"), 
-			   BestSpeedKey, Velocities[BestSpeedKey].X, Velocities[BestSpeedKey].Y, Velocities[BestSpeedKey].Z, BestSpeed);
-#endif
-		AverageFootVel += Velocities[BestSpeedKey];
-		++Num;
-		BestSpeed = 0.0f;
-	}
-
-	AverageFootVel /= Num;
-	float FacingVel = -AverageFootVel | FrameFacingDir;
-	float RightVel = -AverageFootVel | FrameRightDir;
-	float UpVel = -AverageFootVel | FrameUpDir;
-
-	Result.Set(FacingVel, RightVel, UpVel);
-	Result *= Animation.RateScale * RateScale;
-	UE_LOG(LogAnimation, Log, TEXT("%s Locomotion vel = %f %f %f"), *BoneName.ToString(), Result.X, Result.Y, Result.Z);
-	return true;
-}
-
-//======================================================================================================================
-static bool CalculateLocomotionVelocity(
-	FVector&                             Result,
-	const UBlendSpace&                   BlendSpace,
-	const ULocomotionAnalysisProperties* AnalysisProperties,
-	const UAnimSequence&                 Animation,
-	const float                          RateScale)
-{
-	if (!AnalysisProperties)
-	{
-		return false;
-	}
-
-	FVector Result1, Result2;
-	int32 Num = 0;
-	Result.Set(0, 0, 0);
-	if (CalculateLocomotionVelocity(Result1, BlendSpace, AnalysisProperties, 
-									AnalysisProperties->PrimaryBoneSocket, Animation, RateScale))
-	{
-#ifdef ANALYSIS_VERBOSE_LOG
-		UE_LOG(LogAnimation, Log, TEXT("Loco vel from primary = %f %f %f"), Result1.X, Result1.Y, Result1.Z);
-#endif
-		Result += Result1;
-		++Num;
-	}
-	if (CalculateLocomotionVelocity(Result2, BlendSpace, AnalysisProperties, 
-									AnalysisProperties->SecondaryBoneSocket, Animation, RateScale))
-	{
-#ifdef ANALYSIS_VERBOSE_LOG
-		UE_LOG(LogAnimation, Log, TEXT("Loco vel from secondary = %f %f %f"), Result2.X, Result2.Y, Result2.Z);
-#endif
-		Result += Result2;
-		++Num;
-	}
-	if (Num)
-	{
-		Result /= Num;
-		UE_LOG(LogAnimation, Log, TEXT("Loco vel = %f %f %f"), Result.X, Result.Y, Result.Z);
-		return true;
-	}
-	return false;
-}
-
-//======================================================================================================================
-// Calculates the locomotion speed in the character's facing direction
-static bool CalculateLocomotionFwdSpeed(
-	float&                               Result,
-	const UBlendSpace&                   BlendSpace,
-	const ULocomotionAnalysisProperties* AnalysisProperties,
-	const UAnimSequence&                 Animation,
-	const float                          RateScale)
-{
-	FVector Movement;
-	if (CalculateLocomotionVelocity(Movement, BlendSpace, AnalysisProperties, Animation, RateScale))
-	{
-		Result = Movement.X;
-		return true;
-	}
-	return false;
-}
-
-//======================================================================================================================
-// Calculates the locomotion speed in the character's upwards direction
-static bool CalculateLocomotionUpSpeed(
-	float&                               Result,
-	const UBlendSpace&                   BlendSpace,
-	const ULocomotionAnalysisProperties* AnalysisProperties,
-	const UAnimSequence&                 Animation,
-	const float                          RateScale)
-{
-	FVector Movement;
-	if (CalculateLocomotionVelocity(Movement, BlendSpace, AnalysisProperties, Animation, RateScale))
-	{
-		Result = Movement.Z;
-		return true;
-	}
-	return false;
-}
-
-//======================================================================================================================
-// Calculates the locomotion speed in the character's right direction
-static bool CalculateLocomotionRightSpeed(
-	float&                               Result,
-	const UBlendSpace&                   BlendSpace,
-	const ULocomotionAnalysisProperties* AnalysisProperties,
-	const UAnimSequence&                 Animation,
-	const float                          RateScale)
-{
-	FVector Movement;
-	if (CalculateLocomotionVelocity(Movement, BlendSpace, AnalysisProperties, Animation, RateScale))
-	{
-		Result = Movement.Y;
-		return true;
-	}
-	return false;
-}
-
-//======================================================================================================================
-// Calculates the locomotion speed (magnitude)
-static bool CalculateLocomotionSpeed(
-	float&                               Result,
-	const UBlendSpace&                   BlendSpace,
-	const ULocomotionAnalysisProperties* AnalysisProperties,
-	const UAnimSequence&                 Animation,
-	const float                          RateScale)
-{
-	FVector Movement;
-	if (CalculateLocomotionVelocity(Movement, BlendSpace, AnalysisProperties, Animation, RateScale))
-	{
-		Result = Movement.Size();
-		return true;
-	}
-	return false;
-}
-
-//======================================================================================================================
-// Calculates the locomotion direction (degrees)
-static bool CalculateLocomotionDirection(
-	float&                               Result,
-	const UBlendSpace&                   BlendSpace,
-	const ULocomotionAnalysisProperties* AnalysisProperties,
-	const UAnimSequence&                 Animation,
-	const float                          RateScale)
-{
-	FVector Movement;
-	if (CalculateLocomotionVelocity(Movement, BlendSpace, AnalysisProperties, Animation, RateScale))
-	{
-		Result = FMath::RadiansToDegrees(FMath::Atan2(Movement.Y, Movement.X));
-		return true;
-	}
-	return false;
-}
-
-//======================================================================================================================
-// Calculates the locomotion slope angle (degrees) going in the facing direction
-static bool CalculateLocomotionFwdSlope(
-	float&                               Result,
-	const UBlendSpace&                   BlendSpace,
-	const ULocomotionAnalysisProperties* AnalysisProperties,
-	const UAnimSequence&                 Animation,
-	const float                          RateScale)
-{
-	FVector Movement;
-	if (CalculateLocomotionVelocity(Movement, BlendSpace, AnalysisProperties, Animation, RateScale))
-	{
-		if (Movement.X >= 0.0f)
-		{
-			Result = FMath::RadiansToDegrees(FMath::Atan2(Movement.Z, Movement.X));
-		}
-		else
-		{
-			Result = FMath::RadiansToDegrees(FMath::Atan2(-Movement.Z, -Movement.X));
-		}
-
-		return true;
-	}
-	return false;
-}
-
-//======================================================================================================================
-// Calculates the locomotion slope angle (degrees) going in the rightwards direction
-static bool CalculateLocomotionRightSlope(
-	float&                               Result,
-	const UBlendSpace&                   BlendSpace,
-	const ULocomotionAnalysisProperties* AnalysisProperties,
-	const UAnimSequence&                 Animation,
-	const float                          RateScale)
-{
-	FVector Movement;
-	if (CalculateLocomotionVelocity(Movement, BlendSpace, AnalysisProperties, Animation, RateScale))
-	{
-		if (Movement.Y > 0.0f)
-		{
-			Result = FMath::RadiansToDegrees(FMath::Atan2(Movement.Z, Movement.Y));
-		}
-		else
-		{
-			Result = FMath::RadiansToDegrees(FMath::Atan2(-Movement.Z, -Movement.Y));
-		}
-		return true;
-	}
-	return false;
-}
-
-//======================================================================================================================
-/**
- * Helper to extract the component from the FVector functions
- */
-template<typename FunctionType, typename T>
-static bool CalculateComponentSampleValue(
-	float&                     Result,
-	const FunctionType&        Fn,
-	const UBlendSpace&         BlendSpace, 
-	const T*                   AnalysisProperties, 
-	const UAnimSequence&       Animation,
-	const float                RateScale)
-{
-	FVector Value;
-	int32 ComponentIndex = (int32) AnalysisProperties->FunctionAxis;
-	if (Fn(Value, BlendSpace, AnalysisProperties, Animation, RateScale))
-	{
-		Result = Value[ComponentIndex];
-		return true;
-	}
-	return false;
-}
-
-//======================================================================================================================
-void ULinearAnalysisProperties::InitializeFromCache(TSharedPtr<FCachedAnalysisProperties> Cache)
+void ULinearAnalysisProperties::InitializeFromCache(TObjectPtr<UCachedAnalysisProperties> Cache)
 {
 	if (Cache)
 	{
-		Super::InitializeFromCache(Cache);
-		FunctionAxis = Cache->LinearFunctionAxis;
-		Space = Cache->Space;
-		SpaceBoneSocket = Cache->SpaceBoneSocket;
-		BoneSocket = Cache->BoneSocket1;
-		StartTimeFraction = Cache->StartTimeFraction;
-		EndTimeFraction = Cache->EndTimeFraction;
+		UCachedAnalysisProperties* CachePtr = Cache.Get();
+		FunctionAxis = CachePtr->LinearFunctionAxis;
+		BoneSocket = CachePtr->BoneSocket1;
+		Space = CachePtr->Space;
+		SpaceBoneSocket = CachePtr->SpaceBoneSocket;
+		StartTimeFraction = CachePtr->StartTimeFraction;
+		EndTimeFraction = CachePtr->EndTimeFraction;
+		bLockAfterAnalysis = CachePtr->bLockAfterAnalysis;
 	}
 }
 
 //======================================================================================================================
-void ULinearAnalysisProperties::MakeCache(TSharedPtr<FCachedAnalysisProperties>& Cache) const
+void ULinearAnalysisProperties::MakeCache(TObjectPtr<UCachedAnalysisProperties>& Cache, UBlendSpace* BlendSpace)
 {
-	Super::MakeCache(Cache);
-	Cache->LinearFunctionAxis = FunctionAxis;
-	Cache->Space = Space;
-	Cache->SpaceBoneSocket = SpaceBoneSocket;
-	Cache->BoneSocket1 = BoneSocket;
-	Cache->StartTimeFraction = StartTimeFraction;
-	Cache->EndTimeFraction = EndTimeFraction;
-}
-
-//======================================================================================================================
-void UEulerAnalysisProperties::InitializeFromCache(TSharedPtr<FCachedAnalysisProperties> Cache)
-{
-	if (Cache)
+	UCachedAnalysisProperties* CachePtr = Cache.Get();
+	if (!Cache)
 	{
-		Super::InitializeFromCache(Cache);
-		FunctionAxis = Cache->EulerFunctionAxis;
-		Space = Cache->Space;
-		SpaceBoneSocket = Cache->SpaceBoneSocket;
-		CharacterFacingAxis = Cache->CharacterFacingAxis;
-		CharacterUpAxis = Cache->CharacterUpAxis;
-		StartTimeFraction = Cache->StartTimeFraction;
-		EndTimeFraction = Cache->EndTimeFraction;
-		BoneSocket = Cache->BoneSocket1;
-		BoneFacingAxis = Cache->BoneFacingAxis;
-		BoneRightAxis = Cache->BoneRightAxis;
+		Cache = NewObject<UCachedAnalysisProperties>(BlendSpace);
+		CachePtr = Cache.Get();
 	}
+	CachePtr->LinearFunctionAxis = FunctionAxis;
+	CachePtr->BoneSocket1 = BoneSocket;
+	CachePtr->Space = Space;
+	CachePtr->SpaceBoneSocket = SpaceBoneSocket;
+	CachePtr->StartTimeFraction = StartTimeFraction;
+	CachePtr->EndTimeFraction = EndTimeFraction;
+	CachePtr->bLockAfterAnalysis = bLockAfterAnalysis;
 }
 
 //======================================================================================================================
-void UEulerAnalysisProperties::MakeCache(TSharedPtr<FCachedAnalysisProperties>& Cache) const
-{
-	Super::MakeCache(Cache);
-	Cache->EulerFunctionAxis = FunctionAxis;
-	Cache->Space = Space;
-	Cache->SpaceBoneSocket = SpaceBoneSocket;
-	Cache->CharacterFacingAxis = CharacterFacingAxis;
-	Cache->CharacterUpAxis = CharacterUpAxis;
-	Cache->StartTimeFraction = StartTimeFraction;
-	Cache->EndTimeFraction = EndTimeFraction;
-	Cache->BoneSocket1 = BoneSocket;
-	Cache->BoneFacingAxis = BoneFacingAxis;
-	Cache->BoneRightAxis = BoneRightAxis;
-}
-
-//======================================================================================================================
-void UMovementAnalysisProperties::InitializeFromCache(TSharedPtr<FCachedAnalysisProperties> Cache)
+void UEulerAnalysisProperties::InitializeFromCache(TObjectPtr<UCachedAnalysisProperties> Cache)
 {
 	if (Cache)
 	{
-		Super::InitializeFromCache(Cache);
-		Space = Cache->Space;
-		SpaceBoneSocket = Cache->SpaceBoneSocket;
-		CharacterFacingAxis = Cache->CharacterFacingAxis;
-		CharacterUpAxis = Cache->CharacterUpAxis;
-		StartTimeFraction = Cache->StartTimeFraction;
-		EndTimeFraction = Cache->EndTimeFraction;
-		BoneSocket = Cache->BoneSocket1;
+		UCachedAnalysisProperties* CachePtr = Cache.Get();
+		FunctionAxis = CachePtr->EulerFunctionAxis;
+		BoneSocket = CachePtr->BoneSocket1;
+		BoneFacingAxis = CachePtr->BoneFacingAxis;
+		BoneRightAxis = CachePtr->BoneRightAxis;
+		Space = CachePtr->Space;
+		SpaceBoneSocket = CachePtr->SpaceBoneSocket;
+		CharacterFacingAxis = CachePtr->CharacterFacingAxis;
+		CharacterUpAxis = CachePtr->CharacterUpAxis;
+		StartTimeFraction = CachePtr->StartTimeFraction;
+		EndTimeFraction = CachePtr->EndTimeFraction;
+		bLockAfterAnalysis = CachePtr->bLockAfterAnalysis;
 	}
 }
 
 //======================================================================================================================
-void UMovementAnalysisProperties::MakeCache(TSharedPtr<FCachedAnalysisProperties>& Cache) const
+void UEulerAnalysisProperties::MakeCache(TObjectPtr<UCachedAnalysisProperties>& Cache, UBlendSpace* BlendSpace)
 {
-	Super::MakeCache(Cache);
-	Cache->Space = Space;
-	Cache->SpaceBoneSocket = SpaceBoneSocket;
-	Cache->CharacterFacingAxis = CharacterFacingAxis;
-	Cache->CharacterUpAxis = CharacterUpAxis;
-	Cache->StartTimeFraction = StartTimeFraction;
-	Cache->EndTimeFraction = EndTimeFraction;
-	Cache->BoneSocket1 = BoneSocket;
-}
-
-//======================================================================================================================
-void ULocomotionAnalysisProperties::InitializeFromCache(TSharedPtr<FCachedAnalysisProperties> Cache)
-{
-	if (Cache)
+	UCachedAnalysisProperties* CachePtr = Cache.Get();
+	if (!Cache)
 	{
-		Super::InitializeFromCache(Cache);
-		CharacterFacingAxis = Cache->CharacterFacingAxis;
-		CharacterUpAxis = Cache->CharacterUpAxis;
-		PrimaryBoneSocket = Cache->BoneSocket1;
-		SecondaryBoneSocket = Cache->BoneSocket2;
+		Cache = NewObject<UCachedAnalysisProperties>(BlendSpace);
+		CachePtr = Cache.Get();
 	}
+	CachePtr->EulerFunctionAxis = FunctionAxis;
+	CachePtr->BoneSocket1 = BoneSocket;
+	CachePtr->BoneFacingAxis = BoneFacingAxis;
+	CachePtr->BoneRightAxis = BoneRightAxis;
+	CachePtr->Space = Space;
+	CachePtr->SpaceBoneSocket = SpaceBoneSocket;
+	CachePtr->CharacterFacingAxis = CharacterFacingAxis;
+	CachePtr->CharacterUpAxis = CharacterUpAxis;
+	CachePtr->StartTimeFraction = StartTimeFraction;
+	CachePtr->EndTimeFraction = EndTimeFraction;
+	CachePtr->bLockAfterAnalysis = bLockAfterAnalysis;
 }
 
-//======================================================================================================================
-void ULocomotionAnalysisProperties::MakeCache(TSharedPtr<FCachedAnalysisProperties>& Cache) const
-{
-	Super::MakeCache(Cache);
-	Cache->CharacterFacingAxis = CharacterFacingAxis;
-	Cache->CharacterUpAxis = CharacterUpAxis;
-	Cache->BoneSocket1 = PrimaryBoneSocket;
-	Cache->BoneSocket2 = SecondaryBoneSocket;
-}
-
-
 //======================================================================================================================
 class FCoreBlendSpaceAnalysisFeature : public IBlendSpaceAnalysisFeature
 {
@@ -1214,16 +185,7 @@ TArray<FString> FCoreBlendSpaceAnalysisFeature::GetAnalysisFunctions() const
 		TEXT("Orientation"),
 		TEXT("OrientationRate"),
 		TEXT("DeltaOrientation"),
-		TEXT("AngularVelocity"),
-		TEXT("MovementSpeed"),
-		TEXT("MovementDirection"),
-		TEXT("LocomotionRightSpeed"),
-		TEXT("LocomotionForwardSpeed"),
-		TEXT("LocomotionUpSpeed"),
-		TEXT("LocomotionSpeed"),
-		TEXT("LocomotionDirection"),
-		TEXT("LocomotionForwardSlope"),
-		TEXT("LocomotionRightSlope")
+		TEXT("AngularVelocity")
 	};
 	return Functions;
 }
@@ -1245,21 +207,6 @@ UAnalysisProperties* FCoreBlendSpaceAnalysisFeature::MakeAnalysisProperties(UObj
 	{
 		Result = NewObject<UEulerAnalysisProperties>(Outer);
 	}
-	else if (FunctionName.Equals(TEXT("MovementSpeed")) ||
-			 FunctionName.Equals(TEXT("MovementDirection")))
-	{
-		Result = NewObject<UMovementAnalysisProperties>(Outer);
-	}
-	else if (FunctionName.Equals(TEXT("LocomotionForwardSpeed")) ||
-			 FunctionName.Equals(TEXT("LocomotionRightSpeed")) ||
-			 FunctionName.Equals(TEXT("LocomotionUpSpeed")) ||
-			 FunctionName.Equals(TEXT("LocomotionSpeed")) ||
-			 FunctionName.Equals(TEXT("LocomotionDirection")) ||
-			 FunctionName.Equals(TEXT("LocomotionForwardSlope")) ||
-			 FunctionName.Equals(TEXT("LocomotionRightSlope")))
-	{
-		Result = NewObject<ULocomotionAnalysisProperties>(Outer);
-	}
 
 	if (Result)
 	{
@@ -1282,103 +229,68 @@ bool FCoreBlendSpaceAnalysisFeature::CalculateSampleValue(float&
 	const FString& FunctionName = AnalysisProperties->Function;
 	if (FunctionName.Equals(TEXT("Position")))
 	{
-		return CalculateComponentSampleValue(
-			Result, CalculatePosition<ULinearAnalysisProperties>, BlendSpace, 
+		return BlendSpaceAnalysis::CalculateComponentSampleValue(
+			Result, BlendSpaceAnalysis::CalculatePosition<ULinearAnalysisProperties>, BlendSpace, 
 			Cast<ULinearAnalysisProperties>(AnalysisProperties), Animation, RateScale);
 	}
 	else if (FunctionName.Equals(TEXT("Velocity")))
 	{
-		return CalculateComponentSampleValue(
-			Result, CalculateVelocity<ULinearAnalysisProperties>, BlendSpace, 
+		return BlendSpaceAnalysis::CalculateComponentSampleValue(
+			Result, BlendSpaceAnalysis::CalculateVelocity<ULinearAnalysisProperties>, BlendSpace, 
 			Cast<ULinearAnalysisProperties>(AnalysisProperties), Animation, RateScale);
 	}
 	else if (FunctionName.Equals(TEXT("DeltaPosition")))
 	{
-		return CalculateComponentSampleValue(
-			Result, CalculateDeltaPosition<ULinearAnalysisProperties>, BlendSpace, 
+		return BlendSpaceAnalysis::CalculateComponentSampleValue(
+			Result, BlendSpaceAnalysis::CalculateDeltaPosition<ULinearAnalysisProperties>, BlendSpace, 
 			Cast<ULinearAnalysisProperties>(AnalysisProperties), Animation, RateScale);
 	}
 	else if (FunctionName.Equals(TEXT("AngularVelocity")))
 	{
-		return CalculateComponentSampleValue(
-			Result, CalculateAngularVelocity<ULinearAnalysisProperties>, BlendSpace, 
+		return BlendSpaceAnalysis::CalculateComponentSampleValue(
+			Result, BlendSpaceAnalysis::CalculateAngularVelocity<ULinearAnalysisProperties>, BlendSpace, 
 			Cast<ULinearAnalysisProperties>(AnalysisProperties), Animation, RateScale);
 	}
 	else if (FunctionName.Equals(TEXT("Orientation")))
 	{
-		return CalculateComponentSampleValue(
-			Result, CalculateOrientation<UEulerAnalysisProperties>, BlendSpace, 
+		return BlendSpaceAnalysis::CalculateComponentSampleValue(
+			Result, BlendSpaceAnalysis::CalculateOrientation<UEulerAnalysisProperties>, BlendSpace, 
 			Cast<UEulerAnalysisProperties>(AnalysisProperties), Animation, RateScale);
 	}
 	else if (FunctionName.Equals(TEXT("OrientationRate")))
 	{
-		return CalculateComponentSampleValue(
-			Result, CalculateOrientationRate<UEulerAnalysisProperties>, BlendSpace, 
+		return BlendSpaceAnalysis::CalculateComponentSampleValue(
+			Result, BlendSpaceAnalysis::CalculateOrientationRate<UEulerAnalysisProperties>, BlendSpace, 
 			Cast<UEulerAnalysisProperties>(AnalysisProperties), Animation, RateScale);
 	}
 	else if (FunctionName.Equals(TEXT("DeltaOrientation")))
 	{
-		return CalculateComponentSampleValue(
-			Result, CalculateDeltaOrientation<UEulerAnalysisProperties>, BlendSpace, 
+		return BlendSpaceAnalysis::CalculateComponentSampleValue(
+			Result, BlendSpaceAnalysis::CalculateDeltaOrientation<UEulerAnalysisProperties>, BlendSpace, 
 			Cast<UEulerAnalysisProperties>(AnalysisProperties), Animation, RateScale);
 	}
-	else if (FunctionName.Equals(TEXT("MovementSpeed")))
-	{
-		return CalculateMovementSpeed(
-			Result, BlendSpace, Cast<UMovementAnalysisProperties>(AnalysisProperties), Animation, RateScale);
-	}
-	else if (FunctionName.Equals(TEXT("MovementDirection")))
-	{
-		return CalculateMovementDirection(
-			Result, BlendSpace, Cast<UMovementAnalysisProperties>(AnalysisProperties), Animation, RateScale);
-	}
-	else if (FunctionName.Equals(TEXT("LocomotionForwardSpeed")))
-	{
-		return CalculateLocomotionFwdSpeed(
-			Result, BlendSpace, Cast<ULocomotionAnalysisProperties>(AnalysisProperties), Animation, RateScale);
-	}
-	else if (FunctionName.Equals(TEXT("LocomotionRightSpeed")))
-	{
-		return CalculateLocomotionRightSpeed(
-			Result, BlendSpace, Cast<ULocomotionAnalysisProperties>(AnalysisProperties), Animation, RateScale);
-	}
-	else if (FunctionName.Equals(TEXT("LocomotionUpSpeed")))
-	{
-		return CalculateLocomotionUpSpeed(
-			Result, BlendSpace, Cast<ULocomotionAnalysisProperties>(AnalysisProperties), Animation, RateScale);
-	}
-	else if (FunctionName.Equals(TEXT("LocomotionSpeed")))
-	{
-		return CalculateLocomotionSpeed(
-			Result, BlendSpace, Cast<ULocomotionAnalysisProperties>(AnalysisProperties), Animation, RateScale);
-	}
-	else if (FunctionName.Equals(TEXT("LocomotionDirection")))
-	{
-		return CalculateLocomotionDirection(
-			Result, BlendSpace, Cast<ULocomotionAnalysisProperties>(AnalysisProperties), Animation, RateScale);
-	}
-	else if (FunctionName.Equals(TEXT("LocomotionForwardSlope")))
-	{
-		return CalculateLocomotionFwdSlope(
-			Result, BlendSpace, Cast<ULocomotionAnalysisProperties>(AnalysisProperties), Animation, RateScale);
-	}
-	else if (FunctionName.Equals(TEXT("LocomotionRightSlope")))
-	{
-		return CalculateLocomotionRightSlope(
-			Result, BlendSpace, Cast<ULocomotionAnalysisProperties>(AnalysisProperties), Animation, RateScale);
-	}
 	return false;
 }
 
 //======================================================================================================================
-static TArray<IBlendSpaceAnalysisFeature*> GetAnalysisFeatures()
+static TArray<IBlendSpaceAnalysisFeature*> GetAnalysisFeatures(bool bCoreFeaturesLast)
 {
-	TArray<IBlendSpaceAnalysisFeature*> ModularFeatures = 
+	TArray<IBlendSpaceAnalysisFeature*> ModularFeatures;
+
+	if (!bCoreFeaturesLast)
+	{
+		ModularFeatures.Push(&CoreBlendSpaceAnalysisFeature);
+	}
+
+	TArray<IBlendSpaceAnalysisFeature*> ExtraModularFeatures = 
 		IModularFeatures::Get().GetModularFeatureImplementations<IBlendSpaceAnalysisFeature>(
 			IBlendSpaceAnalysisFeature::GetModuleFeatureName());
+	ModularFeatures += ExtraModularFeatures;
 
-	// Put the core one on last so that user-defined ones can override the default behaviour
-	ModularFeatures.Push(&CoreBlendSpaceAnalysisFeature);
+	if (bCoreFeaturesLast)
+	{
+		ModularFeatures.Push(&CoreBlendSpaceAnalysisFeature);
+	}
 	return ModularFeatures;
 }
 
@@ -1388,7 +300,7 @@ FVector BlendSpaceAnalysis::CalculateSampleValue(const UBlendSpace& BlendSpace,
 {
 
 	FVector AdjustedPosition = OriginalPosition;
-	TArray<IBlendSpaceAnalysisFeature*> ModularFeatures = GetAnalysisFeatures();
+	TArray<IBlendSpaceAnalysisFeature*> ModularFeatures = GetAnalysisFeatures(true);
 	for (int32 Index = 0; Index != 2; ++Index)
 	{
 		bAnalyzed[Index] = false;
@@ -1413,7 +325,7 @@ UAnalysisProperties* BlendSpaceAnalysis::MakeAnalysisProperties(UObject* Outer,
 {
 	UAnalysisProperties* Result = nullptr;
 
-	TArray<IBlendSpaceAnalysisFeature*> ModularFeatures = GetAnalysisFeatures();
+	TArray<IBlendSpaceAnalysisFeature*> ModularFeatures = GetAnalysisFeatures(true);
 	for (const IBlendSpaceAnalysisFeature* Feature : ModularFeatures)
 	{
 		Result = Feature->MakeAnalysisProperties(Outer, FunctionName);
@@ -1429,7 +341,7 @@ UAnalysisProperties* BlendSpaceAnalysis::MakeAnalysisProperties(UObject* Outer,
 TArray<FString> BlendSpaceAnalysis::GetAnalysisFunctions()
 {
 	TArray<FString> FunctionNames;
-	TArray<IBlendSpaceAnalysisFeature*> ModularFeatures = GetAnalysisFeatures();
+	TArray<IBlendSpaceAnalysisFeature*> ModularFeatures = GetAnalysisFeatures(false);
 	for (const IBlendSpaceAnalysisFeature* Feature : ModularFeatures)
 	{
 		TArray<FString> FeatureFunctionNames = Feature->GetAnalysisFunctions();
@@ -1447,5 +359,4 @@ bool BlendSpaceAnalysis::GetLockAfterAnalysis(const TObjectPtr<UAnalysisProperti
 	return AnalysisProperties ? AnalysisProperties->bLockAfterAnalysis : false;
 }
 
-
 #undef LOCTEXT_NAMESPACE
diff --git a/Engine/Source/Editor/Persona/Private/SAnimationBlendSpace.cpp b/Engine/Source/Editor/Persona/Private/SAnimationBlendSpace.cpp
--- a/Engine/Source/Editor/Persona/Private/SAnimationBlendSpace.cpp
+++ b/Engine/Source/Editor/Persona/Private/SAnimationBlendSpace.cpp
@@ -150,14 +150,14 @@ void SBlendSpaceEditor::OnSampleRemoved(const int32 SampleIndex)
 }
 
 //======================================================================================================================
-int32 SBlendSpaceEditor::OnSampleAdded(UAnimSequence* Animation, const FVector& Value)
+int32 SBlendSpaceEditor::OnSampleAdded(UAnimSequence* Animation, const FVector& Value, bool bRunAnalysis)
 {
 	FScopedTransaction ScopedTransaction(LOCTEXT("AddSample", "Adding Blend Grid Sample"));
 
 	FVector AdjustedValue = Value;
 	bool bAnalyzed[3] = { false, false, false };
 
-	if(BlendSpace->IsAsset())
+	if(BlendSpace->IsAsset() && bRunAnalysis)
 	{
 		AdjustedValue = BlendSpaceAnalysis::CalculateSampleValue(*BlendSpace, *Animation, 1.0f, Value, bAnalyzed);
 	}
@@ -190,7 +190,7 @@ int32 SBlendSpaceEditor::OnSampleAdded(UAnimSequence* Animation, const FVector&
 
 		if (OnBlendSpaceSampleAdded.IsBound())
 		{
-			OnBlendSpaceSampleAdded.Execute(Animation, AdjustedValue);
+			OnBlendSpaceSampleAdded.Execute(Animation, AdjustedValue, bRunAnalysis);
 		}
 		BlendSpace->PostEditChange();
 	}
@@ -198,10 +198,10 @@ int32 SBlendSpaceEditor::OnSampleAdded(UAnimSequence* Animation, const FVector&
 	return NewSampleIndex;
 }
 
-void SBlendSpaceEditor::OnSampleDuplicated(const int32 SampleIndex, const FVector& NewValue)
+void SBlendSpaceEditor::OnSampleDuplicated(const int32 SampleIndex, const FVector& NewValue, bool bRunAnalysis)
 {
 	const FBlendSample& OrigSample = BlendSpace->GetBlendSample(SampleIndex);
-	int32 NewSampleIndex = OnSampleAdded(OrigSample.Animation, NewValue);
+	int32 NewSampleIndex = OnSampleAdded(OrigSample.Animation, NewValue, bRunAnalysis);
 	if (NewSampleIndex >= 0)
 	{
 		BlendSpace->LockSample(NewSampleIndex, OrigSample.bLockX, OrigSample.bLockY, OrigSample.bLockZ);
diff --git a/Engine/Source/Editor/Persona/Private/SAnimationBlendSpace.h b/Engine/Source/Editor/Persona/Private/SAnimationBlendSpace.h
--- a/Engine/Source/Editor/Persona/Private/SAnimationBlendSpace.h
+++ b/Engine/Source/Editor/Persona/Private/SAnimationBlendSpace.h
@@ -57,8 +57,8 @@ public:
 protected:
 	void OnSampleMoved(const int32 SampleIndex, const FVector& NewValue, bool bIsInteractive);
 	void OnSampleRemoved(const int32 SampleIndex);
-	int32 OnSampleAdded(UAnimSequence* Animation, const FVector& Value);
-	void OnSampleDuplicated(const int32 SampleIndex, const FVector& NewValue);
+	int32 OnSampleAdded(UAnimSequence* Animation, const FVector& Value, bool bRunAnalysis);
+	void OnSampleDuplicated(const int32 SampleIndex, const FVector& NewValue, bool bRunAnalysis);
 	void OnSampleReplaced(const int32 SampleIndex, UAnimSequence* Animation);
 
 	// Begin SAnimEditorBase overrides
diff --git a/Engine/Source/Editor/Persona/Private/SAnimationBlendSpaceGridWidget.cpp b/Engine/Source/Editor/Persona/Private/SAnimationBlendSpaceGridWidget.cpp
--- a/Engine/Source/Editor/Persona/Private/SAnimationBlendSpaceGridWidget.cpp
+++ b/Engine/Source/Editor/Persona/Private/SAnimationBlendSpaceGridWidget.cpp
@@ -996,7 +996,7 @@ FReply SBlendSpaceGridWidget::OnDrop(const FGeometry& MyGeometry, const FDragDro
 						UAnimSequence* Animation = (UAnimSequence*) Asset;
 						if (OnSampleAdded.IsBound())
 						{
-							OnSampleAdded.Execute(Animation, SampleValue);
+							OnSampleAdded.Execute(Animation, SampleValue, true);
 						}
 					}
 				}	
@@ -1402,7 +1402,7 @@ void SBlendSpaceGridWidget::MakeViewContextMenuEntries(FMenuBuilder& InMenuBuild
 			LOCTEXT("StretchFittingTextToolTip", "Whether to stretch the grid to fit or to fit the grid to the largest axis"),
 			FSlateIcon("EditorStyle", "BlendSpaceEditor.ZoomToFit"),
 			FUIAction(
-				FExecuteAction::CreateLambda([this](){ bStretchToFit = !bStretchToFit; }),
+				FExecuteAction::CreateLambda([this](){ ToggleFittingType(); }),
 				FCanExecuteAction(),
 				FGetActionCheckState::CreateLambda([this](){ return bStretchToFit ? ECheckBoxState::Checked : ECheckBoxState::Unchecked; })
 			),
@@ -1488,19 +1488,44 @@ TSharedPtr<SWidget> SBlendSpaceGridWidget::CreateNewBlendSampleContextMenu(const
 	{
 		MenuBuilder.BeginSection("Sample", LOCTEXT("SampleMenuHeader", "Sample"));
 		{
-			if(!BlendSpace->IsAsset())
+			if(BlendSpace->IsAsset())
+			{
+				// Duplicate an existing sample if one was selected. Note that this could be enabled for the graph too,
+				// but we'd need a way of either duplicating the original graph, or referencing it (and choosing between
+				// the two).
+				if (SelectedSampleIndex != INDEX_NONE)
+				{
+					MenuBuilder.AddMenuEntry(
+						LOCTEXT("DuplicateSample", "Duplicate Sample Here"),
+						LOCTEXT("DuplicateSampleTooltip", "Duplicate the selected sample and place it in the blendspace at this location"),
+						FSlateIcon("EditorStyle", "Icons.Duplicate"),
+						FUIAction(
+							FExecuteAction::CreateLambda(
+								[this, NewSampleValue]()
+								{
+									if (OnSampleDuplicated.IsBound())
+									{
+										OnSampleDuplicated.Execute(SelectedSampleIndex, NewSampleValue, false);
+									}
+								})
+						)
+					);
+				}
+			}
+			else
 			{
+				// Blend space graph - add a new graph sample
 				MenuBuilder.AddMenuEntry(
 					LOCTEXT("AddNewSample", "Add New Sample"),
-					LOCTEXT("AddNewSampleTooltip", "Add a new sample to this blendspace at this location"),
-						FSlateIcon("EditorStyle", "Plus"),
+					LOCTEXT("AddNewSampleTooltip", "Add a new sample to the blendspace at this location"),
+						FSlateIcon("EditorStyle", "Icons.Plus"),
 						FUIAction(
 							FExecuteAction::CreateLambda(
 								[this, NewSampleValue]()
 								{
 									if (OnSampleAdded.IsBound())
 									{
-										OnSampleAdded.Execute(nullptr, NewSampleValue);
+										OnSampleAdded.Execute(nullptr, NewSampleValue, false);
 									}
 								})
 						)
diff --git a/Engine/Source/Editor/Persona/Public/BlendSpaceAnalysis.h b/Engine/Source/Editor/Persona/Public/BlendSpaceAnalysis.h
--- a/Engine/Source/Editor/Persona/Public/BlendSpaceAnalysis.h
+++ b/Engine/Source/Editor/Persona/Public/BlendSpaceAnalysis.h
@@ -6,10 +6,13 @@
 #include "Animation/BoneSocketReference.h"
 #include "Features/IModularFeature.h"
 #include "Animation/BlendSpace.h"
+#include "Animation/AnimSequence.h"
 
 #include "BlendSpaceAnalysis.generated.h"
 
-class UAnimSequence;
+#define LOCTEXT_NAMESPACE "BlendSpaceAnalysis"
+//#define ANALYSIS_VERBOSE_LOG
+
 class UBlendSpace;
 class UAnalysisProperties;
 
@@ -42,63 +45,71 @@ public:
 UENUM()
 enum class EAnalysisSpace : uint8
 {
-	World    UMETA(DisplayName = "World"),
-	Fixed    UMETA(DisplayName = "Fixed"),
-	Changing UMETA(DisplayName = "Changing"),
-	Moving   UMETA(DisplayName = "Moving"),
+	World,
+	Fixed,
+	Changing,
+	Moving,
 };
 
 UENUM()
 enum class EAnalysisLinearAxis : uint8
 {
-	X  UMETA(DisplayName = "X"),
-	Y  UMETA(DisplayName = "Y"),
-	Z  UMETA(DisplayName = "Z"),
+	X,
+	Y,
+	Z,
 };
 
 UENUM()
 enum class EAnalysisEulerAxis : uint8
 {
-	Roll   UMETA(DisplayName = "Roll"),
-	Pitch  UMETA(DisplayName = "Pitch"),
-	Yaw    UMETA(DisplayName = "Yaw"),
+	Roll,
+	Pitch,
+	Yaw,
 };
 
+
 /**
  * This will be used to preserve values as far as possible when switching between analysis functions, so it contains all
  * the parameters used by the engine functions. User defined can inherit from this and add their own - then the
  * user-defined MakeCache function should replace any base class cache that is passed in with their own.
 */
-class FCachedAnalysisProperties
+UCLASS()
+class PERSONA_API UCachedAnalysisProperties : public UObject
 {
+	GENERATED_BODY()
 public:
-	EAnalysisLinearAxis  LinearFunctionAxis = EAnalysisLinearAxis::X;
-	EAnalysisEulerAxis   EulerFunctionAxis = EAnalysisEulerAxis::Pitch;
-	EAnalysisSpace       Space = EAnalysisSpace::World;
-	FBoneSocketTarget    SpaceBoneSocket;
-	EAnalysisLinearAxis  CharacterFacingAxis = EAnalysisLinearAxis::Y;
-	EAnalysisLinearAxis  CharacterUpAxis = EAnalysisLinearAxis::Z;
-	float                StartTimeFraction = 0.0f;
-	float                EndTimeFraction = 1.0f;
-	FBoneSocketTarget    BoneSocket1;
-	FBoneSocketTarget    BoneSocket2;
-	EAnalysisLinearAxis  BoneFacingAxis = EAnalysisLinearAxis::X;
-	EAnalysisLinearAxis  BoneRightAxis = EAnalysisLinearAxis::Y;
-	bool                 bLockAfterAnalysis = false;
+	void CopyFrom(const UCachedAnalysisProperties& Other);
+	EAnalysisLinearAxis     LinearFunctionAxis = EAnalysisLinearAxis::X;
+	EAnalysisEulerAxis      EulerFunctionAxis = EAnalysisEulerAxis::Pitch;
+	FBoneSocketTarget       BoneSocket1;
+	FBoneSocketTarget       BoneSocket2;
+	EAnalysisLinearAxis     BoneFacingAxis = EAnalysisLinearAxis::X;
+	EAnalysisLinearAxis     BoneRightAxis = EAnalysisLinearAxis::Y;
+	EAnalysisSpace          Space = EAnalysisSpace::World;
+	FBoneSocketTarget       SpaceBoneSocket;
+	EAnalysisLinearAxis     CharacterFacingAxis = EAnalysisLinearAxis::Y;
+	EAnalysisLinearAxis     CharacterUpAxis = EAnalysisLinearAxis::Z;
+	float                   StartTimeFraction = 0.0f;
+	float                   EndTimeFraction = 1.0f;
+	bool                    bLockAfterAnalysis = false;
 };
 
 UCLASS()
-class ULinearAnalysisProperties : public UAnalysisProperties
+class PERSONA_API ULinearAnalysisProperties : public UAnalysisProperties
 {
 	GENERATED_BODY()
 public:
-	void InitializeFromCache(TSharedPtr<FCachedAnalysisProperties> Cache) override;
-	void MakeCache(TSharedPtr<FCachedAnalysisProperties>& Cache) const override;
+	void InitializeFromCache(TObjectPtr<UCachedAnalysisProperties> Cache) override;
+	void MakeCache(TObjectPtr<UCachedAnalysisProperties>& Cache, UBlendSpace* BlendSpace) override;
 
 	/** Axis for the analysis function */
 	UPROPERTY(EditAnywhere, DisplayName = "Axis", Category = AnalysisProperties)
 	EAnalysisLinearAxis FunctionAxis = EAnalysisLinearAxis::X;
 
+	/** The bone or socket used for analysis */
+	UPROPERTY(EditAnywhere, DisplayName = "Bone/Socket", Category = AnalysisProperties)
+	FBoneSocketTarget BoneSocket;
+
 	/**
 	* The space in which to perform the analysis. Fixed will use the analysis bone/socket at the first frame
 	* of the analysis time range. Changing will use the analysis bone/socket at the relevant frame during the
@@ -112,10 +123,6 @@ public:
 	UPROPERTY(EditAnywhere, DisplayName = "Analysis Space Bone/Socket", Category = AnalysisProperties, meta = (EditCondition = "Space != EAnalysisSpace::World"))
 	FBoneSocketTarget SpaceBoneSocket;
 
-	/** The bone or socket used for analysis */
-	UPROPERTY(EditAnywhere, DisplayName = "Bone/Socket", Category = AnalysisProperties)
-	FBoneSocketTarget BoneSocket;
-
 	/** Fraction through each animation at which analysis starts */
 	UPROPERTY(EditAnywhere, DisplayName = "Start time fraction", Category = AnalysisProperties, meta = (ClampMin = "0", ClampMax = "1"))
 	float StartTimeFraction = 0.0f;
@@ -126,46 +133,17 @@ public:
 };
 
 UCLASS()
-class UEulerAnalysisProperties : public UAnalysisProperties
+class PERSONA_API UEulerAnalysisProperties : public UAnalysisProperties
 {
 	GENERATED_BODY()
 public:
-	void InitializeFromCache(TSharedPtr<FCachedAnalysisProperties> Cache) override;
-	void MakeCache(TSharedPtr<FCachedAnalysisProperties>& Cache) const override;
+	void InitializeFromCache(TObjectPtr<UCachedAnalysisProperties> Cache) override;
+	void MakeCache(TObjectPtr<UCachedAnalysisProperties>& Cache, UBlendSpace* BlendSpace) override;
 
 	/** Axis for the analysis function */
 	UPROPERTY(EditAnywhere, DisplayName = "Axis", Category = AnalysisProperties)
 	EAnalysisEulerAxis FunctionAxis = EAnalysisEulerAxis::Pitch;
 
-	/**
-	* The space in which to perform the analysis. Fixed will use the analysis bone/socket at the first frame
-	* of the analysis time range. Changing will use the analysis bone/socket at the relevant frame during the
-	* analysis, but calculate velocities assuming that frame isn't moving. Moving will do the same but velocities
-	* as well as positions/rotations will be relative to this moving frame.
-	*/
-	UPROPERTY(EditAnywhere, Category = AnalysisProperties)
-	EAnalysisSpace Space = EAnalysisSpace::World;
-
-	/** Bone or socket that defines the analysis space (when it isn't World) */
-	UPROPERTY(EditAnywhere, DisplayName = "Analysis Space Bone/Socket", Category = AnalysisProperties, meta = (EditCondition = "Space != 0"))
-	FBoneSocketTarget SpaceBoneSocket;
-
-	/** World or bone/socket axis that specifies the character's facing direction */
-	UPROPERTY(EditAnywhere, Category = AnalysisProperties)
-	EAnalysisLinearAxis CharacterFacingAxis = EAnalysisLinearAxis::Y;
-
-	/** World or bone/socket axis that specifies the character's up direction */
-	UPROPERTY(EditAnywhere, Category = AnalysisProperties)
-	EAnalysisLinearAxis CharacterUpAxis = EAnalysisLinearAxis::Z;
-
-	/** Fraction through each animation at which analysis starts */
-	UPROPERTY(EditAnywhere, DisplayName = "Start time fraction", Category = AnalysisProperties, meta = (ClampMin = "0", ClampMax = "1"))
-	float StartTimeFraction = 0.0f;
-	/** Fraction through each animation at which analysis ends */
-
-	UPROPERTY(EditAnywhere, DisplayName = "End time fraction", Category = AnalysisProperties, meta = (ClampMin = "0", ClampMax = "1"))
-	float EndTimeFraction = 1.0f;
-
 	/** The bone or socket used for analysis */
 	UPROPERTY(EditAnywhere, DisplayName = "Bone/Socket", Category = AnalysisProperties)
 	FBoneSocketTarget BoneSocket;
@@ -177,15 +155,6 @@ public:
 	/** Used for some analysis functions - specifies the bone/socket axis that points to the "right" */
 	UPROPERTY(EditAnywhere, Category = AnalysisProperties)
 	EAnalysisLinearAxis BoneRightAxis = EAnalysisLinearAxis::Y;
-};
-
-UCLASS()
-class UMovementAnalysisProperties : public UAnalysisProperties
-{
-	GENERATED_BODY()
-public:
-	void InitializeFromCache(TSharedPtr<FCachedAnalysisProperties> Cache) override;
-	void MakeCache(TSharedPtr<FCachedAnalysisProperties>& Cache) const override;
 
 	/**
 	* The space in which to perform the analysis. Fixed will use the analysis bone/socket at the first frame
@@ -211,38 +180,573 @@ public:
 	/** Fraction through each animation at which analysis starts */
 	UPROPERTY(EditAnywhere, DisplayName = "Start time fraction", Category = AnalysisProperties, meta = (ClampMin = "0", ClampMax = "1"))
 	float StartTimeFraction = 0.0f;
-
 	/** Fraction through each animation at which analysis ends */
+
 	UPROPERTY(EditAnywhere, DisplayName = "End time fraction", Category = AnalysisProperties, meta = (ClampMin = "0", ClampMax = "1"))
 	float EndTimeFraction = 1.0f;
-
-	/** The bone or socket used for analysis */
-	UPROPERTY(EditAnywhere, DisplayName = "Bone/Socket", Category = AnalysisProperties)
-	FBoneSocketTarget BoneSocket;
 };
 
-UCLASS()
-class ULocomotionAnalysisProperties : public UAnalysisProperties
-{
-	GENERATED_BODY()
-public:
-	void InitializeFromCache(TSharedPtr<FCachedAnalysisProperties> Cache) override;
-	void MakeCache(TSharedPtr<FCachedAnalysisProperties>& Cache) const override;
+//======================================================================================================================
+// The following are helper functions which may be useful when implementing analysis functions
+//======================================================================================================================
 
-	/** World or bone/socket axis that specifies the character's facing direction */
-	UPROPERTY(EditAnywhere, Category = AnalysisProperties)
-	EAnalysisLinearAxis CharacterFacingAxis = EAnalysisLinearAxis::Y;
-
-	/** World or bone/socket axis that specifies the character's up direction */
-	UPROPERTY(EditAnywhere, Category = AnalysisProperties)
-	EAnalysisLinearAxis CharacterUpAxis = EAnalysisLinearAxis::Z;
+namespace BlendSpaceAnalysis
+{
 
-	/** The primary bone or socket used for analysis */
-	UPROPERTY(EditAnywhere, DisplayName = "Bone/Socket 1", Category = AnalysisProperties)
-	FBoneSocketTarget PrimaryBoneSocket;
+//======================================================================================================================
+// Retrieves the bone index and transform offset given the BoneSocketTarget. Returns true if found
+PERSONA_API bool GetBoneInfo(const UAnimSequence&     Animation, 
+							 const FBoneSocketTarget& BoneSocket, 
+							 FTransform&              BoneOffset, 
+							 FName&                   BoneName);
+
+//======================================================================================================================
+PERSONA_API FTransform GetBoneTransform(const UAnimSequence& Animation, int32 Key, const FName& BoneName);
+
+//======================================================================================================================
+template<typename T>
+void CalculateFrameTM(
+	bool& bNeedToUpdateFrameTM, FTransform& FrameTM, 
+	const int32 SampleKey, const T& AnalysisProperties, const UAnimSequence& Animation)
+{
+	if (bNeedToUpdateFrameTM)
+	{
+		FrameTM.SetIdentity();
+		if (AnalysisProperties->Space != EAnalysisSpace::World)
+		{
+			FTransform SpaceBoneOffset;
+			FName SpaceBoneName;
+			if (GetBoneInfo(Animation, AnalysisProperties->SpaceBoneSocket, SpaceBoneOffset, SpaceBoneName))
+			{
+				FTransform SpaceBoneTM = GetBoneTransform(Animation, SampleKey, SpaceBoneName);
+				FrameTM = SpaceBoneOffset * SpaceBoneTM;
+			}
+		}
+
+		bNeedToUpdateFrameTM = (
+			AnalysisProperties->Space == EAnalysisSpace::Changing || 
+			AnalysisProperties->Space == EAnalysisSpace::Moving);
+	}
+}
+
+//======================================================================================================================
+PERSONA_API FVector GetAxisFromTM(const FTransform& TM, EAnalysisLinearAxis Axis);
+
+//======================================================================================================================
+template<typename T>
+void GetFrameDirs(
+	FVector& FrameFacingDir, FVector& FrameUpDir, FVector& FrameRightDir, 
+	const FTransform& FrameTM, const T& AnalysisProperties)
+{
+	FrameFacingDir = GetAxisFromTM(FrameTM, AnalysisProperties->CharacterFacingAxis);
+	FrameUpDir = GetAxisFromTM(FrameTM, AnalysisProperties->CharacterUpAxis);
+	FrameRightDir = FVector::CrossProduct(FrameUpDir, FrameFacingDir);
+}
 
-	/** The secondary bone or socket used for analysis */
-	UPROPERTY(EditAnywhere, DisplayName = "Bone/Socket 2", Category = AnalysisProperties)
-	FBoneSocketTarget SecondaryBoneSocket;
-};
+//======================================================================================================================
+/**
+* Helper to extract the component from the FVector functions
+*/
+template<typename FunctionType, typename T>
+static bool CalculateComponentSampleValue(
+	float&                     Result,
+	const FunctionType&        Fn,
+	const UBlendSpace&         BlendSpace, 
+	const T*                   AnalysisProperties, 
+	const UAnimSequence&       Animation,
+	const float                RateScale)
+{
+	FVector Value;
+	int32 ComponentIndex = (int32) AnalysisProperties->FunctionAxis;
+	if (Fn(Value, BlendSpace, AnalysisProperties, Animation, RateScale))
+	{
+		Result = Value[ComponentIndex];
+		return true;
+	}
+	return false;
+}
+
+//======================================================================================================================
+template <typename T>
+static bool CalculatePosition(
+	FVector&                         Result,
+	const UBlendSpace&               BlendSpace,
+	const T*                         AnalysisProperties,
+	const UAnimSequence&             Animation,
+	const float                      RateScale)
+{
+	const int32 NumSampledKeys = Animation.GetNumberOfSampledKeys() - 1;
+	if (!AnalysisProperties || NumSampledKeys <= 0)
+	{
+		return false;
+	}
+
+	FTransform BoneOffset;
+	FName BoneName;
+	if (!GetBoneInfo(Animation, AnalysisProperties->BoneSocket, BoneOffset, BoneName))
+	{
+		return false;
+	}
+
+	int32 FirstKey = FMath::Clamp(
+		(int32) (NumSampledKeys * AnalysisProperties->StartTimeFraction), 0, NumSampledKeys);
+	int32 LastKey = FMath::Clamp(
+		(int32) (NumSampledKeys * AnalysisProperties->EndTimeFraction), FirstKey, NumSampledKeys);
+
+	FTransform FrameTM;
+	bool bNeedToUpdateFrameTM = true;
+	CalculateFrameTM(bNeedToUpdateFrameTM, FrameTM, FirstKey, AnalysisProperties, Animation);
+
+	Result.Set(0, 0, 0);
+	for (int32 Key = FirstKey; Key != LastKey + 1; ++Key)
+	{
+		CalculateFrameTM(bNeedToUpdateFrameTM, FrameTM, Key, AnalysisProperties, Animation);
+		FTransform BoneTM = GetBoneTransform(Animation, Key, BoneName);
+		FTransform TM = BoneOffset * BoneTM;
+		FVector RelativePos = FrameTM.InverseTransformPosition(TM.GetTranslation());
+		Result += RelativePos;
+	}
+	Result /= (1 + LastKey - FirstKey);
+	return true;
+}
+
+//======================================================================================================================
+template <typename T>
+static bool CalculateDeltaPosition(
+	FVector&                         Result,
+	const UBlendSpace&               BlendSpace,
+	const T*                         AnalysisProperties,
+	const UAnimSequence&             Animation,
+	const float                      RateScale)
+{
+	const int32 NumSampledKeys = Animation.GetNumberOfSampledKeys() - 1;
+	if (!AnalysisProperties || NumSampledKeys <= 0)
+	{
+		return false;
+	}
+
+	FTransform BoneOffset;
+	FName BoneName;
+	if (!GetBoneInfo(Animation, AnalysisProperties->BoneSocket, BoneOffset, BoneName))
+	{
+		return false;
+	}
+
+	int32 FirstKey = FMath::Clamp(
+		(int32) (NumSampledKeys * AnalysisProperties->StartTimeFraction), 0, NumSampledKeys);
+	int32 LastKey = FMath::Clamp(
+		(int32) (NumSampledKeys * AnalysisProperties->EndTimeFraction), FirstKey, NumSampledKeys);
+
+	FTransform FrameTM;
+	bool bNeedToUpdateFrameTM = true;
+	CalculateFrameTM(bNeedToUpdateFrameTM, FrameTM, FirstKey, AnalysisProperties, Animation);
+
+	FTransform BoneTM1 = GetBoneTransform(Animation, FirstKey, BoneName);
+	FTransform TM1 = BoneOffset * BoneTM1;
+	FVector RelativePos1 = FrameTM.InverseTransformPosition(TM1.GetTranslation());
+
+	CalculateFrameTM(bNeedToUpdateFrameTM, FrameTM, FirstKey, AnalysisProperties, Animation);
+	FTransform BoneTM2 = GetBoneTransform(Animation, LastKey, BoneName);
+	FTransform TM2 = BoneOffset * BoneTM2;
+	FVector RelativePos2 = FrameTM.InverseTransformPosition(TM2.GetTranslation());
+
+	Result = RelativePos2 - RelativePos1;
+	return true;
+}
+
+//======================================================================================================================
+template <typename T>
+static bool CalculateVelocity(
+	FVector&             Result,
+	const UBlendSpace&   BlendSpace,
+	const T*             AnalysisProperties,
+	const UAnimSequence& Animation,
+	const float          RateScale)
+{
+	const int32 NumSampledKeys = Animation.GetNumberOfSampledKeys() - 1;
+	if (!AnalysisProperties || NumSampledKeys <= 0)
+	{
+		return false;
+	}
+
+	if (NumSampledKeys == 1)
+	{
+		Result.Set(0, 0, 0);
+		return true;
+	}
+
+	FTransform BoneOffset;
+	FName BoneName;
+	if (!GetBoneInfo(Animation, AnalysisProperties->BoneSocket, BoneOffset, BoneName))
+	{
+		return false;
+	}
+
+	float DeltaTime = Animation.GetPlayLength() / NumSampledKeys;
+
+	int32 FirstKey = FMath::Clamp(
+		(int32) (NumSampledKeys * AnalysisProperties->StartTimeFraction), 0, NumSampledKeys);
+	int32 LastKey = FMath::Clamp(
+		(int32) (NumSampledKeys * AnalysisProperties->EndTimeFraction), FirstKey, NumSampledKeys-1);
+
+	if (FirstKey >= LastKey)
+	{
+		Result.Set(0, 0, 0);
+		return true;
+	}
+
+	FTransform FrameTM;
+	bool bNeedToUpdateFrameTM = true;
+	CalculateFrameTM(bNeedToUpdateFrameTM, FrameTM, FirstKey, AnalysisProperties, Animation);
+
+	Result.Set(0, 0, 0);
+	for (int32 Key = FirstKey; Key != LastKey + 1; ++Key)
+	{
+		CalculateFrameTM(bNeedToUpdateFrameTM, FrameTM, Key, AnalysisProperties, Animation);
+
+		FTransform BoneTM1 = GetBoneTransform(Animation, Key, BoneName);
+		FTransform TM1 = BoneOffset * BoneTM1;
+		FVector RelativePos1 = FrameTM.InverseTransformPosition(TM1.GetTranslation());
+
+		int32 NextKey = (Key + 1) % (NumSampledKeys + 1);
+		if (AnalysisProperties->Space == EAnalysisSpace::Moving)
+		{
+			CalculateFrameTM(bNeedToUpdateFrameTM, FrameTM, NextKey, AnalysisProperties, Animation);
+		}
+
+		FTransform BoneTM2 = GetBoneTransform(Animation, NextKey, BoneName);
+		FTransform TM2 = BoneOffset * BoneTM2;
+		FVector RelativePos2 = FrameTM.InverseTransformPosition(TM2.GetTranslation());
+		FVector Velocity = (RelativePos2 - RelativePos1) / DeltaTime;
+
+#ifdef ANALYSIS_VERBOSE_LOG
+		UE_LOG(LogAnimation, Log, TEXT("%d Velocity = %f %f %f Height = %f"), 
+			   Key, Velocity.X, Velocity.Y, Velocity.Z, 0.5f * (RelativePos1 + RelativePos2).Z);
+#endif
+		Result += Velocity;
+	}
+	Result /= (1 + LastKey - FirstKey);
+	Result *= Animation.RateScale * RateScale;
+	UE_LOG(LogAnimation, Log, TEXT("%s vel = %f %f %f"), *BoneName.ToString(), Result.X, Result.Y, Result.Z);
+	return true;
+}
+
+//======================================================================================================================
+template <typename T>
+void CalculateBoneOrientation(
+	FVector&                   RollPitchYaw,
+	const UAnimSequence&       Animation, 
+	const int32                Key, 
+	const FName                BoneName, 
+	const FTransform&          BoneOffset, 
+	const T&                   AnalysisProperties, 
+	const FVector&             FrameFacingDir, 
+	const FVector&             FrameRightDir, 
+	const FVector&             FrameUpDir)
+{
+	FTransform BoneTM = GetBoneTransform(Animation, Key, BoneName);
+
+	FTransform TM = BoneOffset * BoneTM;
+	FQuat AimQuat = TM.GetRotation();
+	FVector AimForwardDir = BlendSpaceAnalysis::GetAxisFromTM(TM, AnalysisProperties->BoneFacingAxis);
+	FVector AimRightDir = BlendSpaceAnalysis::GetAxisFromTM(TM, AnalysisProperties->BoneRightAxis);
+
+	// Note that Yaw is best taken from the AimRightDir - this is to avoid problems when the gun is pointing
+	// up or down - especially if it goes beyond 90 degrees in pitch.
+	float Yaw = FMath::RadiansToDegrees(FMath::Atan2(
+		FVector::DotProduct(AimRightDir, -FrameFacingDir), FVector::DotProduct(AimRightDir, FrameRightDir)));
+
+	// Undo the yaw to get pitch
+	const FQuat YawQuat(FrameUpDir, FMath::DegreesToRadians(Yaw));
+	FVector UnYawedAimForwardDir = YawQuat.UnrotateVector(AimForwardDir);
+	float Up = UnYawedAimForwardDir | FrameUpDir;
+	float Forward = UnYawedAimForwardDir | FrameFacingDir;
+	float Pitch = FMath::RadiansToDegrees(FMath::Atan2(Up, Forward));
+
+	// Undo the pitch to get roll
+	FVector UnYawedAimRightDir = YawQuat.UnrotateVector(AimRightDir);
+	const FQuat PitchQuat(FrameRightDir, -FMath::DegreesToRadians(Pitch));
+
+	FVector UnYawedUnPitchedAimRightDir = PitchQuat.UnrotateVector(UnYawedAimRightDir);
+
+	float Roll = FMath::RadiansToDegrees(FMath::Atan2(
+		FVector::DotProduct(UnYawedUnPitchedAimRightDir, -FrameUpDir), 
+		FVector::DotProduct(UnYawedUnPitchedAimRightDir, FrameRightDir)));
+
+	RollPitchYaw.Set(Roll, Pitch, Yaw);
+}
+
+//======================================================================================================================
+// Note that if a looping animation has 56 keys, then its first key is 0 and last is 55, but these will be identical poses.
+// Thus it has one fewer intervals/unique keys
+template <typename T>
+static bool CalculateOrientation(
+	FVector&                        Result,
+	const UBlendSpace&              BlendSpace,
+	const T*                        AnalysisProperties,
+	const UAnimSequence&            Animation,
+	const float                     RateScale)
+{
+	const int32 NumSampledKeys = Animation.GetNumberOfSampledKeys() - 1;
+	if (!AnalysisProperties || NumSampledKeys <= 0)
+	{
+		return false;
+	}
+
+	FTransform BoneOffset;
+	FName BoneName;
+	if (!GetBoneInfo(Animation, AnalysisProperties->BoneSocket, BoneOffset, BoneName))
+	{
+		return false;
+	}
+
+	int32 FirstKey = FMath::Clamp(
+		(int32) (NumSampledKeys * AnalysisProperties->StartTimeFraction), 0, NumSampledKeys);
+	int32 LastKey = FMath::Clamp(
+		(int32) (NumSampledKeys * AnalysisProperties->EndTimeFraction), FirstKey, NumSampledKeys);
+
+	FTransform FrameTM;
+	bool bNeedToUpdateFrameTM = true;
+	FVector FrameFacingDir, FrameUpDir, FrameRightDir;
+	CalculateFrameTM(bNeedToUpdateFrameTM, FrameTM, FirstKey, AnalysisProperties, Animation);
+	GetFrameDirs(FrameFacingDir, FrameUpDir, FrameRightDir, FrameTM, AnalysisProperties);
+
+	Result.Set(0, 0, 0);
+	for (int32 Key = FirstKey; Key != LastKey + 1; ++Key)
+	{
+		CalculateFrameTM(bNeedToUpdateFrameTM, FrameTM, Key, AnalysisProperties, Animation);
+		GetFrameDirs(FrameFacingDir, FrameUpDir, FrameRightDir, FrameTM, AnalysisProperties);
+
+		FVector RollPitchYaw;
+		CalculateBoneOrientation(
+			RollPitchYaw, Animation, Key, BoneName, BoneOffset, 
+			AnalysisProperties, FrameFacingDir, FrameRightDir, FrameUpDir);
+#ifdef ANALYSIS_VERBOSE_LOG
+		UE_LOG(LogAnimation, Log, TEXT("Roll/pitch/yaw = %f %f %f"), RollPitchYaw.X, RollPitchYaw.Y, RollPitchYaw.Z);
+#endif
+		Result += RollPitchYaw;
+	}
+	Result /= (1 + LastKey - FirstKey);
+	UE_LOG(LogAnimation, Log, TEXT("%s Orientation = %f %f %f"), *BoneName.ToString(), Result.X, Result.Y, Result.Z);
+	return true;
+}
+
+//======================================================================================================================
+// Note that if a looping animation has 56 keys, then its first key is 0 and last is 55, but these will be identical poses.
+// Thus it has one fewer intervals/unique keys
+template <typename T>
+static bool CalculateDeltaOrientation(
+	FVector&                        Result,
+	const UBlendSpace&              BlendSpace,
+	const T*                        AnalysisProperties,
+	const UAnimSequence&            Animation,
+	const float                     RateScale)
+{
+	const int32 NumSampledKeys = Animation.GetNumberOfSampledKeys() - 1;
+	if (!AnalysisProperties || NumSampledKeys <= 0)
+	{
+		return false;
+	}
+
+	FTransform BoneOffset;
+	FName BoneName;
+	if (!GetBoneInfo(Animation, AnalysisProperties->BoneSocket, BoneOffset, BoneName))
+	{
+		return false;
+	}
+
+	int32 FirstKey = FMath::Clamp(
+		(int32) (NumSampledKeys * AnalysisProperties->StartTimeFraction), 0, NumSampledKeys);
+	int32 LastKey = FMath::Clamp(
+		(int32) (NumSampledKeys * AnalysisProperties->EndTimeFraction), FirstKey, NumSampledKeys);
+
+	FTransform FrameTM;
+	bool bNeedToUpdateFrameTM = true;
+	FVector FrameFacingDir, FrameUpDir, FrameRightDir;
+	CalculateFrameTM(bNeedToUpdateFrameTM, FrameTM, FirstKey, AnalysisProperties, Animation);
+	GetFrameDirs(FrameFacingDir, FrameUpDir, FrameRightDir, FrameTM, AnalysisProperties);
+
+	FVector RollPitchYaw1;
+	CalculateBoneOrientation(
+		RollPitchYaw1, Animation, FirstKey, BoneName, BoneOffset,
+		AnalysisProperties, FrameFacingDir, FrameRightDir, FrameUpDir);
+
+	CalculateFrameTM(bNeedToUpdateFrameTM, FrameTM, LastKey, AnalysisProperties, Animation);
+	GetFrameDirs(FrameFacingDir, FrameUpDir, FrameRightDir, FrameTM, AnalysisProperties);
+
+	FVector RollPitchYaw2;
+	CalculateBoneOrientation(
+		RollPitchYaw2, Animation, LastKey, BoneName, BoneOffset,
+		AnalysisProperties, FrameFacingDir, FrameRightDir, FrameUpDir);
+
+	Result = RollPitchYaw2 - RollPitchYaw1;
+	return true;
+}
+
+//======================================================================================================================
+template <typename T>
+static bool CalculateAngularVelocity(
+	FVector&                         Result,
+	const UBlendSpace&               BlendSpace,
+	const T*                         AnalysisProperties,
+	const UAnimSequence&             Animation,
+	const float                      RateScale)
+{
+	const int32 NumSampledKeys = Animation.GetNumberOfSampledKeys() - 1;
+	if (!AnalysisProperties || NumSampledKeys <= 0)
+	{
+		return false;
+	}
+
+	if (NumSampledKeys == 1)
+	{
+		Result.Set(0, 0, 0);
+		return true;
+	}
+
+	FTransform BoneOffset;
+	FName BoneName;
+	if (!GetBoneInfo(Animation, AnalysisProperties->BoneSocket, BoneOffset, BoneName))
+	{
+		return false;
+	}
+
+	float DeltaTime = Animation.GetPlayLength() / NumSampledKeys;
+
+	int32 FirstKey = FMath::Clamp(
+		(int32) (NumSampledKeys * AnalysisProperties->StartTimeFraction), 0, NumSampledKeys);
+	int32 LastKey = FMath::Clamp(
+		(int32) (NumSampledKeys * AnalysisProperties->EndTimeFraction), FirstKey, NumSampledKeys-1);
+
+	if (FirstKey >= LastKey)
+	{
+		Result.Set(0, 0, 0);
+		return true;
+	}
+
+	FTransform FrameTM;
+	bool bNeedToUpdateFrameTM = true;
+	CalculateFrameTM(bNeedToUpdateFrameTM, FrameTM, FirstKey, AnalysisProperties, Animation);
+
+	Result.Set(0, 0, 0);
+	for (int32 Key = FirstKey; Key != LastKey + 1; ++Key)
+	{
+		CalculateFrameTM(bNeedToUpdateFrameTM, FrameTM, Key, AnalysisProperties, Animation);
+
+		FTransform BoneTM1 = GetBoneTransform(Animation, Key, BoneName);
+		FTransform TM1 = BoneOffset * BoneTM1;
+		FQuat RelativeQuat1 = FrameTM.InverseTransformRotation(TM1.GetRotation());
+
+		int32 NextKey = (Key + 1) % (NumSampledKeys + 1);
+		if (AnalysisProperties->Space == EAnalysisSpace::Moving)
+		{
+			CalculateFrameTM(bNeedToUpdateFrameTM, FrameTM, NextKey, AnalysisProperties, Animation);
+		}
+
+		FTransform BoneTM2 = GetBoneTransform(Animation, NextKey, BoneName);
+		FTransform TM2 = BoneOffset * BoneTM2;
+		FQuat RelativeQuat2 = FrameTM.InverseTransformRotation(TM2.GetRotation());
+
+		FQuat Rotation = RelativeQuat2 * RelativeQuat1.Inverse();
+		FVector Axis;
+		float Angle;
+		Rotation.ToAxisAndAngle(Axis, Angle);
+		FVector AngularVelocity = FMath::RadiansToDegrees(Axis * (Angle / DeltaTime));
+#ifdef ANALYSIS_VERBOSE_LOG
+		UE_LOG(LogAnimation, Log, TEXT("Angular Velocity = %f %f %f"), AngularVelocity.X, AngularVelocity.Y, AngularVelocity.Z);
+#endif
+		Result += AngularVelocity;
+	}
+	Result /= (1 + LastKey - FirstKey);
+	Result *= Animation.RateScale * RateScale;
+	UE_LOG(LogAnimation, Log, TEXT("%s angular velocity = %f %f %f"), *BoneName.ToString(), Result.X, Result.Y, Result.Z);
+	return true;
+}
+
+//======================================================================================================================
+template <typename T>
+static bool CalculateOrientationRate(
+	FVector&                        Result,
+	const UBlendSpace&              BlendSpace,
+	const T*                        AnalysisProperties,
+	const UAnimSequence&            Animation,
+	const float                     RateScale)
+{
+	const int32 NumSampledKeys = Animation.GetNumberOfSampledKeys() - 1;
+	if (!AnalysisProperties || NumSampledKeys <= 0)
+	{
+		return false;
+	}
+
+	if (NumSampledKeys == 1)
+	{
+		Result.Set(0, 0, 0);
+		return true;
+	}
+
+	FTransform BoneOffset;
+	FName BoneName;
+	if (!GetBoneInfo(Animation, AnalysisProperties->BoneSocket, BoneOffset, BoneName))
+	{
+		return false;
+	}
+
+	float DeltaTime = Animation.GetPlayLength() / NumSampledKeys;
+
+	int32 FirstKey = FMath::Clamp(
+		(int32) (NumSampledKeys * AnalysisProperties->StartTimeFraction), 0, NumSampledKeys);
+	int32 LastKey = FMath::Clamp(
+		(int32) (NumSampledKeys * AnalysisProperties->EndTimeFraction), FirstKey, NumSampledKeys-1);
+
+	if (FirstKey >= LastKey)
+	{
+		Result.Set(0, 0, 0);
+		return true;
+	}
+
+	FTransform FrameTM;
+	bool bNeedToUpdateFrameTM = true;
+	FVector FrameFacingDir, FrameUpDir, FrameRightDir;
+	CalculateFrameTM(bNeedToUpdateFrameTM, FrameTM, FirstKey, AnalysisProperties, Animation);
+	GetFrameDirs(FrameFacingDir, FrameUpDir, FrameRightDir, FrameTM, AnalysisProperties);
+
+	Result.Set(0, 0, 0);
+	for (int32 Key = FirstKey; Key != LastKey + 1; ++Key)
+	{
+		CalculateFrameTM(bNeedToUpdateFrameTM, FrameTM, Key, AnalysisProperties, Animation);
+		GetFrameDirs(FrameFacingDir, FrameUpDir, FrameRightDir, FrameTM, AnalysisProperties);
+
+		FVector RollPitchYaw1;
+		CalculateBoneOrientation(
+			RollPitchYaw1, Animation, Key, BoneName, BoneOffset, 
+			AnalysisProperties, FrameFacingDir, FrameRightDir, FrameUpDir);
+
+		int32 NextKey = (Key + 1) % (NumSampledKeys + 1);
+		if (AnalysisProperties->Space == EAnalysisSpace::Moving)
+		{
+			CalculateFrameTM(bNeedToUpdateFrameTM, FrameTM, NextKey, AnalysisProperties, Animation);
+			GetFrameDirs(FrameFacingDir, FrameUpDir, FrameRightDir, FrameTM, AnalysisProperties);
+		}
+
+		FVector RollPitchYaw2;
+		CalculateBoneOrientation(
+			RollPitchYaw2, Animation, NextKey, BoneName, BoneOffset, 
+			AnalysisProperties, FrameFacingDir, FrameRightDir, FrameUpDir);
+
+		const FVector OrientationRate = (RollPitchYaw2 - RollPitchYaw1) / DeltaTime;
+#ifdef ANALYSIS_VERBOSE_LOG
+		UE_LOG(LogAnimation, Log, TEXT("Orientation rate = %f %f %f"), OrientationRate.X, OrientationRate.Y, OrientationRate.Z);
+#endif
+		Result += OrientationRate;
+	}
+	Result /= (1 + LastKey - FirstKey);
+	Result *= Animation.RateScale * RateScale;
+	UE_LOG(LogAnimation, Log, TEXT("%s Orientation rate = %f %f %f"), *BoneName.ToString(), Result.X, Result.Y, Result.Z);
+	return true;
+}
+
+
+}
+
+#undef LOCTEXT_NAMESPACE
 
diff --git a/Engine/Source/Editor/Persona/Public/PersonaDelegates.h b/Engine/Source/Editor/Persona/Public/PersonaDelegates.h
--- a/Engine/Source/Editor/Persona/Public/PersonaDelegates.h
+++ b/Engine/Source/Editor/Persona/Public/PersonaDelegates.h
@@ -44,10 +44,10 @@ DECLARE_DELEGATE_OneParam(FOnBlendSpaceSampleDoubleClicked, int32 /*SampleIndex*
 DECLARE_DELEGATE_OneParam(FOnBlendSpaceSampleRemoved, const int32 /*SampleIndex*/);
 
 // Called when a blendspace sample point is added. Returns the new SampleIndex
-DECLARE_DELEGATE_RetVal_TwoParams(int32, FOnBlendSpaceSampleAdded, UAnimSequence* /*InSequence*/, const FVector& /*InSamplePoint*/);
+DECLARE_DELEGATE_RetVal_ThreeParams(int32, FOnBlendSpaceSampleAdded, UAnimSequence* /*InSequence*/, const FVector& /*InSamplePoint*/, bool /*bRunAnalysis*/);
 
 // Called when a blendspace sample point is added
-DECLARE_DELEGATE_TwoParams(FOnBlendSpaceSampleDuplicated, int32 /*SampleIndex*/, const FVector& /*InSamplePoint*/);
+DECLARE_DELEGATE_ThreeParams(FOnBlendSpaceSampleDuplicated, int32 /*SampleIndex*/, const FVector& /*InSamplePoint*/, bool /*bRunAnalysis*/);
 
 // Called when a blendspace sample point is replaced
 DECLARE_DELEGATE_TwoParams(FOnBlendSpaceSampleReplaced, const int32 /*SampleIndex*/, UAnimSequence* /*InSequence*/);
diff --git a/Engine/Source/Runtime/Engine/Classes/Animation/BlendSpace.h b/Engine/Source/Runtime/Engine/Classes/Animation/BlendSpace.h
--- a/Engine/Source/Runtime/Engine/Classes/Animation/BlendSpace.h
+++ b/Engine/Source/Runtime/Engine/Classes/Animation/BlendSpace.h
@@ -17,21 +17,22 @@
 #include "Animation/BoneSocketReference.h"
 #include "BlendSpace.generated.h"
 
-class FCachedAnalysisProperties;
+class UCachedAnalysisProperties;
+class UBlendSpace;
 
 /**
 * The base class for properties to be used in analysis. Engine will inherit from this to define structures used for
 * the functions it supports. User-defined functions will likely need their own analysis structures inheriting from
 * this too.
 */
-UCLASS(config=Engine, MinimalAPI)
+UCLASS(MinimalAPI, config=Engine)
 class UAnalysisProperties : public UObject
 {
 	GENERATED_BODY()
 
 public:
-	ENGINE_API virtual void InitializeFromCache(TSharedPtr<FCachedAnalysisProperties> Cache);
-	ENGINE_API virtual void MakeCache(TSharedPtr<FCachedAnalysisProperties>& Cache) const;
+	virtual void InitializeFromCache(TObjectPtr<UCachedAnalysisProperties> Cache) {};
+	virtual void MakeCache(TObjectPtr<UCachedAnalysisProperties>& Cache, UBlendSpace* BlendSpace) {};
 
 	/** Analysis function for this axis */
 	UPROPERTY()
@@ -704,7 +705,7 @@ public:
 	TObjectPtr<UAnalysisProperties> AnalysisProperties[3];
 
 	/** Cached properties used to initialize properties when newly created. */
-	TSharedPtr<FCachedAnalysisProperties> CachedAnalysisProperties[3];
+	TObjectPtr<UCachedAnalysisProperties> CachedAnalysisProperties[3];
 #endif
 
 	/**
diff --git a/Engine/Source/Runtime/Engine/Private/Animation/BlendSpace.cpp b/Engine/Source/Runtime/Engine/Private/Animation/BlendSpace.cpp
--- a/Engine/Source/Runtime/Engine/Private/Animation/BlendSpace.cpp
+++ b/Engine/Source/Runtime/Engine/Private/Animation/BlendSpace.cpp
@@ -36,29 +36,6 @@ struct FBlendSpaceScratchData : public TThreadSingleton<FBlendSpaceScratchData>
 	TArray<FGridBlendSample, TInlineAllocator<4> > RawGridSamples;
 };
 
-//======================================================================================================================
-void UAnalysisProperties::InitializeFromCache(TSharedPtr<FCachedAnalysisProperties> Cache)
-{
-#if WITH_EDITOR
-	if (Cache)
-	{
-		bLockAfterAnalysis = Cache->bLockAfterAnalysis;
-	}
-#endif
-}
-
-//======================================================================================================================
-void UAnalysisProperties::MakeCache(TSharedPtr<FCachedAnalysisProperties>& Cache) const
-{
-#if WITH_EDITOR
-	if (!Cache)
-	{
-		Cache = MakeShared<FCachedAnalysisProperties>();
-	}
-	Cache->bLockAfterAnalysis = bLockAfterAnalysis;
-#endif
-}
-
 //======================================================================================================================
 UBlendSpace::UBlendSpace(const FObjectInitializer& ObjectInitializer)
 	: Super(ObjectInitializer)
