Control Rig: Sanity check on hit proxy validity

diff --git a/Engine/Source/Editor/Persona/Private/PersonaSelectionComponent.cpp b/Engine/Source/Editor/Persona/Private/PersonaSelectionComponent.cpp
--- a/Engine/Source/Editor/Persona/Private/PersonaSelectionComponent.cpp
+++ b/Engine/Source/Editor/Persona/Private/PersonaSelectionComponent.cpp
@@ -32,12 +32,14 @@ void UPersonaSelectionComponent::Reset()
 {
 	Capsules.Reset();
 	CapsuleIndices.Reset();
+	MarkCapsulesChanged();
 }
 
 void UPersonaSelectionComponent::SetNum(int32 InCount)
 {
 	Capsules.SetNum(InCount);
 	SetCapsulesIndicesToFullArray();
+	MarkCapsulesChanged();
 }
 
 int32 UPersonaSelectionComponent::Add(int32 InCount)
@@ -208,6 +210,12 @@ void FPersonaSelectionComponentProxy::GetDynamicMeshElements(const TArray<const
 				const FPersonaSelectionCapsule& Capsule = SelectionComponent->Capsules[CapsuleIndex];
 				const FLinearColor LineColor = FLinearColor::Yellow;
 				const FTransform& Transform = Capsule.Transform * ComponentToWorld;
+
+				if(!SelectionComponent->HitProxies.IsValidIndex(CapsuleIndex))
+				{
+					break;
+				}
+				
 				HPersonaSelectionHitProxy* HitProxy = SelectionComponent->HitProxies[CapsuleIndex].GetReference();
 
 				GetBoxMesh(
