Persona: Fix non-unity build break in PersonaSelectionComponent

diff --git a/Engine/Source/Editor/Persona/Private/PersonaSelectionComponent.cpp b/Engine/Source/Editor/Persona/Private/PersonaSelectionComponent.cpp
--- a/Engine/Source/Editor/Persona/Private/PersonaSelectionComponent.cpp
+++ b/Engine/Source/Editor/Persona/Private/PersonaSelectionComponent.cpp
@@ -1,6 +1,7 @@
 // Copyright Epic Games, Inc. All Rights Reserved.
 
 #include "PersonaSelectionComponent.h"
+#include "Materials/Material.h"
 
 IMPLEMENT_HIT_PROXY(HPersonaSelectionHitProxy, HHitProxy)
 
