Fix for: Right-clicking in AnimBP Viewport Timeline displays Sequence controls

diff --git a/Engine/Source/Editor/Persona/Private/SAnimationEditorViewport.cpp b/Engine/Source/Editor/Persona/Private/SAnimationEditorViewport.cpp
--- a/Engine/Source/Editor/Persona/Private/SAnimationEditorViewport.cpp
+++ b/Engine/Source/Editor/Persona/Private/SAnimationEditorViewport.cpp
@@ -1643,6 +1643,7 @@ void SAnimationEditorViewportTabBody::UpdateScrubPanel(UAnimationAsset* AnimAsse
 					.ViewInputMin(this, &SAnimationEditorViewportTabBody::GetViewMinInput)
 					.ViewInputMax(this, &SAnimationEditorViewportTabBody::GetViewMaxInput)
 					.bAllowZoom(true)
+					.bDisplayAnimScrubBarEditing(false)
 				];
 		}
 	}
