Fix for lock status after reanalyzing samples

diff --git a/Engine/Source/Editor/Persona/Private/Customization/BlendSpaceDetails.cpp b/Engine/Source/Editor/Persona/Private/Customization/BlendSpaceDetails.cpp
--- a/Engine/Source/Editor/Persona/Private/Customization/BlendSpaceDetails.cpp
+++ b/Engine/Source/Editor/Persona/Private/Customization/BlendSpaceDetails.cpp
@@ -68,7 +68,7 @@ FReply FBlendSpaceDetails::HandleAnalyzeSamples()
 				*BlendSpace, *BlendSpace->SampleData[SampleIndex].Animation, 
 				BlendSpace->SampleData[SampleIndex].RateScale, 
 				BlendSpace->SampleData[SampleIndex].SampleValue, bAnalyzed);
-			if (NewValue != BlendSpace->SampleData[SampleIndex].SampleValue)
+			if (bAnalyzed[0] || bAnalyzed[1] || bAnalyzed[2])
 			{
 				BlendSpace->EditSampleValue(SampleIndex, NewValue);
 				// Note that the sample might not move if the destination position is in use
