Persona Selection Component - Guard against race conditions

diff --git a/Engine/Source/Editor/Persona/Private/PersonaSelectionComponent.cpp b/Engine/Source/Editor/Persona/Private/PersonaSelectionComponent.cpp
--- a/Engine/Source/Editor/Persona/Private/PersonaSelectionComponent.cpp
+++ b/Engine/Source/Editor/Persona/Private/PersonaSelectionComponent.cpp
@@ -31,20 +31,28 @@ UPersonaSelectionComponent::UPersonaSelectionComponent(const FObjectInitializer&
 
 void UPersonaSelectionComponent::Reset()
 {
-	Capsules.Reset();
-	CapsuleIndices.Reset();
+	{
+		FScopeLock ScopeLock( &CriticalSection );
+		Capsules.Reset();
+		CapsuleIndices.Reset();
+	}
 	MarkCapsulesChanged();
 }
 
 void UPersonaSelectionComponent::SetNum(int32 InCount)
 {
-	Capsules.SetNum(InCount);
+	{
+		FScopeLock ScopeLock( &CriticalSection );
+		Capsules.SetNum(InCount);
+	}
 	SetCapsulesIndicesToFullArray();
 	MarkCapsulesChanged();
 }
 
 int32 UPersonaSelectionComponent::Add(int32 InCount)
 {
+	FScopeLock ScopeLock( &CriticalSection );
+
 	int32 FirstIndex = INDEX_NONE;
 	for(int32 Index=0;Index<InCount;Index++)
 	{
@@ -60,6 +68,8 @@ int32 UPersonaSelectionComponent::Add(int32 InCount)
 
 void UPersonaSelectionComponent::SetCapsulesIndicesToFullArray()
 {
+	FScopeLock ScopeLock( &CriticalSection );
+
 	CapsuleIndices.SetNum(Capsules.Num());
 	for(int32 Index=0;Index<CapsuleIndices.Num();Index++)
 	{
@@ -69,6 +79,8 @@ void UPersonaSelectionComponent::SetCapsulesIndicesToFullArray()
 
 void UPersonaSelectionComponent::MarkCapsulesChanged()
 {
+	FScopeLock ScopeLock( &CriticalSection );
+
 	while(HitProxies.Num() > Capsules.Num())
 	{
 		HitProxies.Pop();
@@ -191,6 +203,8 @@ void FPersonaSelectionComponentProxy::GetDynamicMeshElements(const TArray<const
 	}
 
 	const UPersonaSelectionComponent* SelectionComponent = SelectionComponentPtr.Get();
+	FScopeLock ScopeLock( &SelectionComponent->CriticalSection );
+
 	if(SelectionComponent->Capsules.IsEmpty())
 	{
 		return;
diff --git a/Engine/Source/Editor/Persona/Public/PersonaSelectionComponent.h b/Engine/Source/Editor/Persona/Public/PersonaSelectionComponent.h
--- a/Engine/Source/Editor/Persona/Public/PersonaSelectionComponent.h
+++ b/Engine/Source/Editor/Persona/Public/PersonaSelectionComponent.h
@@ -145,6 +145,9 @@ private:
 	/** Fired when the component needs to update a set of capsules */
 	FOnPersonaSelectionUpdateCapsules OnPersonaSelectionUpdateCapsules;
 
+	/** Critical section used to block between renders / updates */
+	mutable FCriticalSection CriticalSection;
+
 	friend class FPersonaSelectionComponentProxy;
 };
 
