Allow selection of an animation with the Enter key

diff --git a/Engine/Source/Editor/Persona/Private/SAnimationSequenceBrowser.cpp b/Engine/Source/Editor/Persona/Private/SAnimationSequenceBrowser.cpp
--- a/Engine/Source/Editor/Persona/Private/SAnimationSequenceBrowser.cpp
+++ b/Engine/Source/Editor/Persona/Private/SAnimationSequenceBrowser.cpp
@@ -471,6 +471,14 @@ SAnimationSequenceBrowser::~SAnimationSequenceBrowser()
 	}
 }
 
+void SAnimationSequenceBrowser::OnRequestOpenAssets(const TArray<FAssetData>& SelectedAssets, bool bFromHistory)
+{
+	if (SelectedAssets.Num() == 1)
+	{
+		OnRequestOpenAsset(SelectedAssets[0], bFromHistory);
+	}
+}
+
 void SAnimationSequenceBrowser::OnRequestOpenAsset(const FAssetData& AssetData, bool bFromHistory)
 {
 	if (UObject* RawAsset = AssetData.GetAsset())
@@ -896,6 +904,7 @@ void SAnimationSequenceBrowser::Construct(const FArguments& InArgs, const TShare
 	Config.OnAssetDoubleClicked = FOnAssetDoubleClicked::CreateSP(this, &SAnimationSequenceBrowser::OnRequestOpenAsset, false);
 	Config.OnGetAssetContextMenu = FOnGetAssetContextMenu::CreateSP(this, &SAnimationSequenceBrowser::OnGetAssetContextMenu);
 	Config.OnAssetTagWantsToBeDisplayed = FOnShouldDisplayAssetTag::CreateSP(this, &SAnimationSequenceBrowser::CanShowColumnForAssetRegistryTag);
+	Config.OnAssetEnterPressed = FOnAssetEnterPressed::CreateSP(this, &SAnimationSequenceBrowser::OnRequestOpenAssets, false);
 	Config.SyncToAssetsDelegates.Add(&SyncToAssetsDelegate);
 	Config.OnShouldFilterAsset = FOnShouldFilterAsset::CreateSP(this, &SAnimationSequenceBrowser::HandleFilterAsset);
 	Config.GetCurrentSelectionDelegates.Add(&GetCurrentSelectionDelegate);
diff --git a/Engine/Source/Editor/Persona/Private/SAnimationSequenceBrowser.h b/Engine/Source/Editor/Persona/Private/SAnimationSequenceBrowser.h
--- a/Engine/Source/Editor/Persona/Private/SAnimationSequenceBrowser.h
+++ b/Engine/Source/Editor/Persona/Private/SAnimationSequenceBrowser.h
@@ -60,6 +60,8 @@ public:
 public:
 	void Construct(const FArguments& InArgs, const TSharedRef<class IPersonaToolkit>& InPersonaToolkit);
 
+	void OnRequestOpenAssets(const TArray<FAssetData>& SelectedAssets, bool bFromHistory);
+
 	void OnRequestOpenAsset(const FAssetData& AssetData, bool bFromHistory);
 
 	virtual FReply OnKeyDown(const FGeometry& MyGeometry, const FKeyEvent& InKeyEvent) override;
