Improved the performence of the Floor Height Offset slider from the Animation Viewport.

Now it only create one transaction when the user is adjusting the slider.

diff --git a/Engine/Source/Editor/Persona/Private/AnimationEditorViewportClient.cpp b/Engine/Source/Editor/Persona/Private/AnimationEditorViewportClient.cpp
--- a/Engine/Source/Editor/Persona/Private/AnimationEditorViewportClient.cpp
+++ b/Engine/Source/Editor/Persona/Private/AnimationEditorViewportClient.cpp
@@ -1736,15 +1736,22 @@ float FAnimationViewportClient::GetFloorOffset() const
 	return 0.0f;
 }
 
-void FAnimationViewportClient::SetFloorOffset( float NewValue )
+void FAnimationViewportClient::SetFloorOffset( float NewValue, bool bCommitted )
 {
 	USkeletalMesh* Mesh = GetPreviewScene()->GetPreviewMeshComponent()->SkeletalMesh;
 
 	if ( Mesh )
 	{
-		// This value is saved in a UPROPERTY for the mesh, so changes are transactional
-		FScopedTransaction Transaction( LOCTEXT( "SetFloorOffset", "Set Floor Offset" ) );
-		Mesh->Modify();
+		if (bCommitted)
+		{
+			PendingTransaction.Reset();
+		}
+		else if (!PendingTransaction.IsValid())
+		{
+			// This value is saved in a UPROPERTY for the mesh, so changes are transactional
+			PendingTransaction = MakeUnique<FScopedTransaction>( LOCTEXT( "SetFloorOffset", "Set Floor Offset" ) );
+			Mesh->Modify();
+		}
 
 		Mesh->SetFloorOffset(NewValue);
 		UpdateCameraSetup(); // This does the actual moving of the floor mesh
diff --git a/Engine/Source/Editor/Persona/Private/AnimationEditorViewportClient.h b/Engine/Source/Editor/Persona/Private/AnimationEditorViewportClient.h
--- a/Engine/Source/Editor/Persona/Private/AnimationEditorViewportClient.h
+++ b/Engine/Source/Editor/Persona/Private/AnimationEditorViewportClient.h
@@ -254,7 +254,7 @@ public:
 	float GetFloorOffset() const;
 
 	/* Sets the floor height offset, saves it to config and invalidates the viewport so it shows up immediately */
-	void SetFloorOffset(float NewValue);
+	void SetFloorOffset(float NewValue, bool bCommitted);
 
 	/** Function to set mesh stat drawing state */
 	void OnSetShowMeshStats(int32 ShowMode);
@@ -439,6 +439,8 @@ private:
 	/** Relative view location stored to match it pre/post tick */
 	FVector RelativeViewLocation;
 
+	TUniquePtr<FScopedTransaction> PendingTransaction;
+
 	// Delegate Handler to allow changing of camera controller
 	void OnCameraControllerChanged();
 };
diff --git a/Engine/Source/Editor/Persona/Private/SAnimViewportToolBar.cpp b/Engine/Source/Editor/Persona/Private/SAnimViewportToolBar.cpp
--- a/Engine/Source/Editor/Persona/Private/SAnimViewportToolBar.cpp
+++ b/Engine/Source/Editor/Persona/Private/SAnimViewportToolBar.cpp
@@ -363,6 +363,7 @@ TSharedRef<SWidget> SAnimViewportToolBar::MakeFloorOffsetWidget() const
 				.MaxSliderValue(100.0f)
 				.Value(this, &SAnimViewportToolBar::OnGetFloorOffset)
 				.OnValueChanged(const_cast<SAnimViewportToolBar*>(this), &SAnimViewportToolBar::OnFloorOffsetChanged)
+				.OnValueCommitted(const_cast<SAnimViewportToolBar*>(this), &SAnimViewportToolBar::OnFloorOffsetCommitted)
 				.ToolTipText(LOCTEXT("FloorOffsetToolTip", "Height offset for the floor mesh (stored per-mesh)"))
 			]
 		];
@@ -1209,7 +1210,16 @@ void SAnimViewportToolBar::OnFloorOffsetChanged( float NewValue )
 {
 	FAnimationViewportClient& AnimViewportClient = (FAnimationViewportClient&)Viewport.Pin()->GetLevelViewportClient();
 
-	AnimViewportClient.SetFloorOffset( NewValue );
+	AnimViewportClient.SetFloorOffset( NewValue, false );
+
+	PinnedCommands->AddCustomWidget(TEXT("FloorOffsetWidget"));
+}
+
+void SAnimViewportToolBar::OnFloorOffsetCommitted( float NewValue, ETextCommit::Type CommitType )
+{
+	FAnimationViewportClient& AnimViewportClient = (FAnimationViewportClient&)Viewport.Pin()->GetLevelViewportClient();
+
+	AnimViewportClient.SetFloorOffset( NewValue, true );
 
 	PinnedCommands->AddCustomWidget(TEXT("FloorOffsetWidget"));
 }
diff --git a/Engine/Source/Editor/Persona/Private/SAnimViewportToolBar.h b/Engine/Source/Editor/Persona/Private/SAnimViewportToolBar.h
--- a/Engine/Source/Editor/Persona/Private/SAnimViewportToolBar.h
+++ b/Engine/Source/Editor/Persona/Private/SAnimViewportToolBar.h
@@ -3,13 +3,14 @@
 #pragma once
 
 #include "CoreMinimal.h"
+#include "Editor/UnrealEd/Public/SViewportToolBar.h"
 #include "Layout/Visibility.h"
+#include "SAnimationEditorViewport.h"
+#include "SEditorViewport.h"
 #include "Styling/SlateColor.h"
+#include "Types/SlateEnums.h"
 #include "Widgets/DeclarativeSyntaxSupport.h"
 #include "Widgets/SWidget.h"
-#include "SEditorViewport.h"
-#include "SAnimationEditorViewport.h"
-#include "Editor/UnrealEd/Public/SViewportToolBar.h"
 
 class FMenuBuilder;
 class SComboButton;
@@ -161,6 +162,7 @@ private:
 	TOptional<float> OnGetFloorOffset() const;
 	/** Called when the floor offset slider is adjusted in the perspective viewport */
 	void OnFloorOffsetChanged( float NewValue );
+	void OnFloorOffsetCommitted ( float NewValue, ETextCommit::Type CommitType );
 
 	/** Called when users select the mirror table in the menu */
 	void OnMirrorDataTableSelected(const FAssetData& SelectedMirrorTableData);
