Improve blendspace drawing with the grid. Makes it default to off when using the grid, too. Show an underlay on "bad" samples when using the grid.

diff --git a/Engine/Source/Editor/Persona/Private/SAnimationBlendSpaceGridWidget.cpp b/Engine/Source/Editor/Persona/Private/SAnimationBlendSpaceGridWidget.cpp
--- a/Engine/Source/Editor/Persona/Private/SAnimationBlendSpaceGridWidget.cpp
+++ b/Engine/Source/Editor/Persona/Private/SAnimationBlendSpaceGridWidget.cpp
@@ -40,6 +40,150 @@
 // Draws additional data on the triangulation to help debugging
 //#define DEBUG_BLENDSPACE_TRIANGULATION
 
+// Threshold for it being considered a problem that when a lookup is made at the same location as a sample, the returned
+// weight is less than this.
+static const float SampleLookupWeightThreshold = 0.2f;
+
+// Flag any triangle that has an interior angle smaller than this (degrees)
+static const float CriticalTriangulationAngle = 4.0f;
+
+// Flag any triangle that has a smaller area than this (normalized units)
+static const float CriticalTriangulationArea = 5e-4f;
+
+//======================================================================================================================
+// Paint a filled triangle
+static void PaintTriangle(
+	const FVector2D&         P0,
+	const FVector2D&         P1,
+	const FVector2D&         P2,
+	const FGeometry&         AllottedGeometry,
+	FLinearColor             Color,
+	const FSlateBrush*       Brush,
+	FSlateWindowElementList& OutDrawElements,
+	int32                    DrawLayerId)
+{
+	const FVector2D* Points[3] = { &P0, &P1, &P2 };
+
+	TArray<FSlateVertex> Vertices;
+	Vertices.Reserve(3);
+
+	for (int32 PointIndex = 0; PointIndex != 3; ++PointIndex)
+	{
+		Vertices.AddZeroed();
+		FSlateVertex& NewVert = Vertices.Last();
+		NewVert.Position = AllottedGeometry.LocalToAbsolute(*Points[PointIndex]);
+		NewVert.Color = Color.ToFColor(false);
+	}
+
+	// Fill by making triangles
+	TArray<SlateIndex> VertexIndices = { 0, 1, 2 };
+	FSlateDrawElement::MakeCustomVerts(
+		OutDrawElements, DrawLayerId, Brush->GetRenderingResource(), Vertices, VertexIndices, nullptr, 0, 0);
+}
+
+//======================================================================================================================
+// Paints a filled polygon with outline, defined by a set of points which don't need to be sorted.
+// This will handle concave polygons, but only if the centroid lies inside the polygon.
+static void PaintPolygon(
+	TArray<FVector2D>&       Points,
+	const FGeometry&         AllottedGeometry,
+	FLinearColor             FillColor,
+	FLinearColor             OutlineColor,
+	const FSlateBrush*       Brush,
+	FSlateWindowElementList& OutDrawElements,
+	int32                    DrawLayerId)
+{
+	TArray<FSlateVertex> Vertices;
+	Vertices.Reserve(Points.Num() + 1);
+
+	// Add a mid-position vertex so that we handle polygons that aren't completely convex
+	Vertices.AddZeroed();
+
+	FSlateVertex& MidVertex = Vertices.Last();
+	for (int32 PointIndex = 0; PointIndex != Points.Num(); ++PointIndex)
+	{
+		Vertices.AddZeroed();
+		FSlateVertex& NewVert = Vertices.Last();
+		NewVert.Position = AllottedGeometry.LocalToAbsolute(Points[PointIndex]);
+		NewVert.Color = FillColor.ToFColor(false);
+		MidVertex.Position += NewVert.Position;
+	}
+	MidVertex.Position /= Points.Num();
+	MidVertex.Color = FillColor.ToFColor(false);
+
+	// Make sure the points all wind correctly relative to the mid point
+	struct FComparePoints
+	{
+		FComparePoints(const FSlateVertex& Mid) : MidPoint(Mid) {}
+		bool operator()(const FSlateVertex& A, const FSlateVertex& B) const
+		{
+			FVector2D DeltaA = A.Position - MidPoint.Position;
+			FVector2D DeltaB = B.Position - MidPoint.Position;
+			float AngleA = FMath::Atan2(DeltaA.Y, DeltaA.X);
+			float AngleB = FMath::Atan2(DeltaB.Y, DeltaB.X);
+			return AngleA < AngleB;
+		}
+		FSlateVertex MidPoint;
+	};
+	Sort(Vertices.GetData() + 1, Vertices.Num() - 1, FComparePoints(MidVertex));
+
+	if (FillColor.A > 0)
+	{
+		// Fill by making triangles
+		TArray<SlateIndex> VertexIndices;
+		for (int VertexIndex = 1; VertexIndex < Vertices.Num(); ++VertexIndex)
+		{
+			VertexIndices.Add(0);
+			VertexIndices.Add(VertexIndex);
+			VertexIndices.Add(VertexIndex + 1 >= Vertices.Num() ? 1 : VertexIndex + 1);
+		}
+
+		FSlateDrawElement::MakeCustomVerts(
+			OutDrawElements, DrawLayerId, Brush->GetRenderingResource(), Vertices, VertexIndices, nullptr, 0, 0);
+	}
+
+	if (OutlineColor.A > 0)
+	{
+		TArray<FVector2D> LinePoints;
+		LinePoints.Reserve(Points.Num() + 1);
+		for (int VertexIndex = 1; VertexIndex <= Vertices.Num(); ++VertexIndex)
+		{
+			LinePoints.Add(VertexIndex < Vertices.Num() ? Vertices[VertexIndex].Position : Vertices[1].Position);
+		}
+		FSlateDrawElement::MakeLines(
+			OutDrawElements, DrawLayerId + 1, FPaintGeometry(), LinePoints, ESlateDrawEffect::None, OutlineColor, true, 1.0f);
+	}
+}
+
+//======================================================================================================================
+static void PaintCircle(
+	const FVector2D&         Centre,
+	const float              Radius,
+	const int32              NumVerts,
+	const FGeometry&         AllottedGeometry,
+	FLinearColor             FillColor,
+	FLinearColor             OutlineColor,
+	const FSlateBrush*       Brush,
+	FSlateWindowElementList& OutDrawElements,
+	int32                    DrawLayerId)
+{
+	// NumVerts needs to be a multiple of 4
+	int32 NumVertsPerSector = ((NumVerts + 3) / 4);
+	TArray<FVector2D> Points;
+	Points.SetNum(NumVertsPerSector * 4);
+	for (int32 Index = 0 ; Index != NumVertsPerSector ; ++Index)
+	{
+		float Angle = Index * (HALF_PI / NumVertsPerSector);
+		float S, C;
+		FMath::SinCos(&S, &C, Angle);
+		Points[Index + NumVertsPerSector * 0] = Centre + FVector2D(C * Radius, S * Radius);
+		Points[Index + NumVertsPerSector * 1] = Centre + FVector2D(-S * Radius, C * Radius);
+		Points[Index + NumVertsPerSector * 2] = Centre + FVector2D(-C * Radius, -S * Radius);
+		Points[Index + NumVertsPerSector * 3] = Centre + FVector2D(S * Radius, -C * Radius);
+	}
+	PaintPolygon(Points, AllottedGeometry, FillColor, OutlineColor, Brush, OutDrawElements, DrawLayerId);
+}
+
 void SBlendSpaceGridWidget::Construct(const FArguments& InArgs)
 {
 	BlendSpaceBase = InArgs._BlendSpaceBase;
@@ -74,7 +218,7 @@
 
 	PreviewFilteredPosition = PreviewPosition;
 
-	bShowTriangulation = true;
+	bShowTriangulation = (BlendSpaceBase.Get() != nullptr && BlendSpaceBase.Get()->bInterpolateUsingGrid) ? false : true;
 	bMouseIsOverGeometry = false;
 	bRefreshCachedData = true;
 	bStretchToFit = true;
@@ -155,7 +299,12 @@
 								.VAlign(VAlign_Center)
 								[
 									SNew(SButton)
-									.ToolTipText(LOCTEXT("ShowTriangulation", "Show Triangulation"))
+									.ToolTipText_Lambda([this]()
+														{ 
+															return (BlendSpaceBase.Get() && BlendSpaceBase.Get()->bInterpolateUsingGrid) 
+																? LOCTEXT("ShowGridToSampleConnections", "Show Grid/Sample Connections")
+																: LOCTEXT("ShowTriangulation", "Show Triangulation");
+														})
 									.OnClicked(this, &SBlendSpaceGridWidget::ToggleTriangulationVisibility)
 									.ButtonColorAndOpacity_Lambda([this]() -> FLinearColor { return bShowTriangulation ? FEditorStyle::GetSlateColor("SelectionColor").GetSpecifiedColor() : FLinearColor::White; })
 									.ContentPadding(1)
@@ -284,3 +433,12 @@
 		+ SVerticalBox::Slot()
 		.AutoHeight()
 		[
+			SNew(STextBlock)
+			.Text(this, &SBlendSpaceGridWidget::GetToolTipSampleValidity)
+			.Font(FCoreStyle::Get().GetFontStyle("ToolTip.LargerFont"))
+			.Visibility_Lambda([this]() { return GetToolTipSampleValidity().IsEmpty() ? EVisibility::Collapsed : EVisibility::Visible; })
+		]
+
+		+ SVerticalBox::Slot()
+		.AutoHeight()
+		[
@@ -328,10 +486,7 @@
 	
 	PaintBackgroundAndGrid(AllottedGeometry, MyCullingRect, OutDrawElements, LayerId);
 
-	if (bShowTriangulation)
-	{
-		PaintTriangulation(AllottedGeometry, MyCullingRect, OutDrawElements, LayerId);
-	}	
+	PaintTriangulation(AllottedGeometry, MyCullingRect, OutDrawElements, LayerId);
 	PaintSampleKeys(AllottedGeometry, MyCullingRect, OutDrawElements, LayerId);
 
 	if(bShowAxisLabels)
@@ -403,9 +558,38 @@
 	DrawLayerId += 3;
 }
 
+//======================================================================================================================
+float SBlendSpaceGridWidget::GetSampleLookupWeight(int32 SampleIndex) const
+{
+	const UBlendSpace* BlendSpace = BlendSpaceBase.Get();
+	if (BlendSpace && BlendSpace->bInterpolateUsingGrid)
+	{
+		const TArray<FBlendSample>& Samples = BlendSpace->GetBlendSamples();
+		const FBlendSample& Sample = Samples[SampleIndex];
+			
+		TArray<FBlendSampleData> SampleDataList;
+		int32 TempTriangulationIndex = -1;
+		BlendSpace->GetSamplesFromBlendInput(Sample.SampleValue, SampleDataList, TempTriangulationIndex, true);
+		float LookedUpSampleWeight = 0.0f;
+		for (const FBlendSampleData& LookedUpSample : SampleDataList)
+		{
+			if (LookedUpSample.SampleDataIndex == SampleIndex)
+			{
+				LookedUpSampleWeight += LookedUpSample.GetWeight();
+			}
+		}
+		return FMath::Clamp(LookedUpSampleWeight, 0.0f, 1.0f);
+	}
+	return 1.0f; // Return 1 to avoid anything treating this as a problem
+}
+
+
+//======================================================================================================================
 void SBlendSpaceGridWidget::PaintSampleKeys(
-	const FGeometry& AllottedGeometry, const FSlateRect& MyCullingRect, 
-	FSlateWindowElementList& OutDrawElements, int32& DrawLayerId) const
+	const FGeometry&         AllottedGeometry, 
+	const FSlateRect&        MyCullingRect, 
+	FSlateWindowElementList& OutDrawElements, 
+	int32&                   DrawLayerId) const
 {
 	const int32 FilteredPositionLayer = DrawLayerId + 1;
 	const int32 PreviewPositionLayer = DrawLayerId + 2;
@@ -442,6 +626,19 @@
 				OutDrawElements, SampleLayer, AllottedGeometry.ToPaintGeometry(GridPosition, KeySize), 
 				KeyBrush, ESlateDrawEffect::None, DrawColor );
 
+			float SampleLookupWeight = GetSampleLookupWeight(SampleIndex);
+			if (SampleLookupWeight <= SampleLookupWeightThreshold)
+			{
+				const FVector2D CirclePosition = SampleValueToScreenPosition(Sample.SampleValue);
+				FLinearColor IsolatedColor = FLinearColor::Red;
+				IsolatedColor.A = SampleLookupWeightThreshold > 0.0f ? 
+					(SampleLookupWeightThreshold - SampleLookupWeight) / SampleLookupWeightThreshold :
+					1.0f;
+				IsolatedColor.A *= 0.4f;
+				PaintCircle(CirclePosition, 8.0f, 12, AllottedGeometry, IsolatedColor, IsolatedColor,
+							LabelBrush, OutDrawElements, SampleLayer - 1);
+			}
+
 			// Draw lines/boxes to indicated which axes are locked
 			if (!bReadOnly)
 			{
@@ -537,7 +734,12 @@
 	DrawLayerId += 3;
 }
 
-void SBlendSpaceGridWidget::PaintAxisText(const FGeometry& AllottedGeometry, const FSlateRect& MyCullingRect, FSlateWindowElementList& OutDrawElements, int32& DrawLayerId) const
+//======================================================================================================================
+void SBlendSpaceGridWidget::PaintAxisText(
+	const FGeometry&         AllottedGeometry, 
+	const FSlateRect&        MyCullingRect, 
+	FSlateWindowElementList& OutDrawElements, 
+	int32&                   DrawLayerId) const
 {
 	const TSharedRef< FSlateFontMeasure > FontMeasure = FSlateApplication::Get().GetRenderer()->GetFontMeasureService();
 	const FVector2D GridCenter = CachedGridRectangle.GetCenter();
@@ -632,146 +834,52 @@
 	DrawLayerId += 1;
 }
 
-// Paint a filled triangle
-static void PaintTriangle(
-	const FVector2D&         P0,
-	const FVector2D&         P1,
-	const FVector2D&         P2,
-	const FGeometry&         AllottedGeometry,
-	FLinearColor             Color,
-	const FSlateBrush*       Brush,
-	FSlateWindowElementList& OutDrawElements,
-	int32                    DrawLayerId)
+//======================================================================================================================
+void SBlendSpaceGridWidget::PaintTriangulation(
+	const FGeometry&         AllottedGeometry, 
+	const FSlateRect&        MyCullingRect, 
+	FSlateWindowElementList& OutDrawElements, 
+	int32&                   DrawLayerId) const
 {
-	const FVector2D* Points[3] = { &P0, &P1, &P2 };
-
-	TArray<FSlateVertex> Vertices;
-	Vertices.Reserve(3);
-
-	for (int32 PointIndex = 0; PointIndex != 3; ++PointIndex)
+	if (const UBlendSpace* BlendSpace = BlendSpaceBase.Get())
 	{
-		Vertices.AddZeroed();
-		FSlateVertex& NewVert = Vertices.Last();
-		NewVert.Position = AllottedGeometry.LocalToAbsolute(*Points[PointIndex]);
-		NewVert.Color = Color.ToFColor(false);
-	}
-
-	// Fill by making triangles
-	TArray<SlateIndex> VertexIndices = { 0, 1, 2 };
-	FSlateDrawElement::MakeCustomVerts(
-		OutDrawElements, DrawLayerId, Brush->GetRenderingResource(), Vertices, VertexIndices, nullptr, 0, 0);
-}
-
-// Paints a filled polygon with outline, defined by a set of points which don't need to be sorted.
-// This will handle concave polygons, but only if the centroid lies inside the polygon.
-static void PaintPolygon(
-	TArray<FVector2D>&       Points,
-	const FGeometry&         AllottedGeometry,
-	FLinearColor             FillColor,
-	FLinearColor             OutlineColor,
-	const FSlateBrush*       Brush,
-	FSlateWindowElementList& OutDrawElements,
-	int32                    DrawLayerId)
-{
-	TArray<FSlateVertex> Vertices;
-	Vertices.Reserve(Points.Num() + 1);
-
-	// Add a mid-position vertex so that we handle polygons that aren't completely convex
-	Vertices.AddZeroed();
-
-	FSlateVertex& MidVertex = Vertices.Last();
-	for (int32 PointIndex = 0; PointIndex != Points.Num(); ++PointIndex)
-	{
-		Vertices.AddZeroed();
-		FSlateVertex& NewVert = Vertices.Last();
-		NewVert.Position = AllottedGeometry.LocalToAbsolute(Points[PointIndex]);
-		NewVert.Color = FillColor.ToFColor(false);
-		MidVertex.Position += NewVert.Position;
-	}
-	MidVertex.Position /= Points.Num();
-	MidVertex.Color = FillColor.ToFColor(false);
-
-	// Make sure the points all wind correctly relative to the mid point
-	struct FComparePoints
-	{
-		FComparePoints(const FSlateVertex& Mid) : MidPoint(Mid) {}
-		bool operator()(const FSlateVertex& A, const FSlateVertex& B) const
+		if ( 
+			(bReadOnly && BlendSpace->bInterpolateUsingGrid) ||
+			(!bReadOnly && !bShowTriangulation)
+			)
 		{
-			FVector2D DeltaA = A.Position - MidPoint.Position;
-			FVector2D DeltaB = B.Position - MidPoint.Position;
-			float AngleA = FMath::Atan2(DeltaA.Y, DeltaA.X);
-			float AngleB = FMath::Atan2(DeltaB.Y, DeltaB.X);
-			return AngleA < AngleB;
-		}
-		FSlateVertex MidPoint;
-	};
-	Sort(Vertices.GetData() + 1, Vertices.Num() - 1, FComparePoints(MidVertex));
-
-	if (FillColor.A > 0)
-	{
-		// Fill by making triangles
-		TArray<SlateIndex> VertexIndices;
-		for (int VertexIndex = 1; VertexIndex < Vertices.Num(); ++VertexIndex)
-		{
-			VertexIndices.Add(0);
-			VertexIndices.Add(VertexIndex);
-			VertexIndices.Add(VertexIndex + 1 >= Vertices.Num() ? 1 : VertexIndex + 1);
+			return;
 		}
 
-		FSlateDrawElement::MakeCustomVerts(
-			OutDrawElements, DrawLayerId, Brush->GetRenderingResource(), Vertices, VertexIndices, nullptr, 0, 0);
-	}
-
-	if (OutlineColor.A > 0)
-	{
-		TArray<FVector2D> LinePoints;
-		LinePoints.Reserve(Points.Num() + 1);
-		for (int VertexIndex = 1; VertexIndex <= Vertices.Num(); ++VertexIndex)
-		{
-			LinePoints.Add(VertexIndex < Vertices.Num() ? Vertices[VertexIndex].Position : Vertices[1].Position);
-		}
-		FSlateDrawElement::MakeLines(
-			OutDrawElements, DrawLayerId + 1, FPaintGeometry(), LinePoints, ESlateDrawEffect::None, OutlineColor, true, 1.0f);
-	}
-}
-
-void SBlendSpaceGridWidget::PaintTriangulation(const FGeometry& AllottedGeometry, const FSlateRect& MyCullingRect, FSlateWindowElementList& OutDrawElements, int32& DrawLayerId) const
-{
-	TArray<FVector2D> PolygonPoints;
-
-	if(const UBlendSpace* BlendSpace = BlendSpaceBase.Get())
-	{
+		TArray<FVector2D> PolygonPoints; 
 		const TArray<FBlendSample>& Samples = BlendSpace->GetBlendSamples();
 		if (BlendSpace->bInterpolateUsingGrid)
 		{
-			// Use the Grid
+			// Use the Grid. After preprocessing, each grid point will have been located within the triangulation, and
+			// will store the indices of the three samples that surround/contribute to it.
 			const TArray<FEditorElement>& EditorElements = BlendSpace->GetGridSamples();
-			for (const FEditorElement& Element : EditorElements)
+			for (int32 ElementIndex = 0 ; ElementIndex != EditorElements.Num() ; ++ElementIndex)
 			{
+				const FEditorElement& Element = EditorElements[ElementIndex];
+				FVector ElementPosition = BlendSpace->GetGridPosition(ElementIndex);
+				FVector2D ElementScreenPosition = SampleValueToScreenPosition(ElementPosition);
 				for (int32 SourceIndex = 0; SourceIndex < 3; ++SourceIndex)
 				{
 					const int32 SourceSampleIndex = Element.Indices[SourceIndex]; 
 					if (Samples.IsValidIndex(SourceSampleIndex))
 					{
 						const FBlendSample& SourceSample = Samples[SourceSampleIndex];
-						for (int32 TargetIndex = 0; TargetIndex < 3; ++TargetIndex)
+						FVector2D SourceSampleScreenPosition = SampleValueToScreenPosition(SourceSample.SampleValue);
+						// Draw line from the grid point to each sample that contributes to it
+						if (Element.Weights[SourceIndex] > 0)
 						{
-							const int32 TargetSampleIndex = Element.Indices[TargetIndex]; 
-							if (Samples.IsValidIndex(TargetSampleIndex))
+							if (ElementScreenPosition != SourceSampleScreenPosition)
 							{
-								if (TargetIndex != SourceIndex)
-								{
-									const FBlendSample& TargetSample = Samples[TargetSampleIndex];
-									TArray<FVector2D> Points;
-
-									Points.Add(SampleValueToScreenPosition(SourceSample.SampleValue));
-									Points.Add(SampleValueToScreenPosition(TargetSample.SampleValue));
-
-									// Draw line from and to element
-									FSlateDrawElement::MakeLines(
-										OutDrawElements, DrawLayerId + 1, AllottedGeometry.ToPaintGeometry(), Points, 
-										ESlateDrawEffect::None,  TriangulationColor.GetSpecifiedColor(), true, 0.5f);
-								}
+								FSlateDrawElement::MakeLines(
+									OutDrawElements, DrawLayerId + 1, AllottedGeometry.ToPaintGeometry(), 
+									{ ElementScreenPosition, SourceSampleScreenPosition },
+									ESlateDrawEffect::None,  TriangulationColor.GetSpecifiedColor(), 
+									true, Element.Weights[SourceIndex]);
 							}
 						}
 					}
@@ -794,11 +902,7 @@
 					ESlateDrawEffect::None, TriangulationColor.GetSpecifiedColor(), true, 0.5f);
 			}
 
-			// Flag any triangle that has an interior angle smaller than this
-			const float CriticalAngle = 4.0f;
-			// Flag any triangle that has a smaller area than this (normalized units)
-			const float CriticalArea = 5e-4f;
-			const float CriticalDot = FMath::Cos(FMath::DegreesToRadians(CriticalAngle));
+			const float CriticalDot = FMath::Cos(FMath::DegreesToRadians(CriticalTriangulationAngle));
 			for (const FBlendSpaceTriangle& Triangle : BlendSpaceData.Triangles)
 			{
 				FLinearColor TriangleFillColor = TriangulationCurrentColor.GetSpecifiedColor();
@@ -827,7 +931,7 @@
 						FVector2D B = NormalizedPositions[Index] - NormalizedPositions[(Index + 1) % 3];
 						float Dot = A.GetSafeNormal() | B.GetSafeNormal();
 						float Area = 0.5f * FMath::Abs(B ^ A);
-						if (Dot > CriticalDot || Area < CriticalArea)
+						if (Dot > CriticalDot || Area < CriticalTriangulationArea)
 						{
 							TriangleFillColor = FLinearColor::Red;
 							TriangleFillColor.A = 0.5f;
@@ -1371,9 +1475,24 @@
 {
 	InMenuBuilder.BeginSection("ViewOptions", LOCTEXT("ViewOptionsMenuHeader", "View Options"));
 	{
+		TAttribute<FText> ShowTriangulation = TAttribute<FText>::Create(
+			[this]()
+			{
+				return (BlendSpaceBase.Get() && BlendSpaceBase.Get()->bInterpolateUsingGrid) 
+					? LOCTEXT("ShowGridToSampleConnections", "Show Grid/Sample Connections")
+					: LOCTEXT("ShowTriangulation", "Show Triangulation");
+			});
+		TAttribute<FText> ShowTriangulationToolTip = TAttribute<FText>::Create(
+			[this]()
+			{
+				return (BlendSpaceBase.Get() && BlendSpaceBase.Get()->bInterpolateUsingGrid) 
+					? LOCTEXT("ShowGridToSampleConnectionsToolTip", "Show which samples each grid point is associated with")
+					: LOCTEXT("ShowTriangulationToolTip", "Show the Delaunay triangulation for all blend space samples");
+			});
+
 		InMenuBuilder.AddMenuEntry(
-			LOCTEXT("ShowTriangulation", "Show Triangulation"),
-			LOCTEXT("ShowTriangulationToolTip", "Show the Delaunay triangulation for all blend space samples"),
+			ShowTriangulation,
+			ShowTriangulationToolTip,
 			FSlateIcon("EditorStyle", "BlendSpaceEditor.ToggleTriangulation"),
 			FUIAction(
 				FExecuteAction::CreateLambda([this](){ bShowTriangulation = !bShowTriangulation; }),
@@ -1731,6 +1850,56 @@
 	bSamplePreviewing = false;
 }
 
+FText SBlendSpaceGridWidget::GetToolTipSampleValidity() const
+{
+	const UBlendSpace* BlendSpace = BlendSpaceBase.Get();
+	FText ToolTipText = FText::GetEmpty();
+	if(!bReadOnly && BlendSpace && BlendSpace->bInterpolateUsingGrid)
+	{
+		int32 SampleIndex = INDEX_NONE;
+		if (DragState == EDragState::None)
+		{
+			SampleIndex = HighlightedSampleIndex;
+		}
+		else if (DragState == EDragState::DragSample)
+		{
+			SampleIndex = DraggedSampleIndex;
+		}
+		else
+		{
+			SampleIndex = INDEX_NONE;
+		}
+
+		if (SampleIndex != INDEX_NONE && BlendSpace->IsValidBlendSampleIndex(SampleIndex))
+		{
+			float SampleLookupWeight = GetSampleLookupWeight(SampleIndex);
+			if (SampleLookupWeight >= 1.0f)
+			{
+				return ToolTipText;
+			}
+			else if (SampleLookupWeight <= 0.0f)
+			{
+				ToolTipText = FText::Format(
+					LOCTEXT("SampleValidityZero", "Self weight is zero"), 
+					SampleLookupWeight);
+			}
+			else if (SampleLookupWeight <= SampleLookupWeightThreshold)
+			{
+				ToolTipText = FText::Format(
+					LOCTEXT("SampleValidityLow", "Self weight is low: {0}"), 
+					SampleLookupWeight);
+			}
+			else if (SampleLookupWeight < 1.0f)
+			{
+				ToolTipText = FText::Format(
+					LOCTEXT("SampleValidity", "Self weight: {0}"), 
+					SampleLookupWeight);
+			}
+		}
+	}
+	return ToolTipText;
+}
+
 FText SBlendSpaceGridWidget::GetToolTipAnimationName() const
 {
 	FText ToolTipText = FText::GetEmpty();
@@ -1822,7 +1991,7 @@
 	if(const UBlendSpace* BlendSpace = BlendSpaceBase.Get())
 	{
 		static const FTextFormat OneAxisFormat = LOCTEXT("OneAxisFormat", "{0}: {1}");
-		static const FTextFormat TwoAxisFormat = LOCTEXT("TwoAxisFormat", "{0}: {1} - {2}: {3}");
+		static const FTextFormat TwoAxisFormat = LOCTEXT("TwoAxisFormat", "{0}: {1}  {2}: {3}");
 		const FTextFormat& ValueFormattingText = (GridType == EGridType::TwoAxis) ? TwoAxisFormat : OneAxisFormat;
 
 		auto AddAdvancedPreview = [this, &ToolTipText, BlendSpace]()
diff --git a/Engine/Source/Editor/Persona/Private/SAnimationBlendSpaceGridWidget.h b/Engine/Source/Editor/Persona/Private/SAnimationBlendSpaceGridWidget.h
--- a/Engine/Source/Editor/Persona/Private/SAnimationBlendSpaceGridWidget.h
+++ b/Engine/Source/Editor/Persona/Private/SAnimationBlendSpaceGridWidget.h
@@ -201,6 +201,7 @@ protected:
 	/** Tool tip text content getters*/
 	FText GetToolTipAnimationName() const;
 	FText GetToolTipSampleValue() const;
+	FText GetToolTipSampleValidity() const;
 	FText GetSampleErrorMessage(const struct FBlendSample& BlendSample) const;
 
 	/** Functionality for the grid input box widgets (also used by FBlendSampleDetails) */
@@ -224,6 +225,12 @@ protected:
 
 	/** Enable/disable the status bar message */
 	void EnableStatusBarMessage(bool bEnable);
+
+	/**
+	 * This returns the weight for this sample if a lookup was made on its position. Ideally this would be 1 
+	 * (and will be for triangulation). It will not always be so when using the grid.
+	 */
+	float GetSampleLookupWeight(int32 SampleIndex) const;
 private:
 	/** Currently visualized blendspace (const to ensure changes to it are only made within SAnimationBlendSpace */
 	TAttribute<const UBlendSpace*> BlendSpaceBase;
diff --git a/Engine/Source/Runtime/Core/Public/UObject/UE5MainStreamObjectVersion.h b/Engine/Source/Runtime/Core/Public/UObject/UE5MainStreamObjectVersion.h
--- a/Engine/Source/Runtime/Core/Public/UObject/UE5MainStreamObjectVersion.h
+++ b/Engine/Source/Runtime/Core/Public/UObject/UE5MainStreamObjectVersion.h
@@ -154,6 +154,9 @@ struct CORE_API FUE5MainStreamObjectVersion
 		// The SkyAtmosphere now light up the heightfog by default, and by default the height fog has a black color.
 		SkyAtmosphereAffectsHeightFogWithBetterDefault,
 
+		// Ordering of samples in BlendSpace
+		BlendSpaceSampleOrdering,
+
 		// -----<new versions can be added above this line>-------------------------------------------------
 		VersionPlusOne,
 		LatestVersion = VersionPlusOne - 1
diff --git a/Engine/Source/Runtime/Engine/Classes/Animation/BlendSpace.h b/Engine/Source/Runtime/Engine/Classes/Animation/BlendSpace.h
--- a/Engine/Source/Runtime/Engine/Classes/Animation/BlendSpace.h
+++ b/Engine/Source/Runtime/Engine/Classes/Animation/BlendSpace.h
@@ -592,6 +592,12 @@ public:
 	*/
 	ENGINE_API const TArray<FEditorElement>& GetGridSamples() const;
 
+	/** Returns the sample position associated with the elements returned by GetGridSamples */
+	ENGINE_API FVector GetGridPosition(int32 GridIndex) const;
+
+	/** Returns the sample position associated with the coordinates */
+	ENGINE_API FVector GetGridPosition(int32 GridX, int32 GridY) const;
+
 	/**
 	 * Returns the runtime triangulation etc data
 	 */
diff --git a/Engine/Source/Runtime/Engine/Private/Animation/BlendSpace.cpp b/Engine/Source/Runtime/Engine/Private/Animation/BlendSpace.cpp
--- a/Engine/Source/Runtime/Engine/Private/Animation/BlendSpace.cpp
+++ b/Engine/Source/Runtime/Engine/Private/Animation/BlendSpace.cpp
@@ -84,8 +84,9 @@ void UBlendSpace::Serialize(FArchive& Ar)
 		}
 	}
 
-	if (Ar.IsLoading() && (Ar.CustomVer(FUE5MainStreamObjectVersion::GUID) <
-		FUE5MainStreamObjectVersion::BlendSpaceRuntimeTriangulation))
+	if (Ar.IsLoading() && 
+		(Ar.CustomVer(FUE5MainStreamObjectVersion::GUID) <
+		 FUE5MainStreamObjectVersion::BlendSpaceRuntimeTriangulation))
 	{
 		// Make old blend spaces use the grid
 		bInterpolateUsingGrid = true;
@@ -116,6 +117,15 @@ void UBlendSpace::Serialize(FArchive& Ar)
 			}
 		}
 	}
+
+	if (Ar.IsLoading() && 
+		(Ar.CustomVer(FUE5MainStreamObjectVersion::GUID) <
+		 FUE5MainStreamObjectVersion::BlendSpaceSampleOrdering))
+	{
+		// Force the data to be updated after the ordering of 2D SampleData has been changed to be consistent with 1D
+		ResampleData();
+	}
+
 #endif // WITH_EDITOR
 }
 
@@ -1911,7 +1921,8 @@ void UBlendSpace::GetRawSamplesFromBlendInput1D(const FVector& BlendInput, TArra
 
 const FEditorElement* UBlendSpace::GetEditorElement(int32 XIndex, int32 YIndex) const
 {
-	int32 Index = XIndex * (BlendParameters[1].GridNum + 1) + YIndex;
+	// Needs to match FBlendSpaceGrid::GetElement and FBlendSpaceGrid::GetElement
+	int32 Index = YIndex * (BlendParameters[0].GridNum + 1) + XIndex;
 	return GetGridSampleInternal(Index);
 }
 
@@ -1984,6 +1995,31 @@ void UBlendSpace::GetRawSamplesFromBlendInput2D(const FVector& BlendInput, TArra
 
 #if WITH_EDITOR
 
+//======================================================================================================================
+FVector UBlendSpace::GetGridPosition(int32 GridX, int32 GridY) const
+{
+	const FVector GridMin(BlendParameters[0].Min, BlendParameters[1].Min, 0.0f);
+	const FVector GridMax(BlendParameters[0].Max, BlendParameters[1].Max, 0.0f);
+	const FVector GridRange(GridMax.X - GridMin.X, GridMax.Y - GridMin.Y, 0.0f);
+	const FIntPoint NumGridDivisions(BlendParameters[0].GridNum, BlendParameters[1].GridNum);
+
+	const FVector GridPoint(
+		GridMin.X + (GridX * GridRange.X / NumGridDivisions.X), 
+		GridMin.Y + (GridY * GridRange.Y / NumGridDivisions.Y),
+		0.0f);
+	return GridPoint;
+}
+
+//======================================================================================================================
+FVector UBlendSpace::GetGridPosition(int32 GridIndex) const
+{
+	// Needs to match FBlendSpaceGrid::GetElement and UBlendSpace::GetEditorElement
+	int32 GridX = GridIndex % (BlendParameters[0].GridNum + 1);
+	int32 GridY = (GridIndex - GridX) / (BlendParameters[0].GridNum + 1);
+	return GetGridPosition(GridX, GridY);
+}
+
+//======================================================================================================================
 // Note that this runs and is needed if the axes settings are change to enable snapping - then we do
 // a pass over all the samples and make sure everything is in order. When samples are being dragged
 // around/dropped then that moving sample will be handled in SBlendSpaceGridWidget so, whilst this
@@ -1992,7 +2028,6 @@ void UBlendSpace::SnapSamplesToClosestGridPoint()
 {
 	if (BlendParameters[0].bSnapToGrid && BlendParameters[1].bSnapToGrid)
 	{
-		TArray<FVector> GridPoints;
 		TArray<int32> ClosestSampleToGridPoint;
 
 		const FVector GridMin(BlendParameters[0].Min, BlendParameters[1].Min, 0.0f);
@@ -2004,26 +2039,13 @@ void UBlendSpace::SnapSamplesToClosestGridPoint()
 		// Snap to nearest in normalized space - not depending on the units of the params (which may be completely different).
 		const float GridRatio = GridStep.Y / GridStep.X;
 
-		for (int32 GridY = 0; GridY < NumGridDivisions.Y + 1; ++GridY)
-		{
-			for (int32 GridX = 0; GridX < NumGridDivisions.X + 1; ++GridX)
-			{
-				const FVector GridPoint(
-					GridMin.X + (GridX * GridRange.X / NumGridDivisions.X), 
-					GridMin.Y + (GridY * GridRange.Y / NumGridDivisions.Y),
-					0.0f);
-
-				//const FVector GridPoint((GridX * GridStep.X) + GridMin.X, (GridY * GridStep.Y) + GridMin.Y, 0.0f);
-				GridPoints.Add(GridPoint);
-			}
-		}
-
-		ClosestSampleToGridPoint.Init(INDEX_NONE, GridPoints.Num());
+		int32 NumGridPoints = (NumGridDivisions.X + 1) * (NumGridDivisions.Y + 1);
+		ClosestSampleToGridPoint.Init(INDEX_NONE, NumGridPoints);
 
 		// Find closest sample to grid point
-		for (int32 PointIndex = 0; PointIndex < GridPoints.Num(); ++PointIndex)
+		for (int32 GridIndex = 0; GridIndex < NumGridPoints; ++GridIndex)
 		{
-			const FVector& GridPoint = GridPoints[PointIndex];
+			const FVector& GridPoint = GetGridPosition(GridIndex);
 			float SmallestDistanceSq = FLT_MAX;
 			int32 Index = INDEX_NONE;
 
@@ -2040,7 +2062,7 @@ void UBlendSpace::SnapSamplesToClosestGridPoint()
 				}
 			}
 
-			ClosestSampleToGridPoint[PointIndex] = Index;
+			ClosestSampleToGridPoint[GridIndex] = Index;
 		}
 
 		// Find closest grid point to sample
@@ -2051,14 +2073,14 @@ void UBlendSpace::SnapSamplesToClosestGridPoint()
 			// Find closest grid point
 			float SmallestDistanceSq = FLT_MAX;
 			int32 Index = INDEX_NONE;
-			for (int32 PointIndex = 0; PointIndex < GridPoints.Num(); ++PointIndex)
+			for (int32 GridIndex = 0; GridIndex < NumGridPoints; ++GridIndex)
 			{
-				FVector Delta = GridPoints[PointIndex] - BlendSample.SampleValue;
+				FVector Delta = GetGridPosition(GridIndex) - BlendSample.SampleValue;
 				Delta.X *= GridRatio;
 				const float DistanceSq = Delta.SizeSquared2D();
 				if (DistanceSq < SmallestDistanceSq)
 				{
-					Index = PointIndex;
+					Index = GridIndex;
 					SmallestDistanceSq = DistanceSq;
 				}
 			}
@@ -2068,7 +2090,7 @@ void UBlendSpace::SnapSamplesToClosestGridPoint()
 			{
 				for (int32 AxisIndex = 0; AxisIndex != 2; ++AxisIndex)
 				{
-					BlendSample.SampleValue[AxisIndex] = GridPoints[Index][AxisIndex];
+					BlendSample.SampleValue[AxisIndex] = GetGridPosition(Index)[AxisIndex];
 				}
 			}
 		}
diff --git a/Engine/Source/Runtime/Engine/Private/Animation/BlendSpace1D.cpp b/Engine/Source/Runtime/Engine/Private/Animation/BlendSpace1D.cpp
--- a/Engine/Source/Runtime/Engine/Private/Animation/BlendSpace1D.cpp
+++ b/Engine/Source/Runtime/Engine/Private/Animation/BlendSpace1D.cpp
@@ -24,7 +24,6 @@ void UBlendSpace1D::SnapSamplesToClosestGridPoint()
 	if (!BlendParameters[0].bSnapToGrid)
 		return;
 
-	TArray<float> GridPoints;
 	TArray<int32> ClosestSampleToGridPoint;
 
 	const float GridMin = BlendParameters[0].Min;
@@ -33,18 +32,12 @@ void UBlendSpace1D::SnapSamplesToClosestGridPoint()
 	const int32 NumGridPoints = BlendParameters[0].GridNum + 1;
 	const float GridStep = GridRange / BlendParameters[0].GridNum;
 
-	for (int32 GridPointIndex = 0; GridPointIndex < NumGridPoints; ++GridPointIndex)
-	{
-		const float GridPointValue = (GridPointIndex * GridStep) + GridMin;
-		GridPoints.Add(GridPointValue);
-	}
-
-	ClosestSampleToGridPoint.Init(INDEX_NONE, GridPoints.Num());
+	ClosestSampleToGridPoint.Init(INDEX_NONE, NumGridPoints);
 
 	// Find closest sample to grid point
-	for (int32 PointIndex = 0; PointIndex < GridPoints.Num(); ++PointIndex)
+	for (int32 PointIndex = 0; PointIndex < NumGridPoints; ++PointIndex)
 	{
-		const float GridPoint = GridPoints[PointIndex];
+		const float GridPoint = GetGridPosition(PointIndex)[0];
 		float SmallestDistance = FLT_MAX;
 		int32 Index = INDEX_NONE;
 
@@ -70,9 +63,9 @@ void UBlendSpace1D::SnapSamplesToClosestGridPoint()
 		// Find closest grid point
 		float SmallestDistance = FLT_MAX;
 		int32 Index = INDEX_NONE;
-		for (int32 PointIndex = 0; PointIndex < GridPoints.Num(); ++PointIndex)
+		for (int32 PointIndex = 0; PointIndex < NumGridPoints; ++PointIndex)
 		{
-			const float Distance = FMath::Abs(GridPoints[PointIndex] - BlendSample.SampleValue[0]);
+			const float Distance = FMath::Abs(GetGridPosition(PointIndex)[0] - BlendSample.SampleValue[0]);
 			if (Distance < SmallestDistance)
 			{
 				Index = PointIndex;
@@ -83,7 +76,7 @@ void UBlendSpace1D::SnapSamplesToClosestGridPoint()
 		// Only move the sample if it is also closest to the grid point
 		if (Index != INDEX_NONE && ClosestSampleToGridPoint[Index] == SampleIndex)
 		{
-			BlendSample.SampleValue[0] = GridPoints[Index];
+			BlendSample.SampleValue[0] = GetGridPosition(Index)[0];
 		}
 	}
 }
diff --git a/Engine/Source/Runtime/Engine/Private/Animation/BlendSpaceHelpers.cpp b/Engine/Source/Runtime/Engine/Private/Animation/BlendSpaceHelpers.cpp
--- a/Engine/Source/Runtime/Engine/Private/Animation/BlendSpaceHelpers.cpp
+++ b/Engine/Source/Runtime/Engine/Private/Animation/BlendSpaceHelpers.cpp
@@ -634,7 +634,7 @@ void FBlendSpaceGrid::GenerateGridElements(const TArray<FVertex>& SamplePoints,
 		for (int32 GridPositionY = 0; GridPositionY < NumGridPointsForAxis.Y; ++GridPositionY)
 		{
 			FTriangle * SelectedTriangle = NULL;
-			FEditorElement& GridPoint = GridPoints[ GridPositionX*NumGridPointsForAxis.Y + GridPositionY ];
+			FEditorElement& GridPoint = GetElement(GridPositionX, GridPositionY);
 
 			GridPointPosition = GetPosFromIndex(GridPositionX, GridPositionY);
 
@@ -758,7 +758,19 @@ const FEditorElement& FBlendSpaceGrid::GetElement(const int32 GridX, const int32
 	check (NumGridPointsForAxis.Y >= GridY);
 
 	check (GridPoints.Num() > 0 );
-	return GridPoints[GridX * NumGridPointsForAxis.Y + GridY];
+	return GridPoints[GridY * NumGridPointsForAxis.X + GridX];
+}
+
+FEditorElement& FBlendSpaceGrid::GetElement(const int32 GridX, const int32 GridY)
+{
+	FEditorElement& Test = const_cast<FEditorElement &>(std::as_const(*this).GetElement(GridX, GridY));
+
+	check (NumGridPointsForAxis.X >= GridX);
+	check (NumGridPointsForAxis.Y >= GridY);
+
+	check (GridPoints.Num() > 0 );
+	FEditorElement& Test2 = GridPoints[GridY * NumGridPointsForAxis.X + GridX];
+	return Test2;
 }
 
 #undef LOCTEXT_NAMESPACE
diff --git a/Engine/Source/Runtime/Engine/Private/Animation/BlendSpaceHelpers.h b/Engine/Source/Runtime/Engine/Private/Animation/BlendSpaceHelpers.h
--- a/Engine/Source/Runtime/Engine/Private/Animation/BlendSpaceHelpers.h
+++ b/Engine/Source/Runtime/Engine/Private/Animation/BlendSpaceHelpers.h
@@ -486,6 +486,7 @@ public:
 	const FVector2D GetPosFromIndex(const int32 GridX, const int32 GridY) const;
 
 private:
+	FEditorElement& GetElement(const int32 GridX, const int32 GridY);
 	// Grid Dimension
 	FVector2D GridMin;
 	FVector2D GridMax;
