When editing notify data, it doesn't mark asset as dirty
#fix added transactions and Modify() calls to notify context menu

diff --git a/Engine/Source/Editor/Persona/Private/SAnimNotifyPanel.cpp b/Engine/Source/Editor/Persona/Private/SAnimNotifyPanel.cpp
--- a/Engine/Source/Editor/Persona/Private/SAnimNotifyPanel.cpp
+++ b/Engine/Source/Editor/Persona/Private/SAnimNotifyPanel.cpp
@@ -2824,6 +2824,9 @@ TSharedPtr<SWidget> SAnimNotifyTrack::SummonContextMenu(const FGeometry& MyGeome
 							{
 								if (InCommitType == ETextCommit::OnEnter && NotifyNodes.IsValidIndex(NodeIndex))
 								{
+									const FScopedTransaction Transaction(LOCTEXT("SetNotifyTimeTransaction", "Set Anim Notify trigger time"));
+									Sequence->Modify();
+									
 									INodeObjectInterface* LocalNodeObject = NotifyNodes[NodeIndex]->NodeObjectInterface;
 
 									float NewTime = FMath::Clamp(InValue, 0.0f, Sequence->GetPlayLength() - LocalNodeObject->GetDuration());
@@ -2867,6 +2870,9 @@ TSharedPtr<SWidget> SAnimNotifyTrack::SummonContextMenu(const FGeometry& MyGeome
 							{
 								if (InCommitType == ETextCommit::OnEnter && NotifyNodes.IsValidIndex(NodeIndex))
 								{
+									const FScopedTransaction Transaction(LOCTEXT("SetNotifyFrameTransaction", "Set Anim Notify trigger frame index"));
+									Sequence->Modify();
+
 									INodeObjectInterface* LocalNodeObject = NotifyNodes[NodeIndex]->NodeObjectInterface;
 
 									float NewTime = FMath::Clamp(Sequence->GetTimeAtFrame(InValue), 0.0f, Sequence->GetPlayLength() - LocalNodeObject->GetDuration());
@@ -2912,6 +2918,9 @@ TSharedPtr<SWidget> SAnimNotifyTrack::SummonContextMenu(const FGeometry& MyGeome
 								{
 									if ( InCommitType == ETextCommit::OnEnter && AnimNotifies.IsValidIndex(NotifyIndex) )
 									{
+										const FScopedTransaction Transaction(LOCTEXT("SetNotifyWeightTransaction", "Set Anim Notify trigger weight"));
+										Sequence->Modify();
+									
 										float NewWeight = FMath::Max(InValue, ZERO_ANIMWEIGHT_THRESH);
 										AnimNotifies[NotifyIndex]->TriggerWeightThreshold = NewWeight;
 
@@ -2946,6 +2955,9 @@ TSharedPtr<SWidget> SAnimNotifyTrack::SummonContextMenu(const FGeometry& MyGeome
 									{
 										if ( InCommitType == ETextCommit::OnEnter && AnimNotifies.IsValidIndex(NotifyIndex) )
 										{
+											const FScopedTransaction Transaction(LOCTEXT("SetNotifyDurationSecondsTransaction", "Set Anim Notify State duration in seconds"));
+											Sequence->Modify();
+											
 											float NewDuration = FMath::Max(InValue, SAnimNotifyNode::MinimumStateDuration);
 											float MaxDuration = Sequence->GetPlayLength() - AnimNotifies[NotifyIndex]->GetTime();
 											NewDuration = FMath::Min(NewDuration, MaxDuration);
@@ -2983,6 +2995,9 @@ TSharedPtr<SWidget> SAnimNotifyTrack::SummonContextMenu(const FGeometry& MyGeome
 									{
 										if ( InCommitType == ETextCommit::OnEnter && AnimNotifies.IsValidIndex(NotifyIndex) )
 										{
+											const FScopedTransaction Transaction(LOCTEXT("SetNotifyDurationFramesTransaction", "Set Anim Notify State duration in frames"));
+											Sequence->Modify();
+											
 											float NewDuration = FMath::Max(Sequence->GetTimeAtFrame(InValue), SAnimNotifyNode::MinimumStateDuration);
 											float MaxDuration = Sequence->GetPlayLength() - AnimNotifies[NotifyIndex]->GetTime();
 											NewDuration = FMath::Min(NewDuration, MaxDuration);
