Fixed a bug where single step in the Animation Blueprint had no effect.  It now steps at a fixed framerate.

diff --git a/Engine/Source/Editor/Persona/Private/SAnimationScrubPanel.cpp b/Engine/Source/Editor/Persona/Private/SAnimationScrubPanel.cpp
--- a/Engine/Source/Editor/Persona/Private/SAnimationScrubPanel.cpp
+++ b/Engine/Source/Editor/Persona/Private/SAnimationScrubPanel.cpp
@@ -85,12 +85,12 @@ FReply SAnimationScrubPanel::OnClick_Forward_Step()
 	}
 	else if (SMC)
 	{
-		// BlendSpace combines animations so there's no such thing as a frame. However, 1/30 is a sensible/common rate.
+		// BlendSpaces and Animation Blueprints combine animations so there's no such thing as a frame. However, 1/30 is a sensible/common rate.
 		const float FixedFrameRate = 30.0f;
 
 		// Advance a single frame, leaving it paused afterwards
 		SMC->GlobalAnimRateScale = 1.0f;
-		SMC->TickAnimation(1.0f / FixedFrameRate, false);
+		GetPreviewScene()->Tick(1.0f / FixedFrameRate); 
 		SMC->GlobalAnimRateScale = 0.0f;
 	}
 
