Moving samples in Blend Spaces creates multiple/continuous undo points
#fix added inter-frame scoped transaction for interactive changes to blend samples, and patched up paths which should broadcast interactive/non-interactive versions of OnSampleMoved
#misc added scoped transaction to blendspace sample analysis
#misc float-to-double truncation fix

diff --git a/Engine/Source/Editor/Persona/Private/Customization/BlendSpaceDetails.cpp b/Engine/Source/Editor/Persona/Private/Customization/BlendSpaceDetails.cpp
--- a/Engine/Source/Editor/Persona/Private/Customization/BlendSpaceDetails.cpp
+++ b/Engine/Source/Editor/Persona/Private/Customization/BlendSpaceDetails.cpp
@@ -53,6 +53,9 @@ FReply FBlendSpaceDetails::HandleAnalyzeSamples()
 {
 	if(BlendSpace->IsAsset())
 	{
+		FScopedTransaction ScopedTransaction(LOCTEXT("AnalyzeBlendSpaceSamples", "Applying Blend Space sample analysis"));
+		BlendSpace->Modify();
+
 		FSlateApplication::Get().DismissAllMenus();
 		bool bChangedOne = false;
 
@@ -86,10 +89,13 @@ FReply FBlendSpaceDetails::HandleAnalyzeSamples()
 
 		if (bChangedOne)
 		{
-			BlendSpace->Modify();
 			FPropertyChangedEvent ChangedEvent(nullptr, EPropertyChangeType::ArrayClear);
 			BlendSpace->PostEditChangeProperty(ChangedEvent);
 		}
+		else
+		{
+			ScopedTransaction.Cancel();
+		}
 	}
 	return FReply::Handled();
 }
diff --git a/Engine/Source/Editor/Persona/Private/SAnimationBlendSpace.cpp b/Engine/Source/Editor/Persona/Private/SAnimationBlendSpace.cpp
--- a/Engine/Source/Editor/Persona/Private/SAnimationBlendSpace.cpp
+++ b/Engine/Source/Editor/Persona/Private/SAnimationBlendSpace.cpp
@@ -115,19 +115,31 @@ void SBlendSpaceEditor::Construct(const FArguments& InArgs, const TSharedRef<cla
 void SBlendSpaceEditor::OnSampleMoved(const int32 SampleIndex, const FVector& NewValue, bool bIsInteractive)
 {
 	bool bMoveSuccessful = true;
+
+	// If this is an interactive operation and a transaction has not been opened yet, or if it's a non-interactive operation setup a transaction
+	if (!bIsInteractive || !SampleMoveTransaction.IsValid())
+	{
+		SampleMoveTransaction = MakeUnique<FScopedTransaction>(LOCTEXT("MoveSample", "Moving Blend Grid Sample"));		
+		BlendSpace->Modify();
+	}
+	
 	if (BlendSpace->IsValidBlendSampleIndex(SampleIndex) && BlendSpace->GetBlendSample(SampleIndex).SampleValue != NewValue && !BlendSpace->IsTooCloseToExistingSamplePoint(NewValue, SampleIndex))
 	{
-		FScopedTransaction ScopedTransaction(LOCTEXT("MoveSample", "Moving Blend Grid Sample"));
-
 		bMoveSuccessful = BlendSpace->EditSampleValue(SampleIndex, NewValue);
 		if (bMoveSuccessful)
 		{
-			BlendSpace->Modify();
 			BlendSpace->ValidateSampleData();
 			BlendSpace->PostEditChange();
 			ResampleData();
 		}
 	}
+
+	// If this was non-interactive operation, either a single change or final change at the end of an interactive operation close out the transaction	
+	if (!bIsInteractive)
+	{
+		ensure(SampleMoveTransaction.IsValid());
+		SampleMoveTransaction.Reset();
+	}
 }
 
 void SBlendSpaceEditor::OnSampleRemoved(const int32 SampleIndex)
diff --git a/Engine/Source/Editor/Persona/Private/SAnimationBlendSpace.h b/Engine/Source/Editor/Persona/Private/SAnimationBlendSpace.h
--- a/Engine/Source/Editor/Persona/Private/SAnimationBlendSpace.h
+++ b/Engine/Source/Editor/Persona/Private/SAnimationBlendSpace.h
@@ -119,4 +119,7 @@ protected:
 
 	/** Flag to check whether or not the preview value should be (re-)set on the next tick */
 	bool bShouldSetPreviewPosition = true;
+
+	/** Cross-framed boundary scoped transaction, used for sticking to single transaction for interactive operations (dragging sample or spinbox)*/
+	TUniquePtr<FScopedTransaction> SampleMoveTransaction;
 };
diff --git a/Engine/Source/Editor/Persona/Private/SAnimationBlendSpaceGridWidget.cpp b/Engine/Source/Editor/Persona/Private/SAnimationBlendSpaceGridWidget.cpp
--- a/Engine/Source/Editor/Persona/Private/SAnimationBlendSpaceGridWidget.cpp
+++ b/Engine/Source/Editor/Persona/Private/SAnimationBlendSpaceGridWidget.cpp
@@ -465,6 +465,8 @@ TSharedPtr<SWidget> SBlendSpaceGridWidget::CreateGridEntryBox(const int32 BoxInd
 		.UndeterminedString(LOCTEXT("MultipleValues", "Multiple Values"))
 		.OnValueCommitted(this, &SBlendSpaceGridWidget::OnInputBoxValueCommited, BoxIndex)
 		.OnValueChanged(this, &SBlendSpaceGridWidget::OnInputBoxValueChanged, BoxIndex, true)
+		.OnBeginSliderMovement(this, &SBlendSpaceGridWidget::OnInputSliderBegin, BoxIndex)
+		.OnEndSliderMovement(this, &SBlendSpaceGridWidget::OnInputSliderEnd, BoxIndex)
 		.LabelVAlign(VAlign_Center)
 		.AllowSpin(true)
 		.MinValue(this, &SBlendSpaceGridWidget::GetInputBoxMinValue, BoxIndex)
@@ -681,7 +683,7 @@ void SBlendSpaceGridWidget::PaintSampleKeys(
 			FVector2D GridPosition = SampleValueToScreenPosition(PreviewFilteredPosition) - (PreviewSize * .5f);
 			FSlateDrawElement::MakeBox(
 				OutDrawElements, FilteredPositionLayer, AllottedGeometry.ToPaintGeometry(GridPosition, PreviewSize), 
-				PreviewBrush, ESlateDrawEffect::None, PreviewKeyColor.GetSpecifiedColor() * 0.7);
+				PreviewBrush, ESlateDrawEffect::None, PreviewKeyColor.GetSpecifiedColor() * 0.7f);
 		}
 
 		// Always draw the preview position
@@ -1187,6 +1189,7 @@ FReply SBlendSpaceGridWidget::OnMouseButtonUp(const FGeometry& MyGeometry, const
 			{
 				// Process drag ending			
 				ResetToolTip();
+				OnSampleMoved.ExecuteIfBound(DraggedSampleIndex, LastDragPosition, false);
 			}
 
 			// Reset drag state and index
@@ -2246,6 +2249,12 @@ void SBlendSpaceGridWidget::OnInputBoxValueCommited(const float NewValue, ETextC
 
 void SBlendSpaceGridWidget::OnInputBoxValueChanged(const float NewValue, const int32 ParameterIndex, bool bIsInteractive)
 {
+	// Ignore any SNumericEntryBox.OnValueChanged broadcasts if sliding has finished and OnInputBoxValueCommited will have been broadcasted already
+	if (bIsInteractive && !bSliderMovement[ParameterIndex])
+	{
+		return;
+	}
+	
 	checkf(ParameterIndex < 2, TEXT("Invalid parameter index, suppose to be within FVector array range"));
 
 	if (SelectedSampleIndex != INDEX_NONE && BlendSpaceBase.Get() != nullptr)
@@ -2275,6 +2284,18 @@ void SBlendSpaceGridWidget::OnInputBoxValueChanged(const float NewValue, const i
 	}
 }
 
+void SBlendSpaceGridWidget::OnInputSliderBegin(const int32 ParameterIndex)
+{
+	ensure(bSliderMovement[ParameterIndex] == false);
+	bSliderMovement[ParameterIndex] = true;
+}
+
+void SBlendSpaceGridWidget::OnInputSliderEnd(const float NewValue, const int32 ParameterIndex)
+{
+	ensure(bSliderMovement[ParameterIndex] == true);
+	bSliderMovement[ParameterIndex] = false;
+}
+
 EVisibility SBlendSpaceGridWidget::GetSampleToolTipVisibility() const
 {
 	// Show tool tip when the grid is empty
@@ -2415,6 +2436,11 @@ void SBlendSpaceGridWidget::OnMouseLeave(const FPointerEvent& MouseEvent)
 void SBlendSpaceGridWidget::OnFocusLost(const FFocusEvent& InFocusEvent)
 {
 	SCompoundWidget::OnFocusLost(InFocusEvent);
+
+	if (DragState == EDragState::DragSample)
+	{
+		OnSampleMoved.ExecuteIfBound(DraggedSampleIndex, LastDragPosition, false);
+	}	
 	HighlightedSampleIndex = DraggedSampleIndex = INDEX_NONE;
 	DragState = EDragState::None;
 	bSamplePreviewing = false;
@@ -2495,7 +2521,7 @@ void SBlendSpaceGridWidget::Tick(const FGeometry& AllottedGeometry, const double
 				if (SampleValue != LastDragPosition)
 				{
 					LastDragPosition = SampleValue;
-					OnSampleMoved.ExecuteIfBound(DraggedSampleIndex, SampleValue, false);
+					OnSampleMoved.ExecuteIfBound(DraggedSampleIndex, SampleValue, true);
 				}
 			}
 			else if (DragState == EDragState::DragDrop || DragState == EDragState::InvalidDragDrop || DragState == EDragState::DragDropOverride)
diff --git a/Engine/Source/Editor/Persona/Private/SAnimationBlendSpaceGridWidget.h b/Engine/Source/Editor/Persona/Private/SAnimationBlendSpaceGridWidget.h
--- a/Engine/Source/Editor/Persona/Private/SAnimationBlendSpaceGridWidget.h
+++ b/Engine/Source/Editor/Persona/Private/SAnimationBlendSpaceGridWidget.h
@@ -212,6 +212,8 @@ protected:
 	float GetInputBoxDelta(const int32 ParameterIndex) const;
 	void OnInputBoxValueCommited(const float NewValue, ETextCommit::Type CommitType, const int32 ParameterIndex);
 	void OnInputBoxValueChanged(const float NewValue, const int32 ParameterIndex, bool bIsInteractive);
+	void OnInputSliderBegin(const int32 ParameterIndex);
+	void OnInputSliderEnd(const float NewValue, const int32 ParameterIndex);
 
 	/** Returns whether or not the sample tool tips should be visible */
 	EVisibility GetSampleToolTipVisibility() const;
@@ -291,6 +293,7 @@ private:
 
 	/** Cached values for the grid input boxes */
 	float CachedInputBoxValues[2];
+	bool bSliderMovement[2] = {false, false};
 
 	// Cached grid data (derived from the blend space)
 	FSlateRect CachedGridRectangle;
