Fix for undo bug caused by CL 16356830

diff --git a/Engine/Source/Editor/Persona/Private/SAnimationBlendSpace.cpp b/Engine/Source/Editor/Persona/Private/SAnimationBlendSpace.cpp
--- a/Engine/Source/Editor/Persona/Private/SAnimationBlendSpace.cpp
+++ b/Engine/Source/Editor/Persona/Private/SAnimationBlendSpace.cpp
@@ -119,7 +119,7 @@ void SBlendSpaceEditor::OnSampleMoved(const int32 SampleIndex, const FVector& Ne
 	// If this is an interactive operation and a transaction has not been opened yet, or if it's a non-interactive operation setup a transaction
 	if (!bIsInteractive || !SampleMoveTransaction.IsValid())
 	{
-		SampleMoveTransaction = MakeUnique<FScopedTransaction>(LOCTEXT("MoveSample", "Moving Blend Grid Sample"));		
+		SampleMoveTransaction = MakeUnique<FScopedTransaction>(LOCTEXT("MoveSample", "Moving Blendspace Sample"));		
 		BlendSpace->Modify();
 	}
 	
@@ -144,12 +144,12 @@ void SBlendSpaceEditor::OnSampleMoved(const int32 SampleIndex, const FVector& Ne
 
 void SBlendSpaceEditor::OnSampleRemoved(const int32 SampleIndex)
 {
-	FScopedTransaction ScopedTransaction(LOCTEXT("RemoveSample", "Removing Blend Grid Sample"));
+	FScopedTransaction ScopedTransaction(LOCTEXT("RemoveSample", "Removing Blendspace Sample"));
+	BlendSpace->Modify();
 
 	const bool bRemoveSuccessful = BlendSpace->DeleteSample(SampleIndex);
 	if (bRemoveSuccessful)
 	{
-		BlendSpace->Modify();
 		ResampleData();
 		BlendSpace->ValidateSampleData();
 
@@ -164,7 +164,8 @@ void SBlendSpaceEditor::OnSampleRemoved(const int32 SampleIndex)
 //======================================================================================================================
 int32 SBlendSpaceEditor::OnSampleAdded(UAnimSequence* Animation, const FVector& Value, bool bRunAnalysis)
 {
-	FScopedTransaction ScopedTransaction(LOCTEXT("AddSample", "Adding Blend Grid Sample"));
+	FScopedTransaction ScopedTransaction(LOCTEXT("AddSample", "Adding Blendspace Sample"));
+	BlendSpace->Modify();
 
 	FVector AdjustedValue = Value;
 	bool bAnalyzed[3] = { false, false, false };
@@ -187,7 +188,6 @@ int32 SBlendSpaceEditor::OnSampleAdded(UAnimSequence* Animation, const FVector&
 
 	if (NewSampleIndex >= 0)
 	{
-		BlendSpace->Modify();
 		BlendSpace->LockSample(
 			NewSampleIndex,
 			bAnalyzed[0] ? BlendSpaceAnalysis::GetLockAfterAnalysis(BlendSpace->AnalysisProperties[0]) : false,
@@ -212,6 +212,9 @@ int32 SBlendSpaceEditor::OnSampleAdded(UAnimSequence* Animation, const FVector&
 
 void SBlendSpaceEditor::OnSampleDuplicated(const int32 SampleIndex, const FVector& NewValue, bool bRunAnalysis)
 {
+	FScopedTransaction ScopedTransaction(LOCTEXT("DuplicateSample", "Duplicating Blendspace Sample"));
+	BlendSpace->Modify();
+
 	const FBlendSample& OrigSample = BlendSpace->GetBlendSample(SampleIndex);
 	int32 NewSampleIndex = OnSampleAdded(OrigSample.Animation, NewValue, bRunAnalysis);
 	if (NewSampleIndex >= 0)
@@ -222,7 +225,8 @@ void SBlendSpaceEditor::OnSampleDuplicated(const int32 SampleIndex, const FVecto
 
 void SBlendSpaceEditor::OnSampleReplaced(int32 InSampleIndex, UAnimSequence* Animation)
 {
-	FScopedTransaction ScopedTransaction(LOCTEXT("UpdateAnimation", "Changing Animation Sequence"));
+	FScopedTransaction ScopedTransaction(LOCTEXT("ReplaceSample", "Replacing Blendspace Sample"));
+	BlendSpace->Modify();
 
 	bool bUpdateSuccessful = false;
 	if(BlendSpace->IsAsset())
@@ -236,7 +240,6 @@ void SBlendSpaceEditor::OnSampleReplaced(int32 InSampleIndex, UAnimSequence* Ani
 
 	if (bUpdateSuccessful)
 	{
-		BlendSpace->Modify();
 		ResampleData();
 		BlendSpace->ValidateSampleData();
 		OnBlendSpaceSampleReplaced.ExecuteIfBound(InSampleIndex, Animation);
