Fix for Blend Space Analysis settings are not reverted after undoing

diff --git a/Engine/Source/Editor/Persona/Private/Customization/BlendSpaceDetails.cpp b/Engine/Source/Editor/Persona/Private/Customization/BlendSpaceDetails.cpp
--- a/Engine/Source/Editor/Persona/Private/Customization/BlendSpaceDetails.cpp
+++ b/Engine/Source/Editor/Persona/Private/Customization/BlendSpaceDetails.cpp
@@ -130,7 +130,6 @@ static EVisibility GetAnalyzeButtonVisibility(
 // that if "None" was selected then the new value will be null
 void FBlendSpaceDetails::HandleAnalysisFunctionChanged(int32 AxisIndex, TSharedPtr<FString> NewFunctionName)
 {
-	BlendSpace->Modify();
 	UAnalysisProperties* NewAnalysisProperties = BlendSpaceAnalysis::MakeAnalysisProperties(BlendSpace, *NewFunctionName);
 	// Preserve values where possible
 	if (BlendSpace->AnalysisProperties[AxisIndex])
diff --git a/Engine/Source/Editor/Persona/Private/PersonaBlendSpaceAnalysis.cpp b/Engine/Source/Editor/Persona/Private/PersonaBlendSpaceAnalysis.cpp
--- a/Engine/Source/Editor/Persona/Private/PersonaBlendSpaceAnalysis.cpp
+++ b/Engine/Source/Editor/Persona/Private/PersonaBlendSpaceAnalysis.cpp
@@ -330,6 +330,8 @@ UAnalysisProperties* BlendSpaceAnalysis::MakeAnalysisProperties(UObject* Outer,
 		Result = Feature->MakeAnalysisProperties(Outer, FunctionName);
 		if (Result)
 		{
+			// Need to explicitly set flags to make undo work on the new object
+			Result->SetFlags(RF_Transactional);
 			return Result;
 		}
 	}
diff --git a/Engine/Source/Runtime/Engine/Classes/Animation/BlendSpace.h b/Engine/Source/Runtime/Engine/Classes/Animation/BlendSpace.h
--- a/Engine/Source/Runtime/Engine/Classes/Animation/BlendSpace.h
+++ b/Engine/Source/Runtime/Engine/Classes/Animation/BlendSpace.h
@@ -715,7 +715,10 @@ public:
 	FInterpolationParameter	InterpolationParam[3];
 
 #if WITH_EDITORONLY_DATA
-	/** Analysis properties for each axis. Note that these can be null. */
+	/** 
+	* Analysis properties for each axis. Note that these can be null. They will be created/set from 
+	* FBlendSpaceDetails::HandleAnalysisFunctionChanged
+	*/
 	UPROPERTY(EditAnywhere, Category = AnalysisProperties)
 	TObjectPtr<UAnalysisProperties> AnalysisProperties[3];
 
