Persona: Guard against null ptr in anim editor preview scene

diff --git a/Engine/Source/Editor/Persona/Private/AnimationEditorPreviewScene.cpp b/Engine/Source/Editor/Persona/Private/AnimationEditorPreviewScene.cpp
--- a/Engine/Source/Editor/Persona/Private/AnimationEditorPreviewScene.cpp
+++ b/Engine/Source/Editor/Persona/Private/AnimationEditorPreviewScene.cpp
@@ -211,13 +211,16 @@ void FAnimationEditorPreviewScene::SetPreviewMeshInternal(USkeletalMesh* NewPrev
 	USkeletalMeshComponent* DebuggedSkeletalMeshComponent = nullptr;
 	if(SkeletalMeshComponent->GetAnimInstance())
 	{
-		UAnimBlueprint* SourceBlueprint = PersonaToolkit.Pin()->GetAnimBlueprint();
-		if(SourceBlueprint)
+		if(PersonaToolkit.IsValid())
 		{
-			UAnimInstance* DebuggedAnimInstance = Cast<UAnimInstance>(SourceBlueprint->GetObjectBeingDebugged());
-			if(DebuggedAnimInstance)
+			UAnimBlueprint* SourceBlueprint = PersonaToolkit.Pin()->GetAnimBlueprint();
+			if(SourceBlueprint)
 			{
-				DebuggedSkeletalMeshComponent = DebuggedAnimInstance->GetSkelMeshComponent();
+				UAnimInstance* DebuggedAnimInstance = Cast<UAnimInstance>(SourceBlueprint->GetObjectBeingDebugged());
+				if(DebuggedAnimInstance)
+				{
+					DebuggedSkeletalMeshComponent = DebuggedAnimInstance->GetSkelMeshComponent();
+				}
 			}
 		}
 	}
@@ -401,17 +404,20 @@ void FAnimationEditorPreviewScene::RefreshAdditionalMeshes(bool bAllowOverrideBa
 
 void FAnimationEditorPreviewScene::AddPreviewAttachedObjects()
 {
-	// Load up mesh attachments...
-	USkeletalMesh* Mesh = PersonaToolkit.Pin()->GetMesh();
-
-	if ( Mesh )
+	if(PersonaToolkit.IsValid())
 	{
-		FPreviewAssetAttachContainer& PreviewAssetAttachContainer = Mesh->GetPreviewAttachedAssetContainer();
-		for(int32 Index = 0; Index < PreviewAssetAttachContainer.Num(); Index++)
+		// Load up mesh attachments...
+		USkeletalMesh* Mesh = PersonaToolkit.Pin()->GetMesh();
+
+		if ( Mesh )
 		{
-			FPreviewAttachedObjectPair& PreviewAttachedObject = PreviewAssetAttachContainer[Index];
+			FPreviewAssetAttachContainer& PreviewAssetAttachContainer = Mesh->GetPreviewAttachedAssetContainer();
+			for(int32 Index = 0; Index < PreviewAssetAttachContainer.Num(); Index++)
+			{
+				FPreviewAttachedObjectPair& PreviewAttachedObject = PreviewAssetAttachContainer[Index];
 
-			AttachObjectToPreviewComponent(PreviewAttachedObject.GetAttachedObject(), PreviewAttachedObject.AttachedTo);
+				AttachObjectToPreviewComponent(PreviewAttachedObject.GetAttachedObject(), PreviewAttachedObject.AttachedTo);
+			}
 		}
 	}
 
