Remove sample connection drawing in blendspace when using the grid

diff --git a/Engine/Source/Editor/Persona/Private/SAnimationBlendSpaceGridWidget.cpp b/Engine/Source/Editor/Persona/Private/SAnimationBlendSpaceGridWidget.cpp
--- a/Engine/Source/Editor/Persona/Private/SAnimationBlendSpaceGridWidget.cpp
+++ b/Engine/Source/Editor/Persona/Private/SAnimationBlendSpaceGridWidget.cpp
@@ -218,7 +218,7 @@ void SBlendSpaceGridWidget::Construct(const FArguments& InArgs)
 
 	PreviewFilteredPosition = PreviewPosition;
 
-	bShowTriangulation = (BlendSpaceBase.Get() != nullptr && BlendSpaceBase.Get()->bInterpolateUsingGrid) ? false : true;
+	bShowTriangulation = true;
 	bMouseIsOverGeometry = false;
 	bRefreshCachedData = true;
 	bStretchToFit = true;
@@ -299,12 +299,7 @@ void SBlendSpaceGridWidget::Construct(const FArguments& InArgs)
 								.VAlign(VAlign_Center)
 								[
 									SNew(SButton)
-									.ToolTipText_Lambda([this]()
-														{ 
-															return (BlendSpaceBase.Get() && BlendSpaceBase.Get()->bInterpolateUsingGrid) 
-																? LOCTEXT("ShowGridToSampleConnections", "Show Grid/Sample Connections")
-																: LOCTEXT("ShowTriangulation", "Show Triangulation");
-														})
+									.ToolTipText(LOCTEXT("ShowTriangulation", "Show Triangulation"))
 									.OnClicked(this, &SBlendSpaceGridWidget::ToggleTriangulationVisibility)
 									.ButtonColorAndOpacity_Lambda([this]() -> FLinearColor { return bShowTriangulation ? FEditorStyle::GetSlateColor("SelectionColor").GetSpecifiedColor() : FLinearColor::White; })
 									.ContentPadding(1)
@@ -843,189 +838,155 @@ void SBlendSpaceGridWidget::PaintTriangulation(
 	FSlateWindowElementList& OutDrawElements, 
 	int32&                   DrawLayerId) const
 {
-	if (const UBlendSpace* BlendSpace = BlendSpaceBase.Get())
+	const UBlendSpace* BlendSpace = BlendSpaceBase.Get();
+	if (!BlendSpace)
 	{
-		if ( 
-			(bReadOnly && BlendSpace->bInterpolateUsingGrid) ||
-			(!bReadOnly && !bShowTriangulation)
-			)
-		{
-			return;
-		}
+		return;
+	}
+	if ((bReadOnly && BlendSpace->bInterpolateUsingGrid) || (!bReadOnly && !bShowTriangulation))
+	{
+		return;
+	}
 
-		TArray<FVector2D> PolygonPoints; 
-		const TArray<FBlendSample>& Samples = BlendSpace->GetBlendSamples();
-		if (BlendSpace->bInterpolateUsingGrid)
+	TArray<FVector2D> PolygonPoints; 
+	const TArray<FBlendSample>& Samples = BlendSpace->GetBlendSamples();
+	if (!BlendSpace->bInterpolateUsingGrid)
+	{
+		// Use runtime triangulation
+		const FBlendSpaceData& BlendSpaceData = BlendSpace->GetBlendSpaceData();
+		for (const FBlendSpaceSegment& Segment : BlendSpaceData.Segments)
 		{
-			// Use the Grid. After preprocessing, each grid point will have been located within the triangulation, and
-			// will store the indices of the three samples that surround/contribute to it.
-			const TArray<FEditorElement>& EditorElements = BlendSpace->GetGridSamples();
-			for (int32 ElementIndex = 0 ; ElementIndex != EditorElements.Num() ; ++ElementIndex)
-			{
-				const FEditorElement& Element = EditorElements[ElementIndex];
-				FVector ElementPosition = BlendSpace->GetGridPosition(ElementIndex);
-				FVector2D ElementScreenPosition = SampleValueToScreenPosition(ElementPosition);
-				for (int32 SourceIndex = 0; SourceIndex < 3; ++SourceIndex)
-				{
-					const int32 SourceSampleIndex = Element.Indices[SourceIndex]; 
-					if (Samples.IsValidIndex(SourceSampleIndex))
-					{
-						const FBlendSample& SourceSample = Samples[SourceSampleIndex];
-						FVector2D SourceSampleScreenPosition = SampleValueToScreenPosition(SourceSample.SampleValue);
-						// Draw line from the grid point to each sample that contributes to it
-						if (Element.Weights[SourceIndex] > 0)
-						{
-							if (ElementScreenPosition != SourceSampleScreenPosition)
-							{
-								FSlateDrawElement::MakeLines(
-									OutDrawElements, DrawLayerId + 1, AllottedGeometry.ToPaintGeometry(), 
-									{ ElementScreenPosition, SourceSampleScreenPosition },
-									ESlateDrawEffect::None,  TriangulationColor.GetSpecifiedColor(), 
-									true, Element.Weights[SourceIndex]);
-							}
-						}
-					}
-				}
-			}
+			int32 SampleIndex = Segment.SampleIndices[0];
+			int32 SampleIndex1 = Segment.SampleIndices[1];
+			TArray<FVector2D> Points;
+			Points.Add(SampleValueToScreenPosition(Samples[SampleIndex].SampleValue));
+			Points.Add(SampleValueToScreenPosition(Samples[SampleIndex1].SampleValue));
+			FSlateDrawElement::MakeLines(
+				OutDrawElements, DrawLayerId + 1, AllottedGeometry.ToPaintGeometry(), Points,
+				ESlateDrawEffect::None, TriangulationColor.GetSpecifiedColor(), true, 0.5f);
 		}
-		else
+
+		const float CriticalDot = FMath::Cos(FMath::DegreesToRadians(CriticalTriangulationAngle));
+		for (const FBlendSpaceTriangle& Triangle : BlendSpaceData.Triangles)
 		{
-			// Use runtime triangulation
-			const FBlendSpaceData& BlendSpaceData = BlendSpace->GetBlendSpaceData();
-			for (const FBlendSpaceSegment& Segment : BlendSpaceData.Segments)
-			{
-				int32 SampleIndex = Segment.SampleIndices[0];
-				int32 SampleIndex1 = Segment.SampleIndices[1];
-				TArray<FVector2D> Points;
-				Points.Add(SampleValueToScreenPosition(Samples[SampleIndex].SampleValue));
-				Points.Add(SampleValueToScreenPosition(Samples[SampleIndex1].SampleValue));
-				FSlateDrawElement::MakeLines(
-					OutDrawElements, DrawLayerId + 1, AllottedGeometry.ToPaintGeometry(), Points,
-					ESlateDrawEffect::None, TriangulationColor.GetSpecifiedColor(), true, 0.5f);
-			}
+			FLinearColor TriangleFillColor = TriangulationCurrentColor.GetSpecifiedColor();
+			FLinearColor TriangleLineColor = TriangulationColor.GetSpecifiedColor();
+			TriangleFillColor.A = 0.03f; // Alpha for tinting the triangulation background
+			int32 TriangleLayer = DrawLayerId + 1;
+
+			FVector2D ScreenPositions[3] = {
+				SampleValueToScreenPosition(Samples[Triangle.SampleIndices[0]].SampleValue),
+				SampleValueToScreenPosition(Samples[Triangle.SampleIndices[1]].SampleValue),
+				SampleValueToScreenPosition(Samples[Triangle.SampleIndices[2]].SampleValue),
+			};
 
-			const float CriticalDot = FMath::Cos(FMath::DegreesToRadians(CriticalTriangulationAngle));
-			for (const FBlendSpaceTriangle& Triangle : BlendSpaceData.Triangles)
+			// If we just have one triangle it's OK for it to be degenerate
+			if (BlendSpaceData.Triangles.Num() > 1)
 			{
-				FLinearColor TriangleFillColor = TriangulationCurrentColor.GetSpecifiedColor();
-				FLinearColor TriangleLineColor = TriangulationColor.GetSpecifiedColor();
-				TriangleFillColor.A = 0.03f; // Alpha for tinting the triangulation background
-				int32 TriangleLayer = DrawLayerId + 1;
-
-				FVector2D ScreenPositions[3] = {
-					SampleValueToScreenPosition(Samples[Triangle.SampleIndices[0]].SampleValue),
-					SampleValueToScreenPosition(Samples[Triangle.SampleIndices[1]].SampleValue),
-					SampleValueToScreenPosition(Samples[Triangle.SampleIndices[2]].SampleValue),
+				FVector2D NormalizedPositions[3] = {
+					SampleValueToNormalizedPosition(Samples[Triangle.SampleIndices[0]].SampleValue),
+					SampleValueToNormalizedPosition(Samples[Triangle.SampleIndices[1]].SampleValue),
+					SampleValueToNormalizedPosition(Samples[Triangle.SampleIndices[2]].SampleValue),
 				};
 
-				// If we just have one triangle it's OK for it to be degenerate
-				if (BlendSpaceData.Triangles.Num() > 1)
+				for (int32 Index = 0; Index != 3; ++Index)
 				{
-					FVector2D NormalizedPositions[3] = {
-						SampleValueToNormalizedPosition(Samples[Triangle.SampleIndices[0]].SampleValue),
-						SampleValueToNormalizedPosition(Samples[Triangle.SampleIndices[1]].SampleValue),
-						SampleValueToNormalizedPosition(Samples[Triangle.SampleIndices[2]].SampleValue),
-					};
-
-					for (int32 Index = 0; Index != 3; ++Index)
+					FVector2D A = NormalizedPositions[(Index + 2) % 3] - NormalizedPositions[(Index + 1) % 3];
+					FVector2D B = NormalizedPositions[Index] - NormalizedPositions[(Index + 1) % 3];
+					float Dot = A.GetSafeNormal() | B.GetSafeNormal();
+					float Area = 0.5f * FMath::Abs(B ^ A);
+					if (Dot > CriticalDot || Area < CriticalTriangulationArea)
 					{
-						FVector2D A = NormalizedPositions[(Index + 2) % 3] - NormalizedPositions[(Index + 1) % 3];
-						FVector2D B = NormalizedPositions[Index] - NormalizedPositions[(Index + 1) % 3];
-						float Dot = A.GetSafeNormal() | B.GetSafeNormal();
-						float Area = 0.5f * FMath::Abs(B ^ A);
-						if (Dot > CriticalDot || Area < CriticalTriangulationArea)
-						{
-							TriangleFillColor = FLinearColor::Red;
-							TriangleFillColor.A = 0.5f;
-							TriangleLineColor = FLinearColor::Red;
-							TriangleLayer = DrawLayerId + 10; // just bump it up so it definitely shows
-						}
+						TriangleFillColor = FLinearColor::Red;
+						TriangleFillColor.A = 0.5f;
+						TriangleLineColor = FLinearColor::Red;
+						TriangleLayer = DrawLayerId + 10; // just bump it up so it definitely shows
 					}
 				}
+			}
 
-				FVector2D MidPoint(0, 0);
-				PolygonPoints.Empty(3);
-				for (int32 Index0 = 0; Index0 != FBlendSpaceTriangle::NUM_VERTICES; ++Index0)
-				{
-					int32 Index1 = (Index0 + 1) % FBlendSpaceTriangle::NUM_VERTICES;
-					int32 Index2 = (Index0 + 2) % FBlendSpaceTriangle::NUM_VERTICES;
-					int32 SampleIndex0 = Triangle.SampleIndices[Index0];
-					int32 SampleIndex1 = Triangle.SampleIndices[Index1];
-					int32 SampleIndex2 = Triangle.SampleIndices[Index2];
+			FVector2D MidPoint(0, 0);
+			PolygonPoints.Empty(3);
+			for (int32 Index0 = 0; Index0 != FBlendSpaceTriangle::NUM_VERTICES; ++Index0)
+			{
+				int32 Index1 = (Index0 + 1) % FBlendSpaceTriangle::NUM_VERTICES;
+				int32 Index2 = (Index0 + 2) % FBlendSpaceTriangle::NUM_VERTICES;
+				int32 SampleIndex0 = Triangle.SampleIndices[Index0];
+				int32 SampleIndex1 = Triangle.SampleIndices[Index1];
+				int32 SampleIndex2 = Triangle.SampleIndices[Index2];
 
-					TArray<FVector2D> Points = { ScreenPositions[Index0], ScreenPositions[Index1] };
-					MidPoint += SampleValueToScreenPosition(Samples[SampleIndex0].SampleValue) / 3.0f;
+				TArray<FVector2D> Points = { ScreenPositions[Index0], ScreenPositions[Index1] };
+				MidPoint += SampleValueToScreenPosition(Samples[SampleIndex0].SampleValue) / 3.0f;
 
-					FSlateDrawElement::MakeLines(
-						OutDrawElements, TriangleLayer, AllottedGeometry.ToPaintGeometry(), Points,
-						ESlateDrawEffect::None, TriangleLineColor, true, 0.5f);
-					PolygonPoints.Push(ScreenPositions[Index0]);
-				}
+				FSlateDrawElement::MakeLines(
+					OutDrawElements, TriangleLayer, AllottedGeometry.ToPaintGeometry(), Points,
+					ESlateDrawEffect::None, TriangleLineColor, true, 0.5f);
+				PolygonPoints.Push(ScreenPositions[Index0]);
+			}
 
-				PaintTriangle(PolygonPoints[0], PolygonPoints[1], PolygonPoints[2], AllottedGeometry, 
-							  TriangleFillColor, LabelBrush, OutDrawElements, DrawLayerId);
+			PaintTriangle(PolygonPoints[0], PolygonPoints[1], PolygonPoints[2], AllottedGeometry, 
+							TriangleFillColor, LabelBrush, OutDrawElements, DrawLayerId);
 
 #ifdef DEBUG_BLENDSPACE_TRIANGULATION
-				// Draw the adjacent triangle indices around the perimeter
-				if (!bSamplePreviewing)
+			// Draw the adjacent triangle indices around the perimeter
+			if (!bSamplePreviewing)
+			{
+				for (int32 Index0 = 0; Index0 != FBlendSpaceTriangle::NUM_VERTICES; ++Index0)
 				{
-					for (int32 Index0 = 0; Index0 != FBlendSpaceTriangle::NUM_VERTICES; ++Index0)
+					if (Triangle.EdgeInfo[Index0].NeighbourTriangleIndex < 0)
 					{
-						if (Triangle.EdgeInfo[Index0].NeighbourTriangleIndex < 0)
-						{
-							int32 Index1 = (Index0 + 1) % FBlendSpaceTriangle::NUM_VERTICES;
-							FVector2D MidEdge = ( ScreenPositions[Index0] + ScreenPositions[Index1]) * 0.5;
-							float PullInAmount = 0.2;
-							FSlateDrawElement::MakeText(
-								OutDrawElements, DrawLayerId + 1, AllottedGeometry.MakeChild(
-									FMath::Lerp(ScreenPositions[Index0], MidEdge, PullInAmount), FVector2D(1.0f, 1.0f)).ToPaintGeometry(),
-								FText::AsNumber(Triangle.EdgeInfo[Index0].AdjacentPerimeterTriangleIndices[0]),
-								FontInfo, ESlateDrawEffect::None, FLinearColor::Red);
-
-							FSlateDrawElement::MakeText(
-								OutDrawElements, DrawLayerId + 1, AllottedGeometry.MakeChild(
-									FMath::Lerp(ScreenPositions[Index1], MidEdge, PullInAmount), FVector2D(1.0f, 1.0f)).ToPaintGeometry(),
-								FText::AsNumber(Triangle.EdgeInfo[Index0].AdjacentPerimeterTriangleIndices[1]),
-								FontInfo, ESlateDrawEffect::None, FLinearColor::Red);
-						}
+						int32 Index1 = (Index0 + 1) % FBlendSpaceTriangle::NUM_VERTICES;
+						FVector2D MidEdge = ( ScreenPositions[Index0] + ScreenPositions[Index1]) * 0.5;
+						float PullInAmount = 0.2;
+						FSlateDrawElement::MakeText(
+							OutDrawElements, DrawLayerId + 1, AllottedGeometry.MakeChild(
+								FMath::Lerp(ScreenPositions[Index0], MidEdge, PullInAmount), FVector2D(1.0f, 1.0f)).ToPaintGeometry(),
+							FText::AsNumber(Triangle.EdgeInfo[Index0].AdjacentPerimeterTriangleIndices[0]),
+							FontInfo, ESlateDrawEffect::None, FLinearColor::Red);
+
+						FSlateDrawElement::MakeText(
+							OutDrawElements, DrawLayerId + 1, AllottedGeometry.MakeChild(
+								FMath::Lerp(ScreenPositions[Index1], MidEdge, PullInAmount), FVector2D(1.0f, 1.0f)).ToPaintGeometry(),
+							FText::AsNumber(Triangle.EdgeInfo[Index0].AdjacentPerimeterTriangleIndices[1]),
+							FontInfo, ESlateDrawEffect::None, FLinearColor::Red);
 					}
 				}
+			}
 #endif
 
 #ifdef DEBUG_BLENDSPACE_TRIANGULATION
-				// Draw the triangle indices for debugging
-				FText Text = FText::AsNumber(&Triangle - &BlendSpaceData.Triangles[0]);
-				FSlateDrawElement::MakeText(
-					OutDrawElements, DrawLayerId + 1, AllottedGeometry.MakeChild(
-						FVector2D(MidPoint.X, MidPoint.Y), FVector2D(1.0f, 1.0f)).ToPaintGeometry(),
-					Text, FontInfo, ESlateDrawEffect::None, FLinearColor::Gray);
+			// Draw the triangle indices for debugging
+			FText Text = FText::AsNumber(&Triangle - &BlendSpaceData.Triangles[0]);
+			FSlateDrawElement::MakeText(
+				OutDrawElements, DrawLayerId + 1, AllottedGeometry.MakeChild(
+					FVector2D(MidPoint.X, MidPoint.Y), FVector2D(1.0f, 1.0f)).ToPaintGeometry(),
+				Text, FontInfo, ESlateDrawEffect::None, FLinearColor::Gray);
 #endif
-			}
 		}
+	}
 
-		// Draw the current triangle (or polygon)
-		if (bSamplePreviewing && FSlateApplication::Get().GetModifierKeys().IsAltDown())
+	// Draw the current triangle (or polygon)
+	if (bSamplePreviewing && FSlateApplication::Get().GetModifierKeys().IsAltDown())
+	{
+		PolygonPoints.Empty(3);
+		for (const FBlendSampleData& PreviewedSample : PreviewedSamples)
 		{
-			PolygonPoints.Empty(3);
-			for (const FBlendSampleData& PreviewedSample : PreviewedSamples)
+			float Weight = PreviewedSample.TotalWeight;
+			if (Weight)
 			{
-				float Weight = PreviewedSample.TotalWeight;
-				if (Weight)
-				{
-					int32 SampleIndex = PreviewedSample.SampleDataIndex;
-					FVector2D Point = SampleValueToScreenPosition(Samples[SampleIndex].SampleValue);
-					PolygonPoints.Push(Point);
-				}
-			}
-			if (PolygonPoints.Num())
-			{
-				FLinearColor FillColor = TriangulationCurrentColor.GetSpecifiedColor();
-				FillColor.A = 0.2f; // Alpha for the current triangulation triangle
-				FLinearColor OutlineColor = FillColor;
-				OutlineColor.A = 0.5f;
-				PaintPolygon(PolygonPoints, AllottedGeometry, FillColor, OutlineColor, LabelBrush, OutDrawElements, DrawLayerId);
+				int32 SampleIndex = PreviewedSample.SampleDataIndex;
+				FVector2D Point = SampleValueToScreenPosition(Samples[SampleIndex].SampleValue);
+				PolygonPoints.Push(Point);
 			}
 		}
+		if (PolygonPoints.Num())
+		{
+			FLinearColor FillColor = TriangulationCurrentColor.GetSpecifiedColor();
+			FillColor.A = 0.2f; // Alpha for the current triangulation triangle
+			FLinearColor OutlineColor = FillColor;
+			OutlineColor.A = 0.5f;
+			PaintPolygon(PolygonPoints, AllottedGeometry, FillColor, OutlineColor, LabelBrush, OutDrawElements, DrawLayerId);
+		}
 	}
 
 	DrawLayerId += 1;
@@ -1478,33 +1439,36 @@ void SBlendSpaceGridWidget::MakeViewContextMenuEntries(FMenuBuilder& InMenuBuild
 {
 	InMenuBuilder.BeginSection("ViewOptions", LOCTEXT("ViewOptionsMenuHeader", "View Options"));
 	{
-		TAttribute<FText> ShowTriangulation = TAttribute<FText>::Create(
-			[this]()
-			{
-				return (BlendSpaceBase.Get() && BlendSpaceBase.Get()->bInterpolateUsingGrid) 
-					? LOCTEXT("ShowGridToSampleConnections", "Show Grid/Sample Connections")
-					: LOCTEXT("ShowTriangulation", "Show Triangulation");
-			});
-		TAttribute<FText> ShowTriangulationToolTip = TAttribute<FText>::Create(
-			[this]()
-			{
-				return (BlendSpaceBase.Get() && BlendSpaceBase.Get()->bInterpolateUsingGrid) 
-					? LOCTEXT("ShowGridToSampleConnectionsToolTip", "Show which samples each grid point is associated with")
-					: LOCTEXT("ShowTriangulationToolTip", "Show the Delaunay triangulation for all blend space samples");
-			});
-
-		InMenuBuilder.AddMenuEntry(
-			ShowTriangulation,
-			ShowTriangulationToolTip,
-			FSlateIcon("EditorStyle", "BlendSpaceEditor.ToggleTriangulation"),
-			FUIAction(
-				FExecuteAction::CreateLambda([this](){ bShowTriangulation = !bShowTriangulation; }),
-				FCanExecuteAction(),
-				FGetActionCheckState::CreateLambda([this](){ return bShowTriangulation ? ECheckBoxState::Checked : ECheckBoxState::Unchecked; })
-			),
-			NAME_None,
-			EUserInterfaceActionType::ToggleButton
-		);
+		if (GetTriangulationButtonVisibility() == EVisibility::Visible)
+		{
+			TAttribute<FText> ShowTriangulation = TAttribute<FText>::Create(
+				[this]()
+				{
+					return (BlendSpaceBase.Get() && BlendSpaceBase.Get()->bInterpolateUsingGrid) 
+						? LOCTEXT("ShowGridToSampleConnections", "Show Grid/Sample Connections")
+						: LOCTEXT("ShowTriangulation", "Show Triangulation");
+				});
+			TAttribute<FText> ShowTriangulationToolTip = TAttribute<FText>::Create(
+				[this]()
+				{
+					return (BlendSpaceBase.Get() && BlendSpaceBase.Get()->bInterpolateUsingGrid) 
+						? LOCTEXT("ShowGridToSampleConnectionsToolTip", "Show which samples each grid point is associated with")
+						: LOCTEXT("ShowTriangulationToolTip", "Show the Delaunay triangulation for all blend space samples");
+				});
+
+			InMenuBuilder.AddMenuEntry(
+				ShowTriangulation,
+				ShowTriangulationToolTip,
+				FSlateIcon("EditorStyle", "BlendSpaceEditor.ToggleTriangulation"),
+				FUIAction(
+					FExecuteAction::CreateLambda([this](){ bShowTriangulation = !bShowTriangulation; }),
+					FCanExecuteAction(),
+					FGetActionCheckState::CreateLambda([this](){ return bShowTriangulation ? ECheckBoxState::Checked : ECheckBoxState::Unchecked; })
+				),
+				NAME_None,
+				EUserInterfaceActionType::ToggleButton
+			);
+		}
 
 		InMenuBuilder.AddMenuEntry(
 			LOCTEXT("ShowAnimationNames", "Show Sample Names"),
@@ -1857,7 +1821,7 @@ FText SBlendSpaceGridWidget::GetToolTipSampleValidity() const
 {
 	const UBlendSpace* BlendSpace = BlendSpaceBase.Get();
 	FText ToolTipText = FText::GetEmpty();
-	if(!bReadOnly && BlendSpace && BlendSpace->bInterpolateUsingGrid)
+	if (!bReadOnly && BlendSpace && BlendSpace->bInterpolateUsingGrid)
 	{
 		int32 SampleIndex = INDEX_NONE;
 		if (DragState == EDragState::None)
@@ -2310,7 +2274,17 @@ EVisibility SBlendSpaceGridWidget::GetPreviewToolTipVisibility() const
 
 EVisibility SBlendSpaceGridWidget::GetTriangulationButtonVisibility() const
 {
-	return (bShowSettingsButtons && (GridType == EGridType::TwoAxis)) ? EVisibility::Visible : EVisibility::Collapsed;
+	if (bShowSettingsButtons && GridType == EGridType::TwoAxis)
+	{
+		if (const UBlendSpace* BlendSpace = BlendSpaceBase.Get())
+		{
+			if (!BlendSpace->bInterpolateUsingGrid)
+			{
+				return EVisibility::Visible;
+			}
+		}
+	}
+	return  EVisibility::Collapsed;
 }
 
 EVisibility SBlendSpaceGridWidget::GetAnimationNamesButtonVisibility() const
