Add a new option to debug skel mesh to allow debug mesh LOD to track the LOD of an attached instance in the level viewport

diff --git a/Engine/Source/Editor/AnimationBlueprintEditor/Private/AnimationBlueprintEditor.cpp b/Engine/Source/Editor/AnimationBlueprintEditor/Private/AnimationBlueprintEditor.cpp
--- a/Engine/Source/Editor/AnimationBlueprintEditor/Private/AnimationBlueprintEditor.cpp
+++ b/Engine/Source/Editor/AnimationBlueprintEditor/Private/AnimationBlueprintEditor.cpp
@@ -1918,6 +1918,7 @@ void FAnimationBlueprintEditor::HandleSetObjectBeingDebugged(UObject* InObject)
 			{
 				GetPreviewScene()->ShowDefaultMode();
 				GetPreviewScene()->GetPreviewMeshComponent()->PreviewInstance->SetDebugSkeletalMeshComponent(nullptr);
+				GetPreviewScene()->GetPreviewMeshComponent()->bTrackAttachedInstanceLOD = false;
 			}
 			else
 			{
@@ -1932,6 +1933,7 @@ void FAnimationBlueprintEditor::HandleSetObjectBeingDebugged(UObject* InObject)
 		// Clear the copy-pose component and set us back to 'normal'
 		GetPreviewScene()->ShowDefaultMode();
 		GetPreviewScene()->GetPreviewMeshComponent()->PreviewInstance->SetDebugSkeletalMeshComponent(nullptr);
+		GetPreviewScene()->GetPreviewMeshComponent()->bTrackAttachedInstanceLOD = false;
 	}
 }
 
diff --git a/Engine/Source/Editor/Persona/Private/AnimViewportLODCommands.cpp b/Engine/Source/Editor/Persona/Private/AnimViewportLODCommands.cpp
--- a/Engine/Source/Editor/Persona/Private/AnimViewportLODCommands.cpp
+++ b/Engine/Source/Editor/Persona/Private/AnimViewportLODCommands.cpp
@@ -6,6 +6,7 @@
 
 void FAnimViewportLODCommands::RegisterCommands()
 {
+	UI_COMMAND( LODDebug, "LOD Debug", "Sync the LOD with the currently debugged instance", EUserInterfaceActionType::RadioButton, FInputChord());
 	UI_COMMAND( LODAuto, "LOD Auto", "Automatically select LOD", EUserInterfaceActionType::RadioButton, FInputChord() );
 	UI_COMMAND( LOD0, "LOD 0", "Force select LOD 0", EUserInterfaceActionType::RadioButton, FInputChord() );
 }
diff --git a/Engine/Source/Editor/Persona/Private/AnimViewportLODCommands.h b/Engine/Source/Editor/Persona/Private/AnimViewportLODCommands.h
--- a/Engine/Source/Editor/Persona/Private/AnimViewportLODCommands.h
+++ b/Engine/Source/Editor/Persona/Private/AnimViewportLODCommands.h
@@ -23,6 +23,9 @@ public:
 	{
 	}
 
+	/** LOD Debug */
+	TSharedPtr< FUICommandInfo > LODDebug;
+
 	/** LOD Auto */
 	TSharedPtr< FUICommandInfo > LODAuto;
 
diff --git a/Engine/Source/Editor/Persona/Private/SAnimViewportToolBar.cpp b/Engine/Source/Editor/Persona/Private/SAnimViewportToolBar.cpp
--- a/Engine/Source/Editor/Persona/Private/SAnimViewportToolBar.cpp
+++ b/Engine/Source/Editor/Persona/Private/SAnimViewportToolBar.cpp
@@ -943,10 +943,17 @@ FText SAnimViewportToolBar::GetLODMenuLabel() const
 	{
 		int32 LODSelectionType = Viewport.Pin()->GetLODSelection();
 
-		if (LODSelectionType > 0)
+		if (Viewport.Pin()->IsTrackingAttachedMeshLOD())
 		{
-			FString TitleLabel = FString::Printf(TEXT("LOD %d"), LODSelectionType - 1);
-			Label = FText::FromString(TitleLabel);
+			Label = FText::Format(LOCTEXT("LODMenu_DebugLabel", "LOD Debug ({0})"), FText::AsNumber(LODSelectionType - 1));
+		}
+		else
+		{
+			if (LODSelectionType > 0)
+			{
+				FString TitleLabel = FString::Printf(TEXT("LOD %d"), LODSelectionType - 1);
+				Label = FText::FromString(TitleLabel);
+			}
 		}
 	}
 	return Label;
@@ -968,6 +975,7 @@ TSharedRef<SWidget> SAnimViewportToolBar::GenerateLODMenu() const
 		// LOD Models
 		InMenuBuilder.BeginSection("AnimViewportPreviewLODs", LOCTEXT("ShowLOD_PreviewLabel", "Preview LODs") );
 		{
+			InMenuBuilder.AddMenuEntry( Actions.LODDebug );
 			InMenuBuilder.AddMenuEntry( Actions.LODAuto );
 			InMenuBuilder.AddMenuEntry( Actions.LOD0 );
 
diff --git a/Engine/Source/Editor/Persona/Private/SAnimationEditorViewport.cpp b/Engine/Source/Editor/Persona/Private/SAnimationEditorViewport.cpp
--- a/Engine/Source/Editor/Persona/Private/SAnimationEditorViewport.cpp
+++ b/Engine/Source/Editor/Persona/Private/SAnimationEditorViewport.cpp
@@ -931,6 +931,21 @@ void SAnimationEditorViewportTabBody::BindCommands()
 
 	CommandList.BeginGroup(TEXT("LOD"));
 
+	UDebugSkelMeshComponent* PreviewComponent = GetPreviewScene()->GetPreviewMeshComponent();
+
+	if (PreviewComponent)
+	{
+		//LOD Debug
+		CommandList.MapAction(
+			ViewportLODMenuCommands.LODDebug,
+			FExecuteAction::CreateSP(this, &SAnimationEditorViewportTabBody::OnSetLODTrackDebuggedInstance),
+			FCanExecuteAction::CreateLambda([PreviewComponent]() { return (bool)PreviewComponent->PreviewInstance->GetDebugSkeletalMeshComponent(); }),
+			FIsActionChecked::CreateLambda([PreviewComponent]() { return PreviewComponent->IsTrackingAttachedLOD(); }),
+			FIsActionButtonVisible::CreateLambda([PreviewComponent]() { return (bool)PreviewComponent->PreviewInstance->GetDebugSkeletalMeshComponent() ; }));
+
+		PreviewComponent->RegisterOnDebugForceLODChangedDelegate(FOnDebugForceLODChanged::CreateSP(this, &SAnimationEditorViewportTabBody::OnDebugForcedLODChanged));
+	}
+
 	//LOD Auto
 	CommandList.MapAction( 
 		ViewportLODMenuCommands.LODAuto,
@@ -1725,9 +1740,26 @@ int32 SAnimationEditorViewportTabBody::GetLODSelection() const
 
 bool SAnimationEditorViewportTabBody::IsLODModelSelected(int32 LODSelectionType) const
 {
+	UDebugSkelMeshComponent* PreviewComponent = GetPreviewScene()->GetPreviewMeshComponent();
+	if (PreviewComponent && PreviewComponent->IsTrackingAttachedLOD())
+	{
+		return false;
+	}
+
 	return GetLODSelection() == LODSelectionType;
 }
 
+bool SAnimationEditorViewportTabBody::IsTrackingAttachedMeshLOD() const
+{
+	UDebugSkelMeshComponent* PreviewComponent = GetPreviewScene()->GetPreviewMeshComponent();
+	if (PreviewComponent)
+	{
+		return PreviewComponent->IsTrackingAttachedLOD();
+	}
+
+	return false;
+}
+
 void SAnimationEditorViewportTabBody::OnSetLODModel(int32 LODSelectionType)
 {
 	UDebugSkelMeshComponent* PreviewComponent = GetPreviewScene()->GetPreviewMeshComponent();
@@ -1735,9 +1767,18 @@ void SAnimationEditorViewportTabBody::OnSetLODModel(int32 LODSelectionType)
 	if( PreviewComponent )
 	{
 		LODSelection = LODSelectionType;
-		PreviewComponent->SetForcedLOD(LODSelectionType);
-		PopulateUVChoices();
-		GetPreviewScene()->BroadcastOnSelectedLODChanged();
+		PreviewComponent->SetDebugForcedLOD(LODSelectionType);
+		PreviewComponent->bTrackAttachedInstanceLOD = false;
+	}
+}
+
+void SAnimationEditorViewportTabBody::OnSetLODTrackDebuggedInstance()
+{
+	UDebugSkelMeshComponent* PreviewComponent = GetPreviewScene()->GetPreviewMeshComponent();
+
+	if (PreviewComponent)
+	{
+		PreviewComponent->bTrackAttachedInstanceLOD = true;
 	}
 }
 
@@ -1752,6 +1793,17 @@ void SAnimationEditorViewportTabBody::OnLODModelChanged()
 	}
 }
 
+void SAnimationEditorViewportTabBody::OnDebugForcedLODChanged()
+{
+	UDebugSkelMeshComponent* PreviewComponent = GetPreviewScene()->GetPreviewMeshComponent();
+
+	if (PreviewComponent)
+	{
+		PopulateUVChoices();
+		GetPreviewScene()->BroadcastOnSelectedLODChanged();
+	}
+}
+
 TSharedRef<FAnimationViewportClient> SAnimationEditorViewportTabBody::GetAnimationViewportClient() const
 {
 	return StaticCastSharedRef<FAnimationViewportClient>(LevelViewportClient.ToSharedRef());
diff --git a/Engine/Source/Editor/Persona/Private/SAnimationEditorViewport.h b/Engine/Source/Editor/Persona/Private/SAnimationEditorViewport.h
--- a/Engine/Source/Editor/Persona/Private/SAnimationEditorViewport.h
+++ b/Engine/Source/Editor/Persona/Private/SAnimationEditorViewport.h
@@ -245,6 +245,7 @@ public:
 
 	/** LOD model selection checking function*/
 	bool IsLODModelSelected( int32 LODSelectionType ) const;
+	bool IsTrackingAttachedMeshLOD() const;
 	int32 GetLODSelection() const;
 
 	/** Function to set the current playback speed*/
@@ -290,7 +291,9 @@ public:
 
 	/** Function to set LOD model selection*/
 	void OnSetLODModel(int32 LODSelectionType);
+	void OnSetLODTrackDebuggedInstance();
 	void OnLODModelChanged();
+	void OnDebugForcedLODChanged();
 
 	/** Get the preview scene we are viewing */
 	TSharedRef<class FAnimationEditorPreviewScene> GetPreviewScene() const { return PreviewScenePtr.Pin().ToSharedRef(); }
diff --git a/Engine/Source/Editor/UnrealEd/Classes/Animation/DebugSkelMeshComponent.h b/Engine/Source/Editor/UnrealEd/Classes/Animation/DebugSkelMeshComponent.h
--- a/Engine/Source/Editor/UnrealEd/Classes/Animation/DebugSkelMeshComponent.h
+++ b/Engine/Source/Editor/UnrealEd/Classes/Animation/DebugSkelMeshComponent.h
@@ -13,6 +13,7 @@
 class Error;
 
 DECLARE_DELEGATE_RetVal(FText, FGetExtendedViewportText);
+DECLARE_DELEGATE(FOnDebugForceLODChanged);
 
 USTRUCT()
 struct FSelectedSocketInfo
@@ -293,6 +294,13 @@ class UNREALED_API UDebugSkelMeshComponent : public USkeletalMeshComponent
 	UPROPERTY(transient)
 	bool bPauseClothingSimulationWithAnim;
 
+	/** Should the LOD of the debug mesh component track the LOD of the instance being debugged */
+	UPROPERTY(transient)
+	bool bTrackAttachedInstanceLOD;
+
+	// Helper method that sets the forced lod
+	void SetDebugForcedLOD(int32 InNewForcedLOD);
+
 	//~ Begin USceneComponent Interface.
 	virtual FBoxSphereBounds CalcBounds(const FTransform& LocalToWorld) const override;
 	//~ End USceneComponent Interface.
@@ -386,6 +394,9 @@ class UNREALED_API UDebugSkelMeshComponent : public USkeletalMeshComponent
 	/** Set whether we should be previewing root motion. Note: disabling root motion preview resets transform. */
 	void SetPreviewRootMotion(bool bInPreviewRootMotion);
 
+	/** Whether the current LOD of the debug mesh is being synced with the attached (preview) mesh instance. */
+	bool IsTrackingAttachedLOD() const;
+
 #if WITH_EDITOR
 	//TODO - This is a really poor way to post errors to the user. Work out a better way.
 	struct FAnimNotifyErrors
@@ -408,8 +419,12 @@ class UNREALED_API UDebugSkelMeshComponent : public USkeletalMeshComponent
 	void UnregisterExtendedViewportTextDelegate(const FDelegateHandle& InDelegateHandle);
 	const TArray<FGetExtendedViewportText>& GetExtendedViewportTextDelegates() const { return ExtendedViewportTextDelegates; }
 
+	FDelegateHandle RegisterOnDebugForceLODChangedDelegate(const FOnDebugForceLODChanged& InDelegate);
+	void UnregisterOnDebugForceLODChangedDelegate();
+
 private:
 	TArray<FGetExtendedViewportText> ExtendedViewportTextDelegates;
+	FOnDebugForceLODChanged OnDebugForceLODChangedDelegate;
 public:
 
 #endif
diff --git a/Engine/Source/Editor/UnrealEd/Private/Animation/DebugSkelMeshComponent.cpp b/Engine/Source/Editor/UnrealEd/Private/Animation/DebugSkelMeshComponent.cpp
--- a/Engine/Source/Editor/UnrealEd/Private/Animation/DebugSkelMeshComponent.cpp
+++ b/Engine/Source/Editor/UnrealEd/Private/Animation/DebugSkelMeshComponent.cpp
@@ -43,11 +43,25 @@ UDebugSkelMeshComponent::UDebugSkelMeshComponent(const FObjectInitializer& Objec
 	bPauseClothingSimulationWithAnim = false;
 	bPerformSingleClothingTick = false;
 
+	bTrackAttachedInstanceLOD = false;
+
 	BoneRadiusMultiplier = 1.f;
 	
 	CachedClothBounds = FBoxSphereBounds(ForceInit);
 }
 
+void UDebugSkelMeshComponent::SetDebugForcedLOD(int32 InNewForcedLOD)
+{
+	SetForcedLOD(InNewForcedLOD);
+
+#if WITH_EDITOR
+	if (OnDebugForceLODChangedDelegate.IsBound())
+	{
+		OnDebugForceLODChangedDelegate.Execute();
+	}
+#endif
+}
+
 FBoxSphereBounds UDebugSkelMeshComponent::CalcBounds(const FTransform& LocalToWorld) const
 {
 	// Override bounds with pre-skinned bounds if asking for them
@@ -197,6 +211,11 @@ void UDebugSkelMeshComponent::SetPreviewRootMotion(bool bInPreviewRootMotion)
 	}
 }
 
+bool UDebugSkelMeshComponent::IsTrackingAttachedLOD() const
+{
+	return bTrackAttachedInstanceLOD;
+}
+
 FPrimitiveSceneProxy* UDebugSkelMeshComponent::CreateSceneProxy()
 {
 	FDebugSkelMeshSceneProxy* Result = NULL;
@@ -652,6 +671,18 @@ void UDebugSkelMeshComponent::UnregisterExtendedViewportTextDelegate(const FDele
 	});
 }
 
+FDelegateHandle UDebugSkelMeshComponent::RegisterOnDebugForceLODChangedDelegate(const FOnDebugForceLODChanged& InDelegate)
+{
+	OnDebugForceLODChangedDelegate = InDelegate;
+	return OnDebugForceLODChangedDelegate.GetHandle();
+}
+
+void UDebugSkelMeshComponent::UnregisterOnDebugForceLODChangedDelegate()
+{
+	checkf(OnDebugForceLODChangedDelegate.IsBound(), TEXT("OnDebugForceLODChangedDelegate is not registered"));
+	OnDebugForceLODChangedDelegate.Unbind();
+}
+
 #endif
 
 void UDebugSkelMeshComponent::ToggleClothSectionsVisibility(bool bShowOnlyClothSections)
@@ -799,6 +830,16 @@ void UDebugSkelMeshComponent::TickComponent(float DeltaTime, enum ELevelTick Tic
 		bRequiredBonesUpToDate = false;
 	}
 
+	if (bTrackAttachedInstanceLOD)
+	{
+		UAnimPreviewInstance* AnimPreviewInstance = Cast<UAnimPreviewInstance>(AnimScriptInstance);
+		USkeletalMeshComponent* TargetMeshComp = PreviewInstance->GetDebugSkeletalMeshComponent();
+		if (TargetMeshComp && TargetMeshComp->GetPredictedLODLevel() + 1 != GetForcedLOD())
+		{
+			SetDebugForcedLOD(TargetMeshComp->GetPredictedLODLevel() + 1);
+		}
+	}
+
 	Super::TickComponent(DeltaTime, TickType, ThisTickFunction);
 
 	// The tick from our super will call ShouldRunClothTick on us which will 'consume' this flag.
