Fixed scaling issue with the scrollbar in the anim editor timeline

diff --git a/Engine/Source/Editor/Persona/Private/AnimTimeline/SAnimTimeline.cpp b/Engine/Source/Editor/Persona/Private/AnimTimeline/SAnimTimeline.cpp
--- a/Engine/Source/Editor/Persona/Private/AnimTimeline/SAnimTimeline.cpp
+++ b/Engine/Source/Editor/Persona/Private/AnimTimeline/SAnimTimeline.cpp
@@ -35,6 +35,7 @@
 #include "Animation/AnimSequence.h"
 #include "Animation/AnimData/AnimDataModel.h"
 #include "Runtime/Engine/Classes/Animation/AnimSequenceHelpers.h"
+#include "Widgets/Layout/SScrollBox.h"
 
 #define LOCTEXT_NAMESPACE "SAnimTimeline"
 
@@ -213,29 +214,35 @@ void SAnimTimeline::Construct(const FArguments& InArgs, const TSharedRef<FAnimMo
 							SNew(SOverlay)
 							+SOverlay::Slot()
 							[
-								SNew(SScrollBorder, Outliner.ToSharedRef())
+								SNew(SVerticalBox)
+								+ SVerticalBox::Slot()
+								.FillHeight(1.f)
 								[
-									SNew(SHorizontalBox)
-
-									// outliner tree
-									+SHorizontalBox::Slot()
-									.FillWidth(FillCoefficient_0)
+									SNew(SScrollBox)
+									+ SScrollBox::Slot()
 									[
-										SNew(SBox)
+										SNew(SHorizontalBox)
+										
+										// outliner tree
+										+ SHorizontalBox::Slot()
+										.FillWidth(FillCoefficient_0)
 										[
-											Outliner.ToSharedRef()
+											SNew(SBox)
+											[
+												Outliner.ToSharedRef()
+											]
 										]
-									]
 
-									// track area
-									+SHorizontalBox::Slot()
-									.FillWidth(FillCoefficient_1)
-									[
-										SNew(SBox)
-										.Padding(ResizeBarPadding)
-										.Clipping(EWidgetClipping::ClipToBounds)
+										// track area
+										+ SHorizontalBox::Slot()
+										.FillWidth(FillCoefficient_1)
 										[
-											TrackArea.ToSharedRef()
+											SNew(SBox)
+											.Padding(ResizeBarPadding)
+											.Clipping(EWidgetClipping::ClipToBounds)
+											[
+												TrackArea.ToSharedRef()
+											]
 										]
 									]
 								]
