Fixed an issue where the animation editor viewport would not get commands routed if the viewport was not focused.

diff --git a/Engine/Source/Editor/Persona/Private/SAnimationEditorViewport.cpp b/Engine/Source/Editor/Persona/Private/SAnimationEditorViewport.cpp
--- a/Engine/Source/Editor/Persona/Private/SAnimationEditorViewport.cpp
+++ b/Engine/Source/Editor/Persona/Private/SAnimationEditorViewport.cpp
@@ -163,6 +163,14 @@ void SAnimationEditorViewport::BindCommands()
 
 	FShowFlagMenuCommands::Get().BindCommands(*CommandList, Client);
 	FBufferVisualizationMenuCommands::Get().BindCommands(*CommandList, Client);
+
+	if (TSharedPtr<SAnimationEditorViewportTabBody> TabBody = TabBodyPtr.Pin())
+	{
+		if (TSharedPtr<FAssetEditorToolkit> ParentAssetEditor = TabBody->GetAssetEditorToolkit())
+		{
+			CommandList->Append(ParentAssetEditor->GetToolkitCommands());
+		}
+	}
 }
 
 //////////////////////////////////////////////////////////////////////////
@@ -1050,6 +1058,9 @@ void SAnimationEditorViewportTabBody::BindCommands()
 	CommandList.MapAction(
 		FEditorViewportCommands::Get().FocusViewportToSelection,
 		FExecuteAction::CreateSP(this, &SAnimationEditorViewportTabBody::HandleFocusCamera));
+
+	TSharedPtr<FUICommandList> ToolkitCommandList = ConstCastSharedRef<FUICommandList>(GetAssetEditorToolkit()->GetToolkitCommands());
+	ToolkitCommandList->Append(UICommandList->AsShared());
 }
 
 void SAnimationEditorViewportTabBody::OnSetTurnTableSpeed(int32 SpeedIndex)
