Small fixes to BlendSpace analysis

diff --git a/Engine/Plugins/Animation/BlendSpaceMotionAnalysis/Source/BlendSpaceMotionAnalysis/Public/LocomotionAnalysis.h b/Engine/Plugins/Animation/BlendSpaceMotionAnalysis/Source/BlendSpaceMotionAnalysis/Public/LocomotionAnalysis.h
--- a/Engine/Plugins/Animation/BlendSpaceMotionAnalysis/Source/BlendSpaceMotionAnalysis/Public/LocomotionAnalysis.h
+++ b/Engine/Plugins/Animation/BlendSpaceMotionAnalysis/Source/BlendSpaceMotionAnalysis/Public/LocomotionAnalysis.h
@@ -43,11 +43,11 @@ public:
 
 	/** World or bone/socket axis that specifies the character's facing direction */
 	UPROPERTY(EditAnywhere, Category = AnalysisProperties)
-	EAnalysisLinearAxis CharacterFacingAxis = EAnalysisLinearAxis::Y;
+	EAnalysisLinearAxis CharacterFacingAxis = EAnalysisLinearAxis::PlusY;
 
 	/** World or bone/socket axis that specifies the character's up direction */
 	UPROPERTY(EditAnywhere, Category = AnalysisProperties)
-	EAnalysisLinearAxis CharacterUpAxis = EAnalysisLinearAxis::Z;
+	EAnalysisLinearAxis CharacterUpAxis = EAnalysisLinearAxis::PlusZ;
 };
 
 //======================================================================================================================
diff --git a/Engine/Plugins/Animation/BlendSpaceMotionAnalysis/Source/BlendSpaceMotionAnalysis/Public/RootMotionAnalysis.h b/Engine/Plugins/Animation/BlendSpaceMotionAnalysis/Source/BlendSpaceMotionAnalysis/Public/RootMotionAnalysis.h
--- a/Engine/Plugins/Animation/BlendSpaceMotionAnalysis/Source/BlendSpaceMotionAnalysis/Public/RootMotionAnalysis.h
+++ b/Engine/Plugins/Animation/BlendSpaceMotionAnalysis/Source/BlendSpaceMotionAnalysis/Public/RootMotionAnalysis.h
@@ -52,11 +52,11 @@ public:
 
 	/** World or bone/socket axis that specifies the character's facing direction */
 	UPROPERTY(EditAnywhere, Category = AnalysisProperties)
-	EAnalysisLinearAxis CharacterFacingAxis = EAnalysisLinearAxis::Y;
+	EAnalysisLinearAxis CharacterFacingAxis = EAnalysisLinearAxis::PlusY;
 
 	/** World or bone/socket axis that specifies the character's up direction */
 	UPROPERTY(EditAnywhere, Category = AnalysisProperties)
-	EAnalysisLinearAxis CharacterUpAxis = EAnalysisLinearAxis::Z;
+	EAnalysisLinearAxis CharacterUpAxis = EAnalysisLinearAxis::PlusZ;
 
 	/** Fraction through each animation at which analysis starts */
 	UPROPERTY(EditAnywhere, DisplayName = "Start time fraction", Category = AnalysisProperties, meta = (ClampMin = "0", ClampMax = "1"))
diff --git a/Engine/Source/Editor/Persona/Private/PersonaBlendSpaceAnalysis.cpp b/Engine/Source/Editor/Persona/Private/PersonaBlendSpaceAnalysis.cpp
--- a/Engine/Source/Editor/Persona/Private/PersonaBlendSpaceAnalysis.cpp
+++ b/Engine/Source/Editor/Persona/Private/PersonaBlendSpaceAnalysis.cpp
@@ -33,9 +33,12 @@ FVector BlendSpaceAnalysis::GetAxisFromTM(const FTransform& TM, EAnalysisLinearA
 {
 	switch (Axis)
 	{
-	case EAnalysisLinearAxis::X: return TM.TransformVectorNoScale(FVector(1, 0, 0));
-	case EAnalysisLinearAxis::Y: return TM.TransformVectorNoScale(FVector(0, 1, 0));
-	case EAnalysisLinearAxis::Z: return TM.TransformVectorNoScale(FVector(0, 0, 1));
+	case EAnalysisLinearAxis::PlusX: return TM.TransformVectorNoScale(FVector(1, 0, 0));
+	case EAnalysisLinearAxis::PlusY: return TM.TransformVectorNoScale(FVector(0, 1, 0));
+	case EAnalysisLinearAxis::PlusZ: return TM.TransformVectorNoScale(FVector(0, 0, 1));
+	case EAnalysisLinearAxis::MinusX: return TM.TransformVectorNoScale(FVector(-1, 0, 0));
+	case EAnalysisLinearAxis::MinusY: return TM.TransformVectorNoScale(FVector(0, -1, 0));
+	case EAnalysisLinearAxis::MinusZ: return TM.TransformVectorNoScale(FVector(0, 0, -1));
 	}
 	return FVector(0, 0, 0);
 }
diff --git a/Engine/Source/Editor/Persona/Public/BlendSpaceAnalysis.h b/Engine/Source/Editor/Persona/Public/BlendSpaceAnalysis.h
--- a/Engine/Source/Editor/Persona/Public/BlendSpaceAnalysis.h
+++ b/Engine/Source/Editor/Persona/Public/BlendSpaceAnalysis.h
@@ -54,9 +54,12 @@ enum class EAnalysisSpace : uint8
 UENUM()
 enum class EAnalysisLinearAxis : uint8
 {
-	X,
-	Y,
-	Z,
+	PlusX UMETA(DisplayName = "+X"),
+	PlusY UMETA(DisplayName = "+Y"),
+	PlusZ UMETA(DisplayName = "+Z"),
+	MinusX UMETA(DisplayName = "-X"),
+	MinusY UMETA(DisplayName = "-Y"),
+	MinusZ UMETA(DisplayName = "-Z"),
 };
 
 UENUM()
@@ -79,16 +82,16 @@ class PERSONA_API UCachedAnalysisProperties : public UObject
 	GENERATED_BODY()
 public:
 	void CopyFrom(const UCachedAnalysisProperties& Other);
-	EAnalysisLinearAxis     LinearFunctionAxis = EAnalysisLinearAxis::X;
+	EAnalysisLinearAxis     LinearFunctionAxis = EAnalysisLinearAxis::PlusX;
 	EAnalysisEulerAxis      EulerFunctionAxis = EAnalysisEulerAxis::Pitch;
 	FBoneSocketTarget       BoneSocket1;
 	FBoneSocketTarget       BoneSocket2;
-	EAnalysisLinearAxis     BoneFacingAxis = EAnalysisLinearAxis::X;
-	EAnalysisLinearAxis     BoneRightAxis = EAnalysisLinearAxis::Y;
+	EAnalysisLinearAxis     BoneFacingAxis = EAnalysisLinearAxis::PlusX;
+	EAnalysisLinearAxis     BoneRightAxis = EAnalysisLinearAxis::PlusY;
 	EAnalysisSpace          Space = EAnalysisSpace::World;
 	FBoneSocketTarget       SpaceBoneSocket;
-	EAnalysisLinearAxis     CharacterFacingAxis = EAnalysisLinearAxis::Y;
-	EAnalysisLinearAxis     CharacterUpAxis = EAnalysisLinearAxis::Z;
+	EAnalysisLinearAxis     CharacterFacingAxis = EAnalysisLinearAxis::PlusY;
+	EAnalysisLinearAxis     CharacterUpAxis = EAnalysisLinearAxis::PlusZ;
 	float                   StartTimeFraction = 0.0f;
 	float                   EndTimeFraction = 1.0f;
 	bool                    bLockAfterAnalysis = false;
@@ -104,7 +107,7 @@ public:
 
 	/** Axis for the analysis function */
 	UPROPERTY(EditAnywhere, DisplayName = "Axis", Category = AnalysisProperties)
-	EAnalysisLinearAxis FunctionAxis = EAnalysisLinearAxis::X;
+	EAnalysisLinearAxis FunctionAxis = EAnalysisLinearAxis::PlusX;
 
 	/** The bone or socket used for analysis */
 	UPROPERTY(EditAnywhere, DisplayName = "Bone/Socket", Category = AnalysisProperties)
@@ -150,11 +153,11 @@ public:
 
 	/** Used for some analysis functions - specifies the bone/socket axis that points in the facing/forwards direction */
 	UPROPERTY(EditAnywhere, Category = AnalysisProperties)
-	EAnalysisLinearAxis BoneFacingAxis = EAnalysisLinearAxis::X;
+	EAnalysisLinearAxis BoneFacingAxis = EAnalysisLinearAxis::PlusX;
 
 	/** Used for some analysis functions - specifies the bone/socket axis that points to the "right" */
 	UPROPERTY(EditAnywhere, Category = AnalysisProperties)
-	EAnalysisLinearAxis BoneRightAxis = EAnalysisLinearAxis::Y;
+	EAnalysisLinearAxis BoneRightAxis = EAnalysisLinearAxis::PlusY;
 
 	/**
 	* The space in which to perform the analysis. Fixed will use the analysis bone/socket at the first frame
@@ -171,11 +174,11 @@ public:
 
 	/** World or bone/socket axis that specifies the character's facing direction */
 	UPROPERTY(EditAnywhere, Category = AnalysisProperties)
-	EAnalysisLinearAxis CharacterFacingAxis = EAnalysisLinearAxis::Y;
+	EAnalysisLinearAxis CharacterFacingAxis = EAnalysisLinearAxis::PlusY;
 
 	/** World or bone/socket axis that specifies the character's up direction */
 	UPROPERTY(EditAnywhere, Category = AnalysisProperties)
-	EAnalysisLinearAxis CharacterUpAxis = EAnalysisLinearAxis::Z;
+	EAnalysisLinearAxis CharacterUpAxis = EAnalysisLinearAxis::PlusZ;
 
 	/** Fraction through each animation at which analysis starts */
 	UPROPERTY(EditAnywhere, DisplayName = "Start time fraction", Category = AnalysisProperties, meta = (ClampMin = "0", ClampMax = "1"))
@@ -389,19 +392,17 @@ static bool CalculateVelocity(
 	int32 LastKey = FMath::Clamp(
 		(int32) (NumSampledKeys * AnalysisProperties->EndTimeFraction), FirstKey, NumSampledKeys-1);
 
-	if (FirstKey >= LastKey)
-	{
-		Result.Set(0, 0, 0);
-		return true;
-	}
+	// First and Last key are for averaging. However, the finite differencing always goes from one frame to the next
+	int32 NumKeys = FMath::Max(1 + LastKey - FirstKey, 1);
 
 	FTransform FrameTM;
 	bool bNeedToUpdateFrameTM = true;
 	CalculateFrameTM(bNeedToUpdateFrameTM, FrameTM, FirstKey, AnalysisProperties, Animation);
 
 	Result.Set(0, 0, 0);
-	for (int32 Key = FirstKey; Key != LastKey + 1; ++Key)
+	for (int32 iKey = 0; iKey != NumKeys; ++iKey)
 	{
+		int32 Key = (FirstKey + iKey) % (NumSampledKeys + 1);
 		CalculateFrameTM(bNeedToUpdateFrameTM, FrameTM, Key, AnalysisProperties, Animation);
 
 		FTransform BoneTM1 = GetBoneTransform(Animation, Key, BoneName);
@@ -618,19 +619,17 @@ static bool CalculateAngularVelocity(
 	int32 LastKey = FMath::Clamp(
 		(int32) (NumSampledKeys * AnalysisProperties->EndTimeFraction), FirstKey, NumSampledKeys-1);
 
-	if (FirstKey >= LastKey)
-	{
-		Result.Set(0, 0, 0);
-		return true;
-	}
+	// First and Last key are for averaging. However, the finite differencing always goes from one frame to the next
+	int32 NumKeys = FMath::Max(1 + LastKey - FirstKey, 1);
 
 	FTransform FrameTM;
 	bool bNeedToUpdateFrameTM = true;
 	CalculateFrameTM(bNeedToUpdateFrameTM, FrameTM, FirstKey, AnalysisProperties, Animation);
 
 	Result.Set(0, 0, 0);
-	for (int32 Key = FirstKey; Key != LastKey + 1; ++Key)
+	for (int32 iKey = 0; iKey != NumKeys; ++iKey)
 	{
+		int32 Key = (FirstKey + iKey) % (NumSampledKeys + 1);
 		CalculateFrameTM(bNeedToUpdateFrameTM, FrameTM, Key, AnalysisProperties, Animation);
 
 		FTransform BoneTM1 = GetBoneTransform(Animation, Key, BoneName);
@@ -698,11 +697,8 @@ static bool CalculateOrientationRate(
 	int32 LastKey = FMath::Clamp(
 		(int32) (NumSampledKeys * AnalysisProperties->EndTimeFraction), FirstKey, NumSampledKeys-1);
 
-	if (FirstKey >= LastKey)
-	{
-		Result.Set(0, 0, 0);
-		return true;
-	}
+	// First and Last key are for averaging. However, the finite differencing always goes from one frame to the next
+	int32 NumKeys = FMath::Max(1 + LastKey - FirstKey, 1);
 
 	FTransform FrameTM;
 	bool bNeedToUpdateFrameTM = true;
@@ -711,8 +707,9 @@ static bool CalculateOrientationRate(
 	GetFrameDirs(FrameFacingDir, FrameUpDir, FrameRightDir, FrameTM, AnalysisProperties);
 
 	Result.Set(0, 0, 0);
-	for (int32 Key = FirstKey; Key != LastKey + 1; ++Key)
+	for (int32 iKey = 0; iKey != NumKeys; ++iKey)
 	{
+		int32 Key = (FirstKey + iKey) % (NumSampledKeys + 1);
 		CalculateFrameTM(bNeedToUpdateFrameTM, FrameTM, Key, AnalysisProperties, Animation);
 		GetFrameDirs(FrameFacingDir, FrameUpDir, FrameRightDir, FrameTM, AnalysisProperties);
 
