Fix PreviewSceneController not being reinitialized after a blueprint re-compile

diff --git a/Engine/Source/Editor/Persona/Private/PreviewSceneCustomizations.cpp b/Engine/Source/Editor/Persona/Private/PreviewSceneCustomizations.cpp
--- a/Engine/Source/Editor/Persona/Private/PreviewSceneCustomizations.cpp
+++ b/Engine/Source/Editor/Persona/Private/PreviewSceneCustomizations.cpp
@@ -66,6 +66,14 @@ FPreviewSceneDescriptionCustomization::~FPreviewSceneDescriptionCustomization()
 		FactoryToUse->RemoveFromRoot();
 		FactoryToUse = nullptr;
 	}
+
+	if (const TSharedPtr<IPersonaToolkit> Toolkit = PersonaToolkit.Pin())
+	{
+		if (UAnimBlueprint* AnimBlueprint = Toolkit->GetAnimBlueprint())
+		{
+			AnimBlueprint->OnCompiled().RemoveAll(this);
+		}
+	}
 }
 
 void FPreviewSceneDescriptionCustomization::CustomizeDetails(IDetailLayoutBuilder& DetailBuilder)
@@ -463,6 +471,12 @@ void FPreviewSceneDescriptionCustomization::CustomizeDetails(IDetailLayoutBuilde
 		DetailBuilder.HideProperty(SkeletalMeshProperty);
 		DetailBuilder.HideProperty(AdditionalMeshesProperty);
 	}
+
+	TSharedPtr<IPersonaToolkit> Toolkit = PersonaToolkit.Pin();
+	if (Toolkit && Toolkit->GetAnimBlueprint())
+	{
+		Toolkit->GetAnimBlueprint()->OnCompiled().AddSP(this, &FPreviewSceneDescriptionCustomization::HandleAnimBlueprintCompiled);
+	}
 }
 
 EVisibility FPreviewSceneDescriptionCustomization::GetSaveButtonVisibility(TSharedRef<IPropertyHandle> InAdditionalMeshesProperty) const
@@ -569,11 +583,7 @@ void FPreviewSceneDescriptionCustomization::OnComboSelectionChanged(TSharedPtr<F
 
 void FPreviewSceneDescriptionCustomization::HandlePreviewControllerPropertyChanged()
 {
-	TSharedPtr<FAnimationEditorPreviewScene> PreviewScenePtr = PreviewScene.Pin();
-	UPersonaPreviewSceneDescription* PersonaPreviewSceneDescription = PreviewScenePtr->GetPreviewSceneDescription();
-	
-	PersonaPreviewSceneDescription->PreviewControllerInstance->UninitializeView(PersonaPreviewSceneDescription, PreviewScenePtr.Get());
-	PersonaPreviewSceneDescription->PreviewControllerInstance->InitializeView(PersonaPreviewSceneDescription, PreviewScenePtr.Get());
+	ReinitializePreviewController();
 }
 
 void FPreviewSceneDescriptionCustomization::HandleMeshChanged(const FAssetData& InAssetData)   
@@ -625,6 +635,20 @@ ECheckBoxState FPreviewSceneDescriptionCustomization::HandleUseCustomAnimBPIsChe
 	return GetDefault<UPersonaOptions>()->bAllowPreviewMeshCollectionsToUseCustomAnimBP? ECheckBoxState::Checked : ECheckBoxState::Unchecked;
 }
 
+void FPreviewSceneDescriptionCustomization::HandleAnimBlueprintCompiled(UBlueprint*)
+{
+	ReinitializePreviewController();
+}
+
+void FPreviewSceneDescriptionCustomization::ReinitializePreviewController()
+{
+	TSharedPtr<FAnimationEditorPreviewScene> PreviewScenePtr = PreviewScene.Pin();
+	UPersonaPreviewSceneDescription* PersonaPreviewSceneDescription = PreviewScenePtr->GetPreviewSceneDescription();
+	
+	PersonaPreviewSceneDescription->PreviewControllerInstance->UninitializeView(PersonaPreviewSceneDescription, PreviewScenePtr.Get());
+	PersonaPreviewSceneDescription->PreviewControllerInstance->InitializeView(PersonaPreviewSceneDescription, PreviewScenePtr.Get());
+}
+
 bool FPreviewSceneDescriptionCustomization::GetReplaceVisibility(TSharedPtr<IPropertyHandle> PropertyHandle) const
 {
 	// Only show the replace button if the current material can be replaced
diff --git a/Engine/Source/Editor/Persona/Private/PreviewSceneCustomizations.h b/Engine/Source/Editor/Persona/Private/PreviewSceneCustomizations.h
--- a/Engine/Source/Editor/Persona/Private/PreviewSceneCustomizations.h
+++ b/Engine/Source/Editor/Persona/Private/PreviewSceneCustomizations.h
@@ -82,6 +82,12 @@ private:
 
 	ECheckBoxState HandleUseCustomAnimBPIsChecked() const;
 
+	// Called when the anim blueprint being edited is compiled.
+	void HandleAnimBlueprintCompiled(UBlueprint* Blueprint);
+
+	// Reinitialize the preview controller in the preview scene.
+	void ReinitializePreviewController();
+
 #if CHAOS_SIMULATION_DETAIL_VIEW_FACTORY_SELECTOR
 	// Make the widget of each item in the preview cloth factory combo box
 	TSharedRef<SWidget> MakeClothingSimulationFactoryWidget(TSharedPtr<TSubclassOf<class UClothingSimulationFactory>> Item) const;
