This CL adds "bVisibleInRayTracing" flag to skeletal mesh sections and exposes it in the editor - fixes UE-105223 and partially addresses MH-4066

diff --git a/Engine/Source/Developer/MeshMergeUtilities/Private/MeshMergeHelpers.cpp b/Engine/Source/Developer/MeshMergeUtilities/Private/MeshMergeHelpers.cpp
--- a/Engine/Source/Developer/MeshMergeUtilities/Private/MeshMergeHelpers.cpp
+++ b/Engine/Source/Developer/MeshMergeUtilities/Private/MeshMergeHelpers.cpp
@@ -141,6 +141,11 @@ void FMeshMergeHelpers::ExtractSections(const USkeletalMeshComponent* Component,
 			SectionInfo.EnabledProperties.Add(GET_MEMBER_NAME_CHECKED(FSkelMeshSection, bCastShadow));
 		}
 
+		if (MeshSection.bVisibleInRayTracing)
+		{
+			SectionInfo.EnabledProperties.Add(GET_MEMBER_NAME_CHECKED(FSkelMeshSection, bVisibleInRayTracing));
+		}
+
 		if (MeshSection.bRecomputeTangent)
 		{
 			SectionInfo.EnabledProperties.Add(GET_MEMBER_NAME_CHECKED(FSkelMeshSection, bRecomputeTangent));
diff --git a/Engine/Source/Editor/ClothPainter/Private/SClothAssetSelector.cpp b/Engine/Source/Editor/ClothPainter/Private/SClothAssetSelector.cpp
--- a/Engine/Source/Editor/ClothPainter/Private/SClothAssetSelector.cpp
+++ b/Engine/Source/Editor/ClothPainter/Private/SClothAssetSelector.cpp
@@ -199,6 +199,7 @@ private:
 			FSkelMeshSourceSectionUserData& SourceSectionUserData = LODModel.UserSectionsData.FindOrAdd(OriginalSectionIndex);
 			SourceSectionUserData.bDisabled = LODModel.Sections[SectionIndex].bDisabled;
 			SourceSectionUserData.bCastShadow = LODModel.Sections[SectionIndex].bCastShadow;
+			SourceSectionUserData.bVisibleInRayTracing = LODModel.Sections[SectionIndex].bVisibleInRayTracing;			
 			SourceSectionUserData.bRecomputeTangent = LODModel.Sections[SectionIndex].bRecomputeTangent;
 			SourceSectionUserData.GenerateUpToLodIndex = LODModel.Sections[SectionIndex].GenerateUpToLodIndex;
 			SourceSectionUserData.CorrespondClothAssetIndex = LODModel.Sections[SectionIndex].CorrespondClothAssetIndex;
diff --git a/Engine/Source/Editor/Persona/Private/PersonaMeshDetails.cpp b/Engine/Source/Editor/Persona/Private/PersonaMeshDetails.cpp
--- a/Engine/Source/Editor/Persona/Private/PersonaMeshDetails.cpp
+++ b/Engine/Source/Editor/Persona/Private/PersonaMeshDetails.cpp
@@ -145,6 +145,7 @@ void SetSkelMeshSourceSectionUserData(FSkeletalMeshLODModel& LODModel, const int
 	FSkelMeshSourceSectionUserData& SourceSectionUserData = LODModel.UserSectionsData.FindOrAdd(OriginalSectionIndex);
 	SourceSectionUserData.bDisabled = LODModel.Sections[SectionIndex].bDisabled;
 	SourceSectionUserData.bCastShadow = LODModel.Sections[SectionIndex].bCastShadow;
+	SourceSectionUserData.bVisibleInRayTracing = LODModel.Sections[SectionIndex].bVisibleInRayTracing;
 	SourceSectionUserData.bRecomputeTangent = LODModel.Sections[SectionIndex].bRecomputeTangent;
 	SourceSectionUserData.RecomputeTangentsVertexMaskChannel = LODModel.Sections[SectionIndex].RecomputeTangentsVertexMaskChannel;
 	SourceSectionUserData.GenerateUpToLodIndex = LODModel.Sections[SectionIndex].GenerateUpToLodIndex;
@@ -1943,6 +1944,7 @@ void FPersonaMeshDetails::OnCopySectionList(int32 LODIndex)
 				JSonSection->SetBoolField(TEXT("RecomputeTangent"), ModelSection.bRecomputeTangent);
 				JSonSection->SetNumberField(TEXT("RecomputeTangentsVertexMaskChannel"), static_cast<uint8>(ModelSection.RecomputeTangentsVertexMaskChannel));
 				JSonSection->SetBoolField(TEXT("CastShadow"), ModelSection.bCastShadow);
+				JSonSection->SetBoolField(TEXT("VisibleInRayTracing"), ModelSection.bVisibleInRayTracing);
 				JSonSection->SetNumberField(TEXT("GenerateUpToLodIndex"), ModelSection.GenerateUpToLodIndex);
 				JSonSection->SetNumberField(TEXT("ChunkedParentSectionIndex"), ModelSection.ChunkedParentSectionIndex);
 				JSonSection->SetStringField(TEXT("ClothingData.AssetGuid"), ModelSection.ClothingData.AssetGuid.ToString(EGuidFormats::Digits));
@@ -2027,6 +2029,7 @@ void FPersonaMeshDetails::OnPasteSectionList(int32 LODIndex)
 							ModelSection.RecomputeTangentsVertexMaskChannel = static_cast<ESkinVertexColorChannel>((uint8)(Value & 0xFF));
 						}
 						(*JSonSection)->TryGetBoolField(TEXT("CastShadow"), ModelSection.bCastShadow);
+						(*JSonSection)->TryGetBoolField(TEXT("VisibleInRayTracing"), ModelSection.bVisibleInRayTracing);
 						if ((*JSonSection)->TryGetNumberField(TEXT("GenerateUpToLodIndex"), Value))
 						{
 							ModelSection.GenerateUpToLodIndex = (int8)Value;
@@ -2077,6 +2080,7 @@ void FPersonaMeshDetails::OnCopySectionItem(int32 LODIndex, int32 SectionIndex)
 				RootJsonObject->SetBoolField(TEXT("RecomputeTangent"), ModelSection.bRecomputeTangent);
 				RootJsonObject->SetNumberField(TEXT("RecomputeTangentsVertexMaskChannel"), static_cast<uint8>(ModelSection.RecomputeTangentsVertexMaskChannel));
 				RootJsonObject->SetBoolField(TEXT("CastShadow"), ModelSection.bCastShadow);
+				RootJsonObject->SetBoolField(TEXT("VisibleInRayTracing"), ModelSection.bVisibleInRayTracing);
 				RootJsonObject->SetNumberField(TEXT("GenerateUpToLodIndex"), ModelSection.GenerateUpToLodIndex);
 				RootJsonObject->SetNumberField(TEXT("ChunkedParentSectionIndex"), ModelSection.ChunkedParentSectionIndex);
 				RootJsonObject->SetStringField(TEXT("ClothingData.AssetGuid"), ModelSection.ClothingData.AssetGuid.ToString(EGuidFormats::Digits));
@@ -2156,6 +2160,7 @@ void FPersonaMeshDetails::OnPasteSectionItem(int32 LODIndex, int32 SectionIndex)
 						ModelSection.RecomputeTangentsVertexMaskChannel = static_cast<ESkinVertexColorChannel>((uint8) (Value & 0xFF));
 					}
 					RootJsonObject->TryGetBoolField(TEXT("CastShadow"), ModelSection.bCastShadow);
+					RootJsonObject->TryGetBoolField(TEXT("VisibleInRayTracing"), ModelSection.bVisibleInRayTracing);
 					if (RootJsonObject->TryGetNumberField(TEXT("GenerateUpToLodIndex"), Value))
 					{
 						ModelSection.GenerateUpToLodIndex = (int8)Value;
@@ -4636,6 +4641,26 @@ TSharedRef<SWidget> FPersonaMeshDetails::OnGenerateCustomSectionWidgetsForSectio
 			]
 		]
 	];
+
+	SectionWidget->AddSlot()
+	.AutoHeight()
+	.Padding(0, 2, 0, 0)
+	[
+		SNew(SHorizontalBox)
+		+SHorizontalBox::Slot()
+		.AutoWidth()
+		.Padding(2, 0, 2, 0)
+		[
+			SNew(SCheckBox)
+			.IsChecked(this, &FPersonaMeshDetails::IsSectionVisibleInRayTracingEnabled, LODIndex, SectionIndex)
+			.OnCheckStateChanged(this, &FPersonaMeshDetails::OnSectionVisibleInRayTracingChanged, LODIndex, SectionIndex)
+			[
+				SNew(STextBlock)
+				.Font(IDetailLayoutBuilder::GetDetailFont())
+				.Text(LOCTEXT("Visible In RayTracing", "Visible In RayTracing"))
+			]
+		]
+	];
 	return SectionWidget;
 }
 
@@ -5183,6 +5208,77 @@ void FPersonaMeshDetails::OnSectionShadowCastingChanged(ECheckBoxState NewState,
 	}
 }
 
+ECheckBoxState FPersonaMeshDetails::IsSectionVisibleInRayTracingEnabled(int32 LODIndex, int32 SectionIndex) const
+{
+	ECheckBoxState State = ECheckBoxState::Unchecked;
+	const USkeletalMesh* Mesh = GetPersonaToolkit()->GetMesh();
+	if (Mesh == nullptr)
+		return State;
+
+	check(Mesh->GetImportedModel());
+
+	if (!Mesh->GetImportedModel()->LODModels.IsValidIndex(LODIndex))
+		return State;
+
+	const FSkeletalMeshLODModel& LODModel = Mesh->GetImportedModel()->LODModels[LODIndex];
+
+	if (!LODModel.Sections.IsValidIndex(SectionIndex))
+		return State;
+
+	const FSkelMeshSection& Section = LODModel.Sections[SectionIndex];
+
+	State = Section.bVisibleInRayTracing ? ECheckBoxState::Checked : ECheckBoxState::Unchecked;
+	return State;
+}
+
+void FPersonaMeshDetails::OnSectionVisibleInRayTracingChanged(ECheckBoxState NewState, int32 LODIndex, int32 SectionIndex)
+{
+	ECheckBoxState State = ECheckBoxState::Unchecked;
+	USkeletalMesh* Mesh = GetPersonaToolkit()->GetMesh();
+	if (Mesh == nullptr)
+		return;
+
+	check(Mesh->GetImportedModel());
+
+	if (!Mesh->GetImportedModel()->LODModels.IsValidIndex(LODIndex))
+		return;
+
+	FSkeletalMeshLODModel& LODModel = Mesh->GetImportedModel()->LODModels[LODIndex];
+
+	if (!LODModel.Sections.IsValidIndex(SectionIndex))
+		return;
+
+	FSkelMeshSection& Section = LODModel.Sections[SectionIndex];
+
+	//Update Original PolygonGroup
+	auto UpdatePolygonGroupVisibleInRayTracing = [&Mesh, &LODModel, &Section, &SectionIndex](bool bVisibleInRayTracing)
+	{
+		FScopedSuspendAlternateSkinWeightPreview ScopedSuspendAlternateSkinnWeightPreview(Mesh);
+		{
+			FScopedSkeletalMeshPostEditChange ScopedPostEditChange(Mesh);
+			Section.bVisibleInRayTracing = bVisibleInRayTracing;
+			//We change only the parent chunk data
+			check(Section.ChunkedParentSectionIndex == INDEX_NONE);
+
+			//The post edit change will kick a build
+			SetSkelMeshSourceSectionUserData(LODModel, SectionIndex, Section.OriginalDataSectionIndex);
+		}
+	};
+
+	if (NewState == ECheckBoxState::Checked)
+	{
+		const FScopedTransaction Transaction(LOCTEXT("PersonaSetSectionVisibleInRayTracingFlag", "Persona editor: Set Visible In RayTracing For Section"));
+		Mesh->Modify();
+		UpdatePolygonGroupVisibleInRayTracing(true);
+	}
+	else if (NewState == ECheckBoxState::Unchecked)
+	{
+		const FScopedTransaction Transaction(LOCTEXT("PersonaClearSectionVisibleInRayTracinggFlag", "Persona editor: Clear Visible In RayTracing For Section"));
+		Mesh->Modify();
+		UpdatePolygonGroupVisibleInRayTracing(false);
+	}
+}
+
 const FString RecomputeTangentsChannelNames[] = {
 			TEXT("None"),
 			TEXT("All"),
diff --git a/Engine/Source/Editor/Persona/Private/PersonaMeshDetails.h b/Engine/Source/Editor/Persona/Private/PersonaMeshDetails.h
--- a/Engine/Source/Editor/Persona/Private/PersonaMeshDetails.h
+++ b/Engine/Source/Editor/Persona/Private/PersonaMeshDetails.h
@@ -512,7 +512,7 @@ private:
 	* Handler for check box display based on whether the material has shadow casting enabled
 	*
 	* @param LODIndex	The LODIndex we want to change
-	* @param SectionIndex	The SectionIndex we change the RecomputeTangent
+	* @param SectionIndex	The SectionIndex we change the ShadowCasting flag
 	*/
 	ECheckBoxState IsSectionShadowCastingEnabled(int32 LODIndex, int32 SectionIndex) const;
 
@@ -520,10 +520,26 @@ private:
 	* Handler for changing shadow casting status on a section
 	*
 	* @param LODIndex	The LODIndex we want to change
-	* @param SectionIndex	The SectionIndex we change the RecomputeTangent
+	* @param SectionIndex	The SectionIndex we change the ShadowCasting flag
 	*/
 	void OnSectionShadowCastingChanged(ECheckBoxState NewState, int32 LODIndex, int32 SectionIndex);
 
+	/**
+	* Handler for check box display based on whether the material has VisibleInRayTracing enabled
+	*
+	* @param LODIndex	The LODIndex we want to change
+	* @param SectionIndex	The SectionIndex we change the VisibleInRayTracing flag
+	*/
+	ECheckBoxState IsSectionVisibleInRayTracingEnabled(int32 LODIndex, int32 SectionIndex) const;
+
+	/**
+	* Handler for changing VisibleInRayTracing status on a section
+	*
+	* @param LODIndex	The LODIndex we want to change
+	* @param SectionIndex	The SectionIndex we change the VisibleInRayTracing flag
+	*/
+	void OnSectionVisibleInRayTracingChanged(ECheckBoxState NewState, int32 LODIndex, int32 SectionIndex);
+
 	/**
 	* Handler for selecting which vertex color to mask the blending of recomputing tangents
 	*
diff --git a/Engine/Source/Editor/UnrealEd/Private/Fbx/FbxSkeletalMeshImport.cpp b/Engine/Source/Editor/UnrealEd/Private/Fbx/FbxSkeletalMeshImport.cpp
--- a/Engine/Source/Editor/UnrealEd/Private/Fbx/FbxSkeletalMeshImport.cpp
+++ b/Engine/Source/Editor/UnrealEd/Private/Fbx/FbxSkeletalMeshImport.cpp
@@ -3799,6 +3799,7 @@ void UnFbx::FFbxImporter::InsertNewLODToBaseSkeletalMesh(USkeletalMesh* InSkelet
 							const FSkelMeshSection& ExistSection = ExistLODModel.Sections[ExistSectionIndex];
 							//We found a match, restore the data
 							NewSection.bCastShadow = ExistSection.bCastShadow;
+							NewSection.bVisibleInRayTracing = ExistSection.bVisibleInRayTracing;
 							NewSection.bRecomputeTangent = ExistSection.bRecomputeTangent;
 							NewSection.RecomputeTangentsVertexMaskChannel = ExistSection.RecomputeTangentsVertexMaskChannel;
 							NewSection.bDisabled = ExistSection.bDisabled;
@@ -3811,6 +3812,7 @@ void UnFbx::FFbxImporter::InsertNewLODToBaseSkeletalMesh(USkeletalMesh* InSkelet
 								FSkelMeshSourceSectionUserData& UserSectionData = NewLODModel.UserSectionsData.FindOrAdd(ParentOriginalSectionIndex);
 								UserSectionData.bDisabled = NewSection.bDisabled;
 								UserSectionData.bCastShadow = NewSection.bCastShadow;
+								UserSectionData.bVisibleInRayTracing = NewSection.bVisibleInRayTracing;
 								UserSectionData.bRecomputeTangent = NewSection.bRecomputeTangent;					
 								UserSectionData.RecomputeTangentsVertexMaskChannel = NewSection.RecomputeTangentsVertexMaskChannel;
 								UserSectionData.GenerateUpToLodIndex = NewSection.GenerateUpToLodIndex;
@@ -4223,6 +4225,7 @@ bool UnFbx::FFbxImporter::ImportSkeletalMeshLOD(USkeletalMesh* InSkeletalMesh, U
 					{
 						//Set the value and exit
 						ImportedSection.bCastShadow = ExistingSection.bCastShadow;
+						ImportedSection.bVisibleInRayTracing = ExistingSection.bVisibleInRayTracing;
 						ImportedSection.bRecomputeTangent = ExistingSection.bRecomputeTangent;
 						ImportedSection.RecomputeTangentsVertexMaskChannel = ExistingSection.RecomputeTangentsVertexMaskChannel;
 						break;
@@ -4232,6 +4235,7 @@ bool UnFbx::FFbxImporter::ImportSkeletalMeshLOD(USkeletalMesh* InSkeletalMesh, U
 				{
 					//Set the value and exit
 					ImportedSection.bCastShadow = ExistingSection.bCastShadow;
+					ImportedSection.bVisibleInRayTracing = ExistingSection.bVisibleInRayTracing;
 					ImportedSection.bRecomputeTangent = ExistingSection.bRecomputeTangent;
 					ImportedSection.RecomputeTangentsVertexMaskChannel = ExistingSection.RecomputeTangentsVertexMaskChannel;
 					break;
diff --git a/Engine/Source/Editor/UnrealEd/Private/ImportUtils/SkeletalMeshImportUtils.cpp b/Engine/Source/Editor/UnrealEd/Private/ImportUtils/SkeletalMeshImportUtils.cpp
--- a/Engine/Source/Editor/UnrealEd/Private/ImportUtils/SkeletalMeshImportUtils.cpp
+++ b/Engine/Source/Editor/UnrealEd/Private/ImportUtils/SkeletalMeshImportUtils.cpp
@@ -260,6 +260,7 @@ void SkeletalMesUtilsImpl::SaveSkeletalMeshLODModelSections(USkeletalMesh* Sourc
 						OriginalMatched[OriginalSectionIndex] = true;
 						OriginalSection.bDisabled = ReduceSection.bDisabled;
 						OriginalSection.bCastShadow = ReduceSection.bCastShadow;
+						OriginalSection.bVisibleInRayTracing = ReduceSection.bVisibleInRayTracing;
 						OriginalSection.bRecomputeTangent = ReduceSection.bRecomputeTangent;
 						OriginalSection.RecomputeTangentsVertexMaskChannel = ReduceSection.RecomputeTangentsVertexMaskChannel;
 						OriginalSection.GenerateUpToLodIndex = ReduceSection.GenerateUpToLodIndex;
@@ -279,6 +280,7 @@ void SkeletalMesUtilsImpl::SaveSkeletalMeshLODModelSections(USkeletalMesh* Sourc
 				{
 					OriginalSection.bDisabled = ReduceUserSectionData->bDisabled;
 					OriginalSection.bCastShadow = ReduceUserSectionData->bCastShadow;
+					OriginalSection.bVisibleInRayTracing = ReduceUserSectionData->bVisibleInRayTracing;
 					OriginalSection.bRecomputeTangent = ReduceUserSectionData->bRecomputeTangent;
 					OriginalSection.RecomputeTangentsVertexMaskChannel = ReduceUserSectionData->RecomputeTangentsVertexMaskChannel;
 					OriginalSection.GenerateUpToLodIndex = ReduceUserSectionData->GenerateUpToLodIndex;
@@ -295,6 +297,7 @@ void SkeletalMesUtilsImpl::SaveSkeletalMeshLODModelSections(USkeletalMesh* Sourc
 	{
 		int32 SectionMaterialIndex = CurrentSection.MaterialIndex;
 		bool SectionCastShadow = CurrentSection.bCastShadow;
+		bool SectionVisibleInRayTracing = CurrentSection.bVisibleInRayTracing;
 		bool SectionRecomputeTangents = CurrentSection.bRecomputeTangent;
 		ESkinVertexColorChannel RecomputeTangentsVertexMaskChannel = CurrentSection.RecomputeTangentsVertexMaskChannel;
 		int32 GenerateUpTo = CurrentSection.GenerateUpToLodIndex;
@@ -303,7 +306,7 @@ void SkeletalMesUtilsImpl::SaveSkeletalMeshLODModelSections(USkeletalMesh* Sourc
 		//Save all the sections, even the chunked sections
 		if (ExistingMeshDataPtr->ExistingImportMaterialOriginalNameData.IsValidIndex(SectionMaterialIndex))
 		{
-			ExistingMeshDataPtr->ExistingImportMeshLodSectionMaterialData[LodIndex].Emplace(ExistingMeshDataPtr->ExistingImportMaterialOriginalNameData[SectionMaterialIndex], SectionCastShadow, SectionRecomputeTangents, RecomputeTangentsVertexMaskChannel, GenerateUpTo, bDisabled);
+			ExistingMeshDataPtr->ExistingImportMeshLodSectionMaterialData[LodIndex].Emplace(ExistingMeshDataPtr->ExistingImportMaterialOriginalNameData[SectionMaterialIndex], SectionCastShadow, SectionVisibleInRayTracing, SectionRecomputeTangents, RecomputeTangentsVertexMaskChannel, GenerateUpTo, bDisabled);
 		}
 	}
 }
@@ -1086,6 +1089,7 @@ void SkeletalMesUtilsImpl::RestoreMaterialNameWorkflowSection(const TSharedPtr<c
 
 			//We have a match put back the data
 			NewSection.bCastShadow = MeshData->ExistingImportMeshLodSectionMaterialData[LodIndex][ExistSectionIndex].bCastShadow;
+			NewSection.bVisibleInRayTracing = MeshData->ExistingImportMeshLodSectionMaterialData[LodIndex][ExistSectionIndex].bVisibleInRayTracing;
 			NewSection.bRecomputeTangent = MeshData->ExistingImportMeshLodSectionMaterialData[LodIndex][ExistSectionIndex].bRecomputeTangents;
 			NewSection.RecomputeTangentsVertexMaskChannel = MeshData->ExistingImportMeshLodSectionMaterialData[LodIndex][ExistSectionIndex].RecomputeTangentsVertexMaskChannel;
 			NewSection.GenerateUpToLodIndex = MeshData->ExistingImportMeshLodSectionMaterialData[LodIndex][ExistSectionIndex].GenerateUpTo;
@@ -1098,6 +1102,7 @@ void SkeletalMesUtilsImpl::RestoreMaterialNameWorkflowSection(const TSharedPtr<c
 				FSkelMeshSourceSectionUserData& UserSectionData = SkeletalMeshLodModel.UserSectionsData.FindOrAdd(ParentOriginalSectionIndex);
 				UserSectionData.bDisabled = NewSection.bDisabled;
 				UserSectionData.bCastShadow = NewSection.bCastShadow;
+				UserSectionData.bVisibleInRayTracing = NewSection.bVisibleInRayTracing;
 				UserSectionData.bRecomputeTangent = NewSection.bRecomputeTangent;
 				UserSectionData.RecomputeTangentsVertexMaskChannel = NewSection.RecomputeTangentsVertexMaskChannel;
 				UserSectionData.GenerateUpToLodIndex = NewSection.GenerateUpToLodIndex;
diff --git a/Engine/Source/Editor/UnrealEd/Public/ImportUtils/SkelImport.h b/Engine/Source/Editor/UnrealEd/Public/ImportUtils/SkelImport.h
--- a/Engine/Source/Editor/UnrealEd/Public/ImportUtils/SkelImport.h
+++ b/Engine/Source/Editor/UnrealEd/Public/ImportUtils/SkelImport.h
@@ -22,9 +22,10 @@ class FSkeletalMeshLODModel;
 
 struct FExistingMeshLodSectionData
 {
-	FExistingMeshLodSectionData(FName InImportedMaterialSlotName, bool InbCastShadow, bool InbRecomputeTangents, ESkinVertexColorChannel InRecomputeTangentsVertexMaskChannel, int32 InGenerateUpTo, bool InbDisabled)
+	FExistingMeshLodSectionData(FName InImportedMaterialSlotName, bool InbCastShadow, bool InbVisibleInRayTracing, bool InbRecomputeTangents, ESkinVertexColorChannel InRecomputeTangentsVertexMaskChannel, int32 InGenerateUpTo, bool InbDisabled)
 	: ImportedMaterialSlotName(InImportedMaterialSlotName)
 	, bCastShadow(InbCastShadow)
+	, bVisibleInRayTracing(InbVisibleInRayTracing)
 	, bRecomputeTangents(InbRecomputeTangents)
 	, RecomputeTangentsVertexMaskChannel(InRecomputeTangentsVertexMaskChannel)
 	, GenerateUpTo(InGenerateUpTo)
@@ -32,6 +33,7 @@ struct FExistingMeshLodSectionData
 	{}
 	FName ImportedMaterialSlotName;
 	bool bCastShadow;
+	bool bVisibleInRayTracing;
 	bool bRecomputeTangents;
 	ESkinVertexColorChannel RecomputeTangentsVertexMaskChannel;
 	int32 GenerateUpTo;
diff --git a/Engine/Source/Runtime/Core/Public/UObject/UE5MainStreamObjectVersion.h b/Engine/Source/Runtime/Core/Public/UObject/UE5MainStreamObjectVersion.h
--- a/Engine/Source/Runtime/Core/Public/UObject/UE5MainStreamObjectVersion.h
+++ b/Engine/Source/Runtime/Core/Public/UObject/UE5MainStreamObjectVersion.h
@@ -172,6 +172,9 @@ struct CORE_API FUE5MainStreamObjectVersion
 		// Ray tracing shadows have three states now (Disabled, Use Project Settings, Enabled)
 		RayTracedShadowsType,
 
+		// Add bVisibleInRayTracing flag to Skeletal Mesh Sections
+		SkelMeshSectionVisibleInRayTracingFlagAdded,
+
 		// -----<new versions can be added above this line>-------------------------------------------------
 		VersionPlusOne,
 		LatestVersion = VersionPlusOne - 1
diff --git a/Engine/Source/Runtime/Engine/Private/GPUSkinCache.cpp b/Engine/Source/Runtime/Engine/Private/GPUSkinCache.cpp
--- a/Engine/Source/Runtime/Engine/Private/GPUSkinCache.cpp
+++ b/Engine/Source/Runtime/Engine/Private/GPUSkinCache.cpp
@@ -1777,7 +1777,7 @@ void FGPUSkinCache::ProcessRayTracingGeometryToUpdate(
 				Segment.NumPrimitives = Section.NumTriangles;
 
 				// TODO: If we are at a dropped LOD, route material index through the LODMaterialMap in the LODInfo struct.
-				Segment.bEnabled = !SkinCacheEntry->GPUSkin->IsMaterialHidden(LODIndex, Section.MaterialIndex) && !Section.bDisabled;
+				Segment.bEnabled = !SkinCacheEntry->GPUSkin->IsMaterialHidden(LODIndex, Section.MaterialIndex) && !Section.bDisabled && Section.bVisibleInRayTracing;
 				Initializer.Segments.Add(Segment);
 			}
 
diff --git a/Engine/Source/Runtime/Engine/Private/SkeletalMesh.cpp b/Engine/Source/Runtime/Engine/Private/SkeletalMesh.cpp
--- a/Engine/Source/Runtime/Engine/Private/SkeletalMesh.cpp
+++ b/Engine/Source/Runtime/Engine/Private/SkeletalMesh.cpp
@@ -906,6 +906,7 @@ void USkeletalMesh::InitResources()
 					const FSkelMeshSourceSectionUserData& SectionUserData = ImportLODModel.UserSectionsData.FindChecked(ImportSection.OriginalDataSectionIndex);
 					bool bImportDataInSync = SectionUserData.bDisabled == ImportSection.bDisabled &&
 						SectionUserData.bCastShadow == ImportSection.bCastShadow &&
+						SectionUserData.bVisibleInRayTracing == ImportSection.bVisibleInRayTracing &&
 						SectionUserData.bRecomputeTangent == ImportSection.bRecomputeTangent &&
 						SectionUserData.RecomputeTangentsVertexMaskChannel == ImportSection.RecomputeTangentsVertexMaskChannel;
 					//Check the cloth only for parent section, since chunked section should not have cloth
@@ -920,6 +921,7 @@ void USkeletalMesh::InitResources()
 					const FSkelMeshRenderSection& RenderSection = RenderLODModel.RenderSections[SectionIndex];
 					bool bRenderDataInSync = SectionUserData.bDisabled == RenderSection.bDisabled &&
 						SectionUserData.bCastShadow == RenderSection.bCastShadow &&
+						SectionUserData.bVisibleInRayTracing == RenderSection.bVisibleInRayTracing &&
 						SectionUserData.bRecomputeTangent == RenderSection.bRecomputeTangent &&
 						SectionUserData.RecomputeTangentsVertexMaskChannel == RenderSection.RecomputeTangentsVertexMaskChannel &&
 						SectionUserData.CorrespondClothAssetIndex == RenderSection.CorrespondClothAssetIndex &&
@@ -6453,7 +6455,7 @@ void FSkeletalMeshSceneProxy::GetDynamicRayTracingInstances(FRayTracingMaterialG
 					FRayTracingGeometrySegment Segment;
 					Segment.FirstPrimitive = Section.BaseIndex / 3;
 					Segment.NumPrimitives = Section.NumTriangles;
-					Segment.bEnabled = !MeshObject->IsMaterialHidden(LODIndex, SectionElementInfo.UseMaterialIndex) && !Section.bDisabled;
+					Segment.bEnabled = !MeshObject->IsMaterialHidden(LODIndex, SectionElementInfo.UseMaterialIndex) && !Section.bDisabled && Section.bVisibleInRayTracing;
 					GeometrySections.Add(Segment);
 				}
 
diff --git a/Engine/Source/Runtime/Engine/Private/SkeletalMeshLODModel.cpp b/Engine/Source/Runtime/Engine/Private/SkeletalMeshLODModel.cpp
--- a/Engine/Source/Runtime/Engine/Private/SkeletalMeshLODModel.cpp
+++ b/Engine/Source/Runtime/Engine/Private/SkeletalMeshLODModel.cpp
@@ -305,6 +305,7 @@ FArchive& operator<<(FArchive& Ar, FSkelMeshSection& S)
 	Ar.UsingCustomVersion(FRenderingObjectVersion::GUID);
 	Ar.UsingCustomVersion(FAnimObjectVersion::GUID); // Also used by FSoftSkinVertex serializer
 	Ar.UsingCustomVersion(FSkeletalMeshCustomVersion::GUID);
+	Ar.UsingCustomVersion(FUE5MainStreamObjectVersion::GUID);
 	Ar.UsingCustomVersion(FRecomputeTangentCustomVersion::GUID);
 	Ar.UsingCustomVersion(FOverlappingVerticesCustomVersion::GUID);
 
@@ -386,6 +387,16 @@ FArchive& operator<<(FArchive& Ar, FSkelMeshSection& S)
 		S.bCastShadow = true;
 	}
 
+	if (Ar.CustomVer(FUE5MainStreamObjectVersion::GUID) >= FUE5MainStreamObjectVersion::SkelMeshSectionVisibleInRayTracingFlagAdded)
+	{
+		Ar << S.bVisibleInRayTracing;
+	}
+	else
+	{
+		// default is to be visible in ray tracing - which is consistent with behaviour before adding this member
+		S.bVisibleInRayTracing = true;
+	}
+
 	if (Ar.CustomVer(FSkeletalMeshCustomVersion::GUID) >= FSkeletalMeshCustomVersion::CombineSectionWithChunk)
 	{
 
@@ -534,6 +545,7 @@ FArchive& operator<<(FArchive& Ar, FSkelMeshSection& S)
 // Serialization.
 FArchive& operator<<(FArchive& Ar, FSkelMeshSourceSectionUserData& S)
 {
+	Ar.UsingCustomVersion(FUE5MainStreamObjectVersion::GUID);
 	Ar.UsingCustomVersion(FRecomputeTangentCustomVersion::GUID); 
 
 	FStripDataFlags StripFlags(Ar);
@@ -556,6 +568,17 @@ FArchive& operator<<(FArchive& Ar, FSkelMeshSourceSectionUserData& S)
 	}
 
 	Ar << S.bCastShadow;
+
+	if (Ar.CustomVer(FUE5MainStreamObjectVersion::GUID) >= FUE5MainStreamObjectVersion::SkelMeshSectionVisibleInRayTracingFlagAdded)
+	{
+		Ar << S.bVisibleInRayTracing;
+	}
+	else
+	{
+		// default is to be visible in ray tracing - which is consistent with behaviour before adding this member
+		S.bVisibleInRayTracing = true;
+	}
+
 	Ar << S.bDisabled;
 	Ar << S.GenerateUpToLodIndex;
 	Ar << S.CorrespondClothAssetIndex;
@@ -1082,6 +1105,7 @@ void FSkeletalMeshLODModel::SyncronizeUserSectionsDataArray(bool bResetNonUsedSe
 		FSkelMeshSection& Section = Sections[SectionIndex];
 		FSkelMeshSourceSectionUserData& SectionUserData = UserSectionsData.FindOrAdd(Section.OriginalDataSectionIndex);
 		Section.bCastShadow					= SectionUserData.bCastShadow;
+		Section.bVisibleInRayTracing		= SectionUserData.bVisibleInRayTracing;
 		Section.bRecomputeTangent			= SectionUserData.bRecomputeTangent;
 		Section.RecomputeTangentsVertexMaskChannel = SectionUserData.RecomputeTangentsVertexMaskChannel;
 		Section.bDisabled					= SectionUserData.bDisabled;
@@ -1106,6 +1130,7 @@ void FSkeletalMeshLODModel::SyncronizeUserSectionsDataArray(bool bResetNonUsedSe
 			}
 			FSkelMeshSourceSectionUserData& SectionUserData = UserSectionsData.FindOrAdd(Section.OriginalDataSectionIndex);
 			SectionUserData.bCastShadow = Section.bCastShadow;
+			SectionUserData.bVisibleInRayTracing = Section.bVisibleInRayTracing;
 			SectionUserData.bRecomputeTangent = Section.bRecomputeTangent;
 			SectionUserData.RecomputeTangentsVertexMaskChannel = Section.RecomputeTangentsVertexMaskChannel;
 			SectionUserData.bDisabled = Section.bDisabled;
@@ -1173,6 +1198,7 @@ void FSkeletalMeshLODModel::UpdateChunkedSectionInfo(const FString& SkeletalMesh
 			FSkelMeshSourceSectionUserData& SectionUserData = UserSectionsData.FindOrAdd(Section.OriginalDataSectionIndex);
 			Section.bDisabled = SectionUserData.bDisabled;
 			Section.bCastShadow = SectionUserData.bCastShadow;
+			Section.bVisibleInRayTracing = SectionUserData.bVisibleInRayTracing;
 			Section.bRecomputeTangent = SectionUserData.bRecomputeTangent;
 			Section.RecomputeTangentsVertexMaskChannel = SectionUserData.RecomputeTangentsVertexMaskChannel;
 			Section.GenerateUpToLodIndex = SectionUserData.GenerateUpToLodIndex;
@@ -1187,6 +1213,7 @@ void FSkeletalMeshLODModel::UpdateChunkedSectionInfo(const FString& SkeletalMesh
 			FSkelMeshSourceSectionUserData& SectionUserData = UserSectionsData.FindOrAdd(OriginalIndex);
 			SectionUserData.bDisabled = Section.bDisabled;
 			SectionUserData.bCastShadow = Section.bCastShadow;
+			SectionUserData.bVisibleInRayTracing = Section.bVisibleInRayTracing;
 			SectionUserData.bRecomputeTangent = Section.bRecomputeTangent;
 			SectionUserData.RecomputeTangentsVertexMaskChannel = Section.RecomputeTangentsVertexMaskChannel;
 			SectionUserData.GenerateUpToLodIndex = Section.GenerateUpToLodIndex;
diff --git a/Engine/Source/Runtime/Engine/Private/SkeletalMeshLODRenderData.cpp b/Engine/Source/Runtime/Engine/Private/SkeletalMeshLODRenderData.cpp
--- a/Engine/Source/Runtime/Engine/Private/SkeletalMeshLODRenderData.cpp
+++ b/Engine/Source/Runtime/Engine/Private/SkeletalMeshLODRenderData.cpp
@@ -12,6 +12,7 @@
 #include "PlatformInfo.h"
 #include "UObject/PropertyPortFlags.h"
 #include "UObject/UE5ReleaseStreamObjectVersion.h"
+#include "UObject/UE5MainStreamObjectVersion.h"
 #include "GPUSkinCache.h"
 
 #if WITH_EDITOR
@@ -74,6 +75,7 @@ FArchive& operator<<(FArchive& Ar, FSkelMeshRenderSection& S)
 {
 	const uint8 DuplicatedVertices = 1;
 	
+	Ar.UsingCustomVersion(FUE5MainStreamObjectVersion::GUID);
 	Ar.UsingCustomVersion(FRecomputeTangentCustomVersion::GUID);
 
 	// DuplicatedVerticesBuffer is used only for SkinCache and Editor features which is SM5 only
@@ -99,6 +101,15 @@ FArchive& operator<<(FArchive& Ar, FSkelMeshRenderSection& S)
 		S.RecomputeTangentsVertexMaskChannel = ESkinVertexColorChannel::None;
 	}
 	Ar << S.bCastShadow;
+	if (Ar.CustomVer(FUE5MainStreamObjectVersion::GUID) >= FUE5MainStreamObjectVersion::SkelMeshSectionVisibleInRayTracingFlagAdded)
+	{
+		Ar << S.bVisibleInRayTracing;
+	}
+	else
+	{
+		// default is to be visible in ray tracing - which is consistent with behaviour before adding this member
+		S.bVisibleInRayTracing = true;
+	}
 	Ar << S.BaseVertexIndex;
 	Ar << S.ClothMappingData;
 	Ar << S.BoneMap;
@@ -366,6 +377,7 @@ void FSkeletalMeshLODRenderData::BuildFromLODModel(const FSkeletalMeshLODModel*
 		NewRenderSection.bRecomputeTangent = ModelSection.bRecomputeTangent;
 		NewRenderSection.RecomputeTangentsVertexMaskChannel = ModelSection.RecomputeTangentsVertexMaskChannel;
 		NewRenderSection.bCastShadow = ModelSection.bCastShadow;
+		NewRenderSection.bVisibleInRayTracing = ModelSection.bVisibleInRayTracing;
 		NewRenderSection.BaseVertexIndex = ModelSection.BaseVertexIndex;
 		NewRenderSection.ClothMappingData = ModelSection.ClothMappingData;
 		NewRenderSection.BoneMap = ModelSection.BoneMap;
diff --git a/Engine/Source/Runtime/Engine/Private/SkeletalMeshRenderData.cpp b/Engine/Source/Runtime/Engine/Private/SkeletalMeshRenderData.cpp
--- a/Engine/Source/Runtime/Engine/Private/SkeletalMeshRenderData.cpp
+++ b/Engine/Source/Runtime/Engine/Private/SkeletalMeshRenderData.cpp
@@ -186,7 +186,7 @@ static void SerializeLODInfoForDDC(USkeletalMesh* SkeletalMesh, FString& KeySuff
 // differences, etc.) replace the version GUID below with a new one.
 // In case of merge conflicts with DDC versions, you MUST generate a new GUID
 // and set this new GUID as the version.
-#define SKELETALMESH_DERIVEDDATA_VER TEXT("1BB523EE3EA9483F90CB4FE406DDF29E")
+#define SKELETALMESH_DERIVEDDATA_VER TEXT("F1115E3EF0EE404692FABA5C90B76988")
 
 const FString& GetSkeletalMeshDerivedDataVersion()
 {
diff --git a/Engine/Source/Runtime/Engine/Private/SkeletalRenderCPUSkin.cpp b/Engine/Source/Runtime/Engine/Private/SkeletalRenderCPUSkin.cpp
--- a/Engine/Source/Runtime/Engine/Private/SkeletalRenderCPUSkin.cpp
+++ b/Engine/Source/Runtime/Engine/Private/SkeletalRenderCPUSkin.cpp
@@ -406,7 +406,7 @@ void FSkeletalMeshObjectCPUSkin::FSkeletalMeshObjectLOD::InitResources(FSkelMesh
 					Segment.MaxVertices = Section.GetNumVertices();
 					Segment.FirstPrimitive = Section.BaseIndex / 3;
 					Segment.NumPrimitives = Section.NumTriangles;
-					Segment.bEnabled = !Section.bDisabled;
+					Segment.bEnabled = !Section.bDisabled && Section.bVisibleInRayTracing;
 					GeometrySections.Add(Segment);
 				}
 				Initializer.Segments = GeometrySections;
diff --git a/Engine/Source/Runtime/Engine/Private/SkeletalRenderStatic.cpp b/Engine/Source/Runtime/Engine/Private/SkeletalRenderStatic.cpp
--- a/Engine/Source/Runtime/Engine/Private/SkeletalRenderStatic.cpp
+++ b/Engine/Source/Runtime/Engine/Private/SkeletalRenderStatic.cpp
@@ -93,7 +93,7 @@ void FSkeletalMeshObjectStatic::InitResources(USkinnedMeshComponent* InMeshCompo
 								Segment.MaxVertices = Section.GetNumVertices();
 								Segment.FirstPrimitive = Section.BaseIndex / 3;
 								Segment.NumPrimitives = Section.NumTriangles;
-								Segment.bEnabled = !Section.bDisabled;
+								Segment.bEnabled = !Section.bDisabled && Section.bVisibleInRayTracing;
 								GeometrySections.Add(Segment);
 							}
 							Initializer.Segments = GeometrySections;
diff --git a/Engine/Source/Runtime/Engine/Public/Rendering/SkeletalMeshLODModel.h b/Engine/Source/Runtime/Engine/Public/Rendering/SkeletalMeshLODModel.h
--- a/Engine/Source/Runtime/Engine/Public/Rendering/SkeletalMeshLODModel.h
+++ b/Engine/Source/Runtime/Engine/Public/Rendering/SkeletalMeshLODModel.h
@@ -82,6 +82,9 @@ struct FSkelMeshSection
 
 	/** This section will cast shadow */
 	bool bCastShadow;
+	
+	/** If true, this section will be visible in ray tracing effects. Turning this off will remove it from ray traced reflections, shadows, etc. */
+	bool bVisibleInRayTracing;
 
 	/** This is set for old 'duplicate' sections that need to get removed on load */
 	bool bLegacyClothingSection_DEPRECATED;
@@ -157,6 +160,7 @@ struct FSkelMeshSection
 		, bRecomputeTangent(false)
 		, RecomputeTangentsVertexMaskChannel(ESkinVertexColorChannel::None)
 		, bCastShadow(true)
+		, bVisibleInRayTracing(true)
 		, bLegacyClothingSection_DEPRECATED(false)
 		, CorrespondClothSectionIndex_DEPRECATED(-1)
 		, BaseVertexIndex(0)
@@ -238,6 +242,9 @@ struct FSkelMeshSourceSectionUserData
 	/** This section will cast shadow */
 	bool bCastShadow;
 
+	/** If true, this section will be visible in ray tracing effects. Turning this off will remove it from ray traced reflections, shadows, etc. */
+	bool bVisibleInRayTracing;
+
 	// INDEX_NONE if not set
 	int16 CorrespondClothAssetIndex;
 
@@ -266,6 +273,7 @@ struct FSkelMeshSourceSectionUserData
 		: bRecomputeTangent(false)
 		, RecomputeTangentsVertexMaskChannel(ESkinVertexColorChannel::None)
 		, bCastShadow(true)
+		, bVisibleInRayTracing(true)
 		, CorrespondClothAssetIndex(INDEX_NONE)
 		, bDisabled(false)
 		, GenerateUpToLodIndex(INDEX_NONE)
@@ -287,6 +295,7 @@ struct FSkelMeshSourceSectionUserData
 			//If the UserSectionData do not exist add it and copy from the section data
 			UserSectionData = &UserSectionsData.Add(Section.OriginalDataSectionIndex);
 			UserSectionData->bCastShadow = Section.bCastShadow;
+			UserSectionData->bVisibleInRayTracing = Section.bVisibleInRayTracing;			
 			UserSectionData->bDisabled = Section.bDisabled;
 			UserSectionData->bRecomputeTangent = Section.bRecomputeTangent;
 			UserSectionData->RecomputeTangentsVertexMaskChannel = Section.RecomputeTangentsVertexMaskChannel;
diff --git a/Engine/Source/Runtime/Engine/Public/Rendering/SkeletalMeshLODRenderData.h b/Engine/Source/Runtime/Engine/Public/Rendering/SkeletalMeshLODRenderData.h
--- a/Engine/Source/Runtime/Engine/Public/Rendering/SkeletalMeshLODRenderData.h
+++ b/Engine/Source/Runtime/Engine/Public/Rendering/SkeletalMeshLODRenderData.h
@@ -37,6 +37,9 @@ struct FSkelMeshRenderSection
 	/** This section will cast shadow */
 	bool bCastShadow;
 
+	/** If true, this section will be visible in ray tracing effects. Turning this off will remove it from ray traced reflections, shadows, etc. */
+	bool bVisibleInRayTracing;
+
 	/** Which channel for masking the recompute tangents */
 	ESkinVertexColorChannel RecomputeTangentsVertexMaskChannel;
 
@@ -74,6 +77,7 @@ struct FSkelMeshRenderSection
 		, NumTriangles(0)
 		, bRecomputeTangent(false)
 		, bCastShadow(true)
+		, bVisibleInRayTracing(true)
 		, RecomputeTangentsVertexMaskChannel(ESkinVertexColorChannel::None)
 		, BaseVertexIndex(0)
 		, NumVertices(0)
