Persona: Don't draw bone capsules if we're not in hit proxy draw mode.

#rnx

diff --git a/Engine/Source/Editor/Persona/Private/PersonaSelectionComponent.cpp b/Engine/Source/Editor/Persona/Private/PersonaSelectionComponent.cpp
--- a/Engine/Source/Editor/Persona/Private/PersonaSelectionComponent.cpp
+++ b/Engine/Source/Editor/Persona/Private/PersonaSelectionComponent.cpp
@@ -3,6 +3,8 @@
 #include "PersonaSelectionComponent.h"
 #include "Materials/Material.h"
 
+PRAGMA_DISABLE_OPTIMIZATION
+
 IMPLEMENT_HIT_PROXY(HPersonaSelectionHitProxy, HHitProxy)
 
 const FPersonaSelectionCapsule& HPersonaSelectionHitProxy::GetCapsule() const
@@ -107,7 +109,7 @@ FBoxSphereBounds UPersonaSelectionComponent::CalcBounds(const FTransform& LocalT
 	{
 		const FTransform ComponentTransform = GetComponentToWorld();
 		
-		for(int32 CapsuleIndex : CapsuleIndices)
+		for(const int32 CapsuleIndex : CapsuleIndices)
 		{
 			const FPersonaSelectionCapsule& Capsule = Capsules[CapsuleIndex];
 			const FTransform& Transform = Capsule.Transform * ComponentTransform;
@@ -202,6 +204,12 @@ void FPersonaSelectionComponentProxy::GetDynamicMeshElements(const TArray<const
 		return;
 	}
 
+	// If we're not drawing hit proxies, then there's nothing for us to do.
+	if (!ViewFamily.EngineShowFlags.HitProxies)
+	{
+		return;
+	}
+
 	const UPersonaSelectionComponent* SelectionComponent = SelectionComponentPtr.Get();
 	FScopeLock ScopeLock( &SelectionComponent->CriticalSection );
 
@@ -214,16 +222,12 @@ void FPersonaSelectionComponentProxy::GetDynamicMeshElements(const TArray<const
 	{
 		if (VisibilityMap & (1 << ViewIndex))
 		{
-			const FSceneView* View = Views[ViewIndex];
-			FPrimitiveDrawInterface* PDI = Collector.GetPDI(ViewIndex);
-
-			FMaterialRenderProxy* MaterialRenderProxy = GEngine->DefaultFlattenMaterial->GetRenderProxy();
+			const FMaterialRenderProxy* MaterialRenderProxy = GEngine->DefaultFlattenMaterial->GetRenderProxy();
 			const FTransform ComponentToWorld = SelectionComponent->GetComponentToWorld();
 
-			for(int32 CapsuleIndex : SelectionComponent->CapsuleIndices)
+			for(const int32 CapsuleIndex : SelectionComponent->CapsuleIndices)
             {
 				const FPersonaSelectionCapsule& Capsule = SelectionComponent->Capsules[CapsuleIndex];
-				const FLinearColor LineColor = FLinearColor::Yellow;
 				const FTransform& Transform = Capsule.Transform * ComponentToWorld;
 
 				if(!SelectionComponent->HitProxies.IsValidIndex(CapsuleIndex))
@@ -251,7 +255,7 @@ void FPersonaSelectionComponentProxy::GetDynamicMeshElements(const TArray<const
 
 /**
 *  Returns a struct that describes to the renderer when to draw this proxy.
-*	@param		Scene view to use to determine our relevence.
+*	@param View Scene view to use to determine our relevance.
 *  @return		View relevance struct
 */
 FPrimitiveViewRelevance FPersonaSelectionComponentProxy::GetViewRelevance(const FSceneView* View) const
@@ -275,3 +279,5 @@ uint32 FPersonaSelectionComponentProxy::GetAllocatedSize(void) const
 {
 	return FPrimitiveSceneProxy::GetAllocatedSize();
 }
+
+PRAGMA_ENABLE_OPTIMIZATION
\ No newline at end of file
