remove PRAGMA_DISABLE_OPTIMIZATION

diff --git a/Engine/Source/Editor/Persona/Private/PersonaSelectionComponent.cpp b/Engine/Source/Editor/Persona/Private/PersonaSelectionComponent.cpp
--- a/Engine/Source/Editor/Persona/Private/PersonaSelectionComponent.cpp
+++ b/Engine/Source/Editor/Persona/Private/PersonaSelectionComponent.cpp
@@ -3,8 +3,6 @@
 #include "PersonaSelectionComponent.h"
 #include "Materials/Material.h"
 
-PRAGMA_DISABLE_OPTIMIZATION
-
 IMPLEMENT_HIT_PROXY(HPersonaSelectionHitProxy, HHitProxy)
 
 const FPersonaSelectionCapsule& HPersonaSelectionHitProxy::GetCapsule() const
@@ -279,5 +277,3 @@ uint32 FPersonaSelectionComponentProxy::GetAllocatedSize(void) const
 {
 	return FPrimitiveSceneProxy::GetAllocatedSize();
 }
-
-PRAGMA_ENABLE_OPTIMIZATION
\ No newline at end of file
