Remove ability to lock samples from BlendSpace. This didn't go into EA so there shouldn't be any backwards capatibility problems.

diff --git a/Engine/Plugins/Animation/BlendSpaceMotionAnalysis/Source/BlendSpaceMotionAnalysis/Private/LocomotionAnalysis.cpp b/Engine/Plugins/Animation/BlendSpaceMotionAnalysis/Source/BlendSpaceMotionAnalysis/Private/LocomotionAnalysis.cpp
--- a/Engine/Plugins/Animation/BlendSpaceMotionAnalysis/Source/BlendSpaceMotionAnalysis/Private/LocomotionAnalysis.cpp
+++ b/Engine/Plugins/Animation/BlendSpaceMotionAnalysis/Source/BlendSpaceMotionAnalysis/Private/LocomotionAnalysis.cpp
@@ -28,7 +28,6 @@ void ULocomotionAnalysisProperties::InitializeFromCache(TObjectPtr<UCachedAnalys
 		SecondaryBoneSocket = CachedMotionAnalysisPropertiesPtr->BoneSocket2;
 		CharacterFacingAxis = CachedMotionAnalysisPropertiesPtr->CharacterFacingAxis;
 		CharacterUpAxis = CachedMotionAnalysisPropertiesPtr->CharacterUpAxis;
-		bLockAfterAnalysis = CachedMotionAnalysisPropertiesPtr->bLockAfterAnalysis;
 	}
 }
 
@@ -51,7 +50,6 @@ void ULocomotionAnalysisProperties::MakeCache(TObjectPtr<UCachedAnalysisProperti
 	CachedMotionAnalysisPropertiesPtr->BoneSocket2 = SecondaryBoneSocket;
 	CachedMotionAnalysisPropertiesPtr->CharacterFacingAxis = CharacterFacingAxis;
 	CachedMotionAnalysisPropertiesPtr->CharacterUpAxis = CharacterUpAxis;
-	CachedMotionAnalysisPropertiesPtr->bLockAfterAnalysis = bLockAfterAnalysis;
 }
 
 //======================================================================================================================
diff --git a/Engine/Plugins/Animation/BlendSpaceMotionAnalysis/Source/BlendSpaceMotionAnalysis/Private/RootMotionAnalysis.cpp b/Engine/Plugins/Animation/BlendSpaceMotionAnalysis/Source/BlendSpaceMotionAnalysis/Private/RootMotionAnalysis.cpp
--- a/Engine/Plugins/Animation/BlendSpaceMotionAnalysis/Source/BlendSpaceMotionAnalysis/Private/RootMotionAnalysis.cpp
+++ b/Engine/Plugins/Animation/BlendSpaceMotionAnalysis/Source/BlendSpaceMotionAnalysis/Private/RootMotionAnalysis.cpp
@@ -31,7 +31,6 @@ void URootMotionAnalysisProperties::InitializeFromCache(TObjectPtr<UCachedAnalys
 		CharacterUpAxis = CachedMotionAnalysisPropertiesPtr->CharacterUpAxis;
 		StartTimeFraction = CachedMotionAnalysisPropertiesPtr->StartTimeFraction;
 		EndTimeFraction = CachedMotionAnalysisPropertiesPtr->EndTimeFraction;
-		bLockAfterAnalysis = CachedMotionAnalysisPropertiesPtr->bLockAfterAnalysis;
 	}
 }
 
@@ -57,7 +56,6 @@ void URootMotionAnalysisProperties::MakeCache(TObjectPtr<UCachedAnalysisProperti
 	CachedMotionAnalysisPropertiesPtr->CharacterUpAxis = CharacterUpAxis;
 	CachedMotionAnalysisPropertiesPtr->StartTimeFraction = StartTimeFraction;
 	CachedMotionAnalysisPropertiesPtr->EndTimeFraction = EndTimeFraction;
-	CachedMotionAnalysisPropertiesPtr->bLockAfterAnalysis = bLockAfterAnalysis;
 }
 
 //======================================================================================================================
diff --git a/Engine/Source/Editor/Persona/Private/Customization/BlendSampleDetails.cpp b/Engine/Source/Editor/Persona/Private/Customization/BlendSampleDetails.cpp
--- a/Engine/Source/Editor/Persona/Private/Customization/BlendSampleDetails.cpp
+++ b/Engine/Source/Editor/Persona/Private/Customization/BlendSampleDetails.cpp
@@ -177,24 +177,6 @@ void FBlendSampleDetails::CustomizeDetails(class IDetailLayoutBuilder& DetailBui
 		}, GridWidget->OnSampleMoved, (BlendSpaceBase != nullptr) ? BlendSpaceBase : BlendSpace,
 		GridWidget->GetSelectedSampleIndex(), false);
 
-	// Locked
-	const bool b1DBlendSpace = BlendSpace->IsA<UBlendSpace1D>();
-	{
-		IDetailPropertyRow& LockPositionXPropertyRow = CategoryBuilder.AddProperty(DetailBuilder.GetProperty(
-            GET_MEMBER_NAME_CHECKED(FBlendSample, bLockX), FBlendSample::StaticStruct()));
-		LockPositionXPropertyRow.DisplayName(FText::Format(
-            LOCTEXT("BlendSpaceSampleLockLabel", "{0} {1}"), FText::FromString(TEXT("Lock ")),
-            FText::FromString(BlendSpace->GetBlendParameter(0).DisplayName)));
-	}
-	if (!b1DBlendSpace)
-	{
-		IDetailPropertyRow& LockPositionYPropertyRow = CategoryBuilder.AddProperty(DetailBuilder.GetProperty(
-			GET_MEMBER_NAME_CHECKED(FBlendSample, bLockY), FBlendSample::StaticStruct()));
-		LockPositionYPropertyRow.DisplayName(FText::Format(
-			LOCTEXT("BlendSpaceSampleLockLabel", "{0} {1}"), FText::FromString(TEXT("Lock ")),
-			FText::FromString(BlendSpace->GetBlendParameter(1).DisplayName)));
-	}
-
 	// Animation and rate
 	if(BlendspaceToUse->IsAsset())
 	{
@@ -262,6 +244,7 @@ void FBlendSampleDetails::CustomizeDetails(class IDetailLayoutBuilder& DetailBui
 				return OrigValue != NewValue;
 			});
 
+		const bool b1DBlendSpace = BlendSpace->IsA<UBlendSpace1D>();
 
 		if (bShowAnalysis)
 		{
diff --git a/Engine/Source/Editor/Persona/Private/Customization/BlendSpaceDetails.cpp b/Engine/Source/Editor/Persona/Private/Customization/BlendSpaceDetails.cpp
--- a/Engine/Source/Editor/Persona/Private/Customization/BlendSpaceDetails.cpp
+++ b/Engine/Source/Editor/Persona/Private/Customization/BlendSpaceDetails.cpp
@@ -59,10 +59,6 @@ FReply FBlendSpaceDetails::HandleAnalyzeSamples()
 		FSlateApplication::Get().DismissAllMenus();
 		bool bChangedOne = false;
 
-		bool bLockX = BlendSpaceAnalysis::GetLockAfterAnalysis(BlendSpace->AnalysisProperties[0]);
-		bool bLockY = BlendSpaceAnalysis::GetLockAfterAnalysis(BlendSpace->AnalysisProperties[1]);
-		bool bLockZ = BlendSpaceAnalysis::GetLockAfterAnalysis(BlendSpace->AnalysisProperties[2]);
-		
 		const int32 NumSamples = BlendSpace->SampleData.Num();
 		for (int32 SampleIndex = 0 ; SampleIndex != NumSamples ; ++SampleIndex)
 		{
@@ -78,11 +74,6 @@ FReply FBlendSpaceDetails::HandleAnalyzeSamples()
 				if (NewValue == BlendSpace->SampleData[SampleIndex].SampleValue)
 				{
 					bChangedOne = true;
-					BlendSpace->LockSample(
-						SampleIndex,
-						bAnalyzed[0] ? BlendSpaceAnalysis::GetLockAfterAnalysis(BlendSpace->AnalysisProperties[0]) : false,
-						bAnalyzed[1] ? BlendSpaceAnalysis::GetLockAfterAnalysis(BlendSpace->AnalysisProperties[1]) : false,
-						bAnalyzed[2] ? BlendSpaceAnalysis::GetLockAfterAnalysis(BlendSpace->AnalysisProperties[2]) : false);
 				}
 			}
 		}
@@ -500,28 +491,6 @@ void FBlendSpaceDetails::CustomizeDetails(class IDetailLayoutBuilder& DetailBuil
 					}
 				}), BlendSpace, SampleIndex, false);
 
-
-				// Locked
-				{
-					TSharedPtr<IPropertyHandle> LockPositionProperty = BlendSampleProperty->GetChildHandle(
-						GET_MEMBER_NAME_CHECKED(FBlendSample, bLockX));
-					IDetailPropertyRow& LockPositionPropertyRow = Group.AddPropertyRow(
-						LockPositionProperty.ToSharedRef());
-					LockPositionPropertyRow.DisplayName(FText::Format(
-						LOCTEXT("BlendSpaceSampleLockLabel", "{0} {1}"), FText::FromString(TEXT("Lock ")),
-						FText::FromString(BlendSpace->GetBlendParameter(0).DisplayName)));
-				}
-				if (!b1DBlendSpace)
-				{
-					TSharedPtr<IPropertyHandle> LockPositionProperty = BlendSampleProperty->GetChildHandle(
-						GET_MEMBER_NAME_CHECKED(FBlendSample, bLockY));
-					IDetailPropertyRow& LockPositionPropertyRow = Group.AddPropertyRow(
-						LockPositionProperty.ToSharedRef());
-					LockPositionPropertyRow.DisplayName(FText::Format(
-						LOCTEXT("BlendSpaceSampleLockLabel", "{0} {1}"), FText::FromString(TEXT("Lock ")),
-						FText::FromString(BlendSpace->GetBlendParameter(1).DisplayName)));
-				}
-				
 				if(BlendSpace->IsAsset())
 				{
 					FDetailWidgetRow& AnimationRow = Group.AddWidgetRow();
diff --git a/Engine/Source/Editor/Persona/Private/PersonaBlendSpaceAnalysis.cpp b/Engine/Source/Editor/Persona/Private/PersonaBlendSpaceAnalysis.cpp
--- a/Engine/Source/Editor/Persona/Private/PersonaBlendSpaceAnalysis.cpp
+++ b/Engine/Source/Editor/Persona/Private/PersonaBlendSpaceAnalysis.cpp
@@ -25,7 +25,6 @@ void UCachedAnalysisProperties::CopyFrom(const UCachedAnalysisProperties& Other)
 	CharacterUpAxis = Other.CharacterUpAxis;
 	StartTimeFraction = Other.StartTimeFraction;
 	EndTimeFraction = Other.EndTimeFraction;
-	bLockAfterAnalysis = Other.bLockAfterAnalysis;
 }
 
 //======================================================================================================================
@@ -92,7 +91,6 @@ void ULinearAnalysisProperties::InitializeFromCache(TObjectPtr<UCachedAnalysisPr
 		SpaceBoneSocket = CachePtr->SpaceBoneSocket;
 		StartTimeFraction = CachePtr->StartTimeFraction;
 		EndTimeFraction = CachePtr->EndTimeFraction;
-		bLockAfterAnalysis = CachePtr->bLockAfterAnalysis;
 	}
 }
 
@@ -111,7 +109,6 @@ void ULinearAnalysisProperties::MakeCache(TObjectPtr<UCachedAnalysisProperties>&
 	CachePtr->SpaceBoneSocket = SpaceBoneSocket;
 	CachePtr->StartTimeFraction = StartTimeFraction;
 	CachePtr->EndTimeFraction = EndTimeFraction;
-	CachePtr->bLockAfterAnalysis = bLockAfterAnalysis;
 }
 
 //======================================================================================================================
@@ -130,7 +127,6 @@ void UEulerAnalysisProperties::InitializeFromCache(TObjectPtr<UCachedAnalysisPro
 		CharacterUpAxis = CachePtr->CharacterUpAxis;
 		StartTimeFraction = CachePtr->StartTimeFraction;
 		EndTimeFraction = CachePtr->EndTimeFraction;
-		bLockAfterAnalysis = CachePtr->bLockAfterAnalysis;
 	}
 }
 
@@ -153,7 +149,6 @@ void UEulerAnalysisProperties::MakeCache(TObjectPtr<UCachedAnalysisProperties>&
 	CachePtr->CharacterUpAxis = CharacterUpAxis;
 	CachePtr->StartTimeFraction = StartTimeFraction;
 	CachePtr->EndTimeFraction = EndTimeFraction;
-	CachePtr->bLockAfterAnalysis = bLockAfterAnalysis;
 }
 
 //======================================================================================================================
@@ -357,10 +352,4 @@ TArray<FString> BlendSpaceAnalysis::GetAnalysisFunctions()
 	return FunctionNames;
 }
 
-//======================================================================================================================
-bool BlendSpaceAnalysis::GetLockAfterAnalysis(const TObjectPtr<UAnalysisProperties>& AnalysisProperties)
-{
-	return AnalysisProperties ? AnalysisProperties->bLockAfterAnalysis : false;
-}
-
 #undef LOCTEXT_NAMESPACE
diff --git a/Engine/Source/Editor/Persona/Private/PersonaBlendSpaceAnalysis.h b/Engine/Source/Editor/Persona/Private/PersonaBlendSpaceAnalysis.h
--- a/Engine/Source/Editor/Persona/Private/PersonaBlendSpaceAnalysis.h
+++ b/Engine/Source/Editor/Persona/Private/PersonaBlendSpaceAnalysis.h
@@ -28,10 +28,4 @@ UAnalysisProperties* MakeAnalysisProperties(UObject* Outer, const FString& Funct
  */
 TArray<FString> GetAnalysisFunctions();
 
-/**
- * Simple helper that avoids the need to check pointer validity when determining if a sample should be locked 
- * after analysis
- */
-bool GetLockAfterAnalysis(const TObjectPtr<UAnalysisProperties>& AnalysisProperties);
-
 };
diff --git a/Engine/Source/Editor/Persona/Private/SAnimationBlendSpace.cpp b/Engine/Source/Editor/Persona/Private/SAnimationBlendSpace.cpp
--- a/Engine/Source/Editor/Persona/Private/SAnimationBlendSpace.cpp
+++ b/Engine/Source/Editor/Persona/Private/SAnimationBlendSpace.cpp
@@ -191,12 +191,6 @@ int32 SBlendSpaceEditor::OnSampleAdded(UAnimSequence* Animation, const FVector&
 
 	if (NewSampleIndex >= 0)
 	{
-		BlendSpace->LockSample(
-			NewSampleIndex,
-			bAnalyzed[0] ? BlendSpaceAnalysis::GetLockAfterAnalysis(BlendSpace->AnalysisProperties[0]) : false,
-			bAnalyzed[1] ? BlendSpaceAnalysis::GetLockAfterAnalysis(BlendSpace->AnalysisProperties[1]) : false,
-			bAnalyzed[2] ? BlendSpaceAnalysis::GetLockAfterAnalysis(BlendSpace->AnalysisProperties[2]) : false);
-
 		ResampleData();
 		BlendSpace->ValidateSampleData();
 
@@ -219,11 +213,7 @@ void SBlendSpaceEditor::OnSampleDuplicated(const int32 SampleIndex, const FVecto
 	BlendSpace->Modify();
 
 	const FBlendSample& OrigSample = BlendSpace->GetBlendSample(SampleIndex);
-	int32 NewSampleIndex = OnSampleAdded(OrigSample.Animation, NewValue, bRunAnalysis);
-	if (NewSampleIndex >= 0)
-	{
-		BlendSpace->LockSample(NewSampleIndex, OrigSample.bLockX, OrigSample.bLockY, OrigSample.bLockZ);
-	}
+	OnSampleAdded(OrigSample.Animation, NewValue, bRunAnalysis);
 }
 
 void SBlendSpaceEditor::OnSampleReplaced(int32 InSampleIndex, UAnimSequence* Animation)
diff --git a/Engine/Source/Editor/Persona/Private/SAnimationBlendSpaceGridWidget.cpp b/Engine/Source/Editor/Persona/Private/SAnimationBlendSpaceGridWidget.cpp
--- a/Engine/Source/Editor/Persona/Private/SAnimationBlendSpaceGridWidget.cpp
+++ b/Engine/Source/Editor/Persona/Private/SAnimationBlendSpaceGridWidget.cpp
@@ -638,42 +638,6 @@ void SBlendSpaceGridWidget::PaintSampleKeys(
 				PaintCircle(CirclePosition, 8.0f, 12, AllottedGeometry, IsolatedColor, IsolatedColor,
 							LabelBrush, OutDrawElements, SampleLayer - 1);
 			}
-
-			// Draw lines/boxes to indicated which axes are locked
-			if (!bReadOnly)
-			{
-				TArray<FVector2D> LinePoints;
-				LinePoints.SetNumUninitialized(2);
-				const float LockSizeScale = 0.5f; // Have the markers just poke out
-				const FVector2D SamplePosition = SampleValueToScreenPosition(Sample.SampleValue) + FVector2D(0.0, 0.0);
-				if (Sample.bLockX)
-				{
-					LinePoints[0] = FVector2D(SamplePosition - FVector2D(KeySize.X, -KeySize.Y) * LockSizeScale);
-					LinePoints[1] = FVector2D(SamplePosition - FVector2D(KeySize.X, +KeySize.Y) * LockSizeScale);
-					FSlateDrawElement::MakeLines(
-						OutDrawElements, SampleLayer - 1, AllottedGeometry.ToPaintGeometry(),
-						LinePoints, ESlateDrawEffect::None, InvalidColor.GetSpecifiedColor(), true, 1.0f);
-					LinePoints[0] = FVector2D(SamplePosition + FVector2D(KeySize.X, -KeySize.Y) * LockSizeScale);
-					LinePoints[1] = FVector2D(SamplePosition + FVector2D(KeySize.X, +KeySize.Y) * LockSizeScale);
-					FSlateDrawElement::MakeLines(
-						OutDrawElements, SampleLayer - 1, AllottedGeometry.ToPaintGeometry(),
-						LinePoints, ESlateDrawEffect::None, InvalidColor.GetSpecifiedColor(), true, 1.0f);
-				}
-				const bool b1DBlendSpace = BlendSpace->IsA<UBlendSpace1D>();
-				if (Sample.bLockY && !BlendSpace->IsA<UBlendSpace1D>())
-				{
-					LinePoints[0] = FVector2D(SamplePosition - FVector2D(-KeySize.X, KeySize.X) * LockSizeScale);
-					LinePoints[1] = FVector2D(SamplePosition - FVector2D(+KeySize.X, KeySize.X) * LockSizeScale);
-					FSlateDrawElement::MakeLines(
-						OutDrawElements, SampleLayer - 1, AllottedGeometry.ToPaintGeometry(),
-						LinePoints, ESlateDrawEffect::None, InvalidColor.GetSpecifiedColor(), true, 1.0f);
-					LinePoints[0] = FVector2D(SamplePosition + FVector2D(-KeySize.X, KeySize.X) * LockSizeScale);
-					LinePoints[1] = FVector2D(SamplePosition + FVector2D(+KeySize.X, KeySize.X) * LockSizeScale);
-					FSlateDrawElement::MakeLines(
-						OutDrawElements, SampleLayer - 1, AllottedGeometry.ToPaintGeometry(),
-						LinePoints, ESlateDrawEffect::None, InvalidColor.GetSpecifiedColor(), true, 1.0f);
-				}
-			}
 		}
 
 		// Always draw the filtered position which comes back from whatever is running
@@ -2503,9 +2467,6 @@ void SBlendSpaceGridWidget::Tick(const FGeometry& AllottedGeometry, const double
 				FVector SampleValue = ScreenPositionToSampleValueWithSnapping(LocalMousePosition, FSlateApplication::Get().GetModifierKeys().IsShiftDown());
 				// Only allow dragging on each axis if not locked
 				const FBlendSample& BlendSample = BlendSpace->GetBlendSample(DraggedSampleIndex);
-				SampleValue.X = BlendSample.bLockX ? BlendSample.SampleValue.X : SampleValue.X;
-				SampleValue.Y = BlendSample.bLockY ? BlendSample.SampleValue.Y : SampleValue.Y;
-				SampleValue.Z = BlendSample.bLockZ ? BlendSample.SampleValue.Z : SampleValue.Z;
 				if (SampleValue != LastDragPosition)
 				{
 					LastDragPosition = SampleValue;
diff --git a/Engine/Source/Editor/Persona/Public/BlendSpaceAnalysis.h b/Engine/Source/Editor/Persona/Public/BlendSpaceAnalysis.h
--- a/Engine/Source/Editor/Persona/Public/BlendSpaceAnalysis.h
+++ b/Engine/Source/Editor/Persona/Public/BlendSpaceAnalysis.h
@@ -94,7 +94,6 @@ public:
 	EAnalysisLinearAxis     CharacterUpAxis = EAnalysisLinearAxis::PlusZ;
 	float                   StartTimeFraction = 0.0f;
 	float                   EndTimeFraction = 1.0f;
-	bool                    bLockAfterAnalysis = false;
 };
 
 UCLASS()
diff --git a/Engine/Source/Runtime/Engine/Classes/Animation/BlendSpace.h b/Engine/Source/Runtime/Engine/Classes/Animation/BlendSpace.h
--- a/Engine/Source/Runtime/Engine/Classes/Animation/BlendSpace.h
+++ b/Engine/Source/Runtime/Engine/Classes/Animation/BlendSpace.h
@@ -37,10 +37,6 @@ public:
 	/** Analysis function for this axis */
 	UPROPERTY()
 	FString Function = TEXT("None");
-
-	/** If set then the sample values will be locked after analysis, preventing them from being moved by accident */
-	UPROPERTY(EditAnywhere, Category = AnalysisProperties)
-	bool bLockAfterAnalysis = false;
 };
 
 /** Interpolation data types. */
@@ -167,40 +163,28 @@ struct FBlendSample
 	UPROPERTY(transient)
 	uint8 bIsValid : 1;
 
-	UPROPERTY(EditAnywhere, Category=BlendSample)
-	uint8 bLockX : 1;
-
-	UPROPERTY(EditAnywhere, Category=BlendSample)
-	uint8 bLockY : 1;
-
-	UPROPERTY(EditAnywhere, Category=BlendSample)
-	uint8 bLockZ : 1;
-
 	// Cache the samples marker data counter so that we can track if it changes and revalidate the blendspace
 	int32 CachedMarkerDataUpdateCounter;
 
 #endif // WITH_EDITORONLY_DATA
 
 	FBlendSample()
-        : Animation(nullptr)
-        , SampleValue(0.f)
-        , RateScale(1.0f)
+		: Animation(nullptr)
+		, SampleValue(0.f)
+		, RateScale(1.0f)
 #if WITH_EDITORONLY_DATA
-        , bIsValid(false)
-		, bLockX(false)
-		, bLockY(false)
-		, bLockZ(false)
+		, bIsValid(false)
 		, CachedMarkerDataUpdateCounter(INDEX_NONE)
 #endif // WITH_EDITORONLY_DATA
 {		
 }
 	
 	FBlendSample(class UAnimSequence* InAnim, FVector InValue, bool bInIsSnapped, bool bInIsValid) 
-        : Animation(InAnim)
-        , SampleValue(InValue)
-        , RateScale(1.0f)
+		: Animation(InAnim)
+		, SampleValue(InValue)
+		, RateScale(1.0f)
 #if WITH_EDITORONLY_DATA
-        , bIsValid(bInIsValid)
+		, bIsValid(bInIsValid)
 		, CachedMarkerDataUpdateCounter(INDEX_NONE)
 #endif // WITH_EDITORONLY_DATA
 	{		 
@@ -569,9 +553,6 @@ public:
 	ENGINE_API int32 AddSample(const FVector& SampleValue);
 	ENGINE_API int32 AddSample(UAnimSequence* AnimationSequence, const FVector& SampleValue);
 
-	/** if requested, locks the position of the sample (e.g. after analysis) */
-	ENGINE_API void LockSample(int32 SampleIndex, bool bLockX, bool bLockY, bool bLockZ);
-
 	ENGINE_API void ExpandRangeForSample(const FVector& SampleValue);
 	
 	/** edit samples */
diff --git a/Engine/Source/Runtime/Engine/Private/Animation/BlendSpace.cpp b/Engine/Source/Runtime/Engine/Private/Animation/BlendSpace.cpp
--- a/Engine/Source/Runtime/Engine/Private/Animation/BlendSpace.cpp
+++ b/Engine/Source/Runtime/Engine/Private/Animation/BlendSpace.cpp
@@ -1229,16 +1229,6 @@ void UBlendSpace::ExpandRangeForSample(const FVector& SampleValue)
 	}
 }
 
-void UBlendSpace::LockSample(int32 SampleIndex, bool bLockX, bool bLockY, bool bLockZ)
-{
-	if (SampleData.IsValidIndex(SampleIndex))
-	{
-		SampleData[SampleIndex].bLockX = bLockX; 
-		SampleData[SampleIndex].bLockY = bLockY; 
-		SampleData[SampleIndex].bLockZ = bLockZ; 
-	}
-}
-
 int32 UBlendSpace::AddSample(const FVector& SampleValue)
 {
 	// We should only be adding samples without a source animation if we are not a standalone asset
