Added helper function to create AnimNotifyClassPicker widget

diff --git a/Engine/Source/Editor/Persona/Private/PersonaUtils.cpp b/Engine/Source/Editor/Persona/Private/PersonaUtils.cpp
--- a/Engine/Source/Editor/Persona/Private/PersonaUtils.cpp
+++ b/Engine/Source/Editor/Persona/Private/PersonaUtils.cpp
@@ -16,6 +16,10 @@
 #include "EditorStyleSet.h"
 #include "Widgets/SBoxPanel.h"
 #include "Widgets/Images/SImage.h"
+#include "Animation/AnimNotifies/AnimNotify.h"
+#include "Animation/AnimNotifies/AnimNotifyState.h"
+#include "ClassViewerFilter.h"
+#include "Modules/ModuleManager.h"
 
 namespace PersonaUtils
 {
@@ -189,4 +193,72 @@ TSharedRef<SWidget> MakeTrackButton(FText HoverText, FOnGetContent MenuContent,
 	return ComboButton;
 }
 
+template<typename NotifyTypeClass>
+static TSharedRef<SWidget> MakeNewNotifyPicker(UAnimSequenceBase* Sequence, const FOnClassPicked& OnClassPicked)
+{
+	class FNotifyStateClassFilter : public IClassViewerFilter
+	{
+	public:
+		FNotifyStateClassFilter(UAnimSequenceBase* InSequence)
+			: Sequence(InSequence)
+		{}
+
+		bool IsClassAllowed(const FClassViewerInitializationOptions& InInitOptions, const UClass* InClass, TSharedRef< FClassViewerFilterFuncs > InFilterFuncs) override
+		{
+			const bool bChildOfObjectClass = InClass->IsChildOf(NotifyTypeClass::StaticClass());
+			const bool bMatchesFlags = !InClass->HasAnyClassFlags(CLASS_Hidden | CLASS_HideDropDown | CLASS_Deprecated | CLASS_Abstract);
+			return bChildOfObjectClass && bMatchesFlags && CastChecked<NotifyTypeClass>(InClass->ClassDefaultObject)->CanBePlaced(Sequence);
+		}
+
+		virtual bool IsUnloadedClassAllowed(const FClassViewerInitializationOptions& InInitOptions, const TSharedRef< const IUnloadedBlueprintData > InUnloadedClassData, TSharedRef< FClassViewerFilterFuncs > InFilterFuncs) override
+		{
+			const bool bChildOfObjectClass = InUnloadedClassData->IsChildOf(NotifyTypeClass::StaticClass());
+			const bool bMatchesFlags = !InUnloadedClassData->HasAnyClassFlags(CLASS_Hidden | CLASS_HideDropDown | CLASS_Deprecated | CLASS_Abstract);
+			bool bValidToPlace = false;
+			if (bChildOfObjectClass)
+			{
+				if (const UClass* NativeBaseClass = InUnloadedClassData->GetNativeParent())
+				{
+					bValidToPlace = CastChecked<NotifyTypeClass>(NativeBaseClass->ClassDefaultObject)->CanBePlaced(Sequence);
+				}
+			}
+
+			return bChildOfObjectClass && bMatchesFlags && bValidToPlace;
+		}
+
+		/** Sequence referenced by outer panel */
+		UAnimSequenceBase* Sequence;
+	};
+
+	FClassViewerInitializationOptions InitOptions;
+	InitOptions.Mode = EClassViewerMode::ClassPicker;
+	InitOptions.bShowObjectRootClass = false;
+	InitOptions.bShowUnloadedBlueprints = true;
+	InitOptions.bShowNoneOption = false;
+	InitOptions.bEnableClassDynamicLoading = true;
+	InitOptions.bExpandRootNodes = true;
+	InitOptions.NameTypeToDisplay = EClassViewerNameTypeToDisplay::DisplayName;
+	InitOptions.ClassFilters.Add(MakeShared<FNotifyStateClassFilter>(Sequence));
+	InitOptions.bShowBackgroundBorder = false;
+
+	FClassViewerModule& ClassViewerModule = FModuleManager::LoadModuleChecked<FClassViewerModule>("ClassViewer");
+
+	return SNew(SBox)
+		.MinDesiredWidth(300.0f)
+		.MaxDesiredHeight(400.0f)
+		[
+			ClassViewerModule.CreateClassViewer(InitOptions, OnClassPicked)
+		];
+}
+
+TSharedRef<SWidget> MakeAnimNotifyPicker(UAnimSequenceBase* Sequence, const FOnClassPicked& OnClassPicked)
+{
+	return MakeNewNotifyPicker<UAnimNotify>(Sequence, OnClassPicked);
+}
+
+TSharedRef<SWidget> MakeAnimNotifyStatePicker(UAnimSequenceBase* Sequence, const FOnClassPicked& OnClassPicked)
+{
+	return MakeNewNotifyPicker<UAnimNotifyState>(Sequence, OnClassPicked);
+}
+
 }
diff --git a/Engine/Source/Editor/Persona/Private/SAnimNotifyPanel.cpp b/Engine/Source/Editor/Persona/Private/SAnimNotifyPanel.cpp
--- a/Engine/Source/Editor/Persona/Private/SAnimNotifyPanel.cpp
+++ b/Engine/Source/Editor/Persona/Private/SAnimNotifyPanel.cpp
@@ -47,7 +47,7 @@
 #include "IAnimationEditor.h"
 #include "IAnimationSequenceBrowser.h"
 #include "AnimTimelineTrack_NotifiesPanel.h"
-
+#include "PersonaUtils.h"
 
 // AnimNotify Drawing
 const float NotifyHeightOffset = 0.f;
@@ -858,11 +858,10 @@ public:
 protected:
 
 	// Build up a "New Notify..." menu
-	template<typename NotifyTypeClass>
-	void MakeNewNotifyPicker(FMenuBuilder& MenuBuilder, bool bIsReplaceWithMenu = false);
 	void FillNewNotifyMenu(FMenuBuilder& MenuBuilderbool, bool bIsReplaceWithMenu = false);
 	void FillNewNotifyStateMenu(FMenuBuilder& MenuBuilder, bool bIsReplaceWithMenu  = false);
 	void FillNewSyncMarkerMenu(FMenuBuilder& MenuBuilder);
+	void OnAnimNotifyClassPicked(UClass* NotifyClass, bool bIsReplaceWithMenu = false);
 
 	// New notify functions
 	void CreateNewBlueprintNotifyAtCursor(FString NewNotifyName, FString BlueprintPath);
@@ -2295,90 +2294,30 @@ FCursorReply SAnimNotifyTrack::OnCursorQuery(const FGeometry& MyGeometry, const
 	return FCursorReply::Unhandled();
 }
 
-template<typename NotifyTypeClass>
-void SAnimNotifyTrack::MakeNewNotifyPicker(FMenuBuilder& MenuBuilder, bool bIsReplaceWithMenu /* = false */)
+void SAnimNotifyTrack::OnAnimNotifyClassPicked(UClass* NotifyClass, bool bIsReplaceWithMenu /* = false */)
 {
-	FText TypeName = NotifyTypeClass::StaticClass() == UAnimNotify::StaticClass() ? LOCTEXT("AnimNotifyName", "anim notify") : LOCTEXT("AnimNotifyStateName", "anim notify state");
-	FText SectionHeaderFormat = bIsReplaceWithMenu ? LOCTEXT("ReplaceWithAnExistingAnimNotify", "Replace with an existing {0}") : LOCTEXT("AddsAnExistingAnimNotify", "Add an existing {0}");
+	FSlateApplication::Get().DismissAllMenus();
 
-	class FNotifyStateClassFilter : public IClassViewerFilter
+	if (bIsReplaceWithMenu)
 	{
-	public:
-		FNotifyStateClassFilter(UAnimSequenceBase* InSequence)
-			: Sequence(InSequence)
-		{}
-
-		bool IsClassAllowed(const FClassViewerInitializationOptions& InInitOptions, const UClass* InClass, TSharedRef< FClassViewerFilterFuncs > InFilterFuncs) override
-		{
-			const bool bChildOfObjectClass = InClass->IsChildOf(NotifyTypeClass::StaticClass());
-			const bool bMatchesFlags = !InClass->HasAnyClassFlags(CLASS_Hidden | CLASS_HideDropDown | CLASS_Deprecated | CLASS_Abstract);
-			return bChildOfObjectClass && bMatchesFlags && CastChecked<NotifyTypeClass>(InClass->ClassDefaultObject)->CanBePlaced(Sequence);
-		}
-
-		virtual bool IsUnloadedClassAllowed(const FClassViewerInitializationOptions& InInitOptions, const TSharedRef< const IUnloadedBlueprintData > InUnloadedClassData, TSharedRef< FClassViewerFilterFuncs > InFilterFuncs) override
-		{
-			const bool bChildOfObjectClass = InUnloadedClassData->IsChildOf(NotifyTypeClass::StaticClass());
-			const bool bMatchesFlags = !InUnloadedClassData->HasAnyClassFlags(CLASS_Hidden | CLASS_HideDropDown | CLASS_Deprecated | CLASS_Abstract);
-			bool bValidToPlace = false;
-			if(bChildOfObjectClass)
-			{
-				if (const UClass* NativeBaseClass = InUnloadedClassData->GetNativeParent())
-				{
-					bValidToPlace = CastChecked<NotifyTypeClass>(NativeBaseClass->ClassDefaultObject)->CanBePlaced(Sequence);
-				}
-			}
-
-			return bChildOfObjectClass && bMatchesFlags && bValidToPlace;
-		}
-
-		/** Sequence referenced by outer panel */
-		UAnimSequenceBase* Sequence;
-	};
+		ReplaceSelectedWithNotify(MakeBlueprintNotifyName(NotifyClass->GetName()), NotifyClass);
+	}
+	else
+	{
+		CreateNewNotifyAtCursor(MakeBlueprintNotifyName(NotifyClass->GetName()), NotifyClass);
+	}
+}
 
+void SAnimNotifyTrack::FillNewNotifyStateMenu(FMenuBuilder& MenuBuilder, bool bIsReplaceWithMenu /* = false */)
+{
 	// MenuBuilder always has a search widget added to it by default, hence if larger then 1 then something else has been added to it
 	if (MenuBuilder.GetMultiBox()->GetBlocks().Num() > 1)
 	{
 		MenuBuilder.AddMenuSeparator();
 	}
 
-	FClassViewerInitializationOptions InitOptions;
-	InitOptions.Mode = EClassViewerMode::ClassPicker;
-	InitOptions.bShowObjectRootClass = false;
-	InitOptions.bShowUnloadedBlueprints = true;
-	InitOptions.bShowNoneOption = false;
-	InitOptions.bEnableClassDynamicLoading = true;
-	InitOptions.bExpandRootNodes = true;
-	InitOptions.NameTypeToDisplay = EClassViewerNameTypeToDisplay::DisplayName;
-	InitOptions.ClassFilters.Add(MakeShared<FNotifyStateClassFilter>(Sequence));
-	InitOptions.bShowBackgroundBorder = false;
-
-	FClassViewerModule& ClassViewerModule = FModuleManager::LoadModuleChecked<FClassViewerModule>("ClassViewer");
-	MenuBuilder.AddWidget(
-		SNew(SBox)
-		.MinDesiredWidth(300.0f)
-		.MaxDesiredHeight(400.0f)
-		[
-			ClassViewerModule.CreateClassViewer(InitOptions,
-				FOnClassPicked::CreateLambda([this, bIsReplaceWithMenu](UClass* InClass)
-				{
-					FSlateApplication::Get().DismissAllMenus();
-					if(bIsReplaceWithMenu)
-					{
-						ReplaceSelectedWithNotify(MakeBlueprintNotifyName(InClass->GetName()), InClass);
-					}
-					else
-					{
-						CreateNewNotifyAtCursor(MakeBlueprintNotifyName(InClass->GetName()), InClass);
-					}
-				}
-			))
-		],
-		FText(), true, false);
-}
-
-void SAnimNotifyTrack::FillNewNotifyStateMenu(FMenuBuilder& MenuBuilder, bool bIsReplaceWithMenu /* = false */)
-{
-	MakeNewNotifyPicker<UAnimNotifyState>(MenuBuilder, bIsReplaceWithMenu);
+	TSharedRef<SWidget> Widget = PersonaUtils::MakeAnimNotifyStatePicker(Sequence, FOnClassPicked::CreateRaw(this, &SAnimNotifyTrack::OnAnimNotifyClassPicked, bIsReplaceWithMenu));
+	MenuBuilder.AddWidget(Widget, FText(), true, false);
 }
 
 void SAnimNotifyTrack::FillNewNotifyMenu(FMenuBuilder& MenuBuilder, bool bIsReplaceWithMenu /* = false */)
@@ -2433,8 +2372,15 @@ void SAnimNotifyTrack::FillNewNotifyMenu(FMenuBuilder& MenuBuilder, bool bIsRepl
 		MenuBuilder.EndSection();
 	}
 
+	// MenuBuilder always has a search widget added to it by default, hence if larger then 1 then something else has been added to it
+	if (MenuBuilder.GetMultiBox()->GetBlocks().Num() > 1)
+	{
+		MenuBuilder.AddMenuSeparator();
+	}
+
 	// Add a notify picker
-	MakeNewNotifyPicker<UAnimNotify>(MenuBuilder, bIsReplaceWithMenu);
+	TSharedRef<SWidget> Widget = PersonaUtils::MakeAnimNotifyPicker(Sequence, FOnClassPicked::CreateRaw(this, &SAnimNotifyTrack::OnAnimNotifyClassPicked, bIsReplaceWithMenu));
+	MenuBuilder.AddWidget(Widget, FText(), true, false);
 }
 
 void SAnimNotifyTrack::FillNewSyncMarkerMenu(FMenuBuilder& MenuBuilder)
diff --git a/Engine/Source/Editor/Persona/Public/PersonaUtils.h b/Engine/Source/Editor/Persona/Public/PersonaUtils.h
--- a/Engine/Source/Editor/Persona/Public/PersonaUtils.h
+++ b/Engine/Source/Editor/Persona/Public/PersonaUtils.h
@@ -4,11 +4,13 @@
 
 #include "CoreMinimal.h"
 #include "Framework/SlateDelegates.h"
+#include "ClassViewerModule.h"
 
 class UAnimInstance;
 class USceneComponent;
 class UAnimBlueprint;
 class SWidget;
+class UAnimSequenceBase;
 
 namespace PersonaUtils
 {
@@ -59,7 +61,14 @@ PERSONA_API int32 CopyPropertiesToCDO(UAnimInstance* InAnimInstance, const FCopy
 
 /** Set the correct object to debug depending on whether we are running a sub-layer etc. */
 PERSONA_API void SetObjectBeingDebugged(UAnimBlueprint* InAnimBlueprint, UAnimInstance* InAnimInstance);
+
 /** Make an add button widget, as seen in the animation timeline */
 PERSONA_API TSharedRef<SWidget> MakeTrackButton(FText HoverText, FOnGetContent MenuContent, const TAttribute<bool>& HoverState);
 
+/** Make AnimNotify class picker widget */
+PERSONA_API TSharedRef<SWidget> MakeAnimNotifyPicker(UAnimSequenceBase* Sequence, const FOnClassPicked& OnClassPicked);
+
+/** Make AnimNotifyState class picker widget */
+PERSONA_API TSharedRef<SWidget> MakeAnimNotifyStatePicker(UAnimSequenceBase* Sequence, const FOnClassPicked& OnClassPicked);
+
 }
