Fix keyboard input closing color picker for anim notifies

#7493

diff --git a/Engine/Source/Editor/Persona/Private/SAnimNotifyPanel.cpp b/Engine/Source/Editor/Persona/Private/SAnimNotifyPanel.cpp
--- a/Engine/Source/Editor/Persona/Private/SAnimNotifyPanel.cpp
+++ b/Engine/Source/Editor/Persona/Private/SAnimNotifyPanel.cpp
@@ -4622,8 +4622,11 @@ void SAnimNotifyPanel::OnPropertyChanged(UObject* ChangedObject, FPropertyChange
 		return;
 	}
 
+	const FName PropertyName = PropertyEvent.GetPropertyName();
+	
 	// Don't process if it's an interactive change; wait till we receive the final event.
-	if(PropertyEvent.ChangeType != EPropertyChangeType::Interactive)
+	// Skip notify color as otherwise we will end up refreshing the details panel before any edits are applied (e.g. with the tab key)
+	if(PropertyEvent.ChangeType != EPropertyChangeType::Interactive && PropertyName != GET_MEMBER_NAME_CHECKED(UAnimNotify, NotifyColor) && PropertyName != GET_MEMBER_NAME_CHECKED(UAnimNotifyState, NotifyColor))
 	{
 		for(FAnimNotifyEvent& Event : Sequence->Notifies)
 		{
