Added warning about incompatible post process anim BP being applied to a skeletal mesh

diff --git a/Engine/Source/Editor/Persona/Persona.Build.cs b/Engine/Source/Editor/Persona/Persona.Build.cs
--- a/Engine/Source/Editor/Persona/Persona.Build.cs
+++ b/Engine/Source/Editor/Persona/Persona.Build.cs
@@ -76,6 +76,7 @@ public class Persona : ModuleRules
 				"TimeManagement",
                 "Sequencer",
 				"StatusBar",
+				"ToolWidgets",
 			}
         );
 
diff --git a/Engine/Source/Editor/Persona/Private/PersonaMeshDetails.cpp b/Engine/Source/Editor/Persona/Private/PersonaMeshDetails.cpp
--- a/Engine/Source/Editor/Persona/Private/PersonaMeshDetails.cpp
+++ b/Engine/Source/Editor/Persona/Private/PersonaMeshDetails.cpp
@@ -86,6 +86,7 @@
 #include "PropertyCustomizationHelpers.h"
 #include "ComponentReregisterContext.h"
 #include "LODInfoUILayout.h"
+#include "SWarningOrErrorBox.h"
 #include "Interfaces/ITargetPlatform.h"
 #include "Interfaces/ITargetPlatformManagerModule.h"
 #include "Misc/CoreMisc.h"
@@ -3856,6 +3857,32 @@ void FPersonaMeshDetails::CustomizeDetails( IDetailLayoutBuilder& DetailLayout )
 		.OnObjectChanged(FOnSetObject::CreateSP(this, &FPersonaMeshDetails::OnSetPostProcessBlueprint, PostProcessHandle))
 	];
 
+	// Add warning if the post process BP is using an incompatible skeleton
+	FDetailWidgetRow& PostProcessWarningRow = SkelMeshCategory.AddCustomRow(LOCTEXT("PostProcessWarningFilterString", "Post Process Blueprint Warning"));
+
+	PostProcessWarningRow
+	.Visibility(MakeAttributeLambda([this]()
+	{
+		if(SkeletalMeshPtr.IsValid())
+		{
+			TSubclassOf<UAnimInstance> PostProcessAnimBlueprintClass = SkeletalMeshPtr->GetPostProcessAnimBlueprint();
+			if(PostProcessAnimBlueprintClass.Get())
+			{
+				if(UAnimBlueprint* AnimBlueprint = Cast<UAnimBlueprint>(PostProcessAnimBlueprintClass->ClassGeneratedBy))
+				{
+					return !SkeletalMeshPtr->GetSkeleton()->IsCompatible(AnimBlueprint->TargetSkeleton) ? EVisibility::Visible : EVisibility::Collapsed;
+				}
+			}
+		}
+		return EVisibility::Collapsed;
+	}))
+	.WholeRowContent()
+	[
+		SNew(SWarningOrErrorBox)
+		.MessageStyle(EMessageStyle::Error)
+		.Message(LOCTEXT("IncompatibleSkeletonError", "Post Process Anim Blueprint for this mesh uses an incompatible skeleton"))
+	];
+	
 	IDetailCategoryBuilder& ImportSettingsCategory = DetailLayout.EditCategory("ImportSettings");
 	TSharedRef<IPropertyHandle> AssetImportProperty = DetailLayout.GetProperty(USkeletalMesh::GetAssetImportDataMemberName(), USkeletalMesh::StaticClass());
 	if (!SkeletalMeshPtr.IsValid() || !IsValid(SkeletalMeshPtr->GetAssetImportData()) || !SkeletalMeshPtr->GetAssetImportData()->IsA<UFbxSkeletalMeshImportData>())
