Added 'Pose Watch Manager' and other pose watch improvements

diff --git a/Engine/Source/Editor/AnimGraph/Private/AnimGraphNode_Base.cpp b/Engine/Source/Editor/AnimGraph/Private/AnimGraphNode_Base.cpp
--- a/Engine/Source/Editor/AnimGraph/Private/AnimGraphNode_Base.cpp
+++ b/Engine/Source/Editor/AnimGraph/Private/AnimGraphNode_Base.cpp
@@ -180,6 +180,9 @@ void UAnimGraphNode_Base::DestroyNode()
 {
 	// This node may have been the last using its extension, so refresh
 	GetAnimBlueprint()->RequestRefreshExtensions();
+	
+	// Cleanup the pose watch if one exists on this node
+	AnimationEditorUtils::RemovePoseWatchFromNode(this, GetAnimBlueprint());
 
 	Super::DestroyNode();
 }
diff --git a/Engine/Source/Editor/AnimationBlueprintEditor/Private/AnimationBlueprintEditor.cpp b/Engine/Source/Editor/AnimationBlueprintEditor/Private/AnimationBlueprintEditor.cpp
--- a/Engine/Source/Editor/AnimationBlueprintEditor/Private/AnimationBlueprintEditor.cpp
+++ b/Engine/Source/Editor/AnimationBlueprintEditor/Private/AnimationBlueprintEditor.cpp
@@ -108,6 +108,7 @@ namespace AnimationBlueprintEditorTabs
 	const FName AssetOverridesTab(TEXT("AnimBlueprintParentPlayerEditor"));
 	const FName SlotNamesTab(TEXT("SkeletonSlotNames"));
 	const FName CurveNamesTab(TEXT("AnimCurveViewerTab"));
+	const FName PoseWatchTab(TEXT("PoseWatchTab"));
 };
 
 /////////////////////////////////////////////////////
@@ -728,7 +729,7 @@ FColor FAnimationBlueprintEditor::ChoosePoseWatchColor() const
 			UPoseWatch* PoseWatch = AnimationEditorUtils::FindPoseWatchForNode(GraphNode, AnimBP);
 			if (PoseWatch)
 			{
-				FColor PoseWatchColor = PoseWatch->PoseWatchColour;
+				FColor PoseWatchColor = PoseWatch->GetColor();
 				int32 Index = ColorCounts.IndexOfByPredicate([PoseWatchColor](const FColorCount& CC){return PoseWatchColor == CC.Color;});
 				if (Index != INDEX_NONE)
 				{
@@ -761,14 +762,16 @@ void FAnimationBlueprintEditor::OnTogglePoseWatch()
 	{
 		if (UAnimGraphNode_Base* SelectedNode = Cast<UAnimGraphNode_Base>(*NodeIt))
 		{
-			UPoseWatch* PoseWatch = AnimationEditorUtils::FindPoseWatchForNode(SelectedNode, AnimBP);
-			if (PoseWatch)
+			UPoseWatch* ExistingPoseWatch = AnimationEditorUtils::FindPoseWatchForNode(SelectedNode, AnimBP);
+			if (ExistingPoseWatch)
 			{
-				AnimationEditorUtils::RemovePoseWatch(PoseWatch, AnimBP);
+				AnimationEditorUtils::RemovePoseWatch(ExistingPoseWatch, AnimBP);
+				AnimationEditorUtils::OnPoseWatchesChanged().Broadcast(AnimBP, ExistingPoseWatch->Node.Get());
 			}
 			else
 			{
-				AnimationEditorUtils::MakePoseWatchForNode(AnimBP, SelectedNode, ChoosePoseWatchColor());
+				UPoseWatch* NewPoseWatch = AnimationEditorUtils::MakePoseWatchForNode(AnimBP, SelectedNode, ChoosePoseWatchColor());
+				AnimationEditorUtils::OnPoseWatchesChanged().Broadcast(AnimBP, NewPoseWatch->Node.Get());
 			}
 		}
 	}
@@ -1897,12 +1900,12 @@ void FAnimationBlueprintEditor::HandlePoseWatchSelectedNodes()
 					if (!PoseWatch)
 					{
 						PoseWatch = AnimationEditorUtils::MakePoseWatchForNode(AnimBP, GraphNode, ChoosePoseWatchColor());
-						PoseWatch->bDeleteOnDeselection = true;
+						PoseWatch->SetShouldDeleteOnDeselect(true);
 					}
 				}
 				else
 				{
-					if (PoseWatch && PoseWatch->bDeleteOnDeselection)
+					if (PoseWatch && PoseWatch->GetShouldDeleteOnDeselect())
 					{
 						AnimationEditorUtils::RemovePoseWatch(PoseWatch, AnimBP);
 					}
@@ -1926,7 +1929,7 @@ void FAnimationBlueprintEditor::RemoveAllSelectionPoseWatches()
 			if (GraphNode)
 			{
 				UPoseWatch* PoseWatch = AnimationEditorUtils::FindPoseWatchForNode(GraphNode, AnimBP);
-				if (PoseWatch && PoseWatch->bDeleteOnDeselection)
+				if (PoseWatch && PoseWatch->GetShouldDeleteOnDeselect())
 				{
 					AnimationEditorUtils::RemovePoseWatch(PoseWatch, AnimBP);
 				}
diff --git a/Engine/Source/Editor/AnimationBlueprintEditor/Private/AnimationBlueprintEditor.h b/Engine/Source/Editor/AnimationBlueprintEditor/Private/AnimationBlueprintEditor.h
--- a/Engine/Source/Editor/AnimationBlueprintEditor/Private/AnimationBlueprintEditor.h
+++ b/Engine/Source/Editor/AnimationBlueprintEditor/Private/AnimationBlueprintEditor.h
@@ -62,6 +62,7 @@ namespace AnimationBlueprintEditorTabs
 	extern const FName AssetOverridesTab;
 	extern const FName SlotNamesTab;
 	extern const FName CurveNamesTab;
+	extern const FName PoseWatchTab;
 };
 
 /**
diff --git a/Engine/Source/Editor/AnimationBlueprintEditor/Private/AnimationBlueprintEditorMode.cpp b/Engine/Source/Editor/AnimationBlueprintEditor/Private/AnimationBlueprintEditorMode.cpp
--- a/Engine/Source/Editor/AnimationBlueprintEditor/Private/AnimationBlueprintEditorMode.cpp
+++ b/Engine/Source/Editor/AnimationBlueprintEditor/Private/AnimationBlueprintEditorMode.cpp
@@ -131,6 +131,7 @@ FAnimationBlueprintEditorMode::FAnimationBlueprintEditorMode(const TSharedRef<FA
 	TabFactories.RegisterFactory(PersonaModule.CreateAnimBlueprintAssetOverridesTabFactory(InAnimationBlueprintEditor, InAnimationBlueprintEditor->GetPersonaToolkit()->GetAnimBlueprint(), InAnimationBlueprintEditor->OnPostUndo));
 	TabFactories.RegisterFactory(PersonaModule.CreateSkeletonSlotNamesTabFactory(InAnimationBlueprintEditor, InAnimationBlueprintEditor->GetSkeletonTree()->GetEditableSkeleton(), FOnObjectSelected::CreateSP(&InAnimationBlueprintEditor.Get(), &FAnimationBlueprintEditor::HandleObjectSelected)));
 	TabFactories.RegisterFactory(PersonaModule.CreateCurveViewerTabFactory(InAnimationBlueprintEditor, InAnimationBlueprintEditor->GetSkeletonTree()->GetEditableSkeleton(), InAnimationBlueprintEditor->GetPersonaToolkit()->GetPreviewScene(), FOnObjectsSelected::CreateSP(&InAnimationBlueprintEditor.Get(), &FAnimationBlueprintEditor::HandleObjectsSelected)));
+	TabFactories.RegisterFactory(PersonaModule.CreatePoseWatchTabFactory(InAnimationBlueprintEditor));
 
 	// setup toolbar - clear existing toolbar extender from the BP mode
 	//@TODO: Keep this in sync with BlueprintEditorModes.cpp
diff --git a/Engine/Source/Editor/AnimationBlueprintEditor/Private/AnimationNodes/SAnimationGraphNode.cpp b/Engine/Source/Editor/AnimationBlueprintEditor/Private/AnimationNodes/SAnimationGraphNode.cpp
--- a/Engine/Source/Editor/AnimationBlueprintEditor/Private/AnimationNodes/SAnimationGraphNode.cpp
+++ b/Engine/Source/Editor/AnimationBlueprintEditor/Private/AnimationNodes/SAnimationGraphNode.cpp
@@ -152,11 +152,11 @@ void SAnimationGraphNode::Construct(const FArguments& InArgs, UAnimGraphNode_Bas
 
 	PoseViewWidget =
 		SNew(SButton)
-		.ToolTipText(LOCTEXT("SpawnColourPicker", "Pose watch active. Click to spawn the pose watch colour picker"))
-		.OnClicked(this, &SAnimationGraphNode::SpawnColourPicker)
+		.ToolTipText(LOCTEXT("TogglePoseWatchVisibility", "Click to toggle visibility"))
+		.OnClicked(this, &SAnimationGraphNode::TogglePoseWatchVisibility)
 		.ButtonColorAndOpacity(this, &SAnimationGraphNode::GetPoseViewColour)
 		[
-			SNew(SImage).Image(FEditorStyle::GetBrush("GenericViewButton"))
+			SNew(SImage).Image(this, &SAnimationGraphNode::GetPoseViewIcon)
 		];
 
 	// Search for an enabled or disabled pose watch on this node
@@ -203,7 +203,7 @@ TArray<FOverlayWidgetInfo> SAnimationGraphNode::GetOverlayWidgets(bool bSelected
 
 		if (PoseWatch.IsValid())
 		{
-			const FSlateBrush* ImageBrush = FEditorStyle::GetBrush("GenericViewButton");
+			const FSlateBrush* ImageBrush = FEditorStyle::GetBrush("Level.VisibleIcon16x");
 
 			FOverlayWidgetInfo Info;
 			Info.OverlayOffset = FVector2D(0 - (ImageBrush->ImageSize.X * 0.5f), -(ImageBrush->ImageSize.Y * 0.5f));
@@ -216,16 +216,29 @@ TArray<FOverlayWidgetInfo> SAnimationGraphNode::GetOverlayWidgets(bool bSelected
 	return Widgets;
 }
 
+const FSlateBrush* SAnimationGraphNode::GetPoseViewIcon() const
+{
+	return FEditorStyle::GetBrush(PoseWatch->GetIsVisible() ? "Level.VisibleIcon16x" : "Level.NotVisibleIcon16x");
+}
+
 FSlateColor SAnimationGraphNode::GetPoseViewColour() const
 {
 	UPoseWatch* CurPoseWatch = PoseWatch.Get();
 	if (CurPoseWatch)
 	{
-		return FSlateColor(CurPoseWatch->PoseWatchColour);
+		FLinearColor OutColor = CurPoseWatch->GetColor();
+		OutColor.A = CurPoseWatch->GetShouldDeleteOnDeselect() ? 0.5 : 0.9;
+		return FSlateColor(OutColor);
 	}
 	return FSlateColor(FColor::White); //Need a return value but should never actually get here
 }
 
+FReply SAnimationGraphNode::TogglePoseWatchVisibility()
+{
+	PoseWatch->ToggleIsVisible();
+	return FReply::Handled();
+}
+
 FReply SAnimationGraphNode::SpawnColourPicker()
 {
 	FSlateApplication::Get().PushMenu(
diff --git a/Engine/Source/Editor/AnimationBlueprintEditor/Public/AnimationNodes/SAnimationGraphNode.h b/Engine/Source/Editor/AnimationBlueprintEditor/Public/AnimationNodes/SAnimationGraphNode.h
--- a/Engine/Source/Editor/AnimationBlueprintEditor/Public/AnimationNodes/SAnimationGraphNode.h
+++ b/Engine/Source/Editor/AnimationBlueprintEditor/Public/AnimationNodes/SAnimationGraphNode.h
@@ -55,8 +55,10 @@ protected:
 private:
 	// Return Pose View Colour for slate indicator
 	FSlateColor GetPoseViewColour() const;
+	const FSlateBrush* GetPoseViewIcon() const;
 
 	FReply SpawnColourPicker();
+	FReply TogglePoseWatchVisibility();
 
 	// Handle the node informing us that the title has changed
 	void HandleNodeTitleChanged();
diff --git a/Engine/Source/Editor/Persona/Private/AnimationEditorViewportClient.cpp b/Engine/Source/Editor/Persona/Private/AnimationEditorViewportClient.cpp
--- a/Engine/Source/Editor/Persona/Private/AnimationEditorViewportClient.cpp
+++ b/Engine/Source/Editor/Persona/Private/AnimationEditorViewportClient.cpp
@@ -1332,7 +1332,7 @@ void FAnimationViewportClient::DrawWatchedPoses(UDebugSkelMeshComponent * MeshCo
 			{
 				for (const FAnimNodePoseWatch& AnimNodePoseWatch : AnimBlueprintGeneratedClass->GetAnimBlueprintDebugData().AnimNodePoseWatch)
 				{
-					if(AnimNodePoseWatch.Object.Get() != nullptr)
+					if(AnimNodePoseWatch.Object.Get() != nullptr && !AnimNodePoseWatch.bIsHidden)
 					{
 						DrawBonesFromCompactPose(*AnimNodePoseWatch.PoseInfo.Get(), MeshComponent, PDI, AnimNodePoseWatch.PoseDrawColour);
 					}
diff --git a/Engine/Source/Editor/Persona/Private/IPoseWatchManager.h b/Engine/Source/Editor/Persona/Private/IPoseWatchManager.h
new file mode 100644
--- /dev/null
+++ b/Engine/Source/Editor/Persona/Private/IPoseWatchManager.h
@@ -0,0 +1,60 @@
+// Copyright Epic Games, Inc. All Rights Reserved.
+
+#pragma once
+
+#include "CoreMinimal.h"
+#include "Misc/Attribute.h"
+#include "PoseWatchManagerFwd.h"
+#include "PoseWatchManagerStandaloneTypes.h"
+#include "Widgets/SCompoundWidget.h"
+#include "Widgets/Views/SHeaderRow.h"
+#include "BlueprintEditor.h"
+#include "Animation/AnimBlueprint.h"
+
+class FPoseWatchManagerDefaultMode;
+template<typename ItemType> class STreeView;
+
+/**
+ * The public interface for the Pose Watch Manager widget
+ */
+class IPoseWatchManager : public SCompoundWidget
+{
+public:
+
+	/** Sends a requests to refresh itself the next chance it gets */
+	virtual void Refresh() = 0;
+
+	/** Sends a request to clear the entire tree and rebuild it from scratch */
+	virtual void FullRefresh() = 0;
+
+	/** @return Returns a string to use for highlighting results in the outliner list */
+	virtual TAttribute<FText> GetFilterHighlightText() const = 0;
+
+	/** Get a const reference to the actual tree hierarchy */
+	virtual const STreeView<FPoseWatchManagerTreeItemPtr>& GetTree() const = 0;
+
+	/** Set the keyboard focus to the outliner */
+	virtual void SetKeyboardFocus() = 0;
+
+	/** Return the sorting mode for the specified ColumnId */
+	virtual EColumnSortMode::Type GetColumnSortMode(const FName ColumnId) const = 0;
+
+	virtual uint32 GetTypeSortPriority(const IPoseWatchManagerTreeItem& Item) const = 0;
+
+	/** Request that the tree be sorted at a convenient time */
+	virtual void RequestSort() = 0;
+
+	/** Executes rename. */
+	virtual void Rename_Execute() = 0;
+
+	/** Set the item selection of the outliner based on a selector function. Any items which return true will be added */
+	virtual void SetSelection(const TFunctionRef<bool(IPoseWatchManagerTreeItem&)> Selector) = 0;
+
+	UAnimBlueprint* AnimBlueprint;
+
+	/** Get the active PoseWatchManagerMode */
+	const TSharedPtr<FPoseWatchManagerDefaultMode> GetMode() const { return Mode; }
+
+protected:
+	TSharedPtr<FPoseWatchManagerDefaultMode> Mode;
+};
diff --git a/Engine/Source/Editor/Persona/Private/IPoseWatchManagerColumn.h b/Engine/Source/Editor/Persona/Private/IPoseWatchManagerColumn.h
new file mode 100644
--- /dev/null
+++ b/Engine/Source/Editor/Persona/Private/IPoseWatchManagerColumn.h
@@ -0,0 +1,30 @@
+// Copyright Epic Games, Inc. All Rights Reserved.
+
+#pragma once
+
+#include "CoreMinimal.h"
+#include "PoseWatchManagerFwd.h"
+#include "PoseWatchManagerStandaloneTypes.h"
+#include "Widgets/SWidget.h"
+#include "Widgets/Views/SHeaderRow.h"
+
+template<typename ItemType> class STableRow;
+
+class IPoseWatchManagerColumn : public TSharedFromThis<IPoseWatchManagerColumn>
+{
+public:
+	virtual ~IPoseWatchManagerColumn() {};
+
+	virtual FName GetColumnID() = 0;
+
+	virtual SHeaderRow::FColumn::FArguments ConstructHeaderRowColumn() = 0;
+
+	virtual const TSharedRef<SWidget> ConstructRowWidget(FPoseWatchManagerTreeItemRef TreeItem, const STableRow<FPoseWatchManagerTreeItemPtr>& Row) = 0;
+
+	virtual void PopulateSearchStrings(const IPoseWatchManagerTreeItem& Item, TArray<FString>& OutSearchStrings) const {}
+
+	virtual bool SupportsSorting() const { return false; }
+
+	virtual void SortItems(TArray<FPoseWatchManagerTreeItemPtr>& OutItems, const EColumnSortMode::Type SortMode) const {}
+};
+
diff --git a/Engine/Source/Editor/Persona/Private/IPoseWatchManagerTreeItem.h b/Engine/Source/Editor/Persona/Private/IPoseWatchManagerTreeItem.h
new file mode 100644
--- /dev/null
+++ b/Engine/Source/Editor/Persona/Private/IPoseWatchManagerTreeItem.h
@@ -0,0 +1,113 @@
+// Copyright Epic Games, Inc. All Rights Reserved.
+
+#pragma once
+
+#include "CoreMinimal.h"
+#include "Widgets/SWidget.h"
+#include "Widgets/Views/STableRow.h"
+#include "PoseWatchManagerStandaloneTypes.h"
+
+class UToolMenu;
+class SPoseWatchManager;
+
+struct PERSONA_API IPoseWatchManagerTreeItem : TSharedFromThis<IPoseWatchManagerTreeItem>
+{
+	friend SPoseWatchManager;
+
+public:
+	/** Delegate for hooking up an inline editable text block to be notified that a rename is requested. */
+	DECLARE_DELEGATE(FOnRenameRequest);
+	FOnRenameRequest RenameRequestEvent;
+
+protected:
+	/** Default constructor */
+	IPoseWatchManagerTreeItem(EPoseWatchTreeItemType InType) : Parent(nullptr), TreeType(InType) {}
+	virtual ~IPoseWatchManagerTreeItem() {}
+
+public:
+	/** Get this item's parent. If nullptr then belongs to the root of the tree */
+	FPoseWatchManagerTreeItemPtr GetParent() const
+	{
+		return Parent.Pin();
+	}
+
+	/** Returns true if this tree item is inside a folder */
+	virtual bool IsAssignedFolder() const = 0;
+
+	/** Sets this tree item's visibility */
+	virtual void SetIsVisible(const bool bVisible) = 0;
+
+	/** Add a child to this tree item */
+	void AddChild(FPoseWatchManagerTreeItemRef Child)
+	{
+		check(!Children.Contains(Child));
+		Child->Parent = AsShared();
+		Children.Add(MoveTemp(Child));
+	}
+
+	/** Get this item's children, if any. */
+	FORCEINLINE const TSet<TWeakPtr<IPoseWatchManagerTreeItem>>& GetChildren() const
+	{
+		return Children;
+	}
+
+	/** Attempt to cast this item to another type if it is of that type. Returns nullptr if it fails */
+	template <typename T>
+	T* CastTo()
+	{
+		return TreeType == T::Type ? StaticCast<T*>(this) : nullptr;
+	}
+	/** Attempt to cast this item to another type if it is of that type. Returns nullptr if it fails */
+	template <typename T>
+	const T* CastTo() const
+	{
+		return TreeType == T::Type ? StaticCast<const T*>(this) : nullptr;
+	}
+	/** Returns true if this item is of the specified type */
+	template <typename T>
+	bool IsA() const
+	{
+		return TreeType == T::Type;
+	}
+
+	/** Returns true if the data the item references is valid */
+	virtual bool IsValid() const = 0;
+
+	/** Get the ID that represents this tree item. Used to reference this item in a map */
+	virtual FObjectKey GetID() const = 0;
+
+	/** Get the raw string to display for this tree item - used for sorting */
+	virtual FString GetDisplayString() const = 0;
+
+	/** Generate the label widget for this item */
+	virtual TSharedRef<SWidget> GenerateLabelWidget(IPoseWatchManager& Outliner, const STableRow<FPoseWatchManagerTreeItemPtr>& InRow) { return SNullWidget::NullWidget; }
+
+	/** Generate a context menu for this item. Only called if *only* this item is selected. */
+	virtual TSharedPtr<SWidget> CreateContextMenu() = 0;
+
+	/** Query this items visibility state. Only called if the item type has visibility info */
+	virtual bool GetVisibility() const = 0;
+
+	/** Returns true if this tree item has at least one child */
+	virtual bool HasChildren() const { return false; }
+
+	/** Returns true if this tree item is expanded */
+	virtual bool IsExpanded() const { return false; }
+
+	virtual bool IsEnabled() const { return true; }
+
+	virtual void OnRemoved() = 0;
+
+protected:
+	/** This item's parent tree item */
+	TWeakPtr<IPoseWatchManagerTreeItem> Parent;
+
+	/** Array of children contained underneath this item */
+	TSet<TWeakPtr<IPoseWatchManagerTreeItem>> Children;
+
+	/** Static type identifier for the base class tree item */
+	static const EPoseWatchTreeItemType Type;
+
+	/** Tree item type identifier */
+	EPoseWatchTreeItemType TreeType;
+};
diff --git a/Engine/Source/Editor/Persona/Private/PersonaModule.cpp b/Engine/Source/Editor/Persona/Private/PersonaModule.cpp
--- a/Engine/Source/Editor/Persona/Private/PersonaModule.cpp
+++ b/Engine/Source/Editor/Persona/Private/PersonaModule.cpp
@@ -323,6 +323,11 @@ TSharedRef<class FWorkflowTabFactory> FPersonaModule::CreateAnimBlueprintPreview
 	return MakeShareable(new FAnimBlueprintPreviewEditorSummoner(InBlueprintEditor, InPreviewScene));
 }
 
+TSharedRef<class FWorkflowTabFactory> FPersonaModule::CreatePoseWatchTabFactory(const TSharedRef<class FBlueprintEditor>& InBlueprintEditor) const
+{
+	return MakeShareable(new FPoseWatchManagerSummoner(InBlueprintEditor));
+}
+
 TSharedRef<class FWorkflowTabFactory> FPersonaModule::CreateAnimBlueprintAssetOverridesTabFactory(const TSharedRef<class FBlueprintEditor>& InBlueprintEditor, UAnimBlueprint* InAnimBlueprint, FSimpleMulticastDelegate& InOnPostUndo) const
 {
 	return MakeShareable(new FAnimBlueprintParentPlayerEditorSummoner(InBlueprintEditor, InOnPostUndo));
diff --git a/Engine/Source/Editor/Persona/Private/PoseWatchManagerColumnColor.cpp b/Engine/Source/Editor/Persona/Private/PoseWatchManagerColumnColor.cpp
new file mode 100644
--- /dev/null
+++ b/Engine/Source/Editor/Persona/Private/PoseWatchManagerColumnColor.cpp
@@ -0,0 +1,137 @@
+// Copyright Epic Games, Inc. All Rights Reserved.
+
+#include "PoseWatchManagerColumnColor.h"
+#include "Widgets/SBoxPanel.h"
+#include "Widgets/Layout/SSpacer.h"
+#include "Widgets/Images/SImage.h"
+#include "Widgets/Colors/SColorBlock.h"
+#include "EditorStyleSet.h"
+#include "Editor.h"
+#include "Widgets/Views/STreeView.h"
+#include "IPoseWatchManager.h"
+#include "Widgets/Colors/SColorPicker.h"
+#include "PoseWatchManagerPoseWatchTreeItem.h"
+#include "Engine/PoseWatch.h"
+
+#define LOCTEXT_NAMESPACE "PoseWatchManagerColumnColor"
+
+class SColorBoxWidget : public SBox
+{
+public:
+	SLATE_BEGIN_ARGS(SColorBoxWidget) {}
+	SLATE_END_ARGS()
+
+	void Construct(const FArguments& InArgs, FPoseWatchManagerPoseWatchTreeItem* InPoseWatchTreeItem)
+	{
+		PoseWatchTreeItem = InPoseWatchTreeItem;
+
+		ChildSlot
+		[
+			SNew(SBox)
+			.Padding(FMargin(4.0f, 0.0f, 0.0f, 0.0f))
+			.VAlign(VAlign_Center)
+			[
+				SAssignNew(ColorWidgetBackgroundBorder, SBorder)
+				.Padding(1)
+				.BorderImage(FAppStyle::Get().GetBrush("ColorPicker.RoundedSolidBackground"))
+				.BorderBackgroundColor(this, &SColorBoxWidget::GetColorWidgetBorderColor)
+				.VAlign(VAlign_Center)
+				[
+					SNew(SOverlay)
+					+ SOverlay::Slot()
+					.VAlign(VAlign_Center)
+					[
+						SAssignNew(ColorPickerParentWidget, SColorBlock)
+						.AlphaBackgroundBrush(FAppStyle::Get().GetBrush("ColorPicker.RoundedAlphaBackground"))
+						.Color(this, &SColorBoxWidget::OnGetColorForColorBlock)
+						.ShowBackgroundForAlpha(true)
+						.AlphaDisplayMode(EColorBlockAlphaDisplayMode::Ignore)
+						.OnMouseButtonDown(this, &SColorBoxWidget::OnMouseButtonDownColorBlock)
+						.Size(FVector2D(16.0f, 16.0f))
+						.CornerRadius(FVector4(4.0f, 4.0f, 4.0f, 4.0f))
+					]
+				]
+			]
+		];
+	}
+
+	FLinearColor OnGetColorForColorBlock() const
+	{
+		return PoseWatchTreeItem->PoseWatch->GetColor();
+	}
+
+	void OnSetColorFromColorPicker(FLinearColor NewColor)
+	{
+		PoseWatchTreeItem->PoseWatch->SetColor(NewColor.ToFColorSRGB());
+	}
+
+	void OnColorPickerCancelled(FLinearColor OriginalColor)
+	{
+		PoseWatchTreeItem->PoseWatch->SetColor(OriginalColor.ToFColorSRGB());
+	}
+
+	FReply OnMouseButtonDownColorBlock(const FGeometry& MyGeometry, const FPointerEvent& MouseEvent)
+	{
+		FColorPickerArgs PickerArgs;
+		{
+			PickerArgs.bUseAlpha = false;
+			PickerArgs.DisplayGamma = TAttribute<float>::Create(TAttribute<float>::FGetter::CreateUObject(GEngine, &UEngine::GetDisplayGamma));
+			PickerArgs.OnColorCommitted = FOnLinearColorValueChanged::CreateSP(this, &SColorBoxWidget::OnSetColorFromColorPicker);
+			PickerArgs.OnColorPickerCancelled = FOnColorPickerCancelled::CreateSP(this, &SColorBoxWidget::OnColorPickerCancelled);
+			PickerArgs.InitialColorOverride = PoseWatchTreeItem->PoseWatch->GetColor();
+			PickerArgs.ParentWidget = ColorPickerParentWidget;
+			PickerArgs.OptionalOwningDetailsView = ColorPickerParentWidget;
+			FWidgetPath ParentWidgetPath;
+			if (FSlateApplication::Get().FindPathToWidget(ColorPickerParentWidget.ToSharedRef(), ParentWidgetPath))
+			{
+				PickerArgs.bOpenAsMenu = FSlateApplication::Get().FindMenuInWidgetPath(ParentWidgetPath).IsValid();
+			}
+		}
+
+		OpenColorPicker(PickerArgs);
+
+		return FReply::Handled();
+	}
+
+	FSlateColor GetColorWidgetBorderColor() const
+	{
+		static const FSlateColor HoveredColor = FAppStyle::Get().GetSlateColor("Colors.Hover");
+		static const FSlateColor DefaultColor = FAppStyle::Get().GetSlateColor("Colors.InputOutline");
+		return ColorWidgetBackgroundBorder->IsHovered() ? HoveredColor : DefaultColor;
+	}
+
+private:
+	FPoseWatchManagerPoseWatchTreeItem* PoseWatchTreeItem;
+	TSharedPtr<SWidget> ColorPickerParentWidget;
+	TSharedPtr<SWidget > ColorWidgetBackgroundBorder;
+};
+
+SHeaderRow::FColumn::FArguments FPoseWatchManagerColumnColor::ConstructHeaderRowColumn()
+{
+	return SHeaderRow::Column(GetColumnID())
+		.FixedWidth(24.f)
+		.HAlignHeader(HAlign_Left)
+		.VAlignHeader(VAlign_Center)
+		.HAlignCell(HAlign_Center)
+		.VAlignCell(VAlign_Center)
+		.DefaultTooltip(FText::FromName(GetColumnID()))
+		.HeaderContentPadding(FMargin(4.0f, 0.0f, 0.0f, 0.0f))
+		[
+			SNew(SImage)
+			.ColorAndOpacity(FSlateColor::UseForeground())
+			.Image(FAppStyle::Get().GetBrush("ColorPicker.Mode"))
+		];
+}
+
+
+const TSharedRef<SWidget> FPoseWatchManagerColumnColor::ConstructRowWidget(FPoseWatchManagerTreeItemRef TreeItem, const STableRow<FPoseWatchManagerTreeItemPtr>& Row)
+{
+	if (FPoseWatchManagerPoseWatchTreeItem* PoseWatchTreeItem = TreeItem->CastTo<FPoseWatchManagerPoseWatchTreeItem>())
+	{
+		return SNew(SColorBoxWidget, PoseWatchTreeItem);
+	}
+	return SNullWidget::NullWidget;
+}
+
+#undef LOCTEXT_NAMESPACE
+
diff --git a/Engine/Source/Editor/Persona/Private/PoseWatchManagerColumnColor.h b/Engine/Source/Editor/Persona/Private/PoseWatchManagerColumnColor.h
new file mode 100644
--- /dev/null
+++ b/Engine/Source/Editor/Persona/Private/PoseWatchManagerColumnColor.h
@@ -0,0 +1,31 @@
+// Copyright Epic Games, Inc. All Rights Reserved.
+
+#pragma once
+
+#include "CoreMinimal.h"
+#include "Widgets/SWidget.h"
+#include "Widgets/Views/SHeaderRow.h"
+#include "PoseWatchManagerPublicTypes.h"
+#include "IPoseWatchManagerColumn.h"
+#include "IPoseWatchManagerTreeItem.h"
+
+template<typename ItemType> class STableRow;
+
+/**
+ * A PoseWatchManager column which handles setting and visualizing a pose watch's color
+ */
+class FPoseWatchManagerColumnColor : public IPoseWatchManagerColumn
+{
+public:
+	FPoseWatchManagerColumnColor() {};
+	virtual ~FPoseWatchManagerColumnColor() {}
+
+	static FName GetID() { return FPoseWatchManagerBuiltInColumnTypes::Color(); }
+
+	// Begin IPoseWatchManagerColumn Implementation
+	virtual FName GetColumnID() override { return GetID();  }
+	virtual SHeaderRow::FColumn::FArguments ConstructHeaderRowColumn() override;
+	virtual const TSharedRef<SWidget> ConstructRowWidget(FPoseWatchManagerTreeItemRef TreeItem, const STableRow<FPoseWatchManagerTreeItemPtr>& Row) override;
+	virtual bool SupportsSorting() const override { return false; }
+	// End IPoseWatchManagerColumn Implementation
+};
diff --git a/Engine/Source/Editor/Persona/Private/PoseWatchManagerColumnLabel.cpp b/Engine/Source/Editor/Persona/Private/PoseWatchManagerColumnLabel.cpp
new file mode 100644
--- /dev/null
+++ b/Engine/Source/Editor/Persona/Private/PoseWatchManagerColumnLabel.cpp
@@ -0,0 +1,28 @@
+// Copyright Epic Games, Inc. All Rights Reserved.
+
+#include "PoseWatchManagerColumnLabel.h"
+#include "PoseWatchManagerFwd.h"
+#include "PoseWatchManagerStandaloneTypes.h"
+#include "IPoseWatchManagerTreeItem.h"
+#include "IPoseWatchManager.h"
+
+#define LOCTEXT_NAMESPACE "PoseWatchManagerColumnLabel"
+
+SHeaderRow::FColumn::FArguments FPoseWatchManagerItemLabelColumn::ConstructHeaderRowColumn()
+{
+	return SHeaderRow::Column(GetColumnID())
+		.FillWidth(5.0f);
+}
+const TSharedRef<SWidget> FPoseWatchManagerItemLabelColumn::ConstructRowWidget(FPoseWatchManagerTreeItemRef TreeItem, const STableRow<FPoseWatchManagerTreeItemPtr>& Row)
+{
+	IPoseWatchManager* PoseWatchManager = WeakPoseWatchManager.Pin().Get();
+	check(PoseWatchManager);
+	return TreeItem->GenerateLabelWidget(*PoseWatchManager, Row);
+}
+
+void FPoseWatchManagerItemLabelColumn::PopulateSearchStrings(const IPoseWatchManagerTreeItem& Item, TArray< FString >& OutSearchStrings) const
+{
+	OutSearchStrings.Add(Item.GetDisplayString());
+}
+
+#undef LOCTEXT_NAMESPACE
diff --git a/Engine/Source/Editor/Persona/Private/PoseWatchManagerColumnLabel.h b/Engine/Source/Editor/Persona/Private/PoseWatchManagerColumnLabel.h
new file mode 100644
--- /dev/null
+++ b/Engine/Source/Editor/Persona/Private/PoseWatchManagerColumnLabel.h
@@ -0,0 +1,35 @@
+// Copyright Epic Games, Inc. All Rights Reserved.
+
+#pragma once
+
+#include "CoreMinimal.h"
+#include "Widgets/SWidget.h"
+#include "Widgets/Views/SHeaderRow.h"
+#include "IPoseWatchManager.h"
+#include "PoseWatchManagerPublicTypes.h"
+#include "PoseWatchManagerStandaloneTypes.h"
+#include "IPoseWatchManagerColumn.h"
+
+template<typename ItemType> class STableRow;
+
+/** A column for the PoseWatchManager that displays the item label */
+class FPoseWatchManagerItemLabelColumn : public IPoseWatchManagerColumn
+{
+
+public:
+	FPoseWatchManagerItemLabelColumn(IPoseWatchManager& PoseWatchManager)
+		: WeakPoseWatchManager(StaticCastSharedRef<IPoseWatchManager>(PoseWatchManager.AsShared()))
+	{}
+
+	static FName GetID() { return FPoseWatchManagerBuiltInColumnTypes::Label(); }
+
+	// Begin IPoseWatchManagerColumn Implementation
+	virtual FName GetColumnID() override { return GetID(); }
+	virtual SHeaderRow::FColumn::FArguments ConstructHeaderRowColumn() override;
+	virtual const TSharedRef< SWidget > ConstructRowWidget(FPoseWatchManagerTreeItemRef TreeItem, const STableRow<FPoseWatchManagerTreeItemPtr>& Row) override;
+	virtual void PopulateSearchStrings(const IPoseWatchManagerTreeItem& Item, TArray< FString >& OutSearchStrings) const override;
+	// End IPoseWatchManagerColumn Implementation
+
+private:
+	TWeakPtr<IPoseWatchManager> WeakPoseWatchManager;
+};
diff --git a/Engine/Source/Editor/Persona/Private/PoseWatchManagerColumnVisibility.cpp b/Engine/Source/Editor/Persona/Private/PoseWatchManagerColumnVisibility.cpp
new file mode 100644
--- /dev/null
+++ b/Engine/Source/Editor/Persona/Private/PoseWatchManagerColumnVisibility.cpp
@@ -0,0 +1,161 @@
+// Copyright Epic Games, Inc. All Rights Reserved.
+
+#include "PoseWatchManagerColumnVisibility.h"
+#include "Widgets/SBoxPanel.h"
+#include "Widgets/Layout/SSpacer.h"
+#include "Widgets/Images/SImage.h"
+#include "EditorStyleSet.h"
+#include "Editor.h"
+#include "Widgets/Views/STreeView.h"
+#include "IPoseWatchManager.h"
+#include "PoseWatchManagerPoseWatchTreeItem.h"
+
+#define LOCTEXT_NAMESPACE "PoseWatchManagerColumnVisibility"
+
+class SVisibilityWidget : public SImage
+{
+public:
+	SLATE_BEGIN_ARGS(SVisibilityWidget) {}
+	SLATE_END_ARGS()
+
+	void Construct(const FArguments& InArgs, TWeakPtr<FPoseWatchManagerColumnVisibility> InWeakColumn, TWeakPtr<IPoseWatchManager> InWeakPoseWatchManager, TWeakPtr<IPoseWatchManagerTreeItem> InWeakTreeItem, const STableRow<FPoseWatchManagerTreeItemPtr>* InRow)
+	{
+		WeakTreeItem = InWeakTreeItem;
+		WeakPoseWatchManager = InWeakPoseWatchManager;
+		WeakColumn = InWeakColumn;
+
+		Row = InRow;
+
+		SImage::Construct(
+			SImage::FArguments()
+			.ColorAndOpacity(this, &SVisibilityWidget::GetForegroundColor)
+			.Image(this, &SVisibilityWidget::GetBrush)
+		);
+
+		static const FName NAME_VisibleHoveredBrush = TEXT("Level.VisibleHighlightIcon16x");
+		static const FName NAME_VisibleNotHoveredBrush = TEXT("Level.VisibleIcon16x");
+		static const FName NAME_NotVisibleHoveredBrush = TEXT("Level.NotVisibleHighlightIcon16x");
+		static const FName NAME_NotVisibleNotHoveredBrush = TEXT("Level.NotVisibleIcon16x");
+
+		VisibleHoveredBrush = FAppStyle::Get().GetBrush(NAME_VisibleHoveredBrush);
+		VisibleNotHoveredBrush = FAppStyle::Get().GetBrush(NAME_VisibleNotHoveredBrush);
+
+		NotVisibleHoveredBrush = FAppStyle::Get().GetBrush(NAME_NotVisibleHoveredBrush);
+		NotVisibleNotHoveredBrush = FAppStyle::Get().GetBrush(NAME_NotVisibleNotHoveredBrush);
+	}
+
+private:
+	FReply HandleClick()
+	{
+		TSharedPtr<IPoseWatchManagerTreeItem> TreeItem = WeakTreeItem.Pin();
+		TreeItem->SetIsVisible(!TreeItem->GetVisibility());
+		return FReply::Handled();
+	}
+
+	/** Called when the mouse button is pressed down on this widget */
+	virtual FReply OnMouseButtonDown(const FGeometry& MyGeometry, const FPointerEvent& MouseEvent) override
+	{
+		if (MouseEvent.GetEffectingButton() != EKeys::LeftMouseButton)
+		{
+			return FReply::Unhandled();
+		}
+
+		return HandleClick();
+	}
+
+	/** Get the brush for this widget */
+	const FSlateBrush* GetBrush() const
+	{
+		if (IsVisible())
+		{
+			return IsHovered() ? VisibleHoveredBrush : VisibleNotHoveredBrush;
+		}
+		else
+		{
+			return IsHovered() ? NotVisibleHoveredBrush : NotVisibleNotHoveredBrush;
+		}
+	}
+
+	virtual FSlateColor GetForegroundColor() const
+	{
+		auto TreeItem = WeakTreeItem.Pin();
+		const bool bIsSelected = WeakPoseWatchManager.Pin()->GetTree().IsItemSelected(TreeItem.ToSharedRef());
+
+		if (!TreeItem->IsEnabled())
+		{
+			return FPoseWatchManagerCommonLabelData::DarkColor;
+		}
+
+		if (IsHovered() && !bIsSelected)
+		{
+			return FAppStyle::Get().GetSlateColor("Colors.ForegroundHover");
+		}
+
+		return FSlateColor::UseForeground();
+	}
+
+	static bool IsVisible(const FPoseWatchManagerTreeItemPtr& Item, const TSharedPtr<FPoseWatchManagerColumnVisibility>& Column)
+	{
+		return (Column.IsValid() && Item && Item.IsValid()) ? Item->GetVisibility() : false;
+	}
+
+	bool IsVisible() const
+	{
+		return IsVisible(WeakTreeItem.Pin(), WeakColumn.Pin());
+	}
+
+	void SetIsVisible(const bool bVisible)
+	{
+		TSharedPtr<IPoseWatchManagerTreeItem> TreeItem = WeakTreeItem.Pin();
+		TreeItem.Get()->SetIsVisible(bVisible);
+		WeakPoseWatchManager.Pin()->Refresh();	
+	}
+
+	/** The tree item we relate to */
+	TWeakPtr<IPoseWatchManagerTreeItem> WeakTreeItem;
+
+	TWeakPtr<IPoseWatchManager> WeakPoseWatchManager;
+
+	/** Weak pointer back to the column */
+	TWeakPtr<FPoseWatchManagerColumnVisibility> WeakColumn;
+
+	/** Weak pointer back to the row */
+	const STableRow<FPoseWatchManagerTreeItemPtr>* Row;
+
+	/** Visibility brushes for the various states */
+	const FSlateBrush* VisibleHoveredBrush;
+	const FSlateBrush* VisibleNotHoveredBrush;
+	const FSlateBrush* NotVisibleHoveredBrush;
+	const FSlateBrush* NotVisibleNotHoveredBrush;
+};
+
+SHeaderRow::FColumn::FArguments FPoseWatchManagerColumnVisibility::ConstructHeaderRowColumn()
+{
+	return SHeaderRow::Column(GetColumnID())
+		.FixedWidth(24.f)
+		.HAlignHeader(HAlign_Left)
+		.VAlignHeader(VAlign_Center)
+		.HAlignCell(HAlign_Center)
+		.VAlignCell(VAlign_Center)
+		.DefaultTooltip(FText::FromName(GetColumnID()))
+		.HeaderContentPadding(FMargin(4.0f, 0.0f, 0.0f, 0.0f))
+		[
+			SNew(SImage)
+			.ColorAndOpacity(FSlateColor::UseForeground())
+			.Image(FAppStyle::Get().GetBrush("Level.VisibleIcon16x"))
+		];
+}
+
+const TSharedRef<SWidget> FPoseWatchManagerColumnVisibility::ConstructRowWidget(FPoseWatchManagerTreeItemRef TreeItem, const STableRow<FPoseWatchManagerTreeItemPtr>& Row)
+{
+	return SNew(SHorizontalBox)
+		+ SHorizontalBox::Slot()
+		.AutoWidth()
+		.VAlign(VAlign_Center)
+		[
+			SNew(SVisibilityWidget, SharedThis(this), WeakPoseWatchManager, TreeItem, &Row)
+		];
+}
+
+#undef LOCTEXT_NAMESPACE
+
diff --git a/Engine/Source/Editor/Persona/Private/PoseWatchManagerColumnVisibility.h b/Engine/Source/Editor/Persona/Private/PoseWatchManagerColumnVisibility.h
new file mode 100644
--- /dev/null
+++ b/Engine/Source/Editor/Persona/Private/PoseWatchManagerColumnVisibility.h
@@ -0,0 +1,33 @@
+// Copyright Epic Games, Inc. All Rights Reserved.
+
+#pragma once
+
+#include "CoreMinimal.h"
+#include "Widgets/SWidget.h"
+#include "Widgets/Views/SHeaderRow.h"
+#include "PoseWatchManagerPublicTypes.h"
+#include "IPoseWatchManagerColumn.h"
+#include "IPoseWatchManager.h"
+
+template<typename ItemType> class STableRow;
+
+/**
+ * A PoseWatchManager column which handles setting and visualizing a pose watch's visibility
+ */
+class FPoseWatchManagerColumnVisibility : public IPoseWatchManagerColumn
+{
+public:
+	FPoseWatchManagerColumnVisibility(IPoseWatchManager& PoseWatchManager)
+		: WeakPoseWatchManager(StaticCastSharedRef<IPoseWatchManager>(PoseWatchManager.AsShared()))
+	{}
+
+	// Begin IPoseWatchManagerColumn Implementation
+	static FName GetID() { return FPoseWatchManagerBuiltInColumnTypes::Visibility(); }
+	virtual FName GetColumnID() override { return GetID(); }
+	virtual SHeaderRow::FColumn::FArguments ConstructHeaderRowColumn() override;
+	virtual const TSharedRef< SWidget > ConstructRowWidget(FPoseWatchManagerTreeItemRef TreeItem, const STableRow<FPoseWatchManagerTreeItemPtr>& Row) override;
+	// End IPoseWatchManagerColumn Implementation
+
+private:
+	TWeakPtr<IPoseWatchManager> WeakPoseWatchManager;
+};
diff --git a/Engine/Source/Editor/Persona/Private/PoseWatchManagerDefaultHierarchy.cpp b/Engine/Source/Editor/Persona/Private/PoseWatchManagerDefaultHierarchy.cpp
new file mode 100644
--- /dev/null
+++ b/Engine/Source/Editor/Persona/Private/PoseWatchManagerDefaultHierarchy.cpp
@@ -0,0 +1,71 @@
+// Copyright Epic Games, Inc. All Rights Reserved.
+
+#include "PoseWatchManagerDefaultHierarchy.h"
+#include "PoseWatchManagerFolderTreeItem.h"
+#include "PoseWatchManagerPoseWatchTreeItem.h"
+#include "PoseWatchManagerDefaultMode.h"
+#include "EditorFolderUtils.h"
+#include "Engine/PoseWatch.h"
+#include "Animation/AnimBlueprint.h"
+#include "SPoseWatchManager.h"
+
+FPoseWatchManagerDefaultHierarchy::FPoseWatchManagerDefaultHierarchy(FPoseWatchManagerDefaultMode* InMode) : Mode(InMode) {}
+
+FPoseWatchManagerTreeItemPtr FPoseWatchManagerDefaultHierarchy::FindParent(const IPoseWatchManagerTreeItem& Item, const TMap<FObjectKey, FPoseWatchManagerTreeItemPtr>& Items) const
+{
+	UPoseWatchFolder* ParentFolder = nullptr;
+	if (const FPoseWatchManagerFolderTreeItem* FolderTreeItem = Item.CastTo<FPoseWatchManagerFolderTreeItem>())
+	{
+		ParentFolder = FolderTreeItem->PoseWatchFolder->GetParent();
+
+	}
+	else if (const FPoseWatchManagerPoseWatchTreeItem* PoseWatchTreeItem = Item.CastTo<FPoseWatchManagerPoseWatchTreeItem>())
+	{
+		ParentFolder = PoseWatchTreeItem->PoseWatch->GetParent();
+	}
+
+	if (ParentFolder)
+	{
+		const FPoseWatchManagerTreeItemPtr* ParentTreeItem = Items.Find(ParentFolder);
+		if (ParentTreeItem)
+		{
+			return *ParentTreeItem;
+		}
+	}
+	
+	return nullptr;
+}
+
+void FPoseWatchManagerDefaultHierarchy::CreateItems(TArray<FPoseWatchManagerTreeItemPtr>& OutItems) const
+{
+	UAnimBlueprint* AnimBlueprint = Mode->PoseWatchManager->AnimBlueprint;
+	for (int32 Index = 0; Index <  AnimBlueprint->PoseWatchFolders.Num(); ++Index)
+	{
+		TObjectPtr<UPoseWatchFolder>& PoseWatchFolder = AnimBlueprint->PoseWatchFolders[Index];
+		if (PoseWatchFolder)
+		{
+			FPoseWatchManagerTreeItemPtr PoseWatchFolderItem = Mode->PoseWatchManager->CreateItemFor<FPoseWatchManagerFolderTreeItem>(PoseWatchFolder);
+			OutItems.Add(PoseWatchFolderItem);
+		}
+		else
+		{
+			AnimBlueprint->PoseWatchFolders.RemoveAtSwap(Index);
+		}
+	}
+	for (int32 Index = 0; Index < AnimBlueprint->PoseWatches.Num(); ++Index)
+	{
+		TObjectPtr<UPoseWatch>& PoseWatch = AnimBlueprint->PoseWatches[Index];
+		if (PoseWatch)
+		{
+			if (!PoseWatch->GetShouldDeleteOnDeselect())
+			{
+				FPoseWatchManagerTreeItemPtr PoseWatchItem = Mode->PoseWatchManager->CreateItemFor<FPoseWatchManagerPoseWatchTreeItem>(PoseWatch);
+				OutItems.Add(PoseWatchItem);
+			}
+		}
+		else
+		{
+			AnimBlueprint->PoseWatches.RemoveAtSwap(Index);
+		}
+	}
+}
diff --git a/Engine/Source/Editor/Persona/Private/PoseWatchManagerDefaultHierarchy.h b/Engine/Source/Editor/Persona/Private/PoseWatchManagerDefaultHierarchy.h
new file mode 100644
--- /dev/null
+++ b/Engine/Source/Editor/Persona/Private/PoseWatchManagerDefaultHierarchy.h
@@ -0,0 +1,21 @@
+// Copyright Epic Games, Inc. All Rights Reserved.
+
+#pragma once
+
+#include "CoreMinimal.h"
+#include "PoseWatchManagerFwd.h"
+#include "IPoseWatchManagerTreeItem.h"
+
+class FPoseWatchManagerDefaultMode;
+
+class FPoseWatchManagerDefaultHierarchy
+{
+public:
+	FPoseWatchManagerDefaultHierarchy(FPoseWatchManagerDefaultMode* InMode);
+
+	FPoseWatchManagerTreeItemPtr FindParent(const IPoseWatchManagerTreeItem& Item, const TMap<FObjectKey, FPoseWatchManagerTreeItemPtr>& Items) const;
+
+	void CreateItems(TArray<FPoseWatchManagerTreeItemPtr>& OutItems) const;
+
+	FPoseWatchManagerDefaultMode* Mode;
+};
\ No newline at end of file
diff --git a/Engine/Source/Editor/Persona/Private/PoseWatchManagerDefaultMode.cpp b/Engine/Source/Editor/Persona/Private/PoseWatchManagerDefaultMode.cpp
new file mode 100644
--- /dev/null
+++ b/Engine/Source/Editor/Persona/Private/PoseWatchManagerDefaultMode.cpp
@@ -0,0 +1,236 @@
+// Copyright Epic Games, Inc. All Rights Reserved.
+
+#include "PoseWatchManagerDefaultMode.h"
+#include "AnimationEditorUtils.h"
+#include "Engine/PoseWatch.h"
+#include "SGraphNode.h"
+#include "SGraphPanel.h"
+#include "IAnimationBlueprintEditor.h"
+#include "PoseWatchManagerFolderTreeItem.h"
+#include "PoseWatchManagerPoseWatchTreeItem.h"
+#include "SPoseWatchManager.h"
+
+#define LOCTEXT_NAMESPACE "PoseWatchDefaultMode"
+
+/** Functor which can be used to get weak actor pointers from a selection */
+struct FWeakPoseWatchSelector
+{
+	bool operator()(const TWeakPtr<IPoseWatchManagerTreeItem>& Item, TWeakObjectPtr<UPoseWatch>& DataOut) const;
+};
+
+FPoseWatchManagerDefaultMode::FPoseWatchManagerDefaultMode(SPoseWatchManager* InPoseWatchManager)
+	: PoseWatchManager(InPoseWatchManager)
+{
+	Rebuild();
+}
+
+void FPoseWatchManagerDefaultMode::Rebuild()
+{
+	Hierarchy = MakeUnique<FPoseWatchManagerDefaultHierarchy>(this);
+}
+
+bool FPoseWatchManagerDefaultMode::ParseDragDrop(FPoseWatchManagerDragDropPayload& OutPayload, const FDragDropOperation& Operation) const
+{
+	if (Operation.IsOfType<FPoseWatchManagerDragDropOp>())
+	{
+		const auto& OutlinerOp = static_cast<const FPoseWatchManagerDragDropOp&>(Operation);
+		if (const auto& PoseWatchOp = OutlinerOp.GetSubOp<FPoseWatchDragDropOp>())
+		{
+			OutPayload.DraggedItem = PoseWatchManager->GetTreeItem(PoseWatchOp->PoseWatch.Get());
+		}
+		if (const auto& FolderOp = OutlinerOp.GetSubOp<FPoseWatchFolderDragDropOp>())
+		{
+			OutPayload.DraggedItem = PoseWatchManager->GetTreeItem(FolderOp->PoseWatchFolder.Get());
+		}
+		return true;
+	}
+
+	return false;
+}
+
+FPoseWatchManagerDragValidationInfo FPoseWatchManagerDefaultMode::ValidateDrop(const IPoseWatchManagerTreeItem& DropTarget, const FPoseWatchManagerDragDropPayload& Payload) const
+{
+	check(Payload.DraggedItem.IsValid())
+	TSharedPtr<IPoseWatchManagerTreeItem> PayloadItem = Payload.DraggedItem.Pin();
+	
+	// Support for removing the parent by dragging an item to the bottom of the tree
+	if (!DropTarget.IsValid())
+	{
+		return FPoseWatchManagerDragValidationInfo(EPoseWatchManagerDropCompatibility::Compatible, LOCTEXT("MoveToRoot", "Move to root"));
+	}
+
+	// Pose watches cannot be a parent
+	if (const FPoseWatchManagerPoseWatchTreeItem* PoseWatchDropTarget = DropTarget.CastTo<FPoseWatchManagerPoseWatchTreeItem>())
+	{
+		return FPoseWatchManagerDragValidationInfo(EPoseWatchManagerDropCompatibility::Incompatible, LOCTEXT("PoseWatchBadParent", "Pose Watches cannot be parents"));
+	}
+
+	// Folders can be parents as long as the payload is not already a parent of the drop target, otherwise there'll be cycles
+	if (const FPoseWatchManagerFolderTreeItem* FolderItem = DropTarget.CastTo<FPoseWatchManagerFolderTreeItem>())
+	{
+		if (FolderItem->PoseWatchFolder.IsValid())
+		{
+			// Dropping a folder into a folder
+			if (const FPoseWatchManagerFolderTreeItem* FolderPayloadItem = PayloadItem.Get()->CastTo<FPoseWatchManagerFolderTreeItem>())
+			{
+				if (FolderPayloadItem->PoseWatchFolder == FolderItem->PoseWatchFolder)
+				{
+					return FPoseWatchManagerDragValidationInfo(EPoseWatchManagerDropCompatibility::Incompatible, LOCTEXT("FolderInItself", "A folder cannot contain itself"));
+				}
+
+				if (FolderPayloadItem->PoseWatchFolder->IsIn(FolderItem->PoseWatchFolder.Get()))
+				{
+					FText ValidationText = FText::Format(LOCTEXT("MovePoseWatchIntoFolder", "This folder is already inside {0}"), FolderItem->PoseWatchFolder->GetLabel());
+					return FPoseWatchManagerDragValidationInfo(EPoseWatchManagerDropCompatibility::Incompatible, ValidationText);
+				}
+
+				//if (FolderPayloadItem->PoseWatchFolder->IsDescendantOf(FolderItem->PoseWatchFolder.Get()))
+				if (FolderItem->PoseWatchFolder->IsDescendantOf(FolderPayloadItem->PoseWatchFolder.Get()))
+				{
+					return FPoseWatchManagerDragValidationInfo(EPoseWatchManagerDropCompatibility::Incompatible, LOCTEXT("ParentFolderNotChildren", "Parent folders cannot be children"));
+				}
+
+				FText OutErrorMessage;
+				FText Label = FolderPayloadItem->PoseWatchFolder->GetLabel();
+				UPoseWatchFolder* ConflictingFolder = nullptr;
+
+				if (FolderPayloadItem->PoseWatchFolder->IsFolderLabelUniqueInFolder(FolderPayloadItem->PoseWatchFolder->GetLabel(), FolderItem->PoseWatchFolder.Get()))
+				{
+					FText ValidationText = FText::Format(LOCTEXT("MovePoseWatchIntoFolder", "Move into {0}"), FolderItem->PoseWatchFolder->GetLabel());
+					return FPoseWatchManagerDragValidationInfo(EPoseWatchManagerDropCompatibility::Compatible, ValidationText);
+				}
+				else
+				{
+					FText ValidationText = FText::Format(LOCTEXT("MovePoseWatchIntoFolder", "A folder with that name already exists within {0}"), FolderItem->PoseWatchFolder->GetLabel());
+					return FPoseWatchManagerDragValidationInfo(EPoseWatchManagerDropCompatibility::Incompatible, ValidationText);
+				}
+			}
+
+			// Dropping a pose watch into a folder
+			if (const FPoseWatchManagerPoseWatchTreeItem* PoseWatchPayloadItem = PayloadItem.Get()->CastTo<FPoseWatchManagerPoseWatchTreeItem>())
+			{
+				if (PoseWatchPayloadItem->PoseWatch->IsIn(FolderItem->PoseWatchFolder.Get()))
+				{
+					FText ValidationText = FText::Format(LOCTEXT("MovePoseWatchIntoFolder", "This pose watch is already inside {0}"), FolderItem->PoseWatchFolder->GetLabel());
+					return FPoseWatchManagerDragValidationInfo(EPoseWatchManagerDropCompatibility::Incompatible, ValidationText);
+				}
+				else
+				{
+					FText ValidationText = FText::Format(LOCTEXT("MovePoseWatchIntoFolder", "Move into {0}"), FolderItem->PoseWatchFolder->GetLabel());
+					return FPoseWatchManagerDragValidationInfo(EPoseWatchManagerDropCompatibility::Compatible, ValidationText);
+				}
+			}
+		}
+		else
+		{
+			if (PayloadItem.Get()->IsAssignedFolder())
+			{
+				return FPoseWatchManagerDragValidationInfo(EPoseWatchManagerDropCompatibility::Compatible, LOCTEXT("MoveToRoot", "Move to root"));
+			}
+			else
+			{
+				return FPoseWatchManagerDragValidationInfo(EPoseWatchManagerDropCompatibility::Incompatible, LOCTEXT("FolderAlreadyInRoot", "This item is already in root"));
+			}
+		}
+	}
+	
+	return FPoseWatchManagerDragValidationInfo(EPoseWatchManagerDropCompatibility::Incompatible, FText());
+}
+
+void FPoseWatchManagerDefaultMode::OnDrop(IPoseWatchManagerTreeItem& DropTarget, const FPoseWatchManagerDragDropPayload& Payload, const FPoseWatchManagerDragValidationInfo& ValidationInfo) const
+{
+	check(Payload.DraggedItem.IsValid());
+	check(DropTarget.IsA<FPoseWatchManagerFolderTreeItem>());
+	const FPoseWatchManagerFolderTreeItem* FolderDropTarget = DropTarget.CastTo<FPoseWatchManagerFolderTreeItem>();
+	const TWeakPtr<IPoseWatchManagerTreeItem> PayloadItem = Payload.DraggedItem;
+
+	if (const FPoseWatchManagerFolderTreeItem* ChildFolderItem = PayloadItem.Pin()->CastTo<FPoseWatchManagerFolderTreeItem>())
+	{
+		ChildFolderItem->PoseWatchFolder->MoveTo(FolderDropTarget->PoseWatchFolder.Get());
+		return;
+	}
+	if (const FPoseWatchManagerPoseWatchTreeItem* ChildPoseWatchItem = PayloadItem.Pin()->CastTo<FPoseWatchManagerPoseWatchTreeItem>())
+	{
+		ChildPoseWatchItem->PoseWatch->MoveTo(FolderDropTarget->PoseWatchFolder.Get());
+		return;
+	}
+	check(false);
+}
+
+TSharedPtr<FDragDropOperation> FPoseWatchManagerDefaultMode::CreateDragDropOperation(const TArray<FPoseWatchManagerTreeItemPtr>& InTreeItems) const
+{
+	check(InTreeItems.Num() == 1)
+	FPoseWatchManagerTreeItemPtr TreeItem = InTreeItems[0];
+
+	FPoseWatchManagerDragDropPayload DraggedObjects(TreeItem);
+
+	TSharedPtr<FPoseWatchManagerDragDropOp> OutlinerOp = MakeShareable(new FPoseWatchManagerDragDropOp());
+
+	if (FPoseWatchManagerPoseWatchTreeItem* PoseWatchTreeItem = TreeItem->CastTo<FPoseWatchManagerPoseWatchTreeItem>())
+	{
+		TSharedPtr<FPoseWatchDragDropOp> Operation = MakeShareable(new FPoseWatchDragDropOp);
+		Operation->Init(TWeakObjectPtr<UPoseWatch>(PoseWatchTreeItem->PoseWatch));
+		OutlinerOp->AddSubOp(Operation);
+	}
+
+	else if (FPoseWatchManagerFolderTreeItem* FolderTreeItem = TreeItem->CastTo<FPoseWatchManagerFolderTreeItem>())
+	{
+		TSharedPtr<FPoseWatchFolderDragDropOp> Operation = MakeShareable(new FPoseWatchFolderDragDropOp);
+		Operation->Init(TWeakObjectPtr<UPoseWatchFolder>(FolderTreeItem->PoseWatchFolder));
+		OutlinerOp->AddSubOp(Operation);
+	}
+
+	OutlinerOp->Construct();
+	return OutlinerOp;
+}
+
+FReply FPoseWatchManagerDefaultMode::OnDragOverItem(const FDragDropEvent& Event, const IPoseWatchManagerTreeItem& Item) const
+{
+	return FReply::Handled();
+}
+
+TSharedPtr<SWidget> FPoseWatchManagerDefaultMode::CreateContextMenu()
+{
+	FPoseWatchManagerTreeItemPtr SelectedItem = PoseWatchManager->GetSelection();
+	return SelectedItem ? SelectedItem->CreateContextMenu() : SNullWidget::NullWidget;
+}
+
+FReply FPoseWatchManagerDefaultMode::OnKeyDown(const FKeyEvent& InKeyEvent)
+{
+	const FPoseWatchManagerTreeItemPtr& Selection = PoseWatchManager->GetSelection();
+
+	if (InKeyEvent.GetKey() == EKeys::F2)
+	{
+		if (Selection)
+		{
+			//FPoseWatchManagerTreeItemPtr ItemToRename = Selection.Get();
+
+			if (Selection.IsValid())
+			{
+				PoseWatchManager->SetPendingRenameItem(Selection);
+				PoseWatchManager->ScrollItemIntoView(Selection);
+			}
+
+			return FReply::Handled();
+		}
+	}
+
+	else if (InKeyEvent.GetKey() == EKeys::F5)
+	{
+		PoseWatchManager->FullRefresh();
+		return FReply::Handled();
+	}
+
+	else if (InKeyEvent.GetKey() == EKeys::Delete || InKeyEvent.GetKey() == EKeys::BackSpace)
+	{
+		if (Selection)
+		{
+			Selection->OnRemoved();
+			return FReply::Handled();
+		}
+	}
+
+	return FReply::Unhandled();
+}
+
+#undef LOCTEXT_NAMESPACE
\ No newline at end of file
diff --git a/Engine/Source/Editor/Persona/Private/PoseWatchManagerDefaultMode.h b/Engine/Source/Editor/Persona/Private/PoseWatchManagerDefaultMode.h
new file mode 100644
--- /dev/null
+++ b/Engine/Source/Editor/Persona/Private/PoseWatchManagerDefaultMode.h
@@ -0,0 +1,38 @@
+// Copyright Epic Games, Inc. All Rights Reserved.
+
+#pragma once
+
+#include "PoseWatchManagerFwd.h"
+#include "PoseWatchManagerDefaultHierarchy.h"
+
+struct FPoseWatchManagerItemSelection;
+class SPoseWatchManager;
+
+class FPoseWatchManagerDefaultMode
+{
+public:
+	FPoseWatchManagerDefaultMode(SPoseWatchManager* InPoseWatchManager);
+
+	void Rebuild();
+
+	FReply OnKeyDown(const FKeyEvent& InKeyEvent);
+
+	TSharedPtr<SWidget> CreateContextMenu();
+
+	bool ParseDragDrop(FPoseWatchManagerDragDropPayload& OutPayload, const FDragDropOperation& Operation) const;
+	
+	FPoseWatchManagerDragValidationInfo ValidateDrop(const IPoseWatchManagerTreeItem& DropTarget, const FPoseWatchManagerDragDropPayload& Payload) const;
+	
+	void OnDrop(IPoseWatchManagerTreeItem& DropTarget, const FPoseWatchManagerDragDropPayload& Payload, const FPoseWatchManagerDragValidationInfo& ValidationInfo) const;
+	
+	TSharedPtr<FDragDropOperation> CreateDragDropOperation(const TArray<FPoseWatchManagerTreeItemPtr>& InTreeItems) const;
+	
+	FReply OnDragOverItem(const FDragDropEvent& Event, const IPoseWatchManagerTreeItem& Item) const;
+
+	int32 GetTypeSortPriority(const IPoseWatchManagerTreeItem& Item) const { return 0; }
+
+
+	SPoseWatchManager* PoseWatchManager;
+
+	TUniquePtr<FPoseWatchManagerDefaultHierarchy> Hierarchy;
+};
\ No newline at end of file
diff --git a/Engine/Source/Editor/Persona/Private/PoseWatchManagerDragDrop.cpp b/Engine/Source/Editor/Persona/Private/PoseWatchManagerDragDrop.cpp
new file mode 100644
--- /dev/null
+++ b/Engine/Source/Editor/Persona/Private/PoseWatchManagerDragDrop.cpp
@@ -0,0 +1,69 @@
+// Copyright Epic Games, Inc. All Rights Reserved.
+
+#include "PoseWatchManagerDragDrop.h"
+#include "IPoseWatchManagerTreeItem.h"
+#include "PoseWatchManagerStandaloneTypes.h"
+
+#define LOCTEXT_NAMESPACE "SPoseWatchManager"
+
+
+FPoseWatchManagerDragDropOp::FPoseWatchManagerDragDropOp()
+	: OverrideText()
+	, OverrideIcon(nullptr)
+{}
+
+EVisibility FPoseWatchManagerDragDropOp::GetOverrideVisibility() const
+{
+	return OverrideText.IsEmpty() && OverrideIcon == nullptr ? EVisibility::Collapsed : EVisibility::Visible;
+}
+
+EVisibility FPoseWatchManagerDragDropOp::GetDefaultVisibility() const
+{
+	return OverrideText.IsEmpty() && OverrideIcon == nullptr ? EVisibility::Visible : EVisibility::Collapsed;
+}
+
+TSharedPtr<SWidget> FPoseWatchManagerDragDropOp::GetDefaultDecorator() const
+{
+	TSharedRef<SVerticalBox> VerticalBox = SNew(SVerticalBox);
+
+	VerticalBox->AddSlot()
+	[
+		SNew(SBorder)
+		.BorderImage(FEditorStyle::GetBrush("Graph.ConnectorFeedback.Border"))
+		.Visibility(this, &FPoseWatchManagerDragDropOp::GetOverrideVisibility)
+		.Content()
+		[
+			SNew(SHorizontalBox)
+
+			+ SHorizontalBox::Slot()
+			.AutoWidth()
+			.Padding(0.0f, 0.0f, 3.0f, 0.0f)
+			[
+				SNew(SImage)
+				.Image(this, &FPoseWatchManagerDragDropOp::GetOverrideIcon)
+			]
+
+			+ SHorizontalBox::Slot()
+			.AutoWidth()
+			.VAlign(VAlign_Center)
+			[
+				SNew(STextBlock)
+				.Text(this, &FPoseWatchManagerDragDropOp::GetOverrideText)
+			]
+		]
+	];
+
+	for (auto& SubOp : SubOps)
+	{
+		auto Content = SubOp->GetDefaultDecorator();
+		if (Content.IsValid())
+		{
+			Content->SetVisibility(TAttribute<EVisibility>(this, &FPoseWatchManagerDragDropOp::GetDefaultVisibility));
+			VerticalBox->AddSlot()[Content.ToSharedRef()];
+		}
+	}
+
+	return VerticalBox;
+}
+
+#undef LOCTEXT_NAMESPACE
diff --git a/Engine/Source/Editor/Persona/Private/PoseWatchManagerDragDrop.h b/Engine/Source/Editor/Persona/Private/PoseWatchManagerDragDrop.h
new file mode 100644
--- /dev/null
+++ b/Engine/Source/Editor/Persona/Private/PoseWatchManagerDragDrop.h
@@ -0,0 +1,142 @@
+// Copyright Epic Games, Inc. All Rights Reserved.
+
+#pragma once
+
+#include "CoreMinimal.h"
+#include "PoseWatchManagerFwd.h"
+#include "Layout/Visibility.h"
+#include "Input/DragAndDrop.h"
+#include "DragAndDrop/CompositeDragDropOp.h"
+#include "IPoseWatchManagerTreeItem.h"
+#include "Engine/PoseWatch.h"
+
+/** Enum to describe the compatibility of a drag drop operation */
+enum class EPoseWatchManagerDropCompatibility : uint8
+{
+	Compatible,
+	Incompatible
+};
+
+struct FPoseWatchManagerDragDropPayload
+{
+	FPoseWatchManagerDragDropPayload(const FDragDropOperation& InOperation = FDragDropOperation())
+		: SourceOperation(InOperation)
+	{}
+
+	template<typename TreeType>
+	FPoseWatchManagerDragDropPayload(const TreeType& InDraggedItem, const FDragDropOperation& InOperation = FDragDropOperation())
+		: SourceOperation(InOperation)
+	{
+		DraggedItem = InDraggedItem;
+	}
+
+	TWeakPtr<IPoseWatchManagerTreeItem> DraggedItem;
+
+	/** The source FDragDropOperation */
+	const FDragDropOperation& SourceOperation;
+};
+
+/** Struct used for validation of a drag/drop operation in the pose watch manager */
+struct FPoseWatchManagerDragValidationInfo
+{
+	/** The tooltip type to display on the operation */
+	EPoseWatchManagerDropCompatibility CompatibilityType;
+
+	/** The tooltip text to display on the operation */
+	FText ValidationText;
+
+	/** Construct this validation information out of a tooltip type and some text */
+	FPoseWatchManagerDragValidationInfo(const EPoseWatchManagerDropCompatibility InCompatibilityType, const FText InValidationText)
+		: CompatibilityType(InCompatibilityType)
+		, ValidationText(InValidationText)
+	{}
+
+	/** Return a generic invalid result */
+	static FPoseWatchManagerDragValidationInfo Invalid()
+	{
+		return FPoseWatchManagerDragValidationInfo(EPoseWatchManagerDropCompatibility::Incompatible, FText());
+	}
+
+	/** @return true if this operation is valid, false otherwise */
+	bool IsValid() const
+	{
+		return CompatibilityType == EPoseWatchManagerDropCompatibility::Compatible;
+	}
+};
+
+/** A drag/drop operation that was started from the pose watch manager */
+struct PERSONA_API FPoseWatchManagerDragDropOp : public FCompositeDragDropOp
+{
+	DRAG_DROP_OPERATOR_TYPE(FPoseWatchManagerDragDropOp, FCompositeDragDropOp);
+
+	FPoseWatchManagerDragDropOp();
+
+	using FDragDropOperation::Construct;
+
+	void ResetTooltip()
+	{
+		OverrideText = FText();
+		OverrideIcon = nullptr;
+	}
+
+	void SetTooltip(FText InOverrideText, const FSlateBrush* InOverrideIcon)
+	{
+		OverrideText = InOverrideText;
+		OverrideIcon = InOverrideIcon;
+	}
+
+	virtual TSharedPtr<SWidget> GetDefaultDecorator() const override;
+
+private:
+
+	EVisibility GetOverrideVisibility() const;
+	EVisibility GetDefaultVisibility() const;
+
+	FText OverrideText;
+	FText GetOverrideText() const { return OverrideText; }
+
+	const FSlateBrush* OverrideIcon;
+	const FSlateBrush* GetOverrideIcon() const { return OverrideIcon; }
+};
+
+
+class FPoseWatchDragDropOp : public FDecoratedDragDropOp
+{
+public:
+	DRAG_DROP_OPERATOR_TYPE(FPoseWatchDragDropOp, FDecoratedDragDropOp)
+
+	/** Pose Watch that we are dragging */
+	TWeakObjectPtr<UPoseWatch> PoseWatch;
+
+	void Init(const TWeakObjectPtr<UPoseWatch>& InPoseWatch)
+	{
+		PoseWatch = InPoseWatch;
+
+		// Set text and icon
+		UClass* CommonSelClass = NULL;
+		CurrentIconBrush = FEditorStyle::Get().GetBrush(TEXT("ClassIcon.PoseAsset"));
+		CurrentHoverText = PoseWatch->GetLabel();
+
+		SetupDefaults();
+	}
+};
+
+class FPoseWatchFolderDragDropOp : public FDecoratedDragDropOp
+{
+public:
+	DRAG_DROP_OPERATOR_TYPE(FPoseWatchFolderDragDropOp, FDecoratedDragDropOp)
+
+	/** Pose Watch Folder that we are dragging */
+	TWeakObjectPtr<UPoseWatchFolder> PoseWatchFolder;
+
+	void Init(const TWeakObjectPtr<UPoseWatchFolder>& InPoseWatchFolder)
+	{
+		PoseWatchFolder = InPoseWatchFolder;
+
+		CurrentIconBrush = FEditorStyle::Get().GetBrush(TEXT("SceneOutliner.FolderClosed"));
+		CurrentHoverText = PoseWatchFolder->GetLabel();
+
+
+		SetupDefaults();
+	}
+};
diff --git a/Engine/Source/Editor/Persona/Private/PoseWatchManagerFolderTreeItem.cpp b/Engine/Source/Editor/Persona/Private/PoseWatchManagerFolderTreeItem.cpp
new file mode 100644
--- /dev/null
+++ b/Engine/Source/Editor/Persona/Private/PoseWatchManagerFolderTreeItem.cpp
@@ -0,0 +1,193 @@
+// Copyright Epic Games, Inc. All Rights Reserved.
+
+#include "PoseWatchManagerFolderTreeItem.h"
+#include "Textures/SlateIcon.h"
+#include "Widgets/Text/SInlineEditableTextBlock.h"
+#include "Framework/Commands/UIAction.h"
+#include "ToolMenus.h"
+#include "EditorStyleSet.h"
+#include "PoseWatchManagerDragDrop.h"
+#include "SPoseWatchManager.h"
+#include "Framework/MultiBox/MultiBoxBuilder.h"
+#include "Engine/PoseWatch.h"
+
+#define LOCTEXT_NAMESPACE "PoseWatchManagerFolderTreeItem"
+
+const EPoseWatchTreeItemType FPoseWatchManagerFolderTreeItem::Type(EPoseWatchTreeItemType::Folder);
+
+FPoseWatchManagerFolderTreeItem::FPoseWatchManagerFolderTreeItem(UPoseWatchFolder* InPoseWatchFolder)
+	: IPoseWatchManagerTreeItem(Type)
+	, ID(InPoseWatchFolder)
+	, PoseWatchFolder(InPoseWatchFolder)
+{
+}
+
+bool FPoseWatchManagerFolderTreeItem::IsAssignedFolder() const
+{
+	return PoseWatchFolder->IsAssignedFolder();
+}
+
+FObjectKey FPoseWatchManagerFolderTreeItem::GetID() const
+{
+	return ID;
+}
+
+FString FPoseWatchManagerFolderTreeItem::GetDisplayString() const
+{
+	return PoseWatchFolder->GetLabel().ToString();
+}
+
+void FPoseWatchManagerFolderTreeItem::SetIsVisible(const bool bVisible)
+{
+	PoseWatchFolder->SetIsVisible(bVisible);
+}
+
+TSharedPtr<SWidget> FPoseWatchManagerFolderTreeItem::CreateContextMenu()
+{
+	FMenuBuilder MenuBuilder(true, nullptr);
+
+	UPoseWatchFolder* FolderItemPtr = PoseWatchFolder.Get();
+	MenuBuilder.BeginSection(FName(), LOCTEXT("PoseWatchFolder", "Pose Watch Folder"));
+	MenuBuilder.AddMenuEntry(
+		LOCTEXT("DeleteFolder", "Delete Folder"),
+		LOCTEXT("DeleteFolderDescription", "Delete this folder and move its contents to this folder's parent"),
+		FSlateIcon(),
+		FUIAction(
+			FExecuteAction::CreateLambda([FolderItemPtr, this]() {
+				FolderItemPtr->OnRemoved();
+			})
+		)
+	);
+	MenuBuilder.EndSection();
+
+	return MenuBuilder.MakeWidget();
+}
+
+
+struct SPoseWatchManagerFolderTreeItemLabel : FPoseWatchManagerCommonLabelData, public SCompoundWidget
+{
+	SLATE_BEGIN_ARGS(SPoseWatchManagerFolderTreeItemLabel) {}
+	SLATE_END_ARGS()
+
+	void Construct(const FArguments& InArgs, FPoseWatchManagerFolderTreeItem& FolderItem, IPoseWatchManager& PoseWatchManager, const STableRow<FPoseWatchManagerTreeItemPtr>& InRow)
+	{
+		TreeItemPtr = StaticCastSharedRef<FPoseWatchManagerFolderTreeItem>(FolderItem.AsShared());
+		WeakPoseWatchManager = StaticCastSharedRef<IPoseWatchManager>(PoseWatchManager.AsShared());
+
+		TSharedPtr<SInlineEditableTextBlock> InlineTextBlock = SNew(SInlineEditableTextBlock)
+			.Text(this, &SPoseWatchManagerFolderTreeItemLabel::GetDisplayText)
+			.HighlightText(PoseWatchManager.GetFilterHighlightText())
+			.ColorAndOpacity(this, &SPoseWatchManagerFolderTreeItemLabel::GetForegroundColor)
+			.OnTextCommitted(this, &SPoseWatchManagerFolderTreeItemLabel::OnLabelCommitted)
+			.OnVerifyTextChanged(this, &SPoseWatchManagerFolderTreeItemLabel::OnVerifyItemLabelChanged)
+			.IsSelected(FIsSelected::CreateSP(&InRow, &STableRow<FPoseWatchManagerTreeItemPtr>::IsSelectedExclusively))
+			.IsReadOnly_Lambda([Item = FolderItem.AsShared(), this]()
+		{
+			return !CanExecuteRenameRequest(Item.Get());
+		});
+
+		FolderItem.RenameRequestEvent.BindSP(InlineTextBlock.Get(), &SInlineEditableTextBlock::EnterEditingMode);
+
+		ChildSlot
+		[
+			SNew(SHorizontalBox)
+
+			+ SHorizontalBox::Slot()
+			.AutoWidth()
+			.VAlign(VAlign_Center)
+			.Padding(FPoseWatchManagerDefaultTreeItemMetrics::IconPadding())
+			[
+				SNew(SBox)
+				.WidthOverride(FPoseWatchManagerDefaultTreeItemMetrics::IconSize())
+				.HeightOverride(FPoseWatchManagerDefaultTreeItemMetrics::IconSize())
+			[
+				SNew(SImage)
+				.Image(this, &SPoseWatchManagerFolderTreeItemLabel::GetIcon)
+			]
+		]
+
+		+ SHorizontalBox::Slot()
+			.FillWidth(1.0f)
+			.VAlign(VAlign_Center)
+			.Padding(0.0f, 2.0f)
+			[
+				InlineTextBlock.ToSharedRef()
+			]
+		];
+	}
+
+private:
+	TWeakPtr<FPoseWatchManagerFolderTreeItem> TreeItemPtr;
+
+	FText GetDisplayText() const
+	{
+		return TreeItemPtr.Pin()->PoseWatchFolder->GetLabel();
+	}
+
+	const FSlateBrush* GetIcon() const
+	{
+		auto TreeItem = TreeItemPtr.Pin();
+		if (!TreeItem.IsValid())
+		{
+			return FEditorStyle::Get().GetBrush(TEXT("SceneOutliner.FolderClosed"));
+		}
+
+		if (TreeItemPtr.Pin()->IsExpanded() && TreeItem->PoseWatchFolder->HasChildren())
+		{
+			return FEditorStyle::Get().GetBrush(TEXT("SceneOutliner.FolderOpen"));
+		}
+		else
+		{
+			return FEditorStyle::Get().GetBrush(TEXT("SceneOutliner.FolderClosed"));
+		}
+	}
+
+	FSlateColor GetForegroundColor() const
+	{
+		if (auto BaseColor = FPoseWatchManagerCommonLabelData::GetForegroundColor(*TreeItemPtr.Pin()))
+		{
+			return BaseColor.GetValue();
+		}
+
+		return FSlateColor::UseForeground();
+	}
+
+	bool OnVerifyItemLabelChanged(const FText& InLabel, FText& OutErrorMessage)
+	{
+		return TreeItemPtr.Pin()->PoseWatchFolder.Get()->ValidateLabelRename(InLabel, OutErrorMessage);
+	}
+
+	void OnLabelCommitted(const FText& InLabel, ETextCommit::Type InCommitInfo)
+	{
+		check(TreeItemPtr.Pin()->PoseWatchFolder.Get()->SetLabel(InLabel));
+		WeakPoseWatchManager.Pin()->FullRefresh();
+		WeakPoseWatchManager.Pin()->SetKeyboardFocus();
+	}
+};
+
+TSharedRef<SWidget> FPoseWatchManagerFolderTreeItem::GenerateLabelWidget(IPoseWatchManager& Outliner, const STableRow<FPoseWatchManagerTreeItemPtr>& InRow)
+{
+	return SNew(SPoseWatchManagerFolderTreeItemLabel, *this, Outliner, InRow);
+}
+
+void FPoseWatchManagerFolderTreeItem::OnRemoved()
+{
+	PoseWatchFolder->OnRemoved();
+}
+
+bool FPoseWatchManagerFolderTreeItem::GetVisibility() const
+{
+	return PoseWatchFolder.IsValid() ? PoseWatchFolder.Get()->GetIsVisible() : false;
+}
+
+bool FPoseWatchManagerFolderTreeItem::HasChildren() const
+{
+	return PoseWatchFolder.IsValid() ? PoseWatchFolder.Get()->HasChildren() : false;
+}
+
+bool FPoseWatchManagerFolderTreeItem::IsExpanded() const
+{
+	return PoseWatchFolder.IsValid() ? PoseWatchFolder.Get()->GetIsExpanded() : false;
+}
+
+#undef LOCTEXT_NAMESPACE
diff --git a/Engine/Source/Editor/Persona/Private/PoseWatchManagerFolderTreeItem.h b/Engine/Source/Editor/Persona/Private/PoseWatchManagerFolderTreeItem.h
new file mode 100644
--- /dev/null
+++ b/Engine/Source/Editor/Persona/Private/PoseWatchManagerFolderTreeItem.h
@@ -0,0 +1,39 @@
+// Copyright Epic Games, Inc. All Rights Reserved.
+
+#pragma once
+
+#include "CoreMinimal.h"
+#include "Widgets/SWidget.h"
+#include "IPoseWatchManagerTreeItem.h"
+
+class UPoseWatch;
+
+struct FPoseWatchManagerFolderTreeItem : IPoseWatchManagerTreeItem
+{
+public:
+	FPoseWatchManagerFolderTreeItem(UPoseWatchFolder* InPoseWatchFolder);
+
+	/** Constant identifier for this tree item */
+	const FObjectKey ID;
+	
+	/** The pose watch folder this item is describing */
+	TWeakObjectPtr<UPoseWatchFolder> PoseWatchFolder;
+	
+	/** Static type identifier for this tree item class */
+	static const EPoseWatchTreeItemType Type;
+
+	bool IsValid() const override { return PoseWatchFolder.IsValid(); }
+
+	/* Begin IPoseWatchManagerTreeItem Implementation */
+	virtual FObjectKey GetID() const override;
+	virtual FString GetDisplayString() const override;
+	virtual TSharedRef<SWidget> GenerateLabelWidget(IPoseWatchManager& Outliner, const STableRow<FPoseWatchManagerTreeItemPtr>& InRow) override;
+	virtual TSharedPtr<SWidget> CreateContextMenu() override;
+	virtual bool GetVisibility() const override;
+	virtual bool IsAssignedFolder() const override;
+	virtual void SetIsVisible(const bool bVisible) override;
+	virtual bool HasChildren() const override;
+	virtual void OnRemoved() override;
+	virtual bool IsExpanded() const override;
+	/* End IPoseWatchManagerTreeItem Implementation */
+};
diff --git a/Engine/Source/Editor/Persona/Private/PoseWatchManagerFwd.h b/Engine/Source/Editor/Persona/Private/PoseWatchManagerFwd.h
new file mode 100644
--- /dev/null
+++ b/Engine/Source/Editor/Persona/Private/PoseWatchManagerFwd.h
@@ -0,0 +1,24 @@
+// Copyright Epic Games, Inc. All Rights Reserved.
+
+#pragma once
+
+#include "CoreMinimal.h"
+#include "Misc/CoreDelegates.h"
+
+class IPoseWatchManager;
+class IPoseWatchManagerColumn;
+
+class UPoseWatch;
+class UPoseWatchFolder;
+
+struct FPoseWatchManagerInitializationOptions;
+
+struct IPoseWatchManagerTreeItem;
+struct FPoseWatchManagerPoseWatchTreeItem;
+struct FPoseWatchManagerFolderTreeItem;
+
+struct FObjectKey;
+
+
+struct FPoseWatchManagerDragDropPayload;
+struct FPoseWatchManagerDragValidationInfo;
\ No newline at end of file
diff --git a/Engine/Source/Editor/Persona/Private/PoseWatchManagerPoseWatchTreeItem.cpp b/Engine/Source/Editor/Persona/Private/PoseWatchManagerPoseWatchTreeItem.cpp
new file mode 100644
--- /dev/null
+++ b/Engine/Source/Editor/Persona/Private/PoseWatchManagerPoseWatchTreeItem.cpp
@@ -0,0 +1,247 @@
+// Copyright Epic Games, Inc. All Rights Reserved.
+
+#include "PoseWatchManagerPoseWatchTreeItem.h"
+#include "Textures/SlateIcon.h"
+#include "Widgets/Text/SInlineEditableTextBlock.h"
+#include "Framework/Commands/UIAction.h"
+#include "ToolMenus.h"
+#include "EditorStyleSet.h"
+#include "PoseWatchManagerDragDrop.h"
+#include "SPoseWatchManager.h"
+#include "Framework/MultiBox/MultiBoxBuilder.h"
+#include "Engine/PoseWatch.h"
+
+#define LOCTEXT_NAMESPACE "PoseWatchManagerPoseWatchTreeItem"
+
+struct PERSONA_API SPoseWatchManagerPoseWatchTreeLabel : FPoseWatchManagerCommonLabelData, public SCompoundWidget
+{
+	SLATE_BEGIN_ARGS(SPoseWatchManagerPoseWatchTreeLabel) {}
+	SLATE_END_ARGS()
+
+	void Construct(const FArguments& InArgs, FPoseWatchManagerPoseWatchTreeItem& PoseWatchTreeItem, IPoseWatchManager& PoseWatchManager, const STableRow<FPoseWatchManagerTreeItemPtr>& InRow)
+	{
+		WeakPoseWatchManager = StaticCastSharedRef<IPoseWatchManager>(PoseWatchManager.AsShared());
+
+		TreeItemPtr = StaticCastSharedRef<FPoseWatchManagerPoseWatchTreeItem>(PoseWatchTreeItem.AsShared());
+		WeakPoseWatchPtr = PoseWatchTreeItem.PoseWatch;
+
+		HighlightText = PoseWatchManager.GetFilterHighlightText();
+
+		TSharedPtr<SInlineEditableTextBlock> InlineTextBlock;
+
+		auto MainContent = SNew(SHorizontalBox)
+
+			+ SHorizontalBox::Slot()
+			.VAlign(VAlign_Center)
+			[
+				SAssignNew(InlineTextBlock, SInlineEditableTextBlock)
+				.Text(this, &SPoseWatchManagerPoseWatchTreeLabel::GetDisplayText)
+				.ToolTipText(this, &SPoseWatchManagerPoseWatchTreeLabel::GetTooltipText)
+				.HighlightText(HighlightText)
+				.ColorAndOpacity(this, &SPoseWatchManagerPoseWatchTreeLabel::GetForegroundColor)
+				.OnTextCommitted(this, &SPoseWatchManagerPoseWatchTreeLabel::OnLabelCommitted)
+				.OnVerifyTextChanged(this, &SPoseWatchManagerPoseWatchTreeLabel::OnVerifyItemLabelChanged)
+				.OnEnterEditingMode(this, &SPoseWatchManagerPoseWatchTreeLabel::OnEnterEditingMode)
+				.OnExitEditingMode(this, &SPoseWatchManagerPoseWatchTreeLabel::OnExitEditingMode)
+				.IsSelected(FIsSelected::CreateSP(&InRow, &STableRow<FPoseWatchManagerTreeItemPtr>::IsSelectedExclusively))
+			]
+
+			+ SHorizontalBox::Slot()
+			.VAlign(VAlign_Center)
+			.AutoWidth()
+			.Padding(0.0f, 0.f, 3.0f, 0.0f)
+			[
+				SNew(STextBlock)
+				.Text(this, &SPoseWatchManagerPoseWatchTreeLabel::GetTypeText)
+				.Visibility(this, &SPoseWatchManagerPoseWatchTreeLabel::GetTypeTextVisibility)
+				.HighlightText(HighlightText)
+			];
+
+		if (WeakPoseWatchManager.Pin())
+		{
+			PoseWatchTreeItem.RenameRequestEvent.BindSP(InlineTextBlock.Get(), &SInlineEditableTextBlock::EnterEditingMode);
+		}
+
+		ChildSlot
+		[
+			SNew(SHorizontalBox)
+
+			+ SHorizontalBox::Slot()
+			.AutoWidth()
+			.VAlign(VAlign_Center)
+			.Padding(FPoseWatchManagerDefaultTreeItemMetrics::IconPadding())
+			[
+				SNew(SBox)
+				.WidthOverride(FPoseWatchManagerDefaultTreeItemMetrics::IconSize())
+				.HeightOverride(FPoseWatchManagerDefaultTreeItemMetrics::IconSize())
+				[
+					SNew(SImage)
+					.Image(this, &SPoseWatchManagerPoseWatchTreeLabel::GetIcon)
+					.ToolTipText(this, &SPoseWatchManagerPoseWatchTreeLabel::GetIconTooltip)
+					.ColorAndOpacity(FSlateColor::UseForeground())
+				]
+			]
+
+			+ SHorizontalBox::Slot()
+			.FillWidth(1.0f)
+			.VAlign(VAlign_Center)
+			.Padding(0.0f, 0.0f)
+			[
+				MainContent
+			]
+		];
+	}
+
+private:
+	TWeakPtr<FPoseWatchManagerPoseWatchTreeItem> TreeItemPtr;
+	TWeakObjectPtr<UPoseWatch> WeakPoseWatchPtr;
+	TAttribute<FText> HighlightText;
+	bool bInEditingMode = false;
+
+	FText GetDisplayText() const
+	{
+		return WeakPoseWatchPtr.Get()->GetLabel();
+	}
+
+	FText GetTooltipText() const
+	{
+		if (!TreeItemPtr.Pin()->IsEnabled())
+		{
+			return LOCTEXT("PoseWatchDisabled", "This pose watch is disabled because it is not connected to the 'Output Pose' node");
+		}
+		return GetDisplayText();
+	}
+
+	FText GetTypeText() const
+	{
+		if (const UPoseWatch* PoseWatch = WeakPoseWatchPtr.Get())
+		{
+			return FText::FromName(PoseWatch->Node->GetFName());
+		}
+		return FText();
+	}
+
+	EVisibility GetTypeTextVisibility() const
+	{
+		return HighlightText.Get().IsEmpty() ? EVisibility::Collapsed : EVisibility::Visible;
+	}
+
+	const FSlateBrush* GetIcon() const
+	{
+		return FEditorStyle::Get().GetBrush(TEXT("ClassIcon.PoseAsset"));
+	}
+
+	const FSlateBrush* GetIconOverlay() const
+	{
+		return nullptr;
+	}
+
+	FText GetIconTooltip() const
+	{
+		return LOCTEXT("PoseWatch", "Pose Watch");
+	}
+
+	FSlateColor GetForegroundColor() const
+	{
+		if (auto BaseColor = FPoseWatchManagerCommonLabelData::GetForegroundColor(*TreeItemPtr.Pin()))
+		{
+			return BaseColor.GetValue();
+		}
+		return FSlateColor::UseForeground();
+	}
+
+	bool OnVerifyItemLabelChanged(const FText& InLabel, FText& OutErrorMessage)
+	{
+		return WeakPoseWatchPtr->ValidateLabelRename(InLabel, OutErrorMessage);
+	}
+
+	void OnLabelCommitted(const FText& InLabel, ETextCommit::Type InCommitInfo)
+	{
+		check(WeakPoseWatchPtr->SetLabel(InLabel));
+		WeakPoseWatchManager.Pin()->FullRefresh();
+		WeakPoseWatchManager.Pin()->SetKeyboardFocus();
+	}
+
+	void OnEnterEditingMode()
+	{
+		bInEditingMode = true;
+	}
+
+	void OnExitEditingMode()
+	{
+		bInEditingMode = false;
+	}
+};
+
+const EPoseWatchTreeItemType FPoseWatchManagerPoseWatchTreeItem::Type(EPoseWatchTreeItemType::PoseWatch);
+
+FPoseWatchManagerPoseWatchTreeItem::FPoseWatchManagerPoseWatchTreeItem(UPoseWatch* InPoseWatch)
+	: IPoseWatchManagerTreeItem(Type)
+	, ID(InPoseWatch)
+	, PoseWatch(InPoseWatch)
+{
+	check(InPoseWatch);
+}
+
+FObjectKey FPoseWatchManagerPoseWatchTreeItem::GetID() const
+{
+	return ID;
+}
+
+FString FPoseWatchManagerPoseWatchTreeItem::GetDisplayString() const
+{
+	return PoseWatch->GetLabel().ToString();
+}
+
+bool FPoseWatchManagerPoseWatchTreeItem::IsAssignedFolder() const
+{
+	return PoseWatch->IsAssignedFolder();
+}
+
+TSharedRef<SWidget> FPoseWatchManagerPoseWatchTreeItem::GenerateLabelWidget(IPoseWatchManager& PoseWatchManager, const STableRow<FPoseWatchManagerTreeItemPtr>& InRow)
+{
+	return SNew(SPoseWatchManagerPoseWatchTreeLabel, *this, PoseWatchManager, InRow);
+}
+
+bool FPoseWatchManagerPoseWatchTreeItem::GetVisibility() const
+{
+	return PoseWatch->GetIsVisible();
+}
+
+void FPoseWatchManagerPoseWatchTreeItem::SetIsVisible(const bool bVisible)
+{
+	PoseWatch->SetIsVisible(bVisible);
+}
+
+TSharedPtr<SWidget> FPoseWatchManagerPoseWatchTreeItem::CreateContextMenu()
+{
+	FMenuBuilder MenuBuilder(true, nullptr);
+
+	UPoseWatch* PoseWatchPtr = PoseWatch.Get();
+	MenuBuilder.BeginSection(FName(), LOCTEXT("PoseWatch", "Pose Watch"));
+	MenuBuilder.AddMenuEntry(
+		LOCTEXT("DeletePoseWatch", "Delete Pose Watch"),
+		LOCTEXT("DeleteFolderDescription", "Delete the selected pose watch"),
+		FSlateIcon(),
+		FUIAction(
+			FExecuteAction::CreateLambda([PoseWatchPtr]() {
+				PoseWatchPtr->OnRemoved();
+			})
+		)
+	);
+	MenuBuilder.EndSection();
+
+	return MenuBuilder.MakeWidget();
+}
+
+void FPoseWatchManagerPoseWatchTreeItem::OnRemoved()
+{
+	PoseWatch->OnRemoved();
+}
+
+bool FPoseWatchManagerPoseWatchTreeItem::IsEnabled() const
+{
+	return PoseWatch->GetIsEnabled();
+}
+
+#undef LOCTEXT_NAMESPACE
diff --git a/Engine/Source/Editor/Persona/Private/PoseWatchManagerPoseWatchTreeItem.h b/Engine/Source/Editor/Persona/Private/PoseWatchManagerPoseWatchTreeItem.h
new file mode 100644
--- /dev/null
+++ b/Engine/Source/Editor/Persona/Private/PoseWatchManagerPoseWatchTreeItem.h
@@ -0,0 +1,38 @@
+// Copyright Epic Games, Inc. All Rights Reserved.
+
+#pragma once
+
+#include "CoreMinimal.h"
+#include "Widgets/SWidget.h"
+#include "IPoseWatchManagerTreeItem.h"
+
+class UPoseWatch;
+
+struct FPoseWatchManagerPoseWatchTreeItem : IPoseWatchManagerTreeItem
+{
+public:
+	FPoseWatchManagerPoseWatchTreeItem(UPoseWatch* InPoseWatch);
+	
+	/** Constant identifier for this tree item */
+	const FObjectKey ID;
+	
+	/** The pose watch this item is describing */
+	TWeakObjectPtr<UPoseWatch> PoseWatch;
+
+	/** Static type identifier for this tree item class */
+	static const EPoseWatchTreeItemType Type;
+
+	bool IsValid() const override { return PoseWatch.IsValid();	}
+
+	/* Begin IPoseWatchManagerTreeItem Implementation */
+	virtual FObjectKey GetID() const override;
+	virtual FString GetDisplayString() const override;
+	virtual TSharedRef<SWidget> GenerateLabelWidget(IPoseWatchManager& Outliner, const STableRow<FPoseWatchManagerTreeItemPtr>& InRow) override;
+	virtual bool GetVisibility() const override;
+	virtual bool IsAssignedFolder() const;
+	virtual void SetIsVisible(const bool bVisible) override;
+	virtual TSharedPtr<SWidget> CreateContextMenu() override;
+	virtual void OnRemoved() override;
+	virtual bool IsEnabled() const override;
+	/* End IPoseWatchManagerTreeItem Implementation */
+};
diff --git a/Engine/Source/Editor/Persona/Private/PoseWatchManagerPublicTypes.h b/Engine/Source/Editor/Persona/Private/PoseWatchManagerPublicTypes.h
new file mode 100644
--- /dev/null
+++ b/Engine/Source/Editor/Persona/Private/PoseWatchManagerPublicTypes.h
@@ -0,0 +1,49 @@
+// Copyright Epic Games, Inc. All Rights Reserved.
+
+#pragma once
+
+#include "CoreMinimal.h"
+#include "Layout/Margin.h"
+
+/** Container for built in column types. Function-static so they are available without linking */
+struct FPoseWatchManagerBuiltInColumnTypes
+{
+	static const FName& Visibility()
+	{
+		static FName Visibility("Visibility");
+		return Visibility;
+	}
+
+	static const FName& Label()
+	{
+		static FName Label("Item Label");
+		return Label;
+	}
+
+	static FName& Color()
+	{
+		static FName Color("Color");
+		return Color;
+	}
+};
+
+enum class EPoseWatchManagerColumnVisibility : uint8
+{
+	Visible,
+	Invisible,
+};
+
+struct FPoseWatchManagerInitializationOptions
+{
+	/** The blueprint editor this pose watch manager is reflecting */
+	TWeakPtr<class FBlueprintEditor> BlueprintEditor;
+
+	FPoseWatchManagerInitializationOptions() {}
+};
+
+struct FPoseWatchManagerDefaultTreeItemMetrics
+{
+	static int32	RowHeight() { return 20; };
+	static int32	IconSize() { return 16; };
+	static FMargin	IconPadding() { return FMargin(0.f, 1.f, 6.f, 1.f); };
+};
\ No newline at end of file
diff --git a/Engine/Source/Editor/Persona/Private/PoseWatchManagerStandaloneTypes.cpp b/Engine/Source/Editor/Persona/Private/PoseWatchManagerStandaloneTypes.cpp
new file mode 100644
--- /dev/null
+++ b/Engine/Source/Editor/Persona/Private/PoseWatchManagerStandaloneTypes.cpp
@@ -0,0 +1,43 @@
+// Copyright Epic Games, Inc. All Rights Reserved.
+
+#include "PoseWatchManagerStandaloneTypes.h"
+
+#include "EditorActorFolders.h"
+#include "IPoseWatchManagerTreeItem.h"
+#include "IPoseWatchManager.h"
+#include "PoseWatchManagerDragDrop.h"
+#include "PoseWatchManagerDefaultMode.h"
+
+#define LOCTEXT_NAMESPACE "PoseWatchManagerStandaloneTypes"
+
+const FLinearColor FPoseWatchManagerCommonLabelData::DarkColor(0.15f, 0.15f, 0.15f);
+
+TOptional<FLinearColor> FPoseWatchManagerCommonLabelData::GetForegroundColor(const IPoseWatchManagerTreeItem& TreeItem) const
+{
+	if (!TreeItem.IsValid() || !TreeItem.IsEnabled())
+	{
+		return DarkColor;
+	}
+
+	// Darken items that aren't suitable targets for an active drag and drop action
+	if (FSlateApplication::Get().IsDragDropping())
+	{
+		TSharedPtr<FDragDropOperation> DragDropOp = FSlateApplication::Get().GetDragDroppingContent();
+
+		FPoseWatchManagerDragDropPayload DraggedObjects(*DragDropOp);
+		const auto Outliner = WeakPoseWatchManager.Pin();
+		if (Outliner->GetMode()->ParseDragDrop(DraggedObjects, *DragDropOp) && !Outliner->GetMode()->ValidateDrop(TreeItem, DraggedObjects).IsValid())
+		{
+			return DarkColor;
+		}
+	}
+
+	return TOptional<FLinearColor>();
+}
+
+bool FPoseWatchManagerCommonLabelData::CanExecuteRenameRequest(const IPoseWatchManagerTreeItem& Item) const
+{
+	return true;
+}
+
+#undef LOCTEXT_NAMESPACE
diff --git a/Engine/Source/Editor/Persona/Private/PoseWatchManagerStandaloneTypes.h b/Engine/Source/Editor/Persona/Private/PoseWatchManagerStandaloneTypes.h
new file mode 100644
--- /dev/null
+++ b/Engine/Source/Editor/Persona/Private/PoseWatchManagerStandaloneTypes.h
@@ -0,0 +1,46 @@
+// Copyright Epic Games, Inc. All Rights Reserved.
+
+#pragma once
+
+#include "CoreMinimal.h"
+#include "Misc/Paths.h"
+#include "UObject/ObjectKey.h"
+#include "Templates/MaxSizeof.h"
+#include "PoseWatchManagerFwd.h"
+
+typedef TSharedPtr<IPoseWatchManagerTreeItem> FPoseWatchManagerTreeItemPtr;
+typedef TSharedRef<IPoseWatchManagerTreeItem> FPoseWatchManagerTreeItemRef;
+
+typedef TMap<FObjectKey, FPoseWatchManagerTreeItemPtr> FPoseWatchManagerTreeItemMap;
+
+enum EPoseWatchTreeItemType
+{
+	PoseWatch,
+	Folder
+};
+
+struct PERSONA_API FPoseWatchManagerCommonLabelData
+{
+	TWeakPtr<IPoseWatchManager> WeakPoseWatchManager;
+	static const FLinearColor DarkColor;
+
+	TOptional<FLinearColor> GetForegroundColor(const IPoseWatchManagerTreeItem& TreeItem) const;
+
+	bool CanExecuteRenameRequest(const IPoseWatchManagerTreeItem& Item) const;
+};
+
+
+struct FPoseWatchManagerHierarchyChangedData
+{
+	enum
+	{
+		Added,
+		FullRefresh
+	} Type;
+
+	// This event may pass one of two kinds of data, depending on the type of event
+	TArray<FPoseWatchManagerTreeItemPtr> Items;
+
+	/** Actions to apply to items */
+	uint8 ItemActions = 0;
+};
diff --git a/Engine/Source/Editor/Persona/Private/SPoseWatchManager.cpp b/Engine/Source/Editor/Persona/Private/SPoseWatchManager.cpp
new file mode 100644
--- /dev/null
+++ b/Engine/Source/Editor/Persona/Private/SPoseWatchManager.cpp
@@ -0,0 +1,898 @@
+// Copyright Epic Games, Inc. All Rights Reserved.
+
+#include "SPoseWatchManager.h"
+
+#include "EdMode.h"
+#include "Editor.h"
+#include "Editor/UnrealEdEngine.h"
+#include "EditorModeManager.h"
+#include "EditorStyleSet.h"
+#include "Engine/GameViewportClient.h"
+#include "Engine/Selection.h"
+#include "EngineUtils.h"
+#include "Framework/Application/SlateApplication.h"
+#include "Framework/Commands/UIAction.h"
+#include "Framework/MultiBox/MultiBoxBuilder.h"
+#include "HAL/PlatformApplicationMisc.h"
+#include "IPoseWatchManagerColumn.h"
+#include "Kismet2/ComponentEditorUtils.h"
+#include "Layout/WidgetPath.h"
+#include "Modules/ModuleManager.h"
+#include "ScopedTransaction.h"
+#include "Textures/SlateIcon.h"
+#include "ToolMenus.h"
+#include "UnrealEdGlobals.h"
+#include "UObject/PackageReload.h"
+#include "Widgets/Images/SImage.h"
+#include "Widgets/Input/SButton.h"
+#include "Widgets/Input/SComboButton.h"
+#include "Widgets/Input/SSearchBox.h"
+#include "Widgets/Layout/SSeparator.h"
+#include "Widgets/SOverlay.h"
+#include "EditorFolderUtils.h"
+#include "IAnimationBlueprintEditor.h"
+#include "Animation/AnimBlueprint.h"
+#include "EdGraph/EdGraph.h"
+#include "GraphEditAction.h"
+#include "AnimationEditorUtils.h"
+#include "PoseWatchManagerDefaultHierarchy.h"
+#include "PoseWatchManagerDefaultMode.h"
+
+#include "PoseWatchManagerPoseWatchTreeItem.h"
+
+#include "PoseWatchManagerColumnVisibility.h"
+#include "PoseWatchManagerColumnColor.h"
+#include "PoseWatchManagerColumnLabel.h"
+
+#include "Framework/Notifications/NotificationManager.h"
+#include "Widgets/Notifications/SNotificationList.h"
+
+#define LOCTEXT_NAMESPACE "SPoseWatchManager"
+
+SPoseWatchManager::SPoseWatchManager()
+{}
+
+SPoseWatchManager::~SPoseWatchManager()
+{
+	SearchBoxFilter->OnChanged().RemoveAll(this);
+
+	AnimationEditorUtils::OnPoseWatchesChanged().RemoveAll(this);
+	FCoreUObjectDelegates::OnPackageReloaded.RemoveAll(this);
+}
+
+void SPoseWatchManager::Construct(const FArguments& InArgs, const FPoseWatchManagerInitializationOptions& InInitOptions)
+{
+	AnimationEditorUtils::OnPoseWatchesChanged().AddSP(this, &SPoseWatchManager::OnPoseWatchesChanged);
+
+	Mode = MakeShared<FPoseWatchManagerDefaultMode>(this);
+
+	FBlueprintEditor* BlueprintEditor = InInitOptions.BlueprintEditor.Pin().Get();
+	if (IAnimationBlueprintEditor* AnimBlueprintEditor = static_cast<IAnimationBlueprintEditor*>(BlueprintEditor))
+	{
+		AnimBlueprint = static_cast<UAnimBlueprint*>(AnimBlueprintEditor->GetBlueprintObj());
+	}
+	check(AnimBlueprint)
+
+	bProcessingFullRefresh = false;
+	bFullRefresh = true;
+	bNeedsRefresh = true;
+	bNeedsColumnRefresh = true;
+	bIsReentrant = false;
+	bSortDirty = true;
+	bSelectionDirty = true;
+	bPendingFocusNextFrame = false;
+
+	SortByColumn = FPoseWatchManagerBuiltInColumnTypes::Label();
+	SortMode = EColumnSortMode::Ascending;
+
+
+	//Setup the search box
+	{
+		auto Delegate = PoseWatchManager::TreeItemTextFilter::FItemToStringArray::CreateSP(this, &SPoseWatchManager::PopulateSearchStrings);
+		SearchBoxFilter = MakeShareable(new PoseWatchManager::TreeItemTextFilter(Delegate));
+	}
+
+	TSharedRef<SVerticalBox> VerticalBox = SNew(SVerticalBox);
+
+
+	SearchBoxFilter->OnChanged().AddSP(this, &SPoseWatchManager::FullRefresh);
+
+	HeaderRowWidget =
+		SNew(SHeaderRow)
+		.Visibility(EVisibility::Visible)
+		.CanSelectGeneratedColumn(true);
+
+	SetupColumns(*HeaderRowWidget);
+
+	ChildSlot
+		[
+			VerticalBox
+		];
+
+	auto Toolbar = SNew(SHorizontalBox);
+
+	Toolbar->AddSlot()
+		.VAlign(VAlign_Center)
+		[
+			SAssignNew(FilterTextBoxWidget, SSearchBox)
+			.Visibility(EVisibility::Visible)
+			.HintText(LOCTEXT("FilterSearch", "Search..."))
+			.ToolTipText(LOCTEXT("FilterSearchHint", "Type here to search"))
+			.OnTextChanged(this, &SPoseWatchManager::OnFilterTextChanged)
+		];
+
+
+	Toolbar->AddSlot()
+		.VAlign(VAlign_Center)
+		.AutoWidth()
+		.Padding(4.f, 0.f, 0.f, 0.f)
+		[
+			SNew(SButton)
+			.ButtonStyle(FAppStyle::Get(), "SimpleButton")
+			.ToolTipText(LOCTEXT("CreateFolderToolTip", "Create a new folder"))
+			.OnClicked(this, &SPoseWatchManager::OnCreateFolderClicked)
+			[
+				SNew(SImage)
+				.ColorAndOpacity(FSlateColor::UseForeground())
+				.Image(FAppStyle::Get().GetBrush("SceneOutliner.NewFolderIcon"))
+			]
+		];
+
+	VerticalBox->AddSlot()
+		.AutoHeight()
+		.Padding(8.0f, 8.0f, 8.0f, 4.0f)
+		[
+			Toolbar
+		];
+
+	VerticalBox->AddSlot()
+		.FillHeight(1.0)
+		[
+			SNew(SOverlay)
+			+ SOverlay::Slot()
+			.HAlign(HAlign_Center)
+			[
+				SNew(STextBlock)
+				.Visibility(this, &SPoseWatchManager::GetEmptyLabelVisibility)
+				.Text(LOCTEXT("EmptyLabel", "Empty"))
+				.ColorAndOpacity(FLinearColor(0.4f, 1.0f, 0.4f))
+			]
+
+		+ SOverlay::Slot()
+		[
+			SNew(SBorder).BorderImage(FAppStyle::Get().GetBrush("Brushes.Recessed"))
+		]
+
+		+ SOverlay::Slot()
+		[
+			SAssignNew(PoseWatchManagerTreeView, SPoseWatchManagerTreeView, StaticCastSharedRef<SPoseWatchManager>(AsShared()))
+
+			// Determined by the mode
+			.SelectionMode(this, &SPoseWatchManager::GetSelectionMode)
+
+			// Point the tree to our array of root-level items.  Whenever this changes, we'll call RequestTreeRefresh()
+			.TreeItemsSource(&RootTreeItems)
+
+			// Find out when the user selects something in the tree
+			.OnSelectionChanged(this, &SPoseWatchManager::OnManagerTreeSelectionChanged)
+
+			// Called when an item is scrolled into view
+			.OnItemScrolledIntoView(this, &SPoseWatchManager::OnManagerTreeItemScrolledIntoView)
+
+			// Called when an item is expanded or collapsed
+			.OnExpansionChanged(this, &SPoseWatchManager::OnItemExpansionChanged)
+
+			// Called to child items for any given parent item
+			.OnGetChildren(this, &SPoseWatchManager::OnGetChildrenForManagerTree)
+
+			// Generates the actual widget for a tree item
+			.OnGenerateRow(this, &SPoseWatchManager::OnGenerateRowForManagerTree)
+
+			// Use the level viewport context menu as the right click menu for tree items
+			.OnContextMenuOpening(this, &SPoseWatchManager::OnOpenContextMenu)
+
+			// Header for the tree
+			.HeaderRow(HeaderRowWidget)
+
+			// Make it easier to see hierarchies when there are a lot of items
+			.HighlightParentNodesForSelection(true)
+		]
+	];
+
+	// Don't allow tool-tips over the header
+	HeaderRowWidget->EnableToolTipForceField(true);
+
+	// Populate our data set
+	Populate();
+}
+
+void SPoseWatchManager::SetupColumns(SHeaderRow& HeaderRow)
+{
+	Columns.Empty(3);
+	Columns.Add(FPoseWatchManagerColumnVisibility::GetID(), MakeShared<FPoseWatchManagerColumnVisibility>(*this));
+	Columns.Add(FPoseWatchManagerColumnColor::GetID(), MakeShared<FPoseWatchManagerColumnColor>());
+	Columns.Add(FPoseWatchManagerItemLabelColumn::GetID(), MakeShared<FPoseWatchManagerItemLabelColumn>(*this));
+
+	HeaderRow.ClearColumns();
+
+	for (auto& Pair : Columns)
+	{
+		FName ID = Pair.Key;
+		TSharedPtr<IPoseWatchManagerColumn> Column = Pair.Value;
+
+		if (ensure(Column.IsValid()))
+		{
+			auto ColumnArgs = Column->ConstructHeaderRowColumn();
+
+			if (Column->SupportsSorting())
+			{
+				ColumnArgs
+					.SortMode(this, &SPoseWatchManager::GetColumnSortMode, ID)
+					.OnSort(this, &SPoseWatchManager::OnColumnSortModeChanged);
+			}
+
+			ColumnArgs.DefaultLabel(FText::FromName(ID));
+			ColumnArgs.ShouldGenerateWidget(true);
+
+			HeaderRow.AddColumn(ColumnArgs);
+			HeaderRowWidget->SetShowGeneratedColumn(ID);
+		}
+	}
+
+	bNeedsColumnRefresh = false;
+}
+
+void SPoseWatchManager::RefreshColumns()
+{
+	bNeedsColumnRefresh = true;
+}
+
+void SPoseWatchManager::OnItemAdded(const FObjectKey& ItemID, uint8 ActionMask)
+{
+	NewItemActions.Add(ItemID, ActionMask);
+}
+
+ESelectionMode::Type SPoseWatchManager::GetSelectionMode() const
+{
+	return ESelectionMode::Single;
+}
+
+void SPoseWatchManager::Refresh()
+{
+	bNeedsRefresh = true;
+}
+
+void SPoseWatchManager::FullRefresh()
+{
+	bDisableIntermediateSorting = true;
+	bFullRefresh = true;
+	Refresh();
+}
+
+void SPoseWatchManager::RefreshSelection()
+{
+	bSelectionDirty = true;
+}
+
+void SPoseWatchManager::Populate()
+{
+	// Block events while we clear out the list
+	TGuardValue<bool> ReentrantGuard(bIsReentrant, true);
+
+	bool bMadeAnySignificantChanges = false;
+	if (bFullRefresh)
+	{
+		// Clear the selection here - RepopulateEntireTree will reconstruct it.
+		PoseWatchManagerTreeView->ClearSelection();
+
+		RepopulateEntireTree();
+
+		bMadeAnySignificantChanges = true;
+		bFullRefresh = false;
+	}
+
+	int32 Index = 0;
+	while (Index < PendingOperations.Num())
+	{
+		auto& PendingOp = PendingOperations[Index];
+		switch (PendingOp.Type)
+		{
+		case PoseWatchManager::FPendingTreeOperation::Added:
+			bMadeAnySignificantChanges = AddItemToTree(PendingOp.Item) || bMadeAnySignificantChanges;
+			break;
+
+		default:
+			check(false);
+			break;
+		}
+
+		++Index;
+	}
+
+	PendingOperations.RemoveAt(0, Index);
+
+	// Check if we need to sort because we are finished with the populating operations
+	bool bFinalSort = false;
+	if (PendingOperations.Num() == 0)
+	{
+		// When done processing a FullRefresh Scroll to First item in selection as it may have been
+		// scrolled out of view by the Refresh
+		if (bProcessingFullRefresh)
+		{
+			FPoseWatchManagerTreeItemPtr ItemToScroll = GetSelection();
+			if (ItemToScroll)
+			{
+				ScrollItemIntoView(ItemToScroll);
+			}
+		}
+
+		bProcessingFullRefresh = false;
+		// We're fully refreshed now.
+		NewItemActions.Empty();
+		bNeedsRefresh = false;
+		if (bDisableIntermediateSorting)
+		{
+			bDisableIntermediateSorting = false;
+			bFinalSort = true;
+		}
+	}
+
+	// If we are allowing intermediate sorts and met the conditions, or this is the final sort after all ops are complete
+	if ((bMadeAnySignificantChanges && !bDisableIntermediateSorting) || bFinalSort)
+	{
+		RequestSort();
+	}
+}
+
+
+void SPoseWatchManager::EmptyTreeItems()
+{
+	PendingOperations.Empty();
+	TreeItemMap.Reset();
+	PendingTreeItemMap.Empty();
+	RootTreeItems.Empty();
+}
+
+void SPoseWatchManager::AddPendingItem(FPoseWatchManagerTreeItemPtr Item)
+{
+	if (Item)
+	{
+		PendingTreeItemMap.Add(Item->GetID(), Item);
+		PendingOperations.Emplace(PoseWatchManager::FPendingTreeOperation::Added, Item.ToSharedRef());
+	}
+}
+
+void SPoseWatchManager::AddPendingItemAndChildren(FPoseWatchManagerTreeItemPtr Item)
+{
+	if (Item)
+	{
+		AddPendingItem(Item);
+		Refresh();
+	}
+}
+
+void SPoseWatchManager::RepopulateEntireTree()
+{
+	EmptyTreeItems();
+
+	// Rebuild the hierarchy
+	Mode->Rebuild();
+
+	// Create all the items which match the filters, parent-child relationships are handled when each item is actually added to the tree
+
+	TArray<FPoseWatchManagerTreeItemPtr> Items;
+	Mode->Hierarchy->CreateItems(Items);
+
+	for (FPoseWatchManagerTreeItemPtr& Item : Items)
+	{
+		if (Item && Item->IsValid())
+		{
+			AddPendingItem(Item);
+			if (FPoseWatchManagerFolderTreeItem* FolderItem = Item->CastTo<FPoseWatchManagerFolderTreeItem>())
+			{
+				PoseWatchManagerTreeView->SetItemExpansion(Item, FolderItem->PoseWatchFolder->GetIsExpanded());
+			}
+		}
+	}
+
+	bProcessingFullRefresh = PendingOperations.Num() > 0;
+
+	Refresh();
+}
+
+FPoseWatchManagerTreeItemPtr SPoseWatchManager::GetTreeItem(FObjectKey ItemID, bool bIncludePending)
+{
+	FPoseWatchManagerTreeItemPtr Result = TreeItemMap.FindRef(ItemID);
+	if (bIncludePending && !Result.IsValid())
+	{
+		Result = PendingTreeItemMap.FindRef(ItemID);
+	}
+	return Result;
+}
+
+FPoseWatchManagerTreeItemPtr SPoseWatchManager::EnsureParentForItem(FPoseWatchManagerTreeItemRef Item)
+{
+	FPoseWatchManagerTreeItemPtr Parent = Mode->Hierarchy->FindParent(*Item, TreeItemMap);
+	if (Parent.IsValid())
+	{
+		return Parent;
+	}
+		
+	// Try to find the parent in the pending items
+	Parent = Mode->Hierarchy->FindParent(*Item, PendingTreeItemMap);
+	if (Parent.IsValid())
+	{
+		AddUnfilteredItemToTree(Parent.ToSharedRef());
+		return Parent;
+	}
+
+	return nullptr;
+}
+
+bool SPoseWatchManager::AddItemToTree(FPoseWatchManagerTreeItemRef Item)
+{
+	const auto ItemID = Item->GetID();
+
+	PendingTreeItemMap.Remove(ItemID);
+
+	// If a tree item already exists that represents the same data or if the item represents invalid data, bail
+	if (TreeItemMap.Find(ItemID) || !Item->IsValid())
+	{
+		return false;
+	}
+
+	AddUnfilteredItemToTree(Item);
+
+	// Check if we need to do anything with this new item
+	if (uint8* ActionMask = NewItemActions.Find(ItemID))
+	{
+		if (*ActionMask & PoseWatchManager::ENewItemAction::Select)
+		{
+			PoseWatchManagerTreeView->ClearSelection();
+			PoseWatchManagerTreeView->SetItemSelection(Item, true);
+		}
+
+		if (*ActionMask & PoseWatchManager::ENewItemAction::Rename)
+		{
+			PendingRenameItem = Item;
+		}
+
+		if (*ActionMask & (PoseWatchManager::ENewItemAction::ScrollIntoView | PoseWatchManager::ENewItemAction::Rename))
+		{
+			ScrollItemIntoView(Item);
+		}
+	}
+
+	return true;
+}
+
+void SPoseWatchManager::AddUnfilteredItemToTree(FPoseWatchManagerTreeItemRef Item)
+{
+	FPoseWatchManagerTreeItemPtr Parent = EnsureParentForItem(Item);
+
+	const FObjectKey ItemID = Item->GetID();
+	check(!TreeItemMap.Contains(ItemID));
+
+	TreeItemMap.Add(ItemID, Item);
+
+	if (Parent.IsValid())
+	{
+		Parent->AddChild(Item);
+	}
+	else
+	{
+		RootTreeItems.Add(Item);
+	}
+}
+
+void SPoseWatchManager::PopulateSearchStrings(const IPoseWatchManagerTreeItem& Item, TArray< FString >& OutSearchStrings) const
+{
+	for (const auto& Pair : Columns)
+	{
+		Pair.Value->PopulateSearchStrings(Item, OutSearchStrings);
+	}
+}
+
+TSharedPtr<SWidget> SPoseWatchManager::OnOpenContextMenu()
+{
+	return Mode->CreateContextMenu();
+}
+
+void SPoseWatchManager::Rename_Execute()
+{
+	FPoseWatchManagerTreeItemPtr ItemToRename = GetSelection();
+	if (ItemToRename && ItemToRename.IsValid())
+	{
+		PendingRenameItem = ItemToRename->AsShared();
+		ScrollItemIntoView(ItemToRename);
+	}
+}
+
+void SPoseWatchManager::SetColumnVisibility(FName ColumnId, bool bIsVisible)
+{
+	if (Columns.Contains(ColumnId))
+	{
+		HeaderRowWidget->SetShowGeneratedColumn(ColumnId, bIsVisible);
+	}
+}
+
+void SPoseWatchManager::SetSelection(const TFunctionRef<bool(IPoseWatchManagerTreeItem&)> Selector)
+{
+	TArray<FPoseWatchManagerTreeItemPtr> ItemsToAdd;
+	for (const auto& Pair : TreeItemMap)
+	{
+		FPoseWatchManagerTreeItemPtr ItemPtr = Pair.Value;
+		if (IPoseWatchManagerTreeItem* Item = ItemPtr.Get())
+		{
+			if (Selector(*Item))
+			{
+				ItemsToAdd.Add(ItemPtr);
+			}
+		}
+	}
+
+	SetItemSelection(ItemsToAdd, true);
+}
+
+void SPoseWatchManager::SetItemSelection(const TArray<FPoseWatchManagerTreeItemPtr>& InItems, bool bSelected, ESelectInfo::Type SelectInfo)
+{
+	PoseWatchManagerTreeView->ClearSelection();
+	PoseWatchManagerTreeView->SetItemSelection(InItems, bSelected, SelectInfo);
+}
+
+void SPoseWatchManager::SetItemSelection(const FPoseWatchManagerTreeItemPtr& InItem, bool bSelected, ESelectInfo::Type SelectInfo)
+{
+	PoseWatchManagerTreeView->ClearSelection();
+	PoseWatchManagerTreeView->SetItemSelection(InItem, bSelected, SelectInfo);
+}
+
+void SPoseWatchManager::ClearSelection()
+{
+	if (!bIsReentrant)
+	{
+		PoseWatchManagerTreeView->ClearSelection();
+	}
+}
+
+void SPoseWatchManager::OnPoseWatchesChanged(UAnimBlueprint* InAnimBlueprint, UEdGraphNode* InNode)
+{
+	for (UPoseWatch* SomePoseWatch : AnimBlueprint->PoseWatches)
+	{
+		if (SomePoseWatch->Node == InNode)
+		{
+			FPoseWatchManagerHierarchyChangedData EventData;
+			EventData.Type = FPoseWatchManagerHierarchyChangedData::Added;
+			EventData.Items.Add(Mode->PoseWatchManager->CreateItemFor<FPoseWatchManagerPoseWatchTreeItem>(SomePoseWatch));
+			EventData.ItemActions = PoseWatchManager::ENewItemAction::Select | PoseWatchManager::ENewItemAction::Rename;
+			OnHierarchyChangedEvent(EventData);
+			return;
+		}
+	}
+
+	// A pose watch for this node was not found, it must've been removed
+	FullRefresh();
+}
+
+FReply SPoseWatchManager::OnCreateFolderClicked()
+{
+	CreateFolder();
+	return FReply::Handled();
+}
+
+void SPoseWatchManager::CreateFolder()
+{
+	check(AnimBlueprint);
+	UPoseWatchFolder* NewPoseWatchFolder = NewObject<UPoseWatchFolder>(AnimBlueprint);
+	NewPoseWatchFolder->SetUniqueDefaultLabel();
+	AnimBlueprint->PoseWatchFolders.Add(NewPoseWatchFolder);
+
+	FPoseWatchManagerHierarchyChangedData EventData;
+	EventData.Type = FPoseWatchManagerHierarchyChangedData::Added;
+	EventData.Items.Add(Mode->PoseWatchManager->CreateItemFor<FPoseWatchManagerFolderTreeItem>(NewPoseWatchFolder));
+	EventData.ItemActions = PoseWatchManager::ENewItemAction::Select | PoseWatchManager::ENewItemAction::Rename;
+	OnHierarchyChangedEvent(EventData);
+}
+
+void SPoseWatchManager::ScrollItemIntoView(const FPoseWatchManagerTreeItemPtr& Item)
+{
+	auto Parent = Item->GetParent();
+	while (Parent.IsValid())
+	{
+		PoseWatchManagerTreeView->SetItemExpansion(Parent->AsShared(), true);
+		Parent = Parent->GetParent();
+	}
+	PoseWatchManagerTreeView->RequestScrollIntoView(Item);
+}
+
+TSharedRef< ITableRow > SPoseWatchManager::OnGenerateRowForManagerTree(FPoseWatchManagerTreeItemPtr Item, const TSharedRef< STableViewBase >& OwnerTable)
+{
+	return SNew(SPoseWatchManagerTreeRow, PoseWatchManagerTreeView.ToSharedRef(), SharedThis(this)).Item(Item);
+}
+
+void SPoseWatchManager::OnGetChildrenForManagerTree(FPoseWatchManagerTreeItemPtr InParent, TArray< FPoseWatchManagerTreeItemPtr >& OutChildren)
+{
+	for (auto& WeakChild : InParent->GetChildren())
+	{
+		auto Child = WeakChild.Pin();
+		// Should never have bogus entries in this list
+		check(Child.IsValid());
+		OutChildren.Add(Child);
+	}
+
+	// If the item needs it's children sorting, do that now
+	if (OutChildren.Num())
+	{
+		// Sort the children we returned
+		SortItems(OutChildren);
+
+		// Empty out the children and repopulate them in the correct order
+		InParent->Children.Empty();
+		for (auto& Child : OutChildren)
+		{
+			InParent->Children.Emplace(Child);
+		}
+	}
+}
+
+void SPoseWatchManager::OnManagerTreeSelectionChanged(FPoseWatchManagerTreeItemPtr TreeItem, ESelectInfo::Type SelectInfo)
+{
+	if (SelectInfo == ESelectInfo::Direct)
+	{
+		return;
+	}
+
+	if (!bIsReentrant)
+	{
+		TGuardValue<bool> ReentrantGuard(bIsReentrant, true);
+		//Mode->OnItemSelectionChanged(TreeItem, SelectInfo, FPoseWatchManagerItemSelection(*PoseWatchManagerTreeView));
+
+		OnItemSelectionChanged.Broadcast(TreeItem, SelectInfo);
+	}
+}
+
+void SPoseWatchManager::OnManagerTreeItemScrolledIntoView(FPoseWatchManagerTreeItemPtr TreeItem, const TSharedPtr<ITableRow>& Widget)
+{
+	if (TreeItem == PendingRenameItem.Pin())
+	{
+		PendingRenameItem = nullptr;
+		TreeItem->RenameRequestEvent.ExecuteIfBound();
+	}
+}
+
+void SPoseWatchManager::OnItemExpansionChanged(FPoseWatchManagerTreeItemPtr TreeItem, bool bIsExpanded) const
+{
+	check(TreeItem.Get()->IsA<FPoseWatchManagerFolderTreeItem>());
+	FPoseWatchManagerFolderTreeItem* FolderTreeItem = TreeItem.Get()->CastTo<FPoseWatchManagerFolderTreeItem>();
+	FolderTreeItem->PoseWatchFolder->SetIsExpanded(bIsExpanded);
+}
+
+void SPoseWatchManager::OnHierarchyChangedEvent(FPoseWatchManagerHierarchyChangedData Event)
+{
+	if (Event.Type == FPoseWatchManagerHierarchyChangedData::Added)
+	{
+		for (const auto& TreeItemPtr : Event.Items)
+		{
+			if (TreeItemPtr.IsValid() && !TreeItemMap.Find(TreeItemPtr->GetID()))
+			{
+				AddPendingItemAndChildren(TreeItemPtr);
+				if (Event.ItemActions)
+				{
+					NewItemActions.Add(TreeItemPtr->GetID(), Event.ItemActions);
+				}
+			}
+		}
+		FullRefresh();
+	}
+	else if (Event.Type == FPoseWatchManagerHierarchyChangedData::FullRefresh)
+	{
+		FullRefresh();
+	}
+}
+
+void SPoseWatchManager::OnFilterTextChanged(const FText& InFilterText)
+{
+	SearchBoxFilter->SetRawFilterText(InFilterText);
+	FilterTextBoxWidget->SetError(SearchBoxFilter->GetFilterErrorText());
+}
+
+EVisibility SPoseWatchManager::GetFilterStatusVisibility() const
+{
+	return IsTextFilterActive() ? EVisibility::Visible : EVisibility::Collapsed;
+}
+
+EVisibility SPoseWatchManager::GetEmptyLabelVisibility() const
+{
+	return (IsTextFilterActive() || RootTreeItems.Num() > 0) ? EVisibility::Collapsed : EVisibility::Visible;
+}
+
+bool SPoseWatchManager::IsTextFilterActive() const
+{
+	return FilterTextBoxWidget->GetText().ToString().Len() > 0;
+}
+
+const FSlateBrush* SPoseWatchManager::GetFilterButtonGlyph() const
+{
+	if (IsTextFilterActive())
+	{
+		return FEditorStyle::GetBrush(TEXT("SceneOutliner.FilterCancel"));
+	}
+	else
+	{
+		return FEditorStyle::GetBrush(TEXT("SceneOutliner.FilterSearch"));
+	}
+}
+
+FString SPoseWatchManager::GetFilterButtonToolTip() const
+{
+	return IsTextFilterActive() ? LOCTEXT("ClearSearchFilter", "Clear search filter").ToString() : LOCTEXT("StartSearching", "Search").ToString();
+
+}
+
+TAttribute<FText> SPoseWatchManager::GetFilterHighlightText() const
+{
+	return TAttribute<FText>::Create(TAttribute<FText>::FGetter::CreateStatic([](TWeakPtr<PoseWatchManager::TreeItemTextFilter> Filter) {
+		auto FilterPtr = Filter.Pin();
+		return FilterPtr.IsValid() ? FilterPtr->GetRawFilterText() : FText();
+		}, TWeakPtr<PoseWatchManager::TreeItemTextFilter>(SearchBoxFilter)));
+}
+
+void SPoseWatchManager::SetKeyboardFocus()
+{
+	if (SupportsKeyboardFocus())
+	{
+		FWidgetPath PoseWatchManagerTreeViewWidgetPath;
+		// NOTE: Careful, GeneratePathToWidget can be reentrant in that it can call visibility delegates and such
+		FSlateApplication::Get().GeneratePathToWidgetUnchecked(PoseWatchManagerTreeView.ToSharedRef(), PoseWatchManagerTreeViewWidgetPath);
+		FSlateApplication::Get().SetKeyboardFocus(PoseWatchManagerTreeViewWidgetPath, EFocusCause::SetDirectly);
+	}
+}
+
+bool SPoseWatchManager::SupportsKeyboardFocus() const
+{
+	return true;
+}
+
+FReply SPoseWatchManager::OnKeyDown(const FGeometry& MyGeometry, const FKeyEvent& InKeyEvent)
+{
+	return Mode->OnKeyDown(InKeyEvent);
+}
+
+void SPoseWatchManager::Tick(const FGeometry& AllottedGeometry, const double InCurrentTime, const float InDeltaTime)
+{
+	if (bPendingFocusNextFrame && FilterTextBoxWidget->GetVisibility() == EVisibility::Visible)
+	{
+		FWidgetPath WidgetToFocusPath;
+		FSlateApplication::Get().GeneratePathToWidgetUnchecked(FilterTextBoxWidget.ToSharedRef(), WidgetToFocusPath);
+		FSlateApplication::Get().SetKeyboardFocus(WidgetToFocusPath, EFocusCause::SetDirectly);
+		bPendingFocusNextFrame = false;
+	}
+
+	if (bNeedsColumnRefresh)
+	{
+		SetupColumns(*HeaderRowWidget);
+	}
+
+	if (bNeedsRefresh)
+	{
+		if (!bIsReentrant)
+		{
+			Populate();
+		}
+	}
+
+	if (bSortDirty)
+	{
+		SortItems(RootTreeItems);
+		PoseWatchManagerTreeView->RequestTreeRefresh();
+		bSortDirty = false;
+	}
+
+	if (bSelectionDirty)
+	{
+		//Mode->SynchronizeSelection();
+		bSelectionDirty = false;
+	}
+}
+
+EColumnSortMode::Type SPoseWatchManager::GetColumnSortMode(const FName ColumnId) const
+{
+	if (SortByColumn == ColumnId)
+	{
+		auto Column = Columns.FindRef(ColumnId);
+		if (Column.IsValid() && Column->SupportsSorting())
+		{
+			return SortMode;
+		}
+	}
+
+	return EColumnSortMode::None;
+}
+
+void SPoseWatchManager::OnColumnSortModeChanged(const EColumnSortPriority::Type SortPriority, const FName& ColumnId, const EColumnSortMode::Type InSortMode)
+{
+	auto Column = Columns.FindRef(ColumnId);
+	if (!Column.IsValid() || !Column->SupportsSorting())
+	{
+		return;
+	}
+
+	SortByColumn = ColumnId;
+	SortMode = InSortMode;
+
+	RequestSort();
+}
+
+void SPoseWatchManager::RequestSort()
+{
+	bSortDirty = true;
+}
+
+void SPoseWatchManager::SortItems(TArray<FPoseWatchManagerTreeItemPtr>& Items) const
+{
+	auto Column = Columns.FindRef(SortByColumn);
+	if (Column.IsValid())
+	{
+		Column->SortItems(Items, SortMode);
+	}
+}
+
+void SPoseWatchManager::SetItemExpansionRecursive(FPoseWatchManagerTreeItemPtr Model, bool bInExpansionState)
+{
+	if (Model.IsValid())
+	{
+		PoseWatchManagerTreeView->SetItemExpansion(Model, bInExpansionState);
+		for (auto& Child : Model->Children)
+		{
+			if (Child.IsValid())
+			{
+				SetItemExpansionRecursive(Child.Pin(), bInExpansionState);
+			}
+		}
+	}
+}
+
+TSharedPtr<FDragDropOperation> SPoseWatchManager::CreateDragDropOperation(const TArray<FPoseWatchManagerTreeItemPtr>& InTreeItems) const
+{
+	return Mode->CreateDragDropOperation(InTreeItems);
+}
+
+/** Parse a drag drop operation into a payload */
+bool SPoseWatchManager::ParseDragDrop(FPoseWatchManagerDragDropPayload& OutPayload, const FDragDropOperation& Operation) const
+{
+	return Mode->ParseDragDrop(OutPayload, Operation);
+}
+
+/** Validate a drag drop operation on a drop target */
+FPoseWatchManagerDragValidationInfo SPoseWatchManager::ValidateDrop(const IPoseWatchManagerTreeItem& DropTarget, const FPoseWatchManagerDragDropPayload& Payload) const
+{
+	return Mode->ValidateDrop(DropTarget, Payload);
+}
+
+/** Called when a payload is dropped onto a target */
+void SPoseWatchManager::OnDropPayload(IPoseWatchManagerTreeItem& DropTarget, const FPoseWatchManagerDragDropPayload& Payload, const FPoseWatchManagerDragValidationInfo& ValidationInfo) const
+{
+	return Mode->OnDrop(DropTarget, Payload, ValidationInfo);
+}
+
+/** Called when a payload is dragged over an item */
+FReply SPoseWatchManager::OnDragOverItem(const FDragDropEvent& Event, const IPoseWatchManagerTreeItem& Item) const
+{
+	return Mode->OnDragOverItem(Event, Item);
+}
+
+FPoseWatchManagerTreeItemPtr SPoseWatchManager::FindParent(const IPoseWatchManagerTreeItem& InItem) const
+{
+	FPoseWatchManagerTreeItemPtr Parent = Mode->Hierarchy->FindParent(InItem, TreeItemMap);
+	if (!Parent.IsValid())
+	{
+		Parent = Mode->Hierarchy->FindParent(InItem, PendingTreeItemMap);
+	}
+	return Parent;
+}
+
+uint32 SPoseWatchManager::GetTypeSortPriority(const IPoseWatchManagerTreeItem& Item) const
+{
+	return Mode->GetTypeSortPriority(Item);
+}
+
+#undef LOCTEXT_NAMESPACE
diff --git a/Engine/Source/Editor/Persona/Private/SPoseWatchManager.h b/Engine/Source/Editor/Persona/Private/SPoseWatchManager.h
new file mode 100644
--- /dev/null
+++ b/Engine/Source/Editor/Persona/Private/SPoseWatchManager.h
@@ -0,0 +1,414 @@
+// Copyright Epic Games, Inc. All Rights Reserved.
+#pragma once
+
+#include "CoreMinimal.h"
+#include "UObject/Object.h"
+#include "EditorUndoClient.h"
+#include "Framework/Commands/UICommandList.h"
+#include "Framework/Views/ITypedTableView.h"
+#include "Input/Reply.h"
+#include "Layout/Visibility.h"
+#include "Misc/Attribute.h"
+#include "Misc/TextFilter.h"
+#include "SlateFwd.h"
+#include "Styling/SlateColor.h"
+#include "Widgets/DeclarativeSyntaxSupport.h"
+#include "Widgets/SWidget.h"
+#include "Widgets/Views/STableRow.h"
+#include "Widgets/Views/STableViewBase.h"
+#include "Delegates/DelegateCombinations.h"
+#include "Framework/Views/ITypedTableView.h"
+#include "Types/SlateEnums.h"
+
+#include "BlueprintEditor.h"
+#include "Engine/PoseWatch.h"
+#include "PoseWatchManagerFolderTreeItem.h"
+
+
+#include "IPoseWatchManager.h"
+#include "PoseWatchManagerFwd.h"
+
+#include "SPoseWatchManagerTreeView.h"
+#include "PoseWatchManagerPublicTypes.h"
+#include "PoseWatchManagerStandaloneTypes.h"
+
+#include "PoseWatchManagerDragDrop.h"
+
+
+class FMenuBuilder;
+class UToolMenu;
+class IPoseWatchManagerColumn;
+class FPoseWatchManagerDefaultHierarchy;
+class SComboButton;
+
+template<typename ItemType> class STreeView;
+
+
+namespace PoseWatchManager
+{
+	DECLARE_EVENT_OneParam(SPoseWatchManager, FTreeItemPtrEvent, FPoseWatchManagerTreeItemPtr);
+
+	DECLARE_EVENT_TwoParams(SPoseWatchManager, FOnItemSelectionChanged, FPoseWatchManagerTreeItemPtr, ESelectInfo::Type);
+
+	typedef TTextFilter<const IPoseWatchManagerTreeItem&> TreeItemTextFilter;
+
+	/** Structure that defines an operation that should be applied to the tree */
+	struct FPendingTreeOperation
+	{
+		enum EType { Added, Removed, Moved };
+		FPendingTreeOperation(EType InType, TSharedRef<IPoseWatchManagerTreeItem> InItem) : Type(InType), Item(InItem) { }
+
+		/** The type of operation that is to be applied */
+		EType Type;
+
+		/** The tree item to which this operation relates */
+		FPoseWatchManagerTreeItemRef Item;
+	};
+
+	/** Set of actions to apply to new tree items */
+	namespace ENewItemAction
+	{
+		enum Type
+		{
+			/** Do nothing when it is created */
+			None = 0,
+			/** Select the item when it is created */
+			Select = 1 << 0,
+			/** Scroll the item into view when it is created */
+			ScrollIntoView = 1 << 1,
+			/** Interactively rename the item when it is created (implies the above) */
+			Rename = 1 << 2,
+		};
+	}
+}
+
+class SPoseWatchManager : public IPoseWatchManager, public FGCObject
+{
+public:
+	SLATE_BEGIN_ARGS(SPoseWatchManager)
+	{}
+	SLATE_END_ARGS()
+
+	void Construct(const FArguments& InArgs, const FPoseWatchManagerInitializationOptions& InitOptions);
+
+	SPoseWatchManager();
+	virtual ~SPoseWatchManager();
+
+	// SWidget interface 
+	virtual void Tick(const FGeometry& AllottedGeometry, const double InCurrentTime, const float InDeltaTime) override;
+	virtual bool SupportsKeyboardFocus() const override;
+	virtual FReply OnKeyDown(const FGeometry& MyGeometry, const FKeyEvent& InKeyEvent) override;
+
+	/** Sends a requests to refresh the next chance it gets */
+	virtual void Refresh() override;
+
+	void RefreshSelection();
+
+	/** Get a const reference to the actual tree hierarchy */
+	virtual const STreeView<FPoseWatchManagerTreeItemPtr>& GetTree() const override
+	{
+		return *PoseWatchManagerTreeView;
+	}
+
+	virtual const TSharedPtr<SPoseWatchManagerTreeView>& GetTreeView() const
+	{
+		return PoseWatchManagerTreeView;
+	}
+
+	/** @return Returns a string to use for highlighting results in the list */
+	virtual TAttribute<FText> GetFilterHighlightText() const;
+
+	/** Set the keyboard focus to the manager */
+	virtual void SetKeyboardFocus() override;
+
+	void SetColumnVisibility(FName ColumnId, bool bIsVisible);
+
+	/** Returns the current sort mode of the specified column */
+	virtual EColumnSortMode::Type GetColumnSortMode(const FName ColumnId) const;
+
+	/** Request that the tree be sorted at a convenient time */
+	virtual void RequestSort();
+
+	/** Requests a full refresh, which will clear the entire tree and rebuild it from scratch */
+	virtual void FullRefresh() override;
+
+public:
+	/** Event to react to a user double click on a item */
+	PoseWatchManager::FTreeItemPtrEvent& GetDoubleClickEvent() { return OnDoubleClickOnTreeEvent; }
+
+	PoseWatchManager::FOnItemSelectionChanged& GetOnItemSelectionChanged() { return OnItemSelectionChanged; }
+
+	/** Set the item selection of the manager based on a selector function. Any items which return true will be added */
+	virtual void SetSelection(const TFunctionRef<bool(IPoseWatchManagerTreeItem&)> Selector) override;
+
+	/** Set the selection status of a set of items in the scene manager */
+	void SetItemSelection(const TArray<FPoseWatchManagerTreeItemPtr>& InItems, bool bSelected, ESelectInfo::Type SelectInfo = ESelectInfo::Direct);
+
+	/** Set the selection status of a single item in the scene manager */
+	void SetItemSelection(const FPoseWatchManagerTreeItemPtr& InItem, bool bSelected, ESelectInfo::Type SelectInfo = ESelectInfo::Direct);
+
+
+	virtual FPoseWatchManagerTreeItemPtr GetSelection() const { return PoseWatchManagerTreeView->GetSelectedItems().Num() == 1 ? PoseWatchManagerTreeView->GetSelectedItems()[0] : nullptr; }
+
+	FPoseWatchManagerTreeItemPtr FindParent(const IPoseWatchManagerTreeItem& InItem) const;
+
+	/** Deselect all selected items */
+	void ClearSelection();
+
+	/** Sets the next item to rename */
+	void SetPendingRenameItem(const FPoseWatchManagerTreeItemPtr& InItem) { PendingRenameItem = InItem; Refresh(); }
+
+	/** Retrieve an IPoseWatchManagerTreeItem by its ID if it exists in the tree */
+	FPoseWatchManagerTreeItemPtr GetTreeItem(FObjectKey, bool bIncludePending = false);
+
+	/** Create a drag drop operation */
+	TSharedPtr<FDragDropOperation> CreateDragDropOperation(const TArray<FPoseWatchManagerTreeItemPtr>& InTreeItems) const;
+
+	/** Parse a drag drop operation into a payload */
+	bool ParseDragDrop(FPoseWatchManagerDragDropPayload& OutPayload, const FDragDropOperation& Operation) const;
+
+	/** Validate a drag drop operation on a drop target */
+	FPoseWatchManagerDragValidationInfo ValidateDrop(const IPoseWatchManagerTreeItem& DropTarget, const FPoseWatchManagerDragDropPayload& Payload) const;
+
+	/** Called when a payload is dropped onto a target */
+	void OnDropPayload(IPoseWatchManagerTreeItem& DropTarget, const FPoseWatchManagerDragDropPayload& Payload, const FPoseWatchManagerDragValidationInfo& ValidationInfo) const;
+
+	/** Called when a payload is dragged over an item */
+	FReply OnDragOverItem(const FDragDropEvent& Event, const IPoseWatchManagerTreeItem& Item) const;
+
+	virtual uint32 GetTypeSortPriority(const IPoseWatchManagerTreeItem& Item) const;
+
+
+	/** Used to test if manager related selection changes have already been handled */
+	bool GetIsReentrant() const { return bIsReentrant; }
+
+private:
+	/** Empty all the tree item containers maintained by this manager */
+	void EmptyTreeItems();
+
+	/** Apply incremental changes to, or a complete repopulation of the tree  */
+	void Populate();
+
+	/** Repopulates the entire tree */
+	void RepopulateEntireTree();
+
+	/** Adds a single new item to the pending map and creates an add operation for it */
+	void AddPendingItem(FPoseWatchManagerTreeItemPtr Item);
+
+	/** Adds a new item and all of its children to the pending items. */
+	void AddPendingItemAndChildren(FPoseWatchManagerTreeItemPtr Item);
+
+	/** Attempts to add a pending item to the current tree. Will add any parents if required. */
+	bool AddItemToTree(FPoseWatchManagerTreeItemRef InItem);
+
+	/** Add an item to the tree, even if it doesn't match the filter terms. Used to add parent's that would otherwise be filtered out */
+	void AddUnfilteredItemToTree(FPoseWatchManagerTreeItemRef Item);
+
+	/** Ensure that the specified item's parent is added to the tree, if applicable */
+	FPoseWatchManagerTreeItemPtr EnsureParentForItem(FPoseWatchManagerTreeItemRef Item);
+
+public:
+	// Test the filters using stack-allocated data to prevent unnecessary heap allocations
+	template <typename TreeItemType, typename TreeItemData>
+	FPoseWatchManagerTreeItemPtr CreateItemFor(const TreeItemData& Data, bool bForce = false)
+	{
+		const TreeItemType Temporary(Data);
+
+		if (bForce || SearchBoxFilter->PassesFilter(Temporary))
+		{
+			FPoseWatchManagerTreeItemPtr Result = MakeShareable(new TreeItemType(Data));
+			return Result;
+		}
+
+		return nullptr;
+	}
+
+	/** Instruct the manager to perform an action on the specified item when it is created */
+	void OnItemAdded(const FObjectKey& ItemID, uint8 ActionMask);
+
+	void OnPoseWatchesChanged(UAnimBlueprint* InAnimBlueprint, UEdGraphNode* InNode);
+
+	virtual void Rename_Execute() override;
+
+	/** Get the columns to be displayed in this manager */
+	const TMap<FName, TSharedPtr<IPoseWatchManagerColumn>>& GetColumns() const
+	{
+		return Columns;
+	}
+
+	/** @return	Returns true if the text filter is currently active */
+	bool IsTextFilterActive() const;
+
+	bool PassesTextFilter(const FPoseWatchManagerTreeItemPtr& Item) const
+	{
+		return SearchBoxFilter->PassesFilter(*Item);
+	}
+
+	bool HasSelectorFocus(FPoseWatchManagerTreeItemPtr Item) const
+	{
+		return PoseWatchManagerTreeView->Private_HasSelectorFocus(Item);
+	}
+
+private:
+	/** Map of columns that are shown on this manager. */
+	TMap<FName, TSharedPtr<IPoseWatchManagerColumn>> Columns;
+
+	TArray<FName> HiddenColumnsList;
+
+	/** Set up the columns required for this manager */
+	void SetupColumns(SHeaderRow& HeaderRow);
+
+	/** Refresh the scene manager for when a colum was added or removed */
+	void RefreshColumns();
+
+	/** Populates OutSearchStrings with the strings associated with TreeItem that should be used in searching */
+	void PopulateSearchStrings(const IPoseWatchManagerTreeItem& TreeItem, OUT TArray< FString >& OutSearchStrings) const;
+
+public:
+	/** Scroll the specified item into view */
+	void ScrollItemIntoView(const FPoseWatchManagerTreeItemPtr& Item);
+
+private:
+	/** Called by STreeView to generate a table row for the specified item */
+	TSharedRef< ITableRow > OnGenerateRowForManagerTree(FPoseWatchManagerTreeItemPtr Item, const TSharedRef< STableViewBase >& OwnerTable);
+
+	/** Called by STreeView to get child items for the specified parent item */
+	void OnGetChildrenForManagerTree(FPoseWatchManagerTreeItemPtr InParent, TArray< FPoseWatchManagerTreeItemPtr >& OutChildren);
+
+	/** Called by STreeView when the tree's selection has changed */
+	void OnManagerTreeSelectionChanged(FPoseWatchManagerTreeItemPtr TreeItem, ESelectInfo::Type SelectInfo);
+
+	/** Called by STreeView when an item is scrolled into view */
+	void OnManagerTreeItemScrolledIntoView(FPoseWatchManagerTreeItemPtr TreeItem, const TSharedPtr<ITableRow>& Widget);
+
+	/** Called when an item in the tree has been collapsed or expanded */
+	void OnItemExpansionChanged(FPoseWatchManagerTreeItemPtr TreeItem, bool bIsExpanded) const;
+
+private:
+	/** Event required to keep the pose watch manager up to date */
+	void OnHierarchyChangedEvent(FPoseWatchManagerHierarchyChangedData Event);
+
+private:
+	/** Called by the editable text control when the filter text is changed by the user */
+	void OnFilterTextChanged(const FText& InFilterText);
+
+	/** Called by the filter button to get the image to display in the button */
+	const FSlateBrush* GetFilterButtonGlyph() const;
+
+	/** @return	The filter button tool-tip text */
+	FString GetFilterButtonToolTip() const;
+
+	/** @return	Returns whether the filter status line should be drawn */
+	EVisibility GetFilterStatusVisibility() const;
+
+	/**	Returns the current visibility of the Empty label */
+	EVisibility GetEmptyLabelVisibility() const;
+
+	/** Returns the selection mode*/
+	ESelectionMode::Type GetSelectionMode() const;
+
+public:
+	TSharedPtr<SWidget> OnOpenContextMenu();
+
+private:
+	/** Called when the user has clicked the button to add a new folder */
+	FReply OnCreateFolderClicked();
+
+	void CreateFolder();
+
+private:
+	/** Context menu opening delegate provided by the client */
+	FOnContextMenuOpening OnContextMenuOpening;
+
+	/** List of pending operations to be applied to the tree */
+	TArray<PoseWatchManager::FPendingTreeOperation> PendingOperations;
+
+	/** Map of actions to apply to new tree items */
+	TMap<FObjectKey, uint8> NewItemActions;
+
+	/** Our tree view */
+	TSharedPtr< SPoseWatchManagerTreeView > PoseWatchManagerTreeView;
+
+	/** A map of all items we have in the tree */
+	FPoseWatchManagerTreeItemMap TreeItemMap;
+
+	/** Pending tree items that are yet to be added the tree */
+	FPoseWatchManagerTreeItemMap PendingTreeItemMap;
+
+	/** Folders pending selection */
+	TArray<FName> PendingFoldersSelect;
+
+	/** Root level tree items */
+	TArray<FPoseWatchManagerTreeItemPtr> RootTreeItems;
+
+	/** The button that displays view options */
+	TSharedPtr<SComboButton> ViewOptionsComboButton;
+
+private:
+	/** true if the manager  needs to be repopulated at the next appropriate opportunity */
+	uint8 bNeedsRefresh : 1;
+
+	/** Processing a full refresh until pending items are processed */
+	uint8 bProcessingFullRefresh : 1;
+
+	/** true if should do a full refresh */
+	uint8 bFullRefresh : 1;
+
+	/** true if should refresh selection */
+	uint8 bSelectionDirty : 1;
+
+	uint8 bNeedsColumnRefresh : 1;
+
+	/** Reentrancy guard */
+	bool bIsReentrant;
+
+	/* Widget containing the filtering text box */
+	TSharedPtr< SSearchBox > FilterTextBoxWidget;
+
+	/** The header row of the manager */
+	TSharedPtr< SHeaderRow > HeaderRowWidget;
+
+	/** The TextFilter attached to the SearchBox widget */
+	TSharedPtr< PoseWatchManager::TreeItemTextFilter > SearchBoxFilter;
+
+	/** True if the search box will take keyboard focus next frame */
+	bool bPendingFocusNextFrame;
+
+	/** The tree item that is currently pending a rename */
+	TWeakPtr<IPoseWatchManagerTreeItem> PendingRenameItem;
+
+	TMap<FName, const FSlateBrush*> CachedIcons;
+
+	TUniquePtr<FPoseWatchManagerDefaultHierarchy> Hierarchy;
+
+	PoseWatchManager::FTreeItemPtrEvent OnDoubleClickOnTreeEvent;
+
+	PoseWatchManager::FOnItemSelectionChanged OnItemSelectionChanged;
+
+private:
+	virtual void AddReferencedObjects(FReferenceCollector& Collector) override {};
+	virtual FString GetReferencerName() const override
+	{
+		return TEXT("PoseWatchManager::SPoseWatchManager");
+	}
+
+	bool bDisableIntermediateSorting;
+
+	/** true if currently needs to be sorted */
+	bool bSortDirty;
+
+	/** Specify which column to sort with */
+	FName SortByColumn;
+
+	/** Currently selected sorting mode */
+	EColumnSortMode::Type SortMode;
+
+	/** Handles column sorting mode change */
+	void OnColumnSortModeChanged(const EColumnSortPriority::Type SortPriority, const FName& ColumnId, const EColumnSortMode::Type InSortMode);
+
+	/** Sort the specified array of items based on the current sort column */
+	void SortItems(TArray<FPoseWatchManagerTreeItemPtr>& Items) const;
+
+	/** Handler for recursively expanding/collapsing items */
+	void SetItemExpansionRecursive(FPoseWatchManagerTreeItemPtr Model, bool bInExpansionState);
+};
diff --git a/Engine/Source/Editor/Persona/Private/SPoseWatchManagerTreeView.cpp b/Engine/Source/Editor/Persona/Private/SPoseWatchManagerTreeView.cpp
new file mode 100644
--- /dev/null
+++ b/Engine/Source/Editor/Persona/Private/SPoseWatchManagerTreeView.cpp
@@ -0,0 +1,380 @@
+// Copyright Epic Games, Inc. All Rights Reserved.
+
+#include "SPoseWatchManagerTreeView.h"
+#include "Framework/Application/SlateApplication.h"
+#include "EditorStyleSet.h"
+#include "IPoseWatchManagerColumn.h"
+#include "PoseWatchManagerPublicTypes.h"
+#include "DragAndDrop/DecoratedDragDropOp.h"
+#include "PoseWatchManagerDragDrop.h"
+#include "PoseWatchManagerStandaloneTypes.h"
+#include "SPoseWatchManager.h"
+
+#include "PoseWatchManagerFolderTreeItem.h"
+
+#define LOCTEXT_NAMESPACE "SPoseWatchManager"
+
+static void UpdateOperationDecorator(const FDragDropEvent& Event, const FPoseWatchManagerDragValidationInfo& ValidationInfo)
+{
+	const FSlateBrush* Icon = ValidationInfo.IsValid() ? FEditorStyle::GetBrush(TEXT("Graph.ConnectorFeedback.OK")) : FEditorStyle::GetBrush(TEXT("Graph.ConnectorFeedback.Error"));
+
+	FDragDropOperation* Operation = Event.GetOperation().Get();
+
+	if (Operation && Operation->IsOfType<FPoseWatchManagerDragDropOp>())
+	{
+		FPoseWatchManagerDragDropOp* PoseWatchDecoratedOp = static_cast<FPoseWatchManagerDragDropOp*>(Operation);
+		PoseWatchDecoratedOp->SetTooltip(ValidationInfo.ValidationText, Icon);
+	}
+}
+
+static void ResetOperationDecorator(const FDragDropEvent& Event)
+{
+	FDragDropOperation* Operation = Event.GetOperation().Get();
+	if (Operation)
+	{
+		if (Operation->IsOfType<FPoseWatchManagerDragDropOp>())
+		{
+			static_cast<FPoseWatchManagerDragDropOp*>(Operation)->ResetTooltip();
+		}
+	}
+}
+
+static FReply HandleOnDragDetected(const FGeometry& MyGeometry, const FPointerEvent& MouseEvent, TWeakPtr<SPoseWatchManagerTreeView> Table)
+{
+	TSharedPtr<SPoseWatchManagerTreeView> TablePtr = Table.Pin();
+	if (TablePtr.IsValid() && MouseEvent.IsMouseButtonDown(EKeys::LeftMouseButton))
+	{
+		TSharedPtr<FDragDropOperation> Operation = TablePtr->GetPoseWatchManagerPtr().Pin()->CreateDragDropOperation(TablePtr->GetSelectedItems());
+
+		if (Operation.IsValid())
+		{
+			return FReply::Handled().BeginDragDrop(Operation.ToSharedRef());
+		}
+	}
+
+	return FReply::Unhandled();
+}
+
+FReply HandleDrop(TSharedPtr<SPoseWatchManager> PoseWatchManagerPtr, const FDragDropEvent& DragDropEvent, IPoseWatchManagerTreeItem& DropTarget, FPoseWatchManagerDragValidationInfo& ValidationInfo, bool bApplyDrop = false)
+{
+	if (!PoseWatchManagerPtr.IsValid())
+	{
+		return FReply::Unhandled();
+	}
+
+	FPoseWatchManagerDragDropPayload DraggedObjects(*DragDropEvent.GetOperation());
+	// Validate now to make sure we don't doing anything we shouldn't
+	if (!PoseWatchManagerPtr->ParseDragDrop(DraggedObjects, *DragDropEvent.GetOperation()))
+	{
+		return FReply::Unhandled();
+	}
+
+	ValidationInfo = PoseWatchManagerPtr->ValidateDrop(StaticCast<IPoseWatchManagerTreeItem&>(DropTarget), DraggedObjects);
+
+	if (!ValidationInfo.IsValid())
+	{
+		// Return handled here to stop anything else trying to handle it - the operation is invalid as far as we're concerned
+		return FReply::Handled();
+	}
+
+	if (bApplyDrop)
+	{
+		PoseWatchManagerPtr->OnDropPayload(DropTarget, DraggedObjects, ValidationInfo);
+	}
+
+	return FReply::Handled();
+}
+
+FReply HandleDropFromWeak(TWeakPtr<SPoseWatchManager> PoseWatchManagerWeak, const FDragDropEvent& DragDropEvent, IPoseWatchManagerTreeItem& DropTarget, FPoseWatchManagerDragValidationInfo& ValidationInfo, bool bApplyDrop = false)
+{
+	return HandleDrop(PoseWatchManagerWeak.Pin(), DragDropEvent, DropTarget, ValidationInfo, bApplyDrop);
+}
+
+void SPoseWatchManagerTreeView::Construct(const SPoseWatchManagerTreeView::FArguments& InArgs, TSharedRef<SPoseWatchManager> Owner)
+{
+	PoseWatchManagerWeak = Owner;
+	STreeView::Construct(InArgs);
+}
+
+void SPoseWatchManagerTreeView::FlashHighlightOnItem(FPoseWatchManagerTreeItemPtr FlashHighlightOnItem)
+{
+	TSharedPtr< SPoseWatchManagerTreeRow > RowWidget = StaticCastSharedPtr< SPoseWatchManagerTreeRow >(WidgetGenerator.GetWidgetForItem(FlashHighlightOnItem));
+	if (RowWidget.IsValid())
+	{
+		RowWidget->FlashHighlight();
+	}
+}
+
+FReply SPoseWatchManagerTreeView::OnDragOver(const FGeometry& MyGeometry, const FDragDropEvent& DragDropEvent)
+{
+	FPoseWatchManagerDragValidationInfo ValidationInfo = FPoseWatchManagerDragValidationInfo::Invalid();
+
+	FPoseWatchManagerFolderTreeItem DropTarget(nullptr);
+
+	FReply Reply = HandleDropFromWeak(PoseWatchManagerWeak, DragDropEvent, DropTarget, ValidationInfo);
+	if (Reply.IsEventHandled())
+	{
+		UpdateOperationDecorator(DragDropEvent, ValidationInfo);
+	}
+
+	return Reply;
+}
+
+void SPoseWatchManagerTreeView::OnDragLeave(const FDragDropEvent& DragDropEvent)
+{
+	ResetOperationDecorator(DragDropEvent);
+}
+
+FReply SPoseWatchManagerTreeView::OnDrop(const FGeometry& MyGeometry, const FDragDropEvent& DragDropEvent)
+{
+	FPoseWatchManagerDragValidationInfo ValidationInfo = FPoseWatchManagerDragValidationInfo::Invalid();
+
+	// Dropping to folder nullptr will assign the payload to the root of the tree
+	FPoseWatchManagerFolderTreeItem DropTarget(nullptr);
+
+	FReply Reply = HandleDropFromWeak(PoseWatchManagerWeak, DragDropEvent, DropTarget, ValidationInfo, true);
+	//PoseWatchManagerPtr.Get()->FullRefresh();
+	PoseWatchManagerWeak.Pin()->FullRefresh();
+
+	return Reply;
+}
+
+FReply SPoseWatchManagerTreeRow::OnDrop(const FGeometry& MyGeometry, const FDragDropEvent& DragDropEvent)
+{
+	auto ItemPtr = Item.Pin();
+	auto PoseWatchManagerPtr = PoseWatchManagerWeak.Pin();
+	if (ItemPtr.IsValid() && PoseWatchManagerPtr.IsValid())
+	{
+		FPoseWatchManagerDragValidationInfo ValidationInfo = FPoseWatchManagerDragValidationInfo::Invalid();
+		FReply Reply = HandleDrop(PoseWatchManagerPtr, DragDropEvent, *ItemPtr, ValidationInfo, true);
+		PoseWatchManagerPtr.Get()->FullRefresh();
+		return Reply;
+	}
+
+	return FReply::Unhandled();
+}
+
+void SPoseWatchManagerTreeRow::OnDragEnter(const FGeometry& MyGeometry, const FDragDropEvent& DragDropEvent)
+{
+	FPoseWatchManagerTreeItemPtr ItemPtr = Item.Pin();
+	TSharedPtr<SPoseWatchManager> PoseWatchManagerPtr = PoseWatchManagerWeak.Pin();
+	if (ItemPtr.IsValid() && PoseWatchManagerPtr.IsValid())
+	{
+		FPoseWatchManagerDragValidationInfo ValidationInfo = FPoseWatchManagerDragValidationInfo::Invalid();
+
+		FReply Reply = HandleDrop(PoseWatchManagerPtr, DragDropEvent, *ItemPtr, ValidationInfo, false);
+		if (Reply.IsEventHandled()) 
+		{
+			UpdateOperationDecorator(DragDropEvent, ValidationInfo);
+		}
+	}
+}
+
+void SPoseWatchManagerTreeRow::OnDragLeave(const FDragDropEvent& DragDropEvent)
+{
+	ResetOperationDecorator(DragDropEvent);
+}
+
+FReply SPoseWatchManagerTreeRow::OnDragOver(const FGeometry& MyGeometry, const FDragDropEvent& DragDropEvent)
+{
+	TSharedPtr<SPoseWatchManager> PoseWatchManagerPtr = PoseWatchManagerWeak.Pin();
+	if (SPoseWatchManager* PoseWatchManager = PoseWatchManagerPtr.Get())
+	{
+		if (const FPoseWatchManagerTreeItemPtr ItemPtr = Item.Pin())
+		{
+			return PoseWatchManager->OnDragOverItem(DragDropEvent, *ItemPtr.Get());
+		}
+		return FReply::Unhandled();
+	}
+
+	return FReply::Handled();
+}
+
+FReply SPoseWatchManagerTreeRow::OnMouseButtonDown(const FGeometry& MyGeometry, const FPointerEvent& MouseEvent)
+{
+	FPoseWatchManagerTreeItemPtr ItemPtr = Item.Pin();
+	if (ItemPtr.IsValid())
+	{
+		if (MouseEvent.GetEffectingButton() == EKeys::LeftMouseButton)
+		{
+			FReply Reply = SMultiColumnTableRow<FPoseWatchManagerTreeItemPtr>::OnMouseButtonDown(MyGeometry, MouseEvent);
+			return Reply.DetectDrag(SharedThis(this), EKeys::LeftMouseButton);
+		}
+	}
+
+	return FReply::Handled();
+}
+
+FReply SPoseWatchManagerTreeRow::OnMouseButtonUp(const FGeometry& MyGeometry, const FPointerEvent& MouseEvent)
+{
+	FPoseWatchManagerTreeItemPtr ItemPtr = Item.Pin();
+	// We don't to change the selection when it is a left click since this was handle in the on mouse down
+	if (ItemPtr.IsValid())
+	{
+		return SMultiColumnTableRow<FPoseWatchManagerTreeItemPtr>::OnMouseButtonUp(MyGeometry, MouseEvent);
+	}
+
+	return FReply::Handled();
+}
+
+TSharedRef<SWidget> SPoseWatchManagerTreeRow::GenerateWidgetForColumn(const FName& ColumnName)
+{
+	FPoseWatchManagerTreeItemPtr ItemPtr = Item.Pin();
+	if (!ItemPtr.IsValid())
+	{
+		return SNullWidget::NullWidget;
+	}
+
+	// Create the widget for this item
+	TSharedRef<SWidget> NewItemWidget = SNullWidget::NullWidget;
+
+	check(PoseWatchManagerWeak.IsValid());
+	auto Column = PoseWatchManagerWeak.Pin()->GetColumns().FindRef(ColumnName);
+	if (Column.IsValid())
+	{
+		NewItemWidget = Column->ConstructRowWidget(ItemPtr.ToSharedRef(), *this);
+	}
+
+	if (ColumnName == FPoseWatchManagerBuiltInColumnTypes::Label())
+	{
+		// The first column gets the tree expansion arrow for this row
+		return SNew(SBox)
+			.MinDesiredHeight(FPoseWatchManagerDefaultTreeItemMetrics::RowHeight())
+			[
+				SNew(SHorizontalBox)
+
+				+ SHorizontalBox::Slot()
+				.AutoWidth()
+				.Padding(6, 0, 0, 0)
+				[
+					SNew(SExpanderArrow, SharedThis(this)).IndentAmount(12)
+				]
+
+				+ SHorizontalBox::Slot()
+				.FillWidth(1.0f)
+				[
+					NewItemWidget
+				]
+			];
+	}
+	else
+	{
+		// No expansion arrow needed
+		return NewItemWidget;
+	}
+}
+
+void SPoseWatchManagerTreeRow::Construct(const FArguments& InArgs, const TSharedRef<SPoseWatchManagerTreeView>& PoseWatchManagerTreeView, TSharedRef<SPoseWatchManager> PoseWatchManager)
+{
+	Item = InArgs._Item->AsShared();
+	PoseWatchManagerWeak = PoseWatchManager;
+	LastHighlightInteractionTime = 0.0;
+
+	auto Args = FSuperRowType::FArguments()
+		.Style(&FEditorStyle::Get().GetWidgetStyle<FTableRowStyle>("SceneOutliner.TableViewRow"));
+
+
+	Args.OnDragDetected_Static(HandleOnDragDetected, TWeakPtr<SPoseWatchManagerTreeView>(PoseWatchManagerTreeView));
+
+	SMultiColumnTableRow<FPoseWatchManagerTreeItemPtr>::Construct(Args, PoseWatchManagerTreeView);
+}
+
+const float SPoseWatchManagerTreeRow::HighlightRectLeftOffset = 0.0f;
+const float SPoseWatchManagerTreeRow::HighlightRectRightOffset = 0.0f;
+const float SPoseWatchManagerTreeRow::HighlightTargetSpringConstant = 25.0f;
+const float SPoseWatchManagerTreeRow::HighlightTargetEffectDuration = 0.5f;
+const float SPoseWatchManagerTreeRow::HighlightTargetOpacity = 0.8f;
+const float SPoseWatchManagerTreeRow::LabelChangedAnimOffsetPercent = 0.2f;
+
+void SPoseWatchManagerTreeRow::FlashHighlight()
+{
+	LastHighlightInteractionTime = FSlateApplication::Get().GetCurrentTime();
+}
+
+void SPoseWatchManagerTreeRow::Tick(const FGeometry& AllottedGeometry, const double InCurrentTime, const float InDeltaTime)
+{
+	// Call parent implementation.
+	SCompoundWidget::Tick(AllottedGeometry, InCurrentTime, InDeltaTime);
+
+	// We'll draw with the 'focused' look if we're either focused or we have a context menu summoned
+	const bool bShouldAppearFocused = HasKeyboardFocus();
+
+	// Update highlight 'target' effect
+	{
+		const float HighlightLeftX = HighlightRectLeftOffset;
+		const float HighlightRightX = HighlightRectRightOffset + AllottedGeometry.GetLocalSize().X;
+
+		HighlightTargetLeftSpring.SetTarget(HighlightLeftX);
+		HighlightTargetRightSpring.SetTarget(HighlightRightX);
+
+		float TimeSinceHighlightInteraction = (float)(InCurrentTime - LastHighlightInteractionTime);
+		if (TimeSinceHighlightInteraction <= HighlightTargetEffectDuration || bShouldAppearFocused)
+		{
+			HighlightTargetLeftSpring.Tick(InDeltaTime);
+			HighlightTargetRightSpring.Tick(InDeltaTime);
+		}
+	}
+}
+
+int32 SPoseWatchManagerTreeRow::OnPaint(const FPaintArgs& Args, const FGeometry& AllottedGeometry, const FSlateRect& MyCullingRect, FSlateWindowElementList& OutDrawElements, int32 LayerId, const FWidgetStyle& InWidgetStyle, bool bParentEnabled) const
+{
+	int32 StartLayer = SMultiColumnTableRow::OnPaint(Args, AllottedGeometry, MyCullingRect, OutDrawElements, LayerId, InWidgetStyle, bParentEnabled);
+
+	const int32 TextLayer = 1;
+
+	// See if a disabled effect should be used
+	bool bEnabled = ShouldBeEnabled(bParentEnabled);
+	ESlateDrawEffect DrawEffects = (bEnabled) ? ESlateDrawEffect::None : ESlateDrawEffect::DisabledEffect;
+
+	const double CurrentTime = FSlateApplication::Get().GetCurrentTime();
+
+	// We'll draw with the 'focused' look if we're either focused or we have a context menu summoned
+	const bool bShouldAppearFocused = HasKeyboardFocus();
+
+	// Draw highlight targeting effect
+	const float TimeSinceHighlightInteraction = (float)(CurrentTime - LastHighlightInteractionTime);
+	if (TimeSinceHighlightInteraction <= HighlightTargetEffectDuration)
+	{
+
+		// Compute animation progress
+		float EffectAlpha = FMath::Clamp(TimeSinceHighlightInteraction / HighlightTargetEffectDuration, 0.0f, 1.0f);
+		EffectAlpha = 1.0f - EffectAlpha * EffectAlpha;  // Inverse square falloff (looks nicer!)
+
+		// Apply extra opacity falloff when dehighlighting
+		float EffectOpacity = EffectAlpha;
+
+		// Figure out a universally visible highlight color.
+		FLinearColor HighlightTargetColorAndOpacity = ((FLinearColor::White - GetColorAndOpacity()) * 0.5f + FLinearColor(+0.4f, +0.1f, -0.2f)) * InWidgetStyle.GetColorAndOpacityTint();
+		HighlightTargetColorAndOpacity.A = HighlightTargetOpacity * EffectOpacity * 255.0f;
+
+		// Compute the bounds offset of the highlight target from where the highlight target spring
+		// extents currently lie.  This is used to "grow" or "shrink" the highlight as needed.
+		const float LabelChangedAnimOffset = LabelChangedAnimOffsetPercent * AllottedGeometry.GetLocalSize().Y;
+
+		// Choose an offset amount depending on whether we're highlighting, or clearing highlight
+		const float EffectOffset = EffectAlpha * LabelChangedAnimOffset;
+
+		const float HighlightLeftX = HighlightTargetLeftSpring.GetPosition() - EffectOffset;
+		const float HighlightRightX = HighlightTargetRightSpring.GetPosition() + EffectOffset;
+		const float HighlightTopY = 0.0f - LabelChangedAnimOffset;
+		const float HighlightBottomY = AllottedGeometry.GetLocalSize().Y + EffectOffset;
+
+		const FVector2D DrawPosition = FVector2D(HighlightLeftX, HighlightTopY);
+		const FVector2D DrawSize = FVector2D(HighlightRightX - HighlightLeftX, HighlightBottomY - HighlightTopY);
+
+		const FSlateBrush* StyleInfo = FEditorStyle::GetBrush("SceneOutliner.ChangedItemHighlight");
+
+		// NOTE: We rely on scissor clipping for the highlight rectangle
+		FSlateDrawElement::MakeBox(
+			OutDrawElements,
+			LayerId + TextLayer,
+			AllottedGeometry.ToPaintGeometry(DrawPosition, DrawSize),	// Position, Size, Scale
+			StyleInfo,													// Style
+			DrawEffects,												// Effects to use
+			HighlightTargetColorAndOpacity);							// Color
+	}
+
+	return LayerId + TextLayer;
+}
+
+#undef LOCTEXT_NAMESPACE
diff --git a/Engine/Source/Editor/Persona/Private/SPoseWatchManagerTreeView.h b/Engine/Source/Editor/Persona/Private/SPoseWatchManagerTreeView.h
new file mode 100644
--- /dev/null
+++ b/Engine/Source/Editor/Persona/Private/SPoseWatchManagerTreeView.h
@@ -0,0 +1,108 @@
+// Copyright Epic Games, Inc. All Rights Reserved.
+#pragma once
+
+#include "CoreMinimal.h"
+#include "PoseWatchManagerFwd.h"
+#include "PoseWatchManagerStandaloneTypes.h"
+#include "Input/Reply.h"
+#include "Widgets/SWidget.h"
+#include "Widgets/DeclarativeSyntaxSupport.h"
+#include "Widgets/Views/STableRow.h"
+#include "Widgets/Views/STreeView.h"
+#include "Animation/SlateSprings.h"
+
+class SPoseWatchManager;
+
+class SPoseWatchManagerTreeView : public STreeView< FPoseWatchManagerTreeItemPtr >
+{
+public:
+	void Construct(const FArguments& InArgs, TSharedRef<SPoseWatchManager> Owner);
+
+	void FlashHighlightOnItem(FPoseWatchManagerTreeItemPtr FlashHighlightOnItem);
+
+	const TWeakPtr<SPoseWatchManager>& GetPoseWatchManagerPtr() { return PoseWatchManagerWeak; }
+
+	
+
+protected:
+	virtual FReply OnDragOver(const FGeometry& MyGeometry, const FDragDropEvent& DragDropEvent) override;
+
+	virtual void OnDragLeave(const FDragDropEvent& DragDropEvent) override;
+
+	virtual FReply OnDrop(const FGeometry& MyGeometry, const FDragDropEvent& DragDropEvent) override;
+
+	TWeakPtr<SPoseWatchManager> PoseWatchManagerWeak;
+};
+
+class SPoseWatchManagerTreeRow
+	: public SMultiColumnTableRow< FPoseWatchManagerTreeItemPtr >
+{
+public:
+	SLATE_BEGIN_ARGS(SPoseWatchManagerTreeRow) {}
+		SLATE_ARGUMENT(FPoseWatchManagerTreeItemPtr, Item)
+	SLATE_END_ARGS()
+
+	/** Construct function for this widget */
+	void Construct(const FArguments& InArgs, const TSharedRef<SPoseWatchManagerTreeView>& PoseWatchManagerTreeView, TSharedRef<SPoseWatchManager> PoseWatchManager);
+
+	/** Overridden from SMultiColumnTableRow.  Generates a widget for this column of the tree row. */
+	virtual TSharedRef<SWidget> GenerateWidgetForColumn(const FName& ColumnName) override;
+
+	void FlashHighlight();
+
+protected:
+	virtual FReply OnMouseButtonDown(const FGeometry& MyGeometry, const FPointerEvent& MouseEvent) override;
+
+	virtual FReply OnMouseButtonUp(const FGeometry& MyGeometry, const FPointerEvent& MouseEvent) override;
+
+	virtual void OnDragEnter(const FGeometry& MyGeometry, const FDragDropEvent& DragDropEvent) override;
+
+	virtual void OnDragLeave(const FDragDropEvent& DragDropEvent) override;
+
+	virtual FReply OnDragOver(const FGeometry& MyGeometry, const FDragDropEvent& DragDropEvent) override;
+
+	virtual FReply OnDrop(const FGeometry& MyGeometry, const FDragDropEvent& DragDropEvent) override;
+
+private:
+
+	/** Weak reference to the outliner widget that owns our list */
+	TWeakPtr< SPoseWatchManager > PoseWatchManagerWeak;
+
+	/** The item associated with this row of data */
+	TWeakPtr<IPoseWatchManagerTreeItem> Item;
+
+protected:
+
+	virtual void Tick(const FGeometry& AllottedGeometry, const double InCurrentTime, const float InDeltaTime) override;
+
+	virtual int32 OnPaint(const FPaintArgs& Args, const FGeometry& AllottedGeometry, const FSlateRect& MyCullingRect, FSlateWindowElementList& OutDrawElements, int32 LayerId, const FWidgetStyle& InWidgetStyle, bool bParentEnabled) const override;
+
+private:
+	/** How many pixels to extend the highlight rectangle's left side horizontally */
+	static const float HighlightRectLeftOffset;
+
+	/** How many pixels to extend the highlight rectangle's right side horizontally */
+	static const float HighlightRectRightOffset;
+
+	/** How quickly the highlight 'targeting' rectangle will slide around.  Larger is faster. */
+	static const float HighlightTargetSpringConstant;
+
+	/** Duration of animation highlight target effects */
+	static const float HighlightTargetEffectDuration;
+
+	/** Opacity of the highlight target effect overlay */
+	static const float HighlightTargetOpacity;
+
+	/** How large the highlight target effect will be when highlighting, as a scalar percentage of font height */
+	static const float LabelChangedAnimOffsetPercent;
+
+	/** Highlight "targeting" visual effect left position */
+	FFloatSpring1D HighlightTargetLeftSpring;
+
+	/** Highlight "targeting" visual effect right position */
+	FFloatSpring1D HighlightTargetRightSpring;
+
+	/** Last time that the user had a major interaction with the highlight */
+	double LastHighlightInteractionTime;
+};
+
diff --git a/Engine/Source/Editor/Persona/Private/TabSpawners.cpp b/Engine/Source/Editor/Persona/Private/TabSpawners.cpp
--- a/Engine/Source/Editor/Persona/Private/TabSpawners.cpp
+++ b/Engine/Source/Editor/Persona/Private/TabSpawners.cpp
@@ -18,6 +18,7 @@
 #include "SAnimCurveViewer.h"
 #include "SAnimationSequenceBrowser.h"
 #include "SAnimationEditorViewport.h"
+#include "SPoseWatchManager.h"
 #include "SRetargetManager.h"
 #include "SKismetInspector.h"
 #include "Widgets/Input/SButton.h"
@@ -25,6 +26,8 @@
 #include "IPersonaPreviewScene.h"
 #include "PreviewSceneCustomizations.h"
 #include "Engine/PreviewMeshCollection.h"
+#include "PoseWatchManagerPublicTypes.h"
+#include "PoseWatchManagerDefaultMode.h"
 
 #define LOCTEXT_NAMESPACE "PersonaModes"
 
@@ -63,6 +66,8 @@ const FName FPersonaTabs::SkeletonSlotGroupNamesID("SkeletonSlotGroupNames");
 const FName FPersonaTabs::BlendProfileManagerID("BlendProfileManager");
 const FName FPersonaTabs::AnimMontageSectionsID("AnimMontageSections");
 
+const FName FPersonaTabs::PoseWatchManagerID("PoseWatchManager");
+
 const FName FPersonaTabs::AdvancedPreviewSceneSettingsID("AdvancedPreviewTab");
 
 const FName FPersonaTabs::DetailsID("DetailsTab");
@@ -570,6 +575,32 @@ FText FAnimBlueprintParentPlayerEditorSummoner::GetTabToolTipText(const FWorkflo
 	return LOCTEXT("AnimSubClassTabToolTip", "Editor for overriding the animation assets referenced by the parent animation graph.");
 }
 
+//////////////////////////////////////////////////////////////////////////
+// FPoseWatchManagerSummoner
+
+FPoseWatchManagerSummoner::FPoseWatchManagerSummoner(TSharedPtr<class FBlueprintEditor> InBlueprintEditor)
+	: FWorkflowTabFactory(FPersonaTabs::PoseWatchManagerID, InBlueprintEditor)
+	, BlueprintEditor(InBlueprintEditor)
+{
+	TabLabel = LOCTEXT("PoseWatchManager", "Pose Watch Manager");
+	TabIcon = FSlateIcon(FEditorStyle::GetStyleSetName(), "ClassIcon.PoseAsset");
+	bIsSingleton = true;
+}
+
+TSharedRef<SWidget> FPoseWatchManagerSummoner::CreateTabBody(const FWorkflowTabSpawnInfo& Info) const
+{
+	FPoseWatchManagerInitializationOptions Options;
+	Options.BlueprintEditor = BlueprintEditor;
+
+	return SNew(SPoseWatchManager, Options)
+		.IsEnabled(FSlateApplication::Get().GetNormalExecutionAttribute());
+}
+
+FText FPoseWatchManagerSummoner::GetTabToolTipText(const FWorkflowTabSpawnInfo& Info) const
+{
+	return LOCTEXT("PoseWatchTabToolTip", "Shows all active pose watches.");
+}
+
 /////////////////////////////////////////////////////
 // FAdvancedPreviewSceneTabSummoner
 
diff --git a/Engine/Source/Editor/Persona/Private/TabSpawners.h b/Engine/Source/Editor/Persona/Private/TabSpawners.h
--- a/Engine/Source/Editor/Persona/Private/TabSpawners.h
+++ b/Engine/Source/Editor/Persona/Private/TabSpawners.h
@@ -312,6 +312,21 @@ private:
 	FSimpleMulticastDelegate& OnPostUndo;
 };
 
+/////////////////////////////////////////////////////
+// FPoseWatchManagerSummoner
+class FPoseWatchManagerSummoner : public FWorkflowTabFactory
+{
+public:
+	FPoseWatchManagerSummoner(TSharedPtr<class FBlueprintEditor> InBlueprintEditor);
+
+	virtual TSharedRef<SWidget> CreateTabBody(const FWorkflowTabSpawnInfo& Info) const override;
+	virtual FText GetTabToolTipText(const FWorkflowTabSpawnInfo& Info) const override;
+
+private:
+	TWeakPtr<class FBlueprintEditor> BlueprintEditor;
+};
+
+
 /////////////////////////////////////////////////////
 // FAdvancedPreviewSceneTabSummoner
 
diff --git a/Engine/Source/Editor/Persona/Public/PersonaModule.h b/Engine/Source/Editor/Persona/Public/PersonaModule.h
--- a/Engine/Source/Editor/Persona/Public/PersonaModule.h
+++ b/Engine/Source/Editor/Persona/Public/PersonaModule.h
@@ -298,6 +298,9 @@ public:
 	/** Create a tab factory for editing anim blueprint preview & defaults */
 	virtual TSharedRef<FWorkflowTabFactory> CreateAnimBlueprintPreviewTabFactory(const TSharedRef<class FBlueprintEditor>& InBlueprintEditor, const TSharedRef<IPersonaPreviewScene>& InPreviewScene) const;
 
+	/** Create a tab factory for the pose watch manager */
+	virtual TSharedRef<FWorkflowTabFactory> CreatePoseWatchTabFactory(const TSharedRef<class FBlueprintEditor>& InBlueprintEditor) const;
+
 	/** Create a tab factory for editing anim blueprint parent overrides */
 	virtual TSharedRef<FWorkflowTabFactory> CreateAnimBlueprintAssetOverridesTabFactory(const TSharedRef<class FBlueprintEditor>& InBlueprintEditor, UAnimBlueprint* InAnimBlueprint, FSimpleMulticastDelegate& InOnPostUndo) const;
 
diff --git a/Engine/Source/Editor/Persona/Public/PersonaTabs.h b/Engine/Source/Editor/Persona/Public/PersonaTabs.h
--- a/Engine/Source/Editor/Persona/Public/PersonaTabs.h
+++ b/Engine/Source/Editor/Persona/Public/PersonaTabs.h
@@ -39,6 +39,7 @@ struct PERSONA_API FPersonaTabs
 	static const FName CurveNameManagerID;
 	static const FName BlendProfileManagerID;
 	static const FName AnimMontageSectionsID;
+	static const FName PoseWatchManagerID;
 
 	// Advanced Preview Scene
 	static const FName AdvancedPreviewSceneSettingsID;
diff --git a/Engine/Source/Editor/UnrealEd/Private/AnimationEditorUtils.cpp b/Engine/Source/Editor/UnrealEd/Private/AnimationEditorUtils.cpp
--- a/Engine/Source/Editor/UnrealEd/Private/AnimationEditorUtils.cpp
+++ b/Engine/Source/Editor/UnrealEd/Private/AnimationEditorUtils.cpp
@@ -844,23 +844,22 @@ namespace AnimationEditorUtils
 	void SetPoseWatch(UPoseWatch* PoseWatch, UAnimBlueprint* AnimBlueprintIfKnown)
 	{
 #if WITH_EDITORONLY_DATA
-		if (UAnimGraphNode_Base* TargetNode = Cast<UAnimGraphNode_Base>(PoseWatch->Node))
-		{
-			UAnimBlueprint* AnimBlueprint = AnimBlueprintIfKnown ? AnimBlueprintIfKnown : Cast<UAnimBlueprint>(FBlueprintEditorUtils::FindBlueprintForNode(TargetNode));
-
-			if ((AnimBlueprint != NULL) && (AnimBlueprint->GeneratedClass != NULL))
+			if (UAnimGraphNode_Base* TargetNode = Cast<UAnimGraphNode_Base>(PoseWatch->Node))
 			{
-				if (UAnimBlueprintGeneratedClass* AnimBPGenClass = Cast<UAnimBlueprintGeneratedClass>(*AnimBlueprint->GeneratedClass))
+				UAnimBlueprint* AnimBlueprint = AnimBlueprintIfKnown ? AnimBlueprintIfKnown : Cast<UAnimBlueprint>(FBlueprintEditorUtils::FindBlueprintForNode(TargetNode));
+				if ((AnimBlueprint != NULL) && (AnimBlueprint->GeneratedClass != NULL))
 				{
-					// Find the insertion point from the debugging data
-					int32 LinkID = AnimBPGenClass->GetLinkIDForNode<FAnimNode_Base>(TargetNode);
-					AnimBPGenClass->GetAnimBlueprintDebugData().AddPoseWatch(LinkID, PoseWatch->PoseWatchColour);
-
-					OnPoseWatchesChangedDelegate.Broadcast(AnimBlueprint, TargetNode);
+					if (UAnimBlueprintGeneratedClass* AnimBPGenClass = Cast<UAnimBlueprintGeneratedClass>(*AnimBlueprint->GeneratedClass))
+					{
+						// Find the insertion point from the debugging data
+						int32 LinkID = AnimBPGenClass->GetLinkIDForNode<FAnimNode_Base>(TargetNode);
+						PoseWatch->SetIsEnabled(LinkID != INDEX_NONE);
+						AnimBPGenClass->GetAnimBlueprintDebugData().AddPoseWatch(LinkID, PoseWatch->GetColor(), PoseWatch->GetIsVisible());
+						OnPoseWatchesChangedDelegate.Broadcast(AnimBlueprint, TargetNode);
+					}
 				}
 			}
-		}
-#endif	//#if WITH_EDITORONLY_DATA
+#endif
 	}
 
 	UPoseWatch* FindPoseWatchForNode(const UEdGraphNode* Node, UAnimBlueprint* AnimBlueprintIfKnown)
@@ -897,7 +896,8 @@ namespace AnimationEditorUtils
 #if WITH_EDITORONLY_DATA
 		UPoseWatch* NewPoseWatch = NewObject<UPoseWatch>(AnimBlueprint);
 		NewPoseWatch->Node = Node;
-		NewPoseWatch->PoseWatchColour = PoseWatchColour;
+		NewPoseWatch->SetColor(PoseWatchColour);
+		NewPoseWatch->SetUniqueDefaultLabel();
 		AnimBlueprint->PoseWatches.Add(NewPoseWatch);
 		SetPoseWatch(NewPoseWatch, AnimBlueprint);
 		return NewPoseWatch;
@@ -912,16 +912,13 @@ namespace AnimationEditorUtils
 		if (UAnimGraphNode_Base* TargetNode = Cast<UAnimGraphNode_Base>(PoseWatch->Node))
 		{
 			UAnimBlueprint* AnimBlueprint = AnimBlueprintIfKnown ? AnimBlueprintIfKnown : Cast<UAnimBlueprint>(FBlueprintEditorUtils::FindBlueprintForNode(TargetNode));
-
 			if (AnimBlueprint)
 			{
 				AnimBlueprint->PoseWatches.Remove(PoseWatch);
-
 				if (UAnimBlueprintGeneratedClass* AnimBPGenClass = AnimBlueprint->GetAnimBlueprintGeneratedClass())
 				{
 					int32 LinkID = AnimBPGenClass->GetLinkIDForNode<FAnimNode_Base>(Cast<UAnimGraphNode_Base>(PoseWatch->Node));
 					AnimBPGenClass->GetAnimBlueprintDebugData().RemovePoseWatch(LinkID);
-
 					OnPoseWatchesChangedDelegate.Broadcast(AnimBlueprint, TargetNode);
 				}
 			}
@@ -929,6 +926,31 @@ namespace AnimationEditorUtils
 #endif
 	}
 
+	void RemovePoseWatchFromNode(UEdGraphNode* Node, UAnimBlueprint* AnimBlueprint)
+	{
+#if WITH_EDITORONLY_DATA
+		if (AnimBlueprint)
+		{
+			for (UPoseWatch* SomePoseWatch : AnimBlueprint->PoseWatches)
+			{
+				if (SomePoseWatch->Node == Node)
+				{
+					if (UAnimBlueprintGeneratedClass* AnimBPGenClass = AnimBlueprint->GetAnimBlueprintGeneratedClass())
+					{
+						int32 LinkID = AnimBPGenClass->GetLinkIDForNode<FAnimNode_Base>(Cast<UAnimGraphNode_Base>(SomePoseWatch->Node));
+						AnimBPGenClass->GetAnimBlueprintDebugData().RemovePoseWatch(LinkID);
+
+						OnPoseWatchesChangedDelegate.Broadcast(AnimBlueprint, Node);
+					}
+
+					SomePoseWatch->OnRemoved();
+					return;
+				}
+			}
+		}
+#endif
+	}
+
 	FOnPoseWatchesChanged& OnPoseWatchesChanged()
 	{
 		return OnPoseWatchesChangedDelegate;
@@ -944,23 +966,45 @@ namespace AnimationEditorUtils
 	void UpdatePoseWatchColour(UPoseWatch* PoseWatch, FColor NewPoseWatchColour)
 	{
 #if WITH_EDITORONLY_DATA
-		PoseWatch->PoseWatchColour = NewPoseWatchColour;
+		UAnimBlueprintGeneratedClass* AnimBPGenClass = nullptr;
+		int32 LinkID = GetPoseWatchNodeLinkID(PoseWatch, AnimBPGenClass);
+		if (LinkID != INDEX_NONE)
+		{
+			AnimBPGenClass->GetAnimBlueprintDebugData().UpdatePoseWatchColour(LinkID, NewPoseWatchColour);
+		}
+#endif
+	}
 
+	int32 GetPoseWatchNodeLinkID(UPoseWatch* PoseWatch, UAnimBlueprintGeneratedClass*& OutAnimBPGenClass)
+	{
+#if WITH_EDITORONLY_DATA
 		if (UAnimGraphNode_Base* TargetNode = Cast<UAnimGraphNode_Base>(PoseWatch->Node))
 		{
 			UAnimBlueprint* AnimBlueprint = Cast<UAnimBlueprint>(FBlueprintEditorUtils::FindBlueprintForNode(TargetNode));
 
-			if ((AnimBlueprint != NULL) && (AnimBlueprint->GeneratedClass != NULL))
+			if ((AnimBlueprint != nullptr) && (AnimBlueprint->GeneratedClass != nullptr))
 			{
 				if (UAnimBlueprintGeneratedClass* AnimBPGenClass = Cast<UAnimBlueprintGeneratedClass>(*AnimBlueprint->GeneratedClass))
 				{
 					// Find the insertion point from the debugging data
-					int32 LinkID = AnimBPGenClass->GetLinkIDForNode<FAnimNode_Base>(TargetNode);
-
-					AnimBPGenClass->GetAnimBlueprintDebugData().UpdatePoseWatchColour(LinkID, NewPoseWatchColour);
+					OutAnimBPGenClass = AnimBPGenClass;
+					return AnimBPGenClass->GetLinkIDForNode<FAnimNode_Base>(TargetNode);
 				}
 			}
 		}
+#endif
+		return INDEX_NONE;
+	}
+
+	void UpdatePoseWatchVisibility(UPoseWatch* PoseWatch)
+	{
+#if WITH_EDITORONLY_DATA
+		UAnimBlueprintGeneratedClass* AnimBPGenClass = nullptr;
+		int32 LinkID = GetPoseWatchNodeLinkID(PoseWatch, AnimBPGenClass);
+		if (LinkID != INDEX_NONE)
+		{
+			AnimBPGenClass->GetAnimBlueprintDebugData().SetPoseWatchVisibility(LinkID, PoseWatch->GetIsVisible());
+		}
 #endif
 	}
 
diff --git a/Engine/Source/Editor/UnrealEd/Public/AnimationEditorUtils.h b/Engine/Source/Editor/UnrealEd/Public/AnimationEditorUtils.h
--- a/Engine/Source/Editor/UnrealEd/Public/AnimationEditorUtils.h
+++ b/Engine/Source/Editor/UnrealEd/Public/AnimationEditorUtils.h
@@ -24,6 +24,8 @@ class UAnimSequence;
 class UEdGraph;
 class UPoseWatch;
 class UEdGraphNode;
+class UAnimBlueprintGeneratedClass;
+class UAnimGraphNode_Base;
 
 /** dialog to prompt users to decide an animation asset name */
 class SCreateAnimationAssetDlg : public SWindow
@@ -192,11 +194,14 @@ namespace AnimationEditorUtils
 	// Is the supplied UEdGraph an Animation Graph
 	UNREALED_API bool IsAnimGraph(UEdGraph* Graph);
 
+	int32 GetPoseWatchNodeLinkID(UPoseWatch* PoseWatch, OUT UAnimBlueprintGeneratedClass*& AnimBPGenClass);
 	UNREALED_API void SetPoseWatch(UPoseWatch* PoseWatch, UAnimBlueprint* AnimBlueprintIfKnown = nullptr);
 	UNREALED_API UPoseWatch* FindPoseWatchForNode(const UEdGraphNode* Node, UAnimBlueprint* AnimBlueprintIfKnown=nullptr);
 	UNREALED_API UPoseWatch* MakePoseWatchForNode(UAnimBlueprint* AnimBlueprint, UEdGraphNode* Node, FColor PoseWatchColour);
 	UNREALED_API void RemovePoseWatch(UPoseWatch* PoseWatch, UAnimBlueprint* AnimBlueprintIfKnown=nullptr);
+	UNREALED_API void RemovePoseWatchFromNode(UEdGraphNode* Node, UAnimBlueprint* AnimBlueprint);
 	UNREALED_API void UpdatePoseWatchColour(UPoseWatch* PoseWatch, FColor NewPoseWatchColour);
+	UNREALED_API void UpdatePoseWatchVisibility(UPoseWatch* PoseWatch);
 	UNREALED_API TArrayView<const FColor> GetPoseWatchColorPalette();
 
 	// Delegate fired when a pose watch is added or removed
diff --git a/Engine/Source/Runtime/Engine/Classes/Animation/AnimBlueprint.h b/Engine/Source/Runtime/Engine/Classes/Animation/AnimBlueprint.h
--- a/Engine/Source/Runtime/Engine/Classes/Animation/AnimBlueprint.h
+++ b/Engine/Source/Runtime/Engine/Classes/Animation/AnimBlueprint.h
@@ -13,6 +13,8 @@ class SWidget;
 class UAnimationAsset;
 class USkeletalMesh;
 class USkeleton;
+class UPoseWatch;
+class UPoseWatchFolder;
 struct FAnimBlueprintDebugData;
 
 USTRUCT()
@@ -216,8 +218,11 @@ public:
 
 	// Array of active pose watches (pose watch allows us to see the bone pose at a 
 	// particular point of the anim graph) 
-	UPROPERTY(transient)
-	TArray<TObjectPtr<class UPoseWatch>> PoseWatches;
+	UPROPERTY()
+	TArray<TObjectPtr<UPoseWatchFolder>> PoseWatchFolders;
+	
+	UPROPERTY()
+	TArray<TObjectPtr<UPoseWatch>> PoseWatches;
 
 private:
 	friend class FAnimBlueprintCompilerContext;
diff --git a/Engine/Source/Runtime/Engine/Classes/Animation/AnimBlueprintGeneratedClass.h b/Engine/Source/Runtime/Engine/Classes/Animation/AnimBlueprintGeneratedClass.h
--- a/Engine/Source/Runtime/Engine/Classes/Animation/AnimBlueprintGeneratedClass.h
+++ b/Engine/Source/Runtime/Engine/Classes/Animation/AnimBlueprintGeneratedClass.h
@@ -6,13 +6,13 @@
 #include "UObject/ObjectMacros.h"
 #include "Misc/Guid.h"
 #include "Animation/AnimTypes.h"
-#include "Engine/PoseWatch.h"
 #include "Engine/BlueprintGeneratedClass.h"
 #include "Animation/AnimStateMachineTypes.h"
 #include "Animation/AnimClassInterface.h"
 #include "Animation/AnimNodeBase.h"
 #include "Animation/BlendSpace.h"
 #include "Animation/ExposedValueHandler.h"
+#include "Engine/PoseWatch.h"
 
 #include "AnimBlueprintGeneratedClass.generated.h"
 
@@ -303,9 +303,10 @@ public:
 	void RecordSequencePlayer(int32 InNodeID, float InPosition, float InLength, int32 InFrameCount);
 	void RecordBlendSpacePlayer(int32 InNodeID, const UBlendSpace* InBlendSpace, const FVector& InPosition, const FVector& InFilteredPosition);
 
-	void AddPoseWatch(int32 NodeID, FColor Color);
+	void AddPoseWatch(int32 NodeID, FColor Color, bool bIsVisible);
 	void RemovePoseWatch(int32 NodeID);
 	void UpdatePoseWatchColour(int32 NodeID, FColor Color);
+	void SetPoseWatchVisibility(int32 NodeID, bool bIsVisible);
 
 	TArrayView<const FName> GetNodeAttributes(TWeakObjectPtr<UAnimGraphNode_Base> InAnimGraphNode) const;
 #endif
diff --git a/Engine/Source/Runtime/Engine/Classes/Engine/PoseWatch.h b/Engine/Source/Runtime/Engine/Classes/Engine/PoseWatch.h
--- a/Engine/Source/Runtime/Engine/Classes/Engine/PoseWatch.h
+++ b/Engine/Source/Runtime/Engine/Classes/Engine/PoseWatch.h
@@ -7,7 +7,9 @@
 #include "UObject/Object.h"
 #include "PoseWatch.generated.h"
 
+
 struct FCompactHeapPose;
+class UAnimBlueprint;
 
 struct FAnimNodePoseWatch
 {
@@ -16,23 +18,228 @@ struct FAnimNodePoseWatch
 	TSharedPtr<FCompactHeapPose>	PoseInfo;
 	FColor							PoseDrawColour;
 	int32							NodeID;
+	bool							bIsHidden;
 };
 
+namespace PoseWatchUtil
+{
+	/** Gets all pose watches that are parented to Folder, if Folder is nullptr then gets orphans */
+	TSet<UPoseWatch*> GetChildrenPoseWatchOf(const UPoseWatchFolder* Folder, const UAnimBlueprint* AnimBlueprint);
+
+	/** Gets all pose watches folders that are parented to Folder, if Folder is nullptr then gets orphans */
+	TSet<UPoseWatchFolder*> GetChildrenPoseWatchFoldersOf(const UPoseWatchFolder* Folder, const UAnimBlueprint* AnimBlueprint);
+
+	/** Finds unique name for a pose watch or folder inside of InParent */
+	template <typename T>
+	FText FindUniqueNameInFolder(UPoseWatchFolder* InParent, const T* Item, TArray<TObjectPtr<T>> Collection);
+}
+
 UCLASS()
-class ENGINE_API UPoseWatch
+class ENGINE_API UPoseWatchFolder
 	: public UObject
 {
 	GENERATED_UCLASS_BODY()
 
 public:
-	// Node that we are trying to watch
+#if WITH_EDITOR
+	/** Returns the slash delimited path of this folder (e.g. MyFolder/MyNestedFolder/MyPoseWatch)*/
+	const FText GetPath() const;
+
+	/** The default name given to all new folders */
+	FText GetDefaultLabel() const;
+
+	/** Returns the display label assigned to this pose watch folder */
+	FText GetLabel() const;
+
+	/** Returns the visibility of this pose watch folder */
+	bool GetIsVisible() const;
+
+	/** Returns the parent folder this folder belongs to, if any */
+	UPoseWatchFolder* GetParent() const;
+	
+	/**  Attempts to set this folder's parent, returns false if unsuccessful (if moving causes a name clash)
+	 *	 Using bForce may change the folder's label to ensure no name clashes
+	 */
+	bool SetParent(UPoseWatchFolder* Parent, bool bForce = false);
+
+	/** Alias of SetParent */
+	void MoveTo(UPoseWatchFolder* InFolder);
+
+	/** Attempts to set the label, returns false if unsuccessful (if there's a name clash with another  folder in the current directory) */
+	bool SetLabel(const FText& InLabel);
+
+	/** Sets the visibility of this folder, must contain at least one post watch descendant to become visible */
+	void SetIsVisible(bool bInIsVisible, bool bUpdateChildren=true);
+
+	/** Called before the pose watch folder is deleted to cleanup it's children and update it's parent */
+	void OnRemoved();
+
+	/** Returns true if InFolder is the parent of this */
+	bool IsIn(const UPoseWatchFolder* InFolder) const;
+
+	/** Returns true if this is a descendant of InFolder */
+	bool IsDescendantOf(const UPoseWatchFolder* InFolder) const;
+
+	/** Returns true if this folder is inside another folder */
+	bool IsAssignedFolder() const;
+
+	/** Returns true if there would be no name clashes when assigning the label InLabel */
+	bool ValidateLabelRename(const FText& InLabel, FText& OutErrorMessage);
+
+	/** Returns true if InLabel is a unique folder label among the children of InFolder,  excluding this */
+	bool IsFolderLabelUniqueInFolder(const FText& InLabel, UPoseWatchFolder* InFolder) const;
+
+	/** Returns true if at least one UPoseWatch/UPoseWatchFolder has this as it's parent */
+	bool HasChildren() const;
+
+	/** Takes GetDefaultLabel() and generates unique labels to avoid name clashes (e.g. NewFolder1, NewFolder2, ...) */
+	void SetUniqueDefaultLabel();
+
+	/** Called when a child is removed/added to a folder */
+	void UpdateVisibility();
+
+	/** Returns the anim blueprint this pose watch folder is stored inside */
+	UAnimBlueprint* GetAnimBlueprint() const;
+
+	void SetIsExpanded(bool bInIsExpanded);
+
+	bool GetIsExpanded() const;
+
+private:
+	/** Returns a unique name for a new folder placed within InParent */
+	FText FindUniqueNameInFolder(UPoseWatchFolder* InParent) const;
+
+	/** Returns true if there is at least one direct pose watch child */
+	bool HasPoseWatchChildren() const;
+
+	/** Returns true if there is at least one pose watch descendant (nested folders) */
+	bool HasPoseWatchDescendents() const;
+#endif // WITH_EDITOR
+
+#if WITH_EDITORONLY_DATA
+protected:
 	UPROPERTY()
-	TObjectPtr<class UEdGraphNode> Node;
+	FText Label;
 
 	UPROPERTY()
-	FColor PoseWatchColour;
+	TWeakObjectPtr<UPoseWatchFolder> Parent;
 
-	/** If true then the watch will be deleted when the connected node is deleted */
 	UPROPERTY()
-	bool bDeleteOnDeselection = false;
+	bool bIsVisible = false;
+
+	UPROPERTY(Transient)
+	bool bIsExpanded = true;
+#endif // WITH_EDITORONLY_DATA
 };
+
+
+UCLASS()
+class ENGINE_API UPoseWatch
+	: public UObject
+{
+	GENERATED_UCLASS_BODY()
+
+public:
+#if WITH_EDITOR
+	/** Returns the slash delimited path of this pose watch */
+	const FText GetPath() const;
+
+	/** The default name given to all new pose watches */
+	FText GetDefaultLabel() const;
+
+	/** Returns the display label assigned to this pose watch */
+	FText GetLabel() const;
+
+	/** Returns the visibility of this pose watch */
+	bool GetIsVisible() const;
+
+	/** Returns true if the pose watch is connected to the output pose */
+	bool GetIsEnabled() const;
+
+	/** Returns the color to display the pose watch using */
+	FColor GetColor() const;
+
+	/** Returns true if this pose watch should be deleted after the user has deselected its assigned node (Editor preference) */
+	bool GetShouldDeleteOnDeselect() const;
+
+	/** Returns the parent folder this folder belongs to, if any */
+	UPoseWatchFolder* GetParent() const;
+
+	/**  Attempts to set this pose watch's parent, returns false if unsuccessful (if moving causes a name clash)
+	 *	 Using bForce may change the pose watch's label to ensure no name clashes
+	 */
+	bool SetParent(UPoseWatchFolder* InParent, bool bForce=false);
+
+	/** If set, denotes the pose watch is able to be drawn to the viewport */
+	void SetIsEnabled(bool bInIsEnabled);
+
+	/** Alias of SetParent */
+	void MoveTo(UPoseWatchFolder* InFolder);
+
+	/** Attempts to set the label, returns false if unsuccessful (if there's a name clash with another pose watch in the current directory) */
+	bool SetLabel(const FText& InLabel);
+
+	/** Sets whether or not to render this pose watch to the viewport */
+	void SetIsVisible(bool bInIsVisible);
+
+	/** Sets the display color of this pose watch in the UI and viewport */
+	void SetColor(const FColor& InColor);
+
+	/** Sets whether this pose watch should delete after deselecting it's assigned node (Editor preference) */
+	void SetShouldDeleteOnDeselect(const bool bInDeleteOnDeselection);
+
+	/** Called when a pose watch is deleted to update it's parent */
+	void OnRemoved();
+
+	/** Toggle's the pose watch's visibility */
+	void ToggleIsVisible();
+
+	/** Returns true if this pose watch is inside InFolder */
+	bool IsIn(const UPoseWatchFolder* InFolder) const;
+
+	/** Returns true if this pose watch is inside some pose watch folder */
+	bool IsAssignedFolder() const;
+
+	/** Returns true if there would be no name clashes when assigning the label InLabel */
+	bool ValidateLabelRename(const FText& InLabel, FText& OutErrorMessage);
+
+	/** Returns true if InLabel is a unique pose watch label among the children of InFolder, excluding this */
+	bool IsPoseWatchLabelUniqueInFolder(const FText& InLabel, UPoseWatchFolder* InFolder) const;
+
+	/** Takes GetDefaultLabel() and generates unique labels to avoid name clashes (e.g. PoseWatch1, PoseWatch2, ...) */
+	void SetUniqueDefaultLabel();
+
+	/** Returns the anim blueprint this pose watch is stored inside */
+	UAnimBlueprint* GetAnimBlueprint() const;
+
+private:
+	/** Returns a unique name for a new pose watch  placed within InParent */
+	FText FindUniqueNameInFolder(UPoseWatchFolder* InParent) const;
+#endif // WITH_EDITOR
+
+#if WITH_EDITORONLY_DATA
+public:
+	UPROPERTY()
+	TWeakObjectPtr<class UEdGraphNode> Node;
+protected:
+	UPROPERTY()
+	bool bDeleteOnDeselection = false;
+
+	// If true will draw the pose to the viewport
+	UPROPERTY()
+	bool bIsVisible = true;
+
+	// If true, the pose is able to be drawn to the viewport
+	UPROPERTY(Transient)
+	bool bIsEnabled = false;
+
+	UPROPERTY()
+	FColor Color;
+
+	UPROPERTY()
+	FText Label;
+
+	UPROPERTY()
+	TWeakObjectPtr<UPoseWatchFolder> Parent;
+#endif // WITH_EDITORONLY_DATA
+};
\ No newline at end of file
diff --git a/Engine/Source/Runtime/Engine/Private/Animation/AnimBlueprintGeneratedClass.cpp b/Engine/Source/Runtime/Engine/Private/Animation/AnimBlueprintGeneratedClass.cpp
--- a/Engine/Source/Runtime/Engine/Private/Animation/AnimBlueprintGeneratedClass.cpp
+++ b/Engine/Source/Runtime/Engine/Private/Animation/AnimBlueprintGeneratedClass.cpp
@@ -228,13 +228,14 @@ void FAnimBlueprintDebugData::RecordBlendSpacePlayer(int32 InNodeID, const UBlen
 	BlendSpacePlayerRecordsThisFrame.Emplace(InNodeID, InBlendSpace, InPosition, InFilteredPosition);
 }
 
-void FAnimBlueprintDebugData::AddPoseWatch(int32 NodeID, FColor Color)
+void FAnimBlueprintDebugData::AddPoseWatch(int32 NodeID, FColor Color, bool bIsVisible)
 {
 	for (FAnimNodePoseWatch& PoseWatch : AnimNodePoseWatch)
 	{
 		if (PoseWatch.NodeID == NodeID)
 		{
 			PoseWatch.PoseDrawColour = Color;
+			PoseWatch.bIsHidden = !bIsVisible;
 			return;
 		}
 	}
@@ -244,6 +245,7 @@ void FAnimBlueprintDebugData::AddPoseWatch(int32 NodeID, FColor Color)
 	FAnimNodePoseWatch& NewAnimNodePoseWatch = AnimNodePoseWatch.Last();
 	NewAnimNodePoseWatch.NodeID = NodeID;
 	NewAnimNodePoseWatch.PoseDrawColour = Color;
+	NewAnimNodePoseWatch.bIsHidden = !bIsVisible;
 	NewAnimNodePoseWatch.PoseInfo = MakeShareable(new FCompactHeapPose());
 	NewAnimNodePoseWatch.Object = nullptr;
 }
@@ -260,6 +262,18 @@ void FAnimBlueprintDebugData::RemovePoseWatch(int32 NodeID)
 	}
 }
 
+void FAnimBlueprintDebugData::SetPoseWatchVisibility(int32 NodeID, bool bIsVisible)
+{
+	for (FAnimNodePoseWatch& PoseWatch : AnimNodePoseWatch)
+	{
+		if (PoseWatch.NodeID == NodeID)
+		{
+			PoseWatch.bIsHidden = !bIsVisible;
+			return;
+		}
+	}
+}
+
 void FAnimBlueprintDebugData::UpdatePoseWatchColour(int32 NodeID, FColor Color)
 {
 	for (FAnimNodePoseWatch& PoseWatch : AnimNodePoseWatch)
diff --git a/Engine/Source/Runtime/Engine/Private/Engine/PoseWatch.cpp b/Engine/Source/Runtime/Engine/Private/Engine/PoseWatch.cpp
--- a/Engine/Source/Runtime/Engine/Private/Engine/PoseWatch.cpp
+++ b/Engine/Source/Runtime/Engine/Private/Engine/PoseWatch.cpp
@@ -1,8 +1,540 @@
 // Copyright Epic Games, Inc. All Rights Reserved.
 
 #include "Engine/PoseWatch.h"
+#include "Animation/AnimBlueprint.h"
+
+#if WITH_EDITOR
+#include "AnimationEditorUtils.h"
+#endif
+
+#define LOCTEXT_NAMESPACE "PoseWatch"
+
+UPoseWatchFolder::UPoseWatchFolder(const FObjectInitializer& ObjectInitializer)
+	: Super(ObjectInitializer)
+{
+#if WITH_EDITOR
+	Label = GetDefaultLabel();
+#endif
+}
 
 UPoseWatch::UPoseWatch(const FObjectInitializer& ObjectInitializer)
 	: Super(ObjectInitializer)
 {
+#if WITH_EDITOR
+	Label = GetDefaultLabel();
+#endif
+}
+
+#if WITH_EDITOR
+
+
+// PoseWatchUtil
+
+TSet<UPoseWatch*> PoseWatchUtil::GetChildrenPoseWatchOf(const UPoseWatchFolder* Folder, const UAnimBlueprint* AnimBlueprint)
+{
+	TSet<UPoseWatch*> Children;
+
+	for (UPoseWatch* SomePoseWatch : AnimBlueprint->PoseWatches)
+	{
+		if (SomePoseWatch->IsIn(Folder))
+		{
+			Children.Add(SomePoseWatch);
+		}
+	}
+
+	return Children;
+}
+
+TSet<UPoseWatchFolder*> PoseWatchUtil::GetChildrenPoseWatchFoldersOf(const UPoseWatchFolder* Folder, const UAnimBlueprint* AnimBlueprint)
+{
+	TSet<UPoseWatchFolder*> Children;
+
+	for (UPoseWatchFolder* SomePoseWatchFolder : AnimBlueprint->PoseWatchFolders)
+	{
+		if (SomePoseWatchFolder->IsIn(Folder))
+		{
+			Children.Add(SomePoseWatchFolder);
+		}
+	}
+
+	return Children;
+}
+
+template <typename T>
+FText PoseWatchUtil::FindUniqueNameInFolder(UPoseWatchFolder* InParent, const T* Item, TArray< TObjectPtr<T>> Collection)
+{
+	static const FTextFormat LabelFormat = FTextFormat::FromString("{0}{1}");
+
+	uint32 Counter = 1;
+	bool bNameClashes = false;
+	FText NewLabel = Item->GetDefaultLabel();
+	UAnimBlueprint* AnimBlueprint = CastChecked<UAnimBlueprint>(Item->GetOuter());
+
+	do
+	{
+		bNameClashes = false;
+		NewLabel = FText::Format(LabelFormat, Item->GetLabel(), FText::AsNumber(Counter));
+		for (T* SomeItem : Collection)
+		{
+			if (SomeItem->IsIn(InParent) && SomeItem->GetLabel().ToString().Equals(NewLabel.ToString()))
+			{
+				++Counter;
+				bNameClashes = true;
+				break;
+			}
+		}
+	} while (bNameClashes);
+
+	return NewLabel;
+}
+
+// UPoseWatchFolder
+
+const FText UPoseWatchFolder::GetPath() const
+{
+	if (Parent.Get() != nullptr)
+	{
+		return FText::Format(LOCTEXT("Path", "{0}/{1}"), Parent->GetPath(), Label);
+	}
+	return Label;
+}
+
+FText UPoseWatchFolder::GetDefaultLabel() const
+{
+	return LOCTEXT("PoseWatchFolderDefaultName", "NewFolder");
+}
+
+FText UPoseWatchFolder::GetLabel() const
+{
+	return Label;
+}
+
+bool UPoseWatchFolder::GetIsVisible() const
+{
+	return bIsVisible;
+}
+
+UPoseWatchFolder* UPoseWatchFolder::GetParent() const
+{
+	return Parent.Get();
+}
+
+bool UPoseWatchFolder::SetParent(UPoseWatchFolder* InParent, bool bForce)
+{
+	if (IsFolderLabelUniqueInFolder(Label, InParent))
+	{
+		Parent = InParent;
+		return true;
+	}
+	else if (bForce)
+	{
+		Label = FindUniqueNameInFolder(InParent);
+		check(IsFolderLabelUniqueInFolder(Label, InParent))
+		Parent = InParent;
+		return true;
+	}
+	return false;
+}
+
+bool UPoseWatchFolder::IsFolderLabelUniqueInFolder(const FText& InLabel, UPoseWatchFolder* InFolder) const
+{
+	for (UPoseWatchFolder* SomeChildFolder : PoseWatchUtil::GetChildrenPoseWatchFoldersOf(InFolder, GetAnimBlueprint()))
+	{
+		if (SomeChildFolder->GetLabel().ToString().Equals(InLabel.ToString()))
+		{
+			if (SomeChildFolder != this)
+			{
+				return false;
+			}
+		}
+	}
+	return true;
 }
+
+void UPoseWatchFolder::MoveTo(UPoseWatchFolder* InFolder)
+{
+	SetParent(InFolder);
+}
+
+bool UPoseWatchFolder::SetLabel(const FText& InLabel)
+{
+	if (IsFolderLabelUniqueInFolder(InLabel, Parent.Get()))
+	{
+		Label = InLabel;
+		return true;
+	}
+	return false;
+}
+
+void UPoseWatchFolder::SetIsVisible(bool bInIsVisible, bool bUpdateChildren)
+{
+	// Can only become visible if there are no children descendents
+	if (!HasPoseWatchDescendents() && bInIsVisible)
+	{
+		bIsVisible = false;
+		return;
+	}
+
+	bIsVisible = bInIsVisible;
+
+	if (bUpdateChildren)
+	{
+		for (UPoseWatch* SomePoseWatch : PoseWatchUtil::GetChildrenPoseWatchOf(this, GetAnimBlueprint()))
+		{
+			SomePoseWatch->SetIsVisible(bInIsVisible);
+		}
+		for (UPoseWatchFolder* SomeChildFolder : PoseWatchUtil::GetChildrenPoseWatchFoldersOf(this, GetAnimBlueprint()))
+		{
+			SomeChildFolder->SetIsVisible(bInIsVisible);
+		}
+	}
+}
+
+void UPoseWatchFolder::SetIsExpanded(bool bInIsExpanded)
+{
+	bIsExpanded = bInIsExpanded;
+}
+
+bool UPoseWatchFolder::GetIsExpanded() const
+{
+	return bIsExpanded;
+}
+
+void UPoseWatchFolder::OnRemoved()
+{
+	// Move all this folder's children to this folder's parent
+	for (UPoseWatch* SomePoseWatch : PoseWatchUtil::GetChildrenPoseWatchOf(this, GetAnimBlueprint()))
+	{
+		SomePoseWatch->SetParent(Parent.Get(), /* bForce*/ true);
+	}
+	for (UPoseWatchFolder* SomePoseWatchFolder : PoseWatchUtil::GetChildrenPoseWatchFoldersOf(this, GetAnimBlueprint()))
+	{
+		SomePoseWatchFolder->SetParent(Parent.Get(), /* bForce*/ true);
+	}
+
+	UAnimBlueprint* AnimBlueprint = CastChecked<UAnimBlueprint>(GetOuter());
+	AnimBlueprint->PoseWatchFolders.Remove(this);
+
+	if (Parent.IsValid())
+	{
+		Parent->UpdateVisibility();
+	}
+
+	AnimationEditorUtils::OnPoseWatchesChanged().Broadcast(GetAnimBlueprint(), nullptr);
+}
+
+void UPoseWatchFolder::UpdateVisibility()
+{
+	bool bNewIsVisible = false;
+
+	for (UPoseWatch* SomePoseWatch : PoseWatchUtil::GetChildrenPoseWatchOf(this, GetAnimBlueprint()))
+	{
+		bNewIsVisible |= SomePoseWatch->GetIsVisible();
+	}
+	for (UPoseWatchFolder* SomePoseWatchFolder : PoseWatchUtil::GetChildrenPoseWatchFoldersOf(this, GetAnimBlueprint()))
+	{
+		bNewIsVisible |= SomePoseWatchFolder->GetIsVisible();
+	}
+
+	SetIsVisible(bNewIsVisible, false);
+	if (Parent.IsValid())
+	{
+		Parent->UpdateVisibility();
+	}
+}
+
+UAnimBlueprint* UPoseWatchFolder::GetAnimBlueprint() const
+{
+	return  CastChecked<UAnimBlueprint>(GetOuter());
+}
+
+bool UPoseWatchFolder::IsIn(const UPoseWatchFolder* InFolder) const
+{
+	return Parent.Get() == InFolder;
+}
+
+bool UPoseWatchFolder::IsDescendantOf(const UPoseWatchFolder* InFolder) const
+{
+	if (IsIn(InFolder))
+	{
+		return true;
+	}
+
+	TWeakObjectPtr<UPoseWatchFolder> ParentFolder = Parent;
+	while (ParentFolder.IsValid())
+	{
+		if (ParentFolder->IsIn(InFolder))
+		{
+			return true;
+		}
+		ParentFolder = ParentFolder->Parent;
+	}
+	return false;
+}
+
+bool UPoseWatchFolder::IsAssignedFolder() const
+{
+	return Parent != nullptr;
+}
+
+bool UPoseWatchFolder::ValidateLabelRename(const FText& InLabel, FText& OutErrorMessage)
+{
+	if (!IsFolderLabelUniqueInFolder(InLabel, Parent.Get()))
+	{
+		OutErrorMessage = LOCTEXT("PoseWatchFolderNameTaken", "A folder already has this name at this level");
+		return false;
+	}
+	return true;
+}
+
+bool UPoseWatchFolder::HasChildren() const
+{
+	if (PoseWatchUtil::GetChildrenPoseWatchFoldersOf(this, GetAnimBlueprint()).Num() > 0)
+	{
+		return true;
+	}
+	if (PoseWatchUtil::GetChildrenPoseWatchOf(this, GetAnimBlueprint()).Num() > 0)
+	{
+		return true;
+	}
+	return false;
+}
+
+void UPoseWatchFolder::SetUniqueDefaultLabel()
+{
+	Label = GetDefaultLabel();
+	Label = FindUniqueNameInFolder(Parent.Get());
+}
+
+FText UPoseWatchFolder::FindUniqueNameInFolder(UPoseWatchFolder* InParent) const
+{
+	return PoseWatchUtil::FindUniqueNameInFolder(InParent, this, GetAnimBlueprint()->PoseWatchFolders);
+}
+
+bool UPoseWatchFolder::HasPoseWatchChildren() const
+{
+	return PoseWatchUtil::GetChildrenPoseWatchOf(this, GetAnimBlueprint()).Num() > 0;
+}
+
+bool UPoseWatchFolder::HasPoseWatchDescendents() const
+{
+	if (HasPoseWatchChildren())
+	{
+		return true;
+	}
+	for (UPoseWatchFolder* SomePoseWatchFolder : PoseWatchUtil::GetChildrenPoseWatchFoldersOf(this, GetAnimBlueprint()))
+	{
+		if (SomePoseWatchFolder->HasPoseWatchDescendents())
+		{
+			return true;
+		}
+	}
+	return false;
+}
+
+// UPoseWatch
+
+const FText UPoseWatch::GetPath() const
+{
+	check(!Label.IsEmpty())
+	if (Parent.Get() != nullptr)
+	{
+		return FText::Format(LOCTEXT("Path", "{0}/{1}"), Parent->GetPath(), Label);
+	}
+	return Label;
+}
+
+FText UPoseWatch::GetLabel() const
+{
+	return Label;
+}
+
+FText UPoseWatch::GetDefaultLabel() const
+{
+	if (Node.IsValid())
+	{
+		return Node->GetNodeTitle(ENodeTitleType::ListView);
+	}
+	return LOCTEXT("NewPoseWatch", "NewPoseWatch");
+}
+
+bool UPoseWatch::GetIsVisible() const
+{
+	return bIsVisible;
+}
+
+FColor UPoseWatch::GetColor() const
+{
+	return Color;
+}
+
+bool UPoseWatch::GetShouldDeleteOnDeselect() const
+{
+	return bDeleteOnDeselection;
+}
+
+void UPoseWatch::OnRemoved()
+{
+	UAnimBlueprint* AnimBlueprint = CastChecked<UAnimBlueprint>(GetOuter());
+	AnimBlueprint->PoseWatches.Remove(this);
+	
+	AnimationEditorUtils::RemovePoseWatch(this, AnimBlueprint);
+
+
+	if (Parent.IsValid())
+	{
+		Parent->UpdateVisibility();
+	}
+
+	AnimationEditorUtils::OnPoseWatchesChanged().Broadcast(GetAnimBlueprint(), Node.Get());
+}
+
+UPoseWatchFolder* UPoseWatch::GetParent() const
+{
+	return Parent.Get();
+}
+
+bool UPoseWatch::SetParent(UPoseWatchFolder* InParent, bool bForce)
+{
+	if (!IsPoseWatchLabelUniqueInFolder(Label, InParent))
+	{
+		if (!bForce)
+		{
+			return false;
+		}
+		Label = FindUniqueNameInFolder(InParent);
+	}
+
+	UPoseWatchFolder* OldParent = Parent.Get();
+	Parent = InParent;
+
+	if (OldParent)
+	{
+		OldParent->UpdateVisibility();
+	}
+
+	if (InParent)
+	{
+		InParent->UpdateVisibility();
+		InParent->SetIsExpanded(true);
+	}
+
+	return true;
+}
+
+bool UPoseWatch::GetIsEnabled() const
+{
+	return bIsEnabled;
+}
+
+void UPoseWatch::SetIsEnabled(bool bInIsEnabled)
+{
+	bIsEnabled = bInIsEnabled;
+}
+
+void UPoseWatch::MoveTo(UPoseWatchFolder* InFolder)
+{
+	SetParent(InFolder);
+}
+
+bool UPoseWatch::SetLabel(const FText& InLabel)
+{
+	UPoseWatch* ConflictingPoseWatch = nullptr;
+	if (IsPoseWatchLabelUniqueInFolder(InLabel, Parent.Get()))
+	{
+		Label = InLabel;
+		return true;
+	}
+	else if (ConflictingPoseWatch == this)
+	{
+		Label = InLabel;
+		return true;
+	}
+	return false;
+}
+
+void UPoseWatch::SetIsVisible(bool bInIsVisible)
+{
+	bIsVisible = bInIsVisible;
+	AnimationEditorUtils::UpdatePoseWatchVisibility(this);
+
+	if (Parent.IsValid())
+	{
+		Parent->UpdateVisibility();
+	}
+}
+
+void UPoseWatch::SetColor(const FColor& InColor)
+{
+	Color = InColor;
+	AnimationEditorUtils::UpdatePoseWatchColour(this, Color);
+}
+
+void UPoseWatch::SetShouldDeleteOnDeselect(const bool bInDeleteOnDeselection)
+{
+	bDeleteOnDeselection = bInDeleteOnDeselection;
+}
+
+void UPoseWatch::ToggleIsVisible()
+{
+	SetIsVisible(!bIsVisible);
+}
+
+bool UPoseWatch::IsIn(const UPoseWatchFolder* InFolder) const
+{
+	return Parent.Get() == InFolder;
+}
+
+bool UPoseWatch::IsAssignedFolder() const
+{
+	return Parent != nullptr;
+}
+
+bool UPoseWatch::ValidateLabelRename(const FText& InLabel, FText& OutErrorMessage)
+{
+	UPoseWatch* ConflictingPoseWatch = nullptr;
+	if (!IsPoseWatchLabelUniqueInFolder(InLabel, Parent.Get()))
+	{
+		if (ConflictingPoseWatch != this)
+		{
+			OutErrorMessage = LOCTEXT("PoseWatchNameTaken", "A pose watch already has this name at this level");
+			return false;
+		}
+	}
+	return true;
+}
+
+bool UPoseWatch::IsPoseWatchLabelUniqueInFolder(const FText& InLabel, UPoseWatchFolder* InFolder) const
+{
+	for (UPoseWatch* SomePoseWatch : PoseWatchUtil::GetChildrenPoseWatchOf(InFolder, GetAnimBlueprint()))
+	{
+		if (SomePoseWatch->GetLabel().ToString().Equals(InLabel.ToString()))
+		{
+			if (SomePoseWatch != this)
+			{
+				return false;
+			}
+		}
+	}
+	return true;
+}
+
+void UPoseWatch::SetUniqueDefaultLabel()
+{
+	Label = GetDefaultLabel();
+	Label = FindUniqueNameInFolder(Parent.Get());
+}
+
+UAnimBlueprint* UPoseWatch::GetAnimBlueprint() const
+{
+	return  CastChecked<UAnimBlueprint>(GetOuter());
+}
+
+FText UPoseWatch::FindUniqueNameInFolder(UPoseWatchFolder* InParent) const
+{
+	return PoseWatchUtil::FindUniqueNameInFolder(InParent, this, GetAnimBlueprint()->PoseWatches);
+}
+
+#endif // WITH_EDITOR
+
+#undef LOCTEXT_NAMESPACE
\ No newline at end of file
