Pose Watch Manager - Doubling clicking tree item jumps to node in graph

diff --git a/Engine/Source/Editor/Persona/Private/SPoseWatchManager.cpp b/Engine/Source/Editor/Persona/Private/SPoseWatchManager.cpp
--- a/Engine/Source/Editor/Persona/Private/SPoseWatchManager.cpp
+++ b/Engine/Source/Editor/Persona/Private/SPoseWatchManager.cpp
@@ -30,8 +30,8 @@
 #include "Widgets/Layout/SSeparator.h"
 #include "Widgets/SOverlay.h"
 #include "EditorFolderUtils.h"
-#include "IAnimationBlueprintEditor.h"
 #include "Animation/AnimBlueprint.h"
+#include "BlueprintEditor.h"
 #include "EdGraph/EdGraph.h"
 #include "GraphEditAction.h"
 #include "AnimationEditorUtils.h"
@@ -66,12 +66,8 @@ void SPoseWatchManager::Construct(const FArguments& InArgs, const FPoseWatchMana
 
 	Mode = MakeShared<FPoseWatchManagerDefaultMode>(this);
 
-	FBlueprintEditor* BlueprintEditor = InInitOptions.BlueprintEditor.Pin().Get();
-	if (IAnimationBlueprintEditor* AnimBlueprintEditor = static_cast<IAnimationBlueprintEditor*>(BlueprintEditor))
-	{
-		AnimBlueprint = static_cast<UAnimBlueprint*>(AnimBlueprintEditor->GetBlueprintObj());
-	}
-	check(AnimBlueprint)
+	BlueprintEditor = InInitOptions.BlueprintEditor.Pin().Get();
+	AnimBlueprint = CastChecked<UAnimBlueprint>(BlueprintEditor->GetBlueprintObj());
 
 	bProcessingFullRefresh = false;
 	bFullRefresh = true;
@@ -176,6 +172,9 @@ void SPoseWatchManager::Construct(const FArguments& InArgs, const FPoseWatchMana
 			// Find out when the user selects something in the tree
 			.OnSelectionChanged(this, &SPoseWatchManager::OnManagerTreeSelectionChanged)
 
+			// Called when the user double-clicks with LMB on an item in the list
+			.OnMouseButtonDoubleClick(this, &SPoseWatchManager::OnManagerTreeDoubleClick)
+
 			// Called when an item is scrolled into view
 			.OnItemScrolledIntoView(this, &SPoseWatchManager::OnManagerTreeItemScrolledIntoView)
 
@@ -895,4 +894,13 @@ uint32 SPoseWatchManager::GetTypeSortPriority(const IPoseWatchManagerTreeItem& I
 	return Mode->GetTypeSortPriority(Item);
 }
 
+void SPoseWatchManager::OnManagerTreeDoubleClick(FPoseWatchManagerTreeItemPtr TreeItem)
+{
+	if (FPoseWatchManagerPoseWatchTreeItem* PoseWatchTreeItem = TreeItem->CastTo<FPoseWatchManagerPoseWatchTreeItem>())
+	{
+		UPoseWatch* PoseWatch = PoseWatchTreeItem->PoseWatch.Get();
+		BlueprintEditor->JumpToHyperlink(PoseWatch->Node.Get(), false);
+	}
+}
+
 #undef LOCTEXT_NAMESPACE
diff --git a/Engine/Source/Editor/Persona/Private/SPoseWatchManager.h b/Engine/Source/Editor/Persona/Private/SPoseWatchManager.h
--- a/Engine/Source/Editor/Persona/Private/SPoseWatchManager.h
+++ b/Engine/Source/Editor/Persona/Private/SPoseWatchManager.h
@@ -281,6 +281,8 @@ private:
 	/** Called by STreeView when an item is scrolled into view */
 	void OnManagerTreeItemScrolledIntoView(FPoseWatchManagerTreeItemPtr TreeItem, const TSharedPtr<ITableRow>& Widget);
 
+	void OnManagerTreeDoubleClick(FPoseWatchManagerTreeItemPtr TreeItem);
+
 	/** Called when an item in the tree has been collapsed or expanded */
 	void OnItemExpansionChanged(FPoseWatchManagerTreeItemPtr TreeItem, bool bIsExpanded) const;
 
@@ -403,6 +405,8 @@ private:
 	/** Currently selected sorting mode */
 	EColumnSortMode::Type SortMode;
 
+	FBlueprintEditor* BlueprintEditor;
+
 	/** Handles column sorting mode change */
 	void OnColumnSortModeChanged(const EColumnSortPriority::Type SortPriority, const FName& ColumnId, const EColumnSortMode::Type InSortMode);
 
