Fix CIS build
#rnx

diff --git a/Engine/Source/Editor/Persona/Private/SkeletalMeshReferenceSectionDetails.h b/Engine/Source/Editor/Persona/Private/SkeletalMeshReferenceSectionDetails.h
--- a/Engine/Source/Editor/Persona/Private/SkeletalMeshReferenceSectionDetails.h
+++ b/Engine/Source/Editor/Persona/Private/SkeletalMeshReferenceSectionDetails.h
@@ -4,6 +4,7 @@
 
 #include "CoreMinimal.h"
 #include "IDetailCustomization.h"
+#include "IPropertyTypeCustomization.h"
 #include "Types/SlateEnums.h"
 
 class FSkeletalMeshLODModel;
@@ -15,6 +16,12 @@ class USkeletalMesh;
 class FSectionReferenceCustomization : public IPropertyTypeCustomization
 {
 public:
+	virtual ~FSectionReferenceCustomization()
+	{
+		SectionIndexProperty = nullptr;
+		TargetSkeletalMesh = nullptr;
+	}
+
 	static TSharedRef<IPropertyTypeCustomization> MakeInstance();
 
 	// IPropertyTypeCustomization interface
