[Backout] - CL18301216
Original CL Desc
-----------------------------------------------------------------
Fix CIS build
#rnx

diff --git a/Engine/Source/Editor/Persona/Private/SkeletalMeshReferenceSectionDetails.h b/Engine/Source/Editor/Persona/Private/SkeletalMeshReferenceSectionDetails.h
--- a/Engine/Source/Editor/Persona/Private/SkeletalMeshReferenceSectionDetails.h
+++ b/Engine/Source/Editor/Persona/Private/SkeletalMeshReferenceSectionDetails.h
@@ -4,7 +4,6 @@
 
 #include "CoreMinimal.h"
 #include "IDetailCustomization.h"
-#include "IPropertyTypeCustomization.h"
 #include "Types/SlateEnums.h"
 
 class FSkeletalMeshLODModel;
@@ -16,12 +15,6 @@ class USkeletalMesh;
 class FSectionReferenceCustomization : public IPropertyTypeCustomization
 {
 public:
-	virtual ~FSectionReferenceCustomization()
-	{
-		SectionIndexProperty = nullptr;
-		TargetSkeletalMesh = nullptr;
-	}
-
 	static TSharedRef<IPropertyTypeCustomization> MakeInstance();
 
 	// IPropertyTypeCustomization interface
