Fixed virtual bone skeleton tree item not selecting when a virtual bone is selected in the viewport.
Fixed font not found warning spam when a virtual bone skeleton tree item is selected.
Fixed virtual bone name not showing when selected in viewport.
Cleaned up parity of bone names text shadows

diff --git a/Engine/Source/Editor/Persona/Private/EditModes/SkeletonSelectionEditMode.cpp b/Engine/Source/Editor/Persona/Private/EditModes/SkeletonSelectionEditMode.cpp
--- a/Engine/Source/Editor/Persona/Private/EditModes/SkeletonSelectionEditMode.cpp
+++ b/Engine/Source/Editor/Persona/Private/EditModes/SkeletonSelectionEditMode.cpp
@@ -339,15 +339,16 @@ FIntPoint FSkeletonSelectionEditMode::GetDPIUnscaledSize(FViewport* Viewport, FV
 void FSkeletonSelectionEditMode::DrawHUD(FEditorViewportClient* ViewportClient, FViewport* Viewport, const FSceneView* View, FCanvas* Canvas)
 {
 	UDebugSkelMeshComponent* PreviewMeshComponent = GetAnimPreviewScene().GetPreviewMeshComponent();
+	FReferenceSkeleton& RefSkeleton = PreviewMeshComponent->SkeletalMesh->GetRefSkeleton();
+	int32 BoneIndex = GetAnimPreviewScene().GetSelectedBoneIndex();
 
 	// Draw name of selected bone
-	if (IsSelectedBoneRequired())
+	if (IsSelectedBoneRequired() || RefSkeleton.GetRequiredVirtualBones().Contains(BoneIndex))
 	{
 		const FIntPoint ViewPortSize = GetDPIUnscaledSize(Viewport, ViewportClient);
 		const int32 HalfX = ViewPortSize.X / 2;
 		const int32 HalfY = ViewPortSize.Y / 2;
 
-		int32 BoneIndex = GetAnimPreviewScene().GetSelectedBoneIndex();
 		const FName BoneName = PreviewMeshComponent->SkeletalMesh->GetRefSkeleton().GetBoneName(BoneIndex);
 
 		FMatrix BoneMatrix = PreviewMeshComponent->GetBoneMatrix(BoneIndex);
@@ -360,6 +361,7 @@ void FSkeletonSelectionEditMode::DrawHUD(FEditorViewportClient* ViewportClient,
 			FQuat BoneQuat = PreviewMeshComponent->GetBoneQuaternion(BoneName);
 			FVector Loc = PreviewMeshComponent->GetBoneLocation(BoneName);
 			FCanvasTextItem TextItem(FVector2D(XPos, YPos), FText::FromString(BoneName.ToString()), GEngine->GetSmallFont(), FLinearColor::White);
+			TextItem.EnableShadow(FLinearColor::Black);
 			Canvas->DrawItem(TextItem);
 		}
 	}
@@ -383,6 +385,7 @@ void FSkeletonSelectionEditMode::DrawHUD(FEditorViewportClient* ViewportClient,
 			const int32 XPos = HalfX + (HalfX * Proj.X);
 			const int32 YPos = HalfY + (HalfY * (Proj.Y * -1));
 			FCanvasTextItem TextItem(FVector2D(XPos, YPos), FText::FromString(Socket->SocketName.ToString()), GEngine->GetSmallFont(), FLinearColor::White);
+			TextItem.EnableShadow(FLinearColor::Black);
 			Canvas->DrawItem(TextItem);
 		}
 	}
diff --git a/Engine/Source/Editor/SkeletonEditor/Private/SSkeletonTree.cpp b/Engine/Source/Editor/SkeletonEditor/Private/SSkeletonTree.cpp
--- a/Engine/Source/Editor/SkeletonEditor/Private/SSkeletonTree.cpp
+++ b/Engine/Source/Editor/SkeletonEditor/Private/SSkeletonTree.cpp
@@ -1736,7 +1736,7 @@ void SSkeletonTree::SetSelectedBone( const FName& BoneName, ESelectInfo::Type In
 		{
 			TSharedPtr<ISkeletonTreeItem> SkeletonRow = *(SkeletonRowIt);
 
-			if (SkeletonRow->GetFilterResult() != ESkeletonTreeFilterResult::Hidden && SkeletonRow->IsOfType<FSkeletonTreeBoneItem>() && SkeletonRow->GetRowItemName() == BoneName)
+			if (SkeletonRow->GetFilterResult() != ESkeletonTreeFilterResult::Hidden && (SkeletonRow->IsOfType<FSkeletonTreeBoneItem>() || SkeletonRow->IsOfType<FSkeletonTreeVirtualBoneItem>()) && SkeletonRow->GetRowItemName() == BoneName)
 			{
 				SkeletonTreeView->SetItemSelection(SkeletonRow, true, InSelectInfo);
 				ExpandTreeOnSelection(SkeletonRow);
diff --git a/Engine/Source/Editor/SkeletonEditor/Private/SkeletonTreeVirtualBoneItem.cpp b/Engine/Source/Editor/SkeletonEditor/Private/SkeletonTreeVirtualBoneItem.cpp
--- a/Engine/Source/Editor/SkeletonEditor/Private/SkeletonTreeVirtualBoneItem.cpp
+++ b/Engine/Source/Editor/SkeletonEditor/Private/SkeletonTreeVirtualBoneItem.cpp
@@ -93,7 +93,7 @@ TSharedRef< SWidget > FSkeletonTreeVirtualBoneItem::GenerateWidgetForDataColumn(
 
 FSlateFontInfo FSkeletonTreeVirtualBoneItem::GetBoneTextFont() const
 {
-	return FEditorStyle::GetWidgetStyle<FTextBlockStyle>("SkeletonTree.BoldFont").Font;
+	return FEditorStyle::GetWidgetStyle<FTextBlockStyle>("SkeletonTree.NormalFont").Font;
 }
 
 FSlateColor FSkeletonTreeVirtualBoneItem::GetBoneTextColor(FIsSelected InIsSelected) const
