Added a nullptr check.

diff --git a/Engine/Source/Editor/Persona/Private/EditModes/SkeletonSelectionEditMode.cpp b/Engine/Source/Editor/Persona/Private/EditModes/SkeletonSelectionEditMode.cpp
--- a/Engine/Source/Editor/Persona/Private/EditModes/SkeletonSelectionEditMode.cpp
+++ b/Engine/Source/Editor/Persona/Private/EditModes/SkeletonSelectionEditMode.cpp
@@ -339,6 +339,11 @@ FIntPoint FSkeletonSelectionEditMode::GetDPIUnscaledSize(FViewport* Viewport, FV
 void FSkeletonSelectionEditMode::DrawHUD(FEditorViewportClient* ViewportClient, FViewport* Viewport, const FSceneView* View, FCanvas* Canvas)
 {
 	UDebugSkelMeshComponent* PreviewMeshComponent = GetAnimPreviewScene().GetPreviewMeshComponent();
+	if (PreviewMeshComponent == nullptr || PreviewMeshComponent->SkeletalMesh == nullptr)
+	{
+		return;
+	}
+
 	FReferenceSkeleton& RefSkeleton = PreviewMeshComponent->SkeletalMesh->GetRefSkeleton();
 	int32 BoneIndex = GetAnimPreviewScene().GetSelectedBoneIndex();
 
