Misc pose watch manager changes.
- Pose watch manager now in default anim bp layout
- Fix for being unable to make pose watches created from 'Automatically create pose watch on selection' setting a permanent pose watch and vice versa
- Cleaned up drag-and-drop validation
- Updated context menus
- Improved resolving name clashes

diff --git a/Engine/Source/Editor/AnimationBlueprintEditor/Private/AnimationBlueprintEditor.cpp b/Engine/Source/Editor/AnimationBlueprintEditor/Private/AnimationBlueprintEditor.cpp
--- a/Engine/Source/Editor/AnimationBlueprintEditor/Private/AnimationBlueprintEditor.cpp
+++ b/Engine/Source/Editor/AnimationBlueprintEditor/Private/AnimationBlueprintEditor.cpp
@@ -108,7 +108,7 @@ namespace AnimationBlueprintEditorTabs
 	const FName AssetOverridesTab(TEXT("AnimBlueprintParentPlayerEditor"));
 	const FName SlotNamesTab(TEXT("SkeletonSlotNames"));
 	const FName CurveNamesTab(TEXT("AnimCurveViewerTab"));
-	const FName PoseWatchTab(TEXT("PoseWatchTab"));
+	const FName PoseWatchTab(TEXT("PoseWatchManager"));
 };
 
 /////////////////////////////////////////////////////
@@ -695,64 +695,6 @@ void FAnimationBlueprintEditor::OnRemovePosePin()
 	}
 }
 
-FColor FAnimationBlueprintEditor::ChoosePoseWatchColor() const
-{
-	TArrayView<const FColor> PoseWatchColors = AnimationEditorUtils::GetPoseWatchColorPalette();
-	const int32 NumColors = PoseWatchColors.Num();
-
-	if (NumColors <= 0)
-	{
-		return FColor::Red;
-	}
-
-	struct FColorCount 
-	{ 
-		FColor Color; 
-		int    Count;
-		bool operator<(const FColorCount& Other) const { return Count < Other.Count; }
-	};
-	TArray<FColorCount> ColorCounts;
-	ColorCounts.SetNum(NumColors);
-	for (int ColorIndex = 0 ; ColorIndex != NumColors ; ++ColorIndex)
-	{
-		ColorCounts[ColorIndex].Color = PoseWatchColors[ColorIndex];
-		ColorCounts[ColorIndex].Count = 0;
-	}
-
-	TSharedPtr<SGraphEditor> FocusedGraphEd = FocusedGraphEdPtr.Pin();
-	if (FocusedGraphEd.IsValid())
-	{
-		UAnimBlueprint* AnimBP = GetAnimBlueprint();
-		for (UEdGraphNode* Node : FocusedGraphEd->GetCurrentGraph()->Nodes)
-		{
-			UAnimGraphNode_Base* GraphNode = Cast<UAnimGraphNode_Base>(Node);
-			UPoseWatch* PoseWatch = AnimationEditorUtils::FindPoseWatchForNode(GraphNode, AnimBP);
-			if (PoseWatch)
-			{
-				FColor PoseWatchColor = PoseWatch->GetColor();
-				int32 Index = ColorCounts.IndexOfByPredicate([PoseWatchColor](const FColorCount& CC){return PoseWatchColor == CC.Color;});
-				if (Index != INDEX_NONE)
-				{
-					++ColorCounts[Index].Count;
-				}
-			}
-		}
-	}
-
-	// Pick the colour with the lowest count, preferring the original order if there's a tie (which there normally will be)
-	int BestIndex = 0;
-	for (int ColorIndex = 1 ; ColorIndex < NumColors ; ++ColorIndex)
-	{
-		if (ColorCounts[ColorIndex].Count < ColorCounts[BestIndex].Count)
-		{
-			BestIndex = ColorIndex;
-		}
-	}
-
-	return ColorCounts[BestIndex].Color;
-}
-
-
 void FAnimationBlueprintEditor::OnTogglePoseWatch()
 {
 	const FGraphPanelSelectionSet SelectedNodes = GetSelectedNodes();
@@ -765,12 +707,24 @@ void FAnimationBlueprintEditor::OnTogglePoseWatch()
 			UPoseWatch* ExistingPoseWatch = AnimationEditorUtils::FindPoseWatchForNode(SelectedNode, AnimBP);
 			if (ExistingPoseWatch)
 			{
-				AnimationEditorUtils::RemovePoseWatch(ExistingPoseWatch, AnimBP);
+				// Promote the temporary pose watch to permanent 
+				if (ExistingPoseWatch->GetShouldDeleteOnDeselect())
+				{
+					ExistingPoseWatch->SetShouldDeleteOnDeselect(false);
+				}
+				else if (GetDefault<UAnimationBlueprintEditorSettings>()->bPoseWatchSelectedNodes)
+				{
+					ExistingPoseWatch->SetShouldDeleteOnDeselect(true);
+				}
+				else
+				{
+					AnimationEditorUtils::RemovePoseWatch(ExistingPoseWatch, AnimBP);
+				}
 				AnimationEditorUtils::OnPoseWatchesChanged().Broadcast(AnimBP, ExistingPoseWatch->Node.Get());
 			}
 			else
 			{
-				UPoseWatch* NewPoseWatch = AnimationEditorUtils::MakePoseWatchForNode(AnimBP, SelectedNode, ChoosePoseWatchColor());
+				UPoseWatch* NewPoseWatch = AnimationEditorUtils::MakePoseWatchForNode(AnimBP, SelectedNode);
 				AnimationEditorUtils::OnPoseWatchesChanged().Broadcast(AnimBP, NewPoseWatch->Node.Get());
 			}
 		}
@@ -1898,7 +1852,7 @@ void FAnimationBlueprintEditor::HandlePoseWatchSelectedNodes()
 				{
 					if (!PoseWatch)
 					{
-						PoseWatch = AnimationEditorUtils::MakePoseWatchForNode(AnimBP, GraphNode, ChoosePoseWatchColor());
+						PoseWatch = AnimationEditorUtils::MakePoseWatchForNode(AnimBP, GraphNode);
 						PoseWatch->SetShouldDeleteOnDeselect(true);
 					}
 				}
diff --git a/Engine/Source/Editor/AnimationBlueprintEditor/Private/AnimationBlueprintEditorMode.cpp b/Engine/Source/Editor/AnimationBlueprintEditor/Private/AnimationBlueprintEditorMode.cpp
--- a/Engine/Source/Editor/AnimationBlueprintEditor/Private/AnimationBlueprintEditorMode.cpp
+++ b/Engine/Source/Editor/AnimationBlueprintEditor/Private/AnimationBlueprintEditorMode.cpp
@@ -56,6 +56,7 @@ FAnimationBlueprintEditorMode::FAnimationBlueprintEditorMode(const TSharedRef<FA
 						->SetSizeCoefficient(0.5f)
 						->AddTab(AnimationBlueprintEditorTabs::CurveNamesTab, ETabState::ClosedTab)
 						->AddTab(AnimationBlueprintEditorTabs::SkeletonTreeTab, ETabState::ClosedTab)
+						->AddTab(AnimationBlueprintEditorTabs::PoseWatchTab, ETabState::OpenedTab)
 						->AddTab(FBlueprintEditorTabs::MyBlueprintID, ETabState::OpenedTab)
 					)
 				)
diff --git a/Engine/Source/Editor/Persona/Private/AnimationEditorViewportClient.cpp b/Engine/Source/Editor/Persona/Private/AnimationEditorViewportClient.cpp
--- a/Engine/Source/Editor/Persona/Private/AnimationEditorViewportClient.cpp
+++ b/Engine/Source/Editor/Persona/Private/AnimationEditorViewportClient.cpp
@@ -37,6 +37,7 @@
 #include "Animation/MorphTarget.h"
 #include "Rendering/SkeletalMeshModel.h"
 #include "UnrealWidget.h"
+#include "Engine/PoseWatch.h"
 
 namespace {
 	static const float AnimationEditorViewport_RotateSpeed = 0.02f;
@@ -1332,9 +1333,9 @@ void FAnimationViewportClient::DrawWatchedPoses(UDebugSkelMeshComponent * MeshCo
 			{
 				for (const FAnimNodePoseWatch& AnimNodePoseWatch : AnimBlueprintGeneratedClass->GetAnimBlueprintDebugData().AnimNodePoseWatch)
 				{
-					if(AnimNodePoseWatch.Object.Get() != nullptr && !AnimNodePoseWatch.bIsHidden)
+					if (AnimNodePoseWatch.Object.Get() != nullptr && AnimNodePoseWatch.PoseWatch->GetIsVisible() && AnimNodePoseWatch.PoseWatch->GetIsEnabled())
 					{
-						DrawBonesFromCompactPose(*AnimNodePoseWatch.PoseInfo.Get(), MeshComponent, PDI, AnimNodePoseWatch.PoseDrawColour);
+						DrawBonesFromCompactPose(*AnimNodePoseWatch.PoseInfo.Get(), MeshComponent, PDI, AnimNodePoseWatch.PoseWatch->GetColor());
 					}
 				}
 			}
diff --git a/Engine/Source/Editor/Persona/Private/PoseWatchManagerColumnVisibility.cpp b/Engine/Source/Editor/Persona/Private/PoseWatchManagerColumnVisibility.cpp
--- a/Engine/Source/Editor/Persona/Private/PoseWatchManagerColumnVisibility.cpp
+++ b/Engine/Source/Editor/Persona/Private/PoseWatchManagerColumnVisibility.cpp
@@ -83,7 +83,7 @@ private:
 
 		if (!TreeItem->IsEnabled())
 		{
-			return FPoseWatchManagerCommonLabelData::DarkColor;
+			return FPoseWatchManagerCommonLabelData::DisabledColor;
 		}
 
 		if (IsHovered() && !bIsSelected)
diff --git a/Engine/Source/Editor/Persona/Private/PoseWatchManagerDefaultMode.cpp b/Engine/Source/Editor/Persona/Private/PoseWatchManagerDefaultMode.cpp
--- a/Engine/Source/Editor/Persona/Private/PoseWatchManagerDefaultMode.cpp
+++ b/Engine/Source/Editor/Persona/Private/PoseWatchManagerDefaultMode.cpp
@@ -53,85 +53,78 @@ FPoseWatchManagerDragValidationInfo FPoseWatchManagerDefaultMode::ValidateDrop(c
 	check(Payload.DraggedItem.IsValid())
 	TSharedPtr<IPoseWatchManagerTreeItem> PayloadItem = Payload.DraggedItem.Pin();
 	
-	// Support for removing the parent by dragging an item to the bottom of the tree
-	if (!DropTarget.IsValid())
-	{
-		return FPoseWatchManagerDragValidationInfo(EPoseWatchManagerDropCompatibility::Compatible, LOCTEXT("MoveToRoot", "Move to root"));
-	}
-
 	// Pose watches cannot be a parent
 	if (const FPoseWatchManagerPoseWatchTreeItem* PoseWatchDropTarget = DropTarget.CastTo<FPoseWatchManagerPoseWatchTreeItem>())
 	{
-		return FPoseWatchManagerDragValidationInfo(EPoseWatchManagerDropCompatibility::Incompatible, LOCTEXT("PoseWatchBadParent", "Pose Watches cannot be parents"));
+		return FPoseWatchManagerDragValidationInfo(EPoseWatchManagerDropCompatibility::Incompatible, LOCTEXT("PoseWatchBadParent", "A pose watch cannot be a parent"));
 	}
 
-	// Folders can be parents as long as the payload is not already a parent of the drop target, otherwise there'll be cycles
-	if (const FPoseWatchManagerFolderTreeItem* FolderItem = DropTarget.CastTo<FPoseWatchManagerFolderTreeItem>())
+	// Drop target must either be a valid UPoseWatchFolder or nullptr denoting the root of the tree
+	UPoseWatchFolder* TargetFolder = nullptr;
+	FText TargetFolderLabel = LOCTEXT("Root", "root");
+	if (DropTarget.IsValid())
 	{
-		if (FolderItem->PoseWatchFolder.IsValid())
+		if (const FPoseWatchManagerFolderTreeItem* DropTargetFolder = DropTarget.CastTo<FPoseWatchManagerFolderTreeItem>())
 		{
-			// Dropping a folder into a folder
-			if (const FPoseWatchManagerFolderTreeItem* FolderPayloadItem = PayloadItem.Get()->CastTo<FPoseWatchManagerFolderTreeItem>())
+			TargetFolder = DropTargetFolder->PoseWatchFolder.Get();
+			if (TargetFolder)
 			{
-				if (FolderPayloadItem->PoseWatchFolder == FolderItem->PoseWatchFolder)
-				{
-					return FPoseWatchManagerDragValidationInfo(EPoseWatchManagerDropCompatibility::Incompatible, LOCTEXT("FolderInItself", "A folder cannot contain itself"));
-				}
-
-				if (FolderPayloadItem->PoseWatchFolder->IsIn(FolderItem->PoseWatchFolder.Get()))
-				{
-					FText ValidationText = FText::Format(LOCTEXT("MovePoseWatchIntoFolder", "This folder is already inside {0}"), FolderItem->PoseWatchFolder->GetLabel());
-					return FPoseWatchManagerDragValidationInfo(EPoseWatchManagerDropCompatibility::Incompatible, ValidationText);
-				}
-
-				//if (FolderPayloadItem->PoseWatchFolder->IsDescendantOf(FolderItem->PoseWatchFolder.Get()))
-				if (FolderItem->PoseWatchFolder->IsDescendantOf(FolderPayloadItem->PoseWatchFolder.Get()))
-				{
-					return FPoseWatchManagerDragValidationInfo(EPoseWatchManagerDropCompatibility::Incompatible, LOCTEXT("ParentFolderNotChildren", "Parent folders cannot be children"));
-				}
-
-				FText OutErrorMessage;
-				FText Label = FolderPayloadItem->PoseWatchFolder->GetLabel();
-				UPoseWatchFolder* ConflictingFolder = nullptr;
-
-				if (FolderPayloadItem->PoseWatchFolder->IsFolderLabelUniqueInFolder(FolderPayloadItem->PoseWatchFolder->GetLabel(), FolderItem->PoseWatchFolder.Get()))
-				{
-					FText ValidationText = FText::Format(LOCTEXT("MovePoseWatchIntoFolder", "Move into {0}"), FolderItem->PoseWatchFolder->GetLabel());
-					return FPoseWatchManagerDragValidationInfo(EPoseWatchManagerDropCompatibility::Compatible, ValidationText);
-				}
-				else
-				{
-					FText ValidationText = FText::Format(LOCTEXT("MovePoseWatchIntoFolder", "A folder with that name already exists within {0}"), FolderItem->PoseWatchFolder->GetLabel());
-					return FPoseWatchManagerDragValidationInfo(EPoseWatchManagerDropCompatibility::Incompatible, ValidationText);
-				}
+				TargetFolderLabel = TargetFolder->GetLabel();
 			}
+		}
+	}
 
-			// Dropping a pose watch into a folder
-			if (const FPoseWatchManagerPoseWatchTreeItem* PoseWatchPayloadItem = PayloadItem.Get()->CastTo<FPoseWatchManagerPoseWatchTreeItem>())
-			{
-				if (PoseWatchPayloadItem->PoseWatch->IsIn(FolderItem->PoseWatchFolder.Get()))
-				{
-					FText ValidationText = FText::Format(LOCTEXT("MovePoseWatchIntoFolder", "This pose watch is already inside {0}"), FolderItem->PoseWatchFolder->GetLabel());
-					return FPoseWatchManagerDragValidationInfo(EPoseWatchManagerDropCompatibility::Incompatible, ValidationText);
-				}
-				else
-				{
-					FText ValidationText = FText::Format(LOCTEXT("MovePoseWatchIntoFolder", "Move into {0}"), FolderItem->PoseWatchFolder->GetLabel());
-					return FPoseWatchManagerDragValidationInfo(EPoseWatchManagerDropCompatibility::Compatible, ValidationText);
-				}
-			}
+	if (FPoseWatchManagerPoseWatchTreeItem* PayloadPoseWatchItem = PayloadItem->CastTo<FPoseWatchManagerPoseWatchTreeItem>())
+	{
+		// Dropping a pose watch into a folder
+		UPoseWatch* PayloadPoseWatch = PayloadPoseWatchItem->PoseWatch.Get();
+		check(PayloadPoseWatch);
+
+		if (PayloadPoseWatch->IsIn(TargetFolder))
+		{
+			FText ValidationText = FText::Format(LOCTEXT("PoseWatchAlreadyInDirectory", "This pose watch is already inside {0}"), TargetFolderLabel);
+			return FPoseWatchManagerDragValidationInfo(EPoseWatchManagerDropCompatibility::Incompatible, ValidationText);
 		}
-		else
+
+		if (!PayloadPoseWatch->IsPoseWatchLabelUniqueInFolder(PayloadPoseWatch->GetLabel(), TargetFolder))
 		{
-			if (PayloadItem.Get()->IsAssignedFolder())
-			{
-				return FPoseWatchManagerDragValidationInfo(EPoseWatchManagerDropCompatibility::Compatible, LOCTEXT("MoveToRoot", "Move to root"));
-			}
-			else
-			{
-				return FPoseWatchManagerDragValidationInfo(EPoseWatchManagerDropCompatibility::Incompatible, LOCTEXT("FolderAlreadyInRoot", "This item is already in root"));
-			}
+			FText ValidationText = FText::Format(LOCTEXT("MovePoseWatchIntoFolder", "A pose watch with that name already exists inside {0}"), TargetFolderLabel);
+			return FPoseWatchManagerDragValidationInfo(EPoseWatchManagerDropCompatibility::Incompatible, ValidationText);
 		}
+
+		FText ValidationText = FText::Format(LOCTEXT("MovePoseWatchIntoFolder", "Move {0} into {1}"), PayloadPoseWatch->GetLabel(), TargetFolderLabel);
+		return FPoseWatchManagerDragValidationInfo(EPoseWatchManagerDropCompatibility::Compatible, ValidationText);
+	}
+	else if (FPoseWatchManagerFolderTreeItem* PayloadFolderItem = PayloadItem->CastTo<FPoseWatchManagerFolderTreeItem>())
+	{
+		// Dropping a folder into a folder
+		UPoseWatchFolder* PayloadFolder = PayloadFolderItem->PoseWatchFolder.Get();
+		check(PayloadFolder);
+
+		if (PayloadFolder == TargetFolder)
+		{
+			return FPoseWatchManagerDragValidationInfo(EPoseWatchManagerDropCompatibility::Incompatible, LOCTEXT("FolderInItself", "A folder cannot contain itself"));
+		}
+
+		if (PayloadFolder->IsIn(TargetFolder))
+		{
+			FText ValidationText = FText::Format(LOCTEXT("FolderAlreadyInDirectory", "This folder is already inside {0}"), TargetFolderLabel);
+			return FPoseWatchManagerDragValidationInfo(EPoseWatchManagerDropCompatibility::Incompatible, ValidationText);
+		}
+
+		if (TargetFolder && TargetFolder->IsDescendantOf(PayloadFolder))
+		{
+			return FPoseWatchManagerDragValidationInfo(EPoseWatchManagerDropCompatibility::Incompatible, LOCTEXT("ParentFolderNotChildren", "Parent folders cannot be children"));
+		}
+
+		if (!PayloadFolder->IsFolderLabelUniqueInFolder(PayloadFolder->GetLabel(), TargetFolder))
+		{
+			FText ValidationText = FText::Format(LOCTEXT("FolderNameConflictInFolder", "A folder with that name already exists inside {0}"), TargetFolderLabel);
+			return FPoseWatchManagerDragValidationInfo(EPoseWatchManagerDropCompatibility::Incompatible, ValidationText);
+		}
+
+		FText ValidationText = FText::Format(LOCTEXT("MoveFolderIntoFolder", "Move {0} into {1}"), PayloadFolder->GetLabel(), TargetFolderLabel);
+		return FPoseWatchManagerDragValidationInfo(EPoseWatchManagerDropCompatibility::Compatible, ValidationText);
 	}
 	
 	return FPoseWatchManagerDragValidationInfo(EPoseWatchManagerDropCompatibility::Incompatible, FText());
@@ -189,48 +182,4 @@ FReply FPoseWatchManagerDefaultMode::OnDragOverItem(const FDragDropEvent& Event,
 	return FReply::Handled();
 }
 
-TSharedPtr<SWidget> FPoseWatchManagerDefaultMode::CreateContextMenu()
-{
-	FPoseWatchManagerTreeItemPtr SelectedItem = PoseWatchManager->GetSelection();
-	return SelectedItem ? SelectedItem->CreateContextMenu() : SNullWidget::NullWidget;
-}
-
-FReply FPoseWatchManagerDefaultMode::OnKeyDown(const FKeyEvent& InKeyEvent)
-{
-	const FPoseWatchManagerTreeItemPtr& Selection = PoseWatchManager->GetSelection();
-
-	if (InKeyEvent.GetKey() == EKeys::F2)
-	{
-		if (Selection)
-		{
-			//FPoseWatchManagerTreeItemPtr ItemToRename = Selection.Get();
-
-			if (Selection.IsValid())
-			{
-				PoseWatchManager->SetPendingRenameItem(Selection);
-				PoseWatchManager->ScrollItemIntoView(Selection);
-			}
-
-			return FReply::Handled();
-		}
-	}
-
-	else if (InKeyEvent.GetKey() == EKeys::F5)
-	{
-		PoseWatchManager->FullRefresh();
-		return FReply::Handled();
-	}
-
-	else if (InKeyEvent.GetKey() == EKeys::Delete || InKeyEvent.GetKey() == EKeys::BackSpace)
-	{
-		if (Selection)
-		{
-			Selection->OnRemoved();
-			return FReply::Handled();
-		}
-	}
-
-	return FReply::Unhandled();
-}
-
 #undef LOCTEXT_NAMESPACE
\ No newline at end of file
diff --git a/Engine/Source/Editor/Persona/Private/PoseWatchManagerDefaultMode.h b/Engine/Source/Editor/Persona/Private/PoseWatchManagerDefaultMode.h
--- a/Engine/Source/Editor/Persona/Private/PoseWatchManagerDefaultMode.h
+++ b/Engine/Source/Editor/Persona/Private/PoseWatchManagerDefaultMode.h
@@ -15,10 +15,6 @@ public:
 
 	void Rebuild();
 
-	FReply OnKeyDown(const FKeyEvent& InKeyEvent);
-
-	TSharedPtr<SWidget> CreateContextMenu();
-
 	bool ParseDragDrop(FPoseWatchManagerDragDropPayload& OutPayload, const FDragDropOperation& Operation) const;
 	
 	FPoseWatchManagerDragValidationInfo ValidateDrop(const IPoseWatchManagerTreeItem& DropTarget, const FPoseWatchManagerDragDropPayload& Payload) const;
diff --git a/Engine/Source/Editor/Persona/Private/PoseWatchManagerPoseWatchTreeItem.cpp b/Engine/Source/Editor/Persona/Private/PoseWatchManagerPoseWatchTreeItem.cpp
--- a/Engine/Source/Editor/Persona/Private/PoseWatchManagerPoseWatchTreeItem.cpp
+++ b/Engine/Source/Editor/Persona/Private/PoseWatchManagerPoseWatchTreeItem.cpp
@@ -107,7 +107,7 @@ private:
 	{
 		if (!TreeItemPtr.Pin()->IsEnabled())
 		{
-			return LOCTEXT("PoseWatchDisabled", "This pose watch is disabled because it is not connected to the 'Output Pose' node");
+			return LOCTEXT("PoseWatchDisabled", "This pose watch is disabled because it is not being evaluated");
 		}
 		return GetDisplayText();
 	}
diff --git a/Engine/Source/Editor/Persona/Private/PoseWatchManagerStandaloneTypes.cpp b/Engine/Source/Editor/Persona/Private/PoseWatchManagerStandaloneTypes.cpp
--- a/Engine/Source/Editor/Persona/Private/PoseWatchManagerStandaloneTypes.cpp
+++ b/Engine/Source/Editor/Persona/Private/PoseWatchManagerStandaloneTypes.cpp
@@ -10,13 +10,13 @@
 
 #define LOCTEXT_NAMESPACE "PoseWatchManagerStandaloneTypes"
 
-const FLinearColor FPoseWatchManagerCommonLabelData::DarkColor(0.15f, 0.15f, 0.15f);
+const FLinearColor FPoseWatchManagerCommonLabelData::DisabledColor(0.15f, 0.15f, 0.15f);
 
 TOptional<FLinearColor> FPoseWatchManagerCommonLabelData::GetForegroundColor(const IPoseWatchManagerTreeItem& TreeItem) const
 {
 	if (!TreeItem.IsValid() || !TreeItem.IsEnabled())
 	{
-		return DarkColor;
+		return DisabledColor;
 	}
 
 	// Darken items that aren't suitable targets for an active drag and drop action
@@ -28,7 +28,7 @@ TOptional<FLinearColor> FPoseWatchManagerCommonLabelData::GetForegroundColor(con
 		const auto Outliner = WeakPoseWatchManager.Pin();
 		if (Outliner->GetMode()->ParseDragDrop(DraggedObjects, *DragDropOp) && !Outliner->GetMode()->ValidateDrop(TreeItem, DraggedObjects).IsValid())
 		{
-			return DarkColor;
+			return DisabledColor;
 		}
 	}
 
diff --git a/Engine/Source/Editor/Persona/Private/PoseWatchManagerStandaloneTypes.h b/Engine/Source/Editor/Persona/Private/PoseWatchManagerStandaloneTypes.h
--- a/Engine/Source/Editor/Persona/Private/PoseWatchManagerStandaloneTypes.h
+++ b/Engine/Source/Editor/Persona/Private/PoseWatchManagerStandaloneTypes.h
@@ -6,7 +6,9 @@
 #include "Misc/Paths.h"
 #include "UObject/ObjectKey.h"
 #include "Templates/MaxSizeof.h"
-#include "PoseWatchManagerFwd.h"
+
+struct IPoseWatchManagerTreeItem;
+class IPoseWatchManager;
 
 typedef TSharedPtr<IPoseWatchManagerTreeItem> FPoseWatchManagerTreeItemPtr;
 typedef TSharedRef<IPoseWatchManagerTreeItem> FPoseWatchManagerTreeItemRef;
@@ -22,7 +24,7 @@ enum EPoseWatchTreeItemType
 struct PERSONA_API FPoseWatchManagerCommonLabelData
 {
 	TWeakPtr<IPoseWatchManager> WeakPoseWatchManager;
-	static const FLinearColor DarkColor;
+	static const FLinearColor DisabledColor;
 
 	TOptional<FLinearColor> GetForegroundColor(const IPoseWatchManagerTreeItem& TreeItem) const;
 
diff --git a/Engine/Source/Editor/Persona/Private/SPoseWatchManager.cpp b/Engine/Source/Editor/Persona/Private/SPoseWatchManager.cpp
--- a/Engine/Source/Editor/Persona/Private/SPoseWatchManager.cpp
+++ b/Engine/Source/Editor/Persona/Private/SPoseWatchManager.cpp
@@ -46,6 +46,7 @@
 
 #include "Framework/Notifications/NotificationManager.h"
 #include "Widgets/Notifications/SNotificationList.h"
+#include "Framework/Commands/GenericCommands.h"
 
 #define LOCTEXT_NAMESPACE "SPoseWatchManager"
 
@@ -81,6 +82,8 @@ void SPoseWatchManager::Construct(const FArguments& InArgs, const FPoseWatchMana
 	SortByColumn = FPoseWatchManagerBuiltInColumnTypes::Label();
 	SortMode = EColumnSortMode::Ascending;
 
+	FGenericCommands::Register();
+	BindCommands();
 
 	//Setup the search box
 	{
@@ -494,7 +497,30 @@ void SPoseWatchManager::PopulateSearchStrings(const IPoseWatchManagerTreeItem& I
 
 TSharedPtr<SWidget> SPoseWatchManager::OnOpenContextMenu()
 {
-	return Mode->CreateContextMenu();
+	FMenuBuilder MenuBuilder(true, CommandList);
+
+	MenuBuilder.BeginSection("Hierarchy", LOCTEXT("HierarchyMenuHeader", "Hierarchy"));
+	{
+		MenuBuilder.AddMenuEntry(
+			LOCTEXT("CreateFolder", "Create Folder"),
+			LOCTEXT("CreateFolderDescription", "Create a new folder"),
+			FSlateIcon(FEditorStyle::GetStyleSetName(), "SceneOutliner.NewFolderIcon"),
+			FUIAction(FExecuteAction::CreateSP(this, &SPoseWatchManager::CreateFolder))
+		);
+	}
+	MenuBuilder.EndSection();
+
+	if (PoseWatchManagerTreeView.Get()->GetNumItemsSelected() > 0)
+	{
+		MenuBuilder.BeginSection("Edit", LOCTEXT("EditMenuHeader", "Edit"));
+		{
+			MenuBuilder.AddMenuEntry(FGenericCommands::Get().Delete);
+			MenuBuilder.AddMenuEntry(FGenericCommands::Get().Rename);
+		}
+		MenuBuilder.EndSection();
+	}
+
+	return MenuBuilder.MakeWidget();
 }
 
 void SPoseWatchManager::Rename_Execute()
@@ -507,6 +533,16 @@ void SPoseWatchManager::Rename_Execute()
 	}
 }
 
+void SPoseWatchManager::Delete_Execute()
+{
+	if (PoseWatchManagerTreeView.Get()->GetNumItemsSelected() != 1)
+	{
+		return;
+	}
+	FPoseWatchManagerTreeItemPtr SelectedItem = PoseWatchManagerTreeView.Get()->GetSelectedItems()[0];
+	SelectedItem->OnRemoved();
+}
+
 void SPoseWatchManager::SetColumnVisibility(FName ColumnId, bool bIsVisible)
 {
 	if (Columns.Contains(ColumnId))
@@ -582,6 +618,21 @@ void SPoseWatchManager::CreateFolder()
 {
 	check(AnimBlueprint);
 	UPoseWatchFolder* NewPoseWatchFolder = NewObject<UPoseWatchFolder>(AnimBlueprint);
+
+	if (PoseWatchManagerTreeView.Get()->GetNumItemsSelected() == 1)
+	{
+		// Create the folder at the same level as the selected item
+		FPoseWatchManagerTreeItemPtr SelectedItem = PoseWatchManagerTreeView.Get()->GetSelectedItems()[0];
+		if (FPoseWatchManagerFolderTreeItem* SelectedFolderItem = SelectedItem->CastTo<FPoseWatchManagerFolderTreeItem>())
+		{
+			NewPoseWatchFolder->SetParent(SelectedFolderItem->PoseWatchFolder.Get(), true /* bForce */);
+		}
+		else if (FPoseWatchManagerPoseWatchTreeItem* SelectedPoseWatchItem = SelectedItem->CastTo<FPoseWatchManagerPoseWatchTreeItem>())
+		{
+			NewPoseWatchFolder->SetParent(SelectedPoseWatchItem->PoseWatch.Get()->GetParent(), true /* bForce */);
+		}
+	}
+
 	NewPoseWatchFolder->SetUniqueDefaultLabel();
 	AnimBlueprint->PoseWatchFolders.Add(NewPoseWatchFolder);
 
@@ -753,7 +804,7 @@ bool SPoseWatchManager::SupportsKeyboardFocus() const
 
 FReply SPoseWatchManager::OnKeyDown(const FGeometry& MyGeometry, const FKeyEvent& InKeyEvent)
 {
-	return Mode->OnKeyDown(InKeyEvent);
+	return CommandList->ProcessCommandBindings(InKeyEvent) ? FReply::Handled() : FReply::Unhandled();
 }
 
 void SPoseWatchManager::Tick(const FGeometry& AllottedGeometry, const double InCurrentTime, const float InDeltaTime)
@@ -903,4 +954,14 @@ void SPoseWatchManager::OnManagerTreeDoubleClick(FPoseWatchManagerTreeItemPtr Tr
 	}
 }
 
+void SPoseWatchManager::BindCommands()
+{
+	// This should not be called twice on the same instance
+	check(!CommandList.IsValid());
+	CommandList = MakeShareable(new FUICommandList);
+		
+	CommandList->MapAction(FGenericCommands::Get().Delete, FExecuteAction::CreateSP(this, &SPoseWatchManager::Delete_Execute));
+	CommandList->MapAction(FGenericCommands::Get().Rename, FExecuteAction::CreateSP(this, &SPoseWatchManager::Rename_Execute));
+}
+
 #undef LOCTEXT_NAMESPACE
diff --git a/Engine/Source/Editor/Persona/Private/SPoseWatchManager.h b/Engine/Source/Editor/Persona/Private/SPoseWatchManager.h
--- a/Engine/Source/Editor/Persona/Private/SPoseWatchManager.h
+++ b/Engine/Source/Editor/Persona/Private/SPoseWatchManager.h
@@ -230,6 +230,8 @@ public:
 
 	virtual void Rename_Execute() override;
 
+	virtual void Delete_Execute();
+
 	/** Get the columns to be displayed in this manager */
 	const TMap<FName, TSharedPtr<IPoseWatchManagerColumn>>& GetColumns() const
 	{
@@ -253,16 +255,14 @@ private:
 	/** Map of columns that are shown on this manager. */
 	TMap<FName, TSharedPtr<IPoseWatchManagerColumn>> Columns;
 
-	TArray<FName> HiddenColumnsList;
-
 	/** Set up the columns required for this manager */
 	void SetupColumns(SHeaderRow& HeaderRow);
 
-	/** Refresh the scene manager for when a colum was added or removed */
+	/** Refresh the scene manager for when a column was added or removed */
 	void RefreshColumns();
 
 	/** Populates OutSearchStrings with the strings associated with TreeItem that should be used in searching */
-	void PopulateSearchStrings(const IPoseWatchManagerTreeItem& TreeItem, OUT TArray< FString >& OutSearchStrings) const;
+	void PopulateSearchStrings(const IPoseWatchManagerTreeItem& Item, TArray< FString >& OutSearchStrings) const;
 
 public:
 	/** Scroll the specified item into view */
@@ -309,15 +309,19 @@ private:
 	/** Returns the selection mode*/
 	ESelectionMode::Type GetSelectionMode() const;
 
-public:
+private:
+	/** Called when the user right clicks in the tree view */
 	TSharedPtr<SWidget> OnOpenContextMenu();
 
-private:
 	/** Called when the user has clicked the button to add a new folder */
 	FReply OnCreateFolderClicked();
 
+	/** Requests a new folder be created, confirmation received with OnHierarchyChangedEvent */
 	void CreateFolder();
 
+	/** Binds our UI commands to delegates. */
+	void BindCommands();
+
 private:
 	/** Context menu opening delegate provided by the client */
 	FOnContextMenuOpening OnContextMenuOpening;
@@ -337,15 +341,10 @@ private:
 	/** Pending tree items that are yet to be added the tree */
 	FPoseWatchManagerTreeItemMap PendingTreeItemMap;
 
-	/** Folders pending selection */
-	TArray<FName> PendingFoldersSelect;
-
 	/** Root level tree items */
 	TArray<FPoseWatchManagerTreeItemPtr> RootTreeItems;
 
-	/** The button that displays view options */
-	TSharedPtr<SComboButton> ViewOptionsComboButton;
-
+	TSharedPtr<FUICommandList> CommandList;
 private:
 	/** true if the manager  needs to be repopulated at the next appropriate opportunity */
 	uint8 bNeedsRefresh : 1;
@@ -379,8 +378,6 @@ private:
 	/** The tree item that is currently pending a rename */
 	TWeakPtr<IPoseWatchManagerTreeItem> PendingRenameItem;
 
-	TMap<FName, const FSlateBrush*> CachedIcons;
-
 	TUniquePtr<FPoseWatchManagerDefaultHierarchy> Hierarchy;
 
 	PoseWatchManager::FTreeItemPtrEvent OnDoubleClickOnTreeEvent;
diff --git a/Engine/Source/Editor/UnrealEd/Private/AnimationEditorUtils.cpp b/Engine/Source/Editor/UnrealEd/Private/AnimationEditorUtils.cpp
--- a/Engine/Source/Editor/UnrealEd/Private/AnimationEditorUtils.cpp
+++ b/Engine/Source/Editor/UnrealEd/Private/AnimationEditorUtils.cpp
@@ -853,8 +853,7 @@ namespace AnimationEditorUtils
 					{
 						// Find the insertion point from the debugging data
 						int32 LinkID = AnimBPGenClass->GetLinkIDForNode<FAnimNode_Base>(TargetNode);
-						PoseWatch->SetIsEnabled(LinkID != INDEX_NONE);
-						AnimBPGenClass->GetAnimBlueprintDebugData().AddPoseWatch(LinkID, PoseWatch->GetColor(), PoseWatch->GetIsVisible());
+						AnimBPGenClass->GetAnimBlueprintDebugData().AddPoseWatch(LinkID, PoseWatch);
 						OnPoseWatchesChangedDelegate.Broadcast(AnimBlueprint, TargetNode);
 					}
 				}
@@ -891,12 +890,11 @@ namespace AnimationEditorUtils
 #endif
 	}
 
-	UPoseWatch* MakePoseWatchForNode(UAnimBlueprint* AnimBlueprint, UEdGraphNode* Node, FColor PoseWatchColour)
+	UPoseWatch* MakePoseWatchForNode(UAnimBlueprint* AnimBlueprint, UEdGraphNode* Node)
 	{
 #if WITH_EDITORONLY_DATA
 		UPoseWatch* NewPoseWatch = NewObject<UPoseWatch>(AnimBlueprint);
 		NewPoseWatch->Node = Node;
-		NewPoseWatch->SetColor(PoseWatchColour);
 		NewPoseWatch->SetUniqueDefaultLabel();
 		AnimBlueprint->PoseWatches.Add(NewPoseWatch);
 		SetPoseWatch(NewPoseWatch, AnimBlueprint);
@@ -956,25 +954,6 @@ namespace AnimationEditorUtils
 		return OnPoseWatchesChangedDelegate;
 	}
 
-	TArrayView<const FColor> GetPoseWatchColorPalette()
-	{
-		static const FColor PoseWatchColors[] = { FColor::Red, FColor::Green, FColor::Blue, FColor::Cyan, FColor::Orange, FColor::Purple, FColor::Yellow, FColor::Black };
-		size_t NumColors = sizeof(PoseWatchColors) / sizeof(PoseWatchColors[0]);
-		return MakeArrayView(PoseWatchColors, NumColors);
-	}
-
-	void UpdatePoseWatchColour(UPoseWatch* PoseWatch, FColor NewPoseWatchColour)
-	{
-#if WITH_EDITORONLY_DATA
-		UAnimBlueprintGeneratedClass* AnimBPGenClass = nullptr;
-		int32 LinkID = GetPoseWatchNodeLinkID(PoseWatch, AnimBPGenClass);
-		if (LinkID != INDEX_NONE)
-		{
-			AnimBPGenClass->GetAnimBlueprintDebugData().UpdatePoseWatchColour(LinkID, NewPoseWatchColour);
-		}
-#endif
-	}
-
 	int32 GetPoseWatchNodeLinkID(UPoseWatch* PoseWatch, UAnimBlueprintGeneratedClass*& OutAnimBPGenClass)
 	{
 #if WITH_EDITORONLY_DATA
@@ -996,18 +975,6 @@ namespace AnimationEditorUtils
 		return INDEX_NONE;
 	}
 
-	void UpdatePoseWatchVisibility(UPoseWatch* PoseWatch)
-	{
-#if WITH_EDITORONLY_DATA
-		UAnimBlueprintGeneratedClass* AnimBPGenClass = nullptr;
-		int32 LinkID = GetPoseWatchNodeLinkID(PoseWatch, AnimBPGenClass);
-		if (LinkID != INDEX_NONE)
-		{
-			AnimBPGenClass->GetAnimBlueprintDebugData().SetPoseWatchVisibility(LinkID, PoseWatch->GetIsVisible());
-		}
-#endif
-	}
-
 	void SetupDebugLinkedAnimInstances(UAnimBlueprint* InAnimBlueprint, UObject* InRootObjectBeingDebugged)
 	{
 		check(IsInGameThread());
diff --git a/Engine/Source/Editor/UnrealEd/Public/AnimationEditorUtils.h b/Engine/Source/Editor/UnrealEd/Public/AnimationEditorUtils.h
--- a/Engine/Source/Editor/UnrealEd/Public/AnimationEditorUtils.h
+++ b/Engine/Source/Editor/UnrealEd/Public/AnimationEditorUtils.h
@@ -197,12 +197,9 @@ namespace AnimationEditorUtils
 	int32 GetPoseWatchNodeLinkID(UPoseWatch* PoseWatch, OUT UAnimBlueprintGeneratedClass*& AnimBPGenClass);
 	UNREALED_API void SetPoseWatch(UPoseWatch* PoseWatch, UAnimBlueprint* AnimBlueprintIfKnown = nullptr);
 	UNREALED_API UPoseWatch* FindPoseWatchForNode(const UEdGraphNode* Node, UAnimBlueprint* AnimBlueprintIfKnown=nullptr);
-	UNREALED_API UPoseWatch* MakePoseWatchForNode(UAnimBlueprint* AnimBlueprint, UEdGraphNode* Node, FColor PoseWatchColour);
+	UNREALED_API UPoseWatch* MakePoseWatchForNode(UAnimBlueprint* AnimBlueprint, UEdGraphNode* Node);
 	UNREALED_API void RemovePoseWatch(UPoseWatch* PoseWatch, UAnimBlueprint* AnimBlueprintIfKnown=nullptr);
 	UNREALED_API void RemovePoseWatchFromNode(UEdGraphNode* Node, UAnimBlueprint* AnimBlueprint);
-	UNREALED_API void UpdatePoseWatchColour(UPoseWatch* PoseWatch, FColor NewPoseWatchColour);
-	UNREALED_API void UpdatePoseWatchVisibility(UPoseWatch* PoseWatch);
-	UNREALED_API TArrayView<const FColor> GetPoseWatchColorPalette();
 
 	// Delegate fired when a pose watch is added or removed
 	DECLARE_MULTICAST_DELEGATE_TwoParams(FOnPoseWatchesChanged, UAnimBlueprint* /*InAnimBlueprint*/, UEdGraphNode* /*InNode*/);
diff --git a/Engine/Source/Runtime/Engine/Classes/Animation/AnimBlueprintGeneratedClass.h b/Engine/Source/Runtime/Engine/Classes/Animation/AnimBlueprintGeneratedClass.h
--- a/Engine/Source/Runtime/Engine/Classes/Animation/AnimBlueprintGeneratedClass.h
+++ b/Engine/Source/Runtime/Engine/Classes/Animation/AnimBlueprintGeneratedClass.h
@@ -25,6 +25,7 @@ class UAnimStateAliasNode;
 class UAnimStateTransitionNode;
 class UEdGraph;
 class USkeleton;
+class UPoseWatch;
 struct FAnimSubsystem;
 struct FAnimSubsystemInstance;
 struct FPropertyAccessLibrary;
@@ -303,10 +304,8 @@ public:
 	void RecordSequencePlayer(int32 InNodeID, float InPosition, float InLength, int32 InFrameCount);
 	void RecordBlendSpacePlayer(int32 InNodeID, const UBlendSpace* InBlendSpace, const FVector& InPosition, const FVector& InFilteredPosition);
 
-	void AddPoseWatch(int32 NodeID, FColor Color, bool bIsVisible);
+	void AddPoseWatch(int32 NodeID, UPoseWatch*);
 	void RemovePoseWatch(int32 NodeID);
-	void UpdatePoseWatchColour(int32 NodeID, FColor Color);
-	void SetPoseWatchVisibility(int32 NodeID, bool bIsVisible);
 
 	TArrayView<const FName> GetNodeAttributes(TWeakObjectPtr<UAnimGraphNode_Base> InAnimGraphNode) const;
 #endif
diff --git a/Engine/Source/Runtime/Engine/Classes/Engine/PoseWatch.h b/Engine/Source/Runtime/Engine/Classes/Engine/PoseWatch.h
--- a/Engine/Source/Runtime/Engine/Classes/Engine/PoseWatch.h
+++ b/Engine/Source/Runtime/Engine/Classes/Engine/PoseWatch.h
@@ -15,10 +15,9 @@ struct FAnimNodePoseWatch
 {
 	// Object (anim instance) that this pose came from
 	TWeakObjectPtr<const UObject>	Object;
+	TWeakObjectPtr<UPoseWatch>		PoseWatch;
 	TSharedPtr<FCompactHeapPose>	PoseInfo;
-	FColor							PoseDrawColour;
 	int32							NodeID;
-	bool							bIsHidden;
 };
 
 namespace PoseWatchUtil
@@ -29,9 +28,16 @@ namespace PoseWatchUtil
 	/** Gets all pose watches folders that are parented to Folder, if Folder is nullptr then gets orphans */
 	TSet<UPoseWatchFolder*> GetChildrenPoseWatchFoldersOf(const UPoseWatchFolder* Folder, const UAnimBlueprint* AnimBlueprint);
 
+	/** Returns a UPoseWatch/UPoseWatchFolder from PoseWatchCollection that is inside InFolder has has label Label, if one exists */
+	template <typename T>
+	T* FindInFolderInCollection(const FName& Label, UPoseWatchFolder* InFolder, const TArray<TObjectPtr<T>>& PoseWatchCollection);
+
 	/** Finds unique name for a pose watch or folder inside of InParent */
 	template <typename T>
-	FText FindUniqueNameInFolder(UPoseWatchFolder* InParent, const T* Item, TArray<TObjectPtr<T>> Collection);
+	FText FindUniqueNameInFolder(UPoseWatchFolder* InParent, const T* Item, const TArray<TObjectPtr<T>>& Collection);
+
+	/** Returns a color from a palette of predefined colors in round-robin order */
+	FColor ChoosePoseWatchColor();
 }
 
 UCLASS()
diff --git a/Engine/Source/Runtime/Engine/Private/Animation/AnimBlueprintGeneratedClass.cpp b/Engine/Source/Runtime/Engine/Private/Animation/AnimBlueprintGeneratedClass.cpp
--- a/Engine/Source/Runtime/Engine/Private/Animation/AnimBlueprintGeneratedClass.cpp
+++ b/Engine/Source/Runtime/Engine/Private/Animation/AnimBlueprintGeneratedClass.cpp
@@ -20,6 +20,7 @@
 #include "Animation/AnimSubsystem.h"
 #include "EdGraph/EdGraphNode.h"
 #include "Animation/AnimSubsystemInstance.h"
+#include "Engine/PoseWatch.h"
 
 /////////////////////////////////////////////////////
 // FStateMachineDebugData
@@ -228,14 +229,12 @@ void FAnimBlueprintDebugData::RecordBlendSpacePlayer(int32 InNodeID, const UBlen
 	BlendSpacePlayerRecordsThisFrame.Emplace(InNodeID, InBlendSpace, InPosition, InFilteredPosition);
 }
 
-void FAnimBlueprintDebugData::AddPoseWatch(int32 NodeID, FColor Color, bool bIsVisible)
+void FAnimBlueprintDebugData::AddPoseWatch(int32 NodeID, UPoseWatch* InPoseWatch)
 {
 	for (FAnimNodePoseWatch& PoseWatch : AnimNodePoseWatch)
 	{
 		if (PoseWatch.NodeID == NodeID)
 		{
-			PoseWatch.PoseDrawColour = Color;
-			PoseWatch.bIsHidden = !bIsVisible;
 			return;
 		}
 	}
@@ -244,8 +243,7 @@ void FAnimBlueprintDebugData::AddPoseWatch(int32 NodeID, FColor Color, bool bIsV
 	AnimNodePoseWatch.Add(FAnimNodePoseWatch());
 	FAnimNodePoseWatch& NewAnimNodePoseWatch = AnimNodePoseWatch.Last();
 	NewAnimNodePoseWatch.NodeID = NodeID;
-	NewAnimNodePoseWatch.PoseDrawColour = Color;
-	NewAnimNodePoseWatch.bIsHidden = !bIsVisible;
+	NewAnimNodePoseWatch.PoseWatch = InPoseWatch;
 	NewAnimNodePoseWatch.PoseInfo = MakeShareable(new FCompactHeapPose());
 	NewAnimNodePoseWatch.Object = nullptr;
 }
@@ -262,30 +260,6 @@ void FAnimBlueprintDebugData::RemovePoseWatch(int32 NodeID)
 	}
 }
 
-void FAnimBlueprintDebugData::SetPoseWatchVisibility(int32 NodeID, bool bIsVisible)
-{
-	for (FAnimNodePoseWatch& PoseWatch : AnimNodePoseWatch)
-	{
-		if (PoseWatch.NodeID == NodeID)
-		{
-			PoseWatch.bIsHidden = !bIsVisible;
-			return;
-		}
-	}
-}
-
-void FAnimBlueprintDebugData::UpdatePoseWatchColour(int32 NodeID, FColor Color)
-{
-	for (FAnimNodePoseWatch& PoseWatch : AnimNodePoseWatch)
-	{
-		if (PoseWatch.NodeID == NodeID)
-		{
-			PoseWatch.PoseDrawColour = Color;
-			return;
-		}
-	}
-}
-
 TArrayView<const FName> FAnimBlueprintDebugData::GetNodeAttributes(TWeakObjectPtr<UAnimGraphNode_Base> InAnimGraphNode) const
 {
 	const TArray<FName>* Attributes = NodeAttributes.Find(InAnimGraphNode);
diff --git a/Engine/Source/Runtime/Engine/Private/Animation/AnimInstanceProxy.cpp b/Engine/Source/Runtime/Engine/Private/Animation/AnimInstanceProxy.cpp
--- a/Engine/Source/Runtime/Engine/Private/Animation/AnimInstanceProxy.cpp
+++ b/Engine/Source/Runtime/Engine/Private/Animation/AnimInstanceProxy.cpp
@@ -425,6 +425,10 @@ void FAnimInstanceProxy::PreUpdate(UAnimInstance* InAnimInstance, float DeltaSec
 		if (DebugData)
 		{
 			PoseWatchEntriesForThisFrame = DebugData->AnimNodePoseWatch;
+			for (FAnimNodePoseWatch& PoseWatch : PoseWatchEntriesForThisFrame)
+			{
+				PoseWatch.PoseWatch.Get()->SetIsEnabled(false);
+			}
 		}
 	}
 #endif
@@ -2961,6 +2965,7 @@ void FAnimInstanceProxy::RegisterWatchedPose(const FCompactPose& Pose, int32 Lin
 			{
 				PoseWatch.PoseInfo->CopyBonesFrom(Pose);
 				PoseWatch.Object = GetAnimInstanceObject();
+				PoseWatch.PoseWatch->SetIsEnabled(true);
 				break;
 			}
 		}
@@ -2979,6 +2984,7 @@ void FAnimInstanceProxy::RegisterWatchedPose(const FCSPose<FCompactPose>& Pose,
 				FCSPose<FCompactPose>::ConvertComponentPosesToLocalPoses(Pose, TempPose);
 				PoseWatch.PoseInfo->CopyBonesFrom(TempPose);
 				PoseWatch.Object = GetAnimInstanceObject();
+				PoseWatch.PoseWatch->SetIsEnabled(true);
 				break;
 			}
 		}
diff --git a/Engine/Source/Runtime/Engine/Private/Engine/PoseWatch.cpp b/Engine/Source/Runtime/Engine/Private/Engine/PoseWatch.cpp
--- a/Engine/Source/Runtime/Engine/Private/Engine/PoseWatch.cpp
+++ b/Engine/Source/Runtime/Engine/Private/Engine/PoseWatch.cpp
@@ -22,6 +22,7 @@ UPoseWatch::UPoseWatch(const FObjectInitializer& ObjectInitializer)
 {
 #if WITH_EDITOR
 	Label = GetDefaultLabel();
+	SetColor(PoseWatchUtil::ChoosePoseWatchColor());
 #endif
 }
 
@@ -61,33 +62,41 @@ TSet<UPoseWatchFolder*> PoseWatchUtil::GetChildrenPoseWatchFoldersOf(const UPose
 }
 
 template <typename T>
-FText PoseWatchUtil::FindUniqueNameInFolder(UPoseWatchFolder* InParent, const T* Item, TArray< TObjectPtr<T>> Collection)
+T* PoseWatchUtil::FindInFolderInCollection(const FName& Label, UPoseWatchFolder* InFolder, const TArray<TObjectPtr<T>>& PoseWatchCollection)
 {
-	static const FTextFormat LabelFormat = FTextFormat::FromString("{0}{1}");
-
-	uint32 Counter = 1;
-	bool bNameClashes = false;
-	FText NewLabel = Item->GetDefaultLabel();
-	UAnimBlueprint* AnimBlueprint = CastChecked<UAnimBlueprint>(Item->GetOuter());
-
-	do
+	for (const TObjectPtr<T>& PoseWatchItem : PoseWatchCollection)
 	{
-		bNameClashes = false;
-		NewLabel = FText::Format(LabelFormat, Item->GetLabel(), FText::AsNumber(Counter));
-		for (T* SomeItem : Collection)
+		if (PoseWatchItem->IsIn(InFolder) && PoseWatchItem->GetLabel().ToString().Equals(Label.ToString()))
 		{
-			if (SomeItem->IsIn(InParent) && SomeItem->GetLabel().ToString().Equals(NewLabel.ToString()))
-			{
-				++Counter;
-				bNameClashes = true;
-				break;
-			}
+			return PoseWatchItem;
 		}
-	} while (bNameClashes);
+	}
+	return nullptr;
+}
+
+template <typename T>
+FText PoseWatchUtil::FindUniqueNameInFolder(UPoseWatchFolder* InParent, const T* Item, const TArray<TObjectPtr<T>>& Collection)
+{
+	FName NewName(Item->GetLabel().ToString());
+	FString BaseLabel = Item->GetLabel().ToString();
+	int32 Index = 0;
+	T* ConflictingItem = nullptr;
+	do
+	{
+		++Index;
+		NewName = *FString::Printf(TEXT("%s%d"), *BaseLabel, Index);
+		ConflictingItem = FindInFolderInCollection(NewName, InParent, Collection);
+	} while (ConflictingItem != nullptr && ConflictingItem != Item);
 
-	return NewLabel;
+	return FText::FromName(NewName);
 }
 
+FColor PoseWatchUtil::ChoosePoseWatchColor()
+{
+	return FColor::MakeRandomColor();
+}
+
+
 // UPoseWatchFolder
 
 const FText UPoseWatchFolder::GetPath() const
@@ -456,7 +465,6 @@ bool UPoseWatch::SetLabel(const FText& InLabel)
 void UPoseWatch::SetIsVisible(bool bInIsVisible)
 {
 	bIsVisible = bInIsVisible;
-	AnimationEditorUtils::UpdatePoseWatchVisibility(this);
 
 	if (Parent.IsValid())
 	{
@@ -467,7 +475,6 @@ void UPoseWatch::SetIsVisible(bool bInIsVisible)
 void UPoseWatch::SetColor(const FColor& InColor)
 {
 	Color = InColor;
-	AnimationEditorUtils::UpdatePoseWatchColour(this, Color);
 }
 
 void UPoseWatch::SetShouldDeleteOnDeselect(const bool bInDeleteOnDeselection)
