Non-unity fix

diff --git a/Engine/Source/Editor/Persona/Private/PoseWatchManagerFolderTreeItem.h b/Engine/Source/Editor/Persona/Private/PoseWatchManagerFolderTreeItem.h
--- a/Engine/Source/Editor/Persona/Private/PoseWatchManagerFolderTreeItem.h
+++ b/Engine/Source/Editor/Persona/Private/PoseWatchManagerFolderTreeItem.h
@@ -7,6 +7,7 @@
 #include "IPoseWatchManagerTreeItem.h"
 
 class UPoseWatch;
+class UPoseWatchFolder;
 
 struct FPoseWatchManagerFolderTreeItem : IPoseWatchManagerTreeItem
 {
