Fixed ensure when opening some linked animation blueprints

Linked anim BPs can be debugged, but logic to setup the debugged object based on skeletal mesh made some assumptions that only the main anim instance would be the class being set. Expanding the check here for that fallback logic application to cast a wider net than just the preview instance fixes an ensure.
Found while investgating UE-110612.

diff --git a/Engine/Source/Editor/Persona/Private/AnimationEditorPreviewScene.cpp b/Engine/Source/Editor/Persona/Private/AnimationEditorPreviewScene.cpp
--- a/Engine/Source/Editor/Persona/Private/AnimationEditorPreviewScene.cpp
+++ b/Engine/Source/Editor/Persona/Private/AnimationEditorPreviewScene.cpp
@@ -276,9 +276,9 @@ void FAnimationEditorPreviewScene::SetPreviewMeshInternal(USkeletalMesh* NewPrev
 
 	// Setting the skeletal mesh to in the PreviewScene can change AnimScriptInstance so we must re register it
 	// with the AnimBlueprint
-	if (DebuggedSkeletalMeshComponent && DebuggedSkeletalMeshComponent->GetAnimInstance() && !DebuggedSkeletalMeshComponent->GetAnimInstance()->IsA<UAnimPreviewInstance>())
+	UAnimBlueprint* SourceBlueprint = PersonaToolkit.Pin()->GetAnimBlueprint();
+	if (DebuggedSkeletalMeshComponent && DebuggedSkeletalMeshComponent->GetAnimInstance() && DebuggedSkeletalMeshComponent->GetAnimInstance()->IsA(SourceBlueprint->GeneratedClass))
 	{
-		UAnimBlueprint* SourceBlueprint = PersonaToolkit.Pin()->GetAnimBlueprint();
 		PersonaUtils::SetObjectBeingDebugged(SourceBlueprint, DebuggedSkeletalMeshComponent->GetAnimInstance());
 	}
 
