Added copy and paste for blend samples

diff --git a/Engine/Source/Editor/Persona/Private/SAnimationBlendSpaceGridWidget.cpp b/Engine/Source/Editor/Persona/Private/SAnimationBlendSpaceGridWidget.cpp
--- a/Engine/Source/Editor/Persona/Private/SAnimationBlendSpaceGridWidget.cpp
+++ b/Engine/Source/Editor/Persona/Private/SAnimationBlendSpaceGridWidget.cpp
@@ -35,6 +35,18 @@
 #include "Framework/MultiBox/MultiBoxBuilder.h"
 #include "Styling/StyleColors.h"
 
+#include "JsonObjectConverter.h"
+#include "HAL/PlatformApplicationMisc.h"
+#include "AnimGraphNode_BlendSpaceGraphBase.h"
+#include "BlendSpaceGraph.h"
+#include "AnimationBlendSpaceSampleGraph.h"
+#include "EdGraphUtilities.h"
+#include "AnimGraphNode_BlendSpaceSampleResult.h"
+#include "ScopedTransaction.h"
+#include "Framework/Commands/GenericCommands.h"
+#include "Framework/Commands/UICommandList.h"
+
+
 #define LOCTEXT_NAMESPACE "SAnimationBlendSpaceGridWidget"
 
 // Draws additional data on the triangulation to help debugging
@@ -50,6 +62,10 @@ static const float CriticalTriangulationAngle = 4.0f;
 // Flag any triangle that has a smaller area than this (normalized units)
 static const float CriticalTriangulationArea = 5e-4f;
 
+// Identifies the clipboard contents as being a blend sample
+static const FString BlendSampleClipboardHeaderAsset = TEXT("COPY_BLENDSAMPLE_ASSET");
+static const FString BlendSampleClipboardHeaderGraph = TEXT("COPY_BLENDSAMPLE_GRAPH");
+
 //======================================================================================================================
 // Paint a filled triangle
 static void PaintTriangle(
@@ -227,6 +243,10 @@ void SBlendSpaceGridWidget::Construct(const FArguments& InArgs)
 	bStretchToFit = true;
 	bShowAnimationNames = false;
 
+	// Register and bind all our menu commands
+	FGenericCommands::Register();
+	BindCommands();
+
 	InvalidSamplePositionDragDropText = FText::FromString(TEXT("Invalid Sample Position"));
 
 	// Retrieve UI color values
@@ -1333,6 +1353,11 @@ FReply SBlendSpaceGridWidget::OnKeyDown(const FGeometry& MyGeometry, const FKeyE
 {
 	if(!bReadOnly && BlendSpaceBase.IsSet())
 	{
+		if (UICommandList->ProcessCommandBindings(InKeyEvent))
+		{
+			return FReply::Handled();
+		}
+
 		// Start previewing when either one of the shift keys is pressed
 		if (IsHovered() && bMouseIsOverGeometry)
 		{
@@ -1387,23 +1412,6 @@ FReply SBlendSpaceGridWidget::OnKeyUp(const FGeometry& MyGeometry, const FKeyEve
 			return FReply::Handled();
 		}
 
-		// If delete is pressed and we currently have a sample selected remove it from the blendspace
-		if (InKeyEvent.GetKey() == EKeys::Delete)
-		{
-			if (SelectedSampleIndex != INDEX_NONE)
-			{
-				OnSampleRemoved.ExecuteIfBound(SelectedSampleIndex);
-		
-				if (SelectedSampleIndex == HighlightedSampleIndex)
-				{
-					HighlightedSampleIndex = INDEX_NONE;
-					ResetToolTip();
-				}
-
-				SelectedSampleIndex = INDEX_NONE;
-			}
-		}
-
 		// Pressing esc will remove the current key selection
 		if( InKeyEvent.GetKey() == EKeys::Escape)
 		{
@@ -1481,8 +1489,8 @@ void SBlendSpaceGridWidget::MakeViewContextMenuEntries(FMenuBuilder& InMenuBuild
 TSharedPtr<SWidget> SBlendSpaceGridWidget::CreateBlendSampleContextMenu()
 {
 	const bool bShouldCloseWindowAfterMenuSelection = true;
-	FMenuBuilder MenuBuilder(bShouldCloseWindowAfterMenuSelection, nullptr);
-	
+	FMenuBuilder MenuBuilder(bShouldCloseWindowAfterMenuSelection, UICommandList);
+
 	TSharedPtr<IStructureDetailsView> StructureDetailsView;
 	// Initialize details view
 	FDetailsViewArgs DetailsViewArgs;
@@ -1523,6 +1531,9 @@ TSharedPtr<SWidget> SBlendSpaceGridWidget::CreateBlendSampleContextMenu()
 	
 	MenuBuilder.BeginSection("Sample", LOCTEXT("SampleMenuHeader", "Sample"));
 	{
+		MenuBuilder.AddMenuEntry(FGenericCommands::Get().Cut);
+		MenuBuilder.AddMenuEntry(FGenericCommands::Get().Copy);
+		MenuBuilder.AddMenuEntry(FGenericCommands::Get().Delete);
 		MenuBuilder.AddWidget(StructureDetailsView->GetWidget().ToSharedRef(), FText::GetEmpty(), true);
 	}
 	MenuBuilder.EndSection();
@@ -1535,7 +1546,7 @@ TSharedPtr<SWidget> SBlendSpaceGridWidget::CreateBlendSampleContextMenu()
 TSharedPtr<SWidget> SBlendSpaceGridWidget::CreateNewBlendSampleContextMenu(const FVector2D& InMousePosition)
 {
 	const bool bShouldCloseWindowAfterMenuSelection = true;
-	FMenuBuilder MenuBuilder(bShouldCloseWindowAfterMenuSelection, nullptr);
+	FMenuBuilder MenuBuilder(bShouldCloseWindowAfterMenuSelection, UICommandList);
 
 	FVector NewSampleValue;
 	if(FSlateApplication::Get().GetModifierKeys().IsShiftDown())
@@ -1553,31 +1564,10 @@ TSharedPtr<SWidget> SBlendSpaceGridWidget::CreateNewBlendSampleContextMenu(const
 	{
 		MenuBuilder.BeginSection("Sample", LOCTEXT("SampleMenuHeader", "Sample"));
 		{
-			if(BlendSpace->IsAsset())
-			{
-				// Duplicate an existing sample if one was selected. Note that this could be enabled for the graph too,
-				// but we'd need a way of either duplicating the original graph, or referencing it (and choosing between
-				// the two).
-				if (SelectedSampleIndex != INDEX_NONE)
-				{
-					MenuBuilder.AddMenuEntry(
-						LOCTEXT("DuplicateSample", "Duplicate Sample Here"),
-						LOCTEXT("DuplicateSampleTooltip", "Duplicate the selected sample and place it in the blendspace at this location"),
-						FSlateIcon("EditorStyle", "Icons.Duplicate"),
-						FUIAction(
-							FExecuteAction::CreateLambda(
-								[this, NewSampleValue]()
-								{
-									if (OnSampleDuplicated.IsBound())
-									{
-										OnSampleDuplicated.Execute(SelectedSampleIndex, NewSampleValue, false);
-									}
-								})
-						)
-					);
-				}
-			}
-			else
+			MenuBuilder.AddMenuEntry(FGenericCommands::Get().Paste);
+			MenuBuilder.AddMenuEntry(FGenericCommands::Get().Delete);
+
+			if(!BlendSpace->IsAsset())
 			{
 				// Blend space graph - add a new graph sample
 				MenuBuilder.AddMenuEntry(
@@ -2689,4 +2679,214 @@ const bool SBlendSpaceGridWidget::IsPreviewing() const
 	return (bSamplePreviewing && !TargetPosition.IsSet()) || (TargetPosition.IsSet() && bIsManualPreviewing);
 }
 
+void SBlendSpaceGridWidget::BindCommands()
+{
+	// This should not be called twice on the same instance
+	check(!UICommandList.IsValid());
+	UICommandList = MakeShared<FUICommandList>();
+
+	UICommandList->MapAction(
+		FGenericCommands::Get().Cut,
+		FExecuteAction::CreateSP(this, &SBlendSpaceGridWidget::OnBlendSampleCut),
+		FCanExecuteAction::CreateSP(this, &SBlendSpaceGridWidget::CanBlendSampleCutCopy)
+	);
+	
+	UICommandList->MapAction(
+		FGenericCommands::Get().Copy,
+		FExecuteAction::CreateSP(this, &SBlendSpaceGridWidget::OnBlendSampleCopy),
+		FCanExecuteAction::CreateSP(this, &SBlendSpaceGridWidget::CanBlendSampleCutCopy)
+	);
+
+	UICommandList->MapAction(
+		FGenericCommands::Get().Paste,
+		FExecuteAction::CreateSP(this, &SBlendSpaceGridWidget::OnBlendSamplePaste),
+		FCanExecuteAction::CreateSP(this, &SBlendSpaceGridWidget::CanBlendSamplePaste)
+	);
+
+	UICommandList->MapAction(
+		FGenericCommands::Get().Delete,
+		FExecuteAction::CreateSP(this, &SBlendSpaceGridWidget::OnBlendSampleDelete),
+		FCanExecuteAction::CreateSP(this, &SBlendSpaceGridWidget::CanBlendSampleDelete)
+	);
+}
+
+bool SBlendSpaceGridWidget::CanBlendSampleCutCopy()
+{
+	return BlendSpaceBase.Get()->IsValidBlendSampleIndex(SelectedSampleIndex);
+}
+
+bool SBlendSpaceGridWidget::CanBlendSamplePaste()
+{
+	FString PastedText;
+	FPlatformApplicationMisc::ClipboardPaste(PastedText);
+	return PastedText.StartsWith(BlendSpaceBase.Get()->IsAsset() ? BlendSampleClipboardHeaderAsset : BlendSampleClipboardHeaderGraph);
+}
+
+bool SBlendSpaceGridWidget::CanBlendSampleDelete()
+{
+	return BlendSpaceBase.Get()->IsValidBlendSampleIndex(SelectedSampleIndex);
+}
+
+void SBlendSpaceGridWidget::OnBlendSampleDelete()
+{
+	if (SelectedSampleIndex != INDEX_NONE)
+	{
+		OnSampleRemoved.ExecuteIfBound(SelectedSampleIndex);
+
+		if (SelectedSampleIndex == HighlightedSampleIndex)
+		{
+			HighlightedSampleIndex = INDEX_NONE;
+			ResetToolTip();
+		}
+
+		SelectedSampleIndex = INDEX_NONE;
+	}
+}
+
+void SBlendSpaceGridWidget::OnBlendSampleCut()
+{
+	if (BlendSpaceBase.Get()->IsValidBlendSampleIndex(SelectedSampleIndex))
+	{
+		OnBlendSampleCopy();
+		OnSampleRemoved.ExecuteIfBound(SelectedSampleIndex);
+	}
+}
+
+void SBlendSpaceGridWidget::OnBlendSampleCopy()
+{
+	typedef TJsonWriter<TCHAR, TPrettyJsonPrintPolicy<TCHAR>> FStringWriter;
+	typedef TJsonWriterFactory<TCHAR, TPrettyJsonPrintPolicy<TCHAR>> FStringWriterFactory;
+	TSharedRef<FJsonObject> RootJsonObject = MakeShareable(new FJsonObject());
+
+	const UBlendSpace* BlendSpace = BlendSpaceBase.Get();
+	if (!BlendSpace->IsValidBlendSampleIndex(SelectedSampleIndex))
+	{
+		return;
+	}
+
+	if (BlendSpace->IsAsset())
+	{
+		FBlendSample SampleToCopy = BlendSpace->GetBlendSample(SelectedSampleIndex);
+		FJsonObjectConverter::UStructToJsonObject(FBlendSample::StaticStruct(), &SampleToCopy, RootJsonObject, 0 /* CheckFlags */, 0 /* SkipFlags */);
+	}
+	else
+	{
+		UBlendSpaceGraph* BlendSpaceGraph = CastChecked<UBlendSpaceGraph>(BlendSpace->GetOuter());
+		UAnimGraphNode_BlendSpaceGraphBase* BlendSpaceNode = CastChecked<UAnimGraphNode_BlendSpaceGraphBase>(BlendSpaceGraph->GetOuter());
+		UAnimationBlendSpaceSampleGraph* GraphSampleToCopy = CastChecked<UAnimationBlendSpaceSampleGraph>(BlendSpaceNode->GetGraphs()[SelectedSampleIndex]);
+
+		TSet<UObject*> NodesSet;
+		for (TObjectPtr<UEdGraphNode>& Node : GraphSampleToCopy->Nodes)
+		{
+			NodesSet.Add(Node);
+		}
+
+		FStringOutputDevice NodesString;
+		FEdGraphUtilities::ExportNodesToText(NodesSet, NodesString);
+		RootJsonObject->SetField(TEXT("Nodes"), MakeShared<FJsonValueString>(NodesString));
+
+		// Output Animation Pose node is not pasted as it already exists in the pasted graph
+		// Take note of the node that connects to it to reconstruct the link (if exists) on paste
+		UEdGraphPin* ResultNodePosePin = GraphSampleToCopy->ResultNode->FindPinChecked(TEXT("Result"), EEdGraphPinDirection::EGPD_Input);
+		if (ResultNodePosePin->LinkedTo.Num() == 1)
+		{
+			UEdGraphPin* ConnectedNodePosePin = ResultNodePosePin->LinkedTo[0];
+			UEdGraphNode* OwningNode = ConnectedNodePosePin->GetOwningNode();
+			RootJsonObject->SetField(TEXT("NodeConnectedToResult"), MakeShared<FJsonValueString>(OwningNode->NodeGuid.ToString()));
+		}
+	}
+
+	FString SerializedStr;
+	TSharedRef<FStringWriter> Writer = FStringWriterFactory::Create(&SerializedStr);
+	FJsonSerializer::Serialize(RootJsonObject, Writer);
+	SerializedStr = *FString::Printf(TEXT("%s%s"), BlendSpace->IsAsset() ? *BlendSampleClipboardHeaderAsset : *BlendSampleClipboardHeaderGraph, *SerializedStr);
+	FPlatformApplicationMisc::ClipboardCopy(*SerializedStr);
+}
+
+void SBlendSpaceGridWidget::OnBlendSamplePaste()
+{
+	FString PastedText;
+	FPlatformApplicationMisc::ClipboardPaste(PastedText);
+
+	check(OnSampleAdded.IsBound());
+
+	const FVector SampleValue = ScreenPositionToSampleValue(LocalMousePosition, true); // Paste the sample at the cursor's location
+	int32 NewSampleIndex = INDEX_NONE;
+
+	if (BlendSpaceBase.Get()->IsAsset())
+	{
+		check(PastedText.StartsWith(BlendSampleClipboardHeaderAsset));
+		PastedText.RightChopInline(BlendSampleClipboardHeaderAsset.Len());
+
+		TSharedPtr<FJsonObject> RootJsonObject;
+		TSharedRef<TJsonReader<>> Reader = TJsonReaderFactory<>::Create(PastedText);
+		if (!FJsonSerializer::Deserialize(Reader, RootJsonObject))
+		{
+			return;
+		}
+
+		FBlendSample SampleToPaste;
+		if (FJsonObjectConverter::JsonObjectToUStruct(RootJsonObject.ToSharedRef(), FBlendSample::StaticStruct(), &SampleToPaste, 0 /* CheckFlags */, 0 /* SkipFlags */))
+		{
+			const FScopedTransaction Transaction(LOCTEXT("PasteBlendSpaceSample", "Paste Blend Space Sample"));
+			NewSampleIndex = OnSampleAdded.Execute(SampleToPaste.Animation, SampleValue, false);
+		}
+	}
+	else
+	{
+		check(PastedText.StartsWith(BlendSampleClipboardHeaderGraph));
+		PastedText.RightChopInline(BlendSampleClipboardHeaderGraph.Len());
+
+		TSharedPtr<FJsonObject> RootJsonObject;
+		TSharedRef<TJsonReader<>> Reader = TJsonReaderFactory<>::Create(PastedText);
+		if (!FJsonSerializer::Deserialize(Reader, RootJsonObject))
+		{
+			return;
+		}
+
+		check(RootJsonObject->HasField(TEXT("Nodes")));
+
+		const FScopedTransaction Transaction(LOCTEXT("PasteBlendSpaceSample", "Paste Blend Space Sample"));
+
+		UBlendSpaceGraph* BlendSpaceGraph = CastChecked<UBlendSpaceGraph>(BlendSpaceBase.Get()->GetOuter());
+		UAnimGraphNode_BlendSpaceGraphBase* BlendSpaceNode = CastChecked<UAnimGraphNode_BlendSpaceGraphBase>(BlendSpaceGraph->GetOuter());
+		UAnimationBlendSpaceSampleGraph* DestinationGraph = nullptr;
+		NewSampleIndex = OnSampleAdded.Execute(nullptr, SampleValue, false);
+		DestinationGraph = CastChecked<UAnimationBlendSpaceSampleGraph>(BlendSpaceNode->GetGraphs()[NewSampleIndex]);
+
+		FString NodesSetString = RootJsonObject->GetStringField("Nodes");
+		TSet<UEdGraphNode*> ImportedNodes;
+		FEdGraphUtilities::ImportNodesFromText(DestinationGraph, NodesSetString, ImportedNodes);
+
+		// Reconstruct link to output, if exists
+		if (RootJsonObject->HasField(TEXT("NodeConnectedToResult")))
+		{
+			FString NodeConnectedToResult = RootJsonObject->GetStringField("NodeConnectedToResult");
+			FGuid ResultNodeGuid(NodeConnectedToResult);
+			
+			UEdGraphPin* ResultNodePosePin = DestinationGraph->ResultNode->FindPinChecked(TEXT("Result"), EEdGraphPinDirection::EGPD_Input);
+			check(ResultNodePosePin->LinkedTo.Num() == 0);
+			for (UEdGraphNode* Node : ImportedNodes)
+			{
+				if (Node->NodeGuid == ResultNodeGuid)
+				{
+					UEdGraphPin* PosePin = Node->FindPinChecked(TEXT("Pose"), EEdGraphPinDirection::EGPD_Output);
+					check(PosePin->LinkedTo.Num() == 0);
+
+					PosePin->MakeLinkTo(ResultNodePosePin);
+					break;
+				}
+			}
+		}
+
+		for (UEdGraphNode* ImportedNode : ImportedNodes)
+		{
+			ImportedNode->CreateNewGuid();
+		}
+	}
+
+	SelectedSampleIndex = NewSampleIndex;
+	HighlightedSampleIndex = INDEX_NONE;
+}
+
 #undef LOCTEXT_NAMESPACE // "SAnimationBlendSpaceGridWidget"
diff --git a/Engine/Source/Editor/Persona/Private/SAnimationBlendSpaceGridWidget.h b/Engine/Source/Editor/Persona/Private/SAnimationBlendSpaceGridWidget.h
--- a/Engine/Source/Editor/Persona/Private/SAnimationBlendSpaceGridWidget.h
+++ b/Engine/Source/Editor/Persona/Private/SAnimationBlendSpaceGridWidget.h
@@ -218,6 +218,20 @@ protected:
 	void OnInputSliderBegin(const int32 ParameterIndex);
 	void OnInputSliderEnd(const float NewValue, const int32 ParameterIndex);
 
+	/** Binds commands to functions in this class */
+	void BindCommands();
+
+	/** Handlers for blend sample actions*/
+	void OnBlendSampleCut();
+	void OnBlendSampleCopy();
+	void OnBlendSamplePaste();
+	void OnBlendSampleDelete();
+
+	/** Returns whether or not the action is possible */
+	bool CanBlendSampleCutCopy();
+	bool CanBlendSamplePaste();
+	bool CanBlendSampleDelete();
+
 	/** Returns whether or not the sample tool tips should be visible */
 	EVisibility GetSampleToolTipVisibility() const;
 	EVisibility GetPreviewToolTipVisibility() const;
@@ -311,6 +325,9 @@ private:
 	TArray<FVector2D> CachedGridPoints;
 	TArray<FVector> CachedSamplePoints;
 
+	/** Commands that are bound to delegates */
+	TSharedPtr<class FUICommandList> UICommandList;
+
 	/** Whether or not the cached data should be refreshed on the next tick*/
 	bool bRefreshCachedData;
 
