Add a NoRefStreamingLODBias setting for static and skeletal meshes. As the name suggests, it is a LOD bias that is only applied when a mesh doesn't have any component reference. The setting can be found in the static and skeletal mesh editors under the LOD Settings section. The default value is -1 which means using the value of r.Streaming.DefaultNoRefLODBias. If it is greater than or equal to 0, then the specified value will be used. The setting is implemented as a per quality level variable so users can specify different values for different quality levels. Users select what quality levels to use by setting r.Streaming.NoRefLODBiasQualityLevel under the ViewDistanceQuality sections in platform scalability ini files. For example, adding r.Streaming.NoRefLODBiasQualityLevel=0 and r.Streaming.NoRefLODBiasQualityLevel=1 to [ViewDistanceQuality[REDACTED_HANDLE]] and [ViewDistanceQuality[REDACTED_HANDLE]], respectively, in WindowsClientEngine.ini will make low and medium no-ref LOD biases selectable at runtime and they will be changed accordingly when changing view distance qualities. Since no value is specified for view distance quality 2 and up, the default no-ref LOD bias will be used. Note that this is the default value of the NoRefStreamingLODBias setting on a specific mesh. It is not r.Streaming.DefaultNoRefLODBias. The default value will be automatically set to the value of the lowest selected quality level. In the above example, if NoRefStreamingLODBias is 2 for low and 1 for medium, the default will be 2. If this is not desired, it is advised to explicitly specify the desired quality levels for all view distance qualities.

diff --git a/Engine/Source/Editor/Persona/Private/PersonaMeshDetails.cpp b/Engine/Source/Editor/Persona/Private/PersonaMeshDetails.cpp
--- a/Engine/Source/Editor/Persona/Private/PersonaMeshDetails.cpp
+++ b/Engine/Source/Editor/Persona/Private/PersonaMeshDetails.cpp
@@ -90,6 +90,7 @@
 #include "Interfaces/ITargetPlatform.h"
 #include "Interfaces/ITargetPlatformManagerModule.h"
 #include "Misc/CoreMisc.h"
+#include "SPerQualityLevelPropertiesWidget.h"
 #include "UObject/SavePackage.h"
 
 #define LOCTEXT_NAMESPACE "PersonaMeshDetails"
@@ -3218,6 +3219,26 @@ void FPersonaMeshDetails::CustomizeLODSettingsCategories(IDetailLayoutBuilder& D
 		]
 	];
 
+	LODSettingsCategory.AddCustomRow(LOCTEXT("NoRefStreamingLODBias", "NoRef Streaming LOD Bias"))
+	.RowTag("NoRefStreamingLODBias")
+	.IsEnabled(TAttribute<bool>::CreateLambda([this]() { return GetLODCount() > 1; }))
+	.NameContent()
+	[
+		SNew(STextBlock)
+		.Font(IDetailLayoutBuilder::GetDetailFont())
+		.Text(LOCTEXT("NoRefStreamingLODBias", "NoRef Streaming LOD Bias"))
+	]
+	.ValueContent()
+	.MinDesiredWidth(float(SkelMesh->GetNoRefStreamingLODBias().PerQuality.Num() + 1) * 125.f)
+	.MaxDesiredWidth(float((int32)QualityLevelProperty::EQualityLevels::Num + 1) * 125.f)
+	[
+		SNew(SPerQualityLevelPropertiesWidget)
+		.OnGenerateWidget(this, &FPersonaMeshDetails::GetNoRefStreamingLODBiasWidget)
+		.OnAddEntry(this, &FPersonaMeshDetails::AddNoRefStreamingLODBiasOverride)
+		.OnRemoveEntry(this, &FPersonaMeshDetails::RemoveNoRefStreamingLODBiasOverride)
+		.EntryNames(this, &FPersonaMeshDetails::GetNoRefStreamingLODBiasOverrideNames)
+	];
+
 	TSharedPtr<IPropertyHandle> MinLODPropertyHandle = DetailLayout.GetProperty(USkeletalMesh::GetMinLodMemberName(), USkeletalMesh::StaticClass());
 	IDetailPropertyRow& MinLODRow = LODSettingsCategory.AddProperty(MinLODPropertyHandle);
 	MinLODRow.IsEnabled(TAttribute<bool>::Create(TAttribute<bool>::FGetter::CreateSP(this, &FPersonaMeshDetails::IsLODInfoEditingEnabled, -1)));
@@ -3400,6 +3421,119 @@ void FPersonaMeshDetails::RefreshMeshDetailLayout()
 	MeshDetailLayout->ForceRefreshDetails();
 }
 
+void FPersonaMeshDetails::OnNoRefStreamingLODBiasChanged(int32 NewValue, FName QualityLevel)
+{
+	USkeletalMesh* SkelMesh = GetPersonaToolkit()->GetMesh();
+	check(SkelMesh);
+
+	{
+		FSkinnedMeshComponentRecreateRenderStateContext ReregisterContext(SkelMesh);
+		NewValue = FMath::Clamp<int32>(NewValue, -1, SkelMesh->GetLODNum() - 1);
+		FPerQualityLevelInt NoRefStreamingLODBias = SkelMesh->GetNoRefStreamingLODBias();
+		int32 QLKey = QualityLevelProperty::FNameToQualityLevel(QualityLevel);
+		if (QualityLevel == NAME_None || QLKey == INDEX_NONE)
+		{
+			NoRefStreamingLODBias.Default = NewValue;
+		}
+		else
+		{
+			int32* ValuePtr = NoRefStreamingLODBias.PerQuality.Find(QLKey);
+			if (ValuePtr != nullptr)
+			{
+				*ValuePtr = NewValue;
+			}
+		}
+		SkelMesh->SetNoRefStreamingLODBias(MoveTemp(NoRefStreamingLODBias));
+		SkelMesh->Modify();
+	}
+	RefreshMeshDetailLayout();
+}
+
+void FPersonaMeshDetails::OnNoRefStreamingLODBiasCommitted(int32 InValue, ETextCommit::Type CommitInfo, FName QualityLevel)
+{
+	OnNoRefStreamingLODBiasChanged(InValue, QualityLevel);
+}
+
+int32 FPersonaMeshDetails::GetNoRefStreamingLODBias(FName QualityLevel) const
+{
+	USkeletalMesh* SkelMesh = GetPersonaToolkit()->GetMesh();
+	check(SkelMesh);
+
+	int32 QLKey = QualityLevelProperty::FNameToQualityLevel(QualityLevel);
+	const int32* ValuePtr = (QualityLevel == NAME_None) ? nullptr : SkelMesh->GetNoRefStreamingLODBias().PerQuality.Find(QLKey);
+	return (ValuePtr != nullptr) ? *ValuePtr : SkelMesh->GetNoRefStreamingLODBias().Default;
+}
+
+TSharedRef<SWidget> FPersonaMeshDetails::GetNoRefStreamingLODBiasWidget(FName QualityLevelName) const
+{
+	return SNew(SSpinBox<int32>)
+		.Font(IDetailLayoutBuilder::GetDetailFont())
+		.Value(this, &FPersonaMeshDetails::GetNoRefStreamingLODBias, QualityLevelName)
+		.OnValueChanged(const_cast<FPersonaMeshDetails*>(this), &FPersonaMeshDetails::OnNoRefStreamingLODBiasChanged, QualityLevelName)
+		.OnValueCommitted(const_cast<FPersonaMeshDetails*>(this), &FPersonaMeshDetails::OnNoRefStreamingLODBiasCommitted, QualityLevelName)
+		.MinValue(-1)
+		.MaxValue(GetPersonaToolkit()->GetMesh()->GetLODNum() - 1)
+		.ToolTipText(this, &FPersonaMeshDetails::GetNoRefStreamingLODBiasTooltip)
+		.IsEnabled(FPersonaMeshDetails::GetLODCount() > 1);
+}
+
+bool FPersonaMeshDetails::AddNoRefStreamingLODBiasOverride(FName QualityLevelName)
+{
+	FScopedTransaction Transaction(LOCTEXT("AddNoRefStreamingLODBiasOverride", "Add NoRef Streaming LOD Bias Override"));
+	USkeletalMesh* SkelMesh = GetPersonaToolkit()->GetMesh();
+	check(SkelMesh);
+	SkelMesh->Modify();
+	int32 QLKey = QualityLevelProperty::FNameToQualityLevel(QualityLevelName);
+	if (SkelMesh->GetNoRefStreamingLODBias().PerQuality.Find(QLKey) == nullptr)
+	{
+		FPerQualityLevelInt NoRefStreamingLODBias = SkelMesh->GetNoRefStreamingLODBias();
+		int32 Value = NoRefStreamingLODBias.Default;
+		NoRefStreamingLODBias.PerQuality.Add(QLKey, Value);
+		SkelMesh->SetNoRefStreamingLODBias(MoveTemp(NoRefStreamingLODBias));
+		OnNoRefStreamingLODBiasChanged(Value, QualityLevelName);
+		return true;
+	}
+	return false;
+}
+
+bool FPersonaMeshDetails::RemoveNoRefStreamingLODBiasOverride(FName QualityLevelName)
+{
+	FScopedTransaction Transaction(LOCTEXT("RemoveNoRefStreamingLODBiasOverride", "Remove NoRef Streaming LOD Bias Override"));
+	USkeletalMesh* SkelMesh = GetPersonaToolkit()->GetMesh();
+	check(SkelMesh);
+	SkelMesh->Modify();
+
+	FPerQualityLevelInt NoRefStreamingLODBias = SkelMesh->GetNoRefStreamingLODBias();
+	int32 QL = QualityLevelProperty::FNameToQualityLevel(QualityLevelName);
+	if (QL != INDEX_NONE && NoRefStreamingLODBias.PerQuality.Remove(QL) != 0)
+	{
+		int32 Value = NoRefStreamingLODBias.Default;
+		SkelMesh->SetNoRefStreamingLODBias(MoveTemp(NoRefStreamingLODBias));
+		OnNoRefStreamingLODBiasChanged(Value, QualityLevelName);
+		return true;
+	}
+	return false;
+}
+
+TArray<FName> FPersonaMeshDetails::GetNoRefStreamingLODBiasOverrideNames() const
+{
+	USkeletalMesh* SkelMesh = GetPersonaToolkit()->GetMesh();
+	check(SkelMesh);
+
+	TArray<FName> OverrideNames;
+	for (const TPair<int32, int32>& Pair : SkelMesh->GetNoRefStreamingLODBias().PerQuality)
+	{
+		OverrideNames.Add(QualityLevelProperty::QualityLevelToFName(Pair.Key));
+	}
+	OverrideNames.Sort(FNameLexicalLess());
+	return OverrideNames;
+}
+
+FText FPersonaMeshDetails::GetNoRefStreamingLODBiasTooltip() const
+{
+	return LOCTEXT("NoRefStreamingLODBiasTooltip", "LOD bias for preloading no-ref mesh LODs. To use platform default, set to -1.");
+}
+
 FReply FPersonaMeshDetails::OnApplyChanges()
 {
 	ApplyChanges();
diff --git a/Engine/Source/Editor/Persona/Private/PersonaMeshDetails.h b/Engine/Source/Editor/Persona/Private/PersonaMeshDetails.h
--- a/Engine/Source/Editor/Persona/Private/PersonaMeshDetails.h
+++ b/Engine/Source/Editor/Persona/Private/PersonaMeshDetails.h
@@ -586,6 +586,15 @@ private:
 
 	void RefreshMeshDetailLayout();
 
+	void OnNoRefStreamingLODBiasChanged(int32 NewValue, FName QualityLevel);
+	void OnNoRefStreamingLODBiasCommitted(int32 InValue, ETextCommit::Type CommitInfo, FName QualityLevel);
+	int32 GetNoRefStreamingLODBias(FName QualityLevel) const;
+	TSharedRef<SWidget> GetNoRefStreamingLODBiasWidget(FName QualityLevelName) const;
+	bool AddNoRefStreamingLODBiasOverride(FName QualityLevelName);
+	bool RemoveNoRefStreamingLODBiasOverride(FName QualityLevelName);
+	TArray<FName> GetNoRefStreamingLODBiasOverrideNames() const;
+	FText GetNoRefStreamingLODBiasTooltip() const;
+
 	/** apply LOD changes if the user modified LOD reduction settings */
 	FReply OnApplyChanges();
 	/** regenerate one specific LOD Index no dependencies*/
diff --git a/Engine/Source/Editor/StaticMeshEditor/Private/StaticMeshEditorTools.cpp b/Engine/Source/Editor/StaticMeshEditor/Private/StaticMeshEditorTools.cpp
--- a/Engine/Source/Editor/StaticMeshEditor/Private/StaticMeshEditorTools.cpp
+++ b/Engine/Source/Editor/StaticMeshEditor/Private/StaticMeshEditorTools.cpp
@@ -3297,6 +3297,26 @@ void FLevelOfDetailSettingsLayout::AddToDetailsPanel( IDetailLayoutBuilder& Deta
 			.EntryNames(this, &FLevelOfDetailSettingsLayout::GetMinQualityLevelLODOverrideNames)
 		];
 
+	LODSettingsCategory.AddCustomRow(LOCTEXT("NoRefStreamingLODBias", "NoRef Streaming LOD Bias"))
+	.RowTag("NoRefStreamingLODBias")
+	.IsEnabled(TAttribute<bool>::CreateLambda([this]() { return GetLODCount() > 1; }))
+	.NameContent()
+		[
+			SNew(STextBlock)
+			.Font(IDetailLayoutBuilder::GetDetailFont())
+			.Text(LOCTEXT("NoRefStreamingLODBias", "NoRef Streaming LOD Bias"))
+		]
+	.ValueContent()
+		.MinDesiredWidth(float(StaticMesh->GetNoRefStreamingLODBias().PerQuality.Num() + 1) * 125.f)
+		.MaxDesiredWidth(float((int32)QualityLevelProperty::EQualityLevels::Num + 1) * 125.f)
+		[
+			SNew(SPerQualityLevelPropertiesWidget)
+			.OnGenerateWidget(this, &FLevelOfDetailSettingsLayout::GetNoRefStreamingLODBiasWidget)
+			.OnAddEntry(this, &FLevelOfDetailSettingsLayout::AddNoRefStreamingLODBiasOverride)
+			.OnRemoveEntry(this, &FLevelOfDetailSettingsLayout::RemoveNoRefStreamingLODBiasOverride)
+			.EntryNames(this, &FLevelOfDetailSettingsLayout::GetNoRefStreamingLODBiasOverrideNames)
+		];
+
 	{
 		TAttribute<TArray<FName>> PlatformOverrideNames = TAttribute<TArray<FName>>::CreateSP(this, &FLevelOfDetailSettingsLayout::GetNumStreamedLODsPlatformOverrideNames);
 
@@ -4421,6 +4441,119 @@ TArray<FName> FLevelOfDetailSettingsLayout::GetMinQualityLevelLODOverrideNames()
 	return OverrideNames;
 }
 
+void FLevelOfDetailSettingsLayout::OnNoRefStreamingLODBiasChanged(int32 NewValue, FName QualityLevel)
+{
+	UStaticMesh* StaticMesh = StaticMeshEditor.GetStaticMesh();
+	check(StaticMesh);
+
+	{
+		FStaticMeshComponentRecreateRenderStateContext ReregisterContext(StaticMesh, false);
+		NewValue = FMath::Clamp<int32>(NewValue, -1, MAX_STATIC_MESH_LODS - 1);
+		FPerQualityLevelInt NoRefStreamingLODBias = StaticMesh->GetNoRefStreamingLODBias();
+		int32 QLKey = QualityLevelProperty::FNameToQualityLevel(QualityLevel);
+		if (QualityLevel == NAME_None || QLKey == INDEX_NONE)
+		{
+			NoRefStreamingLODBias.Default = NewValue;
+		}
+		else
+		{
+			int32* ValuePtr = NoRefStreamingLODBias.PerQuality.Find(QLKey);
+			if (ValuePtr != nullptr)
+			{
+				*ValuePtr = NewValue;
+			}
+		}
+		StaticMesh->SetNoRefStreamingLODBias(MoveTemp(NoRefStreamingLODBias));
+		StaticMesh->Modify();
+	}
+	StaticMeshEditor.RefreshViewport();
+}
+
+void FLevelOfDetailSettingsLayout::OnNoRefStreamingLODBiasCommitted(int32 InValue, ETextCommit::Type CommitInfo, FName QualityLevel)
+{
+	OnNoRefStreamingLODBiasChanged(InValue, QualityLevel);
+}
+
+int32 FLevelOfDetailSettingsLayout::GetNoRefStreamingLODBias(FName QualityLevel) const
+{
+	UStaticMesh* StaticMesh = StaticMeshEditor.GetStaticMesh();
+	check(StaticMesh);
+
+	int32 QLKey = QualityLevelProperty::FNameToQualityLevel(QualityLevel);
+	const int32* ValuePtr = (QualityLevel == NAME_None) ? nullptr : StaticMesh->GetNoRefStreamingLODBias().PerQuality.Find(QLKey);
+	return (ValuePtr != nullptr) ? *ValuePtr : StaticMesh->GetNoRefStreamingLODBias().Default;
+}
+
+TSharedRef<SWidget> FLevelOfDetailSettingsLayout::GetNoRefStreamingLODBiasWidget(FName QualityLevelName) const
+{
+	return SNew(SSpinBox<int32>)
+		.Font(IDetailLayoutBuilder::GetDetailFont())
+		.Value(this, &FLevelOfDetailSettingsLayout::GetNoRefStreamingLODBias, QualityLevelName)
+		.OnValueChanged(const_cast<FLevelOfDetailSettingsLayout*>(this), &FLevelOfDetailSettingsLayout::OnNoRefStreamingLODBiasChanged, QualityLevelName)
+		.OnValueCommitted(const_cast<FLevelOfDetailSettingsLayout*>(this), &FLevelOfDetailSettingsLayout::OnNoRefStreamingLODBiasCommitted, QualityLevelName)
+		.MinValue(-1)
+		.MaxValue(MAX_STATIC_MESH_LODS - 1)
+		.ToolTipText(this, &FLevelOfDetailSettingsLayout::GetNoRefStreamingLODBiasTooltip)
+		.IsEnabled(FLevelOfDetailSettingsLayout::GetLODCount() > 1);
+}
+
+bool FLevelOfDetailSettingsLayout::AddNoRefStreamingLODBiasOverride(FName QualityLevelName)
+{
+	FScopedTransaction Transaction(LOCTEXT("AddNoRefStreamingLODBiasOverride", "Add NoRef Streaming LOD Bias Override"));
+	UStaticMesh* StaticMesh = StaticMeshEditor.GetStaticMesh();
+	check(StaticMesh);
+	StaticMesh->Modify();
+	int32 QLKey = QualityLevelProperty::FNameToQualityLevel(QualityLevelName);
+	if (StaticMesh->GetNoRefStreamingLODBias().PerQuality.Find(QLKey) == nullptr)
+	{
+		FPerQualityLevelInt NoRefStreamingLODBias = StaticMesh->GetNoRefStreamingLODBias();
+		int32 Value = NoRefStreamingLODBias.Default;
+		NoRefStreamingLODBias.PerQuality.Add(QLKey, Value);
+		StaticMesh->SetNoRefStreamingLODBias(MoveTemp(NoRefStreamingLODBias));
+		OnNoRefStreamingLODBiasChanged(Value, QualityLevelName);
+		return true;
+	}
+	return false;
+}
+
+bool FLevelOfDetailSettingsLayout::RemoveNoRefStreamingLODBiasOverride(FName QualityLevelName)
+{
+	FScopedTransaction Transaction(LOCTEXT("RemoveNoRefStreamingLODBiasOverride", "Remove NoRef Streaming LOD Bias Override"));
+	UStaticMesh* StaticMesh = StaticMeshEditor.GetStaticMesh();
+	check(StaticMesh);
+	StaticMesh->Modify();
+
+	FPerQualityLevelInt NoRefStreamingLODBias = StaticMesh->GetNoRefStreamingLODBias();
+	int32 QL = QualityLevelProperty::FNameToQualityLevel(QualityLevelName);
+	if (QL != INDEX_NONE && NoRefStreamingLODBias.PerQuality.Remove(QL) != 0)
+	{
+		int32 Value = NoRefStreamingLODBias.Default;
+		StaticMesh->SetNoRefStreamingLODBias(MoveTemp(NoRefStreamingLODBias));
+		OnNoRefStreamingLODBiasChanged(Value, QualityLevelName);
+		return true;
+	}
+	return false;
+}
+
+TArray<FName> FLevelOfDetailSettingsLayout::GetNoRefStreamingLODBiasOverrideNames() const
+{
+	UStaticMesh* StaticMesh = StaticMeshEditor.GetStaticMesh();
+	check(StaticMesh);
+
+	TArray<FName> OverrideNames;
+	for (const TPair<int32, int32>& Pair : StaticMesh->GetNoRefStreamingLODBias().PerQuality)
+	{
+		OverrideNames.Add(QualityLevelProperty::QualityLevelToFName(Pair.Key));
+	}
+	OverrideNames.Sort(FNameLexicalLess());
+	return OverrideNames;
+}
+
+FText FLevelOfDetailSettingsLayout::GetNoRefStreamingLODBiasTooltip() const
+{
+	return LOCTEXT("NoRefStreamingLODBiasTooltip", "LOD bias for preloading no-ref mesh LODs. To use platform default, set to -1.");
+}
+
 /** @return - whether value was different */
 static bool UpdateStaticMeshNumStreamedLODsHelper(UStaticMesh* StaticMesh, int32 NewValue, FName Platform)
 {
diff --git a/Engine/Source/Editor/StaticMeshEditor/Private/StaticMeshEditorTools.h b/Engine/Source/Editor/StaticMeshEditor/Private/StaticMeshEditorTools.h
--- a/Engine/Source/Editor/StaticMeshEditor/Private/StaticMeshEditorTools.h
+++ b/Engine/Source/Editor/StaticMeshEditor/Private/StaticMeshEditorTools.h
@@ -559,6 +559,14 @@ private:
 	bool RemoveMinLODQualityLevelOverride(FName QualityLevelName);
 	TArray<FName> GetMinQualityLevelLODOverrideNames() const;
 
+	void OnNoRefStreamingLODBiasChanged(int32 NewValue, FName QualityLevel);
+	void OnNoRefStreamingLODBiasCommitted(int32 InValue, ETextCommit::Type CommitInfo, FName QualityLevel);
+	int32 GetNoRefStreamingLODBias(FName QualityLevel) const;
+	TSharedRef<SWidget> GetNoRefStreamingLODBiasWidget(FName QualityLevelName) const;
+	bool AddNoRefStreamingLODBiasOverride(FName QualityLevelName);
+	bool RemoveNoRefStreamingLODBiasOverride(FName QualityLevelName);
+	TArray<FName> GetNoRefStreamingLODBiasOverrideNames() const;
+
 	void OnNumStreamedLODsChanged(int32 NewValue, FName Platform);
 	void OnNumStreamedLODsCommitted(int32 InValue, ETextCommit::Type CommitInfo, FName Platform);
 	int32 GetNumStreamedLODs(FName Platform) const;
@@ -595,6 +603,7 @@ private:
 	void UpdateLODNames();
 	FText GetLODCountTooltip() const;
 	FText GetMinLODTooltip() const;
+	FText GetNoRefStreamingLODBiasTooltip() const;
 	FText GetNumStreamedLODsTooltip() const;
 
 	FText GetLODCustomModeNameContent(int32 LODIndex) const;
diff --git a/Engine/Source/Runtime/Engine/Classes/Engine/StreamableRenderAsset.h b/Engine/Source/Runtime/Engine/Classes/Engine/StreamableRenderAsset.h
--- a/Engine/Source/Runtime/Engine/Classes/Engine/StreamableRenderAsset.h
+++ b/Engine/Source/Runtime/Engine/Classes/Engine/StreamableRenderAsset.h
@@ -11,6 +11,7 @@
 #include "Templates/RefCounting.h"
 #include "RenderAssetUpdate.h"
 #include "Streaming/StreamableRenderResourceState.h"
+#include "PerQualityLevelProperties.h"
 #include "StreamableRenderAsset.generated.h"
 
 #define STREAMABLERENDERASSET_NODEFAULT(FuncName) LowLevelFatalError(TEXT("UStreamableRenderAsset::%s has no default implementation"), TEXT(#FuncName))
@@ -204,6 +205,18 @@ public:
 	ENGINE_API virtual void BeginDestroy() override;
 	ENGINE_API virtual bool IsReadyForFinishDestroy() override;
 
+	const FPerQualityLevelInt& GetNoRefStreamingLODBias() const
+	{
+		return NoRefStreamingLODBias;
+	}
+
+	void SetNoRefStreamingLODBias(FPerQualityLevelInt NewValue)
+	{
+		NoRefStreamingLODBias = MoveTemp(NewValue);
+	}
+
+	ENGINE_API int32 GetCurrentNoRefStreamingLODBias() const;
+
 protected:
 	
 	// Also returns false if the render resource is non existent, to prevent stalling on an event that will never complete.
@@ -243,6 +256,9 @@ public:
 	int32 NumCinematicMipLevels;
 
 protected:
+	UPROPERTY()
+	FPerQualityLevelInt NoRefStreamingLODBias;
+
 	/** FStreamingRenderAsset index used by the texture streaming system. */
 	UPROPERTY(transient, duplicatetransient, NonTransactional)
 	int32 StreamingIndex = INDEX_NONE;
diff --git a/Engine/Source/Runtime/Engine/Private/StreamableRenderAsset.cpp b/Engine/Source/Runtime/Engine/Private/StreamableRenderAsset.cpp
--- a/Engine/Source/Runtime/Engine/Private/StreamableRenderAsset.cpp
+++ b/Engine/Source/Runtime/Engine/Private/StreamableRenderAsset.cpp
@@ -5,12 +5,24 @@
 #include "ContentStreaming.h"
 #include "Rendering/NaniteCoarseMeshStreamingManager.h"
 
+static const TCHAR* GNoRefBiasQualityLevelCVarName = TEXT("r.Streaming.NoRefLODBiasQualityLevel");
+static const TCHAR* GNoRefBiasQualityLevelScalabilitySection = TEXT("ViewDistanceQuality");
+static int32 GNoRefBiasQualityLevel = -1;
+static FAutoConsoleVariableRef CVarNoRefBiasQualityLevel(
+	GNoRefBiasQualityLevelCVarName,
+	GNoRefBiasQualityLevel,
+	TEXT("The quality level for the no-ref mesh streaming LOD bias"),
+	ECVF_Scalability);
+
 extern bool TrackRenderAssetEvent(struct FStreamingRenderAsset* StreamingRenderAsset, UStreamableRenderAsset* RenderAsset, bool bForceMipLevelsToBeResident, const FRenderAssetStreamingManager* Manager);
 
 UStreamableRenderAsset::UStreamableRenderAsset(const FObjectInitializer& ObjectInitializer)
 	: Super(ObjectInitializer)
 {
 	check(sizeof(FStreamableRenderResourceState) == sizeof(uint64));
+
+	SetNoRefStreamingLODBias(-1);
+	NoRefStreamingLODBias.Init(GNoRefBiasQualityLevelCVarName, GNoRefBiasQualityLevelScalabilitySection);
 }
 
 void UStreamableRenderAsset::RegisterMipLevelChangeCallback(UPrimitiveComponent* Component, int32 LODIndex, float TimeoutSecs, bool bOnStreamIn, FLODStreamingCallback&& Callback)
@@ -394,4 +406,9 @@ bool UStreamableRenderAsset::IsReadyForFinishDestroy()
 	}
 
 	return !PendingUpdate.IsValid();
-}
\ No newline at end of file
+}
+
+int32 UStreamableRenderAsset::GetCurrentNoRefStreamingLODBias() const
+{
+	return GetNoRefStreamingLODBias().GetValue(GNoRefBiasQualityLevel);
+}
diff --git a/Engine/Source/Runtime/Engine/Private/Streaming/StreamingTexture.cpp b/Engine/Source/Runtime/Engine/Private/Streaming/StreamingTexture.cpp
--- a/Engine/Source/Runtime/Engine/Private/Streaming/StreamingTexture.cpp
+++ b/Engine/Source/Runtime/Engine/Private/Streaming/StreamingTexture.cpp
@@ -21,6 +21,13 @@ static FAutoConsoleVariableRef CVarUseMobileLODBiasOnDesktopES31(
 );
 #endif
 
+static int32 GDefaultNoRefBias = 0;
+static FAutoConsoleVariableRef CVarDefaultNoRefBias(
+	TEXT("r.Streaming.DefaultNoRefLODBias"),
+	GDefaultNoRefBias,
+	TEXT("The default LOD bias for no-ref meshes"),
+	ECVF_Scalability);
+
 FStreamingRenderAsset::FStreamingRenderAsset(
 	UStreamableRenderAsset* InRenderAsset,
 	const int32* NumStreamedMips,
@@ -209,6 +216,17 @@ void FStreamingRenderAsset::UpdateDynamicData(const int32* NumStreamedMips, int3
 #endif
 		}
 
+		int32 LocalNoRefBias = 0;
+		if (RenderAssetType == EStreamableRenderAssetType::StaticMesh || RenderAssetType == EStreamableRenderAssetType::SkeletalMesh)
+		{
+			LocalNoRefBias = RenderAsset->GetCurrentNoRefStreamingLODBias();
+			if (LocalNoRefBias < 0)
+			{
+				LocalNoRefBias = GDefaultNoRefBias;
+			}
+		}
+		NoRefLODBias = LocalNoRefBias;
+
 		// If the optional mips are not available, or if we shouldn't load them now, clamp the possible mips requested. 
 		// (when the non-optional mips are not yet loaded, loading optional mips generates cross files requests).
 		// This is not bullet proof though since the texture/mesh could have a pending stream-out request.
@@ -239,6 +257,7 @@ void FStreamingRenderAsset::UpdateDynamicData(const int32* NumStreamedMips, int3
 		RequestedMips = 0;
 		MinAllowedMips = 0;
 		MaxAllowedMips = 0;
+		NoRefLODBias = 0;
 		OptionalMipsState = EOptionalMipsState::OMS_NotCached;
 		LastRenderTime = FLT_MAX;	
 	}
@@ -297,18 +316,21 @@ int32 FStreamingRenderAsset::GetWantedMipsFromSize(float Size, float InvMaxScree
 	else
 	{
 		check(RenderAssetType == EStreamableRenderAssetType::StaticMesh || RenderAssetType == EStreamableRenderAssetType::SkeletalMesh || RenderAssetType == EStreamableRenderAssetType::LandscapeMeshMobile);
-		if (Size != FLT_MAX)
+		if (Size == FLT_MAX)
+		{
+			return MaxAllowedMips;
+		}
+		int32 Result = MaxAllowedMips;
+		const float NormalizedSize = Size * InvMaxScreenSizeOverAllViews;
+		for (int32 NumMips = MinAllowedMips; NumMips <= MaxAllowedMips; ++NumMips)
 		{
-			const float NormalizedSize = Size * InvMaxScreenSizeOverAllViews;
-			for (int32 NumMips = MinAllowedMips; NumMips <= MaxAllowedMips; ++NumMips)
+			if (GetNormalizedScreenSize(NumMips) >= NormalizedSize)
 			{
-				if (GetNormalizedScreenSize(NumMips) >= NormalizedSize)
-				{
-					return NumMips;
-				}
+				Result = NumMips;
+				break;
 			}
 		}
-		return MaxAllowedMips;
+		return bUseUnkownRefHeuristic ? FMath::Max(Result - NoRefLODBias, MinAllowedMips) : Result;;
 	}
 }
 
diff --git a/Engine/Source/Runtime/Engine/Private/Streaming/StreamingTexture.h b/Engine/Source/Runtime/Engine/Private/Streaming/StreamingTexture.h
--- a/Engine/Source/Runtime/Engine/Private/Streaming/StreamingTexture.h
+++ b/Engine/Source/Runtime/Engine/Private/Streaming/StreamingTexture.h
@@ -220,6 +220,8 @@ struct FStreamingRenderAsset
 	int32			MinAllowedMips;
 	/** (2) Max mip to be requested by the streaming  */
 	int32			MaxAllowedMips;
+	/** (2) LOD bias applied to no-ref meshes */
+	int32			NoRefLODBias;
 	/** (2) How much game time has elapsed since the texture was bound for rendering. Based on FApp::GetCurrentTime(). */
 	float			LastRenderTime;
 
