Fix crash when closing skeletalmesheditor and there was a selected morph target.
#rnx

diff --git a/Engine/Source/Editor/Persona/Private/SMorphTargetViewer.cpp b/Engine/Source/Editor/Persona/Private/SMorphTargetViewer.cpp
--- a/Engine/Source/Editor/Persona/Private/SMorphTargetViewer.cpp
+++ b/Engine/Source/Editor/Persona/Private/SMorphTargetViewer.cpp
@@ -606,13 +606,16 @@ void SMorphTargetViewer::SetSelectedMorphTargets(const TArray<FName>& SelectedMo
 
 		if (SelectedMorphTargetNames.Num() > 0)
 		{
-			for (const FName& MorphTargetName : SelectedMorphTargetNames)
+			if (SkeletalMesh)
 			{
-				int32 MorphtargetIdx;
-				UMorphTarget* MorphTarget = SkeletalMesh->FindMorphTargetAndIndex(MorphTargetName, MorphtargetIdx);
-				if (MorphTarget != nullptr)
+				for (const FName& MorphTargetName : SelectedMorphTargetNames)
 				{
-					PreviewComponent->MorphTargetOfInterests.AddUnique(MorphTarget);
+					int32 MorphtargetIdx;
+					UMorphTarget* MorphTarget = SkeletalMesh->FindMorphTargetAndIndex(MorphTargetName, MorphtargetIdx);
+					if (MorphTarget != nullptr)
+					{
+						PreviewComponent->MorphTargetOfInterests.AddUnique(MorphTarget);
+					}
 				}
 			}
 
