Fix crash deleting a skeleton

Hardened code that could potentially reference a deleted USkeleton

diff --git a/Engine/Source/Editor/Persona/Private/SAnimationEditorViewport.cpp b/Engine/Source/Editor/Persona/Private/SAnimationEditorViewport.cpp
--- a/Engine/Source/Editor/Persona/Private/SAnimationEditorViewport.cpp
+++ b/Engine/Source/Editor/Persona/Private/SAnimationEditorViewport.cpp
@@ -244,7 +244,7 @@ FText SAnimationEditorViewportTabBody::GetDisplayString() const
 {
 	class UDebugSkelMeshComponent* Component = GetPreviewScene()->GetPreviewMeshComponent();
 	TSharedPtr<IEditableSkeleton> EditableSkeleton = GetPreviewScene()->GetPersonaToolkit()->GetEditableSkeleton();
-	FName TargetSkeletonName = EditableSkeleton.IsValid() ? EditableSkeleton->GetSkeleton().GetFName() : NAME_None;
+	FName TargetSkeletonName = (EditableSkeleton.IsValid() && EditableSkeleton->IsSkeletonValid()) ? EditableSkeleton->GetSkeleton().GetFName() : NAME_None;
 
 	FText DefaultText;
 
diff --git a/Engine/Source/Editor/SkeletonEditor/Private/EditableSkeleton.cpp b/Engine/Source/Editor/SkeletonEditor/Private/EditableSkeleton.cpp
--- a/Engine/Source/Editor/SkeletonEditor/Private/EditableSkeleton.cpp
+++ b/Engine/Source/Editor/SkeletonEditor/Private/EditableSkeleton.cpp
@@ -141,6 +141,11 @@ const USkeleton& FEditableSkeleton::GetSkeleton() const
 	return *Skeleton;
 }
 
+bool FEditableSkeleton::IsSkeletonValid() const
+{
+	return Skeleton != nullptr;
+}
+
 const TArray<class UBlendProfile*>& FEditableSkeleton::GetBlendProfiles() const
 {
 	return Skeleton->BlendProfiles;
diff --git a/Engine/Source/Editor/SkeletonEditor/Private/EditableSkeleton.h b/Engine/Source/Editor/SkeletonEditor/Private/EditableSkeleton.h
--- a/Engine/Source/Editor/SkeletonEditor/Private/EditableSkeleton.h
+++ b/Engine/Source/Editor/SkeletonEditor/Private/EditableSkeleton.h
@@ -29,6 +29,7 @@ public:
 
 	/** IEditableSkeleton interface */
 	virtual const USkeleton& GetSkeleton() const override;
+	virtual bool IsSkeletonValid() const override;
 	virtual const TArray<class UBlendProfile*>& GetBlendProfiles() const override;
 	virtual class UBlendProfile* GetBlendProfile(const FName& InBlendProfileName) override;
 	virtual class UBlendProfile* CreateNewBlendProfile(const FName& InBlendProfileName) override;
diff --git a/Engine/Source/Editor/SkeletonEditor/Public/IEditableSkeleton.h b/Engine/Source/Editor/SkeletonEditor/Public/IEditableSkeleton.h
--- a/Engine/Source/Editor/SkeletonEditor/Public/IEditableSkeleton.h
+++ b/Engine/Source/Editor/SkeletonEditor/Public/IEditableSkeleton.h
@@ -27,6 +27,9 @@ public:
 	/** Get the skeleton. Const - to modify the skeleton you should use the wrapper methods below */
 	virtual const class USkeleton& GetSkeleton() const = 0;
 
+	/** Checks to see if the editable skeleton is valid. It is possible to force-delete the USkeleton. */
+	virtual bool IsSkeletonValid() const = 0;
+	
 	/** Get the blend profiles that this skeleton currently contains */
 	virtual const TArray<class UBlendProfile*>& GetBlendProfiles() const  = 0;
 
