CTRL+Z/CTRL+Y doesn't undo/redo changes in Additional Meshes in Animation Sequence
#fix ensure that the transient preview mesh collection asset is marked as transactional

diff --git a/Engine/Source/Editor/Persona/Private/AnimationEditorPreviewScene.cpp b/Engine/Source/Editor/Persona/Private/AnimationEditorPreviewScene.cpp
--- a/Engine/Source/Editor/Persona/Private/AnimationEditorPreviewScene.cpp
+++ b/Engine/Source/Editor/Persona/Private/AnimationEditorPreviewScene.cpp
@@ -91,7 +91,7 @@ FAnimationEditorPreviewScene::FAnimationEditorPreviewScene(const ConstructionVal
 	{
 		FactoryToUse->CurrentSkeleton = MakeWeakObjectPtr(const_cast<USkeleton*>(&InEditableSkeleton->GetSkeleton()));
 	}
-	PreviewSceneDescription->DefaultAdditionalMeshes = CastChecked<UPreviewMeshCollection>(FactoryToUse->FactoryCreateNew(UPreviewMeshCollection::StaticClass(), PreviewSceneDescription, "UnsavedCollection", RF_Transient, nullptr, nullptr));
+	PreviewSceneDescription->DefaultAdditionalMeshes = CastChecked<UPreviewMeshCollection>(FactoryToUse->FactoryCreateNew(UPreviewMeshCollection::StaticClass(), PreviewSceneDescription, "UnsavedCollection", RF_Transient | RF_Transactional, nullptr, nullptr));
 
 	if (!PreviewSceneDescription->AdditionalMeshes.IsValid())
 	{
