- Multiple Preview Mesh Set Temporarily toasts can be present at the same time
- Incorrect Preview Mesh can show as selected if the user presses the Apply to Asset button after selection
#fix fade-out any pending Apply Preview Mesh toast when adding a new one
#fix retrieve preview-mesh from preview scene rather than viewed asset

diff --git a/Engine/Source/Editor/Persona/Private/PersonaModule.cpp b/Engine/Source/Editor/Persona/Private/PersonaModule.cpp
--- a/Engine/Source/Editor/Persona/Private/PersonaModule.cpp
+++ b/Engine/Source/Editor/Persona/Private/PersonaModule.cpp
@@ -892,6 +892,7 @@ void FPersonaModule::AddCommonToolbarExtensions(FToolBarBuilder& InToolbarBuilde
 			TSharedPtr<SNotificationItem> Notification;
 		};
 
+		static TWeakPtr<SNotificationItem> WeakPendingApplyMeshNotification = nullptr;
 		auto CreatePreviewMeshComboButtonContents = [WeakPersonaToolkit]()
 		{
 			FMenuBuilder MenuBuilder(true, nullptr);
@@ -921,7 +922,15 @@ void FPersonaModule::AddCommonToolbarExtensions(FToolBarBuilder& InToolbarBuilde
 								FSimpleDelegate::CreateStatic(&FNotificationHandler::HandleApplyPreviewMesh, NotificationHandler, WeakPersonaToolkit),
 								SNotificationItem::CS_Success));
 
+						// Fade-out previously added SetPreviewMesh notification (if any)
+						const TSharedPtr<SNotificationItem> PendingNotification = WeakPendingApplyMeshNotification.Pin();
+						if(PendingNotification.IsValid())
+						{
+							PendingNotification->Fadeout();
+						}
+
 						NotificationHandler->Notification = FSlateNotificationManager::Get().AddNotification(Info);
+						WeakPendingApplyMeshNotification = NotificationHandler->Notification;
 						if (NotificationHandler->Notification.IsValid())
 						{
 							NotificationHandler->Notification->SetCompletionState(SNotificationItem::CS_Success);
@@ -962,7 +971,7 @@ void FPersonaModule::AddCommonToolbarExtensions(FToolBarBuilder& InToolbarBuilde
 				});
 				if (WeakPersonaToolkit.IsValid())
 				{
-					AssetPickerConfig.InitialAssetSelection = FAssetData(WeakPersonaToolkit.Pin()->GetPreviewMesh());
+					AssetPickerConfig.InitialAssetSelection = FAssetData(WeakPersonaToolkit.Pin()->GetPreviewScene()->GetPreviewMesh());
 				}
 
 				FContentBrowserModule& ContentBrowserModule = FModuleManager::Get().LoadModuleChecked<FContentBrowserModule>(TEXT("ContentBrowser"));
