Added transform attribute visualization to Persona

diff --git a/Engine/Source/Editor/Persona/Private/AnimViewportShowCommands.cpp b/Engine/Source/Editor/Persona/Private/AnimViewportShowCommands.cpp
--- a/Engine/Source/Editor/Persona/Private/AnimViewportShowCommands.cpp
+++ b/Engine/Source/Editor/Persona/Private/AnimViewportShowCommands.cpp
@@ -42,6 +42,7 @@ void FAnimViewportShowCommands::RegisterCommands()
 	UI_COMMAND( ShowSourceRawAnimation, "Source Animation", "Display skeleton in source raw animation if you have track curves modified", EUserInterfaceActionType::ToggleButton, FInputChord() );
 	UI_COMMAND( ShowBakedAnimation, "Baked Animation", "Display skeleton in baked raw animation if you have track curves modified", EUserInterfaceActionType::ToggleButton, FInputChord() );
 	UI_COMMAND( ShowSockets, "Sockets", "Display socket hitpoints", EUserInterfaceActionType::ToggleButton, FInputChord() );
+	UI_COMMAND( ShowAttributes, "Attributes", "Display transform attributes in the viewport", EUserInterfaceActionType::ToggleButton, FInputChord() );
 
 	UI_COMMAND( ShowBoneDrawNone, "None", "Hide bone selection", EUserInterfaceActionType::RadioButton, FInputChord());
 	UI_COMMAND( ShowBoneDrawSelected, "Selected Only", "Show only the selected bone", EUserInterfaceActionType::RadioButton, FInputChord());
diff --git a/Engine/Source/Editor/Persona/Private/AnimViewportShowCommands.h b/Engine/Source/Editor/Persona/Private/AnimViewportShowCommands.h
--- a/Engine/Source/Editor/Persona/Private/AnimViewportShowCommands.h
+++ b/Engine/Source/Editor/Persona/Private/AnimViewportShowCommands.h
@@ -116,6 +116,9 @@ public:
 	/** Show socket hit point diamonds */
 	TSharedPtr< FUICommandInfo > ShowSockets;
 
+	/** Show transform attributes */
+	TSharedPtr< FUICommandInfo > ShowAttributes;
+
 	/** Hide all local axes */
 	TSharedPtr< FUICommandInfo > ShowLocalAxesNone;
 	
diff --git a/Engine/Source/Editor/Persona/Private/AnimationEditorViewportClient.cpp b/Engine/Source/Editor/Persona/Private/AnimationEditorViewportClient.cpp
--- a/Engine/Source/Editor/Persona/Private/AnimationEditorViewportClient.cpp
+++ b/Engine/Source/Editor/Persona/Private/AnimationEditorViewportClient.cpp
@@ -13,6 +13,7 @@
 #include "Preferences/PersonaOptions.h"
 #include "Engine/CollisionProfile.h"
 #include "Animation/AnimBlueprintGeneratedClass.h"
+#include "Animation/BuiltInAttributeTypes.h"
 #include "Animation/MirrorDataTable.h"
 #include "GameFramework/WorldSettings.h"
 #include "PersonaModule.h"
@@ -498,6 +499,11 @@ void FAnimationViewportClient::Draw(const FSceneView* View, FPrimitiveDrawInterf
 				DrawSockets(PreviewMeshComponent, PreviewMeshComponent->SkeletalMesh->GetMeshOnlySocketList(), GetAnimPreviewScene()->GetSelectedSocket(), PDI, false);
 			}
 		}
+
+		if (PreviewMeshComponent->bDrawAttributes)
+		{
+			DrawAttributes(PreviewMeshComponent, PDI);
+		}
 	}
 
 	if (bFocusOnDraw)
@@ -527,6 +533,12 @@ void FAnimationViewportClient::DrawCanvas( FViewport& InViewport, FSceneView& Vi
 			ShowBoneNames(&Canvas, &View);
 		}
 
+		// Display attribute names
+		if (PreviewMeshComponent->bDrawAttributes)
+		{
+			ShowAttributeNames(&Canvas, &View);
+		}
+
 		if (bDrawUVs)
 		{
 			DrawUVsForMesh(Viewport, &Canvas, 1.0f);
@@ -721,6 +733,49 @@ void FAnimationViewportClient::ShowBoneNames( FCanvas* Canvas, FSceneView* View
 	}
 }
 
+void FAnimationViewportClient::ShowAttributeNames(FCanvas* Canvas, FSceneView* View)
+{
+	UDebugSkelMeshComponent* MeshComponent = GetAnimPreviewScene()->GetPreviewMeshComponent();
+
+	if (MeshComponent && MeshComponent->SkeletalMesh)
+	{
+		const int32 HalfX = Viewport->GetSizeXY().X / 2 / GetDPIScale();
+		const int32 HalfY = Viewport->GetSizeXY().Y / 2 / GetDPIScale();
+
+		const UE::Anim::FHeapAttributeContainer& Attributes = MeshComponent->GetCustomAttributes();
+
+		const int32 TransformAnimationAttributeTypeIndex = Attributes.FindTypeIndex(FTransformAnimationAttribute::StaticStruct());
+		if (TransformAnimationAttributeTypeIndex != INDEX_NONE)
+		{
+			const TArray<UE::Anim::FAttributeId, FDefaultAllocator>& AttributeIdentifiers = Attributes.GetKeys(TransformAnimationAttributeTypeIndex);
+			const TArray<UE::Anim::TWrappedAttribute<FDefaultAllocator>>& AttributeValues = Attributes.GetValues(TransformAnimationAttributeTypeIndex);
+			check(AttributeIdentifiers.Num() == AttributeValues.Num());
+
+			for (int32 AttributeIndex = 0; AttributeIndex < AttributeValues.Num(); ++AttributeIndex)
+			{
+				if (const FTransformAnimationAttribute* AttributeValue = AttributeValues[AttributeIndex].GetPtr<FTransformAnimationAttribute>())
+				{
+					const UE::Anim::FAttributeId& AttributeIdentifier = AttributeIdentifiers[AttributeIndex];
+
+					const FTransform AttributeParentTransform = MeshComponent->GetDrawTransform(AttributeIdentifier.GetIndex()) * MeshComponent->GetComponentTransform();
+					const FTransform AttributeTransform = AttributeValue->Value * AttributeParentTransform;
+
+					const FPlane proj = View->Project(AttributeTransform.GetLocation());
+					if (proj.W > 0.f)
+					{
+						const int32 XPos = HalfX + (HalfX * proj.X);
+						const int32 YPos = HalfY + (HalfY * (proj.Y * -1));
+
+						FCanvasTextItem TextItem(FVector2D(XPos, YPos), FText::FromName(AttributeIdentifier.GetName()), GEngine->GetSmallFont(), FLinearColor(0.0f, 1.0f, 1.0f));
+						TextItem.EnableShadow(FLinearColor::Black);
+						Canvas->DrawItem(TextItem);
+					}
+				}
+			}
+		}
+	}
+}
+
 bool FAnimationViewportClient::ShouldDisplayAdditiveScaleErrorMessage() const
 {
 	UAnimSequence* AnimSequence = Cast<UAnimSequence>(GetAnimPreviewScene()->GetPreviewAnimationAsset());
@@ -1550,6 +1605,37 @@ void FAnimationViewportClient::DrawMeshSubsetBones(const UDebugSkelMeshComponent
 	}
 }
 
+void FAnimationViewportClient::DrawAttributes(UDebugSkelMeshComponent* MeshComponent, FPrimitiveDrawInterface* PDI) const
+{
+	if (MeshComponent && MeshComponent->SkeletalMesh)
+	{
+		const UE::Anim::FHeapAttributeContainer& Attributes = MeshComponent->GetCustomAttributes();
+
+		const int32 TransformAnimationAttributeTypeIndex = Attributes.FindTypeIndex(FTransformAnimationAttribute::StaticStruct());
+		if (TransformAnimationAttributeTypeIndex != INDEX_NONE)
+		{
+			const TArray<UE::Anim::FAttributeId, FDefaultAllocator>& AttributeIdentifiers = Attributes.GetKeys(TransformAnimationAttributeTypeIndex);
+			const TArray<UE::Anim::TWrappedAttribute<FDefaultAllocator>>& AttributeValues = Attributes.GetValues(TransformAnimationAttributeTypeIndex);
+			check(AttributeIdentifiers.Num() == AttributeValues.Num());
+
+			for (int32 AttributeIndex = 0; AttributeIndex < AttributeValues.Num(); ++AttributeIndex)
+			{
+				if (const FTransformAnimationAttribute* AttributeValue = AttributeValues[AttributeIndex].GetPtr<FTransformAnimationAttribute>())
+				{
+					const UE::Anim::FAttributeId& AttributeIdentifier = AttributeIdentifiers[AttributeIndex];
+
+					const FTransform AttributeParentTransform = MeshComponent->GetDrawTransform(AttributeIdentifier.GetIndex()) * MeshComponent->GetComponentTransform();
+					const FTransform AttributeTransform = AttributeValue->Value * AttributeParentTransform;
+
+					DrawWireDiamond(PDI, AttributeTransform.ToMatrixNoScale(), 2.0f, FLinearColor(0.0f, 1.0f, 1.0f), SDPG_Foreground);
+					SkeletalDebugRendering::DrawAxes(PDI, AttributeTransform, SDPG_Foreground, 0.0f, 10.0f);
+					//DrawDashedLine(PDI, AttributeTransform.GetLocation(), AttributeParentTransform.GetLocation(), FLinearColor(0.0f, 1.0f, 1.0f), 2.0f, SDPG_World);
+				}
+			}
+		}
+	}
+}
+
 void FAnimationViewportClient::DrawSockets(const UDebugSkelMeshComponent* InPreviewMeshComponent, TArray<USkeletalMeshSocket*>& InSockets, FSelectedSocketInfo InSelectedSocket, FPrimitiveDrawInterface* PDI, bool bUseSkeletonSocketColor)
 {
 	if (InPreviewMeshComponent && InPreviewMeshComponent->SkeletalMesh)
diff --git a/Engine/Source/Editor/Persona/Private/AnimationEditorViewportClient.h b/Engine/Source/Editor/Persona/Private/AnimationEditorViewportClient.h
--- a/Engine/Source/Editor/Persona/Private/AnimationEditorViewportClient.h
+++ b/Engine/Source/Editor/Persona/Private/AnimationEditorViewportClient.h
@@ -91,7 +91,10 @@ class FAnimationViewportClient : public FEditorViewportClient
 protected:
 
 	/** Function to display bone names*/
-	void ShowBoneNames( FCanvas* Canvas, FSceneView* View );
+	void ShowBoneNames(FCanvas* Canvas, FSceneView* View);
+
+	/** Function to display transform attribute names*/
+	void ShowAttributeNames(FCanvas* Canvas, FSceneView* View);
 
 	/** Function to display debug lines generated from skeletal controls in animBP mode */
 	void DrawNodeDebugLines(TArray<FText>& Lines, FCanvas* Canvas, FSceneView* View);
@@ -393,6 +396,9 @@ private:
 	/** Draw Sub set of Bones **/
 	void DrawMeshSubsetBones(const UDebugSkelMeshComponent* MeshComponent, const TArray<int32>& BonesOfInterest, FPrimitiveDrawInterface* PDI) const;
 
+	/** Draws active transform attributes */
+	void DrawAttributes(UDebugSkelMeshComponent* MeshComponent, FPrimitiveDrawInterface* PDI) const;
+
 	/** Draws bones from watched poses*/
 	void DrawWatchedPoses(UDebugSkelMeshComponent * MeshComponent, FPrimitiveDrawInterface* PDI);
 
diff --git a/Engine/Source/Editor/Persona/Private/SAnimViewportToolBar.cpp b/Engine/Source/Editor/Persona/Private/SAnimViewportToolBar.cpp
--- a/Engine/Source/Editor/Persona/Private/SAnimViewportToolBar.cpp
+++ b/Engine/Source/Editor/Persona/Private/SAnimViewportToolBar.cpp
@@ -699,6 +699,7 @@ TSharedRef<SWidget> SAnimViewportToolBar::GenerateCharacterMenu() const
 					SubMenuBuilder.BeginSection("BonesAndSockets", LOCTEXT("CharacterMenu_BonesAndSocketsLabel", "Show"));
 					{
 						SubMenuBuilder.AddMenuEntry(FAnimViewportShowCommands::Get().ShowSockets);
+						SubMenuBuilder.AddMenuEntry(FAnimViewportShowCommands::Get().ShowAttributes);
 						SubMenuBuilder.AddMenuEntry(FAnimViewportShowCommands::Get().ShowBoneNames);
 					}
 					SubMenuBuilder.EndSection();
diff --git a/Engine/Source/Editor/Persona/Private/SAnimationEditorViewport.cpp b/Engine/Source/Editor/Persona/Private/SAnimationEditorViewport.cpp
--- a/Engine/Source/Editor/Persona/Private/SAnimationEditorViewport.cpp
+++ b/Engine/Source/Editor/Persona/Private/SAnimationEditorViewport.cpp
@@ -842,6 +842,13 @@ void SAnimationEditorViewportTabBody::BindCommands()
 		FCanExecuteAction(),
 		FIsActionChecked::CreateSP(this, &SAnimationEditorViewportTabBody::IsShowingSockets));
 
+	// Show transform attributes
+	CommandList.MapAction(
+		ViewportShowMenuCommands.ShowAttributes,
+		FExecuteAction::CreateSP(this, &SAnimationEditorViewportTabBody::OnShowAttributes),
+		FCanExecuteAction(),
+		FIsActionChecked::CreateSP(this, &SAnimationEditorViewportTabBody::IsShowingAttributes));
+
 	// Set bone drawing mode
 	CommandList.BeginGroup(TEXT("BoneDrawingMode"));
 
@@ -1390,6 +1397,23 @@ bool SAnimationEditorViewportTabBody::IsShowingSockets() const
 	return PreviewComponent != NULL && PreviewComponent->bDrawSockets;
 }
 
+void SAnimationEditorViewportTabBody::OnShowAttributes()
+{
+	UDebugSkelMeshComponent* PreviewComponent = GetPreviewScene()->GetPreviewMeshComponent();
+	if (PreviewComponent)
+	{
+		PreviewComponent->bDrawAttributes = !PreviewComponent->bDrawAttributes;
+		PreviewComponent->MarkRenderStateDirty();
+		RefreshViewport();
+	}
+}
+
+bool SAnimationEditorViewportTabBody::IsShowingAttributes() const
+{
+	UDebugSkelMeshComponent* PreviewComponent = GetPreviewScene()->GetPreviewMeshComponent();
+	return PreviewComponent != NULL && PreviewComponent->bDrawAttributes;
+}
+
 void SAnimationEditorViewportTabBody::OnToggleAutoAlignFloor()
 {
 	TSharedRef<FAnimationViewportClient> AnimViewportClient = StaticCastSharedRef<FAnimationViewportClient>(LevelViewportClient.ToSharedRef());
diff --git a/Engine/Source/Editor/Persona/Private/SAnimationEditorViewport.h b/Engine/Source/Editor/Persona/Private/SAnimationEditorViewport.h
--- a/Engine/Source/Editor/Persona/Private/SAnimationEditorViewport.h
+++ b/Engine/Source/Editor/Persona/Private/SAnimationEditorViewport.h
@@ -392,6 +392,12 @@ private:
 	/** Function to check whether socket hit points are displayed or not*/
 	bool IsShowingSockets() const;
 
+	/** Function to show/hide extracted transform attributes */
+	void OnShowAttributes();
+
+	/** Function to check whether extracted transform attributes are displayed or not */
+	bool IsShowingAttributes() const;
+
 	/** Function to show/hide mesh info*/
 	void OnShowDisplayInfo(int32 DisplayInfoMode);
 
diff --git a/Engine/Source/Editor/UnrealEd/Classes/Animation/DebugSkelMeshComponent.h b/Engine/Source/Editor/UnrealEd/Classes/Animation/DebugSkelMeshComponent.h
--- a/Engine/Source/Editor/UnrealEd/Classes/Animation/DebugSkelMeshComponent.h
+++ b/Engine/Source/Editor/UnrealEd/Classes/Animation/DebugSkelMeshComponent.h
@@ -189,6 +189,10 @@ class UNREALED_API UDebugSkelMeshComponent : public USkeletalMeshComponent
 	UPROPERTY(transient)
 	uint32 bDrawSockets:1;
 
+	/** Attribute visualization */
+	UPROPERTY(transient)
+	uint32 bDrawAttributes : 1;
+
 	/** Skeleton sockets visible? */
 	UPROPERTY(transient)
 	uint32 bSkeletonSocketsVisible:1;
