Persona shortcut icon and colour update

diff --git a/Engine/Source/Developer/AssetTools/Private/AssetTypeActions/AssetTypeActions_SkeletalMesh.h b/Engine/Source/Developer/AssetTools/Private/AssetTypeActions/AssetTypeActions_SkeletalMesh.h
--- a/Engine/Source/Developer/AssetTools/Private/AssetTypeActions/AssetTypeActions_SkeletalMesh.h
+++ b/Engine/Source/Developer/AssetTools/Private/AssetTypeActions/AssetTypeActions_SkeletalMesh.h
@@ -15,7 +15,7 @@ public:
 
 	// IAssetTypeActions Implementation
 	virtual FText GetName() const override { return NSLOCTEXT("AssetTypeActions", "AssetTypeActions_SkeletalMesh", "Skeletal Mesh"); }
-	virtual FColor GetTypeColor() const override { return FColor(255,0,255); }
+	virtual FColor GetTypeColor() const override { return FColor(241,163,241); }
 	virtual UClass* GetSupportedClass() const override { return USkeletalMesh::StaticClass(); }
 	virtual bool HasActions ( const TArray<UObject*>& InObjects ) const override { return true; }
 	virtual void GetActions(const TArray<UObject*>& InObjects, struct FToolMenuSection& Section) override;
diff --git a/Engine/Source/Editor/Persona/Private/PersonaAssetFamily.cpp b/Engine/Source/Editor/Persona/Private/PersonaAssetFamily.cpp
--- a/Engine/Source/Editor/Persona/Private/PersonaAssetFamily.cpp
+++ b/Engine/Source/Editor/Persona/Private/PersonaAssetFamily.cpp
@@ -6,6 +6,7 @@
 #include "AssetRegistryModule.h"
 #include "Animation/AnimBlueprint.h"
 #include "PhysicsEngine/PhysicsAsset.h"
+#include "AssetToolsModule.h"
 
 #define LOCTEXT_NAMESPACE "PersonaAssetFamily"
 
@@ -260,6 +261,46 @@ const FSlateBrush* FPersonaAssetFamily::GetAssetTypeDisplayIcon(UClass* InAssetC
 	return nullptr;
 }	
 
+FSlateColor FPersonaAssetFamily::GetAssetTypeDisplayTint(UClass* InAssetClass) const
+{
+	static const FAssetToolsModule& AssetToolsModule = FModuleManager::LoadModuleChecked<FAssetToolsModule>("AssetTools");
+
+	UClass* UseAssetClass = nullptr;
+	if (InAssetClass)
+	{
+		if (InAssetClass->IsChildOf<USkeleton>())
+		{
+			UseAssetClass = USkeleton::StaticClass();
+		}
+		else if (InAssetClass->IsChildOf<UAnimationAsset>())
+		{
+			UseAssetClass = UAnimationAsset::StaticClass();
+		}
+		else if (InAssetClass->IsChildOf<USkeletalMesh>())
+		{
+			UseAssetClass = USkeletalMesh::StaticClass();
+		}
+		else if (InAssetClass->IsChildOf<UAnimBlueprint>())
+		{
+			UseAssetClass = UAnimBlueprint::StaticClass();
+		}
+		else if (InAssetClass->IsChildOf<UPhysicsAsset>())
+		{
+			UseAssetClass = UPhysicsAsset::StaticClass();
+		}
+	}
+
+	if (UseAssetClass)
+	{
+		TWeakPtr<IAssetTypeActions> AssetTypeActions = AssetToolsModule.Get().GetAssetTypeActionsForClass(UseAssetClass);
+		if (AssetTypeActions.IsValid())
+		{
+			return AssetTypeActions.Pin()->GetTypeColor();
+		}
+	}
+	return FSlateColor::UseForeground();
+}
+
 bool FPersonaAssetFamily::IsAssetCompatible(const FAssetData& InAssetData) const
 {
 	UClass* Class = InAssetData.GetClass();
diff --git a/Engine/Source/Editor/Persona/Private/PersonaAssetFamily.h b/Engine/Source/Editor/Persona/Private/PersonaAssetFamily.h
--- a/Engine/Source/Editor/Persona/Private/PersonaAssetFamily.h
+++ b/Engine/Source/Editor/Persona/Private/PersonaAssetFamily.h
@@ -24,6 +24,7 @@ public:
 	virtual void FindAssetsOfType(UClass* InAssetClass, TArray<FAssetData>& OutAssets) const override;
 	virtual FText GetAssetTypeDisplayName(UClass* InAssetClass) const override;
 	virtual const FSlateBrush* GetAssetTypeDisplayIcon(UClass* InAssetClass) const override;
+	virtual FSlateColor GetAssetTypeDisplayTint(UClass* InAssetClass) const override;
 	virtual bool IsAssetCompatible(const FAssetData& InAssetData) const override;
 	virtual UClass* GetAssetFamilyClass(UClass* InClass) const override;
 	virtual void RecordAssetOpened(const FAssetData& InAssetData) override;
diff --git a/Engine/Source/Editor/Persona/Private/SAssetFamilyShortcutBar.cpp b/Engine/Source/Editor/Persona/Private/SAssetFamilyShortcutBar.cpp
--- a/Engine/Source/Editor/Persona/Private/SAssetFamilyShortcutBar.cpp
+++ b/Engine/Source/Editor/Persona/Private/SAssetFamilyShortcutBar.cpp
@@ -95,7 +95,7 @@ public:
 					.Padding(FMargin(28.f, 4.f))
 					[
 						SNew(SImage)
-						.ColorAndOpacity(FSlateColor::UseForeground())
+						.ColorAndOpacity(this, &SAssetShortcut::GetAssetTint)
 						.Image(this, &SAssetShortcut::GetAssetIcon)
 					]
 
@@ -131,7 +131,7 @@ public:
 					.Padding(FMargin(16.f, 4.f))
 					[
 						SNew(SImage)
-						.ColorAndOpacity(FSlateColor::UseForeground())
+						.ColorAndOpacity(this, &SAssetShortcut::GetAssetTint)
 						.Image(this, &SAssetShortcut::GetAssetIcon)
 					]
 
@@ -217,6 +217,15 @@ public:
 		return AssetFamily->GetAssetTypeDisplayIcon(AssetData.GetClass());	
 	}
 
+	FSlateColor GetAssetTint() const
+	{
+		if (GetCheckState() == ECheckBoxState::Checked)
+		{
+			return FSlateColor::UseForeground();
+		}
+		return AssetFamily->GetAssetTypeDisplayTint(AssetData.GetClass());
+	}
+
 	ECheckBoxState GetCheckState() const
 	{
 		if(HostingApp.IsValid())
diff --git a/Engine/Source/Editor/Persona/Public/IAssetFamily.h b/Engine/Source/Editor/Persona/Public/IAssetFamily.h
--- a/Engine/Source/Editor/Persona/Public/IAssetFamily.h
+++ b/Engine/Source/Editor/Persona/Public/IAssetFamily.h
@@ -43,6 +43,9 @@ public:
 	/** Gets the slate brush that represents this asset family */
 	virtual const FSlateBrush* GetAssetTypeDisplayIcon(UClass* InAssetClass) const = 0;
 
+	/** Gets the color to tint the asset display icon */
+	virtual FSlateColor GetAssetTypeDisplayTint(UClass* InAssetClass) const = 0;
+
 	/** Check whether an asset is compatible with this family */
 	virtual bool IsAssetCompatible(const FAssetData& InAssetData) const = 0;
 
