Anim: Remove AnimInstance->Montage_Stop which is stopping all montages when unbinding from Sequencer. Sequencer already stores and restores the specific montages that need to be stopped in FStopPlayingMontageTokenProducer.

diff --git a/Engine/Source/Editor/Persona/Private/AnimationEditorPreviewScene.cpp b/Engine/Source/Editor/Persona/Private/AnimationEditorPreviewScene.cpp
--- a/Engine/Source/Editor/Persona/Private/AnimationEditorPreviewScene.cpp
+++ b/Engine/Source/Editor/Persona/Private/AnimationEditorPreviewScene.cpp
@@ -332,9 +332,11 @@ void FAnimationEditorPreviewScene::RefreshAdditionalMeshes(bool bAllowOverrideBa
 	// remove all components
 	for (USkeletalMeshComponent* Component : AdditionalMeshes)
 	{
-		const UAnimInstance* AnimInst = Component->GetAnimInstance();
+		UAnimInstance* AnimInst = Component->GetAnimInstance();
 		if (AnimInst && AnimInst->IsA(UAnimPreviewAttacheInstance::StaticClass()))
 		{
+			AnimInst->Montage_Stop(0.0f);
+
 			FAnimCustomInstanceHelper::UnbindFromSkeletalMeshComponent<UAnimPreviewAttacheInstance>(Component);
 		}
 		
diff --git a/Engine/Source/Runtime/AnimGraphRuntime/Public/AnimCustomInstanceHelper.h b/Engine/Source/Runtime/AnimGraphRuntime/Public/AnimCustomInstanceHelper.h
--- a/Engine/Source/Runtime/AnimGraphRuntime/Public/AnimCustomInstanceHelper.h
+++ b/Engine/Source/Runtime/AnimGraphRuntime/Public/AnimCustomInstanceHelper.h
@@ -262,8 +262,6 @@ public:
 		
 			if (AnimInstance)
 			{
-				AnimInstance->Montage_Stop(0.0f);
-
 				const TArray<UAnimInstance*>& LinkedInstances = const_cast<const USkeletalMeshComponent*>(InSkeletalMeshComponent)->GetLinkedAnimInstances();
 				for (UAnimInstance* LinkedInstance : LinkedInstances)
 				{
