Fix for PWM SVisibilityWidget crash on Linux

diff --git a/Engine/Source/Editor/Persona/Private/PoseWatchManagerColumnVisibility.cpp b/Engine/Source/Editor/Persona/Private/PoseWatchManagerColumnVisibility.cpp
--- a/Engine/Source/Editor/Persona/Private/PoseWatchManagerColumnVisibility.cpp
+++ b/Engine/Source/Editor/Persona/Private/PoseWatchManagerColumnVisibility.cpp
@@ -12,122 +12,125 @@
 
 #define LOCTEXT_NAMESPACE "PoseWatchManagerColumnVisibility"
 
-class SVisibilityWidget : public SImage
+namespace PoseWatchManager
 {
-public:
-	SLATE_BEGIN_ARGS(SVisibilityWidget) {}
-	SLATE_END_ARGS()
-
-	void Construct(const FArguments& InArgs, TWeakPtr<FPoseWatchManagerColumnVisibility> InWeakColumn, TWeakPtr<IPoseWatchManager> InWeakPoseWatchManager, TWeakPtr<IPoseWatchManagerTreeItem> InWeakTreeItem, const STableRow<FPoseWatchManagerTreeItemPtr>* InRow)
+	class SVisibilityWidget : public SImage
 	{
-		WeakTreeItem = InWeakTreeItem;
-		WeakPoseWatchManager = InWeakPoseWatchManager;
-		WeakColumn = InWeakColumn;
+	public:
+		SLATE_BEGIN_ARGS(SVisibilityWidget) {}
+		SLATE_END_ARGS()
 
-		Row = InRow;
+			void Construct(const FArguments& InArgs, TWeakPtr<FPoseWatchManagerColumnVisibility> InWeakColumn, TWeakPtr<IPoseWatchManager> InWeakPoseWatchManager, TWeakPtr<IPoseWatchManagerTreeItem> InWeakTreeItem, const STableRow<FPoseWatchManagerTreeItemPtr>* InRow)
+		{
+			WeakTreeItem = InWeakTreeItem;
+			WeakPoseWatchManager = InWeakPoseWatchManager;
+			WeakColumn = InWeakColumn;
 
-		SImage::Construct(
-			SImage::FArguments()
-			.ColorAndOpacity(this, &SVisibilityWidget::GetForegroundColor)
-			.Image(this, &SVisibilityWidget::GetBrush)
-		);
+			Row = InRow;
 
-		static const FName NAME_VisibleHoveredBrush = TEXT("Level.VisibleHighlightIcon16x");
-		static const FName NAME_VisibleNotHoveredBrush = TEXT("Level.VisibleIcon16x");
-		static const FName NAME_NotVisibleHoveredBrush = TEXT("Level.NotVisibleHighlightIcon16x");
-		static const FName NAME_NotVisibleNotHoveredBrush = TEXT("Level.NotVisibleIcon16x");
+			SImage::Construct(
+				SImage::FArguments()
+				.ColorAndOpacity(this, &SVisibilityWidget::GetForegroundColor)
+				.Image(this, &SVisibilityWidget::GetBrush)
+			);
 
-		VisibleHoveredBrush = FAppStyle::Get().GetBrush(NAME_VisibleHoveredBrush);
-		VisibleNotHoveredBrush = FAppStyle::Get().GetBrush(NAME_VisibleNotHoveredBrush);
+			static const FName NAME_VisibleHoveredBrush = TEXT("Level.VisibleHighlightIcon16x");
+			static const FName NAME_VisibleNotHoveredBrush = TEXT("Level.VisibleIcon16x");
+			static const FName NAME_NotVisibleHoveredBrush = TEXT("Level.NotVisibleHighlightIcon16x");
+			static const FName NAME_NotVisibleNotHoveredBrush = TEXT("Level.NotVisibleIcon16x");
 
-		NotVisibleHoveredBrush = FAppStyle::Get().GetBrush(NAME_NotVisibleHoveredBrush);
-		NotVisibleNotHoveredBrush = FAppStyle::Get().GetBrush(NAME_NotVisibleNotHoveredBrush);
-	}
+			VisibleHoveredBrush = FAppStyle::Get().GetBrush(NAME_VisibleHoveredBrush);
+			VisibleNotHoveredBrush = FAppStyle::Get().GetBrush(NAME_VisibleNotHoveredBrush);
 
-private:
-	FReply HandleClick()
-	{
-		TSharedPtr<IPoseWatchManagerTreeItem> TreeItem = WeakTreeItem.Pin();
-		TreeItem->SetIsVisible(!TreeItem->GetVisibility());
-		return FReply::Handled();
-	}
+			NotVisibleHoveredBrush = FAppStyle::Get().GetBrush(NAME_NotVisibleHoveredBrush);
+			NotVisibleNotHoveredBrush = FAppStyle::Get().GetBrush(NAME_NotVisibleNotHoveredBrush);
+		}
 
-	/** Called when the mouse button is pressed down on this widget */
-	virtual FReply OnMouseButtonDown(const FGeometry& MyGeometry, const FPointerEvent& MouseEvent) override
-	{
-		if (MouseEvent.GetEffectingButton() != EKeys::LeftMouseButton)
+	private:
+		FReply HandleClick()
 		{
-			return FReply::Unhandled();
+			TSharedPtr<IPoseWatchManagerTreeItem> TreeItem = WeakTreeItem.Pin();
+			TreeItem->SetIsVisible(!TreeItem->GetVisibility());
+			return FReply::Handled();
 		}
 
-		return HandleClick();
-	}
-
-	/** Get the brush for this widget */
-	const FSlateBrush* GetBrush() const
-	{
-		if (IsVisible())
+		/** Called when the mouse button is pressed down on this widget */
+		virtual FReply OnMouseButtonDown(const FGeometry& MyGeometry, const FPointerEvent& MouseEvent) override
 		{
-			return IsHovered() ? VisibleHoveredBrush : VisibleNotHoveredBrush;
+			if (MouseEvent.GetEffectingButton() != EKeys::LeftMouseButton)
+			{
+				return FReply::Unhandled();
+			}
+
+			return HandleClick();
 		}
-		else
+
+		/** Get the brush for this widget */
+		const FSlateBrush* GetBrush() const
 		{
-			return IsHovered() ? NotVisibleHoveredBrush : NotVisibleNotHoveredBrush;
+			if (IsVisible())
+			{
+				return IsHovered() ? VisibleHoveredBrush : VisibleNotHoveredBrush;
+			}
+			else
+			{
+				return IsHovered() ? NotVisibleHoveredBrush : NotVisibleNotHoveredBrush;
+			}
 		}
-	}
 
-	virtual FSlateColor GetForegroundColor() const
-	{
-		auto TreeItem = WeakTreeItem.Pin();
-		const bool bIsSelected = WeakPoseWatchManager.Pin()->GetTree().IsItemSelected(TreeItem.ToSharedRef());
-
-		if (!TreeItem->IsEnabled())
+		virtual FSlateColor GetForegroundColor() const
 		{
-			return FPoseWatchManagerCommonLabelData::DisabledColor;
+			auto TreeItem = WeakTreeItem.Pin();
+			const bool bIsSelected = WeakPoseWatchManager.Pin()->GetTree().IsItemSelected(TreeItem.ToSharedRef());
+
+			if (!TreeItem->IsEnabled())
+			{
+				return FPoseWatchManagerCommonLabelData::DisabledColor;
+			}
+
+			if (IsHovered() && !bIsSelected)
+			{
+				return FAppStyle::Get().GetSlateColor("Colors.ForegroundHover");
+			}
+
+			return FSlateColor::UseForeground();
 		}
 
-		if (IsHovered() && !bIsSelected)
+		static bool IsVisible(const FPoseWatchManagerTreeItemPtr& Item, const TSharedPtr<FPoseWatchManagerColumnVisibility>& Column)
 		{
-			return FAppStyle::Get().GetSlateColor("Colors.ForegroundHover");
+			return (Column.IsValid() && Item && Item.IsValid()) ? Item->GetVisibility() : false;
 		}
 
-		return FSlateColor::UseForeground();
-	}
-
-	static bool IsVisible(const FPoseWatchManagerTreeItemPtr& Item, const TSharedPtr<FPoseWatchManagerColumnVisibility>& Column)
-	{
-		return (Column.IsValid() && Item && Item.IsValid()) ? Item->GetVisibility() : false;
-	}
-
-	bool IsVisible() const
-	{
-		return IsVisible(WeakTreeItem.Pin(), WeakColumn.Pin());
-	}
+		bool IsVisible() const
+		{
+			return IsVisible(WeakTreeItem.Pin(), WeakColumn.Pin());
+		}
 
-	void SetIsVisible(const bool bVisible)
-	{
-		TSharedPtr<IPoseWatchManagerTreeItem> TreeItem = WeakTreeItem.Pin();
-		TreeItem.Get()->SetIsVisible(bVisible);
-		WeakPoseWatchManager.Pin()->Refresh();	
-	}
+		void SetIsVisible(const bool bVisible)
+		{
+			TSharedPtr<IPoseWatchManagerTreeItem> TreeItem = WeakTreeItem.Pin();
+			TreeItem.Get()->SetIsVisible(bVisible);
+			WeakPoseWatchManager.Pin()->Refresh();
+		}
 
-	/** The tree item we relate to */
-	TWeakPtr<IPoseWatchManagerTreeItem> WeakTreeItem;
+		/** The tree item we relate to */
+		TWeakPtr<IPoseWatchManagerTreeItem> WeakTreeItem;
 
-	TWeakPtr<IPoseWatchManager> WeakPoseWatchManager;
+		TWeakPtr<IPoseWatchManager> WeakPoseWatchManager;
 
-	/** Weak pointer back to the column */
-	TWeakPtr<FPoseWatchManagerColumnVisibility> WeakColumn;
+		/** Weak pointer back to the column */
+		TWeakPtr<FPoseWatchManagerColumnVisibility> WeakColumn;
 
-	/** Weak pointer back to the row */
-	const STableRow<FPoseWatchManagerTreeItemPtr>* Row;
+		/** Weak pointer back to the row */
+		const STableRow<FPoseWatchManagerTreeItemPtr>* Row;
 
-	/** Visibility brushes for the various states */
-	const FSlateBrush* VisibleHoveredBrush;
-	const FSlateBrush* VisibleNotHoveredBrush;
-	const FSlateBrush* NotVisibleHoveredBrush;
-	const FSlateBrush* NotVisibleNotHoveredBrush;
-};
+		/** Visibility brushes for the various states */
+		const FSlateBrush* VisibleHoveredBrush;
+		const FSlateBrush* VisibleNotHoveredBrush;
+		const FSlateBrush* NotVisibleHoveredBrush;
+		const FSlateBrush* NotVisibleNotHoveredBrush;
+	};
+}
 
 SHeaderRow::FColumn::FArguments FPoseWatchManagerColumnVisibility::ConstructHeaderRowColumn()
 {
@@ -153,7 +156,7 @@ const TSharedRef<SWidget> FPoseWatchManagerColumnVisibility::ConstructRowWidget(
 		.AutoWidth()
 		.VAlign(VAlign_Center)
 		[
-			SNew(SVisibilityWidget, SharedThis(this), WeakPoseWatchManager, TreeItem, &Row)
+			SNew(PoseWatchManager::SVisibilityWidget, SharedThis(this), WeakPoseWatchManager, TreeItem, &Row)
 		];
 }
 
