Moving anim sequence browser context menu over to UToolsMenu API

diff --git a/Engine/Source/Editor/Persona/Private/SAnimationSequenceBrowser.cpp b/Engine/Source/Editor/Persona/Private/SAnimationSequenceBrowser.cpp
--- a/Engine/Source/Editor/Persona/Private/SAnimationSequenceBrowser.cpp
+++ b/Engine/Source/Editor/Persona/Private/SAnimationSequenceBrowser.cpp
@@ -45,6 +45,8 @@
 #include "SAnimCurveViewer.h"
 #include "IEditableSkeleton.h"
 #include "SAnimCurvePicker.h"
+#include "ToolMenus.h"
+#include "AnimationSequenceBrowserMenuContexts.h"
 
 #define LOCTEXT_NAMESPACE "SequenceBrowser"
 
@@ -508,88 +510,94 @@ void SAnimationSequenceBrowser::OnRequestOpenAsset(const FAssetData& AssetData,
 
 TSharedPtr<SWidget> SAnimationSequenceBrowser::OnGetAssetContextMenu(const TArray<FAssetData>& SelectedAssets)
 {
-	bool bHasSelectedAnimSequence = false;
-	bool bHasSelectedAnimAsset = false;
-	if ( SelectedAssets.Num() )
+	UToolMenus* ToolMenus = UToolMenus::Get();
+	const FName SequenceBrowserContextMenuName = "AnimationSequenceBrowser.SequenceContextMenu";
+	if (!ToolMenus->IsMenuRegistered(SequenceBrowserContextMenuName))
 	{
-		for(auto Iter = SelectedAssets.CreateConstIterator(); Iter; ++Iter)
+		ToolMenus->RegisterMenu(SequenceBrowserContextMenuName);
+	}
+	UToolMenu* Menu = ToolMenus->FindMenu(SequenceBrowserContextMenuName);
+
+	Menu->AddDynamicSection("SequenceContextMenuDynamic", FNewToolMenuDelegate::CreateLambda([this, SelectedAssets](UToolMenu* InMenu)
+	{
+		bool bHasSelectedAnimSequence = false;
+		bool bHasSelectedAnimAsset = false;
+		if (SelectedAssets.Num())
 		{
-			UObject* Asset =  Iter->GetAsset();
-			if(Cast<UAnimSequence>(Asset))
+			for (auto Iter = SelectedAssets.CreateConstIterator(); Iter; ++Iter)
 			{
-				bHasSelectedAnimSequence = true;
-			}
-			if (Cast<UAnimationAsset>(Asset))
-			{
-				bHasSelectedAnimAsset = true;
+				UObject* Asset = Iter->GetAsset();
+				if (Cast<UAnimSequence>(Asset))
+				{
+					bHasSelectedAnimSequence = true;
+				}
+				if (Cast<UAnimationAsset>(Asset))
+				{
+					bHasSelectedAnimAsset = true;
+				}
 			}
 		}
-	}
-
-	FMenuBuilder MenuBuilder(/*bInShouldCloseWindowAfterMenuSelection=*/ true, Commands);
 
-	if(bHasSelectedAnimSequence)
-	{
-		MenuBuilder.BeginSection("AnimationSequenceOptions", LOCTEXT("AnimationHeading", "Animation"));
+		if (bHasSelectedAnimSequence)
 		{
-			MenuBuilder.AddMenuEntry(
-				LOCTEXT("RunCompressionOnAnimations", "Apply Compression"),
-				LOCTEXT("RunCompressionOnAnimations_ToolTip", "Apply a compression scheme from the options given to the selected animations"),
+			FToolMenuSection& Section = InMenu->AddSection("SequenceOptions", LOCTEXT("AssetHeading", "Asset"));
+
+			Section.AddMenuEntry("ApplyCompression",
+				LOCTEXT("ApplyCompression", "Apply Compression"),
+				LOCTEXT("ApplyCompression_ToolTip", "Apply a compression scheme from the options given to the selected animations"),
 				FSlateIcon(),
 				FUIAction(
-				FExecuteAction::CreateSP(this, &SAnimationSequenceBrowser::OnApplyCompression, SelectedAssets),
-				FCanExecuteAction()
+					FExecuteAction::CreateSP(this, &SAnimationSequenceBrowser::OnApplyCompression, SelectedAssets),
+					FCanExecuteAction()
 				)
-				);
+			);
 
-			MenuBuilder.AddMenuEntry(
+			Section.AddMenuEntry("ExportAnimationsToFBX",
 				LOCTEXT("ExportAnimationsToFBX", "Export to FBX"),
 				LOCTEXT("ExportAnimationsToFBX_ToolTip", "Export Animation(s) To FBX"),
 				FSlateIcon(),
 				FUIAction(
-				FExecuteAction::CreateSP(this, &SAnimationSequenceBrowser::OnExportToFBX, SelectedAssets),
-				FCanExecuteAction()
+					FExecuteAction::CreateSP(this, &SAnimationSequenceBrowser::OnExportToFBX, SelectedAssets),
+					FCanExecuteAction()
 				)
-				);
+			);
 
-			MenuBuilder.AddMenuEntry(
+			Section.AddMenuEntry("AddLoopingInterpolation",
 				LOCTEXT("AddLoopingInterpolation", "Add Looping Interpolation"),
 				LOCTEXT("AddLoopingInterpolation_ToolTip", "Add an extra frame at the end of the animation to create better looping"),
 				FSlateIcon(),
 				FUIAction(
-				FExecuteAction::CreateSP(this, &SAnimationSequenceBrowser::OnAddLoopingInterpolation, SelectedAssets),
-				FCanExecuteAction()
+					FExecuteAction::CreateSP(this, &SAnimationSequenceBrowser::OnAddLoopingInterpolation, SelectedAssets),
+					FCanExecuteAction()
 				)
-				);
+			);
 
-			MenuBuilder.AddMenuEntry(
+			Section.AddMenuEntry("ReimportAnimation",
 				LOCTEXT("ReimportAnimation", "Reimport Animation"),
-				LOCTEXT("ReimportAnimation_ToolTip", "Reimport current animaion."),
+				LOCTEXT("ReimportAnimation_ToolTip", "Reimport current animation."),
 				FSlateIcon(),
 				FUIAction(
-				FExecuteAction::CreateSP(this, &SAnimationSequenceBrowser::OnReimportAnimation, SelectedAssets),
-				FCanExecuteAction()
+					FExecuteAction::CreateSP(this, &SAnimationSequenceBrowser::OnReimportAnimation, SelectedAssets),
+					FCanExecuteAction()
 				)
-				);
+			);
 
-			MenuBuilder.AddMenuEntry(
+			Section.AddMenuEntry("SetCurrentPreviewMesh",
 				LOCTEXT("SetCurrentPreviewMesh", "Set Current Preview Mesh"),
 				LOCTEXT("SetCurrentPreviewMesh_ToolTip", "Set current preview mesh to be used when previewed by this asset. This only applies when you open Persona using this asset."),
 				FSlateIcon(),
 				FUIAction(
-				FExecuteAction::CreateSP(this, &SAnimationSequenceBrowser::OnSetCurrentPreviewMesh, SelectedAssets),
-				FCanExecuteAction()
+					FExecuteAction::CreateSP(this, &SAnimationSequenceBrowser::OnSetCurrentPreviewMesh, SelectedAssets),
+					FCanExecuteAction()
 				)
-				);
+			);
 		}
-		MenuBuilder.EndSection();
-	}
 
-	if (SelectedAssets.Num() == 1 && SelectedAssets[0].GetClass()->IsChildOf(USoundWave::StaticClass()))
-	{
-		MenuBuilder.BeginSection("AnimationSequenceAudioOptions", LOCTEXT("AudioOptionsHeading", "Audio"));
+		if (SelectedAssets.Num() == 1 && SelectedAssets[0].GetClass()->IsChildOf(USoundWave::StaticClass()))
 		{
-			MenuBuilder.AddMenuEntry(
+			FToolMenuSection& Section = InMenu->AddSection("SequenceAudioOptions", LOCTEXT("AssetHeading", "Asset"));
+
+			Section.AddMenuEntry("PlayAudio",
 				LOCTEXT("PlayAudio", "Play Audio"),
 				LOCTEXT("PlayAudio_ToolTip", "Play this audio asset as a preview"),
 				FSlateIcon(),
@@ -604,7 +612,7 @@ TSharedPtr<SWidget> SAnimationSequenceBrowser::OnGetAssetContextMenu(const TArra
 			{
 				if (AudioComponent->IsPlaying())
 				{
-					MenuBuilder.AddMenuEntry(
+					Section.AddMenuEntry("StopAudio",
 						LOCTEXT("StopAudio", "Stop Audio"),
 						LOCTEXT("StopAudio_ToolTip", "Stop the currently playing preview audio"),
 						FSlateIcon(),
@@ -616,43 +624,53 @@ TSharedPtr<SWidget> SAnimationSequenceBrowser::OnGetAssetContextMenu(const TArra
 				}
 			}
 		}
-		MenuBuilder.EndSection();
-	}
 
-	MenuBuilder.BeginSection("AnimationSequenceOptions", LOCTEXT("OptionsHeading", "Options") );
-	{
-		MenuBuilder.AddMenuEntry(
-			LOCTEXT("SaveSelectedAssets", "Save"),
-			LOCTEXT("SaveSelectedAssets_ToolTip", "Save the selected assets"),
-			FSlateIcon(FEditorStyle::GetStyleSetName(), "Level.SaveIcon16x"),
-			FUIAction(
-				FExecuteAction::CreateSP( this, &SAnimationSequenceBrowser::SaveSelectedAssets, SelectedAssets),
-				FCanExecuteAction::CreateSP( this, &SAnimationSequenceBrowser::CanSaveSelectedAssets, SelectedAssets)
-				)
-			);
+		{
+			FToolMenuSection& Section = InMenu->AddSection("SequenceGeneralOptions", LOCTEXT("AnimationSequenceGeneralOptions", "Options"));
 
-		MenuBuilder.AddMenuEntry(
-			LOCTEXT("AnimSequenceBase_OpenInNewWindow", "Open In New Window"),
-			LOCTEXT("AnimSequenceBase_OpenInNewWindowTooltip", "Will always open asset in a new window, and not re-use existing window. (Shift+Double-Click)"),
-			FSlateIcon(FEditorStyle::GetStyleSetName(), "ContentBrowser.AssetActions.OpenInExternalEditor"),
-			FUIAction(
-				FExecuteAction::CreateSP(this, &SAnimationSequenceBrowser::OpenInNewWindow, SelectedAssets),
-				FCanExecuteAction()
-			)
-		);
-			
-		MenuBuilder.AddMenuEntry(FGlobalEditorCommonCommands::Get().FindInContentBrowser);
+			Section.AddMenuEntry("SaveSelectedAssets",
+				LOCTEXT("SaveSelectedAssets", "Save"),
+				LOCTEXT("SaveSelectedAssets_ToolTip", "Save the selected assets"),
+				FSlateIcon(FEditorStyle::GetStyleSetName(), "Level.SaveIcon16x"),
+				FUIAction(
+					FExecuteAction::CreateSP(this, &SAnimationSequenceBrowser::SaveSelectedAssets, SelectedAssets),
+					FCanExecuteAction::CreateSP(this, &SAnimationSequenceBrowser::CanSaveSelectedAssets, SelectedAssets)
+				));
+
+			Section.AddMenuEntry("OpenSequenceInNewWindow",
+				LOCTEXT("AnimSequenceBase_OpenInNewWindow", "Open In New Window"),
+				LOCTEXT("AnimSequenceBase_OpenInNewWindowTooltip", "Will always open asset in a new window, and not re-use existing window. (Shift+Double-Click)"),
+				FSlateIcon(FEditorStyle::GetStyleSetName(), "ContentBrowser.AssetActions.OpenInExternalEditor"),
+				FUIAction(
+					FExecuteAction::CreateSP(this, &SAnimationSequenceBrowser::OpenInNewWindow, SelectedAssets),
+					FCanExecuteAction()
+				));
+
+			Section.AddMenuEntry(FGlobalEditorCommonCommands::Get().FindInContentBrowser);
+		}
+	}));
+
+	UAnimationSequenceBrowserContextMenuContext* ContextObject = NewObject<UAnimationSequenceBrowserContextMenuContext>();
+	ContextObject->OwningAnimSequenceBrowser = SharedThis(this);
+	ContextObject->SelectedObjects.Reset();
+	for (auto Iter = SelectedAssets.CreateConstIterator(); Iter; ++Iter)
+	{
+		UObject* Asset = Iter->GetAsset();
+		if (Asset != nullptr)
+		{
+			ContextObject->SelectedObjects.Add(Asset);
+		}	
 	}
-	MenuBuilder.EndSection();
 
-	return MenuBuilder.MakeWidget();
+	FToolMenuContext MenuContext(Commands, nullptr, ContextObject);
+	return ToolMenus->GenerateWidget(SequenceBrowserContextMenuName, MenuContext);
 }
 
 void SAnimationSequenceBrowser::OpenInNewWindow(TArray<FAssetData> AnimationAssets)
 {
-	for(auto Iter = AnimationAssets.CreateConstIterator(); Iter; ++Iter)
+	for (auto Iter = AnimationAssets.CreateConstIterator(); Iter; ++Iter)
 	{
-		UObject* Asset =  Iter->GetAsset();
+		UObject* Asset = Iter->GetAsset();
 		UAnimationAsset* AnimationAsset = Cast<UAnimationAsset>(Asset);
 		if (AnimationAsset)
 		{
@@ -662,7 +680,6 @@ void SAnimationSequenceBrowser::OpenInNewWindow(TArray<FAssetData> AnimationAsse
 	}
 }
 
-
 void SAnimationSequenceBrowser::FindInContentBrowser()
 {
 	TArray<FAssetData> CurrentSelection = GetCurrentSelectionDelegate.Execute();
diff --git a/Engine/Source/Editor/Persona/Public/AnimationSequenceBrowserMenuContexts.h b/Engine/Source/Editor/Persona/Public/AnimationSequenceBrowserMenuContexts.h
new file mode 100644
--- /dev/null
+++ b/Engine/Source/Editor/Persona/Public/AnimationSequenceBrowserMenuContexts.h
@@ -0,0 +1,33 @@
+#pragma once
+
+#include "CoreMinimal.h"
+#include "UObject/ObjectMacros.h"
+#include "UObject/Object.h"
+
+#include "AnimationSequenceBrowserMenuContexts.generated.h"
+
+class SAnimationSequenceBrowser;
+
+UCLASS(BlueprintType)
+class PERSONA_API UAnimationSequenceBrowserContextMenuContext : public UObject
+{
+	GENERATED_BODY()
+
+public:
+	TWeakPtr<SAnimationSequenceBrowser> OwningAnimSequenceBrowser;
+
+	UPROPERTY()
+	TArray<TWeakObjectPtr<UObject>> SelectedObjects;
+
+	UFUNCTION(BlueprintCallable, Category = AnimationEditorExtensions)
+	TArray<UObject*> GetSelectedObjects() const
+	{
+		TArray<UObject*> Result;
+		Result.Reserve(SelectedObjects.Num());
+		for (const TWeakObjectPtr<UObject>& Object : SelectedObjects)
+		{
+			Result.Add(Object.Get());
+		}
+		return Result;
+	}
+};
