Added copyright notice to AnimationSequenceBrowserMenuContexts

diff --git a/Engine/Source/Editor/Persona/Public/AnimationSequenceBrowserMenuContexts.h b/Engine/Source/Editor/Persona/Public/AnimationSequenceBrowserMenuContexts.h
--- a/Engine/Source/Editor/Persona/Public/AnimationSequenceBrowserMenuContexts.h
+++ b/Engine/Source/Editor/Persona/Public/AnimationSequenceBrowserMenuContexts.h
@@ -1,3 +1,5 @@
+// Copyright Epic Games, Inc. All Rights Reserved.
+
 #pragma once
 
 #include "CoreMinimal.h"
