Montage Editor Displays Assets with a Slight Offset In Track
#fix use external scrollbar rather than creating an additional one which caused excessive padding

diff --git a/Engine/Source/Editor/Persona/Private/AnimTimeline/SAnimTimeline.cpp b/Engine/Source/Editor/Persona/Private/AnimTimeline/SAnimTimeline.cpp
--- a/Engine/Source/Editor/Persona/Private/AnimTimeline/SAnimTimeline.cpp
+++ b/Engine/Source/Editor/Persona/Private/AnimTimeline/SAnimTimeline.cpp
@@ -219,6 +219,7 @@ void SAnimTimeline::Construct(const FArguments& InArgs, const TSharedRef<FAnimMo
 								.FillHeight(1.f)
 								[
 									SNew(SScrollBox)
+									.ExternalScrollbar(ScrollBar)
 									+ SScrollBox::Slot()
 									[
 										SNew(SHorizontalBox)
