Exposing Retarget Base Pose editing back in editor after premature deprecation.

diff --git a/Engine/Source/Editor/Persona/Private/SRetargetSources.cpp b/Engine/Source/Editor/Persona/Private/SRetargetSources.cpp
--- a/Engine/Source/Editor/Persona/Private/SRetargetSources.cpp
+++ b/Engine/Source/Editor/Persona/Private/SRetargetSources.cpp
@@ -1,10 +1,22 @@
 // Copyright Epic Games, Inc. All Rights Reserved.
 
 #include "SRetargetSources.h"
+#include "Misc/MessageDialog.h"
 #include "Widgets/SBoxPanel.h"
+#include "Widgets/Text/STextBlock.h"
+#include "EditorStyleSet.h"
+#include "Widgets/Layout/SSeparator.h"
+#include "Widgets/Input/SButton.h"
+#include "Animation/DebugSkelMeshComponent.h"
+#include "Widgets/SToolTip.h"
+#include "IDocumentation.h"
+#include "ScopedTransaction.h"
 #include "SRetargetSourceWindow.h"
+#include "AnimPreviewInstance.h"
 #include "IEditableSkeleton.h"
 #include "PropertyCustomizationHelpers.h"
+#include "Framework/MultiBox/MultiBoxBuilder.h"
+
 
 #define LOCTEXT_NAMESPACE "SRetargetSources"
 
@@ -16,13 +28,75 @@ void SRetargetSources::Construct(
 {
 	EditableSkeletonPtr = InEditableSkeleton;
 	PreviewScenePtr = InPreviewScene;
+	InOnPostUndo.Add(FSimpleDelegate::CreateSP(this, &SRetargetSources::PostUndo));
+
 
 	const FString DocLink = TEXT("Shared/Editors/Persona");
 	ChildSlot
 	[
 		SNew (SVerticalBox)
+
+		+ SVerticalBox::Slot()
+		.Padding(5, 5)
+		.AutoHeight()
+		[
+			SNew(STextBlock)
+			.TextStyle(FEditorStyle::Get(), "Persona.RetargetManager.ImportantText")
+			.Text(LOCTEXT("BasePose_Title", "Edit Retarget Base Pose"))
+		]
+
+		+ SVerticalBox::Slot()
+		.AutoHeight()		// This is required to make the scrollbar work, as content overflows Slate containers by default
+		.VAlign(VAlign_Center)
+		.HAlign(HAlign_Left)
+		.Padding(5, 5)
+		[
+			SNew(SHorizontalBox)
+
+			+SHorizontalBox::Slot()
+			.AutoWidth()
+			.HAlign(HAlign_Center)
+			[
+				SNew(SButton)
+				.OnClicked(FOnClicked::CreateSP(this, &SRetargetSources::OnModifyPose))
+				.HAlign(HAlign_Center)
+				.VAlign(VAlign_Center)
+				.Text(LOCTEXT("ModifyRetargetBasePose_Label", "Modify Pose"))
+				.ToolTipText(LOCTEXT("ModifyRetargetBasePose_Tooltip", "The Retarget Base Pose is used to retarget animation from Skeletal Meshes with varying ref poses to this Skeletal Mesh. \nPrefer fixing ref poses before importing if possible."))
+			]
+
+			+SHorizontalBox::Slot()
+			.AutoWidth()
+			.HAlign(HAlign_Center)
+			[
+				SNew(SButton)
+				.OnClicked(FOnClicked::CreateSP(this, &SRetargetSources::OnViewRetargetBasePose))
+				.HAlign(HAlign_Center)
+				.VAlign(VAlign_Center)
+				.Text(this, &SRetargetSources::GetToggleRetargetBasePose)
+				.ToolTipText(LOCTEXT("ViewRetargetBasePose_Tooltip", "Toggle to View/Edit Retarget Base Pose"))
+			]
+		]
+
+		+SVerticalBox::Slot()
+		.Padding(5, 5)
+		.AutoHeight()
+		[
+			SNew(SSeparator)
+			.Orientation(Orient_Horizontal)
+		]
+		
+		+ SVerticalBox::Slot()
+		.Padding(5, 5)
+		.AutoHeight()
+		[
+			SNew(STextBlock)
+			.TextStyle(FEditorStyle::Get(), "Persona.RetargetManager.ImportantText")
+			.Text(LOCTEXT("RetargetSource_Title", "Manage Retarget Sources"))
+		]
+		
 		+ SVerticalBox::Slot()
-		.Padding(2, 5)
+		.Padding(5, 5)
 		.FillHeight(0.5)
 		[
 			// construct retarget source window
@@ -31,5 +105,332 @@ void SRetargetSources::Construct(
 	];
 }
 
-#undef LOCTEXT_NAMESPACE
+FReply SRetargetSources::OnViewRetargetBasePose() const
+{
+	UDebugSkelMeshComponent * PreviewMeshComp = PreviewScenePtr.Pin()->GetPreviewMeshComponent();
+	if (PreviewMeshComp && PreviewMeshComp->PreviewInstance)
+	{
+		const FScopedTransaction Transaction(LOCTEXT("ViewRetargetBasePose_Action", "Edit Retarget Base Pose"));
+		PreviewMeshComp->PreviewInstance->SetForceRetargetBasePose(!PreviewMeshComp->PreviewInstance->GetForceRetargetBasePose());
+		PreviewMeshComp->Modify();
+		// reset all bone transform since you don't want to keep any bone transform change
+		PreviewMeshComp->PreviewInstance->ResetModifiedBone();
+		// add root 
+		if (PreviewMeshComp->PreviewInstance->GetForceRetargetBasePose())
+		{
+			PreviewMeshComp->BonesOfInterest.Add(0);
+		}
+	}
+
+	return FReply::Handled();
+}
+
+FReply SRetargetSources::OnModifyPose()
+{
+	// create context menu
+	TSharedPtr< SWindow > Parent = FSlateApplication::Get().GetActiveTopLevelWindow();
+	if (Parent.IsValid())
+	{
+		FSlateApplication::Get().PushMenu(
+			Parent.ToSharedRef(),
+			FWidgetPath(),
+			OnModifyPoseContextMenu(),
+			FSlateApplication::Get().GetCursorPos(),
+			FPopupTransitionEffect(FPopupTransitionEffect::TypeInPopup));
+	}
+
+	return FReply::Handled();
+}
+
+TSharedRef<SWidget> SRetargetSources::OnModifyPoseContextMenu() 
+{
+	FMenuBuilder MenuBuilder(false, nullptr);
+
+	MenuBuilder.BeginSection("ModifyPose_Label", LOCTEXT("ModifyPose", "Set Pose"));
+	{
+		FUIAction Action_ReferencePose
+		(
+			FExecuteAction::CreateSP(this, &SRetargetSources::ResetRetargetBasePose)
+		);
+
+		MenuBuilder.AddMenuEntry
+		(
+			LOCTEXT("ModifyPoseContextMenu_Reset", "Reset"),
+			LOCTEXT("ModifyPoseContextMenu_Reset_Desc", "Reset to reference pose"),
+			FSlateIcon(FEditorStyle::GetStyleSetName(), "Profiler.EventGraph.SelectStack"), Action_ReferencePose, NAME_None, EUserInterfaceActionType::Button
+		);
+
+		FUIAction Action_UseCurrentPose
+		(
+			FExecuteAction::CreateSP(this, &SRetargetSources::UseCurrentPose)
+		);
+
+		MenuBuilder.AddMenuEntry
+		(
+			LOCTEXT("ModifyPoseContextMenu_UseCurrentPose", "Use CurrentPose"),
+			LOCTEXT("ModifyPoseContextMenu_UseCurrentPose_Desc", "Use Current Pose"),
+			FSlateIcon(FEditorStyle::GetStyleSetName(), "Profiler.EventGraph.SelectStack"), Action_UseCurrentPose, NAME_None, EUserInterfaceActionType::Button
+		);
+
+		MenuBuilder.AddMenuSeparator();
+
+		MenuBuilder.AddWidget(
+			SNew(SVerticalBox)
+
+			+ SVerticalBox::Slot()
+			.AutoHeight()
+			.Padding(3)
+			[
+				SNew(SObjectPropertyEntryBox)
+				.AllowedClass(UPoseAsset::StaticClass())
+				.OnObjectChanged(this, &SRetargetSources::SetSelectedPose)
+				.OnShouldFilterAsset(this, &SRetargetSources::ShouldFilterAsset)
+				.ObjectPath(this, &SRetargetSources::GetSelectedPose)
+			]
+			+ SVerticalBox::Slot()
+			.AutoHeight()
+			.Padding(3)
+			[
+				SAssignNew(PoseAssetNameWidget, SPoseAssetNameWidget)
+				.OnSelectionChanged(this, &SRetargetSources::SetPoseName)
+			]
+
+			+ SVerticalBox::Slot()
+			.AutoHeight()
+			.Padding(3)
+			[
+				SNew(SButton)
+				.OnClicked(FOnClicked::CreateSP(this, &SRetargetSources::OnImportPose))
+				.IsEnabled(this, &SRetargetSources::CanImportPose)
+				.HAlign(HAlign_Center)
+				.VAlign(VAlign_Center)
+				.Text(LOCTEXT("ImportRetargetBasePose_Label", "Import"))
+				.ToolTipText(LOCTEXT("ImportRetargetBasePose_Tooltip", "Import the selected pose to Retarget Base Pose"))
+			]
+			,
+			FText()
+		);
+
+		if (SelectedPoseAsset.IsValid())
+		{
+			PoseAssetNameWidget->SetPoseAsset(SelectedPoseAsset.Get());
+		}
+		// import pose 
+		MenuBuilder.EndSection();
+	}
+
+	return MenuBuilder.MakeWidget();
+}
+
+bool SRetargetSources::CanImportPose() const
+{
+	return (SelectedPoseAsset.IsValid() && SelectedPoseAsset.Get()->ContainsPose(FName(*SelectedPoseName)));
+}
+
+void SRetargetSources::SetSelectedPose(const FAssetData& InAssetData)
+{
+	if (PoseAssetNameWidget.IsValid())
+	{
+		SelectedPoseAsset = Cast<UPoseAsset>(InAssetData.GetAsset());
+		if (SelectedPoseAsset.IsValid())
+		{
+			PoseAssetNameWidget->SetPoseAsset(SelectedPoseAsset.Get());
+		}
+	}
+}
+
+FString SRetargetSources::GetSelectedPose() const
+{
+	return SelectedPoseAsset->GetPathName();
+}
 
+bool SRetargetSources::ShouldFilterAsset(const FAssetData& InAssetData)
+{
+	if (InAssetData.GetClass() == UPoseAsset::StaticClass() &&
+		EditableSkeletonPtr.IsValid())
+	{
+		FString SkeletonString = FAssetData(&EditableSkeletonPtr.Pin()->GetSkeleton()).GetExportTextName();
+		FAssetDataTagMapSharedView::FFindTagResult Result = InAssetData.TagsAndValues.FindTag("Skeleton");
+		return (!Result.IsSet() || SkeletonString != Result.GetValue());
+	}
+
+	return false;
+}
+
+void SRetargetSources::ResetRetargetBasePose()
+{
+	UDebugSkelMeshComponent * PreviewMeshComp = PreviewScenePtr.Pin()->GetPreviewMeshComponent();
+	if(PreviewMeshComp && PreviewMeshComp->SkeletalMesh)
+	{
+		USkeletalMesh * PreviewMesh = PreviewMeshComp->SkeletalMesh;
+
+		check(PreviewMesh && &EditableSkeletonPtr.Pin()->GetSkeleton() == PreviewMesh->GetSkeleton());
+
+		if(PreviewMesh)
+		{
+			const FScopedTransaction Transaction(LOCTEXT("ResetRetargetBasePose_Action", "Reset Retarget Base Pose"));
+			PreviewMesh->Modify();
+			// reset to original ref pose
+			PreviewMesh->SetRetargetBasePose(PreviewMesh->GetRefSkeleton().GetRefBonePose());
+			TurnOnPreviewRetargetBasePose();
+		}
+	}
+
+	FSlateApplication::Get().DismissAllMenus();
+}
+
+void SRetargetSources::UseCurrentPose()
+{
+	UDebugSkelMeshComponent * PreviewMeshComp = PreviewScenePtr.Pin()->GetPreviewMeshComponent();
+	if (PreviewMeshComp && PreviewMeshComp->SkeletalMesh)
+	{
+		USkeletalMesh * PreviewMesh = PreviewMeshComp->SkeletalMesh;
+
+		check(PreviewMesh && &EditableSkeletonPtr.Pin()->GetSkeleton() == PreviewMesh->GetSkeleton());
+
+		if (PreviewMesh)
+		{
+			const FScopedTransaction Transaction(LOCTEXT("RetargetBasePose_UseCurrentPose_Action", "Retarget Base Pose : Use Current Pose"));
+			PreviewMesh->Modify();
+			// get space bases and calculate local
+			const TArray<FTransform> & SpaceBases = PreviewMeshComp->GetComponentSpaceTransforms();
+			// @todo check to see if skeleton vs preview mesh makes it different for missing bones
+			const FReferenceSkeleton& RefSkeleton = PreviewMesh->GetRefSkeleton();
+			TArray<FTransform> & NewRetargetBasePose = PreviewMesh->GetRetargetBasePose();
+			// if you're using master pose component in preview, this won't work
+			check(PreviewMesh->GetRefSkeleton().GetNum() == SpaceBases.Num());
+			int32 TotalNumBones = PreviewMesh->GetRefSkeleton().GetNum();
+			NewRetargetBasePose.Empty(TotalNumBones);
+			NewRetargetBasePose.AddUninitialized(TotalNumBones);
+
+			for (int32 BoneIndex = 0; BoneIndex < TotalNumBones; ++BoneIndex)
+			{
+				// this is slower, but skeleton can have more bones, so I can't just access
+				// Parent index from Skeleton. Going safer route
+				FName BoneName = PreviewMeshComp->GetBoneName(BoneIndex);
+				FName ParentBoneName = PreviewMeshComp->GetParentBone(BoneName);
+				int32 ParentIndex = RefSkeleton.FindBoneIndex(ParentBoneName);
+
+				if (ParentIndex != INDEX_NONE)
+				{
+					NewRetargetBasePose[BoneIndex] = SpaceBases[BoneIndex].GetRelativeTransform(SpaceBases[ParentIndex]);
+				}
+				else
+				{
+					NewRetargetBasePose[BoneIndex] = SpaceBases[BoneIndex];
+				}
+			}
+
+			// Clear PreviewMeshComp bone modified, they're baked now
+			PreviewMeshComp->PreviewInstance->ResetModifiedBone();
+			TurnOnPreviewRetargetBasePose();
+		}
+	}
+	FSlateApplication::Get().DismissAllMenus();
+}
+
+void SRetargetSources::SetPoseName(TSharedPtr<FString> PoseName, ESelectInfo::Type SelectionType)
+{
+	SelectedPoseName = *PoseName.Get();
+}
+
+FReply SRetargetSources::OnImportPose()
+{
+	if (CanImportPose())
+	{
+		UPoseAsset* RawPoseAsset = SelectedPoseAsset.Get();
+		ImportPose(RawPoseAsset, FName(*SelectedPoseName));
+	}
+
+	FSlateApplication::Get().DismissAllMenus();
+
+	return FReply::Handled();
+}
+
+void SRetargetSources::ImportPose(const UPoseAsset* PoseAsset, const FName& PoseName)
+{
+	// Get transforms from pose (this also converts from additive if necessary)
+	const int32 PoseIndex = PoseAsset->GetPoseIndexByName(PoseName);
+	if (PoseIndex != INDEX_NONE)
+	{
+		TArray<FTransform> PoseTransforms;
+		if (PoseAsset->GetFullPose(PoseIndex, PoseTransforms))
+		{
+			const TArray<FName>	PoseTrackNames = PoseAsset->GetTrackNames();
+
+			ensureAlways(PoseTrackNames.Num() == PoseTransforms.Num());
+
+			// now I have pose, I have to copy to the retarget base pose
+			UDebugSkelMeshComponent * PreviewMeshComp = PreviewScenePtr.Pin()->GetPreviewMeshComponent();
+			if (PreviewMeshComp && PreviewMeshComp->SkeletalMesh)
+			{
+				USkeletalMesh * PreviewMesh = PreviewMeshComp->SkeletalMesh;
+
+				check(PreviewMesh && &EditableSkeletonPtr.Pin()->GetSkeleton() == PreviewMesh->GetSkeleton());
+
+				if (PreviewMesh)
+				{
+					// Check if we have bones for all the tracks. If not, then fail so that the user doesn't end up
+					// with partial or broken retarget setup.
+					for (int32 TrackIndex = 0; TrackIndex < PoseTrackNames.Num(); ++TrackIndex)
+					{
+						const int32 BoneIndex = PreviewMesh->GetRefSkeleton().FindBoneIndex(PoseTrackNames[TrackIndex]);
+						if (BoneIndex == INDEX_NONE)
+						{
+							FMessageDialog::Open(EAppMsgType::Ok, FText::FromString(TEXT("Pose asset does not match the preview mesh skeleton. Aborting.")));
+							return;
+						}
+					}
+					
+					const FScopedTransaction Transaction(LOCTEXT("ImportRetargetBasePose_Action", "Import Retarget Base Pose"));
+					PreviewMesh->Modify();
+
+					// reset to original ref pose first
+					PreviewMesh->SetRetargetBasePose(PreviewMesh->GetRefSkeleton().GetRefBonePose());
+
+					// now override imported pose
+					for (int32 TrackIndex = 0; TrackIndex < PoseTrackNames.Num(); ++TrackIndex)
+					{
+						const int32 BoneIndex = PreviewMesh->GetRefSkeleton().FindBoneIndex(PoseTrackNames[TrackIndex]);
+						PreviewMesh->GetRetargetBasePose()[BoneIndex] = PoseTransforms[TrackIndex];
+					}
+
+					TurnOnPreviewRetargetBasePose();
+				}
+			}
+		}
+	}
+}
+
+void SRetargetSources::PostUndo()
+{
+}
+
+FText SRetargetSources::GetToggleRetargetBasePose() const
+{
+	UDebugSkelMeshComponent * PreviewMeshComp = PreviewScenePtr.Pin()->GetPreviewMeshComponent();
+	if(PreviewMeshComp && PreviewMeshComp->PreviewInstance)
+	{
+		if (PreviewMeshComp->PreviewInstance->GetForceRetargetBasePose())
+		{
+			return LOCTEXT("HideRetargetBasePose_Label", "Hide Pose");
+		}
+		else
+		{
+			return LOCTEXT("ViewRetargetBasePose_Label", "View Pose");
+		}
+	}
+
+	return LOCTEXT("InvalidRetargetBasePose_Label", "No Mesh for Base Pose");
+}
+
+void SRetargetSources::TurnOnPreviewRetargetBasePose()
+{
+	UDebugSkelMeshComponent * PreviewMeshComp = PreviewScenePtr.Pin()->GetPreviewMeshComponent();
+	if (PreviewMeshComp && PreviewMeshComp->PreviewInstance)
+	{
+		PreviewMeshComp->PreviewInstance->SetForceRetargetBasePose(true);
+	}
+}
+
+#undef LOCTEXT_NAMESPACE
diff --git a/Engine/Source/Editor/Persona/Private/SRetargetSources.h b/Engine/Source/Editor/Persona/Private/SRetargetSources.h
--- a/Engine/Source/Editor/Persona/Private/SRetargetSources.h
+++ b/Engine/Source/Editor/Persona/Private/SRetargetSources.h
@@ -4,8 +4,11 @@
 
 #include "CoreMinimal.h"
 #include "Widgets/DeclarativeSyntaxSupport.h"
+#include "Input/Reply.h"
 #include "Widgets/SCompoundWidget.h"
 #include "IPersonaPreviewScene.h"
+	
+#include "SPoseAssetNameWidget.h"
 
 class IEditableSkeleton;
 struct FAssetData;
@@ -22,9 +25,37 @@ public:
 
 private:
 
+	FReply OnViewRetargetBasePose() const;
+	FReply OnModifyPose();
+	FReply OnSaveRetargetBasePose();
+
 	/** The editable skeleton */
 	TWeakPtr<class IEditableSkeleton> EditableSkeletonPtr;
 
 	/** The preview scene  */
 	TWeakPtr<class IPersonaPreviewScene> PreviewScenePtr;
+
+	
+	/** Delegate for undo/redo transaction **/
+	void PostUndo();
+	FText GetToggleRetargetBasePose() const;	
+	void ResetRetargetBasePose();
+	void UseCurrentPose();
+	void ImportPose(const UPoseAsset* PoseAsset, const FName& PoseName);
+
+	TSharedRef<SWidget> OnModifyPoseContextMenu();
+
+	void SetSelectedPose(const FAssetData& InAssetData);
+	FString GetSelectedPose() const;
+	bool ShouldFilterAsset(const FAssetData& InAssetData);
+
+	TSharedPtr<SPoseAssetNameWidget> PoseAssetNameWidget;
+	TWeakObjectPtr<UPoseAsset> SelectedPoseAsset;
+	FString SelectedPoseName;
+
+	void SetPoseName(TSharedPtr<FString> PoseName, ESelectInfo::Type SelectionType);
+	FReply OnImportPose();
+
+	void TurnOnPreviewRetargetBasePose();
+	bool CanImportPose() const;
 };
