Fixed nested structs that are bound to anim node pins not showing their child values

diff --git a/Engine/Source/Editor/AnimGraph/Private/AnimGraphNode_Base.cpp b/Engine/Source/Editor/AnimGraph/Private/AnimGraphNode_Base.cpp
--- a/Engine/Source/Editor/AnimGraph/Private/AnimGraphNode_Base.cpp
+++ b/Engine/Source/Editor/AnimGraph/Private/AnimGraphNode_Base.cpp
@@ -126,6 +126,8 @@ void UAnimGraphNode_Base::SetPinVisibility(bool bInVisible, int32 InOptionalPinI
 		
 		FOptionalPinManager::EvaluateOldShownPins(ShowPinForProperties, OldShownPins, this);
 		ReconstructNode();
+
+		PinVisibilityChangedEvent.Broadcast(bInVisible, InOptionalPinIndex);
 	}
 }
 
diff --git a/Engine/Source/Editor/AnimGraph/Public/AnimGraphNode_Base.h b/Engine/Source/Editor/AnimGraph/Public/AnimGraphNode_Base.h
--- a/Engine/Source/Editor/AnimGraph/Public/AnimGraphNode_Base.h
+++ b/Engine/Source/Editor/AnimGraph/Public/AnimGraphNode_Base.h
@@ -395,6 +395,10 @@ public:
 	DECLARE_EVENT_OneParam(UAnimGraphNode_Base, FOnNodePropertyChangedEvent, FPropertyChangedEvent&);
 	FOnNodePropertyChangedEvent& OnNodePropertyChanged() { return PropertyChangeEvent;	}
 
+	// Event that observers can bind to so that they are notified about changes to pin visibility
+	DECLARE_EVENT_TwoParams(UAnimGraphNode_Base, FPinVisibilityChangedEvent, bool /*bInVisible*/, int32 /*InOptionalPinIndex*/);
+	FPinVisibilityChangedEvent& OnPinVisibilityChanged() { return PinVisibilityChangedEvent; }
+	
 	/**
 	 * Helper function to check whether a pin is valid and linked to something else in the graph
 	 * @param	InPinName		The name of the pin @see UEdGraphNode::FindPin
@@ -593,6 +597,8 @@ protected:
 
 	FOnNodeTitleChangedEvent NodeTitleChangedEvent;
 
+	FPinVisibilityChangedEvent PinVisibilityChangedEvent;
+
 	// Helper function used to refresh the type of a binding
 	void RecalculateBindingType(FAnimGraphNodePropertyBinding& InBinding);
 	
diff --git a/Engine/Source/Editor/Persona/Private/AnimGraphNodeDetails.cpp b/Engine/Source/Editor/Persona/Private/AnimGraphNodeDetails.cpp
--- a/Engine/Source/Editor/Persona/Private/AnimGraphNodeDetails.cpp
+++ b/Engine/Source/Editor/Persona/Private/AnimGraphNodeDetails.cpp
@@ -58,6 +58,8 @@ TSharedRef<IDetailCustomization> FAnimGraphNodeDetails::MakeInstance()
 
 void FAnimGraphNodeDetails::CustomizeDetails(class IDetailLayoutBuilder& DetailBuilder)
 {
+	DetailLayoutBuilder = &DetailBuilder;
+	
 	TArray< TWeakObjectPtr<UObject> > SelectedObjectsList;
 	DetailBuilder.GetObjectsBeingCustomized(SelectedObjectsList);
 
@@ -75,6 +77,8 @@ void FAnimGraphNodeDetails::CustomizeDetails(class IDetailLayoutBuilder& DetailB
 		return;
 	}
 
+	AnimGraphNode->OnPinVisibilityChanged().AddSP(this, &FAnimGraphNodeDetails::OnPinVisibilityChanged);
+
 	// make sure type matches with all the nodes. 
 	const UAnimGraphNode_Base* FirstNodeType = AnimGraphNode;
 	for (int32 Index = 1; Index < SelectedObjectsList.Num(); ++Index)
@@ -85,6 +89,10 @@ void FAnimGraphNodeDetails::CustomizeDetails(class IDetailLayoutBuilder& DetailB
 			// if type mismatches, multi selection doesn't work, just return
 			return;
 		}
+		else
+		{
+			CurrentNode->OnPinVisibilityChanged().AddSP(this, &FAnimGraphNodeDetails::OnPinVisibilityChanged);
+		}
 	}
 
 	TargetSkeleton = AnimGraphNode->GetAnimBlueprint()->TargetSkeleton;
@@ -195,8 +203,9 @@ void FAnimGraphNodeDetails::CustomizeDetails(class IDetailLayoutBuilder& DetailB
 						}
 
 						NameWidget = PropertyNameWidget;
-						
-						PropertyRow.CustomWidget()
+
+						const bool bShowChildren = GetVisibilityOfProperty(ShowHidePropertyHandle) == EVisibility::Visible;
+						PropertyRow.CustomWidget(bShowChildren)
 						.NameContent()
 						.MinDesiredWidth(Row.NameWidget.MinWidth)
 						.MaxDesiredWidth(Row.NameWidget.MaxWidth)
@@ -209,28 +218,6 @@ void FAnimGraphNodeDetails::CustomizeDetails(class IDetailLayoutBuilder& DetailB
 						[
 							ValueWidget.ToSharedRef()
 						];
-
-						// Set visibility of children based on parent if this is an array property
-						if(TargetPropertyHandle->AsArray().IsValid())
-						{
-							uint32 NumChildren = 0;
-							TargetPropertyHandle->GetNumChildren(NumChildren);
-							for(uint32 ChildIndex = 0; ChildIndex < NumChildren; ++ChildIndex)
-							{
-								TSharedPtr<IPropertyHandle> ChildHandle = TargetPropertyHandle->GetChildHandle(ChildIndex);
-								IDetailPropertyRow& ChildRow = DetailBuilder.AddPropertyToCategory(ChildHandle);
-
-								ChildRow.Visibility(MakeAttributeLambda([this, WeakShowHidePropertyHandle = TWeakPtr<IPropertyHandle>(ShowHidePropertyHandle)]()
-								{
-									if(TSharedPtr<IPropertyHandle> PinnedShowHidePropertyHandle = WeakShowHidePropertyHandle.Pin())
-									{
-										return GetVisibilityOfProperty(PinnedShowHidePropertyHandle.ToSharedRef());
-									}
-
-									return EVisibility::Visible;
-								}));
-							}
-						}
 					}
 					else if (InternalCustomWidget != SNullWidget::NullWidget)
 					{
@@ -313,6 +300,10 @@ void FAnimGraphNodeDetails::OnBlendProfileChanged(UBlendProfile* NewProfile, TSh
 	}
 }
 
+void FAnimGraphNodeDetails::OnPinVisibilityChanged(bool bInIsVisible, int32 InOptionalPinIndex)
+{
+	DetailLayoutBuilder->ForceRefreshDetails();
+}
 
 TSharedRef<IPropertyTypeCustomization> FInputScaleBiasCustomization::MakeInstance() 
 {
diff --git a/Engine/Source/Editor/Persona/Private/AnimGraphNodeDetails.h b/Engine/Source/Editor/Persona/Private/AnimGraphNodeDetails.h
--- a/Engine/Source/Editor/Persona/Private/AnimGraphNodeDetails.h
+++ b/Engine/Source/Editor/Persona/Private/AnimGraphNodeDetails.h
@@ -63,11 +63,17 @@ protected:
 	/** Called when a blend profile is selected */
 	void OnBlendProfileChanged(UBlendProfile* NewProfile, TSharedPtr<IPropertyHandle> PropertyHandle);
 
+	/** Called when pin visibility changes */
+	void OnPinVisibilityChanged(bool bInIsVisible, int32 InOptionalPinIndex);
+	
 	/** The skeleton we're operating on */
-	USkeleton* TargetSkeleton;
+	USkeleton* TargetSkeleton = nullptr;
 
 	/** Path to the current blueprints skeleton to allow us to filter asset pickers */
 	FString TargetSkeletonName;
+
+	/** Builder that built this customization (used for refresh) */
+	IDetailLayoutBuilder* DetailLayoutBuilder = nullptr;
 };
 
 /////////////////////////////////////////////////////
