Made it so that proportionally linked notifies are not moved when the length of a sequence changes, and fixed some clipping behaviour that was scaling the position of notifies rather than actually clipping them.

diff --git a/Engine/Source/Editor/Persona/Private/AnimTimeline/AnimModel_AnimMontage.cpp b/Engine/Source/Editor/Persona/Private/AnimTimeline/AnimModel_AnimMontage.cpp
--- a/Engine/Source/Editor/Persona/Private/AnimTimeline/AnimModel_AnimMontage.cpp
+++ b/Engine/Source/Editor/Persona/Private/AnimTimeline/AnimModel_AnimMontage.cpp
@@ -256,14 +256,12 @@ bool FAnimModel_AnimMontage::ClampToEndTime(float NewEndTime)
 	bool bClampingNeeded = (SequenceLength > 0.f && NewEndTime < SequenceLength);
 	if(bClampingNeeded)
 	{
-		float ratio = NewEndTime / AnimMontage->GetPlayLength();
-
 		for(int32 i=0; i < AnimMontage->CompositeSections.Num(); i++)
 		{
 			if(AnimMontage->CompositeSections[i].GetTime() > NewEndTime)
 			{
 				float CurrentTime = AnimMontage->CompositeSections[i].GetTime();
-				AnimMontage->CompositeSections[i].SetTime(CurrentTime * ratio);
+				AnimMontage->CompositeSections[i].SetTime(NewEndTime);
 			}
 		}
 
@@ -274,7 +272,7 @@ bool FAnimModel_AnimMontage::ClampToEndTime(float NewEndTime)
 
 			if(NotifyTime >= NewEndTime)
 			{
-				Notify.SetTime(NotifyTime * ratio);
+				Notify.SetTime(NewEndTime);
 				Notify.TriggerTimeOffset = GetTriggerTimeOffsetForType(AnimMontage->CalculateOffsetForNotify(Notify.GetTime()));
 			}
 		}
diff --git a/Engine/Source/Runtime/Engine/Private/Animation/AnimSequenceBase.cpp b/Engine/Source/Runtime/Engine/Private/Animation/AnimSequenceBase.cpp
--- a/Engine/Source/Runtime/Engine/Private/Animation/AnimSequenceBase.cpp
+++ b/Engine/Source/Runtime/Engine/Private/Animation/AnimSequenceBase.cpp
@@ -1007,6 +1007,12 @@ void UAnimSequenceBase::OnModelModified(const EAnimDataModelNotifyType& NotifyTy
 			const float InsertTime = T0;
 			for (FAnimNotifyEvent& Notify : Notifies)
 			{
+				// Proportional notifies don't need to be adjusted
+				if (Notify.GetLinkMethod() == EAnimLinkMethod::Proportional)
+				{
+					continue;
+				}
+
 				float CurrentTime = Notify.GetTime();
 				float NewDuration = 0.f;
 
@@ -1053,6 +1059,12 @@ void UAnimSequenceBase::OnModelModified(const EAnimDataModelNotifyType& NotifyTy
 			// re-locate notifies
 			for (FAnimNotifyEvent& Notify : Notifies)
 			{
+				// Proportional notifies don't need to be adjusted
+				if (Notify.GetLinkMethod() == EAnimLinkMethod::Proportional)
+				{
+					continue;
+				}
+
 				float CurrentTime = Notify.GetTime();
 				float NewDuration = 0.f;
 				
