Cloth - Update tooltip on the delete cloth asset button.

diff --git a/Engine/Source/Editor/Persona/Private/PersonaMeshDetails.cpp b/Engine/Source/Editor/Persona/Private/PersonaMeshDetails.cpp
--- a/Engine/Source/Editor/Persona/Private/PersonaMeshDetails.cpp
+++ b/Engine/Source/Editor/Persona/Private/PersonaMeshDetails.cpp
@@ -5892,7 +5892,7 @@ void FPersonaMeshDetails::OnGenerateElementForClothingAsset( TSharedRef<IPropert
 			.ContentPadding(0)
 			.ForegroundColor( FSlateColor::UseForeground() )
 			.ButtonColorAndOpacity(FLinearColor(1.0f,1.0f,1.0f,0.0f))
-			.ToolTipText(LOCTEXT("RemoveApexFileTip", "Remove this APEX asset"))
+			.ToolTipText(LOCTEXT("RemoveClothingAssetTip", "Remove this clothing asset"))
 			[ 
 				SNew( SImage )
 				.Image( FEditorStyle::GetBrush("PropertyWindow.Button_Clear") )
