Cloth Painting - Fix crashes in Cloth Painting Edit Mode caused by invalid static_casts in PersonaEditorModeManager

==BUG==
This bug has manifested as a couple of Mac crashes when Cloth Painting. The cause of the crash is that in PersonaEditorModeManager, it assumed that all FEdModes are IPersonaEditModes, and static_casts the mode to an IPersonaEditMode and calls virtual functions on them. ClothPaintingEditModes are FEdModes, but not IPersonaEditModes, so this would end up calling a different function which could corrupt its data. Currently, on Mac, this would corrupt the SharedPtr to the LegacyEditorMode in the ULegacyEdModeWrapper, causing a crash.

Any action which FPersonaEditorModeManager::GetCameraTarget to be called while in the Cloth Paint edit mode, would cause this crash. The two reported issues were toggling a Mask checkbox (UE-136071) and undoing painting (UE-141936).

==FIX==
This is a minimal fix to workaround this inheritance problem. UE-143249 exists for a better solution to this problem.

Explicitly check that the FEditorModeID is not the ClothPaintMode ID before static_casting to IPersonaEditMode.

==TESTING==
Tested that true IPersonaEditModes continue to work correctly (including checking in the debugger that they continue to call GetCameraTarget) as well as ClothPaintingEditMode no longer crashes as it does not call GetCameraTarget.

diff --git a/Engine/Source/Editor/Persona/Private/PersonaEditorModeManager.cpp b/Engine/Source/Editor/Persona/Private/PersonaEditorModeManager.cpp
--- a/Engine/Source/Editor/Persona/Private/PersonaEditorModeManager.cpp
+++ b/Engine/Source/Editor/Persona/Private/PersonaEditorModeManager.cpp
@@ -12,13 +12,17 @@ bool FPersonaEditorModeManager::GetCameraTarget(FSphere& OutTarget) const
 	for (UEdMode* Mode : ActiveScriptableModes)
 	{
 		FEdMode* LegacyMode = Mode->AsLegacyMode();
-		if (IPersonaEditMode* EditMode = static_cast<IPersonaEditMode*>(LegacyMode))
+		// Hack for UE-136071, UE-141936. ClothPaintModes are not IPersonaEditModes, but are FEdModes.
+		if (LegacyMode && LegacyMode->GetID() != FEditorModeID("ClothPaintMode"))
 		{
-			FSphere Target;
-			if (EditMode->GetCameraTarget(Target))
+			if (IPersonaEditMode* EditMode = static_cast<IPersonaEditMode*>(LegacyMode))
 			{
-				OutTarget = Target;
-				return true;
+				FSphere Target;
+				if (EditMode->GetCameraTarget(Target))
+				{
+					OutTarget = Target;
+					return true;
+				}
 			}
 		}
 	}
@@ -31,9 +35,13 @@ void FPersonaEditorModeManager::GetOnScreenDebugInfo(TArray<FText>& OutDebugText
 	for (UEdMode* Mode : ActiveScriptableModes)
 	{
 		FEdMode* LegacyMode = Mode->AsLegacyMode();
-		if (IPersonaEditMode* EditMode = static_cast<IPersonaEditMode*>(LegacyMode))
+		// Hack for UE-136071, UE-141936. ClothPaintModes are not IPersonaEditModes, but are FEdModes.
+		if (LegacyMode && LegacyMode->GetID() != FEditorModeID("ClothPaintMode"))
 		{
-			EditMode->GetOnScreenDebugInfo(OutDebugText);
+			if (IPersonaEditMode* EditMode = static_cast<IPersonaEditMode*>(LegacyMode))
+			{
+				EditMode->GetOnScreenDebugInfo(OutDebugText);
+			}
 		}
 	}
 }
