Fix for blend profile drop down not showing for linked layer blend settings

diff --git a/Engine/Source/Editor/Persona/Private/BlendProfileCustomization.cpp b/Engine/Source/Editor/Persona/Private/BlendProfileCustomization.cpp
--- a/Engine/Source/Editor/Persona/Private/BlendProfileCustomization.cpp
+++ b/Engine/Source/Editor/Persona/Private/BlendProfileCustomization.cpp
@@ -97,15 +97,23 @@ void FBlendProfileCustomization::OnBlendProfileChanged(UBlendProfile* NewProfile
 
 USkeleton* FBlendProfileCustomization::GetSkeletonFromOuter(const UObject* Outer)
 {
+	const UAnimBlueprint* AnimBlueprint = nullptr;
 	if (const UEdGraphNode* OuterEdGraphNode = Cast<UEdGraphNode>(Outer))
 	{
-		// If outer is a node, we can get the skeleton from its anim blueprint.
-		if (const UAnimBlueprint* AnimBlueprint = Cast<UAnimBlueprint>(FBlueprintEditorUtils::FindBlueprintForNode(OuterEdGraphNode)))
-		{
-			return AnimBlueprint->TargetSkeleton;
-		}
+		AnimBlueprint = Cast<UAnimBlueprint>(FBlueprintEditorUtils::FindBlueprintForNode(OuterEdGraphNode));
+	}
+	else if (const UEdGraph* OuterEdGraph = Cast<UEdGraph>(Outer))
+	{
+		AnimBlueprint = Cast<UAnimBlueprint>(FBlueprintEditorUtils::FindBlueprintForGraph(OuterEdGraph));
 	}
-	else if (const UAnimationAsset* OuterAnimAsset = Cast<UAnimationAsset>(Outer))
+
+	// If outer belongs to an anim blueprint, grab its skeleton.
+	if (AnimBlueprint)
+	{
+		return AnimBlueprint->TargetSkeleton;
+	}
+
+	if (const UAnimationAsset* OuterAnimAsset = Cast<UAnimationAsset>(Outer))
 	{
 		// If outer is an anim asset, grab the skeleton
 		return OuterAnimAsset->GetSkeleton();
