Fix crash when selecting a slot node in a template animation blueprint

Exposes the slot name 'raw' rather than selectable via a skeleton-bound widget

diff --git a/Engine/Source/Editor/Persona/Private/Customization/AnimGraphNodeSlotDetails.cpp b/Engine/Source/Editor/Persona/Private/Customization/AnimGraphNodeSlotDetails.cpp
--- a/Engine/Source/Editor/Persona/Private/Customization/AnimGraphNodeSlotDetails.cpp
+++ b/Engine/Source/Editor/Persona/Private/Customization/AnimGraphNodeSlotDetails.cpp
@@ -29,14 +29,12 @@ void FAnimGraphNodeSlotDetails::CustomizeDetails(class IDetailLayoutBuilder& Det
 {
 	SlotNodeNamePropertyHandle = DetailBuilder.GetProperty(TEXT("Node.SlotName"));
 
-	DetailBuilder.HideProperty(SlotNodeNamePropertyHandle);
-
 	TArray<UObject*> OuterObjects;
 	SlotNodeNamePropertyHandle->GetOuterObjects(OuterObjects);
 
 	check(SlotNodeNamePropertyHandle.IsValid());
 
-	Skeleton = NULL;
+	Skeleton = nullptr;
 
 	// only support one object, it doesn't make sense to have same slot node name in multiple nodes
 	for (auto Object : OuterObjects)
@@ -57,56 +55,72 @@ void FAnimGraphNodeSlotDetails::CustomizeDetails(class IDetailLayoutBuilder& Det
 		}
 	}
 
-	check (Skeleton);
-
-	// this is used for 2 things, to create name, but also for another pop-up box to show off, it's a bit hacky to use this, but I need widget to parent on
-	TSharedRef<SWidget> SlotNodePropertyNameWidget = SlotNodeNamePropertyHandle->CreatePropertyNameWidget();
-
-	// fill combo with groups and slots
-	RefreshComboLists();
+	if(Skeleton == nullptr)
+	{
+		IDetailCategoryBuilder& SlotNameGroup = DetailBuilder.EditCategory(TEXT("Settings"));
+		SlotNameGroup.AddCustomRow(LOCTEXT("SlotNameTitleLabel", "Slot Name"))
+		.NameContent()
+		[
+			SlotNodeNamePropertyHandle->CreatePropertyNameWidget()
+		]
+		.ValueContent()
+		[
+			SlotNodeNamePropertyHandle->CreatePropertyValueWidget()
+		];
+	}
+	else
+	{
+		DetailBuilder.HideProperty(SlotNodeNamePropertyHandle);
+		
+		// this is used for 2 things, to create name, but also for another pop-up box to show off, it's a bit hacky to use this, but I need widget to parent on
+		TSharedRef<SWidget> SlotNodePropertyNameWidget = SlotNodeNamePropertyHandle->CreatePropertyNameWidget();
 
-	int32 FoundIndex = SlotNameList.Find(SlotNameComboSelectedName);
-	TSharedPtr<FString> ComboBoxSelectedItem = SlotNameComboListItems[FoundIndex];
+		// fill combo with groups and slots
+		RefreshComboLists();
 
-	// create combo box
-	IDetailCategoryBuilder& SlotNameGroup = DetailBuilder.EditCategory(TEXT("Settings"));
-	SlotNameGroup.AddCustomRow(LOCTEXT("SlotNameTitleLabel", "Slot Name"))
-	.NameContent()
-	[
-		SlotNodePropertyNameWidget
-	]
-	.ValueContent()
-	.MinDesiredWidth(125.f * 3.f)
-	.MaxDesiredWidth(125.f * 3.f)
-	[
-		SNew(SHorizontalBox)
+		int32 FoundIndex = SlotNameList.Find(SlotNameComboSelectedName);
+		TSharedPtr<FString> ComboBoxSelectedItem = SlotNameComboListItems[FoundIndex];
 
-		+SHorizontalBox::Slot()
-		.AutoWidth()
+		// create combo box
+		IDetailCategoryBuilder& SlotNameGroup = DetailBuilder.EditCategory(TEXT("Settings"));
+		SlotNameGroup.AddCustomRow(LOCTEXT("SlotNameTitleLabel", "Slot Name"))
+		.NameContent()
 		[
-			SAssignNew(SlotNameComboBox, STextComboBox)
-			.OptionsSource(&SlotNameComboListItems)
-			.OnSelectionChanged(this, &FAnimGraphNodeSlotDetails::OnSlotNameChanged)
-			.OnComboBoxOpening(this, &FAnimGraphNodeSlotDetails::OnSlotListOpening)
-			.InitiallySelectedItem(ComboBoxSelectedItem)
-			.ContentPadding(2)
-			.ToolTipText(FText::FromString(*ComboBoxSelectedItem))
+			SlotNodePropertyNameWidget
 		]
-
-		+SHorizontalBox::Slot()
-		.AutoWidth()
+		.ValueContent()
+		.MinDesiredWidth(125.f * 3.f)
+		.MaxDesiredWidth(125.f * 3.f)
 		[
-			SNew(SButton)
-			.Text(LOCTEXT("AnimSlotNode_DetailPanelManageButtonLabel", "Anim Slot Manager"))
-			.ToolTipText(LOCTEXT("AnimSlotNode_DetailPanelManageButtonToolTipText", "Open Anim Slot Manager to edit Slots and Groups."))
-			.OnClicked(this, &FAnimGraphNodeSlotDetails::OnOpenAnimSlotManager)
-			.Content()
+			SNew(SHorizontalBox)
+
+			+SHorizontalBox::Slot()
+			.AutoWidth()
 			[
-				SNew(SImage)
-				.Image(FEditorStyle::GetBrush("MeshPaint.FindInCB"))
+				SAssignNew(SlotNameComboBox, STextComboBox)
+				.OptionsSource(&SlotNameComboListItems)
+				.OnSelectionChanged(this, &FAnimGraphNodeSlotDetails::OnSlotNameChanged)
+				.OnComboBoxOpening(this, &FAnimGraphNodeSlotDetails::OnSlotListOpening)
+				.InitiallySelectedItem(ComboBoxSelectedItem)
+				.ContentPadding(2)
+				.ToolTipText(FText::FromString(*ComboBoxSelectedItem))
 			]
-		]
-	];
+
+			+SHorizontalBox::Slot()
+			.AutoWidth()
+			[
+				SNew(SButton)
+				.Text(LOCTEXT("AnimSlotNode_DetailPanelManageButtonLabel", "Anim Slot Manager"))
+				.ToolTipText(LOCTEXT("AnimSlotNode_DetailPanelManageButtonToolTipText", "Open Anim Slot Manager to edit Slots and Groups."))
+				.OnClicked(this, &FAnimGraphNodeSlotDetails::OnOpenAnimSlotManager)
+				.Content()
+				[
+					SNew(SImage)
+					.Image(FEditorStyle::GetBrush("MeshPaint.FindInCB"))
+				]
+			]
+		];
+	}
 }
 
 void FAnimGraphNodeSlotDetails::RefreshComboLists(bool bOnlyRefreshIfDifferent /*= false*/)
