[UE][Feature] RBAN Better Debug Draw - 2nd attempt
- Added Debug Visualization of physics bodies and constraints for RBAN nodes in AnimBP editor
- Added checkboxes to filter debug Visualization to Phat skeleton tree
- Created a new PhysicsAssetRenderSettings class that incorporates the debug rendering and filtering settings from Phat
- Created a new PhysicsAssetRenderUtilities namespace that incorporates the debug rendering and filtering code from Phat
- Synchronize debug visualization of physics bodies and constraints between Phat and AnimBP editors

oringinaly submitted as cl-19242421 but failed on non-unity build so backed out

diff --git a/Engine/Source/Editor/AnimGraph/AnimGraph.Build.cs b/Engine/Source/Editor/AnimGraph/AnimGraph.Build.cs
--- a/Engine/Source/Editor/AnimGraph/AnimGraph.Build.cs
+++ b/Engine/Source/Editor/AnimGraph/AnimGraph.Build.cs
@@ -23,6 +23,7 @@ public class AnimGraph : ModuleRules
 				"Slate",
 				"AnimGraphRuntime",
 				"BlueprintGraph",
+				"Persona"
 			}
 		);
 
@@ -54,7 +55,6 @@ public class AnimGraph : ModuleRules
 
         PrivateIncludePathModuleNames.AddRange(
             new string[] {
-                "Persona",
                 "SkeletonEditor",
                 "AdvancedPreviewScene",
                 "AnimationBlueprintEditor",
diff --git a/Engine/Source/Editor/AnimGraph/Private/AnimGraphNode_RigidBody.cpp b/Engine/Source/Editor/AnimGraph/Private/AnimGraphNode_RigidBody.cpp
--- a/Engine/Source/Editor/AnimGraph/Private/AnimGraphNode_RigidBody.cpp
+++ b/Engine/Source/Editor/AnimGraph/Private/AnimGraphNode_RigidBody.cpp
@@ -1,8 +1,17 @@
 // Copyright Epic Games, Inc. All Rights Reserved.
 
 #include "AnimGraphNode_RigidBody.h"
+#include "AnimNodeEditModes.h"
 #include "Kismet2/CompilerResultsLog.h"
 #include "BoneControllers/AnimNode_RigidBody.h"
+#include "PhysicsAssetRenderUtils.h"
+
+// Details includes
+#include "PropertyHandle.h"
+#include "DetailLayoutBuilder.h"
+#include "DetailWidgetRow.h"
+#include "DetailCategoryBuilder.h"
+#include "Widgets/Input/SButton.h"
 
 /////////////////////////////////////////////////////
 // UAnimGraphNode_RigidBody
@@ -41,4 +50,121 @@ void UAnimGraphNode_RigidBody::ValidateAnimNodeDuringCompilation(USkeleton* ForS
 	Super::ValidateAnimNodeDuringCompilation(ForSkeleton, MessageLog);
 }
 
+void UAnimGraphNode_RigidBody::Draw(FPrimitiveDrawInterface* PDI, USkeletalMeshComponent* PreviewSkelMeshComp) const
+{
+	if (const FAnimNode_RigidBody* const RuntimeRigidBodyNode = GetDebuggedAnimNode<FAnimNode_RigidBody>())
+	{
+		if (UPhysicsAsset* const PhysicsAsset = RuntimeRigidBodyNode->GetPhysicsAsset())
+		{
+			PhysicsAssetRender::DebugDraw(PreviewSkelMeshComp, PhysicsAsset, PDI);
+		}
+	}
+}
+
+void UAnimGraphNode_RigidBody::CustomizeDetails(IDetailLayoutBuilder& DetailBuilder)
+{
+	Super::CustomizeDetails(DetailBuilder);
+
+	IDetailCategoryBuilder& ViewportCategory = DetailBuilder.EditCategory(TEXT("Debug Visualization"));
+	FDetailWidgetRow& WidgetRow = ViewportCategory.AddCustomRow(LOCTEXT("ToggleDebugVisualizationButtonRow", "DebugVisualization"));
+	FAnimNode_RigidBody* const RigidBodyNode = static_cast<FAnimNode_RigidBody*>(GetDebuggedAnimNode());
+
+	WidgetRow
+		[
+			SNew(SHorizontalBox)
+			// Show/Hide Bodies button.
+			+ SHorizontalBox::Slot()
+			[
+				SNew(SButton)
+				.HAlign(HAlign_Center)
+				.VAlign(VAlign_Center)
+				.OnClicked_Lambda([this](){ this->ToggleBodyVisibility(); return FReply::Handled(); })
+				.ButtonColorAndOpacity_Lambda([this](){ return (AreAnyBodiesHidden()) ? FAppStyle::Get().GetSlateColor("Colors.AccentRed") : FAppStyle::Get().GetSlateColor("Colors.AccentGreen"); })
+				.Content()
+				[
+					SNew(STextBlock)
+					.Text_Lambda([this]() { return (AreAnyBodiesHidden()) ? LOCTEXT("ShowAllBodiesButtonText", "Show All Bodies") : LOCTEXT("HideAllBodiesButtonText", "Hide All Bodies"); })
+					.ToolTipText(LOCTEXT("ToggleBodyVisibilityButtonToolTip", "Toggle debug visualization of all physics bodies"))
+				]
+			]
+			// Show/Hide Constraints button.
+			+ SHorizontalBox::Slot()
+			[
+				SNew(SButton)
+				.HAlign(HAlign_Center)
+				.VAlign(VAlign_Center)
+				.OnClicked_Lambda([this](){ this->ToggleConstraintVisibility(); return FReply::Handled(); })
+				.ButtonColorAndOpacity_Lambda([this]() { return (AreAnyConstraintsHidden()) ? FAppStyle::Get().GetSlateColor("Colors.AccentRed") : FAppStyle::Get().GetSlateColor("Colors.AccentGreen"); })
+				.Content()
+				[
+					SNew(STextBlock)
+					.Text_Lambda([this]() { return (AreAnyConstraintsHidden()) ? LOCTEXT("ShowAllConstraintsButtonText", "Show All Constraints") : LOCTEXT("HideAllConstraintsButtonText", "Hide All Constraints"); })
+					.ToolTipText(LOCTEXT("ToggleConstraintVisibilityButtonToolTip", "Toggle debug visualization of all physics constriants"))
+				]
+			]
+		];
+}
+
+void UAnimGraphNode_RigidBody::PostEditChangeProperty(struct FPropertyChangedEvent& PropertyChangedEvent)
+{
+	Super::PostEditChangeProperty(PropertyChangedEvent);
+
+	if (UPhysicsAssetRenderUtilities* PhysicsAssetRenderUtilities = GetMutableDefault<UPhysicsAssetRenderUtilities>())
+	{
+		PhysicsAssetRenderUtilities->SaveConfig();
+	}
+}
+
+void UAnimGraphNode_RigidBody::ToggleBodyVisibility()
+{
+	FAnimNode_RigidBody* const RigidBodyNode = static_cast<FAnimNode_RigidBody*>(GetDebuggedAnimNode());
+	FPhysicsAssetRenderSettings* const PhysicsAssetRenderSettings = GetRenderSettings();
+
+	if (PhysicsAssetRenderSettings && RigidBodyNode)
+	{
+		PhysicsAssetRenderSettings->ToggleShowAllBodies(RigidBodyNode->GetPhysicsAsset());
+	}
+}
+
+void UAnimGraphNode_RigidBody::ToggleConstraintVisibility()
+{
+	FAnimNode_RigidBody* const RigidBodyNode = static_cast<FAnimNode_RigidBody*>(GetDebuggedAnimNode());
+	FPhysicsAssetRenderSettings* const PhysicsAssetRenderSettings = GetRenderSettings();
+
+	if (PhysicsAssetRenderSettings && RigidBodyNode)
+	{
+		PhysicsAssetRenderSettings->ToggleShowAllConstraints(RigidBodyNode->GetPhysicsAsset());
+	}
+}
+
+FPhysicsAssetRenderSettings* UAnimGraphNode_RigidBody::GetRenderSettings() const
+{
+	if (FAnimNode_RigidBody* const RigidBodyNode = static_cast<FAnimNode_RigidBody*>(GetDebuggedAnimNode()))
+	{
+		return UPhysicsAssetRenderUtilities::GetSettings(RigidBodyNode->GetPhysicsAsset());
+	}
+
+	return nullptr;
+}
+
+bool UAnimGraphNode_RigidBody::AreAnyBodiesHidden() const
+{
+	if (FPhysicsAssetRenderSettings* const PhysicsAssetRenderSettings = GetRenderSettings())
+	{
+		return PhysicsAssetRenderSettings->AreAnyBodiesHidden();
+	}
+
+	return false;
+}
+
+bool UAnimGraphNode_RigidBody::AreAnyConstraintsHidden() const
+{
+	if (FPhysicsAssetRenderSettings* const PhysicsAssetRenderSettings = GetRenderSettings())
+	{
+		return PhysicsAssetRenderSettings->AreAnyConstraintsHidden();
+	}
+
+	return false;
+}
+
 #undef LOCTEXT_NAMESPACE
diff --git a/Engine/Source/Editor/AnimGraph/Public/AnimGraphNode_RigidBody.h b/Engine/Source/Editor/AnimGraph/Public/AnimGraphNode_RigidBody.h
--- a/Engine/Source/Editor/AnimGraph/Public/AnimGraphNode_RigidBody.h
+++ b/Engine/Source/Editor/AnimGraph/Public/AnimGraphNode_RigidBody.h
@@ -21,9 +21,22 @@ public:
 	virtual void ValidateAnimNodeDuringCompilation(USkeleton* ForSkeleton, FCompilerResultsLog& MessageLog) override;
 	// End of UEdGraphNode interface
 
+	// UAnimGraphNode_Base
+	virtual void Draw(FPrimitiveDrawInterface* PDI, USkeletalMeshComponent* PreviewSkelMeshComp) const override;
+
+	virtual void CustomizeDetails(IDetailLayoutBuilder& DetailBuilder) override;
+	virtual void PostEditChangeProperty(struct FPropertyChangedEvent& PropertyChangedEvent) override;
+
+	void ToggleBodyVisibility();
+	void ToggleConstraintVisibility();
+
 protected:
 	// UAnimGraphNode_SkeletalControlBase interface
 	virtual FText GetControllerDescription() const override;
 	virtual const FAnimNode_SkeletalControlBase* GetNode() const override { return &Node; }
 	// End of UAnimGraphNode_SkeletalControlBase interface
+
+	struct FPhysicsAssetRenderSettings* GetRenderSettings() const;
+	bool AreAnyBodiesHidden() const;
+	bool AreAnyConstraintsHidden() const;
 };
diff --git a/Engine/Source/Editor/Persona/Private/PersonaModule.cpp b/Engine/Source/Editor/Persona/Private/PersonaModule.cpp
--- a/Engine/Source/Editor/Persona/Private/PersonaModule.cpp
+++ b/Engine/Source/Editor/Persona/Private/PersonaModule.cpp
@@ -75,6 +75,7 @@
 #include "Widgets/Notifications/SNotificationList.h"
 #include "Framework/Notifications/NotificationManager.h"
 #include "PhysicsEngine/PhysicsAsset.h"
+#include "PhysicsAssetRenderUtils.h"
 #include "PersonaPreviewSceneDescription.h"
 #include "PersonaPreviewSceneAnimationController.h"
 #include "PersonaPreviewSceneRefPoseController.h"
@@ -113,6 +114,8 @@ void FPersonaModule::StartupModule()
 	// Make sure the advanced preview scene module is loaded 
 	FModuleManager::Get().LoadModuleChecked("AdvancedPreviewScene");
 
+	UPhysicsAssetRenderUtilities::Initialise();
+
 	{
 		FPropertyEditorModule& PropertyModule = FModuleManager::LoadModuleChecked<FPropertyEditorModule>("PropertyEditor");
 		PropertyModule.RegisterCustomClassLayout( "SkeletalMeshSocket", FOnGetDetailCustomizationInstance::CreateStatic( &FSkeletalMeshSocketDetails::MakeInstance ) );
diff --git a/Engine/Source/Editor/Persona/Private/PhysicsAssetRenderUtils.cpp b/Engine/Source/Editor/Persona/Private/PhysicsAssetRenderUtils.cpp
new file mode 100644
--- /dev/null
+++ b/Engine/Source/Editor/Persona/Private/PhysicsAssetRenderUtils.cpp
@@ -0,0 +1,617 @@
+// Copyright Epic Games, Inc. All Rights Reserved.
+
+#include "PhysicsAssetRenderUtils.h"
+#include "Materials/MaterialInterface.h"
+#include "Materials/MaterialInstanceDynamic.h"
+#include "Components/SkeletalMeshComponent.h"
+#include "Materials/Material.h"
+#include "Preferences/PhysicsAssetEditorOptions.h"
+#include "PhysicsEngine/PhysicsAsset.h"
+#include "PhysicsEngine/PhysicsConstraintTemplate.h"
+#include "Chaos/Core.h"
+#include "SkeletalMeshTypes.h"
+#include "EditorStyleSet.h"
+#include "Animation/DebugSkelMeshComponent.h"
+#include "AssetRegistry/AssetRegistryModule.h"
+#include "Internationalization/TextKey.h"
+
+#define LOCTEXT_NAMESPACE "PhysicsAssetRenderUtils"
+
+// How large to make the constraint arrows.
+// The factor of 60 was found experimentally, to look reasonable in comparison with the rest of the constraint visuals.
+constexpr float ConstraintArrowScale = 60.0f;
+
+////////////////////////////////////////
+// Utility Functions 
+
+FTransform GetConstraintMatrix(const USkeletalMeshComponent* const SkeletalMeshComponent, const UPhysicsConstraintTemplate* const ConstraintSetup, const EConstraintFrame::Type Frame, const int32 BoneIndex)
+{
+	if (!ConstraintSetup || !SkeletalMeshComponent)
+	{
+		return FTransform::Identity;
+	}
+
+	FTransform LFrame = ConstraintSetup->DefaultInstance.GetRefFrame(Frame);
+	FTransform BoneTM = SkeletalMeshComponent->GetBoneTransform(BoneIndex);
+	BoneTM.RemoveScaling();
+
+	return LFrame * BoneTM;
+}
+
+////////////////////////////////////////
+// struct FPhysicsAssetRenderSettings
+
+FPhysicsAssetRenderSettings::FPhysicsAssetRenderSettings()
+	: CollisionViewMode(EPhysicsAssetEditorCollisionViewMode::Solid)
+	, ConstraintViewMode(EPhysicsAssetEditorConstraintViewMode::AllLimits)
+	, ConstraintDrawSize(1.0f)
+	, PhysicsBlend(1.0f)
+	, bHideKinematicBodies(false)
+	, bHideSimulatedBodies(false)
+	, bRenderOnlySelectedConstraints(false)
+	, bShowCOM(false)
+	, bShowConstraintsAsPoints(false)
+	, BoneUnselectedColor(170, 155, 225)
+	, NoCollisionColor(200, 200, 200)
+	, COMRenderColor(255, 255, 100)
+	, COMRenderSize(5.0f)
+	, InfluenceLineLength(2.0f)
+{
+	UMaterialInterface* BaseBoneUnselectedMaterial = LoadObject<UMaterialInterface>(NULL, TEXT("/Engine/EditorMaterials/PhAT_UnselectedMaterial.PhAT_UnselectedMaterial"), NULL, LOAD_None, NULL);
+	BoneUnselectedMaterial = UMaterialInstanceDynamic::Create(BaseBoneUnselectedMaterial, GetTransientPackage());
+	check(BoneUnselectedMaterial);
+
+	UMaterialInterface* BaseBoneNoCollisionMaterial = LoadObject<UMaterialInterface>(NULL, TEXT("/Engine/EditorMaterials/PhAT_NoCollisionMaterial.PhAT_NoCollisionMaterial"), NULL, LOAD_None, NULL);
+	BoneNoCollisionMaterial = UMaterialInstanceDynamic::Create(BaseBoneNoCollisionMaterial, GetTransientPackage());
+	check(BoneNoCollisionMaterial);
+}
+
+bool FPhysicsAssetRenderSettings::IsBodyHidden(const int32 BodyIndex) const
+{
+	return HiddenBodies.Contains(BodyIndex);
+}
+
+bool FPhysicsAssetRenderSettings::IsConstraintHidden(const int32 ConstraintIndex) const
+{
+	return HiddenConstraints.Contains(ConstraintIndex);
+}
+
+bool FPhysicsAssetRenderSettings::AreAnyBodiesHidden() const
+{
+	return HiddenBodies.Num() > 0;
+}
+
+bool FPhysicsAssetRenderSettings::AreAnyConstraintsHidden() const
+{
+	return HiddenConstraints.Num() > 0;
+}
+
+void FPhysicsAssetRenderSettings::HideBody(const int32 BodyIndex)
+{
+	if (!HiddenBodies.Contains(BodyIndex))
+	{
+		HiddenBodies.Add(BodyIndex);
+	}
+}
+
+void FPhysicsAssetRenderSettings::ShowBody(const int32 BodyIndex)
+{
+	if (HiddenBodies.Contains(BodyIndex))
+	{
+		HiddenBodies.RemoveSwap(BodyIndex);
+	}
+}
+
+void FPhysicsAssetRenderSettings::ShowAllBodies()
+{
+	HiddenBodies.Reset();
+}
+
+void FPhysicsAssetRenderSettings::ShowAllConstraints()
+{
+	HiddenConstraints.Reset();
+}
+
+void FPhysicsAssetRenderSettings::ShowAll()
+{
+	ShowAllBodies();
+	ShowAllConstraints();
+}
+
+void FPhysicsAssetRenderSettings::HideAllBodies(const UPhysicsAsset* const PhysicsAsset)
+{
+	HiddenBodies.Reset();
+
+	if (PhysicsAsset != nullptr)
+	{
+		for (int32 i = 0; i < PhysicsAsset->SkeletalBodySetups.Num(); ++i)
+		{
+			HiddenBodies.Add(i);
+		}
+	}
+}
+
+void FPhysicsAssetRenderSettings::HideAllConstraints(const UPhysicsAsset* const PhysicsAsset)
+{
+	HiddenConstraints.Reset();
+
+	if (PhysicsAsset != nullptr)
+	{
+		for (int32 i = 0; i < PhysicsAsset->ConstraintSetup.Num(); ++i)
+		{
+			HiddenConstraints.Add(i);
+		}
+	}
+}
+
+void FPhysicsAssetRenderSettings::HideAll(const UPhysicsAsset* const PhysicsAsset)
+{
+	HideAllBodies(PhysicsAsset);
+	HideAllConstraints(PhysicsAsset);
+}
+
+void FPhysicsAssetRenderSettings::ToggleShowBody(const int32 BodyIndex)
+{
+	if (IsBodyHidden(BodyIndex))
+	{
+		ShowBody(BodyIndex);
+	}
+	else
+	{
+		HideBody(BodyIndex);
+	}
+}
+
+void FPhysicsAssetRenderSettings::ToggleShowConstraint(const int32 ConstraintIndex)
+{
+	if (IsConstraintHidden(ConstraintIndex))
+	{
+		ShowConstraint(ConstraintIndex);
+	}
+	else
+	{
+		HideConstraint(ConstraintIndex);
+	}
+}
+
+void FPhysicsAssetRenderSettings::ToggleShowAllBodies(const UPhysicsAsset* const PhysicsAsset)
+{
+	if (AreAnyBodiesHidden())
+	{
+		ShowAllBodies();
+	}
+	else
+	{
+		HideAllBodies(PhysicsAsset);
+	}
+}
+
+void FPhysicsAssetRenderSettings::ToggleShowAllConstraints(const UPhysicsAsset* const PhysicsAsset)
+{
+	if (AreAnyConstraintsHidden())
+	{
+		ShowAllConstraints();
+	}
+	else
+	{
+		HideAllConstraints(PhysicsAsset);
+	}
+}
+
+void FPhysicsAssetRenderSettings::HideConstraint(const int32 ConstraintIndex)
+{
+	if (!HiddenConstraints.Contains(ConstraintIndex))
+	{
+		HiddenConstraints.Add(ConstraintIndex);
+	}
+}
+
+void FPhysicsAssetRenderSettings::ShowConstraint(const int32 ConstraintIndex)
+{
+	if (HiddenConstraints.Contains(ConstraintIndex))
+	{
+		HiddenConstraints.RemoveSwap(ConstraintIndex);
+	}
+}
+
+void FPhysicsAssetRenderSettings::SetHiddenBodies(const TArray<int32>& InHiddenBodies)
+{
+	HiddenBodies.Reset();
+	HiddenBodies.Append(InHiddenBodies);
+}
+
+void FPhysicsAssetRenderSettings::SetHiddenConstraints(const TArray<int32>& InHiddenConstraints)
+{
+	HiddenConstraints.Reset();
+	HiddenConstraints.Append(InHiddenConstraints);
+}
+
+void FPhysicsAssetRenderSettings::ResetEditorViewportOptions()
+{
+	const FPhysicsAssetRenderSettings DefaultObject = FPhysicsAssetRenderSettings();
+
+	CollisionViewMode = DefaultObject.CollisionViewMode;
+	ConstraintViewMode = DefaultObject.ConstraintViewMode;
+	ConstraintDrawSize = DefaultObject.ConstraintDrawSize;
+	PhysicsBlend = DefaultObject.PhysicsBlend;
+	bHideKinematicBodies = DefaultObject.bHideKinematicBodies;
+	bHideSimulatedBodies = DefaultObject.bHideSimulatedBodies;
+	bRenderOnlySelectedConstraints = DefaultObject.bRenderOnlySelectedConstraints;
+	bShowConstraintsAsPoints = DefaultObject.bShowConstraintsAsPoints;
+}
+
+namespace PhysicsAssetRender
+{
+	void DebugDraw(USkeletalMeshComponent* const SkeletalMeshComponent, UPhysicsAsset* const PhysicsAsset, FPrimitiveDrawInterface* PDI)
+	{
+		// Draw Bodies.
+		{
+			auto TransformFn = [](const UPhysicsAsset* PhysicsAsset, const FTransform& BoneTM, const int32 BodyIndex, const EAggCollisionShape::Type PrimType, const int32 PrimIndex, const float Scale) { return GetPrimitiveTransform(PhysicsAsset, BoneTM, BodyIndex, PrimType, PrimIndex, Scale);  };
+			auto ColorFn    = [](const int32 BodyIndex, const EAggCollisionShape::Type PrimitiveType, const int32 PrimitiveIndex, const FPhysicsAssetRenderSettings& Settings) { return GetPrimitiveColor(BodyIndex, PrimitiveType, PrimitiveIndex, Settings); };
+			auto MaterialFn = [](const int32 BodyIndex, const EAggCollisionShape::Type PrimitiveType, const int32 PrimitiveIndex, const FPhysicsAssetRenderSettings& Settings) { return GetPrimitiveMaterial(BodyIndex, PrimitiveType, PrimitiveIndex, Settings); };
+			auto HitProxyFn = [](const int32 BodyIndex, const EAggCollisionShape::Type PrimitiveType, const int32 PrimitiveIndex) { return nullptr; };
+
+			DebugDrawBodies(SkeletalMeshComponent, PhysicsAsset, PDI, ColorFn, MaterialFn, TransformFn, HitProxyFn);
+		}
+
+		// Draw Constraints.
+		{
+			auto HitProxyFn = [](const int32) { return nullptr; };
+			auto IsSelectedFn = [](const uint32) { return false; };
+			DebugDrawConstraints(SkeletalMeshComponent, PhysicsAsset, PDI, IsSelectedFn, false, HitProxyFn);
+		}
+	}
+
+	// Contains info needed for drawing, in a form that can be sorted by distance.
+	struct DrawElement
+	{
+		DrawElement(
+			FTransform        InTM,
+			HHitProxy* InHitProxy,
+			float             InScale,
+			const FSceneView* InView,
+			FKShapeElem* InShapeElem,
+			const FName& InBoneName
+		)
+			: TM(InTM), HitProxy(InHitProxy), Scale(InScale), ShapeElem(InShapeElem)
+#ifdef UE_BUILD_DEBUG
+			, BoneName(InBoneName)
+#endif
+		{
+			// The TM position is at the center of the objects, so can be used directly for sorting.
+			// Sorting by distance (rather than along the view direction) reduces flickering when
+			// the camera is rotated without moving it.
+			FVector ViewDir = InView->GetViewDirection();
+			//Distance = TM.GetTranslation() | ViewDir; // Would sort along the view direction
+			SortingMetric = (TM.GetTranslation() - InView->ViewLocation).SquaredLength(); // Sorts by Distance
+		}
+
+		UMaterialInterface* Material = 0;
+		FColor                Color = FColor(ForceInitToZero);
+		FTransform            TM;
+		HHitProxy* HitProxy = nullptr;
+		float                 Scale;
+		FKShapeElem* ShapeElem = nullptr;
+		float                 SortingMetric;
+#ifdef UE_BUILD_DEBUG
+		FName                 BoneName;
+#endif
+	};
+	
+	template< typename TElementContainer > void CollectDrawElements(TArray<DrawElement>& DrawElements, const EAggCollisionShape::Type PrimitiveType, const TElementContainer ElementContainer, const FTransform& BoneTM, const float Scale, const uint32 BodyIndex, const FPhysicsAssetRenderSettings* const RenderSettings, UPhysicsAsset* const PhysicsAsset, FPrimitiveDrawInterface* PDI, GetPrimitiveRef< FColor > GetPrimitiveColor, GetPrimitiveRef< UMaterialInterface* > GetPrimitiveMaterial, GetPrimitiveTransformRef GetPrimitiveTransform, CreateBodyHitProxyFn CreateHitProxy)
+	{
+		if (PhysicsAsset && RenderSettings && PDI && PhysicsAsset->SkeletalBodySetups.IsValidIndex(BodyIndex))
+		{
+			for (int32 ElementIndex = 0; ElementIndex < ElementContainer.Num(); ++ElementIndex)
+			{
+				const FName BoneName = PhysicsAsset->SkeletalBodySetups[BodyIndex]->BoneName;
+
+				DrawElement DE(
+					GetPrimitiveTransform(PhysicsAsset, BoneTM, BodyIndex, PrimitiveType, ElementIndex, Scale),
+					CreateHitProxy(BodyIndex, PrimitiveType, ElementIndex),
+					Scale, PDI->View, (FKShapeElem*)(&ElementContainer[ElementIndex]), BoneName
+				);
+
+				if (RenderSettings->CollisionViewMode == EPhysicsAssetEditorCollisionViewMode::Solid || RenderSettings->CollisionViewMode == EPhysicsAssetEditorCollisionViewMode::SolidWireframe)
+				{
+					DE.Material = GetPrimitiveMaterial(BodyIndex, PrimitiveType, ElementIndex, *RenderSettings);
+				}
+				if (RenderSettings->CollisionViewMode == EPhysicsAssetEditorCollisionViewMode::SolidWireframe || RenderSettings->CollisionViewMode == EPhysicsAssetEditorCollisionViewMode::Wireframe)
+				{
+					DE.Color = GetPrimitiveColor(BodyIndex, PrimitiveType, ElementIndex, *RenderSettings);
+				}
+				if (DE.Material || DE.Color.A)
+				{
+					DrawElements.Add(DE);
+				}
+			}
+		}
+	}
+
+	void DebugDrawBodies(USkeletalMeshComponent* const SkeletalMeshComponent, UPhysicsAsset* const PhysicsAsset, FPrimitiveDrawInterface* PDI, GetPrimitiveRef< FColor > GetPrimitiveColor, GetPrimitiveRef< UMaterialInterface* > GetPrimitiveMaterial, GetPrimitiveTransformRef GetPrimitiveTransform, CreateBodyHitProxyFn CreateHitProxy)
+	{
+		if (!SkeletalMeshComponent || !PhysicsAsset || !PhysicsAsset->PreviewSkeletalMesh)
+		{
+			// Nothing to draw without an asset, this can happen if the preview scene has no skeletal mesh
+			return;
+		}
+
+		const FPhysicsAssetRenderSettings* const RenderSettings = UPhysicsAssetRenderUtilities::GetSettings(PhysicsAsset);
+
+		if (!RenderSettings)
+		{
+			return;
+		}
+
+		// Draw bodies
+		for (int32 i = 0; i < PhysicsAsset->SkeletalBodySetups.Num(); ++i)
+		{
+			if (!ensure(PhysicsAsset->SkeletalBodySetups[i]))
+			{
+				continue;
+			}
+			if ((PhysicsAsset->SkeletalBodySetups[i]->PhysicsType == EPhysicsType::PhysType_Kinematic &&
+				RenderSettings->bHideKinematicBodies) ||
+				(PhysicsAsset->SkeletalBodySetups[i]->PhysicsType == EPhysicsType::PhysType_Simulated &&
+					RenderSettings->bHideSimulatedBodies)
+				)
+			{
+				continue;
+			}
+			if (RenderSettings->IsBodyHidden(i))
+			{
+				continue;
+			}
+
+			const int32 BoneIndex = SkeletalMeshComponent->GetBoneIndex(PhysicsAsset->SkeletalBodySetups[i]->BoneName);
+
+			// If we found a bone for it, draw the collision.
+			// The logic is as follows; always render in the ViewMode requested when not in hit mode - but if we are in hit mode and the right editing mode, render as solid
+			if (BoneIndex != INDEX_NONE)
+			{
+				FTransform BoneTM = SkeletalMeshComponent->GetBoneTransform(BoneIndex);
+				const float Scale = BoneTM.GetScale3D().GetAbsMax();
+				BoneTM.RemoveScaling();
+
+				FKAggregateGeom* const AggGeom = &PhysicsAsset->SkeletalBodySetups[i]->AggGeom;
+
+				TArray<DrawElement> DrawElements;
+
+				CollectDrawElements(DrawElements, EAggCollisionShape::Sphere, AggGeom->SphereElems, BoneTM, Scale, i, RenderSettings, PhysicsAsset, PDI, GetPrimitiveColor, GetPrimitiveMaterial, GetPrimitiveTransform, CreateHitProxy);
+				CollectDrawElements(DrawElements, EAggCollisionShape::Box, AggGeom->BoxElems, BoneTM, Scale, i, RenderSettings, PhysicsAsset, PDI, GetPrimitiveColor, GetPrimitiveMaterial, GetPrimitiveTransform, CreateHitProxy);
+				CollectDrawElements(DrawElements, EAggCollisionShape::Sphyl, AggGeom->SphylElems, BoneTM, Scale, i, RenderSettings, PhysicsAsset, PDI, GetPrimitiveColor, GetPrimitiveMaterial, GetPrimitiveTransform, CreateHitProxy);
+				CollectDrawElements(DrawElements, EAggCollisionShape::Convex, AggGeom->ConvexElems, BoneTM, Scale, i, RenderSettings, PhysicsAsset, PDI, GetPrimitiveColor, GetPrimitiveMaterial, GetPrimitiveTransform, CreateHitProxy);
+				CollectDrawElements(DrawElements, EAggCollisionShape::TaperedCapsule, AggGeom->TaperedCapsuleElems, BoneTM, Scale, i, RenderSettings, PhysicsAsset, PDI, GetPrimitiveColor, GetPrimitiveMaterial, GetPrimitiveTransform, CreateHitProxy);
+
+				// Sort elements
+				DrawElements.Sort([](const DrawElement& DE1, const DrawElement& DE2) {
+					return DE1.SortingMetric > DE2.SortingMetric;
+					});
+
+				// Render sorted elements.
+				for (const DrawElement& DE : DrawElements)
+				{
+					PDI->SetHitProxy(DE.HitProxy);
+					if (DE.Material)
+					{
+						DE.ShapeElem->DrawElemSolid(PDI, DE.TM, DE.Scale, DE.Material->GetRenderProxy());
+					}
+					if (DE.Color.A)
+					{
+						DE.ShapeElem->DrawElemWire(PDI, DE.TM, DE.Scale, DE.Color);
+					}
+				}
+	
+				PDI->SetHitProxy(NULL);
+			}
+
+			if (RenderSettings->bShowCOM && SkeletalMeshComponent->Bodies.IsValidIndex(i))
+			{
+				SkeletalMeshComponent->Bodies[i]->DrawCOMPosition(PDI, RenderSettings->COMRenderSize, RenderSettings->COMRenderColor);
+			}
+		}
+	}
+
+	void DebugDrawConstraints(USkeletalMeshComponent* const SkeletalMeshComponent, UPhysicsAsset* const PhysicsAsset, FPrimitiveDrawInterface* PDI, TFunctionRef< bool(const uint32) > IsConstraintSelected, const bool bRunningSimulation, CreateConstraintHitProxyFn CreateHitProxy)
+	{
+		check(SkeletalMeshComponent);
+		check(PhysicsAsset);
+
+		const FPhysicsAssetRenderSettings* const RenderSettings = UPhysicsAssetRenderUtilities::GetSettings(PhysicsAsset);
+
+		if (!RenderSettings)
+		{
+			return;
+		}
+
+		if (RenderSettings->ConstraintViewMode != EPhysicsAssetEditorConstraintViewMode::None)
+		{
+			for (int32 ConstraintIndex = 0; ConstraintIndex < PhysicsAsset->ConstraintSetup.Num(); ++ConstraintIndex)
+			{
+				const bool bConstraintSelected = IsConstraintSelected(ConstraintIndex);
+
+				if ((!RenderSettings->bRenderOnlySelectedConstraints || (RenderSettings->bRenderOnlySelectedConstraints && bConstraintSelected)) &&
+					!RenderSettings->IsConstraintHidden(ConstraintIndex))
+				{
+					const bool bDrawLimits = (RenderSettings->ConstraintViewMode == EPhysicsAssetEditorConstraintViewMode::AllLimits) || bConstraintSelected;
+					const bool bDrawSelected = !bRunningSimulation && bConstraintSelected;
+
+					PDI->SetHitProxy(CreateHitProxy(ConstraintIndex));
+					
+					UPhysicsConstraintTemplate* const ConstraintSetup = PhysicsAsset->ConstraintSetup[ConstraintIndex];
+					check(ConstraintSetup);
+					if (ConstraintSetup)
+					{
+						const int32 BoneIndex1 = SkeletalMeshComponent->GetBoneIndex(ConstraintSetup->DefaultInstance.ConstraintBone1);
+						const int32 BoneIndex2 = SkeletalMeshComponent->GetBoneIndex(ConstraintSetup->DefaultInstance.ConstraintBone2);
+
+						// if bone doesn't exist, do not draw it. It crashes in random points when we try to manipulate. 
+						if ((BoneIndex1 != INDEX_NONE) && (BoneIndex2 != INDEX_NONE))
+						{
+							const float DrawScale = ConstraintArrowScale * RenderSettings->ConstraintDrawSize;
+							FTransform Con1Frame = GetConstraintMatrix(SkeletalMeshComponent, ConstraintSetup, EConstraintFrame::Frame1, BoneIndex1);
+							FTransform Con2Frame = GetConstraintMatrix(SkeletalMeshComponent, ConstraintSetup, EConstraintFrame::Frame2, BoneIndex2);
+
+							ConstraintSetup->DefaultInstance.DrawConstraint(PDI, RenderSettings->ConstraintDrawSize, DrawScale, bDrawLimits, bDrawSelected, Con1Frame, Con2Frame, RenderSettings->bShowConstraintsAsPoints);
+
+						}
+					}
+
+					PDI->SetHitProxy(NULL);
+				}
+			}
+		}
+	}
+
+	FTransform GetPrimitiveTransform(const UPhysicsAsset* PhysicsAsset, const FTransform& BoneTM, const int32 BodyIndex, const EAggCollisionShape::Type PrimType, const int32 PrimIndex, const float Scale)
+	{
+		UBodySetup* SharedBodySetup = PhysicsAsset->SkeletalBodySetups[BodyIndex];
+		FVector Scale3D(Scale);
+
+		FTransform PrimTM = FTransform::Identity;
+
+		if (PrimType == EAggCollisionShape::Sphere)
+		{
+			PrimTM = SharedBodySetup->AggGeom.SphereElems[PrimIndex].GetTransform();
+		}
+		else if (PrimType == EAggCollisionShape::Box)
+		{
+			PrimTM = SharedBodySetup->AggGeom.BoxElems[PrimIndex].GetTransform();
+		}
+		else if (PrimType == EAggCollisionShape::Sphyl)
+		{
+			PrimTM = SharedBodySetup->AggGeom.SphylElems[PrimIndex].GetTransform();
+		}
+		else if (PrimType == EAggCollisionShape::Convex)
+		{
+			PrimTM = SharedBodySetup->AggGeom.ConvexElems[PrimIndex].GetTransform();
+		}
+		else if (PrimType == EAggCollisionShape::TaperedCapsule)
+		{
+			PrimTM = SharedBodySetup->AggGeom.TaperedCapsuleElems[PrimIndex].GetTransform();
+		}
+		else
+		{
+			// Should never reach here
+			check(0);
+		}
+
+		PrimTM.ScaleTranslation(Scale3D);
+		return PrimTM * BoneTM;
+	}
+
+	FColor GetPrimitiveColor(const int32 BodyIndex, const EAggCollisionShape::Type PrimitiveType, const int32 PrimitiveIndex, const FPhysicsAssetRenderSettings& Settings)
+	{
+		return (PrimitiveType == EAggCollisionShape::TaperedCapsule) ? Settings.NoCollisionColor : Settings.BoneUnselectedColor;
+	}
+
+	UMaterialInterface* GetPrimitiveMaterial(const int32 BodyIndex, const EAggCollisionShape::Type PrimitiveType, const int32 PrimitiveIndex, const FPhysicsAssetRenderSettings& Settings)
+	{
+		return (PrimitiveType == EAggCollisionShape::TaperedCapsule) ? Settings.BoneNoCollisionMaterial : Settings.BoneUnselectedMaterial;
+	}
+}
+
+void UPhysicsAssetRenderUtilities::Initialise()
+{
+	if (UPhysicsAssetRenderUtilities* PhysicsAssetRenderUtilities = GetMutableDefault<UPhysicsAssetRenderUtilities>())
+	{
+		PhysicsAssetRenderUtilities->InitialiseImpl();
+	}
+}
+
+void UPhysicsAssetRenderUtilities::OnAssetRenamed(FAssetData const& AssetInfo, const FString& InOldPhysicsAssetPathName)
+{
+	FPhysicsAssetRenderSettings Temp;
+
+	const uint32 OldPhysicsAssetPathNameHash = GetPathNameHash(InOldPhysicsAssetPathName);
+
+	if (IdToSettingsMap.RemoveAndCopyValue(OldPhysicsAssetPathNameHash, Temp))
+	{
+		const uint32 PhysicsAssetPathNameHash = GetPathNameHash(AssetInfo.ObjectPath.ToString());
+		IdToSettingsMap.Add(PhysicsAssetPathNameHash, Temp);
+		SaveConfig();
+	}
+}
+
+void UPhysicsAssetRenderUtilities::OnAssetRemoved(UObject* Object)
+{
+	if (Object)
+	{
+		const uint32 PhysicsAssetPathNameHash = GetPathNameHash(Object->GetPathName());
+		IdToSettingsMap.Remove(PhysicsAssetPathNameHash);
+	}
+}
+
+FPhysicsAssetRenderSettings* UPhysicsAssetRenderUtilities::GetSettings(const UPhysicsAsset* InPhysicsAsset)
+{
+	const uint32 PhysicsAssetPathNameHash = GetPathNameHash(InPhysicsAsset);
+
+	if (PhysicsAssetPathNameHash)
+	{
+		return GetSettings(PhysicsAssetPathNameHash);
+	}
+
+	return nullptr;
+}
+
+FPhysicsAssetRenderSettings* UPhysicsAssetRenderUtilities::GetSettings(const FString& InPhysicsAssetPathName)
+{
+	const uint32 PhysicsAssetPathNameHash = GetPathNameHash(InPhysicsAssetPathName);
+
+	if (PhysicsAssetPathNameHash)
+	{
+		return GetSettings(PhysicsAssetPathNameHash);
+	}
+
+	return nullptr;
+}
+
+FPhysicsAssetRenderSettings* UPhysicsAssetRenderUtilities::GetSettings(const uint32 InPhysicsAssetPathNameHash)
+{
+	if (UPhysicsAssetRenderUtilities* PhysicsAssetRenderUtilities = GetMutableDefault<UPhysicsAssetRenderUtilities>())
+	{
+		return PhysicsAssetRenderUtilities->GetSettingsImpl(InPhysicsAssetPathNameHash);
+	}
+
+	return nullptr;
+}
+
+void UPhysicsAssetRenderUtilities::InitialiseImpl()
+{
+	LoadConfig();
+
+	FAssetRegistryModule& AssetRegistryModule = FModuleManager::LoadModuleChecked<FAssetRegistryModule>(TEXT("AssetRegistry"));
+	AssetRegistryModule.Get().OnAssetRenamed().AddUObject(this, &UPhysicsAssetRenderUtilities::OnAssetRenamed);
+	AssetRegistryModule.Get().OnInMemoryAssetDeleted().AddUObject(this, &UPhysicsAssetRenderUtilities::OnAssetRemoved);
+}
+
+FPhysicsAssetRenderSettings* UPhysicsAssetRenderUtilities::GetSettingsImpl(const uint32 InPhysicsAssetPathNameHash)
+{
+	FPhysicsAssetRenderSettings* Settings = IdToSettingsMap.Find(InPhysicsAssetPathNameHash);
+
+	if (!Settings)
+	{
+		IdToSettingsMap.Add(InPhysicsAssetPathNameHash, FPhysicsAssetRenderSettings());
+		Settings = IdToSettingsMap.Find(InPhysicsAssetPathNameHash);
+	}
+
+	return Settings;
+}
+
+
+uint32 UPhysicsAssetRenderUtilities::GetPathNameHash(const UPhysicsAsset* InPhysicsAsset)
+{
+	if (InPhysicsAsset)
+	{
+		return GetPathNameHash(InPhysicsAsset->GetPathName());
+	}
+
+	return 0;
+}
+	
+uint32 UPhysicsAssetRenderUtilities::GetPathNameHash(const FString& InPhysicsAssetPathName)
+{
+	return TextKeyUtil::HashString(InPhysicsAssetPathName);
+}
+
+#undef LOCTEXT_NAMESPACE
diff --git a/Engine/Source/Editor/Persona/Public/PhysicsAssetRenderUtils.h b/Engine/Source/Editor/Persona/Public/PhysicsAssetRenderUtils.h
new file mode 100644
--- /dev/null
+++ b/Engine/Source/Editor/Persona/Public/PhysicsAssetRenderUtils.h
@@ -0,0 +1,181 @@
+// Copyright Epic Games, Inc. All Rights Reserved.
+
+#pragma once
+
+#include "CoreMinimal.h"
+#include "UObject/ObjectMacros.h"
+#include "Math/Color.h"
+#include "PhysicsEngine/ShapeElem.h"
+#include "SceneManagement.h"
+#include "PhysicsAssetRenderUtils.generated.h"
+
+enum class EPhysicsAssetEditorConstraintViewMode : uint8;
+enum class EPhysicsAssetEditorCollisionViewMode : uint8;
+
+class UPhysicsAsset;
+class HHitProxy;
+
+//////////////////////////////////////////////////////////////////////////
+// FPhysicsAssetRenderSettings
+
+/**
+	Per Physics Asset parameters that determine how debug draw functions 
+	should render that asset in an editor viewport.
+	
+	These parameters are used across different editor modes to ensure the 
+	debug draw is consistent. This makes it easier to create or debug 
+	physics assets whilst switching between editor modes.
+*/
+USTRUCT()
+struct PERSONA_API FPhysicsAssetRenderSettings
+{
+	GENERATED_BODY()
+
+public:
+
+	FPhysicsAssetRenderSettings();
+
+	/** Accessors/helper methods */
+	bool IsBodyHidden(const int32 BodyIndex) const;
+	bool IsConstraintHidden(const int32 ConstraintIndex) const;
+	bool AreAnyBodiesHidden() const;
+	bool AreAnyConstraintsHidden() const;
+	void HideBody(const int32 BodyIndex);
+	void ShowBody(const int32 BodyIndex);
+	void HideConstraint(const int32 ConstraintIndex);
+	void ShowConstraint(const int32 ConstraintIndex);
+	void ShowAllBodies();
+	void ShowAllConstraints();
+	void ShowAll();
+	void HideAllBodies(const UPhysicsAsset* const PhysicsAsset);
+	void HideAllConstraints(const UPhysicsAsset* const PhysicsAsset);
+	void HideAll(const UPhysicsAsset* const PhysicsAsset);
+	void ToggleShowBody(const int32 BodyIndex);
+	void ToggleShowConstraint(const int32 ConstraintIndex);
+	void ToggleShowAllBodies(const UPhysicsAsset* const PhysicsAsset);
+	void ToggleShowAllConstraints(const UPhysicsAsset* const PhysicsAsset);
+
+	void SetHiddenBodies(const TArray<int32>& InHiddenBodies);
+	void SetHiddenConstraints(const TArray<int32>& InHiddenConstraints);
+
+	void ResetEditorViewportOptions();
+
+	// Physics Asset Editor Viewport Options
+	UPROPERTY()
+	EPhysicsAssetEditorCollisionViewMode CollisionViewMode;
+
+	UPROPERTY()
+	EPhysicsAssetEditorConstraintViewMode ConstraintViewMode;
+
+	UPROPERTY()
+	float ConstraintDrawSize;
+
+	UPROPERTY()
+	float PhysicsBlend;
+
+	UPROPERTY()
+	bool bHideKinematicBodies;
+
+	UPROPERTY()
+	bool bHideSimulatedBodies;
+
+	UPROPERTY()
+	bool bRenderOnlySelectedConstraints;
+
+	UPROPERTY()
+	bool bShowCOM;
+
+	UPROPERTY()
+	bool bShowConstraintsAsPoints;
+
+	// Draw colors
+	UPROPERTY()
+	FColor BoneUnselectedColor;
+
+	UPROPERTY()
+	FColor NoCollisionColor;
+
+	UPROPERTY()
+	FColor COMRenderColor;
+
+	UPROPERTY()
+	float COMRenderSize;
+
+	UPROPERTY()
+	float InfluenceLineLength;
+
+	// Materials
+	UPROPERTY()
+	TObjectPtr<class UMaterialInstanceDynamic> BoneUnselectedMaterial;
+
+	UPROPERTY()
+	TObjectPtr<class UMaterialInstanceDynamic> BoneNoCollisionMaterial;
+
+	UPROPERTY()
+	TArray<int32> HiddenBodies;
+
+	UPROPERTY()
+	TArray<int32> HiddenConstraints;
+
+};
+
+//////////////////////////////////////////////////////////////////////////
+// UPhysicsAssetRenderUtilities
+
+/** Factory class for FPhysicsAssetRenderSettings. */
+UCLASS(config=EditorPerProjectUserSettings)
+class PERSONA_API UPhysicsAssetRenderUtilities : public UObject
+{
+	GENERATED_BODY()
+
+public:
+
+	static void Initialise();
+
+	/** Returns an existing render settings object or creates and returns a new one if none exists. */
+	static FPhysicsAssetRenderSettings* GetSettings(const UPhysicsAsset* InPhysicsAsset);
+	static FPhysicsAssetRenderSettings* GetSettings(const FString& InPhysicsAssetPathName);
+	static FPhysicsAssetRenderSettings* GetSettings(const uint32 InPhysicsAssetPathNameHash);
+	
+	static uint32 GetPathNameHash(const UPhysicsAsset* InPhysicsAsset);
+	static uint32 GetPathNameHash(const FString& InPhysicsAssetPathName);
+
+	void OnAssetRenamed(FAssetData const& AssetInfo, const FString& InOldPhysicsAssetPathName);
+	void OnAssetRemoved(UObject* Object);
+
+private:
+
+	void InitialiseImpl();
+	FPhysicsAssetRenderSettings* GetSettingsImpl(const uint32 InPhysicsAssetPathNameHash);
+
+	UPROPERTY()
+	TMap< uint32, FPhysicsAssetRenderSettings > IdToSettingsMap;
+};
+
+//////////////////////////////////////////////////////////////////////////
+// PhysicsAssetRender
+
+/**
+	Namespace containing functions for debug draw of Physics Assets in the editor viewport.
+*/
+namespace PhysicsAssetRender
+{
+	template< typename TReturnType > using GetPrimitiveRef = TFunctionRef< TReturnType (const int32 BodyIndex, const EAggCollisionShape::Type PrimitiveType, const int32 PrimitiveIndex, const FPhysicsAssetRenderSettings& Settings) >;
+	using GetPrimitiveTransformRef = TFunctionRef< FTransform (const UPhysicsAsset* PhysicsAsset, const FTransform& BoneTM, const int32 BodyIndex, const EAggCollisionShape::Type PrimType, const int32 PrimIndex, const float Scale) >;
+	using CreateBodyHitProxyFn = TFunctionRef< HHitProxy* (const int32 BodyIndex, const EAggCollisionShape::Type PrimitiveType, const int32 PrimitiveIndex) >;
+	using CreateConstraintHitProxyFn = TFunctionRef< HHitProxy* (const int32 InConstraintIndex) >;
+
+	/** Debug draw Physics Asset bodies and constraints using the default callbacks */
+	PERSONA_API void DebugDraw(class USkeletalMeshComponent* const SkeletalMeshComponent, class UPhysicsAsset* const PhysicsAsset, FPrimitiveDrawInterface* PDI);
+	
+	/** Debug draw Physics Asset bodies using the supplied custom callbacks */
+	PERSONA_API void DebugDrawBodies(class USkeletalMeshComponent* const SkeletalMeshComponent, class UPhysicsAsset* const PhysicsAsset, FPrimitiveDrawInterface* PDI, GetPrimitiveRef< FColor > GetPrimitiveColor, GetPrimitiveRef< class UMaterialInterface* > GetPrimitiveMaterial, GetPrimitiveTransformRef GetPrimitiveTransform, CreateBodyHitProxyFn CreateHitProxy);
+	
+	/** Debug draw Physics Asset constraints using the supplied custom callbacks */
+	PERSONA_API void DebugDrawConstraints(class USkeletalMeshComponent* const SkeletalMeshComponent, class UPhysicsAsset* const PhysicsAsset, FPrimitiveDrawInterface* PDI, TFunctionRef< bool(const uint32) > IsConstraintSelected, const bool bRunningSimulation, CreateConstraintHitProxyFn CreateHitProxy);
+	
+	/** Default callbacks used by DebugDraw */
+	PERSONA_API FTransform GetPrimitiveTransform(const UPhysicsAsset* PhysicsAsset, const FTransform& BoneTM, const int32 BodyIndex, const EAggCollisionShape::Type PrimType, const int32 PrimIndex, const float Scale);
+	PERSONA_API FColor GetPrimitiveColor(const int32 BodyIndex, const EAggCollisionShape::Type PrimitiveType, const int32 PrimitiveIndex, const FPhysicsAssetRenderSettings& Settings);
+	PERSONA_API class UMaterialInterface* GetPrimitiveMaterial(const int32 BodyIndex, const EAggCollisionShape::Type PrimitiveType, const int32 PrimitiveIndex, const FPhysicsAssetRenderSettings& Settings);
+}
\ No newline at end of file
diff --git a/Engine/Source/Editor/PhysicsAssetEditor/Private/PhysicsAssetEditor.cpp b/Engine/Source/Editor/PhysicsAssetEditor/Private/PhysicsAssetEditor.cpp
--- a/Engine/Source/Editor/PhysicsAssetEditor/Private/PhysicsAssetEditor.cpp
+++ b/Engine/Source/Editor/PhysicsAssetEditor/Private/PhysicsAssetEditor.cpp
@@ -24,6 +24,7 @@
 #include "PhysicsAssetEditorModule.h"
 #include "ScopedTransaction.h"
 #include "PhysicsAssetEditorActions.h"
+#include "PhysicsAssetRenderUtils.h"
 #include "PhysicsAssetEditorSkeletalMeshComponent.h"
 #include "Templates/TypeHash.h"
 
@@ -185,6 +186,7 @@ void FPhysicsAssetEditor::InitPhysicsAssetEditor(const EToolkitMode::Type Mode,
 	SkeletonTreeArgs.OnSelectionChanged = FOnSkeletonTreeSelectionChanged::CreateSP(this, &FPhysicsAssetEditor::HandleSelectionChanged);
 	SkeletonTreeArgs.PreviewScene = PersonaToolkit->GetPreviewScene();
 	SkeletonTreeArgs.bShowBlendProfiles = false;
+	SkeletonTreeArgs.bShowDebugVisualizationOptions = true;
 	SkeletonTreeArgs.bAllowMeshOperations = false;
 	SkeletonTreeArgs.bAllowSkeletonOperations = false;
 	SkeletonTreeArgs.bHideBonesByDefault = true;
@@ -382,6 +384,23 @@ FLinearColor FPhysicsAssetEditor::GetWorldCentricTabColorScale() const
 	return FLinearColor(0.3f, 0.2f, 0.5f, 0.5f);
 }
 
+void FPhysicsAssetEditor::OnClose()
+{
+	// Clear render settings from editor viewport. These settings must be applied to the rendering in all editors 
+	// when an asset is open in the Physics Asset Editor but should not persist after the editor has been closed.
+	if (FPhysicsAssetRenderSettings* const RenderSettings = UPhysicsAssetRenderUtilities::GetSettings(SharedData->PhysicsAsset))
+	{
+		RenderSettings->ResetEditorViewportOptions();
+	}
+
+	if (UPhysicsAssetRenderUtilities* PhysicsAssetRenderUtilities = GetMutableDefault<UPhysicsAssetRenderUtilities>())
+	{
+		PhysicsAssetRenderUtilities->SaveConfig();
+	}
+
+	IPhysicsAssetEditor::OnClose();
+}
+
 void FPhysicsAssetEditor::AddReferencedObjects(FReferenceCollector& Collector)
 {
 	SharedData->AddReferencedObjects(Collector);
@@ -485,6 +504,11 @@ void FPhysicsAssetEditor::OnFinishedChangingProperties(const FPropertyChangedEve
 		}
 	}
 
+	if (UPhysicsAssetRenderUtilities* PhysicsAssetRenderUtilities = GetMutableDefault<UPhysicsAssetRenderUtilities>())
+	{
+		PhysicsAssetRenderUtilities->SaveConfig();
+	}
+
 	RecreatePhysicsState();
 
 	RefreshPreviewViewport();
@@ -1507,6 +1531,7 @@ TSharedRef<ISkeletonTree> FPhysicsAssetEditor::BuildMenuWidgetNewConstraintForBo
 	SkeletonTreeArgs.bAllowSkeletonOperations = false;
 	SkeletonTreeArgs.bShowBlendProfiles = false;
 	SkeletonTreeArgs.bShowFilterMenu = false;
+	SkeletonTreeArgs.bShowDebugVisualizationOptions = true;
 	SkeletonTreeArgs.bHideBonesByDefault = true;
 	SkeletonTreeArgs.Builder = Builder;
 	SkeletonTreeArgs.PreviewScene = GetPersonaToolkit()->GetPreviewScene();
diff --git a/Engine/Source/Editor/PhysicsAssetEditor/Private/PhysicsAssetEditor.h b/Engine/Source/Editor/PhysicsAssetEditor/Private/PhysicsAssetEditor.h
--- a/Engine/Source/Editor/PhysicsAssetEditor/Private/PhysicsAssetEditor.h
+++ b/Engine/Source/Editor/PhysicsAssetEditor/Private/PhysicsAssetEditor.h
@@ -105,6 +105,7 @@ public:
 
 	/** IHasPersonaToolkit interface */
 	virtual TSharedRef<IPersonaToolkit> GetPersonaToolkit() const override { return PersonaToolkit.ToSharedRef(); }
+	virtual void OnClose() override;
 
 	/** FGCObject interface */
 	virtual void AddReferencedObjects(FReferenceCollector& Collector) override;
diff --git a/Engine/Source/Editor/PhysicsAssetEditor/Private/PhysicsAssetEditorSharedData.cpp b/Engine/Source/Editor/PhysicsAssetEditor/Private/PhysicsAssetEditorSharedData.cpp
--- a/Engine/Source/Editor/PhysicsAssetEditor/Private/PhysicsAssetEditorSharedData.cpp
+++ b/Engine/Source/Editor/PhysicsAssetEditor/Private/PhysicsAssetEditorSharedData.cpp
@@ -2,6 +2,7 @@
 
 #include "PhysicsAssetEditorSharedData.h"
 #include "PhysicsAssetEditorPhysicsHandleComponent.h"
+#include "PhysicsAssetRenderUtils.h"
 #include "PhysicsEngine/RigidBodyIndexPair.h"
 #include "Misc/MessageDialog.h"
 #include "Modules/ModuleManager.h"
@@ -735,7 +736,7 @@ void FPhysicsAssetEditorSharedData::ToggleShowSelected()
 	{
 		for (const FSelection& Selection : SelectedConstraints)
 		{
-			if (HiddenConstraints.Contains(Selection.Index))
+			if (IsConstraintHidden(Selection.Index))
 			{
 				bAllSelectedVisible = false;
 				break;
@@ -746,7 +747,7 @@ void FPhysicsAssetEditorSharedData::ToggleShowSelected()
 	{
 		for (const FSelection& Selection : SelectedBodies)
 		{
-			if (HiddenBodies.Contains(Selection.Index))
+			if (IsBodyHidden(Selection.Index))
 			{
 				bAllSelectedVisible = false;
 			}
@@ -772,7 +773,7 @@ void FPhysicsAssetEditorSharedData::ToggleShowOnlySelected()
 	{
 		for (const FSelection& Selection : SelectedConstraints)
 		{
-			if (HiddenConstraints.Contains(Selection.Index))
+			if (IsConstraintHidden(Selection.Index))
 			{
 				bAllSelectedVisible = false;
 				break;
@@ -783,7 +784,7 @@ void FPhysicsAssetEditorSharedData::ToggleShowOnlySelected()
 	{
 		for (const FSelection& Selection : SelectedBodies)
 		{
-			if (HiddenBodies.Contains(Selection.Index))
+			if (IsBodyHidden(Selection.Index))
 			{
 				bAllSelectedVisible = false;
 			}
@@ -799,7 +800,7 @@ void FPhysicsAssetEditorSharedData::ToggleShowOnlySelected()
 			if (!SelectedConstraints.ContainsByPredicate([ConstraintIndex](FSelection& V) { return V.Index == ConstraintIndex; } ))
 			{
 				// Is it hidden?
-				if (!HiddenConstraints.Contains(ConstraintIndex))
+				if (!IsConstraintHidden(ConstraintIndex))
 				{
 					bAllNotSelectedHidden = false;
 					break;
@@ -815,7 +816,7 @@ void FPhysicsAssetEditorSharedData::ToggleShowOnlySelected()
 			if (!SelectedBodies.ContainsByPredicate([BodyIndex](FSelection& V) { return V.Index == BodyIndex; }))
 			{
 				// Is it hidden?
-				if (!HiddenBodies.Contains(BodyIndex))
+				if (!IsBodyHidden(BodyIndex))
 				{
 					bAllNotSelectedHidden = false;
 					break;
@@ -835,33 +836,79 @@ void FPhysicsAssetEditorSharedData::ToggleShowOnlySelected()
 	}
 }
 
+bool FPhysicsAssetEditorSharedData::IsBodyHidden(const int32 BodyIndex) const
+{
+	if (FPhysicsAssetRenderSettings* const PhysicsAssetRenderSettings = GetRenderSettings())
+	{
+		return PhysicsAssetRenderSettings->IsBodyHidden(BodyIndex);
+	}
+	
+	return false;
+}
+
+bool FPhysicsAssetEditorSharedData::IsConstraintHidden(const int32 ConstraintIndex) const
+{
+	if (FPhysicsAssetRenderSettings* const PhysicsAssetRenderSettings = GetRenderSettings())
+	{
+		return PhysicsAssetRenderSettings->IsConstraintHidden(ConstraintIndex);
+	}
+
+	return false;
+}
+
+void FPhysicsAssetEditorSharedData::HideBody(const int32 BodyIndex)
+{
+	if (FPhysicsAssetRenderSettings* const PhysicsAssetRenderSettings = GetRenderSettings())
+	{
+		PhysicsAssetRenderSettings->HideBody(BodyIndex);
+	}
+}
+
+void FPhysicsAssetEditorSharedData::ShowBody(const int32 BodyIndex)
+{
+	if (FPhysicsAssetRenderSettings* const PhysicsAssetRenderSettings = GetRenderSettings())
+	{
+		PhysicsAssetRenderSettings->ShowBody(BodyIndex);
+	}
+}
+
+void FPhysicsAssetEditorSharedData::HideConstraint(const int32 ConstraintIndex)
+{
+	if (FPhysicsAssetRenderSettings* const PhysicsAssetRenderSettings = GetRenderSettings())
+	{
+		PhysicsAssetRenderSettings->HideConstraint(ConstraintIndex);
+	}
+}
+
+void FPhysicsAssetEditorSharedData::ShowConstraint(const int32 ConstraintIndex)
+{
+	if (FPhysicsAssetRenderSettings* const PhysicsAssetRenderSettings = GetRenderSettings())
+	{
+		PhysicsAssetRenderSettings->ShowConstraint(ConstraintIndex);
+	}
+}
+
 void FPhysicsAssetEditorSharedData::ShowAll()
 {
-	HiddenConstraints.Empty();
-	HiddenBodies.Empty();
+	if (FPhysicsAssetRenderSettings* const PhysicsAssetRenderSettings = GetRenderSettings())
+	{
+		PhysicsAssetRenderSettings->ShowAll();
+	}
 }
 
 void FPhysicsAssetEditorSharedData::HideAllBodies()
 {
-	if (PhysicsAsset != nullptr)
+	if (FPhysicsAssetRenderSettings* const PhysicsAssetRenderSettings = GetRenderSettings())
 	{
-		HiddenBodies.Empty();
-		for (int32 i = 0; i < PhysicsAsset->SkeletalBodySetups.Num(); ++i)
-		{
-			HiddenBodies.Add(i);
-		}
+		PhysicsAssetRenderSettings->HideAllBodies(PhysicsAsset);
 	}
 }
 
 void FPhysicsAssetEditorSharedData::HideAllConstraints()
 {
-	if (PhysicsAsset != nullptr)
+	if (FPhysicsAssetRenderSettings* const PhysicsAssetRenderSettings = GetRenderSettings())
 	{
-		HiddenConstraints.Empty();
-		for (int32 i = 0; i < PhysicsAsset->ConstraintSetup.Num(); ++i)
-		{
-			HiddenConstraints.Add(i);
-		}
+		PhysicsAssetRenderSettings->HideAllConstraints(PhysicsAsset);
 	}
 }
 
@@ -875,17 +922,11 @@ void FPhysicsAssetEditorSharedData::ShowSelected()
 {
 	for (const FSelection& Selection : SelectedConstraints)
 	{
-		if (HiddenConstraints.Contains(Selection.Index))
-		{
-			HiddenConstraints.RemoveSwap(Selection.Index);
-		}
+		ShowConstraint(Selection.Index);
 	}
 	for (const FSelection& Selection : SelectedBodies)
 	{
-		if (HiddenBodies.Contains(Selection.Index))
-		{
-			HiddenBodies.RemoveSwap(Selection.Index);
-		}
+		ShowBody(Selection.Index);
 	}
 }
 
@@ -893,39 +934,41 @@ void FPhysicsAssetEditorSharedData::HideSelected()
 {
 	for (const FSelection& Selection : SelectedConstraints)
 	{
-		if (!HiddenConstraints.Contains(Selection.Index))
-		{
-			HiddenConstraints.Add(Selection.Index);
-		}
+		HideConstraint(Selection.Index);
 	}
 	for (const FSelection& Selection : SelectedBodies)
 	{
-		if (!HiddenBodies.Contains(Selection.Index))
-		{
-			HiddenBodies.Add(Selection.Index);
-		}
+		HideBody(Selection.Index);
 	}
 }
 
 void FPhysicsAssetEditorSharedData::ToggleShowOnlyColliding()
 {
 	// important that we check this before calling ShowAll
-	const bool bIsShowingColliding = (HiddenBodies == NoCollisionBodies);
+	bool bIsShowingColliding = true;
 
-	// in any case first show all
-	ShowAll();
-
-	if (!bIsShowingColliding)
+	for (const int32 BodyIndex : NoCollisionBodies)
 	{
-		// only works if one only body is selected
-		if (PhysicsAsset != nullptr && SelectedBodies.Num() == 1)
-		{
+		bIsShowingColliding &= IsBodyHidden(BodyIndex);
 
-			// NoCollisionBodies already contains the non colliding bodies from the one selection
-			HiddenBodies.Empty();
-			HiddenBodies.Append(NoCollisionBodies);
+		if (!bIsShowingColliding)
+		{
+			break;
 		}
 	}
+
+	// in any case first show all
+	ShowAll();
+
+	FPhysicsAssetRenderSettings* const PhysicsAssetRenderSettings = GetRenderSettings();
+	
+	// only works if one only body is selected
+	if (!bIsShowingColliding && PhysicsAssetRenderSettings && (SelectedBodies.Num() == 1))
+	{	
+		// NoCollisionBodies already contains the non colliding bodies from the one selection
+		PhysicsAssetRenderSettings->SetHiddenBodies(NoCollisionBodies);
+	}
+
 }
 
 void FPhysicsAssetEditorSharedData::ToggleShowOnlyConstrained()
@@ -936,10 +979,13 @@ void FPhysicsAssetEditorSharedData::ToggleShowOnlyConstrained()
 	}
 
 	// important that we check this before calling ShowAll
-	if (bool bIsAlreadyShowingConstrained = (HiddenBodies.Num() > 0))
 	{
-		HiddenBodies.Empty();
-		return;
+		FPhysicsAssetRenderSettings* const PhysicsAssetRenderSettings = GetRenderSettings();
+		if (PhysicsAssetRenderSettings && PhysicsAssetRenderSettings->AreAnyBodiesHidden())
+		{
+			PhysicsAssetRenderSettings->ShowAllBodies();
+			return;
+		}
 	}
 
 	// first Hide all bodies and then show only the ones that needs to be
@@ -948,7 +994,7 @@ void FPhysicsAssetEditorSharedData::ToggleShowOnlyConstrained()
 	// add  the current selection of bodies
 	for (const FSelection& SelectedBody : SelectedBodies)
 	{
-		HiddenBodies.RemoveSwap(SelectedBody.Index);
+		ShowBody(SelectedBody.Index);
 	}
 
 	// collect connected bodies from the selected constraints
@@ -957,16 +1003,16 @@ void FPhysicsAssetEditorSharedData::ToggleShowOnlyConstrained()
 		UPhysicsConstraintTemplate* ConstraintTemplate = PhysicsAsset->ConstraintSetup[Selection.Index];
 		FConstraintInstance& DefaultInstance = ConstraintTemplate->DefaultInstance;
 
-		// Add bothe connected bodies
+		// Add both connected bodies
 		int32 Body1IndexToAdd = PhysicsAsset->FindBodyIndex(DefaultInstance.ConstraintBone1);
 		if (Body1IndexToAdd != INDEX_NONE)
 		{
-			HiddenBodies.RemoveSwap(Body1IndexToAdd);
+			ShowBody(Body1IndexToAdd);
 		}
 		int32 Body2IndexToAdd = PhysicsAsset->FindBodyIndex(DefaultInstance.ConstraintBone2);
 		if (Body2IndexToAdd != INDEX_NONE)
 		{
-			HiddenBodies.RemoveSwap(Body2IndexToAdd);
+			ShowBody(Body2IndexToAdd);
 		}
 	}
 
@@ -991,7 +1037,7 @@ void FPhysicsAssetEditorSharedData::ToggleShowOnlyConstrained()
 				int32 BodyIndexToAdd = PhysicsAsset->FindBodyIndex(OtherConnectedBody);
 				if (BodyIndexToAdd != INDEX_NONE)
 				{
-					HiddenBodies.RemoveSwap(BodyIndexToAdd);
+					ShowBody(BodyIndexToAdd);
 				}
 			}
 		}
@@ -2837,6 +2883,10 @@ void FPhysicsAssetEditorSharedData::UpdateClothPhysics()
 	}
 }
 
+FPhysicsAssetRenderSettings* FPhysicsAssetEditorSharedData::GetRenderSettings() const
+{
+	return UPhysicsAssetRenderUtilities::GetSettings(PhysicsAsset);
+}
 
 
 #undef LOCTEXT_NAMESPACE
diff --git a/Engine/Source/Editor/PhysicsAssetEditor/Private/PhysicsAssetEditorSharedData.h b/Engine/Source/Editor/PhysicsAssetEditor/Private/PhysicsAssetEditorSharedData.h
--- a/Engine/Source/Editor/PhysicsAssetEditor/Private/PhysicsAssetEditorSharedData.h
+++ b/Engine/Source/Editor/PhysicsAssetEditor/Private/PhysicsAssetEditorSharedData.h
@@ -145,6 +145,12 @@ public:
 	bool IsBodySelected(const FSelection& Body) const;
 	void ToggleSelectionType(bool bIgnoreUserConstraints = true);
 	void ToggleShowSelected();
+	bool IsBodyHidden(const int32 BodyIndex) const;
+	bool IsConstraintHidden(const int32 ConstraintIndex) const;
+	void HideBody(const int32 BodyIndex);
+	void ShowBody(const int32 BodyIndex);
+	void HideConstraint(const int32 ConstraintIndex);
+	void ShowConstraint(const int32 ConstraintIndex);
 	void ShowAll();
 	void HideAll();
 	void HideAllBodies();
@@ -315,8 +321,6 @@ public:
 
 	TArray<FSelection> SelectedBodies;
 
-	TArray<int32> HiddenBodies;
-	TArray<int32> HiddenConstraints;
 	FSelection * GetSelectedBody()
 	{
 		int32 Count = SelectedBodies.Num();
@@ -337,6 +341,8 @@ public:
 		return Count ? &SelectedConstraints[Count - 1] : NULL;
 	}
 
+	struct FPhysicsAssetRenderSettings* GetRenderSettings() const;
+
 	/** Show flags */
 	bool bShowCOM;
 
diff --git a/Engine/Source/Editor/PhysicsAssetEditor/Private/PhysicsAssetEditorSkeletalMeshComponent.cpp b/Engine/Source/Editor/PhysicsAssetEditor/Private/PhysicsAssetEditorSkeletalMeshComponent.cpp
--- a/Engine/Source/Editor/PhysicsAssetEditor/Private/PhysicsAssetEditorSkeletalMeshComponent.cpp
+++ b/Engine/Source/Editor/PhysicsAssetEditor/Private/PhysicsAssetEditorSkeletalMeshComponent.cpp
@@ -11,6 +11,7 @@
 #include "PhysicsAssetEditorHitProxies.h"
 #include "PhysicsAssetEditorSkeletalMeshComponent.h"
 #include "PhysicsAssetEditorAnimInstance.h"
+#include "PhysicsAssetRenderUtils.h"
 #include "PhysicsEngine/PhysicsConstraintTemplate.h"
 #include "PhysicsEngine/PhysicsAsset.h"
 #include "Chaos/Core.h"
@@ -21,13 +22,8 @@
 
 namespace
 {
-	// How large to make the constraint arrows.
-	// The factor of 60 was found experimentally, to look reasonable in comparison with the rest of the constraint visuals.
-	constexpr float ConstraintArrowScale = 60.0f;
-
 	bool bDebugViewportClicks = false;
 	FAutoConsoleVariableRef CVarChaosImmPhysStepTime(TEXT("p.PhAT.DebugViewportClicks"), bDebugViewportClicks, TEXT("Set to 1 to show mouse click results in PhAT"));
-
 }
 
 UPhysicsAssetEditorSkeletalMeshComponent::UPhysicsAssetEditorSkeletalMeshComponent(const FObjectInitializer& ObjectInitializer)
@@ -72,8 +68,7 @@ UPhysicsAssetEditorSkeletalMeshComponent::UPhysicsAssetEditorSkeletalMeshCompone
 
 	bSelectable = false;
 }
-
-void UPhysicsAssetEditorSkeletalMeshComponent::RenderAssetTools(const FSceneView* View, class FPrimitiveDrawInterface* PDI)
+void UPhysicsAssetEditorSkeletalMeshComponent::DebugDraw(const FSceneView* View, FPrimitiveDrawInterface* PDI)
 {
 	check(SharedData);
 
@@ -107,243 +102,43 @@ void UPhysicsAssetEditorSkeletalMeshComponent::RenderAssetTools(const FSceneView
 
 	ElemSelectedMaterial->SetVectorParameterValue(SelectionColorName, LinearSelectionColor);
 
-	// Contains info needed for drawing, in a form that can be sorted by distance.
-	struct DrawElement
-	{
-		DrawElement(
-			FTransform        InTM, 
-			HHitProxy*        InHitProxy, 
-			float             InScale, 
-			const FSceneView* InView,
-			FKShapeElem*      InShapeElem,
-			const FName&      InBoneName
-		)
-			: TM(InTM), HitProxy(InHitProxy), Scale(InScale), ShapeElem(InShapeElem) 
-#ifdef UE_BUILD_DEBUG
-			, BoneName(InBoneName)
-#endif
-		{
-			// The TM position is at the center of the objects, so can be used directly for sorting.
-			// Sorting by distance (rather than along the view direction) reduces flickering when
-			// the camera is rotated without moving it.
-			FVector ViewDir = InView->GetViewDirection();
-			//Distance = TM.GetTranslation() | ViewDir; // Would sort along the view direction
-			SortingMetric = (TM.GetTranslation() - InView->ViewLocation).SquaredLength(); // Sorts by Distance
-		}
-
-		UMaterialInterface*   Material = 0;
-		FColor                Color = FColor(ForceInitToZero);
-		FTransform            TM;
-		HHitProxy*            HitProxy = nullptr;
-		float                 Scale;
-		FKShapeElem*          ShapeElem = nullptr;
-		float                 SortingMetric;
-#ifdef UE_BUILD_DEBUG
-		FName                 BoneName;
-#endif
-	};
-	TArray<DrawElement> DrawElements;
+	
 
-	// Draw bodies
-	for (int32 i = 0; i <PhysicsAsset->SkeletalBodySetups.Num(); ++i)
+	FPhysicsAssetRenderSettings* const RenderSettings = UPhysicsAssetRenderUtilities::GetSettings(PhysicsAsset);
+	
+	if (RenderSettings)
 	{
-		if (!ensure(PhysicsAsset->SkeletalBodySetups[i]))
-		{
-			continue;
-		}
-		if ((PhysicsAsset->SkeletalBodySetups[i]->PhysicsType == EPhysicsType::PhysType_Kinematic &&
-			SharedData->EditorOptions->bHideKinematicBodies) ||
-			(PhysicsAsset->SkeletalBodySetups[i]->PhysicsType == EPhysicsType::PhysType_Simulated &&
-			SharedData->EditorOptions->bHideSimulatedBodies)
-			)
-		{
-			continue;
-		}
-		if (SharedData->HiddenBodies.Contains(i))
-		{
-			continue;
-		}
-		int32 BoneIndex = GetBoneIndex(PhysicsAsset->SkeletalBodySetups[i]->BoneName);
+		// Copy render settings from editor viewport. These settings must be applied to the rendering in all editors 
+		// when an asset is open in the Physics Asset Editor but should not persist after the editor has been closed.
+		RenderSettings->CollisionViewMode = SharedData->GetCurrentCollisionViewMode(SharedData->bRunningSimulation);
+		RenderSettings->ConstraintViewMode = SharedData->GetCurrentConstraintViewMode(SharedData->bRunningSimulation);
+		RenderSettings->ConstraintDrawSize = SharedData->EditorOptions->ConstraintDrawSize;
+		RenderSettings->PhysicsBlend = SharedData->EditorOptions->PhysicsBlend;
+		RenderSettings->bHideKinematicBodies = SharedData->EditorOptions->bHideKinematicBodies;
+		RenderSettings->bHideSimulatedBodies = SharedData->EditorOptions->bHideSimulatedBodies;
+		RenderSettings->bRenderOnlySelectedConstraints = SharedData->EditorOptions->bRenderOnlySelectedConstraints;
+		RenderSettings->bShowConstraintsAsPoints = SharedData->EditorOptions->bShowConstraintsAsPoints;
 
-		// If we found a bone for it, draw the collision.
-		// The logic is as follows; always render in the ViewMode requested when not in hit mode - but if we are in hit mode and the right editing mode, render as solid
-		if (BoneIndex != INDEX_NONE)
+		// Draw Bodies.
 		{
-			FTransform BoneTM = GetBoneTransform(BoneIndex);
-			float Scale = BoneTM.GetScale3D().GetAbsMax();
-			BoneTM.RemoveScaling();
-
-			FKAggregateGeom* AggGeom = &PhysicsAsset->SkeletalBodySetups[i]->AggGeom;
-
-			for (int32 j = 0; j <AggGeom->SphereElems.Num(); ++j)
-			{
-				DrawElement DE(
-					GetPrimitiveTransform(BoneTM, i, EAggCollisionShape::Sphere, j, Scale), 
-					new HPhysicsAssetEditorEdBoneProxy(i, EAggCollisionShape::Sphere, j),
-					Scale, View, &AggGeom->SphereElems[j], GetBoneName(BoneIndex)
-				);
-
-				//solids are drawn if it's the ViewMode and we're not doing a hit, or if it's hitAndBodyMode
-				if (CollisionViewMode == EPhysicsAssetEditorCollisionViewMode::Solid || CollisionViewMode == EPhysicsAssetEditorCollisionViewMode::SolidWireframe)
-				{
-					DE.Material = GetPrimitiveMaterial(i, EAggCollisionShape::Sphere, j);
-				}
-				if (CollisionViewMode == EPhysicsAssetEditorCollisionViewMode::SolidWireframe || CollisionViewMode == EPhysicsAssetEditorCollisionViewMode::Wireframe)
-				{
-					DE.Color = GetPrimitiveColor(i, EAggCollisionShape::Sphere, j);
-				}
-				if (DE.Material || DE.Color.A)
-				{
-					DrawElements.Add(DE);
-				}
-			}
-
-			for (int32 j = 0; j <AggGeom->BoxElems.Num(); ++j)
-			{
-				DrawElement DE(
-					GetPrimitiveTransform(BoneTM, i, EAggCollisionShape::Box, j, Scale),
-					new HPhysicsAssetEditorEdBoneProxy(i, EAggCollisionShape::Box, j),
-					Scale, View, &AggGeom->BoxElems[j], GetBoneName(BoneIndex)
-				);
-
-				if (CollisionViewMode == EPhysicsAssetEditorCollisionViewMode::Solid || CollisionViewMode == EPhysicsAssetEditorCollisionViewMode::SolidWireframe)
-				{
-					DE.Material = GetPrimitiveMaterial(i, EAggCollisionShape::Box, j);
-				}
-				if (CollisionViewMode == EPhysicsAssetEditorCollisionViewMode::SolidWireframe || CollisionViewMode == EPhysicsAssetEditorCollisionViewMode::Wireframe)
-				{
-					DE.Color = GetPrimitiveColor(i, EAggCollisionShape::Box, j);
-				}
-				if (DE.Material || DE.Color.A)
-				{
-					DrawElements.Add(DE);
-				}
-			}
-
-			for (int32 j = 0; j <AggGeom->SphylElems.Num(); ++j)
-			{
-				DrawElement DE(
-					GetPrimitiveTransform(BoneTM, i, EAggCollisionShape::Sphyl, j, Scale),
-					new HPhysicsAssetEditorEdBoneProxy(i, EAggCollisionShape::Sphyl, j),
-					Scale, View, &AggGeom->SphylElems[j], GetBoneName(BoneIndex)
-				);
-
-				if (CollisionViewMode == EPhysicsAssetEditorCollisionViewMode::Solid || CollisionViewMode == EPhysicsAssetEditorCollisionViewMode::SolidWireframe)
-				{
-					DE.Material = GetPrimitiveMaterial(i, EAggCollisionShape::Sphyl, j);
-				}
-				if (CollisionViewMode == EPhysicsAssetEditorCollisionViewMode::SolidWireframe || CollisionViewMode == EPhysicsAssetEditorCollisionViewMode::Wireframe)
-				{
-					DE.Color = GetPrimitiveColor(i, EAggCollisionShape::Sphyl, j);
-				}
-				if (DE.Material || DE.Color.A)
-				{
-					DrawElements.Add(DE);
-				}
-			}
-
-			for (int32 j = 0; j <AggGeom->ConvexElems.Num(); ++j)
-			{
-				DrawElement DE(
-					GetPrimitiveTransform(BoneTM, i, EAggCollisionShape::Convex, j, Scale),
-					new HPhysicsAssetEditorEdBoneProxy(i, EAggCollisionShape::Convex, j),
-					Scale, View, &AggGeom->ConvexElems[j], GetBoneName(BoneIndex)
-				);
-
-				if (CollisionViewMode == EPhysicsAssetEditorCollisionViewMode::Solid || CollisionViewMode == EPhysicsAssetEditorCollisionViewMode::SolidWireframe)
-				{
-					DE.Material = GetPrimitiveMaterial(i, EAggCollisionShape::Convex, j);
-				}
-				if (CollisionViewMode == EPhysicsAssetEditorCollisionViewMode::SolidWireframe || CollisionViewMode == EPhysicsAssetEditorCollisionViewMode::Wireframe)
-				{
-					DE.Color = GetPrimitiveColor(i, EAggCollisionShape::Convex, j);
-				}
-				if (DE.Material || DE.Color.A)
-				{
-					DrawElements.Add(DE);
-				}
-			}
-
-			for (int32 j = 0; j <AggGeom->TaperedCapsuleElems.Num(); ++j)
-			{
-				DrawElement DE(
-					GetPrimitiveTransform(BoneTM, i, EAggCollisionShape::TaperedCapsule, j, Scale),
-					new HPhysicsAssetEditorEdBoneProxy(i, EAggCollisionShape::TaperedCapsule, j),
-					Scale, View, &AggGeom->TaperedCapsuleElems[j], GetBoneName(BoneIndex)
-				);
-
-				if (CollisionViewMode == EPhysicsAssetEditorCollisionViewMode::Solid || CollisionViewMode == EPhysicsAssetEditorCollisionViewMode::SolidWireframe)
-				{
-					DE.Material = GetPrimitiveMaterial(i, EAggCollisionShape::TaperedCapsule, j);
-				}
-				if (CollisionViewMode == EPhysicsAssetEditorCollisionViewMode::SolidWireframe || CollisionViewMode == EPhysicsAssetEditorCollisionViewMode::Wireframe)
-				{
-					DE.Color = GetPrimitiveColor(i, EAggCollisionShape::TaperedCapsule, j);
-				}
-				if (DE.Material || DE.Color.A)
-				{
-					DrawElements.Add(DE);
-				}
-			}
+			auto TransformFn = [this](const UPhysicsAsset* PhysicsAsset, const FTransform& BoneTM, const int32 BodyIndex, const EAggCollisionShape::Type PrimType, const int32 PrimIndex, const float Scale) { return this->GetPrimitiveTransform(BoneTM, BodyIndex, PrimType, PrimIndex, Scale);  };
+			auto ColorFn = [this](const int32 BodyIndex, const EAggCollisionShape::Type PrimitiveType, const int32 PrimitiveIndex, const FPhysicsAssetRenderSettings& Settings) { return this->GetPrimitiveColor(BodyIndex, PrimitiveType, PrimitiveIndex); };
+			auto MaterialFn = [this](const int32 BodyIndex, const EAggCollisionShape::Type PrimitiveType, const int32 PrimitiveIndex, const FPhysicsAssetRenderSettings& Settings) { return this->GetPrimitiveMaterial(BodyIndex, PrimitiveType, PrimitiveIndex); };
+			auto HitProxyFn = [](const int32 BodyIndex, const EAggCollisionShape::Type PrimitiveType, const int32 PrimitiveIndex) { return new HPhysicsAssetEditorEdBoneProxy(BodyIndex, PrimitiveType, PrimitiveIndex); };
 
-			if (SharedData->bShowCOM && Bodies.IsValidIndex(i))
-			{
-				Bodies[i]->DrawCOMPosition(PDI, COMRenderSize, SharedData->COMRenderColor);
-			}
-		}
-	}
-
-	// Sort elements
-	DrawElements.Sort([this](const DrawElement& DE1, const DrawElement& DE2) {
-		return DE1.SortingMetric > DE2.SortingMetric; 
-		});
-
-
-	// Render sorted elements.
-	for (const DrawElement& DE : DrawElements)
-	{
-		PDI->SetHitProxy(DE.HitProxy);
-		if (DE.Material)
-		{
-			DE.ShapeElem->DrawElemSolid(PDI, DE.TM, DE.Scale, DE.Material->GetRenderProxy());
-		}
-		if (DE.Color.A)
-		{
-			DE.ShapeElem->DrawElemWire(PDI, DE.TM, DE.Scale, DE.Color);
+			PhysicsAssetRender::DebugDrawBodies(this, PhysicsAsset, PDI, ColorFn, MaterialFn, TransformFn, HitProxyFn);
 		}
-		PDI->SetHitProxy(NULL);
-	}
 
-	// Draw Constraints
-	EPhysicsAssetEditorConstraintViewMode ConstraintViewMode = SharedData->GetCurrentConstraintViewMode(SharedData->bRunningSimulation);
-	if (ConstraintViewMode != EPhysicsAssetEditorConstraintViewMode::None)
-	{
-		for (int32 i = 0; i <PhysicsAsset->ConstraintSetup.Num(); ++i)
+		// Draw Constraints.
 		{
-			if((!SharedData->EditorOptions->bRenderOnlySelectedConstraints || (SharedData->EditorOptions->bRenderOnlySelectedConstraints && SharedData->IsConstraintSelected(i))) &&
-				!SharedData->HiddenConstraints.Contains(i))
-			{
-				int32 BoneIndex1 = GetBoneIndex(PhysicsAsset->ConstraintSetup[i]->DefaultInstance.ConstraintBone1);
-				int32 BoneIndex2 = GetBoneIndex(PhysicsAsset->ConstraintSetup[i]->DefaultInstance.ConstraintBone2);
-				// if bone doesn't exist, do not draw it. It crashes in random points when we try to manipulate. 
-				if (BoneIndex1 != INDEX_NONE && BoneIndex2 != INDEX_NONE)
-				{
-					PDI->SetHitProxy(new HPhysicsAssetEditorEdConstraintProxy(i));
+			auto HitProxyFn = [](const int32 InConstraintIndex) { return new HPhysicsAssetEditorEdConstraintProxy(InConstraintIndex); };
+			auto IsConstraintSelectedFn = [this](const uint32 InConstraintIndex) { return this->SharedData->IsConstraintSelected(InConstraintIndex); };
 
-					DrawConstraint(i, View, PDI, SharedData->EditorOptions->bShowConstraintsAsPoints);
-
-					PDI->SetHitProxy(NULL);
-				}
-			}
+			PhysicsAssetRender::DebugDrawConstraints(this, PhysicsAsset, PDI, IsConstraintSelectedFn, SharedData->bRunningSimulation, HitProxyFn);
 		}
 	}
 }
 
-void UPhysicsAssetEditorSkeletalMeshComponent::DebugDraw(const FSceneView* View, FPrimitiveDrawInterface* PDI)
-{
-	RenderAssetTools(View, PDI);
-}
-
 FPrimitiveSceneProxy* UPhysicsAssetEditorSkeletalMeshComponent::CreateSceneProxy()
 {
 	FPrimitiveSceneProxy* Proxy = NULL;
@@ -370,33 +165,7 @@ bool ConstraintInSelected(int32 Index, const TArray<FPhysicsAssetEditorSharedDat
 	return false;
 }
 
-void UPhysicsAssetEditorSkeletalMeshComponent::DrawConstraint(int32 ConstraintIndex, const FSceneView* View, FPrimitiveDrawInterface* PDI, bool bDrawAsPoint)
-{
-	EPhysicsAssetEditorConstraintViewMode ConstraintViewMode = SharedData->GetCurrentConstraintViewMode(SharedData->bRunningSimulation);
-	bool bDrawLimits = false;
-	bool bConstraintSelected = ConstraintInSelected(ConstraintIndex, SharedData->SelectedConstraints);
-	if (ConstraintViewMode == EPhysicsAssetEditorConstraintViewMode::AllLimits || bConstraintSelected)
-	{
-		bDrawLimits = true;
-	}
-
-	bool bDrawSelected = false;
-	if (!SharedData->bRunningSimulation && bConstraintSelected)
-	{
-		bDrawSelected = true;
-	}
-
-	UPhysicsConstraintTemplate* ConstraintSetup = SharedData->PhysicsAsset->ConstraintSetup[ConstraintIndex];
-
-	FTransform Con1Frame = SharedData->GetConstraintMatrix(ConstraintIndex, EConstraintFrame::Frame1, 1.f);
-	FTransform Con2Frame = SharedData->GetConstraintMatrix(ConstraintIndex, EConstraintFrame::Frame2, 1.f);
-
-	const float DrawScale = ConstraintArrowScale * SharedData->EditorOptions->ConstraintDrawSize;
-
-	ConstraintSetup->DefaultInstance.DrawConstraint(PDI, SharedData->EditorOptions->ConstraintDrawSize, DrawScale, bDrawLimits, bDrawSelected, Con1Frame, Con2Frame, bDrawAsPoint);
-}
-
-FTransform UPhysicsAssetEditorSkeletalMeshComponent::GetPrimitiveTransform(FTransform& BoneTM, int32 BodyIndex, EAggCollisionShape::Type PrimType, int32 PrimIndex, float Scale)
+FTransform UPhysicsAssetEditorSkeletalMeshComponent::GetPrimitiveTransform(const FTransform& BoneTM, const int32 BodyIndex, const EAggCollisionShape::Type PrimType, const int32 PrimIndex, const float Scale) const
 {
 	UBodySetup* SharedBodySetup = SharedData->PhysicsAsset->SkeletalBodySetups[BodyIndex];
 	FVector Scale3D(Scale);
@@ -454,7 +223,7 @@ FTransform UPhysicsAssetEditorSkeletalMeshComponent::GetPrimitiveTransform(FTran
 	return FTransform::Identity;
 }
 
-FColor UPhysicsAssetEditorSkeletalMeshComponent::GetPrimitiveColor(int32 BodyIndex, EAggCollisionShape::Type PrimitiveType, int32 PrimitiveIndex)
+FColor UPhysicsAssetEditorSkeletalMeshComponent::GetPrimitiveColor(const int32 BodyIndex, const EAggCollisionShape::Type PrimitiveType, const int32 PrimitiveIndex) const
 {
 	UBodySetup* SharedBodySetup = SharedData->PhysicsAsset->SkeletalBodySetups[BodyIndex];
 
@@ -526,7 +295,7 @@ FColor UPhysicsAssetEditorSkeletalMeshComponent::GetPrimitiveColor(int32 BodyInd
 	return BoneUnselectedColor;
 }
 
-UMaterialInterface* UPhysicsAssetEditorSkeletalMeshComponent::GetPrimitiveMaterial(int32 BodyIndex, EAggCollisionShape::Type PrimitiveType, int32 PrimitiveIndex)
+UMaterialInterface* UPhysicsAssetEditorSkeletalMeshComponent::GetPrimitiveMaterial(const int32 BodyIndex, const EAggCollisionShape::Type PrimitiveType, const int32 PrimitiveIndex) const
 {
 	if (SharedData->bRunningSimulation)
 	{
diff --git a/Engine/Source/Editor/PhysicsAssetEditor/Private/PhysicsAssetEditorSkeletalMeshComponent.h b/Engine/Source/Editor/PhysicsAssetEditor/Private/PhysicsAssetEditorSkeletalMeshComponent.h
--- a/Engine/Source/Editor/PhysicsAssetEditor/Private/PhysicsAssetEditorSkeletalMeshComponent.h
+++ b/Engine/Source/Editor/PhysicsAssetEditor/Private/PhysicsAssetEditorSkeletalMeshComponent.h
@@ -59,16 +59,10 @@ class UPhysicsAssetEditorSkeletalMeshComponent : public UDebugSkelMeshComponent
 	/** Debug drawing */
 	void DebugDraw(const FSceneView* View, FPrimitiveDrawInterface* PDI);
 
-	/** Handles most of the rendering logic for this component */
-	void RenderAssetTools(const FSceneView* View, FPrimitiveDrawInterface* PDI);
-
-	/** Draws a constraint */
-	void DrawConstraint(int32 ConstraintIndex, const FSceneView* View, FPrimitiveDrawInterface* PDI, bool bDrawAsPoint);
-
 	/** Accessors/helper methods */
-	FTransform GetPrimitiveTransform(FTransform& BoneTM, int32 BodyIndex, EAggCollisionShape::Type PrimType, int32 PrimIndex, float Scale);
-	FColor GetPrimitiveColor(int32 BodyIndex, EAggCollisionShape::Type PrimitiveType, int32 PrimitiveIndex);
-	UMaterialInterface* GetPrimitiveMaterial(int32 BodyIndex, EAggCollisionShape::Type PrimitiveType, int32 PrimitiveIndex);
+	FTransform GetPrimitiveTransform(const FTransform& BoneTM, const int32 BodyIndex, const EAggCollisionShape::Type PrimType, const int32 PrimIndex, const float Scale) const;
+	FColor GetPrimitiveColor(const int32 BodyIndex, const EAggCollisionShape::Type PrimitiveType, const int32 PrimitiveIndex) const;
+	UMaterialInterface* GetPrimitiveMaterial(const int32 BodyIndex, const EAggCollisionShape::Type PrimitiveType, const int32 PrimitiveIndex) const;
 
 	/** Manipulator methods */
 	virtual void Grab(FName InBoneName, const FVector& Location, const FRotator& Rotation, bool bRotationConstrained);
diff --git a/Engine/Source/Editor/PhysicsAssetEditor/Private/PhysicsAssetEditorSkeletonTreeBuilder.cpp b/Engine/Source/Editor/PhysicsAssetEditor/Private/PhysicsAssetEditorSkeletonTreeBuilder.cpp
--- a/Engine/Source/Editor/PhysicsAssetEditor/Private/PhysicsAssetEditorSkeletonTreeBuilder.cpp
+++ b/Engine/Source/Editor/PhysicsAssetEditor/Private/PhysicsAssetEditorSkeletonTreeBuilder.cpp
@@ -145,7 +145,7 @@ void FPhysicsAssetEditorSkeletonTreeBuilder::AddBodies(FSkeletonTreeBuilderOutpu
 
 						if (bHasShapes)
 						{
-							Output.Add(MakeShared<FSkeletonTreePhysicsBodyItem>(BodySetup, BodySetupIndex, BoneName, true, bHasShapes, SkeletonTreePtr.Pin().ToSharedRef()), BoneName, "FSkeletonTreeBoneItem", true);
+							Output.Add(MakeShared<FSkeletonTreePhysicsBodyItem>(BodySetup, BodySetupIndex, BoneName, true, bHasShapes, PhysicsAsset, SkeletonTreePtr.Pin().ToSharedRef()), BoneName, "FSkeletonTreeBoneItem", true);
 
 							int32 ShapeIndex;
 							for (ShapeIndex = 0; ShapeIndex < AggGeom.SphereElems.Num(); ++ShapeIndex)
@@ -184,7 +184,7 @@ void FPhysicsAssetEditorSkeletonTreeBuilder::AddBodies(FSkeletonTreeBuilderOutpu
 							if (bJointMatches || bUserConstraintMatches)
 							{
 								const bool bIsConstraintOnParentBone = !(ConstraintInstance.JointName == BoneName || ConstraintInstance.ConstraintBone1 == BoneName);
-								Output.Add(MakeShared<FSkeletonTreePhysicsConstraintItem>(PhysicsAsset->ConstraintSetup[ConstraintIndex], ConstraintIndex, BoneName, bIsConstraintOnParentBone, SkeletonTreePtr.Pin().ToSharedRef()), BoneName, FSkeletonTreePhysicsBodyItem::GetTypeId());
+								Output.Add(MakeShared<FSkeletonTreePhysicsConstraintItem>(PhysicsAsset->ConstraintSetup[ConstraintIndex], ConstraintIndex, BoneName, bIsConstraintOnParentBone, PhysicsAsset, SkeletonTreePtr.Pin().ToSharedRef()), BoneName, FSkeletonTreePhysicsBodyItem::GetTypeId());
 							}
 						}
 
diff --git a/Engine/Source/Editor/PhysicsAssetEditor/Private/SkeletonTreePhysicsBodyItem.cpp b/Engine/Source/Editor/PhysicsAssetEditor/Private/SkeletonTreePhysicsBodyItem.cpp
--- a/Engine/Source/Editor/PhysicsAssetEditor/Private/SkeletonTreePhysicsBodyItem.cpp
+++ b/Engine/Source/Editor/PhysicsAssetEditor/Private/SkeletonTreePhysicsBodyItem.cpp
@@ -1,39 +1,43 @@
 // Copyright Epic Games, Inc. All Rights Reserved.
 
 #include "SkeletonTreePhysicsBodyItem.h"
-#include "Widgets/Text/STextBlock.h"
-#include "Widgets/Images/SImage.h"
 #include "EditorStyleSet.h"
+#include "PhysicsAssetRenderUtils.h"
+#include "PhysicsEngine/PhysicsAsset.h"
 
 #define LOCTEXT_NAMESPACE "FSkeletonTreePhysicsBodyItem"
 
-void FSkeletonTreePhysicsBodyItem::GenerateWidgetForNameColumn( TSharedPtr< SHorizontalBox > Box, const TAttribute<FText>& FilterText, FIsSelected InIsSelected )
+FSkeletonTreePhysicsBodyItem::FSkeletonTreePhysicsBodyItem(USkeletalBodySetup* InBodySetup, int32 InBodySetupIndex, const FName& InBoneName, bool bInHasBodySetup, bool bInHasShapes, class UPhysicsAsset* const InPhysicsAsset, const TSharedRef<class ISkeletonTree>& InSkeletonTree)
+	: FSkeletonTreePhysicsItem(InPhysicsAsset, InSkeletonTree)
+	, BodySetup(InBodySetup)
+	, BodySetupIndex(InBodySetupIndex)
+	, bHasBodySetup(bInHasBodySetup)
+	, bHasShapes(bInHasShapes)
 {
-	Box->AddSlot()
-	.AutoWidth()
-	.Padding(FMargin(0.0f, 1.0f))
-	[
-		SNew( SImage )
-		.ColorAndOpacity(FSlateColor::UseForeground())
-		.Image(this, &FSkeletonTreePhysicsBodyItem::GetBrush)
-	];
+	DisplayName = InBoneName;
+}
 
-	Box->AddSlot()
-	.AutoWidth()
-	.Padding(2, 0, 0, 0)
-	[
-		SNew(STextBlock)
-		.ColorAndOpacity(this, &FSkeletonTreePhysicsBodyItem::GetBodyTextColor)
-		.Text(FText::FromName(BoneName))
-		.HighlightText(FilterText)
-		.Font(FEditorStyle::GetFontStyle("PhysicsAssetEditor.Tree.Font"))
-		.ToolTipText(FText::Format(LOCTEXT("BodyTooltip", "Aggregate physics body for bone '{0}'. Bodies can consist of multiple shapes."), FText::FromName(BoneName)))
-	];
+UObject* FSkeletonTreePhysicsBodyItem::GetObject() const
+{
+	return BodySetup;
 }
 
-TSharedRef< SWidget > FSkeletonTreePhysicsBodyItem::GenerateWidgetForDataColumn(const FName& DataColumnName, FIsSelected InIsSelected)
+void FSkeletonTreePhysicsBodyItem::OnToggleItemDisplayed(ECheckBoxState InCheckboxState)
 {
-	return SNullWidget::NullWidget;
+	if (FPhysicsAssetRenderSettings* RenderSettings = GetRenderSettings())
+	{
+		RenderSettings->ToggleShowBody(BodySetupIndex);
+	}
+}
+
+ECheckBoxState FSkeletonTreePhysicsBodyItem::IsItemDisplayed() const
+{
+	if (FPhysicsAssetRenderSettings* RenderSettings = GetRenderSettings())
+	{
+		return RenderSettings->IsBodyHidden(BodySetupIndex) ? ECheckBoxState::Unchecked : ECheckBoxState::Checked;
+	}
+
+	return ECheckBoxState::Undetermined;
 }
 
 const FSlateBrush* FSkeletonTreePhysicsBodyItem::GetBrush() const
@@ -41,7 +45,7 @@ const FSlateBrush* FSkeletonTreePhysicsBodyItem::GetBrush() const
 	return BodySetup->PhysicsType == EPhysicsType::PhysType_Kinematic ? FEditorStyle::GetBrush("PhysicsAssetEditor.Tree.KinematicBody") : FEditorStyle::GetBrush("PhysicsAssetEditor.Tree.Body");
 }
 
-FSlateColor FSkeletonTreePhysicsBodyItem::GetBodyTextColor() const
+FSlateColor FSkeletonTreePhysicsBodyItem::GetTextColor() const
 {
 	FLinearColor Color(1.0f, 1.0f, 1.0f);
 
@@ -62,4 +66,9 @@ FSlateColor FSkeletonTreePhysicsBodyItem::GetBodyTextColor() const
 	}
 }
 
+FText FSkeletonTreePhysicsBodyItem::GetNameColumnToolTip() const
+{
+	return FText::Format(LOCTEXT("BodyTooltip", "Aggregate physics body for bone '{0}'. Bodies can consist of multiple shapes."), FText::FromName(GetRowItemName()));
+}
+
 #undef LOCTEXT_NAMESPACE
diff --git a/Engine/Source/Editor/PhysicsAssetEditor/Private/SkeletonTreePhysicsBodyItem.h b/Engine/Source/Editor/PhysicsAssetEditor/Private/SkeletonTreePhysicsBodyItem.h
--- a/Engine/Source/Editor/PhysicsAssetEditor/Private/SkeletonTreePhysicsBodyItem.h
+++ b/Engine/Source/Editor/PhysicsAssetEditor/Private/SkeletonTreePhysicsBodyItem.h
@@ -4,53 +4,40 @@
 
 #include "CoreMinimal.h"
 #include "Styling/SlateColor.h"
-#include "Input/Reply.h"
-#include "Widgets/SWidget.h"
-#include "Widgets/Views/STableViewBase.h"
-#include "Widgets/Views/STableRow.h"
-#include "SkeletonTreeItem.h"
-#include "PhysicsEngine/PhysicsAsset.h"
-
-class FSkeletonTreePhysicsBodyItem : public FSkeletonTreeItem
+#include "SkeletonTreePhysicsItem.h"
+
+class FSkeletonTreePhysicsBodyItem : public FSkeletonTreePhysicsItem
 {
 public:
-	SKELETON_TREE_ITEM_TYPE(FSkeletonTreePhysicsBodyItem, FSkeletonTreeItem)
-
-	FSkeletonTreePhysicsBodyItem(USkeletalBodySetup* InBodySetup, int32 InBodySetupIndex, const FName& InBoneName, bool bInHasBodySetup, bool bInHasShapes, const TSharedRef<class ISkeletonTree>& InSkeletonTree)
-		: FSkeletonTreeItem(InSkeletonTree)
-		, BodySetup(InBodySetup)
-		, BodySetupIndex(InBodySetupIndex)
-		, BoneName(InBoneName)
-		, bHasBodySetup(bInHasBodySetup)
-		, bHasShapes(bInHasShapes)
-	{}
-
-	/** ISkeletonTreeItem interface */
-	virtual void GenerateWidgetForNameColumn(TSharedPtr< SHorizontalBox > Box, const TAttribute<FText>& FilterText, FIsSelected InIsSelected) override;
-	virtual TSharedRef< SWidget > GenerateWidgetForDataColumn(const FName& DataColumnName, FIsSelected InIsSelected) override;	
-	virtual FName GetRowItemName() const override { return BoneName; }
-	virtual UObject* GetObject() const override { return BodySetup; }
+	SKELETON_TREE_ITEM_TYPE(FSkeletonTreePhysicsBodyItem, FSkeletonTreePhysicsItem)
+
+	FSkeletonTreePhysicsBodyItem(class USkeletalBodySetup* InBodySetup, int32 InBodySetupIndex, const FName& InBoneName, bool bInHasBodySetup, bool bInHasShapes, class UPhysicsAsset* const InPhysicsAsset, const TSharedRef<class ISkeletonTree>& InSkeletonTree);
+
+	virtual UObject* GetObject() const override;
 
 	/** Get the index of the body setup in the physics asset */
 	int32 GetBodySetupIndex() const { return BodySetupIndex; }
 
+	/** UI Callbacks */
+	virtual void OnToggleItemDisplayed(ECheckBoxState InCheckboxState) override;
+	virtual ECheckBoxState IsItemDisplayed() const override;
+
 private:
 	/** Gets the icon to display for this body */
-	const FSlateBrush* GetBrush() const;
+	virtual const FSlateBrush* GetBrush() const override;
 
 	/** Gets the color to display the item's text */
-	FSlateColor GetBodyTextColor() const;
+	virtual FSlateColor GetTextColor() const override;
+
+	/** Gets the tool tip to display on hovering over the item's name */
+	virtual FText GetNameColumnToolTip() const override;
 
-private:
 	/** The body setup we are representing */
 	USkeletalBodySetup* BodySetup;
 
 	/** The index of the body setup in the physics asset */
 	int32 BodySetupIndex;
 
-	/** The name of the bone that this body is bound to */
-	FName BoneName;
-
 	/** Whether there is a body set up for this bone */
 	bool bHasBodySetup;
 
diff --git a/Engine/Source/Editor/PhysicsAssetEditor/Private/SkeletonTreePhysicsConstraintItem.cpp b/Engine/Source/Editor/PhysicsAssetEditor/Private/SkeletonTreePhysicsConstraintItem.cpp
--- a/Engine/Source/Editor/PhysicsAssetEditor/Private/SkeletonTreePhysicsConstraintItem.cpp
+++ b/Engine/Source/Editor/PhysicsAssetEditor/Private/SkeletonTreePhysicsConstraintItem.cpp
@@ -1,14 +1,13 @@
 // Copyright Epic Games, Inc. All Rights Reserved.
 
 #include "SkeletonTreePhysicsConstraintItem.h"
-#include "Widgets/Text/STextBlock.h"
-#include "Widgets/Images/SImage.h"
 #include "EditorStyleSet.h"
+#include "PhysicsAssetRenderUtils.h"
 
 #define LOCTEXT_NAMESPACE "FSkeletonTreePhysicsConstraintItem"
 
-FSkeletonTreePhysicsConstraintItem::FSkeletonTreePhysicsConstraintItem(UPhysicsConstraintTemplate* InConstraint, int32 InConstraintIndex, const FName& InBoneName, bool bInIsConstraintOnParentBody, const TSharedRef<class ISkeletonTree>& InSkeletonTree)
-	: FSkeletonTreeItem(InSkeletonTree)
+FSkeletonTreePhysicsConstraintItem::FSkeletonTreePhysicsConstraintItem(UPhysicsConstraintTemplate* InConstraint, int32 InConstraintIndex, const FName& InBoneName, bool bInIsConstraintOnParentBody, class UPhysicsAsset* const InPhysicsAsset, const TSharedRef<class ISkeletonTree>& InSkeletonTree)
+	: FSkeletonTreePhysicsItem(InPhysicsAsset, InSkeletonTree)
 	, Constraint(InConstraint)
 	, ConstraintIndex(InConstraintIndex)
 	, bIsConstraintOnParentBody(bInIsConstraintOnParentBody)
@@ -18,38 +17,35 @@ FSkeletonTreePhysicsConstraintItem::FSkeletonTreePhysicsConstraintItem(UPhysicsC
 	DisplayName = *Label.ToString();
 }
 
-void FSkeletonTreePhysicsConstraintItem::GenerateWidgetForNameColumn( TSharedPtr< SHorizontalBox > Box, const TAttribute<FText>& FilterText, FIsSelected InIsSelected )
+UObject* FSkeletonTreePhysicsConstraintItem::GetObject() const
 {
-	Box->AddSlot()
-	.AutoWidth()
-	.Padding(FMargin(0.0f, 1.0f))
-	[
-		SNew( SImage )
-		.ColorAndOpacity(FSlateColor::UseForeground())
-		.Image(FEditorStyle::GetBrush("PhysicsAssetEditor.Tree.Constraint"))
-	];
+	return Constraint;
+}
 
-	const FConstraintInstance& ConstraintInstance = Constraint->DefaultInstance;
+void FSkeletonTreePhysicsConstraintItem::OnToggleItemDisplayed(ECheckBoxState InCheckboxState)
+{
+	if (FPhysicsAssetRenderSettings* RenderSettings = GetRenderSettings())
+	{
+		RenderSettings->ToggleShowConstraint(ConstraintIndex);
+	}
+}
+
+ECheckBoxState FSkeletonTreePhysicsConstraintItem::IsItemDisplayed() const
+{
+	if (FPhysicsAssetRenderSettings* RenderSettings = GetRenderSettings())
+	{
+		return RenderSettings->IsConstraintHidden(ConstraintIndex) ? ECheckBoxState::Unchecked : ECheckBoxState::Checked;
+	}
 
-	Box->AddSlot()
-		.AutoWidth()
-		.Padding(2, 0, 0, 0)
-		[
-			SNew(STextBlock)
-			.ColorAndOpacity(this, &FSkeletonTreePhysicsConstraintItem::GetConstraintTextColor)
-			.Text(FText::FromName(DisplayName))
-			.HighlightText(FilterText)
-			.Font(FEditorStyle::GetFontStyle("PhysicsAssetEditor.Tree.Font"))
-			.ToolTipText(FText::Format(LOCTEXT("ConstraintTooltip", "Constraint linking child body [{0}] to parent body [{1}]"), FText::FromName(ConstraintInstance.ConstraintBone1), FText::FromName(ConstraintInstance.ConstraintBone2)))
-	];
+	return ECheckBoxState::Undetermined;
 }
 
-TSharedRef< SWidget > FSkeletonTreePhysicsConstraintItem::GenerateWidgetForDataColumn(const FName& DataColumnName, FIsSelected InIsSelected)
+const FSlateBrush* FSkeletonTreePhysicsConstraintItem::GetBrush() const
 {
-	return SNullWidget::NullWidget;
+	return	FEditorStyle::GetBrush("PhysicsAssetEditor.Tree.Constraint");
 }
 
-FSlateColor FSkeletonTreePhysicsConstraintItem::GetConstraintTextColor() const
+FSlateColor FSkeletonTreePhysicsConstraintItem::GetTextColor() const
 {
 	const FLinearColor Color(1.0f, 1.0f, 1.0f);
 	const bool bInCurrentProfile = Constraint->GetCurrentConstraintProfileName() == NAME_None || Constraint->ContainsConstraintProfile(Constraint->GetCurrentConstraintProfileName());
@@ -63,4 +59,15 @@ FSlateColor FSkeletonTreePhysicsConstraintItem::GetConstraintTextColor() const
 	}
 }
 
+FText FSkeletonTreePhysicsConstraintItem::GetNameColumnToolTip() const
+{
+	if (Constraint)
+	{
+		const FConstraintInstance& ConstraintInstance = Constraint->DefaultInstance;
+		return FText::Format(LOCTEXT("ConstraintTooltip", "Constraint linking child body [{0}] to parent body [{1}]"), FText::FromName(ConstraintInstance.ConstraintBone1), FText::FromName(ConstraintInstance.ConstraintBone2));
+	}
+
+	return FText::GetEmpty();
+}
+
 #undef LOCTEXT_NAMESPACE
diff --git a/Engine/Source/Editor/PhysicsAssetEditor/Private/SkeletonTreePhysicsConstraintItem.h b/Engine/Source/Editor/PhysicsAssetEditor/Private/SkeletonTreePhysicsConstraintItem.h
--- a/Engine/Source/Editor/PhysicsAssetEditor/Private/SkeletonTreePhysicsConstraintItem.h
+++ b/Engine/Source/Editor/PhysicsAssetEditor/Private/SkeletonTreePhysicsConstraintItem.h
@@ -4,26 +4,18 @@
 
 #include "CoreMinimal.h"
 #include "Styling/SlateColor.h"
-#include "Input/Reply.h"
-#include "Widgets/SWidget.h"
-#include "Widgets/Views/STableViewBase.h"
-#include "Widgets/Views/STableRow.h"
-#include "SkeletonTreeItem.h"
-#include "PhysicsEngine/PhysicsAsset.h"
+#include "SkeletonTreePhysicsItem.h"
 #include "PhysicsEngine/PhysicsConstraintTemplate.h"
 
-class FSkeletonTreePhysicsConstraintItem : public FSkeletonTreeItem
+class FSkeletonTreePhysicsConstraintItem : public FSkeletonTreePhysicsItem
 {
 public:
-	SKELETON_TREE_ITEM_TYPE(FSkeletonTreePhysicsConstraintItem, FSkeletonTreeItem)
+	SKELETON_TREE_ITEM_TYPE(FSkeletonTreePhysicsConstraintItem, FSkeletonTreePhysicsItem)
 
-	FSkeletonTreePhysicsConstraintItem(UPhysicsConstraintTemplate* InConstraint, int32 InConstraintIndex, const FName& InBoneName, bool bInIsConstraintOnParentBody, const TSharedRef<class ISkeletonTree>& InSkeletonTree);
+	FSkeletonTreePhysicsConstraintItem(class UPhysicsConstraintTemplate* InConstraint, int32 InConstraintIndex, const FName& InBoneName, bool bInIsConstraintOnParentBody, class UPhysicsAsset* const InPhysicsAsset, const TSharedRef<class ISkeletonTree>& InSkeletonTree);
 
 	/** ISkeletonTreeItem interface */
-	virtual void GenerateWidgetForNameColumn(TSharedPtr< SHorizontalBox > Box, const TAttribute<FText>& FilterText, FIsSelected InIsSelected) override;
-	virtual TSharedRef< SWidget > GenerateWidgetForDataColumn(const FName& DataColumnName, FIsSelected InIsSelected) override;	
-	virtual FName GetRowItemName() const override { return DisplayName; }
-	virtual UObject* GetObject() const override { return Constraint; }
+	virtual UObject* GetObject() const override;
 
 	/** Get the index of the constraint in the physics asset */
 	int32 GetConstraintIndex() const { return ConstraintIndex; }
@@ -31,10 +23,20 @@ public:
 	/** since constraint are showing on both parent and child, gets  if this tree item is the one on the parent body */
 	bool IsConstraintOnParentBody() const { return bIsConstraintOnParentBody; }
 
-private:
-	FSlateColor GetConstraintTextColor() const;
+	/** UI Callbacks */
+	virtual void OnToggleItemDisplayed(ECheckBoxState InCheckboxState) override;
+	virtual ECheckBoxState IsItemDisplayed() const override;
 
 private:
+	/** Gets the icon to display for this body */
+	virtual const FSlateBrush* GetBrush() const override;
+
+	/** Gets the color to display the item's text */
+	virtual FSlateColor GetTextColor() const override;
+
+	/** Gets the tool tip to display on hovering over the item's name */
+	virtual FText GetNameColumnToolTip() const override;
+
 	/** The constraint we are representing */
 	UPhysicsConstraintTemplate* Constraint;
 
@@ -43,7 +45,4 @@ private:
 
 	/** since constraint are showing on both parent and child, indicates if this tree item is the one on the parent body */
 	bool bIsConstraintOnParentBody;
-
-	/** The display name of the item */
-	FName DisplayName;
 };
diff --git a/Engine/Source/Editor/PhysicsAssetEditor/Private/SkeletonTreePhysicsItem.cpp b/Engine/Source/Editor/PhysicsAssetEditor/Private/SkeletonTreePhysicsItem.cpp
new file mode 100644
--- /dev/null
+++ b/Engine/Source/Editor/PhysicsAssetEditor/Private/SkeletonTreePhysicsItem.cpp
@@ -0,0 +1,74 @@
+// Copyright Epic Games, Inc. All Rights Reserved.
+
+#include "SkeletonTreePhysicsItem.h"
+#include "Widgets/Text/STextBlock.h"
+#include "Widgets/Images/SImage.h"
+#include "Widgets/Input/SCheckBox.h"
+#include "EditorStyleSet.h"
+#include "PhysicsAssetRenderUtils.h"
+
+#define LOCTEXT_NAMESPACE "FSkeletonTreePhysicsItem"
+
+FSkeletonTreePhysicsItem::FSkeletonTreePhysicsItem(class UPhysicsAsset* const InPhysicsAsset, const TSharedRef<class ISkeletonTree>& InSkeletonTree)
+	: FSkeletonTreeItem(InSkeletonTree)
+	, PhysicsAssetPathNameHash(UPhysicsAssetRenderUtilities::GetPathNameHash(InPhysicsAsset))
+{}
+
+void FSkeletonTreePhysicsItem::GenerateWidgetForNameColumn(TSharedPtr< SHorizontalBox > Box, const TAttribute<FText>& FilterText, FIsSelected InIsSelected)
+{
+	Box->AddSlot()
+		.AutoWidth()
+		.Padding(FMargin(0.0f, 1.0f))
+		[
+			SNew(SImage)
+			.ColorAndOpacity(FSlateColor::UseForeground())
+			.Image(GetBrush())
+		];
+
+	Box->AddSlot()
+		.AutoWidth()
+		.Padding(2, 0, 0, 0)
+		[
+			SNew(STextBlock)
+			.ColorAndOpacity(this, &FSkeletonTreePhysicsItem::GetTextColor)
+			.Text(FText::FromName(GetRowItemName()))
+			.HighlightText(FilterText)
+			.Font(FEditorStyle::GetFontStyle("PhysicsAssetEditor.Tree.Font"))
+			.ToolTipText(GetNameColumnToolTip())
+		];
+}
+
+
+TSharedRef< SWidget > FSkeletonTreePhysicsItem::GenerateWidgetForDataColumn(const FName& DataColumnName, FIsSelected InIsSelected)
+{
+	if (DataColumnName == ISkeletonTree::DebugVisualizationOptionsName())
+	{
+		// Create a widget that displays an open/closed eye icon when the body is shown/hidden in the viewport debug draw.
+		return SNew(SHorizontalBox)
+		+ SHorizontalBox::Slot()
+		.AutoWidth()
+		.HAlign(HAlign_Left)
+		[
+			SNew(SCheckBox)
+			.ToolTipText(LOCTEXT("SkeletonTreePhysicsItemVisibilityCheckBoxToolTip", "Click to toggle visibility of debug draw for this item"))
+			.OnCheckStateChanged(this, &FSkeletonTreePhysicsItem::OnToggleItemDisplayed)
+			.IsChecked(this, &FSkeletonTreePhysicsItem::IsItemDisplayed)
+			.CheckedHoveredImage(FEditorStyle::GetBrush("Level.VisibleIcon16x"))
+			.CheckedImage(FEditorStyle::GetBrush("Level.VisibleIcon16x"))
+			.CheckedPressedImage(FEditorStyle::GetBrush("Level.VisibleIcon16x"))
+			.UncheckedHoveredImage(FEditorStyle::GetBrush("Level.NotVisibleIcon16x"))
+			.UncheckedImage(FEditorStyle::GetBrush("Level.NotVisibleIcon16x"))
+			.UncheckedPressedImage(FEditorStyle::GetBrush("Level.NotVisibleIcon16x"))
+		];
+	}
+
+	return SNullWidget::NullWidget;
+}
+
+
+FPhysicsAssetRenderSettings* FSkeletonTreePhysicsItem::GetRenderSettings() const
+{
+	return UPhysicsAssetRenderUtilities::GetSettings(PhysicsAssetPathNameHash);
+}
+
+#undef LOCTEXT_NAMESPACE
diff --git a/Engine/Source/Editor/PhysicsAssetEditor/Private/SkeletonTreePhysicsItem.h b/Engine/Source/Editor/PhysicsAssetEditor/Private/SkeletonTreePhysicsItem.h
new file mode 100644
--- /dev/null
+++ b/Engine/Source/Editor/PhysicsAssetEditor/Private/SkeletonTreePhysicsItem.h
@@ -0,0 +1,48 @@
+// Copyright Epic Games, Inc. All Rights Reserved.
+
+#pragma once
+
+#include "CoreMinimal.h"
+#include "Styling/SlateColor.h"
+#include "Input/Reply.h"
+#include "Widgets/SWidget.h"
+#include "Widgets/Views/STableViewBase.h"
+#include "Widgets/Views/STableRow.h"
+#include "SkeletonTreeItem.h"
+
+class FSkeletonTreePhysicsItem : public FSkeletonTreeItem
+{
+public:
+	SKELETON_TREE_ITEM_TYPE(FSkeletonTreePhysicsItem, FSkeletonTreeItem)
+
+	FSkeletonTreePhysicsItem(class UPhysicsAsset* const InPhysicsAsset, const TSharedRef<class ISkeletonTree>& InSkeletonTree);
+	
+	/** ISkeletonTreeItem interface */
+	virtual void GenerateWidgetForNameColumn(TSharedPtr< SHorizontalBox > Box, const TAttribute<FText>& FilterText, FIsSelected InIsSelected) override;
+	virtual TSharedRef< SWidget > GenerateWidgetForDataColumn(const FName& DataColumnName, FIsSelected InIsSelected) override;
+	virtual FName GetRowItemName() const override { return DisplayName; }
+
+	/** UI Callbacks */
+	virtual void OnToggleItemDisplayed(ECheckBoxState InCheckboxState) = 0;
+	virtual ECheckBoxState IsItemDisplayed() const = 0;
+
+private:
+	/** Gets the icon to display for this body */
+	virtual const FSlateBrush* GetBrush() const = 0;
+
+	/** Gets the color to display the item's text */
+	virtual FSlateColor GetTextColor() const = 0;
+
+	/** Gets the tool tip to display on hovering over the item's name */
+	virtual FText GetNameColumnToolTip() const = 0;
+
+protected:
+	/** Gets the physics asset render settings for this physics item */
+	struct FPhysicsAssetRenderSettings* GetRenderSettings() const;
+
+	/** Unique ID of the physics asset for this physics item */
+	uint32 PhysicsAssetPathNameHash;
+
+	/** The name of the item in the tree */
+	FName DisplayName;
+};
diff --git a/Engine/Source/Editor/SkeletonEditor/Private/SSkeletonTree.cpp b/Engine/Source/Editor/SkeletonEditor/Private/SSkeletonTree.cpp
--- a/Engine/Source/Editor/SkeletonEditor/Private/SSkeletonTree.cpp
+++ b/Engine/Source/Editor/SkeletonEditor/Private/SSkeletonTree.cpp
@@ -72,6 +72,7 @@
 const FName	ISkeletonTree::Columns::Name("Name");
 const FName	ISkeletonTree::Columns::Retargeting("Retargeting");
 const FName ISkeletonTree::Columns::BlendProfile("BlendProfile");
+const FName ISkeletonTree::Columns::DebugVisualization("DebugVisualization");
 
 // This is mostly duplicated from SListView, to allow for us to avoid selecting collapsed items
 template <typename ItemType>
@@ -164,6 +165,7 @@ void SSkeletonTree::Construct(const FArguments& InArgs, const TSharedRef<FEditab
 	Mode = InSkeletonTreeArgs.Mode;
 	bAllowMeshOperations = InSkeletonTreeArgs.bAllowMeshOperations;
 	bAllowSkeletonOperations = InSkeletonTreeArgs.bAllowSkeletonOperations;
+	bShowDebugVisualizationOptions = InSkeletonTreeArgs.bShowDebugVisualizationOptions;
 	Extenders = InSkeletonTreeArgs.Extenders;
 	OnGetFilterText = InSkeletonTreeArgs.OnGetFilterText;
 	Builder = InSkeletonTreeArgs.Builder;
@@ -410,6 +412,12 @@ void SSkeletonTree::BindCommands()
 		FCanExecuteAction(),
 		FIsActionChecked::CreateSP(this, &SSkeletonTree::IsShowingAdvancedOptions));
 
+	CommandList.MapAction(
+		MenuActions.ShowDebugVisualization,
+		FExecuteAction::CreateSP(this, &SSkeletonTree::OnChangeShowingDebugVisualizationOptions),
+		FCanExecuteAction(),
+		FIsActionChecked::CreateSP(this, &SSkeletonTree::IsShowingDebugVisualizationOptions));
+
 	// Socket manipulation commands
 	CommandList.MapAction(
 		MenuActions.AddSocket,
@@ -560,6 +568,7 @@ void SSkeletonTree::CreateTreeColumns()
 	TArray<FName> HiddenColumnsList;
 	HiddenColumnsList.Add(ISkeletonTree::Columns::Retargeting);
 	HiddenColumnsList.Add(ISkeletonTree::Columns::BlendProfile);
+	HiddenColumnsList.Add(ISkeletonTree::Columns::DebugVisualization);
 
 	TSharedRef<SHeaderRow> TreeHeaderRow = 
 	SNew(SHeaderRow)
@@ -575,6 +584,10 @@ void SSkeletonTree::CreateTreeColumns()
 	.DefaultLabel(LOCTEXT("SkeletonBoneTranslationRetargetingLabel", "Translation Retargeting"))
 	.FillWidth(0.25f)
 
+	+ SHeaderRow::Column(ISkeletonTree::Columns::DebugVisualization)
+	.DefaultLabel(LOCTEXT("SkeletonBoneDebugVisualizationLabel", "Debug"))
+	.FillWidth(0.25f)
+
 	+ SHeaderRow::Column(ISkeletonTree::Columns::BlendProfile)
 	.DefaultLabel(LOCTEXT("BlendProfile", "Blend Profile"))
 	.FillWidth(0.25f)
@@ -1780,6 +1793,11 @@ TSharedRef< SWidget > SSkeletonTree::CreateFilterMenu()
 		}
 		MenuBuilder.AddMenuEntry(Actions.FilteringFlattensHierarchy);
 		MenuBuilder.AddMenuEntry(Actions.HideParentsWhenFiltering);
+		
+		if (bShowDebugVisualizationOptions)
+		{
+			MenuBuilder.AddMenuEntry(Actions.ShowDebugVisualization);
+		}
 	}
 	MenuBuilder.EndSection();
 
@@ -2044,6 +2062,16 @@ bool SSkeletonTree::IsShowingAdvancedOptions() const
 	return SkeletonTreeView->GetHeaderRow()->IsColumnVisible(ISkeletonTree::Columns::Retargeting);
 }
 
+void SSkeletonTree::OnChangeShowingDebugVisualizationOptions()
+{
+	SkeletonTreeView->GetHeaderRow()->SetShowGeneratedColumn(ISkeletonTree::Columns::DebugVisualization, !IsShowingDebugVisualizationOptions());
+}
+
+bool SSkeletonTree::IsShowingDebugVisualizationOptions() const
+{
+	return SkeletonTreeView->GetHeaderRow()->IsColumnVisible(ISkeletonTree::Columns::DebugVisualization);
+}
+
 UBlendProfile* SSkeletonTree::GetSelectedBlendProfile()
 {
 	return BlendProfilePicker->GetSelectedBlendProfile();
diff --git a/Engine/Source/Editor/SkeletonEditor/Private/SSkeletonTree.h b/Engine/Source/Editor/SkeletonEditor/Private/SSkeletonTree.h
--- a/Engine/Source/Editor/SkeletonEditor/Private/SSkeletonTree.h
+++ b/Engine/Source/Editor/SkeletonEditor/Private/SSkeletonTree.h
@@ -263,6 +263,12 @@ private:
 	/**  Handler for when we change the "Show Retargeting Options" check box */
 	void OnChangeShowingAdvancedOptions();
 
+	/** Handler for "Show Debug Visualization Options" check box IsChecked functionality */
+	bool IsShowingDebugVisualizationOptions() const;
+
+	/**  Handler for when we change the "Show Debug Visualization Options" check box */
+	void OnChangeShowingDebugVisualizationOptions();
+
 	/** This replicates the socket filter to the previewcomponent so that the viewport can use the same settings */
 	void SetPreviewComponentSocketFilter() const;
 
@@ -415,6 +421,9 @@ private:
 	/** Whether to allow operations that modify the mesh */
 	bool bAllowSkeletonOperations;
 
+	/** Whether to show the filter option to allow filtering of debug draw elements in the viewport. */
+	bool bShowDebugVisualizationOptions;
+
 	/** Extenders for menus */
 	TSharedPtr<FExtender> Extenders;
 
diff --git a/Engine/Source/Editor/SkeletonEditor/Private/SkeletonTreeCommands.cpp b/Engine/Source/Editor/SkeletonEditor/Private/SkeletonTreeCommands.cpp
--- a/Engine/Source/Editor/SkeletonEditor/Private/SkeletonTreeCommands.cpp
+++ b/Engine/Source/Editor/SkeletonEditor/Private/SkeletonTreeCommands.cpp
@@ -12,6 +12,7 @@ void FSkeletonTreeCommands::RegisterCommands()
 	UI_COMMAND( ShowWeightedBones, "Show Weighted Bones", "Show bones that have vertices weighted to them", EUserInterfaceActionType::RadioButton, FInputChord() );
 	UI_COMMAND( HideBones, "Hide Bones", "Hides all bones (sockets and attached assets will still be listed)", EUserInterfaceActionType::RadioButton, FInputChord() );
 	UI_COMMAND( ShowRetargeting, "Show Retargeting Options", "Show retargeting options for each bone.", EUserInterfaceActionType::ToggleButton, FInputChord());
+	UI_COMMAND( ShowDebugVisualization, "Show Debug Visualization Options", "Show debug visualization options for objects in the viewport", EUserInterfaceActionType::ToggleButton, FInputChord());
 	UI_COMMAND( FilteringFlattensHierarchy, "Filtering Flattens Hierarchy", "Whether to keep the hierarchy or flatten it when searching for tree items", EUserInterfaceActionType::ToggleButton, FInputChord());
 	UI_COMMAND( HideParentsWhenFiltering, "Hide Parents When Filtering", "Whether to show parent items grayed out, or hide them entirely when filtering", EUserInterfaceActionType::ToggleButton, FInputChord());
 	
diff --git a/Engine/Source/Editor/SkeletonEditor/Private/SkeletonTreeCommands.h b/Engine/Source/Editor/SkeletonEditor/Private/SkeletonTreeCommands.h
--- a/Engine/Source/Editor/SkeletonEditor/Private/SkeletonTreeCommands.h
+++ b/Engine/Source/Editor/SkeletonEditor/Private/SkeletonTreeCommands.h
@@ -45,6 +45,9 @@ public:
 	/** Show retargeting options */
 	TSharedPtr< FUICommandInfo > ShowRetargeting;
 
+	/** Show debug visualization options */
+	TSharedPtr< FUICommandInfo > ShowDebugVisualization;
+
 	/** Flatten hierarchy on filter */
 	TSharedPtr< FUICommandInfo > FilteringFlattensHierarchy;
 
diff --git a/Engine/Source/Editor/SkeletonEditor/Public/ISkeletonTree.h b/Engine/Source/Editor/SkeletonEditor/Public/ISkeletonTree.h
--- a/Engine/Source/Editor/SkeletonEditor/Public/ISkeletonTree.h
+++ b/Engine/Source/Editor/SkeletonEditor/Public/ISkeletonTree.h
@@ -48,6 +48,7 @@ struct FSkeletonTreeArgs
 		, ContextName(TEXT("SkeletonTree"))
 		, bShowBlendProfiles(true)
 		, bShowFilterMenu(true)
+		, bShowDebugVisualizationOptions(false)
 		, bAllowMeshOperations(true)
 		, bAllowSkeletonOperations(true)
 		, bHideBonesByDefault(false)
@@ -80,6 +81,9 @@ struct FSkeletonTreeArgs
 	/** Whether to show the filter menu to allow filtering of active bones, sockets etc. */
 	bool bShowFilterMenu;
 
+	/** Whether to show the filter option to allow filtering of debug draw elements in the viewport. */
+	bool bShowDebugVisualizationOptions;
+
 	/** Whether to allow operations that modify the mesh */
 	bool bAllowMeshOperations;
 
@@ -99,8 +103,11 @@ public:
 		static const FName Name;
 		static const FName Retargeting;
 		static const FName BlendProfile;
+		static const FName DebugVisualization;
 	};
 
+	static const FName DebugVisualizationOptionsName() { return Columns::DebugVisualization; };
+
 	/** Manually refresh the tree */
 	virtual void Refresh() = 0;
 
diff --git a/Engine/Source/Runtime/AnimGraphRuntime/Public/BoneControllers/AnimNode_RigidBody.h b/Engine/Source/Runtime/AnimGraphRuntime/Public/BoneControllers/AnimNode_RigidBody.h
--- a/Engine/Source/Runtime/AnimGraphRuntime/Public/BoneControllers/AnimNode_RigidBody.h
+++ b/Engine/Source/Runtime/AnimGraphRuntime/Public/BoneControllers/AnimNode_RigidBody.h
@@ -148,6 +148,8 @@ struct ANIMGRAPHRUNTIME_API FAnimNode_RigidBody : public FAnimNode_SkeletalContr
 	// TEMP: Exposed for use in PhAt as a quick way to get drag handles working with Chaos
 	virtual ImmediatePhysics::FSimulation* GetSimulation() { return PhysicsSimulation; }
 
+	UPhysicsAsset* GetPhysicsAsset() const { return UsePhysicsAsset; }
+
 public:
 	/** Physics asset to use. If empty use the skeletal mesh's default physics asset */
 	UPROPERTY(EditAnywhere, Category = Settings)
