[UE][FIX] Crash in PhysicsAssetRenderUtils after merge

#rmx

diff --git a/Engine/Source/Editor/Persona/Private/PhysicsAssetRenderUtils.cpp b/Engine/Source/Editor/Persona/Private/PhysicsAssetRenderUtils.cpp
--- a/Engine/Source/Editor/Persona/Private/PhysicsAssetRenderUtils.cpp
+++ b/Engine/Source/Editor/Persona/Private/PhysicsAssetRenderUtils.cpp
@@ -270,7 +270,7 @@ namespace PhysicsAssetRender
 			HHitProxy* InHitProxy,
 			float             InScale,
 			const FSceneView* InView,
-			FKShapeElem* InShapeElem,
+			const FKShapeElem* InShapeElem,
 			const FName& InBoneName
 		)
 			: TM(InTM), HitProxy(InHitProxy), Scale(InScale), ShapeElem(InShapeElem)
@@ -291,14 +291,14 @@ namespace PhysicsAssetRender
 		FTransform            TM;
 		HHitProxy* HitProxy = nullptr;
 		float                 Scale;
-		FKShapeElem* ShapeElem = nullptr;
+		const FKShapeElem* ShapeElem = nullptr;
 		float                 SortingMetric;
 #ifdef UE_BUILD_DEBUG
 		FName                 BoneName;
 #endif
 	};
 	
-	template< typename TElementContainer > void CollectDrawElements(TArray<DrawElement>& DrawElements, const EAggCollisionShape::Type PrimitiveType, const TElementContainer ElementContainer, const FTransform& BoneTM, const float Scale, const uint32 BodyIndex, const FPhysicsAssetRenderSettings* const RenderSettings, UPhysicsAsset* const PhysicsAsset, FPrimitiveDrawInterface* PDI, GetPrimitiveRef< FColor > GetPrimitiveColor, GetPrimitiveRef< UMaterialInterface* > GetPrimitiveMaterial, GetPrimitiveTransformRef GetPrimitiveTransform, CreateBodyHitProxyFn CreateHitProxy)
+	template< typename TElementContainer > void CollectDrawElements(TArray<DrawElement>& DrawElements, const EAggCollisionShape::Type PrimitiveType, const TElementContainer& ElementContainer, const FTransform& BoneTM, const float Scale, const uint32 BodyIndex, const FPhysicsAssetRenderSettings* const RenderSettings, UPhysicsAsset* const PhysicsAsset, FPrimitiveDrawInterface* PDI, GetPrimitiveRef< FColor > GetPrimitiveColor, GetPrimitiveRef< UMaterialInterface* > GetPrimitiveMaterial, GetPrimitiveTransformRef GetPrimitiveTransform, CreateBodyHitProxyFn CreateHitProxy)
 	{
 		if (PhysicsAsset && RenderSettings && PDI && PhysicsAsset->SkeletalBodySetups.IsValidIndex(BodyIndex))
 		{
@@ -309,7 +309,7 @@ namespace PhysicsAssetRender
 				DrawElement DE(
 					GetPrimitiveTransform(PhysicsAsset, BoneTM, BodyIndex, PrimitiveType, ElementIndex, Scale),
 					CreateHitProxy(BodyIndex, PrimitiveType, ElementIndex),
-					Scale, PDI->View, (FKShapeElem*)(&ElementContainer[ElementIndex]), BoneName
+					Scale, PDI->View, &ElementContainer[ElementIndex], BoneName
 				);
 
 				if (RenderSettings->CollisionViewMode == EPhysicsAssetEditorCollisionViewMode::Solid || RenderSettings->CollisionViewMode == EPhysicsAssetEditorCollisionViewMode::SolidWireframe)
