Added a conenience RMB menu item to copy curve names from the timeline editor to the clipboard.

diff --git a/Engine/Source/Editor/Persona/Private/AnimSequenceTimelineCommands.cpp b/Engine/Source/Editor/Persona/Private/AnimSequenceTimelineCommands.cpp
--- a/Engine/Source/Editor/Persona/Private/AnimSequenceTimelineCommands.cpp
+++ b/Engine/Source/Editor/Persona/Private/AnimSequenceTimelineCommands.cpp
@@ -19,6 +19,7 @@ void FAnimSequenceTimelineCommands::RegisterCommands()
 	UI_COMMAND(ConvertMetaDataToCurve, "Convert To Curve", "Convert this metadata curve to a full curve", EUserInterfaceActionType::Button, FInputChord());
 	UI_COMMAND(RemoveCurve, "Remove Curve", "Remove this curve", EUserInterfaceActionType::Button, FInputChord());
 	UI_COMMAND(RemoveAllCurves, "Remove All Curves", "Remove all the curves in this animation", EUserInterfaceActionType::Button, FInputChord());
+	UI_COMMAND(CopySelectedCurveNames, "Copy Selected Curve Names", "Copy the name of the selected curves to the clipboard", EUserInterfaceActionType::Button, FInputChord());
 	UI_COMMAND(DisplaySeconds, "Seconds", "Display the time in seconds", EUserInterfaceActionType::RadioButton, FInputChord());
 	UI_COMMAND(DisplayFrames, "Frames", "Display the time in frames", EUserInterfaceActionType::RadioButton, FInputChord());
 	UI_COMMAND(DisplayPercentage, "Percentage", "Display the percentage along with the time with the scrubber", EUserInterfaceActionType::ToggleButton, FInputChord());
diff --git a/Engine/Source/Editor/Persona/Private/AnimSequenceTimelineCommands.h b/Engine/Source/Editor/Persona/Private/AnimSequenceTimelineCommands.h
--- a/Engine/Source/Editor/Persona/Private/AnimSequenceTimelineCommands.h
+++ b/Engine/Source/Editor/Persona/Private/AnimSequenceTimelineCommands.h
@@ -26,7 +26,6 @@ public:
 	TSharedPtr<FUICommandInfo> EditSelectedCurves;
 
 	TSharedPtr<FUICommandInfo> RemoveSelectedCurves;
-
 	TSharedPtr<FUICommandInfo> AddNotifyTrack;
 
 	TSharedPtr<FUICommandInfo> InsertNotifyTrack;
@@ -49,6 +48,8 @@ public:
 
 	TSharedPtr<FUICommandInfo> RemoveAllCurves;
 
+	TSharedPtr<FUICommandInfo> CopySelectedCurveNames;
+	
 	TSharedPtr<FUICommandInfo> DisplaySeconds;
 
 	TSharedPtr<FUICommandInfo> DisplayFrames;
@@ -64,7 +65,6 @@ public:
 	TSharedPtr<FUICommandInfo> SnapToMontageSections;
 
 	TSharedPtr<FUICommandInfo> SnapToCompositeSegments;
-
 public:
 	virtual void RegisterCommands() override;
 };
diff --git a/Engine/Source/Editor/Persona/Private/AnimTimeline/AnimModel_AnimSequenceBase.cpp b/Engine/Source/Editor/Persona/Private/AnimTimeline/AnimModel_AnimSequenceBase.cpp
--- a/Engine/Source/Editor/Persona/Private/AnimTimeline/AnimModel_AnimSequenceBase.cpp
+++ b/Engine/Source/Editor/Persona/Private/AnimTimeline/AnimModel_AnimSequenceBase.cpp
@@ -22,6 +22,7 @@
 #include "IPersonaPreviewScene.h"
 #include "Animation/DebugSkelMeshComponent.h"
 #include "AnimPreviewInstance.h"
+#include "HAL/PlatformApplicationMisc.h"
 #include "ScopedTransaction.h"
 
 
@@ -91,6 +92,10 @@ void FAnimModel_AnimSequenceBase::Initialize()
 		Commands.RemoveSelectedCurves,
 		FExecuteAction::CreateSP(this, &FAnimModel_AnimSequenceBase::RemoveSelectedCurves));
 
+	CommandList->MapAction(
+		Commands.CopySelectedCurveNames,
+		FExecuteAction::CreateSP(this, &FAnimModel_AnimSequenceBase::CopySelectedCurveNamesToClipboard));
+	
 	CommandList->MapAction(
 		Commands.DisplayFrames,
 		FExecuteAction::CreateSP(this, &FAnimModel_AnimSequenceBase::SetDisplayFormat, EFrameNumberDisplayFormats::Frames),
@@ -144,6 +149,7 @@ void FAnimModel_AnimSequenceBase::Initialize()
 		FIsActionButtonVisible::CreateSP(this, &FAnimModel_AnimSequenceBase::IsSnapAvailable, FAnimModel::FSnapType::MontageSection.Type));
 }
 
+
 void FAnimModel_AnimSequenceBase::RefreshTracks()
 {
 	ClearTrackSelection();
@@ -485,6 +491,31 @@ void FAnimModel_AnimSequenceBase::RemoveSelectedCurves()
 	}
 }
 
+
+void FAnimModel_AnimSequenceBase::CopySelectedCurveNamesToClipboard()
+{
+	TArray<FString> TrackNames; 
+	for(TSharedRef<FAnimTimelineTrack>& SelectedTrack : SelectedTracks)
+	{
+		if(SelectedTrack->IsA<FAnimTimelineTrack_FloatCurve>())
+		{
+			TSharedRef<FAnimTimelineTrack_FloatCurve> FloatCurveTrack = StaticCastSharedRef<FAnimTimelineTrack_FloatCurve>(SelectedTrack);
+			TrackNames.Add(FloatCurveTrack->GetName().DisplayName.ToString());
+		}
+		else if(SelectedTrack->IsA<FAnimTimelineTrack_TransformCurve>())
+		{
+			TSharedRef<FAnimTimelineTrack_TransformCurve> TransformCurveTrack = StaticCastSharedRef<FAnimTimelineTrack_TransformCurve>(SelectedTrack);
+			TrackNames.Add(TransformCurveTrack->GetName().DisplayName.ToString());
+		}
+
+	}
+	if (!TrackNames.IsEmpty())
+	{
+		FPlatformApplicationMisc::ClipboardCopy(*FString::Join(TrackNames, TEXT("\n")));
+	}
+}
+
+
 void FAnimModel_AnimSequenceBase::SetDisplayFormat(EFrameNumberDisplayFormats InFormat)
 {
 	GetMutableDefault<UPersonaOptions>()->TimelineDisplayFormat = InFormat;
diff --git a/Engine/Source/Editor/Persona/Private/AnimTimeline/AnimModel_AnimSequenceBase.h b/Engine/Source/Editor/Persona/Private/AnimTimeline/AnimModel_AnimSequenceBase.h
--- a/Engine/Source/Editor/Persona/Private/AnimTimeline/AnimModel_AnimSequenceBase.h
+++ b/Engine/Source/Editor/Persona/Private/AnimTimeline/AnimModel_AnimSequenceBase.h
@@ -74,6 +74,7 @@ private:
 	void EditSelectedCurves();
 	bool CanEditSelectedCurves() const;
 	void RemoveSelectedCurves();
+	void CopySelectedCurveNamesToClipboard();
 	void SetDisplayFormat(EFrameNumberDisplayFormats InFormat);
 	bool IsDisplayFormatChecked(EFrameNumberDisplayFormats InFormat) const;
 	void ToggleDisplayPercentage();
diff --git a/Engine/Source/Editor/Persona/Private/AnimTimeline/AnimTimelineTrack_Curve.cpp b/Engine/Source/Editor/Persona/Private/AnimTimeline/AnimTimelineTrack_Curve.cpp
--- a/Engine/Source/Editor/Persona/Private/AnimTimeline/AnimTimelineTrack_Curve.cpp
+++ b/Engine/Source/Editor/Persona/Private/AnimTimeline/AnimTimelineTrack_Curve.cpp
@@ -324,6 +324,7 @@ void FAnimTimelineTrack_Curve::AddToContextMenu(FMenuBuilder& InMenuBuilder, TSe
 		{
 			InMenuBuilder.AddMenuEntry(FAnimSequenceTimelineCommands::Get().EditSelectedCurves);
 			InMenuBuilder.AddMenuEntry(FAnimSequenceTimelineCommands::Get().RemoveSelectedCurves);
+			InMenuBuilder.AddMenuEntry(FAnimSequenceTimelineCommands::Get().CopySelectedCurveNames);
 		}
 		InMenuBuilder.EndSection();
 
