Adding compatible skeletons to Retarget Sources

diff --git a/Engine/Source/Editor/Persona/Private/SRetargetSourceWindow.cpp b/Engine/Source/Editor/Persona/Private/SRetargetSourceWindow.cpp
--- a/Engine/Source/Editor/Persona/Private/SRetargetSourceWindow.cpp
+++ b/Engine/Source/Editor/Persona/Private/SRetargetSourceWindow.cpp
@@ -169,9 +169,7 @@ void SRetargetSourceWindow::Construct(const FArguments& InArgs, const TSharedRef
 	EditableSkeletonPtr = InEditableSkeleton;
 
 	InOnPostUndo.Add(FSimpleDelegate::CreateSP( this, &SRetargetSourceWindow::PostUndo ) );
-
-	FText SkeletonName = FText::FromString(InEditableSkeleton->GetSkeleton().GetName());
-
+	
 	ChildSlot
 	[
 		SNew( SVerticalBox )
@@ -601,5 +599,163 @@ FReply SRetargetSourceWindow::OnUpdateAllRetargetSourceButtonClicked()
 	OnRefreshRetargetSource(true);
 	return FReply::Handled();
 }
+
+void SCompatibleSkeletons::Construct(
+	const FArguments& InArgs,
+	const TSharedRef<IEditableSkeleton>& InEditableSkeleton,
+	FSimpleMulticastDelegate& InOnPostUndo)
+{
+	EditableSkeletonPtr = InEditableSkeleton;
+	UpdateCompatibleSkeletonAssets(InEditableSkeleton->GetSkeleton());
+
+	ChildSlot
+	[
+		SNew( SVerticalBox )
+		
+		+ SVerticalBox::Slot()
+		.AutoHeight()
+		[
+			SNew( SHorizontalBox )
+
+			+SHorizontalBox::Slot()
+			.HAlign(HAlign_Right)
+			.AutoWidth()
+			.Padding(2, 0)
+			[
+				SNew(SButton)
+				.OnClicked(FOnClicked::CreateSP(this, &SCompatibleSkeletons::OnAddSkeletonClicked))
+				.HAlign(HAlign_Center)
+				.VAlign(VAlign_Center)
+				.Text(LOCTEXT("AddCompatibleSkeletonButton_Label", "Add Skeleton"))
+				.ToolTipText(LOCTEXT("AddCompatibleSkeletonButton_ToolTip", "When Skeleton assets share an identical hierarchy, bone names and orientations they can be added to the Compatible Skeletons list. Animation assets authored on Compatible Skeletons can then be used in this Skeleton's Animation Blueprints."))
+			]
+
+			+SHorizontalBox::Slot()
+			.HAlign(HAlign_Right)
+			.AutoWidth()
+			.Padding(2, 0)
+			[
+				SNew(SButton)
+				.OnClicked(FOnClicked::CreateSP(this, &SCompatibleSkeletons::OnRemoveSkeletonClicked))
+				.HAlign(HAlign_Center)
+				.VAlign(VAlign_Center)
+				.Text(LOCTEXT("RemoveCompatibleSkeletonButton_Label", "Remove Selected"))
+				.ToolTipText(LOCTEXT("RemoveCompatibleSkeletonButton_ToolTip", "Remove the selected skeleton assets from the list of Compatible Skeletons."))
+			]
+		]
+
+		+ SVerticalBox::Slot()
+		.FillHeight( 1.0f )
+		.Padding(2.0f)
+		[
+			SAssignNew(CompatibleSkeletonListView, SListView<TSharedRef<FSoftObjectPath>>)
+			.ListItemsSource(&CompatibleSkeletonAssets)
+			.OnGenerateRow(this, &SCompatibleSkeletons::GenerateRowForItem)
+			.ItemHeight( 22.0f )
+		]
+	];
+}
+
+void SCompatibleSkeletons::UpdateCompatibleSkeletonAssets(const USkeleton& Skeleton)
+{
+	CompatibleSkeletonAssets.Empty();
+
+	const TArray<TSoftObjectPtr<USkeleton>>& CompatibleSkeletons = Skeleton.GetCompatibleSkeletons();
+	for (const TSoftObjectPtr<USkeleton>& SoftCompatibleSkeleton : CompatibleSkeletons)
+	{
+		TSharedRef<FSoftObjectPath> AssetPath = MakeShareable(new FSoftObjectPath(SoftCompatibleSkeleton.ToSoftObjectPath()));
+		CompatibleSkeletonAssets.Add(AssetPath);
+	}
+}
+
+FReply SCompatibleSkeletons::OnAddSkeletonClicked()
+{
+	// show list of skeletalmeshes that they can choose from
+	FContentBrowserModule& ContentBrowserModule = FModuleManager::Get().LoadModuleChecked<FContentBrowserModule>(TEXT("ContentBrowser"));
+
+	FAssetPickerConfig AssetPickerConfig;
+	AssetPickerConfig.Filter.ClassNames.Add(USkeleton::StaticClass()->GetFName());
+	AssetPickerConfig.OnAssetSelected = FOnAssetSelected::CreateSP(this, &SCompatibleSkeletons::OnAssetSelectedFromSkeletonPicker);
+	AssetPickerConfig.bAllowNullSelection = false;
+	AssetPickerConfig.InitialAssetViewType = EAssetViewType::Tile;
+
+	const TSharedRef<SWidget> Widget = SNew(SBox)
+		.WidthOverride(384)
+		.HeightOverride(768)
+		[
+			SNew(SBorder)
+			.BorderBackgroundColor(FLinearColor(0.25f, 0.25f, 0.25f, 1.f))
+			.Padding( 2 )
+			[
+				SNew(SBorder)
+				.BorderImage( FEditorStyle::GetBrush("ToolPanel.GroupBorder") )
+				.Padding( 8 )
+				[
+					ContentBrowserModule.Get().CreateAssetPicker(AssetPickerConfig)
+				]
+			]
+		];
+
+	FSlateApplication::Get().PushMenu(
+		AsShared(),
+		FWidgetPath(),
+		Widget,
+		FSlateApplication::Get().GetCursorPos(),
+		FPopupTransitionEffect( FPopupTransitionEffect::TopMenu )
+		);
+	
+	return FReply::Handled();
+}
+
+FReply SCompatibleSkeletons::OnRemoveSkeletonClicked()
+{
+	TArray<TSharedRef<FSoftObjectPath>> SelectedAssets = CompatibleSkeletonListView->GetSelectedItems();
+	for (TSharedRef<FSoftObjectPath>& SkeletonAssetPath : SelectedAssets)
+	{
+		USkeleton* SkeletonToRemove = Cast<USkeleton>(SkeletonAssetPath.Get().TryLoad());
+		EditableSkeletonPtr.Pin()->RemoveCompatibleSkeleton(SkeletonToRemove);
+	}
+
+	UpdateCompatibleSkeletonAssets(EditableSkeletonPtr.Pin()->GetSkeleton());
+	CompatibleSkeletonListView->RequestListRefresh();
+
+	FAssetNotifications::SkeletonNeedsToBeSaved(&EditableSkeletonPtr.Pin()->GetSkeleton());
+	
+	return FReply::Handled();
+}
+
+TSharedRef<ITableRow> SCompatibleSkeletons::GenerateRowForItem(TSharedRef<FSoftObjectPath> Item, const TSharedRef<STableViewBase>& OwnerTable) const
+{
+	return SNew(STableRow<TSharedPtr<USkeleton>>, OwnerTable)
+		.Content()
+		[
+			SNew(STextBlock).Text(FText::FromString(Item.Get().GetAssetName()))
+		];
+}
+
+void SCompatibleSkeletons::OnAssetSelectedFromSkeletonPicker(const FAssetData& AssetData)
+{
+	// make sure we haven't already added this asset as a compatible skeleton
+	const USkeleton& Skeleton = EditableSkeletonPtr.Pin()->GetSkeleton();
+	const FString AssetFullPath = AssetData.ToSoftObjectPath().ToString();
+	const TArray<TSoftObjectPtr<USkeleton>>& CompatibleSkeletons = Skeleton.GetCompatibleSkeletons();
+	for (const TSoftObjectPtr<USkeleton>& CompatibleSkeleton : CompatibleSkeletons)
+	{
+		if (CompatibleSkeleton.ToSoftObjectPath() == AssetData.ToSoftObjectPath())
+		{
+			FSlateApplication::Get().DismissAllMenus();
+			return;
+		}
+	}
+
+	const USkeleton* CompatibleSkeleton = CastChecked<USkeleton>(AssetData.GetAsset());
+	EditableSkeletonPtr.Pin()->AddCompatibleSkeleton(CompatibleSkeleton);
+	FAssetNotifications::SkeletonNeedsToBeSaved(&EditableSkeletonPtr.Pin()->GetSkeleton());
+	FSlateApplication::Get().DismissAllMenus();
+
+	UpdateCompatibleSkeletonAssets(EditableSkeletonPtr.Pin()->GetSkeleton());
+	CompatibleSkeletonListView->RequestListRefresh();
+}
+
 #undef LOCTEXT_NAMESPACE
 
diff --git a/Engine/Source/Editor/Persona/Private/SRetargetSourceWindow.h b/Engine/Source/Editor/Persona/Private/SRetargetSourceWindow.h
--- a/Engine/Source/Editor/Persona/Private/SRetargetSourceWindow.h
+++ b/Engine/Source/Editor/Persona/Private/SRetargetSourceWindow.h
@@ -18,7 +18,7 @@ template <typename ItemType> class SListView;
 //////////////////////////////////////////////////////////////////////////
 // FDisplayedRetargetSourceInfo
 
-class FDisplayedRetargetSourceInfo
+class  FDisplayedRetargetSourceInfo
 {
 public:
 	FName Name;
@@ -226,3 +226,29 @@ private:
 	FReply OnAddRetargetSourceButtonClicked();
 	FReply OnUpdateAllRetargetSourceButtonClicked();
 };
+
+class SCompatibleSkeletons : public SCompoundWidget
+{
+public:
+	SLATE_BEGIN_ARGS( SCompatibleSkeletons ){}
+	SLATE_END_ARGS()
+
+	void Construct(const FArguments& InArgs, const TSharedRef<class IEditableSkeleton>& InEditableSkeleton, FSimpleMulticastDelegate& InOnPostUndo);
+
+	FReply OnAddSkeletonClicked();
+
+	FReply OnRemoveSkeletonClicked();
+
+	TSharedRef<ITableRow> GenerateRowForItem(TSharedRef<FSoftObjectPath> Item, const TSharedRef<STableViewBase>& OwnerTable) const;
+
+	void OnAssetSelectedFromSkeletonPicker(const FAssetData& AssetData);
+
+private:
+	void UpdateCompatibleSkeletonAssets(const USkeleton& Skeleton);
+	
+	TArray<TSharedRef<FSoftObjectPath>> CompatibleSkeletonAssets;
+
+	TSharedPtr<SListView<TSharedRef<FSoftObjectPath>>> CompatibleSkeletonListView;
+	
+	TWeakPtr<IEditableSkeleton> EditableSkeletonPtr;
+};
diff --git a/Engine/Source/Editor/Persona/Private/SRetargetSources.cpp b/Engine/Source/Editor/Persona/Private/SRetargetSources.cpp
--- a/Engine/Source/Editor/Persona/Private/SRetargetSources.cpp
+++ b/Engine/Source/Editor/Persona/Private/SRetargetSources.cpp
@@ -102,6 +102,30 @@ void SRetargetSources::Construct(
 			// construct retarget source window
 			SNew(SRetargetSourceWindow, InEditableSkeleton, InOnPostUndo)
 		]
+
+		+SVerticalBox::Slot()
+		.Padding(5, 5)
+		.AutoHeight()
+		[
+			SNew(SSeparator)
+			.Orientation(Orient_Horizontal)
+		]
+
+		+ SVerticalBox::Slot()
+		.Padding(5, 5)
+		.AutoHeight()
+		[
+			SNew(STextBlock)
+			.TextStyle(FEditorStyle::Get(), "Persona.RetargetManager.ImportantText")
+			.Text(LOCTEXT("CompatibleSkeletons_Title", "Manage Compatible Skeletons"))
+		]
+
+		+ SVerticalBox::Slot()
+		.Padding(5, 5)
+		.FillHeight(0.5)
+		[
+			SNew(SCompatibleSkeletons, InEditableSkeleton, InOnPostUndo)
+		]
 	];
 }
 
diff --git a/Engine/Source/Editor/SkeletonEditor/Private/EditableSkeleton.cpp b/Engine/Source/Editor/SkeletonEditor/Private/EditableSkeleton.cpp
--- a/Engine/Source/Editor/SkeletonEditor/Private/EditableSkeleton.cpp
+++ b/Engine/Source/Editor/SkeletonEditor/Private/EditableSkeleton.cpp
@@ -1493,6 +1493,22 @@ void FEditableSkeleton::RefreshRetargetSources(const TArray<FName>& InRetargetSo
 	}
 }
 
+void FEditableSkeleton::AddCompatibleSkeleton(const USkeleton* InCompatibleSkeleton)
+{
+	const FScopedTransaction Transaction(LOCTEXT("AddedCompatibleSkeletons", "Add Compatible Skeleton"));
+	Skeleton->Modify();
+
+	Skeleton->AddCompatibleSkeleton(InCompatibleSkeleton);
+}
+
+void FEditableSkeleton::RemoveCompatibleSkeleton(const USkeleton* InCompatibleSkeleton)
+{
+	const FScopedTransaction Transaction(LOCTEXT("RemoveCompatibleSkeletons", "Remove Compatible Skeleton"));
+	Skeleton->Modify();
+
+	Skeleton->RemoveCompatibleSkeleton(InCompatibleSkeleton);
+}
+
 void FEditableSkeleton::RefreshRigConfig()
 {
 	Skeleton->RefreshRigConfig();
diff --git a/Engine/Source/Editor/SkeletonEditor/Private/EditableSkeleton.h b/Engine/Source/Editor/SkeletonEditor/Private/EditableSkeleton.h
--- a/Engine/Source/Editor/SkeletonEditor/Private/EditableSkeleton.h
+++ b/Engine/Source/Editor/SkeletonEditor/Private/EditableSkeleton.h
@@ -62,6 +62,8 @@ public:
 	virtual void AddRetargetSource(const FName& InName, USkeletalMesh* InReferenceMesh) override;
 	virtual void DeleteRetargetSources(const TArray<FName>& InRetargetSourceNames) override;
 	virtual void RefreshRetargetSources(const TArray<FName>& InRetargetSourceNames) override;
+	virtual void AddCompatibleSkeleton(const USkeleton* InCompatibleSkeleton) override;
+	virtual void RemoveCompatibleSkeleton(const USkeleton* InCompatibleSkeleton) override;
 	virtual void RefreshRigConfig() override;
 	virtual void SetRigConfig(URig* InRig) override;
 	virtual void SetRigBoneMapping(const FName& InNodeName, const FName& InBoneName) override;
diff --git a/Engine/Source/Editor/SkeletonEditor/Public/IEditableSkeleton.h b/Engine/Source/Editor/SkeletonEditor/Public/IEditableSkeleton.h
--- a/Engine/Source/Editor/SkeletonEditor/Public/IEditableSkeleton.h
+++ b/Engine/Source/Editor/SkeletonEditor/Public/IEditableSkeleton.h
@@ -150,6 +150,12 @@ public:
 	/** Refresh retarget sources */
 	virtual void RefreshRetargetSources(const TArray<FName>& InRetargetSourceNames) = 0;
 
+	/** Add a compatible skeleton */
+	virtual void AddCompatibleSkeleton(const USkeleton* InCompatibleSkeleton) = 0;
+
+	/** Remove a compatible skeleton */
+	virtual void RemoveCompatibleSkeleton(const USkeleton* InCompatibleSkeleton) = 0;
+
 	/** Refreshes the rig config, validating the mappings */
 	virtual void RefreshRigConfig() = 0;
 
