#UE Do not load materials in FPhysicsAssetRenderSettings constructor, which is called during boot to look for uninitialized properties. Load instead in UPhysicsAssetRenderUtilities's initialize function, and do not make a MID as it is not needed.

diff --git a/Engine/Source/Editor/Persona/Private/PhysicsAssetRenderUtils.cpp b/Engine/Source/Editor/Persona/Private/PhysicsAssetRenderUtils.cpp
--- a/Engine/Source/Editor/Persona/Private/PhysicsAssetRenderUtils.cpp
+++ b/Engine/Source/Editor/Persona/Private/PhysicsAssetRenderUtils.cpp
@@ -56,13 +56,14 @@ FPhysicsAssetRenderSettings::FPhysicsAssetRenderSettings()
 	, COMRenderColor(255, 255, 100)
 	, COMRenderSize(5.0f)
 	, InfluenceLineLength(2.0f)
+{}
+
+void FPhysicsAssetRenderSettings::InitPhysicsAssetRenderSettings(UMaterialInterface* InBoneUnselectedMaterial, UMaterialInterface* InBoneNoCollisionMaterial)
 {
-	UMaterialInterface* BaseBoneUnselectedMaterial = LoadObject<UMaterialInterface>(NULL, TEXT("/Engine/EditorMaterials/PhAT_UnselectedMaterial.PhAT_UnselectedMaterial"), NULL, LOAD_None, NULL);
-	BoneUnselectedMaterial = UMaterialInstanceDynamic::Create(BaseBoneUnselectedMaterial, GetTransientPackage());
+	BoneUnselectedMaterial = InBoneUnselectedMaterial;
 	check(BoneUnselectedMaterial);
 
-	UMaterialInterface* BaseBoneNoCollisionMaterial = LoadObject<UMaterialInterface>(NULL, TEXT("/Engine/EditorMaterials/PhAT_NoCollisionMaterial.PhAT_NoCollisionMaterial"), NULL, LOAD_None, NULL);
-	BoneNoCollisionMaterial = UMaterialInstanceDynamic::Create(BaseBoneNoCollisionMaterial, GetTransientPackage());
+	BoneNoCollisionMaterial = InBoneNoCollisionMaterial;
 	check(BoneNoCollisionMaterial);
 }
 
@@ -600,6 +601,9 @@ void UPhysicsAssetRenderUtilities::InitialiseImpl()
 	FAssetRegistryModule& AssetRegistryModule = FModuleManager::LoadModuleChecked<FAssetRegistryModule>(TEXT("AssetRegistry"));
 	AssetRegistryModule.Get().OnAssetRenamed().AddUObject(this, &UPhysicsAssetRenderUtilities::OnAssetRenamed);
 	AssetRegistryModule.Get().OnInMemoryAssetDeleted().AddUObject(this, &UPhysicsAssetRenderUtilities::OnAssetRemoved);
+
+	BoneUnselectedMaterial = LoadObject<UMaterialInterface>(NULL, TEXT("/Engine/EditorMaterials/PhAT_UnselectedMaterial.PhAT_UnselectedMaterial"), NULL, LOAD_None, NULL);
+	BoneNoCollisionMaterial = LoadObject<UMaterialInterface>(NULL, TEXT("/Engine/EditorMaterials/PhAT_NoCollisionMaterial.PhAT_NoCollisionMaterial"), NULL, LOAD_None, NULL);
 }
 
 FPhysicsAssetRenderSettings* UPhysicsAssetRenderUtilities::GetSettingsImpl(const uint32 InPhysicsAssetPathNameHash)
@@ -608,8 +612,9 @@ FPhysicsAssetRenderSettings* UPhysicsAssetRenderUtilities::GetSettingsImpl(const
 
 	if (!Settings)
 	{
-		IdToSettingsMap.Add(InPhysicsAssetPathNameHash, FPhysicsAssetRenderSettings());
-		Settings = IdToSettingsMap.Find(InPhysicsAssetPathNameHash);
+		FPhysicsAssetRenderSettings& NewSettings = IdToSettingsMap.Add(InPhysicsAssetPathNameHash, FPhysicsAssetRenderSettings());
+		NewSettings.InitPhysicsAssetRenderSettings(BoneUnselectedMaterial, BoneNoCollisionMaterial);
+		Settings = &NewSettings;
 	}
 
 	return Settings;
diff --git a/Engine/Source/Editor/Persona/Public/PhysicsAssetRenderUtils.h b/Engine/Source/Editor/Persona/Public/PhysicsAssetRenderUtils.h
--- a/Engine/Source/Editor/Persona/Public/PhysicsAssetRenderUtils.h
+++ b/Engine/Source/Editor/Persona/Public/PhysicsAssetRenderUtils.h
@@ -35,6 +35,7 @@ struct PERSONA_API FPhysicsAssetRenderSettings
 public:
 
 	FPhysicsAssetRenderSettings();
+	void InitPhysicsAssetRenderSettings(class UMaterialInterface* InBoneUnselectedMaterial, class UMaterialInterface* InBoneNoCollisionMaterial);
 
 	/** Accessors/helper methods */
 	bool IsBodyHidden(const int32 BodyIndex) const;
@@ -107,10 +108,10 @@ public:
 
 	// Materials
 	UPROPERTY()
-	TObjectPtr<class UMaterialInstanceDynamic> BoneUnselectedMaterial;
+	TObjectPtr<class UMaterialInterface> BoneUnselectedMaterial;
 
 	UPROPERTY()
-	TObjectPtr<class UMaterialInstanceDynamic> BoneNoCollisionMaterial;
+	TObjectPtr<class UMaterialInterface> BoneNoCollisionMaterial;
 
 	UPROPERTY()
 	TArray<int32> HiddenBodies;
@@ -155,6 +156,12 @@ private:
 	UPROPERTY()
 	TMap< uint32, FPhysicsAssetRenderSettings > IdToSettingsMap;
 
+	UPROPERTY(transient)
+	TObjectPtr<class UMaterialInterface> BoneUnselectedMaterial;
+
+	UPROPERTY(transient)
+	TObjectPtr<class UMaterialInterface> BoneNoCollisionMaterial;
+
 	IPhysicsAssetRenderInterface* PhysicsAssetRenderInterface;
 };
 
