#UE Fix uninitialized properties

diff --git a/Engine/Source/Editor/Persona/Private/PhysicsAssetRenderUtils.cpp b/Engine/Source/Editor/Persona/Private/PhysicsAssetRenderUtils.cpp
--- a/Engine/Source/Editor/Persona/Private/PhysicsAssetRenderUtils.cpp
+++ b/Engine/Source/Editor/Persona/Private/PhysicsAssetRenderUtils.cpp
@@ -56,6 +56,8 @@ FPhysicsAssetRenderSettings::FPhysicsAssetRenderSettings()
 	, COMRenderColor(255, 255, 100)
 	, COMRenderSize(5.0f)
 	, InfluenceLineLength(2.0f)
+	, BoneUnselectedMaterial(nullptr)
+	, BoneNoCollisionMaterial(nullptr)
 {}
 
 void FPhysicsAssetRenderSettings::InitPhysicsAssetRenderSettings(UMaterialInterface* InBoneUnselectedMaterial, UMaterialInterface* InBoneNoCollisionMaterial)
