[MLDeformer] A complete redesign, turning this into a whole framework that can handle different models.

diff --git a/Engine/Source/Editor/Persona/Private/AnimationEditorPreviewScene.cpp b/Engine/Source/Editor/Persona/Private/AnimationEditorPreviewScene.cpp
--- a/Engine/Source/Editor/Persona/Private/AnimationEditorPreviewScene.cpp
+++ b/Engine/Source/Editor/Persona/Private/AnimationEditorPreviewScene.cpp
@@ -1163,14 +1163,17 @@ void FAnimationEditorPreviewScene::Tick(float InDeltaTime)
 		GetWorld()->Tick(LEVELTICK_All, InDeltaTime);
 	}
 
-	// Handle updating the preview component to represent the effects of root motion	
-	const FBoxSphereBounds& Bounds = GetFloorBounds();
-	SkeletalMeshComponent->ConsumeRootMotion(Bounds.GetBox().Min, Bounds.GetBox().Max);
-
-	if (LastCachedLODForPreviewComponent != SkeletalMeshComponent->GetPredictedLODLevel())
+	if (SkeletalMeshComponent)
 	{
-		OnLODChanged.Broadcast();
-		LastCachedLODForPreviewComponent = SkeletalMeshComponent->GetPredictedLODLevel();
+		// Handle updating the preview component to represent the effects of root motion	
+		const FBoxSphereBounds& Bounds = GetFloorBounds();
+		SkeletalMeshComponent->ConsumeRootMotion(Bounds.GetBox().Min, Bounds.GetBox().Max);
+
+		if (LastCachedLODForPreviewComponent != SkeletalMeshComponent->GetPredictedLODLevel())
+		{
+			OnLODChanged.Broadcast();
+			LastCachedLODForPreviewComponent = SkeletalMeshComponent->GetPredictedLODLevel();
+		}
 	}
 
 	OnPostTickDelegate.Broadcast();
diff --git a/Engine/Source/Editor/Persona/Private/AnimationEditorViewportClient.cpp b/Engine/Source/Editor/Persona/Private/AnimationEditorViewportClient.cpp
--- a/Engine/Source/Editor/Persona/Private/AnimationEditorViewportClient.cpp
+++ b/Engine/Source/Editor/Persona/Private/AnimationEditorViewportClient.cpp
@@ -141,7 +141,10 @@ FAnimationViewportClient::FAnimationViewportClient(const TSharedRef<IPersonaPrev
 	}
 
 	InPreviewScene->RegisterOnPreviewMeshChanged(FOnPreviewMeshChanged::CreateRaw(this, &FAnimationViewportClient::HandleSkeletalMeshChanged));
-	HandleSkeletalMeshChanged(nullptr, InPreviewScene->GetPreviewMeshComponent()->SkeletalMesh);
+	if (InPreviewScene->GetPreviewMeshComponent())
+	{
+		HandleSkeletalMeshChanged(nullptr, InPreviewScene->GetPreviewMeshComponent()->SkeletalMesh);
+	}
 	InPreviewScene->RegisterOnInvalidateViews(FSimpleDelegate::CreateRaw(this, &FAnimationViewportClient::HandleInvalidateViews));
 	InPreviewScene->RegisterOnFocusViews(FSimpleDelegate::CreateRaw(this, &FAnimationViewportClient::HandleFocusViews));
 	InPreviewScene->RegisterOnPreTick(FSimpleDelegate::CreateRaw(this, &FAnimationViewportClient::HandlePreviewScenePreTick));
@@ -446,7 +449,7 @@ void FAnimationViewportClient::Draw(const FSceneView* View, FPrimitiveDrawInterf
 	FEditorViewportClient::Draw(View, PDI);
 
 	UDebugSkelMeshComponent* PreviewMeshComponent = GetAnimPreviewScene()->GetPreviewMeshComponent();
-	if (PreviewMeshComponent->SkeletalMesh)
+	if (PreviewMeshComponent && PreviewMeshComponent->SkeletalMesh)
 	{
 		// Can't have both bones of interest and sockets of interest set
 		check( !(GetAnimPreviewScene()->GetSelectedBoneIndex() != INDEX_NONE && GetAnimPreviewScene()->GetSelectedSocket().IsValid() ) )
