Correctly tag custom rows in anim sequence and skeletal mesh details customizations

diff --git a/Engine/Source/Editor/DetailCustomizations/Private/AnimSequenceDetails.cpp b/Engine/Source/Editor/DetailCustomizations/Private/AnimSequenceDetails.cpp
--- a/Engine/Source/Editor/DetailCustomizations/Private/AnimSequenceDetails.cpp
+++ b/Engine/Source/Editor/DetailCustomizations/Private/AnimSequenceDetails.cpp
@@ -92,7 +92,9 @@ void FAnimSequenceDetails::CustomizeDetails(IDetailLayoutBuilder& DetailBuilder)
 	}
 
 	// add widget for editing retarget source
-	AnimationCategory.AddCustomRow(RetargetSourceNameHandler->GetPropertyDisplayName())
+	AnimationCategory
+	.AddCustomRow(RetargetSourceNameHandler->GetPropertyDisplayName())
+	.RowTag(RetargetSourceNameHandler->GetProperty()->GetFName())
 	.NameContent()
 	[
 		RetargetSourceNameHandler->CreatePropertyNameWidget()
@@ -116,7 +118,9 @@ void FAnimSequenceDetails::CustomizeDetails(IDetailLayoutBuilder& DetailBuilder)
 		]	
 	];
 
-	AnimationCategory.AddCustomRow(RetargetSourceAssetHandle->GetPropertyDisplayName())
+	AnimationCategory
+	.AddCustomRow(RetargetSourceAssetHandle->GetPropertyDisplayName())
+	.RowTag(RetargetSourceAssetHandle->GetProperty()->GetFName())
 	.NameContent()
 	[
 		RetargetSourceAssetHandle->CreatePropertyNameWidget()
diff --git a/Engine/Source/Editor/DetailCustomizations/Private/AnimationAssetDetail.cpp b/Engine/Source/Editor/DetailCustomizations/Private/AnimationAssetDetail.cpp
--- a/Engine/Source/Editor/DetailCustomizations/Private/AnimationAssetDetail.cpp
+++ b/Engine/Source/Editor/DetailCustomizations/Private/AnimationAssetDetail.cpp
@@ -46,7 +46,9 @@ void FAnimationAssetDetails::CustomizeDetails(IDetailLayoutBuilder& DetailBuilde
 
 	// add widget for editing source animation 
 	IDetailCategoryBuilder& AnimationCategory = DetailBuilder.EditCategory("Animation");
-	AnimationCategory.AddCustomRow(PreviewPoseAssetHandler->GetPropertyDisplayName())
+	AnimationCategory
+	.AddCustomRow(PreviewPoseAssetHandler->GetPropertyDisplayName())
+	.RowTag(PreviewPoseAssetHandler->GetProperty()->GetFName())
 	.NameContent()
 	[
 		PreviewPoseAssetHandler->CreatePropertyNameWidget()
diff --git a/Engine/Source/Editor/Persona/Private/PersonaMeshDetails.cpp b/Engine/Source/Editor/Persona/Private/PersonaMeshDetails.cpp
--- a/Engine/Source/Editor/Persona/Private/PersonaMeshDetails.cpp
+++ b/Engine/Source/Editor/Persona/Private/PersonaMeshDetails.cpp
@@ -473,6 +473,7 @@ void FSkeletalMeshReductionSettingsLayout::GenerateChildContent(IDetailChildrenB
 	if (bUseThirdPartyUI)
 	{
 		ChildrenBuilder.AddCustomRow( LOCTEXT("ReductionReductionMethod", "Reduction_ReductionMethod") )
+		.RowTag(TEXT("ReductionMethod"))
 		.NameContent()
 		[
 			SNew(STextBlock)
@@ -500,6 +501,7 @@ void FSkeletalMeshReductionSettingsLayout::GenerateChildContent(IDetailChildrenB
 			LOCTEXT("PercentTriangles_Row", "Triangle Percentage"),
 			LOCTEXT("PercentTriangles", "Percent of Triangles"),
 			LOCTEXT("PercentTriangles_DeviationToolTip", "The percentage of triangles to retain as a ratio, e.g. 0.1 indicates 10 percent."),
+			TEXT("PercentTriangles"),
 			0.0f,
 			1.0f,
 			FGetFloatDelegate::CreateRaw(this, &FSkeletalMeshReductionSettingsLayout::GetNumTrianglesPercentage),
@@ -510,6 +512,7 @@ void FSkeletalMeshReductionSettingsLayout::GenerateChildContent(IDetailChildrenB
 			LOCTEXT("Accuracy_Row", "Accuracy Percentage"),
 			LOCTEXT("PercentAccuracy", "Accuracy Percentage"),
 			LOCTEXT("PercentAccuracy_ToolTip", "The simplification uses this as how much deviate from source mesh. Better works with hard surface meshes."),
+			TEXT("PercentAccuracy"),
 			0.0f,
 			1.0f,
 			FGetFloatDelegate::CreateRaw(this, &FSkeletalMeshReductionSettingsLayout::GetAccuracyPercentage),
@@ -517,9 +520,10 @@ void FSkeletalMeshReductionSettingsLayout::GenerateChildContent(IDetailChildrenB
 			.Visibility(TAttribute<EVisibility>::Create(TAttribute<EVisibility>::FGetter::CreateSP(this, &FSkeletalMeshReductionSettingsLayout::GetVisibiltyIfCurrentReductionMethodIsNot, SMOT_NumOfTriangles)));
 
 
-		auto AddImportanceRow = [this, &ChildrenBuilder](const FText RowTitleText, const FText RowNameContentText, const EImportanceType ImportanceType)
+		auto AddImportanceRow = [this, &ChildrenBuilder](const FText RowTitleText, const FText RowNameContentText, FName RowTag, const EImportanceType ImportanceType)
 		{
 			ChildrenBuilder.AddCustomRow(RowTitleText)
+			.RowTag(RowTag)
 			.NameContent()
 			[
 				SNew(STextBlock)
@@ -543,15 +547,16 @@ void FSkeletalMeshReductionSettingsLayout::GenerateChildContent(IDetailChildrenB
 			];
 		};
 		
-		AddImportanceRow(LOCTEXT("ReductionSilhouetteImportance", "Reduction_SilhouetteImportance"), LOCTEXT("SilhouetteImportance", "Silhouette"), ID_Silhouette);
-		AddImportanceRow(LOCTEXT("ReductionTextureImportance", "Reduction_TextureImportance"), LOCTEXT("TextureImportance", "Texture"), ID_Texture);
-		AddImportanceRow(LOCTEXT("ReductionShadingImportance", "Reduction_ShadingImportance"), LOCTEXT("ShadingImportance", "Shading"), ID_Shading);
-		AddImportanceRow(LOCTEXT("ReductionSkinningImportance", "Reduction_SkinningImportance"), LOCTEXT("SkinningImportance", "Skinning"), ID_Skinning);
+		AddImportanceRow(LOCTEXT("ReductionSilhouetteImportance", "Reduction_SilhouetteImportance"), LOCTEXT("SilhouetteImportance", "Silhouette"), TEXT("SilhouetteImportance"), ID_Silhouette);
+		AddImportanceRow(LOCTEXT("ReductionTextureImportance", "Reduction_TextureImportance"), LOCTEXT("TextureImportance", "Texture"), TEXT("TextureImportance"), ID_Texture);
+		AddImportanceRow(LOCTEXT("ReductionShadingImportance", "Reduction_ShadingImportance"), LOCTEXT("ShadingImportance", "Shading"), TEXT("ShadingImportance"), ID_Shading);
+		AddImportanceRow(LOCTEXT("ReductionSkinningImportance", "Reduction_SkinningImportance"), LOCTEXT("SkinningImportance", "Skinning"), TEXT("SkinningImportance"), ID_Skinning);
 
 		AddBoolRow(ChildrenBuilder,
 			LOCTEXT("RemapMorphTargets_Row", "RemapMorphTargets"),
 			LOCTEXT("RemapMorphTargets_RowNameContent", "Remap Morph Targets"),
 			LOCTEXT("RemapMorphTargets_RowNameContentTooltip", "Remap the morph targets from the base LOD onto the reduce LOD."),
+			TEXT("RemapMorphTargets"),
 			FGetCheckBoxStateDelegate::CreateRaw(this, &FSkeletalMeshReductionSettingsLayout::GetRemapMorphTargets),
 			FSetCheckBoxStateDelegate::CreateRaw(this, &FSkeletalMeshReductionSettingsLayout::SetRemapMorphTargets));
 
@@ -559,6 +564,7 @@ void FSkeletalMeshReductionSettingsLayout::GenerateChildContent(IDetailChildrenB
 			LOCTEXT("RecalcNormals_Row", "Recalculate Normals"),
 			LOCTEXT("RecalcNormals_RowNameContent", "Recompute Normal"),
 			LOCTEXT("RecalcNormals_RowNameContentTooltip", "Whether Normal smoothing groups should be preserved. If true then Hard Edge Angle (NormalsThreshold) is used."),
+			TEXT("RecalcNormals"),
 			FGetCheckBoxStateDelegate::CreateRaw(this, &FSkeletalMeshReductionSettingsLayout::ShouldRecomputeNormals),
 			FSetCheckBoxStateDelegate::CreateRaw(this, &FSkeletalMeshReductionSettingsLayout::OnRecomputeNormalsChanged));
 
@@ -566,6 +572,7 @@ void FSkeletalMeshReductionSettingsLayout::GenerateChildContent(IDetailChildrenB
 			LOCTEXT("NormalsThreshold_Row", "Normals Threshold"),
 			LOCTEXT("NormalsThreshold_RowNameContent", "Hard Edge Angle"),
 			LOCTEXT("NormalsThreshold_RowNameContentToolTip", "If the angle between two triangles are above this value, the normals will not be smooth over the edge between those two triangles. Set in degrees. This is only used when Recalculate Normals is set to true."),
+			TEXT("NormalsThreshold"),
 			0.0f,
 			360.0f,
 			FGetFloatDelegate::CreateRaw(this, &FSkeletalMeshReductionSettingsLayout::GetNormalsThreshold),
@@ -575,6 +582,7 @@ void FSkeletalMeshReductionSettingsLayout::GenerateChildContent(IDetailChildrenB
 			LOCTEXT("WeldingThreshold_Row", "Welding Threshold"),
 			LOCTEXT("WeldingThreshold_RowNameContent", "Welding Threshold"),
 			LOCTEXT("WeldingThreshold_RowNameContentToolTip", "The welding threshold distance.Vertices under this distance will be welded."),
+			TEXT("WeldingThreshold"),
 			0.0f,
 			1000.0f,
 			FGetFloatDelegate::CreateRaw(this, &FSkeletalMeshReductionSettingsLayout::GetWeldingThreshold),
@@ -584,6 +592,7 @@ void FSkeletalMeshReductionSettingsLayout::GenerateChildContent(IDetailChildrenB
 			LOCTEXT("MaxBonesPerVertex_Row", "MaxBonesPerVertex"),
 			LOCTEXT("MaxBonesPerVertex", "Max Bones Influence"),
 			LOCTEXT("MaxBonesPerVertex_ToolTip", "Maximum number of bones that can be assigned to each vertex."),
+			TEXT("MaxBonesPerVertex"),
 			1,
 			INT_MAX,
 			FGetIntegerDelegate::CreateRaw(this, &FSkeletalMeshReductionSettingsLayout::GetMaxBonesPerVertex),
@@ -620,6 +629,7 @@ void FSkeletalMeshReductionSettingsLayout::GenerateChildContent(IDetailChildrenB
 				LOCTEXT("PercentTriangles_Row", "Triangle Percentage"),
 				LOCTEXT("PercentTriangles", "Percent of Triangles"),
 				LOCTEXT("PercentTriangles_ToolTip", "The simplification uses this percentage of source mesh's triangle count as a target."),
+				TEXT("PercentTriangles"),
 				0.0f,
 				1.0f,
 				FGetFloatDelegate::CreateRaw(this, &FSkeletalMeshReductionSettingsLayout::GetNumTrianglesPercentage),
@@ -633,6 +643,7 @@ void FSkeletalMeshReductionSettingsLayout::GenerateChildContent(IDetailChildrenB
 				LOCTEXT("Percentvertices_Row", "Vertices Percentage"),
 				LOCTEXT("PercentVertices", "Percent of Vertices"),
 				LOCTEXT("PercentVertices_ToolTip", "The percentage of vertices to retain as a ratio, e.g. 0.1 indicates 10 percent."),
+				TEXT("PercentVertices"),
 				0.0f,
 				1.0f,
 				FGetFloatDelegate::CreateRaw(this, &FSkeletalMeshReductionSettingsLayout::GetNumVerticesPercentage),
@@ -646,6 +657,7 @@ void FSkeletalMeshReductionSettingsLayout::GenerateChildContent(IDetailChildrenB
 				LOCTEXT("MaxTriangles_Row", "Max Number of Triangles"),
 				LOCTEXT("MaxTriangles", "Max Triangles Count"),
 				LOCTEXT("MaxTriangles_ToolTip", "The maximum number of triangles to retain."),
+				TEXT("MaxTriangles"),
 				0,
 				INT_MAX,
 				FGetIntegerDelegate::CreateRaw(this, &FSkeletalMeshReductionSettingsLayout::GetNumMaxTrianglesCount),
@@ -666,6 +678,7 @@ void FSkeletalMeshReductionSettingsLayout::GenerateChildContent(IDetailChildrenB
 				LOCTEXT("MaxVertices_Row", "Max Number of Vertices"),
 				LOCTEXT("MaxVertices", "Max Vertex Count"),
 				LOCTEXT("MaxVertices_ToolTip", "The maximum number of vertices to retain."),
+				TEXT("MaxVertices"),
 				0,
 				INT_MAX,
 				FGetIntegerDelegate::CreateRaw(this, &FSkeletalMeshReductionSettingsLayout::GetNumMaxVerticesCount),
@@ -685,6 +698,7 @@ void FSkeletalMeshReductionSettingsLayout::GenerateChildContent(IDetailChildrenB
 				LOCTEXT("MaxTrianglesPercentage_Row", "Max Number of Triangles"),
 				LOCTEXT("MaxTrianglesPercentage", "Max Triangles Count"),
 				LOCTEXT("MaxTrianglesPercentage_ToolTip", "The maximum number of triangles to retain when using percentage criterion."),
+				TEXT("MaxTrianglesPercentage"),
 				0,
 				MAX_uint32,
 				FGetUnsignedIntegerDelegate::CreateRaw(this, &FSkeletalMeshReductionSettingsLayout::GetNumMaxTrianglesPercentageCount),
@@ -704,6 +718,7 @@ void FSkeletalMeshReductionSettingsLayout::GenerateChildContent(IDetailChildrenB
 				LOCTEXT("MaxVerticesPercentage_Row", "Max Number of Vertices"),
 				LOCTEXT("MaxVerticesPercentage", "Max Vertex Count"),
 				LOCTEXT("MaxVerticesPercentage_ToolTip", "The maximum number of vertices to retain when using percentage criterion."),
+				TEXT("MaxVerticesPercentage"),
 				0,
 				MAX_uint32,
 				FGetUnsignedIntegerDelegate::CreateRaw(this, &FSkeletalMeshReductionSettingsLayout::GetNumMaxVerticesPercentageCount),
@@ -721,6 +736,7 @@ void FSkeletalMeshReductionSettingsLayout::GenerateChildContent(IDetailChildrenB
 			LOCTEXT("RemapMorphTargets_Row", "RemapMorphTargets"),
 			LOCTEXT("RemapMorphTargets_RowNameContent", "Remap Morph Targets"),
 			LOCTEXT("RemapMorphTargets_RowNameContentTooltip", "Remap the morph targets from the base LOD onto the reduce LOD."),
+			TEXT("RemapMorphTargets"),
 			FGetCheckBoxStateDelegate::CreateRaw(this, &FSkeletalMeshReductionSettingsLayout::GetRemapMorphTargets),
 			FSetCheckBoxStateDelegate::CreateRaw(this, &FSkeletalMeshReductionSettingsLayout::SetRemapMorphTargets));
 
@@ -728,6 +744,7 @@ void FSkeletalMeshReductionSettingsLayout::GenerateChildContent(IDetailChildrenB
 			LOCTEXT("MaxBonesPerVertex_Row", "MaxBonesPerVertex"),
 			LOCTEXT("MaxBonesPerVertex", "Max Bones Influence"),
 			LOCTEXT("MaxBonesPerVertex_ToolTip", "Maximum number of bones that can be assigned to each vertex."),
+			TEXT("MaxBonesPerVertex"),
 			1,
 			INT_MAX,
 			FGetIntegerDelegate::CreateRaw(this, &FSkeletalMeshReductionSettingsLayout::GetMaxBonesPerVertex),
@@ -737,6 +754,7 @@ void FSkeletalMeshReductionSettingsLayout::GenerateChildContent(IDetailChildrenB
 			LOCTEXT("EnforceBoneBoundaries_Row", "EnforceBoneBoundaries"),
 			LOCTEXT("EnforceBoneBoundaries_RowNameContent", "Enforce Bone Boundaries"),
 			LOCTEXT("EnforceBoneBoundaries_RowNameContentTooltip", "Penalize edge collapse between vertices that have different major bones.  This will help articulated segments like tongues but can lead to undesirable results under extreme simplification."),
+			TEXT("EnforceBoneBoundaries"),
 			FGetCheckBoxStateDelegate::CreateRaw(this, &FSkeletalMeshReductionSettingsLayout::GetEnforceBoneBoundaries),
 			FSetCheckBoxStateDelegate::CreateRaw(this, &FSkeletalMeshReductionSettingsLayout::SetEnforceBoneBoundaries));
 
@@ -744,6 +762,7 @@ void FSkeletalMeshReductionSettingsLayout::GenerateChildContent(IDetailChildrenB
 			LOCTEXT("MergeCoincidentVertBones_Row", "MergeCoincidentVertBones"),
 			LOCTEXT("MergeCoincidentVertBones_RowNameContent", "Merge Coincident Vertices Bones"),
 			LOCTEXT("MergeCoincidentVertBones_RowNameContentTooltip", "If enabled this option make sure vertices that share the same location (e.g. UV boundaries) have the same bone weights. This can fix cracks when the characters animate."),
+			TEXT("MergeCoincidentVertBones"),
 			FGetCheckBoxStateDelegate::CreateRaw(this, &FSkeletalMeshReductionSettingsLayout::GetMergeCoincidentVertBones),
 			FSetCheckBoxStateDelegate::CreateRaw(this, &FSkeletalMeshReductionSettingsLayout::SetMergeCoincidentVertBones));
 
@@ -751,6 +770,7 @@ void FSkeletalMeshReductionSettingsLayout::GenerateChildContent(IDetailChildrenB
 			LOCTEXT("VolumeImportance_Row", "Volume Importance"),
 			LOCTEXT("VolumeImportance", "Volumetric Correction"),
 			LOCTEXT("VolumeImportance_ToolTip", "Default value of 1 attempts to preserve volume.  Smaller values will loose volume by flattening curved surfaces, and larger values will accentuate curved surfaces."),
+			TEXT("VolumeImportance"),
 			0.0f,
 			2.0f,
 			FGetFloatDelegate::CreateRaw(this, &FSkeletalMeshReductionSettingsLayout::GetVolumeImportance),
@@ -760,6 +780,7 @@ void FSkeletalMeshReductionSettingsLayout::GenerateChildContent(IDetailChildrenB
 			LOCTEXT("LockEdges_Row", "LockEdges"),
 			LOCTEXT("LockEdges_RowNameContent", "Lock Mesh Edges"),
 			LOCTEXT("LockEdges_RowNameContentTooltip", "Preserve cuts in the mesh surface by locking vertices in place.  Increases the quality of the simplified mesh at edges at the cost of more triangles."),
+			TEXT("LockEdges"),
 			FGetCheckBoxStateDelegate::CreateRaw(this, &FSkeletalMeshReductionSettingsLayout::GetLockEdges),
 			FSetCheckBoxStateDelegate::CreateRaw(this, &FSkeletalMeshReductionSettingsLayout::SetLockEdges));
 
@@ -767,6 +788,7 @@ void FSkeletalMeshReductionSettingsLayout::GenerateChildContent(IDetailChildrenB
 			LOCTEXT("LockColorBoundaries_Row", "LockColorBoundaries"),
 			LOCTEXT("LockColorBoundaries_RowNameContent", "Lock Vertex Color Boundaries"),
 			LOCTEXT("LockColorBoundaries_RowNameContentTooltip", "Locking edges that connect two vertex colors.  Increases the quality of the simplified mesh at edges at the cost of more triangles."),
+			TEXT("LockColorBoundaries"),
 			FGetCheckBoxStateDelegate::CreateRaw(this, &FSkeletalMeshReductionSettingsLayout::GetLockColorBounaries),
 			FSetCheckBoxStateDelegate::CreateRaw(this, &FSkeletalMeshReductionSettingsLayout::SetLockColorBounaries));
 	}
@@ -779,7 +801,7 @@ bool FSkeletalMeshReductionSettingsLayout::IsReductionEnabled() const
 	return IsLODSettingsEnabledDelegate.IsBound() ? IsLODSettingsEnabledDelegate.Execute(LODIndex) : false;
 }
 
-FDetailWidgetRow& FSkeletalMeshReductionSettingsLayout::AddFloatRow(IDetailChildrenBuilder& ChildrenBuilder, const FText RowTitleText, const FText RowNameContentText, const FText RowNameContentTootlipText, const float MinSliderValue, const float MaxSliderValue, FGetFloatDelegate GetterDelegate, FSetFloatDelegate SetterDelegate)
+FDetailWidgetRow& FSkeletalMeshReductionSettingsLayout::AddFloatRow(IDetailChildrenBuilder& ChildrenBuilder, const FText RowTitleText, const FText RowNameContentText, const FText RowNameContentTootlipText, FName RowTag, const float MinSliderValue, const float MaxSliderValue, FGetFloatDelegate GetterDelegate, FSetFloatDelegate SetterDelegate)
 {
 	int32 SliderDataIndex = SliderStateDataArray.Num();
 	FSliderStateData& SliderData = SliderStateDataArray.AddDefaulted_GetRef();
@@ -824,6 +846,7 @@ FDetailWidgetRow& FSkeletalMeshReductionSettingsLayout::AddFloatRow(IDetailChild
 	};
 	
 	FDetailWidgetRow& Row = ChildrenBuilder.AddCustomRow(RowTitleText)
+	.RowTag(RowTag)
 	.NameContent()
 	[
 		SNew(STextBlock)
@@ -846,9 +869,10 @@ FDetailWidgetRow& FSkeletalMeshReductionSettingsLayout::AddFloatRow(IDetailChild
 	return Row;
 }
 
-FDetailWidgetRow& FSkeletalMeshReductionSettingsLayout::AddBoolRow(IDetailChildrenBuilder& ChildrenBuilder, const FText RowTitleText, const FText RowNameContentText, const FText RowNameContentToolitipText, FGetCheckBoxStateDelegate GetterDelegate, FSetCheckBoxStateDelegate SetterDelegate)
+FDetailWidgetRow& FSkeletalMeshReductionSettingsLayout::AddBoolRow(IDetailChildrenBuilder& ChildrenBuilder, const FText RowTitleText, const FText RowNameContentText, const FText RowNameContentToolitipText, FName RowTag, FGetCheckBoxStateDelegate GetterDelegate, FSetCheckBoxStateDelegate SetterDelegate)
 {
 	FDetailWidgetRow& Row = ChildrenBuilder.AddCustomRow(RowTitleText)
+	.RowTag(RowTag)
 	.NameContent()
 	[
 		SNew(STextBlock)
@@ -866,7 +890,7 @@ FDetailWidgetRow& FSkeletalMeshReductionSettingsLayout::AddBoolRow(IDetailChildr
 	return Row;
 }
 
-FDetailWidgetRow& FSkeletalMeshReductionSettingsLayout::AddIntegerRow(IDetailChildrenBuilder& ChildrenBuilder, const FText RowTitleText, const FText RowNameContentText, const FText RowNameContentTootlipText, const int32 MinSliderValue, const int32 MaxSliderValue, FGetIntegerDelegate GetterDelegate, FSetIntegerDelegate SetterDelegate)
+FDetailWidgetRow& FSkeletalMeshReductionSettingsLayout::AddIntegerRow(IDetailChildrenBuilder& ChildrenBuilder, const FText RowTitleText, const FText RowNameContentText, const FText RowNameContentTootlipText, FName RowTag, const int32 MinSliderValue, const int32 MaxSliderValue, FGetIntegerDelegate GetterDelegate, FSetIntegerDelegate SetterDelegate)
 {
 	int32 SliderDataIndex = SliderStateDataArray.Num();
 	FSliderStateData& SliderData = SliderStateDataArray.AddDefaulted_GetRef();
@@ -911,6 +935,7 @@ FDetailWidgetRow& FSkeletalMeshReductionSettingsLayout::AddIntegerRow(IDetailChi
 	};
 
 	FDetailWidgetRow& Row = ChildrenBuilder.AddCustomRow(RowTitleText)
+	.RowTag(RowTag)
 	.NameContent()
 	[
 		SNew(STextBlock)
@@ -933,7 +958,7 @@ FDetailWidgetRow& FSkeletalMeshReductionSettingsLayout::AddIntegerRow(IDetailChi
 	return Row;
 }
 
-FDetailWidgetRow& FSkeletalMeshReductionSettingsLayout::AddUnsignedIntegerRow(IDetailChildrenBuilder& ChildrenBuilder, const FText RowTitleText, const FText RowNameContentText, const FText RowNameContentTootlipText, const uint32 MinSliderValue, const uint32 MaxSliderValue, FGetUnsignedIntegerDelegate GetterDelegate, FSetUnsignedIntegerDelegate SetterDelegate)
+FDetailWidgetRow& FSkeletalMeshReductionSettingsLayout::AddUnsignedIntegerRow(IDetailChildrenBuilder& ChildrenBuilder, const FText RowTitleText, const FText RowNameContentText, const FText RowNameContentTootlipText, FName RowTag, const uint32 MinSliderValue, const uint32 MaxSliderValue, FGetUnsignedIntegerDelegate GetterDelegate, FSetUnsignedIntegerDelegate SetterDelegate)
 {
 	uint32 SliderDataIndex = SliderStateDataArray.Num();
 	FSliderStateData& SliderData = SliderStateDataArray.AddDefaulted_GetRef();
@@ -978,17 +1003,18 @@ FDetailWidgetRow& FSkeletalMeshReductionSettingsLayout::AddUnsignedIntegerRow(ID
 	};
 
 	FDetailWidgetRow& Row = ChildrenBuilder.AddCustomRow(RowTitleText)
-		.NameContent()
-		[
-			SNew(STextBlock)
-			.Font(IDetailLayoutBuilder::GetDetailFont())
+	.RowTag(RowTag)
+	.NameContent()
+	[
+		SNew(STextBlock)
+		.Font(IDetailLayoutBuilder::GetDetailFont())
 		.Text(RowNameContentText)
 		.ToolTipText(RowNameContentTootlipText)
-		]
+	]
 	.ValueContent()
-		[
-			SNew(SSpinBox<uint32>)
-			.Font(IDetailLayoutBuilder::GetDetailFont())
+	[
+		SNew(SSpinBox<uint32>)
+		.Font(IDetailLayoutBuilder::GetDetailFont())
 		.MinValue(MinSliderValue)
 		.MaxValue(MaxSliderValue)
 		.Value_Lambda(GetValueHelperFunc)
@@ -996,7 +1022,7 @@ FDetailWidgetRow& FSkeletalMeshReductionSettingsLayout::AddUnsignedIntegerRow(ID
 		.OnBeginSliderMovement_Lambda(BeginSliderMovementHelperFunc)
 		.OnEndSliderMovement_Lambda(EndSliderMovementHelperFunc)
 		.IsEnabled(this, &FSkeletalMeshReductionSettingsLayout::IsReductionEnabled)
-		];
+	];
 	return Row;
 }
 
@@ -1014,6 +1040,7 @@ void FSkeletalMeshReductionSettingsLayout::AddBaseLODRow(IDetailChildrenBuilder&
 			LOCTEXT("ReductionBaseLOD", "Reduction_BaseLOD"),
 			LOCTEXT("BaseLOD", "Base LOD"),
 			LOCTEXT("BaseLODTooltip", "Base LOD index to generate this LOD. By default, we generate from LOD 0"),
+			TEXT("BaseLOD"),
 			0,
 			MaxBaseLOD,
 			FGetIntegerDelegate::CreateRaw(this, &FSkeletalMeshReductionSettingsLayout::GetBaseLODValue),
@@ -2564,7 +2591,9 @@ void FPersonaMeshDetails::AddLODLevelCategories(IDetailLayoutBuilder& DetailLayo
 		{
 			FString MaterialCategoryName = FString(TEXT("Material Slots"));
 			IDetailCategoryBuilder& MaterialCategory = DetailLayout.EditCategory(*MaterialCategoryName, FText::GetEmpty(), ECategoryPriority::Important);
-			MaterialCategory.AddCustomRow(LOCTEXT("AddLODLevelCategories_MaterialArrayOperationAdd", "Materials Operation Add Material Slot"))
+			MaterialCategory
+				.AddCustomRow(LOCTEXT("AddLODLevelCategories_MaterialArrayOperationAdd", "Materials Operation Add Material Slot"))
+				.RowTag(TEXT("MaterialSlots"))
 				.CopyAction(FUIAction(FExecuteAction::CreateSP(this, &FPersonaMeshDetails::OnCopyMaterialList), FCanExecuteAction::CreateSP(this, &FPersonaMeshDetails::OnCanCopyMaterialList)))
 				.PasteAction(FUIAction(FExecuteAction::CreateSP(this, &FPersonaMeshDetails::OnPasteMaterialList)))
 				.NameContent()
@@ -2648,6 +2677,7 @@ void FPersonaMeshDetails::AddLODLevelCategories(IDetailLayoutBuilder& DetailLayo
 		LodCustomCategory = &LODCustomModeCategory;
 
 		LODCustomModeCategory.AddCustomRow((LOCTEXT("LODCustomModeSelect", "Select LOD")))
+		.RowTag("SelectLOD")
 		.NameContent()
 		[
 			SNew(STextBlock)
@@ -2661,6 +2691,7 @@ void FPersonaMeshDetails::AddLODLevelCategories(IDetailLayoutBuilder& DetailLayo
 		];
 
 		LODCustomModeCategory.AddCustomRow((LOCTEXT("LODCustomModeFirstRowName", "LODCustomMode")))
+		.RowTag("LODCustomMode")
 		.NameContent()
 		[
 			SNew(STextBlock)
@@ -2712,6 +2743,7 @@ void FPersonaMeshDetails::AddLODLevelCategories(IDetailLayoutBuilder& DetailLayo
 			bool IsViewportLOD = (CurrentLodIndex == 0 ? 0 : CurrentLodIndex - 1) == LODIndex;
 			DetailDisplayLODs.Add(true); //Enable all LOD in custum mode
 			LODCustomModeCategory.AddCustomRow(( LOCTEXT("LODCustomModeRowName", "LODCheckBoxRowName")), true)
+			.RowTag("LODCheckBoxRowName")
 			.NameContent()
 			[
 				SNew(STextBlock)
@@ -2824,6 +2856,7 @@ void FPersonaMeshDetails::AddLODLevelCategories(IDetailLayoutBuilder& DetailLayo
 				{
 					FString MeshDescriptionReferenceIDString = LODModel.GetLODModelDeriveDataKey();
 					LODCategory.AddCustomRow(LOCTEXT("LODButtonsRow", "LOD Buttons"))
+						.RowTag(TEXT("LODButtons"))
 						.ValueContent()
 						.HAlign(HAlign_Fill)
 						[
@@ -3125,6 +3158,7 @@ void FPersonaMeshDetails::CustomizeLODSettingsCategories(IDetailLayoutBuilder& D
 	TSharedPtr<SWidget> LodTextPtr;
 
 	LODSettingsCategory.AddCustomRow(LOCTEXT("LODImport", "LOD Import"))
+	.RowTag("LODImport")
 	.NameContent()
 	[
 		SAssignNew(LodTextPtr, STextBlock)
@@ -3144,6 +3178,7 @@ void FPersonaMeshDetails::CustomizeLODSettingsCategories(IDetailLayoutBuilder& D
 	// Add Number of LODs slider.
 	const int32 MinAllowedLOD = 1;
 	LODSettingsCategory.AddCustomRow(LOCTEXT("NumberOfLODs", "Number of LODs"))
+	.RowTag("NumberOfLODs")
 	.Visibility(TAttribute<EVisibility>::Create(TAttribute<EVisibility>::FGetter::CreateLambda([]()->EVisibility { return IsAutoMeshReductionAvailable()? EVisibility::Visible : EVisibility::Hidden; })))
 	.NameContent()
 	[
@@ -3165,6 +3200,7 @@ void FPersonaMeshDetails::CustomizeLODSettingsCategories(IDetailLayoutBuilder& D
 	];
 
 	LODSettingsCategory.AddCustomRow(LOCTEXT("ApplyChanges", "Apply Changes"))
+	.RowTag("ApplyChanges")
 	.Visibility(TAttribute<EVisibility>::Create(TAttribute<EVisibility>::FGetter::CreateLambda([]()->EVisibility { return IsAutoMeshReductionAvailable() ? EVisibility::Visible : EVisibility::Hidden; })))
 	.ValueContent()
 	.HAlign(HAlign_Left)
@@ -3183,6 +3219,7 @@ void FPersonaMeshDetails::CustomizeLODSettingsCategories(IDetailLayoutBuilder& D
 	TSharedPtr<IPropertyHandle> LODSettingAssetPropertyHandle = DetailLayout.GetProperty(USkeletalMesh::GetLODSettingsMemberName(), USkeletalMesh::StaticClass());
 	DetailLayout.HideProperty(LODSettingAssetPropertyHandle);
 	LODSettingsCategory.AddCustomRow(LODSettingAssetPropertyHandle->GetPropertyDisplayName())
+	.RowTag(TEXT("LodSettingsAsset"))
 	.NameContent()
 	[
 		LODSettingAssetPropertyHandle->CreatePropertyNameWidget()
@@ -4177,6 +4214,7 @@ void FPersonaMeshDetails::CustomizeDetails( IDetailLayoutBuilder& DetailLayout )
 	PostProcessHandle->MarkHiddenByCustomization();
 
 	FDetailWidgetRow& PostProcessRow = SkelMeshCategory.AddCustomRow(LOCTEXT("PostProcessFilterString", "Post Process Blueprint"));
+	PostProcessRow.RowTag(TEXT("PostProcessAnimBlueprint"));
 	PostProcessRow.NameContent()
 	[
 		PostProcessHandle->CreatePropertyNameWidget()
@@ -4194,6 +4232,7 @@ void FPersonaMeshDetails::CustomizeDetails( IDetailLayoutBuilder& DetailLayout )
 
 	// Add warning if the post process BP is using an incompatible skeleton
 	FDetailWidgetRow& PostProcessWarningRow = SkelMeshCategory.AddCustomRow(LOCTEXT("PostProcessWarningFilterString", "Post Process Blueprint Warning"));
+	PostProcessWarningRow.RowTag(TEXT("PostProcessAnimBlueprint"));
 
 	PostProcessWarningRow
 	.Visibility(MakeAttributeLambda([this]()
diff --git a/Engine/Source/Editor/Persona/Private/PersonaMeshDetails.h b/Engine/Source/Editor/Persona/Private/PersonaMeshDetails.h
--- a/Engine/Source/Editor/Persona/Private/PersonaMeshDetails.h
+++ b/Engine/Source/Editor/Persona/Private/PersonaMeshDetails.h
@@ -123,16 +123,16 @@ private:
 	virtual void GenerateChildContent(IDetailChildrenBuilder& ChildrenBuilder) override;
 	virtual void Tick(float DeltaTime) override {}
 	virtual bool RequiresTick() const override { return false; }
-	virtual FName GetName() const override { static FName MeshReductionSettings("MeshReductionSettings"); return MeshReductionSettings; }
+	virtual FName GetName() const override { static FName MeshReductionSettings("SkeletalMeshOptimizationSettings"); return MeshReductionSettings; }
 	virtual bool InitiallyCollapsed() const override { return true; }
 
 	bool IsReductionEnabled() const;
 
 	//Custom Row Add utilities
-	FDetailWidgetRow& AddFloatRow(IDetailChildrenBuilder& ChildrenBuilder, const FText RowTitleText, const FText RowNameContentText, const FText RowNameContentTootlipText, const float MinSliderValue, const float MaxSliderValue, FGetFloatDelegate GetterDelegate, FSetFloatDelegate SetterDelegate);
-	FDetailWidgetRow& AddBoolRow(IDetailChildrenBuilder& ChildrenBuilder, const FText RowTitleText, const FText RowNameContentText, const FText RowNameContentToolitipText, FGetCheckBoxStateDelegate GetterDelegate, FSetCheckBoxStateDelegate SetterDelegate);
-	FDetailWidgetRow& AddIntegerRow(IDetailChildrenBuilder& ChildrenBuilder, const FText RowTitleText, const FText RowNameContentText, const FText RowNameContentTootlipText, const int32 MinSliderValue, const int32 MaxSliderValue, FGetIntegerDelegate GetterDelegate, FSetIntegerDelegate SetterDelegate);
-	FDetailWidgetRow& AddUnsignedIntegerRow(IDetailChildrenBuilder& ChildrenBuilder, const FText RowTitleText, const FText RowNameContentText, const FText RowNameContentTootlipText, const uint32 MinSliderValue, const uint32 MaxSliderValue, FGetUnsignedIntegerDelegate GetterDelegate, FSetUnsignedIntegerDelegate SetterDelegate);
+	FDetailWidgetRow& AddFloatRow(IDetailChildrenBuilder& ChildrenBuilder, const FText RowTitleText, const FText RowNameContentText, const FText RowNameContentTootlipText, FName RowTag, const float MinSliderValue, const float MaxSliderValue, FGetFloatDelegate GetterDelegate, FSetFloatDelegate SetterDelegate);
+	FDetailWidgetRow& AddBoolRow(IDetailChildrenBuilder& ChildrenBuilder, const FText RowTitleText, const FText RowNameContentText, const FText RowNameContentToolitipText, FName RowTag, FGetCheckBoxStateDelegate GetterDelegate, FSetCheckBoxStateDelegate SetterDelegate);
+	FDetailWidgetRow& AddIntegerRow(IDetailChildrenBuilder& ChildrenBuilder, const FText RowTitleText, const FText RowNameContentText, const FText RowNameContentTootlipText, FName RowTag, const int32 MinSliderValue, const int32 MaxSliderValue, FGetIntegerDelegate GetterDelegate, FSetIntegerDelegate SetterDelegate);
+	FDetailWidgetRow& AddUnsignedIntegerRow(IDetailChildrenBuilder& ChildrenBuilder, const FText RowTitleText, const FText RowNameContentText, const FText RowNameContentTootlipText, FName RowTag, const uint32 MinSliderValue, const uint32 MaxSliderValue, FGetUnsignedIntegerDelegate GetterDelegate, FSetUnsignedIntegerDelegate SetterDelegate);
 	void AddBaseLODRow(IDetailChildrenBuilder& ChildrenBuilder);
 
 	void SetPercentAndAbsoluteVisibility(FDetailWidgetRow& Row, SkeletalMeshTerminationCriterion FirstCriterion, SkeletalMeshTerminationCriterion SecondCriterion);
