Added radio button for Buffer Visualization view mode in Control Rig and Persona Viewports

diff --git a/Engine/Source/Editor/Persona/Private/SAnimViewportToolBar.cpp b/Engine/Source/Editor/Persona/Private/SAnimViewportToolBar.cpp
--- a/Engine/Source/Editor/Persona/Private/SAnimViewportToolBar.cpp
+++ b/Engine/Source/Editor/Persona/Private/SAnimViewportToolBar.cpp
@@ -1269,17 +1269,36 @@ void SAnimViewportToolBar::AddMenuExtender(FName MenuToExtend, FMenuExtensionDel
 TSharedRef<FExtender> SAnimViewportToolBar::GetViewMenuExtender(TSharedPtr<class SEditorViewport> InRealViewport)
 {
 	TSharedRef<FExtender> Extender(new FExtender());
+
 	Extender->AddMenuExtension(
 		TEXT("ViewMode"),
 		EExtensionHook::After,
 		InRealViewport->GetCommandList(),
-		FMenuExtensionDelegate::CreateLambda([](FMenuBuilder& InMenuBuilder)
+		FMenuExtensionDelegate::CreateLambda([this](FMenuBuilder& InMenuBuilder)
 		{
-			InMenuBuilder.AddSubMenu(LOCTEXT("VisualizeBufferViewModeDisplayName", "Buffer Visualization"),
+			InMenuBuilder.AddSubMenu(
+				LOCTEXT("VisualizeBufferViewModeDisplayName", "Buffer Visualization"),
 				LOCTEXT("BufferVisualizationMenu_ToolTip", "Select a mode for buffer visualization"),
 				FNewMenuDelegate::CreateStatic(&FBufferVisualizationMenuCommands::BuildVisualisationSubMenu),
-				false,
-				FSlateIcon(FEditorStyle::GetStyleSetName(), "EditorViewport.VisualizeBufferMode"));
+				FUIAction(
+					FExecuteAction(),
+					FCanExecuteAction(),
+					FIsActionChecked::CreateLambda([this]()
+					{
+						const TSharedPtr<SAnimationEditorViewportTabBody> ViewportPtr = Viewport.Pin();
+						if (ViewportPtr.IsValid())
+						{
+							const FEditorViewportClient& ViewportClient = ViewportPtr->GetViewportClient();
+							return ViewportClient.IsViewModeEnabled(VMI_VisualizeBuffer);
+						}
+						return false;
+					})
+				),
+				"VisualizeBufferViewMode",
+				EUserInterfaceActionType::RadioButton,
+				/* bInOpenSubMenuOnClick = */ false,
+				FSlateIcon(FEditorStyle::GetStyleSetName(), "EditorViewport.VisualizeBufferMode")
+				);
 		}));
 
 	return Extender;
