Fix slowdown in animation sequences with lots of curve tracks

Switched to a SScrollBorder instead of SScrollBox so that widget size is correctly calculated and we dont end up trying to create widgets for all tracks

diff --git a/Engine/Source/Editor/Persona/Private/AnimTimeline/SAnimTimeline.cpp b/Engine/Source/Editor/Persona/Private/AnimTimeline/SAnimTimeline.cpp
--- a/Engine/Source/Editor/Persona/Private/AnimTimeline/SAnimTimeline.cpp
+++ b/Engine/Source/Editor/Persona/Private/AnimTimeline/SAnimTimeline.cpp
@@ -218,9 +218,7 @@ void SAnimTimeline::Construct(const FArguments& InArgs, const TSharedRef<FAnimMo
 								+ SVerticalBox::Slot()
 								.FillHeight(1.f)
 								[
-									SNew(SScrollBox)
-									.ExternalScrollbar(ScrollBar)
-									+ SScrollBox::Slot()
+									SNew(SScrollBorder, Outliner.ToSharedRef())
 									[
 										SNew(SHorizontalBox)
 										
