Animation Curve Editor is Very Slow In Engine
#fix ensure that interactive changes in the curve editor for Animation does not trigger recompression or caching of FRichCurve
#misc fix up curve editor drag/spinbox operations to actually use the property change flag(s)
#misc AnimSequence now evaluates float curves directly from model (which allows previewing interactive changes)

diff --git a/Engine/Source/Editor/CurveEditor/Private/DragOperations/CurveEditorDragOperation_Tangent.cpp b/Engine/Source/Editor/CurveEditor/Private/DragOperations/CurveEditorDragOperation_Tangent.cpp
--- a/Engine/Source/Editor/CurveEditor/Private/DragOperations/CurveEditorDragOperation_Tangent.cpp
+++ b/Engine/Source/Editor/CurveEditor/Private/DragOperations/CurveEditorDragOperation_Tangent.cpp
@@ -58,7 +58,7 @@ void FCurveEditorDragOperation_Tangent::OnDrag(FVector2D InitialPosition, FVecto
 	TArray<FKeyAttributes> NewKeyAttributesScratch;
 	TArray<FKeyPosition> KeyPositions;
 
-	for (const FKeyData& KeyData : KeysByCurve)
+	for (FKeyData& KeyData : KeysByCurve)
 	{
 		const SCurveEditorView* View = CurveEditor->FindFirstInteractiveView(KeyData.CurveID);
 		if (!View)
@@ -175,7 +175,8 @@ void FCurveEditorDragOperation_Tangent::OnDrag(FVector2D InitialPosition, FVecto
 				
 			NewKeyAttributesScratch.Add(NewAttributes);
 		}
-		Curve->SetKeyAttributes(KeyData.Handles, NewKeyAttributesScratch);
+		Curve->SetKeyAttributes(KeyData.Handles, NewKeyAttributesScratch, EPropertyChangeType::Interactive);
+		KeyData.LastDraggedAttributes = NewKeyAttributesScratch;
 	}
 }
 
@@ -187,7 +188,7 @@ void FCurveEditorDragOperation_Tangent::OnCancelDrag()
 	{
 		if (FCurveModel* Curve = CurveEditor->FindCurve(KeyData.CurveID))
 		{
-			Curve->SetKeyAttributes(KeyData.Handles, KeyData.Attributes);
+			Curve->SetKeyAttributes(KeyData.Handles, KeyData.Attributes, EPropertyChangeType::ValueSet);
 		}
 	}
 
@@ -197,6 +198,15 @@ void FCurveEditorDragOperation_Tangent::OnCancelDrag()
 void FCurveEditorDragOperation_Tangent::OnEndDrag(FVector2D InitialPosition, FVector2D CurrentPosition, const FPointerEvent& MouseEvent)
 {
 	ICurveEditorKeyDragOperation::OnEndDrag(InitialPosition, CurrentPosition, MouseEvent);
+
+	for (const FKeyData& KeyData : KeysByCurve)
+	{
+		if (FCurveModel* Curve = CurveEditor->FindCurve(KeyData.CurveID))
+		{
+			Curve->SetKeyAttributes(KeyData.Handles, KeyData.LastDraggedAttributes, EPropertyChangeType::ValueSet);
+		}
+	}
+	
 	CurveEditor->SuppressBoundTransformUpdates(false);
 }
 
diff --git a/Engine/Source/Editor/CurveEditor/Private/DragOperations/CurveEditorDragOperation_Tangent.h b/Engine/Source/Editor/CurveEditor/Private/DragOperations/CurveEditorDragOperation_Tangent.h
--- a/Engine/Source/Editor/CurveEditor/Private/DragOperations/CurveEditorDragOperation_Tangent.h
+++ b/Engine/Source/Editor/CurveEditor/Private/DragOperations/CurveEditorDragOperation_Tangent.h
@@ -44,6 +44,8 @@ private:
 		TArray<ECurvePointType> PointTypes;
 		/** The key attributes for each of the above handles */
 		TArray<FKeyAttributes> Attributes;
+		/** Used in OnEndDrag to send final key updates */
+		TArray<FKeyAttributes> LastDraggedAttributes;
 	};
 
 	/** Key dragging data stored per-curve */
diff --git a/Engine/Source/Editor/CurveEditor/Private/RichCurveKeyProxy.h b/Engine/Source/Editor/CurveEditor/Private/RichCurveKeyProxy.h
--- a/Engine/Source/Editor/CurveEditor/Private/RichCurveKeyProxy.h
+++ b/Engine/Source/Editor/CurveEditor/Private/RichCurveKeyProxy.h
@@ -55,21 +55,10 @@ private:
 					Owner->Modify();
 
 					FRichCurveKey& ActualKey = RichCurve.GetKey(KeyHandle);
-
-					const float PreviousTime = ActualKey.Time;
-					const float NewTime      = Value.Time;
-
-					// Copy the key properties
-					ActualKey = Value;
-					ActualKey.Time = PreviousTime;
-					if (PreviousTime != NewTime)
-					{
-						RichCurve.SetKeyTime(KeyHandle, NewTime);
-			
-					}
-
+					const FKeyPosition KeyPosition(Value.Time, Value.Value);
+					Model->SetKeyPositions({KeyHandle}, {KeyPosition}, PropertyChangedEvent.ChangeType);
+				
 					Owner->PostEditChangeProperty(PropertyChangedEvent);
-
 					Model->OnCurveModified().Broadcast();
 				}
 			}
diff --git a/Engine/Source/Editor/Persona/Private/SAnimSequenceCurveEditor.cpp b/Engine/Source/Editor/Persona/Private/SAnimSequenceCurveEditor.cpp
--- a/Engine/Source/Editor/Persona/Private/SAnimSequenceCurveEditor.cpp
+++ b/Engine/Source/Editor/Persona/Private/SAnimSequenceCurveEditor.cpp
@@ -63,8 +63,47 @@ const FRichCurve& FRichCurveEditorModelNamed::GetReadOnlyRichCurve() const
 	return const_cast<FRichCurveEditorModelNamed*>(this)->GetRichCurve();
 }
 
-void FRichCurveEditorModelNamed::CurveHasChanged()
+void FRichCurveEditorModelNamed::SetKeyPositions(TArrayView<const FKeyHandle> InKeys, TArrayView<const FKeyPosition> InKeyPositions, EPropertyChangeType::Type ChangeType)
 {
+	const bool bInteractiveChange = ChangeType == EPropertyChangeType::Interactive;
+
+	// Open bracket in case this is an interactive change
+	if (bInteractiveChange && !InteractiveBracket.IsValid())
+	{
+		IAnimationDataController& Controller = AnimSequence->GetController();
+		InteractiveBracket = MakeUnique<IAnimationDataController::FScopedBracket>(Controller, LOCTEXT("SetKeyPositions", "Set Key Positions"));
+	}
+	
+	FRichCurveEditorModel::SetKeyPositions(InKeys, InKeyPositions, ChangeType);
+
+	// Close bracket, if open, in case this is was a non-interactive change
+	if (!bInteractiveChange && InteractiveBracket.IsValid())
+	{
+		InteractiveBracket.Reset();
+	}
+}
+
+void FRichCurveEditorModelNamed::SetKeyAttributes(TArrayView<const FKeyHandle> InKeys, TArrayView<const FKeyAttributes> InAttributes, EPropertyChangeType::Type ChangeType)
+{
+	const bool bInteractiveChange = ChangeType == EPropertyChangeType::Interactive;
+
+	// Open bracket in case this is an interactive change
+	if (bInteractiveChange && !InteractiveBracket.IsValid())
+	{
+		IAnimationDataController& Controller = AnimSequence->GetController();
+		InteractiveBracket = MakeUnique<IAnimationDataController::FScopedBracket>(Controller, LOCTEXT("SetKeyAttributes", "Set Key Attributes"));
+	}
+		
+	FRichCurveEditorModel::SetKeyAttributes(InKeys, InAttributes, ChangeType);
+	// Close bracket, if open, in case this is was a non-interactive change
+	if (!bInteractiveChange && InteractiveBracket.IsValid())
+	{
+		InteractiveBracket.Reset();
+	}
+}
+
+void FRichCurveEditorModelNamed::CurveHasChanged()
+{	
 	IAnimationDataController& Controller = AnimSequence->GetController();
 
 	switch (Type)
diff --git a/Engine/Source/Editor/Persona/Private/SAnimSequenceCurveEditor.h b/Engine/Source/Editor/Persona/Private/SAnimSequenceCurveEditor.h
--- a/Engine/Source/Editor/Persona/Private/SAnimSequenceCurveEditor.h
+++ b/Engine/Source/Editor/Persona/Private/SAnimSequenceCurveEditor.h
@@ -31,6 +31,9 @@ public:
 	virtual FRichCurve& GetRichCurve() override;
 	virtual const FRichCurve& GetReadOnlyRichCurve() const override;
 
+	virtual void SetKeyPositions(TArrayView<const FKeyHandle> InKeys, TArrayView<const FKeyPosition> InKeyPositions, EPropertyChangeType::Type ChangeType) override;
+	virtual void SetKeyAttributes(TArrayView<const FKeyHandle> InKeys, TArrayView<const FKeyAttributes> InAttributes, EPropertyChangeType::Type ChangeType = EPropertyChangeType::Unspecified) override;
+
 	void CurveHasChanged();
 	void OnModelHasChanged(const EAnimDataModelNotifyType& NotifyType, UAnimDataModel* Model, const FAnimDataModelNotifPayload& Payload);
 	void UpdateCachedCurve();
@@ -45,6 +48,8 @@ public:
 	UE::Anim::FAnimDataModelNotifyCollector NotifyCollector;
 	FRichCurve CachedCurve;
 	bool bCurveRemoved;
+
+	TUniquePtr<IAnimationDataController::FScopedBracket> InteractiveBracket;
 };
 
 class SAnimSequenceCurveEditor : public IAnimSequenceCurveEditor, public FEditorUndoClient
diff --git a/Engine/Source/Runtime/Engine/Private/Animation/AnimSequence.cpp b/Engine/Source/Runtime/Engine/Private/Animation/AnimSequence.cpp
--- a/Engine/Source/Runtime/Engine/Private/Animation/AnimSequence.cpp
+++ b/Engine/Source/Runtime/Engine/Private/Animation/AnimSequence.cpp
@@ -3967,7 +3967,23 @@ void UAnimSequence::EvaluateCurveData(FBlendedCurve& OutCurve, float CurrentTime
 	check(!bForceUseRawData || CanEvaluateRawAnimationData());
 	if (CanEvaluateRawAnimationData() && bEvaluateRawData)
 	{
+#if WITH_EDITOR
+		// Evaluate float curves from the AnimationData Model
+		if (OutCurve.NumValidCurveCount > 0)
+		{
+			for (auto CurveIter = DataModel->GetFloatCurves().CreateConstIterator(); CurveIter; ++CurveIter)
+			{
+				const FFloatCurve& Curve = *CurveIter;
+				if (OutCurve.IsEnabled(Curve.Name.UID))
+				{
+					const float Value = Curve.Evaluate(CurrentTime);
+					OutCurve.Set(Curve.Name.UID, Value);
+				}
+			}
+		}
+#else
 		Super::EvaluateCurveData(OutCurve, CurrentTime, bForceUseRawData);
+#endif
 	}
 	else if(IsCurveCompressedDataValid() && CompressedData.CurveCompressionCodec)
 	{
@@ -5403,10 +5419,11 @@ void UAnimSequence::OnModelModified(const EAnimDataModelNotifyType& NotifyType,
 		case EAnimDataModelNotifyType::CurveFlagsChanged:
 		case EAnimDataModelNotifyType::CurveScaled:
 		{
+			ClearCompressedCurveData();
+			UpdateRawDataGuid(RegenerateGUID);
+				
 			if (NotifyCollector.IsNotWithinBracket())
-			{
-				UpdateRawDataGuid(RegenerateGUID);
-				ClearCompressedCurveData();
+			{						
 				RecompressAnimationData();
 			}
 
