ControlRig: Added Animation Attribute Support

diff --git a/Engine/Plugins/Animation/ControlRig/Source/ControlRig/Private/AnimNode_ControlRigBase.cpp b/Engine/Plugins/Animation/ControlRig/Source/ControlRig/Private/AnimNode_ControlRigBase.cpp
--- a/Engine/Plugins/Animation/ControlRig/Source/ControlRig/Private/AnimNode_ControlRigBase.cpp
+++ b/Engine/Plugins/Animation/ControlRig/Source/ControlRig/Private/AnimNode_ControlRigBase.cpp
@@ -485,6 +485,10 @@ void FAnimNode_ControlRigBase::ExecuteControlRig(FPoseContext& InOutput)
 {
 	if (UControlRig* ControlRig = GetControlRig())
 	{
+		// temporarily give control rig access to the stack allocated attribute container
+		// control rig may have rig units that can add/get attributes to/from this container
+		UControlRig::FAnimAttributeContainerPtrScope AttributeScope(ControlRig, InOutput.CustomAttributes);
+		
 		// first update input to the system
 		UpdateInput(ControlRig, InOutput);
 
diff --git a/Engine/Plugins/Animation/ControlRig/Source/ControlRig/Private/ControlRig.cpp b/Engine/Plugins/Animation/ControlRig/Source/ControlRig/Private/ControlRig.cpp
--- a/Engine/Plugins/Animation/ControlRig/Source/ControlRig/Private/ControlRig.cpp
+++ b/Engine/Plugins/Animation/ControlRig/Source/ControlRig/Private/ControlRig.cpp
@@ -63,6 +63,9 @@ UControlRig::UControlRig(const FObjectInitializer& ObjectInitializer)
 	, ControlRigLog(nullptr)
 	, bEnableControlRigLogging(true)
 #endif
+#if WITH_EDITOR
+	, bEnableAnimAttributeTrace(false)
+#endif 
 	, DataSourceRegistry(nullptr)
 	, EventQueue()
 #if WITH_EDITOR
@@ -372,6 +375,20 @@ TArray<FRigVMExternalVariable> UControlRig::GetExternalVariablesImpl(bool bFallb
 	return ExternalVariables;
 }
 
+UControlRig::FAnimAttributeContainerPtrScope::FAnimAttributeContainerPtrScope(UControlRig* InControlRig,
+	UE::Anim::FStackAttributeContainer& InExternalContainer)
+{
+	ControlRig = InControlRig;
+	ControlRig->ExternalAnimAttributeContainer = &InExternalContainer;
+}
+
+UControlRig::FAnimAttributeContainerPtrScope::~FAnimAttributeContainerPtrScope()
+{
+	// control rig should not hold on to this container since it is stack allocated
+	// and should not be used outside of stack, see FPoseContext
+	ControlRig->ExternalAnimAttributeContainer = nullptr;
+}
+
 TArray<FRigVMExternalVariable> UControlRig::GetPublicVariables() const
 {
 	TArray<FRigVMExternalVariable> PublicVariables;
@@ -670,6 +687,8 @@ void UControlRig::Execute(const EControlRigState InState, const FName& InEventNa
 	DrawInterface.Reset();
 	Context.DrawInterface = &DrawInterface;
 
+	Context.AnimAttributeContainer = ExternalAnimAttributeContainer;
+
 	// draw container contains persistent draw instructions, 
 	// so we cannot call Reset(), which will clear them,
 	// instead, we re-initialize them from the CDO
@@ -881,7 +900,14 @@ void UControlRig::Execute(const EControlRigState InState, const FName& InEventNa
 			// Transform Overrride is generated using a Transient Control 
 			ApplyTransformOverrideForUserCreatedBones();
 #endif
-
+			
+#if WITH_EDITOR
+			if (bEnableAnimAttributeTrace && ExternalAnimAttributeContainer != nullptr)
+			{
+				InputAnimAttributeSnapshot.CopyFrom(*ExternalAnimAttributeContainer);
+			}
+#endif
+			
 			if (InState == EControlRigState::Update && InEventName == FRigUnit_BeginExecution::EventName)
 			{
 				if (PreForwardsSolveEvent.IsBound())
@@ -893,6 +919,13 @@ void UControlRig::Execute(const EControlRigState InState, const FName& InEventNa
 
 			ExecuteUnits(Context, InEventName);
 
+#if WITH_EDITOR
+			if (bEnableAnimAttributeTrace && ExternalAnimAttributeContainer != nullptr)
+			{
+				OutputAnimAttributeSnapshot.CopyFrom(*ExternalAnimAttributeContainer);
+			}
+#endif
+
 			if (InState == EControlRigState::Update && InEventName == FRigUnit_BeginExecution::EventName)
 			{
 				if (PostForwardsSolveEvent.IsBound())
diff --git a/Engine/Plugins/Animation/ControlRig/Source/ControlRig/Private/Units/Animation/RigUnit_AnimAttribute.cpp b/Engine/Plugins/Animation/ControlRig/Source/ControlRig/Private/Units/Animation/RigUnit_AnimAttribute.cpp
new file mode 100644
--- /dev/null
+++ b/Engine/Plugins/Animation/ControlRig/Source/ControlRig/Private/Units/Animation/RigUnit_AnimAttribute.cpp
@@ -0,0 +1,153 @@
+// Copyright Epic Games, Inc. All Rights Reserved.
+
+#include "Units/Animation/RigUnit_AnimAttribute.h"
+
+#include "Animation/BuiltInAttributeTypes.h"
+#include "Units/RigUnitContext.h"
+#include "Components/SkeletalMeshComponent.h"
+
+template<typename T>
+class TAnimAttributeType;
+
+template<>
+class TAnimAttributeType<int>
+{
+public:
+	using Type = FIntegerAnimationAttribute;
+};
+
+template<>
+class TAnimAttributeType<float>
+{
+public:
+	using Type = FFloatAnimationAttribute;
+};
+template<>
+class TAnimAttributeType<FTransform>
+{
+public:
+	using Type = FTransformAnimationAttribute;
+};
+
+template<typename T>
+FORCEINLINE_DEBUGGABLE static T* GetAnimAttributeValue(
+	bool bAddIfNotFound,
+	const FRigUnitContext& Context,
+	const FName& Name,
+	const FName& BoneName,
+	FName& CachedBoneName,
+	int32& CachedBoneIndex)
+{
+	if (Name.IsNone())
+	{
+		return nullptr;
+	}
+
+	if (!Context.AnimAttributeContainer)
+	{
+		return nullptr;
+	}
+	
+	const USkeletalMeshComponent* OwningComponent = Cast<USkeletalMeshComponent>(Context.OwningComponent);
+
+	if (!OwningComponent ||
+		!OwningComponent->SkeletalMesh)
+	{
+		return nullptr;
+	}
+
+	if (BoneName == NAME_None)
+	{
+		// default to use root bone
+		CachedBoneIndex = 0;
+	}
+	else
+	{
+		// Invalidate cache if input changed
+		if (CachedBoneName != BoneName)
+		{
+			CachedBoneIndex = OwningComponent->SkeletalMesh->GetRefSkeleton().FindBoneIndex(BoneName);
+		}
+	}
+	
+	CachedBoneName = BoneName;
+
+	if (CachedBoneIndex != INDEX_NONE)
+	{
+		const UE::Anim::FAttributeId Id = {Name, FCompactPoseBoneIndex(CachedBoneIndex)} ;
+		typename TAnimAttributeType<T>::Type* Attribute = bAddIfNotFound ?
+			Context.AnimAttributeContainer->FindOrAdd<typename TAnimAttributeType<T>::Type>(Id) :
+			Context.AnimAttributeContainer->Find<typename TAnimAttributeType<T>::Type>(Id);
+		if (Attribute)
+		{
+			return &Attribute->Value;
+		}
+	}
+
+	return nullptr;
+}
+
+
+FRigUnit_SetAnimAttribute_Integer_Execute()
+{
+	DECLARE_SCOPE_HIERARCHICAL_COUNTER_RIGUNIT()
+	
+	int* ValuePtr = GetAnimAttributeValue<int>(true ,Context, Name, BoneName, CachedBoneName, CachedBoneIndex);
+	if (ValuePtr)
+	{
+		*ValuePtr = Value;
+	}
+}
+
+FRigUnit_SetAnimAttribute_Float_Execute()
+{
+	DECLARE_SCOPE_HIERARCHICAL_COUNTER_RIGUNIT()
+
+	float* ValuePtr = GetAnimAttributeValue<float>(true, Context, Name, BoneName, CachedBoneName, CachedBoneIndex);
+	if (ValuePtr)
+	{
+		*ValuePtr = Value;
+	}
+}
+
+FRigUnit_SetAnimAttribute_Transform_Execute()
+{
+    DECLARE_SCOPE_HIERARCHICAL_COUNTER_RIGUNIT()
+
+	FTransform* ValuePtr = GetAnimAttributeValue<FTransform>(true, Context, Name, BoneName, CachedBoneName, CachedBoneIndex);
+	if (ValuePtr)
+	{
+		*ValuePtr = Value;
+	}
+}
+
+FRigUnit_GetAnimAttribute_Integer_Execute()
+{
+	DECLARE_SCOPE_HIERARCHICAL_COUNTER_RIGUNIT()
+	
+	const int* ValuePtr = GetAnimAttributeValue<int>(false ,Context, Name, BoneName, CachedBoneName, CachedBoneIndex);
+	bWasFound = ValuePtr ? true : false;
+	Value = ValuePtr ? *ValuePtr : FallbackValue;
+}
+
+FRigUnit_GetAnimAttribute_Float_Execute()
+{
+	DECLARE_SCOPE_HIERARCHICAL_COUNTER_RIGUNIT()
+
+	const float* ValuePtr = GetAnimAttributeValue<float>(false ,Context, Name, BoneName, CachedBoneName, CachedBoneIndex);
+	bWasFound = ValuePtr ? true : false;
+	Value = ValuePtr ? *ValuePtr : FallbackValue;
+}
+
+FRigUnit_GetAnimAttribute_Transform_Execute()
+{
+	DECLARE_SCOPE_HIERARCHICAL_COUNTER_RIGUNIT()
+
+	const FTransform* ValuePtr = GetAnimAttributeValue<FTransform>(false, Context, Name, BoneName, CachedBoneName, CachedBoneIndex);
+	bWasFound = ValuePtr ? true : false;
+	Value = ValuePtr ? *ValuePtr : FallbackValue;
+}
+
+
+
+
diff --git a/Engine/Plugins/Animation/ControlRig/Source/ControlRig/Private/Units/Animation/RigUnit_AnimAttribute.h b/Engine/Plugins/Animation/ControlRig/Source/ControlRig/Private/Units/Animation/RigUnit_AnimAttribute.h
new file mode 100644
--- /dev/null
+++ b/Engine/Plugins/Animation/ControlRig/Source/ControlRig/Private/Units/Animation/RigUnit_AnimAttribute.h
@@ -0,0 +1,259 @@
+// Copyright Epic Games, Inc. All Rights Reserved.
+
+#pragma once
+#include "Units/RigUnit.h"
+#include "RigUnit_AnimAttribute.generated.h"
+
+/**
+ * Sets the value of an animation attribute with the matching names.
+ * If the attribute was not found, a new attribute is created.
+ *
+ * Animation Attributes allow dynamically added data to flow from
+ * one Anim Node to other Anim Nodes downstream in the Anim Graph.
+ */
+USTRUCT(meta=(Abstract, Category="Animation Attribute", TemplateName="Set Attribute", NodeColor = "0.0 0.36470600962638855 1.0", Varying))
+struct FRigUnit_SetAnimAttributeBase : public FRigUnitMutable
+{
+	GENERATED_BODY()
+};
+
+
+USTRUCT(meta=(DisplayName="Set Attribute - Integer"))
+struct CONTROLRIG_API FRigUnit_SetAnimAttribute_Integer: public FRigUnit_SetAnimAttributeBase
+{
+	GENERATED_BODY()
+	
+	FRigUnit_SetAnimAttribute_Integer()
+	{
+		Name = TEXT("NewAttributeName");
+		BoneName = NAME_None;
+		Value = 0;
+		CachedBoneName = NAME_None;
+		CachedBoneIndex = 0;
+	}
+
+	RIGVM_METHOD()
+	virtual void Execute(const FRigUnitContext& Context) override;
+
+	UPROPERTY(meta=(Input))
+	FName Name;
+	
+	UPROPERTY(meta=(Input, CustomWidget="BoneName"))
+	FName BoneName;
+
+	UPROPERTY(meta=(Input))
+	int32 Value;
+
+	UPROPERTY()
+	FName CachedBoneName;
+	
+	UPROPERTY()
+	int32 CachedBoneIndex;
+};
+
+
+USTRUCT(meta=(DisplayName="Set Attribute - Float"))
+struct CONTROLRIG_API FRigUnit_SetAnimAttribute_Float: public FRigUnit_SetAnimAttributeBase
+{
+	GENERATED_BODY()
+	
+	FRigUnit_SetAnimAttribute_Float()
+	{
+		Name = TEXT("NewAttributeName");
+		BoneName = NAME_None;
+		Value = 0.f;
+		CachedBoneName = NAME_None;
+		CachedBoneIndex = 0;
+	}
+
+	RIGVM_METHOD()
+	virtual void Execute(const FRigUnitContext& Context) override;
+
+	UPROPERTY(meta=(Input))
+	FName Name;
+	
+	UPROPERTY(meta=(Input, CustomWidget="BoneName"))
+	FName BoneName;
+
+	UPROPERTY(meta=(Input))
+	float Value;
+
+	UPROPERTY()
+	FName CachedBoneName;
+	
+	UPROPERTY()
+	int32 CachedBoneIndex;
+};
+
+
+USTRUCT(meta=(DisplayName="Set Attribute - Transform"))
+struct CONTROLRIG_API FRigUnit_SetAnimAttribute_Transform: public FRigUnit_SetAnimAttributeBase
+{
+	GENERATED_BODY()
+	
+	FRigUnit_SetAnimAttribute_Transform()
+	{
+		Name = TEXT("NewAttributeName");
+		BoneName = NAME_None;
+		Value = FTransform::Identity;
+		CachedBoneName = NAME_None;
+		CachedBoneIndex = 0;
+	}
+
+	RIGVM_METHOD()
+	virtual void Execute(const FRigUnitContext& Context) override;
+
+	UPROPERTY(meta=(Input))
+	FName Name;
+	
+	UPROPERTY(meta=(Input, CustomWidget="BoneName"))
+	FName BoneName;
+
+	UPROPERTY(meta=(Input))
+	FTransform Value;
+
+	UPROPERTY()
+	FName CachedBoneName;
+	
+	UPROPERTY()
+	int32 CachedBoneIndex;
+};
+
+/**
+ * Gets the value of an animation attribute with the matching names.
+ * If the attribute was not found, the fallback value is outputted.
+ *
+ * Animation Attributes allow dynamically added data to flow from
+ * one Anim Node to other Anim Nodes downstream in the Anim Graph.
+ */
+USTRUCT(meta=(Abstract, Category="Animation Attribute", TemplateName="Get Attribute", NodeColor="0.462745, 1,0, 0.329412", Varying))
+struct FRigUnit_GetAnimAttributeBase : public FRigUnit
+{
+	GENERATED_BODY()
+};
+
+
+USTRUCT(meta=(DisplayName="Get Attribute - Integer"))
+struct CONTROLRIG_API FRigUnit_GetAnimAttribute_Integer: public FRigUnit_GetAnimAttributeBase
+{
+	GENERATED_BODY()
+	
+	FRigUnit_GetAnimAttribute_Integer()
+	{
+		Name = TEXT("AttributeName");
+		BoneName = NAME_None;
+		FallbackValue = 0;
+		Value = 0;
+		bWasFound = false;
+		CachedBoneName = NAME_None;
+		CachedBoneIndex = 0;
+	}
+
+	RIGVM_METHOD()
+	virtual void Execute(const FRigUnitContext& Context) override;
+
+	UPROPERTY(meta=(Input))
+	FName Name;
+	
+	UPROPERTY(meta=(Input, CustomWidget="BoneName"))
+	FName BoneName;
+
+	UPROPERTY(meta=(Input))
+	int32 FallbackValue;
+	
+	UPROPERTY(meta=(Output))
+	int32 Value;
+
+	UPROPERTY(meta=(Output))
+	bool bWasFound;
+
+	UPROPERTY()
+	FName CachedBoneName;
+	
+	UPROPERTY()
+	int32 CachedBoneIndex;
+};
+
+
+USTRUCT(meta=(DisplayName="Get Attribute - Float"))
+struct CONTROLRIG_API FRigUnit_GetAnimAttribute_Float: public FRigUnit_GetAnimAttributeBase
+{
+	GENERATED_BODY()
+	
+	FRigUnit_GetAnimAttribute_Float()
+	{
+		Name = TEXT("AttributeName");
+		BoneName = NAME_None;
+		FallbackValue = 0.f;
+		Value = 0.f;
+		bWasFound = false;
+		CachedBoneName = NAME_None;
+		CachedBoneIndex = 0;
+	}
+
+	RIGVM_METHOD()
+	virtual void Execute(const FRigUnitContext& Context) override;
+
+	UPROPERTY(meta=(Input))
+	FName Name;
+	
+	UPROPERTY(meta=(Input, CustomWidget="BoneName"))
+	FName BoneName;
+
+	UPROPERTY(meta=(Input))
+	float FallbackValue;
+	
+	UPROPERTY(meta=(Output))
+	float Value;
+
+	UPROPERTY(meta=(Output))
+	bool bWasFound;
+
+	UPROPERTY()
+	FName CachedBoneName;
+	
+	UPROPERTY()
+	int32 CachedBoneIndex;
+};
+
+
+USTRUCT(meta=(DisplayName="Get Attribute - Transform"))
+struct CONTROLRIG_API FRigUnit_GetAnimAttribute_Transform: public FRigUnit_GetAnimAttributeBase
+{
+	GENERATED_BODY()
+	
+	FRigUnit_GetAnimAttribute_Transform()
+	{
+		Name = TEXT("AttributeName");
+		BoneName = NAME_None;
+		FallbackValue = FTransform::Identity;
+		Value = FTransform::Identity;
+		bWasFound = false;
+		CachedBoneName = NAME_None;
+		CachedBoneIndex = 0;
+	}
+
+	RIGVM_METHOD()
+	virtual void Execute(const FRigUnitContext& Context) override;
+
+	UPROPERTY(meta=(Input))
+	FName Name;
+	
+	UPROPERTY(meta=(Input, CustomWidget="BoneName"))
+	FName BoneName;
+
+	UPROPERTY(meta=(Input))
+	FTransform FallbackValue;
+	
+	UPROPERTY(meta=(Output))
+	FTransform Value;
+
+	UPROPERTY(meta=(Output))
+	bool bWasFound;
+
+	UPROPERTY()
+	FName CachedBoneName;
+	
+	UPROPERTY()
+	int32 CachedBoneIndex;
+};
\ No newline at end of file
diff --git a/Engine/Plugins/Animation/ControlRig/Source/ControlRig/Public/ControlRig.h b/Engine/Plugins/Animation/ControlRig/Source/ControlRig/Public/ControlRig.h
--- a/Engine/Plugins/Animation/ControlRig/Source/ControlRig/Public/ControlRig.h
+++ b/Engine/Plugins/Animation/ControlRig/Source/ControlRig/Public/ControlRig.h
@@ -21,6 +21,7 @@
 #include "Interfaces/Interface_AssetUserData.h"
 #include "Animation/AnimInstance.h"
 #include "Animation/AnimInstanceProxy.h"
+#include "Animation/AttributesRuntime.h"
 
 #if WITH_EDITOR
 #include "RigVMModel/RigVMPin.h"
@@ -516,9 +517,31 @@ private:
 		}
 		return nullptr;
 	}
+public:
+	class FAnimAttributeContainerPtrScope
+	{
+	public:
+		FAnimAttributeContainerPtrScope(UControlRig* InControlRig, UE::Anim::FStackAttributeContainer& InExternalContainer);
+		~FAnimAttributeContainerPtrScope();
 
+		UControlRig* ControlRig;
+	};
+	
 private:
+	UE::Anim::FStackAttributeContainer* ExternalAnimAttributeContainer;
 
+#if WITH_EDITOR
+	void SetEnableAnimAttributeTrace(bool bInEnable)
+	{
+		bEnableAnimAttributeTrace = bInEnable;
+	};
+	
+	bool bEnableAnimAttributeTrace;
+	
+	UE::Anim::FHeapAttributeContainer InputAnimAttributeSnapshot;
+	UE::Anim::FHeapAttributeContainer OutputAnimAttributeSnapshot;
+#endif
+	
 	UPROPERTY()
 	FControlRigDrawContainer DrawContainer;
 
@@ -874,6 +897,7 @@ private:
 	friend class FControlRigEditor;
 	friend class SRigCurveContainer;
 	friend class SRigHierarchy;
+	friend class SControlRigAnimAttributeView;
 	friend class UEngineTestControlRig;
  	friend class FControlRigEditMode;
 	friend class UControlRigBlueprint;
diff --git a/Engine/Plugins/Animation/ControlRig/Source/ControlRig/Public/Units/RigUnitContext.h b/Engine/Plugins/Animation/ControlRig/Source/ControlRig/Public/Units/RigUnitContext.h
--- a/Engine/Plugins/Animation/ControlRig/Source/ControlRig/Public/Units/RigUnitContext.h
+++ b/Engine/Plugins/Animation/ControlRig/Source/ControlRig/Public/Units/RigUnitContext.h
@@ -8,6 +8,7 @@
 #include "Rigs/RigNameCache.h"
 #include "ControlRigLog.h"
 #include "AnimationDataSource.h"
+#include "Animation/AttributesRuntime.h"
 #include "Drawing/ControlRigDrawInterface.h"
 #include "GameFramework/Actor.h"
 #include "Components/SceneComponent.h"
@@ -41,7 +42,8 @@ struct FRigUnitContext
 {
 	/** default constructor */
 	FRigUnitContext()
-		: DrawInterface(nullptr)
+		: AnimAttributeContainer(nullptr)
+		, DrawInterface(nullptr)
 		, DrawContainer(nullptr)
 		, DataSourceRegistry(nullptr)
 		, DeltaTime(0.f)
@@ -61,6 +63,9 @@ struct FRigUnitContext
 	{
 	}
 
+	/** An external anim attribute container */
+	UE::Anim::FStackAttributeContainer* AnimAttributeContainer;
+	
 	/** The draw interface for the units to use */
 	FControlRigDrawInterface* DrawInterface;
 
diff --git a/Engine/Plugins/Animation/ControlRig/Source/ControlRigEditor/Private/Editor/ControlRigEditorMode.cpp b/Engine/Plugins/Animation/ControlRig/Source/ControlRigEditor/Private/Editor/ControlRigEditorMode.cpp
--- a/Engine/Plugins/Animation/ControlRig/Source/ControlRigEditor/Private/Editor/ControlRigEditorMode.cpp
+++ b/Engine/Plugins/Animation/ControlRig/Source/ControlRigEditor/Private/Editor/ControlRigEditorMode.cpp
@@ -11,6 +11,7 @@
 #include "RigCurveContainerTabSummoner.h"
 #include "RigInfluenceMapTabSummoner.h"
 #include "RigValidationTabSummoner.h"
+#include "RigAnimAttributeTabSummoner.h"
 #include "ToolMenus.h"
 
 FControlRigEditorMode::FControlRigEditorMode(const TSharedRef<FControlRigEditor>& InControlRigEditor)
@@ -24,6 +25,7 @@ FControlRigEditorMode::FControlRigEditorMode(const TSharedRef<FControlRigEditor>
 	TabFactories.RegisterFactory(MakeShared<FRigCurveContainerTabSummoner>(InControlRigEditor));
 	//TabFactories.RegisterFactory(MakeShared<FRigInfluenceMapTabSummoner>(InControlRigEditor));
 	TabFactories.RegisterFactory(MakeShared<FRigValidationTabSummoner>(InControlRigEditor));
+	TabFactories.RegisterFactory(MakeShared<FRigAnimAttributeTabSummoner>(InControlRigEditor));
 
 	FPersonaModule& PersonaModule = FModuleManager::LoadModuleChecked<FPersonaModule>("Persona");
 
@@ -110,6 +112,7 @@ FControlRigEditorMode::FControlRigEditorMode(const TSharedRef<FControlRigEditor>
 						->SetSizeCoefficient(1.f)
 						->AddTab(FBlueprintEditorTabs::DetailsID, ETabState::OpenedTab)
 						->AddTab(FPersonaTabs::AdvancedPreviewSceneSettingsID, ETabState::OpenedTab)
+						->AddTab(FRigAnimAttributeTabSummoner::TabID, ETabState::OpenedTab)
 						->SetForegroundTab(FBlueprintEditorTabs::DetailsID)
 					)
 				)
diff --git a/Engine/Plugins/Animation/ControlRig/Source/ControlRigEditor/Private/Editor/RigAnimAttributeTabSummoner.cpp b/Engine/Plugins/Animation/ControlRig/Source/ControlRigEditor/Private/Editor/RigAnimAttributeTabSummoner.cpp
new file mode 100644
--- /dev/null
+++ b/Engine/Plugins/Animation/ControlRig/Source/ControlRigEditor/Private/Editor/RigAnimAttributeTabSummoner.cpp
@@ -0,0 +1,29 @@
+// Copyright Epic Games, Inc. All Rights Reserved.
+
+#include "RigAnimAttributeTabSummoner.h"
+#include "SControlRigAnimAttributeView.h"
+
+#include "ControlRigEditor.h"
+
+#define LOCTEXT_NAMESPACE "RigAnimAttributeTabSummoner"
+
+const FName FRigAnimAttributeTabSummoner::TabID(TEXT("RigAnimAttribute"));
+
+FRigAnimAttributeTabSummoner::FRigAnimAttributeTabSummoner(const TSharedRef<FControlRigEditor>& InControlRigEditor)
+	: FWorkflowTabFactory(TabID, InControlRigEditor)
+	, ControlRigEditor(InControlRigEditor)
+{
+	TabLabel = LOCTEXT("RigAnimAttributeTabLabel", "Animation Attributes");
+
+	bIsSingleton = true;
+
+	ViewMenuDescription = LOCTEXT("RigAnimAttribute_ViewMenu_Desc", "Animation Attribute");
+	ViewMenuTooltip = LOCTEXT("RigAnimAttribute_ViewMenu_ToolTip", "Show the Animation Attribute tab");
+}
+
+TSharedRef<SWidget> FRigAnimAttributeTabSummoner::CreateTabBody(const FWorkflowTabSpawnInfo& Info) const
+{
+	return SNew(SControlRigAnimAttributeView, ControlRigEditor.Pin().ToSharedRef());
+}
+
+#undef LOCTEXT_NAMESPACE 
diff --git a/Engine/Plugins/Animation/ControlRig/Source/ControlRigEditor/Private/Editor/RigAnimAttributeTabSummoner.h b/Engine/Plugins/Animation/ControlRig/Source/ControlRigEditor/Private/Editor/RigAnimAttributeTabSummoner.h
new file mode 100644
--- /dev/null
+++ b/Engine/Plugins/Animation/ControlRig/Source/ControlRigEditor/Private/Editor/RigAnimAttributeTabSummoner.h
@@ -0,0 +1,22 @@
+// Copyright Epic Games, Inc. All Rights Reserved.
+
+#pragma once
+
+#include "CoreMinimal.h"
+#include "WorkflowOrientedApp/WorkflowTabFactory.h"
+
+class FControlRigEditor;
+
+struct FRigAnimAttributeTabSummoner : public FWorkflowTabFactory
+{
+public:
+	static const FName TabID;
+	
+public:
+	FRigAnimAttributeTabSummoner(const TSharedRef<FControlRigEditor>& InControlRigEditor);
+	
+	virtual TSharedRef<SWidget> CreateTabBody(const FWorkflowTabSpawnInfo& Info) const override;
+	
+protected:
+	TWeakPtr<FControlRigEditor> ControlRigEditor;
+};
diff --git a/Engine/Plugins/Animation/ControlRig/Source/ControlRigEditor/Private/Editor/SControlRigAnimAttributeView.cpp b/Engine/Plugins/Animation/ControlRig/Source/ControlRigEditor/Private/Editor/SControlRigAnimAttributeView.cpp
new file mode 100644
--- /dev/null
+++ b/Engine/Plugins/Animation/ControlRig/Source/ControlRigEditor/Private/Editor/SControlRigAnimAttributeView.cpp
@@ -0,0 +1,116 @@
+// Copyright Epic Games, Inc. All Rights Reserved.
+
+#include "SControlRigAnimAttributeView.h"
+#include "ControlRigEditor.h"
+#include "ControlRigBlueprint.h"
+#include "Engine/Console.h"
+#include "ControlRig/Private/Units/Execution/RigUnit_BeginExecution.h"
+#include "Persona/Public/SAnimAttributeView.h"
+
+#define LOCTEXT_NAMESPACE "SControlRigAnimAttributeView"
+
+static const FName IncomingSnapshotName(TEXT("Incoming Anim Attributes"));
+static const FName OutgoingSnapshotName(TEXT("Outgoing Anim Attributes"));
+static const TArray<FName> SnapshotNames = {IncomingSnapshotName, OutgoingSnapshotName};
+
+SControlRigAnimAttributeView::~SControlRigAnimAttributeView()
+{
+	StopObservingCurrentControlRig();
+	
+	ControlRigBlueprint->OnSetObjectBeingDebugged().RemoveAll(this);
+}
+
+void SControlRigAnimAttributeView::Construct(const FArguments& InArgs, TSharedRef<FControlRigEditor> InControlRigEditor)
+{
+	SAssignNew(AttributeView, SAnimAttributeView);
+	
+	ChildSlot
+	[
+		AttributeView.ToSharedRef()
+	];
+
+	ControlRigEditor = InControlRigEditor;
+	ControlRigBlueprint = ControlRigEditor.Pin()->GetControlRigBlueprint();
+	ControlRigBlueprint->OnSetObjectBeingDebugged().AddRaw(this, &SControlRigAnimAttributeView::HandleSetObjectBeingDebugged);
+
+	if (UControlRig* ControlRig = Cast<UControlRig>(ControlRigBlueprint->GetObjectBeingDebugged()))
+	{
+		StartObservingNewControlRig(ControlRig);
+	}
+}
+
+void SControlRigAnimAttributeView::HandleSetObjectBeingDebugged(UObject* InObject)
+{
+	if(ControlRigBeingDebuggedPtr.Get() == InObject)
+	{
+		return;
+	}
+
+	StopObservingCurrentControlRig();
+	
+	if(UControlRig* ControlRig = Cast<UControlRig>(InObject))
+	{
+		StartObservingNewControlRig(ControlRig);
+	}
+}
+
+
+void SControlRigAnimAttributeView::StartObservingNewControlRig(UControlRig* InControlRig) 
+{
+	if (InControlRig)
+	{
+		InControlRig->OnPostForwardsSolve_AnyThread().RemoveAll(this);
+		InControlRig->OnPostForwardsSolve_AnyThread().AddRaw(this, &SControlRigAnimAttributeView::HandleControlRigPostForwardSolve);
+		InControlRig->SetEnableAnimAttributeTrace(true);
+	}
+	
+	ControlRigBeingDebuggedPtr = InControlRig;
+}
+
+void SControlRigAnimAttributeView::StopObservingCurrentControlRig()
+{
+	if(ControlRigBeingDebuggedPtr.IsValid())
+	{
+		if(UControlRig* ControlRig = ControlRigBeingDebuggedPtr.Get())
+		{
+			if(!ControlRig->HasAnyFlags(RF_BeginDestroyed))
+			{
+				ControlRig->OnPostForwardsSolve_AnyThread().RemoveAll(this);
+				ControlRig->SetEnableAnimAttributeTrace(false);
+			}
+		}
+	}
+	
+	ControlRigBeingDebuggedPtr.Reset();
+}
+
+void SControlRigAnimAttributeView::HandleControlRigPostForwardSolve(
+	UControlRig* InControlRig,
+	const EControlRigState InState,
+	const FName& InEventName) const
+{
+	if (IsValid(InControlRig) && ensure(InControlRig == ControlRigBeingDebuggedPtr))
+	{
+		if (InState == EControlRigState::Update && InEventName == FRigUnit_BeginExecution::EventName)
+		{
+			if (const USkeletalMeshComponent* SkeletalMeshComponent = Cast<USkeletalMeshComponent>(InControlRig->OuterSceneComponent.Get()))
+			{
+				const UE::Anim::FHeapAttributeContainer& InputSnapshot = InControlRig->InputAnimAttributeSnapshot;
+				const UE::Anim::FHeapAttributeContainer& OutputSnapshot = InControlRig->OutputAnimAttributeSnapshot;
+
+				const TArray<TTuple<FName, const UE::Anim::FHeapAttributeContainer&>> Snapshots =
+				{
+					{TEXT("Input") , InputSnapshot  },
+					{TEXT("Output"), OutputSnapshot }
+				};
+				
+				AttributeView->DisplayNewAttributeContainerSnapshots(Snapshots, SkeletalMeshComponent);
+				return;
+			}
+		}
+	}
+
+	AttributeView->ClearListView();
+}
+
+#undef LOCTEXT_NAMESPACE
diff --git a/Engine/Plugins/Animation/ControlRig/Source/ControlRigEditor/Private/Editor/SControlRigAnimAttributeView.h b/Engine/Plugins/Animation/ControlRig/Source/ControlRigEditor/Private/Editor/SControlRigAnimAttributeView.h
new file mode 100644
--- /dev/null
+++ b/Engine/Plugins/Animation/ControlRig/Source/ControlRigEditor/Private/Editor/SControlRigAnimAttributeView.h
@@ -0,0 +1,38 @@
+// Copyright Epic Games, Inc. All Rights Reserved.
+
+#pragma once
+
+#include "CoreMinimal.h"
+#include "Widgets/SCompoundWidget.h"
+
+class FControlRigEditor;
+class UControlRigBlueprint;
+class UControlRig;
+enum class EControlRigState : uint8;
+
+class SAnimAttributeView;
+
+
+class SControlRigAnimAttributeView : public SCompoundWidget
+{
+	SLATE_BEGIN_ARGS( SControlRigAnimAttributeView )
+	{}
+	
+	SLATE_END_ARGS()
+
+	~SControlRigAnimAttributeView();
+	
+	void Construct( const FArguments& InArgs, TSharedRef<FControlRigEditor> InControlRigEditor);
+
+	void HandleSetObjectBeingDebugged(UObject* InObject);
+	
+	void StartObservingNewControlRig(UControlRig* InControlRig);
+	void StopObservingCurrentControlRig();
+
+	void HandleControlRigPostForwardSolve(UControlRig* InControlRig, const EControlRigState InState, const FName& InEventName) const;
+	
+	TWeakPtr<FControlRigEditor> ControlRigEditor;
+	TWeakObjectPtr<UControlRigBlueprint> ControlRigBlueprint;
+	TWeakObjectPtr<UControlRig> ControlRigBeingDebuggedPtr;
+	TSharedPtr<SAnimAttributeView> AttributeView;
+};
diff --git a/Engine/Source/Editor/Persona/Private/SAnimAttributeView.cpp b/Engine/Source/Editor/Persona/Private/SAnimAttributeView.cpp
new file mode 100644
--- /dev/null
+++ b/Engine/Source/Editor/Persona/Private/SAnimAttributeView.cpp
@@ -0,0 +1,657 @@
+// Copyright Epic Games, Inc. All Rights Reserved.
+
+#include "SAnimAttributeView.h"
+#include "Animation/BuiltInAttributeTypes.h"
+#include "Components/SkeletalMeshComponent.h"
+#include "Modules/ModuleManager.h"
+#include "PropertyEditorModule.h"
+#include "IStructureDetailsView.h"
+#include "Widgets/Input/SSearchBox.h"
+#include "Widgets/Layout/SExpandableArea.h"
+#include "Framework/MultiBox/MultiBoxBuilder.h"
+
+#define LOCTEXT_NAMESPACE "SAnimAttributeView"
+
+static const FName ColumnId_AnimAttributeName( TEXT("AttributeName") );
+static const FName ColumnId_AnimAttributeBoneName( TEXT("BoneName") );
+static const FName ColumnId_AnimAttributeTypeName( TEXT("TypeName") );
+static const FName ColumnId_AnimAttributeSnapshotName( TEXT("SnapshotName") );
+
+static const int32 MinimalNumColumns = 3;
+
+static FName GetAnimAttributeTypeName(const UScriptStruct* InType)
+{
+	return *(InType->GetName().Replace(TEXT("AnimationAttribute"), TEXT("")));
+}
+
+TSharedRef<IStructureDetailsView> SAnimAttributeView::CreateValueViewWidget()
+{
+	FDetailsViewArgs DetailsViewArgs;
+	{
+		DetailsViewArgs.bLockable = false;
+		DetailsViewArgs.bUpdatesFromSelection = false;
+		DetailsViewArgs.bAllowSearch = false;
+		DetailsViewArgs.NameAreaSettings = FDetailsViewArgs::ENameAreaSettings::HideNameArea;
+		DetailsViewArgs.bAllowMultipleTopLevelObjects = false;
+		DetailsViewArgs.bShowPropertyMatrixButton = false;
+	}
+
+	const FStructureDetailsViewArgs StructureViewArgs;
+	
+	FPropertyEditorModule& PropertyEditorModule = FModuleManager::Get().LoadModuleChecked<FPropertyEditorModule>("PropertyEditor");
+	TSharedRef<IStructureDetailsView> ValueView = PropertyEditorModule.CreateStructureDetailView(DetailsViewArgs, StructureViewArgs, nullptr);
+	ValueView->GetDetailsView()->SetIsPropertyEditingEnabledDelegate(
+		FIsPropertyEditingEnabled::CreateLambda( []() { return false; }));
+
+	return ValueView;
+}
+
+TSharedRef<ITableRow> SAnimAttributeView::MakeTableRowWidget(
+	TSharedPtr<FAnimAttributeEntry> InItem,
+	const TSharedRef<STableViewBase>& InOwnerTable)
+{
+	return InItem->MakeTableRowWidget(InOwnerTable);
+}
+
+
+TSharedRef<FAnimAttributeEntry> FAnimAttributeEntry::MakeEntry(
+	const FAnimationAttributeIdentifier& InIdentifier,
+	const FName& InSnapshotDisplayName)
+{
+	return MakeShareable(new FAnimAttributeEntry(InIdentifier, InSnapshotDisplayName)); 
+}
+
+FName SAnimAttributeView::GetSnapshotColumnDisplayName(const TArray<FName>& InSnapshotNames)
+{
+	if (InSnapshotNames.Num() == 0)
+	{
+		return NAME_None;
+	}
+	
+	if (InSnapshotNames.Num() == 1)
+	{
+		return InSnapshotNames[0];
+	}
+	return *(InSnapshotNames[0].ToString() + " - " + InSnapshotNames.Last().ToString());
+}
+
+
+FAnimAttributeEntry::FAnimAttributeEntry(const FAnimationAttributeIdentifier& InIdentifier, const FName& InSnapshotDisplayName)
+{
+	Identifier = InIdentifier;
+
+	SnapshotDisplayName = InSnapshotDisplayName;
+	
+	CachedTypeName = GetAnimAttributeTypeName(InIdentifier.GetType());
+}
+
+TSharedRef<ITableRow> FAnimAttributeEntry::MakeTableRowWidget(const TSharedRef<STableViewBase>& InOwnerTable)
+{
+	return SNew(SAnimAttributeEntry, InOwnerTable, SharedThis(this));
+}
+
+FName FAnimAttributeEntry::GetDisplayName() const
+{
+	return GetName();
+}
+
+void SAnimAttributeEntry::Construct(const FArguments& InArgs, const TSharedRef<STableViewBase>& InOwnerTable, TSharedRef<FAnimAttributeEntry> InEntry)
+{
+	Entry = InEntry;
+	SMultiColumnTableRow< TSharedPtr<FAnimAttributeEntry> >::Construct(
+		FSuperRowType::FArguments(),
+		InOwnerTable );
+}
+
+TSharedRef<SWidget> SAnimAttributeEntry::GenerateWidgetForColumn(const FName& ColumnName)
+{
+	if (ColumnName == ColumnId_AnimAttributeName)
+	{
+		return
+			SNew(SVerticalBox)
+			+ SVerticalBox::Slot()
+			.AutoHeight()
+			.Padding(4)
+			.VAlign(VAlign_Center)
+			[
+				SNew(STextBlock)
+				.Text(this, &SAnimAttributeEntry::GetEntryName)	
+			];
+	}
+	else if (ColumnName == ColumnId_AnimAttributeBoneName)
+	{
+		return
+			SNew(SVerticalBox)
+			+ SVerticalBox::Slot()
+			.AutoHeight()
+			.Padding(4)
+			.VAlign(VAlign_Center)
+			[
+				SNew(STextBlock)
+				.Text(this, &SAnimAttributeEntry::GetEntryBoneName)	
+			];
+	}
+	else if (ColumnName == ColumnId_AnimAttributeTypeName)
+	{
+		return
+			SNew(SVerticalBox)
+			+ SVerticalBox::Slot()
+			.AutoHeight()
+			.Padding(4)
+			.VAlign(VAlign_Center)
+			[
+				SNew(STextBlock)
+				.Text(this, &SAnimAttributeEntry::GetEntryTypeName)	
+			];
+	}
+	else if (ColumnName == ColumnId_AnimAttributeSnapshotName)
+	{
+		return
+			SNew(SVerticalBox)
+			+ SVerticalBox::Slot()
+			.AutoHeight()
+			.Padding(4)
+			.VAlign(VAlign_Center)
+			[
+				SNew(STextBlock)
+				.Text(this, &SAnimAttributeEntry::GetEntrySnapshotDisplayName)	
+			];
+	}
+
+	return SNullWidget::NullWidget;
+}
+
+FText SAnimAttributeEntry::GetEntryName() const
+{
+	return FText::FromName(Entry.Pin()->GetName());
+}
+
+FText SAnimAttributeEntry::GetEntryBoneName() const
+{
+	return FText::FromName(Entry.Pin()->GetBoneName());
+}
+
+FText SAnimAttributeEntry::GetEntryTypeName() const
+{
+	return FText::FromName(Entry.Pin()->GetTypeName());
+}
+
+FText SAnimAttributeEntry::GetEntrySnapshotDisplayName() const
+{
+	return FText::FromName(Entry.Pin()->GetSnapshotDisplayName());
+}
+
+SAnimAttributeView::SAnimAttributeView()
+{
+	bShouldRefreshListView = false;
+	bShouldRefreshValueView = false;
+	CachedNumSnapshots = 0;
+	ColumnIdToSort = ColumnId_AnimAttributeSnapshotName;
+	ActiveSortMode = EColumnSortMode::Ascending;
+}
+
+void SAnimAttributeView::Construct(const FArguments& InArgs)
+{
+	OnGetAttributeSnapshotColumnDisplayName = InArgs._OnGetAttributeSnapshotColumnDisplayName;
+	SnapshotColumnLabelOverride = InArgs._SnapshotColumnLabelOverride;
+
+	SExpandableArea::FArguments ExpandableAreaArgs;
+	ExpandableAreaArgs.AreaTitleFont(FEditorStyle::GetFontStyle("DetailsView.CategoryFontStyle"));
+	
+	ChildSlot
+	[
+		SNew(SSplitter)
+		.Orientation(Orient_Vertical)
+		+SSplitter::Slot()
+		.Value(0.6)
+		.MinSize(80.0f)
+		[
+			SNew(SVerticalBox)
+			+SVerticalBox::Slot()
+			.AutoHeight()
+			.VAlign(VAlign_Top)
+			[
+				SNew(SHorizontalBox)
+				+SHorizontalBox::Slot()
+				.VAlign(VAlign_Center)
+				.Padding(3.0f, 1.0f)
+				[
+					SNew(SSearchBox)
+					.OnTextChanged(this, &SAnimAttributeView::OnFilterTextChanged)
+				]
+			]
+			+SVerticalBox::Slot()
+			[
+				SAssignNew(AttributeListView, SListView<TSharedPtr<FAnimAttributeEntry>>)
+				.ListItemsSource(&FilteredAttributeEntries)
+				.OnGenerateRow_Static(&SAnimAttributeView::MakeTableRowWidget)
+				.OnSelectionChanged(this, &SAnimAttributeView::OnSelectionChanged)
+				.HeaderRow
+				(
+					SAssignNew(HeaderRow, SHeaderRow)
+					+ SHeaderRow::Column( ColumnId_AnimAttributeName)
+					.FillWidth(1.f)
+					.DefaultLabel( LOCTEXT( "AnimAttributeNameLabel", "Name" ) )
+					.SortMode_Raw(this, &SAnimAttributeView::GetSortModeForColumn, ColumnId_AnimAttributeName)
+					.OnSort_Raw(this, &SAnimAttributeView::OnSortAttributeEntries)
+
+					+ SHeaderRow::Column( ColumnId_AnimAttributeBoneName)
+					.FillWidth(1.f)
+					.DefaultLabel( LOCTEXT( "AnimAttributeBoneNameLabel", "Bone" ) )
+					.SortMode_Raw(this, &SAnimAttributeView::GetSortModeForColumn, ColumnId_AnimAttributeBoneName)
+					.OnSort_Raw(this, &SAnimAttributeView::OnSortAttributeEntries)
+					
+					+ SHeaderRow::Column( ColumnId_AnimAttributeTypeName)
+					.FillWidth(1.f)
+					.DefaultLabel( LOCTEXT( "AnimAttributeTypeNameLabel", "Type" ) )
+					.SortMode_Raw(this, &SAnimAttributeView::GetSortModeForColumn, ColumnId_AnimAttributeTypeName)
+					.OnSort_Raw(this, &SAnimAttributeView::OnSortAttributeEntries)
+				)
+			]
+		]
+		+SSplitter::Slot()
+		.Value(0.4)
+		.MinSize(30.0f)
+		[
+			SAssignNew(ValueViewBox, SVerticalBox)
+		]	
+	];
+}
+
+void SAnimAttributeView::DisplayNewAttributeContainerSnapshots(
+	const TArray<TTuple<FName, const UE::Anim::FHeapAttributeContainer&>>& InSnapshots,
+	const USkeletalMeshComponent* InOwningComponent)
+{
+	if (!ensure(InSnapshots.Num() != 0))
+	{
+		ClearListView();
+		return;
+	}
+	
+	// we need the SKM to look up bone names
+	if (!InOwningComponent || !InOwningComponent->SkeletalMesh)
+	{
+		ClearListView();
+		return;
+	}
+
+	if (ShouldInvalidateListViewCache(InSnapshots, InOwningComponent))
+	{
+		CachedNumSnapshots = InSnapshots.Num();
+		CachedSnapshotNameIndexMap.Reset();
+		
+		CachedAttributeIdentifierLists.Reset(InSnapshots.Num());
+		CachedAttributeSnapshotMap.Reset();
+		
+		CachedAttributeIdentifierLists.AddDefaulted( InSnapshots.Num());
+		
+		for (int32 SnapshotIndex = 0; SnapshotIndex < InSnapshots.Num(); SnapshotIndex++)
+		{
+			const TTuple<FName, const UE::Anim::FHeapAttributeContainer&> Snapshot = InSnapshots[SnapshotIndex];
+			CachedSnapshotNameIndexMap.FindOrAdd(Snapshot.Key) = SnapshotIndex;
+			
+			TTuple<FName, TArray<FAnimationAttributeIdentifier>>& CachedIdentifierList = CachedAttributeIdentifierLists[SnapshotIndex];
+			CachedIdentifierList.Key = Snapshot.Key;
+
+			const UE::Anim::FHeapAttributeContainer& AttributeContainer = Snapshot.Value;
+			TArray<FAnimationAttributeIdentifier>& CachedIdentifiers = CachedIdentifierList.Value;
+
+			CachedIdentifiers.AddDefaulted(AttributeContainer.Num());
+			
+			const TArray<TWeakObjectPtr<UScriptStruct>>& Types = AttributeContainer.GetUniqueTypes();
+
+			int32 CachedListIdentifierIndex = 0;
+			for (int32 TypeIndex = 0; TypeIndex < Types.Num(); TypeIndex++)
+			{
+				const TArray<UE::Anim::FAttributeId>& Ids = AttributeContainer.GetKeys(TypeIndex);
+			
+				for (int32 IdIndex = 0; IdIndex < Ids.Num(); IdIndex++)
+				{
+					const UE::Anim::FAttributeId& Id = Ids[IdIndex];
+					const FName BoneName = InOwningComponent->SkeletalMesh->GetRefSkeleton().GetBoneName(Id.GetIndex());
+				
+					const FAnimationAttributeIdentifier Identifier(
+						Id.GetName(), Id.GetIndex() ,BoneName, Types[TypeIndex].Get());
+
+					CachedIdentifiers[CachedListIdentifierIndex] = Identifier;
+
+					CachedAttributeSnapshotMap.FindOrAdd(Identifier).Add(Snapshot.Key);
+				
+					CachedListIdentifierIndex++;
+				}
+			}
+		}
+
+		// filtered list should also be refreshed since it depends on the cache
+		RefreshFilteredAttributeEntries();
+
+		// delay value view refresh until tick since this function can be called from animation thread
+		bShouldRefreshValueView = true;
+
+		return;
+	}
+
+	if (SelectedAttribute.IsSet())
+	{
+		for (const FAttributeValueView& ValueView : SelectedAttributeSnapshotValueViews)
+		{
+			const int32* SnapshotIndex = CachedSnapshotNameIndexMap.Find(ValueView.SnapshotName);
+			if (ensure(SnapshotIndex) && ensure(InSnapshots.IsValidIndex(*SnapshotIndex)))
+			{
+				const TTuple<FName, const UE::Anim::FHeapAttributeContainer&> Snapshot = InSnapshots[*SnapshotIndex];
+				const UE::Anim::FHeapAttributeContainer& AttributeContainer = Snapshot.Value;
+
+				ValueView.UpdateValue(AttributeContainer);
+			}
+		}
+	}	
+}
+
+
+
+bool SAnimAttributeView::ShouldInvalidateListViewCache(
+	const TArray<TTuple<FName, const UE::Anim::FHeapAttributeContainer&>>& InSnapshots,
+	const USkeletalMeshComponent* InOwningComponent)
+{
+	// we need the SKM to look up bone names
+	// it should have been checked
+	check(InOwningComponent);
+
+	if (InSnapshots.Num() != CachedAttributeIdentifierLists.Num())
+	{
+		return true;
+	}
+	
+	for (int32 SnapshotIndex = 0; SnapshotIndex < InSnapshots.Num(); SnapshotIndex++)
+	{
+		const TTuple<FName, const UE::Anim::FHeapAttributeContainer&> Snapshot = InSnapshots[SnapshotIndex];
+		const TTuple<FName, TArray<FAnimationAttributeIdentifier>>& CachedIdList = CachedAttributeIdentifierLists[SnapshotIndex];
+		
+		if (Snapshot.Key != CachedIdList.Key)
+		{
+			return true;
+		}
+
+		if (Snapshot.Value.Num() != CachedIdList.Value.Num())
+		{
+			return true;
+		}
+	}
+
+	for (int32 SnapshotIndex = 0; SnapshotIndex < InSnapshots.Num(); SnapshotIndex++)
+	{
+		const TTuple<FName, const UE::Anim::FHeapAttributeContainer&> Snapshot = InSnapshots[SnapshotIndex];
+		const TTuple<FName, TArray<FAnimationAttributeIdentifier>> CachedIdentifierList = CachedAttributeIdentifierLists[SnapshotIndex];
+
+		const UE::Anim::FHeapAttributeContainer& AttributeContainer = Snapshot.Value;
+		const TArray<FAnimationAttributeIdentifier>& CachedIdentifiers = CachedIdentifierList.Value;
+		
+		const TArray<TWeakObjectPtr<UScriptStruct>>& Types = AttributeContainer.GetUniqueTypes();
+
+		int32 CachedListIdentifierIndex = 0;
+		for (int32 TypeIndex = 0; TypeIndex < Types.Num(); TypeIndex++)
+		{
+			const TArray<UE::Anim::FAttributeId>& Ids = AttributeContainer.GetKeys(TypeIndex);
+			
+			for (int32 IdIndex = 0; IdIndex < Ids.Num(); IdIndex++)
+			{
+				const UE::Anim::FAttributeId& Id = Ids[IdIndex];
+				const FName BoneName = InOwningComponent->SkeletalMesh->GetRefSkeleton().GetBoneName(Id.GetIndex());
+				
+				FAnimationAttributeIdentifier Identifier(
+					Id.GetName(), Id.GetIndex() ,BoneName, Types[TypeIndex].Get());
+
+				if (!(Identifier == CachedIdentifiers[CachedListIdentifierIndex]))
+				{
+					return true;
+				}
+
+				CachedListIdentifierIndex++;
+			}
+		}
+	}	
+	
+
+	return false;
+}
+
+
+void SAnimAttributeView::OnSelectionChanged(TSharedPtr<FAnimAttributeEntry> InEntry, ESelectInfo::Type InSelectType)
+{
+	if (InEntry.IsValid())
+	{
+		SelectedAttribute = *InEntry.Get();
+	}
+	else
+	{
+		SelectedAttribute.Reset();
+	}
+
+	RefreshValueView();
+}
+
+void SAnimAttributeView::OnFilterTextChanged(const FText& InText)
+{
+	if (FilterText == InText.ToString())
+	{
+		return;
+	}
+	
+	FilterText = InText.ToString();
+
+	RefreshFilteredAttributeEntries();
+}
+
+EColumnSortMode::Type SAnimAttributeView::GetSortModeForColumn(FName InColumnId) const
+{
+	if (ColumnIdToSort == InColumnId)
+	{
+		return ActiveSortMode;
+	}
+		
+	return EColumnSortMode::None;
+}
+
+void SAnimAttributeView::OnSortAttributeEntries(
+   EColumnSortPriority::Type InPriority,
+   const FName& InColumnId,
+   EColumnSortMode::Type InSortMode)
+{
+	ColumnIdToSort = InColumnId;
+	ActiveSortMode = InSortMode;
+	
+	ExecuteSort();
+}
+
+void SAnimAttributeView::ExecuteSort()
+{
+	static const TArray<FName> ColumnIds =
+	{
+		ColumnId_AnimAttributeSnapshotName,
+		ColumnId_AnimAttributeName,
+		ColumnId_AnimAttributeBoneName,
+		ColumnId_AnimAttributeTypeName
+	};
+
+	TArray<FName> ColumnIdsBySortOrder = {ColumnIdToSort};
+	for (const FName& Id : ColumnIds)
+	{
+		if (Id != ColumnIdToSort)
+		{
+			ColumnIdsBySortOrder.Add(Id);
+		}
+	}
+
+	FilteredAttributeEntries.Sort(
+		[ColumnIdsBySortOrder, this](const TSharedPtr<FAnimAttributeEntry>& Left, const TSharedPtr<FAnimAttributeEntry>& Right)
+		{
+			int32 CompareResult = 0;
+			for (const FName& ColumnId : ColumnIdsBySortOrder)
+			{
+				if (ColumnId == ColumnId_AnimAttributeSnapshotName)
+				{
+					CompareResult = Left->GetSnapshotDisplayName().Compare(Right->GetSnapshotDisplayName());	
+				}
+				if (ColumnId == ColumnId_AnimAttributeName)
+				{
+					CompareResult = Left->GetName().Compare(Right->GetName());
+				}
+				else if (ColumnId == ColumnId_AnimAttributeBoneName)
+				{
+					CompareResult = Left->GetBoneName().Compare(Right->GetBoneName());
+				}
+				else if (ColumnId == ColumnId_AnimAttributeTypeName)
+				{
+					CompareResult = Left->GetTypeName().Compare(Right->GetTypeName());
+				}
+				
+				if (CompareResult != 0)
+				{
+					// we have a winner
+					return ActiveSortMode == EColumnSortMode::Ascending ? CompareResult < 0 : CompareResult > 0;
+				}
+			}
+
+			// keep the original order if two entries are the same (though we should never have identical entries)
+			return ActiveSortMode == EColumnSortMode::Ascending ? true : false;
+		});
+
+	bShouldRefreshListView = true;
+}
+
+void SAnimAttributeView::RefreshFilteredAttributeEntries()
+{
+	FilteredAttributeEntries.Reset(CachedAttributeSnapshotMap.Num());
+	
+	for (const TTuple<FAnimationAttributeIdentifier, TArray<FName>>& Identifier : CachedAttributeSnapshotMap)
+	{
+		FName SnapshotDisplayName;
+		if (OnGetAttributeSnapshotColumnDisplayName.IsBound())
+		{
+			SnapshotDisplayName = OnGetAttributeSnapshotColumnDisplayName.Execute(Identifier.Value);	
+		}
+		else
+		{
+			SnapshotDisplayName = GetSnapshotColumnDisplayName(Identifier.Value);
+		}
+		
+		if 	(FilterText.IsEmpty() ||
+			Identifier.Key.GetName().ToString().Contains(FilterText) ||
+			Identifier.Key.GetBoneName().ToString().Contains(FilterText) ||
+			GetAnimAttributeTypeName(Identifier.Key.GetType()).ToString().Contains(FilterText) ||
+			SnapshotDisplayName.ToString().Contains(FilterText))
+		{
+			FilteredAttributeEntries.Add(FAnimAttributeEntry::MakeEntry(Identifier.Key, SnapshotDisplayName));
+		}
+	}
+
+	ExecuteSort();
+
+	// delay the refresh to until tick since this function
+	// can be invoked from animation thread
+	bShouldRefreshListView = true;
+}
+
+void SAnimAttributeView::RefreshValueView()
+{
+	for (const FAttributeValueView& ValueView : SelectedAttributeSnapshotValueViews)
+	{
+		ValueViewBox->RemoveSlot(ValueView.ViewWidget->GetWidget().ToSharedRef());
+	}
+	
+	SelectedAttributeSnapshotValueViews.Reset();
+	
+	if (SelectedAttribute.IsSet())
+	{
+		const FAnimationAttributeIdentifier& Identifier = SelectedAttribute->GetAnimationAttributeIdentifier();
+
+		if (const TArray<FName>* SnapshotNames = CachedAttributeSnapshotMap.Find(Identifier))
+		{
+			for (const FName& SnapshotName : *SnapshotNames)
+			{
+				SelectedAttributeSnapshotValueViews.Add(
+					FAttributeValueView(
+							SnapshotName,
+							SelectedAttribute.GetValue()));
+			}
+		}
+	}
+
+	
+	// for (int32 ViewIndex = SelectedAttributeSnapshotValueViews.Num() - 1; ViewIndex >= 0; ViewIndex--)
+	for (int32 ViewIndex = 0; ViewIndex < SelectedAttributeSnapshotValueViews.Num(); ViewIndex++)
+	{
+		const FAttributeValueView& ValueView = SelectedAttributeSnapshotValueViews[ViewIndex];
+		// slots are added in reverse order
+		ValueViewBox->AddSlot()
+		.AutoHeight()
+		[
+			ValueView.ViewWidget->GetWidget().ToSharedRef()
+		];	
+	}
+}
+
+void SAnimAttributeView::Tick(const FGeometry& AllottedGeometry, const double InCurrentTime, const float InDeltaTime)
+{
+	SCompoundWidget::Tick(AllottedGeometry, InCurrentTime, InDeltaTime);
+
+	
+	if (bShouldRefreshListView)
+	{
+		if (CachedNumSnapshots == 1 && HeaderRow->GetColumns().Num() > MinimalNumColumns)
+		{
+			HeaderRow->RemoveColumn(ColumnId_AnimAttributeSnapshotName);
+		}
+		else if (CachedNumSnapshots >= 1 && HeaderRow->GetColumns().Num() <= MinimalNumColumns)
+		{
+			SHeaderRow::FColumn::FArguments	ColumnArgs;
+
+			ColumnArgs
+				.ColumnId(ColumnId_AnimAttributeSnapshotName)
+				.FillWidth(1.f)
+				.DefaultLabel( SnapshotColumnLabelOverride )
+				.SortMode_Raw(this, &SAnimAttributeView::GetSortModeForColumn, ColumnId_AnimAttributeSnapshotName)
+				.OnSort_Raw(this, &SAnimAttributeView::OnSortAttributeEntries);
+			
+			HeaderRow->AddColumn(ColumnArgs);
+		}
+		
+		AttributeListView->RequestListRefresh();
+		bShouldRefreshListView = false;
+	}
+
+	if (bShouldRefreshValueView)
+	{
+		RefreshValueView();
+		bShouldRefreshValueView = false;
+	}
+}
+
+SAnimAttributeView::FAttributeValueView::FAttributeValueView(FName InSnapshotName,
+	const FAnimAttributeEntry& InSelectedAttribute)
+{
+	SubjectAttribute = InSelectedAttribute;
+	SnapshotName = InSnapshotName;
+	StructData = MakeShareable(new FStructOnScope(InSelectedAttribute.GetScriptStruct()));
+	ViewWidget = CreateValueViewWidget();
+
+	const FString DisplayName = InSelectedAttribute.GetDisplayName().ToString() + TEXT(" - ") + InSnapshotName.ToString(); 
+	
+	ViewWidget->SetCustomName(FText::FromString(DisplayName));
+	ViewWidget->SetStructureData(StructData);
+}
+
+void SAnimAttributeView::FAttributeValueView::UpdateValue(const UE::Anim::FHeapAttributeContainer& InAttributeContainer) const
+{
+	const uint8* ValuePtr = InAttributeContainer.Find(SubjectAttribute.GetScriptStruct(), SubjectAttribute.GetAttributeId());
+
+	if (ensure(ValuePtr))
+	{
+		FMemory::Memcpy(
+			StructData->GetStructMemory(),
+			ValuePtr,
+			StructData->GetStruct()->GetStructureSize());
+	}		
+}
+
+#undef LOCTEXT_NAMESPACE
diff --git a/Engine/Source/Editor/Persona/Public/SAnimAttributeView.h b/Engine/Source/Editor/Persona/Public/SAnimAttributeView.h
new file mode 100644
--- /dev/null
+++ b/Engine/Source/Editor/Persona/Public/SAnimAttributeView.h
@@ -0,0 +1,211 @@
+// Copyright Epic Games, Inc. All Rights Reserved.
+
+#pragma once
+
+#include "CoreMinimal.h"
+#include "IStructureDetailsView.h"
+#include "Widgets/Views/STableRow.h"
+#include "Widgets/SCompoundWidget.h"
+#include "Animation/AttributesRuntime.h"
+#include "Animation/AnimData/AttributeIdentifier.h"
+
+class ITableRow;
+class STableViewBase;
+class IStructureDetailsView;
+class FStructOnScope;
+class SHeaderRow;
+class SVerticalBox;
+
+namespace EColumnSortMode
+{
+	enum Type;
+}
+
+
+class FAnimAttributeEntry : public TSharedFromThis<FAnimAttributeEntry>
+{
+public:
+	static TSharedRef<FAnimAttributeEntry> MakeEntry(const FAnimationAttributeIdentifier& InIdentifier, const FName& InSnapshotDisplayName);
+
+	FAnimAttributeEntry() = default;
+	FAnimAttributeEntry(const FAnimationAttributeIdentifier& InIdentifier, const FName& InSnapshotDisplayName);
+	
+	TSharedRef<ITableRow> MakeTableRowWidget(const TSharedRef<STableViewBase>& InOwnerTable);
+
+	FName GetName() const { return Identifier.GetName(); }
+	FName GetBoneName() const { return Identifier.GetBoneName(); }
+	int32 GetBoneIndex() const { return Identifier.GetBoneIndex(); }
+	FName GetTypeName() const { return CachedTypeName; }
+	UScriptStruct* GetScriptStruct() const { return Identifier.GetType(); }
+	FName GetSnapshotDisplayName() const { return SnapshotDisplayName; }
+	
+	FName GetDisplayName() const;
+	
+	UE::Anim::FAttributeId GetAttributeId() const
+	{
+		return UE::Anim::FAttributeId(GetName(), FCompactPoseBoneIndex(GetBoneIndex()));
+	};
+
+	const FAnimationAttributeIdentifier& GetAnimationAttributeIdentifier() const
+	{
+		return Identifier;
+	};
+
+	bool operator==(const FAnimAttributeEntry& InOther) const
+	{
+		return Identifier == InOther.Identifier;
+	}
+
+	bool operator!=(const FAnimAttributeEntry& InOther) const
+	{
+		return !(*this == InOther);
+	}
+	
+	
+private:
+	FAnimationAttributeIdentifier Identifier;
+	
+	FName SnapshotDisplayName;
+
+	FName CachedTypeName;
+};
+
+class SAnimAttributeEntry : public SMultiColumnTableRow<TSharedPtr<FAnimAttributeEntry>>
+{
+	SLATE_BEGIN_ARGS( SAnimAttributeEntry )
+	{}
+	
+	SLATE_END_ARGS()
+
+	void Construct( const FArguments& InArgs, const TSharedRef<STableViewBase>& InOwnerTable, TSharedRef<FAnimAttributeEntry> InEntry);
+
+	/** Overridden from SMultiColumnTableRow.  Generates a widget for this column of the tree row. */
+	virtual TSharedRef<SWidget> GenerateWidgetForColumn(const FName& ColumnName) override;
+	
+	FText GetEntryName()const ;
+	FText GetEntryBoneName() const;
+	FText GetEntryTypeName() const;
+	FText GetEntrySnapshotDisplayName() const;
+private:
+	TWeakPtr<FAnimAttributeEntry> Entry;
+};
+
+class PERSONA_API SAnimAttributeView : public SCompoundWidget
+{
+private:	
+	static TSharedRef<IStructureDetailsView> CreateValueViewWidget();
+	
+	static TSharedRef<ITableRow> MakeTableRowWidget(
+		TSharedPtr<FAnimAttributeEntry> InItem,
+		const TSharedRef<STableViewBase>& InOwnerTable);
+
+	static FName GetSnapshotColumnDisplayName(const TArray<FName>& InSnapshotNames);
+
+public:
+	DECLARE_DELEGATE_RetVal_OneParam(FName, FOnGetAttributeSnapshotColumnDisplayName, const TArray<FName>& /** Snapshot Names */)
+	
+	SLATE_BEGIN_ARGS( SAnimAttributeView )
+		: _SnapshotColumnLabelOverride(FText::FromString(TEXT("Direction")))
+		{}
+		// override what is displayed in the snapshot column, given a set of snapshots that contains the attribute
+		SLATE_EVENT(FOnGetAttributeSnapshotColumnDisplayName, OnGetAttributeSnapshotColumnDisplayName)
+	
+		// override the label on the snapshot column, a typical choice is "Direction"
+		SLATE_ATTRIBUTE(FText, SnapshotColumnLabelOverride)
+	SLATE_END_ARGS()
+	
+	SAnimAttributeView();
+	
+	void Construct(const FArguments& InArgs);
+
+	void DisplayNewAttributeContainerSnapshots(
+		const TArray<TTuple<FName, const UE::Anim::FHeapAttributeContainer&>>& InSnapshots, 
+		const USkeletalMeshComponent* InOwningComponent);
+
+	void ClearListView();
+	
+	virtual void Tick(const FGeometry& AllottedGeometry, const double InCurrentTime, const float InDeltaTime) override;
+	
+private:
+
+	bool ShouldInvalidateListViewCache(
+		const TArray<TTuple<FName, const UE::Anim::FHeapAttributeContainer&>>& InSnapshots,
+		const USkeletalMeshComponent* InOwningComponent);
+	
+	void OnSelectionChanged(TSharedPtr<FAnimAttributeEntry> InEntry, ESelectInfo::Type InSelectType);
+
+	void OnFilterTextChanged(const FText& InText);
+
+	EColumnSortMode::Type GetSortModeForColumn(FName InColumnId) const;
+	
+	void OnSortAttributeEntries(
+		EColumnSortPriority::Type InPriority,
+		const FName& InColumnId,
+		EColumnSortMode::Type InSortMode);
+
+	void ExecuteSort();
+	
+	void RefreshFilteredAttributeEntries();
+
+	void RefreshValueView();
+
+private:
+	/** list view */
+	TSharedPtr<SListView<TSharedPtr<FAnimAttributeEntry>>> AttributeListView;
+	bool bShouldRefreshListView;
+
+	TSharedPtr<SHeaderRow> HeaderRow;
+	
+	FName ColumnIdToSort;
+	EColumnSortMode::Type ActiveSortMode;
+	FOnGetAttributeSnapshotColumnDisplayName OnGetAttributeSnapshotColumnDisplayName;
+	TAttribute<FText> SnapshotColumnLabelOverride;
+
+	int32 CachedNumSnapshots;
+	// cache all attributes in the attribute container that the list view is observing
+	// such that we can use it to detect if a change to the attribute container occured
+	// and refresh the list accordingly
+	TArray<TTuple<FName, TArray<FAnimationAttributeIdentifier>>> CachedAttributeIdentifierLists;
+
+	// for each attribute, save the name of the attribute container snapshot that contains it
+	TMap<FAnimationAttributeIdentifier, TArray<FName>> CachedAttributeSnapshotMap;
+
+	TMap<FName, int32> CachedSnapshotNameIndexMap;
+
+	// attributes to be displayed
+	TArray<TSharedPtr<FAnimAttributeEntry>> FilteredAttributeEntries;
+
+	FString FilterText;
+
+	
+	/** value view */
+	TSharedPtr<SVerticalBox> ValueViewBox;
+	bool bShouldRefreshValueView;
+	TOptional<FAnimAttributeEntry> SelectedAttribute;
+
+	struct FAttributeValueView
+	{
+		FAttributeValueView(FName InSnapshotName, const FAnimAttributeEntry& InSelectedAttribute);
+
+		void UpdateValue(const UE::Anim::FHeapAttributeContainer& InAttributeContainer) const;
+
+		FAnimAttributeEntry SubjectAttribute;
+		FName SnapshotName;
+		TSharedPtr<FStructOnScope> StructData;
+		TSharedPtr<IStructureDetailsView> ViewWidget;
+	};
+	
+	TArray<FAttributeValueView> SelectedAttributeSnapshotValueViews;
+};
+
+
+inline void SAnimAttributeView::ClearListView()
+{
+	CachedNumSnapshots = 0;
+	CachedAttributeIdentifierLists.Reset();
+	CachedAttributeSnapshotMap.Reset();
+	CachedSnapshotNameIndexMap.Reset();
+	FilteredAttributeEntries.Reset();
+
+	bShouldRefreshListView = true;
+}
diff --git a/Engine/Source/Runtime/Engine/Classes/Animation/AttributesContainer.h b/Engine/Source/Runtime/Engine/Classes/Animation/AttributesContainer.h
--- a/Engine/Source/Runtime/Engine/Classes/Animation/AttributesContainer.h
+++ b/Engine/Source/Runtime/Engine/Classes/Animation/AttributesContainer.h
@@ -285,7 +285,7 @@ namespace UE
 			*
 			* @return	The ptr to the added entry, to be used for populating the data, nullptr if adding it failed
 			*/
-			uint8* Add(const UScriptStruct* InScriptStruct, const FAttributeId& InAttributeId)
+			uint8* Add(UScriptStruct* InScriptStruct, const FAttributeId& InAttributeId)
 			{
 				const int32 TypeIndex = FindOrAddTypeIndex(InScriptStruct);
 
@@ -337,7 +337,7 @@ namespace UE
 			*
 			* @return	The ptr to the added/existing entry, to be used for populating the data, nullptr if adding it failed
 			*/
-			uint8* FindOrAdd(const UScriptStruct* InScriptStruct, const FAttributeId& InAttributeId)
+			uint8* FindOrAdd(UScriptStruct* InScriptStruct, const FAttributeId& InAttributeId)
 			{
 				const int32 TypeIndex = FindOrAddTypeIndex(InScriptStruct);
 				TArray<FAttributeId>& AttributeIds = AttributeIdentifiers[TypeIndex];
@@ -672,7 +672,7 @@ namespace UE
 			/*
 			* @return Array of all the contained attribute types (UScriptStruct)
 			*/
-			const TArray<TWeakObjectPtr<const UScriptStruct>>& GetUniqueTypes() const { return UniqueTypes; }
+			const TArray<TWeakObjectPtr<UScriptStruct>>& GetUniqueTypes() const { return UniqueTypes; }
 
 			/*
 			* Populate out parameter OutAttributeKeyNames for all contained attribute key names
@@ -703,6 +703,20 @@ namespace UE
 		*/
 		const TArray<int32>& GetUniqueTypedBoneIndices(int32 TypeIndex) const { return UniqueTypedBoneIndices[TypeIndex]; }
 
+		/*
+		* @return total number of attributes
+		*/
+		int32 Num() const
+		{
+			int32 NumAttributes = 0;
+			for (const TArray<FAttributeId>& AttributesPerType : AttributeIdentifiers)
+			{
+				NumAttributes += AttributesPerType.Num();
+			}
+				
+			return NumAttributes;
+		}
+
 		public:
 			/* Deprecated API */
 			template<typename DataType> 
@@ -770,7 +784,7 @@ namespace UE
 			TArray<FAttributeId>& GetKeysInternal(int32 TypeIndex) { return AttributeIdentifiers[TypeIndex]; }
 			
 			/** Find or add a new root-level entry for the provided attribute data type, returning the index into the arrays representing the type */
-			int32 FindOrAddTypeIndex(const UScriptStruct* InScriptStruct)
+			int32 FindOrAddTypeIndex(UScriptStruct* InScriptStruct)
 			{
 				const int32 OldNum = UniqueTypes.Num();
 				const int32 TypeIndex = UniqueTypes.AddUnique(InScriptStruct);
@@ -792,7 +806,7 @@ namespace UE
 			TArray<TArray<int32>> UniqueTypedBoneIndices;
 			TArray<TArray<FAttributeId>> AttributeIdentifiers;
 			TArray<TArray<TWrappedAttribute<InAllocator>>> Values;
-			TArray<TWeakObjectPtr<const UScriptStruct>> UniqueTypes;
+			TArray<TWeakObjectPtr<UScriptStruct>> UniqueTypes;
 
 			template <typename OtherBoneIndexType, typename OtherAllocator>
 			friend struct TAttributeContainer;
diff --git a/Engine/Source/Runtime/Engine/Private/Animation/AttributesRuntime.cpp b/Engine/Source/Runtime/Engine/Private/Animation/AttributesRuntime.cpp
--- a/Engine/Source/Runtime/Engine/Private/Animation/AttributesRuntime.cpp
+++ b/Engine/Source/Runtime/Engine/Private/Animation/AttributesRuntime.cpp
@@ -44,13 +44,13 @@ void Attributes::BlendAttributes(const TArrayView<const FStackAttributeContainer
 
 	if (SourceAttributes.Num())
 	{
-		TSet<TWeakObjectPtr<const UScriptStruct>, DefaultKeyFuncs<TWeakObjectPtr<const UScriptStruct>>, TInlineSetAllocator<4>> UniqueStructs;
+		TSet<TWeakObjectPtr<UScriptStruct>, DefaultKeyFuncs<TWeakObjectPtr<UScriptStruct>>, TInlineSetAllocator<4>> UniqueStructs;
 		for (const FStackAttributeContainer& CustomAttributes : SourceAttributes)
 		{
 			UniqueStructs.Append(CustomAttributes.GetUniqueTypes());
 		}
 		
-		for (const TWeakObjectPtr<const UScriptStruct>& WeakScriptStruct : UniqueStructs)
+		for (const TWeakObjectPtr<UScriptStruct>& WeakScriptStruct : UniqueStructs)
 		{
 			const IAttributeBlendOperator* Operator = AttributeTypes::GetTypeOperator(WeakScriptStruct);
 			FAttributeBlendData AttributeBlendData = FAttributeBlendData::PerContainerWeighted(SourceAttributes, SourceWeights, WeakScriptStruct.Get());
@@ -65,13 +65,13 @@ void Attributes::BlendAttributes(const TArrayView<const FStackAttributeContainer
 
 	if (SourceAttributes.Num())
 	{
-		TSet<TWeakObjectPtr<const UScriptStruct>, DefaultKeyFuncs<TWeakObjectPtr<const UScriptStruct>>, TInlineSetAllocator<4>> UniqueStructs;
+		TSet<TWeakObjectPtr<UScriptStruct>, DefaultKeyFuncs<TWeakObjectPtr<UScriptStruct>>, TInlineSetAllocator<4>> UniqueStructs;
 		for (const FStackAttributeContainer* CustomAttributes : SourceAttributes)
 		{
 			UniqueStructs.Append(CustomAttributes->GetUniqueTypes());
 		}
 
-		for (const TWeakObjectPtr<const UScriptStruct>& WeakScriptStruct : UniqueStructs)
+		for (const TWeakObjectPtr<UScriptStruct>& WeakScriptStruct : UniqueStructs)
 		{
 			const IAttributeBlendOperator* Operator = AttributeTypes::GetTypeOperator(WeakScriptStruct);
 			FAttributeBlendData AttributeBlendData = FAttributeBlendData::PerContainerPtrWeighted(SourceAttributes, SourceWeights, WeakScriptStruct.Get());
@@ -86,13 +86,13 @@ void Attributes::BlendAttributes(const TArrayView<const FStackAttributeContainer
 
 	if (SourceAttributes.Num())
 	{
-		TSet<TWeakObjectPtr<const UScriptStruct>, DefaultKeyFuncs<TWeakObjectPtr<const UScriptStruct>>, TInlineSetAllocator<4>> UniqueStructs;
+		TSet<TWeakObjectPtr<UScriptStruct>, DefaultKeyFuncs<TWeakObjectPtr<UScriptStruct>>, TInlineSetAllocator<4>> UniqueStructs;
 		for (const FStackAttributeContainer& CustomAttributes : SourceAttributes)
 		{
 			UniqueStructs.Append(CustomAttributes.GetUniqueTypes());
 		}
 
-		for (const TWeakObjectPtr<const UScriptStruct>& WeakScriptStruct : UniqueStructs)
+		for (const TWeakObjectPtr<UScriptStruct>& WeakScriptStruct : UniqueStructs)
 		{
 			const IAttributeBlendOperator* Operator = AttributeTypes::GetTypeOperator(WeakScriptStruct);
 			FAttributeBlendData AttributeBlendData = FAttributeBlendData::PerContainerRemappedWeighted(SourceAttributes, SourceWeights, SourceWeightsIndices, WeakScriptStruct.Get());
@@ -105,12 +105,12 @@ void Attributes::OverrideAttributes(const FStackAttributeContainer& SourceAttrib
 {
 	QUICK_SCOPE_CYCLE_COUNTER(STAT_OverrideAttributes_Weighted);
 
-	TSet<TWeakObjectPtr<const UScriptStruct>, DefaultKeyFuncs<TWeakObjectPtr<const UScriptStruct>>, TInlineSetAllocator<4>> UniqueStructs;
+	TSet<TWeakObjectPtr<UScriptStruct>, DefaultKeyFuncs<TWeakObjectPtr<UScriptStruct>>, TInlineSetAllocator<4>> UniqueStructs;
 	UniqueStructs.Append(SourceAttributes.GetUniqueTypes());
 	UniqueStructs.Append(OutAttributes.GetUniqueTypes());
 
 	const TArray<const FStackAttributeContainer*> Array = { &SourceAttributes };
-	for (const TWeakObjectPtr<const UScriptStruct>& WeakScriptStruct : UniqueStructs)
+	for (const TWeakObjectPtr<UScriptStruct>& WeakScriptStruct : UniqueStructs)
 	{
 		const IAttributeBlendOperator* Operator = AttributeTypes::GetTypeOperator(WeakScriptStruct);
 
@@ -124,12 +124,12 @@ void Attributes::AccumulateAttributes(const FStackAttributeContainer& SourceAttr
 {
 	QUICK_SCOPE_CYCLE_COUNTER(STAT_AccumulateAttributes);
 
-	TSet<TWeakObjectPtr<const UScriptStruct>, DefaultKeyFuncs<TWeakObjectPtr<const UScriptStruct>>, TInlineSetAllocator<4>> UniqueStructs;
+	TSet<TWeakObjectPtr<UScriptStruct>, DefaultKeyFuncs<TWeakObjectPtr<UScriptStruct>>, TInlineSetAllocator<4>> UniqueStructs;
 	UniqueStructs.Append(SourceAttributes.GetUniqueTypes());
 	UniqueStructs.Append(OutAttributes.GetUniqueTypes());
 
 	const TArray<const FStackAttributeContainer*> Array = { &SourceAttributes };
-	for (const TWeakObjectPtr<const UScriptStruct>& WeakScriptStruct : UniqueStructs)
+	for (const TWeakObjectPtr<UScriptStruct>& WeakScriptStruct : UniqueStructs)
 	{
 		const IAttributeBlendOperator* Operator = AttributeTypes::GetTypeOperator(WeakScriptStruct);
 		FAttributeBlendData AttributeBlendData = FAttributeBlendData::SingleAdditiveContainerUniformWeighted(Array, Weight, AdditiveType, WeakScriptStruct.Get());
@@ -141,11 +141,11 @@ void Attributes::ConvertToAdditive(const FStackAttributeContainer& BaseAttribute
 {
 	QUICK_SCOPE_CYCLE_COUNTER(STAT_SubtractAttributes);
 
-	TSet<TWeakObjectPtr<const UScriptStruct>, DefaultKeyFuncs<TWeakObjectPtr<const UScriptStruct>>, TInlineSetAllocator<4>> UniqueStructs;
+	TSet<TWeakObjectPtr<UScriptStruct>, DefaultKeyFuncs<TWeakObjectPtr<UScriptStruct>>, TInlineSetAllocator<4>> UniqueStructs;
 	UniqueStructs.Append(BaseAttributes.GetUniqueTypes());
 
 	const TArray<const FStackAttributeContainer*> Array = { &BaseAttributes };
-	for (const TWeakObjectPtr<const UScriptStruct>& WeakScriptStruct : UniqueStructs)
+	for (const TWeakObjectPtr<UScriptStruct>& WeakScriptStruct : UniqueStructs)
 	{
 		const IAttributeBlendOperator* Operator = AttributeTypes::GetTypeOperator(WeakScriptStruct);
 
@@ -159,9 +159,9 @@ void Attributes::CopyAndRemapAttributes(const FMeshAttributeContainer& SourceAtt
 {
 	QUICK_SCOPE_CYCLE_COUNTER(STAT_CopyAndRemapAttributes);
 
-	for (const TWeakObjectPtr<const UScriptStruct> WeakScriptStruct : SourceAttributes.GetUniqueTypes())
+	for (const TWeakObjectPtr<UScriptStruct> WeakScriptStruct : SourceAttributes.GetUniqueTypes())
 	{
-		const UScriptStruct* ScriptStruct = WeakScriptStruct.Get();
+		UScriptStruct* ScriptStruct = WeakScriptStruct.Get();
 		const int32 TypeIndex = SourceAttributes.FindTypeIndex(ScriptStruct);
 		if (TypeIndex != INDEX_NONE)
 		{
@@ -196,12 +196,12 @@ void Attributes::InterpolateAttributes(FMeshAttributeContainer& FromAttributes,
 {
 	QUICK_SCOPE_CYCLE_COUNTER(STAT_InterpolateAttributes);
 
-	TSet<TWeakObjectPtr<const UScriptStruct>> FromUniqueStructs;
-	TSet<TWeakObjectPtr<const UScriptStruct>> ToUniqueStructs;
-	TSet<TWeakObjectPtr<const UScriptStruct>> OverlappingStructs;
+	TSet<TWeakObjectPtr<UScriptStruct>> FromUniqueStructs;
+	TSet<TWeakObjectPtr<UScriptStruct>> ToUniqueStructs;
+	TSet<TWeakObjectPtr<UScriptStruct>> OverlappingStructs;
 	{
-		const TArray<TWeakObjectPtr<const UScriptStruct>, FDefaultAllocator>& SourceUniqueTypes = FromAttributes.GetUniqueTypes();
-		const TArray<TWeakObjectPtr<const UScriptStruct>, FDefaultAllocator>& OutUniqueTypes = ToAttributes.GetUniqueTypes();
+		const TArray<TWeakObjectPtr<UScriptStruct>, FDefaultAllocator>& SourceUniqueTypes = FromAttributes.GetUniqueTypes();
+		const TArray<TWeakObjectPtr<UScriptStruct>, FDefaultAllocator>& OutUniqueTypes = ToAttributes.GetUniqueTypes();
 
 		for (const auto& Type : SourceUniqueTypes)
 		{
@@ -227,7 +227,7 @@ void Attributes::InterpolateAttributes(FMeshAttributeContainer& FromAttributes,
 	// Handle unique types first
 	{
 		// From
-		for (const TWeakObjectPtr<const UScriptStruct>& WeakScriptStruct : FromUniqueStructs)
+		for (const TWeakObjectPtr<UScriptStruct>& WeakScriptStruct : FromUniqueStructs)
 		{
 			const int32 TypeIndex = FromAttributes.FindTypeIndex(WeakScriptStruct.Get());
 			const TArray<FAttributeId, FDefaultAllocator>& Identifiers = FromAttributes.GetKeys(TypeIndex);
@@ -246,7 +246,7 @@ void Attributes::InterpolateAttributes(FMeshAttributeContainer& FromAttributes,
 		}
 
 		// To 
-		for (const TWeakObjectPtr<const UScriptStruct>& WeakScriptStruct : ToUniqueStructs)
+		for (const TWeakObjectPtr<UScriptStruct>& WeakScriptStruct : ToUniqueStructs)
 		{
 			const int32 TypeIndex = ToAttributes.FindTypeIndex(WeakScriptStruct.Get());
 			const TArray<FAttributeId, FDefaultAllocator>& Identifiers = ToAttributes.GetKeys(TypeIndex);
@@ -267,7 +267,7 @@ void Attributes::InterpolateAttributes(FMeshAttributeContainer& FromAttributes,
 	}
 
 	// Overlapping types	
-	for (const TWeakObjectPtr<const UScriptStruct>& WeakScriptStruct : OverlappingStructs)
+	for (const TWeakObjectPtr<UScriptStruct>& WeakScriptStruct : OverlappingStructs)
 	{
 		const int32 FromTypeIndex = FromAttributes.FindTypeIndex(WeakScriptStruct.Get());
 		const TArray<FAttributeId, FDefaultAllocator>& FromIdentifiers = FromAttributes.GetKeys(FromTypeIndex);
@@ -330,11 +330,11 @@ void Attributes::BlendAttributesPerBone(const FStackAttributeContainer& SourceAt
 {
 	QUICK_SCOPE_CYCLE_COUNTER(STAT_BlendAttributesPerBone);
 
-	TSet<TWeakObjectPtr<const UScriptStruct>, DefaultKeyFuncs<TWeakObjectPtr<const UScriptStruct>>, TInlineSetAllocator<4>> UniqueStructs;
+	TSet<TWeakObjectPtr<UScriptStruct>, DefaultKeyFuncs<TWeakObjectPtr<UScriptStruct>>, TInlineSetAllocator<4>> UniqueStructs;
 	UniqueStructs.Append(SourceAttributes1.GetUniqueTypes());
 	UniqueStructs.Append(SourceAttributes2.GetUniqueTypes());
 
-	for (const TWeakObjectPtr<const UScriptStruct>& WeakScriptStruct : UniqueStructs)
+	for (const TWeakObjectPtr<UScriptStruct>& WeakScriptStruct : UniqueStructs)
 	{
 		const IAttributeBlendOperator* Operator = AttributeTypes::GetTypeOperator(WeakScriptStruct);
 		FAttributeBlendData AttributeBlendData = FAttributeBlendData::PerBoneSingleContainerWeighted(SourceAttributes1, SourceAttributes2, WeightsOfSource2, WeakScriptStruct.Get());
@@ -346,7 +346,7 @@ void Attributes::BlendAttributesPerBone(TArrayView<const FStackAttributeContaine
 {
 	QUICK_SCOPE_CYCLE_COUNTER(STAT_BlendAttributesPerBoneBlendSample);
 
-	TSet<TWeakObjectPtr<const UScriptStruct>, DefaultKeyFuncs<TWeakObjectPtr<const UScriptStruct>>, TInlineSetAllocator<4>> UniqueStructs;
+	TSet<TWeakObjectPtr<UScriptStruct>, DefaultKeyFuncs<TWeakObjectPtr<UScriptStruct>>, TInlineSetAllocator<4>> UniqueStructs;
 	for (const FStackAttributeContainer& CustomAttributes : SourceAttributes)
 	{
 		UniqueStructs.Append(CustomAttributes.GetUniqueTypes());
@@ -356,7 +356,7 @@ void Attributes::BlendAttributesPerBone(TArrayView<const FStackAttributeContaine
 	const bool bContainsPerBoneWeights = BlendSampleDataCache[0].PerBoneBlendData.Num() > 0;
 	if (bContainsPerBoneWeights)
 	{
-		for (const TWeakObjectPtr<const UScriptStruct>& WeakScriptStruct : UniqueStructs)
+		for (const TWeakObjectPtr<UScriptStruct>& WeakScriptStruct : UniqueStructs)
 		{
 			const IAttributeBlendOperator* Operator = AttributeTypes::GetTypeOperator(WeakScriptStruct);
 
@@ -372,7 +372,7 @@ void Attributes::BlendAttributesPerBone(TArrayView<const FStackAttributeContaine
 			Weights.Add(SampleData.GetClampedWeight());
 		}
 
-		for (const TWeakObjectPtr<const UScriptStruct>& WeakScriptStruct : UniqueStructs)
+		for (const TWeakObjectPtr<UScriptStruct>& WeakScriptStruct : UniqueStructs)
 		{
 			const IAttributeBlendOperator* Operator = AttributeTypes::GetTypeOperator(WeakScriptStruct);
 
@@ -386,7 +386,7 @@ void Attributes::BlendAttributesPerBone(TArrayView<const FStackAttributeContaine
 {
 	QUICK_SCOPE_CYCLE_COUNTER(STAT_BlendAttributesPerBoneBlendSample);
 
-	TSet<TWeakObjectPtr<const UScriptStruct>, DefaultKeyFuncs<TWeakObjectPtr<const UScriptStruct>>, TInlineSetAllocator<4>> UniqueStructs;
+	TSet<TWeakObjectPtr<UScriptStruct>, DefaultKeyFuncs<TWeakObjectPtr<UScriptStruct>>, TInlineSetAllocator<4>> UniqueStructs;
 	for (const FStackAttributeContainer& CustomAttributes : SourceAttributes)
 	{
 		UniqueStructs.Append(CustomAttributes.GetUniqueTypes());
@@ -396,7 +396,7 @@ void Attributes::BlendAttributesPerBone(TArrayView<const FStackAttributeContaine
 	const bool bContainsPerBoneWeights = BlendSampleDataCache[0].PerBoneBlendData.Num() > 0;
 	if (bContainsPerBoneWeights)
 	{
-		for (const TWeakObjectPtr<const UScriptStruct>& WeakScriptStruct : UniqueStructs)
+		for (const TWeakObjectPtr<UScriptStruct>& WeakScriptStruct : UniqueStructs)
 		{
 			const IAttributeBlendOperator* Operator = AttributeTypes::GetTypeOperator(WeakScriptStruct);
 
@@ -413,7 +413,7 @@ void Attributes::BlendAttributesPerBone(TArrayView<const FStackAttributeContaine
 			Weights.Add(SampleData.GetClampedWeight());
 		}
 
-		for (const TWeakObjectPtr<const UScriptStruct>& WeakScriptStruct : UniqueStructs)
+		for (const TWeakObjectPtr<UScriptStruct>& WeakScriptStruct : UniqueStructs)
 		{
 			const IAttributeBlendOperator* Operator = AttributeTypes::GetTypeOperator(WeakScriptStruct);
 
@@ -427,7 +427,7 @@ void Attributes::BlendAttributesPerBoneFilter(const FStackAttributeContainer& Ba
 {
 	QUICK_SCOPE_CYCLE_COUNTER(STAT_BlendAttributesPerBoneFilter);
 
-	TSet<TWeakObjectPtr<const UScriptStruct>, DefaultKeyFuncs<TWeakObjectPtr<const UScriptStruct>>, TInlineSetAllocator<4>> UniqueStructs;
+	TSet<TWeakObjectPtr<UScriptStruct>, DefaultKeyFuncs<TWeakObjectPtr<UScriptStruct>>, TInlineSetAllocator<4>> UniqueStructs;
 	for (const FStackAttributeContainer& CustomAttributes : BlendAttributes)
 	{
 		UniqueStructs.Append(CustomAttributes.GetUniqueTypes());
@@ -435,7 +435,7 @@ void Attributes::BlendAttributesPerBoneFilter(const FStackAttributeContainer& Ba
 
 	UniqueStructs.Append(BaseAttributes.GetUniqueTypes());
 
-	for (const TWeakObjectPtr<const UScriptStruct>& WeakScriptStruct : UniqueStructs)
+	for (const TWeakObjectPtr<UScriptStruct>& WeakScriptStruct : UniqueStructs)
 	{
 		const IAttributeBlendOperator* Operator = AttributeTypes::GetTypeOperator(WeakScriptStruct);
 		FAttributeBlendData AttributeBlendData = FAttributeBlendData::PerBoneFilteredWeighted(BaseAttributes, BlendAttributes, BoneBlendWeights, WeakScriptStruct.Get());
@@ -476,7 +476,7 @@ void Attributes::MirrorAttributes(FStackAttributeContainer& CustomAttributes, co
 		return true;
 	};
 
-	for (const TWeakObjectPtr<const UScriptStruct>& WeakScriptStruct : CustomAttributes.GetUniqueTypes())
+	for (const TWeakObjectPtr<UScriptStruct>& WeakScriptStruct : CustomAttributes.GetUniqueTypes())
 	{
 		const int32 TypeIndex = CustomAttributes.FindTypeIndex(WeakScriptStruct.Get());
 		if (TypeIndex != INDEX_NONE)
