AnimAttributeView: minor fix, hitting an avoidable ensure

diff --git a/Engine/Source/Editor/Persona/Private/SAnimAttributeView.cpp b/Engine/Source/Editor/Persona/Private/SAnimAttributeView.cpp
--- a/Engine/Source/Editor/Persona/Private/SAnimAttributeView.cpp
+++ b/Engine/Source/Editor/Persona/Private/SAnimAttributeView.cpp
@@ -546,10 +546,29 @@ void SAnimAttributeView::RefreshFilteredAttributeEntries()
 	}
 
 	ExecuteSort();
-
+	
 	// delay the refresh to until tick since this function
 	// can be invoked from animation thread
 	bShouldRefreshListView = true;
+
+	bool bSelectedAttributeStillValid = false;
+	for (const TSharedPtr<FAnimAttributeEntry>& Entry : FilteredAttributeEntries)
+	{
+		if (SelectedAttribute.IsSet())
+		{
+			if (*(Entry) == SelectedAttribute.GetValue())
+			{
+				bSelectedAttributeStillValid = true;
+				break;
+			}
+		}	
+	}
+
+	if (!bSelectedAttributeStillValid)
+	{
+		SelectedAttribute.Reset();
+		bShouldRefreshValueView = true;
+	}
 }
 
 void SAnimAttributeView::RefreshValueView()
diff --git a/Engine/Source/Editor/Persona/Public/SAnimAttributeView.h b/Engine/Source/Editor/Persona/Public/SAnimAttributeView.h
--- a/Engine/Source/Editor/Persona/Public/SAnimAttributeView.h
+++ b/Engine/Source/Editor/Persona/Public/SAnimAttributeView.h
@@ -53,7 +53,7 @@ public:
 
 	bool operator==(const FAnimAttributeEntry& InOther) const
 	{
-		return Identifier == InOther.Identifier;
+		return Identifier == InOther.Identifier && SnapshotDisplayName == InOther.SnapshotDisplayName;
 	}
 
 	bool operator!=(const FAnimAttributeEntry& InOther) const
