Removed OnStopEditingCurves delegate and behaviour from Persona module, this is now superseded with AnimDataModel notify callbacks instead
#misc fixed use after release crash

diff --git a/Engine/Source/Editor/Persona/Private/AnimTimeline/AnimModel_AnimSequenceBase.cpp b/Engine/Source/Editor/Persona/Private/AnimTimeline/AnimModel_AnimSequenceBase.cpp
--- a/Engine/Source/Editor/Persona/Private/AnimTimeline/AnimModel_AnimSequenceBase.cpp
+++ b/Engine/Source/Editor/Persona/Private/AnimTimeline/AnimModel_AnimSequenceBase.cpp
@@ -43,25 +43,14 @@ FAnimModel_AnimSequenceBase::FAnimModel_AnimSequenceBase(const TSharedRef<IPerso
 		bElementNodeDisplayFlag = false;
 	}
 
-	AnimSequenceBase->RegisterOnNotifyChanged(UAnimSequenceBase::FOnNotifyChanged::CreateRaw(this, &FAnimModel_AnimSequenceBase::RefreshSnapTimes));
-	
-	if(GEditor)
-	{
-		GEditor->RegisterForUndo(this);
-	}
+	AnimSequenceBase->RegisterOnNotifyChanged(UAnimSequenceBase::FOnNotifyChanged::CreateRaw(this, &FAnimModel_AnimSequenceBase::RefreshSnapTimes));	
 
 	AnimSequenceBase->GetDataModel()->GetModifiedEvent().AddRaw(this, &FAnimModel_AnimSequenceBase::OnDataModelChanged);
 }
 
 FAnimModel_AnimSequenceBase::~FAnimModel_AnimSequenceBase()
 {
-	if(GEditor)
-	{
-		GEditor->UnregisterForUndo(this);
-	}
-
 	AnimSequenceBase->UnregisterOnNotifyChanged(this);
-
 	AnimSequenceBase->GetDataModel()->GetModifiedEvent().RemoveAll(this);
 }
 
@@ -430,14 +419,6 @@ void FAnimModel_AnimSequenceBase::RemoveSelectedCurves()
 				FSmartName TrackName;
 				if (AnimSequenceBase->GetSkeleton()->GetSmartNameByUID(USkeleton::AnimCurveMappingName, CurveName.UID, TrackName))
 				{
-					// Stop editing this curve in the external editor window
-					FSmartName Name;
-					ERawCurveTrackTypes Type;
-					int32 CurveEditIndex;
-					FloatCurveTrack->GetCurveEditInfo(0, Name, Type, CurveEditIndex);
-					IAnimationEditor::FCurveEditInfo EditInfo(Name, Type, CurveEditIndex);
-					OnStopEditingCurves.ExecuteIfBound(TArray<IAnimationEditor::FCurveEditInfo>({ EditInfo }));
-
 					Controller.RemoveCurve(FloatCurveId);
 					bDeletedCurve = true;
 				}
@@ -457,20 +438,6 @@ void FAnimModel_AnimSequenceBase::RemoveSelectedCurves()
 				FSmartName CurveToDelete;
 				if (AnimSequenceBase->GetSkeleton()->GetSmartNameByUID(USkeleton::AnimTrackCurveMappingName, CurveName.UID, CurveToDelete))
 				{
-					// Stop editing these curves in the external editor window
-					TArray<IAnimationEditor::FCurveEditInfo> CurveEditInfo;
-					for(int32 CurveIndex = 0; CurveIndex < TransformCurveTrack->GetCurves().Num(); ++CurveIndex)
-					{
-						FSmartName Name;
-						ERawCurveTrackTypes Type;
-						int32 CurveEditIndex;
-						TransformCurveTrack->GetCurveEditInfo(CurveIndex, Name, Type, CurveEditIndex);
-						IAnimationEditor::FCurveEditInfo EditInfo(Name, Type, CurveEditIndex);
-						CurveEditInfo.Add(EditInfo);
-					}
-
-					OnStopEditingCurves.ExecuteIfBound(CurveEditInfo);
-
 					Controller.RemoveCurve(TransformCurveId);
 					bDeletedCurve = true;
 				}
@@ -546,19 +513,6 @@ bool FAnimModel_AnimSequenceBase::IsDisplaySecondaryChecked() const
 	return GetDefault<UPersonaOptions>()->bTimelineDisplayFormatSecondary;
 }
 
-void FAnimModel_AnimSequenceBase::HandleUndoRedo()
-{
-	// Close any curves that are no longer editable
-	for (const FFloatCurve& FloatCurve : AnimSequenceBase->GetDataModel()->GetCurveData().FloatCurves)
-	{
-		if(FloatCurve.GetCurveTypeFlag(AACF_Metadata))
-		{
-			IAnimationEditor::FCurveEditInfo CurveEditInfo(FloatCurve.Name, ERawCurveTrackTypes::RCT_Float, 0);
-			OnStopEditingCurves.ExecuteIfBound(TArray<IAnimationEditor::FCurveEditInfo>({ CurveEditInfo }));
-		}
-	}
-}
-
 void FAnimModel_AnimSequenceBase::UpdateRange()
 {
 	FAnimatedRange OldPlaybackRange = PlaybackRange;
diff --git a/Engine/Source/Editor/Persona/Private/AnimTimeline/AnimModel_AnimSequenceBase.h b/Engine/Source/Editor/Persona/Private/AnimTimeline/AnimModel_AnimSequenceBase.h
--- a/Engine/Source/Editor/Persona/Private/AnimTimeline/AnimModel_AnimSequenceBase.h
+++ b/Engine/Source/Editor/Persona/Private/AnimTimeline/AnimModel_AnimSequenceBase.h
@@ -18,7 +18,7 @@ class FAnimTimelineTrack_Attributes;
 enum class EFrameNumberDisplayFormats : uint8;
 
 /** Anim model for an anim sequence base */
-class FAnimModel_AnimSequenceBase : public FAnimModel, public FEditorUndoClient
+class FAnimModel_AnimSequenceBase : public FAnimModel
 {
 public:
 	FAnimModel_AnimSequenceBase(const TSharedRef<IPersonaPreviewScene>& InPreviewScene, const TSharedRef<IEditableSkeleton>& InEditableSkeleton, const TSharedRef<FUICommandList>& InCommandList, UAnimSequenceBase* InAnimSequenceBase);
@@ -31,18 +31,11 @@ public:
 	virtual void Initialize() override;
 	virtual void UpdateRange() override;
 
-	/** FEditorUndoClient interface */
-	virtual void PostUndo(bool bSuccess) override { HandleUndoRedo(); }
-	virtual void PostRedo(bool bSuccess) override { HandleUndoRedo(); }
-
 	const TSharedPtr<FAnimTimelineTrack_Notifies>& GetNotifyRoot() const { return NotifyRoot; }
 
 	/** Delegate used to edit curves */
 	FOnEditCurves OnEditCurves;
 
-	/** Delegate used to edit curves */
-	FOnStopEditingCurves OnStopEditingCurves;
-
 	/** Notify track timing options */
 	bool IsNotifiesTimingElementDisplayEnabled(ETimingElementType::Type ElementType) const;
 	void ToggleNotifiesTimingElementDisplayEnabled(ETimingElementType::Type ElementType);
@@ -81,7 +74,6 @@ private:
 	bool IsDisplayPercentageChecked() const;
 	void ToggleDisplaySecondary();
 	bool IsDisplaySecondaryChecked() const;
-	void HandleUndoRedo();
 
 private:
 	/** The anim sequence base we wrap */
diff --git a/Engine/Source/Editor/Persona/Private/AnimTimeline/AnimTimelineTrack_FloatCurve.cpp b/Engine/Source/Editor/Persona/Private/AnimTimeline/AnimTimelineTrack_FloatCurve.cpp
--- a/Engine/Source/Editor/Persona/Private/AnimTimeline/AnimTimelineTrack_FloatCurve.cpp
+++ b/Engine/Source/Editor/Persona/Private/AnimTimeline/AnimTimelineTrack_FloatCurve.cpp
@@ -21,7 +21,6 @@
 #include "AnimModel_AnimSequenceBase.h"
 #include "Animation/AnimData/AnimDataModel.h"
 #include "SAnimOutlinerItem.h"
-#include "Animation/AnimData/AnimDataModel.h"
 
 #define LOCTEXT_NAMESPACE "FAnimTimelineTrack_FloatCurve"
 
@@ -153,10 +152,6 @@ void FAnimTimelineTrack_FloatCurve::ConvertCurveToMetaData()
 {
 	UAnimSequenceBase* AnimSequenceBase = GetModel()->GetAnimSequenceBase();
 
-	// Stop editing this curve in the external editor window
-	IAnimationEditor::FCurveEditInfo EditInfo(CurveName, ERawCurveTrackTypes::RCT_Float, 0);
-	StaticCastSharedRef<FAnimModel_AnimSequenceBase>(GetModel())->OnStopEditingCurves.ExecuteIfBound(TArray<IAnimationEditor::FCurveEditInfo>({ EditInfo }));
-
 	IAnimationDataController& Controller = AnimSequenceBase->GetController();
 	IAnimationDataController::FScopedBracket ScopedBracket(Controller, LOCTEXT("ConvertCurveToMetaData_Bracket", "Converting curve to metadata"));
 	Controller.SetCurveFlag(CurveId, AACF_Metadata, true);
@@ -185,8 +180,7 @@ void FAnimTimelineTrack_FloatCurve::RemoveCurve()
 		{
 			// Stop editing this curve in the external editor window
 			IAnimationEditor::FCurveEditInfo EditInfo(CurveName, ERawCurveTrackTypes::RCT_Float, 0);
-			StaticCastSharedRef<FAnimModel_AnimSequenceBase>(GetModel())->OnStopEditingCurves.ExecuteIfBound(TArray<IAnimationEditor::FCurveEditInfo>({ EditInfo }));
-
+			
 			AnimSequenceBase->Modify(true);
 
 			IAnimationDataController& Controller = AnimSequenceBase->GetController();
@@ -207,10 +201,6 @@ void FAnimTimelineTrack_FloatCurve::OnCommitCurveName(const FText& InText, EText
 		FText CurrentCurveName = GetLabel();
 		if (!CurrentCurveName.EqualToCaseIgnored(InText))
 		{
-			// Stop editing this curve in the external editor window
-			IAnimationEditor::FCurveEditInfo EditInfo(CurveName, ERawCurveTrackTypes::RCT_Float, 0);
-			StaticCastSharedRef<FAnimModel_AnimSequenceBase>(GetModel())->OnStopEditingCurves.ExecuteIfBound(TArray<IAnimationEditor::FCurveEditInfo>({ EditInfo }));
-
 			// Check that the name doesn't already exist
 			const FName RequestedName = FName(*InText.ToString());
 
diff --git a/Engine/Source/Editor/Persona/Private/AnimTimeline/AnimTimelineTrack_TransformCurve.cpp b/Engine/Source/Editor/Persona/Private/AnimTimeline/AnimTimelineTrack_TransformCurve.cpp
--- a/Engine/Source/Editor/Persona/Private/AnimTimeline/AnimTimelineTrack_TransformCurve.cpp
+++ b/Engine/Source/Editor/Persona/Private/AnimTimeline/AnimTimelineTrack_TransformCurve.cpp
@@ -127,21 +127,6 @@ void FAnimTimelineTrack_TransformCurve::DeleteTrack()
 		FSmartName CurveToDelete;
 		if (AnimSequenceBase->GetSkeleton()->GetSmartNameByUID(USkeleton::AnimTrackCurveMappingName, TransformCurve->Name.UID, CurveToDelete))
 		{
-			// Stop editing these curves in the external editor window
-			TArray<IAnimationEditor::FCurveEditInfo> CurveEditInfo;
-			for(int32 CurveIndex = 0; CurveIndex < Curves.Num(); ++CurveIndex)
-			{
-				FSmartName Name;
-				ERawCurveTrackTypes Type;
-				int32 CurveEditIndex;
-				GetCurveEditInfo(CurveIndex, Name, Type, CurveEditIndex);
-				IAnimationEditor::FCurveEditInfo EditInfo(Name, Type, CurveEditIndex);
-				CurveEditInfo.Add(EditInfo);
-			}
-
-			BaseModel->OnStopEditingCurves.ExecuteIfBound(CurveEditInfo);
-
-
 			IAnimationDataController& Controller = AnimSequenceBase->GetController();
 			Controller.RemoveCurve(CurveId);
 
diff --git a/Engine/Source/Editor/Persona/Private/PersonaModule.cpp b/Engine/Source/Editor/Persona/Private/PersonaModule.cpp
--- a/Engine/Source/Editor/Persona/Private/PersonaModule.cpp
+++ b/Engine/Source/Editor/Persona/Private/PersonaModule.cpp
@@ -391,19 +391,13 @@ TSharedRef<SWidget> FPersonaModule::CreateEditorWidgetForAnimDocument(const TSha
 			WeakHostingApp.Pin()->EditCurves(InAnimSequence, InCurveInfo, InExternalTimeSliderController);
 		};
 
-		auto OnStopEditingCurves = [WeakHostingApp](const TArray<IAnimationEditor::FCurveEditInfo>& InCurveInfo)
-		{ 
-			WeakHostingApp.Pin()->StopEditingCurves(InCurveInfo);
-		};
-
 		if (UAnimSequence* Sequence = Cast<UAnimSequence>(InAnimAsset))
 		{
 			Result = SNew(SSequenceEditor, InArgs.PreviewScene.Pin().ToSharedRef(), InArgs.EditableSkeleton.Pin().ToSharedRef(), InHostingApp->GetToolkitCommands())
 				.Sequence(Sequence)
 				.OnObjectsSelected(InArgs.OnDespatchObjectsSelected)
 				.OnInvokeTab(InArgs.OnDespatchInvokeTab)
-				.OnEditCurves_Lambda(OnEditCurves)
-				.OnStopEditingCurves_Lambda(OnStopEditingCurves);
+				.OnEditCurves_Lambda(OnEditCurves);
 
 			OutDocumentLink = TEXT("Engine/Animation/Sequences");
 		}
@@ -413,8 +407,7 @@ TSharedRef<SWidget> FPersonaModule::CreateEditorWidgetForAnimDocument(const TSha
 				.Composite(Composite)
 				.OnObjectsSelected(InArgs.OnDespatchObjectsSelected)
 				.OnInvokeTab(InArgs.OnDespatchInvokeTab)
-				.OnEditCurves_Lambda(OnEditCurves)
-				.OnStopEditingCurves_Lambda(OnStopEditingCurves);
+				.OnEditCurves_Lambda(OnEditCurves);
 
 			OutDocumentLink = TEXT("Engine/Animation/AnimationComposite");
 		}
@@ -427,8 +420,7 @@ TSharedRef<SWidget> FPersonaModule::CreateEditorWidgetForAnimDocument(const TSha
 				.OnSectionsChanged(InArgs.OnDespatchSectionsChanged)
 				.OnInvokeTab(InArgs.OnDespatchInvokeTab)
 				.OnObjectsSelected(InArgs.OnDespatchObjectsSelected)
-				.OnEditCurves_Lambda(OnEditCurves)
-				.OnStopEditingCurves_Lambda(OnStopEditingCurves);
+				.OnEditCurves_Lambda(OnEditCurves);
 
 			OutDocumentLink = TEXT("Engine/Animation/AnimMontage");
 		}
@@ -438,8 +430,7 @@ TSharedRef<SWidget> FPersonaModule::CreateEditorWidgetForAnimDocument(const TSha
 				.StreamableAnim(StreamableAnim)
 				.OnObjectsSelected(InArgs.OnDespatchObjectsSelected)
 				.OnInvokeTab(InArgs.OnDespatchInvokeTab)
-				.OnEditCurves_Lambda(OnEditCurves)
-				.OnStopEditingCurves_Lambda(OnStopEditingCurves);
+				.OnEditCurves_Lambda(OnEditCurves);
 
 			OutDocumentLink = TEXT("Engine/Animation/Sequences");
 		}
diff --git a/Engine/Source/Editor/Persona/Private/SAnimCompositeEditor.cpp b/Engine/Source/Editor/Persona/Private/SAnimCompositeEditor.cpp
--- a/Engine/Source/Editor/Persona/Private/SAnimCompositeEditor.cpp
+++ b/Engine/Source/Editor/Persona/Private/SAnimCompositeEditor.cpp
@@ -32,7 +32,6 @@ void SAnimCompositeEditor::Construct(const FArguments& InArgs, const TSharedRef<
 		InOnEditCurves.ExecuteIfBound(InAnimSequence, InCurveInfo, TimelineWidget->GetTimeSliderController());
 	});
 
-	AnimModel->OnStopEditingCurves = InArgs._OnStopEditingCurves;
 	AnimModel->Initialize();
 
 	SAnimEditorBase::Construct( SAnimEditorBase::FArguments()
diff --git a/Engine/Source/Editor/Persona/Private/SAnimCompositeEditor.h b/Engine/Source/Editor/Persona/Private/SAnimCompositeEditor.h
--- a/Engine/Source/Editor/Persona/Private/SAnimCompositeEditor.h
+++ b/Engine/Source/Editor/Persona/Private/SAnimCompositeEditor.h
@@ -33,7 +33,6 @@ public:
 		SLATE_EVENT(FOnObjectsSelected, OnObjectsSelected)
 		SLATE_EVENT(FOnInvokeTab, OnInvokeTab)
 		SLATE_EVENT(FOnEditCurves, OnEditCurves)
-		SLATE_EVENT(FOnStopEditingCurves, OnStopEditingCurves)
 
 	SLATE_END_ARGS()
 
diff --git a/Engine/Source/Editor/Persona/Private/SAnimSequenceCurveEditor.cpp b/Engine/Source/Editor/Persona/Private/SAnimSequenceCurveEditor.cpp
--- a/Engine/Source/Editor/Persona/Private/SAnimSequenceCurveEditor.cpp
+++ b/Engine/Source/Editor/Persona/Private/SAnimSequenceCurveEditor.cpp
@@ -363,19 +363,11 @@ public:
 	TWeakPtr<ITimeSliderController> ExternalTimeSliderController;
 };
 
-SAnimSequenceCurveEditor::SAnimSequenceCurveEditor()
-{
-	if(GEditor)
-	{
-		GEditor->RegisterForUndo(this);
-	}
-}
-
 SAnimSequenceCurveEditor::~SAnimSequenceCurveEditor()
 {
-	if(GEditor)
+	if(AnimSequence)
 	{
-		GEditor->UnregisterForUndo(this);
+		AnimSequence->GetDataModel()->GetModifiedEvent().RemoveAll(this);		
 	}
 }
 
@@ -391,6 +383,8 @@ void SAnimSequenceCurveEditor::Construct(const FArguments& InArgs, const TShared
 
 	AnimSequence = InAnimSequence;
 
+	AnimSequence->GetDataModel()->GetModifiedEvent().AddRaw(this, &SAnimSequenceCurveEditor::OnModelHasChanged);
+
 	CurveEditorTree = SNew(SCurveEditorTree, CurveEditor);
 
 	TSharedRef<SCurveEditorPanel> CurveEditorPanel = SNew(SCurveEditorPanel, CurveEditor.ToSharedRef())
@@ -444,6 +438,40 @@ void SAnimSequenceCurveEditor::Construct(const FArguments& InArgs, const TShared
 	];
 }
 
+void SAnimSequenceCurveEditor::OnModelHasChanged(const EAnimDataModelNotifyType& NotifyType, UAnimDataModel* Model, const FAnimDataModelNotifPayload& Payload)
+{
+	auto StopEditingCurve = [this, NotifyType, &Payload, Model]()
+	{
+		const FCurvePayload& TypedPayload = Payload.GetPayload<FCurvePayload>();
+		const FAnimationCurveIdentifier& CurveId = TypedPayload.Identifier;
+		const int32 ChannelIndices = CurveId.CurveType == ERawCurveTrackTypes::RCT_Transform ? 9 : 1;
+		for (int32 ChannelIndex = 0; ChannelIndex < ChannelIndices; ++ChannelIndex)
+		{
+			RemoveCurve(CurveId.InternalName, CurveId.CurveType, ChannelIndex);
+		}
+	};
+	
+	switch(NotifyType)
+	{
+	case EAnimDataModelNotifyType::CurveRemoved:
+	case EAnimDataModelNotifyType::CurveRenamed:
+		{
+			StopEditingCurve();
+			break;
+		}			
+	case EAnimDataModelNotifyType::CurveFlagsChanged:
+		{
+			const FCurveFlagsChangedPayload& TypedPayload = Payload.GetPayload<FCurveFlagsChangedPayload>();
+			if(Model->FindCurve(TypedPayload.Identifier)->GetCurveTypeFlag(AACF_Metadata))
+			{
+				StopEditingCurve();
+			}
+			break;
+		}
+	}
+	
+}
+
 TSharedRef<SWidget> SAnimSequenceCurveEditor::MakeToolbar(TSharedRef<SCurveEditorPanel> InEditorPanel)
 {
 	FToolBarBuilder ToolBarBuilder(InEditorPanel->GetCommands(), FMultiBoxCustomization::None, InEditorPanel->GetToolbarExtender(), true);
@@ -486,8 +514,10 @@ void SAnimSequenceCurveEditor::RemoveCurve(const FSmartName& InName, ERawCurveTr
 		FRichCurveEditorModelNamed* Model = static_cast<FRichCurveEditorModelNamed*>(CurvePair.Value.Get());
 		if(Model->Name == InName && Model->Type == InType && Model->CurveIndex == InCurveIndex)
 		{
+			// Cache ID to prevent use after release
+			const FCurveEditorTreeItemID TreeId = Model->TreeId; 
 			CurveEditor->RemoveCurve(CurvePair.Key);
-			CurveEditor->RemoveTreeItem(Model->TreeId);
+			CurveEditor->RemoveTreeItem(TreeId);
 			break;
 		}
 	}
@@ -498,26 +528,4 @@ void SAnimSequenceCurveEditor::ZoomToFit()
 	CurveEditor->ZoomToFit(EAxisList::Y);
 }
 
-void SAnimSequenceCurveEditor::PostUndoRedo()
-{
-	// Check if our curves are still valid & remove if not
-	bool bRemoved = false;
-	do
-	{
-		bRemoved = false;
-		for(const TPair<FCurveModelID, TUniquePtr<FCurveModel>>& Pair : CurveEditor->GetCurves())
-		{
-			FRichCurveEditorModelNamed* Model = static_cast<FRichCurveEditorModelNamed*>(Pair.Value.Get());
-			if(!Model->IsValid())
-			{
-				CurveEditor->RemoveCurve(Pair.Key);
-				CurveEditor->RemoveTreeItem(Model->TreeId);
-				bRemoved = true;
-				break;
-			}
-		}
-	}
-	while(bRemoved);
-}
-
 #undef LOCTEXT_NAMESPACE
\ No newline at end of file
diff --git a/Engine/Source/Editor/Persona/Private/SAnimSequenceCurveEditor.h b/Engine/Source/Editor/Persona/Private/SAnimSequenceCurveEditor.h
--- a/Engine/Source/Editor/Persona/Private/SAnimSequenceCurveEditor.h
+++ b/Engine/Source/Editor/Persona/Private/SAnimSequenceCurveEditor.h
@@ -52,7 +52,7 @@ public:
 	TUniquePtr<IAnimationDataController::FScopedBracket> InteractiveBracket;
 };
 
-class SAnimSequenceCurveEditor : public IAnimSequenceCurveEditor, public FEditorUndoClient
+class SAnimSequenceCurveEditor : public IAnimSequenceCurveEditor
 {
 	SLATE_BEGIN_ARGS(SAnimSequenceCurveEditor) {}
 
@@ -62,7 +62,6 @@ class SAnimSequenceCurveEditor : public IAnimSequenceCurveEditor, public FEditor
 
 	SLATE_END_ARGS()
 
-	SAnimSequenceCurveEditor();
 	~SAnimSequenceCurveEditor();
 
 	void Construct(const FArguments& InArgs, const TSharedRef<IPersonaPreviewScene>& InPreviewScene, UAnimSequenceBase* InAnimSequence);
@@ -73,16 +72,12 @@ class SAnimSequenceCurveEditor : public IAnimSequenceCurveEditor, public FEditor
 	virtual void RemoveCurve(const FSmartName& InName, ERawCurveTrackTypes InType, int32 InCurveIndex) override;
 	virtual void ZoomToFit() override;
 	
-	virtual void PostUndo(bool bSuccess) override { PostUndoRedo(); }
-	virtual void PostRedo(bool bSuccess) override { PostUndoRedo(); }
+	void OnModelHasChanged(const EAnimDataModelNotifyType& NotifyType, UAnimDataModel* Model, const FAnimDataModelNotifPayload& Payload);
 
 private:
 	// Build the toolbar for this curve editor
 	TSharedRef<SWidget> MakeToolbar(TSharedRef<SCurveEditorPanel> InEditorPanel);
 
-	// Handle undo/redo to check underlying curve data is still valid
-	void PostUndoRedo();
-
 private:
 	/** The actual curve editor */
 	TSharedPtr<FCurveEditor> CurveEditor;
diff --git a/Engine/Source/Editor/Persona/Private/SAnimStreamableEditor.cpp b/Engine/Source/Editor/Persona/Private/SAnimStreamableEditor.cpp
--- a/Engine/Source/Editor/Persona/Private/SAnimStreamableEditor.cpp
+++ b/Engine/Source/Editor/Persona/Private/SAnimStreamableEditor.cpp
@@ -28,7 +28,6 @@ void SAnimStreamableEditor::Construct(const FArguments& InArgs, const TSharedRef
 		InOnEditCurves.ExecuteIfBound(InAnimSequence, InCurveInfo, TimelineWidget->GetTimeSliderController());
 	});
 
-	AnimModel->OnStopEditingCurves = InArgs._OnStopEditingCurves;
 	AnimModel->Initialize();
 
 	SAnimEditorBase::Construct(SAnimEditorBase::FArguments()
diff --git a/Engine/Source/Editor/Persona/Private/SAnimStreamableEditor.h b/Engine/Source/Editor/Persona/Private/SAnimStreamableEditor.h
--- a/Engine/Source/Editor/Persona/Private/SAnimStreamableEditor.h
+++ b/Engine/Source/Editor/Persona/Private/SAnimStreamableEditor.h
@@ -27,7 +27,6 @@ public:
 	SLATE_EVENT(FOnObjectsSelected, OnObjectsSelected)
 	SLATE_EVENT(FOnInvokeTab, OnInvokeTab)
 	SLATE_EVENT(FOnEditCurves, OnEditCurves)
-	SLATE_EVENT(FOnStopEditingCurves, OnStopEditingCurves)
 
 	SLATE_END_ARGS()
 
diff --git a/Engine/Source/Editor/Persona/Private/SMontageEditor.cpp b/Engine/Source/Editor/Persona/Private/SMontageEditor.cpp
--- a/Engine/Source/Editor/Persona/Private/SMontageEditor.cpp
+++ b/Engine/Source/Editor/Persona/Private/SMontageEditor.cpp
@@ -51,7 +51,6 @@ void SMontageEditor::Construct(const FArguments& InArgs, const FMontageEditorReq
 		InOnEditCurves.ExecuteIfBound(InAnimSequence, InCurveInfo, TimelineWidget->GetTimeSliderController());
 	});
 
-	AnimModel->OnStopEditingCurves = InArgs._OnStopEditingCurves;
 	AnimModel->Initialize();
 
 	MontageObj->RegisterOnMontageChanged(UAnimMontage::FOnMontageChanged::CreateSP(AnimModel.Get(), &FAnimModel_AnimMontage::RefreshTracks));
diff --git a/Engine/Source/Editor/Persona/Private/SMontageEditor.h b/Engine/Source/Editor/Persona/Private/SMontageEditor.h
--- a/Engine/Source/Editor/Persona/Private/SMontageEditor.h
+++ b/Engine/Source/Editor/Persona/Private/SMontageEditor.h
@@ -54,7 +54,6 @@ public:
 	SLATE_ARGUMENT( UAnimMontage*, Montage )
 	SLATE_EVENT(FOnObjectsSelected, OnObjectsSelected)
 	SLATE_EVENT(FOnEditCurves, OnEditCurves)
-	SLATE_EVENT(FOnStopEditingCurves, OnStopEditingCurves)
 
 	SLATE_END_ARGS()
 
diff --git a/Engine/Source/Editor/Persona/Private/SSequenceEditor.cpp b/Engine/Source/Editor/Persona/Private/SSequenceEditor.cpp
--- a/Engine/Source/Editor/Persona/Private/SSequenceEditor.cpp
+++ b/Engine/Source/Editor/Persona/Private/SSequenceEditor.cpp
@@ -28,7 +28,6 @@ void SSequenceEditor::Construct(const FArguments& InArgs, TSharedRef<class IPers
 		InOnEditCurves.ExecuteIfBound(InAnimSequence, InCurveInfo, TimelineWidget->GetTimeSliderController());
 	});
 
-	AnimModel->OnStopEditingCurves = InArgs._OnStopEditingCurves;
 	AnimModel->OnSelectObjects = FOnObjectsSelected::CreateSP(this, &SAnimEditorBase::OnSelectionChanged);
 	AnimModel->OnInvokeTab = InArgs._OnInvokeTab;
 	AnimModel->Initialize();
diff --git a/Engine/Source/Editor/Persona/Private/SSequenceEditor.h b/Engine/Source/Editor/Persona/Private/SSequenceEditor.h
--- a/Engine/Source/Editor/Persona/Private/SSequenceEditor.h
+++ b/Engine/Source/Editor/Persona/Private/SSequenceEditor.h
@@ -29,7 +29,6 @@ public:
 		SLATE_EVENT(FOnObjectsSelected, OnObjectsSelected)
 		SLATE_EVENT(FOnInvokeTab, OnInvokeTab)
 		SLATE_EVENT(FOnEditCurves, OnEditCurves)
-		SLATE_EVENT(FOnStopEditingCurves, OnStopEditingCurves)
 
 	SLATE_END_ARGS()
 
diff --git a/Engine/Source/Editor/Persona/Public/PersonaDelegates.h b/Engine/Source/Editor/Persona/Public/PersonaDelegates.h
--- a/Engine/Source/Editor/Persona/Public/PersonaDelegates.h
+++ b/Engine/Source/Editor/Persona/Public/PersonaDelegates.h
@@ -35,6 +35,7 @@ DECLARE_DELEGATE_RetVal(UObject*, FOnGetAsset);
 DECLARE_DELEGATE_ThreeParams(FOnEditCurves, UAnimSequenceBase* /*InAnimSequence*/, const TArray<IAnimationEditor::FCurveEditInfo>& /*InCurveInfo*/, const TSharedPtr<ITimeSliderController>& /*InTimeSliderController*/);
 
 // Called to stop editing curves in the curve editor
+UE_DEPRECATED(5.1, "FOnStopEditingCurves has been deprecated and its use-case replaced with Animation Data Model notifies")
 DECLARE_DELEGATE_OneParam(FOnStopEditingCurves, const TArray<IAnimationEditor::FCurveEditInfo>& /*InCurveInfo*/);
 
 // Called when the user navigates with pageup
