Added track filtering to anim timeline

Also edigrated CL 19957678 to fix small timelines not allowing scroll/scrub interactions

diff --git a/Engine/Source/Editor/Persona/Private/AnimTimeline/AnimModel.cpp b/Engine/Source/Editor/Persona/Private/AnimTimeline/AnimModel.cpp
--- a/Engine/Source/Editor/Persona/Private/AnimTimeline/AnimModel.cpp
+++ b/Engine/Source/Editor/Persona/Private/AnimTimeline/AnimModel.cpp
@@ -346,4 +346,31 @@ void FAnimModel::BuildContextMenu(FMenuBuilder& InMenuBuilder)
 	}
 }
 
+void FAnimModel::AddRootTrack(TSharedRef<FAnimTimelineTrack> InTrack)
+{
+	if (GetMutableDefault<UPersonaOptions>()->GetAllowedAnimationEditorTracks().PassesFilter(InTrack->GetTypeName()))
+	{
+		PRAGMA_DISABLE_DEPRECATION_WARNINGS
+		RootTracks.Add(InTrack);
+		PRAGMA_ENABLE_DEPRECATION_WARNINGS
+	}
+}
+
+void FAnimModel::ClearRootTracks()
+{
+	PRAGMA_DISABLE_DEPRECATION_WARNINGS
+	RootTracks.Empty();
+	PRAGMA_ENABLE_DEPRECATION_WARNINGS
+}
+
+void FAnimModel::ForEachRootTrack(TFunctionRef<void(FAnimTimelineTrack&)> InFunction)
+{
+	PRAGMA_DISABLE_DEPRECATION_WARNINGS
+	for (TSharedRef<FAnimTimelineTrack>& Track : RootTracks)
+	{
+		InFunction(Track.Get());
+	}
+	PRAGMA_ENABLE_DEPRECATION_WARNINGS
+}
+
 #undef LOCTEXT_NAMESPACE
\ No newline at end of file
diff --git a/Engine/Source/Editor/Persona/Private/AnimTimeline/AnimModel.h b/Engine/Source/Editor/Persona/Private/AnimTimeline/AnimModel.h
--- a/Engine/Source/Editor/Persona/Private/AnimTimeline/AnimModel.h
+++ b/Engine/Source/Editor/Persona/Private/AnimTimeline/AnimModel.h
@@ -23,14 +23,37 @@ class FAnimModel : public TSharedFromThis<FAnimModel>, public FGCObject
 public:
 	FAnimModel(const TSharedRef<IPersonaPreviewScene>& InPreviewScene, const TSharedRef<IEditableSkeleton>& InEditableSkeleton, const TSharedRef<FUICommandList>& InCommandList);
 
+	PRAGMA_DISABLE_DEPRECATION_WARNINGS
 	virtual ~FAnimModel() {}
+	PRAGMA_ENABLE_DEPRECATION_WARNINGS
 
 	/** Bind commands and perform any one-time initialization */
 	virtual void Initialize();
 
 	/** Get the root tracks representing the tree */
-	TArray<TSharedRef<FAnimTimelineTrack>>& GetRootTracks() { return RootTracks; }
-	const TArray<TSharedRef<FAnimTimelineTrack>>& GetRootTracks() const { return RootTracks; }
+	UE_DEPRECATED(5.1, "Mutable access to root tracks is deprecated, please use AddRootTrack/ClearRootTracks/ForEachRootTrack")
+	TArray<TSharedRef<FAnimTimelineTrack>>& GetRootTracks() 
+	{
+		PRAGMA_DISABLE_DEPRECATION_WARNINGS
+		return RootTracks;
+		PRAGMA_ENABLE_DEPRECATION_WARNINGS
+	}
+
+	/** Get the root tracks representing the tree */
+	const TArray<TSharedRef<FAnimTimelineTrack>>& GetRootTracks() const 
+	{
+		PRAGMA_DISABLE_DEPRECATION_WARNINGS
+		return RootTracks;
+		PRAGMA_ENABLE_DEPRECATION_WARNINGS
+	}
+
+	/** Get the root tracks representing the tree - alias to aid deprecation */
+	const TArray<TSharedRef<FAnimTimelineTrack>>& GetAllRootTracks() const
+	{
+		PRAGMA_DISABLE_DEPRECATION_WARNINGS
+		return RootTracks;
+		PRAGMA_ENABLE_DEPRECATION_WARNINGS
+	}
 
 	/** Get the current view range */
 	FAnimatedRange GetViewRange() const;
@@ -172,8 +195,19 @@ public:
 	/** Build a context menu for selected items */
 	virtual void BuildContextMenu(FMenuBuilder& InMenuBuilder);
 
+	/** Run a predicate for each root track */
+	void ForEachRootTrack(TFunctionRef<void(FAnimTimelineTrack&)> InFunction);
+
 protected:
+	/** Adds a track to the root tracks collection */
+	void AddRootTrack(TSharedRef<FAnimTimelineTrack> InTrack);
+
+	/** Clears all root tracks */
+	void ClearRootTracks();
+
+protected:	
 	/** Tracks used to generate a tree */
+	UE_DEPRECATED(5.1, "Direct access to RootTracks is deprecated, please use GetRootTracks/AddRootTrack/ClearRootTracks")
 	TArray<TSharedRef<FAnimTimelineTrack>> RootTracks;
 
 	/** Tracks that are selected */
diff --git a/Engine/Source/Editor/Persona/Private/AnimTimeline/AnimModel_AnimComposite.cpp b/Engine/Source/Editor/Persona/Private/AnimTimeline/AnimModel_AnimComposite.cpp
--- a/Engine/Source/Editor/Persona/Private/AnimTimeline/AnimModel_AnimComposite.cpp
+++ b/Engine/Source/Editor/Persona/Private/AnimTimeline/AnimModel_AnimComposite.cpp
@@ -22,7 +22,7 @@ void FAnimModel_AnimComposite::RefreshTracks()
 	ClearTrackSelection();
 
 	// Clear all tracks
-	RootTracks.Empty();
+	ClearRootTracks();
 
 	// Add the composite root track
 	if(!CompositeRoot.IsValid())
@@ -31,7 +31,7 @@ void FAnimModel_AnimComposite::RefreshTracks()
 	}
 
 	CompositeRoot->ClearChildren();
-	RootTracks.Add(CompositeRoot.ToSharedRef());
+	AddRootTrack(CompositeRoot.ToSharedRef());
 
 	TSharedRef<FAnimTimelineTrack_CompositePanel> CompositePanel = MakeShared<FAnimTimelineTrack_CompositePanel>(SharedThis(this));
 	CompositeRoot->AddChild(CompositePanel);
diff --git a/Engine/Source/Editor/Persona/Private/AnimTimeline/AnimModel_AnimMontage.cpp b/Engine/Source/Editor/Persona/Private/AnimTimeline/AnimModel_AnimMontage.cpp
--- a/Engine/Source/Editor/Persona/Private/AnimTimeline/AnimModel_AnimMontage.cpp
+++ b/Engine/Source/Editor/Persona/Private/AnimTimeline/AnimModel_AnimMontage.cpp
@@ -46,7 +46,7 @@ void FAnimModel_AnimMontage::RefreshTracks()
 	ClearTrackSelection();
 
 	// Clear all tracks
-	RootTracks.Empty();
+	ClearRootTracks();
 
 	bool bIsChildAnimMontage = AnimMontage->HasParentAsset();
 
@@ -57,7 +57,7 @@ void FAnimModel_AnimMontage::RefreshTracks()
 	}
 
 	MontageRoot->ClearChildren();
-	RootTracks.Add(MontageRoot.ToSharedRef());
+	AddRootTrack(MontageRoot.ToSharedRef());
 
 	// Create & add the montage panel
 	MontagePanel = MakeShared<FAnimTimelineTrack_MontagePanel>(SharedThis(this));
diff --git a/Engine/Source/Editor/Persona/Private/AnimTimeline/AnimModel_AnimSequenceBase.cpp b/Engine/Source/Editor/Persona/Private/AnimTimeline/AnimModel_AnimSequenceBase.cpp
--- a/Engine/Source/Editor/Persona/Private/AnimTimeline/AnimModel_AnimSequenceBase.cpp
+++ b/Engine/Source/Editor/Persona/Private/AnimTimeline/AnimModel_AnimSequenceBase.cpp
@@ -144,7 +144,7 @@ void FAnimModel_AnimSequenceBase::RefreshTracks()
 	ClearTrackSelection();
 
 	// Clear all tracks
-	RootTracks.Empty();
+	ClearRootTracks();
 
 	// Add notifies
 	RefreshNotifyTracks();
@@ -178,7 +178,7 @@ void FAnimModel_AnimSequenceBase::RefreshNotifyTracks()
 	}
 
 	NotifyRoot->ClearChildren();
-	RootTracks.Add(NotifyRoot.ToSharedRef());
+	AddRootTrack(NotifyRoot.ToSharedRef());
 
 	if(!NotifyPanel.IsValid())
 	{
@@ -198,7 +198,7 @@ void FAnimModel_AnimSequenceBase::RefreshCurveTracks()
 	}
 
 	CurveRoot->ClearChildren();
-	RootTracks.Add(CurveRoot.ToSharedRef());
+	AddRootTrack(CurveRoot.ToSharedRef());
 
 	// Next add a track for each float curve
 	const FAnimationCurveData& AnimationModelCurveData = AnimSequenceBase->GetDataModel()->GetCurveData();
@@ -217,7 +217,7 @@ void FAnimModel_AnimSequenceBase::RefreshCurveTracks()
 		}
 
 		AdditiveRoot->ClearChildren();
-		RootTracks.Add(AdditiveRoot.ToSharedRef());
+		AddRootTrack(AdditiveRoot.ToSharedRef());
 
 		// Next add a track for each transform curve
 		for(const FTransformCurve& TransformCurve : AnimationModelCurveData.TransformCurves)
@@ -279,7 +279,7 @@ void FAnimModel_AnimSequenceBase::RefreshAttributeTracks()
 		}
 
 		AttributesRoot->ClearChildren();
-		RootTracks.Add(AttributesRoot.ToSharedRef());
+		AddRootTrack(AttributesRoot.ToSharedRef());
 			   
 		TMap<FName, TSharedPtr<FAnimTimelineTrack_PerBoneAttributes>> BoneTracks;
 
diff --git a/Engine/Source/Editor/Persona/Private/AnimTimeline/AnimTimelineTrack.cpp b/Engine/Source/Editor/Persona/Private/AnimTimeline/AnimTimelineTrack.cpp
--- a/Engine/Source/Editor/Persona/Private/AnimTimeline/AnimTimelineTrack.cpp
+++ b/Engine/Source/Editor/Persona/Private/AnimTimeline/AnimTimelineTrack.cpp
@@ -198,4 +198,12 @@ void FAnimTimelineTrack::OnSetInputViewRange(float ViewMin, float ViewMax)
 	GetModel()->SetViewRange(TRange<double>(ViewMin, ViewMax));
 }
 
+void FAnimTimelineTrack::AddChild(const TSharedRef<FAnimTimelineTrack>& InChild)
+{
+	if (GetMutableDefault<UPersonaOptions>()->GetAllowedAnimationEditorTracks().PassesFilter(InChild->GetTypeName()))
+	{
+		Children.Add(InChild); 
+	}
+}
+
 #undef LOCTEXT_NAMESPACE
\ No newline at end of file
diff --git a/Engine/Source/Editor/Persona/Private/AnimTimeline/AnimTimelineTrack.h b/Engine/Source/Editor/Persona/Private/AnimTimeline/AnimTimelineTrack.h
--- a/Engine/Source/Editor/Persona/Private/AnimTimeline/AnimTimelineTrack.h
+++ b/Engine/Source/Editor/Persona/Private/AnimTimeline/AnimTimelineTrack.h
@@ -91,7 +91,7 @@ public:
 	const TArray<TSharedRef<FAnimTimelineTrack>>& GetChildren() { return Children; }
 
 	/** Add a child track */
-	void AddChild(const TSharedRef<FAnimTimelineTrack>& InChild) { Children.Add(InChild); }
+	void AddChild(const TSharedRef<FAnimTimelineTrack>& InChild);
 
 	/** Clear all child tracks */
 	void ClearChildren() { Children.Empty(); }
diff --git a/Engine/Source/Editor/Persona/Private/AnimTimeline/SAnimOutliner.cpp b/Engine/Source/Editor/Persona/Private/AnimTimeline/SAnimOutliner.cpp
--- a/Engine/Source/Editor/Persona/Private/AnimTimeline/SAnimOutliner.cpp
+++ b/Engine/Source/Editor/Persona/Private/AnimTimeline/SAnimOutliner.cpp
@@ -42,7 +42,7 @@ void SAnimOutliner::Construct(const FArguments& InArgs, const TSharedRef<FAnimMo
 	STreeView::Construct
 	(
 		STreeView::FArguments()
-		.TreeItemsSource(&InAnimModel->GetRootTracks())
+		.TreeItemsSource(&InAnimModel->GetAllRootTracks())
 		.SelectionMode(ESelectionMode::Multi)
 		.OnGenerateRow(this, &SAnimOutliner::HandleGenerateRow) 
 		.OnGetChildren(this, &SAnimOutliner::HandleGetChildren)
@@ -54,10 +54,10 @@ void SAnimOutliner::Construct(const FArguments& InArgs, const TSharedRef<FAnimMo
 	);
 
 	// expand all
-	for(TSharedRef<FAnimTimelineTrack>& RootTrack : InAnimModel->GetRootTracks())
+	InAnimModel->ForEachRootTrack([this](FAnimTimelineTrack& InRootTrack)
 	{
-		RootTrack->Traverse_ParentFirst([this](FAnimTimelineTrack& InTrack){ SetItemExpansion(InTrack.AsShared(), InTrack.IsExpanded()); return true; });
-	}
+		InRootTrack.Traverse_ParentFirst([this](FAnimTimelineTrack& InTrack){ SetItemExpansion(InTrack.AsShared(), InTrack.IsExpanded()); return true; });
+	});
 }
 
 void SAnimOutliner::Tick(const FGeometry& AllottedGeometry, const double InCurrentTime, const float InDeltaTime)
diff --git a/Engine/Source/Editor/UnrealEd/Classes/Preferences/PersonaOptions.h b/Engine/Source/Editor/UnrealEd/Classes/Preferences/PersonaOptions.h
--- a/Engine/Source/Editor/UnrealEd/Classes/Preferences/PersonaOptions.h
+++ b/Engine/Source/Editor/UnrealEd/Classes/Preferences/PersonaOptions.h
@@ -13,6 +13,7 @@
 #include "UObject/ObjectMacros.h"
 #include "UObject/Object.h"
 #include "Engine/EngineBaseTypes.h"
+#include "Misc/NamePermissionList.h"
 #include "PersonaOptions.generated.h"
 
 enum class EFrameNumberDisplayFormats : uint8;
@@ -245,7 +246,12 @@ public:
 		OnSettingsChange.Remove(Object);
 	}
 
+	FNamePermissionList& GetAllowedAnimationEditorTracks() { return AllowedAnimationEditorTracks; }
+
 protected:
 	// UObject interface
 	virtual void PostEditChangeProperty(struct FPropertyChangedEvent& PropertyChangedEvent) override;
+
+private:
+	FNamePermissionList AllowedAnimationEditorTracks;
 };
