Refactor persona editor modes/mode manager to use ITF context objects for persona-only interfaces instead of static casting

diff --git a/Engine/Plugins/Animation/IKRig/Source/IKRigEditor/Private/RetargetEditor/IKRetargetEditPoseMode.cpp b/Engine/Plugins/Animation/IKRig/Source/IKRigEditor/Private/RetargetEditor/IKRetargetEditPoseMode.cpp
--- a/Engine/Plugins/Animation/IKRig/Source/IKRigEditor/Private/RetargetEditor/IKRetargetEditPoseMode.cpp
+++ b/Engine/Plugins/Animation/IKRig/Source/IKRigEditor/Private/RetargetEditor/IKRetargetEditPoseMode.cpp
@@ -402,6 +402,8 @@ bool FIKRetargetEditPoseMode::GetCustomInputCoordinateSystem(FMatrix& InMatrix,
 
 void FIKRetargetEditPoseMode::Enter()
 {
+	IPersonaEditMode::Enter();
+	
 	const TSharedPtr<FIKRetargetEditorController> Controller = EditorController.Pin();
 	if (!Controller.IsValid())
 	{
@@ -426,6 +428,8 @@ void FIKRetargetEditPoseMode::Exit()
 	}
 
 	Controller->AssetController->GetAsset()->SetOutputMode(ERetargeterOutputMode::RunRetarget);
+
+	IPersonaEditMode::Exit();
 }
 
 void FIKRetargetEditPoseMode::UpdateWidgetTransform()
diff --git a/Engine/Source/Editor/AnimGraph/AnimGraph.Build.cs b/Engine/Source/Editor/AnimGraph/AnimGraph.Build.cs
--- a/Engine/Source/Editor/AnimGraph/AnimGraph.Build.cs
+++ b/Engine/Source/Editor/AnimGraph/AnimGraph.Build.cs
@@ -42,6 +42,7 @@ public class AnimGraph : ModuleRules
 				"Kismet",
 				"EditorWidgets",
 				"ToolWidgets",
+				"AnimationEditMode",
 			}
 		);
 
@@ -54,8 +55,7 @@ public class AnimGraph : ModuleRules
 
         PrivateIncludePathModuleNames.AddRange(
             new string[] {
-				"Persona",
-				"SkeletonEditor",
+	            "SkeletonEditor",
                 "AdvancedPreviewScene",
                 "AnimationBlueprintEditor",
             }
diff --git a/Engine/Source/Editor/AnimGraph/Public/IAnimNodeEditMode.h b/Engine/Source/Editor/AnimGraph/Public/IAnimNodeEditMode.h
--- a/Engine/Source/Editor/AnimGraph/Public/IAnimNodeEditMode.h
+++ b/Engine/Source/Editor/AnimGraph/Public/IAnimNodeEditMode.h
@@ -4,10 +4,10 @@
 
 #include "CoreMinimal.h"
 #include "UnrealWidgetFwd.h"
-#include "IPersonaEditMode.h"
+#include "AnimationEditMode.h"
 
 /** Base interface for skeletal control edit modes */
-class IAnimNodeEditMode : public IPersonaEditMode
+class IAnimNodeEditMode : public FAnimationEditMode
 {
 public:
 	/** Returns the coordinate system that should be used for this bone */
diff --git a/Engine/Source/Editor/AnimationEditMode/AnimationEditMode.Build.cs b/Engine/Source/Editor/AnimationEditMode/AnimationEditMode.Build.cs
new file mode 100644
--- /dev/null
+++ b/Engine/Source/Editor/AnimationEditMode/AnimationEditMode.Build.cs
@@ -0,0 +1,20 @@
+﻿// Copyright Epic Games, Inc. All Rights Reserved.
+
+using UnrealBuildTool;
+
+public class AnimationEditMode : ModuleRules
+{
+	public AnimationEditMode(ReadOnlyTargetRules Target) : base(Target)
+	{
+		PrivateDependencyModuleNames.AddRange(
+			new string[] {
+				"Core",
+				"CoreUObject",
+				"Engine",
+				"UnrealEd",
+				"EditorInteractiveToolsFramework",
+				"InteractiveToolsFramework",
+			}
+		);
+	}
+}
\ No newline at end of file
diff --git a/Engine/Source/Editor/AnimationEditMode/Private/AnimationEditMode.cpp b/Engine/Source/Editor/AnimationEditMode/Private/AnimationEditMode.cpp
new file mode 100644
--- /dev/null
+++ b/Engine/Source/Editor/AnimationEditMode/Private/AnimationEditMode.cpp
@@ -0,0 +1,53 @@
+﻿// Copyright Epic Games, Inc. All Rights Reserved.
+
+#include "AnimationEditMode.h"
+
+#include "EdModeInteractiveToolsContext.h"
+#include "EditorModeManager.h"
+#include "ContextObjectStore.h"
+
+bool UAnimationEditModeContext::GetCameraTarget(FSphere& OutTarget) const
+{
+	check(EditMode);
+	return EditMode->GetCameraTarget(OutTarget);
+}
+
+IPersonaPreviewScene& UAnimationEditModeContext::GetAnimPreviewScene() const
+{
+	check(EditMode);
+	return EditMode->GetAnimPreviewScene();
+}
+
+void UAnimationEditModeContext::GetOnScreenDebugInfo(TArray<FText>& OutDebugInfo) const
+{
+	check(EditMode);
+	return EditMode->GetOnScreenDebugInfo(OutDebugInfo);
+};
+
+void FAnimationEditMode::Enter()
+{
+	FEdMode::Enter();
+
+	if (const UEdMode* EdMode = GetModeManager()->GetActiveScriptableMode(Info.ID))
+	{
+		UContextObjectStore* ContextObjectStore = EdMode->GetInteractiveToolsContext()->ContextObjectStore;
+		ContextObjectStore->AddContextObject(AnimationEditModeContext.Get());
+	}
+}
+
+void FAnimationEditMode::Exit()
+{
+	if (const UEdMode* EdMode = GetModeManager()->GetActiveScriptableMode(Info.ID))
+	{
+		UContextObjectStore* ContextObjectStore = EdMode->GetInteractiveToolsContext()->ContextObjectStore;
+		ContextObjectStore->RemoveContextObject(AnimationEditModeContext.Get());
+	}
+	
+	FEdMode::Exit();
+}
+
+void FAnimationEditMode::AddReferencedObjects(FReferenceCollector& Collector)
+{
+	FEdMode::AddReferencedObjects(Collector);
+	Collector.AddReferencedObject(AnimationEditModeContext);
+}
diff --git a/Engine/Source/Editor/AnimationEditMode/Public/AnimationEditContext.h b/Engine/Source/Editor/AnimationEditMode/Public/AnimationEditContext.h
new file mode 100644
--- /dev/null
+++ b/Engine/Source/Editor/AnimationEditMode/Public/AnimationEditContext.h
@@ -0,0 +1,36 @@
+﻿// Copyright Epic Games, Inc. All Rights Reserved.
+
+#pragma once
+
+#include "UObject/UObjectHierarchyFwd.h"
+
+#include "AnimationEditContext.generated.h"
+
+UINTERFACE(MinimalAPI)
+class UAnimationEditContext : public UInterface
+{
+	GENERATED_BODY()
+};
+
+class IAnimationEditContext
+{
+	GENERATED_BODY()
+public:
+	/** 
+	 * Get a camera target used to focus the viewport on an object when a user presses 'F' (default). 
+	 * @param	 OutTarget	 The target object
+	 * @return true if the target sphere was filled-in
+	 */
+	virtual bool GetCameraTarget(FSphere& OutTarget) const { return false; }
+
+	/** @return the anim preview scene */
+	virtual class IPersonaPreviewScene& GetAnimPreviewScene() const = 0;
+
+	/**
+	* Function to collect strings from nodes to display in the viewport.
+	* Use this rather than DrawHUD when adding general text to the viewport.
+	* Display of this assumes that this will be mostly used by skeletal controls.
+	* @param	OutDebugInfo	Text to display
+	*/
+	virtual void GetOnScreenDebugInfo(TArray<FText>& OutDebugInfo) const {}
+};
diff --git a/Engine/Source/Editor/AnimationEditMode/Public/AnimationEditMode.h b/Engine/Source/Editor/AnimationEditMode/Public/AnimationEditMode.h
new file mode 100644
--- /dev/null
+++ b/Engine/Source/Editor/AnimationEditMode/Public/AnimationEditMode.h
@@ -0,0 +1,44 @@
+﻿// Copyright Epic Games, Inc. All Rights Reserved.
+
+#pragma once
+
+#include "CoreMinimal.h"
+#include "EdMode.h"
+#include "AnimationEditContext.h"
+
+#include "AnimationEditMode.generated.h"
+
+class FAnimationEditMode;
+
+/**
+ *	A compatibility context object to support IPersonaEditMode-based code. It simply calls into a different
+ *	IAnimationEditContext in its implementations.
+ */
+UCLASS(MinimalAPI)
+class UAnimationEditModeContext : public UObject, public IAnimationEditContext
+{
+	GENERATED_BODY()
+public:
+	virtual bool GetCameraTarget(FSphere& OutTarget) const override;
+	virtual class IPersonaPreviewScene& GetAnimPreviewScene() const override;
+	virtual void GetOnScreenDebugInfo(TArray<FText>& OutDebugInfo) const override;
+private:
+	IAnimationEditContext* EditMode = nullptr;
+	static UAnimationEditModeContext* CreateFor(IAnimationEditContext* InEditMode)
+	{
+		UAnimationEditModeContext* NewPersonaContext = NewObject<UAnimationEditModeContext>();
+		NewPersonaContext->EditMode = InEditMode;
+		return NewPersonaContext;
+	}
+	friend FAnimationEditMode;
+};
+
+class ANIMATIONEDITMODE_API FAnimationEditMode : public FEdMode, public IAnimationEditContext
+{
+public:
+	virtual void Enter() override;
+	virtual void Exit() override;
+	virtual void AddReferencedObjects(FReferenceCollector& Collector) override;
+private:
+	TObjectPtr<UAnimationEditModeContext> AnimationEditModeContext{UAnimationEditModeContext::CreateFor(this)};
+};
diff --git a/Engine/Source/Editor/Persona/Persona.Build.cs b/Engine/Source/Editor/Persona/Persona.Build.cs
--- a/Engine/Source/Editor/Persona/Persona.Build.cs
+++ b/Engine/Source/Editor/Persona/Persona.Build.cs
@@ -20,6 +20,7 @@ public class Persona : ModuleRules
         PublicDependencyModuleNames.AddRange(
             new string[] {
                 "AdvancedPreviewScene",
+                "AnimationEditMode",
             }
         );
 
@@ -33,6 +34,7 @@ public class Persona : ModuleRules
                 "MeshReductionInterface",
                 "SequenceRecorder",
                 "AnimationBlueprintEditor",
+                "EditorInteractiveToolsFramework",
 			}
 		);
 
@@ -77,7 +79,8 @@ public class Persona : ModuleRules
                 "Sequencer",
 				"StatusBar",
 				"ToolWidgets",
-			}
+				"InteractiveToolsFramework",
+            }
         );
 
         DynamicallyLoadedModuleNames.AddRange(
diff --git a/Engine/Source/Editor/Persona/Private/AnimationEditorViewportClient.cpp b/Engine/Source/Editor/Persona/Private/AnimationEditorViewportClient.cpp
--- a/Engine/Source/Editor/Persona/Private/AnimationEditorViewportClient.cpp
+++ b/Engine/Source/Editor/Persona/Private/AnimationEditorViewportClient.cpp
@@ -35,6 +35,8 @@
 #include "AudioDevice.h"
 #include "RawIndexBuffer.h"
 #include "CameraController.h"
+#include "ContextObjectStore.h"
+#include "EdModeInteractiveToolsContext.h"
 #include "Animation/MorphTarget.h"
 #include "Rendering/SkeletalMeshModel.h"
 #include "UnrealWidget.h"
@@ -1721,10 +1723,10 @@ void FAnimationViewportClient::DrawSockets(const UDebugSkelMeshComponent* InPrev
 FSphere FAnimationViewportClient::GetCameraTarget()
 {
 	// give the editor mode a chance to give us a camera target
-	if (GetPersonaModeManager())
+	if (IPersonaManagerContext* PersonaModeManagerContext = GetModeTools()->GetInteractiveToolsContext()->ContextObjectStore->FindContext<UPersonaEditorModeManagerContext>())
 	{
 		FSphere Target;
-		if (GetPersonaModeManager()->GetCameraTarget(Target))
+		if (PersonaModeManagerContext->GetCameraTarget(Target))
 		{
 			return Target;
 		}
@@ -2096,7 +2098,7 @@ TSharedRef<FAnimationEditorPreviewScene> FAnimationViewportClient::GetAnimPrevie
 
 IPersonaEditorModeManager* FAnimationViewportClient::GetPersonaModeManager() const
 {
-	return static_cast<IPersonaEditorModeManager*>(ModeTools.Get());
+	return ModeTools->GetInteractiveToolsContext()->ContextObjectStore->FindContext<UPersonaEditorModeManagerContext>()->GetPersonaEditorModeManager();
 }
 
 void FAnimationViewportClient::HandleInvalidateViews()
diff --git a/Engine/Source/Editor/Persona/Private/AnimationEditorViewportClient.h b/Engine/Source/Editor/Persona/Private/AnimationEditorViewportClient.h
--- a/Engine/Source/Editor/Persona/Private/AnimationEditorViewportClient.h
+++ b/Engine/Source/Editor/Persona/Private/AnimationEditorViewportClient.h
@@ -301,6 +301,7 @@ public:
 	int32 GetViewportIndex() const { return ViewportIndex; }
 
 	/** Get the persona mode manager */
+	UE_DEPRECATED(5.1, "Use the UPersonaEditorModeManagerContext object stored in the editor mode tools' context store instead.")
 	class IPersonaEditorModeManager* GetPersonaModeManager() const;
 
 private:
diff --git a/Engine/Source/Editor/Persona/Private/IPersonaEditMode.cpp b/Engine/Source/Editor/Persona/Private/IPersonaEditMode.cpp
new file mode 100644
--- /dev/null
+++ b/Engine/Source/Editor/Persona/Private/IPersonaEditMode.cpp
@@ -0,0 +1,26 @@
+﻿// Copyright Epic Games, Inc. All Rights Reserved.
+
+#include "IPersonaEditMode.h"
+
+IPersonaEditMode::IPersonaEditMode()
+{
+}
+
+IPersonaEditMode::~IPersonaEditMode()
+{
+}
+
+void IPersonaEditMode::Enter()
+{
+	FAnimationEditMode::Enter();
+}
+
+void IPersonaEditMode::Exit()
+{
+	FAnimationEditMode::Exit();
+}
+
+void IPersonaEditMode::AddReferencedObjects(FReferenceCollector& Collector)
+{
+	FAnimationEditMode::AddReferencedObjects(Collector);
+}
diff --git a/Engine/Source/Editor/Persona/Private/IPersonaEditorModeManager.cpp b/Engine/Source/Editor/Persona/Private/IPersonaEditorModeManager.cpp
new file mode 100644
--- /dev/null
+++ b/Engine/Source/Editor/Persona/Private/IPersonaEditorModeManager.cpp
@@ -0,0 +1,41 @@
+﻿// Copyright Epic Games, Inc. All Rights Reserved.
+
+#include "IPersonaEditorModeManager.h"
+
+#include "ContextObjectStore.h"
+#include "EdModeInteractiveToolsContext.h"
+
+IPersonaEditorModeManager* UPersonaEditorModeManagerContext::GetPersonaEditorModeManager() const
+{
+	return ModeManager;
+}
+
+bool UPersonaEditorModeManagerContext::GetCameraTarget(FSphere& OutTarget) const
+{
+	check(ModeManager);
+	return ModeManager->GetCameraTarget(OutTarget);
+}
+
+void UPersonaEditorModeManagerContext::GetOnScreenDebugInfo(TArray<FText>& OutDebugText) const
+{
+	check(ModeManager);
+	return ModeManager->GetOnScreenDebugInfo(OutDebugText);
+}
+
+IPersonaEditorModeManager::IPersonaEditorModeManager() : FAssetEditorModeManager()
+{
+	UContextObjectStore* ContextObjectStore = GetInteractiveToolsContext()->ContextObjectStore;
+	ContextObjectStore->AddContextObject(PersonaModeManagerContext.Get());
+}
+
+IPersonaEditorModeManager::~IPersonaEditorModeManager()
+{
+	UContextObjectStore* ContextObjectStore = GetInteractiveToolsContext()->ContextObjectStore;
+	ContextObjectStore->RemoveContextObject(PersonaModeManagerContext.Get());
+}
+
+void IPersonaEditorModeManager::AddReferencedObjects(FReferenceCollector& Collector)
+{
+	FAssetEditorModeManager::AddReferencedObjects(Collector);
+	Collector.AddReferencedObject(PersonaModeManagerContext);
+}
\ No newline at end of file
diff --git a/Engine/Source/Editor/Persona/Private/PersonaEditorModeManager.cpp b/Engine/Source/Editor/Persona/Private/PersonaEditorModeManager.cpp
--- a/Engine/Source/Editor/Persona/Private/PersonaEditorModeManager.cpp
+++ b/Engine/Source/Editor/Persona/Private/PersonaEditorModeManager.cpp
@@ -1,46 +1,41 @@
 // Copyright Epic Games, Inc. All Rights Reserved.
 
 #include "PersonaEditorModeManager.h"
+
 #include "IPersonaEditMode.h"
 #include "IPersonaPreviewScene.h"
 #include "Selection.h"
 #include "Animation/DebugSkelMeshComponent.h"
-
+#include "EdModeInteractiveToolsContext.h"
+#include "ContextObjectStore.h"
 
 bool FPersonaEditorModeManager::GetCameraTarget(FSphere& OutTarget) const
 {
-	for (UEdMode* Mode : ActiveScriptableModes)
+	for (const UEdMode* Mode : ActiveScriptableModes)
 	{
-		FEdMode* LegacyMode = Mode->AsLegacyMode();
-		// Hack for UE-136071, UE-141936. ClothPaintModes are not IPersonaEditModes, but are FEdModes.
-		if (LegacyMode && LegacyMode->GetID() != FEditorModeID("ClothPaintMode"))
+		if (const UEditorInteractiveToolsContext* ModeInteractiveToolsContext = Mode->GetInteractiveToolsContext())
 		{
-			if (IPersonaEditMode* EditMode = static_cast<IPersonaEditMode*>(LegacyMode))
+			if (const IAnimationEditContext* PersonaContext = ModeInteractiveToolsContext->ContextObjectStore->FindContext<UAnimationEditModeContext>())
 			{
-				FSphere Target;
-				if (EditMode->GetCameraTarget(Target))
+				if (PersonaContext->GetCameraTarget(OutTarget))
 				{
-					OutTarget = Target;
 					return true;
 				}
 			}
 		}
 	}
-
 	return false;
 }
 
 void FPersonaEditorModeManager::GetOnScreenDebugInfo(TArray<FText>& OutDebugText) const
 {
-	for (UEdMode* Mode : ActiveScriptableModes)
+	for (const UEdMode* Mode : ActiveScriptableModes)
 	{
-		FEdMode* LegacyMode = Mode->AsLegacyMode();
-		// Hack for UE-136071, UE-141936. ClothPaintModes are not IPersonaEditModes, but are FEdModes.
-		if (LegacyMode && LegacyMode->GetID() != FEditorModeID("ClothPaintMode"))
+		if (const UEditorInteractiveToolsContext* ModeInteractiveToolsContext = Mode->GetInteractiveToolsContext())
 		{
-			if (IPersonaEditMode* EditMode = static_cast<IPersonaEditMode*>(LegacyMode))
+			if (const IAnimationEditContext* PersonaContext = ModeInteractiveToolsContext->ContextObjectStore->FindContext<UAnimationEditModeContext>())
 			{
-				EditMode->GetOnScreenDebugInfo(OutDebugText);
+				PersonaContext->GetOnScreenDebugInfo(OutDebugText);
 			}
 		}
 	}
diff --git a/Engine/Source/Editor/Persona/Private/SAnimationEditorViewport.cpp b/Engine/Source/Editor/Persona/Private/SAnimationEditorViewport.cpp
--- a/Engine/Source/Editor/Persona/Private/SAnimationEditorViewport.cpp
+++ b/Engine/Source/Editor/Persona/Private/SAnimationEditorViewport.cpp
@@ -33,6 +33,10 @@
 #include "Editor/EditorPerProjectUserSettings.h"
 #include "Materials/Material.h"
 #include "EditorFontGlyphs.h"
+#include "Experimental/EditorInteractiveToolsFramework/Public/EdModeInteractiveToolsContext.h"
+#include "ContextObjectStore.h"
+#include "IPersonaEditMode.h"
+#include "PersonaContext.h"
 
 #include "SkeletalMeshTypes.h"
 #include "IPersonaToolkit.h"
@@ -291,11 +295,10 @@ FText SAnimationEditorViewportTabBody::GetDisplayString() const
 	else if(AnimViewportClient->IsShowingSelectedNodeStats())
 	{
 		// Allow edit modes (inc. skeletal control modes) to draw with the canvas, and collect on screen strings to draw later
-		if (AnimViewportClient->GetPersonaModeManager())
+		if (IAnimationEditContext* PersonaContext = AnimViewportClient->GetModeTools()->GetInteractiveToolsContext()->ContextObjectStore->FindContext<UAnimationEditModeContext>())
 		{
 			TArray<FText> EditModeDebugText;
-			AnimViewportClient->GetPersonaModeManager()->GetOnScreenDebugInfo(EditModeDebugText);
-
+			PersonaContext->GetOnScreenDebugInfo(EditModeDebugText);
 			for(FText& Text : EditModeDebugText)
 			{
 				DefaultText = ConcatenateLine(DefaultText, Text);
diff --git a/Engine/Source/Editor/Persona/Public/IPersonaEditMode.h b/Engine/Source/Editor/Persona/Public/IPersonaEditMode.h
--- a/Engine/Source/Editor/Persona/Public/IPersonaEditMode.h
+++ b/Engine/Source/Editor/Persona/Public/IPersonaEditMode.h
@@ -2,27 +2,16 @@
 
 #pragma once
 
-#include "CoreMinimal.h"
-#include "EdMode.h"
+#include "AnimationEditMode.h"
 
-class IPersonaEditMode : public FEdMode
+class PERSONA_API IPersonaEditMode : public FAnimationEditMode
 {
+	// These implementations are duplicated in the Persona module so that code that links only to Persona and not to
+	// the AnimationEditMode module can work properly
 public:
-	/** 
-	 * Get a camera target used to focus the viewport on an object when a user presses 'F' (default). 
-	 * @param	OutTarget	The target object
-	 * @return true if the target sphere was filled-in
-	 */
-	virtual bool GetCameraTarget(FSphere& OutTarget) const = 0;
-
-	/** @return the anim preview scene */
-	virtual class IPersonaPreviewScene& GetAnimPreviewScene() const = 0;
-
-	/**
-	* Function to collect strings from nodes to display in the viewport.
-	* Use this rather than DrawHUD when adding general text to the viewport.
-	* Display of this assumes that this will be mostly used by skeletal controls.
-	* @param	OutDebugInfo	Text to display
-	*/
-	virtual void GetOnScreenDebugInfo(TArray<FText>& OutDebugInfo) const = 0;
+	IPersonaEditMode();
+	virtual ~IPersonaEditMode() override;
+	virtual void Enter() override;
+	virtual void Exit() override;
+	virtual void AddReferencedObjects(FReferenceCollector& Collector) override;
 };
diff --git a/Engine/Source/Editor/Persona/Public/IPersonaEditorModeManager.h b/Engine/Source/Editor/Persona/Public/IPersonaEditorModeManager.h
--- a/Engine/Source/Editor/Persona/Public/IPersonaEditorModeManager.h
+++ b/Engine/Source/Editor/Persona/Public/IPersonaEditorModeManager.h
@@ -5,9 +5,18 @@
 #include "CoreMinimal.h"
 #include "AssetEditorModeManager.h"
 
+#include "IPersonaEditorModeManager.generated.h"
+
+UINTERFACE(MinimalAPI)
+class UPersonaManagerContext : public UInterface
+{
+	GENERATED_BODY()
+};
+
 /** Persona-specific extensions to the asset editor mode manager */
-class IPersonaEditorModeManager : public FAssetEditorModeManager
+class IPersonaManagerContext
 {
+	GENERATED_BODY()
 public:
 	/** 
 	 * Get a camera target for when the user focuses the viewport
@@ -22,3 +31,35 @@ public:
 	 */
 	virtual void GetOnScreenDebugInfo(TArray<FText>& OutDebugText) const = 0;
 };
+
+class IPersonaEditorModeManager;
+
+UCLASS()
+class UPersonaEditorModeManagerContext : public UObject, public IPersonaManagerContext
+{
+	GENERATED_BODY()
+public:
+	// Only for use by FAnimationViewportClient::GetPersonaModeManager for compatibility
+	IPersonaEditorModeManager* GetPersonaEditorModeManager() const;
+	virtual bool GetCameraTarget(FSphere& OutTarget) const override;
+	virtual void GetOnScreenDebugInfo(TArray<FText>& OutDebugText) const override;
+private:
+	IPersonaEditorModeManager* ModeManager = nullptr;
+	static UPersonaEditorModeManagerContext* CreateFor(IPersonaEditorModeManager* InModeManager)
+	{
+		UPersonaEditorModeManagerContext* NewPersonaContext = NewObject<UPersonaEditorModeManagerContext>();
+		NewPersonaContext->ModeManager = InModeManager;
+		return NewPersonaContext;
+	}
+	friend IPersonaEditorModeManager;
+};
+
+class IPersonaEditorModeManager : public FAssetEditorModeManager, public IPersonaManagerContext
+{
+public:
+	IPersonaEditorModeManager();
+	virtual ~IPersonaEditorModeManager() override;
+	virtual void AddReferencedObjects(FReferenceCollector& Collector) override;
+private:
+	TObjectPtr<UPersonaEditorModeManagerContext> PersonaModeManagerContext{UPersonaEditorModeManagerContext::CreateFor(this)};
+};
diff --git a/Engine/Source/Editor/Persona/Public/PersonaContext.h b/Engine/Source/Editor/Persona/Public/PersonaContext.h
new file mode 100644
--- /dev/null
+++ b/Engine/Source/Editor/Persona/Public/PersonaContext.h
@@ -0,0 +1,36 @@
+// Copyright Epic Games, Inc. All Rights Reserved.
+
+#pragma once
+
+#include "UObject/UObjectHierarchyFwd.h"
+
+#include "PersonaContext.generated.h"
+
+UINTERFACE(MinimalAPI)
+class UPersonaContext : public UInterface
+{
+	GENERATED_BODY()
+};
+
+class IPersonaContext
+{
+	GENERATED_BODY()
+public:
+	/** 
+	 * Get a camera target used to focus the viewport on an object when a user presses 'F' (default). 
+	 * @param	 OutTarget	 The target object
+	 * @return true if the target sphere was filled-in
+	 */
+	virtual bool GetCameraTarget(FSphere& OutTarget) const { return false; }
+
+	/** @return the anim preview scene */
+	virtual class IPersonaPreviewScene& GetAnimPreviewScene() const = 0;
+
+	/**
+	* Function to collect strings from nodes to display in the viewport.
+	* Use this rather than DrawHUD when adding general text to the viewport.
+	* Display of this assumes that this will be mostly used by skeletal controls.
+	* @param	OutDebugInfo	Text to display
+	*/
+	virtual void GetOnScreenDebugInfo(TArray<FText>& OutDebugInfo) const {}
+};
diff --git a/Engine/Source/Editor/UnrealEd/Private/Tools/LegacyEdMode.cpp b/Engine/Source/Editor/UnrealEd/Private/Tools/LegacyEdMode.cpp
--- a/Engine/Source/Editor/UnrealEd/Private/Tools/LegacyEdMode.cpp
+++ b/Engine/Source/Editor/UnrealEd/Private/Tools/LegacyEdMode.cpp
@@ -281,6 +281,7 @@ bool ULegacyEdModeWrapper::ShouldDrawBrushWireframe(AActor* InActor) const
 
 void ULegacyEdModeWrapper::Enter()
 {
+	CreateInteractiveToolsContexts();
 	LegacyEditorMode->Enter();
 	Toolkit = LegacyEditorMode->GetToolkit();
 }
@@ -290,6 +291,7 @@ void ULegacyEdModeWrapper::Exit()
 	Toolkit.Reset();
 
 	LegacyEditorMode->Exit();
+	DestroyInteractiveToolsContexts();
 }
 
 FEdMode* ULegacyEdModeWrapper::AsLegacyMode()
diff --git a/Engine/Source/Editor/UnrealEd/Private/Tools/UEdMode.cpp b/Engine/Source/Editor/UnrealEd/Private/Tools/UEdMode.cpp
--- a/Engine/Source/Editor/UnrealEd/Private/Tools/UEdMode.cpp
+++ b/Engine/Source/Editor/UnrealEd/Private/Tools/UEdMode.cpp
@@ -47,13 +47,7 @@ void UEdMode::Enter()
 			return true;
 		});
 
-	// Editor-scope ToolsContext is provided by the Mode Manager
-	EditorToolsContext = Owner->GetInteractiveToolsContext();
-	check(EditorToolsContext.IsValid());
-
-	// Mode-scope ToolsContext is created derived from the Editor-Scope ToolsContext so that the UInputRouter can be shared
-	ModeToolsContext = EditorToolsContext->CreateNewChildEdModeToolsContext();
-	check(ModeToolsContext != nullptr);
+	CreateInteractiveToolsContexts();
 
 	ModeToolsContext->ToolManager->OnToolStarted.AddUObject(this, &UEdMode::OnToolStarted);
 	ModeToolsContext->ToolManager->OnToolEnded.AddUObject(this, &UEdMode::OnToolEnded);
@@ -145,10 +139,7 @@ void UEdMode::Exit()
 
 		// Shutdown the Mode-scope ToolsContext, and notify the EditorToolsContext so it can release the reference it holds.
 		// Do this before shutting down the Toolkit as if a Mode-scope Tool is still active it will need to clean up
-		if (ModeToolsContext != nullptr)
-		{
-			ModeToolsContext->ShutdownContext();
-		}
+		DestroyInteractiveToolsContexts();
 	}
 
 	if (Toolkit.IsValid())
@@ -263,6 +254,25 @@ void UEdMode::OnModeActivated(const FEditorModeID& InID, bool bIsActive)
 	}
 }
 
+void UEdMode::CreateInteractiveToolsContexts()
+{
+	// Editor-scope ToolsContext is provided by the Mode Manager
+	EditorToolsContext = Owner->GetInteractiveToolsContext();
+	check(EditorToolsContext.IsValid());
+
+	// Mode-scope ToolsContext is created derived from the Editor-Scope ToolsContext so that the UInputRouter can be shared
+	ModeToolsContext = EditorToolsContext->CreateNewChildEdModeToolsContext();
+	check(ModeToolsContext);
+}
+
+void UEdMode::DestroyInteractiveToolsContexts()
+{
+	if (ModeToolsContext)
+	{
+		ModeToolsContext->ShutdownContext();
+	}
+}
+
 bool UEdMode::IsSnapRotationEnabled()
 {
 	return GetDefault<ULevelEditorViewportSettings>()->RotGridEnabled;
diff --git a/Engine/Source/Editor/UnrealEd/Public/Tools/UEdMode.h b/Engine/Source/Editor/UnrealEd/Public/Tools/UEdMode.h
--- a/Engine/Source/Editor/UnrealEd/Public/Tools/UEdMode.h
+++ b/Engine/Source/Editor/UnrealEd/Public/Tools/UEdMode.h
@@ -291,4 +291,8 @@ protected:
 	UPROPERTY(Transient)
 	TObjectPtr<UObject> SettingsObject;
 
+protected:
+
+	void CreateInteractiveToolsContexts();
+	void DestroyInteractiveToolsContexts();
 };
