Fix non unity build

diff --git a/Engine/Source/Editor/Persona/Private/SAnimationEditorViewport.cpp b/Engine/Source/Editor/Persona/Private/SAnimationEditorViewport.cpp
--- a/Engine/Source/Editor/Persona/Private/SAnimationEditorViewport.cpp
+++ b/Engine/Source/Editor/Persona/Private/SAnimationEditorViewport.cpp
@@ -33,10 +33,9 @@
 #include "Editor/EditorPerProjectUserSettings.h"
 #include "Materials/Material.h"
 #include "EditorFontGlyphs.h"
-#include "Experimental/EditorInteractiveToolsFramework/Public/EdModeInteractiveToolsContext.h"
+#include "EdModeInteractiveToolsContext.h"
 #include "ContextObjectStore.h"
 #include "IPersonaEditMode.h"
-#include "PersonaContext.h"
 
 #include "SkeletalMeshTypes.h"
 #include "IPersonaToolkit.h"
diff --git a/Engine/Source/Editor/Persona/Public/PersonaContext.h b/Engine/Source/Editor/Persona/Public/PersonaContext.h
deleted file mode 100644
