Adding skeleton view to retargeter and replacing bone rendering method in all animation editors.

diff --git a/Engine/Plugins/Animation/ControlRig/Source/ControlRigEditor/Private/Editor/ControlRigEditor.cpp b/Engine/Plugins/Animation/ControlRig/Source/ControlRigEditor/Private/Editor/ControlRigEditor.cpp
--- a/Engine/Plugins/Animation/ControlRig/Source/ControlRigEditor/Private/Editor/ControlRigEditor.cpp
+++ b/Engine/Plugins/Animation/ControlRig/Source/ControlRigEditor/Private/Editor/ControlRigEditor.cpp
@@ -547,8 +547,6 @@ void FControlRigEditor::InitControlRigEditor(const EToolkitMode::Type Mode, cons
 			Compile();
 		}
 	}
-	
-	OnToolbarBoneRadiusChanged(InControlRigBlueprint->DebugBoneRadius);
 
 	FFunctionGraphTask::CreateAndDispatchWhenReady([this]()
 	{
@@ -1348,11 +1346,6 @@ void FControlRigEditor::HandleSetObjectBeingDebugged(UObject* InObject)
 		DebuggedControlRig->ControlRigLog = &ControlRigLog;
 		DebuggedControlRig->DynamicHierarchy->HierarchyForSelectionPtr = GetControlRigBlueprint()->Hierarchy;
 
-		if (UControlRigBlueprint* RigBlueprint = Cast<UControlRigBlueprint>(GetBlueprintObj()))
-		{
-			DebuggedControlRig->DebugBoneRadiusMultiplier = RigBlueprint->DebugBoneRadius;
-		}
-
 		UControlRigSkeletalMeshComponent* EditorSkelComp = Cast<UControlRigSkeletalMeshComponent>(GetPersonaToolkit()->GetPreviewScene()->GetPreviewMeshComponent());
 		if (EditorSkelComp)
 		{
@@ -3646,27 +3639,6 @@ void FControlRigEditor::HandleViewportCreated(const TSharedRef<class IPersonaVie
 					LOCTEXT("ControlRigAxesScale", "Axes Scale")
 				);
 
-				InMenuBuilder.AddWidget(
-					SNew(SBox)
-					.HAlign(HAlign_Right)
-					[
-						SNew(SBox)
-						.Padding(FMargin(4.0f, 0.0f, 0.0f, 0.0f))
-						.WidthOverride(100.0f)
-						[
-							SNew(SNumericEntryBox<float>)
-							.Font(FAppStyle::GetFontStyle(TEXT("MenuItem.Font")))
-							.AllowSpin(true)
-							.MinSliderValue(0.0f)
-							.MaxSliderValue(10.0f)
-							.Value(this, &FControlRigEditor::GetToolbarBoneRadius)
-							.OnValueChanged(this, &FControlRigEditor::OnToolbarBoneRadiusChanged)
-							.ToolTipText(LOCTEXT("ControlRigBoneRadiusToolTip", "The radius of the bones drawn for selected rig elements"))
-						]
-					], 
-					LOCTEXT("ControlRigBoneRadius", "Bone Radius")
-				);
-
 				if (UControlRigBlueprint* ControlRigBlueprint = CastChecked<UControlRigBlueprint>(GetBlueprintObj()))
 				{
 					for (UEdGraph* Graph : ControlRigBlueprint->UbergraphPages)
@@ -3753,6 +3725,30 @@ void FControlRigEditor::HandleViewportCreated(const TSharedRef<class IPersonaVie
 		}
 		return FReply::Unhandled();
 	});
+
+	// register callbacks to allow control rig asset to store the Bone Size viewport setting
+	FEditorViewportClient& ViewportClient = InViewport->GetViewportClient();
+	if (FAnimationViewportClient* AnimViewportClient = static_cast<FAnimationViewportClient*>(&ViewportClient))
+	{
+		AnimViewportClient->OnSetBoneSize.BindLambda([this](float InBoneSize)
+		{
+			if (UControlRigBlueprint* RigBlueprint = Cast<UControlRigBlueprint>(GetBlueprintObj()))
+			{
+				RigBlueprint->Modify();
+				RigBlueprint->DebugBoneRadius = InBoneSize;
+			}
+		});
+		
+		AnimViewportClient->OnGetBoneSize.BindLambda([this]() -> float
+		{
+			if (UControlRigBlueprint* RigBlueprint = Cast<UControlRigBlueprint>(GetBlueprintObj()))
+			{
+				return RigBlueprint->DebugBoneRadius;
+			}
+
+			return 1.0f;
+		});
+	}
 }
 
 TOptional<float> FControlRigEditor::GetToolbarAxesScale() const
@@ -3772,35 +3768,6 @@ void FControlRigEditor::OnToolbarAxesScaleChanged(float InValue)
 	}
 }
 
-TOptional<float> FControlRigEditor::GetToolbarBoneRadius() const
-{
-	if (UControlRigBlueprint* ControlRigBlueprint = CastChecked<UControlRigBlueprint>(GetBlueprintObj()))
-	{
-		return ControlRigBlueprint->DebugBoneRadius;
-	}
-	return 0.f;
-}
-
-void FControlRigEditor::OnToolbarBoneRadiusChanged(float InValue)
-{
-	if (UControlRigBlueprint* ControlRigBlueprint = CastChecked<UControlRigBlueprint>(GetBlueprintObj()))
-	{
-		if(!FMath::IsNearlyEqual(ControlRigBlueprint->DebugBoneRadius, InValue))
-		{
-			ControlRigBlueprint->Modify();
-			ControlRigBlueprint->DebugBoneRadius = InValue;
-		}
-	}
-	if (UDebugSkelMeshComponent* MeshComponent = GetPersonaToolkit()->GetPreviewScene()->GetPreviewMeshComponent())
-	{
-		MeshComponent->BoneRadiusMultiplier = InValue;
-	}
-	if(UControlRig* DebuggedControlRig = Cast<UControlRig>(GetBlueprintObj()->GetObjectBeingDebugged()))
-	{
-		DebuggedControlRig->DebugBoneRadiusMultiplier = InValue;
-	}
-}
-
 ECheckBoxState FControlRigEditor::GetToolbarDrawAxesOnSelection() const
 {
 	if (const UControlRigEditModeSettings* Settings = GetDefault<UControlRigEditModeSettings>())
@@ -4044,15 +4011,6 @@ void FControlRigEditor::HandlePreviewSceneCreated(const TSharedRef<IPersonaPrevi
 			GEditor->UnregisterForUndo(AnimationEditorPreviewScene);
 		}
 	}
-
-	if(HasEditingObject())
-	{
-		if(UControlRigBlueprint* Blueprint = Cast<UControlRigBlueprint>(GetBlueprintObj()))
-		{
-			OnToolbarBoneRadiusChanged(Blueprint->DebugBoneRadius);
-		}
-	}
-
 }
 
 void FControlRigEditor::UpdateControlRig()
diff --git a/Engine/Plugins/Animation/ControlRig/Source/ControlRigEditor/Private/Editor/ControlRigEditor.h b/Engine/Plugins/Animation/ControlRig/Source/ControlRigEditor/Private/Editor/ControlRigEditor.h
--- a/Engine/Plugins/Animation/ControlRig/Source/ControlRigEditor/Private/Editor/ControlRigEditor.h
+++ b/Engine/Plugins/Animation/ControlRig/Source/ControlRigEditor/Private/Editor/ControlRigEditor.h
@@ -303,8 +303,6 @@ private:
 	void OnToolbarDrawAxesOnSelectionChanged(ECheckBoxState InNewValue);
 	TOptional<float> GetToolbarAxesScale() const;
 	void OnToolbarAxesScaleChanged(float InValue);
-	TOptional<float> GetToolbarBoneRadius() const;
-	void OnToolbarBoneRadiusChanged(float InValue);
 
 		/** Handle switching skeletal meshes */
 	void HandlePreviewMeshChanged(USkeletalMesh* InOldSkeletalMesh, USkeletalMesh* InNewSkeletalMesh);
diff --git a/Engine/Plugins/Animation/ControlRig/Source/ControlRigEditor/Private/Editor/ControlRigEditorEditMode.cpp b/Engine/Plugins/Animation/ControlRig/Source/ControlRigEditor/Private/Editor/ControlRigEditorEditMode.cpp
--- a/Engine/Plugins/Animation/ControlRig/Source/ControlRigEditor/Private/Editor/ControlRigEditorEditMode.cpp
+++ b/Engine/Plugins/Animation/ControlRig/Source/ControlRigEditor/Private/Editor/ControlRigEditorEditMode.cpp
@@ -20,95 +20,195 @@ void FControlRigEditorEditMode::Render(const FSceneView* View, FViewport* Viewpo
 	}
 	EBoneDrawMode::Type BoneDrawMode = (EBoneDrawMode::Type)ConfigOption->DefaultBoneDrawSelection;
 
-	if (bDrawHierarchyBones && BoneDrawMode != EBoneDrawMode::None)
+	// is the viewport configured to draw bones AND should Control Rig draw bones at all (only if Skeletal Mesh not present)
+	const bool bShouldDrawBones = bDrawHierarchyBones && BoneDrawMode != EBoneDrawMode::None;
+	if (!bShouldDrawBones)
 	{
-		TArray<UControlRig*> ControlRigs = GetControlRigsArray(true /*bIsVisible*/);
-		if (ControlRigs.Num() > 0)
+		return;
+	}
+
+	// is there a control rig to edit?
+	TArray<UControlRig*> ControlRigs = GetControlRigsArray(true );
+	if (ControlRigs.IsEmpty())
+	{
+		return;
+	}
+
+	UControlRig* ControlRig = ControlRigs[0]; //just one control rig in the CR asset editor
+	if (!ControlRig) 
+	{
+		return;
+	}
+	
+	URigHierarchy* Hierarchy = ControlRig->GetHierarchy();
+	URigHierarchy* HierarchyForSelection = Hierarchy;
+	
+	if (const UClass* ControlRigClass = ControlRig->GetClass())
+	{
+		if (const UControlRigBlueprint* RigBlueprint = Cast<UControlRigBlueprint>(ControlRigClass->ClassGeneratedBy))
 		{
-			if (UControlRig* ControlRig = ControlRigs[0]) //just one control rig in the CR asset editor
-			{
-				URigHierarchy* Hierarchy = ControlRig->GetHierarchy();
-				URigHierarchy* HierarchyForSelection = Hierarchy;
+			HierarchyForSelection = RigBlueprint->Hierarchy;
+		}
+	}
+
+	// maintain 3 masks for drawing bones depending on viewport settings, selected and affected states
+	TBitArray<> BonesToDraw(false, Hierarchy->Num());
+	TBitArray<> SelectedBones(false, Hierarchy->Num());
+	TBitArray<> AffectedBones(false, Hierarchy->Num());
+
+	// get the viewport states to filter which bones to draw at all
+	const bool bDrawAllBones = BoneDrawMode == EBoneDrawMode::All;
+	const bool bDrawSelected = BoneDrawMode == EBoneDrawMode::Selected;
+	const bool bDrawSelectedAndParents = BoneDrawMode == EBoneDrawMode::SelectedAndParents;
+	const bool bDrawSelectedAndChildren = BoneDrawMode == EBoneDrawMode::SelectedAndChildren;
+	const bool bDrawSelectedAndParentsAndChildren = BoneDrawMode == EBoneDrawMode::SelectedAndParentsAndChildren;
+	
+	// determine which elements are bones and optionally which are selected
+	for (int32 ElementIndex = 0; ElementIndex < Hierarchy->Num(); ElementIndex++)
+	{
+		FRigBoneElement* BoneElement = Cast<FRigBoneElement>(Hierarchy->Get(ElementIndex));
+		if(!BoneElement)
+		{
+			continue;
+		}
+		
+		const int32 BoneIndex = BoneElement->GetIndex();
+
+		// optionally draw ALL bones
+		if (bDrawAllBones)
+		{
+			BonesToDraw[BoneIndex] = true;
+		}
 
-				if (UClass* ControlRigClass = ControlRig->GetClass())
+		// record which bones are selected
+		const bool bIsSelected = HierarchyForSelection->IsSelected(BoneElement->GetKey());
+		if (bIsSelected)
+		{
+			SelectedBones[BoneIndex] = true;
+		}
+		
+		// draw selected bones (optionally)
+		const bool bFilterOnlySelected = bDrawSelected || bDrawSelectedAndParents || bDrawSelectedAndChildren || bDrawSelectedAndParentsAndChildren;
+		if (bIsSelected && bFilterOnlySelected)
+		{
+			BonesToDraw[BoneIndex] = true;
+		
+			// add children of selected to list of bones to draw (optionally)
+			if (bDrawSelectedAndChildren || bDrawSelectedAndParentsAndChildren)
+			{
+				FRigBaseElementChildrenArray AllChildren = Hierarchy->GetChildren(BoneElement, true);
+				for (FRigBaseElement* Child : AllChildren)
 				{
-					if (UControlRigBlueprint* RigBlueprint = Cast<UControlRigBlueprint>(ControlRigClass->ClassGeneratedBy))
+					if(FRigBoneElement* ChildBone = Cast<FRigBoneElement>(Child))
 					{
-						HierarchyForSelection = RigBlueprint->Hierarchy;
+						BonesToDraw[ChildBone->GetIndex()] = true;
 					}
 				}
+			}
+		}
+	}
 
-				if (BoneDrawMode == EBoneDrawMode::SelectedAndParents)
-				{
-					if (BoneHasSelectedChild.Num() != Hierarchy->Num())
-					{
-						BoneHasSelectedChild.SetNumZeroed(Hierarchy->Num());
-					}
+	// add parents of selected (optionally)
+	if (bDrawSelectedAndParents || bDrawSelectedAndParentsAndChildren)
+	{
+		for (int32 ElementIndex=0; ElementIndex<Hierarchy->Num(); ++ElementIndex)
+		{
+			if (!SelectedBones[ElementIndex])
+			{
+				continue;
+			}
+			
+			int32 ParentIndex = Hierarchy->GetFirstParent(ElementIndex);
+			while (ParentIndex != INDEX_NONE)
+			{
+				BonesToDraw[ParentIndex] = true;
+				ParentIndex = Hierarchy->GetFirstParent(ParentIndex);
+			}
+		}
+	}
 
-					Hierarchy->ForEach<FRigBoneElement>([this, Hierarchy, HierarchyForSelection](FRigBoneElement* BoneElement) -> bool
-						{
-							BoneHasSelectedChild[BoneElement->GetIndex()] = false;
-							if (!HierarchyForSelection->IsSelected(BoneElement->GetKey()))
-							{
-								return true;
-							}
-
-							int32 ParentIndex = Hierarchy->GetFirstParent(BoneElement->GetIndex());
-							while (ParentIndex != INDEX_NONE)
-							{
-								BoneHasSelectedChild[ParentIndex] = true;
-								ParentIndex = Hierarchy->GetFirstParent(ParentIndex);
-							}
-
-							return true;
-						});
-				}
+	// determine bones are "affected" (child of a selected bone, but not selected themselves)
+	Hierarchy->ForEach<FRigBoneElement>([this, Hierarchy, &SelectedBones, &AffectedBones](FRigBoneElement* BoneElement) -> bool
+	{
+		const int32 BoneIndex = BoneElement->GetIndex();
+		int32 ParentIndex = Hierarchy->GetFirstParent(BoneIndex);
+		while (ParentIndex != INDEX_NONE)
+		{
+			if (SelectedBones[ParentIndex])
+			{
+				AffectedBones[BoneIndex] = true;
+				return true;
+			}
+			ParentIndex = Hierarchy->GetFirstParent(ParentIndex);
+		}
 
-				Hierarchy->ForEach<FRigBoneElement>([this, PDI, Hierarchy, HierarchyForSelection, &BoneDrawMode](FRigBoneElement* BoneElement) -> bool
-					{
-						const int32 ParentIndex = Hierarchy->GetFirstParent(BoneElement->GetIndex());
-						const bool bSelected = HierarchyForSelection->IsSelected(BoneElement->GetKey());
-
-						if (!bSelected)
-						{
-							if (BoneDrawMode == EBoneDrawMode::Selected)
-							{
-								return true;
-							}
-							else if (BoneDrawMode == EBoneDrawMode::SelectedAndParents)
-							{
-								if (!BoneHasSelectedChild[BoneElement->GetIndex()])
-								{
-									return true;
-								}
-							}
-						}
-
-						const FLinearColor LineColor = bSelected ? FLinearColor(1.0f, 0.34f, 0.0f, 1.0f) : FLinearColor::White;
-
-						FVector Start, End;
-						if (ParentIndex != INDEX_NONE)
-						{
-							Start = Hierarchy->GetGlobalTransform(ParentIndex).GetLocation();
-							End = Hierarchy->GetGlobalTransform(BoneElement->GetIndex()).GetLocation();
-						}
-						else
-						{
-							Start = FVector::ZeroVector;
-							End = Hierarchy->GetGlobalTransform(BoneElement->GetIndex()).GetLocation();
-						}
-
-						const float BoneLength = (End - Start).Size();
-						const float Radius = FMath::Clamp(BoneLength * 0.05f, 0.1f, 10000.f);
-
-						//Render Sphere for bone end point and a cone between it and its parent.
-						PDI->SetHitProxy(new HPersonaBoneHitProxy(BoneElement->GetIndex(), BoneElement->GetName()));
-						SkeletalDebugRendering::DrawWireBone(PDI, Start, End, LineColor, SDPG_Foreground, Radius);
-						PDI->SetHitProxy(nullptr);
-						return true;
-					});
+		return true;
+	});
+
+	// get size setting for bone drawing
+	float BoneRadius = 1.0f;
+	if (FAnimationViewportClient* AnimViewportClient = static_cast<FAnimationViewportClient*>(Viewport->GetClient()))
+	{
+		BoneRadius = AnimViewportClient->GetBoneDrawSize();
+	}
+
+	// spin through all bones and render them
+	TArray<FVector> ChildPositions;
+	TArray<FLinearColor> ChildColors;
+	TArray<int32> ChildrenIndices;
+	for ( int32 ElementIndex=0; ElementIndex<BonesToDraw.Num(); ++ElementIndex )
+	{
+		// skips bones that should not be drawn
+		if (!BonesToDraw[ElementIndex])
+		{
+			continue;
+		}
+
+		// determine color of bone based on selection / affected state
+		const bool bIsSelected = SelectedBones[ElementIndex];
+		const bool bIsAffected = AffectedBones[ElementIndex];
+		FLinearColor DefaultBoneColor = SkeletalDebugRendering::DEFAULT_BONE_COLOR;
+		FLinearColor BoneColor = bIsAffected ? SkeletalDebugRendering::AFFECTED_BONE_COLOR : DefaultBoneColor;
+		BoneColor = bIsSelected ? SkeletalDebugRendering::SELECTED_BONE_COLOR : BoneColor;
+
+		// draw the little coordinate frame inside the bone ONLY if selected or affected
+		const bool bDrawAxesInsideBone = bIsAffected|| bIsSelected;
+
+		// draw cone to each child
+		// but use a different color if this bone is NOT selected, but the child IS selected
+		ChildPositions.Reset();
+		ChildColors.Reset();
+		ChildrenIndices = Hierarchy->GetChildren(ElementIndex);
+		for (int32 ChildIndex : ChildrenIndices)
+		{
+			FTransform ChildTransform = Hierarchy->GetGlobalTransform(ChildIndex);
+			ChildPositions.Add(ChildTransform.GetLocation());
+			FLinearColor ChildLineColor = BoneColor;
+			if (!bIsSelected && SelectedBones[ChildIndex])
+			{
+				ChildLineColor = SkeletalDebugRendering::PARENT_OF_SELECTED_BONE_COLOR;
 			}
+			ChildColors.Add(ChildLineColor);
+		}
+
+		const FName BoneName = Hierarchy->GetKey(ElementIndex).Name;
+		const FTransform BoneTransform = Hierarchy->GetGlobalTransform(ElementIndex);
+		
+		PDI->SetHitProxy(new HPersonaBoneHitProxy(ElementIndex, BoneName));
+		SkeletalDebugRendering::DrawWireBoneAdvanced(
+			PDI,
+			BoneTransform,
+			ChildPositions,
+			ChildColors,
+			BoneColor,
+			SDPG_Foreground,
+			BoneRadius,
+			bDrawAxesInsideBone);
+		if (Hierarchy->GetFirstParent(ElementIndex) == INDEX_NONE)
+		{
+			SkeletalDebugRendering::DrawRootCone(PDI, BoneTransform, FVector::Zero(), BoneRadius);
 		}
+		PDI->SetHitProxy(nullptr);
 	}
 }
 
diff --git a/Engine/Plugins/Animation/ControlRig/Source/ControlRigEditor/Private/Editor/ControlRigEditorEditMode.h b/Engine/Plugins/Animation/ControlRig/Source/ControlRigEditor/Private/Editor/ControlRigEditorEditMode.h
--- a/Engine/Plugins/Animation/ControlRig/Source/ControlRigEditor/Private/Editor/ControlRigEditorEditMode.h
+++ b/Engine/Plugins/Animation/ControlRig/Source/ControlRigEditor/Private/Editor/ControlRigEditorEditMode.h
@@ -28,7 +28,4 @@ public:
 
 	/** Drawing options */
 	UPersonaOptions* ConfigOption = nullptr;
-
-	/** Flag to know if a bone is visible or not */
-	TArray<bool> BoneHasSelectedChild;
 };
\ No newline at end of file
diff --git a/Engine/Plugins/Animation/IKRig/Source/IKRig/Private/AnimNodes/AnimNode_RetargetPoseFromMesh.cpp b/Engine/Plugins/Animation/IKRig/Source/IKRig/Private/AnimNodes/AnimNode_RetargetPoseFromMesh.cpp
--- a/Engine/Plugins/Animation/IKRig/Source/IKRig/Private/AnimNodes/AnimNode_RetargetPoseFromMesh.cpp
+++ b/Engine/Plugins/Animation/IKRig/Source/IKRig/Private/AnimNodes/AnimNode_RetargetPoseFromMesh.cpp
@@ -194,16 +194,6 @@ void FAnimNode_RetargetPoseFromMesh::UpdateSpeedValuesFromCurves()
 	}
 }
 
-#if WITH_EDITOR
-void FAnimNode_RetargetPoseFromMesh::SetProcessorNeedsInitialized()
-{
-	if (Processor)
-	{
-		Processor->SetNeedsInitialized();
-	}
-}
-#endif
-
 UIKRetargetProcessor* FAnimNode_RetargetPoseFromMesh::GetRetargetProcessor() const
 {
 	return Processor;
diff --git a/Engine/Plugins/Animation/IKRig/Source/IKRig/Private/Retargeter/IKRetargetProcessor.cpp b/Engine/Plugins/Animation/IKRig/Source/IKRig/Private/Retargeter/IKRetargetProcessor.cpp
--- a/Engine/Plugins/Animation/IKRig/Source/IKRig/Private/Retargeter/IKRetargetProcessor.cpp
+++ b/Engine/Plugins/Animation/IKRig/Source/IKRig/Private/Retargeter/IKRetargetProcessor.cpp
@@ -11,7 +11,12 @@
 
 #define LOCTEXT_NAMESPACE "IKRetargetProcessor"
 
-void FRetargetSkeleton::Initialize(USkeletalMesh* InSkeletalMesh)
+// This is the default end of branch index value, meaning we haven't cached it yet
+#define RETARGETSKELETON_INVALID_BRANCH_INDEX -2
+
+void FRetargetSkeleton::Initialize(
+	USkeletalMesh* InSkeletalMesh,
+	const TArray<FBoneChain>& BoneChains)
 {	
 	// record which skeletal mesh this is running on
 	SkeletalMesh = InSkeletalMesh;
@@ -24,8 +29,26 @@ void FRetargetSkeleton::Initialize(USkeletalMesh* InSkeletalMesh)
 		ParentIndices.Add(RefSkeleton.GetParentIndex(BoneIndex));	
 	}
 
+	// determine set of bones referenced by one of the retarget bone chains
+	// this is the set of bones that will be affected by the retarget pose
+	IsBoneInAnyChain.Init(false, BoneNames.Num());
+	for (const FBoneChain& BoneChain : BoneChains)
+	{
+		TArray<int32> BonesInChain;
+		if (FResolvedBoneChain(BoneChain, *this, BonesInChain).IsValid())
+		{
+			for (const int32 BoneInChain : BonesInChain)
+			{
+				IsBoneInAnyChain[BoneInChain] = true;
+			}
+		}
+	}
+
 	// update retarget pose to reflect custom offsets
 	GenerateRetargetPose();
+
+	// initialize branch caching
+	CachedEndOfBranchIndices.Init(RETARGETSKELETON_INVALID_BRANCH_INDEX, ParentIndices.Num());
 }
 
 void FRetargetSkeleton::Reset()
@@ -129,11 +152,58 @@ FTransform FRetargetSkeleton::GetGlobalRefPoseOfSingleBone(
 	return ChildLocalTransform * ParentGlobalTransform;
 }
 
+int32 FRetargetSkeleton::GetCachedEndOfBranchIndex(const int32 InBoneIndex) const
+{
+	if (!CachedEndOfBranchIndices.IsValidIndex(InBoneIndex))
+	{
+		return INDEX_NONE;
+	}
+
+	// already cached
+	if (CachedEndOfBranchIndices[InBoneIndex] != RETARGETSKELETON_INVALID_BRANCH_INDEX)
+	{
+		return CachedEndOfBranchIndices[InBoneIndex];
+	}
+
+	const int32 NumBones = BoneNames.Num();
+	
+	// if we're asking for root's branch, get the last bone  
+	if (InBoneIndex == 0)
+	{
+		CachedEndOfBranchIndices[InBoneIndex] = NumBones-1;
+		return CachedEndOfBranchIndices[InBoneIndex];
+	}
+
+	CachedEndOfBranchIndices[InBoneIndex] = INDEX_NONE;
+	const int32 StartParentIndex = GetParentIndex(InBoneIndex);
+	int32 BoneIndex = InBoneIndex + 1;
+	int32 ParentIndex = GetParentIndex(BoneIndex);
+
+	// if next child bone's parent is less than or equal to StartParentIndex,
+	// we are leaving the branch so no need to go further
+	while (ParentIndex > StartParentIndex && BoneIndex < NumBones)
+	{
+		CachedEndOfBranchIndices[InBoneIndex] = BoneIndex;
+				
+		BoneIndex++;
+		ParentIndex = GetParentIndex(BoneIndex);
+	}
+
+	return CachedEndOfBranchIndices[InBoneIndex];
+}
+
 void FRetargetSkeleton::GetChildrenIndices(const int32 BoneIndex, TArray<int32>& OutChildren) const
 {
-	for (int32 ChildBoneIndex=0; ChildBoneIndex<ParentIndices.Num(); ++ChildBoneIndex)
+	const int32 LastBranchIndex = GetCachedEndOfBranchIndex(BoneIndex);
+	if (LastBranchIndex == INDEX_NONE)
 	{
-		if (ParentIndices[ChildBoneIndex] == BoneIndex)
+		// no children (leaf bone)
+		return;
+	}
+	
+	for (int32 ChildBoneIndex = BoneIndex + 1; ChildBoneIndex <= LastBranchIndex; ChildBoneIndex++)
+	{
+		if (GetParentIndex(ChildBoneIndex) == BoneIndex)
 		{
 			OutChildren.Add(ChildBoneIndex);
 		}
@@ -142,13 +212,16 @@ void FRetargetSkeleton::GetChildrenIndices(const int32 BoneIndex, TArray<int32>&
 
 void FRetargetSkeleton::GetChildrenIndicesRecursive(const int32 BoneIndex, TArray<int32>& OutChildren) const
 {
-	const int32 NumBones = BoneNames.Num();
-	for (int32 ChildIndex=BoneIndex+1; ChildIndex<NumBones; ++ChildIndex)
+	const int32 LastBranchIndex = GetCachedEndOfBranchIndex(BoneIndex);
+	if (LastBranchIndex == INDEX_NONE)
 	{
-		if (!OutChildren.Contains(ChildIndex) && IsParentOfChild(BoneIndex, ChildIndex))
-		{
-			OutChildren.Add(ChildIndex);
-		}
+		// no children (leaf bone)
+		return;
+	}
+	
+	for (int32 ChildBoneIndex = BoneIndex + 1; ChildBoneIndex <= LastBranchIndex; ChildBoneIndex++)
+	{
+		OutChildren.Add(ChildBoneIndex);
 	}
 }
 
@@ -170,7 +243,7 @@ bool FRetargetSkeleton::IsParentOfChild(const int32 PotentialParentIndex, const
 
 int32 FRetargetSkeleton::GetParentIndex(const int32 BoneIndex) const
 {
-	if (BoneIndex < 0 || BoneIndex>ParentIndices.Num() || BoneIndex == INDEX_NONE)
+	if (BoneIndex < 0 || BoneIndex>=ParentIndices.Num() || BoneIndex == INDEX_NONE)
 	{
 		return INDEX_NONE;
 	}
@@ -184,33 +257,18 @@ void FTargetSkeleton::Initialize(
 	const FName& RetargetRootBone,
 	const TArray<FBoneChain>& TargetChains)
 {
-	FRetargetSkeleton::Initialize(InSkeletalMesh);
-
-	// initialize storage for output pose (the result of the retargeting)
-	OutputGlobalPose = RetargetGlobalPose;
+	FRetargetSkeleton::Initialize(InSkeletalMesh, TargetChains);
 
 	// make storage for per-bone "Is Retargeted" flag (used for hierarchy updates)
 	// these are bones that are in a target chain that is mapped to a source chain (ie, will actually be retargeted)
 	// these flags are actually set later in init phase when bone chains are mapped together
-	IsBoneRetargeted.Init(false, OutputGlobalPose.Num());
+	IsBoneRetargeted.Init(false, BoneNames.Num());
 
-	// determine set of bones referenced by one of the target bone chains to be retargeted
-	// this is the set of bones that will be affected by the retarget pose
-	IsBoneInAnyTargetChain.Init(false, OutputGlobalPose.Num());
-	for (const FBoneChain& TargetChain : TargetChains)
-	{
-		TArray<int32> BonesInChain;
-		if (FResolvedBoneChain(TargetChain, *this, BonesInChain).IsValid())
-		{
-			for (int32 BoneInChain : BonesInChain)
-			{
-				IsBoneInAnyTargetChain[BoneInChain] = true;
-			}
-		}
-	}
+	// initialize storage for output pose (the result of the retargeting)
+	OutputGlobalPose = RetargetGlobalPose;
 
 	// generate the retarget pose (applies stored offsets)
-	// NOTE: this must be done AFTER generating IsBoneInAnyTargetChain array above
+	// NOTE: this must be done AFTER generating IsBoneInAnyTargetChain array above (FRetargetSkeleton::Initialize)
 	GenerateRetargetPose(RetargetPose, RetargetRootBone);
 }
 
@@ -248,7 +306,7 @@ void FTargetSkeleton::GenerateRetargetPose(const FIKRetargetPose* InRetargetPose
 			continue;
 		}
 
-		if (!IsBoneInAnyTargetChain[BoneIndex] && BoneIndex!=RootBoneIndex)
+		if (!IsBoneInAnyChain[BoneIndex] && BoneIndex!=RootBoneIndex)
 		{
 			// this can happen if a retarget pose includes bone edits from a bone chain that was subsequently removed,
 			// and the asset has not run through the "CleanChainMapping" operation yet (happens on load)
@@ -1245,7 +1303,7 @@ void UIKRetargetProcessor::Initialize(
 	}
 	
 	// initialize skeleton data for source and target
-	SourceSkeleton.Initialize(SourceSkeletalMesh);
+	SourceSkeleton.Initialize(SourceSkeletalMesh, RetargeterAsset->GetSourceIKRig()->GetRetargetChains());
 	TargetSkeleton.Initialize(
 		TargetSkeletalMesh,
 		RetargeterAsset->GetCurrentRetargetPose(),
@@ -1573,14 +1631,6 @@ void UIKRetargetProcessor::ResetPlanting()
 	}
 }
 
-FTransform UIKRetargetProcessor::GetTargetBoneRetargetPoseGlobalTransform(const int32& TargetBoneIndex) const
-{
-	check(TargetSkeleton.BoneNames.IsValidIndex(TargetBoneIndex))
-
-	// get the current retarget pose
-	return TargetSkeleton.RetargetGlobalPose[TargetBoneIndex];
-}
-
 FTransform UIKRetargetProcessor::GetTargetBoneRetargetPoseLocalTransform(const int32& TargetBoneIndex) const
 {
 	check(TargetSkeleton.BoneNames.IsValidIndex(TargetBoneIndex))
diff --git a/Engine/Plugins/Animation/IKRig/Source/IKRig/Public/AnimNodes/AnimNode_RetargetPoseFromMesh.h b/Engine/Plugins/Animation/IKRig/Source/IKRig/Public/AnimNodes/AnimNode_RetargetPoseFromMesh.h
--- a/Engine/Plugins/Animation/IKRig/Source/IKRig/Public/AnimNodes/AnimNode_RetargetPoseFromMesh.h
+++ b/Engine/Plugins/Animation/IKRig/Source/IKRig/Public/AnimNodes/AnimNode_RetargetPoseFromMesh.h
@@ -51,11 +51,6 @@ struct IKRIG_API FAnimNode_RetargetPoseFromMesh : public FAnimNode_Base
 	virtual void PreUpdate(const UAnimInstance* InAnimInstance) override;
 	// End of FAnimNode_Base interface
 
-#if WITH_EDITOR
-	/** Force reinitialization. */
-	void SetProcessorNeedsInitialized();
-#endif
-
 	/** Access to the runtime processor */
 	UIKRetargetProcessor* GetRetargetProcessor() const;
 
diff --git a/Engine/Plugins/Animation/IKRig/Source/IKRig/Public/IKRigDefinition.h b/Engine/Plugins/Animation/IKRig/Source/IKRig/Public/IKRigDefinition.h
--- a/Engine/Plugins/Animation/IKRig/Source/IKRig/Public/IKRigDefinition.h
+++ b/Engine/Plugins/Animation/IKRig/Source/IKRig/Public/IKRigDefinition.h
@@ -189,13 +189,9 @@ public:
 	TSoftObjectPtr<USkeletalMesh> PreviewSkeletalMesh;
 
 #if WITH_EDITORONLY_DATA
-
-	/**Draw bones in the viewport.*/
-	UPROPERTY(EditAnywhere, Category = "Viewport Bone Settings")
-	bool DrawBones = true;
 	
-	/**The size of the Bones in the editor viewport.*/
-	UPROPERTY(EditAnywhere, Category = "Viewport Bone Settings", meta = (ClampMin = "0.01", UIMin = "0.1", UIMax = "100.0"))
+	/**The size of the Bones in the editor viewport. This is set by callbacks from the viewport Character>Bones menu*/
+	UPROPERTY()
 	float BoneSize = 2.0f;
 
 	/**Draw bones in the viewport.*/
diff --git a/Engine/Plugins/Animation/IKRig/Source/IKRig/Public/Retargeter/IKRetargetProcessor.h b/Engine/Plugins/Animation/IKRig/Source/IKRig/Public/Retargeter/IKRetargetProcessor.h
--- a/Engine/Plugins/Animation/IKRig/Source/IKRig/Public/Retargeter/IKRetargetProcessor.h
+++ b/Engine/Plugins/Animation/IKRig/Source/IKRig/Public/Retargeter/IKRetargetProcessor.h
@@ -30,7 +30,14 @@ struct IKRIG_API FRetargetSkeleton
 	
 	USkeletalMesh* SkeletalMesh;
 
-	void Initialize(USkeletalMesh* InSkeletalMesh);
+	// true for bones that are in a retarget bone chain and thus are posed by a retarget pose
+	// the application of a retarget pose to a bone happens regardless of whether the chain is mapped or not
+	// (this must happen or else child chains may not end up in the correct retarget pose)
+	TArray<bool> IsBoneInAnyChain;
+
+	void Initialize(
+		USkeletalMesh* InSkeletalMesh,
+		const TArray<FBoneChain>& BoneChains);
 
 	void Reset();
 
@@ -64,11 +71,19 @@ struct IKRIG_API FRetargetSkeleton
 		const int32 BoneIndex,
 		const TArray<FTransform>& InGlobalPose) const;
 
+	int32 GetCachedEndOfBranchIndex(const int32 InBoneIndex) const;
+
 	void GetChildrenIndices(const int32 BoneIndex, TArray<int32>& OutChildren) const;
 
 	void GetChildrenIndicesRecursive(const int32 BoneIndex, TArray<int32>& OutChildren) const;
 	
 	bool IsParentOfChild(const int32 PotentialParentIndex, const int32 ChildBoneIndex) const;
+
+private:
+	
+	/** One index per-bone. Lazy-filled on request. Stores the last element of the branch below the bone.
+	 * You can iterate between in the indices stored here and the bone in question to iterate over all children recursively */
+	mutable TArray<int32> CachedEndOfBranchIndices;
 };
 
 struct FTargetSkeleton : public FRetargetSkeleton
@@ -77,9 +92,6 @@ struct FTargetSkeleton : public FRetargetSkeleton
 	// true for bones that are in a target chain that is ALSO mapped to a source chain
 	// ie, bones that are actually posed based on a mapped source chain
 	TArray<bool> IsBoneRetargeted;
-	// true for bones that are in a target chain and thus are posed by a retarget pose
-	// the application of a retarget pose to a bone happens regardless of whether they are mapped to a source chain or not
-	TArray<bool> IsBoneInAnyTargetChain;
 
 	void Initialize(
 		USkeletalMesh* InSkeletalMesh, 
@@ -483,24 +495,20 @@ public:
 	*/
 	TArray<FTransform>& RunRetargeter(const TArray<FTransform>& InSourceGlobalPose, const TMap<FName, float>& SpeedValuesFromCurves, const float DeltaTime);
 
-	/**
-	* Get the Global transform, in the currently used retarget pose, for a bone in the target skeletal mesh
-	* @param TargetBoneIndex - The index of a bone in the target skeleton.
-	* @return The Global space transform for the bone or Identity if bone not found.
-	* @warning This function is only valid to call after the retargeter has been initialized (which generates global retarget pose)
-	*/
-	FTransform GetTargetBoneRetargetPoseGlobalTransform(const int32& TargetBoneIndex) const;
-
+	
 	FTransform GetTargetBoneRetargetPoseLocalTransform(const int32& TargetBoneIndex) const;
 
 	/** Get read-only access to the target skeleton. */
 	const FTargetSkeleton& GetTargetSkeleton() const { return TargetSkeleton; };
 
-	/** Get index of the root bone of the target skeleton. */
-	const int32 GetTargetSkeletonRootBone() const { return RootRetargeter.Target.BoneIndex; };
-
 	/** Get read-only access to the source skeleton. */
 	const FRetargetSkeleton& GetSourceSkeleton() const { return SourceSkeleton; };
+
+	/** Get index of the root bone of the source skeleton. */
+	const int32 GetSourceRetargetRoot() const { return RootRetargeter.Source.BoneIndex; };
+
+	/** Get index of the root bone of the target skeleton. */
+	const int32 GetTargetRetargetRoot() const { return RootRetargeter.Target.BoneIndex; };
 	
 	/** Get whether this processor is ready to call RunRetargeter() and generate new poses. */
 	bool IsInitialized() const { return bIsInitialized; };
diff --git a/Engine/Plugins/Animation/IKRig/Source/IKRig/Public/Retargeter/IKRetargeter.h b/Engine/Plugins/Animation/IKRig/Source/IKRig/Public/Retargeter/IKRetargeter.h
--- a/Engine/Plugins/Animation/IKRig/Source/IKRig/Public/Retargeter/IKRetargeter.h
+++ b/Engine/Plugins/Animation/IKRig/Source/IKRig/Public/Retargeter/IKRetargeter.h
@@ -353,8 +353,8 @@ public:
 	UPROPERTY(EditAnywhere, Category = PreviewSettings)
 	FVector SourceMeshOffset;
 
-	/** The visual size of the bones in the viewport when editing the retarget pose.*/
-	UPROPERTY(EditAnywhere, Category = PreviewSettings, meta = (ClampMin = "0.0", UIMin = "0.01", UIMax = "10.0"))
+	/** The visual size of the bones in the viewport (saved between sessions). This is set from the viewport Character>Bones menu*/
+	UPROPERTY()
 	float BoneDrawSize = 8.0f;
 	
 private:
diff --git a/Engine/Plugins/Animation/IKRig/Source/IKRigEditor/Private/RetargetEditor/IKRetargetAnimInstance.cpp b/Engine/Plugins/Animation/IKRig/Source/IKRigEditor/Private/RetargetEditor/IKRetargetAnimInstance.cpp
--- a/Engine/Plugins/Animation/IKRig/Source/IKRigEditor/Private/RetargetEditor/IKRetargetAnimInstance.cpp
+++ b/Engine/Plugins/Animation/IKRig/Source/IKRigEditor/Private/RetargetEditor/IKRetargetAnimInstance.cpp
@@ -22,11 +22,6 @@ UIKRetargetProcessor* UIKRetargetAnimInstance::GetRetargetProcessor() const
 	return IKRetargeterNode.GetRetargetProcessor();
 }
 
-void UIKRetargetAnimInstance::SetProcessorNeedsInitialized()
-{
-	IKRetargeterNode.SetProcessorNeedsInitialized();
-}
-
 FAnimInstanceProxy* UIKRetargetAnimInstance::CreateAnimInstanceProxy()
 {
 	LLM_SCOPE_BYNAME(TEXT("Animation/IKRig"));
diff --git a/Engine/Plugins/Animation/IKRig/Source/IKRigEditor/Private/RetargetEditor/IKRetargetAnimInstanceProxy.cpp b/Engine/Plugins/Animation/IKRig/Source/IKRigEditor/Private/RetargetEditor/IKRetargetAnimInstanceProxy.cpp
--- a/Engine/Plugins/Animation/IKRig/Source/IKRigEditor/Private/RetargetEditor/IKRetargetAnimInstanceProxy.cpp
+++ b/Engine/Plugins/Animation/IKRig/Source/IKRigEditor/Private/RetargetEditor/IKRetargetAnimInstanceProxy.cpp
@@ -51,6 +51,9 @@ void FIKRetargetAnimInstanceProxy::SetRetargetAssetAndSourceComponent(
 	IKRetargetNode->bUseAttachedParent = false;
 	IKRetargetNode->SourceMeshComponent = InSourceMeshComponent;
 	IKRetargetNode->bDriveWithAsset = true;
-	IKRetargetNode->SetProcessorNeedsInitialized();
+	if (UIKRetargetProcessor* Processor = IKRetargetNode->GetRetargetProcessor())
+	{
+		Processor->SetNeedsInitialized();
+	}
 }
 
diff --git a/Engine/Plugins/Animation/IKRig/Source/IKRigEditor/Private/RetargetEditor/IKRetargetApplicationMode.cpp b/Engine/Plugins/Animation/IKRig/Source/IKRigEditor/Private/RetargetEditor/IKRetargetApplicationMode.cpp
--- a/Engine/Plugins/Animation/IKRig/Source/IKRigEditor/Private/RetargetEditor/IKRetargetApplicationMode.cpp
+++ b/Engine/Plugins/Animation/IKRig/Source/IKRigEditor/Private/RetargetEditor/IKRetargetApplicationMode.cpp
@@ -10,6 +10,7 @@
 
 #include "RetargetEditor/IKRetargetChainTabSummoner.h"
 #include "RetargetEditor/IKRetargetEditor.h"
+#include "RetargetEditor/IKRetargetHierarchyTabSummoner.h"
 #include "RetargetEditor/IKRetargetOutputLogTabSummoner.h"
 
 #define LOCTEXT_NAMESPACE "IKRetargetMode"
@@ -28,6 +29,7 @@ FIKRetargetApplicationMode::FIKRetargetApplicationMode(
 	ViewportArgs.bShowStats = false;
 	ViewportArgs.bShowTurnTable = false;
 	ViewportArgs.ContextName = TEXT("IKRetargetEditor.Viewport");
+	ViewportArgs.OnViewportCreated = FOnViewportCreated::CreateSP(IKRetargetEditor, &FIKRetargetEditor::HandleViewportCreated);
 
 	// register Persona tabs
 	FPersonaModule& PersonaModule = FModuleManager::LoadModuleChecked<FPersonaModule>("Persona");
@@ -38,9 +40,10 @@ FIKRetargetApplicationMode::FIKRetargetApplicationMode(
 	TabFactories.RegisterFactory(MakeShared<FIKRetargetChainTabSummoner>(IKRetargetEditor));
 	TabFactories.RegisterFactory(MakeShared<FIKRetargetAssetBrowserTabSummoner>(IKRetargetEditor));
 	TabFactories.RegisterFactory(MakeShared<FIKRetargetOutputLogTabSummoner>(IKRetargetEditor));
+	TabFactories.RegisterFactory(MakeShared<FIKRetargetHierarchyTabSummoner>(IKRetargetEditor));
 
 	// create tab layout
-	TabLayout = FTabManager::NewLayout("Standalone_IKRetargetEditor_Layout_v1.012")
+	TabLayout = FTabManager::NewLayout("Standalone_IKRetargetEditor_Layout_v1.014")
 		->AddArea
 		(
 			FTabManager::NewPrimaryArea()
@@ -53,14 +56,27 @@ FIKRetargetApplicationMode::FIKRetargetApplicationMode(
 				->Split
 				(
 					FTabManager::NewSplitter()
-					->SetSizeCoefficient(0.9f)
+					->SetSizeCoefficient(0.8f)
 					->SetOrientation(Orient_Vertical)
 					->Split
 					(
+					FTabManager::NewSplitter()
+						//->SetSizeCoefficient(0.8f)
+						->SetOrientation(Orient_Horizontal)
+						->Split
+						(
+							FTabManager::NewStack()
+							->SetSizeCoefficient(0.2f)
+							->SetHideTabWell(true)
+							->AddTab(FIKRetargetHierarchyTabSummoner::TabID, ETabState::OpenedTab)
+						)
+						->Split
+						(
 						FTabManager::NewStack()
-						->SetSizeCoefficient(0.8f)
-						->SetHideTabWell(true)
-						->AddTab(FPersonaTabs::PreviewViewportID, ETabState::OpenedTab)
+							->SetSizeCoefficient(0.8f)
+							->SetHideTabWell(true)
+							->AddTab(FPersonaTabs::PreviewViewportID, ETabState::OpenedTab)
+						)
 					)
 					->Split
 					(
diff --git a/Engine/Plugins/Animation/IKRig/Source/IKRigEditor/Private/RetargetEditor/IKRetargetDefaultMode.cpp b/Engine/Plugins/Animation/IKRig/Source/IKRigEditor/Private/RetargetEditor/IKRetargetDefaultMode.cpp
--- a/Engine/Plugins/Animation/IKRig/Source/IKRigEditor/Private/RetargetEditor/IKRetargetDefaultMode.cpp
+++ b/Engine/Plugins/Animation/IKRig/Source/IKRigEditor/Private/RetargetEditor/IKRetargetDefaultMode.cpp
@@ -57,12 +57,9 @@
 		PrimitiveComponent->SelectionOverrideDelegate = UPrimitiveComponent::FSelectionOverride::CreateRaw(this, &FIKRetargetDefaultMode::ComponentSelectionOverride);
 	}
 
-	// deselect all
-	SetSelectedComponent(nullptr);
-
 	// update offsets on preview meshes
-	Controller->AddOffsetAndUpdatePreviewMeshPosition(FVector::ZeroVector, Controller->SourceSkelMeshComponent);
-	Controller->AddOffsetAndUpdatePreviewMeshPosition(FVector::ZeroVector, Controller->TargetSkelMeshComponent);
+	Controller->AddOffsetToMeshComponent(FVector::ZeroVector, Controller->SourceSkelMeshComponent);
+	Controller->AddOffsetToMeshComponent(FVector::ZeroVector, Controller->TargetSkelMeshComponent);
 
 	bIsInitialized = true;
 }
@@ -89,8 +86,14 @@
 
 bool FIKRetargetDefaultMode::UsesTransformWidget(UE::Widget::EWidgetMode CheckMode) const
 {
+	const TSharedPtr<FIKRetargetEditorController> Controller = EditorController.Pin();
+	if (!Controller.IsValid())
+	{
+		return false; 
+	}
+	
 	const bool bTranslating = CheckMode == UE::Widget::EWidgetMode::WM_Translate;
-	return bTranslating && IsValid(SelectedComponent);
+	return bTranslating && IsValid(Controller->GetSelectedMesh());
 }
 
 FVector FIKRetargetDefaultMode::GetWidgetLocation() const
@@ -101,12 +104,12 @@
 		return FVector::ZeroVector; 
 	}
 	
-	if (SelectedComponent.IsNull())
+	if (!Controller->GetSelectedMesh())
 	{
 		return FVector::ZeroVector; // shouldn't get here
 	}
 
-	return SelectedComponent->GetComponentTransform().GetLocation();
+	return Controller->GetSelectedMesh()->GetComponentTransform().GetLocation();
 }
 
 bool FIKRetargetDefaultMode::HandleClick(FEditorViewportClient* InViewportClient, HHitProxy* HitProxy, const FViewportClick& Click)
@@ -116,19 +119,42 @@
 	{
 		return false; 
 	}
-	
-	// did we select an actor in the viewport?
+
 	const bool bLeftButtonClicked = Click.GetKey() == EKeys::LeftMouseButton;
+	constexpr bool bFromHierarchy = false;
+	
+	// did we click on an actor in the viewport?
 	const bool bHitActor = HitProxy && HitProxy->IsA(HActor::StaticGetType());
 	if (bLeftButtonClicked && bHitActor)
 	{
 		const HActor* ActorProxy = static_cast<HActor*>(HitProxy);
-		SetSelectedComponent(const_cast<UPrimitiveComponent*>(ActorProxy->PrimComponent));
+		Controller->SetSelectedMesh(const_cast<UPrimitiveComponent*>(ActorProxy->PrimComponent));
 		return true;
 	}
 
-	SetSelectedComponent(nullptr);
-	return false;
+	// did we click on a bone in the viewport?
+	const bool bHitBone = HitProxy && HitProxy->IsA(HPersonaBoneHitProxy::StaticGetType());
+	if (bLeftButtonClicked && bHitBone)
+	{
+		const HPersonaBoneHitProxy* BoneProxy = static_cast<HPersonaBoneHitProxy*>(HitProxy);
+		const TArray<FName> BoneNames{BoneProxy->BoneName};
+		
+		const bool bCtrlOrShiftHeld = Click.IsControlDown() || Click.IsShiftDown();
+		const EBoneSelectionEdit EditMode = bCtrlOrShiftHeld ? EBoneSelectionEdit::Add : EBoneSelectionEdit::Replace;
+		
+		Controller->EditBoneSelection(BoneNames, EditMode, bFromHierarchy);
+		
+		return true;
+	}
+
+	// did we click in empty space in viewport?
+	if (!(bHitActor || bHitBone))
+	{
+		// deselect all meshes, bones, chains and update details view
+		Controller->ClearSelection();
+	}
+	
+	return true;
 }
 
 bool FIKRetargetDefaultMode::StartTracking(FEditorViewportClient* InViewportClient, FViewport* InViewport)
@@ -149,7 +175,7 @@
 	}
 
 	const bool bTranslating = InViewportClient->GetWidgetMode() == UE::Widget::EWidgetMode::WM_Translate;
-	if (bTranslating && IsValid(SelectedComponent))
+	if (bTranslating && IsValid(Controller->GetSelectedMesh()))
 	{
 		bIsTranslating = true;
 		GEditor->BeginTransaction(LOCTEXT("MovePreviewMesh", "Move Preview Mesh"));
@@ -174,23 +200,23 @@
 	FRotator& InRot,
 	FVector& InScale)
 {
-	if (!(bIsTranslating && IsValid(SelectedComponent)))
+	const TSharedPtr<FIKRetargetEditorController> Controller = EditorController.Pin();
+	if (!Controller.IsValid())
 	{
+		return false; 
+	}
+	
+	if (!(bIsTranslating && IsValid(Controller->GetSelectedMesh())))
+	{
 		return false; // not handled
 	}
 
-	const TSharedPtr<FIKRetargetEditorController> Controller = EditorController.Pin();
-	if (!Controller.IsValid())
-	{
-		return false; 
-	}
-
 	if(InViewportClient->GetWidgetMode() != UE::Widget::WM_Translate)
 	{
 		return false;
 	}
 
-	Controller->AddOffsetAndUpdatePreviewMeshPosition(InDrag, SelectedComponent);
+	Controller->AddOffsetToMeshComponent(InDrag, Controller->GetSelectedMesh());
 	
 	return true;
 }
@@ -203,12 +229,12 @@
 		return false; 
 	}
 
-	if (SelectedComponent.IsNull())
+	if (!Controller->GetSelectedMesh())
 	{
 		return false;
 	}
 
-	InMatrix = SelectedComponent->GetComponentTransform().ToMatrixNoScale().RemoveTranslation();
+	InMatrix = Controller->GetSelectedMesh()->GetComponentTransform().ToMatrixNoScale().RemoveTranslation();
 	return true;
 }
 
@@ -227,13 +253,16 @@
 		return; 
 	}
 
+	// record which skeleton is being viewed/edited
+	SkeletonMode = Controller->GetSkeletonMode();
+
 	// allow selection of meshes in this mode
 	Controller->Editor.Pin()->GetPersonaToolkit()->GetPreviewScene()->SetAllowMeshHitProxies(true);
 	Controller->SourceSkelMeshComponent->bSelectable = true;
 	Controller->TargetSkelMeshComponent->bSelectable = true;
 
-	// deselect when entering mode
-	SetSelectedComponent(nullptr);
+	// clear all selection when entering mode
+	Controller->ClearSelection();
 }
 
 void FIKRetargetDefaultMode::Exit()
@@ -249,39 +278,32 @@
 	Controller->SourceSkelMeshComponent->bSelectable = false;
 	Controller->TargetSkelMeshComponent->bSelectable = false;
 	
-	// deselect all
-	SetSelectedComponent(nullptr);
+	// deselect everything
+	Controller->ClearSelection();
 
 	IPersonaEditMode::Exit();
 }
 
 bool FIKRetargetDefaultMode::ComponentSelectionOverride(const UPrimitiveComponent* InComponent) const
 {
-	return InComponent == SelectedComponent;
+	const TSharedPtr<FIKRetargetEditorController> Controller = EditorController.Pin();
+	if (!Controller.IsValid())
+	{
+		return false; 
+	}
+	
+	return InComponent == Controller->GetSelectedMesh();
 }
 
-void FIKRetargetDefaultMode::SetSelectedComponent(UPrimitiveComponent* InComponent)
+UDebugSkelMeshComponent* FIKRetargetDefaultMode::GetCurrentlyEditedMesh() const
 {
 	const TSharedPtr<FIKRetargetEditorController> Controller = EditorController.Pin();
 	if (!Controller.IsValid())
 	{
-		return; 
+		return nullptr; 
 	}
 	
-	SelectedComponent = InComponent;
-
-	Controller->SourceSkelMeshComponent->PushSelectionToProxy();
-	Controller->TargetSkelMeshComponent->PushSelectionToProxy();
-	
-	Controller->SourceSkelMeshComponent->MarkRenderStateDirty();
-	Controller->TargetSkelMeshComponent->MarkRenderStateDirty();
-
-	// when clicking in empty space, show global details
-	if (SelectedComponent == nullptr)
-	{
-		Controller->ChainsView->ClearSelection();
-		Controller->DetailsView->SetObject(Controller->AssetController->GetAsset());
-	}
+	return SkeletonMode == EIKRetargetSkeletonMode::Source ? Controller->SourceSkelMeshComponent : Controller->TargetSkelMeshComponent;
 }
 
 void FIKRetargetDefaultMode::ApplyOffsetToMeshTransform(const FVector& Offset, USceneComponent* Component)
@@ -303,6 +325,15 @@
 	{
 		Initialize();
 	}
+
+	// update skeleton drawing mode
+	const TSharedPtr<FIKRetargetEditorController> Controller = EditorController.Pin();
+	if (Controller.IsValid())
+	{
+		const bool bEditingSource = Controller->GetSkeletonMode() == EIKRetargetSkeletonMode::Source;
+		Controller->SourceSkelMeshComponent->SkeletonDrawMode = bEditingSource ? ESkeletonDrawMode::Default : ESkeletonDrawMode::GreyedOut;
+		Controller->TargetSkelMeshComponent->SkeletonDrawMode = !bEditingSource ? ESkeletonDrawMode::Default : ESkeletonDrawMode::GreyedOut;
+	}
 }
 
 void FIKRetargetDefaultMode::DrawHUD(FEditorViewportClient* ViewportClient, FViewport* Viewport, const FSceneView* View, FCanvas* Canvas)
diff --git a/Engine/Plugins/Animation/IKRig/Source/IKRigEditor/Private/RetargetEditor/IKRetargetDetails.cpp b/Engine/Plugins/Animation/IKRig/Source/IKRigEditor/Private/RetargetEditor/IKRetargetDetails.cpp
--- a/Engine/Plugins/Animation/IKRig/Source/IKRigEditor/Private/RetargetEditor/IKRetargetDetails.cpp
+++ b/Engine/Plugins/Animation/IKRig/Source/IKRigEditor/Private/RetargetEditor/IKRetargetDetails.cpp
@@ -80,10 +80,6 @@ void FIKRetargeterDetails::CustomizeDetails(IDetailLayoutBuilder& DetailBuilder)
 			MakeToolbar(Commands)
 		]
 	];
-
-	// add the bone size slider
-	TSharedRef<IPropertyHandle> BoneSizeHandle = DetailBuilder.GetProperty(GET_MEMBER_NAME_CHECKED(UIKRetargeter, BoneDrawSize));
-	EditPoseCategoryBuilder.AddProperty(BoneSizeHandle);
 }
 
 TSharedRef<SWidget>  FIKRetargeterDetails::MakeToolbar(TSharedPtr<FUICommandList> Commands)
diff --git a/Engine/Plugins/Animation/IKRig/Source/IKRigEditor/Private/RetargetEditor/IKRetargetEditPoseMode.cpp b/Engine/Plugins/Animation/IKRig/Source/IKRigEditor/Private/RetargetEditor/IKRetargetEditPoseMode.cpp
--- a/Engine/Plugins/Animation/IKRig/Source/IKRigEditor/Private/RetargetEditor/IKRetargetEditPoseMode.cpp
+++ b/Engine/Plugins/Animation/IKRig/Source/IKRigEditor/Private/RetargetEditor/IKRetargetEditPoseMode.cpp
@@ -50,85 +50,114 @@
 {
 	FEdMode::Render(View, Viewport, PDI);
 
-	const TSharedPtr<FIKRetargetEditorController> Controller = EditorController.Pin();
-	if (!Controller.IsValid())
+	const FIKRetargetEditorController* Controller = EditorController.Pin().Get();
+	if (!Controller)
 	{
 		return;
 	}
 
-	const UIKRetargetProcessor* RetargetProcessor = Controller->GetRetargetProcessor();
-	if (!(RetargetProcessor && RetargetProcessor->IsInitialized()))
+	const UIKRetargetProcessor* Processor = EditorController.Pin()->GetRetargetProcessor();
+	if (!(Processor && Processor->IsInitialized()))
 	{
 		return;
 	}
 
-	const UIKRetargeter* Asset = Controller->AssetController->GetAsset();
-	const FTargetSkeleton& TargetSkeleton = RetargetProcessor->GetTargetSkeleton();
-	const int32 RootBoneIndex = RetargetProcessor->GetTargetSkeletonRootBone();
+	// render the skeleton
+	RenderSkeleton(PDI, Controller, Processor);
+}
 
+void FIKRetargetEditPoseMode::RenderSkeleton(
+	FPrimitiveDrawInterface* PDI,
+	const FIKRetargetEditorController* Controller,
+	const UIKRetargetProcessor* Processor)
+{
+	check(Processor && Controller && Processor->IsInitialized())
+	
+	const FRetargetSkeleton& Skeleton = GetCurrentlyEditedSkeleton(*Processor);
+	
+	// get selected and affected bones in this skeleton
 	TSet<int32> OutAffectedBones;
 	TSet<int32> OutSelectedBones;
-	UIKRigProcessor* IKRigProcessor = RetargetProcessor->GetTargetIKRigProcessor();
-	GetAffectedBones(Controller.Get(), IKRigProcessor, OutAffectedBones, OutSelectedBones);
+	GetSelectedAndAffectedBones(Controller, Skeleton, OutAffectedBones, OutSelectedBones);
 
+	const UIKRetargeter* Asset = Controller->AssetController->GetAsset();
+	float Scale;
+	FVector Offset;
+	GetEditedComponentScaleAndOffset(Scale, Offset);
+	const float BoneDrawSize = Asset->BoneDrawSize;
+	const int32 RootBoneIndex = GetEditedRetargetRootBoneIndex(*Processor);
 	const float MaxDrawRadius = Controller->TargetSkelMeshComponent->Bounds.SphereRadius * 0.01f;
-	for (int32 BoneIndex = 0; BoneIndex<TargetSkeleton.BoneNames.Num(); ++BoneIndex)
-	{		
-		// filter out bones that cannot be edited as part of the retarget pose
-		const bool bIsRootBone = BoneIndex==RootBoneIndex;
-		const bool bIsInTargetChain = TargetSkeleton.IsBoneInAnyTargetChain[BoneIndex];
-		if (!(bIsRootBone || bIsInTargetChain))
-		{
-			continue;
-		}
 
+	// loop over whole skeleton and draw each bone
+	for (int32 BoneIndex = 0; BoneIndex<Skeleton.BoneNames.Num(); ++BoneIndex)
+	{
 		// selected bones are drawn with different color, affected bones are highlighted
 		const bool bIsSelected = OutSelectedBones.Contains(BoneIndex);
 		const bool bIsAffected = OutAffectedBones.Contains(BoneIndex);
-		FLinearColor LineColor = bIsAffected ? IKRigDebugRendering::AFFECTED_BONE_COLOR : IKRigDebugRendering::DESELECTED_BONE_COLOR;
-		LineColor = bIsSelected ? IKRigDebugRendering::SELECTED_BONE_COLOR : LineColor;
-
-		const float BoneRadiusSetting = Asset->BoneDrawSize;
-		const float BoneRadius = FMath::Min(1.0f, MaxDrawRadius) * BoneRadiusSetting;
+		const bool bIsDisabled = !(Skeleton.IsBoneInAnyChain[BoneIndex] || BoneIndex==RootBoneIndex);
+		FLinearColor LineColor = bIsAffected ? SkeletalDebugRendering::AFFECTED_BONE_COLOR : SkeletalDebugRendering::DEFAULT_BONE_COLOR;
+		LineColor = bIsSelected ? SkeletalDebugRendering::SELECTED_BONE_COLOR : LineColor;
+		LineColor = bIsDisabled ? SkeletalDebugRendering::DISABLED_BONE_COLOR : LineColor;
 		
+		const float BoneRadius = FMath::Min(1.0f, MaxDrawRadius) * BoneDrawSize;
+		
 		// get the location of the bone on the currently initialized target skeletal mesh
 		// along with array of child positions
-		const FTransform BoneTransform = Controller->GetTargetBoneGlobalTransform(RetargetProcessor, BoneIndex);
-		FVector Start;
-		TArray<FVector> ChildLocations;
-		Controller->GetTargetBoneLineSegments(RetargetProcessor, BoneIndex, Start, ChildLocations);
+		const FTransform BoneTransform = Controller->GetGlobalRetargetPoseOfBone(Skeleton, BoneIndex, Scale, Offset);
+		TArray<FVector> ChildrenPositions;
+		TArray<int32> ChildrenIndices;
+		Controller->GetGlobalRetargetPoseOfImmediateChildren(
+			Skeleton,
+			BoneIndex,
+			Scale,
+			Offset,
+			ChildrenIndices,
+			ChildrenPositions);
+
+		// draw cone to parent with different color if this bone is NOT selected, but the child IS selected
+		TArray<FLinearColor> ChildColors;
+		for (int32 ChildIndex : ChildrenIndices)
+		{
+			FLinearColor ChildLineColor = LineColor;
+			if (!bIsSelected && OutSelectedBones.Contains(ChildIndex))
+			{
+				ChildLineColor = SkeletalDebugRendering::PARENT_OF_SELECTED_BONE_COLOR;
+			}
+			ChildColors.Add(ChildLineColor);
+		}
 		
 		// draw the bone proxy
-		PDI->SetHitProxy(new HIKRetargetEditorBoneProxy(TargetSkeleton.BoneNames[BoneIndex]));
-		IKRigDebugRendering::DrawWireBone(
+		if (!bIsDisabled)
+		{
+			PDI->SetHitProxy(new HIKRetargetEditorBoneProxy(Skeleton.BoneNames[BoneIndex]));
+		}
+		SkeletalDebugRendering::DrawWireBoneAdvanced(
 			PDI,
 			BoneTransform,
-			ChildLocations,
+			ChildrenPositions,
+			ChildColors,
 			LineColor,
 			SDPG_Foreground,
 			BoneRadius,
 			bIsSelected || bIsAffected);
-		PDI->SetHitProxy(NULL);
+		if (!bIsDisabled)
+		{
+			PDI->SetHitProxy(NULL);
+		}
 	}
 }
 
-void FIKRetargetEditPoseMode::GetAffectedBones(
-	FIKRetargetEditorController* Controller,
-	UIKRigProcessor* Processor,
+void FIKRetargetEditPoseMode::GetSelectedAndAffectedBones(
+	const FIKRetargetEditorController* Controller,
+	const FRetargetSkeleton& Skeleton,
 	TSet<int32>& OutAffectedBones,
 	TSet<int32>& OutSelectedBones) const
 {
-	if (!(Controller && Processor))
-	{
-		return;
-	}
-
-	const FIKRigSkeleton& Skeleton = Processor->GetSkeleton();
-
 	// record selected bone indices
-	for (const FName& SelectedBone : BoneEdit.SelectedBones)
+	const TArray<FName>& SelectedBones = Controller->GetSelectedBones();
+	for (const FName& SelectedBone : SelectedBones)
 	{
-		OutSelectedBones.Add(Skeleton.GetBoneIndexFromName(SelectedBone));
+		OutSelectedBones.Add(Skeleton.FindBoneIndexByName(SelectedBone));
 	}
 
 	// "affected bones" are the selected bones AND their children, recursively
@@ -160,8 +189,14 @@
 
 bool FIKRetargetEditPoseMode::UsesTransformWidget(UE::Widget::EWidgetMode CheckMode) const
 {
+	const TSharedPtr<FIKRetargetEditorController> Controller = EditorController.Pin();
+	if (!Controller.IsValid())
+	{
+		return false; 
+	}
+	
 	const bool bIsOnlyRootSelected = IsOnlyRootSelected();
-	const bool bIsAnyBoneSelected = !BoneEdit.SelectedBones.IsEmpty();
+	const bool bIsAnyBoneSelected = !Controller->GetSelectedBones().IsEmpty();
 
 	if (!bIsAnyBoneSelected)
 	{
@@ -183,62 +218,68 @@
 
 FVector FIKRetargetEditPoseMode::GetWidgetLocation() const
 {
-	if (BoneEdit.SelectedBones.IsEmpty())
-	{
-		return FVector::ZeroVector; 
-	}
-
 	const TSharedPtr<FIKRetargetEditorController> Controller = EditorController.Pin();
 	if (!Controller.IsValid())
 	{
 		return FVector::ZeroVector; 
 	}
 
-	const FTargetSkeleton& TargetSkeleton = Controller->GetRetargetProcessor()->GetTargetSkeleton();
-	const int32 BoneIndex = TargetSkeleton.FindBoneIndexByName(BoneEdit.SelectedBones.Last());
+	const TArray<FName>& SelectedBones = Controller->GetSelectedBones();
+	if (SelectedBones.IsEmpty())
+	{
+		return FVector::ZeroVector; 
+	}
+
+	const UIKRetargetProcessor* RetargetProcessor = Controller->GetRetargetProcessor();
+	if (!(RetargetProcessor && RetargetProcessor->IsInitialized()))
+	{
+		return FVector::ZeroVector;
+	}
+
+	const bool bIsSource = Controller->GetSkeletonMode() == EIKRetargetSkeletonMode::Source;
+	
+	const FRetargetSkeleton& Skeleton = bIsSource ? RetargetProcessor->GetSourceSkeleton() : RetargetProcessor->GetTargetSkeleton();
+	const int32 BoneIndex = Skeleton.FindBoneIndexByName(SelectedBones.Last());
 	if (BoneIndex == INDEX_NONE)
 	{
 		return FVector::ZeroVector;
 	}
-
-	const UIKRetargetProcessor* RetargetProcessor = Controller->GetRetargetProcessor();
-	if (!(RetargetProcessor && RetargetProcessor->IsInitialized()))
-	{
-		return FVector::ZeroVector;
-	}
 	
-	return Controller->GetTargetBoneGlobalTransform(RetargetProcessor, BoneIndex).GetTranslation();
+	const UIKRetargeter* Asset = Controller->AssetController->GetAsset();
+	const float Scale = bIsSource ? 1.0f : Asset->TargetMeshScale;
+	const FVector Offset = bIsSource ? Asset->SourceMeshOffset : Asset->TargetMeshOffset;
+	
+	return Controller->GetGlobalRetargetPoseOfBone(Skeleton, BoneIndex, Scale, Offset).GetTranslation();
 }
 
 bool FIKRetargetEditPoseMode::HandleClick(FEditorViewportClient* InViewportClient, HHitProxy* HitProxy, const FViewportClick& Click)
 {
-	// check for selections
-	if (Click.GetKey() == EKeys::LeftMouseButton)
+	const TSharedPtr<FIKRetargetEditorController> Controller = EditorController.Pin();
+	if (!Controller.IsValid())
 	{
-		// selected bone ?
-		if (HitProxy && HitProxy->IsA(HIKRetargetEditorBoneProxy::StaticGetType()))
-		{
-			HIKRetargetEditorBoneProxy* BoneProxy = static_cast<HIKRetargetEditorBoneProxy*>(HitProxy);
-			const bool bReplaceSelection = !(InViewportClient->IsCtrlPressed() || InViewportClient->IsShiftPressed());
-			HandleBoneSelectedInViewport(BoneProxy->BoneName, bReplaceSelection);
-			return true;
-		}
-		else
-		{
-			// clicking in empty space clears selection
-			const bool bReplaceSelection = true;
-			HandleBoneSelectedInViewport(NAME_None, bReplaceSelection);
+		return false; 
+	}
 
-			// show asset in details view
-			const TSharedPtr<FIKRetargetEditorController> Controller = EditorController.Pin();
-			if (Controller.IsValid())
-			{
-				Controller->DetailsView->SetObject(Controller->AssetController->GetAsset());
-			}
-		}
+	
+	if (Click.GetKey() != EKeys::LeftMouseButton)
+	{
+		return false;
 	}
 	
-	return false;
+	// selected bone ?
+	if (HitProxy && HitProxy->IsA(HIKRetargetEditorBoneProxy::StaticGetType()))
+	{
+		HIKRetargetEditorBoneProxy* BoneProxy = static_cast<HIKRetargetEditorBoneProxy*>(HitProxy);
+		const TArray<FName> BoneNames{BoneProxy->BoneName};
+		constexpr bool bFromHierarchy = false;
+		const EBoneSelectionEdit EditMode = Click.IsControlDown() || Click.IsShiftDown() ? EBoneSelectionEdit::Add : EBoneSelectionEdit::Replace;
+		Controller->EditBoneSelection(BoneNames, EditMode, bFromHierarchy);
+		return true;
+	}
+
+	// clicking in empty space clears selection
+	Controller->ClearSelection();
+	return true;
 }
 
 bool FIKRetargetEditPoseMode::StartTracking(FEditorViewportClient* InViewportClient, FViewport* InViewport)
@@ -261,7 +302,7 @@
 	// get state of viewport
 	const bool bTranslating = InViewportClient->GetWidgetMode() == UE::Widget::EWidgetMode::WM_Translate;
 	const bool bRotating = InViewportClient->GetWidgetMode() == UE::Widget::EWidgetMode::WM_Rotate;
-	const bool bAnyBoneSelected = !BoneEdit.SelectedBones.IsEmpty();
+	const bool bAnyBoneSelected = !Controller->GetSelectedBones().IsEmpty();
 	const bool bOnlyRootSelected = IsOnlyRootSelected();
 
 	// is any bone being rotated?
@@ -347,10 +388,11 @@
 		const FQuat BoneLocalOffset = BoneLocal * BoneLocalPlusOffset.Inverse();
 		
 		// apply rotation delta to all selected bones
-		for (int32 SelectionIndex=0; SelectionIndex<BoneEdit.SelectedBones.Num(); ++SelectionIndex)
+		const TArray<FName>& SelectedBones = Controller->GetSelectedBones();
+		for (int32 SelectionIndex=0; SelectionIndex<SelectedBones.Num(); ++SelectionIndex)
 		{
 			const FQuat TotalOffsetLocal = BoneLocalOffset * BoneEdit.PrevLocalOffsets[SelectionIndex];
-			const FName& BoneName = BoneEdit.SelectedBones[SelectionIndex];
+			const FName& BoneName = SelectedBones[SelectionIndex];
 			Controller->AssetController->SetRotationOffsetForRetargetPoseBone(BoneName, TotalOffsetLocal);
 		}
 
@@ -375,17 +417,17 @@
 
 bool FIKRetargetEditPoseMode::GetCustomDrawingCoordinateSystem(FMatrix& InMatrix, void* InData)
 {
-	if (BoneEdit.SelectedBones.IsEmpty())
+	const TSharedPtr<FIKRetargetEditorController> Controller = EditorController.Pin();
+	if (!Controller.IsValid())
 	{
+		return false; 
+	}
+
+	if (Controller->GetSelectedBones().IsEmpty())
+	{
 		return false; // nothing selected to manipulate
 	}
 
-	const TSharedPtr<FIKRetargetEditorController> Controller = EditorController.Pin();
-	if (!Controller.IsValid())
-	{
-		return false; 
-	}
-
 	if (TrackingState == FIKRetargetTrackingState::None)
 	{
 		UpdateWidgetTransform();
@@ -410,13 +452,28 @@
 		return;
 	}
 
+	const UIKRetargetProcessor* Processor = Controller->GetRetargetProcessor();
+	if (!Processor)
+	{
+		return;
+	}
+
 	// put source mesh in reference pose
 	Controller->SourceSkelMeshComponent->ShowReferencePose(true);
 	// have to move component back to offset position because ShowReferencePose() sets it back to origin
-	Controller->AddOffsetAndUpdatePreviewMeshPosition(FVector::ZeroVector, Controller->SourceSkelMeshComponent);
-
+	Controller->AddOffsetToMeshComponent(FVector::ZeroVector, Controller->SourceSkelMeshComponent);
 	// put asset in mode where target mesh will output retarget pose (for preview purposes)
 	Controller->AssetController->GetAsset()->SetOutputMode(ERetargeterOutputMode::EditRetargetPose);
+
+	// clear selection
+	BoneEdit.Reset();
+	
+	// record skeleton to edit (must be constant between enter/exit)
+	SkeletonMode = Controller->GetSkeletonMode();
+	// hide skeleton we're editing so that we can draw our own editable version of it
+	const bool bEditingSource = SkeletonMode == EIKRetargetSkeletonMode::Source;
+	Controller->SourceSkelMeshComponent->SkeletonDrawMode = bEditingSource ? ESkeletonDrawMode::Hidden : ESkeletonDrawMode::GreyedOut;
+	Controller->TargetSkelMeshComponent->SkeletonDrawMode = !bEditingSource ? ESkeletonDrawMode::Hidden : ESkeletonDrawMode::GreyedOut;
 }
 
 void FIKRetargetEditPoseMode::Exit()
@@ -430,15 +487,44 @@
 	IPersonaEditMode::Exit();
 }
 
-const TArray<FName>& FIKRetargetEditPoseMode::GetSelectedBones() const
+const FRetargetSkeleton& FIKRetargetEditPoseMode::GetCurrentlyEditedSkeleton(const UIKRetargetProcessor& Processor) const
 {
-	return BoneEdit.SelectedBones;
+	return SkeletonMode == EIKRetargetSkeletonMode::Source ? Processor.GetSourceSkeleton() : Processor.GetTargetSkeleton();
 }
 
+void FIKRetargetEditPoseMode::GetEditedComponentScaleAndOffset(float& OutScale, FVector& OutOffset) const
+{
+	OutScale = 1.0f;
+	OutOffset = FVector::Zero();
+
+	const bool bIsSource = EditorController.Pin()->GetSkeletonMode() == EIKRetargetSkeletonMode::Source;
+	const UIKRetargeter* Asset = EditorController.Pin()->AssetController->GetAsset();
+	if (!Asset)
+	{
+		return;
+	}
+	
+	OutScale = bIsSource ? 1.0f : Asset->TargetMeshScale;
+	OutOffset = bIsSource ? Asset->SourceMeshOffset : Asset->TargetMeshOffset;
+}
+
+int32 FIKRetargetEditPoseMode::GetEditedRetargetRootBoneIndex(const UIKRetargetProcessor& Processor) const
+{
+	const bool bIsSource = EditorController.Pin()->GetSkeletonMode() == EIKRetargetSkeletonMode::Source;
+	return bIsSource ? Processor.GetSourceRetargetRoot() : Processor.GetTargetRetargetRoot();
+}
+
 void FIKRetargetEditPoseMode::UpdateWidgetTransform()
 {
 	const TSharedPtr<FIKRetargetEditorController> Controller = EditorController.Pin();
-	if (!Controller.IsValid() || BoneEdit.SelectedBones.IsEmpty())
+	if (!Controller.IsValid())
+	{
+		BoneEdit.GlobalTransform = FTransform::Identity;
+		return;
+	}
+
+	const TArray<FName>& SelectedBones = Controller->GetSelectedBones();
+	if (SelectedBones.IsEmpty())
 	{
 		BoneEdit.GlobalTransform = FTransform::Identity;
 		return;
@@ -451,24 +537,27 @@
 	}
 	
 	const UIKRetargeterController* AssetController = Controller->AssetController;
-	const FTargetSkeleton& TargetSkeleton = RetargetProcessor->GetTargetSkeleton();
+	const FRetargetSkeleton& Skeleton = GetCurrentlyEditedSkeleton(*RetargetProcessor);
+	float Scale;
+	FVector Offset;
+	GetEditedComponentScaleAndOffset(Scale,Offset);
 
-	BoneEdit.Name = BoneEdit.SelectedBones.Last();
-	BoneEdit.Index = TargetSkeleton.FindBoneIndexByName(BoneEdit.Name);
-	BoneEdit.GlobalTransform = Controller->GetTargetBoneGlobalTransform(RetargetProcessor, BoneEdit.Index);
+	BoneEdit.Name = SelectedBones.Last();
+	BoneEdit.Index = Skeleton.FindBoneIndexByName(BoneEdit.Name);
+	BoneEdit.GlobalTransform = Controller->GetGlobalRetargetPoseOfBone(Skeleton, BoneEdit.Index, Scale, Offset);
 	BoneEdit.AccumulatedGlobalOffset = FQuat::Identity;
 
 	BoneEdit.PrevLocalOffsets.Reset();
-	for (int32 SelectionIndex=0; SelectionIndex<BoneEdit.SelectedBones.Num(); ++SelectionIndex)
+	for (int32 SelectionIndex=0; SelectionIndex<SelectedBones.Num(); ++SelectionIndex)
 	{
-		FQuat PrevLocalOffset = AssetController->GetRotationOffsetForRetargetPoseBone(BoneEdit.SelectedBones[SelectionIndex]);
+		FQuat PrevLocalOffset = AssetController->GetRotationOffsetForRetargetPoseBone(SelectedBones[SelectionIndex]);
 		BoneEdit.PrevLocalOffsets.Add(PrevLocalOffset);
 	}
 	
-	const int32 ParentIndex = TargetSkeleton.GetParentIndex(BoneEdit.Index);
+	const int32 ParentIndex = Skeleton.GetParentIndex(BoneEdit.Index);
 	if (ParentIndex != INDEX_NONE)
 	{
-		BoneEdit.ParentGlobalTransform = Controller->GetTargetBoneGlobalTransform(RetargetProcessor, ParentIndex);
+		BoneEdit.ParentGlobalTransform = Controller->GetGlobalRetargetPoseOfBone(Skeleton, ParentIndex, Scale, Offset);
 	}else
 	{
 		BoneEdit.ParentGlobalTransform = FTransform::Identity;
@@ -477,43 +566,32 @@
 
 bool FIKRetargetEditPoseMode::IsRootSelected() const
 {
-	if (BoneEdit.SelectedBones.IsEmpty())
-	{
-		return false;
-	}
-
 	const TSharedPtr<FIKRetargetEditorController> Controller = EditorController.Pin();
 	if (!Controller.IsValid())
 	{
 		return false; 
 	}
-
+	
+	const TArray<FName>& SelectedBones = Controller->GetSelectedBones();
 	const FName& RootName = Controller->AssetController->GetTargetRootBone();
-	for (const FName& SelectedBone : BoneEdit.SelectedBones)
-	{
-		if (RootName == SelectedBone)
-		{
-			return true;
-		}
-	}
-
-	return false;
+	return SelectedBones.Contains(RootName);
 }
 
 bool FIKRetargetEditPoseMode::IsOnlyRootSelected() const
 {
-	if (BoneEdit.SelectedBones.Num() != 1)
-	{
-		return false;
-	}
-
 	const TSharedPtr<FIKRetargetEditorController> Controller = EditorController.Pin();
 	if (!Controller.IsValid())
 	{
 		return false; 
 	}
 
-	return Controller->AssetController->GetTargetRootBone() == BoneEdit.SelectedBones[0];
+	const TArray<FName>& SelectedBones = Controller->GetSelectedBones();
+	if (SelectedBones.Num() != 1)
+	{
+		return false;
+	}
+
+	return Controller->AssetController->GetTargetRootBone() == SelectedBones[0];
 }
 
 void FIKRetargetEditPoseMode::Tick(FEditorViewportClient* ViewportClient, float DeltaTime)
@@ -528,36 +606,4 @@
 	FEdMode::DrawHUD(ViewportClient, Viewport, View, Canvas);
 }
 
-bool FIKRetargetEditPoseMode::IsBoneSelected(const FName& BoneName) const
-{
-	return BoneEdit.SelectedBones.Contains(BoneName);
-}
-
-void FIKRetargetEditPoseMode::HandleBoneSelectedInViewport(const FName& BoneName, bool bReplace)
-{
-	if (bReplace)
-	{
-		BoneEdit.SelectedBones.Reset();
-		if (BoneName != NAME_None)
-		{
-			BoneEdit.SelectedBones.Add(BoneName);
-		}
-	}
-	else
-	{
-		if (BoneName != NAME_None)
-		{
-			const bool bAlreadySelected = BoneEdit.SelectedBones.Contains(BoneName);
-			if (bAlreadySelected)
-			{
-				BoneEdit.SelectedBones.Remove(BoneName);	
-			}
-			else
-			{
-				BoneEdit.SelectedBones.Add(BoneName);
-			}
-		}	
-	}
-}
-
 #undef LOCTEXT_NAMESPACE
diff --git a/Engine/Plugins/Animation/IKRig/Source/IKRigEditor/Private/RetargetEditor/IKRetargetEditor.cpp b/Engine/Plugins/Animation/IKRig/Source/IKRigEditor/Private/RetargetEditor/IKRetargetEditor.cpp
--- a/Engine/Plugins/Animation/IKRig/Source/IKRigEditor/Private/RetargetEditor/IKRetargetEditor.cpp
+++ b/Engine/Plugins/Animation/IKRig/Source/IKRigEditor/Private/RetargetEditor/IKRetargetEditor.cpp
@@ -2,6 +2,7 @@
 
 #include "RetargetEditor/IKRetargetEditor.h"
 
+#include "Persona/Private/AnimationEditorViewportClient.h"
 #include "AnimationEditorPreviewActor.h"
 #include "EditorModeManager.h"
 #include "GameFramework/WorldSettings.h"
@@ -12,6 +13,7 @@
 #include "ISkeletonEditorModule.h"
 #include "Preferences/PersonaOptions.h"
 #include "AnimCustomInstanceHelper.h"
+#include "IPersonaViewport.h"
 #include "Animation/DebugSkelMeshComponent.h"
 #include "Framework/MultiBox/MultiBoxBuilder.h"
 
@@ -255,8 +257,8 @@ void FIKRetargetEditor::AddReferencedObjects(FReferenceCollector& Collector)
 void FIKRetargetEditor::Tick(float DeltaTime)
 {
 	// update with latest offsets
-	EditorController->AddOffsetAndUpdatePreviewMeshPosition(FVector::ZeroVector, EditorController->SourceSkelMeshComponent);
-	EditorController->AddOffsetAndUpdatePreviewMeshPosition(FVector::ZeroVector, EditorController->TargetSkelMeshComponent);
+	EditorController->AddOffsetToMeshComponent(FVector::ZeroVector, EditorController->SourceSkelMeshComponent);
+	EditorController->AddOffsetToMeshComponent(FVector::ZeroVector, EditorController->TargetSkelMeshComponent);
 
 	// retargeter IK planting must be reset when time is reversed or playback jumps ahead 
 	const float CurrentTime = EditorController->SourceAnimInstance->GetCurrentTime();
@@ -283,6 +285,33 @@ void FIKRetargetEditor::PostRedo(bool bSuccess)
 	EditorController->OnRetargeterNeedsInitialized(EditorController->AssetController->GetAsset());
 }
 
+void FIKRetargetEditor::HandleViewportCreated(const TSharedRef<class IPersonaViewport>& InViewport)
+{
+	// register callbacks to allow the asset to store the Bone Size viewport setting
+	FEditorViewportClient& ViewportClient = InViewport->GetViewportClient();
+	if (FAnimationViewportClient* AnimViewportClient = static_cast<FAnimationViewportClient*>(&ViewportClient))
+	{
+		AnimViewportClient->OnSetBoneSize.BindLambda([this](float InBoneSize)
+		{
+			if (UIKRetargeter* Asset = EditorController->AssetController->GetAsset())
+			{
+				Asset->Modify();
+				Asset->BoneDrawSize = InBoneSize;
+			}
+		});
+		
+		AnimViewportClient->OnGetBoneSize.BindLambda([this]() -> float
+		{
+			if (const UIKRetargeter* Asset = EditorController->AssetController->GetAsset())
+			{
+				return Asset->BoneDrawSize;
+			}
+
+			return 1.0f;
+		});
+	}
+}
+
 void FIKRetargetEditor::HandlePreviewSceneCreated(const TSharedRef<IPersonaPreviewScene>& InPersonaPreviewScene)
 {
 	AAnimationEditorPreviewActor* Actor = InPersonaPreviewScene->GetWorld()->SpawnActor<AAnimationEditorPreviewActor>(AAnimationEditorPreviewActor::StaticClass(), FTransform::Identity);
diff --git a/Engine/Plugins/Animation/IKRig/Source/IKRigEditor/Private/RetargetEditor/IKRetargetEditorController.cpp b/Engine/Plugins/Animation/IKRig/Source/IKRigEditor/Private/RetargetEditor/IKRetargetEditorController.cpp
--- a/Engine/Plugins/Animation/IKRig/Source/IKRigEditor/Private/RetargetEditor/IKRetargetEditorController.cpp
+++ b/Engine/Plugins/Animation/IKRig/Source/IKRigEditor/Private/RetargetEditor/IKRetargetEditorController.cpp
@@ -18,6 +18,7 @@
 #include "RetargetEditor/IKRetargetEditor.h"
 #include "RetargetEditor/IKRetargetFactory.h"
 #include "RetargetEditor/SIKRetargetChainMapList.h"
+#include "RetargetEditor/SIKRetargetHierarchy.h"
 #include "Retargeter/IKRetargeter.h"
 #include "RigEditor/SIKRigOutputLog.h"
 #include "RigEditor/IKRigController.h"
@@ -30,6 +31,7 @@ void FIKRetargetEditorController::Initialize(TSharedPtr<FIKRetargetEditor> InEdi
 	Editor = InEditor;
 	AssetController = UIKRetargeterController::GetController(InAsset);
 	AssetController->SetEditorController(this);
+	SkeletonMode = EIKRetargetSkeletonMode::Target;
 
 	// bind callbacks when SOURCE or TARGET IK Rigs are modified
 	BindToIKRigAsset(AssetController->GetAsset()->GetSourceIKRigWriteable());
@@ -57,7 +59,7 @@ void FIKRetargetEditorController::BindToIKRigAsset(UIKRigDefinition* InIKRig) co
 
 void FIKRetargetEditorController::OnIKRigNeedsInitialized(UIKRigDefinition* ModifiedIKRig) const
 {
-	const UIKRetargeter* Retargeter = AssetController->GetAsset();
+	UIKRetargeter* Retargeter = AssetController->GetAsset();
 	
 	check(ModifiedIKRig && Retargeter)
 	 
@@ -68,10 +70,8 @@ void FIKRetargetEditorController::OnIKRigNeedsInitialized(UIKRigDefinition* Modi
 		return;
 	}
 
-	// the target anim instance has the IK RetargetPoseFromMesh node which needs reinitialized with new asset version
-	ClearOutputLog();
-	TargetAnimInstance->SetProcessorNeedsInitialized();
-	RefreshAllViews();
+	// the target anim instance has the RetargetPoseFromMesh node which needs reinitialized with new asset version
+	OnRetargeterNeedsInitialized(Retargeter);
 }
 
 void FIKRetargetEditorController::OnRetargetChainRenamed(UIKRigDefinition* ModifiedIKRig, FName OldName, FName NewName) const
@@ -88,24 +88,26 @@ void FIKRetargetEditorController::OnRetargetChainRemoved(UIKRigDefinition* Modif
 	RefreshAllViews();
 }
 
-void FIKRetargetEditorController::OnRetargeterNeedsInitialized(const UIKRetargeter* Retargeter) const
+void FIKRetargetEditorController::OnRetargeterNeedsInitialized(UIKRetargeter* Retargeter) const
 {
 	// clear the output log
 	ClearOutputLog();
-	// force reinit the runtime retarget processor
-	TargetAnimInstance->SetProcessorNeedsInitialized();
+	// force reinit the retarget processor (also inits the target IK Rig processor)
+	if (UIKRetargetProcessor* Processor = GetRetargetProcessor())
+	{
+		constexpr bool bSuppressWarnings = false;
+		Processor->Initialize(GetSourceSkeletalMesh(), GetTargetSkeletalMesh(), Retargeter, bSuppressWarnings);
+	}
 	// refresh all the UI views
 	RefreshAllViews();
 }
 
-void FIKRetargetEditorController::AddOffsetAndUpdatePreviewMeshPosition(
-	const FVector& Offset,
-	USceneComponent* Component) const
+void FIKRetargetEditorController::AddOffsetToMeshComponent(const FVector& Offset, USceneComponent* MeshComponent) const
 {
 	UIKRetargeter* Asset = AssetController->GetAsset();
 	FVector Position;
 	float Scale;
-	if (Component == TargetSkelMeshComponent)
+	if (MeshComponent == TargetSkelMeshComponent)
 	{
 		Asset->TargetMeshOffset += Offset;
 		Position = Asset->TargetMeshOffset;
@@ -121,8 +123,8 @@ void FIKRetargetEditorController::AddOffsetAndUpdatePreviewMeshPosition(
 	constexpr bool bSweep = false;
 	constexpr FHitResult* OutSweepHitResult = nullptr;
 	constexpr ETeleportType Teleport = ETeleportType::ResetPhysics;
-	Component->SetWorldLocation(Position, bSweep, OutSweepHitResult, Teleport);
-	Component->SetWorldScale3D(FVector(Scale,Scale,Scale));
+	MeshComponent->SetWorldLocation(Position, bSweep, OutSweepHitResult, Teleport);
+	MeshComponent->SetWorldScale3D(FVector(Scale,Scale,Scale));
 }
 
 USkeletalMesh* FIKRetargetEditorController::GetSourceSkeletalMesh() const
@@ -155,29 +157,25 @@ const USkeleton* FIKRetargetEditorController::GetTargetSkeleton() const
 	return nullptr;
 }
 
-const TArray<FName>& FIKRetargetEditorController::GetSelectedBones() const
+UDebugSkelMeshComponent* FIKRetargetEditorController::GetEditedSkeletalMesh() const
 {
-	if (const FIKRetargetEditPoseMode* EditPoseMode = Editor.Pin()->GetEditorModeManager().GetActiveModeTyped<FIKRetargetEditPoseMode>(FIKRetargetEditPoseMode::ModeName))
-	{
-		return EditPoseMode->GetSelectedBones();
-	}
-	
-	static TArray<FName> Empty;
-	return Empty;
+	return SkeletonMode == EIKRetargetSkeletonMode::Source ? SourceSkelMeshComponent : TargetSkelMeshComponent;
 }
 
-FTransform FIKRetargetEditorController::GetTargetBoneGlobalTransform(
-	const UIKRetargetProcessor* RetargetProcessor,
-	const int32& TargetBoneIndex) const
+FTransform FIKRetargetEditorController::GetGlobalRetargetPoseOfBone(
+	const FRetargetSkeleton& Skeleton,
+	const int32& BoneIndex,
+	const float& Scale,
+	const FVector& Offset) const
 {
-	check(RetargetProcessor && RetargetProcessor->IsInitialized())
-
+	check(Skeleton.BoneNames.IsValidIndex(BoneIndex))
+	
 	// get transform of bone
-	FTransform BoneTransform = RetargetProcessor->GetTargetBoneRetargetPoseGlobalTransform(TargetBoneIndex);
+	FTransform BoneTransform = Skeleton.RetargetGlobalPose[BoneIndex];
 
 	// scale and offset
-	BoneTransform.ScaleTranslation(AssetController->GetAsset()->TargetMeshScale);
-	BoneTransform.AddToTranslation(AssetController->GetAsset()->TargetMeshOffset);
+	BoneTransform.ScaleTranslation(Scale);
+	BoneTransform.AddToTranslation(Offset);
 	BoneTransform.NormalizeRotation();
 
 	return BoneTransform;
@@ -192,41 +190,34 @@ FTransform FIKRetargetEditorController::GetTargetBoneLocalTransform(
 	return RetargetProcessor->GetTargetBoneRetargetPoseLocalTransform(TargetBoneIndex);
 }
 
-bool FIKRetargetEditorController::GetTargetBoneLineSegments(
-	const UIKRetargetProcessor* RetargetProcessor,
-	const int32& TargetBoneIndex,
-	FVector& OutStart,
-	TArray<FVector>& OutChildren) const
+void FIKRetargetEditorController::GetGlobalRetargetPoseOfImmediateChildren(
+	const FRetargetSkeleton& RetargetSkeleton,
+	const int32& BoneIndex,
+	const float& Scale,
+	const FVector& Offset,
+	TArray<int32>& OutChildrenIndices,
+	TArray<FVector>& OutChildrenPositions)
 {
-	// get the runtime processor
-	check(RetargetProcessor && RetargetProcessor->IsInitialized())
+	OutChildrenIndices.Reset();
+	OutChildrenPositions.Reset();
 	
-	// get the target skeleton we want to draw
-	const FTargetSkeleton& TargetSkeleton = RetargetProcessor->GetTargetSkeleton();
-	check(TargetSkeleton.BoneNames.IsValidIndex(TargetBoneIndex))
+	check(RetargetSkeleton.BoneNames.IsValidIndex(BoneIndex))
 
-	// get the origin of the bone chain
-	OutStart = TargetSkeleton.RetargetGlobalPose[TargetBoneIndex].GetTranslation();
+	// get indices of immediate children
+	RetargetSkeleton.GetChildrenIndices(BoneIndex, OutChildrenIndices);
 
-	// get children
-	TArray<int32> ChildIndices;
-	TargetSkeleton.GetChildrenIndices(TargetBoneIndex, ChildIndices);
-	for (const int32& ChildIndex : ChildIndices)
+	// get the positions of the immediate children
+	for (const int32& ChildIndex : OutChildrenIndices)
 	{
-		OutChildren.Emplace(TargetSkeleton.RetargetGlobalPose[ChildIndex].GetTranslation());
+		OutChildrenPositions.Emplace(RetargetSkeleton.RetargetGlobalPose[ChildIndex].GetTranslation());
 	}
 
-	// add the target translation offset and scale
-	const UIKRetargeter* Asset = AssetController->GetAsset();
-	OutStart *= Asset->TargetMeshScale;
-	OutStart += Asset->TargetMeshOffset;
-	for (FVector& ChildPoint : OutChildren)
+	// apply scale and offset to positions
+	for (FVector& ChildPosition : OutChildrenPositions)
 	{
-		ChildPoint *= Asset->TargetMeshScale;
-		ChildPoint += Asset->TargetMeshOffset;
+		ChildPosition *= Scale;
+		ChildPosition += Offset;
 	}
-	
-	return true;
 }
 
 UIKRetargetProcessor* FIKRetargetEditorController::GetRetargetProcessor() const
@@ -271,6 +262,11 @@ void FIKRetargetEditorController::RefreshAllViews() const
 	{
 		AssetBrowserView.Get()->RefreshView();
 	}
+
+	if (HierarchyView.IsValid())
+	{
+		HierarchyView.Get()->RefreshTreeView();
+	}
 }
 
 void FIKRetargetEditorController::PlayAnimationAsset(UAnimationAsset* AssetToPlay)
@@ -285,6 +281,148 @@ void FIKRetargetEditorController::PlayAnimationAsset(UAnimationAsset* AssetToPla
 	}
 }
 
+void FIKRetargetEditorController::SetSkeletonMode(EIKRetargetSkeletonMode NewMode)
+{
+	// already in this mode, so do nothing
+	if (NewMode == SkeletonMode)
+	{
+		return;
+	}
+
+	// clear the selection on old skeleton
+	ClearSelection();
+	
+	// store the new skeleton mode
+	SkeletonMode = NewMode;
+
+	const ERetargeterOutputMode Mode = AssetController->GetAsset()->GetOutputMode();
+	switch (Mode)
+	{
+		case ERetargeterOutputMode::EditRetargetPose:
+		{
+			// TODO restart pose editing on other skeleton once we support source pose editing
+			//FEditorModeTools& EditorModeManager = Editor.Pin()->GetEditorModeManager();
+			//FIKRetargetEditPoseMode* EditMode = EditorModeManager.GetActiveModeTyped<FIKRetargetEditPoseMode>(FIKRetargetEditPoseMode::ModeName);
+			//EditMode->Enter();
+			SetRetargeterMode(ERetargeterOutputMode::ShowRetargetPose);
+			break;
+		}
+		case ERetargeterOutputMode::ShowRetargetPose:
+		case ERetargeterOutputMode::RunRetarget:
+		{
+			// toggle visibility of currently active skeleton
+			const bool bEditingSource = NewMode == EIKRetargetSkeletonMode::Source;
+			SourceSkelMeshComponent->SkeletonDrawMode = bEditingSource ? ESkeletonDrawMode::Default : ESkeletonDrawMode::GreyedOut;
+			TargetSkelMeshComponent->SkeletonDrawMode = bEditingSource ? ESkeletonDrawMode::GreyedOut : ESkeletonDrawMode::Default;
+			break;
+		}
+		default:
+			checkNoEntry();
+	}
+	
+	RefreshAllViews();
+}
+
+void FIKRetargetEditorController::SetSelectedMesh(UPrimitiveComponent* InMeshComponent)
+{
+	SelectedMesh = InMeshComponent;
+	SourceSkelMeshComponent->PushSelectionToProxy();
+	TargetSkelMeshComponent->PushSelectionToProxy();
+	SourceSkelMeshComponent->MarkRenderStateDirty();
+	TargetSkelMeshComponent->MarkRenderStateDirty();
+}
+
+UPrimitiveComponent* FIKRetargetEditorController::GetSelectedMesh()
+{
+	return SelectedMesh;
+}
+
+void FIKRetargetEditorController::EditBoneSelection(
+	const TArray<FName>& InBoneNames,
+	EBoneSelectionEdit EditMode,
+	const bool bFromHierarchyView)
+{
+	// must have a skeletal mesh
+	UDebugSkelMeshComponent* DebugComponent = GetEditedSkeletalMesh();
+	if (!DebugComponent->SkeletalMesh)
+	{
+		return;
+	}
+	
+	switch (EditMode)
+	{
+		case EBoneSelectionEdit::Add:
+		{
+			for (const FName& BoneName : InBoneNames)
+			{
+				SelectedBones.AddUnique(BoneName);
+			}
+			
+			break;
+		}
+		case EBoneSelectionEdit::Remove:
+		{
+			for (const FName& BoneName : InBoneNames)
+			{
+				SelectedBones.Remove(BoneName);
+			}
+			break;
+		}
+		case EBoneSelectionEdit::Replace:
+		{
+			SelectedBones = InBoneNames;
+			break;
+		}
+		default:
+			checkNoEntry();
+	}
+
+	// convert to bone indices
+	const FReferenceSkeleton& RefSkeleton = DebugComponent->SkeletalMesh->GetRefSkeleton();
+	TArray<int32> SelectedBoneIndices;
+	for (const FName& Bone : SelectedBones)
+	{
+		int32 BoneIndex = RefSkeleton.FindBoneIndex(Bone);
+		SelectedBoneIndices.Add(BoneIndex);
+		
+		if (BoneIndex == INDEX_NONE)
+		{
+			ensureMsgf(false, TEXT("Incoming selection list is not compatible with the currently edited skeleton."));
+			SelectedBoneIndices.Reset();
+			break;
+		}
+	}
+
+	// deselect mesh
+	SetSelectedMesh(nullptr);
+	
+	// apply selection to debug mesh component so rendering knows
+	DebugComponent->BonesOfInterest = SelectedBoneIndices;
+
+	// update views
+	if (!bFromHierarchyView && HierarchyView.IsValid())
+	{
+		HierarchyView->RefreshTreeView();
+	}
+}
+
+void FIKRetargetEditorController::ClearSelection()
+{
+	// clear mesh selection
+	SetSelectedMesh(nullptr);
+	
+	// clear bone selection
+	const TArray<FName> Empty;
+	constexpr bool bFromHierarchy = false;
+	EditBoneSelection(Empty, EBoneSelectionEdit::Replace, bFromHierarchy);
+	
+	// deselect all chains
+	ChainsView->ClearSelection();
+
+	// show global details
+	DetailsView->SetObject(AssetController->GetAsset());
+}
+
 void FIKRetargetEditorController::SetRetargeterMode(ERetargeterOutputMode Mode)
 {
 	const bool bWasEditingPose = IsEditingPose();
@@ -325,7 +463,7 @@ void FIKRetargetEditorController::SetRetargeterMode(ERetargeterOutputMode Mode)
 			// put source back in ref pose
 			SourceSkelMeshComponent->ShowReferencePose(true);
 			// have to move component back to offset position because ShowReferencePose() sets it back to origin
-			AddOffsetAndUpdatePreviewMeshPosition(FVector::ZeroVector, SourceSkelMeshComponent);
+			AddOffsetToMeshComponent(FVector::ZeroVector, SourceSkelMeshComponent);
 			break;
 
 		default:
@@ -360,7 +498,10 @@ bool FIKRetargetEditorController::CanEditPose() const
 		return false;
 	}
 
-	return Processor->IsInitialized();
+	const bool bRetargeterIsReady = Processor->IsInitialized();
+	// TODO in the future, add support for editing source pose
+	const bool bEditingTarget = GetSkeletonMode() == EIKRetargetSkeletonMode::Target;
+	return bRetargeterIsReady && bEditingTarget;
 }
 
 bool FIKRetargetEditorController::IsEditingPose() const
@@ -420,9 +561,9 @@ void FIKRetargetEditorController::HandleNewPose()
 					.ButtonStyle(FAppStyle::Get(), "Button")
 					.TextStyle( FAppStyle::Get(), "DialogButtonText" )
 					.Text(LOCTEXT("CancelButtonLabel", "Cancel"))
-					.OnClicked_Lambda( [Window=NewPoseWindow]()
+					.OnClicked_Lambda( [this]()
 					{
-						Window->RequestDestroyWindow();
+						NewPoseWindow->RequestDestroyWindow();
 						return FReply::Handled();
 					})
 				]
@@ -503,9 +644,9 @@ void FIKRetargetEditorController::HandleDuplicatePose()
 					.ButtonStyle(FAppStyle::Get(), "Button")
 					.TextStyle( FAppStyle::Get(), "DialogButtonText" )
 					.Text(LOCTEXT("CancelButtonLabel", "Cancel"))
-					.OnClicked_Lambda( [Window=NewPoseWindow]()
+					.OnClicked_Lambda( [this]()
 					{
-						Window->RequestDestroyWindow();
+						NewPoseWindow->RequestDestroyWindow();
 						return FReply::Handled();
 					})
 				]
@@ -586,9 +727,9 @@ void FIKRetargetEditorController::HandleImportPose()
 					.HAlign(HAlign_Center)
 					.VAlign(VAlign_Center)
 					.Text(LOCTEXT("CancelButtonLabel", "Cancel"))
-					.OnClicked_Lambda( [Window=ImportPoseWindow]()
+					.OnClicked_Lambda( [this]()
 					{
-						Window->RequestDestroyWindow();
+						ImportPoseWindow->RequestDestroyWindow();
 						return FReply::Handled();
 					})
 				]
@@ -781,9 +922,9 @@ void FIKRetargetEditorController::HandleImportPoseFromSequence()
 					.HAlign(HAlign_Center)
 					.VAlign(VAlign_Center)
 					.Text(LOCTEXT("CancelButtonLabel", "Cancel"))
-					.OnClicked_Lambda( [Window=ImportPoseFromSequenceWindow]()
+					.OnClicked_Lambda( [this]()
 					{
-						Window->RequestDestroyWindow();
+						ImportPoseFromSequenceWindow->RequestDestroyWindow();
 						return FReply::Handled();
 					})
 				]
@@ -859,7 +1000,7 @@ FReply FIKRetargetEditorController::OnImportPoseFromSequence()
 	const TArray<FTransform>& RefPose = RefSkeleton.GetRefBonePose();
 	const FTargetSkeleton& TargetSkeleton = Retargeter->GetTargetSkeleton();
 	int32 NumBones = RefSkeleton.GetNum();
-	const int32 RootBoneIndex = Retargeter->GetTargetSkeletonRootBone();
+	const int32 RootBoneIndex = Retargeter->GetTargetRetargetRoot();
 	for (int32 BoneIndex = 0; BoneIndex < NumBones; BoneIndex++)
 	{
 		const FName& BoneName = RefSkeleton.GetBoneName(BoneIndex);
@@ -985,7 +1126,6 @@ void FIKRetargetEditorController::HandleResetSelectedBones() const
 void FIKRetargetEditorController::HandleResetSelectedAndChildrenBones() const
 {
 	const FName CurrentPose = AssetController->GetCurrentRetargetPoseName();
-	const TArray<FName>& SelectedBones = GetSelectedBones();
 	
 	// get list of all children of selected bones
 	const UIKRetargetProcessor* Processor = GetRetargetProcessor();
@@ -1070,9 +1210,9 @@ void FIKRetargetEditorController::HandleRenamePose()
 					.HAlign(HAlign_Center)
 					.VAlign(VAlign_Center)
 					.Text(LOCTEXT("CancelButtonLabel", "Cancel"))
-					.OnClicked_Lambda( [Window=RenamePoseWindow]()
+					.OnClicked_Lambda( [this]()
 					{
-						Window->RequestDestroyWindow();
+						RenamePoseWindow->RequestDestroyWindow();
 						return FReply::Handled();
 					})
 				]
diff --git a/Engine/Plugins/Animation/IKRig/Source/IKRigEditor/Private/RetargetEditor/IKRetargetHierarchyTabSummoner.cpp b/Engine/Plugins/Animation/IKRig/Source/IKRigEditor/Private/RetargetEditor/IKRetargetHierarchyTabSummoner.cpp
new file mode 100644
--- /dev/null
+++ b/Engine/Plugins/Animation/IKRig/Source/IKRigEditor/Private/RetargetEditor/IKRetargetHierarchyTabSummoner.cpp
@@ -0,0 +1,38 @@
+﻿// Copyright Epic Games, Inc. All Rights Reserved.
+
+#include "RetargetEditor/IKRetargetHierarchyTabSummoner.h"
+
+#include "RetargetEditor/SIKRetargetHierarchy.h"
+#include "RetargetEditor/IKRetargetEditor.h"
+#include "RigEditor/IKRigEditorStyle.h"
+
+#include "IDocumentation.h"
+
+#define LOCTEXT_NAMESPACE "IKRetargetHierarchyTabSummoner"
+
+const FName FIKRetargetHierarchyTabSummoner::TabID(TEXT("IKRetargetHierarchy"));
+
+FIKRetargetHierarchyTabSummoner::FIKRetargetHierarchyTabSummoner(const TSharedRef<FIKRetargetEditor>& InEditor)
+	: FWorkflowTabFactory(TabID, InEditor)
+	, RetargetEditor(InEditor)
+{
+	bIsSingleton = true; // only allow a single instance of this tab
+	
+	TabLabel = LOCTEXT("IKRetargetHierarchy_TabLabel", "Hierarchy");
+	TabIcon = FSlateIcon(FIKRigEditorStyle::Get().GetStyleSetName(), "IKRig.TabIcon");
+
+	ViewMenuDescription = LOCTEXT("IKRetargetHierarchy_ViewMenu_Desc", "Hierarchy");
+	ViewMenuTooltip = LOCTEXT("IKRetargetHierarchy_ViewMenu_ToolTip", "Show the Retarget Hierarchy Tab");
+}
+
+TSharedPtr<SToolTip> FIKRetargetHierarchyTabSummoner::CreateTabToolTipWidget(const FWorkflowTabSpawnInfo& Info) const
+{
+	return  IDocumentation::Get()->CreateToolTip(LOCTEXT("IKRetargetHierarchyTooltip", "The IK Retargeting Hierarchy tab lets you view the skeleton hierarchy."), NULL, TEXT("Shared/Editors/Persona"), TEXT("IKRigSkeleton_Window"));
+}
+
+TSharedRef<SWidget> FIKRetargetHierarchyTabSummoner::CreateTabBody(const FWorkflowTabSpawnInfo& Info) const
+{
+	return SNew(SIKRetargetHierarchy, RetargetEditor.Pin()->GetController());
+}
+
+#undef LOCTEXT_NAMESPACE 
diff --git a/Engine/Plugins/Animation/IKRig/Source/IKRigEditor/Private/RetargetEditor/IKRetargeterController.cpp b/Engine/Plugins/Animation/IKRig/Source/IKRigEditor/Private/RetargetEditor/IKRetargeterController.cpp
--- a/Engine/Plugins/Animation/IKRig/Source/IKRigEditor/Private/RetargetEditor/IKRetargeterController.cpp
+++ b/Engine/Plugins/Animation/IKRig/Source/IKRigEditor/Private/RetargetEditor/IKRetargeterController.cpp
@@ -91,6 +91,16 @@ USkeletalMesh* UIKRetargeterController::GetTargetPreviewMesh() const
 	return Asset->GetTargetIKRig()->GetPreviewMesh();
 }
 
+const UIKRigDefinition* UIKRetargeterController::GetSourceIKRig() const
+{
+	return Asset->GetSourceIKRig();
+}
+
+const UIKRigDefinition* UIKRetargeterController::GetTargetIKRig() const
+{
+	return Asset->GetTargetIKRig();
+}
+
 FName UIKRetargeterController::GetSourceRootBone() const
 {
 	const UIKRigDefinition* SourceIKRig = Asset->GetSourceIKRig();
diff --git a/Engine/Plugins/Animation/IKRig/Source/IKRigEditor/Private/RetargetEditor/SIKRetargetHierarchy.cpp b/Engine/Plugins/Animation/IKRig/Source/IKRigEditor/Private/RetargetEditor/SIKRetargetHierarchy.cpp
new file mode 100644
--- /dev/null
+++ b/Engine/Plugins/Animation/IKRig/Source/IKRigEditor/Private/RetargetEditor/SIKRetargetHierarchy.cpp
@@ -0,0 +1,358 @@
+﻿// Copyright Epic Games, Inc. All Rights Reserved.
+
+#include "RetargetEditor/SIKRetargetHierarchy.h"
+
+#include "Preferences/PersonaOptions.h"
+#include "Widgets/Input/SSegmentedControl.h"
+#include "Retargeter/IKRetargetProcessor.h"
+#include "RigEditor/IKRigEditorStyle.h"
+#include "Widgets/Layout/SScrollBox.h"
+
+#define LOCTEXT_NAMESPACE "SIKRetargetHierarchy"
+
+FIKRetargetTreeElement::FIKRetargetTreeElement(
+	const FText& InKey,
+	const TSharedRef<FIKRetargetEditorController>& InEditorController)
+	: Key(InKey)
+	, EditorController(InEditorController)
+{}
+
+TSharedRef<ITableRow> FIKRetargetTreeElement::MakeTreeRowWidget(
+	TSharedRef<FIKRetargetEditorController> InEditorController,
+	const TSharedRef<STableViewBase>& InOwnerTable,
+	TSharedRef<FIKRetargetTreeElement> InTreeElement,
+	TSharedRef<FUICommandList> InCommandList,
+	TSharedPtr<SIKRetargetHierarchy> InHierarchy)
+{
+	return SNew(SIKRetargetHierarchyItem, InEditorController, InOwnerTable, InTreeElement, InCommandList, InHierarchy);
+}
+
+void SIKRetargetHierarchyItem::Construct(
+	const FArguments& InArgs,
+	TSharedRef<FIKRetargetEditorController> InEditorController,
+	const TSharedRef<STableViewBase>& OwnerTable,
+	TSharedRef<FIKRetargetTreeElement> InTreeElement,
+	TSharedRef<FUICommandList> InCommandList,
+	TSharedPtr<SIKRetargetHierarchy> InHierarchy)
+{
+	WeakTreeElement = InTreeElement;
+	EditorController = InEditorController;
+	HierarchyView = InHierarchy;
+
+	const FSlateBrush* Brush = FAppStyle::Get().GetBrush("SkeletonTree.Bone");
+	
+	const FTextBlockStyle NormalText = FIKRigEditorStyle::Get().GetWidgetStyle<FTextBlockStyle>("IKRig.Tree.NormalText");
+	const FTextBlockStyle ItalicText = FIKRigEditorStyle::Get().GetWidgetStyle<FTextBlockStyle>("IKRig.Tree.ItalicText");
+	FSlateFontInfo TextFont = NormalText.Font;
+	FSlateColor TextColor = FLinearColor(0.2f, 0.2f, 0.2f, 0.5f);
+	
+	TSharedPtr<SHorizontalBox> HorizontalBox;
+	STableRow<TSharedPtr<FIKRetargetTreeElement>>::Construct(
+		STableRow<TSharedPtr<FIKRetargetTreeElement>>::FArguments()
+		.ShowWires(true)
+		.Content()
+		[
+			SAssignNew(HorizontalBox, SHorizontalBox)
+			+ SHorizontalBox::Slot()
+			.MaxWidth(18)
+			.FillWidth(1.0)
+			.HAlign(HAlign_Left)
+			.VAlign(VAlign_Center)
+			[
+				SNew(SImage)
+				.Image(Brush)
+			]
+		], OwnerTable);
+
+	HorizontalBox->AddSlot()
+		.AutoWidth()
+		.VAlign(VAlign_Center)
+		[
+			SNew(STextBlock)
+			.Text(this, &SIKRetargetHierarchyItem::GetName)
+			.Font(TextFont)
+			.ColorAndOpacity(TextColor)
+		];
+}
+
+FText SIKRetargetHierarchyItem::GetName() const
+{
+	return WeakTreeElement.Pin()->Key;
+}
+
+void SIKRetargetHierarchy::Construct(
+	const FArguments& InArgs,
+	TSharedRef<FIKRetargetEditorController> InEditorController)
+{
+	EditorController = InEditorController;
+	EditorController.Pin()->HierarchyView = SharedThis(this);
+	CommandList = MakeShared<FUICommandList>();
+	
+	ChildSlot
+	[
+		SNew(SVerticalBox)
+		+SVerticalBox::Slot()
+		.Padding(2.0f)
+		.AutoHeight()
+		[
+			SNew(SBox)
+			.Padding(2.0f)
+			.HAlign(HAlign_Center)
+			[
+				SNew(SSegmentedControl<EIKRetargetSkeletonMode>)
+				.Value_Lambda([this]()
+				{
+					return EditorController.Pin()->GetSkeletonMode();
+				})
+				.OnValueChanged_Lambda([this](EIKRetargetSkeletonMode Mode)
+				{
+					EditorController.Pin()->SetSkeletonMode(Mode);
+				})
+				+SSegmentedControl<EIKRetargetSkeletonMode>::Slot(EIKRetargetSkeletonMode::Source)
+				.Text(LOCTEXT("SourceSkeleton", "Source"))
+				+ SSegmentedControl<EIKRetargetSkeletonMode>::Slot(EIKRetargetSkeletonMode::Target)
+				.Text(LOCTEXT("TargetSkeleton", "Target"))
+			]
+		]
+
+		+SVerticalBox::Slot()
+		.Padding(2.0f)
+		[
+			SNew(SBorder)
+			.Padding(2.0f)
+			.BorderImage(FAppStyle::GetBrush("SCSEditor.TreePanel"))
+			[
+				SAssignNew(TreeView, SIKRetargetHierarchyTreeView)
+				.TreeItemsSource(&RootElements)
+				.SelectionMode(ESelectionMode::Multi)
+				.OnGenerateRow(this, &SIKRetargetHierarchy::MakeTableRowWidget)
+				.OnGetChildren(this, &SIKRetargetHierarchy::HandleGetChildrenForTree)
+				.OnSelectionChanged(this, &SIKRetargetHierarchy::OnSelectionChanged)
+				.OnMouseButtonClick(this, &SIKRetargetHierarchy::OnItemClicked)
+				.OnMouseButtonDoubleClick(this, &SIKRetargetHierarchy::OnItemDoubleClicked)
+				.OnSetExpansionRecursive(this, &SIKRetargetHierarchy::OnSetExpansionRecursive)
+				.HighlightParentNodesForSelection(false)
+				.ItemHeight(24)
+			]
+		]
+	];
+
+	constexpr bool IsInitialSetup = true;
+	RefreshTreeView(IsInitialSetup);
+}
+
+void SIKRetargetHierarchy::ShowItemAfterSelection(FName ItemName)
+{
+	TSharedPtr<FIKRetargetTreeElement> ItemToShow;
+	for (const TSharedPtr<FIKRetargetTreeElement>& Element : AllElements)
+	{
+		if (Element.Get()->Name == ItemName)
+		{
+			ItemToShow = Element;
+		}
+	}
+
+	if (!ItemToShow.IsValid())
+	{
+		return;
+	}
+	
+	if(GetDefault<UPersonaOptions>()->bExpandTreeOnSelection)
+	{
+		TSharedPtr<FIKRetargetTreeElement> ItemToExpand = ItemToShow->Parent;
+		while(ItemToExpand.IsValid())
+		{
+			TreeView->SetItemExpansion(ItemToExpand, true);
+			ItemToExpand = ItemToExpand->Parent;
+		}
+	}
+    
+	TreeView->RequestScrollIntoView(ItemToShow);
+}
+
+void SIKRetargetHierarchy::RefreshTreeView(bool IsInitialSetup)
+{
+	const TSharedPtr<FIKRetargetEditorController> Controller = EditorController.Pin();
+	if (!Controller.IsValid())
+	{
+		return;
+	}
+
+	const TSharedRef<FIKRetargetEditorController> ControllerRef = Controller.ToSharedRef();
+
+	// synchronize selection with editor controller
+	const TArray<FName> SelectedBones = Controller->GetSelectedBones();
+	for (auto Element : AllElements)
+	{
+		const bool bIsSelected = SelectedBones.Contains(Element.Get()->Name);
+		TreeView->SetItemSelection(Element, bIsSelected, ESelectInfo::Direct);
+	}
+	
+	// save expansion and selection state
+	TreeView->SaveAndClearState();
+
+	// reset all tree items
+	RootElements.Reset();
+	AllElements.Reset();
+
+	// validate we have a skeleton to load
+	const UIKRetargetProcessor* Processor = Controller->GetRetargetProcessor();
+	if (!Processor && Processor->IsInitialized())
+	{
+		TreeView->RequestTreeRefresh();
+		return;
+	}
+
+	// get the skeleton that is currently being viewed in the editor
+	const FTargetSkeleton& TargetSkeleton = Processor->GetTargetSkeleton();
+	const FRetargetSkeleton& SourceSkeleton = Processor->GetSourceSkeleton();
+	const bool bViewTarget = Controller->GetSkeletonMode() == EIKRetargetSkeletonMode::Target;
+	const TArray<FName>& BoneNames = bViewTarget ? TargetSkeleton.BoneNames : SourceSkeleton.BoneNames;
+	const TArray<int32>& ParentIndices = bViewTarget ? TargetSkeleton.ParentIndices : SourceSkeleton.ParentIndices;
+	
+	// record bone element indices
+	TMap<FName, int32> BoneTreeElementIndices;
+
+	// create all bone elements
+	for (const FName BoneName : BoneNames)
+	{
+		// create "Bone" tree element for this bone
+		const FText BoneDisplayName = FText::FromName(BoneName);
+		TSharedPtr<FIKRetargetTreeElement> BoneElement = MakeShared<FIKRetargetTreeElement>(BoneDisplayName, ControllerRef);
+		BoneElement.Get()->Name = BoneName;
+		const int32 BoneElementIndex = AllElements.Add(BoneElement);
+		BoneTreeElementIndices.Add(BoneName, BoneElementIndex);
+
+		// TODO, this will show ALL bones in the skeletal mesh. represent bone differently if it's in the retargeted set
+		// TODO, add filtering of retargeted bones, by name, etc..
+	}
+
+	// store children/parent pointers on all bone elements
+	for (int32 BoneIndex=0; BoneIndex<BoneNames.Num(); ++BoneIndex)
+	{
+		const FName BoneName = BoneNames[BoneIndex];
+		const TSharedPtr<FIKRetargetTreeElement> BoneTreeElement = AllElements[BoneTreeElementIndices[BoneName]];
+		const int32 ParentIndex = ParentIndices[BoneIndex];
+		if (ParentIndex < 0)
+		{
+			// store the root element
+			RootElements.Add(BoneTreeElement);
+			// has no parent, so skip storing parent pointer
+			continue;
+		}
+
+		// get parent tree element
+		const FName ParentBoneName = BoneNames[ParentIndex];
+		const TSharedPtr<FIKRetargetTreeElement> ParentBoneTreeElement = AllElements[BoneTreeElementIndices[ParentBoneName]];
+		// store pointer to child on parent
+		ParentBoneTreeElement->Children.Add(BoneTreeElement);
+		// store pointer to parent on child
+		BoneTreeElement->Parent = ParentBoneTreeElement;
+	}
+
+	// expand all elements upon the initial construction of the tree
+	if (IsInitialSetup)
+	{
+		for (TSharedPtr<FIKRetargetTreeElement> RootElement : RootElements)
+		{
+			SetExpansionRecursive(RootElement, false, true);
+		}
+	}
+	else
+	{
+		// restore expansion and selection state
+		for (const TSharedPtr<FIKRetargetTreeElement>& Element : AllElements)
+		{
+			TreeView->RestoreState(Element);
+		}
+	}
+	
+	TreeView->RequestTreeRefresh();
+}
+
+TSharedRef<ITableRow> SIKRetargetHierarchy::MakeTableRowWidget(
+	TSharedPtr<FIKRetargetTreeElement> InItem,
+	const TSharedRef<STableViewBase>& OwnerTable)
+{
+	return InItem->MakeTreeRowWidget(EditorController.Pin().ToSharedRef(), OwnerTable, InItem.ToSharedRef(), CommandList.ToSharedRef(), SharedThis(this));
+}
+
+void SIKRetargetHierarchy::HandleGetChildrenForTree(
+	TSharedPtr<FIKRetargetTreeElement> InItem,
+	TArray<TSharedPtr<FIKRetargetTreeElement>>& OutChildren)
+{
+	OutChildren = InItem.Get()->Children;
+}
+
+void SIKRetargetHierarchy::OnSelectionChanged(
+	TSharedPtr<FIKRetargetTreeElement> Selection,
+	ESelectInfo::Type SelectInfo)
+{
+	if (SelectInfo == ESelectInfo::Direct)
+	{
+		return;
+	}
+	
+	const TSharedPtr<FIKRetargetEditorController> Controller = EditorController.Pin();
+	if (!Controller.IsValid())
+	{
+		return;
+	}
+
+	TArray<FName> SelectedBoneNames;
+	TArray<TSharedPtr<FIKRetargetTreeElement>> SelectedItems = TreeView->GetSelectedItems();
+	for (auto SelectedItem : SelectedItems)
+	{
+		SelectedBoneNames.Add(SelectedItem.Get()->Name);
+	}
+
+	constexpr bool bFromHierarchy = true;
+	Controller->EditBoneSelection(SelectedBoneNames, EBoneSelectionEdit::Replace, bFromHierarchy);
+}
+
+void SIKRetargetHierarchy::OnItemClicked(TSharedPtr<FIKRetargetTreeElement> InItem)
+{
+	// TODO show bone details
+}
+
+void SIKRetargetHierarchy::OnItemDoubleClicked(TSharedPtr<FIKRetargetTreeElement> InItem)
+{
+	if (TreeView->IsItemExpanded(InItem))
+	{
+		SetExpansionRecursive(InItem, false, false);
+	}
+	else
+	{
+		SetExpansionRecursive(InItem, false, true);
+	}
+}
+
+void SIKRetargetHierarchy::OnSetExpansionRecursive(TSharedPtr<FIKRetargetTreeElement> InItem, bool bShouldBeExpanded)
+{
+	SetExpansionRecursive(InItem, false, bShouldBeExpanded);
+}
+
+void SIKRetargetHierarchy::SetExpansionRecursive(
+	TSharedPtr<FIKRetargetTreeElement> InElement,
+	bool bTowardsParent,
+	bool bShouldBeExpanded)
+{
+	TreeView->SetItemExpansion(InElement, bShouldBeExpanded);
+    
+	if (bTowardsParent)
+	{
+		if (InElement->Parent.Get())
+		{
+			SetExpansionRecursive(InElement->Parent, bTowardsParent, bShouldBeExpanded);
+		}
+	}
+	else
+	{
+		for (int32 ChildIndex = 0; ChildIndex < InElement->Children.Num(); ++ChildIndex)
+		{
+			SetExpansionRecursive(InElement->Children[ChildIndex], bTowardsParent, bShouldBeExpanded);
+		}
+	}
+}
+
+#undef LOCTEXT_NAMESPACE
diff --git a/Engine/Plugins/Animation/IKRig/Source/IKRigEditor/Private/RigEditor/IKRigEditMode.cpp b/Engine/Plugins/Animation/IKRig/Source/IKRigEditor/Private/RigEditor/IKRigEditMode.cpp
--- a/Engine/Plugins/Animation/IKRig/Source/IKRigEditor/Private/RigEditor/IKRigEditMode.cpp
+++ b/Engine/Plugins/Animation/IKRig/Source/IKRigEditor/Private/RigEditor/IKRigEditMode.cpp
@@ -107,15 +107,16 @@ void FIKRigEditMode::RenderGoals(FPrimitiveDrawInterface* PDI)
 	{
 		return;
 	}
-	
+
 	TArray<UIKRigEffectorGoal*> Goals = AssetController->GetAllGoals();
 	for (const UIKRigEffectorGoal* Goal : Goals)
 	{
 		const bool bIsSelected = Controller->IsGoalSelected(Goal->GoalName);
 		const float Size = IKRigAsset->GoalSize * Goal->SizeMultiplier;
 		const float Thickness = IKRigAsset->GoalThickness * Goal->ThicknessMultiplier;
+		const FLinearColor Color = bIsSelected ? FLinearColor::Green : FLinearColor::Yellow;
 		PDI->SetHitProxy(new HIKRigEditorGoalProxy(Goal->GoalName));
-		IKRigDebugRendering::DrawGoal(PDI, Goal, bIsSelected, Size, Thickness);
+		IKRigDebugRendering::DrawWireCube(PDI, Goal->CurrentTransform, Color, Size, Thickness);
 		PDI->SetHitProxy(NULL);
 	}
 }
@@ -129,13 +130,6 @@ void FIKRigEditMode::RenderBones(FPrimitiveDrawInterface* PDI)
 		return; 
 	}
 
-	const UIKRigController* AssetController = Controller->AssetController;
-	const UIKRigDefinition* IKRigAsset = AssetController->GetAsset();
-	if (!IKRigAsset->DrawBones)
-	{
-		return;
-	}
-
 	// anim instance initialized?
 	if (!Controller->AnimInstance.IsValid())
 	{
@@ -153,30 +147,29 @@ void FIKRigEditMode::RenderBones(FPrimitiveDrawInterface* PDI)
 		return;
 	}
 
+	// determine size to draw bones from preference stored on the asset
+	const float BoneRadius = FMath::Max(0.01f, Controller->AssetController->GetAsset()->BoneSize);
+
 	// get affected / selected bones
 	TSet<int32> OutAffectedBones;
 	TSet<int32> OutSelectedBones;
 	GetAffectedBones(Controller.Get(), CurrentProcessor,OutAffectedBones,OutSelectedBones);
 
-	// draw affected bones
+	// draw bones
 	TArray<int32> ChildrenIndices;
 	const FIKRigSkeleton& Skeleton = CurrentProcessor->GetSkeleton();
 	const TArray<FTransform>& BoneTransforms = Skeleton.CurrentPoseGlobal;
-	const float MaxDrawRadius = Controller->SkelMeshComponent->Bounds.SphereRadius * 0.01f;
 	for (int32 BoneIndex=0; BoneIndex<Skeleton.BoneNames.Num(); ++BoneIndex)
 	{
 		// selected bones are drawn with different color
 		const bool bIsSelected = OutSelectedBones.Contains(BoneIndex);
 		const bool bIsAffected = OutAffectedBones.Contains(BoneIndex);
-		FLinearColor LineColor = bIsAffected ? IKRigDebugRendering::AFFECTED_BONE_COLOR : IKRigDebugRendering::DESELECTED_BONE_COLOR;
-		LineColor = bIsSelected ? IKRigDebugRendering::SELECTED_BONE_COLOR : LineColor;
+		FLinearColor LineColor = bIsAffected ? SkeletalDebugRendering::AFFECTED_BONE_COLOR : SkeletalDebugRendering::DEFAULT_BONE_COLOR;
+		LineColor = bIsSelected ? SkeletalDebugRendering::SELECTED_BONE_COLOR : LineColor;
 
 		// only draw axes on affected/selected bones
 		const bool bDrawAxes = bIsSelected || bIsAffected;
 		
-		const float BoneRadiusSetting = IKRigAsset->BoneSize;
-		const float BoneRadius = FMath::Min(1.0f, MaxDrawRadius) * BoneRadiusSetting;
-		
 		// draw line from bone to each child
 		FTransform BoneTransform = BoneTransforms[BoneIndex];
 		TArray<FVector> ChildPoints;
@@ -186,11 +179,24 @@ void FIKRigEditMode::RenderBones(FPrimitiveDrawInterface* PDI)
 			ChildPoints.Add(BoneTransforms[ChildIndex].GetLocation());
 		}
 
+		// draw cone to parent with different color if child is selected
+		TArray<FLinearColor> ChildColors;
+		for (int32 ChildIndex : ChildrenIndices)
+		{
+			FLinearColor ChildLineColor = LineColor;
+			if (!bIsSelected && OutSelectedBones.Contains(ChildIndex))
+			{
+				ChildLineColor = SkeletalDebugRendering::PARENT_OF_SELECTED_BONE_COLOR;
+			}
+			ChildColors.Add(ChildLineColor);
+		}
+
 		PDI->SetHitProxy(new HIKRigEditorBoneProxy(Skeleton.BoneNames[BoneIndex]));
-		IKRigDebugRendering::DrawWireBone(
+		SkeletalDebugRendering::DrawWireBoneAdvanced(
 			PDI,
 			BoneTransform,
 			ChildPoints,
+			ChildColors,
 			LineColor,
 			SDPG_Foreground,
 			BoneRadius,
diff --git a/Engine/Plugins/Animation/IKRig/Source/IKRigEditor/Private/RigEditor/IKRigEditorController.cpp b/Engine/Plugins/Animation/IKRig/Source/IKRigEditor/Private/RigEditor/IKRigEditorController.cpp
--- a/Engine/Plugins/Animation/IKRig/Source/IKRigEditor/Private/RigEditor/IKRigEditorController.cpp
+++ b/Engine/Plugins/Animation/IKRig/Source/IKRigEditor/Private/RigEditor/IKRigEditorController.cpp
@@ -786,7 +786,7 @@ TObjectPtr<UIKRigBoneDetails> FIKRigEditorController::CreateBoneDetails(const TS
 	}
 	
 	// create and store a new one
-	UIKRigBoneDetails* NewBoneDetails = NewObject<UIKRigBoneDetails>();
+	UIKRigBoneDetails* NewBoneDetails = NewObject<UIKRigBoneDetails>(AssetController->GetAsset(), FName(InBoneItem->BoneName), RF_Standalone | RF_Transient );
 	NewBoneDetails->SelectedBone = InBoneItem->BoneName;
 	NewBoneDetails->AnimInstancePtr = AnimInstance;
 	NewBoneDetails->AssetPtr = AssetController->GetAsset();
diff --git a/Engine/Plugins/Animation/IKRig/Source/IKRigEditor/Private/RigEditor/IKRigMode.cpp b/Engine/Plugins/Animation/IKRig/Source/IKRigEditor/Private/RigEditor/IKRigMode.cpp
--- a/Engine/Plugins/Animation/IKRig/Source/IKRigEditor/Private/RigEditor/IKRigMode.cpp
+++ b/Engine/Plugins/Animation/IKRig/Source/IKRigEditor/Private/RigEditor/IKRigMode.cpp
@@ -28,6 +28,7 @@ FIKRigMode::FIKRigMode(
 	ViewportArgs.bShowStats = false;
 	ViewportArgs.bShowTurnTable = false;
 	ViewportArgs.ContextName = TEXT("IKRigEditor.Viewport");
+	ViewportArgs.OnViewportCreated = FOnViewportCreated::CreateSP(IKRigEditor, &FIKRigEditorToolkit::HandleViewportCreated);
 
 	// register Persona tabs
 	FPersonaModule& PersonaModule = FModuleManager::LoadModuleChecked<FPersonaModule>("Persona");
diff --git a/Engine/Plugins/Animation/IKRig/Source/IKRigEditor/Private/RigEditor/IKRigToolkit.cpp b/Engine/Plugins/Animation/IKRig/Source/IKRigEditor/Private/RigEditor/IKRigToolkit.cpp
--- a/Engine/Plugins/Animation/IKRig/Source/IKRigEditor/Private/RigEditor/IKRigToolkit.cpp
+++ b/Engine/Plugins/Animation/IKRig/Source/IKRigEditor/Private/RigEditor/IKRigToolkit.cpp
@@ -2,6 +2,7 @@
 
 #include "RigEditor/IKRigToolkit.h"
 
+#include "Persona/Private/AnimationEditorViewportClient.h"
 #include "AnimationEditorPreviewActor.h"
 #include "EditorModeManager.h"
 #include "GameFramework/WorldSettings.h"
@@ -16,6 +17,7 @@
 #include "Framework/MultiBox/MultiBoxBuilder.h"
 
 #include "IKRigDefinition.h"
+#include "IPersonaViewport.h"
 #include "RigEditor/IKRigAnimInstance.h"
 #include "RigEditor/IKRigCommands.h"
 #include "RigEditor/IKRigEditMode.h"
@@ -192,6 +194,8 @@ void FIKRigEditorToolkit::HandlePreviewSceneCreated(const TSharedRef<IPersonaPre
 	
 	// create the preview skeletal mesh component
 	EditorController->SkelMeshComponent = NewObject<UDebugSkelMeshComponent>(Actor);
+	// turn off default bone rendering (we do our own in the IK Rig editor)
+	EditorController->SkelMeshComponent->SkeletonDrawMode = ESkeletonDrawMode::Hidden;
 
 	// setup an apply an anim instance to the skeletal mesh component
 	UIKRigAnimInstance* AnimInstance = NewObject<UIKRigAnimInstance>(EditorController->SkelMeshComponent, TEXT("IKRigAnimScriptInstance"));
@@ -220,4 +224,31 @@ void FIKRigEditorToolkit::HandleDetailsCreated(const TSharedRef<class IDetailsVi
 	EditorController->SetDetailsView(InDetailsView);
 }
 
+void FIKRigEditorToolkit::HandleViewportCreated(const TSharedRef<IPersonaViewport>& InViewport)
+{
+	// register callbacks to allow the asset to store the Bone Size viewport setting
+	FEditorViewportClient& ViewportClient = InViewport->GetViewportClient();
+	if (FAnimationViewportClient* AnimViewportClient = static_cast<FAnimationViewportClient*>(&ViewportClient))
+	{
+		AnimViewportClient->OnSetBoneSize.BindLambda([this](float InBoneSize)
+		{
+			if (UIKRigDefinition* Asset = EditorController->AssetController->GetAsset())
+			{
+				Asset->Modify();
+				Asset->BoneSize = InBoneSize;
+			}
+		});
+		
+		AnimViewportClient->OnGetBoneSize.BindLambda([this]() -> float
+		{
+			if (const UIKRigDefinition* Asset = EditorController->AssetController->GetAsset())
+			{
+				return Asset->BoneSize;
+			}
+
+			return 1.0f;
+		});
+	}
+}
+
 #undef LOCTEXT_NAMESPACE
diff --git a/Engine/Plugins/Animation/IKRig/Source/IKRigEditor/Public/IKRigDebugRendering.h b/Engine/Plugins/Animation/IKRig/Source/IKRigEditor/Public/IKRigDebugRendering.h
--- a/Engine/Plugins/Animation/IKRig/Source/IKRigEditor/Public/IKRigDebugRendering.h
+++ b/Engine/Plugins/Animation/IKRig/Source/IKRigEditor/Public/IKRigDebugRendering.h
@@ -7,10 +7,6 @@
 
 namespace IKRigDebugRendering
 {
-	static const FLinearColor DESELECTED_BONE_COLOR = FLinearColor(0.0f,0.0f,0.025f,1.0f);
-	static const FLinearColor SELECTED_BONE_COLOR = FLinearColor(0.2f,1.0f,0.2f,1.0f);
-	static const FLinearColor AFFECTED_BONE_COLOR = FLinearColor(1.0f,1.0f,1.0f,1.0f);
-	
 	static TArray BoxPoints =
 	{
 		FVector(0.5f, 0.5f, 0.5f),
@@ -41,69 +37,17 @@ namespace IKRigDebugRendering
 		FVector(-0.5f, 0.5f, -0.5f),
 	};
 
-	static void DrawGoal(
+	static void DrawWireCube(
 		FPrimitiveDrawInterface* PDI,
-		const UIKRigEffectorGoal* Goal,
-		bool bIsSelected,
+		const FTransform& Transform,
+		FLinearColor Color,
 		float Size,
 		float Thickness)
 	{
-		const FLinearColor Color = bIsSelected ? FLinearColor::Green : FLinearColor::Yellow;
 		const float Scale = FMath::Clamp(Size, 0.1f, 1000.0f);
-		const FTransform Transform = Goal->CurrentTransform;
 		for (int32 PointIndex = 0; PointIndex < BoxPoints.Num() - 1; PointIndex += 2)
 		{
 			PDI->DrawLine(Transform.TransformPosition(BoxPoints[PointIndex] * Scale), Transform.TransformPosition(BoxPoints[PointIndex + 1] * Scale), Color, SDPG_Foreground, Thickness);
 		}
 	}
-
-	// todo - refactor this and fix SkeletalDebugRendering
-	// this is copied/adapted from SkeletalDebugRendering because annoyingly the bone rendering puts the joint sphere at the
-	// END of the bone. This is so that bones can be rendered individually (needing only start/end point pairs).
-	// Which is easier, but doesn't make sense. We need the bones to render from parent to each immediate child.
-	// This causes several issues:
-	// 1. Bones lack a sphere at the root.
-	// 2. Selected bones render to their parent, rather than to their children.
-	// Also, this function draws the cone BETWEEN the joint spheres, not through them which is a cleaner end result.
-	// 
-	static void DrawWireBone(
-		FPrimitiveDrawInterface* PDI,
-		const FTransform& InBoneTransform,
-		const TArray<FVector>& InChildLocations,
-		const FLinearColor& InColor,
-		ESceneDepthPriorityGroup InDepthPriority,
-		const float SphereRadius,
-		const bool bDrawAxes)
-	{
-	#if ENABLE_DRAW_DEBUG
-		static const int32 NumSphereSides = 10;
-		static const int32 NumConeSides = 4;
-
-		const FVector BoneLocation = InBoneTransform.GetLocation();
-
-		// render sphere at joint origin
-		DrawWireSphere(PDI, BoneLocation, InColor, SphereRadius, NumSphereSides, InDepthPriority, 0.0f, 1.0f);
-		// draw axes at joint location
-		if (bDrawAxes)
-		{
-			SkeletalDebugRendering::DrawAxes(PDI, InBoneTransform, SDPG_Foreground, 0.f , SphereRadius);
-		}
-
-		// draw wire cones to each child
-		for (const FVector& ChildPoint : InChildLocations)
-		{
-			// offset start/end based on bone radius
-			const FVector RadiusOffset = (ChildPoint - BoneLocation).GetSafeNormal() * SphereRadius;
-			const FVector Start = BoneLocation + RadiusOffset;
-			const FVector End = ChildPoint - RadiusOffset;
-			
-			// calc cone size
-			const FVector EndToStart = (Start - End);
-			const float ConeLength = EndToStart.Size();
-			const float Angle = FMath::RadiansToDegrees(FMath::Atan(SphereRadius / ConeLength));
-			TArray<FVector> Verts;
-			DrawWireCone(PDI, Verts, FRotationMatrix::MakeFromX(EndToStart) * FTranslationMatrix(End), ConeLength, Angle, NumConeSides, InColor, InDepthPriority, 0.0f, 1.0f);
-		}
-	#endif
-	}
 }
diff --git a/Engine/Plugins/Animation/IKRig/Source/IKRigEditor/Public/RetargetEditor/IKRetargetAnimInstance.h b/Engine/Plugins/Animation/IKRig/Source/IKRigEditor/Public/RetargetEditor/IKRetargetAnimInstance.h
--- a/Engine/Plugins/Animation/IKRig/Source/IKRigEditor/Public/RetargetEditor/IKRetargetAnimInstance.h
+++ b/Engine/Plugins/Animation/IKRig/Source/IKRigEditor/Public/RetargetEditor/IKRetargetAnimInstance.h
@@ -22,8 +22,6 @@ public:
 
 	UIKRetargetProcessor* GetRetargetProcessor() const;
 
-	void SetProcessorNeedsInitialized();
-
 protected:
 	virtual FAnimInstanceProxy* CreateAnimInstanceProxy() override;
 
diff --git a/Engine/Plugins/Animation/IKRig/Source/IKRigEditor/Public/RetargetEditor/IKRetargetDefaultMode.h b/Engine/Plugins/Animation/IKRig/Source/IKRigEditor/Public/RetargetEditor/IKRetargetDefaultMode.h
--- a/Engine/Plugins/Animation/IKRig/Source/IKRigEditor/Public/RetargetEditor/IKRetargetDefaultMode.h
+++ b/Engine/Plugins/Animation/IKRig/Source/IKRigEditor/Public/RetargetEditor/IKRetargetDefaultMode.h
@@ -3,6 +3,7 @@
 #pragma once
 
 #include "CoreMinimal.h"
+#include "IKRetargetEditorController.h"
 
 #include "Retargeter/IKRetargeter.h"
 #include "IPersonaEditMode.h"
@@ -54,11 +55,13 @@ public:
 	/** END FEdMode interface */
 
 private:
-
-	TObjectPtr<UPrimitiveComponent> SelectedComponent;
+	
 	bool ComponentSelectionOverride(const UPrimitiveComponent* InComponent) const;
-	void SetSelectedComponent(UPrimitiveComponent* InComponent);
 	static void ApplyOffsetToMeshTransform(const FVector& Offset, USceneComponent* Component);
+
+	// the skeleton currently being edited
+	UDebugSkelMeshComponent* GetCurrentlyEditedMesh() const;
+	EIKRetargetSkeletonMode SkeletonMode;
 	
 	/** The hosting app */
 	TWeakPtr<FIKRetargetEditorController> EditorController;
@@ -68,3 +71,4 @@ private:
 
 	bool bIsInitialized = false;
 };
+
diff --git a/Engine/Plugins/Animation/IKRig/Source/IKRigEditor/Public/RetargetEditor/IKRetargetEditPoseMode.h b/Engine/Plugins/Animation/IKRig/Source/IKRigEditor/Public/RetargetEditor/IKRetargetEditPoseMode.h
--- a/Engine/Plugins/Animation/IKRig/Source/IKRigEditor/Public/RetargetEditor/IKRetargetEditPoseMode.h
+++ b/Engine/Plugins/Animation/IKRig/Source/IKRigEditor/Public/RetargetEditor/IKRetargetEditPoseMode.h
@@ -3,9 +3,11 @@
 #pragma once
 
 #include "CoreMinimal.h"
+#include "IKRetargetEditorController.h"
 
 #include "Retargeter/IKRetargeter.h"
 #include "IPersonaEditMode.h"
+#include "Retargeter/IKRetargetProcessor.h"
 
 class UIKRigProcessor;
 class FIKRetargetEditorController;
@@ -29,7 +31,11 @@ struct BoneEdit
 	FQuat AccumulatedGlobalOffset;		// the accumulated offset from rotation gizmo
 	
 	TArray<FQuat> PrevLocalOffsets;		// the prev stored local offsets of all selected bones
-	TArray<FName> SelectedBones;		// the currently selected bones in the viewport
+
+	void Reset()
+	{
+		*this = BoneEdit();
+	}
 };
 
 class FIKRetargetEditPoseMode : public IPersonaEditMode
@@ -67,15 +73,24 @@ public:
 	virtual void Enter() override;
 	virtual void Exit() override;
 	/** END FEdMode interface */
-
-	// get a list of the names of the currently selected bones (empty if non selected)
-	const TArray<FName>& GetSelectedBones() const;
 	
 private:
 
-	void GetAffectedBones(
-		FIKRetargetEditorController* Controller,
-		UIKRigProcessor* Processor,
+	// get the scale and offset associated with the currently edited skeletal mesh component
+	void GetEditedComponentScaleAndOffset(float& OutScale, FVector& OutOffset) const;
+
+	// get the scale and offset associated with the currently edited skeletal mesh component
+	int32 GetEditedRetargetRootBoneIndex(const UIKRetargetProcessor& Processor) const;
+
+	// render skeleton in viewport
+	void RenderSkeleton(
+		FPrimitiveDrawInterface* PDI,
+		const FIKRetargetEditorController* Controller,
+		const UIKRetargetProcessor* Processor);
+
+	void GetSelectedAndAffectedBones(
+		const FIKRetargetEditorController* Controller,
+		const FRetargetSkeleton& Skeleton,
 		TSet<int32>& OutAffectedBones,
 		TSet<int32>& OutSelectedBones) const;
 
@@ -83,11 +98,14 @@ private:
 
 	bool IsRootSelected() const;
 	bool IsOnlyRootSelected() const;
-	bool IsBoneSelected(const FName& BoneName) const;
+	void UpdateWidgetTransform();
 
+	// the bone(s) currently being edited
 	BoneEdit BoneEdit;
-	void UpdateWidgetTransform();
-	void HandleBoneSelectedInViewport(const FName& BoneName, bool bReplace);
+
+	// the skeleton currently being edited
+	const FRetargetSkeleton& GetCurrentlyEditedSkeleton(const UIKRetargetProcessor& Processor) const;
+	EIKRetargetSkeletonMode SkeletonMode;
 	
 	/** The hosting app */
 	TWeakPtr<FIKRetargetEditorController> EditorController;
diff --git a/Engine/Plugins/Animation/IKRig/Source/IKRigEditor/Public/RetargetEditor/IKRetargetEditor.h b/Engine/Plugins/Animation/IKRig/Source/IKRigEditor/Public/RetargetEditor/IKRetargetEditor.h
--- a/Engine/Plugins/Animation/IKRig/Source/IKRigEditor/Public/RetargetEditor/IKRetargetEditor.h
+++ b/Engine/Plugins/Animation/IKRig/Source/IKRigEditor/Public/RetargetEditor/IKRetargetEditor.h
@@ -88,6 +88,7 @@ private:
 	/** END toolbar */
 	
 	/** preview scene setup */
+	void HandleViewportCreated(const TSharedRef<class IPersonaViewport>& InViewport);
 	void HandlePreviewSceneCreated(const TSharedRef<IPersonaPreviewScene>& InPersonaPreviewScene);
 	void HandleDetailsCreated(const TSharedRef<class IDetailsView>& InDetailsView);
 	void OnFinishedChangingDetails(const FPropertyChangedEvent& PropertyChangedEvent);
diff --git a/Engine/Plugins/Animation/IKRig/Source/IKRigEditor/Public/RetargetEditor/IKRetargetEditorController.h b/Engine/Plugins/Animation/IKRig/Source/IKRigEditor/Public/RetargetEditor/IKRetargetEditorController.h
--- a/Engine/Plugins/Animation/IKRig/Source/IKRigEditor/Public/RetargetEditor/IKRetargetEditorController.h
+++ b/Engine/Plugins/Animation/IKRig/Source/IKRigEditor/Public/RetargetEditor/IKRetargetEditorController.h
@@ -7,9 +7,11 @@
 #include "UObject/ObjectMacros.h"
 #include "UObject/ObjectPtr.h"
 #include "Input/Reply.h"
+#include "Retargeter/IKRetargetProcessor.h"
 #include "Widgets/Input/SEditableTextBox.h"
 #include "Widgets/SWindow.h"
 
+class SIKRetargetHierarchy;
 class SIKRigOutputLog;
 class UIKRetargetProcessor;
 class SIKRetargetChainMapList;
@@ -18,6 +20,28 @@
 class UDebugSkelMeshComponent;
 class UIKRigDefinition;
 
+// which skeleton are we editing/viewing?
+enum class EIKRetargetSkeletonMode : uint8
+{
+	/** Editing / viewing the SOURCE skeleton (copy FROM) */
+	Source,
+
+	/** Editing / viewing the TARGET skeleton (copy TO) */
+	Target,
+};
+
+enum class EBoneSelectionEdit : uint8
+{
+	/** add to selection set */
+	Add,
+
+	/** remove from selection */
+	Remove,
+
+	/** replace selection entirely*/
+	Replace
+};
+
 /** a home for cross-widget communication to synchronize state across all tabs and viewport */
 class FIKRetargetEditorController : public TSharedFromThis<FIKRetargetEditorController>
 {
@@ -34,7 +58,7 @@
 	/** callback when IK Rig asset's retarget chain has been removed */
 	void OnRetargetChainRemoved(UIKRigDefinition* ModifiedIKRig, const FName& InChainRemoved) const;
 	/** callback when IK Retargeter asset requires reinitialization */
-	void OnRetargeterNeedsInitialized(const UIKRetargeter* Retargeter) const;
+	void OnRetargeterNeedsInitialized(UIKRetargeter* Retargeter) const;
 	
 	/** all modifications to the data model should go through this controller */
 	UIKRetargeterController* AssetController;
@@ -55,8 +79,6 @@
 	UPROPERTY(transient, NonTransactional)
 	TWeakObjectPtr<class UIKRetargetAnimInstance> TargetAnimInstance;
 
-	void AddOffsetAndUpdatePreviewMeshPosition(const FVector& Offset, USceneComponent* Component) const;
-
 	/** asset properties tab */
 	TSharedPtr<IDetailsView> DetailsView;
 	/** chain list view */
@@ -65,6 +87,8 @@
 	TSharedPtr<SIKRetargetAssetBrowser> AssetBrowserView;
 	/** output log view */
 	TSharedPtr<SIKRigOutputLog> OutputLogView;
+	/** hierarchy view */
+	TSharedPtr<SIKRetargetHierarchy> HierarchyView;
 	/** clear the output log */
 	void ClearOutputLog() const;
 	/** force refresh all views in the editor */
@@ -78,19 +102,26 @@
 	const USkeleton* GetSourceSkeleton() const;
 	/** get the target skeleton asset */
 	const USkeleton* GetTargetSkeleton() const;
+	/** get currently edited debug skeletal mesh */
+	UDebugSkelMeshComponent* GetEditedSkeletalMesh() const;
+	
+	/** get world space pose of a bone (with component scale / offset applied) */
+	FTransform GetGlobalRetargetPoseOfBone(
+		const FRetargetSkeleton& Skeleton,
+		const int32& BoneIndex,
+		const float& Scale,
+		const FVector& Offset) const;
 
-	/** get list of currently selected bones */
-	const TArray<FName>& GetSelectedBones() const;
-	
-	/** get current chain pose */
-	FTransform GetTargetBoneGlobalTransform(const UIKRetargetProcessor* RetargetProcessor, const int32& TargetBoneIndex) const;
 	FTransform GetTargetBoneLocalTransform(const UIKRetargetProcessor* RetargetProcessor, const int32& TargetBoneIndex) const;
-	/** get the line segments to draw from this bone to each child */
-	bool GetTargetBoneLineSegments(
-		const UIKRetargetProcessor* RetargetProcessor,
-		const int32& TargetBoneIndex,
-		FVector& OutStart,
-		TArray<FVector>& OutChildren) const;
+	
+	/** get world space positions of all immediate children of bone (with component scale / offset applied) */
+	static void GetGlobalRetargetPoseOfImmediateChildren(
+		const FRetargetSkeleton& RetargetSkeleton,
+		const int32& BoneIndex,
+		const float& Scale,
+		const FVector& Offset,
+		TArray<int32>& OutChildrenIndices,
+		TArray<FVector>& OutChildrenPositions);
 
 	/** get the retargeter that is running in the viewport (which is a duplicate of the source asset) */
 	UIKRetargetProcessor* GetRetargetProcessor() const;
@@ -106,7 +137,23 @@
 	UAnimationAsset* AnimThatWasPlaying = nullptr;
 	bool bWasPlayingAnim = false;
 	/** END viewport / editor tool mode */
+	
+	/** general editor mode can be either viewing/editing source or target */
+	EIKRetargetSkeletonMode GetSkeletonMode() const { return SkeletonMode; };
+	void SetSkeletonMode(EIKRetargetSkeletonMode Mode);
 
+	/** bone selection management (viewport or hierarchy view) */
+	void EditBoneSelection(
+		const TArray<FName>& InBoneNames,
+		EBoneSelectionEdit EditMode,
+		const bool bFromHierarchyView);
+	void ClearSelection();
+	const TArray<FName>& GetSelectedBones() const {return SelectedBones; };
+	/** mesh selection management (viewport view) */
+	void SetSelectedMesh(UPrimitiveComponent* InComponent);
+	UPrimitiveComponent* GetSelectedMesh();
+	void AddOffsetToMeshComponent(const FVector& Offset, USceneComponent* MeshComponent) const;
+
 	/* START RETARGET POSES */
 	
 	/** go to retarget pose */
@@ -172,4 +219,15 @@
 	TSharedPtr<SEditableTextBox> NewNameEditableText;
 	
 	/* END RETARGET POSES */
+
+private:
+	
+	/** which skeleton are we editing / viewing? */
+	EIKRetargetSkeletonMode SkeletonMode;
+
+	/** current selection set */
+	TArray<FName> SelectedBones;
+
+	/** currently selected mesh */
+	UPrimitiveComponent* SelectedMesh;
 };
diff --git a/Engine/Plugins/Animation/IKRig/Source/IKRigEditor/Public/RetargetEditor/IKRetargetHierarchyTabSummoner.h b/Engine/Plugins/Animation/IKRig/Source/IKRigEditor/Public/RetargetEditor/IKRetargetHierarchyTabSummoner.h
new file mode 100644
--- /dev/null
+++ b/Engine/Plugins/Animation/IKRig/Source/IKRigEditor/Public/RetargetEditor/IKRetargetHierarchyTabSummoner.h
@@ -0,0 +1,21 @@
+﻿// Copyright Epic Games, Inc. All Rights Reserved.
+
+#pragma once
+
+#include "CoreMinimal.h"
+#include "WorkflowOrientedApp/WorkflowTabFactory.h"
+
+class FIKRetargetEditor;
+
+struct FIKRetargetHierarchyTabSummoner : public FWorkflowTabFactory
+{
+public:
+	static const FName TabID;
+	
+	FIKRetargetHierarchyTabSummoner(const TSharedRef<FIKRetargetEditor>& InIKRigEditor);
+	virtual TSharedRef<SWidget> CreateTabBody(const FWorkflowTabSpawnInfo& Info) const override;
+	virtual TSharedPtr<SToolTip> CreateTabToolTipWidget(const FWorkflowTabSpawnInfo& Info) const override;
+
+protected:
+	TWeakPtr<FIKRetargetEditor> RetargetEditor;
+};
diff --git a/Engine/Plugins/Animation/IKRig/Source/IKRigEditor/Public/RetargetEditor/IKRetargeterController.h b/Engine/Plugins/Animation/IKRig/Source/IKRigEditor/Public/RetargetEditor/IKRetargeterController.h
--- a/Engine/Plugins/Animation/IKRig/Source/IKRigEditor/Public/RetargetEditor/IKRetargeterController.h
+++ b/Engine/Plugins/Animation/IKRig/Source/IKRigEditor/Public/RetargetEditor/IKRetargeterController.h
@@ -44,6 +44,10 @@ public:
 	USkeletalMesh* GetSourcePreviewMesh() const;
 	/** Get target skeletal mesh */
 	USkeletalMesh* GetTargetPreviewMesh() const;
+	/** Get source IK Rig asset */
+	const UIKRigDefinition* GetSourceIKRig() const;
+	/** Get source IK Rig asset */
+	const UIKRigDefinition* GetTargetIKRig() const;
 
 	/** Get name of the Root bone used for retargeting the Source skeleton. */
 	FName GetSourceRootBone() const;
@@ -106,7 +110,7 @@ public:
 private:
 	
 	/** Called whenever the retargeter is modified in such a way that would require re-initialization by dependent systems.*/
-	DECLARE_MULTICAST_DELEGATE_OneParam(FOnRetargeterNeedsInitialized, const UIKRetargeter*);
+	DECLARE_MULTICAST_DELEGATE_OneParam(FOnRetargeterNeedsInitialized, UIKRetargeter*);
 	FOnRetargeterNeedsInitialized RetargeterNeedsInitialized;
 	
 public:
diff --git a/Engine/Plugins/Animation/IKRig/Source/IKRigEditor/Public/RetargetEditor/SIKRetargetHierarchy.h b/Engine/Plugins/Animation/IKRig/Source/IKRigEditor/Public/RetargetEditor/SIKRetargetHierarchy.h
new file mode 100644
--- /dev/null
+++ b/Engine/Plugins/Animation/IKRig/Source/IKRigEditor/Public/RetargetEditor/SIKRetargetHierarchy.h
@@ -0,0 +1,96 @@
+﻿// Copyright Epic Games, Inc. All Rights Reserved.
+
+#pragma once
+
+#include "CoreMinimal.h"
+#include "EditorUndoClient.h"
+#include "IKRetargetEditorController.h"
+#include "SBaseHierarchyTreeView.h"
+#include "DragAndDrop/DecoratedDragDropOp.h"
+#include "Widgets/SCompoundWidget.h"
+
+class SIKRetargetHierarchy;
+
+class FIKRetargetTreeElement : public TSharedFromThis<FIKRetargetTreeElement>
+{
+public:
+	
+	FIKRetargetTreeElement(const FText& InKey, const TSharedRef<FIKRetargetEditorController>& InEditorController);
+
+	TSharedRef<ITableRow> MakeTreeRowWidget(
+		TSharedRef<FIKRetargetEditorController> InEditorController,
+		const TSharedRef<STableViewBase>& InOwnerTable,
+		TSharedRef<FIKRetargetTreeElement> InRigTreeElement,
+		TSharedRef<FUICommandList> InCommandList,
+		TSharedPtr<SIKRetargetHierarchy> InHierarchy);
+
+	FText Key;
+	TSharedPtr<FIKRetargetTreeElement> Parent;
+	TArray<TSharedPtr<FIKRetargetTreeElement>> Children;
+	FName Name;
+
+private:
+	
+	TWeakPtr<FIKRetargetEditorController> EditorController;
+};
+
+class SIKRetargetHierarchyItem : public STableRow<TSharedPtr<FIKRetargetTreeElement>>
+{
+public:
+	
+	void Construct(
+		const FArguments& InArgs,
+		TSharedRef<FIKRetargetEditorController> InEditorController,
+		const TSharedRef<STableViewBase>& OwnerTable,
+		TSharedRef<FIKRetargetTreeElement> InTreeElement,
+		TSharedRef<FUICommandList> InCommandList,
+		TSharedPtr<SIKRetargetHierarchy> InHierarchy);
+
+private:
+	
+	FText GetName() const;
+
+	TWeakPtr<FIKRetargetTreeElement> WeakTreeElement;
+	TWeakPtr<FIKRetargetEditorController> EditorController;
+	TWeakPtr<SIKRetargetHierarchy> HierarchyView;
+};
+
+typedef SBaseHierarchyTreeView<FIKRetargetTreeElement> SIKRetargetHierarchyTreeView;
+
+class SIKRetargetHierarchy : public SCompoundWidget, public FEditorUndoClient
+{
+public:
+	SLATE_BEGIN_ARGS(SIKRetargetHierarchy) {}
+	SLATE_END_ARGS()
+
+	void Construct(const FArguments& InArgs, TSharedRef<FIKRetargetEditorController> InEditorController);
+
+	void ShowItemAfterSelection(FName ItemKey);
+
+private:
+	
+	/** centralized editor controls (facilitate cross-communication between multiple UI elements)*/
+	TWeakPtr<FIKRetargetEditorController> EditorController;
+	
+	/** command list we bind to */
+	TSharedPtr<FUICommandList> CommandList;
+	
+	/** tree view widget */
+	TSharedPtr<SIKRetargetHierarchyTreeView> TreeView;
+	TArray<TSharedPtr<FIKRetargetTreeElement>> RootElements;
+	TArray<TSharedPtr<FIKRetargetTreeElement>> AllElements;
+	
+	/** tree view callbacks */
+	void RefreshTreeView(bool IsInitialSetup=false);
+	TSharedRef<ITableRow> MakeTableRowWidget(TSharedPtr<FIKRetargetTreeElement> InItem, const TSharedRef<STableViewBase>& OwnerTable);
+	void HandleGetChildrenForTree(TSharedPtr<FIKRetargetTreeElement> InItem, TArray<TSharedPtr<FIKRetargetTreeElement>>& OutChildren);
+	void OnSelectionChanged(TSharedPtr<FIKRetargetTreeElement> Selection, ESelectInfo::Type SelectInfo);
+	void OnItemClicked(TSharedPtr<FIKRetargetTreeElement> InItem);
+	void OnItemDoubleClicked(TSharedPtr<FIKRetargetTreeElement> InItem);
+	void OnSetExpansionRecursive(TSharedPtr<FIKRetargetTreeElement> InItem, bool bShouldBeExpanded);
+	void SetExpansionRecursive(TSharedPtr<FIKRetargetTreeElement> InElement, bool bTowardsParent, bool bShouldBeExpanded);
+	/** END tree view callbacks */
+
+	friend SIKRetargetHierarchyItem;
+	friend FIKRetargetEditorController;
+};
diff --git a/Engine/Plugins/Animation/IKRig/Source/IKRigEditor/Public/RigEditor/IKRigToolkit.h b/Engine/Plugins/Animation/IKRig/Source/IKRigEditor/Public/RigEditor/IKRigToolkit.h
--- a/Engine/Plugins/Animation/IKRig/Source/IKRigEditor/Public/RigEditor/IKRigToolkit.h
+++ b/Engine/Plugins/Animation/IKRig/Source/IKRigEditor/Public/RigEditor/IKRigToolkit.h
@@ -66,9 +66,10 @@ public:
 
 	TSharedRef<FIKRigEditorController> GetController() const {return EditorController;};
 
-	/** preview scene setup */
+	/** tab creation callbacks */
 	void HandlePreviewSceneCreated(const TSharedRef<IPersonaPreviewScene>& InPersonaPreviewScene);
 	void HandleDetailsCreated(const TSharedRef<class IDetailsView>& InDetailsView) const;
+	void HandleViewportCreated(const TSharedRef<class IPersonaViewport>& InViewport);
 	/** END preview scene setup */
 
 private:
diff --git a/Engine/Plugins/Animation/IKRig/Source/IKRigEditor/Public/RigEditor/SIKRigSkeleton.h b/Engine/Plugins/Animation/IKRig/Source/IKRigEditor/Public/RigEditor/SIKRigSkeleton.h
--- a/Engine/Plugins/Animation/IKRig/Source/IKRigEditor/Public/RigEditor/SIKRigSkeleton.h
+++ b/Engine/Plugins/Animation/IKRig/Source/IKRigEditor/Public/RigEditor/SIKRigSkeleton.h
@@ -4,6 +4,7 @@
 
 #include "CoreMinimal.h"
 #include "EditorUndoClient.h"
+#include "SBaseHierarchyTreeView.h"
 #include "SIKRigRetargetChainList.h"
 #include "DragAndDrop/DecoratedDragDropOp.h"
 #include "Widgets/SCompoundWidget.h"
@@ -100,90 +101,7 @@ public:
 	TWeakPtr<FIKRigTreeElement> Element;
 };
 
-class SIKRigSkeletonTreeView : public STreeView<TSharedPtr<FIKRigTreeElement>>
-{
-public:
-
-	virtual ~SIKRigSkeletonTreeView() {}
-
-	virtual FReply OnFocusReceived(const FGeometry& MyGeometry, const FFocusEvent& InFocusEvent) override
-	{
-		FReply Reply = STreeView<TSharedPtr<FIKRigTreeElement>>::OnFocusReceived(MyGeometry, InFocusEvent);
-		LastClickCycles = FPlatformTime::Cycles();
-		return Reply;
-	}
-
-	/** Save a snapshot of items expansion and selection state */
-	void SaveAndClearState()
-	{
-		SaveAndClearSparseItemInfos();
-		SaveAndClearSelection();
-	}
-	
-	/** Restore items expansion and selection state from the saved snapshot after tree reconstruction */
-	void RestoreState(const TSharedPtr<FIKRigTreeElement>& ItemPtr)
-	{
-		RestoreSparseItemInfos(ItemPtr);
-		RestoreSelection(ItemPtr);
-	}
-
-	/** slow double-click rename state*/
-	uint32 LastClickCycles = 0;
-	TWeakPtr<FIKRigTreeElement> LastSelected;
-
-private:
-
-	/** Save a snapshot of the internal map that tracks item expansion before tree reconstruction */
-	void SaveAndClearSparseItemInfos()
-	{
-		OldSparseItemInfos = SparseItemInfos;
-		ClearExpandedItems();
-	}
-
-	/** Restore the expansion infos map from the saved snapshot after tree reconstruction */
-	void RestoreSparseItemInfos(const TSharedPtr<FIKRigTreeElement>& ItemPtr)
-	{
-		for (const TTuple<TSharedPtr<FIKRigTreeElement>, FSparseItemInfo>& Pair : OldSparseItemInfos)
-		{
-			if (Pair.Key->Key.EqualTo(ItemPtr->Key))
-			{
-				// the SparseItemInfos now reference the new element, but keep the same expansion state
-				SparseItemInfos.Add(ItemPtr, Pair.Value);
-				return;
-			}
-		}
-
-		// set default state as expanded if not found
-		SparseItemInfos.Add(ItemPtr, FSparseItemInfo(true, false));
-	}
-	
-	/** Save a snapshot of the internal set that tracks item selection before tree reconstruction */
-	void SaveAndClearSelection()
-	{
-		OldSelectedItems = SelectedItems;
-		ClearSelection();
-	}
-	
-	/** Restore the selection from the saved snapshot after tree reconstruction */
-	void RestoreSelection(const TSharedPtr<FIKRigTreeElement>& ItemPtr)
-	{
-		for (const TSharedPtr<FIKRigTreeElement>& OldItem : OldSelectedItems)
-		{
-			if (OldItem->Key.EqualTo(ItemPtr->Key))
-			{
-				// select the new element
-				SetItemSelection(ItemPtr, true, ESelectInfo::Direct);
-				return;
-			}
-		}
-	}
-	
-	/** A temporary snapshot of the SparseItemInfos in STreeView, used during SIKRigSkeleton::RefreshTreeView() */
-	TSparseItemMap OldSparseItemInfos;
-
-	/** A temporary snapshot of the SelectedItems in SListView, used during SIKRigSkeleton::RefreshTreeView() */
-	TItemSet OldSelectedItems;
-};
+typedef SBaseHierarchyTreeView<FIKRigTreeElement> SIKRigSkeletonTreeView;
 
 class SIKRigSkeleton : public SCompoundWidget, public FEditorUndoClient
 {
diff --git a/Engine/Plugins/Animation/IKRig/Source/IKRigEditor/Public/SBaseHierarchyTreeView.h b/Engine/Plugins/Animation/IKRig/Source/IKRigEditor/Public/SBaseHierarchyTreeView.h
new file mode 100644
--- /dev/null
+++ b/Engine/Plugins/Animation/IKRig/Source/IKRigEditor/Public/SBaseHierarchyTreeView.h
@@ -0,0 +1,93 @@
+﻿// Copyright Epic Games, Inc. All Rights Reserved.
+
+#pragma once
+
+template<class ItemType>
+class SBaseHierarchyTreeView : public STreeView<TSharedPtr<ItemType>>
+{
+	typedef STreeView<TSharedPtr<ItemType>> BaseView;
+	using TSparseItemMap = typename BaseView::TSparseItemMap;
+	using TItemSet = typename BaseView::TItemSet;
+	
+public:
+
+	virtual ~SBaseHierarchyTreeView() {}
+
+	virtual FReply OnFocusReceived(const FGeometry& MyGeometry, const FFocusEvent& InFocusEvent) override
+	{
+		FReply Reply = STreeView<TSharedPtr<ItemType>>::OnFocusReceived(MyGeometry, InFocusEvent);
+		LastClickCycles = FPlatformTime::Cycles();
+		return Reply;
+	}
+
+	/** Save a snapshot of items expansion and selection state */
+	void SaveAndClearState()
+	{
+		SaveAndClearSparseItemInfos();
+		SaveAndClearSelection();
+	}
+	
+	/** Restore items expansion and selection state from the saved snapshot after tree reconstruction */
+	void RestoreState(const TSharedPtr<ItemType>& ItemPtr)
+	{
+		this->RestoreSparseItemInfos(ItemPtr);
+		this->RestoreSelection(ItemPtr);
+	}
+
+	/** slow double-click rename state*/
+	uint32 LastClickCycles = 0;
+	TWeakPtr<ItemType> LastSelected;
+
+private:
+
+	/** Save a snapshot of the internal map that tracks item expansion before tree reconstruction */
+	void SaveAndClearSparseItemInfos()
+	{
+		OldSparseItemInfos = BaseView::SparseItemInfos;
+		BaseView::ClearExpandedItems();
+	}
+
+	/** Restore the expansion infos map from the saved snapshot after tree reconstruction */
+	void RestoreSparseItemInfos(const TSharedPtr<ItemType>& ItemPtr)
+	{
+		for (const TTuple<TSharedPtr<ItemType>, FSparseItemInfo>& Pair : OldSparseItemInfos)
+		{
+			if (Pair.Key->Key.EqualTo(ItemPtr->Key))
+			{
+				// the SparseItemInfos now reference the new element, but keep the same expansion state
+				BaseView::SparseItemInfos.Add(ItemPtr, Pair.Value);
+				return;
+			}
+		}
+
+		// set default state as expanded if not found
+		BaseView::SparseItemInfos.Add(ItemPtr, FSparseItemInfo(true, false));
+	}
+	
+	/** Save a snapshot of the internal set that tracks item selection before tree reconstruction */
+	void SaveAndClearSelection()
+	{
+		OldSelectedItems = BaseView::SelectedItems;
+		BaseView::ClearSelection();
+	}
+	
+	/** Restore the selection from the saved snapshot after tree reconstruction */
+	void RestoreSelection(const TSharedPtr<ItemType>& ItemPtr)
+	{
+		for (const TSharedPtr<ItemType>& OldItem : OldSelectedItems)
+		{
+			if (OldItem->Key.EqualTo(ItemPtr->Key))
+			{
+				// select the new element
+				BaseView::SetItemSelection(ItemPtr, true, ESelectInfo::Direct);
+				return;
+			}
+		}
+	}
+	
+	/** A temporary snapshot of the SparseItemInfos in STreeView, used during SIKRetargetHierarchy::RefreshTreeView() */
+	TSparseItemMap OldSparseItemInfos;
+
+	/** A temporary snapshot of the SelectedItems in SListView, used during SIKRetargetHierarchy::RefreshTreeView() */
+	TItemSet OldSelectedItems;
+};
\ No newline at end of file
diff --git a/Engine/Source/Editor/Persona/Private/AnimationEditorPreviewScene.cpp b/Engine/Source/Editor/Persona/Private/AnimationEditorPreviewScene.cpp
--- a/Engine/Source/Editor/Persona/Private/AnimationEditorPreviewScene.cpp
+++ b/Engine/Source/Editor/Persona/Private/AnimationEditorPreviewScene.cpp
@@ -126,6 +126,23 @@ FAnimationEditorPreviewScene::~FAnimationEditorPreviewScene()
 	}
 }
 
+TArray<UDebugSkelMeshComponent*> FAnimationEditorPreviewScene::GetAllPreviewMeshComponents() const
+{
+	TArray<UDebugSkelMeshComponent*> PreviewMeshComponents;
+	GetActor()->GetComponents(PreviewMeshComponents, true);
+	return PreviewMeshComponents;
+}
+
+void FAnimationEditorPreviewScene::ForEachPreviewMesh(TFunction<void (UDebugSkelMeshComponent*)> PerMeshFunction)
+{
+	TArray<UDebugSkelMeshComponent*> PreviewMeshComponents;
+	GetActor()->GetComponents(PreviewMeshComponents, true);
+	for (UDebugSkelMeshComponent* PreviewMesh : PreviewMeshComponents)
+	{
+		PerMeshFunction(PreviewMesh);
+	}
+}
+
 void FAnimationEditorPreviewScene::SetPreviewMeshComponent(UDebugSkelMeshComponent* InSkeletalMeshComponent) 
 {
 	SkeletalMeshComponent = InSkeletalMeshComponent; 
diff --git a/Engine/Source/Editor/Persona/Private/AnimationEditorPreviewScene.h b/Engine/Source/Editor/Persona/Private/AnimationEditorPreviewScene.h
--- a/Engine/Source/Editor/Persona/Private/AnimationEditorPreviewScene.h
+++ b/Engine/Source/Editor/Persona/Private/AnimationEditorPreviewScene.h
@@ -32,6 +32,8 @@ public:
 	virtual void InvalidateViews() override;
 	virtual void FocusViews() override;
 	virtual UDebugSkelMeshComponent* GetPreviewMeshComponent() const override { return SkeletalMeshComponent; }
+	virtual TArray<UDebugSkelMeshComponent*> GetAllPreviewMeshComponents() const override;
+	virtual FORCEINLINE_DEBUGGABLE void ForEachPreviewMesh(TFunction<void (UDebugSkelMeshComponent*)> PerMeshFunction) override;
 	virtual void SetPreviewMeshComponent(UDebugSkelMeshComponent* InSkeletalMeshComponent) override;
 	virtual void SetAdditionalMeshes(class UDataAsset* InAdditionalMeshes) override;
 	virtual void SetAdditionalMeshesSelectable(bool bSelectable) override;
@@ -232,7 +234,6 @@ public:
 		return CameraOverride;
 	}
 
-
 	virtual void HandleSkeletonTreeSelectionChanged(const TArrayView<TSharedPtr<ISkeletonTreeItem>>& InSelectedItems, ESelectInfo::Type InSelectInfo) override;
 	/** FPreviewScene interface */
 	virtual void Tick(float InDeltaTime) override;
diff --git a/Engine/Source/Editor/Persona/Private/AnimationEditorViewportClient.cpp b/Engine/Source/Editor/Persona/Private/AnimationEditorViewportClient.cpp
--- a/Engine/Source/Editor/Persona/Private/AnimationEditorViewportClient.cpp
+++ b/Engine/Source/Editor/Persona/Private/AnimationEditorViewportClient.cpp
@@ -450,9 +450,15 @@ void FAnimationViewportClient::Draw(const FSceneView* View, FPrimitiveDrawInterf
 {
 	FEditorViewportClient::Draw(View, PDI);
 
-	UDebugSkelMeshComponent* PreviewMeshComponent = GetAnimPreviewScene()->GetPreviewMeshComponent();
-	if (PreviewMeshComponent && PreviewMeshComponent->SkeletalMesh)
+	// draw bones for all debug skeletal meshes
+	TArray<UDebugSkelMeshComponent*> PreviewMeshComponents = GetPreviewScene()->GetAllPreviewMeshComponents();
+	for (UDebugSkelMeshComponent* PreviewMeshComponent : PreviewMeshComponents)
 	{
+		if (!(PreviewMeshComponent && PreviewMeshComponent->SkeletalMesh))
+		{
+			continue;
+		}
+
 		// Can't have both bones of interest and sockets of interest set
 		check( !(GetAnimPreviewScene()->GetSelectedBoneIndex() != INDEX_NONE && GetAnimPreviewScene()->GetSelectedSocket().IsValid() ) )
 
@@ -529,24 +535,29 @@ void FAnimationViewportClient::DrawCanvas( FViewport& InViewport, FSceneView& Vi
 {
 	FEditorViewportClient::DrawCanvas(InViewport, View, Canvas);
 
-	UDebugSkelMeshComponent* PreviewMeshComponent = GetAnimPreviewScene()->GetPreviewMeshComponent();
-	if (PreviewMeshComponent != nullptr)
+	TArray<UDebugSkelMeshComponent*> PreviewMeshComponents = GetPreviewScene()->GetAllPreviewMeshComponents();
+	for  (UDebugSkelMeshComponent* PreviewMeshComponent : PreviewMeshComponents)
 	{
+		if (!PreviewMeshComponent)
+		{
+			continue;
+		}
+		
 		// Display bone names
 		if (PreviewMeshComponent->bShowBoneNames)
 		{
-			ShowBoneNames(&Canvas, &View);
+			ShowBoneNames(&Canvas, &View, PreviewMeshComponent);
 		}
 
 		// Display attribute names
 		if (PreviewMeshComponent->bDrawAttributes)
 		{
-			ShowAttributeNames(&Canvas, &View);
+			ShowAttributeNames(&Canvas, &View, PreviewMeshComponent);
 		}
 
 		if (bDrawUVs)
 		{
-			DrawUVsForMesh(Viewport, &Canvas, 1.0f);
+			DrawUVsForMesh(Viewport, &Canvas, 1.0f, PreviewMeshComponent);
 		}
 
 		// Debug draw clothing texts
@@ -554,10 +565,13 @@ void FAnimationViewportClient::DrawCanvas( FViewport& InViewport, FSceneView& Vi
 	}
 }
 
-void FAnimationViewportClient::DrawUVsForMesh(FViewport* InViewport, FCanvas* InCanvas, int32 InTextYPos)
+void FAnimationViewportClient::DrawUVsForMesh(FViewport* InViewport, FCanvas* InCanvas, int32 InTextYPos, UDebugSkelMeshComponent* PreviewMeshComponent)
 {
-	UDebugSkelMeshComponent* PreviewMeshComponent = GetAnimPreviewScene()->GetPreviewMeshComponent();
-
+	if (!PreviewMeshComponent)
+	{
+		return;
+	}
+	
 	//use the overridden LOD level
 	const uint32 LODLevel = FMath::Clamp(PreviewMeshComponent->GetForcedLOD() - 1, 0, PreviewMeshComponent->SkeletalMesh->GetLODNum() - 1);
 
@@ -655,10 +669,9 @@ void FAnimationViewportClient::SetCameraTargetLocation(const FSphere &BoundSpher
 	SetViewLocation( NewViewLocation );
 }
 
-void FAnimationViewportClient::ShowBoneNames( FCanvas* Canvas, FSceneView* View )
+void FAnimationViewportClient::ShowBoneNames( FCanvas* Canvas, FSceneView* View, UDebugSkelMeshComponent* PreviewMeshComponent )
 {
-	UDebugSkelMeshComponent* PreviewMeshComponent = GetAnimPreviewScene()->GetPreviewMeshComponent();
-	if (PreviewMeshComponent == nullptr || PreviewMeshComponent->MeshObject == nullptr)
+	if (!(PreviewMeshComponent && PreviewMeshComponent->MeshObject))
 	{
 		return;
 	}
@@ -738,43 +751,43 @@ void FAnimationViewportClient::ShowBoneNames( FCanvas* Canvas, FSceneView* View
 	}
 }
 
-void FAnimationViewportClient::ShowAttributeNames(FCanvas* Canvas, FSceneView* View)
+void FAnimationViewportClient::ShowAttributeNames(FCanvas* Canvas, FSceneView* View, UDebugSkelMeshComponent* MeshComponent) const
 {
-	UDebugSkelMeshComponent* MeshComponent = GetAnimPreviewScene()->GetPreviewMeshComponent();
-
-	if (MeshComponent && MeshComponent->SkeletalMesh)
+	if (!MeshComponent || !MeshComponent->SkeletalMesh)
 	{
-		const int32 HalfX = Viewport->GetSizeXY().X / 2 / GetDPIScale();
-		const int32 HalfY = Viewport->GetSizeXY().Y / 2 / GetDPIScale();
+		return;
+	}
+	
+	const int32 HalfX = Viewport->GetSizeXY().X / 2 / GetDPIScale();
+	const int32 HalfY = Viewport->GetSizeXY().Y / 2 / GetDPIScale();
 
-		const UE::Anim::FMeshAttributeContainer& Attributes = MeshComponent->GetCustomAttributes();
+	const UE::Anim::FMeshAttributeContainer& Attributes = MeshComponent->GetCustomAttributes();
 
-		const int32 TransformAnimationAttributeTypeIndex = Attributes.FindTypeIndex(FTransformAnimationAttribute::StaticStruct());
-		if (TransformAnimationAttributeTypeIndex != INDEX_NONE)
-		{
-			const TArray<UE::Anim::FAttributeId, FDefaultAllocator>& AttributeIdentifiers = Attributes.GetKeys(TransformAnimationAttributeTypeIndex);
-			const TArray<UE::Anim::TWrappedAttribute<FDefaultAllocator>>& AttributeValues = Attributes.GetValues(TransformAnimationAttributeTypeIndex);
-			check(AttributeIdentifiers.Num() == AttributeValues.Num());
+	const int32 TransformAnimationAttributeTypeIndex = Attributes.FindTypeIndex(FTransformAnimationAttribute::StaticStruct());
+	if (TransformAnimationAttributeTypeIndex != INDEX_NONE)
+	{
+		const TArray<UE::Anim::FAttributeId, FDefaultAllocator>& AttributeIdentifiers = Attributes.GetKeys(TransformAnimationAttributeTypeIndex);
+		const TArray<UE::Anim::TWrappedAttribute<FDefaultAllocator>>& AttributeValues = Attributes.GetValues(TransformAnimationAttributeTypeIndex);
+		check(AttributeIdentifiers.Num() == AttributeValues.Num());
 
-			for (int32 AttributeIndex = 0; AttributeIndex < AttributeValues.Num(); ++AttributeIndex)
+		for (int32 AttributeIndex = 0; AttributeIndex < AttributeValues.Num(); ++AttributeIndex)
+		{
+			if (const FTransformAnimationAttribute* AttributeValue = AttributeValues[AttributeIndex].GetPtr<FTransformAnimationAttribute>())
 			{
-				if (const FTransformAnimationAttribute* AttributeValue = AttributeValues[AttributeIndex].GetPtr<FTransformAnimationAttribute>())
-				{
-					const UE::Anim::FAttributeId& AttributeIdentifier = AttributeIdentifiers[AttributeIndex];
+				const UE::Anim::FAttributeId& AttributeIdentifier = AttributeIdentifiers[AttributeIndex];
 
-					const FTransform AttributeParentTransform = MeshComponent->GetDrawTransform(AttributeIdentifier.GetIndex()) * MeshComponent->GetComponentTransform();
-					const FTransform AttributeTransform = AttributeValue->Value * AttributeParentTransform;
+				const FTransform AttributeParentTransform = MeshComponent->GetDrawTransform(AttributeIdentifier.GetIndex()) * MeshComponent->GetComponentTransform();
+				const FTransform AttributeTransform = AttributeValue->Value * AttributeParentTransform;
 
-					const FPlane proj = View->Project(AttributeTransform.GetLocation());
-					if (proj.W > 0.f)
-					{
-						const int32 XPos = HalfX + (HalfX * proj.X);
-						const int32 YPos = HalfY + (HalfY * (proj.Y * -1));
+				const FPlane proj = View->Project(AttributeTransform.GetLocation());
+				if (proj.W > 0.f)
+				{
+					const int32 XPos = HalfX + (HalfX * proj.X);
+					const int32 YPos = HalfY + (HalfY * (proj.Y * -1));
 
-						FCanvasTextItem TextItem(FVector2D(XPos, YPos), FText::FromName(AttributeIdentifier.GetName()), GEngine->GetSmallFont(), FLinearColor(0.0f, 1.0f, 1.0f));
-						TextItem.EnableShadow(FLinearColor::Black);
-						Canvas->DrawItem(TextItem);
-					}
+					FCanvasTextItem TextItem(FVector2D(XPos, YPos), FText::FromName(AttributeIdentifier.GetName()), GEngine->GetSmallFont(), FLinearColor(0.0f, 1.0f, 1.0f));
+					TextItem.EnableShadow(FLinearColor::Black);
+					Canvas->DrawItem(TextItem);
 				}
 			}
 		}
@@ -1287,6 +1300,30 @@ ELocalAxesMode::Type FAnimationViewportClient::GetLocalAxesMode() const
 	return (ELocalAxesMode::Type)ConfigOption->DefaultLocalAxesSelection;
 }
 
+void FAnimationViewportClient::SetBoneDrawSize(const float InBoneDrawSize)
+{
+	BoneDrawSize = FMath::Max(0.01f,InBoneDrawSize);
+
+	// optionally inform editors that may want to maintain bone size between sessions
+	if (OnSetBoneSize.IsBound())
+	{
+		OnSetBoneSize.Execute(BoneDrawSize);
+	}
+	
+	RedrawRequested(nullptr);
+}
+
+float FAnimationViewportClient::GetBoneDrawSize() const
+{
+	// optionally get bone size from editors that may be storing this between sessions
+	if (OnGetBoneSize.IsBound())
+	{
+		return OnGetBoneSize.Execute();
+	}
+	
+	return BoneDrawSize;
+}
+
 void FAnimationViewportClient::SetBoneDrawMode(EBoneDrawMode::Type AxesMode)
 {
 	ConfigOption->SetDefaultBoneDrawSelection(AxesMode);
@@ -1302,65 +1339,111 @@ EBoneDrawMode::Type FAnimationViewportClient::GetBoneDrawMode() const
 	return (EBoneDrawMode::Type)ConfigOption->DefaultBoneDrawSelection;
 }
 
-void FAnimationViewportClient::DrawBonesFromTransforms(TArray<FTransform>& Transforms, UDebugSkelMeshComponent * MeshComponent, FPrimitiveDrawInterface* PDI, FLinearColor BoneColour, FLinearColor RootBoneColour) const
+void FAnimationViewportClient::DrawBonesFromTransforms(
+	TArray<FTransform>& Transforms,
+	UDebugSkelMeshComponent * MeshComponent,
+	FPrimitiveDrawInterface* PDI,
+	FLinearColor BoneColour,
+	FLinearColor RootBoneColour) const
 {
-	if ( Transforms.Num() > 0 && MeshComponent->SkeletalMesh )
+	if (Transforms.IsEmpty() ||
+		!MeshComponent->SkeletalMesh ||
+		MeshComponent->SkeletonDrawMode == ESkeletonDrawMode::Hidden)
 	{
-		TArray<FTransform> WorldTransforms;
-		WorldTransforms.AddUninitialized(Transforms.Num());
+		return;
+	}
 
-		TArray<FLinearColor> BoneColours;
-		BoneColours.AddUninitialized(Transforms.Num());
+	if (MeshComponent->SkeletonDrawMode == ESkeletonDrawMode::GreyedOut)
+	{
+		BoneColour = SkeletalDebugRendering::DISABLED_BONE_COLOR;
+	}
+	
+	TArray<FTransform> WorldTransforms;
+	WorldTransforms.AddUninitialized(Transforms.Num());
 
-		// we could cache parent bones as we calculate, but right now I'm not worried about perf issue of this
-		const TArray<FBoneIndexType>& DrawBoneIndices = MeshComponent->GetDrawBoneIndices();
-		for ( int32 Index=0; Index < MeshComponent->GetDrawBoneIndices().Num(); ++Index )
-		{
-			const int32 BoneIndex = DrawBoneIndices[Index];
-			const int32 ParentIndex = MeshComponent->GetReferenceSkeleton().GetParentIndex(BoneIndex);
+	TArray<FLinearColor> BoneColours;
+	BoneColours.AddUninitialized(Transforms.Num());
 
-			WorldTransforms[BoneIndex] = Transforms[BoneIndex] * MeshComponent->GetComponentTransform();
-			BoneColours[BoneIndex] = (ParentIndex >= 0) ? BoneColour : RootBoneColour;
-		}
+	// we could cache parent bones as we calculate, but right now I'm not worried about perf issue of this
+	const TArray<FBoneIndexType>& DrawBoneIndices = MeshComponent->GetDrawBoneIndices();
+	for ( int32 Index=0; Index < MeshComponent->GetDrawBoneIndices().Num(); ++Index )
+	{
+		const int32 BoneIndex = DrawBoneIndices[Index];
+		const int32 ParentIndex = MeshComponent->GetReferenceSkeleton().GetParentIndex(BoneIndex);
 
-		DrawBones(DrawBoneIndices, MeshComponent->GetReferenceSkeleton(), WorldTransforms, MeshComponent->BonesOfInterest, PDI, BoneColours, MeshComponent->Bounds.SphereRadius);
+		WorldTransforms[BoneIndex] = Transforms[BoneIndex] * MeshComponent->GetComponentTransform();
+		BoneColours[BoneIndex] = (ParentIndex >= 0) ? BoneColour : RootBoneColour;
 	}
+
+	constexpr bool bForceDraw = false;
+	const bool bAddHitProxy = MeshComponent->SkeletonDrawMode != ESkeletonDrawMode::GreyedOut;
+
+	DrawBones(
+		MeshComponent->GetComponentLocation(),
+		DrawBoneIndices,
+		MeshComponent->GetReferenceSkeleton(),
+		WorldTransforms,
+		MeshComponent->BonesOfInterest,
+		BoneColours,
+		PDI,
+		bForceDraw,
+		bAddHitProxy);
 }
 
-void FAnimationViewportClient::DrawBonesFromCompactPose(const FCompactHeapPose& Pose, UDebugSkelMeshComponent * MeshComponent, FPrimitiveDrawInterface* PDI, const FLinearColor& DrawColour) const
+void FAnimationViewportClient::DrawBonesFromCompactPose(
+	const FCompactHeapPose& Pose,
+	UDebugSkelMeshComponent* MeshComponent,
+	FPrimitiveDrawInterface* PDI,
+	const FLinearColor& DrawColour) const
 {
-	if (Pose.GetNumBones() > 0 && MeshComponent != nullptr)
+	if (Pose.GetNumBones() == 0 ||
+		!MeshComponent ||
+		!MeshComponent->SkeletalMesh ||
+		MeshComponent->SkeletonDrawMode != ESkeletonDrawMode::Hidden)
 	{
-		TArray<FTransform> WorldTransforms;
-		WorldTransforms.AddUninitialized(Pose.GetBoneContainer().GetNumBones());
+		return;
+	}
 
-		TArray<FLinearColor> BoneColours;
-		BoneColours.AddUninitialized(Pose.GetBoneContainer().GetNumBones());
+	// optionally override draw color
+	const FLinearColor BoneColor = MeshComponent->SkeletonDrawMode == ESkeletonDrawMode::GreyedOut ? SkeletalDebugRendering::DISABLED_BONE_COLOR : DrawColour;
+	
+	TArray<FTransform> WorldTransforms;
+	WorldTransforms.AddUninitialized(Pose.GetBoneContainer().GetNumBones());
 
-		// we could cache parent bones as we calculate, but right now I'm not worried about perf issue of this
-		for (FCompactPoseBoneIndex BoneIndex : Pose.ForEachBoneIndex())
-		{
-			FMeshPoseBoneIndex MeshBoneIndex = Pose.GetBoneContainer().MakeMeshPoseIndex(BoneIndex);
+	TArray<FLinearColor> BoneColours;
+	BoneColours.AddUninitialized(Pose.GetBoneContainer().GetNumBones());
 
-			int32 ParentIndex = Pose.GetBoneContainer().GetParentBoneIndex(MeshBoneIndex.GetInt());
+	// we could cache parent bones as we calculate, but right now I'm not worried about perf issue of this
+	for (FCompactPoseBoneIndex BoneIndex : Pose.ForEachBoneIndex())
+	{
+		FMeshPoseBoneIndex MeshBoneIndex = Pose.GetBoneContainer().MakeMeshPoseIndex(BoneIndex);
 
-			if (ParentIndex == INDEX_NONE)
-			{
-				WorldTransforms[MeshBoneIndex.GetInt()] = Pose[BoneIndex] * MeshComponent->GetComponentTransform();
-			}
-			else
-			{
-				WorldTransforms[MeshBoneIndex.GetInt()] = Pose[BoneIndex] * WorldTransforms[ParentIndex];
-			}
-			BoneColours[MeshBoneIndex.GetInt()] = DrawColour;
-		}
+		int32 ParentIndex = Pose.GetBoneContainer().GetParentBoneIndex(MeshBoneIndex.GetInt());
 
-		if (MeshComponent && MeshComponent->SkeletalMesh)
+		if (ParentIndex == INDEX_NONE)
 		{
-			DrawBones(MeshComponent->GetDrawBoneIndices(), MeshComponent->GetReferenceSkeleton(), WorldTransforms, MeshComponent->BonesOfInterest, PDI, BoneColours, MeshComponent->Bounds.SphereRadius, 1.0f, true);
+			WorldTransforms[MeshBoneIndex.GetInt()] = Pose[BoneIndex] * MeshComponent->GetComponentTransform();
 		}
-
+		else
+		{
+			WorldTransforms[MeshBoneIndex.GetInt()] = Pose[BoneIndex] * WorldTransforms[ParentIndex];
+		}
+		BoneColours[MeshBoneIndex.GetInt()] = BoneColor;
 	}
+
+	constexpr bool bForceDraw = true;
+	const bool bAddHitProxy = MeshComponent->SkeletonDrawMode != ESkeletonDrawMode::GreyedOut;
+
+	DrawBones(
+		MeshComponent->GetComponentLocation(),
+		MeshComponent->GetDrawBoneIndices(),
+		MeshComponent->GetReferenceSkeleton(),
+		WorldTransforms,
+		MeshComponent->BonesOfInterest,
+		BoneColours,
+		PDI,
+		bForceDraw,
+		bAddHitProxy);
 }
 
 void FAnimationViewportClient::DrawMeshBonesUncompressedAnimation(UDebugSkelMeshComponent * MeshComponent, FPrimitiveDrawInterface* PDI) const
@@ -1423,203 +1506,274 @@ void FAnimationViewportClient::DrawMeshBonesBakedAnimation(UDebugSkelMeshCompone
 	}
 }
 
-void FAnimationViewportClient::DrawMeshBones(UDebugSkelMeshComponent * MeshComponent, FPrimitiveDrawInterface* PDI) const
+void FAnimationViewportClient::DrawMeshBones(UDebugSkelMeshComponent* MeshComponent, FPrimitiveDrawInterface* PDI) const
 {
-	if ( MeshComponent && MeshComponent->SkeletalMesh && MeshComponent->GetNumDrawTransform() > 0)
+	if (!MeshComponent ||
+		!MeshComponent->SkeletalMesh ||
+		MeshComponent->GetNumDrawTransform() == 0 ||
+		MeshComponent->SkeletonDrawMode == ESkeletonDrawMode::Hidden)
 	{
-		TArray<FTransform> WorldTransforms;
-		WorldTransforms.AddUninitialized(MeshComponent->GetNumDrawTransform());
-
-		TArray<FLinearColor> BoneColours;
-		BoneColours.AddUninitialized(MeshComponent->GetNumDrawTransform());
-
-		TArray<int32> SelectedBones = MeshComponent->BonesOfInterest;
-
-		// we could cache parent bones as we calculate, but right now I'm not worried about perf issue of this
-		const TArray<FBoneIndexType>& DrawBoneIndices = MeshComponent->GetDrawBoneIndices();
-		for ( int32 Index=0; Index<DrawBoneIndices.Num(); ++Index )
-		{
-			const int32 BoneIndex = DrawBoneIndices[Index];
-			const int32 ParentIndex = MeshComponent->GetReferenceSkeleton().GetParentIndex(BoneIndex);
-
-			WorldTransforms[BoneIndex] = MeshComponent->GetDrawTransform(BoneIndex) * MeshComponent->GetComponentTransform();
-			BoneColours[BoneIndex] = (ParentIndex >= 0) ? FLinearColor::White : FLinearColor::Red;
-		}
+		return;
+	}
+	
+	TArray<FTransform> WorldTransforms;
+	WorldTransforms.AddUninitialized(MeshComponent->GetNumDrawTransform());
 
-		// Color virtual bones
-		for (int16 VirtualBoneIndex : MeshComponent->GetReferenceSkeleton().GetRequiredVirtualBones())
-		{
-			BoneColours[VirtualBoneIndex] = FLinearColor(0.4f, 0.4f, 1.0f, 1.0f);
-		}
+	TArray<FLinearColor> BoneColours;
+	BoneColours.AddUninitialized(MeshComponent->GetNumDrawTransform());
 
-		// Color selected bones
-		for (int32 SelectedBoneIndex : SelectedBones)
+	// factor skeleton draw mode into color selection
+	const FLinearColor BoneColor = MeshComponent->SkeletonDrawMode == ESkeletonDrawMode::GreyedOut ? SkeletalDebugRendering::DISABLED_BONE_COLOR : SkeletalDebugRendering::DEFAULT_BONE_COLOR;
+	const FLinearColor VirtualBoneColor = MeshComponent->SkeletonDrawMode == ESkeletonDrawMode::GreyedOut ? SkeletalDebugRendering::DISABLED_BONE_COLOR : SkeletalDebugRendering::VIRTUAL_BONE_COLOR;
+	
+	// we could cache parent bones as we calculate, but right now I'm not worried about perf issue of this
+	const TArray<FBoneIndexType>& DrawBoneIndices = MeshComponent->GetDrawBoneIndices();
+	for ( int32 Index=0; Index<DrawBoneIndices.Num(); ++Index )
+	{
+		const int32 BoneIndex = DrawBoneIndices[Index];
+		WorldTransforms[BoneIndex] = MeshComponent->GetDrawTransform(BoneIndex) * MeshComponent->GetComponentTransform();
+		BoneColours[BoneIndex] = BoneColor;
+	}
+
+	// color virtual bones
+	for (int16 VirtualBoneIndex : MeshComponent->GetReferenceSkeleton().GetRequiredVirtualBones())
+	{
+		BoneColours[VirtualBoneIndex] = VirtualBoneColor;
+	}
+
+	constexpr bool bForceDraw = false;
+
+	// don't allow selection if the skeleton draw mode is greyed out
+	const bool bAddHitProxy = MeshComponent->SkeletonDrawMode != ESkeletonDrawMode::GreyedOut;
+
+	DrawBones(
+		MeshComponent->GetComponentLocation(),
+		DrawBoneIndices,
+		MeshComponent->GetReferenceSkeleton(),
+		WorldTransforms,
+		MeshComponent->BonesOfInterest,
+		BoneColours,
+		PDI,
+		bForceDraw,
+		bAddHitProxy);
+}
+
+void FAnimationViewportClient::DrawBones(
+	const FVector& ComponentOrigin,
+	const TArray<FBoneIndexType>& RequiredBones,
+	const FReferenceSkeleton& RefSkeleton,
+	const TArray<FTransform>& WorldTransforms,
+	const TArray<int32>& InSelectedBones,
+	const TArray<FLinearColor>& BoneColors,
+	FPrimitiveDrawInterface* PDI,
+	bool bForceDraw,
+	bool bAddHitProxy) const
+{
+	// first determine which bones to draw, and which to filter out
+	TBitArray<> BonesToDraw(false, RefSkeleton.GetNum());
+	const bool bDrawSelected = GetBoneDrawMode() == EBoneDrawMode::Selected;
+	const bool bDrawSelectedAndParents = GetBoneDrawMode() == EBoneDrawMode::SelectedAndParents;
+	const bool bDrawSelectedAndChildren = GetBoneDrawMode() == EBoneDrawMode::SelectedAndChildren;
+	const bool bDrawSelectedAndParentsAndChildren = GetBoneDrawMode() == EBoneDrawMode::SelectedAndParentsAndChildren;
+	
+	// add selected bones
+	if (bDrawSelected || bDrawSelectedAndParents || bDrawSelectedAndChildren || bDrawSelectedAndParentsAndChildren)
+	{
+		for (int32 BoneIndex : InSelectedBones)
 		{
-			BoneColours[SelectedBoneIndex] = FLinearColor(1.0f, 0.34f, 0.0f, 1.0f);
+			if (BoneIndex != INDEX_NONE)
+			{
+				BonesToDraw[BoneIndex] = true;
+			}
 		}
-
-		DrawBones(DrawBoneIndices, MeshComponent->GetReferenceSkeleton(), WorldTransforms, MeshComponent->BonesOfInterest, PDI, BoneColours, MeshComponent->Bounds.SphereRadius, 0.f, false, MeshComponent->BoneRadiusMultiplier);
 	}
-}
-
-void FAnimationViewportClient::DrawBones(const TArray<FBoneIndexType>& RequiredBones, const FReferenceSkeleton& RefSkeleton, const TArray<FTransform> & WorldTransforms, const TArray<int32>& InSelectedBones, FPrimitiveDrawInterface* PDI, const TArray<FLinearColor>& BoneColours, float BoundRadius, float LineThickness/*=0.f*/, bool bForceDraw/*=false*/, float InBoneRadius/*=1.f*/) const
-{
-	TBitArray<> SelectedBones(false, RefSkeleton.GetNum());
 
-	if (InSelectedBones.Num() > 0)
+	// add children of selected
+	if (bDrawSelectedAndChildren || bDrawSelectedAndParentsAndChildren)
 	{
-		// Add the selected bones
-		if (GetBoneDrawMode() == EBoneDrawMode::Selected || GetBoneDrawMode() == EBoneDrawMode::SelectedAndParents || GetBoneDrawMode() == EBoneDrawMode::SelectedAndChildren || GetBoneDrawMode() == EBoneDrawMode::SelectedAndParentsAndChildren)
+		for (int32 BoneIndex = 0; BoneIndex < RefSkeleton.GetNum(); ++BoneIndex)
 		{
-			for (int32 BoneIndex : InSelectedBones)
+			const int32 ParentIndex = RefSkeleton.GetParentIndex(BoneIndex);
+			if (ParentIndex != INDEX_NONE && BonesToDraw[ParentIndex])
 			{
-				if (BoneIndex != INDEX_NONE)
-				{
-					SelectedBones[BoneIndex] = true;
-				}
+				BonesToDraw[BoneIndex] = true;
 			}
 		}
+	}
 
-		// Add the children of the selected bones
-		if (GetBoneDrawMode() == EBoneDrawMode::SelectedAndChildren || GetBoneDrawMode() == EBoneDrawMode::SelectedAndParentsAndChildren)
+	// add parents of selected
+	if (bDrawSelectedAndParents || bDrawSelectedAndParentsAndChildren)
+	{
+		for (const int32 BoneIndex : InSelectedBones)
 		{
-			for (int32 BoneIndex = 0; BoneIndex < RefSkeleton.GetNum(); ++BoneIndex)
+			if (BoneIndex != INDEX_NONE)
 			{
-				int32 ParentIndex = RefSkeleton.GetParentIndex(BoneIndex);
-				if (ParentIndex != INDEX_NONE && SelectedBones[ParentIndex])
+				for (int32 ParentIndex = RefSkeleton.GetParentIndex(BoneIndex); ParentIndex != INDEX_NONE; ParentIndex = RefSkeleton.GetParentIndex(ParentIndex))
 				{
-					SelectedBones[BoneIndex] = true;
+					BonesToDraw[ParentIndex] = true;
 				}
 			}
 		}
+	}
 
-		// Add the parents of the selected bones
-		if (GetBoneDrawMode() == EBoneDrawMode::SelectedAndParents || GetBoneDrawMode() == EBoneDrawMode::SelectedAndParentsAndChildren)
+	// determine which bones are "affected" (these are ALL children of selected bones)
+	TBitArray<> AffectedBones(false, RefSkeleton.GetNum());
+	for (int32 BoneIndex = 0; BoneIndex < RefSkeleton.GetNum(); ++BoneIndex)
+	{
+		for (int32 ParentIndex = RefSkeleton.GetParentIndex(BoneIndex); ParentIndex != INDEX_NONE; ParentIndex = RefSkeleton.GetParentIndex(ParentIndex))
 		{
-			for (int32 BoneIndex : InSelectedBones)
+			if (InSelectedBones.Contains(ParentIndex))
 			{
-				if (BoneIndex != INDEX_NONE)
-				{
-					for (int ParentIndex = RefSkeleton.GetParentIndex(BoneIndex); ParentIndex != INDEX_NONE; ParentIndex = RefSkeleton.GetParentIndex(ParentIndex))
-					{
-						SelectedBones[ParentIndex] = true;
-					}
-				}
+				AffectedBones[BoneIndex] = true;
+				break;
 			}
 		}
 	}
-
-	// we may not want to axis to be too big, so clamp at 1 % of bound
-	const float MaxDrawRadius = BoundRadius * 0.01f;
-	// we could cache parent bones as we calculate, but right now I'm not worried about perf issue of this
+	
+	// spin through all required bones and render them
+	const float BoneRadius = GetBoneDrawSize();
 	for ( int32 Index=0; Index<RequiredBones.Num(); ++Index )
 	{
 		const int32 BoneIndex = RequiredBones[Index];
 
-		if (bForceDraw || GetBoneDrawMode() == EBoneDrawMode::All || SelectedBones[BoneIndex])
+		// skips bones that should not be drawn
+		const bool bDoDraw = bForceDraw || GetBoneDrawMode() == EBoneDrawMode::All || BonesToDraw[BoneIndex];
+		if (!bDoDraw)
 		{
-			const int32 ParentIndex = RefSkeleton.GetParentIndex(BoneIndex);
-			FVector Start, End;
-			float Radius;
-			FLinearColor LineColor = BoneColours[BoneIndex];
+			continue;
+		}
 
-			if (ParentIndex >= 0)
-			{
-				Start = WorldTransforms[ParentIndex].GetLocation();
-				End = WorldTransforms[BoneIndex].GetLocation();
-				// Scale the radius proportionally to the bone length (clamped by bound to not be too long or big)
-				const float BoneLength = (End - Start).Size();
-				Radius = FMath::Clamp(BoneLength * 0.05f, 0.1f, MaxDrawRadius) * InBoneRadius;
-			}
-			else
-			{
-				Start = FVector::ZeroVector;
-				End = WorldTransforms[BoneIndex].GetLocation();
-				// Set the radius to a constant size (rather than proportional to bone length)
-				Radius = FMath::Min(1.0f, MaxDrawRadius) * InBoneRadius;
-			}
+		// determine color of bone based on selection / affected state
+		const bool bIsSelected = InSelectedBones.Contains(BoneIndex);
+		const bool bIsAffected = AffectedBones[BoneIndex];
+		FLinearColor DefaultBoneColor = BoneColors.IsEmpty() ? SkeletalDebugRendering::DEFAULT_BONE_COLOR : BoneColors[Index];
+		FLinearColor BoneColor = bIsAffected ? SkeletalDebugRendering::AFFECTED_BONE_COLOR : DefaultBoneColor;
+		BoneColor = bIsSelected ? SkeletalDebugRendering::SELECTED_BONE_COLOR : BoneColor;
 
-			//Render Sphere for bone end point and a cone between it and its parent.
-			FName BoneName = RefSkeleton.GetBoneName(BoneIndex);
-			PDI->SetHitProxy(new HPersonaBoneHitProxy(BoneIndex, BoneName));
-			SkeletalDebugRendering::DrawWireBone(PDI, Start, End, LineColor, SDPG_Foreground, Radius);
-			PDI->SetHitProxy(nullptr);
+		// draw the little coordinate frame inside the bone ONLY if selected or affected
+		const bool bDrawAxesInsideBone = bIsAffected|| bIsSelected;
 
-			// draw gizmo
-			if ((GetLocalAxesMode() == ELocalAxesMode::All) ||
-				(GetLocalAxesMode() == ELocalAxesMode::Selected && InSelectedBones.Contains(BoneIndex))
-				)
+		// draw cone to each child
+		// but use a different color if this bone is NOT selected, but the child IS selected
+		TArray<FVector> ChildPositions;
+		TArray<FLinearColor> ChildColors;
+		for (int32 ChildIndex = 0; ChildIndex < RefSkeleton.GetNum(); ++ChildIndex)
+		{
+			const int32 ParentIndex = RefSkeleton.GetParentIndex(ChildIndex);
+			if (ParentIndex == BoneIndex)
 			{
-				// we want to say 10 % of bone length is good
-				SkeletalDebugRendering::DrawAxes(PDI, WorldTransforms[BoneIndex], SDPG_Foreground, 0.f , Radius);
+				ChildPositions.Add(WorldTransforms[ChildIndex].GetLocation());
+				FLinearColor ChildLineColor = BoneColor;
+				if (!bIsSelected && InSelectedBones.Contains(ChildIndex))
+				{
+					ChildLineColor = SkeletalDebugRendering::PARENT_OF_SELECTED_BONE_COLOR;
+				}
+				ChildColors.Add(ChildLineColor);
 			}
 		}
+
+		const FName BoneName = RefSkeleton.GetBoneName(BoneIndex);
+		const FTransform BoneTransform = WorldTransforms[BoneIndex];
+
+		if (bAddHitProxy)
+		{
+			PDI->SetHitProxy(new HPersonaBoneHitProxy(BoneIndex, BoneName));
+		}
+		SkeletalDebugRendering::DrawWireBoneAdvanced(
+			PDI,
+			BoneTransform,
+			ChildPositions,
+			ChildColors,
+			BoneColor,
+			SDPG_Foreground,
+			BoneRadius,
+			bDrawAxesInsideBone);
+		if (RefSkeleton.GetParentIndex(BoneIndex) == INDEX_NONE)
+		{
+			SkeletalDebugRendering::DrawRootCone(PDI, BoneTransform, ComponentOrigin, BoneRadius);
+		}
+		if (bAddHitProxy)
+		{
+			PDI->SetHitProxy(nullptr);	
+		}
 	}
 }
 
 void FAnimationViewportClient::DrawMeshSubsetBones(const UDebugSkelMeshComponent* MeshComponent, const TArray<int32>& BonesOfInterest, FPrimitiveDrawInterface* PDI) const
 {
 	// this BonesOfInterest has to be in MeshComponent base, not Skeleton 
-	if ( MeshComponent && MeshComponent->SkeletalMesh && BonesOfInterest.Num() > 0 )
+	if (!MeshComponent ||
+		!MeshComponent->SkeletalMesh ||
+		BonesOfInterest.IsEmpty() ||
+		MeshComponent->SkeletonDrawMode == ESkeletonDrawMode::Hidden)
 	{
-		TArray<FTransform> WorldTransforms;
-		WorldTransforms.AddUninitialized(MeshComponent->GetNumDrawTransform());
+		return;
+	}
+	
+	TArray<FTransform> WorldTransforms;
+	WorldTransforms.AddUninitialized(MeshComponent->GetNumDrawTransform());
 
-		TArray<FLinearColor> BoneColours;
-		BoneColours.AddUninitialized(MeshComponent->GetNumDrawTransform());
+	TArray<FLinearColor> BoneColours;
+	BoneColours.AddUninitialized(MeshComponent->GetNumDrawTransform());
 
-		TArray<FBoneIndexType> RequiredBones;
+	TArray<FBoneIndexType> RequiredBones;
 
-		const FReferenceSkeleton& RefSkeleton = MeshComponent->GetReferenceSkeleton();
+	const FReferenceSkeleton& RefSkeleton = MeshComponent->GetReferenceSkeleton();
 
-		static const FName SelectionColorName("SelectionColor");
+	// we could cache parent bones as we calculate, but right now I'm not worried about perf issue of this
+	const TArray<FBoneIndexType>& DrawBoneIndices = MeshComponent->GetDrawBoneIndices();
+	for ( auto Iter = DrawBoneIndices.CreateConstIterator(); Iter; ++Iter)
+	{
+		const int32 BoneIndex = *Iter;
+		bool bDrawBone = false;
 
-		const FSlateColor SelectionColor = FAppStyle::GetSlateColor(SelectionColorName);
-		const FLinearColor LinearSelectionColor( SelectionColor.IsColorSpecified() ? SelectionColor.GetSpecifiedColor() : FLinearColor::White );
+		const int32 ParentIndex = RefSkeleton.GetParentIndex(BoneIndex);
 
-		// we could cache parent bones as we calculate, but right now I'm not worried about perf issue of this
-		const TArray<FBoneIndexType>& DrawBoneIndices = MeshComponent->GetDrawBoneIndices();
-		for ( auto Iter = DrawBoneIndices.CreateConstIterator(); Iter; ++Iter)
+		// need to see if it's child of any of Bones of interest
+		for (auto SubIter=BonesOfInterest.CreateConstIterator(); SubIter; ++SubIter )
 		{
-			const int32 BoneIndex = *Iter;
-			bool bDrawBone = false;
-
-			const int32 ParentIndex = RefSkeleton.GetParentIndex(BoneIndex);
-
-			// need to see if it's child of any of Bones of interest
-			for (auto SubIter=BonesOfInterest.CreateConstIterator(); SubIter; ++SubIter )
+			const int32 SubBoneIndex = *SubIter;
+			// if I'm child of the BonesOfInterest
+			if(BoneIndex == SubBoneIndex)
 			{
-				const int32 SubBoneIndex = *SubIter;
-				// if I'm child of the BonesOfInterest
-				if(BoneIndex == SubBoneIndex)
-				{
-					//found a bone we are interested in
-					if(ParentIndex >= 0)
-					{
-						WorldTransforms[ParentIndex] = MeshComponent->GetDrawTransform(ParentIndex)*MeshComponent->GetComponentTransform();
-					}
-					BoneColours[BoneIndex] = LinearSelectionColor;
-					bDrawBone = true;
-					break;
-				}
-				else if ( RefSkeleton.BoneIsChildOf(BoneIndex, SubBoneIndex) )
+				//found a bone we are interested in
+				if(ParentIndex >= 0)
 				{
-					BoneColours[BoneIndex] = FLinearColor::White;
-					bDrawBone = true;
-					break;
+					WorldTransforms[ParentIndex] = MeshComponent->GetDrawTransform(ParentIndex)*MeshComponent->GetComponentTransform();
 				}
+				BoneColours[BoneIndex] = SkeletalDebugRendering::SELECTED_BONE_COLOR;
+				bDrawBone = true;
+				break;
 			}
-
-			if (bDrawBone)
+			else if ( RefSkeleton.BoneIsChildOf(BoneIndex, SubBoneIndex) )
 			{
-				//add to the list
-				RequiredBones.AddUnique(BoneIndex);
-				WorldTransforms[BoneIndex] = MeshComponent->GetDrawTransform(BoneIndex) * MeshComponent->GetComponentTransform();
+				BoneColours[BoneIndex] = SkeletalDebugRendering::DEFAULT_BONE_COLOR;
+				bDrawBone = true;
+				break;
 			}
 		}
 
-		DrawBones(RequiredBones, MeshComponent->GetReferenceSkeleton(), WorldTransforms, MeshComponent->BonesOfInterest, PDI, BoneColours, MeshComponent->Bounds.SphereRadius, 0.3f);
+		if (bDrawBone)
+		{
+			//add to the list
+			RequiredBones.AddUnique(BoneIndex);
+			WorldTransforms[BoneIndex] = MeshComponent->GetDrawTransform(BoneIndex) * MeshComponent->GetComponentTransform();
+		}
 	}
+
+	constexpr bool bForceDraw = false;
+	const bool bAddHitProxy = MeshComponent->SkeletonDrawMode != ESkeletonDrawMode::GreyedOut;
+
+	DrawBones(
+		MeshComponent->GetComponentLocation(),
+		RequiredBones,
+		MeshComponent->GetReferenceSkeleton(),
+		WorldTransforms,
+		MeshComponent->BonesOfInterest,
+		BoneColours,
+		PDI,
+		bForceDraw,
+		bAddHitProxy);
 }
 
 void FAnimationViewportClient::DrawAttributes(UDebugSkelMeshComponent* MeshComponent, FPrimitiveDrawInterface* PDI) const
@@ -1931,13 +2085,12 @@ void FAnimationViewportClient::SetFloorOffset( float NewValue )
 
 void FAnimationViewportClient::ToggleCPUSkinning()
 {
-	UDebugSkelMeshComponent* PreviewMeshComponent = GetAnimPreviewScene()->GetPreviewMeshComponent();
-	if (PreviewMeshComponent != nullptr)
+	GetPreviewScene()->ForEachPreviewMesh([this](UDebugSkelMeshComponent* PreviewMeshComponent)
 	{
 		const bool bCurVal = PreviewMeshComponent->GetCPUSkinningEnabled();
 		PreviewMeshComponent->SetCPUSkinningEnabled(!bCurVal);
-		Invalidate();
-	}
+	});
+	Invalidate();
 }
 
 bool FAnimationViewportClient::IsSetCPUSkinningChecked() const
@@ -1952,13 +2105,13 @@ bool FAnimationViewportClient::IsSetCPUSkinningChecked() const
 
 void FAnimationViewportClient::ToggleShowNormals()
 {
-	UDebugSkelMeshComponent* PreviewMeshComponent = GetAnimPreviewScene()->GetPreviewMeshComponent();
-	if(PreviewMeshComponent != nullptr)
+	GetPreviewScene()->ForEachPreviewMesh([this](UDebugSkelMeshComponent* PreviewMeshComponent)
 	{
 		PreviewMeshComponent->bDrawNormals = !PreviewMeshComponent->bDrawNormals;
 		PreviewMeshComponent->MarkRenderStateDirty();
-		Invalidate();
-	}
+	});
+	
+	Invalidate();
 }
 
 bool FAnimationViewportClient::IsSetShowNormalsChecked() const
@@ -1973,13 +2126,13 @@ bool FAnimationViewportClient::IsSetShowNormalsChecked() const
 
 void FAnimationViewportClient::ToggleShowTangents()
 {
-	UDebugSkelMeshComponent* PreviewMeshComponent = GetAnimPreviewScene()->GetPreviewMeshComponent();
-	if(PreviewMeshComponent != nullptr)
+	GetPreviewScene()->ForEachPreviewMesh([this](UDebugSkelMeshComponent* PreviewMeshComponent)
 	{
 		PreviewMeshComponent->bDrawTangents = !PreviewMeshComponent->bDrawTangents;
 		PreviewMeshComponent->MarkRenderStateDirty();
-		Invalidate();
-	}
+	});
+	
+	Invalidate();
 }
 
 bool FAnimationViewportClient::IsSetShowTangentsChecked() const
@@ -1994,13 +2147,13 @@ bool FAnimationViewportClient::IsSetShowTangentsChecked() const
 
 void FAnimationViewportClient::ToggleShowBinormals()
 {
-	UDebugSkelMeshComponent* PreviewMeshComponent = GetAnimPreviewScene()->GetPreviewMeshComponent();
-	if(PreviewMeshComponent != nullptr)
+	GetPreviewScene()->ForEachPreviewMesh([this](UDebugSkelMeshComponent* PreviewMeshComponent)
 	{
 		PreviewMeshComponent->bDrawBinormals = !PreviewMeshComponent->bDrawBinormals;
-		PreviewMeshComponent->MarkRenderStateDirty();
-		Invalidate();
-	}
+			PreviewMeshComponent->MarkRenderStateDirty();
+	});
+	
+	Invalidate();
 }
 
 bool FAnimationViewportClient::IsSetShowBinormalsChecked() const
@@ -2157,12 +2310,12 @@ void FAnimationViewportClient::SetupViewForRendering( FSceneViewFamily& ViewFami
 void FAnimationViewportClient::HandleToggleShowFlag(FEngineShowFlags::EShowFlag EngineShowFlagIndex)
 {
 	FEditorViewportClient::HandleToggleShowFlag(EngineShowFlagIndex);
-	
-	if (UDebugSkelMeshComponent* Component = GetAnimPreviewScene()->GetPreviewMeshComponent())
+
+	GetPreviewScene()->ForEachPreviewMesh([this](UDebugSkelMeshComponent* InMesh)
 	{
-		Component->bDisplayVertexColors = EngineShowFlags.VertexColors;
-		Component->MarkRenderStateDirty();
-	}
+		InMesh->bDisplayVertexColors = EngineShowFlags.VertexColors;
+		InMesh->MarkRenderStateDirty();
+	});
 
 	ConfigOption->SetShowGrid(EngineShowFlags.Grid);
 }
diff --git a/Engine/Source/Editor/Persona/Private/AnimationEditorViewportClient.h b/Engine/Source/Editor/Persona/Private/AnimationEditorViewportClient.h
--- a/Engine/Source/Editor/Persona/Private/AnimationEditorViewportClient.h
+++ b/Engine/Source/Editor/Persona/Private/AnimationEditorViewportClient.h
@@ -20,6 +20,9 @@ class USkeletalMeshSocket;
 struct FCompactHeapPose;
 struct FSkelMeshRenderSection;
 
+DECLARE_DELEGATE_OneParam(FOnBoneSizeSet, float);
+DECLARE_DELEGATE_RetVal(float, FOnGetBoneSize)
+
 //////////////////////////////////////////////////////////////////////////
 // ELocalAxesMode
 
@@ -86,15 +89,15 @@ namespace EAnimationPlaybackSpeeds
 /////////////////////////////////////////////////////////////////////////
 // FAnimationViewportClient
 
-class FAnimationViewportClient : public FEditorViewportClient
+class PERSONA_API FAnimationViewportClient : public FEditorViewportClient
 {
 protected:
 
 	/** Function to display bone names*/
-	void ShowBoneNames(FCanvas* Canvas, FSceneView* View);
+	void ShowBoneNames(FCanvas* Canvas, FSceneView* View, UDebugSkelMeshComponent* MeshComponent);
 
 	/** Function to display transform attribute names*/
-	void ShowAttributeNames(FCanvas* Canvas, FSceneView* View);
+	void ShowAttributeNames(FCanvas* Canvas, FSceneView* View, UDebugSkelMeshComponent* MeshComponent) const;
 
 	/** Function to display debug lines generated from skeletal controls in animBP mode */
 	void DrawNodeDebugLines(TArray<FText>& Lines, FCanvas* Canvas, FSceneView* View);
@@ -130,7 +133,7 @@ public:
 	// End of FEditorViewportClient interface
 
 	/** Draw call to render UV overlay */
-	void DrawUVsForMesh(FViewport* InViewport, FCanvas* InCanvas, int32 InTextYPos);
+	void DrawUVsForMesh(FViewport* InViewport, FCanvas* InCanvas, int32 InTextYPos, UDebugSkelMeshComponent* MeshComponent);
 
 	/** Set the camera follow mode */
 	void SetCameraFollowMode(EAnimationViewportCameraFollowMode Mode, FName InBoneName = NAME_None);
@@ -160,7 +163,7 @@ public:
 	void HandleSkeletalMeshChanged(class USkeletalMesh* OldSkeletalMesh, class USkeletalMesh* NewSkeletalMesh);
 
 	/** Function to display bone names*/
-	void ShowBoneNames(FViewport* Viewport, FCanvas* Canvas);
+	void ShowBoneNames(FViewport* Viewport, FCanvas* Canvas, UDebugSkelMeshComponent* MeshComponent);
 
 	/** Function to enable/disable floor auto align */
 	void OnToggleAutoAlignFloor();
@@ -198,6 +201,10 @@ public:
 	/** Get the Bone local axis mode */
 	ELocalAxesMode::Type GetLocalAxesMode() const;
 
+	/** Access Bone Draw size config option*/
+	void SetBoneDrawSize(const float InBoneDrawSize);
+	float GetBoneDrawSize() const;
+	
 	/** Function to set Bone Draw  mode for the EBoneDrawType */
 	void SetBoneDrawMode(EBoneDrawMode::Type AxesMode);
 
@@ -317,6 +324,10 @@ public:
 	/** persona config options **/
 	UPersonaOptions* ConfigOption;
 
+	/** allow client code to store/serialize bone size if desired */
+	FOnBoneSizeSet OnSetBoneSize;
+	FOnGetBoneSize OnGetBoneSize;
+
 private:
 	/** Weak pointer back to the preview scene we are viewing */
 	TWeakPtr<class IPersonaPreviewScene> PreviewScenePtr;
@@ -393,7 +404,16 @@ private:
 	/** Draw Bones for non retargeted animation. */
 	void DrawMeshBonesBakedAnimation(UDebugSkelMeshComponent * MeshComponent, FPrimitiveDrawInterface* PDI) const;
 	/** Draws Bones for RequiredBones with WorldTransform **/
-	void DrawBones(const TArray<FBoneIndexType> & RequiredBones, const FReferenceSkeleton& RefSkeleton, const TArray<FTransform> & WorldTransforms, const TArray<int32>& InSelectedBones, FPrimitiveDrawInterface* PDI, const TArray<FLinearColor>& BoneColours, float BoundRadius, float LineThickness = 0.f, bool bForceDraw = false, float InBoneRadius = 1.f) const;
+	void DrawBones(
+		const FVector& ComponentOrigin,
+		const TArray<FBoneIndexType>& RequiredBones,
+		const FReferenceSkeleton& RefSkeleton,
+		const TArray<FTransform>& WorldTransforms,
+		const TArray<int32>& InSelectedBones,
+		const TArray<FLinearColor>& BoneColors,
+		FPrimitiveDrawInterface* PDI,
+		bool bForceDraw,
+		bool bAddHitProxy) const;
 	/** Draw Sub set of Bones **/
 	void DrawMeshSubsetBones(const UDebugSkelMeshComponent* MeshComponent, const TArray<int32>& BonesOfInterest, FPrimitiveDrawInterface* PDI) const;
 
@@ -437,6 +457,9 @@ private:
 	void HandlePreviewScenePostTick();
 
 private:
+	/** Size to draw bones in the viewport. Transient setting. */
+	float BoneDrawSize=1.0f;
+	
 	/** Allow mesh stats to be disabled for specific viewport instances */
 	bool bShowMeshStats;
 
diff --git a/Engine/Source/Editor/Persona/Private/SAnimViewportToolBar.cpp b/Engine/Source/Editor/Persona/Private/SAnimViewportToolBar.cpp
--- a/Engine/Source/Editor/Persona/Private/SAnimViewportToolBar.cpp
+++ b/Engine/Source/Editor/Persona/Private/SAnimViewportToolBar.cpp
@@ -162,6 +162,45 @@ protected:
 	TWeakPtr<SAnimationEditorViewportTabBody> AnimViewportPtr;
 };
 
+//Class definition which represents widget to modify Bone Draw Size in viewport
+class SBoneDrawSizeSetting : public SCompoundWidget
+{
+	
+public:
+
+	SLATE_BEGIN_ARGS(SBoneDrawSizeSetting) {}
+	SLATE_ARGUMENT(TWeakPtr<SAnimationEditorViewportTabBody>, AnimEditorViewport)
+	SLATE_END_ARGS()
+
+	void Construct(const FArguments& InArgs )
+	{
+		AnimViewportPtr = InArgs._AnimEditorViewport;
+
+		this->ChildSlot
+		[
+			SNew(SBox)
+			.HAlign(HAlign_Right)
+			[
+				SNew(SBox)
+				.Padding(FMargin(4.0f, 0.0f, 0.0f, 0.0f))
+				.WidthOverride(100.0f)
+				[
+					SNew(SSpinBox<float>)
+					.Font(FAppStyle::GetFontStyle(TEXT("MenuItem.Font")))
+					.ToolTipText(LOCTEXT("BoneDrawSize_ToolTip", "Change bone size in viewport."))
+					.MinValue(0)
+					.MaxValue(10)
+					.Value(AnimViewportPtr.Pin().ToSharedRef(), &SAnimationEditorViewportTabBody::GetBoneDrawSize)
+					.OnValueChanged(SSpinBox<float>::FOnValueChanged::CreateSP(AnimViewportPtr.Pin().ToSharedRef(), &SAnimationEditorViewportTabBody::SetBoneDrawSize))
+				]
+			]
+		];
+	}
+
+protected:
+	TWeakPtr<SAnimationEditorViewportTabBody> AnimViewportPtr;
+};
+
 ///////////////////////////////////////////////////////////
 // SAnimViewportToolBar
 
@@ -694,7 +733,7 @@ TSharedRef<SWidget> SAnimViewportToolBar::GenerateCharacterMenu() const
 			InMenuBuilder.AddSubMenu(
 				LOCTEXT("CharacterMenu_BoneDrawSubMenu", "Bones"),
 				LOCTEXT("CharacterMenu_BoneDrawSubMenuToolTip", "Bone Drawing Options"),
-				FNewMenuDelegate::CreateLambda([](FMenuBuilder& SubMenuBuilder)
+				FNewMenuDelegate::CreateLambda([this](FMenuBuilder& SubMenuBuilder)
 				{
 					SubMenuBuilder.BeginSection("BonesAndSockets", LOCTEXT("CharacterMenu_BonesAndSocketsLabel", "Show"));
 					{
@@ -706,6 +745,8 @@ TSharedRef<SWidget> SAnimViewportToolBar::GenerateCharacterMenu() const
 
 					SubMenuBuilder.BeginSection("AnimViewportPreviewHierarchyBoneDraw", LOCTEXT("CharacterMenu_Actions_BoneDrawing", "Bone Drawing"));
 					{
+						TSharedPtr<SWidget> BoneSizeWidget = SNew(SBoneDrawSizeSetting).AnimEditorViewport(Viewport);
+						SubMenuBuilder.AddWidget(BoneSizeWidget.ToSharedRef(), LOCTEXT("CharacterMenu_Actions_BoneDrawSize", "Bone Draw Size:"));
 						SubMenuBuilder.AddMenuEntry(FAnimViewportShowCommands::Get().ShowBoneDrawAll);
 						SubMenuBuilder.AddMenuEntry(FAnimViewportShowCommands::Get().ShowBoneDrawSelected);
 						SubMenuBuilder.AddMenuEntry(FAnimViewportShowCommands::Get().ShowBoneDrawSelectedAndParents);
diff --git a/Engine/Source/Editor/Persona/Private/SAnimationEditorViewport.cpp b/Engine/Source/Editor/Persona/Private/SAnimationEditorViewport.cpp
--- a/Engine/Source/Editor/Persona/Private/SAnimationEditorViewport.cpp
+++ b/Engine/Source/Editor/Persona/Private/SAnimationEditorViewport.cpp
@@ -1157,67 +1157,58 @@ int32 SAnimationEditorViewportTabBody::GetLODModelCount() const
 
 void SAnimationEditorViewportTabBody::OnShowMorphTargets()
 {
-	if (UDebugSkelMeshComponent* PreviewComponent = GetPreviewScene()->GetPreviewMeshComponent())
+	SetViewportSettingForEachDebugMesh([](UDebugSkelMeshComponent* InMesh)
 	{
-		PreviewComponent->bDisableMorphTarget = !PreviewComponent->bDisableMorphTarget;
-		PreviewComponent->MarkRenderStateDirty();
-		RefreshViewport();
-	}
+		InMesh->bDisableMorphTarget = !InMesh->bDisableMorphTarget;
+	});
 }
 
 void SAnimationEditorViewportTabBody::OnShowBoneNames()
 {
-	if (UDebugSkelMeshComponent* PreviewComponent = GetPreviewScene()->GetPreviewMeshComponent())
+	SetViewportSettingForEachDebugMesh([](UDebugSkelMeshComponent* InMesh)
 	{
-		PreviewComponent->bShowBoneNames = !PreviewComponent->bShowBoneNames;
-		PreviewComponent->MarkRenderStateDirty();
-		RefreshViewport();
-	}
+		InMesh->bShowBoneNames = !InMesh->bShowBoneNames;
+	});
 }
 
 void SAnimationEditorViewportTabBody::OnShowRawAnimation()
 {
-	if (UDebugSkelMeshComponent* PreviewComponent = GetPreviewScene()->GetPreviewMeshComponent())
+	SetViewportSettingForEachDebugMesh([](UDebugSkelMeshComponent* InMesh)
 	{
-		PreviewComponent->bDisplayRawAnimation = !PreviewComponent->bDisplayRawAnimation;
-		PreviewComponent->MarkRenderStateDirty();
-	}
+		InMesh->bDisplayRawAnimation = !InMesh->bDisplayRawAnimation;
+	});
 }
 
 void SAnimationEditorViewportTabBody::OnShowNonRetargetedAnimation()
 {
-	if (UDebugSkelMeshComponent* PreviewComponent = GetPreviewScene()->GetPreviewMeshComponent())
+	SetViewportSettingForEachDebugMesh([](UDebugSkelMeshComponent* InMesh)
 	{
-		PreviewComponent->bDisplayNonRetargetedPose = !PreviewComponent->bDisplayNonRetargetedPose;
-		PreviewComponent->MarkRenderStateDirty();
-	}
+		InMesh->bDisplayNonRetargetedPose = !InMesh->bDisplayNonRetargetedPose;
+	});
 }
 
 void SAnimationEditorViewportTabBody::OnShowSourceRawAnimation()
 {
-	if (UDebugSkelMeshComponent* PreviewComponent = GetPreviewScene()->GetPreviewMeshComponent())
+	SetViewportSettingForEachDebugMesh([](UDebugSkelMeshComponent* InMesh)
 	{
-		PreviewComponent->bDisplaySourceAnimation = !PreviewComponent->bDisplaySourceAnimation;
-		PreviewComponent->MarkRenderStateDirty();
-	}
+		InMesh->bDisplaySourceAnimation = !InMesh->bDisplaySourceAnimation;
+	});
 }
 
 void SAnimationEditorViewportTabBody::OnShowBakedAnimation()
 {
-	if (UDebugSkelMeshComponent* PreviewComponent = GetPreviewScene()->GetPreviewMeshComponent())
+	SetViewportSettingForEachDebugMesh([](UDebugSkelMeshComponent* InMesh)
 	{
-		PreviewComponent->bDisplayBakedAnimation = !PreviewComponent->bDisplayBakedAnimation;
-		PreviewComponent->MarkRenderStateDirty();
-	}
+		InMesh->bDisplayBakedAnimation = !InMesh->bDisplayBakedAnimation;
+	});
 }
 
 void SAnimationEditorViewportTabBody::OnShowAdditiveBase()
 {
-	if (UDebugSkelMeshComponent* PreviewComponent = GetPreviewScene()->GetPreviewMeshComponent())
+	SetViewportSettingForEachDebugMesh([](UDebugSkelMeshComponent* InMesh)
 	{
-		PreviewComponent->bDisplayAdditiveBasePose = !PreviewComponent->bDisplayAdditiveBasePose;
-		PreviewComponent->MarkRenderStateDirty();
-	}
+		InMesh->bDisplayAdditiveBasePose = !InMesh->bDisplayAdditiveBasePose;
+	});
 }
 
 bool SAnimationEditorViewportTabBody::IsPreviewingAnimation() const
@@ -1246,24 +1237,39 @@ bool SAnimationEditorViewportTabBody::IsShowingRawAnimation() const
 
 void SAnimationEditorViewportTabBody::OnToggleDisablePostProcess()
 {
-	if(UDebugSkelMeshComponent* PreviewComponent = GetPreviewScene()->GetPreviewMeshComponent())
+	SetViewportSettingForEachDebugMesh([](UDebugSkelMeshComponent* InMesh)
 	{
-		PreviewComponent->ToggleDisablePostProcessBlueprint();
-
-		AddPostProcessNotification();
-	}
+		InMesh->ToggleDisablePostProcessBlueprint();
+	});
+	
+	AddPostProcessNotification();
 }
 
 bool SAnimationEditorViewportTabBody::CanDisablePostProcess()
 {
-	UDebugSkelMeshComponent* PreviewComponent = GetPreviewScene()->GetPreviewMeshComponent();
-	return PreviewComponent && PreviewComponent->PostProcessAnimInstance;
+	TArray<UDebugSkelMeshComponent*> PreviewMeshComponents = GetPreviewScene()->GetAllPreviewMeshComponents();
+	for (UDebugSkelMeshComponent* PreviewMeshComponent : PreviewMeshComponents)
+	{
+		if (PreviewMeshComponent->PostProcessAnimInstance)
+		{
+			return true;
+		}
+	}
+	return false;
 }
 
 bool SAnimationEditorViewportTabBody::IsDisablePostProcessChecked()
 {
-	UDebugSkelMeshComponent* PreviewComponent = GetPreviewScene()->GetPreviewMeshComponent();
-	return PreviewComponent && PreviewComponent->GetDisablePostProcessBlueprint();
+	TArray<UDebugSkelMeshComponent*> PreviewMeshComponents = GetPreviewScene()->GetAllPreviewMeshComponents();
+	for (UDebugSkelMeshComponent* PreviewMeshComponent : PreviewMeshComponents)
+	{
+		if (PreviewMeshComponent->GetDisablePostProcessBlueprint())
+		{
+			return true;
+		}
+	}
+	
+	return false;
 }
 
 bool SAnimationEditorViewportTabBody::IsShowingNonRetargetedPose() const
@@ -1302,15 +1308,13 @@ bool SAnimationEditorViewportTabBody::IsShowingMeshInfo(int32 DisplayInfoMode) c
 
 void SAnimationEditorViewportTabBody::OnShowOverlayNone()
 {
-	UDebugSkelMeshComponent* PreviewComponent = GetPreviewScene()->GetPreviewMeshComponent();
-	if (PreviewComponent)
+	SetViewportSettingForEachDebugMesh([](UDebugSkelMeshComponent* PreviewMeshComponent)
 	{
-		PreviewComponent->SetShowBoneWeight(false);
-		PreviewComponent->SetShowMorphTargetVerts(false);
-		UpdateShowFlagForMeshEdges();
-		PreviewComponent->MarkRenderStateDirty();
-		RefreshViewport();
-	}
+		PreviewMeshComponent->SetShowBoneWeight(false);
+		PreviewMeshComponent->SetShowMorphTargetVerts(false);
+	});
+	
+	UpdateShowFlagForMeshEdges();
 }
 
 bool SAnimationEditorViewportTabBody::IsShowingOverlayNone() const
@@ -1321,14 +1325,13 @@ bool SAnimationEditorViewportTabBody::IsShowingOverlayNone() const
 
 void SAnimationEditorViewportTabBody::OnShowOverlayBoneWeight()
 {
-	UDebugSkelMeshComponent* PreviewComponent = GetPreviewScene()->GetPreviewMeshComponent();
-	if( PreviewComponent )
+	SetViewportSettingForEachDebugMesh([](UDebugSkelMeshComponent* PreviewMeshComponent)
 	{
-		PreviewComponent->SetShowBoneWeight( !PreviewComponent->bDrawBoneInfluences );
-		UpdateShowFlagForMeshEdges();
-		PreviewComponent->MarkRenderStateDirty();
-		RefreshViewport();
-	}
+		PreviewMeshComponent->SetShowBoneWeight( !PreviewMeshComponent->bDrawBoneInfluences );
+		PreviewMeshComponent->MarkRenderStateDirty();
+	});
+	
+	UpdateShowFlagForMeshEdges();
 }
 
 bool SAnimationEditorViewportTabBody::IsShowingOverlayBoneWeight() const
@@ -1339,14 +1342,13 @@ bool SAnimationEditorViewportTabBody::IsShowingOverlayBoneWeight() const
 
 void SAnimationEditorViewportTabBody::OnShowOverlayMorphTargetVert()
 {
-	UDebugSkelMeshComponent* PreviewComponent = GetPreviewScene()->GetPreviewMeshComponent();
-	if (PreviewComponent)
+	SetViewportSettingForEachDebugMesh([](UDebugSkelMeshComponent* PreviewMeshComponent)
 	{
-		PreviewComponent->SetShowMorphTargetVerts(!PreviewComponent->bDrawMorphTargetVerts);
-		UpdateShowFlagForMeshEdges();
-		PreviewComponent->MarkRenderStateDirty();
-		RefreshViewport();
-	}
+		PreviewMeshComponent->SetShowMorphTargetVerts(!PreviewMeshComponent->bDrawMorphTargetVerts);
+		PreviewMeshComponent->MarkRenderStateDirty();
+	});
+
+	UpdateShowFlagForMeshEdges();
 }
 
 bool SAnimationEditorViewportTabBody::IsShowingOverlayMorphTargetVerts() const
@@ -1355,6 +1357,18 @@ bool SAnimationEditorViewportTabBody::IsShowingOverlayMorphTargetVerts() const
 	return PreviewComponent != NULL && PreviewComponent->bDrawMorphTargetVerts;
 }
 
+void SAnimationEditorViewportTabBody::SetBoneDrawSize(float BoneDrawSize)
+{
+	TSharedRef<FAnimationViewportClient> AnimViewportClient = StaticCastSharedRef<FAnimationViewportClient>(LevelViewportClient.ToSharedRef());
+	AnimViewportClient->SetBoneDrawSize(BoneDrawSize);
+}
+
+float SAnimationEditorViewportTabBody::GetBoneDrawSize() const
+{
+	TSharedRef<FAnimationViewportClient> AnimViewportClient = StaticCastSharedRef<FAnimationViewportClient>(LevelViewportClient.ToSharedRef());
+	return AnimViewportClient->GetBoneDrawSize();
+}
+
 void SAnimationEditorViewportTabBody::OnSetBoneDrawMode(int32 BoneDrawMode)
 {
 	TSharedRef<FAnimationViewportClient> AnimViewportClient = StaticCastSharedRef<FAnimationViewportClient>(LevelViewportClient.ToSharedRef());
@@ -1381,13 +1395,10 @@ bool SAnimationEditorViewportTabBody::IsLocalAxesModeSet(int32 LocalAxesMode) co
 
 void SAnimationEditorViewportTabBody::OnShowSockets()
 {
-	UDebugSkelMeshComponent* PreviewComponent = GetPreviewScene()->GetPreviewMeshComponent();
-	if( PreviewComponent )
+	SetViewportSettingForEachDebugMesh([](UDebugSkelMeshComponent* PreviewMeshComponent)
 	{
-		PreviewComponent->bDrawSockets = !PreviewComponent->bDrawSockets;
-		PreviewComponent->MarkRenderStateDirty();
-		RefreshViewport();
-	}
+		PreviewMeshComponent->bDrawSockets = !PreviewMeshComponent->bDrawSockets;
+	});
 }
 
 bool SAnimationEditorViewportTabBody::IsShowingSockets() const
@@ -1398,13 +1409,10 @@ bool SAnimationEditorViewportTabBody::IsShowingSockets() const
 
 void SAnimationEditorViewportTabBody::OnShowAttributes()
 {
-	UDebugSkelMeshComponent* PreviewComponent = GetPreviewScene()->GetPreviewMeshComponent();
-	if (PreviewComponent)
+	SetViewportSettingForEachDebugMesh([](UDebugSkelMeshComponent* PreviewMeshComponent)
 	{
-		PreviewComponent->bDrawAttributes = !PreviewComponent->bDrawAttributes;
-		PreviewComponent->MarkRenderStateDirty();
-		RefreshViewport();
-	}
+		PreviewMeshComponent->bDrawAttributes = !PreviewMeshComponent->bDrawAttributes;
+	});
 }
 
 bool SAnimationEditorViewportTabBody::IsShowingAttributes() const
@@ -1440,25 +1448,35 @@ bool SAnimationEditorViewportTabBody::IsPlaybackSpeedSelected(int32 PlaybackSpee
 
 void SAnimationEditorViewportTabBody::ShowRetargetBasePose()
 {
-	UDebugSkelMeshComponent* PreviewComponent = GetPreviewScene()->GetPreviewMeshComponent();
-	if(PreviewComponent && PreviewComponent->PreviewInstance)
+	SetViewportSettingForEachDebugMesh([](UDebugSkelMeshComponent* PreviewMeshComponent)
 	{
-		PreviewComponent->PreviewInstance->SetForceRetargetBasePose(!PreviewComponent->PreviewInstance->GetForceRetargetBasePose());
-	}
+		PreviewMeshComponent->PreviewInstance->SetForceRetargetBasePose(!PreviewMeshComponent->PreviewInstance->GetForceRetargetBasePose());
+	});
 }
 
 bool SAnimationEditorViewportTabBody::CanShowRetargetBasePose() const
 {
-	UDebugSkelMeshComponent* PreviewComponent = GetPreviewScene()->GetPreviewMeshComponent();
-	return PreviewComponent != NULL && PreviewComponent->PreviewInstance;
+	TArray<UDebugSkelMeshComponent*> PreviewMeshComponents = GetPreviewScene()->GetAllPreviewMeshComponents();
+	for (UDebugSkelMeshComponent* PreviewMeshComponent : PreviewMeshComponents)
+	{
+		if (PreviewMeshComponent->PreviewInstance)
+		{
+			return true;
+		}
+	}
+	
+	return false;
 }
 
 bool SAnimationEditorViewportTabBody::IsShowRetargetBasePoseEnabled() const
 {
-	UDebugSkelMeshComponent* PreviewComponent = GetPreviewScene()->GetPreviewMeshComponent();
-	if(PreviewComponent && PreviewComponent->PreviewInstance)
+	TArray<UDebugSkelMeshComponent*> PreviewMeshComponents = GetPreviewScene()->GetAllPreviewMeshComponents();
+	for (UDebugSkelMeshComponent* PreviewMeshComponent : PreviewMeshComponents)
 	{
-		return PreviewComponent->PreviewInstance->GetForceRetargetBasePose();
+		if (PreviewMeshComponent && PreviewMeshComponent->PreviewInstance)
+		{
+			return PreviewMeshComponent->PreviewInstance->GetForceRetargetBasePose();
+		}
 	}
 	return false;
 }
@@ -1468,18 +1486,16 @@ void SAnimationEditorViewportTabBody::ShowBound()
 	TSharedRef<FAnimationViewportClient> AnimViewportClient = StaticCastSharedRef<FAnimationViewportClient>(LevelViewportClient.ToSharedRef());	
 	AnimViewportClient->ToggleShowBounds();
 
-	UDebugSkelMeshComponent* PreviewComponent = GetPreviewScene()->GetPreviewMeshComponent();
-	if(PreviewComponent)
+	SetViewportSettingForEachDebugMesh([AnimViewportClient](UDebugSkelMeshComponent* PreviewMeshComponent)
 	{
-		PreviewComponent->bDisplayBound = AnimViewportClient->EngineShowFlags.Bounds;
-		PreviewComponent->RecreateRenderState_Concurrent();
-	}
+		PreviewMeshComponent->bDisplayBound = AnimViewportClient->EngineShowFlags.Bounds;
+		PreviewMeshComponent->RecreateRenderState_Concurrent();
+	});
 }
 
 bool SAnimationEditorViewportTabBody::CanShowBound() const
 {
-	UDebugSkelMeshComponent* PreviewComponent = GetPreviewScene()->GetPreviewMeshComponent();
-	return PreviewComponent != NULL;
+	return !GetPreviewScene()->GetAllPreviewMeshComponents().IsEmpty();
 }
 
 bool SAnimationEditorViewportTabBody::IsShowBoundEnabled() const
@@ -1490,54 +1506,73 @@ bool SAnimationEditorViewportTabBody::IsShowBoundEnabled() const
 
 void SAnimationEditorViewportTabBody::ToggleShowPreviewMesh()
 {
-	if (UDebugSkelMeshComponent* PreviewComponent = GetPreviewScene()->GetPreviewMeshComponent())
+	const bool bCurrentlyVisible = IsShowPreviewMeshEnabled();
+	SetViewportSettingForEachDebugMesh([bCurrentlyVisible](UDebugSkelMeshComponent* PreviewMeshComponent)
 	{
-		bool bCurrentlyVisible = IsShowPreviewMeshEnabled();
-		PreviewComponent->SetVisibility(!bCurrentlyVisible);
-	}
+		PreviewMeshComponent->SetVisibility(!bCurrentlyVisible);
+	});
 }
 
 bool SAnimationEditorViewportTabBody::CanShowPreviewMesh() const
 {
-	UDebugSkelMeshComponent* PreviewComponent = GetPreviewScene()->GetPreviewMeshComponent();
-	return PreviewComponent != NULL;
+	return !GetPreviewScene()->GetAllPreviewMeshComponents().IsEmpty();
 }
 
 bool SAnimationEditorViewportTabBody::IsShowPreviewMeshEnabled() const
 {
-	UDebugSkelMeshComponent* PreviewComponent = GetPreviewScene()->GetPreviewMeshComponent();
-	return (PreviewComponent != NULL) && PreviewComponent->IsVisible();
+	TArray<UDebugSkelMeshComponent*> PreviewMeshComponents = GetPreviewScene()->GetAllPreviewMeshComponents();
+	for (UDebugSkelMeshComponent* PreviewMeshComponent : PreviewMeshComponents)
+	{
+		if(PreviewMeshComponent && PreviewMeshComponent->IsVisible())
+		{
+			return true;
+		}
+	}
+	
+	return false;
 }
 
 void SAnimationEditorViewportTabBody::UseInGameBound()
 {
-	UDebugSkelMeshComponent* PreviewComponent = GetPreviewScene()->GetPreviewMeshComponent();
-	if (PreviewComponent != NULL)
+	SetViewportSettingForEachDebugMesh([](UDebugSkelMeshComponent* PreviewMeshComponent)
 	{
-		PreviewComponent->UseInGameBounds(! PreviewComponent->IsUsingInGameBounds());
-	}
+		PreviewMeshComponent->UseInGameBounds(! PreviewMeshComponent->IsUsingInGameBounds());
+	});
 }
 
 bool SAnimationEditorViewportTabBody::CanUseInGameBound() const
 {
-	UDebugSkelMeshComponent* PreviewComponent = GetPreviewScene()->GetPreviewMeshComponent();
-	return PreviewComponent != NULL && IsShowBoundEnabled();
+	TArray<UDebugSkelMeshComponent*> PreviewMeshComponents = GetPreviewScene()->GetAllPreviewMeshComponents();
+	for (UDebugSkelMeshComponent* PreviewMeshComponent : PreviewMeshComponents)
+	{
+		if(IsShowBoundEnabled())
+		{
+			return true;
+		}
+	}
+	
+	return false;
 }
 
 bool SAnimationEditorViewportTabBody::IsUsingInGameBound() const
 {
-	UDebugSkelMeshComponent* PreviewComponent = GetPreviewScene()->GetPreviewMeshComponent();
-	return PreviewComponent != NULL && PreviewComponent->IsUsingInGameBounds();
+	TArray<UDebugSkelMeshComponent*> PreviewMeshComponents = GetPreviewScene()->GetAllPreviewMeshComponents();
+	for (UDebugSkelMeshComponent* PreviewMeshComponent : PreviewMeshComponents)
+	{
+		if(PreviewMeshComponent->IsUsingInGameBounds())
+		{
+			return true;
+		}
+	}
+	return false;
 }
 
-
 void SAnimationEditorViewportTabBody::UseFixedBounds()
 {
-	UDebugSkelMeshComponent* PreviewComponent = GetPreviewScene()->GetPreviewMeshComponent();
-	if (PreviewComponent != NULL)
+	SetViewportSettingForEachDebugMesh([](UDebugSkelMeshComponent* PreviewMeshComponent)
 	{
-		PreviewComponent->bComponentUseFixedSkelBounds = !PreviewComponent->bComponentUseFixedSkelBounds;
-	}
+		PreviewMeshComponent->bComponentUseFixedSkelBounds = !PreviewMeshComponent->bComponentUseFixedSkelBounds;
+	});
 }
 
 bool SAnimationEditorViewportTabBody::CanUseFixedBounds() const
@@ -1554,11 +1589,10 @@ bool SAnimationEditorViewportTabBody::IsUsingFixedBounds() const
 
 void SAnimationEditorViewportTabBody::UsePreSkinnedBounds()
 {
-	UDebugSkelMeshComponent* PreviewComponent = GetPreviewScene()->GetPreviewMeshComponent();
-	if (PreviewComponent != NULL)
+	GetPreviewScene()->ForEachPreviewMesh([](UDebugSkelMeshComponent* PreviewMeshComponent)
 	{
-		PreviewComponent->UsePreSkinnedBounds(!PreviewComponent->IsUsingPreSkinnedBounds());
-	}
+		PreviewMeshComponent->UsePreSkinnedBounds(!PreviewMeshComponent->IsUsingPreSkinnedBounds());
+	});
 }
 
 bool SAnimationEditorViewportTabBody::CanUsePreSkinnedBounds() const
diff --git a/Engine/Source/Editor/Persona/Private/SAnimationEditorViewport.h b/Engine/Source/Editor/Persona/Private/SAnimationEditorViewport.h
--- a/Engine/Source/Editor/Persona/Private/SAnimationEditorViewport.h
+++ b/Engine/Source/Editor/Persona/Private/SAnimationEditorViewport.h
@@ -289,6 +289,10 @@ public:
 	void SetGravityScale( float SliderPos );
 	float GetGravityScaleSliderValue() const;
 
+	/** Adjustable bone draw size */
+	void SetBoneDrawSize(float BoneDrawSize);
+	float GetBoneDrawSize() const;
+
 	/** Function to set LOD model selection*/
 	void OnSetLODModel(int32 LODSelectionType);
 	void OnSetLODTrackDebuggedInstance();
@@ -311,7 +315,7 @@ private:
 	 * Binds our UI commands to delegates
 	 */ 
 	void BindCommands();
-
+	
 	/** Show Morphtarget of SkeletalMesh **/
 	void OnShowMorphTargets();
 
@@ -373,7 +377,7 @@ private:
 
 	/** Function to check whether morphtarget overlay is displayed or not*/
 	bool IsShowingOverlayMorphTargetVerts() const;
-
+	
 	/** Function to set Local axes mode of the specificed type */
 	void OnSetBoneDrawMode(int32 BoneDrawMode);
 
@@ -422,6 +426,20 @@ private:
 	bool CanShowPreviewMesh() const;
 	bool IsShowPreviewMeshEnabled() const;
 
+	/** Run a lambda function on each preview mesh in the scene */
+	FORCEINLINE_DEBUGGABLE void SetViewportSettingForEachDebugMesh(TFunction<void (UDebugSkelMeshComponent*)> PerMeshFunction)
+	{
+		TArray<UDebugSkelMeshComponent*> PreviewMeshComponents;
+		GetPreviewScene()->GetActor()->GetComponents(PreviewMeshComponents, true);
+		for (UDebugSkelMeshComponent* PreviewMesh : PreviewMeshComponents)
+		{
+			PerMeshFunction(PreviewMesh);
+			PreviewMesh->MarkRenderStateDirty();
+		}
+		
+		RefreshViewport();
+	}
+
 	/** Called to toggle using in-game bound on current preview mesh */
 	void UseInGameBound();
 	bool CanUseInGameBound() const;
diff --git a/Engine/Source/Editor/Persona/Public/IPersonaPreviewScene.h b/Engine/Source/Editor/Persona/Public/IPersonaPreviewScene.h
--- a/Engine/Source/Editor/Persona/Public/IPersonaPreviewScene.h
+++ b/Engine/Source/Editor/Persona/Public/IPersonaPreviewScene.h
@@ -83,6 +83,12 @@ public:
 	/** Get the skeletal mesh component we are using for preview, if any. */
 	virtual UDebugSkelMeshComponent* GetPreviewMeshComponent() const = 0;
 
+	/** Get array of all skeletal mesh components in the preview scene. */
+	virtual TArray<UDebugSkelMeshComponent*> GetAllPreviewMeshComponents() const = 0;
+
+	/** Run a lambda function on each preview mesh in the scene */
+	virtual void ForEachPreviewMesh(TFunction<void (UDebugSkelMeshComponent*)> PerMeshFunction) = 0;
+
 	/** Set the skeletal mesh component we are going to preview. */
 	virtual void SetPreviewMeshComponent(UDebugSkelMeshComponent* InSkeletalMeshComponent) = 0;
 
diff --git a/Engine/Source/Editor/UnrealEd/Classes/Animation/DebugSkelMeshComponent.h b/Engine/Source/Editor/UnrealEd/Classes/Animation/DebugSkelMeshComponent.h
--- a/Engine/Source/Editor/UnrealEd/Classes/Animation/DebugSkelMeshComponent.h
+++ b/Engine/Source/Editor/UnrealEd/Classes/Animation/DebugSkelMeshComponent.h
@@ -145,21 +145,28 @@ private:
 	bool bSelectable;
 };
 
+/** Generic modes used to render debug skeletons depending on editor-specific context */
+UENUM()
+enum class ESkeletonDrawMode : uint8
+{
+	/** Bones are visible and selectable */
+	Default,
+
+	/** Bones are completely hidden */
+	Hidden,
+
+	/** Bones are visible but non-selectable */
+	GreyedOut
+};
+
 UCLASS(transient)
 class UNREALED_API UDebugSkelMeshComponent : public USkeletalMeshComponent
 {
 	GENERATED_UCLASS_BODY()
-
-	/** If true, render a wireframe skeleton of the mesh animated with the raw (uncompressed) animation data. */
-	UPROPERTY()
-	uint32 bRenderRawSkeleton:1;
-
-	/** Holds onto the bone color that will be used to render the bones of its skeletal mesh */
-	//var Color		BoneColor;
 	
-	/** If true then the skeletal mesh associated with the component is drawn. */
+	/** Global drawing mode for this skeleton. Depends on context of specific editor using the component. */
 	UPROPERTY()
-	uint32 bDrawMesh:1;
+	ESkeletonDrawMode SkeletonDrawMode = ESkeletonDrawMode::Default;
 
 	/** If true then the bone names associated with the skeletal mesh are displayed */
 	UPROPERTY()
diff --git a/Engine/Source/Editor/UnrealEd/Private/Animation/DebugSkelMeshComponent.cpp b/Engine/Source/Editor/UnrealEd/Private/Animation/DebugSkelMeshComponent.cpp
--- a/Engine/Source/Editor/UnrealEd/Private/Animation/DebugSkelMeshComponent.cpp
+++ b/Engine/Source/Editor/UnrealEd/Private/Animation/DebugSkelMeshComponent.cpp
@@ -30,7 +30,6 @@
 UDebugSkelMeshComponent::UDebugSkelMeshComponent(const FObjectInitializer& ObjectInitializer)
 	: Super(ObjectInitializer)
 {
-	bDrawMesh = true;
 	PreviewInstance = NULL;
 	bDisplayRawAnimation = false;
 	bDisplayNonRetargetedPose = false;
@@ -46,8 +45,6 @@ UDebugSkelMeshComponent::UDebugSkelMeshComponent(const FObjectInitializer& Objec
 
 	bTrackAttachedInstanceLOD = false;
 
-	BoneRadiusMultiplier = 1.f;
-
 	WireframeMeshOverlayColor = FLinearColor(0.4f, 0.8f, 0.66f);
 	
 	CachedClothBounds = FBoxSphereBounds(ForceInit);
@@ -1287,8 +1284,7 @@ void FDebugSkelMeshSceneProxy::GetDynamicMeshElements(const TArray<const FSceneV
 }
 
 FDebugSkelMeshDynamicData::FDebugSkelMeshDynamicData(UDebugSkelMeshComponent* InComponent)
-	: bDrawMesh(InComponent->bDrawMesh)
-	, bDrawNormals(InComponent->bDrawNormals)
+	: bDrawNormals(InComponent->bDrawNormals)
 	, bDrawTangents(InComponent->bDrawTangents)
 	, bDrawBinormals(InComponent->bDrawBinormals)
 	, bDrawClothPaintPreview(InComponent->bShowClothData)
diff --git a/Engine/Source/Runtime/Engine/Private/SkeletalDebugRendering.cpp b/Engine/Source/Runtime/Engine/Private/SkeletalDebugRendering.cpp
--- a/Engine/Source/Runtime/Engine/Private/SkeletalDebugRendering.cpp
+++ b/Engine/Source/Runtime/Engine/Private/SkeletalDebugRendering.cpp
@@ -7,12 +7,17 @@
 namespace SkeletalDebugRendering
 {
 
-void DrawWireBone(FPrimitiveDrawInterface* PDI, const FVector& InStart, const FVector& InEnd, const FLinearColor& InColor, ESceneDepthPriorityGroup InDepthPriority, const float SphereRadius)
+/** A fast and simple bone drawing function. This draws a sphere and a pyramid connection to the PARENT bone.
+ * Use this for basic debug drawing, but if the user is able to select or edit the bones, prefer DrawWireBoneAdvanced.*/
+void DrawWireBone(
+	FPrimitiveDrawInterface* PDI,
+	const FVector& InStart,
+	const FVector& InEnd,
+	const FLinearColor& InColor,
+	ESceneDepthPriorityGroup InDepthPriority,
+	const float SphereRadius)
 {
 #if ENABLE_DRAW_DEBUG
-	static const int32 NumSphereSides = 10;
-	static const int32 NumConeSides = 4;
-
 	// Calc cone size 
 	const FVector EndToStart = (InStart - InEnd);
 	const float ConeLength = EndToStart.Size();
@@ -22,11 +27,91 @@ void DrawWireBone(FPrimitiveDrawInterface* PDI, const FVector& InStart, const FV
 	DrawWireSphere(PDI, InEnd, InColor, SphereRadius, NumSphereSides, InDepthPriority, 0.0f, 1.0f);
 
 	TArray<FVector> Verts;
-	DrawWireCone(PDI, Verts, FRotationMatrix::MakeFromX(EndToStart) * FTranslationMatrix(InEnd), ConeLength, Angle, NumConeSides, InColor, InDepthPriority, 0.0f, 1.0f);
+	DrawWireCone(
+		PDI,
+		Verts,
+		FRotationMatrix::MakeFromX(EndToStart) * FTranslationMatrix(InEnd),
+		ConeLength,
+		Angle,
+		NumConeSides,
+		InColor,
+		InDepthPriority,
+		0.0f,
+		1.0f);
+#endif
+}
+
+/** An advanced bone drawing function for use with interactive editors where the user can select and manipulate bones.
+ *
+ * Differences from DrawWireBone() include:
+ * 1. Drawing all cone-connections to children as part of the "bone" itself so that the user can select the bone
+ *	  by clicking on any of it's children connections (as in all DCC applications)
+ * 2. Cone-connectors are drawn *between* spheres, not overlapping them (cleaner)
+ * 3. Bone sphere is oriented with bone rotation.
+ * 4. Connections to children can be colored individually to allow highlighting parent connections on selected children.
+ *
+ * This function, and the code required to structure the drawing in this manner, will incur some additional cost over
+ * DrawWireBone(). So in cases where you just want to debug draw a skeleton; with no option to select or manipulate
+ * the bones, it may be preferable to use DrawWireBone().
+ */
+void DrawWireBoneAdvanced(
+	FPrimitiveDrawInterface* PDI,
+	const FTransform& InBoneTransform,
+	const TArray<FVector>& InChildLocations,
+	const TArray<FLinearColor>& InChildColors,
+	const FLinearColor& InColor,
+	ESceneDepthPriorityGroup InDepthPriority,
+	const float SphereRadius,
+	const bool bDrawAxes)
+{
+#if ENABLE_DRAW_DEBUG
+
+	const FVector BoneLocation = InBoneTransform.GetLocation();
+
+	// draw wire sphere at joint origin, oriented with the bone
+	DrawWireSphere(PDI, InBoneTransform, InColor, SphereRadius, NumSphereSides, InDepthPriority, 0.0f, 1.0f);
+
+	// draw axes at joint location
+	if (bDrawAxes)
+	{
+		SkeletalDebugRendering::DrawAxes(PDI, InBoneTransform, SDPG_Foreground, 0.f , SphereRadius);
+	}
+
+	// draw wire cones to each child
+	for (int32 ChildIndex=0; ChildIndex<InChildLocations.Num(); ++ChildIndex)
+	{
+		const FVector& ChildPoint = InChildLocations[ChildIndex];
+		// offset start/end based on bone radius
+		const FVector RadiusOffset = (ChildPoint - BoneLocation).GetSafeNormal() * SphereRadius;
+		const FVector Start = BoneLocation + RadiusOffset;
+		const FVector End = ChildPoint - RadiusOffset;
+			
+		// calc cone size
+		const FVector EndToStart = (Start - End);
+		const float ConeLength = EndToStart.Size();
+		const float Angle = FMath::RadiansToDegrees(FMath::Atan(SphereRadius / ConeLength));
+		TArray<FVector> Verts;
+		DrawWireCone(
+			PDI,
+			Verts,
+			FRotationMatrix::MakeFromX(EndToStart) * FTranslationMatrix(End),
+			ConeLength,
+			Angle,
+			NumConeSides,
+			InChildColors[ChildIndex],
+			InDepthPriority,
+			0.0f,
+			1.0f);
+	}
 #endif
 }
 
-void DrawAxes(FPrimitiveDrawInterface* PDI, const FTransform& Transform, ESceneDepthPriorityGroup InDepthPriority, const float Thickness, const float AxisLength)
+void DrawAxes(
+	FPrimitiveDrawInterface* PDI,
+	const FTransform& Transform,
+	ESceneDepthPriorityGroup InDepthPriority,
+	const float Thickness,
+	const float AxisLength)
 {
 #if ENABLE_DRAW_DEBUG
 	// Display colored coordinate system axes for this joint.
@@ -49,4 +134,35 @@ void DrawAxes(FPrimitiveDrawInterface* PDI, const FTransform& Transform, ESceneD
 #endif
 }
 
+void DrawRootCone(
+	FPrimitiveDrawInterface* PDI,
+	const FTransform& InBoneTransform,
+	const FVector& ComponentOrigin,
+	const float SphereRadius)
+{
+#if ENABLE_DRAW_DEBUG
+	// offset start/end based on bone radius
+	const FVector RadiusOffset = (ComponentOrigin - InBoneTransform.GetLocation()).GetSafeNormal() * SphereRadius;
+	const FVector Start = InBoneTransform.GetLocation() + RadiusOffset;
+	const FVector End = ComponentOrigin;
+			
+	// calc cone size
+	const FVector EndToStart = (Start - End);
+	const float ConeLength = EndToStart.Size();
+	const float Angle = FMath::RadiansToDegrees(FMath::Atan(SphereRadius / ConeLength));
+	TArray<FVector> Verts;
+	DrawWireCone(
+		PDI,
+		Verts,
+		FRotationMatrix::MakeFromX(EndToStart) * FTranslationMatrix(End),
+		ConeLength,
+		Angle,
+		NumConeSides,
+		FLinearColor::Red,
+		ESceneDepthPriorityGroup::SDPG_Foreground,
+		0.0f,
+		1.0f);
+#endif
+}
+
 }
\ No newline at end of file
diff --git a/Engine/Source/Runtime/Engine/Public/SkeletalDebugRendering.h b/Engine/Source/Runtime/Engine/Public/SkeletalDebugRendering.h
--- a/Engine/Source/Runtime/Engine/Public/SkeletalDebugRendering.h
+++ b/Engine/Source/Runtime/Engine/Public/SkeletalDebugRendering.h
@@ -9,6 +9,15 @@ class FPrimitiveDrawInterface;
 
 namespace SkeletalDebugRendering
 {
+	static const FLinearColor DEFAULT_BONE_COLOR = FLinearColor(0.0f,0.0f,0.025f,1.0f);
+	static const FLinearColor SELECTED_BONE_COLOR = FLinearColor(0.2f,1.0f,0.2f,1.0f);
+	static const FLinearColor AFFECTED_BONE_COLOR = FLinearColor(1.0f,1.0f,1.0f,1.0f);
+	static const FLinearColor DISABLED_BONE_COLOR = FLinearColor(0.4f,0.4f,0.4f,1.0f);
+	static const FLinearColor PARENT_OF_SELECTED_BONE_COLOR = FLinearColor(0.85f,0.45f,0.12f,1.0f);
+	static const FLinearColor VIRTUAL_BONE_COLOR = FLinearColor(0.4f, 0.4f, 1.0f, 1.0f);
+
+	static const int32 NumSphereSides = 10;
+	static const int32 NumConeSides = 4;
 
 /**
  * Draw a wireframe bone from InStart to InEnd 
@@ -19,7 +28,13 @@ namespace SkeletalDebugRendering
  * @param	InDepthPriority		The scene depth priority group to use
  * @param	SphereRadius		SphereRadius
  */
-ENGINE_API void DrawWireBone(FPrimitiveDrawInterface* PDI, const FVector& InStart, const FVector& InEnd, const FLinearColor& InColor, ESceneDepthPriorityGroup InDepthPriority, const float SphereRadius = 1.f);
+ENGINE_API void DrawWireBone(
+	FPrimitiveDrawInterface* PDI,
+	const FVector& InStart,
+	const FVector& InEnd,
+	const FLinearColor& InColor,
+	ESceneDepthPriorityGroup InDepthPriority,
+	const float SphereRadius = 1.f);
 
 /**
  * Draw a set of axes to represent a transform
@@ -29,6 +44,44 @@ ENGINE_API void DrawWireBone(FPrimitiveDrawInterface* PDI, const FVector& InStar
  * @param	Thickness			Thickness of Axes
  * @param	AxisLength			AxisLength
  */
-ENGINE_API void DrawAxes(FPrimitiveDrawInterface* PDI, const FTransform& InTransform, ESceneDepthPriorityGroup InDepthPriority, const float Thickness = 0.f, const float AxisLength = 4.f);
+ENGINE_API void DrawAxes(
+	FPrimitiveDrawInterface* PDI,
+	const FTransform& InTransform,
+	ESceneDepthPriorityGroup InDepthPriority,
+	const float Thickness = 0.f,
+	const float AxisLength = 4.f);
+
+/**
+ * Draw a set of axes to represent a transform
+ * @param	PDI					Primitive draw interface to use
+ * @param	InBoneTransform		The bone transform
+ * @param	InChildLocations	The positions of all the children of the bone
+ * @param   InChildColors		The colors to use when drawing the cone to each child
+ * @param	InColor				The color to use for the bone
+ * @param	InDepthPriority		The scene depth priority group to use
+ * @param	SphereRadius		Radius of the ball drawn at the bone location
+ * @param	bDrawAxes			If true, will draw small coordinate axes inside the joint sphere
+ */
+ENGINE_API	void DrawWireBoneAdvanced(
+	FPrimitiveDrawInterface* PDI,
+	const FTransform& InBoneTransform,
+	const TArray<FVector>& InChildLocations,
+	const TArray<FLinearColor>& InChildColors,
+	const FLinearColor& InColor,
+	ESceneDepthPriorityGroup InDepthPriority,
+	const float SphereRadius,
+	const bool bDrawAxes);
 
-}
\ No newline at end of file
+/**
+ * Draw a red cone showing offset of root bone from the component origin
+ * @param	PDI					Primitive draw interface to use
+ * @param	InBoneTransform		The bone transform
+ * @param	ComponentOrigin		The position in world space of the component this bone lives within
+ * @param	SphereRadius		The radius of the root bone
+ */
+ENGINE_API	void DrawRootCone(
+	FPrimitiveDrawInterface* PDI,
+	const FTransform& InBoneTransform,
+	const FVector& ComponentOrigin,
+	const float SphereRadius);
+}
