Fixing editor linux build break from CL 20487861

diff --git a/Engine/Source/Editor/Persona/Private/AnimationEditorPreviewScene.h b/Engine/Source/Editor/Persona/Private/AnimationEditorPreviewScene.h
--- a/Engine/Source/Editor/Persona/Private/AnimationEditorPreviewScene.h
+++ b/Engine/Source/Editor/Persona/Private/AnimationEditorPreviewScene.h
@@ -33,7 +33,7 @@ public:
 	virtual void FocusViews() override;
 	virtual UDebugSkelMeshComponent* GetPreviewMeshComponent() const override { return SkeletalMeshComponent; }
 	virtual TArray<UDebugSkelMeshComponent*> GetAllPreviewMeshComponents() const override;
-	virtual FORCEINLINE_DEBUGGABLE void ForEachPreviewMesh(TFunction<void (UDebugSkelMeshComponent*)> PerMeshFunction) override;
+	virtual void ForEachPreviewMesh(TFunction<void (UDebugSkelMeshComponent*)> PerMeshFunction) override;
 	virtual void SetPreviewMeshComponent(UDebugSkelMeshComponent* InSkeletalMeshComponent) override;
 	virtual void SetAdditionalMeshes(class UDataAsset* InAdditionalMeshes) override;
 	virtual void SetAdditionalMeshesSelectable(bool bSelectable) override;
