Fix skeletal controls (and other bone references) not working in template animation blueprints

Removed ensure/log and replaced with a name-based fallback. Validation occurs on the node anyways based on whether the node needs to validate against a paricular skeleton or not

diff --git a/Engine/Source/Editor/Persona/Private/AnimGraphNodeDetails.cpp b/Engine/Source/Editor/Persona/Private/AnimGraphNodeDetails.cpp
--- a/Engine/Source/Editor/Persona/Private/AnimGraphNodeDetails.cpp
+++ b/Engine/Source/Editor/Persona/Private/AnimGraphNodeDetails.cpp
@@ -45,6 +45,7 @@
 #include "IPropertyAccessEditor.h"
 #include "Algo/Accumulate.h"
 #include "Kismet2/BlueprintEditorUtils.h"
+#include "Widgets/Input/SEditableTextBox.h"
 
 #define LOCTEXT_NAMESPACE "KismetNodeWithOptionalPinsDetails"
 
@@ -499,10 +500,26 @@ void FBoneReferenceCustomization::CustomizeHeader( TSharedRef<IPropertyHandle> S
 	}
 	else
 	{
-		// if this FBoneReference is used by some other Outers, this will fail	
-		// should warn programmers instead of silent fail
-		ensureAlways(!bEnsureOnInvalidSkeleton);
-		UE_LOG(LogAnimation, Warning, TEXT("FBoneReferenceCustomization::CustomizeHeader: SetEditableSkeleton failed to find an appropriate skeleton!"));
+		HeaderRow
+		.NameContent()
+		[
+			StructPropertyHandle->CreatePropertyNameWidget()
+		]
+		.ValueContent()
+		[
+			SNew(SEditableTextBox)
+			.ToolTipText(StructPropertyHandle->GetToolTipText())
+			.OnTextCommitted_Lambda([this](const FText& InText, ETextCommit::Type InTextCommit)
+			{
+				BoneNameProperty->SetValue(FName(InText.ToString()));
+			})
+			.Text_Lambda([this]()
+			{
+				FName Name;
+				BoneNameProperty->GetValue(Name);
+				return FText::FromName(Name);
+			})
+		];
 	}
 }
 
