Fixed crash customizing a node that no longer has a valid anim blueprint

diff --git a/Engine/Source/Editor/Persona/Private/AnimGraphNodeDetails.cpp b/Engine/Source/Editor/Persona/Private/AnimGraphNodeDetails.cpp
--- a/Engine/Source/Editor/Persona/Private/AnimGraphNodeDetails.cpp
+++ b/Engine/Source/Editor/Persona/Private/AnimGraphNodeDetails.cpp
@@ -96,7 +96,7 @@ void FAnimGraphNodeDetails::CustomizeDetails(class IDetailLayoutBuilder& DetailB
 		}
 	}
 
-	TargetSkeleton = AnimGraphNode->GetAnimBlueprint()->TargetSkeleton;
+	TargetSkeleton = AnimGraphNode->HasValidBlueprint() ? AnimGraphNode->GetAnimBlueprint()->TargetSkeleton : nullptr;
 	TargetSkeletonName = TargetSkeleton ? FObjectPropertyBase::GetExportPath(TargetSkeleton) : FString(TEXT(""));
 
 	// Get the node property
