Fixing static analysis issues in animation editor viewport.

diff --git a/Engine/Plugins/Animation/IKRig/Source/IKRigEditor/Private/RetargetEditor/SIKRetargetHierarchy.cpp b/Engine/Plugins/Animation/IKRig/Source/IKRigEditor/Private/RetargetEditor/SIKRetargetHierarchy.cpp
--- a/Engine/Plugins/Animation/IKRig/Source/IKRigEditor/Private/RetargetEditor/SIKRetargetHierarchy.cpp
+++ b/Engine/Plugins/Animation/IKRig/Source/IKRigEditor/Private/RetargetEditor/SIKRetargetHierarchy.cpp
@@ -197,7 +197,7 @@ void SIKRetargetHierarchy::RefreshTreeView(bool IsInitialSetup)
 
 	// validate we have a skeleton to load
 	const UIKRetargetProcessor* Processor = Controller->GetRetargetProcessor();
-	if (!Processor && Processor->IsInitialized())
+	if (!(Processor && Processor->IsInitialized()))
 	{
 		TreeView->RequestTreeRefresh();
 		return;
diff --git a/Engine/Source/Editor/Persona/Private/AnimationEditorViewportClient.cpp b/Engine/Source/Editor/Persona/Private/AnimationEditorViewportClient.cpp
--- a/Engine/Source/Editor/Persona/Private/AnimationEditorViewportClient.cpp
+++ b/Engine/Source/Editor/Persona/Private/AnimationEditorViewportClient.cpp
@@ -1394,24 +1394,24 @@ void FAnimationViewportClient::DrawBonesFromCompactPose(
 	const FCompactHeapPose& Pose,
 	UDebugSkelMeshComponent* MeshComponent,
 	FPrimitiveDrawInterface* PDI,
-	const FLinearColor& DrawColour) const
+	const FLinearColor& DrawColor) const
 {
 	if (Pose.GetNumBones() == 0 ||
 		!MeshComponent ||
 		!MeshComponent->SkeletalMesh ||
-		MeshComponent->SkeletonDrawMode != ESkeletonDrawMode::Hidden)
+		MeshComponent->SkeletonDrawMode == ESkeletonDrawMode::Hidden)
 	{
 		return;
 	}
 
 	// optionally override draw color
-	const FLinearColor BoneColor = MeshComponent->SkeletonDrawMode == ESkeletonDrawMode::GreyedOut ? SkeletalDebugRendering::DISABLED_BONE_COLOR : DrawColour;
+	const FLinearColor BoneColor = MeshComponent->SkeletonDrawMode == ESkeletonDrawMode::GreyedOut ? SkeletalDebugRendering::DISABLED_BONE_COLOR : DrawColor;
 	
 	TArray<FTransform> WorldTransforms;
 	WorldTransforms.AddUninitialized(Pose.GetBoneContainer().GetNumBones());
 
-	TArray<FLinearColor> BoneColours;
-	BoneColours.AddUninitialized(Pose.GetBoneContainer().GetNumBones());
+	TArray<FLinearColor> BoneColors;
+	BoneColors.AddUninitialized(Pose.GetBoneContainer().GetNumBones());
 
 	// we could cache parent bones as we calculate, but right now I'm not worried about perf issue of this
 	for (FCompactPoseBoneIndex BoneIndex : Pose.ForEachBoneIndex())
@@ -1428,7 +1428,7 @@ void FAnimationViewportClient::DrawBonesFromCompactPose(
 		{
 			WorldTransforms[MeshBoneIndex.GetInt()] = Pose[BoneIndex] * WorldTransforms[ParentIndex];
 		}
-		BoneColours[MeshBoneIndex.GetInt()] = BoneColor;
+		BoneColors[MeshBoneIndex.GetInt()] = BoneColor;
 	}
 
 	constexpr bool bForceDraw = true;
@@ -1440,7 +1440,7 @@ void FAnimationViewportClient::DrawBonesFromCompactPose(
 		MeshComponent->GetReferenceSkeleton(),
 		WorldTransforms,
 		MeshComponent->BonesOfInterest,
-		BoneColours,
+		BoneColors,
 		PDI,
 		bForceDraw,
 		bAddHitProxy);
diff --git a/Engine/Source/Editor/Persona/Private/AnimationEditorViewportClient.h b/Engine/Source/Editor/Persona/Private/AnimationEditorViewportClient.h
--- a/Engine/Source/Editor/Persona/Private/AnimationEditorViewportClient.h
+++ b/Engine/Source/Editor/Persona/Private/AnimationEditorViewportClient.h
@@ -392,7 +392,7 @@ private:
 	/** Draws the given array of transforms as bones */
 	void DrawBonesFromTransforms(TArray<FTransform>& Transforms, UDebugSkelMeshComponent * MeshComponent, FPrimitiveDrawInterface* PDI, FLinearColor BoneColour, FLinearColor RootBoneColour) const;
 	/** Draws Bones for a compact pose */
-	void DrawBonesFromCompactPose(const FCompactHeapPose& Pose, UDebugSkelMeshComponent * MeshComponent, FPrimitiveDrawInterface* PDI, const FLinearColor& DrawColour) const;
+	void DrawBonesFromCompactPose(const FCompactHeapPose& Pose, UDebugSkelMeshComponent * MeshComponent, FPrimitiveDrawInterface* PDI, const FLinearColor& DrawColor) const;
 	/** Draws Bones for uncompressed animation **/
 	void DrawMeshBonesUncompressedAnimation(UDebugSkelMeshComponent * MeshComponent, FPrimitiveDrawInterface* PDI) const;
 	/** Draw Bones for non retargeted animation. */
