Fixes animation viewport flicker issue.

diff --git a/Engine/Source/Editor/Persona/Private/SAnimationEditorViewport.cpp b/Engine/Source/Editor/Persona/Private/SAnimationEditorViewport.cpp
--- a/Engine/Source/Editor/Persona/Private/SAnimationEditorViewport.cpp
+++ b/Engine/Source/Editor/Persona/Private/SAnimationEditorViewport.cpp
@@ -1157,58 +1157,72 @@ int32 SAnimationEditorViewportTabBody::GetLODModelCount() const
 
 void SAnimationEditorViewportTabBody::OnShowMorphTargets()
 {
-	SetViewportSettingForEachDebugMesh([](UDebugSkelMeshComponent* InMesh)
+	ForEachDebugMesh([](UDebugSkelMeshComponent* InMesh)
 	{
 		InMesh->bDisableMorphTarget = !InMesh->bDisableMorphTarget;
+		InMesh->MarkRenderStateDirty();
 	});
+	RefreshViewport();
 }
 
 void SAnimationEditorViewportTabBody::OnShowBoneNames()
 {
-	SetViewportSettingForEachDebugMesh([](UDebugSkelMeshComponent* InMesh)
+	ForEachDebugMesh([](UDebugSkelMeshComponent* InMesh)
 	{
 		InMesh->bShowBoneNames = !InMesh->bShowBoneNames;
+		InMesh->MarkRenderStateDirty();
 	});
+	RefreshViewport();
 }
 
 void SAnimationEditorViewportTabBody::OnShowRawAnimation()
 {
-	SetViewportSettingForEachDebugMesh([](UDebugSkelMeshComponent* InMesh)
+	ForEachDebugMesh([](UDebugSkelMeshComponent* InMesh)
 	{
 		InMesh->bDisplayRawAnimation = !InMesh->bDisplayRawAnimation;
+		InMesh->MarkRenderStateDirty();
 	});
+	RefreshViewport();
 }
 
 void SAnimationEditorViewportTabBody::OnShowNonRetargetedAnimation()
 {
-	SetViewportSettingForEachDebugMesh([](UDebugSkelMeshComponent* InMesh)
+	ForEachDebugMesh([](UDebugSkelMeshComponent* InMesh)
 	{
 		InMesh->bDisplayNonRetargetedPose = !InMesh->bDisplayNonRetargetedPose;
+		InMesh->MarkRenderStateDirty();
 	});
+	RefreshViewport();
 }
 
 void SAnimationEditorViewportTabBody::OnShowSourceRawAnimation()
 {
-	SetViewportSettingForEachDebugMesh([](UDebugSkelMeshComponent* InMesh)
+	ForEachDebugMesh([](UDebugSkelMeshComponent* InMesh)
 	{
 		InMesh->bDisplaySourceAnimation = !InMesh->bDisplaySourceAnimation;
+		InMesh->MarkRenderStateDirty();
 	});
+	RefreshViewport();
 }
 
 void SAnimationEditorViewportTabBody::OnShowBakedAnimation()
 {
-	SetViewportSettingForEachDebugMesh([](UDebugSkelMeshComponent* InMesh)
+	ForEachDebugMesh([](UDebugSkelMeshComponent* InMesh)
 	{
 		InMesh->bDisplayBakedAnimation = !InMesh->bDisplayBakedAnimation;
+		InMesh->MarkRenderStateDirty();
 	});
+	RefreshViewport();
 }
 
 void SAnimationEditorViewportTabBody::OnShowAdditiveBase()
 {
-	SetViewportSettingForEachDebugMesh([](UDebugSkelMeshComponent* InMesh)
+	ForEachDebugMesh([](UDebugSkelMeshComponent* InMesh)
 	{
 		InMesh->bDisplayAdditiveBasePose = !InMesh->bDisplayAdditiveBasePose;
+		InMesh->MarkRenderStateDirty();
 	});
+	RefreshViewport();
 }
 
 bool SAnimationEditorViewportTabBody::IsPreviewingAnimation() const
@@ -1237,7 +1251,7 @@ bool SAnimationEditorViewportTabBody::IsShowingRawAnimation() const
 
 void SAnimationEditorViewportTabBody::OnToggleDisablePostProcess()
 {
-	SetViewportSettingForEachDebugMesh([](UDebugSkelMeshComponent* InMesh)
+	ForEachDebugMesh([](UDebugSkelMeshComponent* InMesh)
 	{
 		InMesh->ToggleDisablePostProcessBlueprint();
 	});
@@ -1308,13 +1322,15 @@ bool SAnimationEditorViewportTabBody::IsShowingMeshInfo(int32 DisplayInfoMode) c
 
 void SAnimationEditorViewportTabBody::OnShowOverlayNone()
 {
-	SetViewportSettingForEachDebugMesh([](UDebugSkelMeshComponent* PreviewMeshComponent)
+	ForEachDebugMesh([](UDebugSkelMeshComponent* PreviewMeshComponent)
 	{
 		PreviewMeshComponent->SetShowBoneWeight(false);
 		PreviewMeshComponent->SetShowMorphTargetVerts(false);
+		PreviewMeshComponent->MarkRenderStateDirty();
 	});
-	
+
 	UpdateShowFlagForMeshEdges();
+	RefreshViewport();
 }
 
 bool SAnimationEditorViewportTabBody::IsShowingOverlayNone() const
@@ -1325,13 +1341,14 @@ bool SAnimationEditorViewportTabBody::IsShowingOverlayNone() const
 
 void SAnimationEditorViewportTabBody::OnShowOverlayBoneWeight()
 {
-	SetViewportSettingForEachDebugMesh([](UDebugSkelMeshComponent* PreviewMeshComponent)
+	ForEachDebugMesh([](UDebugSkelMeshComponent* PreviewMeshComponent)
 	{
 		PreviewMeshComponent->SetShowBoneWeight( !PreviewMeshComponent->bDrawBoneInfluences );
 		PreviewMeshComponent->MarkRenderStateDirty();
 	});
 	
 	UpdateShowFlagForMeshEdges();
+	RefreshViewport();
 }
 
 bool SAnimationEditorViewportTabBody::IsShowingOverlayBoneWeight() const
@@ -1342,13 +1359,14 @@ bool SAnimationEditorViewportTabBody::IsShowingOverlayBoneWeight() const
 
 void SAnimationEditorViewportTabBody::OnShowOverlayMorphTargetVert()
 {
-	SetViewportSettingForEachDebugMesh([](UDebugSkelMeshComponent* PreviewMeshComponent)
+	ForEachDebugMesh([](UDebugSkelMeshComponent* PreviewMeshComponent)
 	{
 		PreviewMeshComponent->SetShowMorphTargetVerts(!PreviewMeshComponent->bDrawMorphTargetVerts);
 		PreviewMeshComponent->MarkRenderStateDirty();
 	});
 
 	UpdateShowFlagForMeshEdges();
+	RefreshViewport();
 }
 
 bool SAnimationEditorViewportTabBody::IsShowingOverlayMorphTargetVerts() const
@@ -1395,10 +1413,13 @@ bool SAnimationEditorViewportTabBody::IsLocalAxesModeSet(int32 LocalAxesMode) co
 
 void SAnimationEditorViewportTabBody::OnShowSockets()
 {
-	SetViewportSettingForEachDebugMesh([](UDebugSkelMeshComponent* PreviewMeshComponent)
+	ForEachDebugMesh([](UDebugSkelMeshComponent* PreviewMeshComponent)
 	{
 		PreviewMeshComponent->bDrawSockets = !PreviewMeshComponent->bDrawSockets;
+		PreviewMeshComponent->MarkRenderStateDirty();
 	});
+
+	RefreshViewport();
 }
 
 bool SAnimationEditorViewportTabBody::IsShowingSockets() const
@@ -1409,10 +1430,13 @@ bool SAnimationEditorViewportTabBody::IsShowingSockets() const
 
 void SAnimationEditorViewportTabBody::OnShowAttributes()
 {
-	SetViewportSettingForEachDebugMesh([](UDebugSkelMeshComponent* PreviewMeshComponent)
+	ForEachDebugMesh([](UDebugSkelMeshComponent* PreviewMeshComponent)
 	{
 		PreviewMeshComponent->bDrawAttributes = !PreviewMeshComponent->bDrawAttributes;
+		PreviewMeshComponent->MarkRenderStateDirty();
 	});
+	
+	RefreshViewport();
 }
 
 bool SAnimationEditorViewportTabBody::IsShowingAttributes() const
@@ -1448,7 +1472,7 @@ bool SAnimationEditorViewportTabBody::IsPlaybackSpeedSelected(int32 PlaybackSpee
 
 void SAnimationEditorViewportTabBody::ShowRetargetBasePose()
 {
-	SetViewportSettingForEachDebugMesh([](UDebugSkelMeshComponent* PreviewMeshComponent)
+	ForEachDebugMesh([](UDebugSkelMeshComponent* PreviewMeshComponent)
 	{
 		PreviewMeshComponent->PreviewInstance->SetForceRetargetBasePose(!PreviewMeshComponent->PreviewInstance->GetForceRetargetBasePose());
 	});
@@ -1486,7 +1510,7 @@ void SAnimationEditorViewportTabBody::ShowBound()
 	TSharedRef<FAnimationViewportClient> AnimViewportClient = StaticCastSharedRef<FAnimationViewportClient>(LevelViewportClient.ToSharedRef());	
 	AnimViewportClient->ToggleShowBounds();
 
-	SetViewportSettingForEachDebugMesh([AnimViewportClient](UDebugSkelMeshComponent* PreviewMeshComponent)
+	ForEachDebugMesh([AnimViewportClient](UDebugSkelMeshComponent* PreviewMeshComponent)
 	{
 		PreviewMeshComponent->bDisplayBound = AnimViewportClient->EngineShowFlags.Bounds;
 		PreviewMeshComponent->RecreateRenderState_Concurrent();
@@ -1507,7 +1531,7 @@ bool SAnimationEditorViewportTabBody::IsShowBoundEnabled() const
 void SAnimationEditorViewportTabBody::ToggleShowPreviewMesh()
 {
 	const bool bCurrentlyVisible = IsShowPreviewMeshEnabled();
-	SetViewportSettingForEachDebugMesh([bCurrentlyVisible](UDebugSkelMeshComponent* PreviewMeshComponent)
+	ForEachDebugMesh([bCurrentlyVisible](UDebugSkelMeshComponent* PreviewMeshComponent)
 	{
 		PreviewMeshComponent->SetVisibility(!bCurrentlyVisible);
 	});
@@ -1534,7 +1558,7 @@ bool SAnimationEditorViewportTabBody::IsShowPreviewMeshEnabled() const
 
 void SAnimationEditorViewportTabBody::UseInGameBound()
 {
-	SetViewportSettingForEachDebugMesh([](UDebugSkelMeshComponent* PreviewMeshComponent)
+	ForEachDebugMesh([](UDebugSkelMeshComponent* PreviewMeshComponent)
 	{
 		PreviewMeshComponent->UseInGameBounds(! PreviewMeshComponent->IsUsingInGameBounds());
 	});
@@ -1569,7 +1593,7 @@ bool SAnimationEditorViewportTabBody::IsUsingInGameBound() const
 
 void SAnimationEditorViewportTabBody::UseFixedBounds()
 {
-	SetViewportSettingForEachDebugMesh([](UDebugSkelMeshComponent* PreviewMeshComponent)
+	ForEachDebugMesh([](UDebugSkelMeshComponent* PreviewMeshComponent)
 	{
 		PreviewMeshComponent->bComponentUseFixedSkelBounds = !PreviewMeshComponent->bComponentUseFixedSkelBounds;
 	});
diff --git a/Engine/Source/Editor/Persona/Private/SAnimationEditorViewport.h b/Engine/Source/Editor/Persona/Private/SAnimationEditorViewport.h
--- a/Engine/Source/Editor/Persona/Private/SAnimationEditorViewport.h
+++ b/Engine/Source/Editor/Persona/Private/SAnimationEditorViewport.h
@@ -427,17 +427,14 @@ private:
 	bool IsShowPreviewMeshEnabled() const;
 
 	/** Run a lambda function on each preview mesh in the scene */
-	FORCEINLINE_DEBUGGABLE void SetViewportSettingForEachDebugMesh(TFunction<void (UDebugSkelMeshComponent*)> PerMeshFunction)
+	FORCEINLINE_DEBUGGABLE void ForEachDebugMesh(TFunction<void (UDebugSkelMeshComponent*)> PerMeshFunction)
 	{
 		TArray<UDebugSkelMeshComponent*> PreviewMeshComponents;
 		GetPreviewScene()->GetActor()->GetComponents(PreviewMeshComponents, true);
 		for (UDebugSkelMeshComponent* PreviewMesh : PreviewMeshComponents)
 		{
 			PerMeshFunction(PreviewMesh);
-			PreviewMesh->MarkRenderStateDirty();
 		}
-		
-		RefreshViewport();
 	}
 
 	/** Called to toggle using in-game bound on current preview mesh */
diff --git a/Engine/Source/Editor/UnrealEd/Classes/Animation/DebugSkelMeshComponent.h b/Engine/Source/Editor/UnrealEd/Classes/Animation/DebugSkelMeshComponent.h
--- a/Engine/Source/Editor/UnrealEd/Classes/Animation/DebugSkelMeshComponent.h
+++ b/Engine/Source/Editor/UnrealEd/Classes/Animation/DebugSkelMeshComponent.h
@@ -168,6 +168,10 @@ class UNREALED_API UDebugSkelMeshComponent : public USkeletalMeshComponent
 	UPROPERTY()
 	ESkeletonDrawMode SkeletonDrawMode = ESkeletonDrawMode::Default;
 
+	/** If true then the skeletal mesh associated with the component is drawn. */
+	UPROPERTY()
+	uint32 bDrawMesh:1;
+	
 	/** If true then the bone names associated with the skeletal mesh are displayed */
 	UPROPERTY()
 	uint32 bShowBoneNames:1;
diff --git a/Engine/Source/Editor/UnrealEd/Private/Animation/DebugSkelMeshComponent.cpp b/Engine/Source/Editor/UnrealEd/Private/Animation/DebugSkelMeshComponent.cpp
--- a/Engine/Source/Editor/UnrealEd/Private/Animation/DebugSkelMeshComponent.cpp
+++ b/Engine/Source/Editor/UnrealEd/Private/Animation/DebugSkelMeshComponent.cpp
@@ -30,6 +30,7 @@
 UDebugSkelMeshComponent::UDebugSkelMeshComponent(const FObjectInitializer& ObjectInitializer)
 	: Super(ObjectInitializer)
 {
+	bDrawMesh = true;
 	PreviewInstance = NULL;
 	bDisplayRawAnimation = false;
 	bDisplayNonRetargetedPose = false;
@@ -1284,7 +1285,8 @@ void FDebugSkelMeshSceneProxy::GetDynamicMeshElements(const TArray<const FSceneV
 }
 
 FDebugSkelMeshDynamicData::FDebugSkelMeshDynamicData(UDebugSkelMeshComponent* InComponent)
-	: bDrawNormals(InComponent->bDrawNormals)
+	: bDrawMesh(InComponent->bDrawMesh)
+	, bDrawNormals(InComponent->bDrawNormals)
 	, bDrawTangents(InComponent->bDrawTangents)
 	, bDrawBinormals(InComponent->bDrawBinormals)
 	, bDrawClothPaintPreview(InComponent->bShowClothData)
