Changes to Persona orbit camera
* Fixed asset editor viewport orbit camera dollying in the wrong direction after immediately rotating the camera in certain camera view rotations.
* Fixed camera translation (middle-mouse drag) when using orbit camera in Persona to match other editors.
* Focusing the camera in Persona (hotkey 'F') on nothing now maintains the camera's choice of follow mode.
* Moved 'Follow Bone' camera follow mode to its own submenu like other menus.
* Added toggle for 'Lock Camera Rotation' when following a bone.
* Added a new Camera Follow Mode called 'Follow Root' that follows the root bone but also keeps the mesh centred in the frame.
* Added 'Pause On Orbit' editor preference to Persona which pauses and resumes the preview animation when orbiting.

diff --git a/Engine/Source/Editor/Persona/Private/AnimViewportMenuCommands.cpp b/Engine/Source/Editor/Persona/Private/AnimViewportMenuCommands.cpp
--- a/Engine/Source/Editor/Persona/Private/AnimViewportMenuCommands.cpp
+++ b/Engine/Source/Editor/Persona/Private/AnimViewportMenuCommands.cpp
@@ -8,9 +8,12 @@ void FAnimViewportMenuCommands::RegisterCommands()
 {
 	UI_COMMAND( PreviewSceneSettings, "Preview Scene Settings...", "The Advanced Preview Settings tab will let you alter the preview scene's settings.", EUserInterfaceActionType::Button, FInputChord());
 
-	UI_COMMAND( CameraFollowNone, "Free Camera", "Camera is free to move", EUserInterfaceActionType::RadioButton, FInputChord() );
-	UI_COMMAND( CameraFollowBounds, "Follow Bounds", "Camera follows the bounbds of the mesh", EUserInterfaceActionType::RadioButton, FInputChord() );
-	UI_COMMAND( CameraFollowBone, "Follow Bone", "Camera follows a specified bone", EUserInterfaceActionType::RadioButton, FInputChord() );
+	UI_COMMAND( CameraFollowNone, "Free Camera", "Camera is free to move.", EUserInterfaceActionType::RadioButton, FInputChord() );
+	UI_COMMAND( CameraFollowBounds, "Orbit Bounds", "Camera orbits the bounds of the mesh.", EUserInterfaceActionType::RadioButton, FInputChord() );
+	UI_COMMAND( CameraFollowRoot, "Orbit Root", "Camera orbits the root bone while keeping the mesh vertically centered.", EUserInterfaceActionType::RadioButton, FInputChord() );
+	UI_COMMAND( CameraFollowBone, "Orbit Bone", "Camera focuses on a specified bone.", EUserInterfaceActionType::RadioButton, FInputChord() );
+	
+	UI_COMMAND( TogglePauseAnimationOnCameraMove, "Pause Animation On Camera Move", "Pause the preview animation when moving the camera and resume when finished.", EUserInterfaceActionType::ToggleButton, FInputChord() );
 
 	UI_COMMAND( SetCPUSkinning, "CPU Skinning", "Toggles display of CPU skinning in the Preview Pane.", EUserInterfaceActionType::ToggleButton, FInputChord());
 	UI_COMMAND( SetShowNormals, "Normals", "Toggles display of vertex normals in the Preview Pane.", EUserInterfaceActionType::ToggleButton, FInputChord() );
diff --git a/Engine/Source/Editor/Persona/Private/AnimViewportMenuCommands.h b/Engine/Source/Editor/Persona/Private/AnimViewportMenuCommands.h
--- a/Engine/Source/Editor/Persona/Private/AnimViewportMenuCommands.h
+++ b/Engine/Source/Editor/Persona/Private/AnimViewportMenuCommands.h
@@ -35,6 +35,12 @@ public:
 	/** Select camera follow mode to follow a named bone */
 	TSharedPtr< FUICommandInfo > CameraFollowBone;
 
+	/** Select camera follow mode to orbit the root bone while keeping the mesh vertically centered. */
+	TSharedPtr< FUICommandInfo > CameraFollowRoot;
+
+	/** Toggle whether or not to pause the preview animation when moving the camera */
+	TSharedPtr< FUICommandInfo > TogglePauseAnimationOnCameraMove;
+
 	/** Show vertex normals */
 	TSharedPtr< FUICommandInfo > SetCPUSkinning;
 
diff --git a/Engine/Source/Editor/Persona/Private/AnimationEditorViewportClient.cpp b/Engine/Source/Editor/Persona/Private/AnimationEditorViewportClient.cpp
--- a/Engine/Source/Editor/Persona/Private/AnimationEditorViewportClient.cpp
+++ b/Engine/Source/Editor/Persona/Private/AnimationEditorViewportClient.cpp
@@ -65,6 +65,7 @@ FAnimationViewportClient::FAnimationViewportClient(const TSharedRef<IPersonaPrev
 	: FEditorViewportClient(&InAssetEditorToolkit->GetEditorModeManager(), &InPreviewScene.Get(), StaticCastSharedRef<SEditorViewport>(InAnimationEditorViewport))
 	, PreviewScenePtr(InPreviewScene)
 	, AssetEditorToolkitPtr(InAssetEditorToolkit)
+	, bRotateCameraToFollowBone(false)
 	, AnimationPlaybackSpeedMode(EAnimationPlaybackSpeeds::Normal)
 	, bFocusOnDraw(false)
 	, bFocusUsingCustomCamera(false)
@@ -73,6 +74,8 @@ FAnimationViewportClient::FAnimationViewportClient(const TSharedRef<IPersonaPrev
 	, bInitiallyFocused(false)
 	, OrbitRotation(FQuat::Identity)
 	, ViewportIndex(InViewportIndex)
+	, LastLookAtLocation(FVector::ZeroVector)
+	, bResumeAfterTracking(false)
 {
 	CachedDefaultCameraController = CameraController;
 
@@ -246,6 +249,21 @@ bool FAnimationViewportClient::IsUsingAudioAttenuation() const
 	return ConfigOption->bUseAudioAttenuation;
 }
 
+void FAnimationViewportClient::ToggleRotateCameraToFollowBone()
+{
+	bRotateCameraToFollowBone = !bRotateCameraToFollowBone;
+
+	if (!bRotateCameraToFollowBone && GetCameraFollowMode() == EAnimationViewportCameraFollowMode::Bone)
+	{
+		OrbitRotation = FQuat::Identity;
+	}
+}
+
+bool FAnimationViewportClient::GetShouldRotateCameraToFollowBone() const
+{
+	return bRotateCameraToFollowBone;
+}
+
 void FAnimationViewportClient::SetCameraFollowMode(EAnimationViewportCameraFollowMode InCameraFollowMode, FName InBoneName)
 {
 	bool bCanFollow = true;
@@ -267,26 +285,34 @@ void FAnimationViewportClient::SetCameraFollowMode(EAnimationViewportCameraFollo
 
 		if (PreviewMeshComponent != nullptr)
 		{
+			FVector LookAtLocation = LastLookAtLocation;
+
 			switch(CameraFollowMode)
 			{
 			case EAnimationViewportCameraFollowMode::Bounds:
 				{
-					FBoxSphereBounds Bound = PreviewMeshComponent->CalcBounds(PreviewMeshComponent->GetComponentTransform());
-					SetLookAtLocation(Bound.Origin, true);
-					OrbitRotation = FQuat::Identity;
+					const FBoxSphereBounds Bounds = PreviewMeshComponent->CalcGameBounds(PreviewMeshComponent->GetComponentTransform());
+					LookAtLocation = Bounds.Origin;
+				}
+				break;
+			case EAnimationViewportCameraFollowMode::Root:
+				{
+					LookAtLocation = PreviewMeshComponent->GetBoneTransform(0).GetLocation();
+					LookAtLocation.Z = PreviewMeshComponent->CalcGameBounds(PreviewMeshComponent->GetComponentTransform()).Origin.Z;
 				}
 				break;
 			case EAnimationViewportCameraFollowMode::Bone:
 				{
-					FVector BoneLocation = PreviewMeshComponent->GetBoneLocation(InBoneName);
-					SetLookAtLocation(BoneLocation, true);
-					OrbitRotation = PreviewMeshComponent->GetBoneQuaternion(InBoneName) * FQuat(FVector(0.0f, 1.0f, 0.0f), PI * 0.5f);
+					LookAtLocation = PreviewMeshComponent->GetBoneLocation(InBoneName);
 				}
 				break;
 			}
-		}
 
-		SetViewLocation(GetViewTransform().ComputeOrbitMatrix().Inverse().GetOrigin());
+			OrbitRotation = FQuat::Identity;
+			SetLookAtLocation(LookAtLocation, true);
+			LastLookAtLocation = LookAtLocation;
+			bUsingOrbitCamera = true;
+		}
 	}
 	else
 	{
@@ -302,6 +328,35 @@ void FAnimationViewportClient::SetCameraFollowMode(EAnimationViewportCameraFollo
 	}
 }
 
+void FAnimationViewportClient::OnFocusViewportToSelection()
+{
+	// If focusing on a bone and using a Camera Follow Mode that orbits a bone, update the bone to follow to the selected bone
+	if (CameraFollowMode == EAnimationViewportCameraFollowMode::Root
+		|| CameraFollowMode == EAnimationViewportCameraFollowMode::Bone
+		|| CameraFollowMode == EAnimationViewportCameraFollowMode::Bounds )
+	{
+		int32 SelectedBoneIndex = GetAnimPreviewScene()->GetSelectedBoneIndex();
+		if (SelectedBoneIndex != INDEX_NONE)
+		{
+			const FReferenceSkeleton& ReferenceSkeleton = GetAnimPreviewScene()->GetPreviewMeshComponent()->GetReferenceSkeleton();
+			const FName SelectedBoneName = ReferenceSkeleton.GetBoneName(SelectedBoneIndex);
+			check(SelectedBoneName != NAME_None);
+			bRotateCameraToFollowBone = false;
+			SetCameraFollowMode(EAnimationViewportCameraFollowMode::Bone, SelectedBoneName);
+		}
+		else
+		{
+			SetCameraFollowMode(EAnimationViewportCameraFollowMode::Root, FName());
+		}
+	}
+	else
+	{
+		SetCameraFollowMode(EAnimationViewportCameraFollowMode::None);
+		FocusViewportOnPreviewMesh(false);
+	}
+	
+}
+
 EAnimationViewportCameraFollowMode FAnimationViewportClient::GetCameraFollowMode() const
 {
 	return CameraFollowMode;
@@ -589,67 +644,65 @@ void FAnimationViewportClient::Tick(float DeltaSeconds)
 
 void FAnimationViewportClient::HandlePreviewScenePreTick()
 {
-	RelativeViewLocation = FVector::ZeroVector;
+}
 
+void FAnimationViewportClient::HandlePreviewScenePostTick()
+{
 	UDebugSkelMeshComponent* PreviewMeshComponent = GetAnimPreviewScene()->GetPreviewMeshComponent();
-	if (CameraFollowMode != EAnimationViewportCameraFollowMode::None && PreviewMeshComponent != nullptr)
+
+	const bool bInBoneOrbitMode = CameraFollowMode != EAnimationViewportCameraFollowMode::None;
+	if (IsTracking())
 	{
-		switch(CameraFollowMode)
+		if (PreviewMeshComponent->IsPlaying() && GetDefault<UPersonaOptions>()->bPauseAnimationOnCameraMove)
 		{
-		case EAnimationViewportCameraFollowMode::Bounds:
-			{
-				FBoxSphereBounds Bound = PreviewMeshComponent->CalcBounds(PreviewMeshComponent->GetComponentTransform());
-				RelativeViewLocation = Bound.Origin - GetViewLocation();
-			}
-			break;
-		case EAnimationViewportCameraFollowMode::Bone:
-			{
-				int32 BoneIndex = PreviewMeshComponent->GetBoneIndex(CameraFollowBoneName);
-				if(BoneIndex != INDEX_NONE)
-				{
-					FTransform BoneTransform = PreviewMeshComponent->GetBoneTransform(BoneIndex);
-					RelativeViewLocation = BoneTransform.InverseTransformVector(BoneTransform.GetLocation() - GetViewLocation());
-				}
-			}
-			break;
+			PreviewMeshComponent->Stop();
+			bResumeAfterTracking = true;
+			return;
 		}
 	}
-}
+	else if (bResumeAfterTracking)
+	{
+		PreviewMeshComponent->Play(true);
+		bResumeAfterTracking = false;
+	}
 
-void FAnimationViewportClient::HandlePreviewScenePostTick()
-{
-	UDebugSkelMeshComponent* PreviewMeshComponent = GetAnimPreviewScene()->GetPreviewMeshComponent();
 	if (CameraFollowMode != EAnimationViewportCameraFollowMode::None && PreviewMeshComponent != nullptr)
 	{
+		FVector LookAtLocation = LastLookAtLocation;
+
 		switch(CameraFollowMode)
 		{
 		case EAnimationViewportCameraFollowMode::Bounds:
 			{
-				FBoxSphereBounds Bound = PreviewMeshComponent->CalcBounds(PreviewMeshComponent->GetComponentTransform());
-				SetViewLocation(Bound.Origin + RelativeViewLocation);
-				SetLookAtLocation(Bound.Origin);
+				const FBoxSphereBounds Bounds = PreviewMeshComponent->CalcGameBounds(PreviewMeshComponent->GetComponentTransform());
+				LookAtLocation = Bounds.Origin;
+			}
+			break;
+		case EAnimationViewportCameraFollowMode::Root:
+			{
+				LookAtLocation = PreviewMeshComponent->GetBoneTransform(0).GetLocation();
+				LookAtLocation.Z = PreviewMeshComponent->CalcGameBounds(PreviewMeshComponent->GetComponentTransform()).Origin.Z;
 			}
 			break;
 		case EAnimationViewportCameraFollowMode::Bone:
 			{
-				int32 BoneIndex = PreviewMeshComponent->GetBoneIndex(CameraFollowBoneName);
-				if(BoneIndex != INDEX_NONE)
+				const int32 BoneIndex = PreviewMeshComponent->GetBoneIndex(CameraFollowBoneName);
+				check(BoneIndex != INDEX_NONE);
+				LookAtLocation = PreviewMeshComponent->GetBoneTransform(BoneIndex).GetLocation();
+
+				if (GetShouldRotateCameraToFollowBone())
 				{
-					FTransform BoneTransform = PreviewMeshComponent->GetBoneTransform(BoneIndex);
-					if(IsPerspective())
-					{
-						SetViewLocation(BoneTransform.GetLocation() - BoneTransform.TransformVector(RelativeViewLocation));
-					}
-					else
-					{
-						SetViewLocation(BoneTransform.GetLocation());
-					}
-					SetLookAtLocation(BoneTransform.GetLocation());
-					OrbitRotation = BoneTransform.GetRotation() * FQuat(FVector(0.0f, 1.0f, 0.0f), PI * 0.5f);
+					const FName BoneName = PreviewMeshComponent->GetBoneName(BoneIndex);
+					check(BoneName != NAME_None);
+					OrbitRotation = PreviewMeshComponent->GetBoneQuaternion(BoneName) * FQuat(FVector(0.0f, 1.0f, 0.0f), PI * 0.5f);
 				}
 			}
 			break;
 		}
+
+		const FVector Offset = LookAtLocation - LastLookAtLocation;
+		SetLookAtLocation(GetLookAtLocation() + Offset);
+		LastLookAtLocation = LookAtLocation;
 	}
 }
 
@@ -1238,6 +1291,11 @@ void FAnimationViewportClient::SetViewportType(ELevelViewportType InViewportType
 	{
 		bUsingOrbitCamera = true;
 	}
+
+	if (InViewportType != ELevelViewportType::LVT_Perspective)
+	{
+		SetCameraFollowMode(EAnimationViewportCameraFollowMode::None);
+	}
 }
 
 void FAnimationViewportClient::RotateViewportType()
@@ -1894,9 +1952,7 @@ FSphere FAnimationViewportClient::GetCameraTarget()
 		return DefaultSphere;
 	}
 
-	PreviewMeshComponent->CalcBounds(PreviewMeshComponent->GetComponentTransform());
-
-	FBoxSphereBounds Bounds = PreviewMeshComponent->CalcBounds(FTransform::Identity);
+	FBoxSphereBounds Bounds = PreviewMeshComponent->CalcGameBounds(FTransform::Identity);
 	return Bounds.GetSphere();
 }
 
diff --git a/Engine/Source/Editor/Persona/Private/AnimationEditorViewportClient.h b/Engine/Source/Editor/Persona/Private/AnimationEditorViewportClient.h
--- a/Engine/Source/Editor/Persona/Private/AnimationEditorViewportClient.h
+++ b/Engine/Source/Editor/Persona/Private/AnimationEditorViewportClient.h
@@ -138,6 +138,9 @@ public:
 	/** Set the camera follow mode */
 	void SetCameraFollowMode(EAnimationViewportCameraFollowMode Mode, FName InBoneName = NAME_None);
 
+	/** Called when viewport focuses on a selection */
+	void OnFocusViewportToSelection();
+
 	/** Get the camera follow mode */
 	EAnimationViewportCameraFollowMode GetCameraFollowMode() const;
 
@@ -232,6 +235,11 @@ public:
 	/** Callback for checking the normals show flag. */
 	bool IsSetCPUSkinningChecked() const;
 
+	/** Toggles whether to lock the camera's rotation to a specified bone's orientation */
+	void ToggleRotateCameraToFollowBone();
+
+	/** Whether or not to lock the camera's rotation to a specified bone's orientation */
+	bool GetShouldRotateCameraToFollowBone() const;
 
 	/** Callback for toggling the normals show flag. */
 	void ToggleShowNormals();
@@ -347,6 +355,9 @@ private:
 	/** Should we auto align floor to mesh bounds */
 	bool bAutoAlignFloor;
 
+	/** Whether to lock the camera's rotation to a specified bone's orientation */
+	bool bRotateCameraToFollowBone;
+
 	/** User selected color using color picker */
 	FLinearColor SelectedHSVColor;
 
@@ -475,9 +486,12 @@ private:
 	/** Index (0-3) of this viewport */
 	int32 ViewportIndex;
 
-	/** Relative view location stored to match it pre/post tick */
-	FVector RelativeViewLocation;
+	/** The last location the camera was told to look at */
+	FVector LastLookAtLocation;
 
 	// Delegate Handler to allow changing of camera controller
 	void OnCameraControllerChanged();
+
+	/** True when the preview animation should resume playing upon finishing tracking */
+	bool bResumeAfterTracking;
 };
diff --git a/Engine/Source/Editor/Persona/Private/SAnimViewportToolBar.cpp b/Engine/Source/Editor/Persona/Private/SAnimViewportToolBar.cpp
--- a/Engine/Source/Editor/Persona/Private/SAnimViewportToolBar.cpp
+++ b/Engine/Source/Editor/Persona/Private/SAnimViewportToolBar.cpp
@@ -372,7 +372,6 @@ void SAnimViewportToolBar::Construct(const FArguments& InArgs, TSharedPtr<class
 	// Register all the custom widgets we can use here
 	PinnedCommands->RegisterCustomWidget(IPinnedCommandList::FOnGenerateCustomWidget::CreateSP(this, &SAnimViewportToolBar::MakeFloorOffsetWidget), TEXT("FloorOffsetWidget"), LOCTEXT("FloorHeightOffset", "Floor Height Offset"));
 	PinnedCommands->RegisterCustomWidget(IPinnedCommandList::FOnGenerateCustomWidget::CreateSP(this, &SAnimViewportToolBar::MakeFOVWidget), TEXT("FOVWidget"), LOCTEXT("Viewport_FOVLabel", "Field Of View"));
-	PinnedCommands->RegisterCustomWidget(IPinnedCommandList::FOnGenerateCustomWidget::CreateSP(this, &SAnimViewportToolBar::MakeFollowBoneComboWidget), TEXT("FollowBoneWidget"), LOCTEXT("FollowBoneMenuTitle", "Follow Bone"), FMargin(2.0f, 1.0f, 2.0f, 0.0f), false);
 
 	PinnedCommands->BindCommandList(InViewport->GetCommandList().ToSharedRef());
 
@@ -454,7 +453,7 @@ TSharedRef<SWidget> SAnimViewportToolBar::MakeFollowBoneComboWidget() const
 				}
 				else
 				{
-					return LOCTEXT("FollowBoneMenuTitle", "Follow Bone");
+					return LOCTEXT("FollowBoneMenuTitle", "Focus On Bone");
 				}
 			})
 		];
@@ -479,7 +478,6 @@ TSharedRef<SWidget> SAnimViewportToolBar::MakeFollowBoneWidget(TWeakPtr<SComboBu
 		.MaxDesiredHeight(400.0f)
 		[
 			SAssignNew(BoneTreeMenu, SBoneTreeMenu)
-			.Title(FAnimViewportMenuCommands::Get().CameraFollowBone->GetLabel())
 			.bShowVirtualBones(true)
 			.OnBoneSelectionChanged_Lambda([this](FName InBoneName)
 			{
@@ -553,9 +551,11 @@ TSharedRef<SWidget> SAnimViewportToolBar::GenerateViewMenu() const
 	InMenuBuilder.BeginSection("AnimViewportCamera", LOCTEXT("ViewMenu_CameraLabel", "Camera"));
 	{
 		InMenuBuilder.AddMenuEntry(FEditorViewportCommands::Get().FocusViewportToSelection);
+		InMenuBuilder.AddWidget(MakeFOVWidget(), LOCTEXT("Viewport_FOVLabel", "Field Of View"));
+		InMenuBuilder.AddMenuEntry(FAnimViewportMenuCommands::Get().TogglePauseAnimationOnCameraMove);
 
 		InMenuBuilder.AddSubMenu(
-			LOCTEXT("CameraFollowModeLabel", "Camera Follow Mode"),
+			LOCTEXT("CameraFollowModeLabel", "Follow Mode"),
 			LOCTEXT("CameraFollowModeTooltip", "Set various camera follow modes"),
 			FNewMenuDelegate::CreateLambda([this](FMenuBuilder& InSubMenuBuilder)
 			{
@@ -564,23 +564,57 @@ TSharedRef<SWidget> SAnimViewportToolBar::GenerateViewMenu() const
 					InSubMenuBuilder.PushCommandList(Viewport.Pin()->GetCommandList().ToSharedRef());
 
 					InSubMenuBuilder.AddMenuEntry(FAnimViewportMenuCommands::Get().CameraFollowNone);
+					InSubMenuBuilder.AddMenuEntry(FAnimViewportMenuCommands::Get().CameraFollowRoot);
 					InSubMenuBuilder.AddMenuEntry(FAnimViewportMenuCommands::Get().CameraFollowBounds);
-
+					InSubMenuBuilder.AddSubMenu(
+						LOCTEXT("CameraFollowBone_DisplayName", "Orbit Bone"),
+						LOCTEXT("CameraFollowBone_ToolTip", "Select a bone for the camera to follow and orbit around"),
+						FNewMenuDelegate::CreateLambda([=](FMenuBuilder& SubMenuBuilder)
+						{
+							SubMenuBuilder.BeginSection("CameraFollowModeBoneSubmenu", LOCTEXT("CameraFollowModeBoneSubmenu_Label", "Follow Bone Options"));
+							SubMenuBuilder.AddWidget(MakeFollowBoneWidget(), FText());
+							SubMenuBuilder.AddMenuEntry(
+								LOCTEXT("LockRotation_DisplayName", "Lock Rotation"),
+								LOCTEXT("LockRotation_ToolTip", "Keep viewport camera rotation aligned to the orbited bone."),
+								FSlateIcon(),
+								FUIAction(
+									FExecuteAction::CreateLambda([this]() {
+										Viewport.Pin()->ToggleRotateCameraToFollowBone();
+									}),
+									FCanExecuteAction(),
+									FIsActionChecked::CreateLambda([this]()
+									{
+										return Viewport.Pin()->GetShouldRotateCameraToFollowBone();
+									})
+								),
+								NAME_None,
+								EUserInterfaceActionType::ToggleButton
+							);
+							SubMenuBuilder.EndSection();
+						}),
+						FUIAction(
+							FExecuteAction(),
+							FCanExecuteAction::CreateLambda([this]()
+							{
+								return Viewport.Pin()->CanChangeCameraMode();
+							}),
+							FIsActionChecked::CreateLambda([this]()
+							{
+								return Viewport.Pin()->IsCameraFollowEnabled(EAnimationViewportCameraFollowMode::Bone);
+							})
+						),
+						"CameraFollowBone",
+						EUserInterfaceActionType::RadioButton,
+						/* bInOpenSubMenuOnClick = */ false,
+						FSlateIcon()
+					);
 					InSubMenuBuilder.PopCommandList();
 				}
 				InSubMenuBuilder.EndSection();
-		
-				InSubMenuBuilder.BeginSection("AnimViewportCameraFollowBone", FText());
-				{
-					InSubMenuBuilder.AddWidget(MakeFollowBoneWidget(), FText(), true);
-				}
-				InSubMenuBuilder.EndSection();
 			}),
 			false,
 			FSlateIcon(FAppStyle::GetAppStyleSetName(), "AnimViewportMenu.CameraFollow")
 			);
-
-		InMenuBuilder.AddWidget(MakeFOVWidget(), LOCTEXT("Viewport_FOVLabel", "Field Of View"));
 	}
 	InMenuBuilder.EndSection();
 
diff --git a/Engine/Source/Editor/Persona/Private/SAnimationEditorViewport.cpp b/Engine/Source/Editor/Persona/Private/SAnimationEditorViewport.cpp
--- a/Engine/Source/Editor/Persona/Private/SAnimationEditorViewport.cpp
+++ b/Engine/Source/Editor/Persona/Private/SAnimationEditorViewport.cpp
@@ -156,8 +156,7 @@ void SAnimationEditorViewport::PostRedo( bool bSuccess )
 void SAnimationEditorViewport::OnFocusViewportToSelection()
 {
 	TSharedRef<FAnimationViewportClient> AnimViewportClient = StaticCastSharedRef<FAnimationViewportClient>(LevelViewportClient.ToSharedRef());
-	AnimViewportClient->SetCameraFollowMode(EAnimationViewportCameraFollowMode::None);
-	AnimViewportClient->FocusViewportOnPreviewMesh(false);
+	AnimViewportClient->OnFocusViewportToSelection();
 }
 
 void SAnimationEditorViewport::BindCommands()
@@ -647,6 +646,12 @@ void SAnimationEditorViewportTabBody::BindCommands()
 	//Bind menu commands
 	const FAnimViewportMenuCommands& MenuActions = FAnimViewportMenuCommands::Get();
 
+	CommandList.MapAction(
+		MenuActions.TogglePauseAnimationOnCameraMove,
+		FExecuteAction::CreateSP(this, &SAnimationEditorViewportTabBody::TogglePauseAnimationOnCameraMove),
+		FCanExecuteAction(),
+		FIsActionChecked::CreateSP(this, &SAnimationEditorViewportTabBody::GetShouldPauseAnimationOnCameraMove));
+
 	CommandList.MapAction(
 		MenuActions.CameraFollowNone,
 		FExecuteAction::CreateSP(this, &SAnimationEditorViewportTabBody::SetCameraFollowMode, EAnimationViewportCameraFollowMode::None, FName()),
@@ -659,6 +664,12 @@ void SAnimationEditorViewportTabBody::BindCommands()
 		FCanExecuteAction::CreateSP(this, &SAnimationEditorViewportTabBody::CanChangeCameraMode),
 		FIsActionChecked::CreateSP(this, &SAnimationEditorViewportTabBody::IsCameraFollowEnabled, EAnimationViewportCameraFollowMode::Bounds));
 
+	CommandList.MapAction(
+		MenuActions.CameraFollowRoot,
+		FExecuteAction::CreateSP(this, &SAnimationEditorViewportTabBody::SetCameraFollowMode, EAnimationViewportCameraFollowMode::Root, FName()),
+		FCanExecuteAction::CreateSP(this, &SAnimationEditorViewportTabBody::CanChangeCameraMode),
+		FIsActionChecked::CreateSP(this, &SAnimationEditorViewportTabBody::IsCameraFollowEnabled, EAnimationViewportCameraFollowMode::Root));
+
 	CommandList.MapAction(
 		MenuActions.JumpToDefaultCamera,
 		FExecuteAction::CreateSP(this, &SAnimationEditorViewportTabBody::JumpToDefaultCamera),
@@ -1957,6 +1968,28 @@ bool SAnimationEditorViewportTabBody::IsCameraFollowEnabled(EAnimationViewportCa
 	return (AnimViewportClient->GetCameraFollowMode() == InCameraFollowMode);
 }
 
+void SAnimationEditorViewportTabBody::ToggleRotateCameraToFollowBone()
+{
+	TSharedRef<FAnimationViewportClient> AnimViewportClient = StaticCastSharedRef<FAnimationViewportClient>(LevelViewportClient.ToSharedRef());
+	AnimViewportClient->ToggleRotateCameraToFollowBone();
+}
+
+bool SAnimationEditorViewportTabBody::GetShouldRotateCameraToFollowBone() const
+{
+	TSharedRef<FAnimationViewportClient> AnimViewportClient = StaticCastSharedRef<FAnimationViewportClient>(LevelViewportClient.ToSharedRef());
+	return AnimViewportClient->GetShouldRotateCameraToFollowBone();
+}
+
+void SAnimationEditorViewportTabBody::TogglePauseAnimationOnCameraMove()
+{
+	GetMutableDefault<UPersonaOptions>()->bPauseAnimationOnCameraMove = !GetMutableDefault<UPersonaOptions>()->bPauseAnimationOnCameraMove;
+}
+
+bool SAnimationEditorViewportTabBody::GetShouldPauseAnimationOnCameraMove() const
+{
+	return GetMutableDefault<UPersonaOptions>()->bPauseAnimationOnCameraMove;
+}
+
 FName SAnimationEditorViewportTabBody::GetCameraFollowBoneName() const
 {
 	TSharedRef<FAnimationViewportClient> AnimViewportClient = StaticCastSharedRef<FAnimationViewportClient>(LevelViewportClient.ToSharedRef());
@@ -2616,7 +2649,7 @@ void SAnimationEditorViewportTabBody::AddSkinWeightProfileNotification()
 void SAnimationEditorViewportTabBody::HandleFocusCamera()
 {
 	TSharedRef<FAnimationViewportClient> AnimViewportClient = StaticCastSharedRef<FAnimationViewportClient>(LevelViewportClient.ToSharedRef());
-	AnimViewportClient->SetCameraFollowMode(EAnimationViewportCameraFollowMode::None);
+	// AnimViewportClient->SetCameraFollowMode(EAnimationViewportCameraFollowMode::None);
 	AnimViewportClient->FocusViewportOnPreviewMesh(false);
 }
 
diff --git a/Engine/Source/Editor/Persona/Private/SAnimationEditorViewport.h b/Engine/Source/Editor/Persona/Private/SAnimationEditorViewport.h
--- a/Engine/Source/Editor/Persona/Private/SAnimationEditorViewport.h
+++ b/Engine/Source/Editor/Persona/Private/SAnimationEditorViewport.h
@@ -481,9 +481,11 @@ private:
 	/** Focus the viewport on the preview mesh */
 	void HandleFocusCamera();
 
+public:
 	/** Called to determine whether the camera mode menu options should be enabled */
 	bool CanChangeCameraMode() const;
 
+private:
 	/** Tests to see if bone move mode buttons should be visible */
 	EVisibility GetBoneMoveModeButtonVisibility() const;
 
@@ -525,7 +527,10 @@ public:
 	void SetCameraFollowMode(EAnimationViewportCameraFollowMode InCameraFollowMode, FName InBoneName);
 	bool IsCameraFollowEnabled(EAnimationViewportCameraFollowMode InCameraFollowMode) const;
 	FName GetCameraFollowBoneName() const;
-
+	void ToggleRotateCameraToFollowBone();
+	bool GetShouldRotateCameraToFollowBone() const;
+	void TogglePauseAnimationOnCameraMove();
+	bool GetShouldPauseAnimationOnCameraMove() const;
 	bool IsTurnTableSpeedSelected(int32 SpeedIndex) const;
 
 	/** 
@@ -566,6 +571,9 @@ private:
 	/** Whether we should always show the transform toolbar for this viewport */
 	bool bAlwaysShowTransformToolbar;
 
+	/** Whether to align the camera's rotation to the bone's orientation */
+	bool bCameraFollowLockRotation;
+
 	/** Level viewport client */
 	TSharedPtr<FEditorViewportClient> LevelViewportClient;
 
diff --git a/Engine/Source/Editor/UnrealEd/Classes/Animation/DebugSkelMeshComponent.h b/Engine/Source/Editor/UnrealEd/Classes/Animation/DebugSkelMeshComponent.h
--- a/Engine/Source/Editor/UnrealEd/Classes/Animation/DebugSkelMeshComponent.h
+++ b/Engine/Source/Editor/UnrealEd/Classes/Animation/DebugSkelMeshComponent.h
@@ -418,6 +418,9 @@ class UNREALED_API UDebugSkelMeshComponent : public USkeletalMeshComponent
 	 */
 	bool CheckIfBoundsAreCorrrect();
 
+	/** Get the in-game bounds of the skeleton mesh */
+	FBoxSphereBounds CalcGameBounds(const FTransform& LocalToWorld) const;
+
 	/** 
 	 * Update components position based on animation root motion
 	 */
diff --git a/Engine/Source/Editor/UnrealEd/Classes/Preferences/PersonaOptions.h b/Engine/Source/Editor/UnrealEd/Classes/Preferences/PersonaOptions.h
--- a/Engine/Source/Editor/UnrealEd/Classes/Preferences/PersonaOptions.h
+++ b/Engine/Source/Editor/UnrealEd/Classes/Preferences/PersonaOptions.h
@@ -30,6 +30,9 @@ enum class EAnimationViewportCameraFollowMode : uint8
 
 	/** Follow a bone or socket */
 	Bone,
+
+	/** Follow the root bone while keeping the mesh vertically centered */
+	Root
 };
 
 /** Persistent per-viewport options */
@@ -152,6 +155,10 @@ class UNREALED_API UPersonaOptions : public UObject
 	UPROPERTY(EditAnywhere, config, Category = "Composites and Montages")
 	FLinearColor BranchingPointTimingNodeColor;
 
+	/** Pause the preview animation if playing when moving the camera and resume when finished */
+	UPROPERTY(EditAnywhere, config, Category = "Viewport")
+	bool bPauseAnimationOnCameraMove;
+
 	/** Whether to use a socket editor that is created in-line inside the skeleton tree, or whether to use the separate details panel */
 	UPROPERTY(EditAnywhere, config, Category = "Skeleton Tree")
 	bool bUseInlineSocketEditor;
diff --git a/Engine/Source/Editor/UnrealEd/Private/Animation/DebugSkelMeshComponent.cpp b/Engine/Source/Editor/UnrealEd/Private/Animation/DebugSkelMeshComponent.cpp
--- a/Engine/Source/Editor/UnrealEd/Private/Animation/DebugSkelMeshComponent.cpp
+++ b/Engine/Source/Editor/UnrealEd/Private/Animation/DebugSkelMeshComponent.cpp
@@ -111,6 +111,11 @@ FBoxSphereBounds UDebugSkelMeshComponent::CalcBounds(const FTransform& LocalToWo
 	return Result;
 }
 
+FBoxSphereBounds UDebugSkelMeshComponent::CalcGameBounds(const FTransform& LocalToWorld) const
+{
+	return Super::CalcBounds(LocalToWorld);
+}
+
 bool UDebugSkelMeshComponent::IsUsingInGameBounds() const
 {
 	return bIsUsingInGameBounds;
diff --git a/Engine/Source/Editor/UnrealEd/Private/EditorViewportClient.cpp b/Engine/Source/Editor/UnrealEd/Private/EditorViewportClient.cpp
--- a/Engine/Source/Editor/UnrealEd/Private/EditorViewportClient.cpp
+++ b/Engine/Source/Editor/UnrealEd/Private/EditorViewportClient.cpp
@@ -3508,12 +3508,7 @@ void FEditorViewportClient::OnDollyPerspectiveCamera( const FInputEventState& In
 	FKey Key = InputState.GetKey();
 
 	// Scrolling the mousewheel up/down moves the perspective viewport forwards/backwards.
-	FVector Drag(0,0,0);
-
-	const FRotator& ViewRotation = GetViewRotation();
-	Drag.X = FMath::Cos( ViewRotation.Yaw * PI / 180.f ) * FMath::Cos( ViewRotation.Pitch * PI / 180.f );
-	Drag.Y = FMath::Sin( ViewRotation.Yaw * PI / 180.f ) * FMath::Cos( ViewRotation.Pitch * PI / 180.f );
-	Drag.Z = FMath::Sin( ViewRotation.Pitch * PI / 180.f );
+	FVector Drag = (GetLookAtLocation() - GetViewLocation()).GetSafeNormal();
 
 	if( Key == EKeys::MouseScrollDown )
 	{
diff --git a/Engine/Source/Editor/UnrealEd/Private/PreferenceStubs.cpp b/Engine/Source/Editor/UnrealEd/Private/PreferenceStubs.cpp
--- a/Engine/Source/Editor/UnrealEd/Private/PreferenceStubs.cpp
+++ b/Engine/Source/Editor/UnrealEd/Private/PreferenceStubs.cpp
@@ -107,6 +107,7 @@ UPersonaOptions::UPersonaOptions(const FObjectInitializer& ObjectInitializer)
 	: Super(ObjectInitializer)
 	, DefaultLocalAxesSelection(2)
 	, DefaultBoneDrawSelection(1)
+	, bPauseAnimationOnCameraMove(false)
 	, bAllowPreviewMeshCollectionsToSelectFromDifferentSkeletons(true)
 {
 	AssetEditorOptions.AddUnique(FAssetEditorOptions(TEXT("SkeletonEditor")));
