Non-unity and Static analysis fixes
#rnx

diff --git a/Engine/Plugins/Experimental/PCG/Source/PCG/Private/PCGHelpers.cpp b/Engine/Plugins/Experimental/PCG/Source/PCG/Private/PCGHelpers.cpp
--- a/Engine/Plugins/Experimental/PCG/Source/PCG/Private/PCGHelpers.cpp
+++ b/Engine/Plugins/Experimental/PCG/Source/PCG/Private/PCGHelpers.cpp
@@ -157,6 +157,7 @@ namespace PCGHelpers
 		{
 			return;
 		}
+		checkSlow(ObjectClass); // CanBeExpanded checks this but static analyzer doesn't note that
 
 		for (FProperty* Property = ObjectClass->PropertyLink; Property != nullptr; Property = Property->PropertyLinkNext)
 		{
diff --git a/Engine/Plugins/Runtime/MassEntity/Source/MassEntityDebugger/Private/MassDebuggerModule.cpp b/Engine/Plugins/Runtime/MassEntity/Source/MassEntityDebugger/Private/MassDebuggerModule.cpp
--- a/Engine/Plugins/Runtime/MassEntity/Source/MassEntityDebugger/Private/MassDebuggerModule.cpp
+++ b/Engine/Plugins/Runtime/MassEntity/Source/MassEntityDebugger/Private/MassDebuggerModule.cpp
@@ -5,6 +5,7 @@
 #if WITH_MASSENTITY_DEBUG
 
 #include "IMassDebuggerModule.h"
+#include "Features/IModularFeatures.h"
 
 #include "WorkspaceMenuStructure.h"
 #include "WorkspaceMenuStructureModule.h"
diff --git a/Engine/Plugins/Runtime/MeshModelingToolset/Source/ModelingComponents/Private/ModelingComponentsModule.cpp b/Engine/Plugins/Runtime/MeshModelingToolset/Source/ModelingComponents/Private/ModelingComponentsModule.cpp
--- a/Engine/Plugins/Runtime/MeshModelingToolset/Source/ModelingComponents/Private/ModelingComponentsModule.cpp
+++ b/Engine/Plugins/Runtime/MeshModelingToolset/Source/ModelingComponents/Private/ModelingComponentsModule.cpp
@@ -2,7 +2,7 @@
 
 #include "ModelingComponentsModule.h"
 #include "Components/BaseDynamicMeshComponent.h"
-
+#include "Materials/Material.h"
 
 #define LOCTEXT_NAMESPACE "FModelingComponentsModule"
 
diff --git a/Engine/Plugins/Runtime/WaveTable/Source/WaveTable/Private/WaveTableModule.cpp b/Engine/Plugins/Runtime/WaveTable/Source/WaveTable/Private/WaveTableModule.cpp
--- a/Engine/Plugins/Runtime/WaveTable/Source/WaveTable/Private/WaveTableModule.cpp
+++ b/Engine/Plugins/Runtime/WaveTable/Source/WaveTable/Private/WaveTableModule.cpp
@@ -3,7 +3,7 @@
 #include "WaveTableModule.h"
 
 #include "WaveTableLogging.h"
-
+#include "Modules/ModuleManager.h"
 
 namespace WaveTable
 {
diff --git a/Engine/Source/Editor/Persona/Private/AnimationEditorViewportClient.cpp b/Engine/Source/Editor/Persona/Private/AnimationEditorViewportClient.cpp
--- a/Engine/Source/Editor/Persona/Private/AnimationEditorViewportClient.cpp
+++ b/Engine/Source/Editor/Persona/Private/AnimationEditorViewportClient.cpp
@@ -648,61 +648,62 @@ void FAnimationViewportClient::HandlePreviewScenePreTick()
 
 void FAnimationViewportClient::HandlePreviewScenePostTick()
 {
-	UDebugSkelMeshComponent* PreviewMeshComponent = GetAnimPreviewScene()->GetPreviewMeshComponent();
-
-	const bool bInBoneOrbitMode = CameraFollowMode != EAnimationViewportCameraFollowMode::None;
-	if (IsTracking())
+	if (UDebugSkelMeshComponent* PreviewMeshComponent = GetAnimPreviewScene()->GetPreviewMeshComponent())
 	{
-		if (PreviewMeshComponent->IsPlaying() && GetDefault<UPersonaOptions>()->bPauseAnimationOnCameraMove)
+		const bool bInBoneOrbitMode = CameraFollowMode != EAnimationViewportCameraFollowMode::None;
+		if (IsTracking())
 		{
-			PreviewMeshComponent->Stop();
-			bResumeAfterTracking = true;
-			return;
+			if (PreviewMeshComponent->IsPlaying() && GetDefault<UPersonaOptions>()->bPauseAnimationOnCameraMove)
+			{
+				PreviewMeshComponent->Stop();
+				bResumeAfterTracking = true;
+				return;
+			}
+		}
+		else if (bResumeAfterTracking)
+		{
+			PreviewMeshComponent->Play(true);
+			bResumeAfterTracking = false;
 		}
-	}
-	else if (bResumeAfterTracking)
-	{
-		PreviewMeshComponent->Play(true);
-		bResumeAfterTracking = false;
-	}
-
-	if (CameraFollowMode != EAnimationViewportCameraFollowMode::None && PreviewMeshComponent != nullptr)
-	{
-		FVector LookAtLocation = LastLookAtLocation;
 
-		switch(CameraFollowMode)
+		if (CameraFollowMode != EAnimationViewportCameraFollowMode::None)
 		{
-		case EAnimationViewportCameraFollowMode::Bounds:
-			{
-				const FBoxSphereBounds Bounds = PreviewMeshComponent->CalcGameBounds(PreviewMeshComponent->GetComponentTransform());
-				LookAtLocation = Bounds.Origin;
-			}
-			break;
-		case EAnimationViewportCameraFollowMode::Root:
-			{
-				LookAtLocation = PreviewMeshComponent->GetBoneTransform(0).GetLocation();
-				LookAtLocation.Z = PreviewMeshComponent->CalcGameBounds(PreviewMeshComponent->GetComponentTransform()).Origin.Z;
-			}
-			break;
-		case EAnimationViewportCameraFollowMode::Bone:
-			{
-				const int32 BoneIndex = PreviewMeshComponent->GetBoneIndex(CameraFollowBoneName);
-				check(BoneIndex != INDEX_NONE);
-				LookAtLocation = PreviewMeshComponent->GetBoneTransform(BoneIndex).GetLocation();
+			FVector LookAtLocation = LastLookAtLocation;
 
-				if (GetShouldRotateCameraToFollowBone())
+			switch(CameraFollowMode)
+			{
+			case EAnimationViewportCameraFollowMode::Bounds:
+				{
+					const FBoxSphereBounds Bounds = PreviewMeshComponent->CalcGameBounds(PreviewMeshComponent->GetComponentTransform());
+					LookAtLocation = Bounds.Origin;
+				}
+				break;
+			case EAnimationViewportCameraFollowMode::Root:
+				{
+					LookAtLocation = PreviewMeshComponent->GetBoneTransform(0).GetLocation();
+					LookAtLocation.Z = PreviewMeshComponent->CalcGameBounds(PreviewMeshComponent->GetComponentTransform()).Origin.Z;
+				}
+				break;
+			case EAnimationViewportCameraFollowMode::Bone:
 				{
-					const FName BoneName = PreviewMeshComponent->GetBoneName(BoneIndex);
-					check(BoneName != NAME_None);
-					OrbitRotation = PreviewMeshComponent->GetBoneQuaternion(BoneName) * FQuat(FVector(0.0f, 1.0f, 0.0f), PI * 0.5f);
+					const int32 BoneIndex = PreviewMeshComponent->GetBoneIndex(CameraFollowBoneName);
+					check(BoneIndex != INDEX_NONE);
+					LookAtLocation = PreviewMeshComponent->GetBoneTransform(BoneIndex).GetLocation();
+
+					if (GetShouldRotateCameraToFollowBone())
+					{
+						const FName BoneName = PreviewMeshComponent->GetBoneName(BoneIndex);
+						check(BoneName != NAME_None);
+						OrbitRotation = PreviewMeshComponent->GetBoneQuaternion(BoneName) * FQuat(FVector(0.0f, 1.0f, 0.0f), PI * 0.5f);
+					}
 				}
+				break;
 			}
-			break;
-		}
 
-		const FVector Offset = LookAtLocation - LastLookAtLocation;
-		SetLookAtLocation(GetLookAtLocation() + Offset);
-		LastLookAtLocation = LookAtLocation;
+			const FVector Offset = LookAtLocation - LastLookAtLocation;
+			SetLookAtLocation(GetLookAtLocation() + Offset);
+			LastLookAtLocation = LookAtLocation;
+		}
 	}
 }
 
diff --git a/Engine/Source/Runtime/Experimental/Interchange/Engine/Public/InterchangePipelineConfigurationBase.h b/Engine/Source/Runtime/Experimental/Interchange/Engine/Public/InterchangePipelineConfigurationBase.h
--- a/Engine/Source/Runtime/Experimental/Interchange/Engine/Public/InterchangePipelineConfigurationBase.h
+++ b/Engine/Source/Runtime/Experimental/Interchange/Engine/Public/InterchangePipelineConfigurationBase.h
@@ -5,11 +5,10 @@
 #include "CoreMinimal.h"
 #include "UObject/Object.h"
 #include "UObject/ObjectMacros.h"
+#include "InterchangeSourceData.h"
 
 #include "InterchangePipelineConfigurationBase.generated.h"
 
-class UInterchangeSourceData;
-
 UENUM(BlueprintType, Experimental)
 enum class EInterchangePipelineConfigurationDialogResult : uint8
 {
