Replace direct access to SkeletalMesh object from USkinnedMeshComponent with GetSkeletalMesh function

diff --git a/Engine/Plugins/Runtime/MeshModelingToolset/Source/ModelingComponentsEditorOnly/Private/ToolTargets/SkeletalMeshComponentToolTarget.cpp b/Engine/Plugins/Runtime/MeshModelingToolset/Source/ModelingComponentsEditorOnly/Private/ToolTargets/SkeletalMeshComponentToolTarget.cpp
--- a/Engine/Plugins/Runtime/MeshModelingToolset/Source/ModelingComponentsEditorOnly/Private/ToolTargets/SkeletalMeshComponentToolTarget.cpp
+++ b/Engine/Plugins/Runtime/MeshModelingToolset/Source/ModelingComponentsEditorOnly/Private/ToolTargets/SkeletalMeshComponentToolTarget.cpp
@@ -33,7 +33,7 @@ bool USkeletalMeshComponentReadOnlyToolTarget::IsValid() const
 	{
 		return false;
 	}
-	const USkeletalMesh* SkeletalMesh = SkinnedMeshComponent->SkeletalMesh;
+	const USkeletalMesh* SkeletalMesh = SkinnedMeshComponent->GetSkeletalMesh();
 
 	return USkeletalMeshReadOnlyToolTarget::IsValid(SkeletalMesh);
 }
@@ -54,7 +54,7 @@ void USkeletalMeshComponentReadOnlyToolTarget::GetMaterialSet(FComponentMaterial
 
 	if (bPreferAssetMaterials)
 	{
-		const USkeletalMesh* SkeletalMesh = Cast<USkinnedMeshComponent>(Component)->SkeletalMesh;
+		const USkeletalMesh* SkeletalMesh = Cast<USkinnedMeshComponent>(Component)->GetSkeletalMesh();
 		USkeletalMeshToolTarget::GetMaterialSet(SkeletalMesh, MaterialSetOut, bPreferAssetMaterials);
 	}
 	else
@@ -74,7 +74,7 @@ bool USkeletalMeshComponentReadOnlyToolTarget::CommitMaterialSetUpdate(const FCo
 
 	if (bApplyToAsset)
 	{
-		USkeletalMesh* SkeletalMesh = Cast<USkinnedMeshComponent>(Component)->SkeletalMesh;
+		USkeletalMesh* SkeletalMesh = Cast<USkinnedMeshComponent>(Component)->GetSkeletalMesh();
 
 		// unregister the component while we update it's static mesh
 		TUniquePtr<FComponentReregisterContext> ComponentReregisterContext = MakeUnique<FComponentReregisterContext>(Component);
@@ -115,7 +115,7 @@ const FMeshDescription* USkeletalMeshComponentReadOnlyToolTarget::GetMeshDescrip
 	if (!CachedMeshDescription.IsValid())
 	{
 		CachedMeshDescription = MakeUnique<FMeshDescription>();
-		const USkeletalMesh* SkeletalMesh = Cast<USkinnedMeshComponent>(Component)->SkeletalMesh;
+		const USkeletalMesh* SkeletalMesh = Cast<USkinnedMeshComponent>(Component)->GetSkeletalMesh();
 		USkeletalMeshToolTarget::GetMeshDescription(SkeletalMesh, *CachedMeshDescription);
 	}
 	
@@ -139,7 +139,7 @@ FDynamicMesh3 USkeletalMeshComponentReadOnlyToolTarget::GetDynamicMesh()
 
 USkeletalMesh* USkeletalMeshComponentReadOnlyToolTarget::GetSkeletalMesh() const
 {
-	return IsValid() ? Cast<USkinnedMeshComponent>(Component)->SkeletalMesh : nullptr;
+	return IsValid() ? Cast<USkinnedMeshComponent>(Component)->GetSkeletalMesh() : nullptr;
 }
 
 
@@ -152,7 +152,7 @@ void USkeletalMeshComponentToolTarget::CommitMeshDescription(const FCommitter& C
 {
 	if (ensure(IsValid()) == false) return;
 
-	USkeletalMesh* SkeletalMesh = Cast<USkinnedMeshComponent>(Component)->SkeletalMesh;
+	USkeletalMesh* SkeletalMesh = Cast<USkinnedMeshComponent>(Component)->GetSkeletalMesh();
 	if (!CachedMeshDescription.IsValid())
 	{
 		CachedMeshDescription = MakeUnique<FMeshDescription>();
@@ -193,8 +193,8 @@ bool USkeletalMeshComponentReadOnlyToolTargetFactory::CanBuildTarget(UObject* So
 	// just add another factory that allows that class specifically(but make sure that
 	// GetMeshDescription and such work properly)
 
-	return Cast<USkinnedMeshComponent>(SourceObject) && Cast<USkinnedMeshComponent>(SourceObject)->SkeletalMesh &&
-		ExactCast<USkeletalMesh>(Cast<USkinnedMeshComponent>(SourceObject)->SkeletalMesh) &&
+	return Cast<USkinnedMeshComponent>(SourceObject) && Cast<USkinnedMeshComponent>(SourceObject)->GetSkeletalMesh() &&
+		ExactCast<USkeletalMesh>(Cast<USkinnedMeshComponent>(SourceObject)->GetSkeletalMesh()) &&
 		Requirements.AreSatisfiedBy(USkeletalMeshComponentReadOnlyToolTarget::StaticClass());
 }
 
@@ -216,8 +216,8 @@ bool USkeletalMeshComponentToolTargetFactory::CanBuildTarget(UObject* SourceObje
 	// just add another factory that allows that class specifically(but make sure that
 	// GetMeshDescription and such work properly)
 
-	return Cast<USkinnedMeshComponent>(SourceObject) && Cast<USkinnedMeshComponent>(SourceObject)->SkeletalMesh &&
-		ExactCast<USkeletalMesh>(Cast<USkinnedMeshComponent>(SourceObject)->SkeletalMesh) &&
+	return Cast<USkinnedMeshComponent>(SourceObject) && Cast<USkinnedMeshComponent>(SourceObject)->GetSkeletalMesh() &&
+		ExactCast<USkeletalMesh>(Cast<USkinnedMeshComponent>(SourceObject)->GetSkeletalMesh()) &&
 		Requirements.AreSatisfiedBy(USkeletalMeshComponentToolTarget::StaticClass());
 }
 
diff --git a/Engine/Plugins/Runtime/Oculus/OculusVR/Source/OculusInput/Private/OculusHandComponent.cpp b/Engine/Plugins/Runtime/Oculus/OculusVR/Source/OculusInput/Private/OculusHandComponent.cpp
--- a/Engine/Plugins/Runtime/Oculus/OculusVR/Source/OculusInput/Private/OculusHandComponent.cpp
+++ b/Engine/Plugins/Runtime/Oculus/OculusVR/Source/OculusInput/Private/OculusHandComponent.cpp
@@ -33,7 +33,7 @@ void UOculusHandComponent::BeginPlay()
 	Super::BeginPlay();
 
 	// Use custom mesh if a skeletal mesh is already set, else try to load the runtime mesh
-	if (SkeletalMesh)
+	if (GetSkeletalMesh())
 	{
 		bCustomHandMesh = true;
 		bSkeletalMeshInitialized = true;
@@ -106,7 +106,7 @@ void UOculusHandComponent::TickComponent(float DeltaTime, enum ELevelTick TickTy
 			}
 
 			// Update Bone Pose Rotations
-			if (SkeletalMesh)
+			if (GetSkeletalMesh())
 			{
 				UpdateBonePose();
 			}
@@ -163,7 +163,7 @@ void UOculusHandComponent::UpdateBonePose()
 			else
 			{
 				// Set Remaing Bone Rotations
-				int32 BoneIndex = SkeletalMesh->GetRefSkeleton().FindBoneIndex(BoneElem.Value);
+				int32 BoneIndex = GetSkeletalMesh()->GetRefSkeleton().FindBoneIndex(BoneElem.Value);
 				if (BoneIndex >= 0)
 				{
 					FQuat BoneRotation = UOculusInputFunctionLibrary::GetBoneRotation(SkeletonType, (EBone)BoneElem.Key);
@@ -181,7 +181,7 @@ void UOculusHandComponent::UpdateBonePose()
 		BoneSpaceTransforms[0].SetRotation(RootBoneRotation);
 
 		// Set Remaining Bone Rotations
-		for (uint32 BoneIndex = 1; BoneIndex < (uint32)SkeletalMesh->GetRefSkeleton().GetNum(); BoneIndex++)
+		for (uint32 BoneIndex = 1; BoneIndex < (uint32)GetSkeletalMesh()->GetRefSkeleton().GetNum(); BoneIndex++)
 		{
 			FQuat BoneRotation = UOculusInputFunctionLibrary::GetBoneRotation(SkeletonType, (EBone)BoneIndex);
 			BoneSpaceTransforms[BoneIndex].SetRotation(BoneRotation);
diff --git a/Engine/Plugins/Runtime/Oculus/OculusVR/Source/OculusInput/Private/OculusHandTracking.cpp b/Engine/Plugins/Runtime/Oculus/OculusVR/Source/OculusInput/Private/OculusHandTracking.cpp
--- a/Engine/Plugins/Runtime/Oculus/OculusVR/Source/OculusInput/Private/OculusHandTracking.cpp
+++ b/Engine/Plugins/Runtime/Oculus/OculusVR/Source/OculusInput/Private/OculusHandTracking.cpp
@@ -552,7 +552,7 @@ TArray<FOculusCapsuleCollider> FOculusHandTracking::InitializeHandPhysics(const
 		FVector CapsulePointOne = OvrBoneVectorToFVector(OvrBoneCapsule.Points[1], WorldToMeters);
 		FVector Delta = (CapsulePointOne - CapsulePointZero);
 
-		FName BoneName = HandComponent->SkeletalMesh->GetRefSkeleton().GetBoneName(OvrBoneCapsule.BoneIndex);
+		FName BoneName = HandComponent->GetSkeletalMesh()->GetRefSkeleton().GetBoneName(OvrBoneCapsule.BoneIndex);
 
 		float CapsuleHeight = Delta.Size();
 		float CapsuleRadius = OvrBoneCapsule.Radius * WorldToMeters;
diff --git a/Engine/Plugins/VirtualProduction/RemoteControl/Source/RemoteControlUI/Private/RemoteControlUIModule.cpp b/Engine/Plugins/VirtualProduction/RemoteControl/Source/RemoteControlUI/Private/RemoteControlUIModule.cpp
--- a/Engine/Plugins/VirtualProduction/RemoteControl/Source/RemoteControlUI/Private/RemoteControlUIModule.cpp
+++ b/Engine/Plugins/VirtualProduction/RemoteControl/Source/RemoteControlUI/Private/RemoteControlUIModule.cpp
@@ -961,7 +961,7 @@ UMeshComponent* FRemoteControlUIModule::GetSelectedMeshComponentToBeModified(UOb
 		}
 		else if (USkeletalMeshComponent* SkeletalMeshComponent = Cast<USkeletalMeshComponent>(InMeshComponent))
 		{
-			if (SkeletalMeshComponent->SkeletalMesh == InOwnerObject)
+			if (SkeletalMeshComponent->GetSkeletalMesh() == InOwnerObject)
 			{
 				ModifinedComponent = InMeshComponent;
 			}
diff --git a/Engine/Plugins/VirtualProduction/Takes/Source/TakeTrackRecorders/Private/MovieScene3DTransformTrackRecorder.cpp b/Engine/Plugins/VirtualProduction/Takes/Source/TakeTrackRecorders/Private/MovieScene3DTransformTrackRecorder.cpp
--- a/Engine/Plugins/VirtualProduction/Takes/Source/TakeTrackRecorders/Private/MovieScene3DTransformTrackRecorder.cpp
+++ b/Engine/Plugins/VirtualProduction/Takes/Source/TakeTrackRecorders/Private/MovieScene3DTransformTrackRecorder.cpp
@@ -489,7 +489,7 @@ void UMovieScene3DTransformTrackRecorder::PostProcessAnimationData(UMovieSceneAn
 
 			}
 			// Search for the Root Bone in the Skeleton
-			const USkeletalMesh* SkeletalMesh = SkeletalMeshComponent->MasterPoseComponent != nullptr ? ToRawPtr(SkeletalMeshComponent->MasterPoseComponent->SkeletalMesh) : ToRawPtr(SkeletalMeshComponent->SkeletalMesh);
+			const USkeletalMesh* SkeletalMesh = SkeletalMeshComponent->MasterPoseComponent != nullptr ? ToRawPtr(SkeletalMeshComponent->MasterPoseComponent->GetSkeletalMesh()) : ToRawPtr(SkeletalMeshComponent->GetSkeletalMesh());
 			const UAnimSequence* AnimSequence = AnimTrackRecorder->GetAnimSequence();
 			if (AnimSequence && SkeletalMesh && AnimSequence->GetSkeleton())
 			{
diff --git a/Engine/Plugins/VirtualProduction/Takes/Source/TakeTrackRecorders/Private/MovieSceneAnimationTrackRecorder.cpp b/Engine/Plugins/VirtualProduction/Takes/Source/TakeTrackRecorders/Private/MovieSceneAnimationTrackRecorder.cpp
--- a/Engine/Plugins/VirtualProduction/Takes/Source/TakeTrackRecorders/Private/MovieSceneAnimationTrackRecorder.cpp
+++ b/Engine/Plugins/VirtualProduction/Takes/Source/TakeTrackRecorders/Private/MovieSceneAnimationTrackRecorder.cpp
@@ -29,7 +29,7 @@ bool FMovieSceneAnimationTrackRecorderFactory::CanRecordObject(UObject* InObject
 	if (InObjectToRecord->IsA<USkeletalMeshComponent>())
 	{
 		USkeletalMeshComponent* SkeletalMeshComponent = Cast<USkeletalMeshComponent>(InObjectToRecord);
-		if (SkeletalMeshComponent && SkeletalMeshComponent->SkeletalMesh)
+		if (SkeletalMeshComponent && SkeletalMeshComponent->GetSkeletalMesh())
 		{
 			return true;
 		}
@@ -46,7 +46,7 @@ void UMovieSceneAnimationTrackRecorder::CreateAnimationAssetAndSequence(const AA
 {
 	UMovieSceneAnimationTrackRecorderSettings* AnimSettings = CastChecked<UMovieSceneAnimationTrackRecorderSettings>(Settings.Get());
 
-	SkeletalMesh = SkeletalMeshComponent->SkeletalMesh;
+	SkeletalMesh = SkeletalMeshComponent->GetSkeletalMesh();
 	if (SkeletalMesh.IsValid())
 	{
 		ComponentTransform = SkeletalMeshComponent->GetComponentToWorld().GetRelativeTransform(Actor->GetTransform());
@@ -72,7 +72,7 @@ void UMovieSceneAnimationTrackRecorder::CreateAnimationAssetAndSequence(const AA
 			FAssetRegistryModule::AssetCreated(AnimSequence.Get());
 
 			// Assign the skeleton we're recording to the newly created Animation Sequence.
-			AnimSequence->SetSkeleton(SkeletalMeshComponent->SkeletalMesh->GetSkeleton());
+			AnimSequence->SetSkeleton(SkeletalMeshComponent->GetSkeletalMesh()->GetSkeleton());
 		}
 	}
 
@@ -157,7 +157,7 @@ void UMovieSceneAnimationTrackRecorder::CreateTrackImpl()
 			for (int32 BoneIndex = 0; BoneIndex < SpaceBases->Num(); ++BoneIndex)
 			{
 				// verify if this bone exists in skeleton
-				const int32 BoneTreeIndex = AnimSkeleton->GetSkeletonBoneIndexFromMeshBoneIndex(SkeletalMeshComponent->MasterPoseComponent != nullptr ? SkeletalMeshComponent->MasterPoseComponent->SkeletalMesh : SkeletalMeshComponent->SkeletalMesh, BoneIndex);
+				const int32 BoneTreeIndex = AnimSkeleton->GetSkeletonBoneIndexFromMeshBoneIndex(SkeletalMeshComponent->MasterPoseComponent != nullptr ? SkeletalMeshComponent->MasterPoseComponent->GetSkeletalMesh() : SkeletalMeshComponent->GetSkeletalMesh(), BoneIndex);
 				if (BoneTreeIndex != INDEX_NONE)
 				{
 					// add tracks for the bone existing
diff --git a/Engine/Source/Developer/MeshMergeUtilities/Private/MeshMergeHelpers.cpp b/Engine/Source/Developer/MeshMergeUtilities/Private/MeshMergeHelpers.cpp
--- a/Engine/Source/Developer/MeshMergeUtilities/Private/MeshMergeHelpers.cpp
+++ b/Engine/Source/Developer/MeshMergeUtilities/Private/MeshMergeHelpers.cpp
@@ -109,13 +109,13 @@ void FMeshMergeHelpers::ExtractSections(const UStaticMeshComponent* Component, i
 void FMeshMergeHelpers::ExtractSections(const USkeletalMeshComponent* Component, int32 LODIndex, TArray<FSectionInfo>& OutSections)
 {
 	static UMaterialInterface* DefaultMaterial = UMaterial::GetDefaultMaterial(MD_Surface);
-	FSkeletalMeshModel* Resource = Component->SkeletalMesh->GetImportedModel();
+	FSkeletalMeshModel* Resource = Component->GetSkeletalMesh()->GetImportedModel();
 
 	checkf(Resource->LODModels.IsValidIndex(LODIndex), TEXT("Invalid LOD Index"));
 
 	TArray<FName> MaterialSlotNames = Component->GetMaterialSlotNames();
 
-	const FSkeletalMeshLODInfo* LODInfoPtr = Component->SkeletalMesh->GetLODInfo(LODIndex);
+	const FSkeletalMeshLODInfo* LODInfoPtr = Component->GetSkeletalMesh()->GetLODInfo(LODIndex);
 	check(LODInfoPtr);
 	const FSkeletalMeshLODModel& Model = Resource->LODModels[LODIndex];
 	for (int32 SectionIndex = 0; SectionIndex < Model.Sections.Num(); ++SectionIndex)
@@ -422,10 +422,10 @@ void FMeshMergeHelpers::ExportStaticMeshLOD(const FStaticMeshLODResources& Stati
 
 void FMeshMergeHelpers::RetrieveMesh(const USkeletalMeshComponent* SkeletalMeshComponent, int32 LODIndex, FMeshDescription& OutMeshDescription, bool bPropagateVertexColours)
 {
-	FSkeletalMeshModel* Resource = SkeletalMeshComponent->SkeletalMesh->GetImportedModel();
+	FSkeletalMeshModel* Resource = SkeletalMeshComponent->GetSkeletalMesh()->GetImportedModel();
 	if (Resource->LODModels.IsValidIndex(LODIndex))
 	{
-		FSkeletalMeshLODInfo& SrcLODInfo = *(SkeletalMeshComponent->SkeletalMesh->GetLODInfo(LODIndex));
+		FSkeletalMeshLODInfo& SrcLODInfo = *(SkeletalMeshComponent->GetSkeletalMesh()->GetLODInfo(LODIndex));
 
 		// Get the CPU skinned verts for this LOD
 		TArray<FFinalSkinVertex> FinalVertices;
@@ -483,10 +483,10 @@ void FMeshMergeHelpers::RetrieveMesh(const USkeletalMeshComponent* SkeletalMeshC
 			// use the remapping of material indices if there is a valid value
 			if (SrcLODInfo.LODMaterialMap.IsValidIndex(SectionIndex) && SrcLODInfo.LODMaterialMap[SectionIndex] != INDEX_NONE)
 			{
-				MaterialIndex = FMath::Clamp<int32>(SrcLODInfo.LODMaterialMap[SectionIndex], 0, SkeletalMeshComponent->SkeletalMesh->GetMaterials().Num() - 1);
+				MaterialIndex = FMath::Clamp<int32>(SrcLODInfo.LODMaterialMap[SectionIndex], 0, SkeletalMeshComponent->GetSkeletalMesh()->GetMaterials().Num() - 1);
 			}
 
-			FName ImportedMaterialSlotName = SkeletalMeshComponent->SkeletalMesh->GetMaterials()[MaterialIndex].ImportedMaterialSlotName;
+			FName ImportedMaterialSlotName = SkeletalMeshComponent->GetSkeletalMesh()->GetMaterials()[MaterialIndex].ImportedMaterialSlotName;
 			const FPolygonGroupID SectionPolygonGroupID(SectionIndex);
 			if (!OutMeshDescription.IsPolygonGroupValid(SectionPolygonGroupID))
 			{
diff --git a/Engine/Source/Developer/MeshMergeUtilities/Private/MeshMergeUtilities.cpp b/Engine/Source/Developer/MeshMergeUtilities/Private/MeshMergeUtilities.cpp
--- a/Engine/Source/Developer/MeshMergeUtilities/Private/MeshMergeUtilities.cpp
+++ b/Engine/Source/Developer/MeshMergeUtilities/Private/MeshMergeUtilities.cpp
@@ -441,7 +441,7 @@ void FMeshMergeUtilities::BakeMaterialsForComponent(USkeletalMeshComponent* Skel
 	UMaterialMergeOptions* MergeOptions = GetMutableDefault<UMaterialMergeOptions>();
 	TArray<TWeakObjectPtr<UObject>> Objects{ MergeOptions, AssetOptions, MaterialOptions };
 
-	const int32 NumLODs = SkeletalMeshComponent->SkeletalMesh->GetLODNum();
+	const int32 NumLODs = SkeletalMeshComponent->GetSkeletalMesh()->GetLODNum();
 	IMaterialBakingModule& Module = FModuleManager::Get().LoadModuleChecked<IMaterialBakingModule>("MaterialBaking");
 	if (!Module.SetupMaterialBakeSettings(Objects, NumLODs))
 	{
@@ -449,7 +449,7 @@ void FMeshMergeUtilities::BakeMaterialsForComponent(USkeletalMeshComponent* Skel
 	}
 
 	// Bake out materials for skeletal mesh
-	SkeletalMeshComponent->SkeletalMesh->Modify();
+	SkeletalMeshComponent->GetSkeletalMesh()->Modify();
 	FSkeletalMeshComponentAdapter Adapter(SkeletalMeshComponent);
 	BakeMaterialsForComponent(Objects, &Adapter);
 	SkeletalMeshComponent->MarkRenderStateDirty();
diff --git a/Engine/Source/Developer/MeshMergeUtilities/Private/SkeletalMeshAdapter.cpp b/Engine/Source/Developer/MeshMergeUtilities/Private/SkeletalMeshAdapter.cpp
--- a/Engine/Source/Developer/MeshMergeUtilities/Private/SkeletalMeshAdapter.cpp
+++ b/Engine/Source/Developer/MeshMergeUtilities/Private/SkeletalMeshAdapter.cpp
@@ -11,7 +11,7 @@
 #include "UObject/Package.h"
 
 FSkeletalMeshComponentAdapter::FSkeletalMeshComponentAdapter(USkeletalMeshComponent* InSkeletalMeshComponent)
-	: SkeletalMeshComponent(InSkeletalMeshComponent), SkeletalMesh(InSkeletalMeshComponent->SkeletalMesh)
+	: SkeletalMeshComponent(InSkeletalMeshComponent), SkeletalMesh(InSkeletalMeshComponent->GetSkeletalMesh())
 {
 	checkf(SkeletalMesh != nullptr, TEXT("Invalid skeletal mesh in adapter"));
 	NumLODs = SkeletalMesh->GetLODNum();
diff --git a/Engine/Source/Developer/MeshUtilities/Private/MeshUtilities.cpp b/Engine/Source/Developer/MeshUtilities/Private/MeshUtilities.cpp
--- a/Engine/Source/Developer/MeshUtilities/Private/MeshUtilities.cpp
+++ b/Engine/Source/Developer/MeshUtilities/Private/MeshUtilities.cpp
@@ -319,7 +319,7 @@ static void SkinnedMeshToRawMeshes(USkinnedMeshComponent* InSkinnedMeshComponent
 		FRawMeshTracker& RawMeshTracker = OutRawMeshTrackers[OverallLODIndex];
 		const int32 BaseVertexIndex = RawMesh.VertexPositions.Num();
 
-		FSkeletalMeshLODInfo& SrcLODInfo = *(InSkinnedMeshComponent->SkeletalMesh->GetLODInfo(LODIndexRead));
+		FSkeletalMeshLODInfo& SrcLODInfo = *(InSkinnedMeshComponent->GetSkeletalMesh()->GetLODInfo(LODIndexRead));
 
 		// Get the CPU skinned verts for this LOD
 		TArray<FFinalSkinVertex> FinalVertices;
@@ -389,7 +389,7 @@ static void SkinnedMeshToRawMeshes(USkinnedMeshComponent* InSkinnedMeshComponent
 				// use the remapping of material indices if there is a valid value
 				if (SrcLODInfo.LODMaterialMap.IsValidIndex(SectionIndex) && SrcLODInfo.LODMaterialMap[SectionIndex] != INDEX_NONE)
 				{
-					MaterialIndex = FMath::Clamp<int32>(SrcLODInfo.LODMaterialMap[SectionIndex], 0, InSkinnedMeshComponent->SkeletalMesh->GetMaterials().Num() - 1);
+					MaterialIndex = FMath::Clamp<int32>(SrcLODInfo.LODMaterialMap[SectionIndex], 0, InSkinnedMeshComponent->GetSkeletalMesh()->GetMaterials().Num() - 1);
 				}
 
 				// copy face info
diff --git a/Engine/Source/Developer/PhysicsUtilities/Private/PhysicsAssetUtils.cpp b/Engine/Source/Developer/PhysicsUtilities/Private/PhysicsAssetUtils.cpp
--- a/Engine/Source/Developer/PhysicsUtilities/Private/PhysicsAssetUtils.cpp
+++ b/Engine/Source/Developer/PhysicsUtilities/Private/PhysicsAssetUtils.cpp
@@ -662,19 +662,19 @@ void WeldBodies(UPhysicsAsset* PhysAsset, int32 BaseBodyIndex, int32 AddBodyInde
 	if(BaseBodyIndex == INDEX_NONE || AddBodyIndex == INDEX_NONE)
 		return;
 
-	if (SkelComp == NULL || SkelComp->SkeletalMesh == NULL)
+	if (SkelComp == NULL || SkelComp->GetSkeletalMesh() == NULL)
 	{
 		return;
 	}
 
 	UBodySetup* Body1 = PhysAsset->SkeletalBodySetups[BaseBodyIndex];
-	int32 Bone1Index = SkelComp->SkeletalMesh->GetRefSkeleton().FindBoneIndex(Body1->BoneName);
+	int32 Bone1Index = SkelComp->GetSkeletalMesh()->GetRefSkeleton().FindBoneIndex(Body1->BoneName);
 	check(Bone1Index != INDEX_NONE);
 	FTransform Bone1TM = SkelComp->GetBoneTransform(Bone1Index);
 	Bone1TM.RemoveScaling();
 
 	UBodySetup* Body2 = PhysAsset->SkeletalBodySetups[AddBodyIndex];
-	int32 Bone2Index = SkelComp->SkeletalMesh->GetRefSkeleton().FindBoneIndex(Body2->BoneName);
+	int32 Bone2Index = SkelComp->GetSkeletalMesh()->GetRefSkeleton().FindBoneIndex(Body2->BoneName);
 	check(Bone2Index != INDEX_NONE);
 	FTransform Bone2TM = SkelComp->GetBoneTransform(Bone2Index);
 	Bone2TM.RemoveScaling();
diff --git a/Engine/Source/Developer/SkeletalMeshUtilitiesCommon/Private/LODUtilities.cpp b/Engine/Source/Developer/SkeletalMeshUtilitiesCommon/Private/LODUtilities.cpp
--- a/Engine/Source/Developer/SkeletalMeshUtilitiesCommon/Private/LODUtilities.cpp
+++ b/Engine/Source/Developer/SkeletalMeshUtilitiesCommon/Private/LODUtilities.cpp
@@ -1755,7 +1755,7 @@ void FLODUtilities::RefreshLODChange(const USkeletalMesh* SkeletalMesh)
 	for (FThreadSafeObjectIterator Iter(USkeletalMeshComponent::StaticClass()); Iter; ++Iter)
 	{
 		USkeletalMeshComponent* SkeletalMeshComponent = Cast<USkeletalMeshComponent>(*Iter);
-		if  (SkeletalMeshComponent->SkeletalMesh == SkeletalMesh)
+		if  (SkeletalMeshComponent->GetSkeletalMesh() == SkeletalMesh)
 		{
 			// it needs to recreate IF it already has been created
 			if (SkeletalMeshComponent->IsRegistered())
diff --git a/Engine/Source/Editor/AnimGraph/Private/AnimGraphNode_SkeletalControlBase.cpp b/Engine/Source/Editor/AnimGraph/Private/AnimGraphNode_SkeletalControlBase.cpp
--- a/Engine/Source/Editor/AnimGraph/Private/AnimGraphNode_SkeletalControlBase.cpp
+++ b/Engine/Source/Editor/AnimGraph/Private/AnimGraphNode_SkeletalControlBase.cpp
@@ -81,7 +81,7 @@ void UAnimGraphNode_SkeletalControlBase::CreateOutputPins()
 
 void UAnimGraphNode_SkeletalControlBase::ConvertToComponentSpaceTransform(const USkeletalMeshComponent* SkelComp, const FTransform & InTransform, FTransform & OutCSTransform, int32 BoneIndex, EBoneControlSpace Space) const
 {
-	USkeleton * Skeleton = SkelComp->SkeletalMesh->GetSkeleton();
+	USkeleton * Skeleton = SkelComp->GetSkeletalMesh()->GetSkeleton();
 
 	switch (Space)
 	{
@@ -105,7 +105,7 @@ void UAnimGraphNode_SkeletalControlBase::ConvertToComponentSpaceTransform(const
 			const int32 ParentIndex = Skeleton->GetReferenceSkeleton().GetParentIndex(BoneIndex);
 			if (ParentIndex != INDEX_NONE)																																																																																																																															
 			{
-				const int32 MeshParentIndex = Skeleton->GetMeshBoneIndexFromSkeletonBoneIndex(SkelComp->SkeletalMesh, ParentIndex);
+				const int32 MeshParentIndex = Skeleton->GetMeshBoneIndexFromSkeletonBoneIndex(SkelComp->GetSkeletalMesh(), ParentIndex);
 				if (MeshParentIndex != INDEX_NONE)
 				{
 					const FTransform ParentTM = SkelComp->GetBoneTransform(MeshParentIndex);
@@ -122,7 +122,7 @@ void UAnimGraphNode_SkeletalControlBase::ConvertToComponentSpaceTransform(const
 	case BCS_BoneSpace:
 		if (BoneIndex != INDEX_NONE)
 		{
-			const int32 MeshBoneIndex = Skeleton->GetMeshBoneIndexFromSkeletonBoneIndex(SkelComp->SkeletalMesh, BoneIndex);
+			const int32 MeshBoneIndex = Skeleton->GetMeshBoneIndexFromSkeletonBoneIndex(SkelComp->GetSkeletalMesh(), BoneIndex);
 			if (MeshBoneIndex != INDEX_NONE)
 			{
 				const FTransform BoneTM = SkelComp->GetBoneTransform(MeshBoneIndex);
@@ -136,9 +136,9 @@ void UAnimGraphNode_SkeletalControlBase::ConvertToComponentSpaceTransform(const
 		break;
 
 	default:
-		if (SkelComp->SkeletalMesh)
+		if (SkelComp->GetSkeletalMesh())
 		{
-			UE_LOG(LogAnimation, Warning, TEXT("ConvertToComponentSpaceTransform: Unknown BoneSpace %d  for Mesh: %s"), (uint8)Space, *SkelComp->SkeletalMesh->GetFName().ToString());
+			UE_LOG(LogAnimation, Warning, TEXT("ConvertToComponentSpaceTransform: Unknown BoneSpace %d  for Mesh: %s"), (uint8)Space, *SkelComp->GetSkeletalMesh()->GetFName().ToString());
 		}
 		else
 		{
@@ -250,7 +250,7 @@ FVector UAnimGraphNode_SkeletalControlBase::ConvertWidgetLocation(const USkeleta
 
 	if (MeshBases.GetPose().IsValid())
 	{
-		USkeleton * Skeleton = SkelComp->SkeletalMesh->GetSkeleton();
+		USkeleton * Skeleton = SkelComp->GetSkeletalMesh()->GetSkeleton();
 		const FMeshPoseBoneIndex MeshBoneIndex(SkelComp->GetBoneIndex(BoneName));
 		const FCompactPoseBoneIndex CompactBoneIndex = MeshBases.GetPose().GetBoneContainer().MakeCompactPoseIndex(MeshBoneIndex);
 
diff --git a/Engine/Source/Editor/AnimGraph/Private/AnimGraphNode_Trail.cpp b/Engine/Source/Editor/AnimGraph/Private/AnimGraphNode_Trail.cpp
--- a/Engine/Source/Editor/AnimGraph/Private/AnimGraphNode_Trail.cpp
+++ b/Engine/Source/Editor/AnimGraph/Private/AnimGraphNode_Trail.cpp
@@ -129,7 +129,7 @@ void UAnimGraphNode_Trail::Draw(FPrimitiveDrawInterface* PDI, USkeletalMeshCompo
 {
 	// initialize bone list
 	// draw bone Angular limit
-	if (Node.bLimitRotation && Node.ChainLength > 1 && PreviewSkelMeshComp->SkeletalMesh)
+	if (Node.bLimitRotation && Node.ChainLength > 1 && PreviewSkelMeshComp->GetSkeletalMesh())
 	{
 		TArray<FName> TrailBoneList;
 		TrailBoneList.Reset(Node.ChainLength);
@@ -137,7 +137,7 @@ void UAnimGraphNode_Trail::Draw(FPrimitiveDrawInterface* PDI, USkeletalMeshCompo
 		int32 CurrentIndex = Node.ChainLength - 1;
 		TrailBoneList[CurrentIndex] = Node.TrailBone.BoneName;
 		FName CurrentName = Node.TrailBone.BoneName;
-		const FReferenceSkeleton& RefSkeleton = PreviewSkelMeshComp->SkeletalMesh->GetRefSkeleton();
+		const FReferenceSkeleton& RefSkeleton = PreviewSkelMeshComp->GetSkeletalMesh()->GetRefSkeleton();
 		while (--CurrentIndex >= 0 && CurrentName != NAME_None)
 		{
 			const int32 ParentIndex = RefSkeleton.GetParentIndex(RefSkeleton.FindBoneIndex(CurrentName));
diff --git a/Engine/Source/Editor/AnimGraph/Private/AnimPreviewInstance.cpp b/Engine/Source/Editor/AnimGraph/Private/AnimPreviewInstance.cpp
--- a/Engine/Source/Editor/AnimGraph/Private/AnimPreviewInstance.cpp
+++ b/Engine/Source/Editor/AnimGraph/Private/AnimPreviewInstance.cpp
@@ -175,9 +175,9 @@ bool FAnimPreviewInstanceProxy::Evaluate(FPoseContext& Output)
 		if(bForceRetargetBasePose)
 		{
 			USkeletalMeshComponent* MeshComponent = Output.AnimInstanceProxy->GetSkelMeshComponent();
-			if(MeshComponent && MeshComponent->SkeletalMesh)
+			if(MeshComponent && MeshComponent->GetSkeletalMesh())
 			{
-				FAnimationRuntime::FillWithRetargetBaseRefPose(Output.Pose, GetSkelMeshComponent()->SkeletalMesh);
+				FAnimationRuntime::FillWithRetargetBaseRefPose(Output.Pose, GetSkelMeshComponent()->GetSkeletalMesh());
 			}
 			else
 			{
@@ -369,7 +369,7 @@ void FAnimPreviewInstanceProxy::SetKeyImplementation(const FCompactPose& PreCont
 	UDebugSkelMeshComponent* Component = Cast<UDebugSkelMeshComponent> (GetSkelMeshComponent());
 
 	USkeleton* PreviewSkeleton = (CurrentSequence) ? CurrentSequence->GetSkeleton() : nullptr;
-	if(CurrentSequence && PreviewSkeleton && Component && Component->SkeletalMesh)
+	if(CurrentSequence && PreviewSkeleton && Component && Component->GetSkeletalMesh())
 	{
 		FScopedTransaction ScopedTransaction(LOCTEXT("SetKey", "Set Key"));
 		CurrentSequence->Modify(true);
diff --git a/Engine/Source/Editor/AnimGraph/Private/EditModes/AnimNodeEditMode.cpp b/Engine/Source/Editor/AnimGraph/Private/EditModes/AnimNodeEditMode.cpp
--- a/Engine/Source/Editor/AnimGraph/Private/EditModes/AnimNodeEditMode.cpp
+++ b/Engine/Source/Editor/AnimGraph/Private/EditModes/AnimNodeEditMode.cpp
@@ -423,7 +423,8 @@ FAnimNode_Base* FAnimNodeEditMode::GetActiveWidgetRuntimeAnimNode() const
 
 void FAnimNodeEditMode::ConvertToComponentSpaceTransform(const USkeletalMeshComponent* SkelComp, const FTransform & InTransform, FTransform & OutCSTransform, int32 BoneIndex, EBoneControlSpace Space)
 {
-	USkeleton* Skeleton = SkelComp->SkeletalMesh->GetSkeleton();
+	USkeletalMesh* SkelMesh = SkelComp->GetSkeletalMesh();
+	USkeleton* Skeleton = SkelMesh->GetSkeleton();
 
 	switch (Space)
 	{
@@ -447,7 +448,7 @@ void FAnimNodeEditMode::ConvertToComponentSpaceTransform(const USkeletalMeshComp
 			const int32 ParentIndex = Skeleton->GetReferenceSkeleton().GetParentIndex(BoneIndex);
 			if (ParentIndex != INDEX_NONE)
 			{
-				const int32 MeshParentIndex = Skeleton->GetMeshBoneIndexFromSkeletonBoneIndex(SkelComp->SkeletalMesh, ParentIndex);
+				const int32 MeshParentIndex = Skeleton->GetMeshBoneIndexFromSkeletonBoneIndex(SkelMesh, ParentIndex);
 				if (MeshParentIndex != INDEX_NONE)
 				{
 					const FTransform ParentTM = SkelComp->GetBoneTransform(MeshParentIndex);
@@ -464,7 +465,7 @@ void FAnimNodeEditMode::ConvertToComponentSpaceTransform(const USkeletalMeshComp
 	case BCS_BoneSpace:
 		if (BoneIndex != INDEX_NONE)
 		{
-			const int32 MeshBoneIndex = Skeleton->GetMeshBoneIndexFromSkeletonBoneIndex(SkelComp->SkeletalMesh, BoneIndex);
+			const int32 MeshBoneIndex = Skeleton->GetMeshBoneIndexFromSkeletonBoneIndex(SkelMesh, BoneIndex);
 			if (MeshBoneIndex != INDEX_NONE)
 			{
 				const FTransform BoneTM = SkelComp->GetBoneTransform(MeshBoneIndex);
@@ -478,9 +479,9 @@ void FAnimNodeEditMode::ConvertToComponentSpaceTransform(const USkeletalMeshComp
 		break;
 
 	default:
-		if (SkelComp->SkeletalMesh)
+		if (SkelMesh)
 		{
-			UE_LOG(LogAnimation, Warning, TEXT("ConvertToComponentSpaceTransform: Unknown BoneSpace %d  for Mesh: %s"), (uint8)Space, *SkelComp->SkeletalMesh->GetFName().ToString());
+			UE_LOG(LogAnimation, Warning, TEXT("ConvertToComponentSpaceTransform: Unknown BoneSpace %d  for Mesh: %s"), (uint8)Space, *SkelMesh->GetFName().ToString());
 		}
 		else
 		{
@@ -493,7 +494,8 @@ void FAnimNodeEditMode::ConvertToComponentSpaceTransform(const USkeletalMeshComp
 
 void FAnimNodeEditMode::ConvertToBoneSpaceTransform(const USkeletalMeshComponent* SkelComp, const FTransform & InCSTransform, FTransform & OutBSTransform, int32 BoneIndex, EBoneControlSpace Space)
 {
-	USkeleton* Skeleton = SkelComp->SkeletalMesh->GetSkeleton();
+	USkeletalMesh* SkelMesh = SkelComp->GetSkeletalMesh();
+	USkeleton* Skeleton = SkelMesh->GetSkeleton();
 
 	switch(Space)
 	{
@@ -517,7 +519,7 @@ void FAnimNodeEditMode::ConvertToBoneSpaceTransform(const USkeletalMeshComponent
 				const int32 ParentIndex = Skeleton->GetReferenceSkeleton().GetParentIndex(BoneIndex);
 				if(ParentIndex != INDEX_NONE)
 				{
-					const int32 MeshParentIndex = Skeleton->GetMeshBoneIndexFromSkeletonBoneIndex(SkelComp->SkeletalMesh, ParentIndex);
+					const int32 MeshParentIndex = Skeleton->GetMeshBoneIndexFromSkeletonBoneIndex(SkelMesh, ParentIndex);
 					if(MeshParentIndex != INDEX_NONE)
 					{
 						const FTransform ParentTM = SkelComp->GetBoneTransform(MeshParentIndex);
@@ -536,7 +538,7 @@ void FAnimNodeEditMode::ConvertToBoneSpaceTransform(const USkeletalMeshComponent
 		{
 			if(BoneIndex != INDEX_NONE)
 			{
-				const int32 MeshBoneIndex = Skeleton->GetMeshBoneIndexFromSkeletonBoneIndex(SkelComp->SkeletalMesh, BoneIndex);
+				const int32 MeshBoneIndex = Skeleton->GetMeshBoneIndexFromSkeletonBoneIndex(SkelMesh, BoneIndex);
 				if(MeshBoneIndex != INDEX_NONE)
 				{
 					FTransform BoneCSTransform = SkelComp->GetBoneTransform(MeshBoneIndex);
@@ -552,7 +554,7 @@ void FAnimNodeEditMode::ConvertToBoneSpaceTransform(const USkeletalMeshComponent
 
 		default:
 		{
-			UE_LOG(LogAnimation, Warning, TEXT("ConvertToBoneSpaceTransform: Unknown BoneSpace %d  for Mesh: %s"), (int32)Space, *GetNameSafe(SkelComp->SkeletalMesh));
+			UE_LOG(LogAnimation, Warning, TEXT("ConvertToBoneSpaceTransform: Unknown BoneSpace %d  for Mesh: %s"), (int32)Space, *GetNameSafe(SkelMesh));
 			break;
 		}
 	}
@@ -757,7 +759,7 @@ FVector FAnimNodeEditMode::ConvertWidgetLocation(const USkeletalMeshComponent* S
 	{
 		if (InMeshBases.GetPose().IsValid())
 		{
-			USkeleton* Skeleton = InSkelComp->SkeletalMesh->GetSkeleton();
+			USkeleton* Skeleton = InSkelComp->GetSkeletalMesh()->GetSkeleton();
 			const int32 MeshBoneIndex = InSkelComp->GetBoneIndex(InBoneName);
 			if (MeshBoneIndex != INDEX_NONE)
 			{
diff --git a/Engine/Source/Editor/AnimGraph/Private/EditModes/ModifyBoneEditMode.cpp b/Engine/Source/Editor/AnimGraph/Private/EditModes/ModifyBoneEditMode.cpp
--- a/Engine/Source/Editor/AnimGraph/Private/EditModes/ModifyBoneEditMode.cpp
+++ b/Engine/Source/Editor/AnimGraph/Private/EditModes/ModifyBoneEditMode.cpp
@@ -59,7 +59,7 @@ ECoordSystem FModifyBoneEditMode::GetWidgetCoordinateSystem() const
 FVector FModifyBoneEditMode::GetWidgetLocation() const
 {
 	USkeletalMeshComponent* SkelComp = GetAnimPreviewScene().GetPreviewMeshComponent();
-	USkeleton* Skeleton = SkelComp->SkeletalMesh->GetSkeleton();
+	USkeleton* Skeleton = SkelComp->GetSkeletalMesh()->GetSkeleton();
 	FVector WidgetLoc = FVector::ZeroVector;
 
 	// if the current widget mode is translate, then shows the widget according to translation space
diff --git a/Engine/Source/Editor/AnimGraph/Private/EditModes/ObserveBoneEditMode.cpp b/Engine/Source/Editor/AnimGraph/Private/EditModes/ObserveBoneEditMode.cpp
--- a/Engine/Source/Editor/AnimGraph/Private/EditModes/ObserveBoneEditMode.cpp
+++ b/Engine/Source/Editor/AnimGraph/Private/EditModes/ObserveBoneEditMode.cpp
@@ -40,7 +40,7 @@ ECoordSystem FObserveBoneEditMode::GetWidgetCoordinateSystem() const
 FVector FObserveBoneEditMode::GetWidgetLocation() const
 {
 	USkeletalMeshComponent* SkelComp = GetAnimPreviewScene().GetPreviewMeshComponent();
-	USkeleton* Skeleton = SkelComp->SkeletalMesh->GetSkeleton();
+	USkeleton* Skeleton = SkelComp->GetSkeletalMesh()->GetSkeleton();
 	FVector WidgetLoc = FVector::ZeroVector;
 
 	int32 MeshBoneIndex = SkelComp->GetBoneIndex(GraphNode->Node.BoneToObserve.BoneName);
diff --git a/Engine/Source/Editor/AnimGraph/Private/EditModes/TwoBoneIKEditMode.cpp b/Engine/Source/Editor/AnimGraph/Private/EditModes/TwoBoneIKEditMode.cpp
--- a/Engine/Source/Editor/AnimGraph/Private/EditModes/TwoBoneIKEditMode.cpp
+++ b/Engine/Source/Editor/AnimGraph/Private/EditModes/TwoBoneIKEditMode.cpp
@@ -56,9 +56,9 @@ void FTwoBoneIKEditMode::ExitMode()
 void FTwoBoneIKEditMode::Render(const FSceneView* View, FViewport* Viewport, FPrimitiveDrawInterface* PDI)
 {
 	USkeletalMeshComponent* SkelMeshComp = GetAnimPreviewScene().GetPreviewMeshComponent();
-	if (SkelMeshComp && SkelMeshComp->SkeletalMesh && SkelMeshComp->SkeletalMesh->GetSkeleton())
+	if (SkelMeshComp && SkelMeshComp->GetSkeletalMesh() && SkelMeshComp->GetSkeletalMesh()->GetSkeleton())
 	{
-		USkeleton* Skeleton = SkelMeshComp->SkeletalMesh->GetSkeleton();
+		USkeleton* Skeleton = SkelMeshComp->GetSkeletalMesh()->GetSkeleton();
 
 		PDI->SetHitProxy(new HTwoBoneIKProxy(BSM_EndEffector));
 		DrawTargetLocation(PDI, BSM_EndEffector, FColor(255, 128, 128), FColor(180, 128, 128));
diff --git a/Engine/Source/Editor/AnimationEditor/Private/AnimationEditor.cpp b/Engine/Source/Editor/AnimationEditor/Private/AnimationEditor.cpp
--- a/Engine/Source/Editor/AnimationEditor/Private/AnimationEditor.cpp
+++ b/Engine/Source/Editor/AnimationEditor/Private/AnimationEditor.cpp
@@ -662,7 +662,7 @@ bool FAnimationEditor::ExportToFBX(const TArray<UObject*> AssetsToExport, bool b
 		FPersonaModule& PersonaModule = FModuleManager::GetModuleChecked<FPersonaModule>("Persona");
 		
 		
-		AnimSequenceExportResult = PersonaModule.ExportToFBX(AnimSequences, GetPersonaToolkit()->GetPreviewScene()->GetPreviewMeshComponent()->SkeletalMesh);
+		AnimSequenceExportResult = PersonaModule.ExportToFBX(AnimSequences, GetPersonaToolkit()->GetPreviewScene()->GetPreviewMeshComponent()->GetSkeletalMesh());
 	}
 	return AnimSequenceExportResult;
 }
diff --git a/Engine/Source/Editor/ClothPainter/Private/ClothMeshAdapter.cpp b/Engine/Source/Editor/ClothPainter/Private/ClothMeshAdapter.cpp
--- a/Engine/Source/Editor/ClothPainter/Private/ClothMeshAdapter.cpp
+++ b/Engine/Source/Editor/ClothPainter/Private/ClothMeshAdapter.cpp
@@ -15,7 +15,7 @@ TSharedPtr<IMeshPaintGeometryAdapter> FClothMeshPaintAdapterFactory::Construct(U
 {
 	if (USkeletalMeshComponent* Component = Cast<USkeletalMeshComponent>(InComponent))
 	{
-		if (Component->SkeletalMesh != nullptr)
+		if (Component->GetSkeletalMesh() != nullptr)
 		{
 			TSharedRef<FClothMeshPaintAdapter> Result = MakeShareable(new FClothMeshPaintAdapter());
 			if (Result->Construct(InComponent, InPaintingMeshLODIndex))
@@ -33,9 +33,9 @@ bool FClothMeshPaintAdapter::Construct(UMeshComponent* InComponent, int32 InPain
 	SkeletalMeshComponent = Cast<USkeletalMeshComponent>(InComponent);
 	if (SkeletalMeshComponent != nullptr)
 	{
-		if (SkeletalMeshComponent->SkeletalMesh != nullptr)
+		if (SkeletalMeshComponent->GetSkeletalMesh() != nullptr)
 		{
-			ReferencedSkeletalMesh = SkeletalMeshComponent->SkeletalMesh;
+			ReferencedSkeletalMesh = SkeletalMeshComponent->GetSkeletalMesh();
 			PaintingClothLODIndex = InPaintingClothLODIndex;
 			PaintingClothMaskIndex = INDEX_NONE;
 
@@ -49,7 +49,7 @@ bool FClothMeshPaintAdapter::Construct(UMeshComponent* InComponent, int32 InPain
 
 bool FClothMeshPaintAdapter::Initialize()
 {
-	check(ReferencedSkeletalMesh == SkeletalMeshComponent->SkeletalMesh);
+	check(ReferencedSkeletalMesh == SkeletalMeshComponent->GetSkeletalMesh());
 
 	bool bHaveAsset = false;
 	bool bBaseInit = false;
diff --git a/Engine/Source/Editor/ClothPainter/Private/ClothPainter.cpp b/Engine/Source/Editor/ClothPainter/Private/ClothPainter.cpp
--- a/Engine/Source/Editor/ClothPainter/Private/ClothPainter.cpp
+++ b/Engine/Source/Editor/ClothPainter/Private/ClothPainter.cpp
@@ -73,7 +73,7 @@ bool FClothPainter::PaintInternal(const FVector& InCameraOrigin, const TArrayVie
 
 	if(SkeletalMeshComponent->SelectedClothingGuidForPainting.IsValid() && !bShouldSimulate)
 	{
-		USkeletalMesh* SkelMesh = SkeletalMeshComponent->SkeletalMesh;
+		USkeletalMesh* SkelMesh = SkeletalMeshComponent->GetSkeletalMesh();
 
 		for (const TPair<FVector, FVector>& Ray : Rays)
 		{
@@ -160,7 +160,7 @@ USkeletalMesh* FClothPainter::GetSkeletalMesh() const
 {
 	if(SkeletalMeshComponent)
 	{
-		return SkeletalMeshComponent->SkeletalMesh;
+		return SkeletalMeshComponent->GetSkeletalMesh();
 	}
 
 	return nullptr;
@@ -175,7 +175,7 @@ void FClothPainter::RefreshClothingAssets()
 
 	PaintSettings->ClothingAssets.Reset();
 
-	if(USkeletalMesh* Mesh = SkeletalMeshComponent->SkeletalMesh)
+	if(USkeletalMesh* Mesh = SkeletalMeshComponent->GetSkeletalMesh())
 	{
 		for(UClothingAssetBase* BaseClothingAsset : Mesh->GetMeshClothingAssets())
 		{
diff --git a/Engine/Source/Editor/ClothPainter/Private/ClothingPaintEditMode.cpp b/Engine/Source/Editor/ClothPainter/Private/ClothingPaintEditMode.cpp
--- a/Engine/Source/Editor/ClothPainter/Private/ClothingPaintEditMode.cpp
+++ b/Engine/Source/Editor/ClothPainter/Private/ClothingPaintEditMode.cpp
@@ -88,7 +88,7 @@ void FClothingPaintEditMode::Exit()
 		{
 			MeshComponent->bDisableClothSimulation = false;
 
-			if(USkeletalMesh* SkelMesh = MeshComponent->SkeletalMesh)
+			if(USkeletalMesh* SkelMesh = MeshComponent->GetSkeletalMesh())
 			{
 				for(UClothingAssetBase* AssetBase : SkelMesh->GetMeshClothingAssets())
 				{
@@ -113,7 +113,7 @@ void FClothingPaintEditMode::Exit()
 			for(TObjectIterator<USkeletalMeshComponent> It; It; ++It)
 			{
 				USkeletalMeshComponent* Component = *It;
-				if(Component && !Component->IsTemplate() && Component->SkeletalMesh == SkelMesh)
+				if(Component && !Component->IsTemplate() && Component->GetSkeletalMesh() == SkelMesh)
 				{
 					Component->ReregisterComponent();
 				}
diff --git a/Engine/Source/Editor/DetailCustomizations/Private/AnimMontageSegmentDetails.cpp b/Engine/Source/Editor/DetailCustomizations/Private/AnimMontageSegmentDetails.cpp
--- a/Engine/Source/Editor/DetailCustomizations/Private/AnimMontageSegmentDetails.cpp
+++ b/Engine/Source/Editor/DetailCustomizations/Private/AnimMontageSegmentDetails.cpp
@@ -517,7 +517,7 @@ void SAnimationSegmentViewport::InitSkeleton()
 		if (Skeleton && PreviewMesh)
 		{
 			UAnimSingleNodeInstance* Preview = PreviewComponent->PreviewInstance;
-			if((Preview == nullptr || Preview->GetCurrentAsset() != AnimSequenceBase) || (PreviewComponent->SkeletalMesh != PreviewMesh))
+			if((Preview == nullptr || Preview->GetCurrentAsset() != AnimSequenceBase) || (PreviewComponent->GetSkeletalMesh() != PreviewMesh))
 			{
 				const float PlayRate = PlayRateAttribute.Get().Get(1.f);
 
@@ -565,7 +565,7 @@ void SAnimationSegmentViewport::Tick( const FGeometry& AllottedGeometry, const d
 		{
 			Description->SetText(FText::Format( LOCTEXT("Previewing", "Previewing {0}"), FText::FromString(Component->AnimClass->GetName()) ));
 		}
-		else if (Component->SkeletalMesh == NULL)
+		else if (Component->GetSkeletalMesh() == NULL)
 		{
 			Description->SetText(FText::Format( LOCTEXT("NoMeshFound", "No skeletal mesh found for skeleton '{0}'"), FText::FromString(TargetSkeletonName) ));
 		}
diff --git a/Engine/Source/Editor/DetailCustomizations/Private/AnimSequenceDetails.cpp b/Engine/Source/Editor/DetailCustomizations/Private/AnimSequenceDetails.cpp
--- a/Engine/Source/Editor/DetailCustomizations/Private/AnimSequenceDetails.cpp
+++ b/Engine/Source/Editor/DetailCustomizations/Private/AnimSequenceDetails.cpp
@@ -509,7 +509,7 @@ void SAnimationRefPoseViewport::InitSkeleton()
 		}();
 		
 		const bool bInvalidPreviewInstance = PreviewAnimInstance == nullptr || PreviewAnimInstance->GetCurrentAsset() != PreviewAnimationSequence;
-		const bool bPreviewMeshMismatch = PreviewComponent->SkeletalMesh != PreviewSkeletalMesh;
+		const bool bPreviewMeshMismatch = PreviewComponent->GetSkeletalMesh() != PreviewSkeletalMesh;
 		if(bInvalidPreviewInstance || bPreviewMeshMismatch)
 		{
 			PreviewComponent->SetSkeletalMesh(PreviewSkeletalMesh);
@@ -572,7 +572,7 @@ void SAnimationRefPoseViewport::Tick( const FGeometry& AllottedGeometry, const d
 		{
 			Description->SetText( FText::Format( LOCTEXT( "IncorrectSkeleton", "The preview asset doesn't work for the skeleton '{0}'" ), FText::FromString( TargetSkeletonName ) ) );
 		}
-		else if ( Component->SkeletalMesh == NULL )
+		else if ( Component->GetSkeletalMesh() == NULL )
 		{
 			Description->SetText( FText::Format( LOCTEXT( "NoMeshFound", "No skeletal mesh found for skeleton '{0}'" ), FText::FromString( TargetSkeletonName ) ) );
 		}
diff --git a/Engine/Source/Editor/DetailCustomizations/Private/SkeletalMeshComponentDetails.cpp b/Engine/Source/Editor/DetailCustomizations/Private/SkeletalMeshComponentDetails.cpp
--- a/Engine/Source/Editor/DetailCustomizations/Private/SkeletalMeshComponentDetails.cpp
+++ b/Engine/Source/Editor/DetailCustomizations/Private/SkeletalMeshComponentDetails.cpp
@@ -397,7 +397,7 @@ USkeleton* FSkeletalMeshComponentDetails::GetValidSkeletonFromRegisteredMeshes()
 	for (auto ObjectIter = SelectedObjects.CreateConstIterator(); ObjectIter; ++ObjectIter)
 	{
 		USkeletalMeshComponent* const Mesh = Cast<USkeletalMeshComponent>(ObjectIter->Get());
-		if ( !Mesh || !Mesh->SkeletalMesh )
+		if ( !Mesh || !Mesh->GetSkeletalMesh())
 		{
 			continue;
 		}
@@ -405,14 +405,14 @@ USkeleton* FSkeletalMeshComponentDetails::GetValidSkeletonFromRegisteredMeshes()
 		// If we've not come across a valid skeleton yet, store this one.
 		if (!ResultSkeleton)
 		{
-			ResultSkeleton = Mesh->SkeletalMesh->GetSkeleton();
+			ResultSkeleton = Mesh->GetSkeletalMesh()->GetSkeleton();
 			continue;
 		}
 
 		// We've encountered a valid skeleton before.
 		// If this skeleton is not the same one, that means there are multiple
 		// skeletons selected, so we don't want to take any action.
-		if (Mesh->SkeletalMesh->GetSkeleton() != ResultSkeleton)
+		if (Mesh->GetSkeletalMesh()->GetSkeleton() != ResultSkeleton)
 		{
 			return NULL;
 		}
diff --git a/Engine/Source/Editor/DetailCustomizations/Private/SkinnedMeshComponentDetails.cpp b/Engine/Source/Editor/DetailCustomizations/Private/SkinnedMeshComponentDetails.cpp
--- a/Engine/Source/Editor/DetailCustomizations/Private/SkinnedMeshComponentDetails.cpp
+++ b/Engine/Source/Editor/DetailCustomizations/Private/SkinnedMeshComponentDetails.cpp
@@ -183,7 +183,7 @@ void FSkinnedMeshComponentDetails::PopulateSkinWeightProfileNames()
 	// Retrieve all possible skin weight profiles from the component
 	if (USkinnedMeshComponent* Component = WeakSkinnedMeshComponent.Get())
 	{
-		if (USkeletalMesh* Mesh = Component->SkeletalMesh)
+		if (USkeletalMesh* Mesh = Component->GetSkeletalMesh())
 		{			
 			for (const FSkinWeightProfileInfo& Profile : Mesh->GetSkinWeightProfiles())
 			{
diff --git a/Engine/Source/Editor/LevelEditor/Private/LevelEditorActions.cpp b/Engine/Source/Editor/LevelEditor/Private/LevelEditorActions.cpp
--- a/Engine/Source/Editor/LevelEditor/Private/LevelEditorActions.cpp
+++ b/Engine/Source/Editor/LevelEditor/Private/LevelEditorActions.cpp
@@ -2048,7 +2048,7 @@ bool FLevelEditorActionCallbacks::SaveAnimationFromSkeletalMeshComponent(AActor
 		bool bSimulating = false;
 		for (auto & Comp : SimSkelComponents)
 		{
-			bSimulating |= (Comp->SkeletalMesh && Comp->SkeletalMesh->GetSkeleton() && Comp->IsSimulatingPhysics());
+			bSimulating |= (Comp->GetSkeletalMesh() && Comp->GetSkeletalMesh()->GetSkeleton() && Comp->IsSimulatingPhysics());
 		}
 
 		// if any of them are legitimately simulating
@@ -2062,7 +2062,7 @@ bool FLevelEditorActionCallbacks::SaveAnimationFromSkeletalMeshComponent(AActor
 			{
 				for (auto & Comp : SimSkelComponents)
 				{
-					if (Comp->SkeletalMesh && Comp->SkeletalMesh->GetSkeleton() && Comp->IsSimulatingPhysics())
+					if (Comp->GetSkeletalMesh() && Comp->GetSkeletalMesh()->GetSkeleton() && Comp->IsSimulatingPhysics())
 					{
 						// now record to animation
 						class UAnimSequence* Sequence = LevelEditorModule.OnCaptureSingleFrameAnimSequence().IsBound() ? LevelEditorModule.OnCaptureSingleFrameAnimSequence().Execute(Comp) : nullptr;
@@ -2151,7 +2151,7 @@ void FLevelEditorActionCallbacks::OnKeepSimulationChanges()
 				// need to reinitialize animation
 				for (auto MeshComp : ComponentsToReinitialize)
 				{
-					if(MeshComp->SkeletalMesh)
+					if(MeshComp->GetSkeletalMesh())
 					{
 						MeshComp->InitAnim(true);
 					}
diff --git a/Engine/Source/Editor/MeshPaint/Private/MeshPaintHelpers.cpp b/Engine/Source/Editor/MeshPaint/Private/MeshPaintHelpers.cpp
--- a/Engine/Source/Editor/MeshPaint/Private/MeshPaintHelpers.cpp
+++ b/Engine/Source/Editor/MeshPaint/Private/MeshPaintHelpers.cpp
@@ -833,7 +833,7 @@ void MeshPaintHelpers::FillStaticMeshVertexColors(UStaticMeshComponent* MeshComp
 void MeshPaintHelpers::FillSkeletalMeshVertexColors(USkeletalMeshComponent* MeshComponent, int32 LODIndex, const FColor FillColor, const FColor MaskColor)
 {
 	TUniquePtr< FSkinnedMeshComponentRecreateRenderStateContext > RecreateRenderStateContext;
-	USkeletalMesh* Mesh = MeshComponent->SkeletalMesh;
+	USkeletalMesh* Mesh = MeshComponent->GetSkeletalMesh();
 	if (Mesh)
 	{
 		// Dirty the mesh
@@ -1003,10 +1003,10 @@ void MeshPaintHelpers::ImportVertexColorsFromTexture(UMeshComponent* MeshCompone
 				{
 					USkeletalMeshComponent* SkeletalMeshComponent = Cast<USkeletalMeshComponent>(MeshComponent);
 
-					if (SkeletalMeshComponent->SkeletalMesh)
+					if (SkeletalMeshComponent->GetSkeletalMesh())
 					{
 						// Import colors to skeletal mesh
-						ImportVertexColorsToSkeletalMesh(SkeletalMeshComponent->SkeletalMesh, Options, ColorTexture);
+						ImportVertexColorsToSkeletalMesh(SkeletalMeshComponent->GetSkeletalMesh(), Options, ColorTexture);
 					}			
 				}
 			}
@@ -1254,7 +1254,7 @@ bool MeshPaintHelpers::TryGetNumberOfLODs(const UMeshComponent* MeshComponent, i
 	}
 	else if (const USkeletalMeshComponent* SkeletalMeshComponent = Cast<USkeletalMeshComponent>(MeshComponent))
 	{
-		const USkeletalMesh* SkeletalMesh = SkeletalMeshComponent->SkeletalMesh;
+		const USkeletalMesh* SkeletalMesh = SkeletalMeshComponent->GetSkeletalMesh();
 		if (SkeletalMesh != nullptr)
 		{
 			OutNumLODs = SkeletalMesh->GetLODNum();
@@ -1286,7 +1286,7 @@ int32 MeshPaintHelpers::GetNumberOfUVs(const UMeshComponent* MeshComponent, int3
 	}
 	else if (const USkeletalMeshComponent* SkeletalMeshComponent = Cast<USkeletalMeshComponent>(MeshComponent))
 	{
-		const USkeletalMesh* SkeletalMesh = SkeletalMeshComponent->SkeletalMesh;
+		const USkeletalMesh* SkeletalMesh = SkeletalMeshComponent->GetSkeletalMesh();
 		if (SkeletalMesh != nullptr && SkeletalMesh->GetResourceForRendering() && SkeletalMesh->GetResourceForRendering()->LODRenderData.IsValidIndex(LODIndex))
 		{
 			NumUVs = SkeletalMesh->GetResourceForRendering()->LODRenderData[LODIndex].GetNumTexCoords();
@@ -1324,7 +1324,7 @@ bool MeshPaintHelpers::DoesMeshComponentContainPerLODColors(const UMeshComponent
 	}
 	else if (const USkeletalMeshComponent* SkeletalMeshComponent = Cast<USkeletalMeshComponent>(MeshComponent))
 	{
-		USkeletalMesh* SkeletalMesh = SkeletalMeshComponent->SkeletalMesh;
+		USkeletalMesh* SkeletalMesh = SkeletalMeshComponent->GetSkeletalMesh();
 		if (SkeletalMesh)
 		{
 			const TArray<FSkeletalMeshLODInfo>& LODInfo = SkeletalMesh->GetLODInfoArray();
@@ -1694,7 +1694,7 @@ typedef TOctree2<FPaintedMeshVertex, FVertexColorPropogationOctreeSemantics> TVe
 void MeshPaintHelpers::ApplyVertexColorsToAllLODs(IMeshPaintGeometryAdapter& GeometryInfo, USkeletalMeshComponent* SkeletalMeshComponent)
 {
 	checkf(SkeletalMeshComponent != nullptr, TEXT("Invalid Skeletal Mesh Component"));
-	USkeletalMesh* Mesh = SkeletalMeshComponent->SkeletalMesh;
+	USkeletalMesh* Mesh = SkeletalMeshComponent->GetSkeletalMesh();
 	if (Mesh)
 	{
 		FSkeletalMeshRenderData* Resource = Mesh->GetResourceForRendering();
diff --git a/Engine/Source/Editor/MeshPaint/Private/MeshPaintSkeletalMeshAdapter.cpp b/Engine/Source/Editor/MeshPaint/Private/MeshPaintSkeletalMeshAdapter.cpp
--- a/Engine/Source/Editor/MeshPaint/Private/MeshPaintSkeletalMeshAdapter.cpp
+++ b/Engine/Source/Editor/MeshPaint/Private/MeshPaintSkeletalMeshAdapter.cpp
@@ -108,9 +108,9 @@ bool FMeshPaintGeometryAdapterForSkeletalMeshes::Construct(UMeshComponent* InCom
 	{
 		SkeletalMeshChangedHandle = SkeletalMeshComponent->RegisterOnSkeletalMeshPropertyChanged(USkeletalMeshComponent::FOnSkeletalMeshPropertyChanged::CreateRaw(this, &FMeshPaintGeometryAdapterForSkeletalMeshes::OnSkeletalMeshChanged));
 
-		if (SkeletalMeshComponent->SkeletalMesh != nullptr)
+		if (SkeletalMeshComponent->GetSkeletalMesh() != nullptr)
 		{
-			ReferencedSkeletalMesh = SkeletalMeshComponent->SkeletalMesh;
+			ReferencedSkeletalMesh = SkeletalMeshComponent->GetSkeletalMesh();
 			MeshLODIndex = InMeshLODIndex;
 			const bool bSuccess = Initialize();
 			return bSuccess;
@@ -131,8 +131,8 @@ FMeshPaintGeometryAdapterForSkeletalMeshes::~FMeshPaintGeometryAdapterForSkeleta
 void FMeshPaintGeometryAdapterForSkeletalMeshes::OnSkeletalMeshChanged()
 {
 	OnRemoved();
-	ReferencedSkeletalMesh = SkeletalMeshComponent->SkeletalMesh;
-	if (SkeletalMeshComponent->SkeletalMesh != nullptr)
+	ReferencedSkeletalMesh = SkeletalMeshComponent->GetSkeletalMesh();
+	if (SkeletalMeshComponent->GetSkeletalMesh() != nullptr)
 	{	
 		Initialize();
 		OnAdded();
@@ -149,7 +149,7 @@ void FMeshPaintGeometryAdapterForSkeletalMeshes::OnPostMeshCached(USkeletalMesh*
 
 bool FMeshPaintGeometryAdapterForSkeletalMeshes::Initialize()
 {
-	check(ReferencedSkeletalMesh == SkeletalMeshComponent->SkeletalMesh);
+	check(ReferencedSkeletalMesh == SkeletalMeshComponent->GetSkeletalMesh());
 
 	bool bInitializationResult = false;
 
@@ -212,7 +212,7 @@ void FMeshPaintGeometryAdapterForSkeletalMeshes::OnAdded()
 {
 	checkf(SkeletalMeshComponent, TEXT("Invalid SkeletalMesh Component"));
 	checkf(ReferencedSkeletalMesh, TEXT("Invalid reference to Skeletal Mesh"));
-	checkf(ReferencedSkeletalMesh == SkeletalMeshComponent->SkeletalMesh, TEXT("Referenced Skeletal Mesh does not match one in Component"));
+	checkf(ReferencedSkeletalMesh == SkeletalMeshComponent->GetSkeletalMesh(), TEXT("Referenced Skeletal Mesh does not match one in Component"));
 
 	FSkeletalMeshReferencers& SkeletalMeshReferencers = MeshToComponentMap.FindOrAdd(ReferencedSkeletalMesh);
 
@@ -515,7 +515,7 @@ TSharedPtr<IMeshPaintGeometryAdapter> FMeshPaintGeometryAdapterForSkeletalMeshes
 {
 	if (USkeletalMeshComponent* SkeletalMeshComponent = Cast<USkeletalMeshComponent>(InComponent))
 	{
-		if (SkeletalMeshComponent->SkeletalMesh != nullptr)
+		if (SkeletalMeshComponent->GetSkeletalMesh() != nullptr)
 		{
 			TSharedRef<FMeshPaintGeometryAdapterForSkeletalMeshes> Result = MakeShareable(new FMeshPaintGeometryAdapterForSkeletalMeshes());
 			if (Result->Construct(InComponent, InMeshLODIndex))
diff --git a/Engine/Source/Editor/MeshPaint/Private/TexturePaintHelpers.cpp b/Engine/Source/Editor/MeshPaint/Private/TexturePaintHelpers.cpp
--- a/Engine/Source/Editor/MeshPaint/Private/TexturePaintHelpers.cpp
+++ b/Engine/Source/Editor/MeshPaint/Private/TexturePaintHelpers.cpp
@@ -450,7 +450,7 @@ void TexturePaintHelpers::RetrieveMeshSectionsForMaterialIndices(const UMeshComp
 	}
 	else if (const USkeletalMeshComponent* SkeletalMeshComponent = Cast<const USkeletalMeshComponent>(MeshComponent))
 	{
-		const USkeletalMesh* SkeletalMesh = SkeletalMeshComponent->SkeletalMesh;
+		const USkeletalMesh* SkeletalMesh = SkeletalMeshComponent->GetSkeletalMesh();
 		if (SkeletalMesh)
 		{
 			const FSkeletalMeshRenderData* Resource = SkeletalMesh->GetResourceForRendering();
diff --git a/Engine/Source/Editor/MeshPaint/Public/MeshPaintSkeletalMeshAdapter.h b/Engine/Source/Editor/MeshPaint/Public/MeshPaintSkeletalMeshAdapter.h
--- a/Engine/Source/Editor/MeshPaint/Public/MeshPaintSkeletalMeshAdapter.h
+++ b/Engine/Source/Editor/MeshPaint/Public/MeshPaintSkeletalMeshAdapter.h
@@ -33,7 +33,7 @@ public:
 	virtual bool Initialize() override;
 	virtual void OnAdded() override;
 	virtual void OnRemoved() override;
-	virtual bool IsValid() const override { return SkeletalMeshComponent && ReferencedSkeletalMesh && SkeletalMeshComponent->SkeletalMesh == ReferencedSkeletalMesh; }
+	virtual bool IsValid() const override { return SkeletalMeshComponent && ReferencedSkeletalMesh && SkeletalMeshComponent->GetSkeletalMesh() == ReferencedSkeletalMesh; }
 	virtual bool SupportsTexturePaint() const override { return true; }
 	virtual bool SupportsVertexPaint() const override { return SkeletalMeshComponent != nullptr; }
 	virtual bool LineTraceComponent(struct FHitResult& OutHit, const FVector Start, const FVector End, const struct FCollisionQueryParams& Params) const override;
diff --git a/Engine/Source/Editor/MovieSceneTools/Private/EditModes/SkeletalAnimationTrackEditMode.cpp b/Engine/Source/Editor/MovieSceneTools/Private/EditModes/SkeletalAnimationTrackEditMode.cpp
--- a/Engine/Source/Editor/MovieSceneTools/Private/EditModes/SkeletalAnimationTrackEditMode.cpp
+++ b/Engine/Source/Editor/MovieSceneTools/Private/EditModes/SkeletalAnimationTrackEditMode.cpp
@@ -179,11 +179,11 @@ static void DrawBonesFromCompactPose(const FCompactPose& Pose, USkeletalMeshComp
 	{
 
 		FSkeletalMeshLODRenderData* LODData = nullptr;
-		if (MeshComponent->SkeletalMesh && MeshComponent->SkeletalMesh->GetResourceForRendering()->LODRenderData.Num())
+		if (MeshComponent->GetSkeletalMesh() && MeshComponent->GetSkeletalMesh()->GetResourceForRendering()->LODRenderData.Num())
 		{
 			//Get current LOD
-			const int32 LODIndex = FMath::Clamp(MeshComponent->GetPredictedLODLevel(), 0, MeshComponent->SkeletalMesh->GetResourceForRendering()->LODRenderData.Num() - 1);
-			LODData = &(MeshComponent->SkeletalMesh->GetResourceForRendering()->LODRenderData[LODIndex]);
+			const int32 LODIndex = FMath::Clamp(MeshComponent->GetPredictedLODLevel(), 0, MeshComponent->GetSkeletalMesh()->GetResourceForRendering()->LODRenderData.Num() - 1);
+			LODData = &(MeshComponent->GetSkeletalMesh()->GetResourceForRendering()->LODRenderData[LODIndex]);
 		}
 
 		TArray<FTransform> WorldTransforms;
@@ -219,9 +219,9 @@ static void DrawBonesFromCompactPose(const FCompactPose& Pose, USkeletalMeshComp
 			BoneColours[MeshBoneIndex.GetInt()] = DrawColour;
 		}
 
-		if (MeshComponent && MeshComponent->SkeletalMesh)
+		if (MeshComponent && MeshComponent->GetSkeletalMesh())
 		{
-			DrawBones(Pose.GetBoneContainer().GetBoneIndicesArray(), MeshComponent->SkeletalMesh->GetRefSkeleton(), WorldTransforms,
+			DrawBones(Pose.GetBoneContainer().GetBoneIndicesArray(), MeshComponent->GetSkeletalMesh()->GetRefSkeleton(), WorldTransforms,
 				LODData, PDI, BoneColours, MeshComponent->Bounds.SphereRadius, 1.0f);
 		}
 
diff --git a/Engine/Source/Editor/MovieSceneTools/Private/MovieSceneToolHelpers.cpp b/Engine/Source/Editor/MovieSceneTools/Private/MovieSceneToolHelpers.cpp
--- a/Engine/Source/Editor/MovieSceneTools/Private/MovieSceneToolHelpers.cpp
+++ b/Engine/Source/Editor/MovieSceneTools/Private/MovieSceneToolHelpers.cpp
@@ -4208,7 +4208,7 @@ USkeletalMeshComponent* MovieSceneToolHelpers::AcquireSkeletalMeshFromObject(UOb
 	}
 	else if (USkeletalMeshComponent* SkeletalMeshComponent = Cast<USkeletalMeshComponent>(BoundObject))
 	{
-		if (SkeletalMeshComponent->SkeletalMesh)
+		if (SkeletalMeshComponent->GetSkeletalMesh())
 		{
 			return SkeletalMeshComponent;
 		}
diff --git a/Engine/Source/Editor/MovieSceneTools/Private/TrackEditors/SkeletalAnimationTrackEditor.cpp b/Engine/Source/Editor/MovieSceneTools/Private/TrackEditors/SkeletalAnimationTrackEditor.cpp
--- a/Engine/Source/Editor/MovieSceneTools/Private/TrackEditors/SkeletalAnimationTrackEditor.cpp
+++ b/Engine/Source/Editor/MovieSceneTools/Private/TrackEditors/SkeletalAnimationTrackEditor.cpp
@@ -315,7 +315,7 @@ USkeletalMeshComponent* AcquireSkeletalMeshFromObjectGuid(const FGuid& Guid, TSh
 	}
 	else if(USkeletalMeshComponent* SkeletalMeshComponent = Cast<USkeletalMeshComponent>(BoundObject))
 	{
-		if (SkeletalMeshComponent->SkeletalMesh)
+		if (SkeletalMeshComponent->GetSkeletalMesh())
 		{
 			return SkeletalMeshComponent;
 		}
@@ -327,10 +327,10 @@ USkeletalMeshComponent* AcquireSkeletalMeshFromObjectGuid(const FGuid& Guid, TSh
 USkeleton* GetSkeletonFromComponent(UActorComponent* InComponent)
 {
 	USkeletalMeshComponent* SkeletalMeshComp = Cast<USkeletalMeshComponent>(InComponent);
-	if (SkeletalMeshComp && SkeletalMeshComp->SkeletalMesh && SkeletalMeshComp->SkeletalMesh->GetSkeleton())
+	if (SkeletalMeshComp && SkeletalMeshComp->GetSkeletalMesh() && SkeletalMeshComp->GetSkeletalMesh()->GetSkeleton())
 	{
 		// @todo Multiple actors, multiple components
-		return SkeletalMeshComp->SkeletalMesh->GetSkeleton();
+		return SkeletalMeshComp->GetSkeletalMesh()->GetSkeleton();
 	}
 
 	return nullptr;
@@ -1548,9 +1548,9 @@ void FSkeletalAnimationTrackEditor::HandleCreateAnimationSequence(USkeletalMeshC
 	if (SkelMeshComp)
 	{
 		TArray<TWeakObjectPtr<UObject>> Skels;
-		if (SkelMeshComp->SkeletalMesh)
+		if (SkelMeshComp->GetSkeletalMesh())
 		{
-			Skels.Add(SkelMeshComp->SkeletalMesh);
+			Skels.Add(SkelMeshComp->GetSkeletalMesh());
 		}
 		else
 		{
diff --git a/Engine/Source/Editor/Persona/Private/AnimationEditorPreviewScene.cpp b/Engine/Source/Editor/Persona/Private/AnimationEditorPreviewScene.cpp
--- a/Engine/Source/Editor/Persona/Private/AnimationEditorPreviewScene.cpp
+++ b/Engine/Source/Editor/Persona/Private/AnimationEditorPreviewScene.cpp
@@ -221,7 +221,7 @@ USkeletalMesh* FAnimationEditorPreviewScene::GetPreviewMesh() const
 
 void FAnimationEditorPreviewScene::SetPreviewMeshInternal(USkeletalMesh* NewPreviewMesh)
 {
-	USkeletalMesh* OldPreviewMesh = SkeletalMeshComponent->SkeletalMesh;
+	USkeletalMesh* OldPreviewMesh = SkeletalMeshComponent->GetSkeletalMesh();
 
 	// Store off the old skel mesh we are debugging
 	USkeletalMeshComponent* DebuggedSkeletalMeshComponent = nullptr;
@@ -249,7 +249,7 @@ void FAnimationEditorPreviewScene::SetPreviewMeshInternal(USkeletalMesh* NewPrev
 
 	ValidatePreviewAttachedAssets(NewPreviewMesh);
 
-	if (NewPreviewMesh != SkeletalMeshComponent->SkeletalMesh)
+	if (NewPreviewMesh != SkeletalMeshComponent->GetSkeletalMesh())
 	{
 		// setting skeletalmesh unregister/re-register, 
 		// so I have to save the animation settings and resetting after setting mesh
@@ -680,15 +680,15 @@ void FAnimationEditorPreviewScene::ShowReferencePose(bool bShowRefPose, bool bRe
 		SkeletalMeshComponent->ShowReferencePose(bShowRefPose);
 
 		// Also reset bone transforms
-		if(bResetBoneTransforms && SkeletalMeshComponent->SkeletalMesh != nullptr)
+		if(bResetBoneTransforms && SkeletalMeshComponent->GetSkeletalMesh() != nullptr)
 		{
 			bool bModified = false;
 			FScopedTransaction Transaction(LOCTEXT("ResetBoneTransforms", "Reset Bone Transforms"));
 
-			int32 NumBones = SkeletalMeshComponent->SkeletalMesh->GetRefSkeleton().GetNum();
+			int32 NumBones = SkeletalMeshComponent->GetSkeletalMesh()->GetRefSkeleton().GetNum();
 			for (int32 BoneIndex = 0; BoneIndex < NumBones; ++BoneIndex)
 			{
-				FName BoneName = SkeletalMeshComponent->SkeletalMesh->GetRefSkeleton().GetBoneName(BoneIndex);
+				FName BoneName = SkeletalMeshComponent->GetSkeletalMesh()->GetRefSkeleton().GetBoneName(BoneIndex);
 				const FAnimNode_ModifyBone* ModifiedBone = SkeletalMeshComponent->PreviewInstance->FindModifiedBone(BoneName);
 				if (ModifiedBone != nullptr)
 				{
@@ -898,7 +898,7 @@ void FAnimationEditorPreviewScene::DeselectAll()
 bool FAnimationEditorPreviewScene::IsRecordAvailable() const
 {
 	// make sure mesh exists
-	return (SkeletalMeshComponent->SkeletalMesh != nullptr);
+	return (SkeletalMeshComponent->GetSkeletalMesh() != nullptr);
 }
 
 FSlateIcon FAnimationEditorPreviewScene::GetRecordStatusImage() const
diff --git a/Engine/Source/Editor/Persona/Private/AnimationEditorViewportClient.cpp b/Engine/Source/Editor/Persona/Private/AnimationEditorViewportClient.cpp
--- a/Engine/Source/Editor/Persona/Private/AnimationEditorViewportClient.cpp
+++ b/Engine/Source/Editor/Persona/Private/AnimationEditorViewportClient.cpp
@@ -148,7 +148,7 @@ FAnimationViewportClient::FAnimationViewportClient(const TSharedRef<IPersonaPrev
 	InPreviewScene->RegisterOnPreviewMeshChanged(FOnPreviewMeshChanged::CreateRaw(this, &FAnimationViewportClient::HandleSkeletalMeshChanged));
 	if (InPreviewScene->GetPreviewMeshComponent())
 	{
-		HandleSkeletalMeshChanged(nullptr, InPreviewScene->GetPreviewMeshComponent()->SkeletalMesh);
+		HandleSkeletalMeshChanged(nullptr, InPreviewScene->GetPreviewMeshComponent()->GetSkeletalMesh());
 	}
 	InPreviewScene->RegisterOnInvalidateViews(FSimpleDelegate::CreateRaw(this, &FAnimationViewportClient::HandleInvalidateViews));
 	InPreviewScene->RegisterOnFocusViews(FSimpleDelegate::CreateRaw(this, &FAnimationViewportClient::HandleFocusViews));
@@ -184,7 +184,7 @@ FAnimationViewportClient::~FAnimationViewportClient()
 
 			if (OnMeshChangedDelegateHandle.IsValid())
 			{
-				if (USkeletalMesh* SkelMesh = PreviewMeshComponent->SkeletalMesh)
+				if (USkeletalMesh* SkelMesh = PreviewMeshComponent->GetSkeletalMesh())
 				{
 					SkelMesh->GetOnMeshChanged().Remove(OnMeshChangedDelegateHandle);
 				}
@@ -375,7 +375,7 @@ void FAnimationViewportClient::JumpToDefaultCamera()
 
 void FAnimationViewportClient::SaveCameraAsDefault()
 {
-	USkeletalMesh* SkelMesh = GetAnimPreviewScene()->GetPreviewMeshComponent()->SkeletalMesh;
+	USkeletalMesh* SkelMesh = GetAnimPreviewScene()->GetPreviewMeshComponent()->GetSkeletalMesh();
 	if (SkelMesh)
 	{
 		FScopedTransaction Transaction(LOCTEXT("SaveCameraAsDefault", "Save Camera As Default"));
@@ -403,7 +403,7 @@ bool FAnimationViewportClient::CanSaveCameraAsDefault() const
 
 void FAnimationViewportClient::ClearDefaultCamera()
 {
-	USkeletalMesh* SkelMesh = GetAnimPreviewScene()->GetPreviewMeshComponent()->SkeletalMesh;
+	USkeletalMesh* SkelMesh = GetAnimPreviewScene()->GetPreviewMeshComponent()->GetSkeletalMesh();
 	if (SkelMesh)
 	{
 		FScopedTransaction Transaction(LOCTEXT("ClearDefaultCamera", "Clear Default Camera"));
@@ -421,7 +421,7 @@ void FAnimationViewportClient::ClearDefaultCamera()
 
 bool FAnimationViewportClient::HasDefaultCameraSet() const
 {
-	USkeletalMesh* SkelMesh = GetAnimPreviewScene()->GetPreviewMeshComponent()->SkeletalMesh;
+	USkeletalMesh* SkelMesh = GetAnimPreviewScene()->GetPreviewMeshComponent()->GetSkeletalMesh();
 	return (SkelMesh && SkelMesh->GetHasCustomDefaultEditorCamera());
 }
 
@@ -509,7 +509,7 @@ void FAnimationViewportClient::Draw(const FSceneView* View, FPrimitiveDrawInterf
 	TArray<UDebugSkelMeshComponent*> PreviewMeshComponents = GetPreviewScene()->GetAllPreviewMeshComponents();
 	for (UDebugSkelMeshComponent* PreviewMeshComponent : PreviewMeshComponents)
 	{
-		if (!(PreviewMeshComponent && PreviewMeshComponent->SkeletalMesh))
+		if (!(PreviewMeshComponent && PreviewMeshComponent->GetSkeletalMesh()))
 		{
 			continue;
 		}
@@ -555,14 +555,14 @@ void FAnimationViewportClient::Draw(const FSceneView* View, FPrimitiveDrawInterf
 		// Display socket hit points
 		if (PreviewMeshComponent->bDrawSockets )
 		{
-			if (PreviewMeshComponent->bSkeletonSocketsVisible && PreviewMeshComponent->SkeletalMesh->GetSkeleton() )
+			if (PreviewMeshComponent->bSkeletonSocketsVisible && PreviewMeshComponent->GetSkeletalMesh()->GetSkeleton() )
 			{
-				DrawSockets(PreviewMeshComponent, PreviewMeshComponent->SkeletalMesh->GetSkeleton()->Sockets, GetAnimPreviewScene()->GetSelectedSocket(), PDI, true);
+				DrawSockets(PreviewMeshComponent, PreviewMeshComponent->GetSkeletalMesh()->GetSkeleton()->Sockets, GetAnimPreviewScene()->GetSelectedSocket(), PDI, true);
 			}
 
 			if ( PreviewMeshComponent->bMeshSocketsVisible )
 			{
-				DrawSockets(PreviewMeshComponent, PreviewMeshComponent->SkeletalMesh->GetMeshOnlySocketList(), GetAnimPreviewScene()->GetSelectedSocket(), PDI, false);
+				DrawSockets(PreviewMeshComponent, PreviewMeshComponent->GetSkeletalMesh()->GetMeshOnlySocketList(), GetAnimPreviewScene()->GetSelectedSocket(), PDI, false);
 			}
 		}
 
@@ -628,7 +628,7 @@ void FAnimationViewportClient::DrawUVsForMesh(FViewport* InViewport, FCanvas* In
 	}
 	
 	//use the overridden LOD level
-	const uint32 LODLevel = FMath::Clamp(PreviewMeshComponent->GetForcedLOD() - 1, 0, PreviewMeshComponent->SkeletalMesh->GetLODNum() - 1);
+	const uint32 LODLevel = FMath::Clamp(PreviewMeshComponent->GetForcedLOD() - 1, 0, PreviewMeshComponent->GetSkeletalMesh()->GetLODNum() - 1);
 
 	TArray<FVector2D> SelectedEdgeTexCoords; //No functionality in Persona for this (yet?)
 
@@ -807,7 +807,7 @@ void FAnimationViewportClient::ShowBoneNames( FCanvas* Canvas, FSceneView* View,
 
 void FAnimationViewportClient::ShowAttributeNames(FCanvas* Canvas, FSceneView* View, UDebugSkelMeshComponent* MeshComponent) const
 {
-	if (!MeshComponent || !MeshComponent->SkeletalMesh)
+	if (!MeshComponent || !MeshComponent->GetSkeletalMesh())
 	{
 		return;
 	}
@@ -881,7 +881,7 @@ FText FAnimationViewportClient::GetDisplayInfo(bool bDisplayAllInfo) const
 
 	// if not valid skeletalmesh
 	UDebugSkelMeshComponent* PreviewMeshComponent = GetPreviewScene()->GetPreviewMeshComponent();
-	if (PreviewMeshComponent == nullptr || PreviewMeshComponent->SkeletalMesh == nullptr)
+	if (PreviewMeshComponent == nullptr || PreviewMeshComponent->GetSkeletalMesh() == nullptr)
 	{
 		return FText();
 	}
@@ -891,16 +891,16 @@ FText FAnimationViewportClient::GetDisplayInfo(bool bDisplayAllInfo) const
 		TextValue = ConcatenateLine(TextValue, LOCTEXT("AdditiveRefPoseWarning", "<AnimViewport.WarningText>Additive ref pose contains scales of 0.0, this can cause additive animations to not give the desired results</>"));
 	}
 
-	if (PreviewMeshComponent->SkeletalMesh->GetMorphTargets().Num() > 0)
+	if (PreviewMeshComponent->GetSkeletalMesh()->GetMorphTargets().Num() > 0)
 	{
 		TArray<UMaterial*> ProcessedMaterials;
 		TArray<UMaterial*> MaterialsThatNeedMorphFlagOn;
 		TArray<UMaterial*> MaterialsThatNeedSaving;
 
-		const TIndirectArray<FSkeletalMeshLODModel>& LODModels = PreviewMeshComponent->SkeletalMesh->GetImportedModel()->LODModels;
+		const TIndirectArray<FSkeletalMeshLODModel>& LODModels = PreviewMeshComponent->GetSkeletalMesh()->GetImportedModel()->LODModels;
 		int32 LodNumber = LODModels.Num();
 		TArray<UMaterialInterface*> MaterialUsingMorphTarget;
-		for (UMorphTarget *MorphTarget : PreviewMeshComponent->SkeletalMesh->GetMorphTargets())
+		for (UMorphTarget *MorphTarget : PreviewMeshComponent->GetSkeletalMesh()->GetMorphTargets())
 		{
 			if (MorphTarget == nullptr)
 			{
@@ -915,7 +915,7 @@ FText FAnimationViewportClient::GetDisplayInfo(bool bDisplayAllInfo) const
 						const FSkeletalMeshLODModel& LODModel = LODModels[LodIdx];
 						if (LODModel.Sections.IsValidIndex(SectionIndex))
 						{
-							MaterialUsingMorphTarget.AddUnique(PreviewMeshComponent->SkeletalMesh->GetMaterials()[LODModel.Sections[SectionIndex].MaterialIndex].MaterialInterface);
+							MaterialUsingMorphTarget.AddUnique(PreviewMeshComponent->GetSkeletalMesh()->GetMaterials()[LODModel.Sections[SectionIndex].MaterialIndex].MaterialInterface);
 						}
 					}
 				}
@@ -1406,7 +1406,7 @@ void FAnimationViewportClient::DrawBonesFromTransforms(
 	FLinearColor RootBoneColour) const
 {
 	if (Transforms.IsEmpty() ||
-		!MeshComponent->SkeletalMesh ||
+		!MeshComponent->GetSkeletalMesh() ||
 		MeshComponent->SkeletonDrawMode == ESkeletonDrawMode::Hidden)
 	{
 		return;
@@ -1457,7 +1457,7 @@ void FAnimationViewportClient::DrawBonesFromCompactPose(
 {
 	if (Pose.GetNumBones() == 0 ||
 		!MeshComponent ||
-		!MeshComponent->SkeletalMesh ||
+		!MeshComponent->GetSkeletalMesh() ||
 		MeshComponent->SkeletonDrawMode == ESkeletonDrawMode::Hidden)
 	{
 		return;
@@ -1507,7 +1507,7 @@ void FAnimationViewportClient::DrawBonesFromCompactPose(
 
 void FAnimationViewportClient::DrawMeshBonesUncompressedAnimation(UDebugSkelMeshComponent * MeshComponent, FPrimitiveDrawInterface* PDI) const
 {
-	if ( MeshComponent && MeshComponent->SkeletalMesh )
+	if ( MeshComponent && MeshComponent->GetSkeletalMesh())
 	{
 		DrawBonesFromTransforms(MeshComponent->UncompressedSpaceBases, MeshComponent, PDI, FColor(255, 127, 39, 255), FColor(255, 127, 39, 255));
 	}
@@ -1515,7 +1515,7 @@ void FAnimationViewportClient::DrawMeshBonesUncompressedAnimation(UDebugSkelMesh
 
 void FAnimationViewportClient::DrawMeshBonesNonRetargetedAnimation(UDebugSkelMeshComponent * MeshComponent, FPrimitiveDrawInterface* PDI) const
 {
-	if ( MeshComponent && MeshComponent->SkeletalMesh )
+	if ( MeshComponent && MeshComponent->GetSkeletalMesh())
 	{
 		DrawBonesFromTransforms(MeshComponent->NonRetargetedSpaceBases, MeshComponent, PDI, FColor(159, 159, 39, 255), FColor(159, 159, 39, 255));
 	}
@@ -1523,7 +1523,7 @@ void FAnimationViewportClient::DrawMeshBonesNonRetargetedAnimation(UDebugSkelMes
 
 void FAnimationViewportClient::DrawMeshBonesAdditiveBasePose(UDebugSkelMeshComponent * MeshComponent, FPrimitiveDrawInterface* PDI) const
 {
-	if ( MeshComponent && MeshComponent->SkeletalMesh )
+	if ( MeshComponent && MeshComponent->GetSkeletalMesh())
 	{
 		DrawBonesFromTransforms(MeshComponent->AdditiveBasePoses, MeshComponent, PDI, FColor(0, 159, 0, 255), FColor(0, 159, 0, 255));
 	}
@@ -1531,7 +1531,7 @@ void FAnimationViewportClient::DrawMeshBonesAdditiveBasePose(UDebugSkelMeshCompo
 
 void FAnimationViewportClient::DrawMeshBonesSourceRawAnimation(UDebugSkelMeshComponent * MeshComponent, FPrimitiveDrawInterface* PDI) const
 {
-	if(MeshComponent && MeshComponent->SkeletalMesh)
+	if(MeshComponent && MeshComponent->GetSkeletalMesh())
 	{
 		DrawBonesFromTransforms(MeshComponent->SourceAnimationPoses, MeshComponent, PDI, FColor(195, 195, 195, 255), FColor(195, 159, 195, 255));
 	}
@@ -1559,7 +1559,7 @@ void FAnimationViewportClient::DrawWatchedPoses(UDebugSkelMeshComponent * MeshCo
 
 void FAnimationViewportClient::DrawMeshBonesBakedAnimation(UDebugSkelMeshComponent * MeshComponent, FPrimitiveDrawInterface* PDI) const
 {
-	if(MeshComponent && MeshComponent->SkeletalMesh)
+	if(MeshComponent && MeshComponent->GetSkeletalMesh())
 	{
 		DrawBonesFromTransforms(MeshComponent->BakedAnimationPoses, MeshComponent, PDI, FColor(0, 128, 192, 255), FColor(0, 128, 192, 255));
 	}
@@ -1568,7 +1568,7 @@ void FAnimationViewportClient::DrawMeshBonesBakedAnimation(UDebugSkelMeshCompone
 void FAnimationViewportClient::DrawMeshBones(UDebugSkelMeshComponent* MeshComponent, FPrimitiveDrawInterface* PDI) const
 {
 	if (!MeshComponent ||
-		!MeshComponent->SkeletalMesh ||
+		!MeshComponent->GetSkeletalMesh() ||
 		MeshComponent->GetNumDrawTransform() == 0 ||
 		MeshComponent->SkeletonDrawMode == ESkeletonDrawMode::Hidden)
 	{
@@ -1762,7 +1762,7 @@ void FAnimationViewportClient::DrawMeshSubsetBones(const UDebugSkelMeshComponent
 {
 	// this BonesOfInterest has to be in MeshComponent base, not Skeleton 
 	if (!MeshComponent ||
-		!MeshComponent->SkeletalMesh ||
+		!MeshComponent->GetSkeletalMesh() ||
 		BonesOfInterest.IsEmpty() ||
 		MeshComponent->SkeletonDrawMode == ESkeletonDrawMode::Hidden)
 	{
@@ -1837,7 +1837,7 @@ void FAnimationViewportClient::DrawMeshSubsetBones(const UDebugSkelMeshComponent
 
 void FAnimationViewportClient::DrawAttributes(UDebugSkelMeshComponent* MeshComponent, FPrimitiveDrawInterface* PDI) const
 {
-	if (MeshComponent && MeshComponent->SkeletalMesh)
+	if (MeshComponent && MeshComponent->GetSkeletalMesh())
 	{
 		const UE::Anim::FMeshAttributeContainer& Attributes = MeshComponent->GetCustomAttributes();
 
@@ -1868,7 +1868,7 @@ void FAnimationViewportClient::DrawAttributes(UDebugSkelMeshComponent* MeshCompo
 
 void FAnimationViewportClient::DrawSockets(const UDebugSkelMeshComponent* InPreviewMeshComponent, TArray<USkeletalMeshSocket*>& InSockets, FSelectedSocketInfo InSelectedSocket, FPrimitiveDrawInterface* PDI, bool bUseSkeletonSocketColor)
 {
-	if (InPreviewMeshComponent && InPreviewMeshComponent->SkeletalMesh)
+	if (InPreviewMeshComponent && InPreviewMeshComponent->GetSkeletalMesh())
 	{
 		ELocalAxesMode::Type LocalAxesMode = (ELocalAxesMode::Type)GetDefault<UPersonaOptions>()->DefaultLocalAxesSelection;
 
@@ -1961,7 +1961,7 @@ void FAnimationViewportClient::UpdateCameraSetup()
 {
 	static FRotator CustomOrbitRotation(-33.75, -135, 0);
 	UDebugSkelMeshComponent* PreviewMeshComponent = GetAnimPreviewScene()->GetPreviewMeshComponent();
-	if (PreviewMeshComponent != nullptr && PreviewMeshComponent->SkeletalMesh)
+	if (PreviewMeshComponent != nullptr && PreviewMeshComponent->GetSkeletalMesh())
 	{
 		FSphere BoundSphere = GetCameraTarget();
 		FVector CustomOrbitZoom(0, BoundSphere.W / (75.0f * (float)PI / 360.0f), 0);
@@ -2030,7 +2030,7 @@ FBox FAnimationViewportClient::ComputeBoundingBoxForSelectedEditorSection() cons
 		return FBox(ForceInitToZero);
 	}
 
-	USkeletalMesh* const SkeletalMesh = PreviewMeshComponent->SkeletalMesh;
+	USkeletalMesh* const SkeletalMesh = PreviewMeshComponent->GetSkeletalMesh();
 	FSkeletalMeshObject* const MeshObject = PreviewMeshComponent->MeshObject;
 	if ( !SkeletalMesh || !MeshObject )
 	{
@@ -2083,7 +2083,7 @@ void FAnimationViewportClient::FocusViewportOnPreviewMesh(bool bUseCustomCamera)
 
 	if (UDebugSkelMeshComponent* const PreviewMeshComponent = GetAnimPreviewScene()->GetPreviewMeshComponent())
 	{
-		if (USkeletalMesh* const SkelMesh = PreviewMeshComponent->SkeletalMesh)
+		if (USkeletalMesh* const SkelMesh = PreviewMeshComponent->GetSkeletalMesh())
 		{
 			if (bUseCustomCamera && SkelMesh->GetHasCustomDefaultEditorCamera())
 			{
@@ -2118,7 +2118,7 @@ void FAnimationViewportClient::FocusViewportOnPreviewMesh(bool bUseCustomCamera)
 
 float FAnimationViewportClient::GetFloorOffset() const
 {
-	USkeletalMesh* Mesh = GetPreviewScene()->GetPreviewMeshComponent()->SkeletalMesh;
+	USkeletalMesh* Mesh = GetPreviewScene()->GetPreviewMeshComponent()->GetSkeletalMesh();
 	if ( Mesh )
 	{
 		return Mesh->GetFloorOffset();
@@ -2129,7 +2129,7 @@ float FAnimationViewportClient::GetFloorOffset() const
 
 void FAnimationViewportClient::SetFloorOffset( float NewValue )
 {
-	USkeletalMesh* Mesh = GetPreviewScene()->GetPreviewMeshComponent()->SkeletalMesh;
+	USkeletalMesh* Mesh = GetPreviewScene()->GetPreviewMeshComponent()->GetSkeletalMesh();
 
 	if ( Mesh )
 	{
diff --git a/Engine/Source/Editor/Persona/Private/EditModes/SkeletonSelectionEditMode.cpp b/Engine/Source/Editor/Persona/Private/EditModes/SkeletonSelectionEditMode.cpp
--- a/Engine/Source/Editor/Persona/Private/EditModes/SkeletonSelectionEditMode.cpp
+++ b/Engine/Source/Editor/Persona/Private/EditModes/SkeletonSelectionEditMode.cpp
@@ -38,14 +38,14 @@ bool FSkeletonSelectionEditMode::GetCameraTarget(FSphere& OutTarget) const
 
 	UDebugSkelMeshComponent* PreviewMeshComponent = GetAnimPreviewScene().GetPreviewMeshComponent();
 
-	if(PreviewMeshComponent && PreviewMeshComponent->SkeletalMesh)
+	if(PreviewMeshComponent && PreviewMeshComponent->GetSkeletalMesh())
 	{
 		if (GetAnimPreviewScene().GetSelectedBoneIndex() != INDEX_NONE)
 		{
 			const int32 FocusBoneIndex = GetAnimPreviewScene().GetSelectedBoneIndex();
-			if (FocusBoneIndex != INDEX_NONE && PreviewMeshComponent->SkeletalMesh->GetRefSkeleton().IsValidIndex(FocusBoneIndex))
+			if (FocusBoneIndex != INDEX_NONE && PreviewMeshComponent->GetSkeletalMesh()->GetRefSkeleton().IsValidIndex(FocusBoneIndex))
 			{
-				const FName BoneName = PreviewMeshComponent->SkeletalMesh->GetRefSkeleton().GetBoneName(FocusBoneIndex);
+				const FName BoneName = PreviewMeshComponent->GetSkeletalMesh()->GetRefSkeleton().GetBoneName(FocusBoneIndex);
 				OutTarget.Center = PreviewMeshComponent->GetBoneLocation(BoneName);
 				OutTarget.W = 30.0f;
 				bHandled = true;
@@ -79,7 +79,7 @@ void FSkeletonSelectionEditMode::GetOnScreenDebugInfo(TArray<FText>& OutDebugInf
 
 FSelectedSocketInfo FSkeletonSelectionEditMode::DuplicateAndSelectSocket(const FSelectedSocketInfo& SocketInfoToDuplicate)
 {
-	USkeletalMesh* SkeletalMesh = GetAnimPreviewScene().GetPreviewMeshComponent()->SkeletalMesh;
+	USkeletalMesh* SkeletalMesh = GetAnimPreviewScene().GetPreviewMeshComponent()->GetSkeletalMesh();
 	USkeletalMeshSocket* NewSocket = GetAnimPreviewScene().GetPersonaToolkit()->GetEditableSkeleton()->DuplicateSocket(SocketInfoToDuplicate, SocketInfoToDuplicate.Socket->BoneName, SkeletalMesh);
 
 	FSelectedSocketInfo NewSocketInfo(NewSocket, SocketInfoToDuplicate.bSocketIsOnSkeleton);
@@ -95,13 +95,13 @@ bool FSkeletonSelectionEditMode::StartTracking(FEditorViewportClient* InViewport
 	UE::Widget::EWidgetMode WidgetMode = InViewportClient->GetWidgetMode();
 
 	UDebugSkelMeshComponent* PreviewMeshComponent = GetAnimPreviewScene().GetPreviewMeshComponent();
-	if(PreviewMeshComponent != nullptr && PreviewMeshComponent->SkeletalMesh != nullptr)
+	if(PreviewMeshComponent != nullptr && PreviewMeshComponent->GetSkeletalMesh() != nullptr)
 	{
 		int32 BoneIndex = GetAnimPreviewScene().GetSelectedBoneIndex();
 		USkeletalMeshSocket* SelectedSocket = GetAnimPreviewScene().GetSelectedSocket().Socket;
 		AActor* SelectedActor = GetAnimPreviewScene().GetSelectedActor();
 
-		if ((BoneIndex >= 0 && PreviewMeshComponent->SkeletalMesh->GetRefSkeleton().IsValidIndex(BoneIndex)) || SelectedSocket != nullptr || SelectedActor != nullptr)
+		if ((BoneIndex >= 0 && PreviewMeshComponent->GetSkeletalMesh()->GetRefSkeleton().IsValidIndex(BoneIndex)) || SelectedSocket != nullptr || SelectedActor != nullptr)
 		{
 			bool bValidAxis = false;
 			FVector WorldAxisDir;
@@ -157,7 +157,7 @@ bool FSkeletonSelectionEditMode::StartTracking(FEditorViewportClient* InViewport
 						bInTransaction = true;
 
 						// now modify the bone array
-						const FName BoneName = PreviewMeshComponent->SkeletalMesh->GetRefSkeleton().GetBoneName(BoneIndex);
+						const FName BoneName = PreviewMeshComponent->GetSkeletalMesh()->GetRefSkeleton().GetBoneName(BoneIndex);
 						PreviewMeshComponent->PreviewInstance->ModifyBone(BoneName);
 					}
 				}
@@ -216,7 +216,7 @@ bool FSkeletonSelectionEditMode::InputDelta(FEditorViewportClient* InViewportCli
 		if ( BoneIndex >= 0 )
 		{
 			//Get the skeleton control manipulating this bone
-			const FName BoneName = PreviewMeshComponent->SkeletalMesh->GetRefSkeleton().GetBoneName(BoneIndex);
+			const FName BoneName = PreviewMeshComponent->GetSkeletalMesh()->GetRefSkeleton().GetBoneName(BoneIndex);
 			SkelControl = &(PreviewMeshComponent->PreviewInstance->ModifyBone(BoneName));
 		}
 
@@ -339,12 +339,12 @@ FIntPoint FSkeletonSelectionEditMode::GetDPIUnscaledSize(FViewport* Viewport, FV
 void FSkeletonSelectionEditMode::DrawHUD(FEditorViewportClient* ViewportClient, FViewport* Viewport, const FSceneView* View, FCanvas* Canvas)
 {
 	UDebugSkelMeshComponent* PreviewMeshComponent = GetAnimPreviewScene().GetPreviewMeshComponent();
-	if (PreviewMeshComponent == nullptr || PreviewMeshComponent->SkeletalMesh == nullptr)
+	if (PreviewMeshComponent == nullptr || PreviewMeshComponent->GetSkeletalMesh() == nullptr)
 	{
 		return;
 	}
 
-	FReferenceSkeleton& RefSkeleton = PreviewMeshComponent->SkeletalMesh->GetRefSkeleton();
+	FReferenceSkeleton& RefSkeleton = PreviewMeshComponent->GetSkeletalMesh()->GetRefSkeleton();
 	int32 BoneIndex = GetAnimPreviewScene().GetSelectedBoneIndex();
 
 	// Draw name of selected bone
@@ -354,7 +354,7 @@ void FSkeletonSelectionEditMode::DrawHUD(FEditorViewportClient* ViewportClient,
 		const int32 HalfX = ViewPortSize.X / 2;
 		const int32 HalfY = ViewPortSize.Y / 2;
 
-		const FName BoneName = PreviewMeshComponent->SkeletalMesh->GetRefSkeleton().GetBoneName(BoneIndex);
+		const FName BoneName = PreviewMeshComponent->GetSkeletalMesh()->GetRefSkeleton().GetBoneName(BoneIndex);
 
 		FMatrix BoneMatrix = PreviewMeshComponent->GetBoneMatrix(BoneIndex);
 		const FPlane Proj = View->Project(BoneMatrix.GetOrigin());
@@ -446,10 +446,10 @@ bool FSkeletonSelectionEditMode::UsesTransformWidget(UE::Widget::EWidgetMode Che
 bool FSkeletonSelectionEditMode::GetCustomDrawingCoordinateSystem(FMatrix& InMatrix, void* InData)
 {
 	UDebugSkelMeshComponent* PreviewMeshComponent = GetAnimPreviewScene().GetPreviewMeshComponent();
-	if(PreviewMeshComponent && PreviewMeshComponent->SkeletalMesh)
+	if(PreviewMeshComponent && PreviewMeshComponent->GetSkeletalMesh())
 	{
 		int32 BoneIndex = GetAnimPreviewScene().GetSelectedBoneIndex();
-		if (BoneIndex != INDEX_NONE && PreviewMeshComponent->SkeletalMesh->GetRefSkeleton().IsValidIndex(BoneIndex))
+		if (BoneIndex != INDEX_NONE && PreviewMeshComponent->GetSkeletalMesh()->GetRefSkeleton().IsValidIndex(BoneIndex))
 		{
 			FTransform BoneMatrix = PreviewMeshComponent->GetBoneTransform(BoneIndex);
 
@@ -487,9 +487,9 @@ FVector FSkeletonSelectionEditMode::GetWidgetLocation() const
 	int32 BoneIndex = GetAnimPreviewScene().GetSelectedBoneIndex();
 	if (BoneIndex != INDEX_NONE)
 	{
-		if (PreviewMeshComponent && PreviewMeshComponent->SkeletalMesh && PreviewMeshComponent->SkeletalMesh->GetRefSkeleton().IsValidIndex(BoneIndex))
+		if (PreviewMeshComponent && PreviewMeshComponent->GetSkeletalMesh() && PreviewMeshComponent->GetSkeletalMesh()->GetRefSkeleton().IsValidIndex(BoneIndex))
 		{
-			const FName BoneName = PreviewMeshComponent->SkeletalMesh->GetRefSkeleton().GetBoneName(BoneIndex);
+			const FName BoneName = PreviewMeshComponent->GetSkeletalMesh()->GetRefSkeleton().GetBoneName(BoneIndex);
 
 			const FMatrix BoneMatrix = PreviewMeshComponent->GetBoneMatrix(BoneIndex);
 
diff --git a/Engine/Source/Editor/Persona/Private/PersonaModule.cpp b/Engine/Source/Editor/Persona/Private/PersonaModule.cpp
--- a/Engine/Source/Editor/Persona/Private/PersonaModule.cpp
+++ b/Engine/Source/Editor/Persona/Private/PersonaModule.cpp
@@ -1394,7 +1394,7 @@ bool FPersonaModule::CreateAnimation(const TArray<UObject*> NewAssets, const EPo
 				switch (Option)
 				{
 				case EPoseSourceOption::ReferencePose:
-					bResult &= NewAnimSequence->CreateAnimation(MeshComponent->SkeletalMesh);
+					bResult &= NewAnimSequence->CreateAnimation(MeshComponent->GetSkeletalMesh());
 					break;
 				case EPoseSourceOption::CurrentPose:
 					bResult &= NewAnimSequence->CreateAnimation(MeshComponent);
diff --git a/Engine/Source/Editor/Persona/Private/PersonaPreviewSceneAnimationController.cpp b/Engine/Source/Editor/Persona/Private/PersonaPreviewSceneAnimationController.cpp
--- a/Engine/Source/Editor/Persona/Private/PersonaPreviewSceneAnimationController.cpp
+++ b/Engine/Source/Editor/Persona/Private/PersonaPreviewSceneAnimationController.cpp
@@ -23,7 +23,7 @@ IDetailPropertyRow* UPersonaPreviewSceneAnimationController::AddPreviewControlle
 {
 	TArray<UObject*> ListOfPreviewController{ this };
 
-	const USkeleton* Skeleton = PersonaToolkit->GetPreviewMeshComponent()->SkeletalMesh ? PersonaToolkit->GetPreviewMeshComponent()->SkeletalMesh->GetSkeleton() : nullptr;
+	const USkeleton* Skeleton = PersonaToolkit->GetPreviewMeshComponent()->GetSkeletalMesh() ? PersonaToolkit->GetPreviewMeshComponent()->GetSkeletalMesh()->GetSkeleton() : nullptr;
 	if (Skeleton)
 	{
 		IDetailPropertyRow* NewRow = Category.AddExternalObjectProperty(ListOfPreviewController, Property->GetFName(), PropertyLocation);
diff --git a/Engine/Source/Editor/Persona/Private/SAnimAttributeView.cpp b/Engine/Source/Editor/Persona/Private/SAnimAttributeView.cpp
--- a/Engine/Source/Editor/Persona/Private/SAnimAttributeView.cpp
+++ b/Engine/Source/Editor/Persona/Private/SAnimAttributeView.cpp
@@ -270,7 +270,7 @@ void SAnimAttributeView::DisplayNewAttributeContainerSnapshots(
 	}
 	
 	// we need the SKM to look up bone names
-	if (!InOwningComponent || !InOwningComponent->SkeletalMesh)
+	if (!InOwningComponent || !InOwningComponent->GetSkeletalMesh())
 	{
 		ClearListView();
 		return;
@@ -309,7 +309,7 @@ void SAnimAttributeView::DisplayNewAttributeContainerSnapshots(
 				for (int32 IdIndex = 0; IdIndex < Ids.Num(); IdIndex++)
 				{
 					const UE::Anim::FAttributeId& Id = Ids[IdIndex];
-					const FName BoneName = InOwningComponent->SkeletalMesh->GetRefSkeleton().GetBoneName(Id.GetIndex());
+					const FName BoneName = InOwningComponent->GetSkeletalMesh()->GetRefSkeleton().GetBoneName(Id.GetIndex());
 				
 					const FAnimationAttributeIdentifier Identifier(
 						Id.GetName(), Id.GetIndex() ,BoneName, Types[TypeIndex].Get());
@@ -397,7 +397,7 @@ bool SAnimAttributeView::ShouldInvalidateListViewCache(
 			for (int32 IdIndex = 0; IdIndex < Ids.Num(); IdIndex++)
 			{
 				const UE::Anim::FAttributeId& Id = Ids[IdIndex];
-				const FName BoneName = InOwningComponent->SkeletalMesh->GetRefSkeleton().GetBoneName(Id.GetIndex());
+				const FName BoneName = InOwningComponent->GetSkeletalMesh()->GetRefSkeleton().GetBoneName(Id.GetIndex());
 				
 				FAnimationAttributeIdentifier Identifier(
 					Id.GetName(), Id.GetIndex() ,BoneName, Types[TypeIndex].Get());
diff --git a/Engine/Source/Editor/Persona/Private/SAnimViewportToolBar.cpp b/Engine/Source/Editor/Persona/Private/SAnimViewportToolBar.cpp
--- a/Engine/Source/Editor/Persona/Private/SAnimViewportToolBar.cpp
+++ b/Engine/Source/Editor/Persona/Private/SAnimViewportToolBar.cpp
@@ -870,7 +870,7 @@ void SAnimViewportToolBar::FillCharacterMirrorMenu(FMenuBuilder& MenuBuilder) co
 {
 	FAssetPickerConfig AssetPickerConfig;
 	UDebugSkelMeshComponent* PreviewComp = Viewport.Pin()->GetPreviewScene()->GetPreviewMeshComponent();
-	USkeletalMesh* Mesh = PreviewComp->SkeletalMesh;
+	USkeletalMesh* Mesh = PreviewComp->GetSkeletalMesh();
 	UAnimPreviewInstance* PreviewInstance = PreviewComp->PreviewInstance; 
 	if (Mesh && PreviewInstance)
 	{
@@ -901,7 +901,7 @@ void SAnimViewportToolBar::OnMirrorDataTableSelected(const FAssetData& SelectedM
 	if (Viewport.Pin().IsValid())
 	{
 		UDebugSkelMeshComponent* PreviewComp = Viewport.Pin()->GetPreviewScene()->GetPreviewMeshComponent();
-		USkeletalMesh* Mesh = PreviewComp->SkeletalMesh;
+		USkeletalMesh* Mesh = PreviewComp->GetSkeletalMesh();
 		UAnimPreviewInstance* PreviewInstance = PreviewComp->PreviewInstance; 
 		if (Mesh && PreviewInstance)
 		{
diff --git a/Engine/Source/Editor/Persona/Private/SAnimationEditorViewport.cpp b/Engine/Source/Editor/Persona/Private/SAnimationEditorViewport.cpp
--- a/Engine/Source/Editor/Persona/Private/SAnimationEditorViewport.cpp
+++ b/Engine/Source/Editor/Persona/Private/SAnimationEditorViewport.cpp
@@ -273,7 +273,7 @@ FText SAnimationEditorViewportTabBody::GetDisplayString() const
 				DefaultText = FText::Format(LOCTEXT("PreviewingAnimBP", "Previewing {0}"), FText::FromString(Component->AnimClass->GetName()));
 			}
 		}
-		else if (Component->SkeletalMesh == NULL && TargetSkeletonName != NAME_None)
+		else if (Component->GetSkeletalMesh() == NULL && TargetSkeletonName != NAME_None)
 		{
 			DefaultText = FText::Format(LOCTEXT("NoMeshFound", "No skeletal mesh found for skeleton '{0}'"), FText::FromName(TargetSkeletonName));
 		}
@@ -1159,9 +1159,9 @@ bool SAnimationEditorViewportTabBody::IsTurnTableModeSelected(int32 ModeIndex) c
 int32 SAnimationEditorViewportTabBody::GetLODModelCount() const
 {
 	UDebugSkelMeshComponent* PreviewComponent = GetPreviewScene()->GetPreviewMeshComponent();
-	if( PreviewComponent && PreviewComponent->SkeletalMesh )
+	if( PreviewComponent && PreviewComponent->GetSkeletalMesh())
 	{
-		return PreviewComponent->SkeletalMesh->GetResourceForRendering()->LODRenderData.Num();
+		return PreviewComponent->GetSkeletalMesh()->GetResourceForRendering()->LODRenderData.Num();
 	}
 	return 0;
 }
@@ -1747,7 +1747,7 @@ void SAnimationEditorViewportTabBody::PopulateSkinWeightProfileNames()
 	// Retrieve all possible skin weight profiles from the component
 	if (UDebugSkelMeshComponent* PreviewComponent = GetPreviewScene()->GetPreviewMeshComponent())
 	{
-		if (USkeletalMesh* Mesh = PreviewComponent->SkeletalMesh)
+		if (USkeletalMesh* Mesh = PreviewComponent->GetSkeletalMesh())
 		{
 			for (const FSkinWeightProfileInfo& Profile : Mesh->GetSkinWeightProfiles())
 			{
@@ -2380,9 +2380,9 @@ void SAnimationEditorViewportTabBody::AddPostProcessNotification()
 	{
 		if (UDebugSkelMeshComponent* PreviewComponent = GetPreviewScene()->GetPreviewMeshComponent())
 		{
-			if(PreviewComponent->SkeletalMesh && PreviewComponent->SkeletalMesh->GetPostProcessAnimBlueprint() && PreviewComponent->SkeletalMesh->GetPostProcessAnimBlueprint()->ClassGeneratedBy)
+			if(PreviewComponent->GetSkeletalMesh() && PreviewComponent->GetSkeletalMesh()->GetPostProcessAnimBlueprint() && PreviewComponent->GetSkeletalMesh()->GetPostProcessAnimBlueprint()->ClassGeneratedBy)
 			{
-				return FText::FromString(PreviewComponent->SkeletalMesh->GetPostProcessAnimBlueprint()->ClassGeneratedBy->GetName());
+				return FText::FromString(PreviewComponent->GetSkeletalMesh()->GetPostProcessAnimBlueprint()->ClassGeneratedBy->GetName());
 			}
 		}
 
@@ -2431,9 +2431,9 @@ void SAnimationEditorViewportTabBody::AddPostProcessNotification()
 	{
 		if (UDebugSkelMeshComponent* PreviewComponent = GetPreviewScene()->GetPreviewMeshComponent())
 		{
-			if(PreviewComponent->SkeletalMesh && PreviewComponent->SkeletalMesh->GetPostProcessAnimBlueprint())
+			if(PreviewComponent->GetSkeletalMesh() && PreviewComponent->GetSkeletalMesh()->GetPostProcessAnimBlueprint())
 			{
-				GEditor->GetEditorSubsystem<UAssetEditorSubsystem>()->OpenEditorForAssets(TArray<UObject*>({ PreviewComponent->SkeletalMesh->GetPostProcessAnimBlueprint()->ClassGeneratedBy }));
+				GEditor->GetEditorSubsystem<UAssetEditorSubsystem>()->OpenEditorForAssets(TArray<UObject*>({ PreviewComponent->GetSkeletalMesh()->GetPostProcessAnimBlueprint()->ClassGeneratedBy }));
 			}
 		}
 
diff --git a/Engine/Source/Editor/Persona/Private/SAnimationSequenceBrowser.cpp b/Engine/Source/Editor/Persona/Private/SAnimationSequenceBrowser.cpp
--- a/Engine/Source/Editor/Persona/Private/SAnimationSequenceBrowser.cpp
+++ b/Engine/Source/Editor/Persona/Private/SAnimationSequenceBrowser.cpp
@@ -760,7 +760,7 @@ void SAnimationSequenceBrowser::OnExportToFBX(TArray<FAssetData> SelectedAssets)
 		}
 
 		FPersonaModule& PersonaModule = FModuleManager::GetModuleChecked<FPersonaModule>("Persona");
-		PersonaModule.ExportToFBX(AnimSequences, PersonaToolkitPtr.Pin()->GetPreviewScene()->GetPreviewMeshComponent()->SkeletalMesh);
+		PersonaModule.ExportToFBX(AnimSequences, PersonaToolkitPtr.Pin()->GetPreviewScene()->GetPreviewMeshComponent()->GetSkeletalMesh());
 	}
 }
 
@@ -768,7 +768,7 @@ void SAnimationSequenceBrowser::OnSetCurrentPreviewMesh(TArray<FAssetData> Selec
 {
 	if(SelectedAssets.Num() > 0)
 	{
-		USkeletalMesh* PreviewMesh = PersonaToolkitPtr.Pin()->GetPreviewScene()->GetPreviewMeshComponent()->SkeletalMesh;
+		USkeletalMesh* PreviewMesh = PersonaToolkitPtr.Pin()->GetPreviewScene()->GetPreviewMeshComponent()->GetSkeletalMesh();
 		if (PreviewMesh)
 		{
 			TArray<TWeakObjectPtr<UAnimSequence>> AnimSequences;
@@ -1452,7 +1452,7 @@ bool SAnimationSequenceBrowser::OnVisualizeAssetToolTip(const TSharedPtr<SWidget
 		
 		if(MeshToUse)
 		{
-			if(PreviewComponent->SkeletalMesh != MeshToUse)
+			if(PreviewComponent->GetSkeletalMesh() != MeshToUse)
 			{
 				PreviewComponent->SetSkeletalMesh(MeshToUse);
 			}
diff --git a/Engine/Source/Editor/Persona/Private/SMorphTargetViewer.cpp b/Engine/Source/Editor/Persona/Private/SMorphTargetViewer.cpp
--- a/Engine/Source/Editor/Persona/Private/SMorphTargetViewer.cpp
+++ b/Engine/Source/Editor/Persona/Private/SMorphTargetViewer.cpp
@@ -328,7 +328,7 @@ void SMorphTargetViewer::Construct(const FArguments& InArgs, const TSharedRef<IP
 {
 	PreviewScenePtr = InPreviewScene;
 
-	SkeletalMesh = InPreviewScene->GetPreviewMeshComponent()->SkeletalMesh;
+	SkeletalMesh = InPreviewScene->GetPreviewMeshComponent()->GetSkeletalMesh();
 	InPreviewScene->RegisterOnPreviewMeshChanged( FOnPreviewMeshChanged::CreateSP( this, &SMorphTargetViewer::OnPreviewMeshChanged ) );
 	InPreviewScene->RegisterOnMorphTargetsChanged(FSimpleDelegate::CreateSP(this, &SMorphTargetViewer::OnMorphTargetsChanged));
 	OnPostUndo.Add(FSimpleDelegate::CreateSP(this, &SMorphTargetViewer::OnPostUndo));
diff --git a/Engine/Source/Editor/Persona/Private/SPoseEditor.cpp b/Engine/Source/Editor/Persona/Private/SPoseEditor.cpp
--- a/Engine/Source/Editor/Persona/Private/SPoseEditor.cpp
+++ b/Engine/Source/Editor/Persona/Private/SPoseEditor.cpp
@@ -463,7 +463,7 @@ void RestartAnimations(const USkeleton* CurrentSkeleton)
 	for (FThreadSafeObjectIterator Iter(USkeletalMeshComponent::StaticClass()); Iter; ++Iter)
 	{
 		USkeletalMeshComponent* SkeletalMeshComponent = Cast<USkeletalMeshComponent>(*Iter);
-		if (SkeletalMeshComponent->SkeletalMesh && SkeletalMeshComponent->SkeletalMesh->GetSkeleton() == CurrentSkeleton)
+		if (SkeletalMeshComponent->GetSkeletalMesh() && SkeletalMeshComponent->GetSkeletalMesh()->GetSkeleton() == CurrentSkeleton)
 		{
 			SkeletalMeshComponent->InitAnim(true);
 		}
diff --git a/Engine/Source/Editor/Persona/Private/SRetargetSources.cpp b/Engine/Source/Editor/Persona/Private/SRetargetSources.cpp
--- a/Engine/Source/Editor/Persona/Private/SRetargetSources.cpp
+++ b/Engine/Source/Editor/Persona/Private/SRetargetSources.cpp
@@ -284,9 +284,9 @@ bool SRetargetSources::ShouldFilterAsset(const FAssetData& InAssetData)
 void SRetargetSources::ResetRetargetBasePose()
 {
 	UDebugSkelMeshComponent * PreviewMeshComp = PreviewScenePtr.Pin()->GetPreviewMeshComponent();
-	if(PreviewMeshComp && PreviewMeshComp->SkeletalMesh)
+	if(PreviewMeshComp && PreviewMeshComp->GetSkeletalMesh())
 	{
-		USkeletalMesh * PreviewMesh = PreviewMeshComp->SkeletalMesh;
+		USkeletalMesh * PreviewMesh = PreviewMeshComp->GetSkeletalMesh();
 
 		check(PreviewMesh && &EditableSkeletonPtr.Pin()->GetSkeleton() == PreviewMesh->GetSkeleton());
 
@@ -306,9 +306,9 @@ void SRetargetSources::ResetRetargetBasePose()
 void SRetargetSources::UseCurrentPose()
 {
 	UDebugSkelMeshComponent * PreviewMeshComp = PreviewScenePtr.Pin()->GetPreviewMeshComponent();
-	if (PreviewMeshComp && PreviewMeshComp->SkeletalMesh)
+	if (PreviewMeshComp && PreviewMeshComp->GetSkeletalMesh())
 	{
-		USkeletalMesh * PreviewMesh = PreviewMeshComp->SkeletalMesh;
+		USkeletalMesh * PreviewMesh = PreviewMeshComp->GetSkeletalMesh();
 
 		check(PreviewMesh && &EditableSkeletonPtr.Pin()->GetSkeleton() == PreviewMesh->GetSkeleton());
 
@@ -386,9 +386,9 @@ void SRetargetSources::ImportPose(const UPoseAsset* PoseAsset, const FName& Pose
 
 			// now I have pose, I have to copy to the retarget base pose
 			UDebugSkelMeshComponent * PreviewMeshComp = PreviewScenePtr.Pin()->GetPreviewMeshComponent();
-			if (PreviewMeshComp && PreviewMeshComp->SkeletalMesh)
+			if (PreviewMeshComp && PreviewMeshComp->GetSkeletalMesh())
 			{
-				USkeletalMesh * PreviewMesh = PreviewMeshComp->SkeletalMesh;
+				USkeletalMesh * PreviewMesh = PreviewMeshComp->GetSkeletalMesh();
 
 				check(PreviewMesh && &EditableSkeletonPtr.Pin()->GetSkeleton() == PreviewMesh->GetSkeleton());
 
diff --git a/Engine/Source/Editor/Persona/Private/SkinWeightProfileHelpers.cpp b/Engine/Source/Editor/Persona/Private/SkinWeightProfileHelpers.cpp
--- a/Engine/Source/Editor/Persona/Private/SkinWeightProfileHelpers.cpp
+++ b/Engine/Source/Editor/Persona/Private/SkinWeightProfileHelpers.cpp
@@ -280,7 +280,7 @@ void FSkinWeightProfileHelpers::ClearSkinWeightProfileInstanceOverrides(USkeleta
 	FScopedSkeletalMeshPostEditChange ScopedPostEditChange(InSkeletalMesh, false, true);
 	for (TObjectIterator<USkinnedMeshComponent> It; It; ++It)
 	{
-		if (It->SkeletalMesh == InSkeletalMesh)
+		if (It->GetSkeletalMesh() == InSkeletalMesh)
 		{
 			checkf(!It->IsUnreachable(), TEXT("%s"), *It->GetFullName());
 
diff --git a/Engine/Source/Editor/PhysicsAssetEditor/Private/PhysicsAssetEditor.cpp b/Engine/Source/Editor/PhysicsAssetEditor/Private/PhysicsAssetEditor.cpp
--- a/Engine/Source/Editor/PhysicsAssetEditor/Private/PhysicsAssetEditor.cpp
+++ b/Engine/Source/Editor/PhysicsAssetEditor/Private/PhysicsAssetEditor.cpp
@@ -466,12 +466,12 @@ void FPhysicsAssetEditor::OnAssetReimport(UObject* Object)
 	RefreshHierachyTree();
 	RefreshPreviewViewport();
 
-	if (SharedData->EditorSkelComp && SharedData->EditorSkelComp->SkeletalMesh)
+	if (SharedData->EditorSkelComp && SharedData->EditorSkelComp->GetSkeletalMesh())
 	{
 		IMeshUtilities& MeshUtilities = FModuleManager::Get().LoadModuleChecked<IMeshUtilities>("MeshUtilities");
 		// Update various infos based on the mesh
-		MeshUtilities.CalcBoneVertInfos(SharedData->EditorSkelComp->SkeletalMesh, SharedData->DominantWeightBoneInfos, true);
-		MeshUtilities.CalcBoneVertInfos(SharedData->EditorSkelComp->SkeletalMesh, SharedData->AnyWeightBoneInfos, false);
+		MeshUtilities.CalcBoneVertInfos(SharedData->EditorSkelComp->GetSkeletalMesh(), SharedData->DominantWeightBoneInfos, true);
+		MeshUtilities.CalcBoneVertInfos(SharedData->EditorSkelComp->GetSkeletalMesh(), SharedData->AnyWeightBoneInfos, false);
 	}
 }
 
diff --git a/Engine/Source/Editor/PhysicsAssetEditor/Private/PhysicsAssetEditorSharedData.cpp b/Engine/Source/Editor/PhysicsAssetEditor/Private/PhysicsAssetEditorSharedData.cpp
--- a/Engine/Source/Editor/PhysicsAssetEditor/Private/PhysicsAssetEditorSharedData.cpp
+++ b/Engine/Source/Editor/PhysicsAssetEditor/Private/PhysicsAssetEditorSharedData.cpp
@@ -2615,7 +2615,7 @@ void FPhysicsAssetEditorSharedData::EnableSimulation(bool bEnableSimulation)
 			EditorSkelComp->SetAllBodiesSimulatePhysics(false);
 
 			// make sure we enable the preview animation is any compatible with the skeleton
-			if (PreviewAnimationAsset && EditorSkelComp->SkeletalMesh && PreviewAnimationAsset->GetSkeleton() == EditorSkelComp->SkeletalMesh->GetSkeleton())
+			if (PreviewAnimationAsset && EditorSkelComp->GetSkeletalMesh() && PreviewAnimationAsset->GetSkeleton() == EditorSkelComp->GetSkeletalMesh()->GetSkeleton())
 			{
 				EditorSkelComp->EnablePreview(true, PreviewAnimationAsset);
 				EditorSkelComp->Play(true);
