Fix for assertion hit when orbiting a bone that no longer exists

diff --git a/Engine/Source/Editor/Persona/Private/AnimationEditorViewportClient.cpp b/Engine/Source/Editor/Persona/Private/AnimationEditorViewportClient.cpp
--- a/Engine/Source/Editor/Persona/Private/AnimationEditorViewportClient.cpp
+++ b/Engine/Source/Editor/Persona/Private/AnimationEditorViewportClient.cpp
@@ -687,14 +687,19 @@ void FAnimationViewportClient::HandlePreviewScenePostTick()
 			case EAnimationViewportCameraFollowMode::Bone:
 				{
 					const int32 BoneIndex = PreviewMeshComponent->GetBoneIndex(CameraFollowBoneName);
-					check(BoneIndex != INDEX_NONE);
-					LookAtLocation = PreviewMeshComponent->GetBoneTransform(BoneIndex).GetLocation();
+					if (BoneIndex != INDEX_NONE)
+					{
+						LookAtLocation = PreviewMeshComponent->GetBoneTransform(BoneIndex).GetLocation();
 
-					if (GetShouldRotateCameraToFollowBone())
+						if (GetShouldRotateCameraToFollowBone())
+						{
+							OrbitRotation = PreviewMeshComponent->GetBoneQuaternion(CameraFollowBoneName) * FQuat(FVector(0.0f, 1.0f, 0.0f), PI * 0.5f);
+						}
+					}
+					else
 					{
-						const FName BoneName = PreviewMeshComponent->GetBoneName(BoneIndex);
-						check(BoneName != NAME_None);
-						OrbitRotation = PreviewMeshComponent->GetBoneQuaternion(BoneName) * FQuat(FVector(0.0f, 1.0f, 0.0f), PI * 0.5f);
+						SetCameraFollowMode(EAnimationViewportCameraFollowMode::None);
+						return;
 					}
 				}
 				break;
