Makes bone colors an editor preference.

diff --git a/Engine/Plugins/Animation/ControlRig/Source/ControlRigEditor/Private/Editor/ControlRigEditorEditMode.cpp b/Engine/Plugins/Animation/ControlRig/Source/ControlRigEditor/Private/Editor/ControlRigEditorEditMode.cpp
--- a/Engine/Plugins/Animation/ControlRig/Source/ControlRigEditor/Private/Editor/ControlRigEditorEditMode.cpp
+++ b/Engine/Plugins/Animation/ControlRig/Source/ControlRigEditor/Private/Editor/ControlRigEditorEditMode.cpp
@@ -152,6 +152,9 @@ void FControlRigEditorEditMode::Render(const FSceneView* View, FViewport* Viewpo
 		BoneRadius = AnimViewportClient->GetBoneDrawSize();
 	}
 
+	// use colors from user preferences
+	const UPersonaOptions* PersonaOptions = GetDefault<UPersonaOptions>();
+
 	// spin through all bones and render them
 	TArray<FVector> ChildPositions;
 	TArray<FLinearColor> ChildColors;
@@ -167,9 +170,9 @@ void FControlRigEditorEditMode::Render(const FSceneView* View, FViewport* Viewpo
 		// determine color of bone based on selection / affected state
 		const bool bIsSelected = SelectedBones[ElementIndex];
 		const bool bIsAffected = AffectedBones[ElementIndex];
-		FLinearColor DefaultBoneColor = SkeletalDebugRendering::DEFAULT_BONE_COLOR;
-		FLinearColor BoneColor = bIsAffected ? SkeletalDebugRendering::AFFECTED_BONE_COLOR : DefaultBoneColor;
-		BoneColor = bIsSelected ? SkeletalDebugRendering::SELECTED_BONE_COLOR : BoneColor;
+		FLinearColor DefaultBoneColor = PersonaOptions->DefaultBoneColor;
+		FLinearColor BoneColor = bIsAffected ? PersonaOptions->AffectedBoneColor : DefaultBoneColor;
+		BoneColor = bIsSelected ? PersonaOptions->SelectedBoneColor : BoneColor;
 
 		// draw the little coordinate frame inside the bone ONLY if selected or affected
 		const bool bDrawAxesInsideBone = bIsAffected|| bIsSelected;
@@ -186,7 +189,7 @@ void FControlRigEditorEditMode::Render(const FSceneView* View, FViewport* Viewpo
 			FLinearColor ChildLineColor = BoneColor;
 			if (!bIsSelected && SelectedBones[ChildIndex])
 			{
-				ChildLineColor = SkeletalDebugRendering::PARENT_OF_SELECTED_BONE_COLOR;
+				ChildLineColor = PersonaOptions->ParentOfSelectedBoneColor;
 			}
 			ChildColors.Add(ChildLineColor);
 		}
diff --git a/Engine/Plugins/Animation/IKRig/Source/IKRigEditor/Private/RetargetEditor/IKRetargetEditPoseMode.cpp b/Engine/Plugins/Animation/IKRig/Source/IKRigEditor/Private/RetargetEditor/IKRetargetEditPoseMode.cpp
--- a/Engine/Plugins/Animation/IKRig/Source/IKRigEditor/Private/RetargetEditor/IKRetargetEditPoseMode.cpp
+++ b/Engine/Plugins/Animation/IKRig/Source/IKRigEditor/Private/RetargetEditor/IKRetargetEditPoseMode.cpp
@@ -10,6 +10,7 @@
 #include "IPersonaPreviewScene.h"
 #include "Animation/DebugSkelMeshComponent.h"
 #include "Materials/MaterialInstanceDynamic.h"
+#include "Preferences/PersonaOptions.h"
 #include "RetargetEditor/IKRetargetAnimInstance.h"
 #include "RetargetEditor/IKRetargetHitProxies.h"
 #include "RetargetEditor/IKRetargetEditor.h"
@@ -88,6 +89,9 @@ void FIKRetargetEditPoseMode::RenderSkeleton(
 	const int32 RootBoneIndex = GetEditedRetargetRootBoneIndex(*Processor);
 	const float MaxDrawRadius = Controller->TargetSkelMeshComponent->Bounds.SphereRadius * 0.01f;
 
+	// use colors from user preferences
+	const UPersonaOptions* PersonaOptions = GetDefault<UPersonaOptions>();
+
 	// loop over whole skeleton and draw each bone
 	for (int32 BoneIndex = 0; BoneIndex<Skeleton.BoneNames.Num(); ++BoneIndex)
 	{
@@ -95,9 +99,9 @@ void FIKRetargetEditPoseMode::RenderSkeleton(
 		const bool bIsSelected = OutSelectedBones.Contains(BoneIndex);
 		const bool bIsAffected = OutAffectedBones.Contains(BoneIndex);
 		const bool bIsDisabled = !(Skeleton.IsBoneInAnyChain[BoneIndex] || BoneIndex==RootBoneIndex);
-		FLinearColor LineColor = bIsAffected ? SkeletalDebugRendering::AFFECTED_BONE_COLOR : SkeletalDebugRendering::DEFAULT_BONE_COLOR;
-		LineColor = bIsSelected ? SkeletalDebugRendering::SELECTED_BONE_COLOR : LineColor;
-		LineColor = bIsDisabled ? SkeletalDebugRendering::DISABLED_BONE_COLOR : LineColor;
+		FLinearColor LineColor = bIsAffected ? PersonaOptions->AffectedBoneColor : PersonaOptions->DefaultBoneColor;
+		LineColor = bIsSelected ? PersonaOptions->SelectedBoneColor : LineColor;
+		LineColor = bIsDisabled ? PersonaOptions->DisabledBoneColor : LineColor;
 		
 		const float BoneRadius = FMath::Min(1.0f, MaxDrawRadius) * BoneDrawSize;
 		
@@ -121,7 +125,7 @@ void FIKRetargetEditPoseMode::RenderSkeleton(
 			FLinearColor ChildLineColor = LineColor;
 			if (!bIsSelected && OutSelectedBones.Contains(ChildIndex))
 			{
-				ChildLineColor = SkeletalDebugRendering::PARENT_OF_SELECTED_BONE_COLOR;
+				ChildLineColor = PersonaOptions->ParentOfSelectedBoneColor;
 			}
 			ChildColors.Add(ChildLineColor);
 		}
diff --git a/Engine/Plugins/Animation/IKRig/Source/IKRigEditor/Private/RigEditor/IKRigEditMode.cpp b/Engine/Plugins/Animation/IKRig/Source/IKRigEditor/Private/RigEditor/IKRigEditMode.cpp
--- a/Engine/Plugins/Animation/IKRig/Source/IKRigEditor/Private/RigEditor/IKRigEditMode.cpp
+++ b/Engine/Plugins/Animation/IKRig/Source/IKRigEditor/Private/RigEditor/IKRigEditMode.cpp
@@ -12,6 +12,7 @@
 #include "RigEditor/IKRigHitProxies.h"
 #include "RigEditor/IKRigToolkit.h"
 #include "IKRigDebugRendering.h"
+#include "Preferences/PersonaOptions.h"
 
 #define LOCTEXT_NAMESPACE "IKRetargeterEditMode"
 
@@ -155,6 +156,9 @@ void FIKRigEditMode::RenderBones(FPrimitiveDrawInterface* PDI)
 	TSet<int32> OutSelectedBones;
 	GetAffectedBones(Controller.Get(), CurrentProcessor,OutAffectedBones,OutSelectedBones);
 
+	// use colors from user preferences
+	const UPersonaOptions* PersonaOptions = GetDefault<UPersonaOptions>();
+
 	// draw bones
 	TArray<int32> ChildrenIndices;
 	const FIKRigSkeleton& Skeleton = CurrentProcessor->GetSkeleton();
@@ -164,8 +168,8 @@ void FIKRigEditMode::RenderBones(FPrimitiveDrawInterface* PDI)
 		// selected bones are drawn with different color
 		const bool bIsSelected = OutSelectedBones.Contains(BoneIndex);
 		const bool bIsAffected = OutAffectedBones.Contains(BoneIndex);
-		FLinearColor LineColor = bIsAffected ? SkeletalDebugRendering::AFFECTED_BONE_COLOR : SkeletalDebugRendering::DEFAULT_BONE_COLOR;
-		LineColor = bIsSelected ? SkeletalDebugRendering::SELECTED_BONE_COLOR : LineColor;
+		FLinearColor LineColor = bIsAffected ? PersonaOptions->AffectedBoneColor : PersonaOptions->DefaultBoneColor;
+		LineColor = bIsSelected ? PersonaOptions->SelectedBoneColor : LineColor;
 
 		// only draw axes on affected/selected bones
 		const bool bDrawAxes = bIsSelected || bIsAffected;
@@ -186,7 +190,7 @@ void FIKRigEditMode::RenderBones(FPrimitiveDrawInterface* PDI)
 			FLinearColor ChildLineColor = LineColor;
 			if (!bIsSelected && OutSelectedBones.Contains(ChildIndex))
 			{
-				ChildLineColor = SkeletalDebugRendering::PARENT_OF_SELECTED_BONE_COLOR;
+				ChildLineColor = PersonaOptions->ParentOfSelectedBoneColor;
 			}
 			ChildColors.Add(ChildLineColor);
 		}
diff --git a/Engine/Source/Editor/Persona/Private/AnimationEditorViewportClient.cpp b/Engine/Source/Editor/Persona/Private/AnimationEditorViewportClient.cpp
--- a/Engine/Source/Editor/Persona/Private/AnimationEditorViewportClient.cpp
+++ b/Engine/Source/Editor/Persona/Private/AnimationEditorViewportClient.cpp
@@ -1419,7 +1419,7 @@ void FAnimationViewportClient::DrawBonesFromTransforms(
 
 	if (MeshComponent->SkeletonDrawMode == ESkeletonDrawMode::GreyedOut)
 	{
-		BoneColour = SkeletalDebugRendering::DISABLED_BONE_COLOR;
+		BoneColour = GetDefault<UPersonaOptions>()->DisabledBoneColor;
 	}
 	
 	TArray<FTransform> WorldTransforms;
@@ -1469,7 +1469,7 @@ void FAnimationViewportClient::DrawBonesFromCompactPose(
 	}
 
 	// optionally override draw color
-	const FLinearColor BoneColor = MeshComponent->SkeletonDrawMode == ESkeletonDrawMode::GreyedOut ? SkeletalDebugRendering::DISABLED_BONE_COLOR : DrawColor;
+	const FLinearColor BoneColor = MeshComponent->SkeletonDrawMode == ESkeletonDrawMode::GreyedOut ? GetDefault<UPersonaOptions>()->DisabledBoneColor : DrawColor;
 	
 	TArray<FTransform> WorldTransforms;
 	WorldTransforms.AddUninitialized(Pose.GetBoneContainer().GetNumBones());
@@ -1587,8 +1587,8 @@ void FAnimationViewportClient::DrawMeshBones(UDebugSkelMeshComponent* MeshCompon
 	BoneColours.AddUninitialized(MeshComponent->GetNumDrawTransform());
 
 	// factor skeleton draw mode into color selection
-	const FLinearColor BoneColor = MeshComponent->SkeletonDrawMode == ESkeletonDrawMode::GreyedOut ? SkeletalDebugRendering::DISABLED_BONE_COLOR : SkeletalDebugRendering::DEFAULT_BONE_COLOR;
-	const FLinearColor VirtualBoneColor = MeshComponent->SkeletonDrawMode == ESkeletonDrawMode::GreyedOut ? SkeletalDebugRendering::DISABLED_BONE_COLOR : SkeletalDebugRendering::VIRTUAL_BONE_COLOR;
+	const FLinearColor BoneColor = MeshComponent->SkeletonDrawMode == ESkeletonDrawMode::GreyedOut ? GetDefault<UPersonaOptions>()->DisabledBoneColor : GetDefault<UPersonaOptions>()->DefaultBoneColor;
+	const FLinearColor VirtualBoneColor = MeshComponent->SkeletonDrawMode == ESkeletonDrawMode::GreyedOut ? GetDefault<UPersonaOptions>()->DisabledBoneColor : GetDefault<UPersonaOptions>()->VirtualBoneColor;
 	
 	// we could cache parent bones as we calculate, but right now I'm not worried about perf issue of this
 	const TArray<FBoneIndexType>& DrawBoneIndices = MeshComponent->GetDrawBoneIndices();
@@ -1696,6 +1696,7 @@ void FAnimationViewportClient::DrawBones(
 	
 	// spin through all required bones and render them
 	const float BoneRadius = GetBoneDrawSize();
+	const UPersonaOptions* PersonaOptions = GetDefault<UPersonaOptions>();
 	for ( int32 Index=0; Index<RequiredBones.Num(); ++Index )
 	{
 		const int32 BoneIndex = RequiredBones[Index];
@@ -1710,9 +1711,9 @@ void FAnimationViewportClient::DrawBones(
 		// determine color of bone based on selection / affected state
 		const bool bIsSelected = InSelectedBones.Contains(BoneIndex);
 		const bool bIsAffected = AffectedBones[BoneIndex];
-		FLinearColor DefaultBoneColor = BoneColors.IsEmpty() ? SkeletalDebugRendering::DEFAULT_BONE_COLOR : BoneColors[Index];
-		FLinearColor BoneColor = bIsAffected ? SkeletalDebugRendering::AFFECTED_BONE_COLOR : DefaultBoneColor;
-		BoneColor = bIsSelected ? SkeletalDebugRendering::SELECTED_BONE_COLOR : BoneColor;
+		FLinearColor DefaultBoneColor = BoneColors.IsEmpty() ? PersonaOptions->DefaultBoneColor : BoneColors[Index];
+		FLinearColor BoneColor = bIsAffected ? PersonaOptions->AffectedBoneColor : DefaultBoneColor;
+		BoneColor = bIsSelected ? PersonaOptions->SelectedBoneColor : BoneColor;
 
 		// draw the little coordinate frame inside the bone ONLY if selected or affected
 		const bool bDrawAxesInsideBone = bIsAffected|| bIsSelected;
@@ -1730,7 +1731,7 @@ void FAnimationViewportClient::DrawBones(
 				FLinearColor ChildLineColor = BoneColor;
 				if (!bIsSelected && InSelectedBones.Contains(ChildIndex))
 				{
-					ChildLineColor = SkeletalDebugRendering::PARENT_OF_SELECTED_BONE_COLOR;
+					ChildLineColor = PersonaOptions->ParentOfSelectedBoneColor;
 				}
 				ChildColors.Add(ChildLineColor);
 			}
@@ -1786,6 +1787,7 @@ void FAnimationViewportClient::DrawMeshSubsetBones(const UDebugSkelMeshComponent
 
 	// we could cache parent bones as we calculate, but right now I'm not worried about perf issue of this
 	const TArray<FBoneIndexType>& DrawBoneIndices = MeshComponent->GetDrawBoneIndices();
+	const UPersonaOptions* PersonaOptions = GetDefault<UPersonaOptions>();
 	for ( auto Iter = DrawBoneIndices.CreateConstIterator(); Iter; ++Iter)
 	{
 		const int32 BoneIndex = *Iter;
@@ -1805,13 +1807,13 @@ void FAnimationViewportClient::DrawMeshSubsetBones(const UDebugSkelMeshComponent
 				{
 					WorldTransforms[ParentIndex] = MeshComponent->GetDrawTransform(ParentIndex)*MeshComponent->GetComponentTransform();
 				}
-				BoneColours[BoneIndex] = SkeletalDebugRendering::SELECTED_BONE_COLOR;
+				BoneColours[BoneIndex] = PersonaOptions->SelectedBoneColor;
 				bDrawBone = true;
 				break;
 			}
 			else if ( RefSkeleton.BoneIsChildOf(BoneIndex, SubBoneIndex) )
 			{
-				BoneColours[BoneIndex] = SkeletalDebugRendering::DEFAULT_BONE_COLOR;
+				BoneColours[BoneIndex] = PersonaOptions->DefaultBoneColor;
 				bDrawBone = true;
 				break;
 			}
diff --git a/Engine/Source/Editor/UnrealEd/Classes/Preferences/PersonaOptions.h b/Engine/Source/Editor/UnrealEd/Classes/Preferences/PersonaOptions.h
--- a/Engine/Source/Editor/UnrealEd/Classes/Preferences/PersonaOptions.h
+++ b/Engine/Source/Editor/UnrealEd/Classes/Preferences/PersonaOptions.h
@@ -146,6 +146,24 @@ class UNREALED_API UPersonaOptions : public UObject
 	UPROPERTY(EditAnywhere, config, Category = "Viewport")
 	uint32 DefaultBoneDrawSelection;
 
+	UPROPERTY(EditAnywhere, config, Category = "Viewport")
+	FLinearColor DefaultBoneColor;
+
+	UPROPERTY(EditAnywhere, config, Category = "Viewport")
+	FLinearColor SelectedBoneColor;
+
+	UPROPERTY(EditAnywhere, config, Category = "Viewport")
+	FLinearColor AffectedBoneColor;
+
+	UPROPERTY(EditAnywhere, config, Category = "Viewport")
+	FLinearColor DisabledBoneColor;
+
+	UPROPERTY(EditAnywhere, config, Category = "Viewport")
+	FLinearColor ParentOfSelectedBoneColor;
+
+	UPROPERTY(EditAnywhere, config, Category = "Viewport")
+	FLinearColor VirtualBoneColor;
+
 	UPROPERTY(EditAnywhere, config, Category = "Composites and Montages")
 	FLinearColor SectionTimingNodeColor;
 
diff --git a/Engine/Source/Editor/UnrealEd/Private/PreferenceStubs.cpp b/Engine/Source/Editor/UnrealEd/Private/PreferenceStubs.cpp
--- a/Engine/Source/Editor/UnrealEd/Private/PreferenceStubs.cpp
+++ b/Engine/Source/Editor/UnrealEd/Private/PreferenceStubs.cpp
@@ -124,6 +124,13 @@ UPersonaOptions::UPersonaOptions(const FObjectInitializer& ObjectInitializer)
 	SectionTimingNodeColor = FLinearColor(0.39f, 0.39f, 1.0f, 0.75f);
 	NotifyTimingNodeColor = FLinearColor(0.8f, 0.1f, 0.1f);
 	BranchingPointTimingNodeColor = FLinearColor(0.5f, 1.0f, 1.0f);
+	
+	DefaultBoneColor = FLinearColor(0.0f,0.0f,0.025f,1.0f);
+	SelectedBoneColor = FLinearColor(0.2f,1.0f,0.2f,1.0f);
+	AffectedBoneColor = FLinearColor(1.0f,1.0f,1.0f,1.0f);
+	DisabledBoneColor = FLinearColor(0.4f,0.4f,0.4f,1.0f);
+	ParentOfSelectedBoneColor = FLinearColor(0.85f,0.45f,0.12f,1.0f);
+	VirtualBoneColor = FLinearColor(0.4f, 0.4f, 1.0f, 1.0f);
 
 	bAutoAlignFloorToMesh = true;
 
diff --git a/Engine/Source/Runtime/Engine/Public/SkeletalDebugRendering.h b/Engine/Source/Runtime/Engine/Public/SkeletalDebugRendering.h
--- a/Engine/Source/Runtime/Engine/Public/SkeletalDebugRendering.h
+++ b/Engine/Source/Runtime/Engine/Public/SkeletalDebugRendering.h
@@ -9,13 +9,6 @@ class FPrimitiveDrawInterface;
 
 namespace SkeletalDebugRendering
 {
-	static const FLinearColor DEFAULT_BONE_COLOR = FLinearColor(0.0f,0.0f,0.025f,1.0f);
-	static const FLinearColor SELECTED_BONE_COLOR = FLinearColor(0.2f,1.0f,0.2f,1.0f);
-	static const FLinearColor AFFECTED_BONE_COLOR = FLinearColor(1.0f,1.0f,1.0f,1.0f);
-	static const FLinearColor DISABLED_BONE_COLOR = FLinearColor(0.4f,0.4f,0.4f,1.0f);
-	static const FLinearColor PARENT_OF_SELECTED_BONE_COLOR = FLinearColor(0.85f,0.45f,0.12f,1.0f);
-	static const FLinearColor VIRTUAL_BONE_COLOR = FLinearColor(0.4f, 0.4f, 1.0f, 1.0f);
-
 	static const int32 NumSphereSides = 10;
 	static const int32 NumConeSides = 4;
 
