Making UX menu for adding a float curve and metadata curve look and function the same.

**New**
- Added `FIsCurveMarkedForExclusion` Slate Event to `SAnimCurvePicker`. This delegate is determines if a curve should be stored when building the curve list to display.
- Update the `SAnimCurvePicker` widget class to store temporal arrays as member variables and update calls on arrays from using the `Empty()` method to `Reset()` method to avoid unnecessary allocations.
- Updated the `SAnimCurvePicker` widget class to support querying curves from asset registry AND also from the editable skeleton itself.
- Added callback functions in `FAnimTimelineTrack_Curves` for `SAnimCurvePicker` slate events.

**Changed**
- Now `SAnimCurvePicker` reference/stores curves using `FSmartName` instead of `FString`.
- The `FilterAvailableCurves` method in `SAnimCurvePicker` now displays the curve list sorted alphabetically.
- Updated `SAnimCurvePicker` widget UX to have better padding and avoid overlapping widgets draws.
- Updated `SAnimationSequenceBrowser` to now use `FSmartName` when hooking lambdas for `SAnimCurvePicker` slate events.
- Changed the menu button's label responsible for creating a new metadata entry from "Create New" to "Create Metadata Entry" to be more descriptive and reflect the other submenu button "Create Curve".

**Removed**
- Removed all the widget private member variables from the `FAnimTimelineTrack_Curves` class and helper functions to construct search box with curve list view and moved to using `SAnimCurvePicker` for curve selection.

diff --git a/Engine/Source/Editor/Persona/Private/AnimTimeline/AnimTimelineTrack_Curves.cpp b/Engine/Source/Editor/Persona/Private/AnimTimeline/AnimTimelineTrack_Curves.cpp
--- a/Engine/Source/Editor/Persona/Private/AnimTimeline/AnimTimelineTrack_Curves.cpp
+++ b/Engine/Source/Editor/Persona/Private/AnimTimeline/AnimTimelineTrack_Curves.cpp
@@ -6,6 +6,7 @@
 #include "Widgets/Layout/SScrollBox.h"
 #include "Widgets/Input/SSearchBox.h"
 #include "AnimSequenceTimelineCommands.h"
+#include "SAnimCurvePicker.h"
 #include "Framework/MultiBox/MultiBoxBuilder.h"
 #include "Animation/AnimSequenceBase.h"
 #include "Widgets/Input/STextEntryPopup.h"
@@ -19,8 +20,6 @@
 
 #define LOCTEXT_NAMESPACE "FAnimTimelineTrack_Notifies"
 
-const float FAnimTimelineTrack_Curves::CurveListPadding = 8.0f;
-
 ANIMTIMELINE_IMPLEMENT_TRACK(FAnimTimelineTrack_Curves);
 
 FAnimTimelineTrack_Curves::FAnimTimelineTrack_Curves(const TSharedRef<FAnimModel>& InModel)
@@ -138,272 +137,53 @@ TSharedRef<SWidget> FAnimTimelineTrack_Curves::BuildCurvesSubMenu()
 	return MenuBuilder.MakeWidget();
 }
 
-struct FSmartNameSortItemSortOp
-{
-	bool operator()(const FSmartNameSortItem& A, const FSmartNameSortItem& B) const
-	{
-		return (A.SmartName.Compare(B.SmartName) < 0);
-	}
-};
-
 void FAnimTimelineTrack_Curves::FillMetadataEntryMenu(FMenuBuilder& Builder)
 {
-	UAnimSequenceBase* AnimSequenceBase = GetModel()->GetAnimSequenceBase();
-	USkeleton* CurrentSkeleton = AnimSequenceBase->GetSkeleton();
-	check(CurrentSkeleton);
-
-	const FSmartNameMapping* Mapping = CurrentSkeleton->GetSmartNameContainer(USkeleton::AnimCurveMappingName);
-	TArray<USkeleton::AnimCurveUID> CurveUids;
-	Mapping->FillUidArray(CurveUids);
-
-	Builder.BeginSection(NAME_None, LOCTEXT("MetadataMenu_ListHeading", "Available Names"));
+	// Add new metadata entry menu button
 	{
-		TArray<FSmartNameSortItem> SmartNameList;
-
-		const TArray<FFloatCurve>& FloatCurves = AnimSequenceBase->GetDataModel()->GetFloatCurves();
-
-		for (USkeleton::AnimCurveUID Id : CurveUids)
-		{
-			FSmartName SmartName;
-			if (Mapping->FindSmartNameByUID(Id, SmartName) && !FloatCurves.ContainsByPredicate([SmartName](FFloatCurve& Curve)
-			{
-				return Curve.Name == SmartName;
-			}))
-			{
-				SmartNameList.Add(FSmartNameSortItem(SmartName.DisplayName, Id));
-			}
-		}
-
-		{
-			SmartNameList.Sort(FSmartNameSortItemSortOp());
-
-			for (FSmartNameSortItem SmartNameItem : SmartNameList)
-			{
-				const FText Description = LOCTEXT("NewMetadataSubMenu_ToolTip", "Add an existing metadata curve");
-				const FText Label = FText::FromName(SmartNameItem.SmartName);
-
-				FUIAction UIAction;
-				UIAction.ExecuteAction.BindRaw(
-					this, &FAnimTimelineTrack_Curves::AddMetadataEntry,
-					SmartNameItem.ID);
-
-				Builder.AddMenuEntry(Label, Description, FSlateIcon(), UIAction);
-			}
-		}
-	}
-	Builder.EndSection();
-
-	Builder.AddMenuSeparator();
-
-	const FText Description = LOCTEXT("NewMetadataCreateNew_ToolTip", "Create a new metadata entry");
-	const FText Label = LOCTEXT("NewMetadataCreateNew_Label","Create New");
-	FUIAction UIAction;
-	UIAction.ExecuteAction.BindRaw(this, &FAnimTimelineTrack_Curves::CreateNewMetadataEntryClicked);
-
-	Builder.AddMenuEntry(Label, Description, FSlateIcon(), UIAction);
-}
-
-TSharedRef<ITableRow> FAnimTimelineTrack_Curves::GenerateCurveListRow(FCurveListItem InItem, const TSharedRef<STableViewBase>& OwnerList)
-{
-	return SNew(STableRow<FCurveListItem>, OwnerList)
-		.Padding(FMargin(CurveListPadding, 0.0))
-		[
-			SNew(STextBlock).Text(FText::FromString(InItem.Get()->SmartName.ToString())).HighlightText(SearchText)
-		];
-}
-
-void FAnimTimelineTrack_Curves::OnTypeSelectionChanged(FCurveListItem Selection, ESelectInfo::Type SelectInfo)
-{
-	// When the user is navigating, do not act upon the selection change
-	if (SelectInfo == ESelectInfo::OnNavigation)
-	{
-		return;
-	}
-
-	if (Selection.IsValid())
-	{
-		AddVariableCurve(Selection->ID);
-		FSlateApplication::Get().DismissAllMenus();
-	}
-}
-
-void FAnimTimelineTrack_Curves::OnMouseButtonClicked(FCurveListItem Selection)
-{
-	if (Selection.IsValid())
-	{
-		AddVariableCurve(Selection->ID);
-		FSlateApplication::Get().DismissAllMenus();
-	}
-}
-
-void FAnimTimelineTrack_Curves::OnCurveFilterTextChanged(const FText& NewText)
-{
-	SearchText = NewText;
-	FilteredCurveItems.Empty();
-
-	GetCurvesMatchingSearch(NewText, CurveItems, FilteredCurveItems);
-	CurveListView->RequestListRefresh();
-
-	auto SelectedItems = CurveListView->GetSelectedItems();
-	if (FilteredCurveItems.Num() > 0)
-	{
-		CurveListView->SetSelection(FilteredCurveItems[0], ESelectInfo::OnNavigation);
-	}
-}
-
-bool FAnimTimelineTrack_Curves::GetCurvesMatchingSearch(const FText& InSearchText, const TArray<FCurveListItem>& UnfilteredList, TArray<FCurveListItem>& OutFilteredList)
-{
-	// Trim and sanitized the filter text (so that it more likely matches the action descriptions)
-	FString TrimmedFilterString = FText::TrimPrecedingAndTrailing(InSearchText).ToString();
+		const FText Description = LOCTEXT("NewMetadataCreateNew_ToolTip", "Create a new metadata entry");
+		const FText Label = LOCTEXT("NewMetadataCreateNew_Label","Create Metadata Entry");
+		FUIAction UIAction;
+		UIAction.ExecuteAction.BindRaw(this, &FAnimTimelineTrack_Curves::CreateNewMetadataEntryClicked);
 
-	// Tokenize the search box text into a set of terms; all of them must be present to pass the filter
-	TArray<FString> FilterTerms;
-	TrimmedFilterString.ParseIntoArray(FilterTerms, TEXT(" "), true);
-
-	// Generate a list of sanitized versions of the strings
-	TArray<FString> SanitizedFilterTerms;
-	for (int32 iFilters = 0; iFilters < FilterTerms.Num(); iFilters++)
-	{
-		FString EachString = FName::NameToDisplayString(FilterTerms[iFilters], false);
-		EachString = EachString.Replace(TEXT(" "), TEXT(""));
-		SanitizedFilterTerms.Add(EachString);
+		Builder.AddMenuEntry(Label, Description, FSlateIcon(), UIAction);
 	}
+	
+	Builder.BeginSection(NAME_None, LOCTEXT("MetadataMenu_ListHeading", "Available Names"));
 
-	ensure(SanitizedFilterTerms.Num() == FilterTerms.Num());
-
-	bool bReturnVal = false;
-
-	for (auto it = UnfilteredList.CreateConstIterator(); it; ++it)
+	// Add existing curve to timeline using curve picker
 	{
-		FSmartNameSortItem CurItem = *(*it).Get();
-
-		const bool bIsEmptySearch = InSearchText.IsEmpty();
-		bool bFilterTextMatches = true;
-
-		if (!bIsEmptySearch)
-		{
-			const FString DescriptionString = CurItem.SmartName.ToString();
-			const FString MangledDescriptionString = DescriptionString.Replace(TEXT(" "), TEXT(""));
-
-			for (int32 FilterIndex = 0; FilterIndex < FilterTerms.Num() && bFilterTextMatches; ++FilterIndex)
-			{
-				const bool bMatchesTerm = MangledDescriptionString.Contains(FilterTerms[FilterIndex]) || MangledDescriptionString.Contains(SanitizedFilterTerms[FilterIndex]);
-				bFilterTextMatches = bFilterTextMatches && bMatchesTerm;
-			}
-		}
-		if (bIsEmptySearch || bFilterTextMatches)
-		{
-			OutFilteredList.Add(MakeShared<FSmartNameSortItem>(CurItem));
-			bReturnVal = true;
-		}
+		const TSharedRef<SWidget> CurvePickerWidget = SNew(SAnimCurvePicker, GetModel()->GetEditableSkeleton())
+		.OnCurveNamePicked(this, &FAnimTimelineTrack_Curves::OnMetadataCurveNamePicked)
+		.IsCurveMarkedForExclusion(this, &FAnimTimelineTrack_Curves::IsCurveMarkedForExclusion);
+		Builder.AddWidget(CurvePickerWidget, FText::GetEmpty(), true);
 	}
-	return bReturnVal;
+	
+	Builder.EndSection();
 }
 
-void FAnimTimelineTrack_Curves::OnCurveFilterTextCommitted(const FText& NewText, ETextCommit::Type CommitInfo)
+void FAnimTimelineTrack_Curves::FillVariableCurveMenu(FMenuBuilder& Builder)
 {
-	if (CommitInfo == ETextCommit::OnEnter)
+	// Menu entry to create a new curve
 	{
-		auto SelectedItems = CurveListView->GetSelectedItems();
-		if (SelectedItems.Num() > 0)
-		{
-			CurveListView->SetSelection(SelectedItems[0]);
-		}
+		FText Description = LOCTEXT("NewVariableCurveCreateNew_ToolTip", "Create a new variable curve");
+		FText Label = LOCTEXT("NewVariableCurveCreateNew_Label", "Create Curve");
+		FUIAction UIAction;
+		UIAction.ExecuteAction.BindRaw(this, &FAnimTimelineTrack_Curves::CreateNewCurveClicked);
+		Builder.AddMenuEntry(Label, Description, FSlateIcon(), UIAction);
 	}
-}
-
-void FAnimTimelineTrack_Curves::FillVariableCurveMenu(FMenuBuilder& Builder)
-{
-	FText Description = LOCTEXT("NewVariableCurveCreateNew_ToolTip", "Create a new variable curve");
-	FText Label = LOCTEXT("NewVariableCurveCreateNew_Label", "Create Curve");
-	FUIAction UIAction;
-	UIAction.ExecuteAction.BindRaw(this, &FAnimTimelineTrack_Curves::CreateNewCurveClicked);
-
-	Builder.AddMenuEntry(Label, Description, FSlateIcon(), UIAction);
-
-	UAnimSequenceBase* AnimSequenceBase = GetModel()->GetAnimSequenceBase();
-	USkeleton* CurrentSkeleton = AnimSequenceBase->GetSkeleton();
-	check(CurrentSkeleton);
-
-	const FSmartNameMapping* Mapping = CurrentSkeleton->GetSmartNameContainer(USkeleton::AnimCurveMappingName);
-	TArray<USkeleton::AnimCurveUID> CurveUids;
-	Mapping->FillUidArray(CurveUids);
-	CurveItems.Empty();
+	
 	Builder.BeginSection(NAME_None, LOCTEXT("VariableMenu_ListHeading", "Available Names"));
 
-	for (USkeleton::AnimCurveUID Id : CurveUids)
+	// Add existing curve to timeline using curve picker
 	{
-		if (!AnimSequenceBase->GetDataModel()->FindFloatCurve(FAnimationCurveIdentifier(Id, ERawCurveTrackTypes::RCT_Float)))
-		{
-			FName CurveName;
-			if (Mapping->GetName(Id, CurveName))
-			{
-				CurveItems.Add(MakeShared<FSmartNameSortItem>(FSmartNameSortItem(CurveName, Id)));
-			}
-		}
+		const TSharedRef<SWidget> CurvePickerWidget = SNew(SAnimCurvePicker, GetModel()->GetEditableSkeleton())
+		.OnCurveNamePicked(this, &FAnimTimelineTrack_Curves::OnVariableCurveNamePicked)
+		.IsCurveMarkedForExclusion(this, &FAnimTimelineTrack_Curves::IsCurveMarkedForExclusion);
+		
+		Builder.AddWidget(CurvePickerWidget, FText::GetEmpty(), true);
 	}
-
-	// Build a SearchBox followed by a list of all the available curves
-	FilteredCurveItems = CurveItems; 
-	TSharedPtr<SSearchBox>	CurveFilterTextBox;
-	TSharedPtr<SMenuOwner>	MenuContent;
-
-	SAssignNew(CurveListView, SCurveListView)
-		.ListItemsSource(&FilteredCurveItems)
-		.SelectionMode(ESelectionMode::Single)
-		.OnGenerateRow(this, &FAnimTimelineTrack_Curves::GenerateCurveListRow)
-		.OnMouseButtonClick(this, &FAnimTimelineTrack_Curves::OnMouseButtonClicked)
-		.OnSelectionChanged(this, &FAnimTimelineTrack_Curves::OnTypeSelectionChanged);
-
-	SAssignNew(CurveFilterTextBox, SSearchBox)
-		.OnTextChanged(this, &FAnimTimelineTrack_Curves::OnCurveFilterTextChanged)
-		.OnTextCommitted(this, &FAnimTimelineTrack_Curves::OnCurveFilterTextCommitted);
-
-	SAssignNew(MenuContent, SMenuOwner)
-		[
-			SNew(SListViewSelectorDropdownMenu<FCurveListItem>, CurveFilterTextBox, CurveListView)
-			[
-				SNew(SVerticalBox)
-				+ SVerticalBox::Slot()
-				.AutoHeight()
-				.Padding(FMargin(CurveListPadding, 2.0))
-				[
-					CurveFilterTextBox.ToSharedRef()
-				]
-				+ SVerticalBox::Slot()
-				.FillHeight(1.f)
-				.VAlign(VAlign_Fill)
-				.Padding(CurveListPadding, 2.0)
-				[
-					SNew(SBox)
-					.WidthOverride(300.0f)
-					.HeightOverride(300.f)
-					[
-						SNew(SOverlay)
-						+ SOverlay::Slot()
-						[
-							SNew(SBorder)
-							.BorderImage(FAppStyle::GetBrush("Graph.StateNode.Body"))
-							.BorderBackgroundColor(FAppStyle::Get().GetSlateColor("Colors.Input"))
-						]
-						+SOverlay::Slot()
-						[
-							SNew(SScrollBox)
-							.Orientation(EOrientation::Orient_Vertical)
-							+ SScrollBox::Slot()
-							.HAlign(HAlign_Fill)
-							.VAlign(VAlign_Fill)
-							[
-								CurveListView.ToSharedRef()
-							]
-						]
-					]
-				]
-			]
-		];
-
-	Builder.AddWidget(MenuContent.ToSharedRef(), FText::GetEmpty(),  false);
-
+	
 	Builder.EndSection();
 }
 
@@ -523,4 +303,28 @@ bool FAnimTimelineTrack_Curves::IsShowCurvePointsEnabled() const
 	return GetDefault<UPersonaOptions>()->bTimelineDisplayCurveKeys;
 }
 
-#undef LOCTEXT_NAMESPACE
\ No newline at end of file
+void FAnimTimelineTrack_Curves::OnMetadataCurveNamePicked(const FSmartName& InCurveSmartName)
+{
+	FSlateApplication::Get().DismissAllMenus();
+
+	if (InCurveSmartName.IsValid())
+	{
+		AddMetadataEntry(InCurveSmartName.UID);
+	}
+}
+
+void FAnimTimelineTrack_Curves::OnVariableCurveNamePicked(const FSmartName& InCurveSmartName)
+{
+	FSlateApplication::Get().DismissAllMenus();
+
+	if (InCurveSmartName.IsValid())
+	{
+		AddVariableCurve(InCurveSmartName.UID);
+	}
+}
+
+bool FAnimTimelineTrack_Curves::IsCurveMarkedForExclusion(const FSmartName& InCurveSmartName)
+{
+	return GetModel()->GetAnimSequenceBase()->GetDataModel()->FindFloatCurve(FAnimationCurveIdentifier(InCurveSmartName, ERawCurveTrackTypes::RCT_Float)) != nullptr;
+}
+#undef LOCTEXT_NAMESPACE
diff --git a/Engine/Source/Editor/Persona/Private/AnimTimeline/AnimTimelineTrack_Curves.h b/Engine/Source/Editor/Persona/Private/AnimTimeline/AnimTimelineTrack_Curves.h
--- a/Engine/Source/Editor/Persona/Private/AnimTimeline/AnimTimelineTrack_Curves.h
+++ b/Engine/Source/Editor/Persona/Private/AnimTimeline/AnimTimelineTrack_Curves.h
@@ -7,20 +7,6 @@
 #include "Widgets/Views/STableRow.h"
 #include "Widgets/Views/SListView.h"
 
-struct FSmartNameSortItem
-{
-	FName SmartName;
-	USkeleton::AnimCurveUID ID;
-
-	FSmartNameSortItem(const FName& InSmartName, const USkeleton::AnimCurveUID& InID)
-		: SmartName(InSmartName)
-		, ID(InID)
-	{}
-};
-
-typedef TSharedPtr<FSmartNameSortItem> FCurveListItem;
-typedef SListView<FCurveListItem> SCurveListView;
-
 class FAnimTimelineTrack_Curves : public FAnimTimelineTrack
 {
 	ANIMTIMELINE_DECLARE_TRACK(FAnimTimelineTrack_Curves, FAnimTimelineTrack);
@@ -47,22 +33,11 @@ private:
 	void HandleShowCurvePoints();
 	bool IsShowCurvePointsEnabled() const;
 
-	virtual TSharedRef<ITableRow> GenerateCurveListRow(FCurveListItem InItem, const TSharedRef<STableViewBase>& OwnerList);
-	void OnTypeSelectionChanged(FCurveListItem Selection, ESelectInfo::Type SelectInfo);
-	void OnMouseButtonClicked(FCurveListItem Selection);
-
-	/** Curve searching support */
-	FText SearchText;
-	void OnCurveFilterTextChanged(const FText& NewText);
-	void OnCurveFilterTextCommitted(const FText& NewText, ETextCommit::Type CommitInfo);
-	bool GetCurvesMatchingSearch(const FText& InSearchText, const TArray<FCurveListItem>& UnfilteredList, TArray<FCurveListItem>& OutFilteredList);
+	/** Curve Picker Callbacks */
+	void OnMetadataCurveNamePicked(const FSmartName & InCurveSmartName);
+	void OnVariableCurveNamePicked(const FSmartName & InCurveSmartName);
+	bool IsCurveMarkedForExclusion(const FSmartName & InCurveSmartName);
+	
 private:
 	TSharedPtr<SWidget>	OutlinerWidget;
-
-	TSharedPtr<SCurveListView>	CurveListView;
-
-	TArray<FCurveListItem >	CurveItems;
-	TArray<FCurveListItem >	FilteredCurveItems;
-
-	static const float	CurveListPadding;
 };
\ No newline at end of file
diff --git a/Engine/Source/Editor/Persona/Private/SAnimCurvePicker.cpp b/Engine/Source/Editor/Persona/Private/SAnimCurvePicker.cpp
--- a/Engine/Source/Editor/Persona/Private/SAnimCurvePicker.cpp
+++ b/Engine/Source/Editor/Persona/Private/SAnimCurvePicker.cpp
@@ -9,55 +9,112 @@
 #include "Widgets/SBoxPanel.h"
 #include "Widgets/Text/STextBlock.h"
 #include "Widgets/Layout/SBox.h"
+#include "Widgets/Layout/SScrollBox.h"
+#include "SListViewSelectorDropdownMenu.h"
+#include "Widgets/SOverlay.h"
+#include "Widgets/Layout/SMenuOwner.h"
+#include "Widgets/Layout/SBorder.h"
 
 #define LOCTEXT_NAMESPACE "SAnimCurvePicker"
 
+SAnimCurvePicker::~SAnimCurvePicker()
+{
+}
+
 void SAnimCurvePicker::Construct(const FArguments& InArgs, const TSharedRef<IEditableSkeleton>& InEditableSkeleton)
 {
 	OnCurveNamePicked = InArgs._OnCurveNamePicked;
+	IsCurveMarkedForExclusion = InArgs._IsCurveMarkedForExclusion;
 	EditableSkeleton = InEditableSkeleton;
 
+	SAssignNew(SearchBox, SSearchBox)
+	.HintText(LOCTEXT("SearchBoxHint", "Search"))
+	.OnTextChanged(this, &SAnimCurvePicker::HandleFilterTextChanged);
+
+	SAssignNew(NameListView, SListView<TSharedPtr<FSmartName>>)
+	.SelectionMode(ESelectionMode::Single)
+	.ListItemsSource(&CurveNames)
+	.OnSelectionChanged(this, &SAnimCurvePicker::HandleSelectionChanged)
+	.OnGenerateRow(this, &SAnimCurvePicker::HandleGenerateRow);
+
+	const float HorizontalPadding = 8.0f;
+	const float VerticalPadding = 2.0f;
+	const float WeightOverride = 300.0f;
+	
 	ChildSlot
 	[
-		SNew(SVerticalBox)
-		+SVerticalBox::Slot()
-		.AutoHeight()
-		[
-			SNew(SSearchBox)
-			.HintText(LOCTEXT("SearchBoxHint", "Search Available Curves"))
-			.OnTextChanged(this, &SAnimCurvePicker::HandleFilterTextChanged)
-		]
-		+SVerticalBox::Slot()
-		.FillHeight(1.0f)
+		SNew(SMenuOwner)
 		[
-			SAssignNew(NameListView, SListView<TSharedPtr<FString>>)
-			.SelectionMode(ESelectionMode::Single)
-			.ItemHeight(20.0f)
-			.ListItemsSource(&CurveNames)
-			.OnSelectionChanged(this, &SAnimCurvePicker::HandleSelectionChanged)
-			.OnGenerateRow(this, &SAnimCurvePicker::HandleGenerateRow)
+			SNew(SListViewSelectorDropdownMenu<TSharedPtr<FSmartName>>, SearchBox, NameListView)
+			[
+				SNew(SVerticalBox)
+				+SVerticalBox::Slot()
+				.AutoHeight()
+				.Padding(HorizontalPadding, VerticalPadding)
+				[
+					SearchBox.ToSharedRef()
+				]
+				+SVerticalBox::Slot()
+				.FillHeight(1.0f)
+				.VAlign(VAlign_Fill)
+				.Padding(HorizontalPadding, VerticalPadding)
+				[
+					SNew(SBox)
+					.WidthOverride(WeightOverride)
+					.HeightOverride(WeightOverride)
+					[
+						SNew(SOverlay)
+						+SOverlay::Slot()
+						[
+							SNew(SBorder)
+							.BorderImage(FAppStyle::GetBrush("Graph.StateNode.Body"))
+							.BorderBackgroundColor(FAppStyle::Get().GetSlateColor("Colors.Input"))
+						]
+						+SOverlay::Slot()
+						[
+							SNew(SScrollBox)
+							.Orientation(EOrientation::Orient_Vertical)
+							+ SScrollBox::Slot()
+							.HAlign(HAlign_Fill)
+							.VAlign(VAlign_Fill)
+							[
+								NameListView.ToSharedRef()
+							]
+						]
+					]
+				]
+			]
 		]
 	];
 
 	RefreshListItems();
 }
 
-void SAnimCurvePicker::HandleSelectionChanged(TSharedPtr<FString> InItem, ESelectInfo::Type InSelectionType)
+void SAnimCurvePicker::HandleSelectionChanged(TSharedPtr<FSmartName> InItem, ESelectInfo::Type InSelectionType)
 {
-	OnCurveNamePicked.ExecuteIfBound(FName(**InItem.Get()));
+	// When the user is navigating, do not act upon the selection change
+	if (InSelectionType == ESelectInfo::OnNavigation)
+	{
+		return;
+	}
+	
+	if (InItem.IsValid())
+	{
+		OnCurveNamePicked.ExecuteIfBound(*InItem);
+	}
 }
 
-TSharedRef<ITableRow> SAnimCurvePicker::HandleGenerateRow(TSharedPtr<FString> InItem, const TSharedRef<STableViewBase>& InOwnerTable)
+TSharedRef<ITableRow> SAnimCurvePicker::HandleGenerateRow(TSharedPtr<FSmartName> InItem, const TSharedRef<STableViewBase>& InOwnerTable)
 {
 	return 
-		SNew(STableRow<TSharedPtr<FString>>, InOwnerTable)
+		SNew(STableRow<TSharedPtr<FSmartName>>, InOwnerTable)
+		.Padding(FMargin(8.0f, 0.0f))
 		[
 			SNew(SBox)
-			.MinDesiredHeight(20.0f)
 			.VAlign(VAlign_Center)
 			[
 				SNew(STextBlock)
-				.Text(FText::FromString(*InItem.Get()))
+				.Text(FText::FromName(InItem->DisplayName))
 				.HighlightText_Lambda([this]() { return FText::FromString(FilterText); })
 			]
 		];
@@ -65,53 +122,100 @@ TSharedRef<ITableRow> SAnimCurvePicker::HandleGenerateRow(TSharedPtr<FString> In
 
 void SAnimCurvePicker::RefreshListItems()
 {
-	CurveNames.Empty();
-	UniqueCurveNames.Empty();
+	// Ensure valid skeleton exists
+	const TSharedPtr<IEditableSkeleton> CurrentEditableSkeleton = EditableSkeleton.Pin();
+	check(CurrentEditableSkeleton);
+	const USkeleton & CurrentSkeleton = CurrentEditableSkeleton->GetSkeleton();
+
+	// Get name container for all the curves stored by the skeleton
+	const FSmartNameMapping* CurveNameMapping = CurrentSkeleton.GetSmartNameContainer(USkeleton::AnimCurveMappingName);
+	CurveNameMapping->FillUidArray(CurveSmartNameUids);
+	
+	CurveNames.Reset();
+	UniqueCurveNames.Reset();
 
-	// We use the asset registry to query all assets with the supplied skeleton, and accumulate their curve names
-	FAssetRegistryModule& AssetRegistryModule = FModuleManager::LoadModuleChecked<FAssetRegistryModule>("AssetRegistry");
+	{
+		// We use the asset registry to query all assets with the supplied skeleton, and accumulate their curve names
+		const FAssetRegistryModule& AssetRegistryModule = FModuleManager::LoadModuleChecked<FAssetRegistryModule>("AssetRegistry");
 
-	FARFilter Filter;
-	Filter.bRecursiveClasses = true;
-	Filter.ClassPaths.Add(UAnimationAsset::StaticClass()->GetClassPathName());
-	Filter.TagsAndValues.Add(TEXT("Skeleton"), FAssetData(&EditableSkeleton.Pin()->GetSkeleton()).GetExportTextName());
+		FARFilter Filter;
+		Filter.bRecursiveClasses = true;
+		Filter.ClassPaths.Add(UAnimationAsset::StaticClass()->GetClassPathName());
+		Filter.TagsAndValues.Add(TEXT("Skeleton"), FAssetData(&CurrentSkeleton).GetExportTextName());
 
-	TArray<FAssetData> FoundAssetData;
-	AssetRegistryModule.Get().GetAssets(Filter, FoundAssetData);
+		TArray<FAssetData> FoundAssetData;
+		AssetRegistryModule.Get().GetAssets(Filter, FoundAssetData);
 
-	// now build set of unique curve names
-	for (FAssetData& AssetData : FoundAssetData)
-	{
-		const FString TagValue = AssetData.GetTagValueRef<FString>(USkeleton::CurveNameTag);
-		if (!TagValue.IsEmpty())
+		// Build set of unique curve smart names
+		for (FAssetData& AssetData : FoundAssetData)
 		{
-			// parse notifies
-			TArray<FString> CurveValues;
-			if (TagValue.ParseIntoArray(CurveValues, *USkeleton::CurveTagDelimiter, true) > 0)
+			const FString TagValue = AssetData.GetTagValueRef<FString>(USkeleton::CurveNameTag);
+			if (!TagValue.IsEmpty())
 			{
-				for (const FString& CurveValue : CurveValues)
+				// parse notifies
+				if (TagValue.ParseIntoArray(CurveNamesQueriedFromAssetRegistry, *USkeleton::CurveTagDelimiter, true) > 0)
 				{
-					UniqueCurveNames.Add(CurveValue);
+					for (const FString& CurveNameString : CurveNamesQueriedFromAssetRegistry)
+					{
+						FName CurveName = FName(*CurveNameString);
+						FSmartName InCurveSmartName = {CurveName, CurveNameMapping->FindUID(CurveName)};
+						
+						if (IsCurveMarkedForExclusion.IsBound() && IsCurveMarkedForExclusion.Execute(InCurveSmartName))
+						{
+							continue;
+						}
+						
+						UniqueCurveNames.Add(InCurveSmartName);
+					}
 				}
 			}
 		}
 	}
+	
+	// Add unique curves that were not found using asset registry
+	for (SmartName::UID_Type Id : CurveSmartNameUids)
+	{
+		FName OutCurveName;
+		if (CurveNameMapping->GetName(Id, OutCurveName))
+		{
+			FSmartName InCurveSmartName(OutCurveName, Id);
+
+			if (IsCurveMarkedForExclusion.IsBound() && IsCurveMarkedForExclusion.Execute(InCurveSmartName))
+				continue;
+			
+			UniqueCurveNames.Add(InCurveSmartName);
+		}
+	}
 
 	FilterAvailableCurves();
 }
 
 void SAnimCurvePicker::FilterAvailableCurves()
 {
-	CurveNames.Empty();
-
-	for (const FString& UniqueCurveName : UniqueCurveNames)
+	CurveNames.Reset();
+	
+	// Exact filtering
+	for (const FSmartName& UniqueCurveName : UniqueCurveNames)
 	{
-		if (FilterText.IsEmpty() || UniqueCurveName.Contains(FilterText))
+		if (FilterText.IsEmpty() || UniqueCurveName.DisplayName.ToString().Contains(FilterText))
 		{
-			CurveNames.Add(MakeShared<FString>(UniqueCurveName));
+			CurveNames.Add(MakeShared<FSmartName>(UniqueCurveName));
 		}
 	}
 
+	// Alphabetical sorting
+	{
+		struct FSmartNameSortItemSortOp
+		{
+			FORCEINLINE bool operator()( const TSharedPtr<FSmartName>& A, const TSharedPtr<FSmartName>& B ) const
+			{
+				return (A->DisplayName.Compare(B->DisplayName) < 0);
+			}
+		};
+		CurveNames.Sort(FSmartNameSortItemSortOp());
+	}
+
+	// Rebuild list view
 	NameListView->RequestListRefresh();
 }
 
diff --git a/Engine/Source/Editor/Persona/Private/SAnimCurvePicker.h b/Engine/Source/Editor/Persona/Private/SAnimCurvePicker.h
--- a/Engine/Source/Editor/Persona/Private/SAnimCurvePicker.h
+++ b/Engine/Source/Editor/Persona/Private/SAnimCurvePicker.h
@@ -6,20 +6,33 @@
 #include "Widgets/DeclarativeSyntaxSupport.h"
 #include "Widgets/SCompoundWidget.h"
 #include "Widgets/Views/STableRow.h"
+#include "Animation/SmartName.h"
+#include "Widgets/Input/SSearchBox.h"
 
 class IEditableSkeleton;
 
-DECLARE_DELEGATE_OneParam(FOnCurveNamePicked, const FName& /*PickedName*/)
+DECLARE_DELEGATE_OneParam(FOnCurveNamePicked, const FSmartName& /*PickedName*/)
+DECLARE_DELEGATE_RetVal_OneParam(bool, FIsCurveMarkedForExclusion, const FSmartName& /*CurveName*/)
 
 class SAnimCurvePicker : public SCompoundWidget
 {
 public:
+
+	/** Virtual destructor. */
+	virtual ~SAnimCurvePicker() override;
+	
 	SLATE_BEGIN_ARGS(SAnimCurvePicker) {}
 
 	SLATE_EVENT(FOnCurveNamePicked, OnCurveNamePicked)
-
+	SLATE_EVENT(FIsCurveMarkedForExclusion, IsCurveMarkedForExclusion)
+	
 	SLATE_END_ARGS()
 
+	/**
+	 * Construct this widget.
+	 * @param InArgs - The declaration data for this widget
+	 * @param InEditableSkeleton - The skeleton from which the widget extracts curve information 
+	 */
 	void Construct(const FArguments& InArgs, const TSharedRef<IEditableSkeleton>& InEditableSkeleton);
 
 private:
@@ -30,26 +43,45 @@ private:
 	void FilterAvailableCurves();
 
 	/** UI handlers */
-	void HandleSelectionChanged(TSharedPtr<FString> InItem, ESelectInfo::Type InSelectionType);
-	TSharedRef<ITableRow> HandleGenerateRow(TSharedPtr<FString> InItem, const TSharedRef<STableViewBase>& InOwnerTable);
+	void HandleSelectionChanged(TSharedPtr<FSmartName> InItem, ESelectInfo::Type InSelectionType);
+	TSharedRef<ITableRow> HandleGenerateRow(TSharedPtr<FSmartName> InItem, const TSharedRef<STableViewBase>& InOwnerTable);
 	void HandleFilterTextChanged(const FText& InFilterText);
 
 private:
 	/** Delegate fired when a curve name is picked */
 	FOnCurveNamePicked OnCurveNamePicked;
 
+	/* Filter predicate to determine if curve should be excluded from the picker's curve list */
+	FIsCurveMarkedForExclusion IsCurveMarkedForExclusion;
+	
 	/** The editable skeleton we use to grab curves from */
 	TWeakPtr<IEditableSkeleton> EditableSkeleton;
 
 	/** The names of the curves we are displaying */
-	TArray<TSharedPtr<FString>> CurveNames;
-
+	TArray<TSharedPtr<FSmartName>> CurveNames;
+
+	struct FSmartNameKeyFuncs : public DefaultKeyFuncs<FSmartName>
+	{
+		static FORCEINLINE FSmartName GetSetKey(FSmartName const& Element) { return Element; }
+		static FORCEINLINE uint32 GetKeyHash(FSmartName const& Key) { return GetTypeHash(Key.DisplayName); }
+		static FORCEINLINE bool Matches(FSmartName const& A, FSmartName const& B) { return A == B; }
+	};
+	
 	/** All the unique curve names we can find */
-	TSet<FString> UniqueCurveNames;
+	TSet<FSmartName, FSmartNameKeyFuncs> UniqueCurveNames;
 
 	/** The string we use to filter curve names */
 	FString FilterText;
 
 	/** The list view used to display names */
-	TSharedPtr<SListView<TSharedPtr<FString>>> NameListView;
+	TSharedPtr<SListView<TSharedPtr<FSmartName>>> NameListView;
+
+	/** The search box used to filter curves */
+	TSharedPtr<SSearchBox> SearchBox;
+
+	/** The Uids of all the skeleton's curves in its mapping container */
+	TArray<SmartName::UID_Type> CurveSmartNameUids;
+
+	/** The container used to store curve names obtained from asset registry */
+	TArray<FString> CurveNamesQueriedFromAssetRegistry;
 };
\ No newline at end of file
diff --git a/Engine/Source/Editor/Persona/Private/SAnimationSequenceBrowser.cpp b/Engine/Source/Editor/Persona/Private/SAnimationSequenceBrowser.cpp
--- a/Engine/Source/Editor/Persona/Private/SAnimationSequenceBrowser.cpp
+++ b/Engine/Source/Editor/Persona/Private/SAnimationSequenceBrowser.cpp
@@ -264,10 +264,10 @@ public:
 			.WidthOverride(200.0f)
 			[
 				SNew(SAnimCurvePicker, EditableSkeleton.ToSharedRef())
-				.OnCurveNamePicked_Lambda([this](const FName& InCurveName)
+				.OnCurveNamePicked_Lambda([this](const FSmartName& InCurveSmartName)
 				{
 					FSlateApplication::Get().DismissAllMenus();
-					CurveString = InCurveName.ToString();
+					CurveString = InCurveSmartName.DisplayName.ToString();
 					OnChanged().Broadcast();
 				})
 			];
