AnimationEditorViewport: Avoid throwing errors during tick + print on screen messages

Two separate changes:
a) avoid errors being through during one of the tick methods of the client. this was happening in Control Rig and we would see a flood of errors in the log
b) allow on screen text to be rendered within persona viewports

diff --git a/Engine/Source/Editor/Persona/Private/AnimationEditorViewportClient.cpp b/Engine/Source/Editor/Persona/Private/AnimationEditorViewportClient.cpp
--- a/Engine/Source/Editor/Persona/Private/AnimationEditorViewportClient.cpp
+++ b/Engine/Source/Editor/Persona/Private/AnimationEditorViewportClient.cpp
@@ -591,6 +591,8 @@ void FAnimationViewportClient::DrawCanvas( FViewport& InViewport, FSceneView& Vi
 {
 	FEditorViewportClient::DrawCanvas(InViewport, View, Canvas);
 
+	UWorld* World = nullptr;
+	
 	TArray<UDebugSkelMeshComponent*> PreviewMeshComponents = GetPreviewScene()->GetAllPreviewMeshComponents();
 	for  (UDebugSkelMeshComponent* PreviewMeshComponent : PreviewMeshComponents)
 	{
@@ -618,7 +620,21 @@ void FAnimationViewportClient::DrawCanvas( FViewport& InViewport, FSceneView& Vi
 
 		// Debug draw clothing texts
 		PreviewMeshComponent->DebugDrawClothingTexts(&Canvas, &View);
+
+		if(World == nullptr)
+		{
+			World = PreviewMeshComponent->GetWorld();
+		}
+	}
+
+#if !(UE_BUILD_TEST)
+	if (World && GEngine->bEnableOnScreenDebugMessagesDisplay && GEngine->bEnableOnScreenDebugMessages)
+	{
+		constexpr int32 MessageX = 20;
+		constexpr int32 MessageY = 65;
+		GEngine->DrawOnscreenDebugMessages(World, Viewport, &Canvas, nullptr, MessageX, MessageY);
 	}
+#endif
 }
 
 void FAnimationViewportClient::DrawUVsForMesh(FViewport* InViewport, FCanvas* InCanvas, int32 InTextYPos, UDebugSkelMeshComponent* PreviewMeshComponent)
@@ -654,11 +670,14 @@ void FAnimationViewportClient::HandlePreviewScenePostTick()
 		const bool bInBoneOrbitMode = CameraFollowMode != EAnimationViewportCameraFollowMode::None;
 		if (IsTracking())
 		{
-			if (PreviewMeshComponent->IsPlaying() && GetDefault<UPersonaOptions>()->bPauseAnimationOnCameraMove)
+			if (const UAnimSingleNodeInstance* SingleNodeInstance = PreviewMeshComponent->GetSingleNodeInstance())
 			{
-				PreviewMeshComponent->Stop();
-				bResumeAfterTracking = true;
-				return;
+				if (SingleNodeInstance->IsPlaying() && GetDefault<UPersonaOptions>()->bPauseAnimationOnCameraMove)
+				{
+					PreviewMeshComponent->Stop();
+					bResumeAfterTracking = true;
+					return;
+				}
 			}
 		}
 		else if (bResumeAfterTracking)
