Inclusive Terminology Push
Renamed terms at SRetargetSources (trivial, just a comment)

diff --git a/Engine/Source/Editor/Persona/Private/SRetargetSources.cpp b/Engine/Source/Editor/Persona/Private/SRetargetSources.cpp
--- a/Engine/Source/Editor/Persona/Private/SRetargetSources.cpp
+++ b/Engine/Source/Editor/Persona/Private/SRetargetSources.cpp
@@ -321,7 +321,7 @@ void SRetargetSources::UseCurrentPose()
 			// @todo check to see if skeleton vs preview mesh makes it different for missing bones
 			const FReferenceSkeleton& RefSkeleton = PreviewMesh->GetRefSkeleton();
 			TArray<FTransform> & NewRetargetBasePose = PreviewMesh->GetRetargetBasePose();
-			// if you're using master pose component in preview, this won't work
+			// if you're using leader pose component in preview, this won't work
 			check(PreviewMesh->GetRefSkeleton().GetNum() == SpaceBases.Num());
 			int32 TotalNumBones = PreviewMesh->GetRefSkeleton().GetNum();
 			NewRetargetBasePose.Empty(TotalNumBones);
