Skeletal Mesh - Replace the newly created GetSkeletalMesh function by GetSkeletalMeshAsset so it matches the setter SetSkeletalMeshAsset function (which itself couldn't be named SetSkeletalMesh due to

diff --git a/Engine/Plugins/Runtime/HairStrands/Source/HairStrandsCore/Private/GroomComponent.cpp b/Engine/Plugins/Runtime/HairStrands/Source/HairStrandsCore/Private/GroomComponent.cpp
--- a/Engine/Plugins/Runtime/HairStrands/Source/HairStrandsCore/Private/GroomComponent.cpp
+++ b/Engine/Plugins/Runtime/HairStrands/Source/HairStrandsCore/Private/GroomComponent.cpp
@@ -2250,9 +2250,9 @@ static USkeletalMeshComponent* ValidateBindingAsset(
 
 	// Optional advanced check
 	bool bHasValidSectionCount = SkeletalMeshComponent && SkeletalMeshComponent->GetNumMaterials() < int32(GetHairStrandsMaxSectionCount());
-	if (SkeletalMeshComponent && SkeletalMeshComponent->GetSkeletalMesh() && SkeletalMeshComponent->GetSkeletalMesh()->GetResourceForRendering())
+	if (SkeletalMeshComponent && SkeletalMeshComponent->GetSkeletalMeshAsset() && SkeletalMeshComponent->GetSkeletalMeshAsset()->GetResourceForRendering())
 	{
-		const FSkeletalMeshRenderData* RenderData = SkeletalMeshComponent->GetSkeletalMesh()->GetResourceForRendering();
+		const FSkeletalMeshRenderData* RenderData = SkeletalMeshComponent->GetSkeletalMeshAsset()->GetResourceForRendering();
 		const int32 MaxSectionCount = GetHairStrandsMaxSectionCount();
 		// Check that all LOD are below the number sections
 		const uint32 MeshLODCount = RenderData->LODRenderData.Num();
@@ -2279,7 +2279,7 @@ static USkeletalMeshComponent* ValidateBindingAsset(
 	}
 
 	const bool bIsBindingCompatible =
-		UGroomBindingAsset::IsCompatible(SkeletalMeshComponent ? SkeletalMeshComponent->GetSkeletalMesh() : nullptr, BindingAsset, bValidationEnable) &&
+		UGroomBindingAsset::IsCompatible(SkeletalMeshComponent ? SkeletalMeshComponent->GetSkeletalMeshAsset() : nullptr, BindingAsset, bValidationEnable) &&
 		UGroomBindingAsset::IsCompatible(GroomAsset, BindingAsset, bValidationEnable) &&
 		UGroomBindingAsset::IsBindingAssetValid(BindingAsset, bIsBindingReloading, bValidationEnable);
 
@@ -2303,7 +2303,7 @@ static USkeletalMeshComponent* ValidateBindingAsset(
 					GroupIt < BindingAsset->HairGroupResources.Num() &&
 					CardsLODIndex < uint32(BindingAsset->HairGroupResources[GroupIt].CardsRootResources.Num()) &&
 					BindingAsset->HairGroupResources[GroupIt].CardsRootResources[CardsLODIndex] != nullptr &&
-					((SkeletalMeshComponent && SkeletalMeshComponent->GetSkeletalMesh()) ? SkeletalMeshComponent->GetSkeletalMesh()->GetLODInfoArray().Num() == BindingAsset->HairGroupResources[GroupIt].CardsRootResources[CardsLODIndex]->BulkData.MeshProjectionLODs.Num() : false);
+					((SkeletalMeshComponent && SkeletalMeshComponent->GetSkeletalMeshAsset()) ? SkeletalMeshComponent->GetSkeletalMeshAsset()->GetLODInfoArray().Num() == BindingAsset->HairGroupResources[GroupIt].CardsRootResources[CardsLODIndex]->BulkData.MeshProjectionLODs.Num() : false);
 
 				if (!bIsCardsBindingCompatible)
 				{
@@ -2426,7 +2426,7 @@ void UGroomComponent::InitResources(bool bIsBindingReloading)
 		if (ValidatedMeshComponent)
 		{
 			if (BindingAsset && 
-				((BindingAsset->GroomBindingType == EGroomBindingMeshType::SkeletalMesh && Cast<USkeletalMeshComponent>(ValidatedMeshComponent)->GetSkeletalMesh() == nullptr) ||
+				((BindingAsset->GroomBindingType == EGroomBindingMeshType::SkeletalMesh && Cast<USkeletalMeshComponent>(ValidatedMeshComponent)->GetSkeletalMeshAsset() == nullptr) ||
 				(BindingAsset->GroomBindingType == EGroomBindingMeshType::GeometryCache && Cast<UGeometryCacheComponent>(ValidatedMeshComponent)->GeometryCache == nullptr)))
 			{
 				ValidatedMeshComponent = nullptr;
@@ -2437,7 +2437,7 @@ void UGroomComponent::InitResources(bool bIsBindingReloading)
 	{
 		if (USkeletalMeshComponent* ParentSkelMeshComponent = Cast<USkeletalMeshComponent>(ParentMeshComponent))
 		{
-			ValidatedMeshComponent = ParentSkelMeshComponent->GetSkeletalMesh() ? ParentMeshComponent : nullptr;
+			ValidatedMeshComponent = ParentSkelMeshComponent->GetSkeletalMeshAsset() ? ParentMeshComponent : nullptr;
 		}
 	}
 
@@ -2451,7 +2451,7 @@ void UGroomComponent::InitResources(bool bIsBindingReloading)
 		for (USceneComponent* Curr = this; Curr; Curr = Curr->GetAttachParent())
 		{
 			USkeletalMeshComponent* SkelMeshComp = Cast<USkeletalMeshComponent>(Curr);
-			if (SkelMeshComp && SkelMeshComp->GetSkeletalMesh() && SkelMeshComp->GetSkeletalMesh()->GetSkeleton() == TargetSkeleton)
+			if (SkelMeshComp && SkelMeshComp->GetSkeletalMeshAsset() && SkelMeshComp->GetSkeletalMeshAsset()->GetSkeleton() == TargetSkeleton)
 			{
 				DeformedMeshComponent = SkelMeshComp;
 				break;
@@ -2635,7 +2635,7 @@ void UGroomComponent::InitResources(bool bIsBindingReloading)
 				check(GroupIt < LocalBindingAsset->HairGroupResources.Num());
 				if (USkeletalMeshComponent* SkeletalMeshComponent = Cast<USkeletalMeshComponent>(RegisteredMeshComponent))
 				{
-					check(SkeletalMeshComponent->GetSkeletalMesh() ? SkeletalMeshComponent->GetSkeletalMesh()->GetLODInfoArray().Num() == LocalBindingAsset->HairGroupResources[GroupIt].SimRootResources->BulkData.MeshProjectionLODs.Num() : false);
+					check(SkeletalMeshComponent->GetSkeletalMeshAsset() ? SkeletalMeshComponent->GetSkeletalMeshAsset()->GetLODInfoArray().Num() == LocalBindingAsset->HairGroupResources[GroupIt].SimRootResources->BulkData.MeshProjectionLODs.Num() : false);
 				}
 
 				HairGroupInstance->Guides.RestRootResource = LocalBindingAsset->HairGroupResources[GroupIt].SimRootResources;
@@ -2747,7 +2747,7 @@ void UGroomComponent::InitResources(bool bIsBindingReloading)
 				check(GroupIt < LocalBindingAsset->HairGroupResources.Num());
 				if (USkeletalMeshComponent* SkeletalMeshComponent = Cast<USkeletalMeshComponent>(RegisteredMeshComponent))
 				{
-					check(SkeletalMeshComponent->GetSkeletalMesh() ? SkeletalMeshComponent->GetSkeletalMesh()->GetLODInfoArray().Num() == LocalBindingAsset->HairGroupResources[GroupIt].RenRootResources->BulkData.MeshProjectionLODs.Num() : false);
+					check(SkeletalMeshComponent->GetSkeletalMeshAsset() ? SkeletalMeshComponent->GetSkeletalMeshAsset()->GetLODInfoArray().Num() == LocalBindingAsset->HairGroupResources[GroupIt].RenRootResources->BulkData.MeshProjectionLODs.Num() : false);
 				}
 
 				HairGroupInstance->Strands.RestRootResource = LocalBindingAsset->HairGroupResources[GroupIt].RenRootResources;
@@ -2851,7 +2851,7 @@ void UGroomComponent::InitResources(bool bIsBindingReloading)
 						check(GroupIt < LocalBindingAsset->HairGroupResources.Num());
 						if (USkeletalMeshComponent* SkeletalMeshComponent = Cast<USkeletalMeshComponent>(RegisteredMeshComponent))
 						{
-							check(SkeletalMeshComponent->GetSkeletalMesh() ? SkeletalMeshComponent->GetSkeletalMesh()->GetLODInfoArray().Num() == LocalBindingAsset->HairGroupResources[GroupIt].CardsRootResources[CardsLODIndex]->BulkData.MeshProjectionLODs.Num() : false);
+							check(SkeletalMeshComponent->GetSkeletalMeshAsset() ? SkeletalMeshComponent->GetSkeletalMeshAsset()->GetLODInfoArray().Num() == LocalBindingAsset->HairGroupResources[GroupIt].CardsRootResources[CardsLODIndex]->BulkData.MeshProjectionLODs.Num() : false);
 						}
 
 						InstanceLOD.Guides.RestRootResource = LocalBindingAsset->HairGroupResources[GroupIt].CardsRootResources[CardsLODIndex];
@@ -3289,9 +3289,9 @@ void UGroomComponent::BuildSimulationTransform(FTransform& SimulationTransform)
 	if (SimulationSettings.SimulationSetup.bLocalSimulation)
 	{
 		const USkeletalMeshComponent* SkeletelMeshComponent = Cast<const USkeletalMeshComponent>(RegisteredMeshComponent);
-		if (SkeletelMeshComponent && SkeletelMeshComponent->GetSkeletalMesh() && !SimulationSettings.SimulationSetup.LocalBone.IsEmpty())
+		if (SkeletelMeshComponent && SkeletelMeshComponent->GetSkeletalMeshAsset() && !SimulationSettings.SimulationSetup.LocalBone.IsEmpty())
 		{
-			const FReferenceSkeleton& RefSkeleton = SkeletelMeshComponent->GetSkeletalMesh()->GetRefSkeleton();
+			const FReferenceSkeleton& RefSkeleton = SkeletelMeshComponent->GetSkeletalMeshAsset()->GetRefSkeleton();
 			const FName BoneName(SimulationSettings.SimulationSetup.LocalBone);
 			const int32 BoneIndex = RefSkeleton.FindBoneIndex(BoneName);
 
@@ -3330,14 +3330,14 @@ void UGroomComponent::TickComponent(float DeltaTime, enum ELevelTick TickType, F
 	// When a groom binding and simulation are disabled, and the groom component is parented with a skeletal mesh, we can optionally 
 	// attach the groom to a particular socket/bone
 	USkeletalMeshComponent* SkeletelMeshComponent = Cast<USkeletalMeshComponent>(RegisteredMeshComponent);
-	if (SkeletelMeshComponent && SkeletelMeshComponent->GetSkeletalMesh() && !AttachmentName.IsEmpty())
+	if (SkeletelMeshComponent && SkeletelMeshComponent->GetSkeletalMeshAsset() && !AttachmentName.IsEmpty())
 	{
 		const FName BoneName(AttachmentName);
 		if (GetAttachSocketName() != BoneName)
 		{
 			AttachToComponent(SkeletelMeshComponent, FAttachmentTransformRules(EAttachmentRule::KeepRelative, false), BoneName);
 			const uint32 BoneIndex = SkeletelMeshComponent->GetBoneIndex(BoneName);
-			const FMatrix BoneTransformRaw = SkeletelMeshComponent->GetSkeletalMesh()->GetComposedRefPoseMatrix(BoneIndex);
+			const FMatrix BoneTransformRaw = SkeletelMeshComponent->GetSkeletalMeshAsset()->GetComposedRefPoseMatrix(BoneIndex);
 			const FVector BoneLocation = BoneTransformRaw.GetOrigin();
 			const FQuat BoneRotation = BoneTransformRaw.ToQuat();
 
diff --git a/Engine/Plugins/Runtime/HairStrands/Source/HairStrandsCore/Private/GroomGeometryCache.cpp b/Engine/Plugins/Runtime/HairStrands/Source/HairStrandsCore/Private/GroomGeometryCache.cpp
--- a/Engine/Plugins/Runtime/HairStrands/Source/HairStrandsCore/Private/GroomGeometryCache.cpp
+++ b/Engine/Plugins/Runtime/HairStrands/Source/HairStrandsCore/Private/GroomGeometryCache.cpp
@@ -69,7 +69,7 @@ ENGINE_API void UpdatePreviousRefToLocalMatrices(TArray<FMatrix44f>& ReferenceTo
 		return;
 	}
 
-	FSkeletalMeshRenderData* RenderData = SkeletalMeshComponent->GetSkeletalMesh()->GetResourceForRendering();
+	FSkeletalMeshRenderData* RenderData = SkeletalMeshComponent->GetSkeletalMeshAsset()->GetResourceForRendering();
 	FSkeletalMeshLODRenderData& LODData = RenderData->LODRenderData[LODIndex];
 
 	const bool bNeedPreviousPosition = IsHairStrandContinuousDecimationReorderingEnabled();
diff --git a/Engine/Plugins/Runtime/HairStrands/Source/HairStrandsCore/Private/Niagara/NiagaraDataInterfaceHairStrands.cpp b/Engine/Plugins/Runtime/HairStrands/Source/HairStrandsCore/Private/Niagara/NiagaraDataInterfaceHairStrands.cpp
--- a/Engine/Plugins/Runtime/HairStrands/Source/HairStrandsCore/Private/Niagara/NiagaraDataInterfaceHairStrands.cpp
+++ b/Engine/Plugins/Runtime/HairStrands/Source/HairStrandsCore/Private/Niagara/NiagaraDataInterfaceHairStrands.cpp
@@ -636,9 +636,9 @@ FORCEINLINE bool RequiresSimulationReset(FNiagaraSystemInstance* SystemInstance,
 			for (UActorComponent* ActorComp : RootActor->GetComponents())
 			{
 				USkeletalMeshComponent* SkelMeshComp = Cast<USkeletalMeshComponent>(ActorComp);
-				if (SkelMeshComp && SkelMeshComp->GetSkeletalMesh())
+				if (SkelMeshComp && SkelMeshComp->GetSkeletalMeshAsset())
 				{
-					NewSkeletalMeshes += GetTypeHash(SkelMeshComp->GetSkeletalMesh()->GetName());
+					NewSkeletalMeshes += GetTypeHash(SkelMeshComp->GetSkeletalMeshAsset()->GetName());
 				}
 			}
 		}
diff --git a/Engine/Plugins/Runtime/HairStrands/Source/HairStrandsEditor/Private/GroomCustomAssetEditorToolkit.cpp b/Engine/Plugins/Runtime/HairStrands/Source/HairStrandsEditor/Private/GroomCustomAssetEditorToolkit.cpp
--- a/Engine/Plugins/Runtime/HairStrands/Source/HairStrandsEditor/Private/GroomCustomAssetEditorToolkit.cpp
+++ b/Engine/Plugins/Runtime/HairStrands/Source/HairStrandsEditor/Private/GroomCustomAssetEditorToolkit.cpp
@@ -167,7 +167,7 @@ bool FGroomCustomAssetEditorToolkit::OnShouldFilterAnimAsset(const FAssetData& A
 	// Check the compatible skeletons.
 	if (PreviewSkeletalMeshComponent != nullptr)
 	{
-		USkeleton* Skeleton = PreviewSkeletalMeshComponent->GetSkeletalMesh()->GetSkeleton();
+		USkeleton* Skeleton = PreviewSkeletalMeshComponent->GetSkeletalMeshAsset()->GetSkeleton();
 		if (Skeleton && Skeleton->IsCompatibleSkeletonByAssetData(AssetData))
 		{
 			return false;
@@ -360,7 +360,7 @@ static UAnimationAsset* GetFirstCompatibleAnimAsset(const USkeletalMeshComponent
 		return nullptr;
 	}
 
-	USkeleton* InSkeleton = InComponent->GetSkeletalMesh()->GetSkeleton();
+	USkeleton* InSkeleton = InComponent->GetSkeletalMeshAsset()->GetSkeleton();
 	if (!InSkeleton)
 	{
 		return nullptr;
diff --git a/Engine/Plugins/Runtime/MeshModelingToolset/Source/ModelingComponentsEditorOnly/Private/ToolTargets/SkeletalMeshComponentToolTarget.cpp b/Engine/Plugins/Runtime/MeshModelingToolset/Source/ModelingComponentsEditorOnly/Private/ToolTargets/SkeletalMeshComponentToolTarget.cpp
--- a/Engine/Plugins/Runtime/MeshModelingToolset/Source/ModelingComponentsEditorOnly/Private/ToolTargets/SkeletalMeshComponentToolTarget.cpp
+++ b/Engine/Plugins/Runtime/MeshModelingToolset/Source/ModelingComponentsEditorOnly/Private/ToolTargets/SkeletalMeshComponentToolTarget.cpp
@@ -33,7 +33,7 @@ bool USkeletalMeshComponentReadOnlyToolTarget::IsValid() const
 	{
 		return false;
 	}
-	const USkeletalMesh* SkeletalMesh = SkeletalMeshComponent->GetSkeletalMesh();
+	const USkeletalMesh* SkeletalMesh = SkeletalMeshComponent->GetSkeletalMeshAsset();
 
 	return USkeletalMeshReadOnlyToolTarget::IsValid(SkeletalMesh);
 }
@@ -54,7 +54,7 @@ void USkeletalMeshComponentReadOnlyToolTarget::GetMaterialSet(FComponentMaterial
 
 	if (bPreferAssetMaterials)
 	{
-		const USkeletalMesh* SkeletalMesh = Cast<USkeletalMeshComponent>(Component)->GetSkeletalMesh();
+		const USkeletalMesh* SkeletalMesh = Cast<USkeletalMeshComponent>(Component)->GetSkeletalMeshAsset();
 		USkeletalMeshToolTarget::GetMaterialSet(SkeletalMesh, MaterialSetOut, bPreferAssetMaterials);
 	}
 	else
@@ -74,7 +74,7 @@ bool USkeletalMeshComponentReadOnlyToolTarget::CommitMaterialSetUpdate(const FCo
 
 	if (bApplyToAsset)
 	{
-		USkeletalMesh* SkeletalMesh = Cast<USkeletalMeshComponent>(Component)->GetSkeletalMesh();
+		USkeletalMesh* SkeletalMesh = Cast<USkeletalMeshComponent>(Component)->GetSkeletalMeshAsset();
 
 		// unregister the component while we update it's static mesh
 		TUniquePtr<FComponentReregisterContext> ComponentReregisterContext = MakeUnique<FComponentReregisterContext>(Component);
@@ -115,7 +115,7 @@ const FMeshDescription* USkeletalMeshComponentReadOnlyToolTarget::GetMeshDescrip
 	if (!CachedMeshDescription.IsValid())
 	{
 		CachedMeshDescription = MakeUnique<FMeshDescription>();
-		const USkeletalMesh* SkeletalMesh = Cast<USkeletalMeshComponent>(Component)->GetSkeletalMesh();
+		const USkeletalMesh* SkeletalMesh = Cast<USkeletalMeshComponent>(Component)->GetSkeletalMeshAsset();
 		USkeletalMeshToolTarget::GetMeshDescription(SkeletalMesh, *CachedMeshDescription);
 	}
 	
@@ -139,7 +139,7 @@ FDynamicMesh3 USkeletalMeshComponentReadOnlyToolTarget::GetDynamicMesh()
 
 USkeletalMesh* USkeletalMeshComponentReadOnlyToolTarget::GetSkeletalMesh() const
 {
-	return IsValid() ? Cast<USkeletalMeshComponent>(Component)->GetSkeletalMesh() : nullptr;
+	return IsValid() ? Cast<USkeletalMeshComponent>(Component)->GetSkeletalMeshAsset() : nullptr;
 }
 
 
@@ -152,7 +152,7 @@ void USkeletalMeshComponentToolTarget::CommitMeshDescription(const FCommitter& C
 {
 	if (ensure(IsValid()) == false) return;
 
-	USkeletalMesh* SkeletalMesh = Cast<USkeletalMeshComponent>(Component)->GetSkeletalMesh();
+	USkeletalMesh* SkeletalMesh = Cast<USkeletalMeshComponent>(Component)->GetSkeletalMeshAsset();
 	if (!CachedMeshDescription.IsValid())
 	{
 		CachedMeshDescription = MakeUnique<FMeshDescription>();
@@ -193,8 +193,8 @@ bool USkeletalMeshComponentReadOnlyToolTargetFactory::CanBuildTarget(UObject* So
 	// just add another factory that allows that class specifically(but make sure that
 	// GetMeshDescription and such work properly)
 
-	return Cast<USkeletalMeshComponent>(SourceObject) && Cast<USkeletalMeshComponent>(SourceObject)->GetSkeletalMesh() &&
-		ExactCast<USkeletalMesh>(Cast<USkeletalMeshComponent>(SourceObject)->GetSkeletalMesh()) &&
+	return Cast<USkeletalMeshComponent>(SourceObject) && Cast<USkeletalMeshComponent>(SourceObject)->GetSkeletalMeshAsset() &&
+		ExactCast<USkeletalMesh>(Cast<USkeletalMeshComponent>(SourceObject)->GetSkeletalMeshAsset()) &&
 		Requirements.AreSatisfiedBy(USkeletalMeshComponentReadOnlyToolTarget::StaticClass());
 }
 
@@ -216,8 +216,8 @@ bool USkeletalMeshComponentToolTargetFactory::CanBuildTarget(UObject* SourceObje
 	// just add another factory that allows that class specifically(but make sure that
 	// GetMeshDescription and such work properly)
 
-	return Cast<USkeletalMeshComponent>(SourceObject) && Cast<USkeletalMeshComponent>(SourceObject)->GetSkeletalMesh() &&
-		ExactCast<USkeletalMesh>(Cast<USkeletalMeshComponent>(SourceObject)->GetSkeletalMesh()) &&
+	return Cast<USkeletalMeshComponent>(SourceObject) && Cast<USkeletalMeshComponent>(SourceObject)->GetSkeletalMeshAsset() &&
+		ExactCast<USkeletalMesh>(Cast<USkeletalMeshComponent>(SourceObject)->GetSkeletalMeshAsset()) &&
 		Requirements.AreSatisfiedBy(USkeletalMeshComponentToolTarget::StaticClass());
 }
 
diff --git a/Engine/Plugins/VirtualProduction/RemoteControl/Source/RemoteControlUI/Private/RemoteControlUIModule.cpp b/Engine/Plugins/VirtualProduction/RemoteControl/Source/RemoteControlUI/Private/RemoteControlUIModule.cpp
--- a/Engine/Plugins/VirtualProduction/RemoteControl/Source/RemoteControlUI/Private/RemoteControlUIModule.cpp
+++ b/Engine/Plugins/VirtualProduction/RemoteControl/Source/RemoteControlUI/Private/RemoteControlUIModule.cpp
@@ -1036,7 +1036,7 @@ UMeshComponent* FRemoteControlUIModule::GetSelectedMeshComponentToBeModified(UOb
 		}
 		else if (USkeletalMeshComponent* SkeletalMeshComponent = Cast<USkeletalMeshComponent>(InMeshComponent))
 		{
-			if (SkeletalMeshComponent->GetSkeletalMesh() == InOwnerObject)
+			if (SkeletalMeshComponent->GetSkeletalMeshAsset() == InOwnerObject)
 			{
 				ModifinedComponent = InMeshComponent;
 			}
diff --git a/Engine/Plugins/VirtualProduction/Takes/Source/TakeTrackRecorders/Private/MovieScene3DTransformTrackRecorder.cpp b/Engine/Plugins/VirtualProduction/Takes/Source/TakeTrackRecorders/Private/MovieScene3DTransformTrackRecorder.cpp
--- a/Engine/Plugins/VirtualProduction/Takes/Source/TakeTrackRecorders/Private/MovieScene3DTransformTrackRecorder.cpp
+++ b/Engine/Plugins/VirtualProduction/Takes/Source/TakeTrackRecorders/Private/MovieScene3DTransformTrackRecorder.cpp
@@ -490,8 +490,8 @@ void UMovieScene3DTransformTrackRecorder::PostProcessAnimationData(UMovieSceneAn
 			}
 			// Search for the Root Bone in the Skeleton
 			const USkeletalMesh* SkeletalMesh = Cast<USkeletalMeshComponent>(SkeletalMeshComponent->LeaderPoseComponent) ?
-				Cast<USkeletalMeshComponent>(SkeletalMeshComponent->LeaderPoseComponent)->GetSkeletalMesh() : 
-				SkeletalMeshComponent->GetSkeletalMesh();
+				Cast<USkeletalMeshComponent>(SkeletalMeshComponent->LeaderPoseComponent)->GetSkeletalMeshAsset() : 
+				SkeletalMeshComponent->GetSkeletalMeshAsset();
 			const UAnimSequence* AnimSequence = AnimTrackRecorder->GetAnimSequence();
 			if (AnimSequence && SkeletalMesh && AnimSequence->GetSkeleton())
 			{
diff --git a/Engine/Plugins/VirtualProduction/Takes/Source/TakeTrackRecorders/Private/MovieSceneAnimationTrackRecorder.cpp b/Engine/Plugins/VirtualProduction/Takes/Source/TakeTrackRecorders/Private/MovieSceneAnimationTrackRecorder.cpp
--- a/Engine/Plugins/VirtualProduction/Takes/Source/TakeTrackRecorders/Private/MovieSceneAnimationTrackRecorder.cpp
+++ b/Engine/Plugins/VirtualProduction/Takes/Source/TakeTrackRecorders/Private/MovieSceneAnimationTrackRecorder.cpp
@@ -29,7 +29,7 @@ bool FMovieSceneAnimationTrackRecorderFactory::CanRecordObject(UObject* InObject
 	if (InObjectToRecord->IsA<USkeletalMeshComponent>())
 	{
 		USkeletalMeshComponent* SkeletalMeshComponent = Cast<USkeletalMeshComponent>(InObjectToRecord);
-		if (SkeletalMeshComponent && SkeletalMeshComponent->GetSkeletalMesh())
+		if (SkeletalMeshComponent && SkeletalMeshComponent->GetSkeletalMeshAsset())
 		{
 			return true;
 		}
@@ -46,7 +46,7 @@ void UMovieSceneAnimationTrackRecorder::CreateAnimationAssetAndSequence(const AA
 {
 	UMovieSceneAnimationTrackRecorderSettings* AnimSettings = CastChecked<UMovieSceneAnimationTrackRecorderSettings>(Settings.Get());
 
-	SkeletalMesh = SkeletalMeshComponent->GetSkeletalMesh();
+	SkeletalMesh = SkeletalMeshComponent->GetSkeletalMeshAsset();
 	if (SkeletalMesh.IsValid())
 	{
 		ComponentTransform = SkeletalMeshComponent->GetComponentToWorld().GetRelativeTransform(Actor->GetTransform());
@@ -72,7 +72,7 @@ void UMovieSceneAnimationTrackRecorder::CreateAnimationAssetAndSequence(const AA
 			FAssetRegistryModule::AssetCreated(AnimSequence.Get());
 
 			// Assign the skeleton we're recording to the newly created Animation Sequence.
-			AnimSequence->SetSkeleton(SkeletalMeshComponent->GetSkeletalMesh()->GetSkeleton());
+			AnimSequence->SetSkeleton(SkeletalMeshComponent->GetSkeletalMeshAsset()->GetSkeleton());
 		}
 	}
 
@@ -159,8 +159,8 @@ void UMovieSceneAnimationTrackRecorder::CreateTrackImpl()
 				// verify if this bone exists in skeleton
 				const int32 BoneTreeIndex = AnimSkeleton->GetSkeletonBoneIndexFromMeshBoneIndex(
 					Cast<USkeletalMeshComponent>(SkeletalMeshComponent->LeaderPoseComponent) ? 
-					Cast<USkeletalMeshComponent>(SkeletalMeshComponent->LeaderPoseComponent)->GetSkeletalMesh() :
-					SkeletalMeshComponent->GetSkeletalMesh(), BoneIndex);
+					Cast<USkeletalMeshComponent>(SkeletalMeshComponent->LeaderPoseComponent)->GetSkeletalMeshAsset() :
+					SkeletalMeshComponent->GetSkeletalMeshAsset(), BoneIndex);
 				if (BoneTreeIndex != INDEX_NONE)
 				{
 					// add tracks for the bone existing
diff --git a/Engine/Source/Developer/MeshMergeUtilities/Private/MeshMergeHelpers.cpp b/Engine/Source/Developer/MeshMergeUtilities/Private/MeshMergeHelpers.cpp
--- a/Engine/Source/Developer/MeshMergeUtilities/Private/MeshMergeHelpers.cpp
+++ b/Engine/Source/Developer/MeshMergeUtilities/Private/MeshMergeHelpers.cpp
@@ -109,13 +109,13 @@ void FMeshMergeHelpers::ExtractSections(const UStaticMeshComponent* Component, i
 void FMeshMergeHelpers::ExtractSections(const USkeletalMeshComponent* Component, int32 LODIndex, TArray<FSectionInfo>& OutSections)
 {
 	static UMaterialInterface* DefaultMaterial = UMaterial::GetDefaultMaterial(MD_Surface);
-	FSkeletalMeshModel* Resource = Component->GetSkeletalMesh()->GetImportedModel();
+	FSkeletalMeshModel* Resource = Component->GetSkeletalMeshAsset()->GetImportedModel();
 
 	checkf(Resource->LODModels.IsValidIndex(LODIndex), TEXT("Invalid LOD Index"));
 
 	TArray<FName> MaterialSlotNames = Component->GetMaterialSlotNames();
 
-	const FSkeletalMeshLODInfo* LODInfoPtr = Component->GetSkeletalMesh()->GetLODInfo(LODIndex);
+	const FSkeletalMeshLODInfo* LODInfoPtr = Component->GetSkeletalMeshAsset()->GetLODInfo(LODIndex);
 	check(LODInfoPtr);
 	const FSkeletalMeshLODModel& Model = Resource->LODModels[LODIndex];
 	for (int32 SectionIndex = 0; SectionIndex < Model.Sections.Num(); ++SectionIndex)
@@ -422,10 +422,10 @@ void FMeshMergeHelpers::ExportStaticMeshLOD(const FStaticMeshLODResources& Stati
 
 void FMeshMergeHelpers::RetrieveMesh(const USkeletalMeshComponent* SkeletalMeshComponent, int32 LODIndex, FMeshDescription& OutMeshDescription, bool bPropagateVertexColours)
 {
-	FSkeletalMeshModel* Resource = SkeletalMeshComponent->GetSkeletalMesh()->GetImportedModel();
+	FSkeletalMeshModel* Resource = SkeletalMeshComponent->GetSkeletalMeshAsset()->GetImportedModel();
 	if (Resource->LODModels.IsValidIndex(LODIndex))
 	{
-		FSkeletalMeshLODInfo& SrcLODInfo = *(SkeletalMeshComponent->GetSkeletalMesh()->GetLODInfo(LODIndex));
+		FSkeletalMeshLODInfo& SrcLODInfo = *(SkeletalMeshComponent->GetSkeletalMeshAsset()->GetLODInfo(LODIndex));
 
 		// Get the CPU skinned verts for this LOD
 		TArray<FFinalSkinVertex> FinalVertices;
@@ -483,10 +483,10 @@ void FMeshMergeHelpers::RetrieveMesh(const USkeletalMeshComponent* SkeletalMeshC
 			// use the remapping of material indices if there is a valid value
 			if (SrcLODInfo.LODMaterialMap.IsValidIndex(SectionIndex) && SrcLODInfo.LODMaterialMap[SectionIndex] != INDEX_NONE)
 			{
-				MaterialIndex = FMath::Clamp<int32>(SrcLODInfo.LODMaterialMap[SectionIndex], 0, SkeletalMeshComponent->GetSkeletalMesh()->GetMaterials().Num() - 1);
+				MaterialIndex = FMath::Clamp<int32>(SrcLODInfo.LODMaterialMap[SectionIndex], 0, SkeletalMeshComponent->GetSkeletalMeshAsset()->GetMaterials().Num() - 1);
 			}
 
-			FName ImportedMaterialSlotName = SkeletalMeshComponent->GetSkeletalMesh()->GetMaterials()[MaterialIndex].ImportedMaterialSlotName;
+			FName ImportedMaterialSlotName = SkeletalMeshComponent->GetSkeletalMeshAsset()->GetMaterials()[MaterialIndex].ImportedMaterialSlotName;
 			const FPolygonGroupID SectionPolygonGroupID(SectionIndex);
 			if (!OutMeshDescription.IsPolygonGroupValid(SectionPolygonGroupID))
 			{
diff --git a/Engine/Source/Developer/MeshMergeUtilities/Private/MeshMergeUtilities.cpp b/Engine/Source/Developer/MeshMergeUtilities/Private/MeshMergeUtilities.cpp
--- a/Engine/Source/Developer/MeshMergeUtilities/Private/MeshMergeUtilities.cpp
+++ b/Engine/Source/Developer/MeshMergeUtilities/Private/MeshMergeUtilities.cpp
@@ -441,7 +441,7 @@ void FMeshMergeUtilities::BakeMaterialsForComponent(USkeletalMeshComponent* Skel
 	UMaterialMergeOptions* MergeOptions = GetMutableDefault<UMaterialMergeOptions>();
 	TArray<TWeakObjectPtr<UObject>> Objects{ MergeOptions, AssetOptions, MaterialOptions };
 
-	const int32 NumLODs = SkeletalMeshComponent->GetSkeletalMesh()->GetLODNum();
+	const int32 NumLODs = SkeletalMeshComponent->GetSkeletalMeshAsset()->GetLODNum();
 	IMaterialBakingModule& Module = FModuleManager::Get().LoadModuleChecked<IMaterialBakingModule>("MaterialBaking");
 	if (!Module.SetupMaterialBakeSettings(Objects, NumLODs))
 	{
@@ -449,7 +449,7 @@ void FMeshMergeUtilities::BakeMaterialsForComponent(USkeletalMeshComponent* Skel
 	}
 
 	// Bake out materials for skeletal mesh
-	SkeletalMeshComponent->GetSkeletalMesh()->Modify();
+	SkeletalMeshComponent->GetSkeletalMeshAsset()->Modify();
 	FSkeletalMeshComponentAdapter Adapter(SkeletalMeshComponent);
 	BakeMaterialsForComponent(Objects, &Adapter);
 	SkeletalMeshComponent->MarkRenderStateDirty();
diff --git a/Engine/Source/Developer/MeshMergeUtilities/Private/SkeletalMeshAdapter.cpp b/Engine/Source/Developer/MeshMergeUtilities/Private/SkeletalMeshAdapter.cpp
--- a/Engine/Source/Developer/MeshMergeUtilities/Private/SkeletalMeshAdapter.cpp
+++ b/Engine/Source/Developer/MeshMergeUtilities/Private/SkeletalMeshAdapter.cpp
@@ -11,7 +11,7 @@
 #include "UObject/Package.h"
 
 FSkeletalMeshComponentAdapter::FSkeletalMeshComponentAdapter(USkeletalMeshComponent* InSkeletalMeshComponent)
-	: SkeletalMeshComponent(InSkeletalMeshComponent), SkeletalMesh(InSkeletalMeshComponent->GetSkeletalMesh())
+	: SkeletalMeshComponent(InSkeletalMeshComponent), SkeletalMesh(InSkeletalMeshComponent->GetSkeletalMeshAsset())
 {
 	checkf(SkeletalMesh != nullptr, TEXT("Invalid skeletal mesh in adapter"));
 	NumLODs = SkeletalMesh->GetLODNum();
diff --git a/Engine/Source/Developer/PhysicsUtilities/Private/PhysicsAssetUtils.cpp b/Engine/Source/Developer/PhysicsUtilities/Private/PhysicsAssetUtils.cpp
--- a/Engine/Source/Developer/PhysicsUtilities/Private/PhysicsAssetUtils.cpp
+++ b/Engine/Source/Developer/PhysicsUtilities/Private/PhysicsAssetUtils.cpp
@@ -650,19 +650,19 @@ void WeldBodies(UPhysicsAsset* PhysAsset, int32 BaseBodyIndex, int32 AddBodyInde
 	if(BaseBodyIndex == INDEX_NONE || AddBodyIndex == INDEX_NONE)
 		return;
 
-	if (SkelComp == NULL || SkelComp->GetSkeletalMesh() == NULL)
+	if (SkelComp == NULL || SkelComp->GetSkeletalMeshAsset() == NULL)
 	{
 		return;
 	}
 
 	UBodySetup* Body1 = PhysAsset->SkeletalBodySetups[BaseBodyIndex];
-	int32 Bone1Index = SkelComp->GetSkeletalMesh()->GetRefSkeleton().FindBoneIndex(Body1->BoneName);
+	int32 Bone1Index = SkelComp->GetSkeletalMeshAsset()->GetRefSkeleton().FindBoneIndex(Body1->BoneName);
 	check(Bone1Index != INDEX_NONE);
 	FTransform Bone1TM = SkelComp->GetBoneTransform(Bone1Index);
 	Bone1TM.RemoveScaling();
 
 	UBodySetup* Body2 = PhysAsset->SkeletalBodySetups[AddBodyIndex];
-	int32 Bone2Index = SkelComp->GetSkeletalMesh()->GetRefSkeleton().FindBoneIndex(Body2->BoneName);
+	int32 Bone2Index = SkelComp->GetSkeletalMeshAsset()->GetRefSkeleton().FindBoneIndex(Body2->BoneName);
 	check(Bone2Index != INDEX_NONE);
 	FTransform Bone2TM = SkelComp->GetBoneTransform(Bone2Index);
 	Bone2TM.RemoveScaling();
diff --git a/Engine/Source/Developer/SkeletalMeshUtilitiesCommon/Private/LODUtilities.cpp b/Engine/Source/Developer/SkeletalMeshUtilitiesCommon/Private/LODUtilities.cpp
--- a/Engine/Source/Developer/SkeletalMeshUtilitiesCommon/Private/LODUtilities.cpp
+++ b/Engine/Source/Developer/SkeletalMeshUtilitiesCommon/Private/LODUtilities.cpp
@@ -1755,7 +1755,7 @@ void FLODUtilities::RefreshLODChange(const USkeletalMesh* SkeletalMesh)
 	for (FThreadSafeObjectIterator Iter(USkeletalMeshComponent::StaticClass()); Iter; ++Iter)
 	{
 		USkeletalMeshComponent* SkeletalMeshComponent = Cast<USkeletalMeshComponent>(*Iter);
-		if  (SkeletalMeshComponent->GetSkeletalMesh() == SkeletalMesh)
+		if  (SkeletalMeshComponent->GetSkeletalMeshAsset() == SkeletalMesh)
 		{
 			// it needs to recreate IF it already has been created
 			if (SkeletalMeshComponent->IsRegistered())
diff --git a/Engine/Source/Editor/AnimGraph/Private/AnimGraphNode_SkeletalControlBase.cpp b/Engine/Source/Editor/AnimGraph/Private/AnimGraphNode_SkeletalControlBase.cpp
--- a/Engine/Source/Editor/AnimGraph/Private/AnimGraphNode_SkeletalControlBase.cpp
+++ b/Engine/Source/Editor/AnimGraph/Private/AnimGraphNode_SkeletalControlBase.cpp
@@ -81,7 +81,7 @@ void UAnimGraphNode_SkeletalControlBase::CreateOutputPins()
 
 void UAnimGraphNode_SkeletalControlBase::ConvertToComponentSpaceTransform(const USkeletalMeshComponent* SkelComp, const FTransform & InTransform, FTransform & OutCSTransform, int32 BoneIndex, EBoneControlSpace Space) const
 {
-	USkeleton * Skeleton = SkelComp->GetSkeletalMesh()->GetSkeleton();
+	USkeleton * Skeleton = SkelComp->GetSkeletalMeshAsset()->GetSkeleton();
 
 	switch (Space)
 	{
@@ -105,7 +105,7 @@ void UAnimGraphNode_SkeletalControlBase::ConvertToComponentSpaceTransform(const
 			const int32 ParentIndex = Skeleton->GetReferenceSkeleton().GetParentIndex(BoneIndex);
 			if (ParentIndex != INDEX_NONE)																																																																																																																															
 			{
-				const int32 MeshParentIndex = Skeleton->GetMeshBoneIndexFromSkeletonBoneIndex(SkelComp->GetSkeletalMesh(), ParentIndex);
+				const int32 MeshParentIndex = Skeleton->GetMeshBoneIndexFromSkeletonBoneIndex(SkelComp->GetSkeletalMeshAsset(), ParentIndex);
 				if (MeshParentIndex != INDEX_NONE)
 				{
 					const FTransform ParentTM = SkelComp->GetBoneTransform(MeshParentIndex);
@@ -122,7 +122,7 @@ void UAnimGraphNode_SkeletalControlBase::ConvertToComponentSpaceTransform(const
 	case BCS_BoneSpace:
 		if (BoneIndex != INDEX_NONE)
 		{
-			const int32 MeshBoneIndex = Skeleton->GetMeshBoneIndexFromSkeletonBoneIndex(SkelComp->GetSkeletalMesh(), BoneIndex);
+			const int32 MeshBoneIndex = Skeleton->GetMeshBoneIndexFromSkeletonBoneIndex(SkelComp->GetSkeletalMeshAsset(), BoneIndex);
 			if (MeshBoneIndex != INDEX_NONE)
 			{
 				const FTransform BoneTM = SkelComp->GetBoneTransform(MeshBoneIndex);
@@ -136,9 +136,9 @@ void UAnimGraphNode_SkeletalControlBase::ConvertToComponentSpaceTransform(const
 		break;
 
 	default:
-		if (SkelComp->GetSkeletalMesh())
+		if (SkelComp->GetSkeletalMeshAsset())
 		{
-			UE_LOG(LogAnimation, Warning, TEXT("ConvertToComponentSpaceTransform: Unknown BoneSpace %d  for Mesh: %s"), (uint8)Space, *SkelComp->GetSkeletalMesh()->GetFName().ToString());
+			UE_LOG(LogAnimation, Warning, TEXT("ConvertToComponentSpaceTransform: Unknown BoneSpace %d  for Mesh: %s"), (uint8)Space, *SkelComp->GetSkeletalMeshAsset()->GetFName().ToString());
 		}
 		else
 		{
@@ -250,7 +250,7 @@ FVector UAnimGraphNode_SkeletalControlBase::ConvertWidgetLocation(const USkeleta
 
 	if (MeshBases.GetPose().IsValid())
 	{
-		USkeleton * Skeleton = SkelComp->GetSkeletalMesh()->GetSkeleton();
+		USkeleton * Skeleton = SkelComp->GetSkeletalMeshAsset()->GetSkeleton();
 		const FMeshPoseBoneIndex MeshBoneIndex(SkelComp->GetBoneIndex(BoneName));
 		const FCompactPoseBoneIndex CompactBoneIndex = MeshBases.GetPose().GetBoneContainer().MakeCompactPoseIndex(MeshBoneIndex);
 
diff --git a/Engine/Source/Editor/AnimGraph/Private/AnimGraphNode_Trail.cpp b/Engine/Source/Editor/AnimGraph/Private/AnimGraphNode_Trail.cpp
--- a/Engine/Source/Editor/AnimGraph/Private/AnimGraphNode_Trail.cpp
+++ b/Engine/Source/Editor/AnimGraph/Private/AnimGraphNode_Trail.cpp
@@ -129,7 +129,7 @@ void UAnimGraphNode_Trail::Draw(FPrimitiveDrawInterface* PDI, USkeletalMeshCompo
 {
 	// initialize bone list
 	// draw bone Angular limit
-	if (Node.bLimitRotation && Node.ChainLength > 1 && PreviewSkelMeshComp->GetSkeletalMesh())
+	if (Node.bLimitRotation && Node.ChainLength > 1 && PreviewSkelMeshComp->GetSkeletalMeshAsset())
 	{
 		TArray<FName> TrailBoneList;
 		TrailBoneList.Reset(Node.ChainLength);
@@ -137,7 +137,7 @@ void UAnimGraphNode_Trail::Draw(FPrimitiveDrawInterface* PDI, USkeletalMeshCompo
 		int32 CurrentIndex = Node.ChainLength - 1;
 		TrailBoneList[CurrentIndex] = Node.TrailBone.BoneName;
 		FName CurrentName = Node.TrailBone.BoneName;
-		const FReferenceSkeleton& RefSkeleton = PreviewSkelMeshComp->GetSkeletalMesh()->GetRefSkeleton();
+		const FReferenceSkeleton& RefSkeleton = PreviewSkelMeshComp->GetSkeletalMeshAsset()->GetRefSkeleton();
 		while (--CurrentIndex >= 0 && CurrentName != NAME_None)
 		{
 			const int32 ParentIndex = RefSkeleton.GetParentIndex(RefSkeleton.FindBoneIndex(CurrentName));
diff --git a/Engine/Source/Editor/AnimGraph/Private/AnimPreviewInstance.cpp b/Engine/Source/Editor/AnimGraph/Private/AnimPreviewInstance.cpp
--- a/Engine/Source/Editor/AnimGraph/Private/AnimPreviewInstance.cpp
+++ b/Engine/Source/Editor/AnimGraph/Private/AnimPreviewInstance.cpp
@@ -175,9 +175,9 @@ bool FAnimPreviewInstanceProxy::Evaluate(FPoseContext& Output)
 		if(bForceRetargetBasePose)
 		{
 			USkeletalMeshComponent* MeshComponent = Output.AnimInstanceProxy->GetSkelMeshComponent();
-			if(MeshComponent && MeshComponent->GetSkeletalMesh())
+			if(MeshComponent && MeshComponent->GetSkeletalMeshAsset())
 			{
-				FAnimationRuntime::FillWithRetargetBaseRefPose(Output.Pose, GetSkelMeshComponent()->GetSkeletalMesh());
+				FAnimationRuntime::FillWithRetargetBaseRefPose(Output.Pose, GetSkelMeshComponent()->GetSkeletalMeshAsset());
 			}
 			else
 			{
@@ -369,7 +369,7 @@ void FAnimPreviewInstanceProxy::SetKeyImplementation(const FCompactPose& PreCont
 	UDebugSkelMeshComponent* Component = Cast<UDebugSkelMeshComponent> (GetSkelMeshComponent());
 
 	USkeleton* PreviewSkeleton = (CurrentSequence) ? CurrentSequence->GetSkeleton() : nullptr;
-	if(CurrentSequence && PreviewSkeleton && Component && Component->GetSkeletalMesh())
+	if(CurrentSequence && PreviewSkeleton && Component && Component->GetSkeletalMeshAsset())
 	{
 		FScopedTransaction ScopedTransaction(LOCTEXT("SetKey", "Set Key"));
 		CurrentSequence->Modify(true);
diff --git a/Engine/Source/Editor/AnimGraph/Private/EditModes/AnimNodeEditMode.cpp b/Engine/Source/Editor/AnimGraph/Private/EditModes/AnimNodeEditMode.cpp
--- a/Engine/Source/Editor/AnimGraph/Private/EditModes/AnimNodeEditMode.cpp
+++ b/Engine/Source/Editor/AnimGraph/Private/EditModes/AnimNodeEditMode.cpp
@@ -471,7 +471,7 @@ FAnimNode_Base* FAnimNodeEditMode::GetActiveWidgetRuntimeAnimNode() const
 
 void FAnimNodeEditMode::ConvertToComponentSpaceTransform(const USkeletalMeshComponent* SkelComp, const FTransform & InTransform, FTransform & OutCSTransform, int32 BoneIndex, EBoneControlSpace Space)
 {
-	USkeletalMesh* SkelMesh = SkelComp->GetSkeletalMesh();
+	USkeletalMesh* SkelMesh = SkelComp->GetSkeletalMeshAsset();
 	USkeleton* Skeleton = SkelMesh->GetSkeleton();
 
 	switch (Space)
@@ -542,7 +542,7 @@ void FAnimNodeEditMode::ConvertToComponentSpaceTransform(const USkeletalMeshComp
 
 void FAnimNodeEditMode::ConvertToBoneSpaceTransform(const USkeletalMeshComponent* SkelComp, const FTransform & InCSTransform, FTransform & OutBSTransform, int32 BoneIndex, EBoneControlSpace Space)
 {
-	USkeletalMesh* SkelMesh = SkelComp->GetSkeletalMesh();
+	USkeletalMesh* SkelMesh = SkelComp->GetSkeletalMeshAsset();
 	USkeleton* Skeleton = SkelMesh->GetSkeleton();
 
 	switch(Space)
@@ -807,7 +807,7 @@ FVector FAnimNodeEditMode::ConvertWidgetLocation(const USkeletalMeshComponent* S
 	{
 		if (InMeshBases.GetPose().IsValid())
 		{
-			USkeleton* Skeleton = InSkelComp->GetSkeletalMesh()->GetSkeleton();
+			USkeleton* Skeleton = InSkelComp->GetSkeletalMeshAsset()->GetSkeleton();
 			const int32 MeshBoneIndex = InSkelComp->GetBoneIndex(InBoneName);
 			if (MeshBoneIndex != INDEX_NONE)
 			{
diff --git a/Engine/Source/Editor/AnimGraph/Private/EditModes/ModifyBoneEditMode.cpp b/Engine/Source/Editor/AnimGraph/Private/EditModes/ModifyBoneEditMode.cpp
--- a/Engine/Source/Editor/AnimGraph/Private/EditModes/ModifyBoneEditMode.cpp
+++ b/Engine/Source/Editor/AnimGraph/Private/EditModes/ModifyBoneEditMode.cpp
@@ -59,7 +59,7 @@ ECoordSystem FModifyBoneEditMode::GetWidgetCoordinateSystem() const
 FVector FModifyBoneEditMode::GetWidgetLocation() const
 {
 	USkeletalMeshComponent* SkelComp = GetAnimPreviewScene().GetPreviewMeshComponent();
-	USkeleton* Skeleton = SkelComp->GetSkeletalMesh()->GetSkeleton();
+	USkeleton* Skeleton = SkelComp->GetSkeletalMeshAsset()->GetSkeleton();
 	FVector WidgetLoc = FVector::ZeroVector;
 
 	// if the current widget mode is translate, then shows the widget according to translation space
diff --git a/Engine/Source/Editor/AnimGraph/Private/EditModes/ObserveBoneEditMode.cpp b/Engine/Source/Editor/AnimGraph/Private/EditModes/ObserveBoneEditMode.cpp
--- a/Engine/Source/Editor/AnimGraph/Private/EditModes/ObserveBoneEditMode.cpp
+++ b/Engine/Source/Editor/AnimGraph/Private/EditModes/ObserveBoneEditMode.cpp
@@ -40,7 +40,7 @@ ECoordSystem FObserveBoneEditMode::GetWidgetCoordinateSystem() const
 FVector FObserveBoneEditMode::GetWidgetLocation() const
 {
 	USkeletalMeshComponent* SkelComp = GetAnimPreviewScene().GetPreviewMeshComponent();
-	USkeleton* Skeleton = SkelComp->GetSkeletalMesh()->GetSkeleton();
+	USkeleton* Skeleton = SkelComp->GetSkeletalMeshAsset()->GetSkeleton();
 	FVector WidgetLoc = FVector::ZeroVector;
 
 	int32 MeshBoneIndex = SkelComp->GetBoneIndex(GraphNode->Node.BoneToObserve.BoneName);
diff --git a/Engine/Source/Editor/AnimGraph/Private/EditModes/TwoBoneIKEditMode.cpp b/Engine/Source/Editor/AnimGraph/Private/EditModes/TwoBoneIKEditMode.cpp
--- a/Engine/Source/Editor/AnimGraph/Private/EditModes/TwoBoneIKEditMode.cpp
+++ b/Engine/Source/Editor/AnimGraph/Private/EditModes/TwoBoneIKEditMode.cpp
@@ -56,9 +56,9 @@ void FTwoBoneIKEditMode::ExitMode()
 void FTwoBoneIKEditMode::Render(const FSceneView* View, FViewport* Viewport, FPrimitiveDrawInterface* PDI)
 {
 	USkeletalMeshComponent* SkelMeshComp = GetAnimPreviewScene().GetPreviewMeshComponent();
-	if (SkelMeshComp && SkelMeshComp->GetSkeletalMesh() && SkelMeshComp->GetSkeletalMesh()->GetSkeleton())
+	if (SkelMeshComp && SkelMeshComp->GetSkeletalMeshAsset() && SkelMeshComp->GetSkeletalMeshAsset()->GetSkeleton())
 	{
-		USkeleton* Skeleton = SkelMeshComp->GetSkeletalMesh()->GetSkeleton();
+		USkeleton* Skeleton = SkelMeshComp->GetSkeletalMeshAsset()->GetSkeleton();
 
 		PDI->SetHitProxy(new HTwoBoneIKProxy(BSM_EndEffector));
 		DrawTargetLocation(PDI, BSM_EndEffector, FColor(255, 128, 128), FColor(180, 128, 128));
diff --git a/Engine/Source/Editor/AnimationEditor/Private/AnimationEditor.cpp b/Engine/Source/Editor/AnimationEditor/Private/AnimationEditor.cpp
--- a/Engine/Source/Editor/AnimationEditor/Private/AnimationEditor.cpp
+++ b/Engine/Source/Editor/AnimationEditor/Private/AnimationEditor.cpp
@@ -662,7 +662,7 @@ bool FAnimationEditor::ExportToFBX(const TArray<UObject*> AssetsToExport, bool b
 		FPersonaModule& PersonaModule = FModuleManager::GetModuleChecked<FPersonaModule>("Persona");
 		
 		
-		AnimSequenceExportResult = PersonaModule.ExportToFBX(AnimSequences, GetPersonaToolkit()->GetPreviewScene()->GetPreviewMeshComponent()->GetSkeletalMesh());
+		AnimSequenceExportResult = PersonaModule.ExportToFBX(AnimSequences, GetPersonaToolkit()->GetPreviewScene()->GetPreviewMeshComponent()->GetSkeletalMeshAsset());
 	}
 	return AnimSequenceExportResult;
 }
diff --git a/Engine/Source/Editor/ClothPainter/Private/ClothMeshAdapter.cpp b/Engine/Source/Editor/ClothPainter/Private/ClothMeshAdapter.cpp
--- a/Engine/Source/Editor/ClothPainter/Private/ClothMeshAdapter.cpp
+++ b/Engine/Source/Editor/ClothPainter/Private/ClothMeshAdapter.cpp
@@ -15,7 +15,7 @@ TSharedPtr<IMeshPaintGeometryAdapter> FClothMeshPaintAdapterFactory::Construct(U
 {
 	if (USkeletalMeshComponent* Component = Cast<USkeletalMeshComponent>(InComponent))
 	{
-		if (Component->GetSkeletalMesh() != nullptr)
+		if (Component->GetSkeletalMeshAsset() != nullptr)
 		{
 			TSharedRef<FClothMeshPaintAdapter> Result = MakeShareable(new FClothMeshPaintAdapter());
 			if (Result->Construct(InComponent, InPaintingMeshLODIndex))
@@ -33,9 +33,9 @@ bool FClothMeshPaintAdapter::Construct(UMeshComponent* InComponent, int32 InPain
 	SkeletalMeshComponent = Cast<USkeletalMeshComponent>(InComponent);
 	if (SkeletalMeshComponent != nullptr)
 	{
-		if (SkeletalMeshComponent->GetSkeletalMesh() != nullptr)
+		if (SkeletalMeshComponent->GetSkeletalMeshAsset() != nullptr)
 		{
-			ReferencedSkeletalMesh = SkeletalMeshComponent->GetSkeletalMesh();
+			ReferencedSkeletalMesh = SkeletalMeshComponent->GetSkeletalMeshAsset();
 			PaintingClothLODIndex = InPaintingClothLODIndex;
 			PaintingClothMaskIndex = INDEX_NONE;
 
@@ -49,7 +49,7 @@ bool FClothMeshPaintAdapter::Construct(UMeshComponent* InComponent, int32 InPain
 
 bool FClothMeshPaintAdapter::Initialize()
 {
-	check(ReferencedSkeletalMesh == SkeletalMeshComponent->GetSkeletalMesh());
+	check(ReferencedSkeletalMesh == SkeletalMeshComponent->GetSkeletalMeshAsset());
 
 	bool bHaveAsset = false;
 	bool bBaseInit = false;
diff --git a/Engine/Source/Editor/ClothPainter/Private/ClothPainter.cpp b/Engine/Source/Editor/ClothPainter/Private/ClothPainter.cpp
--- a/Engine/Source/Editor/ClothPainter/Private/ClothPainter.cpp
+++ b/Engine/Source/Editor/ClothPainter/Private/ClothPainter.cpp
@@ -73,7 +73,7 @@ bool FClothPainter::PaintInternal(const FVector& InCameraOrigin, const TArrayVie
 
 	if(SkeletalMeshComponent->SelectedClothingGuidForPainting.IsValid() && !bShouldSimulate)
 	{
-		USkeletalMesh* SkelMesh = SkeletalMeshComponent->GetSkeletalMesh();
+		USkeletalMesh* SkelMesh = SkeletalMeshComponent->GetSkeletalMeshAsset();
 
 		for (const TPair<FVector, FVector>& Ray : Rays)
 		{
@@ -160,7 +160,7 @@ USkeletalMesh* FClothPainter::GetSkeletalMesh() const
 {
 	if(SkeletalMeshComponent)
 	{
-		return SkeletalMeshComponent->GetSkeletalMesh();
+		return SkeletalMeshComponent->GetSkeletalMeshAsset();
 	}
 
 	return nullptr;
@@ -175,7 +175,7 @@ void FClothPainter::RefreshClothingAssets()
 
 	PaintSettings->ClothingAssets.Reset();
 
-	if(USkeletalMesh* Mesh = SkeletalMeshComponent->GetSkeletalMesh())
+	if(USkeletalMesh* Mesh = SkeletalMeshComponent->GetSkeletalMeshAsset())
 	{
 		for(UClothingAssetBase* BaseClothingAsset : Mesh->GetMeshClothingAssets())
 		{
diff --git a/Engine/Source/Editor/ClothPainter/Private/ClothingPaintEditMode.cpp b/Engine/Source/Editor/ClothPainter/Private/ClothingPaintEditMode.cpp
--- a/Engine/Source/Editor/ClothPainter/Private/ClothingPaintEditMode.cpp
+++ b/Engine/Source/Editor/ClothPainter/Private/ClothingPaintEditMode.cpp
@@ -88,7 +88,7 @@ void FClothingPaintEditMode::Exit()
 		{
 			MeshComponent->bDisableClothSimulation = false;
 
-			if(USkeletalMesh* SkelMesh = MeshComponent->GetSkeletalMesh())
+			if(USkeletalMesh* SkelMesh = MeshComponent->GetSkeletalMeshAsset())
 			{
 				for(UClothingAssetBase* AssetBase : SkelMesh->GetMeshClothingAssets())
 				{
@@ -113,7 +113,7 @@ void FClothingPaintEditMode::Exit()
 			for(TObjectIterator<USkeletalMeshComponent> It; It; ++It)
 			{
 				USkeletalMeshComponent* Component = *It;
-				if(Component && !Component->IsTemplate() && Component->GetSkeletalMesh() == SkelMesh)
+				if(Component && !Component->IsTemplate() && Component->GetSkeletalMeshAsset() == SkelMesh)
 				{
 					Component->ReregisterComponent();
 				}
diff --git a/Engine/Source/Editor/DetailCustomizations/Private/AnimMontageSegmentDetails.cpp b/Engine/Source/Editor/DetailCustomizations/Private/AnimMontageSegmentDetails.cpp
--- a/Engine/Source/Editor/DetailCustomizations/Private/AnimMontageSegmentDetails.cpp
+++ b/Engine/Source/Editor/DetailCustomizations/Private/AnimMontageSegmentDetails.cpp
@@ -517,7 +517,7 @@ void SAnimationSegmentViewport::InitSkeleton()
 		if (Skeleton && PreviewMesh)
 		{
 			UAnimSingleNodeInstance* Preview = PreviewComponent->PreviewInstance;
-			if((Preview == nullptr || Preview->GetCurrentAsset() != AnimSequenceBase) || (PreviewComponent->GetSkeletalMesh() != PreviewMesh))
+			if((Preview == nullptr || Preview->GetCurrentAsset() != AnimSequenceBase) || (PreviewComponent->GetSkeletalMeshAsset() != PreviewMesh))
 			{
 				const float PlayRate = PlayRateAttribute.Get().Get(1.f);
 
@@ -565,7 +565,7 @@ void SAnimationSegmentViewport::Tick( const FGeometry& AllottedGeometry, const d
 		{
 			Description->SetText(FText::Format( LOCTEXT("Previewing", "Previewing {0}"), FText::FromString(Component->AnimClass->GetName()) ));
 		}
-		else if (Component->GetSkeletalMesh() == NULL)
+		else if (Component->GetSkeletalMeshAsset() == NULL)
 		{
 			Description->SetText(FText::Format( LOCTEXT("NoMeshFound", "No skeletal mesh found for skeleton '{0}'"), FText::FromString(TargetSkeletonName) ));
 		}
diff --git a/Engine/Source/Editor/DetailCustomizations/Private/AnimSequenceDetails.cpp b/Engine/Source/Editor/DetailCustomizations/Private/AnimSequenceDetails.cpp
--- a/Engine/Source/Editor/DetailCustomizations/Private/AnimSequenceDetails.cpp
+++ b/Engine/Source/Editor/DetailCustomizations/Private/AnimSequenceDetails.cpp
@@ -509,7 +509,7 @@ void SAnimationRefPoseViewport::InitSkeleton()
 		}();
 		
 		const bool bInvalidPreviewInstance = PreviewAnimInstance == nullptr || PreviewAnimInstance->GetCurrentAsset() != PreviewAnimationSequence;
-		const bool bPreviewMeshMismatch = PreviewComponent->GetSkeletalMesh() != PreviewSkeletalMesh;
+		const bool bPreviewMeshMismatch = PreviewComponent->GetSkeletalMeshAsset() != PreviewSkeletalMesh;
 		if(bInvalidPreviewInstance || bPreviewMeshMismatch)
 		{
 			PreviewComponent->SetSkeletalMesh(PreviewSkeletalMesh);
@@ -572,7 +572,7 @@ void SAnimationRefPoseViewport::Tick( const FGeometry& AllottedGeometry, const d
 		{
 			Description->SetText( FText::Format( LOCTEXT( "IncorrectSkeleton", "The preview asset doesn't work for the skeleton '{0}'" ), FText::FromString( TargetSkeletonName ) ) );
 		}
-		else if ( Component->GetSkeletalMesh() == NULL )
+		else if ( Component->GetSkeletalMeshAsset() == NULL )
 		{
 			Description->SetText( FText::Format( LOCTEXT( "NoMeshFound", "No skeletal mesh found for skeleton '{0}'" ), FText::FromString( TargetSkeletonName ) ) );
 		}
diff --git a/Engine/Source/Editor/DetailCustomizations/Private/SkeletalMeshComponentDetails.cpp b/Engine/Source/Editor/DetailCustomizations/Private/SkeletalMeshComponentDetails.cpp
--- a/Engine/Source/Editor/DetailCustomizations/Private/SkeletalMeshComponentDetails.cpp
+++ b/Engine/Source/Editor/DetailCustomizations/Private/SkeletalMeshComponentDetails.cpp
@@ -397,7 +397,7 @@ USkeleton* FSkeletalMeshComponentDetails::GetValidSkeletonFromRegisteredMeshes()
 	for (auto ObjectIter = SelectedObjects.CreateConstIterator(); ObjectIter; ++ObjectIter)
 	{
 		USkeletalMeshComponent* const Mesh = Cast<USkeletalMeshComponent>(ObjectIter->Get());
-		if ( !Mesh || !Mesh->GetSkeletalMesh())
+		if ( !Mesh || !Mesh->GetSkeletalMeshAsset())
 		{
 			continue;
 		}
@@ -405,14 +405,14 @@ USkeleton* FSkeletalMeshComponentDetails::GetValidSkeletonFromRegisteredMeshes()
 		// If we've not come across a valid skeleton yet, store this one.
 		if (!ResultSkeleton)
 		{
-			ResultSkeleton = Mesh->GetSkeletalMesh()->GetSkeleton();
+			ResultSkeleton = Mesh->GetSkeletalMeshAsset()->GetSkeleton();
 			continue;
 		}
 
 		// We've encountered a valid skeleton before.
 		// If this skeleton is not the same one, that means there are multiple
 		// skeletons selected, so we don't want to take any action.
-		if (Mesh->GetSkeletalMesh()->GetSkeleton() != ResultSkeleton)
+		if (Mesh->GetSkeletalMeshAsset()->GetSkeleton() != ResultSkeleton)
 		{
 			return NULL;
 		}
diff --git a/Engine/Source/Editor/DetailCustomizations/Private/SkinnedMeshComponentDetails.cpp b/Engine/Source/Editor/DetailCustomizations/Private/SkinnedMeshComponentDetails.cpp
--- a/Engine/Source/Editor/DetailCustomizations/Private/SkinnedMeshComponentDetails.cpp
+++ b/Engine/Source/Editor/DetailCustomizations/Private/SkinnedMeshComponentDetails.cpp
@@ -183,7 +183,7 @@ void FSkinnedMeshComponentDetails::PopulateSkinWeightProfileNames()
 	// Retrieve all possible skin weight profiles from the component
 	if (USkeletalMeshComponent* Component = Cast<USkeletalMeshComponent>(WeakSkinnedMeshComponent.Get()))
 	{
-		if (USkeletalMesh* Mesh = Component->GetSkeletalMesh())
+		if (USkeletalMesh* Mesh = Component->GetSkeletalMeshAsset())
 		{			
 			for (const FSkinWeightProfileInfo& Profile : Mesh->GetSkinWeightProfiles())
 			{
diff --git a/Engine/Source/Editor/LevelEditor/Private/LevelEditorActions.cpp b/Engine/Source/Editor/LevelEditor/Private/LevelEditorActions.cpp
--- a/Engine/Source/Editor/LevelEditor/Private/LevelEditorActions.cpp
+++ b/Engine/Source/Editor/LevelEditor/Private/LevelEditorActions.cpp
@@ -2061,7 +2061,7 @@ bool FLevelEditorActionCallbacks::SaveAnimationFromSkeletalMeshComponent(AActor
 		bool bSimulating = false;
 		for (auto & Comp : SimSkelComponents)
 		{
-			bSimulating |= (Comp->GetSkeletalMesh() && Comp->GetSkeletalMesh()->GetSkeleton() && Comp->IsSimulatingPhysics());
+			bSimulating |= (Comp->GetSkeletalMeshAsset() && Comp->GetSkeletalMeshAsset()->GetSkeleton() && Comp->IsSimulatingPhysics());
 		}
 
 		// if any of them are legitimately simulating
@@ -2075,7 +2075,7 @@ bool FLevelEditorActionCallbacks::SaveAnimationFromSkeletalMeshComponent(AActor
 			{
 				for (auto & Comp : SimSkelComponents)
 				{
-					if (Comp->GetSkeletalMesh() && Comp->GetSkeletalMesh()->GetSkeleton() && Comp->IsSimulatingPhysics())
+					if (Comp->GetSkeletalMeshAsset() && Comp->GetSkeletalMeshAsset()->GetSkeleton() && Comp->IsSimulatingPhysics())
 					{
 						// now record to animation
 						class UAnimSequence* Sequence = LevelEditorModule.OnCaptureSingleFrameAnimSequence().IsBound() ? LevelEditorModule.OnCaptureSingleFrameAnimSequence().Execute(Comp) : nullptr;
@@ -2164,7 +2164,7 @@ void FLevelEditorActionCallbacks::OnKeepSimulationChanges()
 				// need to reinitialize animation
 				for (auto MeshComp : ComponentsToReinitialize)
 				{
-					if(MeshComp->GetSkeletalMesh())
+					if(MeshComp->GetSkeletalMeshAsset())
 					{
 						MeshComp->InitAnim(true);
 					}
diff --git a/Engine/Source/Editor/MeshPaint/Private/MeshPaintHelpers.cpp b/Engine/Source/Editor/MeshPaint/Private/MeshPaintHelpers.cpp
--- a/Engine/Source/Editor/MeshPaint/Private/MeshPaintHelpers.cpp
+++ b/Engine/Source/Editor/MeshPaint/Private/MeshPaintHelpers.cpp
@@ -833,7 +833,7 @@ void MeshPaintHelpers::FillStaticMeshVertexColors(UStaticMeshComponent* MeshComp
 void MeshPaintHelpers::FillSkeletalMeshVertexColors(USkeletalMeshComponent* MeshComponent, int32 LODIndex, const FColor FillColor, const FColor MaskColor)
 {
 	TUniquePtr< FSkinnedMeshComponentRecreateRenderStateContext > RecreateRenderStateContext;
-	USkeletalMesh* Mesh = MeshComponent->GetSkeletalMesh();
+	USkeletalMesh* Mesh = MeshComponent->GetSkeletalMeshAsset();
 	if (Mesh)
 	{
 		// Dirty the mesh
@@ -1003,10 +1003,10 @@ void MeshPaintHelpers::ImportVertexColorsFromTexture(UMeshComponent* MeshCompone
 				{
 					USkeletalMeshComponent* SkeletalMeshComponent = Cast<USkeletalMeshComponent>(MeshComponent);
 
-					if (SkeletalMeshComponent->GetSkeletalMesh())
+					if (SkeletalMeshComponent->GetSkeletalMeshAsset())
 					{
 						// Import colors to skeletal mesh
-						ImportVertexColorsToSkeletalMesh(SkeletalMeshComponent->GetSkeletalMesh(), Options, ColorTexture);
+						ImportVertexColorsToSkeletalMesh(SkeletalMeshComponent->GetSkeletalMeshAsset(), Options, ColorTexture);
 					}			
 				}
 			}
@@ -1254,7 +1254,7 @@ bool MeshPaintHelpers::TryGetNumberOfLODs(const UMeshComponent* MeshComponent, i
 	}
 	else if (const USkeletalMeshComponent* SkeletalMeshComponent = Cast<USkeletalMeshComponent>(MeshComponent))
 	{
-		const USkeletalMesh* SkeletalMesh = SkeletalMeshComponent->GetSkeletalMesh();
+		const USkeletalMesh* SkeletalMesh = SkeletalMeshComponent->GetSkeletalMeshAsset();
 		if (SkeletalMesh != nullptr)
 		{
 			OutNumLODs = SkeletalMesh->GetLODNum();
@@ -1286,7 +1286,7 @@ int32 MeshPaintHelpers::GetNumberOfUVs(const UMeshComponent* MeshComponent, int3
 	}
 	else if (const USkeletalMeshComponent* SkeletalMeshComponent = Cast<USkeletalMeshComponent>(MeshComponent))
 	{
-		const USkeletalMesh* SkeletalMesh = SkeletalMeshComponent->GetSkeletalMesh();
+		const USkeletalMesh* SkeletalMesh = SkeletalMeshComponent->GetSkeletalMeshAsset();
 		if (SkeletalMesh != nullptr && SkeletalMesh->GetResourceForRendering() && SkeletalMesh->GetResourceForRendering()->LODRenderData.IsValidIndex(LODIndex))
 		{
 			NumUVs = SkeletalMesh->GetResourceForRendering()->LODRenderData[LODIndex].GetNumTexCoords();
@@ -1324,7 +1324,7 @@ bool MeshPaintHelpers::DoesMeshComponentContainPerLODColors(const UMeshComponent
 	}
 	else if (const USkeletalMeshComponent* SkeletalMeshComponent = Cast<USkeletalMeshComponent>(MeshComponent))
 	{
-		USkeletalMesh* SkeletalMesh = SkeletalMeshComponent->GetSkeletalMesh();
+		USkeletalMesh* SkeletalMesh = SkeletalMeshComponent->GetSkeletalMeshAsset();
 		if (SkeletalMesh)
 		{
 			const TArray<FSkeletalMeshLODInfo>& LODInfo = SkeletalMesh->GetLODInfoArray();
@@ -1694,7 +1694,7 @@ typedef TOctree2<FPaintedMeshVertex, FVertexColorPropogationOctreeSemantics> TVe
 void MeshPaintHelpers::ApplyVertexColorsToAllLODs(IMeshPaintGeometryAdapter& GeometryInfo, USkeletalMeshComponent* SkeletalMeshComponent)
 {
 	checkf(SkeletalMeshComponent != nullptr, TEXT("Invalid Skeletal Mesh Component"));
-	USkeletalMesh* Mesh = SkeletalMeshComponent->GetSkeletalMesh();
+	USkeletalMesh* Mesh = SkeletalMeshComponent->GetSkeletalMeshAsset();
 	if (Mesh)
 	{
 		FSkeletalMeshRenderData* Resource = Mesh->GetResourceForRendering();
diff --git a/Engine/Source/Editor/MeshPaint/Private/MeshPaintSkeletalMeshAdapter.cpp b/Engine/Source/Editor/MeshPaint/Private/MeshPaintSkeletalMeshAdapter.cpp
--- a/Engine/Source/Editor/MeshPaint/Private/MeshPaintSkeletalMeshAdapter.cpp
+++ b/Engine/Source/Editor/MeshPaint/Private/MeshPaintSkeletalMeshAdapter.cpp
@@ -108,9 +108,9 @@ bool FMeshPaintGeometryAdapterForSkeletalMeshes::Construct(UMeshComponent* InCom
 	{
 		SkeletalMeshChangedHandle = SkeletalMeshComponent->RegisterOnSkeletalMeshPropertyChanged(USkeletalMeshComponent::FOnSkeletalMeshPropertyChanged::CreateRaw(this, &FMeshPaintGeometryAdapterForSkeletalMeshes::OnSkeletalMeshChanged));
 
-		if (SkeletalMeshComponent->GetSkeletalMesh() != nullptr)
+		if (SkeletalMeshComponent->GetSkeletalMeshAsset() != nullptr)
 		{
-			ReferencedSkeletalMesh = SkeletalMeshComponent->GetSkeletalMesh();
+			ReferencedSkeletalMesh = SkeletalMeshComponent->GetSkeletalMeshAsset();
 			MeshLODIndex = InMeshLODIndex;
 			const bool bSuccess = Initialize();
 			return bSuccess;
@@ -131,8 +131,8 @@ FMeshPaintGeometryAdapterForSkeletalMeshes::~FMeshPaintGeometryAdapterForSkeleta
 void FMeshPaintGeometryAdapterForSkeletalMeshes::OnSkeletalMeshChanged()
 {
 	OnRemoved();
-	ReferencedSkeletalMesh = SkeletalMeshComponent->GetSkeletalMesh();
-	if (SkeletalMeshComponent->GetSkeletalMesh() != nullptr)
+	ReferencedSkeletalMesh = SkeletalMeshComponent->GetSkeletalMeshAsset();
+	if (SkeletalMeshComponent->GetSkeletalMeshAsset() != nullptr)
 	{	
 		Initialize();
 		OnAdded();
@@ -149,7 +149,7 @@ void FMeshPaintGeometryAdapterForSkeletalMeshes::OnPostMeshCached(USkeletalMesh*
 
 bool FMeshPaintGeometryAdapterForSkeletalMeshes::Initialize()
 {
-	check(ReferencedSkeletalMesh == SkeletalMeshComponent->GetSkeletalMesh());
+	check(ReferencedSkeletalMesh == SkeletalMeshComponent->GetSkeletalMeshAsset());
 
 	bool bInitializationResult = false;
 
@@ -212,7 +212,7 @@ void FMeshPaintGeometryAdapterForSkeletalMeshes::OnAdded()
 {
 	checkf(SkeletalMeshComponent, TEXT("Invalid SkeletalMesh Component"));
 	checkf(ReferencedSkeletalMesh, TEXT("Invalid reference to Skeletal Mesh"));
-	checkf(ReferencedSkeletalMesh == SkeletalMeshComponent->GetSkeletalMesh(), TEXT("Referenced Skeletal Mesh does not match one in Component"));
+	checkf(ReferencedSkeletalMesh == SkeletalMeshComponent->GetSkeletalMeshAsset(), TEXT("Referenced Skeletal Mesh does not match one in Component"));
 
 	FSkeletalMeshReferencers& SkeletalMeshReferencers = MeshToComponentMap.FindOrAdd(ReferencedSkeletalMesh);
 
@@ -515,7 +515,7 @@ TSharedPtr<IMeshPaintGeometryAdapter> FMeshPaintGeometryAdapterForSkeletalMeshes
 {
 	if (USkeletalMeshComponent* SkeletalMeshComponent = Cast<USkeletalMeshComponent>(InComponent))
 	{
-		if (SkeletalMeshComponent->GetSkeletalMesh() != nullptr)
+		if (SkeletalMeshComponent->GetSkeletalMeshAsset() != nullptr)
 		{
 			TSharedRef<FMeshPaintGeometryAdapterForSkeletalMeshes> Result = MakeShareable(new FMeshPaintGeometryAdapterForSkeletalMeshes());
 			if (Result->Construct(InComponent, InMeshLODIndex))
diff --git a/Engine/Source/Editor/MeshPaint/Public/MeshPaintSkeletalMeshAdapter.h b/Engine/Source/Editor/MeshPaint/Public/MeshPaintSkeletalMeshAdapter.h
--- a/Engine/Source/Editor/MeshPaint/Public/MeshPaintSkeletalMeshAdapter.h
+++ b/Engine/Source/Editor/MeshPaint/Public/MeshPaintSkeletalMeshAdapter.h
@@ -33,7 +33,7 @@ public:
 	virtual bool Initialize() override;
 	virtual void OnAdded() override;
 	virtual void OnRemoved() override;
-	virtual bool IsValid() const override { return SkeletalMeshComponent && ReferencedSkeletalMesh && SkeletalMeshComponent->GetSkeletalMesh() == ReferencedSkeletalMesh; }
+	virtual bool IsValid() const override { return SkeletalMeshComponent && ReferencedSkeletalMesh && SkeletalMeshComponent->GetSkeletalMeshAsset() == ReferencedSkeletalMesh; }
 	virtual bool SupportsTexturePaint() const override { return true; }
 	virtual bool SupportsVertexPaint() const override { return SkeletalMeshComponent != nullptr; }
 	virtual bool LineTraceComponent(struct FHitResult& OutHit, const FVector Start, const FVector End, const struct FCollisionQueryParams& Params) const override;
diff --git a/Engine/Source/Editor/MovieSceneTools/Private/EditModes/SkeletalAnimationTrackEditMode.cpp b/Engine/Source/Editor/MovieSceneTools/Private/EditModes/SkeletalAnimationTrackEditMode.cpp
--- a/Engine/Source/Editor/MovieSceneTools/Private/EditModes/SkeletalAnimationTrackEditMode.cpp
+++ b/Engine/Source/Editor/MovieSceneTools/Private/EditModes/SkeletalAnimationTrackEditMode.cpp
@@ -179,11 +179,11 @@ static void DrawBonesFromCompactPose(const FCompactPose& Pose, USkeletalMeshComp
 	{
 
 		FSkeletalMeshLODRenderData* LODData = nullptr;
-		if (MeshComponent->GetSkeletalMesh() && MeshComponent->GetSkeletalMesh()->GetResourceForRendering()->LODRenderData.Num())
+		if (MeshComponent->GetSkeletalMeshAsset() && MeshComponent->GetSkeletalMeshAsset()->GetResourceForRendering()->LODRenderData.Num())
 		{
 			//Get current LOD
-			const int32 LODIndex = FMath::Clamp(MeshComponent->GetPredictedLODLevel(), 0, MeshComponent->GetSkeletalMesh()->GetResourceForRendering()->LODRenderData.Num() - 1);
-			LODData = &(MeshComponent->GetSkeletalMesh()->GetResourceForRendering()->LODRenderData[LODIndex]);
+			const int32 LODIndex = FMath::Clamp(MeshComponent->GetPredictedLODLevel(), 0, MeshComponent->GetSkeletalMeshAsset()->GetResourceForRendering()->LODRenderData.Num() - 1);
+			LODData = &(MeshComponent->GetSkeletalMeshAsset()->GetResourceForRendering()->LODRenderData[LODIndex]);
 		}
 
 		TArray<FTransform> WorldTransforms;
@@ -219,9 +219,9 @@ static void DrawBonesFromCompactPose(const FCompactPose& Pose, USkeletalMeshComp
 			BoneColours[MeshBoneIndex.GetInt()] = DrawColour;
 		}
 
-		if (MeshComponent && MeshComponent->GetSkeletalMesh())
+		if (MeshComponent && MeshComponent->GetSkeletalMeshAsset())
 		{
-			DrawBones(Pose.GetBoneContainer().GetBoneIndicesArray(), MeshComponent->GetSkeletalMesh()->GetRefSkeleton(), WorldTransforms,
+			DrawBones(Pose.GetBoneContainer().GetBoneIndicesArray(), MeshComponent->GetSkeletalMeshAsset()->GetRefSkeleton(), WorldTransforms,
 				LODData, PDI, BoneColours, MeshComponent->Bounds.SphereRadius, 1.0f);
 		}
 
diff --git a/Engine/Source/Editor/MovieSceneTools/Private/MovieSceneToolHelpers.cpp b/Engine/Source/Editor/MovieSceneTools/Private/MovieSceneToolHelpers.cpp
--- a/Engine/Source/Editor/MovieSceneTools/Private/MovieSceneToolHelpers.cpp
+++ b/Engine/Source/Editor/MovieSceneTools/Private/MovieSceneToolHelpers.cpp
@@ -4151,7 +4151,7 @@ USkeletalMeshComponent* MovieSceneToolHelpers::AcquireSkeletalMeshFromObject(UOb
 	}
 	else if (USkeletalMeshComponent* SkeletalMeshComponent = Cast<USkeletalMeshComponent>(BoundObject))
 	{
-		if (SkeletalMeshComponent->GetSkeletalMesh())
+		if (SkeletalMeshComponent->GetSkeletalMeshAsset())
 		{
 			return SkeletalMeshComponent;
 		}
diff --git a/Engine/Source/Editor/MovieSceneTools/Private/TrackEditors/SkeletalAnimationTrackEditor.cpp b/Engine/Source/Editor/MovieSceneTools/Private/TrackEditors/SkeletalAnimationTrackEditor.cpp
--- a/Engine/Source/Editor/MovieSceneTools/Private/TrackEditors/SkeletalAnimationTrackEditor.cpp
+++ b/Engine/Source/Editor/MovieSceneTools/Private/TrackEditors/SkeletalAnimationTrackEditor.cpp
@@ -313,7 +313,7 @@ USkeletalMeshComponent* AcquireSkeletalMeshFromObjectGuid(const FGuid& Guid, TSh
 	}
 	else if(USkeletalMeshComponent* SkeletalMeshComponent = Cast<USkeletalMeshComponent>(BoundObject))
 	{
-		if (SkeletalMeshComponent->GetSkeletalMesh())
+		if (SkeletalMeshComponent->GetSkeletalMeshAsset())
 		{
 			return SkeletalMeshComponent;
 		}
@@ -325,10 +325,10 @@ USkeletalMeshComponent* AcquireSkeletalMeshFromObjectGuid(const FGuid& Guid, TSh
 USkeleton* GetSkeletonFromComponent(UActorComponent* InComponent)
 {
 	USkeletalMeshComponent* SkeletalMeshComp = Cast<USkeletalMeshComponent>(InComponent);
-	if (SkeletalMeshComp && SkeletalMeshComp->GetSkeletalMesh() && SkeletalMeshComp->GetSkeletalMesh()->GetSkeleton())
+	if (SkeletalMeshComp && SkeletalMeshComp->GetSkeletalMeshAsset() && SkeletalMeshComp->GetSkeletalMeshAsset()->GetSkeleton())
 	{
 		// @todo Multiple actors, multiple components
-		return SkeletalMeshComp->GetSkeletalMesh()->GetSkeleton();
+		return SkeletalMeshComp->GetSkeletalMeshAsset()->GetSkeleton();
 	}
 
 	return nullptr;
@@ -1551,9 +1551,9 @@ void FSkeletalAnimationTrackEditor::HandleCreateAnimationSequence(USkeletalMeshC
 	if (SkelMeshComp)
 	{
 		TArray<TWeakObjectPtr<UObject>> Skels;
-		if (SkelMeshComp->GetSkeletalMesh())
+		if (SkelMeshComp->GetSkeletalMeshAsset())
 		{
-			Skels.Add(SkelMeshComp->GetSkeletalMesh());
+			Skels.Add(SkelMeshComp->GetSkeletalMeshAsset());
 		}
 		else
 		{
diff --git a/Engine/Source/Editor/Persona/Private/AnimationEditorPreviewScene.cpp b/Engine/Source/Editor/Persona/Private/AnimationEditorPreviewScene.cpp
--- a/Engine/Source/Editor/Persona/Private/AnimationEditorPreviewScene.cpp
+++ b/Engine/Source/Editor/Persona/Private/AnimationEditorPreviewScene.cpp
@@ -221,7 +221,7 @@ USkeletalMesh* FAnimationEditorPreviewScene::GetPreviewMesh() const
 
 void FAnimationEditorPreviewScene::SetPreviewMeshInternal(USkeletalMesh* NewPreviewMesh)
 {
-	USkeletalMesh* OldPreviewMesh = SkeletalMeshComponent->GetSkeletalMesh();
+	USkeletalMesh* OldPreviewMesh = SkeletalMeshComponent->GetSkeletalMeshAsset();
 
 	// Store off the old skel mesh we are debugging
 	USkeletalMeshComponent* DebuggedSkeletalMeshComponent = nullptr;
@@ -249,7 +249,7 @@ void FAnimationEditorPreviewScene::SetPreviewMeshInternal(USkeletalMesh* NewPrev
 
 	ValidatePreviewAttachedAssets(NewPreviewMesh);
 
-	if (NewPreviewMesh != SkeletalMeshComponent->GetSkeletalMesh())
+	if (NewPreviewMesh != SkeletalMeshComponent->GetSkeletalMeshAsset())
 	{
 		// setting skeletalmesh unregister/re-register, 
 		// so I have to save the animation settings and resetting after setting mesh
@@ -680,15 +680,15 @@ void FAnimationEditorPreviewScene::ShowReferencePose(bool bShowRefPose, bool bRe
 		SkeletalMeshComponent->ShowReferencePose(bShowRefPose);
 
 		// Also reset bone transforms
-		if(bResetBoneTransforms && SkeletalMeshComponent->GetSkeletalMesh() != nullptr)
+		if(bResetBoneTransforms && SkeletalMeshComponent->GetSkeletalMeshAsset() != nullptr)
 		{
 			bool bModified = false;
 			FScopedTransaction Transaction(LOCTEXT("ResetBoneTransforms", "Reset Bone Transforms"));
 
-			int32 NumBones = SkeletalMeshComponent->GetSkeletalMesh()->GetRefSkeleton().GetNum();
+			int32 NumBones = SkeletalMeshComponent->GetSkeletalMeshAsset()->GetRefSkeleton().GetNum();
 			for (int32 BoneIndex = 0; BoneIndex < NumBones; ++BoneIndex)
 			{
-				FName BoneName = SkeletalMeshComponent->GetSkeletalMesh()->GetRefSkeleton().GetBoneName(BoneIndex);
+				FName BoneName = SkeletalMeshComponent->GetSkeletalMeshAsset()->GetRefSkeleton().GetBoneName(BoneIndex);
 				const FAnimNode_ModifyBone* ModifiedBone = SkeletalMeshComponent->PreviewInstance->FindModifiedBone(BoneName);
 				if (ModifiedBone != nullptr)
 				{
@@ -898,7 +898,7 @@ void FAnimationEditorPreviewScene::DeselectAll()
 bool FAnimationEditorPreviewScene::IsRecordAvailable() const
 {
 	// make sure mesh exists
-	return (SkeletalMeshComponent->GetSkeletalMesh() != nullptr);
+	return (SkeletalMeshComponent->GetSkeletalMeshAsset() != nullptr);
 }
 
 FSlateIcon FAnimationEditorPreviewScene::GetRecordStatusImage() const
diff --git a/Engine/Source/Editor/Persona/Private/AnimationEditorViewportClient.cpp b/Engine/Source/Editor/Persona/Private/AnimationEditorViewportClient.cpp
--- a/Engine/Source/Editor/Persona/Private/AnimationEditorViewportClient.cpp
+++ b/Engine/Source/Editor/Persona/Private/AnimationEditorViewportClient.cpp
@@ -149,7 +149,7 @@ FAnimationViewportClient::FAnimationViewportClient(const TSharedRef<IPersonaPrev
 	InPreviewScene->RegisterOnPreviewMeshChanged(FOnPreviewMeshChanged::CreateRaw(this, &FAnimationViewportClient::HandleSkeletalMeshChanged));
 	if (InPreviewScene->GetPreviewMeshComponent())
 	{
-		HandleSkeletalMeshChanged(nullptr, InPreviewScene->GetPreviewMeshComponent()->GetSkeletalMesh());
+		HandleSkeletalMeshChanged(nullptr, InPreviewScene->GetPreviewMeshComponent()->GetSkeletalMeshAsset());
 	}
 	InPreviewScene->RegisterOnInvalidateViews(FSimpleDelegate::CreateRaw(this, &FAnimationViewportClient::HandleInvalidateViews));
 	InPreviewScene->RegisterOnFocusViews(FSimpleDelegate::CreateRaw(this, &FAnimationViewportClient::HandleFocusViews));
@@ -185,7 +185,7 @@ FAnimationViewportClient::~FAnimationViewportClient()
 
 			if (OnMeshChangedDelegateHandle.IsValid())
 			{
-				if (USkeletalMesh* SkelMesh = PreviewMeshComponent->GetSkeletalMesh())
+				if (USkeletalMesh* SkelMesh = PreviewMeshComponent->GetSkeletalMeshAsset())
 				{
 					SkelMesh->GetOnMeshChanged().Remove(OnMeshChangedDelegateHandle);
 				}
@@ -376,7 +376,7 @@ void FAnimationViewportClient::JumpToDefaultCamera()
 
 void FAnimationViewportClient::SaveCameraAsDefault()
 {
-	USkeletalMesh* SkelMesh = GetAnimPreviewScene()->GetPreviewMeshComponent()->GetSkeletalMesh();
+	USkeletalMesh* SkelMesh = GetAnimPreviewScene()->GetPreviewMeshComponent()->GetSkeletalMeshAsset();
 	if (SkelMesh)
 	{
 		FScopedTransaction Transaction(LOCTEXT("SaveCameraAsDefault", "Save Camera As Default"));
@@ -404,7 +404,7 @@ bool FAnimationViewportClient::CanSaveCameraAsDefault() const
 
 void FAnimationViewportClient::ClearDefaultCamera()
 {
-	USkeletalMesh* SkelMesh = GetAnimPreviewScene()->GetPreviewMeshComponent()->GetSkeletalMesh();
+	USkeletalMesh* SkelMesh = GetAnimPreviewScene()->GetPreviewMeshComponent()->GetSkeletalMeshAsset();
 	if (SkelMesh)
 	{
 		FScopedTransaction Transaction(LOCTEXT("ClearDefaultCamera", "Clear Default Camera"));
@@ -422,7 +422,7 @@ void FAnimationViewportClient::ClearDefaultCamera()
 
 bool FAnimationViewportClient::HasDefaultCameraSet() const
 {
-	USkeletalMesh* SkelMesh = GetAnimPreviewScene()->GetPreviewMeshComponent()->GetSkeletalMesh();
+	USkeletalMesh* SkelMesh = GetAnimPreviewScene()->GetPreviewMeshComponent()->GetSkeletalMeshAsset();
 	return (SkelMesh && SkelMesh->GetHasCustomDefaultEditorCamera());
 }
 
@@ -510,7 +510,7 @@ void FAnimationViewportClient::Draw(const FSceneView* View, FPrimitiveDrawInterf
 	TArray<UDebugSkelMeshComponent*> PreviewMeshComponents = GetPreviewScene()->GetAllPreviewMeshComponents();
 	for (UDebugSkelMeshComponent* PreviewMeshComponent : PreviewMeshComponents)
 	{
-		if (!(PreviewMeshComponent && PreviewMeshComponent->GetSkeletalMesh()))
+		if (!(PreviewMeshComponent && PreviewMeshComponent->GetSkeletalMeshAsset()))
 		{
 			continue;
 		}
@@ -556,14 +556,14 @@ void FAnimationViewportClient::Draw(const FSceneView* View, FPrimitiveDrawInterf
 		// Display socket hit points
 		if (PreviewMeshComponent->bDrawSockets )
 		{
-			if (PreviewMeshComponent->bSkeletonSocketsVisible && PreviewMeshComponent->GetSkeletalMesh()->GetSkeleton() )
+			if (PreviewMeshComponent->bSkeletonSocketsVisible && PreviewMeshComponent->GetSkeletalMeshAsset()->GetSkeleton() )
 			{
-				DrawSockets(PreviewMeshComponent, PreviewMeshComponent->GetSkeletalMesh()->GetSkeleton()->Sockets, GetAnimPreviewScene()->GetSelectedSocket(), PDI, true);
+				DrawSockets(PreviewMeshComponent, PreviewMeshComponent->GetSkeletalMeshAsset()->GetSkeleton()->Sockets, GetAnimPreviewScene()->GetSelectedSocket(), PDI, true);
 			}
 
 			if ( PreviewMeshComponent->bMeshSocketsVisible )
 			{
-				DrawSockets(PreviewMeshComponent, PreviewMeshComponent->GetSkeletalMesh()->GetMeshOnlySocketList(), GetAnimPreviewScene()->GetSelectedSocket(), PDI, false);
+				DrawSockets(PreviewMeshComponent, PreviewMeshComponent->GetSkeletalMeshAsset()->GetMeshOnlySocketList(), GetAnimPreviewScene()->GetSelectedSocket(), PDI, false);
 			}
 		}
 
@@ -645,7 +645,7 @@ void FAnimationViewportClient::DrawUVsForMesh(FViewport* InViewport, FCanvas* In
 	}
 	
 	//use the overridden LOD level
-	const uint32 LODLevel = FMath::Clamp(PreviewMeshComponent->GetForcedLOD() - 1, 0, PreviewMeshComponent->GetSkeletalMesh()->GetLODNum() - 1);
+	const uint32 LODLevel = FMath::Clamp(PreviewMeshComponent->GetForcedLOD() - 1, 0, PreviewMeshComponent->GetSkeletalMeshAsset()->GetLODNum() - 1);
 
 	TArray<FVector2D> SelectedEdgeTexCoords; //No functionality in Persona for this (yet?)
 
@@ -832,7 +832,7 @@ void FAnimationViewportClient::ShowBoneNames( FCanvas* Canvas, FSceneView* View,
 
 void FAnimationViewportClient::ShowAttributeNames(FCanvas* Canvas, FSceneView* View, UDebugSkelMeshComponent* MeshComponent) const
 {
-	if (!MeshComponent || !MeshComponent->GetSkeletalMesh())
+	if (!MeshComponent || !MeshComponent->GetSkeletalMeshAsset())
 	{
 		return;
 	}
@@ -906,7 +906,7 @@ FText FAnimationViewportClient::GetDisplayInfo(bool bDisplayAllInfo) const
 
 	// if not valid skeletalmesh
 	UDebugSkelMeshComponent* PreviewMeshComponent = GetPreviewScene()->GetPreviewMeshComponent();
-	if (PreviewMeshComponent == nullptr || PreviewMeshComponent->GetSkeletalMesh() == nullptr)
+	if (PreviewMeshComponent == nullptr || PreviewMeshComponent->GetSkeletalMeshAsset() == nullptr)
 	{
 		return FText();
 	}
@@ -916,16 +916,16 @@ FText FAnimationViewportClient::GetDisplayInfo(bool bDisplayAllInfo) const
 		TextValue = ConcatenateLine(TextValue, LOCTEXT("AdditiveRefPoseWarning", "<AnimViewport.WarningText>Additive ref pose contains scales of 0.0, this can cause additive animations to not give the desired results</>"));
 	}
 
-	if (PreviewMeshComponent->GetSkeletalMesh()->GetMorphTargets().Num() > 0)
+	if (PreviewMeshComponent->GetSkeletalMeshAsset()->GetMorphTargets().Num() > 0)
 	{
 		TArray<UMaterial*> ProcessedMaterials;
 		TArray<UMaterial*> MaterialsThatNeedMorphFlagOn;
 		TArray<UMaterial*> MaterialsThatNeedSaving;
 
-		const TIndirectArray<FSkeletalMeshLODModel>& LODModels = PreviewMeshComponent->GetSkeletalMesh()->GetImportedModel()->LODModels;
+		const TIndirectArray<FSkeletalMeshLODModel>& LODModels = PreviewMeshComponent->GetSkeletalMeshAsset()->GetImportedModel()->LODModels;
 		int32 LodNumber = LODModels.Num();
 		TArray<UMaterialInterface*> MaterialUsingMorphTarget;
-		for (UMorphTarget *MorphTarget : PreviewMeshComponent->GetSkeletalMesh()->GetMorphTargets())
+		for (UMorphTarget *MorphTarget : PreviewMeshComponent->GetSkeletalMeshAsset()->GetMorphTargets())
 		{
 			if (MorphTarget == nullptr)
 			{
@@ -940,7 +940,7 @@ FText FAnimationViewportClient::GetDisplayInfo(bool bDisplayAllInfo) const
 						const FSkeletalMeshLODModel& LODModel = LODModels[LodIdx];
 						if (LODModel.Sections.IsValidIndex(SectionIndex))
 						{
-							MaterialUsingMorphTarget.AddUnique(PreviewMeshComponent->GetSkeletalMesh()->GetMaterials()[LODModel.Sections[SectionIndex].MaterialIndex].MaterialInterface);
+							MaterialUsingMorphTarget.AddUnique(PreviewMeshComponent->GetSkeletalMeshAsset()->GetMaterials()[LODModel.Sections[SectionIndex].MaterialIndex].MaterialInterface);
 						}
 					}
 				}
@@ -1447,7 +1447,7 @@ void FAnimationViewportClient::DrawBonesFromTransforms(
 	FLinearColor RootBoneColour) const
 {
 	if (Transforms.IsEmpty() ||
-		!MeshComponent->GetSkeletalMesh() ||
+		!MeshComponent->GetSkeletalMeshAsset() ||
 		MeshComponent->SkeletonDrawMode == ESkeletonDrawMode::Hidden)
 	{
 		return;
@@ -1498,7 +1498,7 @@ void FAnimationViewportClient::DrawBonesFromCompactPose(
 {
 	if (Pose.GetNumBones() == 0 ||
 		!MeshComponent ||
-		!MeshComponent->GetSkeletalMesh() ||
+		!MeshComponent->GetSkeletalMeshAsset() ||
 		MeshComponent->SkeletonDrawMode == ESkeletonDrawMode::Hidden)
 	{
 		return;
@@ -1548,7 +1548,7 @@ void FAnimationViewportClient::DrawBonesFromCompactPose(
 
 void FAnimationViewportClient::DrawMeshBonesUncompressedAnimation(UDebugSkelMeshComponent * MeshComponent, FPrimitiveDrawInterface* PDI) const
 {
-	if ( MeshComponent && MeshComponent->GetSkeletalMesh())
+	if ( MeshComponent && MeshComponent->GetSkeletalMeshAsset())
 	{
 		DrawBonesFromTransforms(MeshComponent->UncompressedSpaceBases, MeshComponent, PDI, FColor(255, 127, 39, 255), FColor(255, 127, 39, 255));
 	}
@@ -1556,7 +1556,7 @@ void FAnimationViewportClient::DrawMeshBonesUncompressedAnimation(UDebugSkelMesh
 
 void FAnimationViewportClient::DrawMeshBonesNonRetargetedAnimation(UDebugSkelMeshComponent * MeshComponent, FPrimitiveDrawInterface* PDI) const
 {
-	if ( MeshComponent && MeshComponent->GetSkeletalMesh())
+	if ( MeshComponent && MeshComponent->GetSkeletalMeshAsset())
 	{
 		DrawBonesFromTransforms(MeshComponent->NonRetargetedSpaceBases, MeshComponent, PDI, FColor(159, 159, 39, 255), FColor(159, 159, 39, 255));
 	}
@@ -1564,7 +1564,7 @@ void FAnimationViewportClient::DrawMeshBonesNonRetargetedAnimation(UDebugSkelMes
 
 void FAnimationViewportClient::DrawMeshBonesAdditiveBasePose(UDebugSkelMeshComponent * MeshComponent, FPrimitiveDrawInterface* PDI) const
 {
-	if ( MeshComponent && MeshComponent->GetSkeletalMesh())
+	if ( MeshComponent && MeshComponent->GetSkeletalMeshAsset())
 	{
 		DrawBonesFromTransforms(MeshComponent->AdditiveBasePoses, MeshComponent, PDI, FColor(0, 159, 0, 255), FColor(0, 159, 0, 255));
 	}
@@ -1572,7 +1572,7 @@ void FAnimationViewportClient::DrawMeshBonesAdditiveBasePose(UDebugSkelMeshCompo
 
 void FAnimationViewportClient::DrawMeshBonesSourceRawAnimation(UDebugSkelMeshComponent * MeshComponent, FPrimitiveDrawInterface* PDI) const
 {
-	if(MeshComponent && MeshComponent->GetSkeletalMesh())
+	if(MeshComponent && MeshComponent->GetSkeletalMeshAsset())
 	{
 		DrawBonesFromTransforms(MeshComponent->SourceAnimationPoses, MeshComponent, PDI, FColor(195, 195, 195, 255), FColor(195, 159, 195, 255));
 	}
@@ -1598,7 +1598,7 @@ void FAnimationViewportClient::DrawWatchedPoses(UDebugSkelMeshComponent * MeshCo
 
 void FAnimationViewportClient::DrawMeshBonesBakedAnimation(UDebugSkelMeshComponent * MeshComponent, FPrimitiveDrawInterface* PDI) const
 {
-	if(MeshComponent && MeshComponent->GetSkeletalMesh())
+	if(MeshComponent && MeshComponent->GetSkeletalMeshAsset())
 	{
 		DrawBonesFromTransforms(MeshComponent->BakedAnimationPoses, MeshComponent, PDI, FColor(0, 128, 192, 255), FColor(0, 128, 192, 255));
 	}
@@ -1607,7 +1607,7 @@ void FAnimationViewportClient::DrawMeshBonesBakedAnimation(UDebugSkelMeshCompone
 void FAnimationViewportClient::DrawMeshBones(UDebugSkelMeshComponent* MeshComponent, FPrimitiveDrawInterface* PDI) const
 {
 	if (!MeshComponent ||
-		!MeshComponent->GetSkeletalMesh() ||
+		!MeshComponent->GetSkeletalMeshAsset() ||
 		MeshComponent->GetNumDrawTransform() == 0 ||
 		MeshComponent->SkeletonDrawMode == ESkeletonDrawMode::Hidden)
 	{
@@ -1705,7 +1705,7 @@ void FAnimationViewportClient::DrawMeshSubsetBones(const UDebugSkelMeshComponent
 {
 	// this BonesOfInterest has to be in MeshComponent base, not Skeleton 
 	if (!MeshComponent ||
-		!MeshComponent->GetSkeletalMesh() ||
+		!MeshComponent->GetSkeletalMeshAsset() ||
 		BonesOfInterest.IsEmpty() ||
 		MeshComponent->SkeletonDrawMode == ESkeletonDrawMode::Hidden)
 	{
@@ -1781,7 +1781,7 @@ void FAnimationViewportClient::DrawMeshSubsetBones(const UDebugSkelMeshComponent
 
 void FAnimationViewportClient::DrawAttributes(UDebugSkelMeshComponent* MeshComponent, FPrimitiveDrawInterface* PDI) const
 {
-	if (MeshComponent && MeshComponent->GetSkeletalMesh())
+	if (MeshComponent && MeshComponent->GetSkeletalMeshAsset())
 	{
 		const UE::Anim::FMeshAttributeContainer& Attributes = MeshComponent->GetCustomAttributes();
 
@@ -1812,7 +1812,7 @@ void FAnimationViewportClient::DrawAttributes(UDebugSkelMeshComponent* MeshCompo
 
 void FAnimationViewportClient::DrawSockets(const UDebugSkelMeshComponent* InPreviewMeshComponent, TArray<USkeletalMeshSocket*>& InSockets, FSelectedSocketInfo InSelectedSocket, FPrimitiveDrawInterface* PDI, bool bUseSkeletonSocketColor)
 {
-	if (InPreviewMeshComponent && InPreviewMeshComponent->GetSkeletalMesh())
+	if (InPreviewMeshComponent && InPreviewMeshComponent->GetSkeletalMeshAsset())
 	{
 		ELocalAxesMode::Type LocalAxesMode = (ELocalAxesMode::Type)GetDefault<UPersonaOptions>()->DefaultLocalAxesSelection;
 
@@ -1905,7 +1905,7 @@ void FAnimationViewportClient::UpdateCameraSetup()
 {
 	static FRotator CustomOrbitRotation(-33.75, -135, 0);
 	UDebugSkelMeshComponent* PreviewMeshComponent = GetAnimPreviewScene()->GetPreviewMeshComponent();
-	if (PreviewMeshComponent != nullptr && PreviewMeshComponent->GetSkeletalMesh())
+	if (PreviewMeshComponent != nullptr && PreviewMeshComponent->GetSkeletalMeshAsset())
 	{
 		FSphere BoundSphere = GetCameraTarget();
 		FVector CustomOrbitZoom(0, BoundSphere.W / (75.0f * (float)PI / 360.0f), 0);
@@ -1974,7 +1974,7 @@ FBox FAnimationViewportClient::ComputeBoundingBoxForSelectedEditorSection() cons
 		return FBox(ForceInitToZero);
 	}
 
-	USkeletalMesh* const SkeletalMesh = PreviewMeshComponent->GetSkeletalMesh();
+	USkeletalMesh* const SkeletalMesh = PreviewMeshComponent->GetSkeletalMeshAsset();
 	FSkeletalMeshObject* const MeshObject = PreviewMeshComponent->MeshObject;
 	if ( !SkeletalMesh || !MeshObject )
 	{
@@ -2027,7 +2027,7 @@ void FAnimationViewportClient::FocusViewportOnPreviewMesh(bool bUseCustomCamera)
 
 	if (UDebugSkelMeshComponent* const PreviewMeshComponent = GetAnimPreviewScene()->GetPreviewMeshComponent())
 	{
-		if (USkeletalMesh* const SkelMesh = PreviewMeshComponent->GetSkeletalMesh())
+		if (USkeletalMesh* const SkelMesh = PreviewMeshComponent->GetSkeletalMeshAsset())
 		{
 			if (bUseCustomCamera && SkelMesh->GetHasCustomDefaultEditorCamera())
 			{
@@ -2062,7 +2062,7 @@ void FAnimationViewportClient::FocusViewportOnPreviewMesh(bool bUseCustomCamera)
 
 float FAnimationViewportClient::GetFloorOffset() const
 {
-	USkeletalMesh* Mesh = GetPreviewScene()->GetPreviewMeshComponent()->GetSkeletalMesh();
+	USkeletalMesh* Mesh = GetPreviewScene()->GetPreviewMeshComponent()->GetSkeletalMeshAsset();
 	if ( Mesh )
 	{
 		return Mesh->GetFloorOffset();
@@ -2073,7 +2073,7 @@ float FAnimationViewportClient::GetFloorOffset() const
 
 void FAnimationViewportClient::SetFloorOffset( float NewValue )
 {
-	USkeletalMesh* Mesh = GetPreviewScene()->GetPreviewMeshComponent()->GetSkeletalMesh();
+	USkeletalMesh* Mesh = GetPreviewScene()->GetPreviewMeshComponent()->GetSkeletalMeshAsset();
 
 	if ( Mesh )
 	{
diff --git a/Engine/Source/Editor/Persona/Private/EditModes/SkeletonSelectionEditMode.cpp b/Engine/Source/Editor/Persona/Private/EditModes/SkeletonSelectionEditMode.cpp
--- a/Engine/Source/Editor/Persona/Private/EditModes/SkeletonSelectionEditMode.cpp
+++ b/Engine/Source/Editor/Persona/Private/EditModes/SkeletonSelectionEditMode.cpp
@@ -38,14 +38,14 @@ bool FSkeletonSelectionEditMode::GetCameraTarget(FSphere& OutTarget) const
 
 	UDebugSkelMeshComponent* PreviewMeshComponent = GetAnimPreviewScene().GetPreviewMeshComponent();
 
-	if(PreviewMeshComponent && PreviewMeshComponent->GetSkeletalMesh())
+	if(PreviewMeshComponent && PreviewMeshComponent->GetSkeletalMeshAsset())
 	{
 		if (GetAnimPreviewScene().GetSelectedBoneIndex() != INDEX_NONE)
 		{
 			const int32 FocusBoneIndex = GetAnimPreviewScene().GetSelectedBoneIndex();
-			if (FocusBoneIndex != INDEX_NONE && PreviewMeshComponent->GetSkeletalMesh()->GetRefSkeleton().IsValidIndex(FocusBoneIndex))
+			if (FocusBoneIndex != INDEX_NONE && PreviewMeshComponent->GetSkeletalMeshAsset()->GetRefSkeleton().IsValidIndex(FocusBoneIndex))
 			{
-				const FName BoneName = PreviewMeshComponent->GetSkeletalMesh()->GetRefSkeleton().GetBoneName(FocusBoneIndex);
+				const FName BoneName = PreviewMeshComponent->GetSkeletalMeshAsset()->GetRefSkeleton().GetBoneName(FocusBoneIndex);
 				OutTarget.Center = PreviewMeshComponent->GetBoneLocation(BoneName);
 				OutTarget.W = 30.0f;
 				bHandled = true;
@@ -79,7 +79,7 @@ void FSkeletonSelectionEditMode::GetOnScreenDebugInfo(TArray<FText>& OutDebugInf
 
 FSelectedSocketInfo FSkeletonSelectionEditMode::DuplicateAndSelectSocket(const FSelectedSocketInfo& SocketInfoToDuplicate)
 {
-	USkeletalMesh* SkeletalMesh = GetAnimPreviewScene().GetPreviewMeshComponent()->GetSkeletalMesh();
+	USkeletalMesh* SkeletalMesh = GetAnimPreviewScene().GetPreviewMeshComponent()->GetSkeletalMeshAsset();
 	USkeletalMeshSocket* NewSocket = GetAnimPreviewScene().GetPersonaToolkit()->GetEditableSkeleton()->DuplicateSocket(SocketInfoToDuplicate, SocketInfoToDuplicate.Socket->BoneName, SkeletalMesh);
 
 	FSelectedSocketInfo NewSocketInfo(NewSocket, SocketInfoToDuplicate.bSocketIsOnSkeleton);
@@ -95,13 +95,13 @@ bool FSkeletonSelectionEditMode::StartTracking(FEditorViewportClient* InViewport
 	UE::Widget::EWidgetMode WidgetMode = InViewportClient->GetWidgetMode();
 
 	UDebugSkelMeshComponent* PreviewMeshComponent = GetAnimPreviewScene().GetPreviewMeshComponent();
-	if(PreviewMeshComponent != nullptr && PreviewMeshComponent->GetSkeletalMesh() != nullptr)
+	if(PreviewMeshComponent != nullptr && PreviewMeshComponent->GetSkeletalMeshAsset() != nullptr)
 	{
 		int32 BoneIndex = GetAnimPreviewScene().GetSelectedBoneIndex();
 		USkeletalMeshSocket* SelectedSocket = GetAnimPreviewScene().GetSelectedSocket().Socket;
 		AActor* SelectedActor = GetAnimPreviewScene().GetSelectedActor();
 
-		if ((BoneIndex >= 0 && PreviewMeshComponent->GetSkeletalMesh()->GetRefSkeleton().IsValidIndex(BoneIndex)) || SelectedSocket != nullptr || SelectedActor != nullptr)
+		if ((BoneIndex >= 0 && PreviewMeshComponent->GetSkeletalMeshAsset()->GetRefSkeleton().IsValidIndex(BoneIndex)) || SelectedSocket != nullptr || SelectedActor != nullptr)
 		{
 			bool bValidAxis = false;
 			FVector WorldAxisDir;
@@ -157,7 +157,7 @@ bool FSkeletonSelectionEditMode::StartTracking(FEditorViewportClient* InViewport
 						bInTransaction = true;
 
 						// now modify the bone array
-						const FName BoneName = PreviewMeshComponent->GetSkeletalMesh()->GetRefSkeleton().GetBoneName(BoneIndex);
+						const FName BoneName = PreviewMeshComponent->GetSkeletalMeshAsset()->GetRefSkeleton().GetBoneName(BoneIndex);
 						PreviewMeshComponent->PreviewInstance->ModifyBone(BoneName);
 					}
 				}
@@ -216,7 +216,7 @@ bool FSkeletonSelectionEditMode::InputDelta(FEditorViewportClient* InViewportCli
 		if ( BoneIndex >= 0 )
 		{
 			//Get the skeleton control manipulating this bone
-			const FName BoneName = PreviewMeshComponent->GetSkeletalMesh()->GetRefSkeleton().GetBoneName(BoneIndex);
+			const FName BoneName = PreviewMeshComponent->GetSkeletalMeshAsset()->GetRefSkeleton().GetBoneName(BoneIndex);
 			SkelControl = &(PreviewMeshComponent->PreviewInstance->ModifyBone(BoneName));
 		}
 
@@ -339,12 +339,12 @@ FIntPoint FSkeletonSelectionEditMode::GetDPIUnscaledSize(FViewport* Viewport, FV
 void FSkeletonSelectionEditMode::DrawHUD(FEditorViewportClient* ViewportClient, FViewport* Viewport, const FSceneView* View, FCanvas* Canvas)
 {
 	UDebugSkelMeshComponent* PreviewMeshComponent = GetAnimPreviewScene().GetPreviewMeshComponent();
-	if (PreviewMeshComponent == nullptr || PreviewMeshComponent->GetSkeletalMesh() == nullptr)
+	if (PreviewMeshComponent == nullptr || PreviewMeshComponent->GetSkeletalMeshAsset() == nullptr)
 	{
 		return;
 	}
 
-	FReferenceSkeleton& RefSkeleton = PreviewMeshComponent->GetSkeletalMesh()->GetRefSkeleton();
+	FReferenceSkeleton& RefSkeleton = PreviewMeshComponent->GetSkeletalMeshAsset()->GetRefSkeleton();
 	int32 BoneIndex = GetAnimPreviewScene().GetSelectedBoneIndex();
 
 	// Draw name of selected bone
@@ -354,7 +354,7 @@ void FSkeletonSelectionEditMode::DrawHUD(FEditorViewportClient* ViewportClient,
 		const int32 HalfX = ViewPortSize.X / 2;
 		const int32 HalfY = ViewPortSize.Y / 2;
 
-		const FName BoneName = PreviewMeshComponent->GetSkeletalMesh()->GetRefSkeleton().GetBoneName(BoneIndex);
+		const FName BoneName = PreviewMeshComponent->GetSkeletalMeshAsset()->GetRefSkeleton().GetBoneName(BoneIndex);
 
 		FMatrix BoneMatrix = PreviewMeshComponent->GetBoneMatrix(BoneIndex);
 		const FPlane Proj = View->Project(BoneMatrix.GetOrigin());
@@ -446,10 +446,10 @@ bool FSkeletonSelectionEditMode::UsesTransformWidget(UE::Widget::EWidgetMode Che
 bool FSkeletonSelectionEditMode::GetCustomDrawingCoordinateSystem(FMatrix& InMatrix, void* InData)
 {
 	UDebugSkelMeshComponent* PreviewMeshComponent = GetAnimPreviewScene().GetPreviewMeshComponent();
-	if(PreviewMeshComponent && PreviewMeshComponent->GetSkeletalMesh())
+	if(PreviewMeshComponent && PreviewMeshComponent->GetSkeletalMeshAsset())
 	{
 		int32 BoneIndex = GetAnimPreviewScene().GetSelectedBoneIndex();
-		if (BoneIndex != INDEX_NONE && PreviewMeshComponent->GetSkeletalMesh()->GetRefSkeleton().IsValidIndex(BoneIndex))
+		if (BoneIndex != INDEX_NONE && PreviewMeshComponent->GetSkeletalMeshAsset()->GetRefSkeleton().IsValidIndex(BoneIndex))
 		{
 			FTransform BoneMatrix = PreviewMeshComponent->GetBoneTransform(BoneIndex);
 
@@ -487,9 +487,9 @@ FVector FSkeletonSelectionEditMode::GetWidgetLocation() const
 	int32 BoneIndex = GetAnimPreviewScene().GetSelectedBoneIndex();
 	if (BoneIndex != INDEX_NONE)
 	{
-		if (PreviewMeshComponent && PreviewMeshComponent->GetSkeletalMesh() && PreviewMeshComponent->GetSkeletalMesh()->GetRefSkeleton().IsValidIndex(BoneIndex))
+		if (PreviewMeshComponent && PreviewMeshComponent->GetSkeletalMeshAsset() && PreviewMeshComponent->GetSkeletalMeshAsset()->GetRefSkeleton().IsValidIndex(BoneIndex))
 		{
-			const FName BoneName = PreviewMeshComponent->GetSkeletalMesh()->GetRefSkeleton().GetBoneName(BoneIndex);
+			const FName BoneName = PreviewMeshComponent->GetSkeletalMeshAsset()->GetRefSkeleton().GetBoneName(BoneIndex);
 
 			const FMatrix BoneMatrix = PreviewMeshComponent->GetBoneMatrix(BoneIndex);
 
diff --git a/Engine/Source/Editor/Persona/Private/PersonaModule.cpp b/Engine/Source/Editor/Persona/Private/PersonaModule.cpp
--- a/Engine/Source/Editor/Persona/Private/PersonaModule.cpp
+++ b/Engine/Source/Editor/Persona/Private/PersonaModule.cpp
@@ -1394,7 +1394,7 @@ bool FPersonaModule::CreateAnimation(const TArray<UObject*> NewAssets, const EPo
 				switch (Option)
 				{
 				case EPoseSourceOption::ReferencePose:
-					bResult &= NewAnimSequence->CreateAnimation(MeshComponent->GetSkeletalMesh());
+					bResult &= NewAnimSequence->CreateAnimation(MeshComponent->GetSkeletalMeshAsset());
 					break;
 				case EPoseSourceOption::CurrentPose:
 					bResult &= NewAnimSequence->CreateAnimation(MeshComponent);
diff --git a/Engine/Source/Editor/Persona/Private/PersonaPreviewSceneAnimationController.cpp b/Engine/Source/Editor/Persona/Private/PersonaPreviewSceneAnimationController.cpp
--- a/Engine/Source/Editor/Persona/Private/PersonaPreviewSceneAnimationController.cpp
+++ b/Engine/Source/Editor/Persona/Private/PersonaPreviewSceneAnimationController.cpp
@@ -23,7 +23,7 @@ IDetailPropertyRow* UPersonaPreviewSceneAnimationController::AddPreviewControlle
 {
 	TArray<UObject*> ListOfPreviewController{ this };
 
-	const USkeleton* Skeleton = PersonaToolkit->GetPreviewMeshComponent()->GetSkeletalMesh() ? PersonaToolkit->GetPreviewMeshComponent()->GetSkeletalMesh()->GetSkeleton() : nullptr;
+	const USkeleton* Skeleton = PersonaToolkit->GetPreviewMeshComponent()->GetSkeletalMeshAsset() ? PersonaToolkit->GetPreviewMeshComponent()->GetSkeletalMeshAsset()->GetSkeleton() : nullptr;
 	if (Skeleton)
 	{
 		IDetailPropertyRow* NewRow = Category.AddExternalObjectProperty(ListOfPreviewController, Property->GetFName(), PropertyLocation);
diff --git a/Engine/Source/Editor/Persona/Private/SAnimAttributeView.cpp b/Engine/Source/Editor/Persona/Private/SAnimAttributeView.cpp
--- a/Engine/Source/Editor/Persona/Private/SAnimAttributeView.cpp
+++ b/Engine/Source/Editor/Persona/Private/SAnimAttributeView.cpp
@@ -270,7 +270,7 @@ void SAnimAttributeView::DisplayNewAttributeContainerSnapshots(
 	}
 	
 	// we need the SKM to look up bone names
-	if (!InOwningComponent || !InOwningComponent->GetSkeletalMesh())
+	if (!InOwningComponent || !InOwningComponent->GetSkeletalMeshAsset())
 	{
 		ClearListView();
 		return;
@@ -309,7 +309,7 @@ void SAnimAttributeView::DisplayNewAttributeContainerSnapshots(
 				for (int32 IdIndex = 0; IdIndex < Ids.Num(); IdIndex++)
 				{
 					const UE::Anim::FAttributeId& Id = Ids[IdIndex];
-					const FName BoneName = InOwningComponent->GetSkeletalMesh()->GetRefSkeleton().GetBoneName(Id.GetIndex());
+					const FName BoneName = InOwningComponent->GetSkeletalMeshAsset()->GetRefSkeleton().GetBoneName(Id.GetIndex());
 				
 					const FAnimationAttributeIdentifier Identifier(
 						Id.GetName(), Id.GetIndex() ,BoneName, Types[TypeIndex].Get());
@@ -397,7 +397,7 @@ bool SAnimAttributeView::ShouldInvalidateListViewCache(
 			for (int32 IdIndex = 0; IdIndex < Ids.Num(); IdIndex++)
 			{
 				const UE::Anim::FAttributeId& Id = Ids[IdIndex];
-				const FName BoneName = InOwningComponent->GetSkeletalMesh()->GetRefSkeleton().GetBoneName(Id.GetIndex());
+				const FName BoneName = InOwningComponent->GetSkeletalMeshAsset()->GetRefSkeleton().GetBoneName(Id.GetIndex());
 				
 				FAnimationAttributeIdentifier Identifier(
 					Id.GetName(), Id.GetIndex() ,BoneName, Types[TypeIndex].Get());
diff --git a/Engine/Source/Editor/Persona/Private/SAnimViewportToolBar.cpp b/Engine/Source/Editor/Persona/Private/SAnimViewportToolBar.cpp
--- a/Engine/Source/Editor/Persona/Private/SAnimViewportToolBar.cpp
+++ b/Engine/Source/Editor/Persona/Private/SAnimViewportToolBar.cpp
@@ -871,7 +871,7 @@ void SAnimViewportToolBar::FillCharacterMirrorMenu(FMenuBuilder& MenuBuilder) co
 {
 	FAssetPickerConfig AssetPickerConfig;
 	UDebugSkelMeshComponent* PreviewComp = Viewport.Pin()->GetPreviewScene()->GetPreviewMeshComponent();
-	USkeletalMesh* Mesh = PreviewComp->GetSkeletalMesh();
+	USkeletalMesh* Mesh = PreviewComp->GetSkeletalMeshAsset();
 	UAnimPreviewInstance* PreviewInstance = PreviewComp->PreviewInstance; 
 	if (Mesh && PreviewInstance)
 	{
@@ -902,7 +902,7 @@ void SAnimViewportToolBar::OnMirrorDataTableSelected(const FAssetData& SelectedM
 	if (Viewport.Pin().IsValid())
 	{
 		UDebugSkelMeshComponent* PreviewComp = Viewport.Pin()->GetPreviewScene()->GetPreviewMeshComponent();
-		USkeletalMesh* Mesh = PreviewComp->GetSkeletalMesh();
+		USkeletalMesh* Mesh = PreviewComp->GetSkeletalMeshAsset();
 		UAnimPreviewInstance* PreviewInstance = PreviewComp->PreviewInstance; 
 		if (Mesh && PreviewInstance)
 		{
diff --git a/Engine/Source/Editor/Persona/Private/SAnimationEditorViewport.cpp b/Engine/Source/Editor/Persona/Private/SAnimationEditorViewport.cpp
--- a/Engine/Source/Editor/Persona/Private/SAnimationEditorViewport.cpp
+++ b/Engine/Source/Editor/Persona/Private/SAnimationEditorViewport.cpp
@@ -274,7 +274,7 @@ FText SAnimationEditorViewportTabBody::GetDisplayString() const
 				DefaultText = FText::Format(LOCTEXT("PreviewingAnimBP", "Previewing {0}"), FText::FromString(Component->AnimClass->GetName()));
 			}
 		}
-		else if (Component->GetSkeletalMesh() == NULL && TargetSkeletonName != NAME_None)
+		else if (Component->GetSkeletalMeshAsset() == NULL && TargetSkeletonName != NAME_None)
 		{
 			DefaultText = FText::Format(LOCTEXT("NoMeshFound", "No skeletal mesh found for skeleton '{0}'"), FText::FromName(TargetSkeletonName));
 		}
@@ -1159,9 +1159,9 @@ bool SAnimationEditorViewportTabBody::IsTurnTableModeSelected(int32 ModeIndex) c
 int32 SAnimationEditorViewportTabBody::GetLODModelCount() const
 {
 	UDebugSkelMeshComponent* PreviewComponent = GetPreviewScene()->GetPreviewMeshComponent();
-	if( PreviewComponent && PreviewComponent->GetSkeletalMesh())
+	if( PreviewComponent && PreviewComponent->GetSkeletalMeshAsset())
 	{
-		return PreviewComponent->GetSkeletalMesh()->GetResourceForRendering()->LODRenderData.Num();
+		return PreviewComponent->GetSkeletalMeshAsset()->GetResourceForRendering()->LODRenderData.Num();
 	}
 	return 0;
 }
@@ -1747,7 +1747,7 @@ void SAnimationEditorViewportTabBody::PopulateSkinWeightProfileNames()
 	// Retrieve all possible skin weight profiles from the component
 	if (UDebugSkelMeshComponent* PreviewComponent = GetPreviewScene()->GetPreviewMeshComponent())
 	{
-		if (USkeletalMesh* Mesh = PreviewComponent->GetSkeletalMesh())
+		if (USkeletalMesh* Mesh = PreviewComponent->GetSkeletalMeshAsset())
 		{
 			for (const FSkinWeightProfileInfo& Profile : Mesh->GetSkinWeightProfiles())
 			{
@@ -2380,9 +2380,9 @@ void SAnimationEditorViewportTabBody::AddPostProcessNotification()
 	{
 		if (UDebugSkelMeshComponent* PreviewComponent = GetPreviewScene()->GetPreviewMeshComponent())
 		{
-			if(PreviewComponent->GetSkeletalMesh() && PreviewComponent->GetSkeletalMesh()->GetPostProcessAnimBlueprint() && PreviewComponent->GetSkeletalMesh()->GetPostProcessAnimBlueprint()->ClassGeneratedBy)
+			if(PreviewComponent->GetSkeletalMeshAsset() && PreviewComponent->GetSkeletalMeshAsset()->GetPostProcessAnimBlueprint() && PreviewComponent->GetSkeletalMeshAsset()->GetPostProcessAnimBlueprint()->ClassGeneratedBy)
 			{
-				return FText::FromString(PreviewComponent->GetSkeletalMesh()->GetPostProcessAnimBlueprint()->ClassGeneratedBy->GetName());
+				return FText::FromString(PreviewComponent->GetSkeletalMeshAsset()->GetPostProcessAnimBlueprint()->ClassGeneratedBy->GetName());
 			}
 		}
 
@@ -2431,9 +2431,9 @@ void SAnimationEditorViewportTabBody::AddPostProcessNotification()
 	{
 		if (UDebugSkelMeshComponent* PreviewComponent = GetPreviewScene()->GetPreviewMeshComponent())
 		{
-			if(PreviewComponent->GetSkeletalMesh() && PreviewComponent->GetSkeletalMesh()->GetPostProcessAnimBlueprint())
+			if(PreviewComponent->GetSkeletalMeshAsset() && PreviewComponent->GetSkeletalMeshAsset()->GetPostProcessAnimBlueprint())
 			{
-				GEditor->GetEditorSubsystem<UAssetEditorSubsystem>()->OpenEditorForAssets(TArray<UObject*>({ PreviewComponent->GetSkeletalMesh()->GetPostProcessAnimBlueprint()->ClassGeneratedBy }));
+				GEditor->GetEditorSubsystem<UAssetEditorSubsystem>()->OpenEditorForAssets(TArray<UObject*>({ PreviewComponent->GetSkeletalMeshAsset()->GetPostProcessAnimBlueprint()->ClassGeneratedBy }));
 			}
 		}
 
diff --git a/Engine/Source/Editor/Persona/Private/SAnimationSequenceBrowser.cpp b/Engine/Source/Editor/Persona/Private/SAnimationSequenceBrowser.cpp
--- a/Engine/Source/Editor/Persona/Private/SAnimationSequenceBrowser.cpp
+++ b/Engine/Source/Editor/Persona/Private/SAnimationSequenceBrowser.cpp
@@ -760,7 +760,7 @@ void SAnimationSequenceBrowser::OnExportToFBX(TArray<FAssetData> SelectedAssets)
 		}
 
 		FPersonaModule& PersonaModule = FModuleManager::GetModuleChecked<FPersonaModule>("Persona");
-		PersonaModule.ExportToFBX(AnimSequences, PersonaToolkitPtr.Pin()->GetPreviewScene()->GetPreviewMeshComponent()->GetSkeletalMesh());
+		PersonaModule.ExportToFBX(AnimSequences, PersonaToolkitPtr.Pin()->GetPreviewScene()->GetPreviewMeshComponent()->GetSkeletalMeshAsset());
 	}
 }
 
@@ -768,7 +768,7 @@ void SAnimationSequenceBrowser::OnSetCurrentPreviewMesh(TArray<FAssetData> Selec
 {
 	if(SelectedAssets.Num() > 0)
 	{
-		USkeletalMesh* PreviewMesh = PersonaToolkitPtr.Pin()->GetPreviewScene()->GetPreviewMeshComponent()->GetSkeletalMesh();
+		USkeletalMesh* PreviewMesh = PersonaToolkitPtr.Pin()->GetPreviewScene()->GetPreviewMeshComponent()->GetSkeletalMeshAsset();
 		if (PreviewMesh)
 		{
 			TArray<TWeakObjectPtr<UAnimSequence>> AnimSequences;
@@ -1452,7 +1452,7 @@ bool SAnimationSequenceBrowser::OnVisualizeAssetToolTip(const TSharedPtr<SWidget
 		
 		if(MeshToUse)
 		{
-			if(PreviewComponent->GetSkeletalMesh() != MeshToUse)
+			if(PreviewComponent->GetSkeletalMeshAsset() != MeshToUse)
 			{
 				PreviewComponent->SetSkeletalMesh(MeshToUse);
 			}
diff --git a/Engine/Source/Editor/Persona/Private/SMorphTargetViewer.cpp b/Engine/Source/Editor/Persona/Private/SMorphTargetViewer.cpp
--- a/Engine/Source/Editor/Persona/Private/SMorphTargetViewer.cpp
+++ b/Engine/Source/Editor/Persona/Private/SMorphTargetViewer.cpp
@@ -329,7 +329,7 @@ void SMorphTargetViewer::Construct(const FArguments& InArgs, const TSharedRef<IP
 {
 	PreviewScenePtr = InPreviewScene;
 
-	SkeletalMesh = InPreviewScene->GetPreviewMeshComponent()->GetSkeletalMesh();
+	SkeletalMesh = InPreviewScene->GetPreviewMeshComponent()->GetSkeletalMeshAsset();
 	InPreviewScene->RegisterOnPreviewMeshChanged( FOnPreviewMeshChanged::CreateSP( this, &SMorphTargetViewer::OnPreviewMeshChanged ) );
 	InPreviewScene->RegisterOnMorphTargetsChanged(FSimpleDelegate::CreateSP(this, &SMorphTargetViewer::OnMorphTargetsChanged));
 	OnPostUndo.Add(FSimpleDelegate::CreateSP(this, &SMorphTargetViewer::OnPostUndo));
