Support for AnimCurves Copying and Pasting in Animation Timeline Editor

**New**
- Added AnimTimelineClipboard.cpp / h.
- Added a "Selection Edit" Section in when Right Click on Animation Sequence Timeline Editor.
- Added `SupportsCopy()` and `Copy(UAnimTimelineClipboardContent* ClipboardContent)` vritual methods to `FAnimTimelineTrack`. These handle copying data to clipboard content object.
- Added methods for clipboard support in `FAnimModel_AnimSequenceBase`: `CopyToClipboard()`, `CutToClipboard()`, `PasteFromClipboard()`, and `PasteDataFromClipboardToSelectedCurve`.
- Added UI Handlers for clipboard support in `FAnimModel_AnimSequenceBase`: `CanCopyToClipboard()`, `CanCutToClipboard()`, `CanPasteFromClipboard()`, and  `CanPasteDataFromClipboardToSelectedCurve()`.

**Changed**
- Updated `FAnimTimelineTrack_Curve`, `FAnimTimelineTrack_FloatCurve, and `FAnimTimelineTrack_TransformCurve` to support clipboard copying.
- Update `FAnimTimelineTrack_Curve::AddToContextMenu` to have a "Curve Edit" and "Selection Edit" section. These provide buttons for Copying, Paste, Cutting and Deleting.

**Removed**
- Removed `RemoveSelectedCurves` variable from the `AnimSequenceTimelineCommands.h` since now the `FAnimTimelineTrack_Curve` Context Menu using the `FGenericCommands::Delete`.

diff --git a/Engine/Source/Editor/Persona/Private/AnimSequenceTimelineCommands.cpp b/Engine/Source/Editor/Persona/Private/AnimSequenceTimelineCommands.cpp
--- a/Engine/Source/Editor/Persona/Private/AnimSequenceTimelineCommands.cpp
+++ b/Engine/Source/Editor/Persona/Private/AnimSequenceTimelineCommands.cpp
@@ -6,8 +6,8 @@
 
 void FAnimSequenceTimelineCommands::RegisterCommands()
 {
+	UI_COMMAND(PasteDataIntoCurve, "Paste Clipboard Data Into Curve", "Paste curve data from the clipboard into the selected curve", EUserInterfaceActionType::Button, FInputChord(EModifierKey::Control | EModifierKey::Shift, EKeys::V));
 	UI_COMMAND(EditSelectedCurves, "Edit Selected Curves", "Edit the selected curves in the curve editor tab", EUserInterfaceActionType::Button, FInputChord());
-	UI_COMMAND(RemoveSelectedCurves, "Remove Selected Curves", "Remove the selected curves", EUserInterfaceActionType::Button, FInputChord());
 	UI_COMMAND(AddNotifyTrack, "Add Notify Track", "Add a new notify track", EUserInterfaceActionType::Button, FInputChord());
 	UI_COMMAND(InsertNotifyTrack, "Insert Notify Track", "Insert a new notify track above here", EUserInterfaceActionType::Button, FInputChord());
 	UI_COMMAND(RemoveNotifyTrack, "Remove Notify Track", "Remove this notify track", EUserInterfaceActionType::Button, FInputChord());
diff --git a/Engine/Source/Editor/Persona/Private/AnimSequenceTimelineCommands.h b/Engine/Source/Editor/Persona/Private/AnimSequenceTimelineCommands.h
--- a/Engine/Source/Editor/Persona/Private/AnimSequenceTimelineCommands.h
+++ b/Engine/Source/Editor/Persona/Private/AnimSequenceTimelineCommands.h
@@ -24,9 +24,10 @@ public:
 	}
 
 	TSharedPtr<FUICommandInfo> EditSelectedCurves;
-
-	TSharedPtr<FUICommandInfo> RemoveSelectedCurves;
+	
 	TSharedPtr<FUICommandInfo> AddNotifyTrack;
+	
+	TSharedPtr<FUICommandInfo> PasteDataIntoCurve;
 
 	TSharedPtr<FUICommandInfo> InsertNotifyTrack;
 
diff --git a/Engine/Source/Editor/Persona/Private/AnimTimeline/AnimModel_AnimSequenceBase.cpp b/Engine/Source/Editor/Persona/Private/AnimTimeline/AnimModel_AnimSequenceBase.cpp
--- a/Engine/Source/Editor/Persona/Private/AnimTimeline/AnimModel_AnimSequenceBase.cpp
+++ b/Engine/Source/Editor/Persona/Private/AnimTimeline/AnimModel_AnimSequenceBase.cpp
@@ -22,9 +22,10 @@
 #include "IPersonaPreviewScene.h"
 #include "Animation/DebugSkelMeshComponent.h"
 #include "AnimPreviewInstance.h"
+#include "AnimTimelineClipboard.h"
 #include "HAL/PlatformApplicationMisc.h"
 #include "ScopedTransaction.h"
-
+#include "Animation/AnimCurveTypes.h"
 
 #define LOCTEXT_NAMESPACE "FAnimModel_AnimSequence"
 
@@ -77,10 +78,6 @@ void FAnimModel_AnimSequenceBase::Initialize()
 		FExecuteAction::CreateSP(this, &FAnimModel_AnimSequenceBase::EditSelectedCurves),
 		FCanExecuteAction::CreateSP(this, &FAnimModel_AnimSequenceBase::CanEditSelectedCurves));
 
-	CommandList->MapAction(
-		Commands.RemoveSelectedCurves,
-		FExecuteAction::CreateSP(this, &FAnimModel_AnimSequenceBase::RemoveSelectedCurves));
-
 	CommandList->MapAction(
 		Commands.CopySelectedCurveNames,
 		FExecuteAction::CreateSP(this, &FAnimModel_AnimSequenceBase::CopySelectedCurveNamesToClipboard));
@@ -136,6 +133,30 @@ void FAnimModel_AnimSequenceBase::Initialize()
 		FCanExecuteAction(),
 		FIsActionChecked::CreateSP(this, &FAnimModel_AnimSequenceBase::IsSnapChecked, FAnimModel::FSnapType::MontageSection.Type),
 		FIsActionButtonVisible::CreateSP(this, &FAnimModel_AnimSequenceBase::IsSnapAvailable, FAnimModel::FSnapType::MontageSection.Type));
+
+	CommandList->MapAction(
+		FGenericCommands::Get().Copy,
+		FExecuteAction::CreateSP(this, &FAnimModel_AnimSequenceBase::CopyToClipboard),
+		FCanExecuteAction::CreateSP(this, &FAnimModel_AnimSequenceBase::CanCopyToClipboard));
+
+	CommandList->MapAction(
+		Commands.PasteDataIntoCurve,
+		FExecuteAction::CreateSP(this, &FAnimModel_AnimSequenceBase::PasteDataFromClipboardToSelectedCurve),
+		FCanExecuteAction::CreateSP(this, &FAnimModel_AnimSequenceBase::CanPasteDataFromClipboardToSelectedCurve));
+
+	CommandList->MapAction(
+		FGenericCommands::Get().Paste,
+		FExecuteAction::CreateSP(this, &FAnimModel_AnimSequenceBase::PasteFromClipboard),
+		FCanExecuteAction::CreateSP(this, &FAnimModel_AnimSequenceBase::CanPasteFromClipboard));
+	
+	CommandList->MapAction(
+		FGenericCommands::Get().Cut,
+		FExecuteAction::CreateSP(this, &FAnimModel_AnimSequenceBase::CutToClipboard),
+		FCanExecuteAction::CreateSP(this, &FAnimModel_AnimSequenceBase::CanCutToClipboard));
+
+	CommandList->MapAction(
+		FGenericCommands::Get().Delete,
+		FExecuteAction::CreateSP(this, &FAnimModel_AnimSequenceBase::RemoveSelectedCurves));
 }
 
 
@@ -318,6 +339,7 @@ void FAnimModel_AnimSequenceBase::OnDataModelChanged(const EAnimDataModelNotifyT
 	switch(NotifyType)
 	{ 
 		case EAnimDataModelNotifyType::CurveAdded:
+		case EAnimDataModelNotifyType::CurveChanged:
 		case EAnimDataModelNotifyType::CurveRemoved:
 		case EAnimDataModelNotifyType::TrackAdded:
 		case EAnimDataModelNotifyType::TrackChanged:
@@ -513,6 +535,96 @@ bool FAnimModel_AnimSequenceBase::IsDisplaySecondaryChecked() const
 	return GetDefault<UPersonaOptions>()->bTimelineDisplayFormatSecondary;
 }
 
+void FAnimModel_AnimSequenceBase::CopyToClipboard() const
+{
+	if (SelectedTracks.Num() > 0)
+	{
+		if (UAnimTimelineClipboardContent * ClipboardContent = UAnimTimelineClipboardContent::Create())
+		{
+			FAnimTimelineClipboardUtilities::CopySelectedTracksToClipboard(SelectedTracks, ClipboardContent);
+			FAnimTimelineClipboardUtilities::CopyContentToClipboard(ClipboardContent);
+		}
+		else
+		{
+			UE_LOG(LogAnimation, Warning, TEXT("Failed to get create valid clipboard object for Animation Timeline while attempting to copy data"));
+		}
+	}
+}
+
+bool FAnimModel_AnimSequenceBase::CanCopyToClipboard()
+{
+	if (!SelectedTracks.IsEmpty())
+	{
+		for (const auto & Track : SelectedTracks)
+		{
+			if (!Track->SupportsCopy())
+			{
+				return false;
+			}
+		}
+	}
+	
+	return true;
+}
+
+void FAnimModel_AnimSequenceBase::PasteDataFromClipboardToSelectedCurve()
+{
+	if (const UAnimTimelineClipboardContent* ClipboardContent = FAnimTimelineClipboardUtilities::GetContentFromClipboard())
+	{
+		// Setup for modifying animation sequence 
+		const FScopedTransaction Transaction(LOCTEXT("AnimSequenceBase_PasteCurveData", "Paste Data To Selected Curve"));
+		AnimSequenceBase->Modify();
+
+		// Paste data
+		FAnimTimelineClipboardUtilities::OverwriteSelectedCurveDataFromClipboard(ClipboardContent, SelectedTracks, AnimSequenceBase);
+	}
+	else
+	{
+		UE_LOG(LogAnimation, Warning, TEXT("Failed to get valid clipboard for Animation Timeline while attempting to paste data"));
+	}
+}
+
+bool FAnimModel_AnimSequenceBase::CanPasteDataFromClipboardToSelectedCurve()
+{
+	return FAnimTimelineClipboardUtilities::CanOverwriteSelectedCurveDataFromClipboard(SelectedTracks);
+}
+
+void FAnimModel_AnimSequenceBase::PasteFromClipboard()
+{
+	if (const UAnimTimelineClipboardContent* ClipboardContent = FAnimTimelineClipboardUtilities::GetContentFromClipboard())
+	{
+		// Setup for modifying animation sequence 
+		const FScopedTransaction Transaction(LOCTEXT("AnimSequenceBase_PasteCurves", "Paste"));
+		AnimSequenceBase->Modify();
+
+		// Paste curves from clipboard
+		FAnimTimelineClipboardUtilities::OverwriteOrAddCurvesFromClipboardContent(ClipboardContent, AnimSequenceBase);
+	}
+	else
+	{
+		UE_LOG(LogAnimation, Warning, TEXT("Failed to get valid clipboard for Animation Timeline while attempting to paste data"));
+	}
+}
+
+bool FAnimModel_AnimSequenceBase::CanPasteFromClipboard()
+{
+	const UAnimTimelineClipboardContent* ClipboardContent = FAnimTimelineClipboardUtilities::GetContentFromClipboard();
+	return ClipboardContent && !ClipboardContent->IsEmpty();
+}
+
+void FAnimModel_AnimSequenceBase::CutToClipboard()
+{
+	const FScopedTransaction Transaction(LOCTEXT("AnimSequenceBase_CutCurveSelection", "Cut Selection"));
+	
+	CopyToClipboard();
+	RemoveSelectedCurves();
+}
+
+bool FAnimModel_AnimSequenceBase::CanCutToClipboard()
+{
+	return CanCopyToClipboard();
+}
+
 void FAnimModel_AnimSequenceBase::UpdateRange()
 {
 	FAnimatedRange OldPlaybackRange = PlaybackRange;
diff --git a/Engine/Source/Editor/Persona/Private/AnimTimeline/AnimModel_AnimSequenceBase.h b/Engine/Source/Editor/Persona/Private/AnimTimeline/AnimModel_AnimSequenceBase.h
--- a/Engine/Source/Editor/Persona/Private/AnimTimeline/AnimModel_AnimSequenceBase.h
+++ b/Engine/Source/Editor/Persona/Private/AnimTimeline/AnimModel_AnimSequenceBase.h
@@ -74,6 +74,22 @@ private:
 	bool IsDisplayPercentageChecked() const;
 	void ToggleDisplaySecondary();
 	bool IsDisplaySecondaryChecked() const;
+	
+	/** Copy selected curves to clipboard */
+	void CopyToClipboard() const;
+	bool CanCopyToClipboard();
+
+	/** Paste curve data into selected curve. Only modifies curves, does not add any new curves. */
+	void PasteDataFromClipboardToSelectedCurve();
+	bool CanPasteDataFromClipboardToSelectedCurve();
+
+	/** Paste curves from clipboard. Adds or overwrites curves (if identifiers collide) */
+	void PasteFromClipboard();
+	bool CanPasteFromClipboard();
+
+	/** Cut selected curves to clipboard */
+	void CutToClipboard();
+	bool CanCutToClipboard();
 
 private:
 	/** The anim sequence base we wrap */
diff --git a/Engine/Source/Editor/Persona/Private/AnimTimeline/AnimTimelineClipboard.cpp b/Engine/Source/Editor/Persona/Private/AnimTimeline/AnimTimelineClipboard.cpp
new file mode 100644
--- /dev/null
+++ b/Engine/Source/Editor/Persona/Private/AnimTimeline/AnimTimelineClipboard.cpp
@@ -0,0 +1,332 @@
+﻿// Copyright Epic Games, Inc. All Rights Reserved.
+
+#include "AnimTimelineClipboard.h"
+#include "AnimTimelineTrack.h"
+#include "AnimTimelineTrack_Curve.h"
+#include "UnrealExporter.h"
+#include "Exporters/Exporter.h"
+#include "HAL/PlatformApplicationMisc.h"
+#include "Factories.h"
+#include "Animation/AnimData/AnimDataModel.h"
+#include "Animation/Skeleton.h"
+#include "Animation/AnimSequenceBase.h"
+
+FAnimationCurveIdentifier UAnimCurveBaseCopyObject::GetAnimationCurveIdentifier() const
+{
+	FAnimationCurveIdentifier OutAnimationCurveIdentifier;
+	
+	OutAnimationCurveIdentifier.Axis = Axis;
+	OutAnimationCurveIdentifier.Channel = Channel;
+	OutAnimationCurveIdentifier.InternalName.DisplayName = DisplayName;
+	OutAnimationCurveIdentifier.InternalName.UID = UID;
+	OutAnimationCurveIdentifier.CurveType = CurveType;
+	
+	return OutAnimationCurveIdentifier;
+}
+
+bool UAnimTimelineClipboardContent::AreAllCurvesOfTrackType(const ERawCurveTrackTypes Type) const
+{
+	for (const TObjectPtr<UAnimCurveBaseCopyObject> CurveCopyObj : Curves)
+	{
+		if (CurveCopyObj->CurveType != Type)
+		{
+			return false;
+		}
+	}
+
+	return !Curves.IsEmpty();
+}
+
+bool UAnimTimelineClipboardContent::IsEmpty() const
+{
+	return Curves.IsEmpty();
+}
+
+UAnimTimelineClipboardContent* UAnimTimelineClipboardContent::Create()
+{
+	return NewObject<UAnimTimelineClipboardContent>(GetTransientPackage());
+}
+
+void FAnimTimelineClipboardUtilities::CopyContentToClipboard(UAnimTimelineClipboardContent* Content)
+{
+	// Clear the mark state for saving.
+	UnMarkAllObjects(EObjectMark(OBJECTMARK_TagExp | OBJECTMARK_TagImp));
+	
+	// Export the clipboard to text.
+	FStringOutputDevice Archive;
+	const FExportObjectInnerContext Context;
+	UExporter::ExportToOutputDevice(&Context, Content, nullptr, Archive, TEXT("copy"), 0, PPF_ExportsNotFullyQualified | PPF_Copy | PPF_Delimited, false, Content->GetOuter());
+	FPlatformApplicationMisc::ClipboardCopy(*Archive);
+}
+
+class FAnimCurvesClipboardObjectTextFactory : public FCustomizableTextObjectFactory
+{
+public:
+	
+	FAnimCurvesClipboardObjectTextFactory()
+		: FCustomizableTextObjectFactory(GWarn)
+		, AnimationCurvesClipboard(nullptr) 
+	{
+	}
+
+	UAnimTimelineClipboardContent* AnimationCurvesClipboard;
+
+protected:
+
+	virtual bool CanCreateClass(UClass* InObjectClass, bool& bOmitSubObjs) const override
+	{
+		if (InObjectClass->IsChildOf(UAnimTimelineClipboardContent::StaticClass()))
+		{
+			return true;
+		}
+		return false;
+	}
+
+
+	virtual void ProcessConstructedObject(UObject* CreatedObject) override
+	{
+		if (CreatedObject->IsA<UAnimTimelineClipboardContent>())
+		{
+			AnimationCurvesClipboard = CastChecked<UAnimTimelineClipboardContent>(CreatedObject);
+		}
+	}
+};
+
+bool FAnimTimelineClipboardUtilities::CanPasteContentToClipboard(const FString& TextToImport)
+{
+	const FAnimCurvesClipboardObjectTextFactory ClipboardContentFactory;
+	return ClipboardContentFactory.CanCreateObjectsFromText(TextToImport);
+}
+
+const UAnimTimelineClipboardContent* FAnimTimelineClipboardUtilities::GetContentFromClipboard()
+{
+	// Get the text from the clipboard.
+	FString ClipboardText;
+	FPlatformApplicationMisc::ClipboardPaste(ClipboardText);
+
+	// Try to create curves clipboard content from text.
+	if (CanPasteContentToClipboard(ClipboardText))
+	{
+		FAnimCurvesClipboardObjectTextFactory ClipboardContentFactory;
+		ClipboardContentFactory.ProcessBuffer(GetTransientPackage(), RF_Transactional, ClipboardText);
+		return ClipboardContentFactory.AnimationCurvesClipboard;
+	}
+
+	return nullptr;
+}
+
+UAnimCurveBaseCopyObject::UAnimCurveBaseCopyObject() : DisplayName(NAME_None), UID(SmartName::MaxUID), CurveType(ERawCurveTrackTypes::RCT_MAX), Channel(ETransformCurveChannel::Invalid), Axis(EVectorCurveChannel::Invalid), OriginName(NAME_None)
+{
+}
+
+bool FAnimTimelineClipboardUtilities::CopySelectedTracksToClipboard(const TSet<TSharedRef<FAnimTimelineTrack>>& CurveTracks, UAnimTimelineClipboardContent* ClipboardContent)
+{
+	check(ClipboardContent != nullptr);
+	
+	const int32 StartSize = ClipboardContent->Curves.Num();
+	
+	for (const TSharedRef<FAnimTimelineTrack>& SelectedTrack : CurveTracks)
+	{
+		SelectedTrack->Copy(ClipboardContent);
+	}
+
+	return StartSize != ClipboardContent->Curves.Num();
+}
+
+bool FAnimTimelineClipboardUtilities::CanOverwriteSelectedCurveDataFromClipboard(const TSet<TSharedRef<FAnimTimelineTrack>>& SelectedTracks)
+{
+	const UAnimTimelineClipboardContent* ClipboardContent = FAnimTimelineClipboardUtilities::GetContentFromClipboard();
+
+	// Validate clipboard
+	if (ClipboardContent && !ClipboardContent->IsEmpty())
+	{
+		// Exit on multiple curves selected
+		if (SelectedTracks.Num() != 1)
+		{
+			return false;
+		}
+
+		const TSharedRef<FAnimTimelineTrack> Track = *SelectedTracks.CreateConstIterator();
+		
+		if (Track->IsA<FAnimTimelineTrack_Curve>())
+		{
+			const FAnimTimelineTrack_Curve & CurveTrack = Track->As<FAnimTimelineTrack_Curve>();
+
+			// Get current track information
+			FSmartName CurveName;
+			ERawCurveTrackTypes CurveType;
+			int32 CurveIndex;
+			CurveTrack.GetCurveEditInfo(0 /* Unused */, CurveName, CurveType, CurveIndex);
+			
+			// Handle special case of pasting float curve into subtrack
+			const bool bIsPastingFloatTrackIntoSubtrack = (!CurveTrack.GetCurves().IsEmpty() && ClipboardContent->AreAllCurvesOfTrackType(ERawCurveTrackTypes::RCT_Float));
+			
+			// Allow pasting of content into current track
+			const bool bDoesSelectedTrackMatchClipboardTrack = ClipboardContent->AreAllCurvesOfTrackType(CurveType);
+			
+			return (bIsPastingFloatTrackIntoSubtrack || bDoesSelectedTrackMatchClipboardTrack) && ClipboardContent->Curves.Num() == 1;
+		}
+	}
+	
+	// Invalid clipboard
+	return false;
+}
+
+void FAnimTimelineClipboardUtilities::OverwriteSelectedCurveDataFromClipboard(const UAnimTimelineClipboardContent* ClipboardContent, const TSet<TSharedRef<FAnimTimelineTrack>>& CurveTracks, UAnimSequenceBase* InTargetSequence)
+{
+	check(ClipboardContent != nullptr);
+	check(InTargetSequence != nullptr)
+	IAnimationDataController& Controller = InTargetSequence->GetController();
+	
+	// TODO: Allow for multiple curve pasting. Currently only single curve selection is supported. Need to sort selected curves or create array.
+	
+	if (const TSharedRef<FAnimTimelineTrack>& Track = *CurveTracks.CreateConstIterator(); Track->IsA<FAnimTimelineTrack_Curve>())
+	{
+		const FAnimTimelineTrack_Curve & CurveTrack = Track->As<FAnimTimelineTrack_Curve>();
+
+		// Get selected curve track information
+		FSmartName SelectedCurveTrackName;
+		ERawCurveTrackTypes SelectedCurveTrackType;	// The value is supposed to always be RCT_TRANSFORM as mentioned in FAnimTimelineTrack_Curve but this is not true.
+		int32 SelectedCurveTrackIndex;
+		CurveTrack.GetCurveEditInfo(0 /* Unused */, SelectedCurveTrackName, SelectedCurveTrackType, SelectedCurveTrackIndex);
+		FAnimationCurveIdentifier SelectedCurveTrackIdentifier(SelectedCurveTrackName, SelectedCurveTrackType);
+		
+		// Special case for pasting into child tracks for vectors tracks inside a transform track
+		if (!CurveTrack.GetCurves().IsEmpty() && SelectedCurveTrackType == ERawCurveTrackTypes::RCT_Transform && ClipboardContent->AreAllCurvesOfTrackType(ERawCurveTrackTypes::RCT_Float))
+		{
+			const int VectorFloatDataIndex = SelectedCurveTrackIndex % 3;	// 0 = X, 1 = Y, 2 = Z
+			const int TransformDataIndex = SelectedCurveTrackIndex / 3;		// 0 = Translation, 1 = Rotation, 2 = Scale
+
+			// Ensure identifier information matches
+			SelectedCurveTrackIdentifier.Channel = static_cast<ETransformCurveChannel>(TransformDataIndex);
+
+			// Update animation model data
+			if (const FTransformCurve* SelectedTransform = Controller.GetModel()->FindTransformCurve(SelectedCurveTrackIdentifier))
+			{
+				FTransformCurve TempTransformCurve;
+
+				// Get new transform with incoming data
+				{
+					TempTransformCurve.CopyCurve(*SelectedTransform);
+					FVectorCurve* NewVectorCurve = TempTransformCurve.GetVectorCurveByIndex(TransformDataIndex);
+					NewVectorCurve->FloatCurves[VectorFloatDataIndex] = CastChecked<UFloatCurveCopyObject>(ClipboardContent->Curves[0])->Curve.FloatCurve;
+				}
+
+				// Update transform 
+				{
+					TArray<float> TimeKeys;
+					TArray<FTransform> TransformValues;
+					TempTransformCurve.GetKeys(TimeKeys, TransformValues);
+					Controller.SetTransformCurveKeys(SelectedCurveTrackIdentifier, TransformValues, TimeKeys);
+				}
+			}
+			
+			return;
+		}
+		
+		// Update curve data for float and transform tracks
+		if (ClipboardContent->AreAllCurvesOfTrackType(SelectedCurveTrackType))
+		{
+			if (SelectedCurveTrackType == ERawCurveTrackTypes::RCT_Float)
+			{
+				const UFloatCurveCopyObject* FloatCurve = Cast<UFloatCurveCopyObject>(ClipboardContent->Curves[0]);
+				Controller.SetCurveKeys(SelectedCurveTrackIdentifier, FloatCurve->Curve.FloatCurve.GetConstRefOfKeys());
+				Controller.SetCurveFlags(SelectedCurveTrackIdentifier, FloatCurve->Curve.GetCurveTypeFlags());
+			}
+			else if (SelectedCurveTrackType == ERawCurveTrackTypes::RCT_Transform)
+			{
+				const UTransformCurveCopyObject* TransformCurve = Cast<UTransformCurveCopyObject>(ClipboardContent->Curves[0]);
+
+				TArray<float> TimeKeys;
+				TArray<FTransform> TransformValues;
+				TransformCurve->Curve.GetKeys(TimeKeys, TransformValues); // TODO: Optimize this if possible!
+				
+				Controller.SetTransformCurveKeys(SelectedCurveTrackIdentifier, TransformValues, TimeKeys);
+				Controller.SetCurveFlags(SelectedCurveTrackIdentifier, TransformCurve->Curve.GetCurveTypeFlags());
+			}
+		}
+	}
+	
+}
+
+void FAnimTimelineClipboardUtilities::OverwriteOrAddCurvesFromClipboardContent(const UAnimTimelineClipboardContent* ClipboardContent, UAnimSequenceBase* InTargetSequence)
+{
+	check(InTargetSequence != nullptr)
+	check(ClipboardContent != nullptr)
+	
+	USkeleton * Skeleton = InTargetSequence->GetSkeleton();
+	check(Skeleton != nullptr)
+	IAnimationDataController& Controller = InTargetSequence->GetController();
+	
+	for (TObjectPtr<UAnimCurveBaseCopyObject> InCurveCopyObj : ClipboardContent->Curves)
+	{
+		// Exit on types not supported by animation controller
+		if (InCurveCopyObj->CurveType != ERawCurveTrackTypes::RCT_Float && InCurveCopyObj->CurveType != ERawCurveTrackTypes::RCT_Transform)
+		{
+			UE_LOG(LogAnimation, Warning, TEXT("Attempting to paste curve { %s } of unsupported type { %s } by animation controller. "), *InCurveCopyObj->DisplayName.ToString(), *UEnum::GetValueAsString(InCurveCopyObj->CurveType));
+			continue;
+		}
+		
+		// Ensure we are not pasting data into the same curve we copying data from
+		if (InCurveCopyObj->OriginName.Compare(Controller.GetModel()->GetOuter()->GetFName()) == 0 && Controller.GetModel()->FindCurve(InCurveCopyObj->GetAnimationCurveIdentifier()))
+		{
+			// Do nothing and exit.
+			UE_LOG(LogAnimation, Warning, TEXT("Attempting to paste already existing curve: %s"), *InCurveCopyObj->DisplayName.ToString());
+		}
+		else
+		{
+			FSmartName CurveTrackName(InCurveCopyObj->DisplayName, SmartName::MaxUID);
+			FAnimationCurveIdentifier CurveIdentifier = InCurveCopyObj->GetAnimationCurveIdentifier();
+			const FName & CurveContainerName = InCurveCopyObj->CurveType == ERawCurveTrackTypes::RCT_Float ? USkeleton::AnimCurveMappingName : USkeleton::AnimTrackCurveMappingName;
+			
+			const SmartName::UID_Type FoundCurveUID = Skeleton->GetUIDByName(CurveContainerName, InCurveCopyObj->DisplayName);
+			if (FoundCurveUID != SmartName::MaxUID)
+			{
+				// Update identifier since skeleton does has curve name in the mapping container
+				CurveTrackName.UID = FoundCurveUID;
+				CurveIdentifier.InternalName.UID = FoundCurveUID;
+			}
+			else
+			{
+				// Add new curve name based on clipboard information
+				Skeleton->AddSmartNameAndModify(CurveContainerName, InCurveCopyObj->DisplayName, CurveTrackName);
+					
+				// Ensure that curve smart name exists
+				FSmartName NewName;
+				ensureAlways(Skeleton->GetSmartNameByUID(CurveContainerName, CurveTrackName.UID, NewName));
+
+				// Update identifier since skeleton does has curve name in the mapping container
+				CurveIdentifier.InternalName = NewName;
+			}
+			
+			// Create curve with clipboard information if needed
+			if (!Controller.GetModel()->FindCurve(CurveIdentifier))
+			{
+				Controller.AddCurve(CurveIdentifier);
+			}
+				
+			// Update data for the curve
+			if (CurveTrackName.IsValid()) 
+			{
+				if (InCurveCopyObj->CurveType == ERawCurveTrackTypes::RCT_Float)
+				{
+					const UFloatCurveCopyObject* FloatCopyObject = Cast<UFloatCurveCopyObject>(InCurveCopyObj);
+					
+					Controller.SetCurveKeys(CurveIdentifier, FloatCopyObject->Curve.FloatCurve.GetConstRefOfKeys());
+					Controller.SetCurveFlags(CurveIdentifier, FloatCopyObject->Curve.GetCurveTypeFlags());
+				}
+				else if (InCurveCopyObj->CurveType == ERawCurveTrackTypes::RCT_Transform)
+				{
+					const UTransformCurveCopyObject* TransformCopyObject = Cast<UTransformCurveCopyObject>(InCurveCopyObj);
+					
+					TArray<float> TimeKeys;
+					TArray<FTransform> TransformValues;
+					TransformCopyObject->Curve.GetKeys(TimeKeys, TransformValues); // TODO: Optimize this if possible!
+		
+					Controller.SetTransformCurveKeys(CurveIdentifier, TransformValues, TimeKeys);
+					Controller.SetCurveFlags(CurveIdentifier, TransformCopyObject->Curve.GetCurveTypeFlags());
+				}
+			}
+		}
+	}
+}
diff --git a/Engine/Source/Editor/Persona/Private/AnimTimeline/AnimTimelineClipboard.h b/Engine/Source/Editor/Persona/Private/AnimTimeline/AnimTimelineClipboard.h
new file mode 100644
--- /dev/null
+++ b/Engine/Source/Editor/Persona/Private/AnimTimeline/AnimTimelineClipboard.h
@@ -0,0 +1,181 @@
+﻿// Copyright Epic Games, Inc. All Rights Reserved.
+
+#pragma once
+
+#include "CoreMinimal.h"
+#include "Animation/AnimCurveTypes.h"
+#include "Animation/AnimTypes.h"
+#include "AnimTimelineTrack.h"
+#include "Animation/AnimData/CurveIdentifier.h"
+#include "AnimTimelineClipboard.generated.h"
+
+class UAnimSequenceBase;
+
+///////////////////////////////////////////////////////////////////////////////
+// Animation Curve Copy Objects
+///////////////////////////////////////////////////////////////////////////////
+
+/** Object used to export AnimCurves to clipboard */
+UCLASS(Transient)
+class UAnimCurveBaseCopyObject : public UObject
+{
+	GENERATED_BODY()
+	
+public:
+
+	UAnimCurveBaseCopyObject();
+
+	// Store information for FAnimationCurveIdentifier.
+
+	/** We declare individual properties of FSmartName since only FSmartName.DisplayName is serializable (has UPROPERTY macro)
+	 *  and FSmartName.UID is not.
+	 */
+	
+	// SmartName Start
+	
+	UPROPERTY()
+	FName DisplayName;
+
+	UPROPERTY()
+	uint16	UID;	// We can't use SmartName::UID_Type typedef since the SmartName struct is not properly prefixed.
+	
+	// SmartName End
+	
+	UPROPERTY()
+	ERawCurveTrackTypes CurveType;
+
+	UPROPERTY()
+	ETransformCurveChannel Channel;
+
+	UPROPERTY()
+	EVectorCurveChannel Axis;
+
+	// Context information
+
+	/** This curve's data owner name (ex. UAnimSequenceBase Name), if any. Used internally to check if the user is attempting to paste curve data into source curve. */
+	UPROPERTY()
+	FName OriginName;
+	
+	// Helper methods
+
+	/** Get the information that identifies this curve */
+	FAnimationCurveIdentifier GetAnimationCurveIdentifier() const;
+
+	/**
+	 * Create a copy object used to export  to clipboard
+	 * @tparam T Curve copy object type to create
+	 * @return Copy object for curve
+	 */
+	template<typename T>
+	static T* Create()
+	{
+		return NewObject<T>(GetTransientPackage(), T::StaticClass(), NAME_None, RF_Transient);
+	}
+};
+
+UCLASS(Transient)
+class UFloatCurveCopyObject : public UAnimCurveBaseCopyObject
+{
+	GENERATED_BODY()
+	
+public:
+	
+	UPROPERTY()
+	FFloatCurve Curve;
+};
+
+UCLASS(Transient)
+class UTransformCurveCopyObject : public UAnimCurveBaseCopyObject
+{
+	GENERATED_BODY()
+	
+public:
+	
+	UPROPERTY()
+	FTransformCurve Curve;
+};
+
+UCLASS(Transient)
+class UVectorCurveCopyObject : public UAnimCurveBaseCopyObject
+{
+	GENERATED_BODY()
+	
+public:
+	
+	UPROPERTY()
+	FVectorCurve Curve;
+};
+
+///////////////////////////////////////////////////////////////////////////////
+// Animation Timeline Clipboard
+///////////////////////////////////////////////////////////////////////////////
+
+UCLASS(Transient)
+class UAnimTimelineClipboardContent : public UObject
+{
+public:
+	GENERATED_BODY()
+
+	/** Animation Timeline Editor Curves */
+	UPROPERTY(Export)
+	TArray<TObjectPtr<UAnimCurveBaseCopyObject>> Curves;
+
+	// Utility functions
+	
+	/** Determine if all curve store in clipboard match the given Curve Track Type. */
+	bool AreAllCurvesOfTrackType(ERawCurveTrackTypes Type) const;
+
+	/** Determine if clipboard content is empty. */
+	bool IsEmpty() const;
+	
+	/** Create new clipboard content object  */
+	static UAnimTimelineClipboardContent* Create();
+};
+
+struct FAnimTimelineClipboardUtilities
+{
+	// Clipboard Content
+
+	/** Copy anim timeline clipboard content to system's clipboard */
+	static void CopyContentToClipboard(UAnimTimelineClipboardContent* Content);
+
+	/** Determine if a anim timeline clipboard content can be constructed from system's clipboard */
+	static bool CanPasteContentToClipboard(const FString& TextToImport);
+
+	/** Attempt to get the anim timeline clipboard content from system's clipboard. */
+	static const UAnimTimelineClipboardContent* GetContentFromClipboard();
+
+	// Tracks
+
+	/**
+	 * Copy all the selected tracks into clipboard content.
+	 * @param CurveTracks Selected Tracks in the Animation Timeline Editor
+	 * @param ClipboardContent Clipboard object where data will be copied to
+	 * @return True if any data was copied
+	 */
+	static bool CopySelectedTracksToClipboard(const TSet<TSharedRef<FAnimTimelineTrack>>& CurveTracks, UAnimTimelineClipboardContent* ClipboardContent);
+
+	// Curve Copying and Pasting
+
+	/**
+	 * Determine if curve data can be pasted into selected curve track data.  
+	 * @param SelectedTracks Selected Tracks in the Animation Timeline Editor
+	 * @return Whether pasting is possible
+	 */
+	static bool CanOverwriteSelectedCurveDataFromClipboard(const TSet<TSharedRef<FAnimTimelineTrack>>& SelectedTracks);
+
+	/**
+	 * Overwrite curve data of the selected curve track with data from the clipboard without changing selected curve identifier.
+	 * @param ClipboardContent Clipboard containing the curve(s) to be pasted
+	 * @param CurveTracks Selected Tracks in the Animation Timeline Editor
+	 * @param InTargetSequence Base Animation Sequence that will be modified
+	 */
+	static void OverwriteSelectedCurveDataFromClipboard(const UAnimTimelineClipboardContent* ClipboardContent, const TSet<TSharedRef<FAnimTimelineTrack>>& CurveTracks, UAnimSequenceBase* InTargetSequence);
+
+	/**
+	 * Add or overwrite, if identifiers collide, all the curves from the clipboard object to the given Skeleton's and Controller's UAnimSequenceBase.
+	 * @param ClipboardContent Clipboard containing the curves to be pasted
+	 * @param InTargetSequence Base Animation Sequence that will be modified
+	 */
+	static void OverwriteOrAddCurvesFromClipboardContent(const UAnimTimelineClipboardContent* ClipboardContent, UAnimSequenceBase* InTargetSequence);
+};
\ No newline at end of file
diff --git a/Engine/Source/Editor/Persona/Private/AnimTimeline/AnimTimelineTrack.h b/Engine/Source/Editor/Persona/Private/AnimTimeline/AnimTimelineTrack.h
--- a/Engine/Source/Editor/Persona/Private/AnimTimeline/AnimTimelineTrack.h
+++ b/Engine/Source/Editor/Persona/Private/AnimTimeline/AnimTimelineTrack.h
@@ -13,6 +13,7 @@ class FAnimModel;
 class FMenuBuilder;
 class SBorder;
 class SHorizontalBox;
+class UAnimTimelineClipboardContent;
 
 /**
  * Structure used to define padding for a track
@@ -183,6 +184,12 @@ public:
 	/** Get whether this track supports filtering in the outliner */
 	virtual bool SupportsFiltering() const { return true; }
 
+	/** Get whether this track supports copying to the clipboard */
+	virtual bool SupportsCopy() const { return false; }
+
+	/** Copy track information to clipboard */
+	virtual void Copy(UAnimTimelineClipboardContent* InOutClipboard) const {};
+	
 	/** Get whether this track can be renamed */
 	virtual bool CanRename() const { return false; }
 
diff --git a/Engine/Source/Editor/Persona/Private/AnimTimeline/AnimTimelineTrack_Curve.cpp b/Engine/Source/Editor/Persona/Private/AnimTimeline/AnimTimelineTrack_Curve.cpp
--- a/Engine/Source/Editor/Persona/Private/AnimTimeline/AnimTimelineTrack_Curve.cpp
+++ b/Engine/Source/Editor/Persona/Private/AnimTimeline/AnimTimelineTrack_Curve.cpp
@@ -20,7 +20,9 @@
 #include "IPersonaPreviewScene.h"
 #include "Animation/DebugSkelMeshComponent.h"
 #include "AnimPreviewInstance.h"
+#include "AnimTimelineClipboard.h"
 #include "SAnimSequenceCurveEditor.h"
+#include "Framework/Commands/GenericCommands.h"
 
 #define LOCTEXT_NAMESPACE "FAnimTimelineTrack_Curve"
 
@@ -320,18 +322,60 @@ void FAnimTimelineTrack_Curve::AddToContextMenu(FMenuBuilder& InMenuBuilder, TSe
 {
 	if(!InOutExistingMenuTypes.Contains(FAnimTimelineTrack_Curve::GetTypeName()))
 	{
-		InMenuBuilder.BeginSection("Curve", LOCTEXT("CurveMenuSection", "Curve"));
+		InMenuBuilder.BeginSection("EditCurve", LOCTEXT("CurveEditMenuSection", "Curve Edit"));
 		{
+			InMenuBuilder.AddMenuEntry(FAnimSequenceTimelineCommands::Get().PasteDataIntoCurve);
+		}
+		InMenuBuilder.EndSection();
+		
+		InMenuBuilder.BeginSection("EditSelection", LOCTEXT("CurveSelectionEditMenuSection", "Selection Edit"));
+		{
+			InMenuBuilder.AddMenuEntry(FGenericCommands::Get().Cut);
+			InMenuBuilder.AddMenuEntry(FGenericCommands::Get().Copy);
+			InMenuBuilder.AddMenuEntry(FGenericCommands::Get().Paste);
+			InMenuBuilder.AddMenuEntry(FGenericCommands::Get().Delete);
+
+			InMenuBuilder.AddSeparator();
+			
 			InMenuBuilder.AddMenuEntry(FAnimSequenceTimelineCommands::Get().EditSelectedCurves);
-			InMenuBuilder.AddMenuEntry(FAnimSequenceTimelineCommands::Get().RemoveSelectedCurves);
 			InMenuBuilder.AddMenuEntry(FAnimSequenceTimelineCommands::Get().CopySelectedCurveNames);
 		}
 		InMenuBuilder.EndSection();
-
+		
 		InOutExistingMenuTypes.Add(FAnimTimelineTrack_Curve::GetTypeName());
 	}
 }
 
+void FAnimTimelineTrack_Curve::Copy(UAnimTimelineClipboardContent* InOutClipboard) const
+{
+	check(InOutClipboard != nullptr)
+	
+	if (Curves.Num() == 1)
+	{
+		UFloatCurveCopyObject * CopyableCurve = UAnimCurveBaseCopyObject::Create<UFloatCurveCopyObject>();
+		const FRichCurve* InCurve = Curves[0];
+
+		// Copy raw curve data
+		CopyableCurve->Curve.Name = { FName(FullCurveName.ToString()), SmartName::MaxUID };
+		CopyableCurve->Curve.FloatCurve = *InCurve;
+		CopyableCurve->Curve.SetCurveTypeFlags(AACF_Editable);
+
+		// Copy curve identifier data
+		CopyableCurve->DisplayName = CopyableCurve->Curve.Name.DisplayName;
+		CopyableCurve->UID = CopyableCurve->Curve.Name.UID;
+		CopyableCurve->CurveType = ERawCurveTrackTypes::RCT_Float;
+
+		// Origin data
+		CopyableCurve->OriginName = GetModel()->GetAnimSequenceBase()->GetFName();
+			
+		InOutClipboard->Curves.Add(CopyableCurve);
+	}
+	else
+	{
+		UE_LOG(LogAnimation, Warning, TEXT("Copying multiple curves from a FAnimTimelineTrack_Curve not supported. Curve: %s"), *FullCurveName.ToString())
+	}
+}
+
 FReply FAnimTimelineTrack_Curve::HandleMouseButtonUp(const FGeometry& InGeometry, const FPointerEvent& InPointerEvent)
 {
 	if(InPointerEvent.GetEffectingButton() == EKeys::RightMouseButton)
diff --git a/Engine/Source/Editor/Persona/Private/AnimTimeline/AnimTimelineTrack_Curve.h b/Engine/Source/Editor/Persona/Private/AnimTimeline/AnimTimelineTrack_Curve.h
--- a/Engine/Source/Editor/Persona/Private/AnimTimeline/AnimTimelineTrack_Curve.h
+++ b/Engine/Source/Editor/Persona/Private/AnimTimeline/AnimTimelineTrack_Curve.h
@@ -28,13 +28,15 @@ public:
 	virtual TSharedRef<SWidget> GenerateContainerWidgetForOutliner(const TSharedRef<SAnimOutlinerItem>& InRow) override;
 	virtual bool SupportsSelection() const override { return true; }
 	virtual void AddToContextMenu(FMenuBuilder& InMenuBuilder, TSet<FName>& InOutExistingMenuTypes) const override;
+	virtual bool SupportsCopy() const override { return true; }
+	virtual void Copy(UAnimTimelineClipboardContent* InOutClipboard) const override;
 	
 	/** FEditorUndoClient interface */
 	virtual void PostUndo(bool bSuccess) override { PostUndoRedo(); }
 	virtual void PostRedo(bool bSuccess) override { PostUndoRedo(); }
 
 	/** Access the curve we are editing */
-	const TArray<const FRichCurve*>& GetCurves() { return Curves; }
+	const TArray<const FRichCurve*>& GetCurves() const { return Curves; }
 
 	/** Get a color for the curve widget (and edited curves) */
 	virtual FLinearColor GetCurveColor(int32 InCurveIndex) const;
diff --git a/Engine/Source/Editor/Persona/Private/AnimTimeline/AnimTimelineTrack_FloatCurve.cpp b/Engine/Source/Editor/Persona/Private/AnimTimeline/AnimTimelineTrack_FloatCurve.cpp
--- a/Engine/Source/Editor/Persona/Private/AnimTimeline/AnimTimelineTrack_FloatCurve.cpp
+++ b/Engine/Source/Editor/Persona/Private/AnimTimeline/AnimTimelineTrack_FloatCurve.cpp
@@ -19,6 +19,7 @@
 #include "Widgets/Input/SComboButton.h"
 #include "Widgets/Layout/SBox.h"
 #include "AnimModel_AnimSequenceBase.h"
+#include "AnimTimelineClipboard.h"
 #include "Animation/AnimData/AnimDataModel.h"
 #include "SAnimOutlinerItem.h"
 
@@ -264,6 +265,30 @@ FText FAnimTimelineTrack_FloatCurve::GetLabel() const
 	return FAnimTimelineTrack_FloatCurve::GetFloatCurveName(GetModel(), FloatCurve->Name);
 }
 
+void FAnimTimelineTrack_FloatCurve::Copy(UAnimTimelineClipboardContent* InOutClipboard) const
+{
+	check(InOutClipboard != nullptr)
+	
+	UFloatCurveCopyObject * CopyableCurve = UAnimCurveBaseCopyObject::Create<UFloatCurveCopyObject>();
+
+	// Copy raw curve data
+	CopyableCurve->Curve.Name = FloatCurve->Name;
+	CopyableCurve->Curve.SetCurveTypeFlags(FloatCurve->GetCurveTypeFlags());
+	CopyableCurve->Curve.CopyCurve(*FloatCurve);
+
+	// Copy curve identifier data
+	CopyableCurve->DisplayName = CurveName.DisplayName;
+	CopyableCurve->UID = CurveName.UID;
+	CopyableCurve->CurveType = ERawCurveTrackTypes::RCT_Float;
+	CopyableCurve->Channel = CurveId.Channel;
+	CopyableCurve->Axis = CurveId.Axis;
+
+	// Origin data
+	CopyableCurve->OriginName = GetModel()->GetAnimSequenceBase()->GetFName();
+	
+	InOutClipboard->Curves.Add(CopyableCurve);
+}
+
 FText FAnimTimelineTrack_FloatCurve::GetFloatCurveName(const TSharedRef<FAnimModel>& InModel, const FSmartName& InSmartName)
 {
 	const FSmartNameMapping* NameMapping = InModel->GetAnimSequenceBase()->GetSkeleton()->GetSmartNameContainer(USkeleton::AnimCurveMappingName);
diff --git a/Engine/Source/Editor/Persona/Private/AnimTimeline/AnimTimelineTrack_FloatCurve.h b/Engine/Source/Editor/Persona/Private/AnimTimeline/AnimTimelineTrack_FloatCurve.h
--- a/Engine/Source/Editor/Persona/Private/AnimTimeline/AnimTimelineTrack_FloatCurve.h
+++ b/Engine/Source/Editor/Persona/Private/AnimTimeline/AnimTimelineTrack_FloatCurve.h
@@ -26,7 +26,9 @@ public:
 	virtual void AddCurveTrackButton(TSharedPtr<SHorizontalBox> InnerHorizontalBox) override;
 	virtual FLinearColor GetCurveColor(int32 InCurveIndex) const override;
 	virtual void GetCurveEditInfo(int32 InCurveIndex, FSmartName& OutName, ERawCurveTrackTypes& OutType, int32& OutCurveIndex) const override;
-
+	virtual bool SupportsCopy() const override { return true; }
+	virtual void Copy(UAnimTimelineClipboardContent* InOutClipboard) const override;
+	
 	/** Access the curve we are editing */
 	const FFloatCurve* GetFloatCurve() { return FloatCurve; }
 
diff --git a/Engine/Source/Editor/Persona/Private/AnimTimeline/AnimTimelineTrack_TransformCurve.cpp b/Engine/Source/Editor/Persona/Private/AnimTimeline/AnimTimelineTrack_TransformCurve.cpp
--- a/Engine/Source/Editor/Persona/Private/AnimTimeline/AnimTimelineTrack_TransformCurve.cpp
+++ b/Engine/Source/Editor/Persona/Private/AnimTimeline/AnimTimelineTrack_TransformCurve.cpp
@@ -12,6 +12,7 @@
 #include "Animation/DebugSkelMeshComponent.h"
 #include "AnimPreviewInstance.h"
 #include "AnimModel_AnimSequenceBase.h"
+#include "AnimTimelineClipboard.h"
 
 #define LOCTEXT_NAMESPACE "FAnimTimelineTrack_TransformCurve"
 
@@ -66,6 +67,30 @@ FText FAnimTimelineTrack_TransformCurve::GetFullCurveName(int32 InCurveIndex) co
 	return FText::Format(LOCTEXT("TransformVectorFormat", "{0}.{1}"), FullCurveName, TrackNames[InCurveIndex]);
 }
 
+void FAnimTimelineTrack_TransformCurve::Copy(UAnimTimelineClipboardContent* InOutClipboard) const
+{
+	check(InOutClipboard != nullptr)
+	
+	UTransformCurveCopyObject * CopyableCurve = UAnimCurveBaseCopyObject::Create<UTransformCurveCopyObject>();
+
+	// Copy raw curve data
+	CopyableCurve->Curve.Name = TransformCurve->Name;
+	CopyableCurve->Curve.SetCurveTypeFlags(TransformCurve->GetCurveTypeFlags());
+	CopyableCurve->Curve.CopyCurve(*TransformCurve);
+
+	// Copy curve identifier data
+	CopyableCurve->DisplayName = CurveName.DisplayName;
+	CopyableCurve->UID = CurveName.UID;
+	CopyableCurve->CurveType = ERawCurveTrackTypes::RCT_Transform;
+	CopyableCurve->Channel = CurveId.Channel;
+	CopyableCurve->Axis = CurveId.Axis;
+
+	// Origin data
+	CopyableCurve->OriginName = GetModel()->GetAnimSequenceBase()->GetFName();
+	
+	InOutClipboard->Curves.Add(CopyableCurve);
+}
+
 FText FAnimTimelineTrack_TransformCurve::GetTransformCurveName(const TSharedRef<FAnimModel>& InModel, const FSmartName& InSmartName)
 {
 	const FSmartNameMapping* NameMapping = InModel->GetAnimSequenceBase()->GetSkeleton()->GetSmartNameContainer(USkeleton::AnimTrackCurveMappingName);
diff --git a/Engine/Source/Editor/Persona/Private/AnimTimeline/AnimTimelineTrack_TransformCurve.h b/Engine/Source/Editor/Persona/Private/AnimTimeline/AnimTimelineTrack_TransformCurve.h
--- a/Engine/Source/Editor/Persona/Private/AnimTimeline/AnimTimelineTrack_TransformCurve.h
+++ b/Engine/Source/Editor/Persona/Private/AnimTimeline/AnimTimelineTrack_TransformCurve.h
@@ -21,7 +21,9 @@ public:
 	virtual TSharedRef<SWidget> BuildCurveTrackMenu() override;
 	virtual bool CanEditCurve(int32 InCurveIndex) const override { return true; }
 	virtual void GetCurveEditInfo(int32 InCurveIndex, FSmartName& OutName, ERawCurveTrackTypes& OutType, int32& OutCurveIndex) const override;
-
+	virtual bool SupportsCopy() const override { return true; }
+	virtual void Copy(UAnimTimelineClipboardContent* InOutClipboard) const override;
+	
 	/** Access the curve we are editing */
 	const FTransformCurve& GetTransformCurve() { return *TransformCurve; }
 
diff --git a/Engine/Source/Editor/Persona/Private/AnimTimeline/SAnimTimeline.cpp b/Engine/Source/Editor/Persona/Private/AnimTimeline/SAnimTimeline.cpp
--- a/Engine/Source/Editor/Persona/Private/AnimTimeline/SAnimTimeline.cpp
+++ b/Engine/Source/Editor/Persona/Private/AnimTimeline/SAnimTimeline.cpp
@@ -6,6 +6,7 @@
 #include "AnimModel.h"
 #include "SAnimOutliner.h"
 #include "SAnimTrackArea.h"
+#include "SAnimTrack.h"
 #include "Widgets/Layout/SSplitter.h"
 #include "SAnimTimelineOverlay.h"
 #include "SAnimTimelineSplitterOverlay.h"
@@ -30,12 +31,14 @@
 #include "ScopedTransaction.h"
 #include "Widgets/Input/STextEntryPopup.h"
 #include "AnimSequenceTimelineCommands.h"
+#include "AnimTimelineTrack_Curve.h"
 #include "Widgets/Input/SSpinBox.h"
 #include "SAnimTimelineTransportControls.h"
 #include "Animation/AnimSequence.h"
 #include "Animation/AnimData/AnimDataModel.h"
 #include "Runtime/Engine/Classes/Animation/AnimSequenceHelpers.h"
 #include "Widgets/Layout/SScrollBox.h"
+#include "Framework/Commands/GenericCommands.h"
 
 #define LOCTEXT_NAMESPACE "SAnimTimeline"
 
@@ -359,7 +362,13 @@ FReply SAnimTimeline::OnMouseButtonUp(const FGeometry& MyGeometry, const FPointe
 
 		const bool bCloseAfterSelection = true;
 		FMenuBuilder MenuBuilder(bCloseAfterSelection, Model.Pin()->GetCommandList());
-
+		
+		MenuBuilder.BeginSection("SelectionEdit", LOCTEXT("TimelineSelectionEditSection", "Selection Edit"));
+		{
+			MenuBuilder.AddMenuEntry(FGenericCommands::Get().Paste);
+		}
+		MenuBuilder.EndSection();
+		
 		MenuBuilder.BeginSection("SnapOptions", LOCTEXT("SnapOptions", "Snapping"));
 		{
 			MenuBuilder.AddMenuEntry(FAnimSequenceTimelineCommands::Get().SnapToFrames);
