Sequencer: Allow reusing an animation asset when baking to it.

Added an argument to ExecuteNewAnimAsset, bAllowReplaceExisting, which when true will look for an existing asset and prompt the user whether to use that asset as opposed to calling create asset. When baking to an animation asset, this allows the user to bake to an existing animation asset. Previously, it would try to create the same named asset but would fail because the asset may be actively referenced.

diff --git a/Engine/Source/Editor/MovieSceneTools/Private/TrackEditors/SkeletalAnimationTrackEditor.cpp b/Engine/Source/Editor/MovieSceneTools/Private/TrackEditors/SkeletalAnimationTrackEditor.cpp
--- a/Engine/Source/Editor/MovieSceneTools/Private/TrackEditors/SkeletalAnimationTrackEditor.cpp
+++ b/Engine/Source/Editor/MovieSceneTools/Private/TrackEditors/SkeletalAnimationTrackEditor.cpp
@@ -1083,7 +1083,7 @@ void FSkeletalAnimationTrackEditor::HandleCreatePoseAsset(FGuid InObjectBinding)
 	{
 		TArray<TWeakObjectPtr<UObject>> Skeletons;
 		Skeletons.Add(Skeleton);
-		AnimationEditorUtils::ExecuteNewAnimAsset<UPoseAssetFactory, UPoseAsset>(Skeletons, FString("_PoseAsset"), FAnimAssetCreated::CreateSP(this, &FSkeletalAnimationTrackEditor::CreatePoseAsset, InObjectBinding), false);
+		AnimationEditorUtils::ExecuteNewAnimAsset<UPoseAssetFactory, UPoseAsset>(Skeletons, FString("_PoseAsset"), FAnimAssetCreated::CreateSP(this, &FSkeletalAnimationTrackEditor::CreatePoseAsset, InObjectBinding), false, false);
 	}
 }
 
@@ -1559,7 +1559,10 @@ void FSkeletalAnimationTrackEditor::HandleCreateAnimationSequence(USkeletalMeshC
 		{
 			Skels.Add(Skeleton);
 		}
-		AnimationEditorUtils::ExecuteNewAnimAsset<UAnimSequenceFactory, UAnimSequence>(Skels, FString("_Sequence"), FAnimAssetCreated::CreateSP(this, &FSkeletalAnimationTrackEditor::CreateAnimationSequence, SkelMeshComp,  Binding,bCreateSoftLink), false);
+	
+		const bool bDoNotShowNameDialog = false;
+		const bool bAllowReplaceExisting = true;
+		AnimationEditorUtils::ExecuteNewAnimAsset<UAnimSequenceFactory, UAnimSequence>(Skels, FString("_Sequence"), FAnimAssetCreated::CreateSP(this, &FSkeletalAnimationTrackEditor::CreateAnimationSequence, SkelMeshComp, Binding, bCreateSoftLink), bDoNotShowNameDialog, bAllowReplaceExisting);
 	}
 }
 
diff --git a/Engine/Source/Editor/Persona/Private/PersonaModule.cpp b/Engine/Source/Editor/Persona/Private/PersonaModule.cpp
--- a/Engine/Source/Editor/Persona/Private/PersonaModule.cpp
+++ b/Engine/Source/Editor/Persona/Private/PersonaModule.cpp
@@ -1171,6 +1171,9 @@ void FPersonaModule::FillCreateAnimationMenu(FMenuBuilder& MenuBuilder, TWeakPtr
 		Objects.Add(PersonaToolkit->GetSkeleton());
 	}
 
+	const bool bDoNotShowNameDialog = false;
+	const bool bAllowReplaceExisting = false;
+
 	// create rig
 	MenuBuilder.BeginSection("CreateAnimationSubMenu", LOCTEXT("CreateAnimationSubMenuHeading", "Create Animation"));
 	{
@@ -1179,7 +1182,7 @@ void FPersonaModule::FillCreateAnimationMenu(FMenuBuilder& MenuBuilder, TWeakPtr
 			LOCTEXT("CreateAnimation_RefPose_Tooltip", "Create Animation from reference pose."),
 			FSlateIcon(),
 			FUIAction(
-				FExecuteAction::CreateStatic(&AnimationEditorUtils::ExecuteNewAnimAsset<UAnimSequenceFactory, UAnimSequence>, Objects, FString("_Sequence"), FAnimAssetCreated::CreateRaw(const_cast<FPersonaModule*>(this), &FPersonaModule::CreateAnimation, EPoseSourceOption::ReferencePose, InWeakPersonaToolkit), false),
+				FExecuteAction::CreateStatic(&AnimationEditorUtils::ExecuteNewAnimAsset<UAnimSequenceFactory, UAnimSequence>, Objects, FString("_Sequence"), FAnimAssetCreated::CreateRaw(const_cast<FPersonaModule*>(this), &FPersonaModule::CreateAnimation, EPoseSourceOption::ReferencePose, InWeakPersonaToolkit), bDoNotShowNameDialog, bAllowReplaceExisting),
 				FCanExecuteAction()
 				)
 			);
@@ -1189,7 +1192,7 @@ void FPersonaModule::FillCreateAnimationMenu(FMenuBuilder& MenuBuilder, TWeakPtr
 			LOCTEXT("CreateAnimation_CurrentPose_Tooltip", "Create Animation from current pose."),
 			FSlateIcon(),
 			FUIAction(
-				FExecuteAction::CreateStatic(&AnimationEditorUtils::ExecuteNewAnimAsset<UAnimSequenceFactory, UAnimSequence>, Objects, FString("_Sequence"), FAnimAssetCreated::CreateRaw(const_cast<FPersonaModule*>(this), &FPersonaModule::CreateAnimation, EPoseSourceOption::CurrentPose, InWeakPersonaToolkit), false),
+				FExecuteAction::CreateStatic(&AnimationEditorUtils::ExecuteNewAnimAsset<UAnimSequenceFactory, UAnimSequence>, Objects, FString("_Sequence"), FAnimAssetCreated::CreateRaw(const_cast<FPersonaModule*>(this), &FPersonaModule::CreateAnimation, EPoseSourceOption::CurrentPose, InWeakPersonaToolkit), bDoNotShowNameDialog, bAllowReplaceExisting),
 				FCanExecuteAction()
 				)
 			);
@@ -1222,6 +1225,9 @@ void FPersonaModule::FillCreateAnimationFromCurrentAnimationMenu(FMenuBuilder& M
 		Objects.Add(PersonaToolkit->GetSkeleton());
 	}
 
+	const bool bDoNotShowNameDialog = false;
+	const bool bAllowReplaceExisting = false;
+
 	// create rig
 	MenuBuilder.BeginSection("CreateAnimationSubMenu", LOCTEXT("CreateAnimationFromCurrentAnimationSubmenuHeading", "Create Animation"));
 	{
@@ -1230,7 +1236,7 @@ void FPersonaModule::FillCreateAnimationFromCurrentAnimationMenu(FMenuBuilder& M
 			LOCTEXT("CreateAnimation_CurrentAnimation_AnimData_Tooltip", "Create Animation from Animation Source Data."),
 			FSlateIcon(),
 			FUIAction(
-				FExecuteAction::CreateStatic(&AnimationEditorUtils::ExecuteNewAnimAsset<UAnimSequenceFactory, UAnimSequence>, Objects, FString("_Sequence"), FAnimAssetCreated::CreateRaw(const_cast<FPersonaModule*>(this), &FPersonaModule::CreateAnimation, EPoseSourceOption::CurrentAnimation_AnimData, InWeakPersonaToolkit), false)
+				FExecuteAction::CreateStatic(&AnimationEditorUtils::ExecuteNewAnimAsset<UAnimSequenceFactory, UAnimSequence>, Objects, FString("_Sequence"), FAnimAssetCreated::CreateRaw(const_cast<FPersonaModule*>(this), &FPersonaModule::CreateAnimation, EPoseSourceOption::CurrentAnimation_AnimData, InWeakPersonaToolkit), bDoNotShowNameDialog, bAllowReplaceExisting)
 			)
 		);
 
@@ -1239,7 +1245,7 @@ void FPersonaModule::FillCreateAnimationFromCurrentAnimationMenu(FMenuBuilder& M
 			LOCTEXT("CreateAnimation_CurrentAnimation_PreviewMesh_Tooltip", "Create Animation by playing on the Current Preview Mesh, including Retargeting, Post Process Graph, or anything you see on the preview mesh."),
 			FSlateIcon(),
 			FUIAction(
-				FExecuteAction::CreateStatic(&AnimationEditorUtils::ExecuteNewAnimAsset<UAnimSequenceFactory, UAnimSequence>, Objects, FString("_Sequence"), FAnimAssetCreated::CreateRaw(const_cast<FPersonaModule*>(this), &FPersonaModule::CreateAnimation, EPoseSourceOption::CurrentAnimation_PreviewMesh, InWeakPersonaToolkit), false)
+				FExecuteAction::CreateStatic(&AnimationEditorUtils::ExecuteNewAnimAsset<UAnimSequenceFactory, UAnimSequence>, Objects, FString("_Sequence"), FAnimAssetCreated::CreateRaw(const_cast<FPersonaModule*>(this), &FPersonaModule::CreateAnimation, EPoseSourceOption::CurrentAnimation_PreviewMesh, InWeakPersonaToolkit), bDoNotShowNameDialog, bAllowReplaceExisting)
 			)
 		);
 	}
@@ -1260,6 +1266,9 @@ void FPersonaModule::FillCreatePoseAssetMenu(FMenuBuilder& MenuBuilder, TWeakPtr
 		Objects.Add(PersonaToolkit->GetSkeleton());
 	}
 
+	const bool bDoNotShowNameDialog = false;
+	const bool bAllowReplaceExisting = false;
+
 	// create rig
 	MenuBuilder.BeginSection("CreatePoseAssetSubMenu", LOCTEXT("CreatePoseAssetSubMenuHeading", "Create PoseAsset"));
 	{
@@ -1268,7 +1277,7 @@ void FPersonaModule::FillCreatePoseAssetMenu(FMenuBuilder& MenuBuilder, TWeakPtr
 			LOCTEXT("CreatePoseAsset_CurrentPose_Tooltip", "Create PoseAsset from current pose."),
 			FSlateIcon(),
 			FUIAction(
-				FExecuteAction::CreateStatic(&AnimationEditorUtils::ExecuteNewAnimAsset<UPoseAssetFactory, UPoseAsset>, Objects, FString("_PoseAsset"), FAnimAssetCreated::CreateRaw(const_cast<FPersonaModule*>(this), &FPersonaModule::CreatePoseAsset, EPoseSourceOption::CurrentPose, InWeakPersonaToolkit), false),
+				FExecuteAction::CreateStatic(&AnimationEditorUtils::ExecuteNewAnimAsset<UPoseAssetFactory, UPoseAsset>, Objects, FString("_PoseAsset"), FAnimAssetCreated::CreateRaw(const_cast<FPersonaModule*>(this), &FPersonaModule::CreatePoseAsset, EPoseSourceOption::CurrentPose, InWeakPersonaToolkit), bDoNotShowNameDialog, bAllowReplaceExisting),
 				FCanExecuteAction()
 			)
 		);
@@ -1278,7 +1287,7 @@ void FPersonaModule::FillCreatePoseAssetMenu(FMenuBuilder& MenuBuilder, TWeakPtr
 			LOCTEXT("CreatePoseAsset_CurrentAnimation_Tooltip", "Create Animation from current animation."),
 			FSlateIcon(),
 			FUIAction(
-				FExecuteAction::CreateStatic(&AnimationEditorUtils::ExecuteNewAnimAsset<UPoseAssetFactory, UPoseAsset>, Objects, FString("_PoseAsset"), FAnimAssetCreated::CreateRaw(const_cast<FPersonaModule*>(this), &FPersonaModule::CreatePoseAsset, EPoseSourceOption::CurrentAnimation_AnimData, InWeakPersonaToolkit), false),
+				FExecuteAction::CreateStatic(&AnimationEditorUtils::ExecuteNewAnimAsset<UPoseAssetFactory, UPoseAsset>, Objects, FString("_PoseAsset"), FAnimAssetCreated::CreateRaw(const_cast<FPersonaModule*>(this), &FPersonaModule::CreatePoseAsset, EPoseSourceOption::CurrentAnimation_AnimData, InWeakPersonaToolkit), bDoNotShowNameDialog, bAllowReplaceExisting),
 				FCanExecuteAction()
 			)
 		);
diff --git a/Engine/Source/Editor/UnrealEd/Private/AnimationEditorUtils.cpp b/Engine/Source/Editor/UnrealEd/Private/AnimationEditorUtils.cpp
--- a/Engine/Source/Editor/UnrealEd/Private/AnimationEditorUtils.cpp
+++ b/Engine/Source/Editor/UnrealEd/Private/AnimationEditorUtils.cpp
@@ -458,7 +458,7 @@ namespace AnimationEditorUtils
 		AssetToolsModule.Get().CreateUniqueAssetName(InBasePackageName, InSuffix, OutPackageName, OutAssetName);
 	}
 
-	void CreateAnimationAssets(const TArray<TWeakObjectPtr<UObject>>& SkeletonsOrSkeletalMeshes, TSubclassOf<UAnimationAsset> AssetClass, const FString& InPrefix, FAnimAssetCreated AssetCreated, UObject* NameBaseObject /*= nullptr*/, bool bDoNotShowNameDialog /*= false*/)
+	void CreateAnimationAssets(const TArray<TWeakObjectPtr<UObject>>& SkeletonsOrSkeletalMeshes, TSubclassOf<UAnimationAsset> AssetClass, const FString& InPrefix, FAnimAssetCreated AssetCreated, UObject* NameBaseObject /*= nullptr*/, bool bDoNotShowNameDialog /*= false*/, bool bAllowReplaceExisting /*= false*/)
 	{
 		TArray<UObject*> ObjectsToSync;
 		for(auto SkelIt = SkeletonsOrSkeletalMeshes.CreateConstIterator(); SkelIt; ++SkelIt)
@@ -498,8 +498,35 @@ namespace AnimationEditorUtils
 
 				// Create the asset, and assign its skeleton
 				FAssetToolsModule& AssetToolsModule = FModuleManager::GetModuleChecked<FAssetToolsModule>("AssetTools");
-				UAnimationAsset* NewAsset = Cast<UAnimationAsset>(AssetToolsModule.Get().CreateAsset(Name, FPackageName::GetLongPackagePath(PackageName), AssetClass, NULL));
 
+				UAnimationAsset* NewAsset = nullptr;
+
+				if (bAllowReplaceExisting)
+				{
+					UPackage* ExistingPackage = FindPackage(nullptr, *PackageName);
+					UObject* ExistingObject = StaticFindObject(AssetClass.Get(), ExistingPackage, *Name);
+					if (ExistingObject)
+					{
+						EAppReturnType::Type UserResponse = FMessageDialog::Open(
+							EAppMsgType::YesNo,
+							FText::Format(LOCTEXT("CreateAnimationAssetsAlreadyExists", "Do you want to replace the existing asset?\n\nAn asset already exists at the import location: {0}"), FText::FromString(PackageName)));
+
+						if (UserResponse == EAppReturnType::Yes)
+						{
+							NewAsset = Cast<UAnimationAsset>(ExistingObject);
+						}
+						else
+						{
+							return;
+						}
+					}
+				}
+				
+				if (!NewAsset)
+				{
+					NewAsset = Cast<UAnimationAsset>(AssetToolsModule.Get().CreateAsset(Name, FPackageName::GetLongPackagePath(PackageName), AssetClass, NULL));
+				}
+				
 				if(NewAsset)
 				{
 					NewAsset->SetSkeleton(Skeleton);
@@ -643,6 +670,8 @@ namespace AnimationEditorUtils
 
 	void FillCreateAssetMenu(FMenuBuilder& MenuBuilder, const TArray<TWeakObjectPtr<UObject>>& SkeletonsOrSkeletalMeshes, FAnimAssetCreated AssetCreated, bool bInContentBrowser)
 	{
+		const bool bAllowReplaceExisting = false;
+
 		MenuBuilder.BeginSection("CreateAnimAssets", LOCTEXT("CreateAnimAssetsMenuHeading", "Anim Assets"));
 		{
 			if(CanCreateAssetOfType(UAnimBlueprint::StaticClass()))
@@ -666,7 +695,7 @@ namespace AnimationEditorUtils
 					LOCTEXT("Skeleton_NewAnimCompositeTooltip", "Creates an AnimComposite using the selected skeleton."),
 					FSlateIcon(FAppStyle::GetAppStyleSetName(), "ClassIcon.AnimComposite"),
 					FUIAction(
-						FExecuteAction::CreateStatic(&ExecuteNewAnimAsset<UAnimCompositeFactory, UAnimComposite>, SkeletonsOrSkeletalMeshes, FString("_Composite"), AssetCreated, bInContentBrowser),
+						FExecuteAction::CreateStatic(&ExecuteNewAnimAsset<UAnimCompositeFactory, UAnimComposite>, SkeletonsOrSkeletalMeshes, FString("_Composite"), AssetCreated, bInContentBrowser, bAllowReplaceExisting),
 						FCanExecuteAction()
 						)
 					);
@@ -679,7 +708,7 @@ namespace AnimationEditorUtils
 					LOCTEXT("Skeleton_NewAnimMontageTooltip", "Creates an AnimMontage using the selected skeleton."),
 					FSlateIcon(FAppStyle::GetAppStyleSetName(), "ClassIcon.AnimMontage"),
 					FUIAction(
-						FExecuteAction::CreateStatic(&ExecuteNewAnimAsset<UAnimMontageFactory, UAnimMontage>, SkeletonsOrSkeletalMeshes, FString("_Montage"), AssetCreated, bInContentBrowser),
+						FExecuteAction::CreateStatic(&ExecuteNewAnimAsset<UAnimMontageFactory, UAnimMontage>, SkeletonsOrSkeletalMeshes, FString("_Montage"), AssetCreated, bInContentBrowser, bAllowReplaceExisting),
 						FCanExecuteAction()
 						)
 					);
@@ -696,7 +725,7 @@ namespace AnimationEditorUtils
 					LOCTEXT("SkeletalMesh_New2DBlendspaceTooltip", "Creates a Blend Space using the selected skeleton."),
 					FSlateIcon(FAppStyle::GetAppStyleSetName(), "ClassIcon.BlendSpace"),
 					FUIAction(
-						FExecuteAction::CreateStatic(&ExecuteNewAnimAsset<UBlendSpaceFactoryNew, UBlendSpace>, SkeletonsOrSkeletalMeshes, FString("_BlendSpace"), AssetCreated, bInContentBrowser),
+						FExecuteAction::CreateStatic(&ExecuteNewAnimAsset<UBlendSpaceFactoryNew, UBlendSpace>, SkeletonsOrSkeletalMeshes, FString("_BlendSpace"), AssetCreated, bInContentBrowser, bAllowReplaceExisting),
 						FCanExecuteAction()
 						)
 					);
@@ -709,7 +738,7 @@ namespace AnimationEditorUtils
 					LOCTEXT("SkeletalMesh_New1DBlendspaceTooltip", "Creates a 1D Blend Space using the selected skeleton."),
 					FSlateIcon(FAppStyle::GetAppStyleSetName(), "ClassIcon.BlendSpace1D"),
 					FUIAction(
-						FExecuteAction::CreateStatic(&ExecuteNewAnimAsset<UBlendSpaceFactory1D, UBlendSpace1D>, SkeletonsOrSkeletalMeshes, FString("_BlendSpace1D"), AssetCreated, bInContentBrowser),
+						FExecuteAction::CreateStatic(&ExecuteNewAnimAsset<UBlendSpaceFactory1D, UBlendSpace1D>, SkeletonsOrSkeletalMeshes, FString("_BlendSpace1D"), AssetCreated, bInContentBrowser, bAllowReplaceExisting),
 						FCanExecuteAction()
 						)
 					);
@@ -726,7 +755,7 @@ namespace AnimationEditorUtils
 					LOCTEXT("SkeletalMesh_New2DAimOffsetTooltip", "Creates a Aim Offset blendspace using the selected skeleton."),
 					FSlateIcon(),
 					FUIAction(
-						FExecuteAction::CreateStatic(&ExecuteNewAnimAsset<UAimOffsetBlendSpaceFactoryNew, UAimOffsetBlendSpace>, SkeletonsOrSkeletalMeshes, FString("_AimOffset2D"), AssetCreated, bInContentBrowser),
+						FExecuteAction::CreateStatic(&ExecuteNewAnimAsset<UAimOffsetBlendSpaceFactoryNew, UAimOffsetBlendSpace>, SkeletonsOrSkeletalMeshes, FString("_AimOffset2D"), AssetCreated, bInContentBrowser, bAllowReplaceExisting),
 						FCanExecuteAction()
 						)
 					);
@@ -739,7 +768,7 @@ namespace AnimationEditorUtils
 					LOCTEXT("SkeletalMesh_New1DAimOffsetTooltip", "Creates a 1D Aim Offset blendspace using the selected skeleton."),
 					FSlateIcon(),
 					FUIAction(
-						FExecuteAction::CreateStatic(&ExecuteNewAnimAsset<UAimOffsetBlendSpaceFactory1D, UAimOffsetBlendSpace1D>, SkeletonsOrSkeletalMeshes, FString("_AimOffset1D"), AssetCreated, bInContentBrowser),
+						FExecuteAction::CreateStatic(&ExecuteNewAnimAsset<UAimOffsetBlendSpaceFactory1D, UAimOffsetBlendSpace1D>, SkeletonsOrSkeletalMeshes, FString("_AimOffset1D"), AssetCreated, bInContentBrowser, bAllowReplaceExisting),
 						FCanExecuteAction()
 						)
 					);
diff --git a/Engine/Source/Editor/UnrealEd/Public/AnimationEditorUtils.h b/Engine/Source/Editor/UnrealEd/Public/AnimationEditorUtils.h
--- a/Engine/Source/Editor/UnrealEd/Public/AnimationEditorUtils.h
+++ b/Engine/Source/Editor/UnrealEd/Public/AnimationEditorUtils.h
@@ -138,7 +138,7 @@ namespace AnimationEditorUtils
 {
 	UNREALED_API FAssetData CreateModalAnimationCompressionSelectionDialog(const FAnimationCompressionSelectionDialogConfig& InConfig);
 
-	UNREALED_API void CreateAnimationAssets(const TArray<TWeakObjectPtr<UObject>>& SkeletonsOrSkeletalMeshes, TSubclassOf<UAnimationAsset> AssetClass, const FString& InPrefix, FAnimAssetCreated AssetCreated, UObject* NameBaseObject = nullptr, bool bDoNotShowNameDialog = false);
+	UNREALED_API void CreateAnimationAssets(const TArray<TWeakObjectPtr<UObject>>& SkeletonsOrSkeletalMeshes, TSubclassOf<UAnimationAsset> AssetClass, const FString& InPrefix, FAnimAssetCreated AssetCreated, UObject* NameBaseObject = nullptr, bool bDoNotShowNameDialog = false, bool bAllowReplaceExisting = false);
 	
 	UNREALED_API void CreateNewAnimBlueprint(TArray<TWeakObjectPtr<UObject>> SkeletonsOrSkeletalMeshes, FAnimAssetCreated AssetCreated, bool bInContentBrowser);
 	UNREALED_API void FillCreateAssetMenu(FMenuBuilder& MenuBuilder, const TArray<TWeakObjectPtr<UObject>>& SkeletonsOrSkeletalMeshes, FAnimAssetCreated AssetCreated, bool bInContentBrowser=true);
@@ -212,7 +212,7 @@ namespace AnimationEditorUtils
 	//////////////////////////////////////////////////////////////////////////////////////////
 
 	template <typename TFactory, typename T>
-	void ExecuteNewAnimAsset(TArray<TWeakObjectPtr<UObject>> SkeletonsOrSkeletalMeshes, const FString InSuffix, FAnimAssetCreated AssetCreated, bool bInContentBrowser)
+	void ExecuteNewAnimAsset(TArray<TWeakObjectPtr<UObject>> SkeletonsOrSkeletalMeshes, const FString InSuffix, FAnimAssetCreated AssetCreated, bool bInContentBrowser, bool bAllowReplaceExisting)
 	{
 		if(bInContentBrowser && SkeletonsOrSkeletalMeshes.Num() == 1)
 		{
@@ -265,7 +265,10 @@ namespace AnimationEditorUtils
 		}
 		else
 		{
-			CreateAnimationAssets(SkeletonsOrSkeletalMeshes, T::StaticClass(), InSuffix, AssetCreated);
+			UObject* NameBaseObject = nullptr;
+			const bool bDoNotShowNameDialog = false;
+
+			CreateAnimationAssets(SkeletonsOrSkeletalMeshes, T::StaticClass(), InSuffix, AssetCreated, NameBaseObject, bDoNotShowNameDialog, bAllowReplaceExisting);
 		}
 	}
 } // namespace AnimationEditorUtils
