Additive Tracks curves cannot be Enabled/Disabled
#fix invert enabled/disabled flag state passed to SetCurveFlag

diff --git a/Engine/Source/Editor/Persona/Private/AnimTimeline/AnimTimelineTrack_TransformCurve.cpp b/Engine/Source/Editor/Persona/Private/AnimTimeline/AnimTimelineTrack_TransformCurve.cpp
--- a/Engine/Source/Editor/Persona/Private/AnimTimeline/AnimTimelineTrack_TransformCurve.cpp
+++ b/Engine/Source/Editor/Persona/Private/AnimTimeline/AnimTimelineTrack_TransformCurve.cpp
@@ -175,7 +175,7 @@ void FAnimTimelineTrack_TransformCurve::ToggleEnabled()
 	UAnimSequenceBase* AnimSequenceBase = GetModel()->GetAnimSequenceBase();
 
 	IAnimationDataController& Controller = AnimSequenceBase->GetController();
-	Controller.SetCurveFlag(CurveId, AACF_Disabled, !IsEnabled());
+	Controller.SetCurveFlag(CurveId, AACF_Disabled, IsEnabled());
 
 	// need to update curves, otherwise they're not disabled
 	if (GetModel()->GetPreviewScene()->GetPreviewMeshComponent()->PreviewInstance != nullptr)
