PR #8674: Add custom anim blueprint option for additional meshes in animation preview (Contributed by [REDACTED_PERSON])

diff --git a/Engine/Source/Editor/Persona/Private/AnimationEditorPreviewScene.cpp b/Engine/Source/Editor/Persona/Private/AnimationEditorPreviewScene.cpp
--- a/Engine/Source/Editor/Persona/Private/AnimationEditorPreviewScene.cpp
+++ b/Engine/Source/Editor/Persona/Private/AnimationEditorPreviewScene.cpp
@@ -411,7 +411,8 @@ void FAnimationEditorPreviewScene::RefreshAdditionalMeshes(bool bAllowOverrideBa
 						{
 							NewComp->SetAnimInstanceClass(AnimInstances[MeshIndex]);
 						}
-						else
+
+						if (NewComp->GetAnimInstance() == nullptr)
 						{
 							bool bWasCreated = false;
 							FAnimCustomInstanceHelper::BindToSkeletalMeshComponent<UAnimPreviewAttacheInstance>(NewComp,bWasCreated);
@@ -633,7 +634,7 @@ void FAnimationEditorPreviewScene::SetPreviewAnimationAsset(UAnimationAsset* Ani
 
 		RemoveAttachedComponent(false);
 
-		if (AnimAsset != NULL)
+		if (AnimAsset != nullptr)
 		{
 			// Early out if the new preview asset is the same as the current one, to avoid replaying from the beginning, etc...
 			if (AnimAsset == GetPreviewAnimationAsset() && SkeletalMeshComponent->IsPreviewOn())
@@ -823,7 +824,7 @@ void FAnimationEditorPreviewScene::SetSelectedBone(const FName& BoneName, ESelec
 		ClearSelectedActor();
 
 		// Add in bone of interest only if we have a preview instance set-up
-		if (SkeletalMeshComponent->PreviewInstance != NULL)
+		if (SkeletalMeshComponent->PreviewInstance != nullptr)
 		{
 			// need to get mesh bone base since BonesOfInterest is saved in SkeletalMeshComponent
 			// and it is used by renderer. It is not Skeleton base
diff --git a/Engine/Source/Editor/Persona/Private/PreviewSceneCustomizations.cpp b/Engine/Source/Editor/Persona/Private/PreviewSceneCustomizations.cpp
--- a/Engine/Source/Editor/Persona/Private/PreviewSceneCustomizations.cpp
+++ b/Engine/Source/Editor/Persona/Private/PreviewSceneCustomizations.cpp
@@ -728,7 +728,8 @@ void FPreviewMeshCollectionEntryCustomization::CustomizeHeader(TSharedRef<IPrope
 		PropertyHandle->GetParentHandle()->SetOnPropertyValueChanged(FSimpleDelegate::CreateSP(this, &FPreviewMeshCollectionEntryCustomization::HandleMeshesArrayChanged, CustomizationUtils.GetPropertyUtilities()));
 
 		TSharedPtr<IPropertyHandle> SkeletalMeshProperty = PropertyHandle->GetChildHandle(GET_MEMBER_NAME_CHECKED(FPreviewMeshCollectionEntry, SkeletalMesh));
-		if (SkeletalMeshProperty.IsValid())
+		TSharedPtr<IPropertyHandle> AnimBlueprintProperty = PropertyHandle->GetChildHandle(GET_MEMBER_NAME_CHECKED(FPreviewMeshCollectionEntry, AnimBlueprint));
+		if (SkeletalMeshProperty.IsValid() && AnimBlueprintProperty.IsValid())
 		{
 			HeaderRow.NameContent()
 			[
@@ -738,12 +739,24 @@ void FPreviewMeshCollectionEntryCustomization::CustomizeHeader(TSharedRef<IPrope
 			.MaxDesiredWidth(250.0f)
 			.MinDesiredWidth(250.0f)
 			[
-				SNew(SObjectPropertyEntryBox)
-				.AllowedClass(USkeletalMesh::StaticClass())
-				.PropertyHandle(SkeletalMeshProperty)
-				.OnShouldFilterAsset(this, &FPreviewMeshCollectionEntryCustomization::HandleShouldFilterAsset, SkeletonName, Skeleton)
-				.OnObjectChanged(this, &FPreviewMeshCollectionEntryCustomization::HandleMeshChanged)
-				.ThumbnailPool(CustomizationUtils.GetThumbnailPool())
+				SNew(SVerticalBox)
+				+SVerticalBox::Slot()
+				[
+					SNew(SObjectPropertyEntryBox)
+					.AllowedClass(USkeletalMesh::StaticClass())
+					.PropertyHandle(SkeletalMeshProperty)
+					.OnShouldFilterAsset(this, &FPreviewMeshCollectionEntryCustomization::HandleShouldFilterAsset, SkeletonName, Skeleton)
+					.OnObjectChanged(this, &FPreviewMeshCollectionEntryCustomization::HandleMeshChanged)
+					.ThumbnailPool(CustomizationUtils.GetThumbnailPool())
+				]
+				+SVerticalBox::Slot()
+				[
+					SNew(SObjectPropertyEntryBox)
+					.AllowedClass(UAnimBlueprint::StaticClass())
+					.PropertyHandle(AnimBlueprintProperty)
+					.OnObjectChanged(this, &FPreviewMeshCollectionEntryCustomization::HandleMeshChanged)
+					.ThumbnailPool(CustomizationUtils.GetThumbnailPool())
+				]
 			];
 		}
 	}
diff --git a/Engine/Source/Runtime/Engine/Classes/Engine/PreviewMeshCollection.h b/Engine/Source/Runtime/Engine/Classes/Engine/PreviewMeshCollection.h
--- a/Engine/Source/Runtime/Engine/Classes/Engine/PreviewMeshCollection.h
+++ b/Engine/Source/Runtime/Engine/Classes/Engine/PreviewMeshCollection.h
@@ -2,6 +2,7 @@
 
 #pragma once
 
+#include "Animation/AnimBlueprint.h"
 #include "CoreMinimal.h"
 #include "UObject/ObjectMacros.h"
 #include "Engine/SkeletalMesh.h"
@@ -32,6 +33,10 @@ struct FPreviewMeshCollectionEntry
 	/** The skeletal mesh to display */
 	UPROPERTY(EditAnywhere, Category = "Skeletal Mesh", meta=(DisplayThumbnail=true, DisallowedClasses = "/Script/ApexDestruction.DestructibleMesh"))
 	TSoftObjectPtr<USkeletalMesh> SkeletalMesh;
+
+	/** The custom animation blueprint for the mesh */
+	UPROPERTY(EditAnywhere, Category = "Anim Blueprint", meta=(DisplayThumbnail=true))
+	TSoftObjectPtr<UAnimBlueprint> AnimBlueprint;
 };
 
 /** A simple collection of skeletal meshes used for in-editor preview */
diff --git a/Engine/Source/Runtime/Engine/Private/Animation/PreviewMeshCollection.cpp b/Engine/Source/Runtime/Engine/Private/Animation/PreviewMeshCollection.cpp
--- a/Engine/Source/Runtime/Engine/Private/Animation/PreviewMeshCollection.cpp
+++ b/Engine/Source/Runtime/Engine/Private/Animation/PreviewMeshCollection.cpp
@@ -15,9 +15,18 @@ void UPreviewMeshCollection::GetPreviewSkeletalMeshes(TArray<USkeletalMesh*>& Ou
 		if (Entry.SkeletalMesh.LoadSynchronous())
 		{
 			OutList.Add(Entry.SkeletalMesh.Get());
+			TSubclassOf<UAnimInstance> SubClass = nullptr;
+			// Load up our custom animation blueprints
+			if (Entry.AnimBlueprint.LoadSynchronous())
+			{
+				const UAnimBlueprint* AnimBP = Entry.AnimBlueprint.Get();
+				if(UClass* AnimInstanceClass = Cast<UClass>(AnimBP->GeneratedClass))
+				{					
+					SubClass = AnimInstanceClass;
+				}
+			}
+			
+			OutAnimBP.Add(SubClass);
 		}
 	}
-
-	// no class, but just send it out
-	OutAnimBP.AddZeroed(OutList.Num());
 }
