CTRL+Z/CTRL+Y doesn't undo/redo changes in checkboxes in Additional Meshes in Animation Sequence
#fix add scoped transactions for affected settings, and make sure UPersonaOptions is marked as transactional

diff --git a/Engine/Source/Editor/Persona/Private/PersonaModule.cpp b/Engine/Source/Editor/Persona/Private/PersonaModule.cpp
--- a/Engine/Source/Editor/Persona/Private/PersonaModule.cpp
+++ b/Engine/Source/Editor/Persona/Private/PersonaModule.cpp
@@ -149,6 +149,8 @@ void FPersonaModule::StartupModule()
 
 	FKismetEditorUtilities::RegisterOnBlueprintCreatedCallback(this, UAnimNotify::StaticClass(), FKismetEditorUtilities::FOnBlueprintCreated::CreateRaw(this, &FPersonaModule::HandleNewAnimNotifyBlueprintCreated));
 	FKismetEditorUtilities::RegisterOnBlueprintCreatedCallback(this, UAnimNotifyState::StaticClass(), FKismetEditorUtilities::FOnBlueprintCreated::CreateRaw(this, &FPersonaModule::HandleNewAnimNotifyStateBlueprintCreated));
+
+	GetMutableDefault<UPersonaOptions>()->SetFlags(RF_Transactional);
 }
 
 void FPersonaModule::ShutdownModule()
diff --git a/Engine/Source/Editor/Persona/Private/PreviewSceneCustomizations.cpp b/Engine/Source/Editor/Persona/Private/PreviewSceneCustomizations.cpp
--- a/Engine/Source/Editor/Persona/Private/PreviewSceneCustomizations.cpp
+++ b/Engine/Source/Editor/Persona/Private/PreviewSceneCustomizations.cpp
@@ -580,7 +580,10 @@ TSharedRef<SWidget> FPreviewSceneDescriptionCustomization::MakeControllerComboEn
 void FPreviewSceneDescriptionCustomization::OnComboSelectionChanged(TSharedPtr<FPersonaModeComboEntry> InSelectedItem, ESelectInfo::Type SelectInfo)
 {
 	TSharedPtr<FAnimationEditorPreviewScene> PreviewScenePtr = PreviewScene.Pin();
+
+	const FScopedTransaction Transaction(LOCTEXT("ChangePreviewSceneController", "Setting Preview Scene Controller"));
 	UPersonaPreviewSceneDescription* PersonaPreviewSceneDescription = PreviewScenePtr->GetPreviewSceneDescription();
+	PersonaPreviewSceneDescription->Modify();
 
 	PersonaPreviewSceneDescription->SetPreviewController(InSelectedItem->Class, PreviewScenePtr.Get());
 
@@ -618,7 +621,11 @@ void FPreviewSceneDescriptionCustomization::HandleAdditionalMeshesChanged(const
 
 void FPreviewSceneDescriptionCustomization::HandleAllowDifferentSkeletonsCheckedStateChanged(ECheckBoxState CheckState)
 {
-	GetMutableDefault<UPersonaOptions>()->bAllowPreviewMeshCollectionsToSelectFromDifferentSkeletons = (CheckState == ECheckBoxState::Checked);
+	const FScopedTransaction Transaction(LOCTEXT("AllowDifferentSkeletons", "Setting Allow Different Skeletons"));
+	UPersonaOptions* PersonaOptions = GetMutableDefault<UPersonaOptions>();
+	PersonaOptions->Modify();
+	
+	PersonaOptions->bAllowPreviewMeshCollectionsToSelectFromDifferentSkeletons = (CheckState == ECheckBoxState::Checked);
 }
 
 ECheckBoxState FPreviewSceneDescriptionCustomization::HandleAllowDifferentSkeletonsIsChecked() const
@@ -628,7 +635,10 @@ ECheckBoxState FPreviewSceneDescriptionCustomization::HandleAllowDifferentSkelet
 
 void FPreviewSceneDescriptionCustomization::HandleUseCustomAnimBPCheckedStateChanged(ECheckBoxState CheckState)
 {
-	GetMutableDefault<UPersonaOptions>()->bAllowPreviewMeshCollectionsToUseCustomAnimBP = (CheckState == ECheckBoxState::Checked);
+	const FScopedTransaction Transaction(LOCTEXT("AllowDifferentSkeletons", "Setting Allow Different Skeletons"));
+	UPersonaOptions* PersonaOptions = GetMutableDefault<UPersonaOptions>();
+	PersonaOptions->Modify();	
+	PersonaOptions->bAllowPreviewMeshCollectionsToUseCustomAnimBP = (CheckState == ECheckBoxState::Checked);
 
 	if (PreviewScene.IsValid())
 	{
