Allow adding curves to edit in anim sequences
#feature now have the ability to edit multiple curves simultaneously inside of the Animation Editor by adding/removing curves of interest

diff --git a/Engine/Source/Editor/AnimationEditor/Private/AnimationEditor.cpp b/Engine/Source/Editor/AnimationEditor/Private/AnimationEditor.cpp
--- a/Engine/Source/Editor/AnimationEditor/Private/AnimationEditor.cpp
+++ b/Engine/Source/Editor/AnimationEditor/Private/AnimationEditor.cpp
@@ -555,8 +555,6 @@ void FAnimationEditor::EditCurves(UAnimSequenceBase* InAnimSequence, const TArra
 
 	check(CurveEditor.IsValid());
 
-	CurveEditor.Pin()->ResetCurves();
-
 	for(const FCurveEditInfo& CurveInfo : InCurveInfo)
 	{
 		CurveEditor.Pin()->AddCurve(CurveInfo.CurveDisplayName, CurveInfo.CurveColor, CurveInfo.Name, CurveInfo.Type, CurveInfo.CurveIndex, CurveInfo.OnCurveModified);
diff --git a/Engine/Source/Editor/Persona/Private/SAnimSequenceCurveEditor.cpp b/Engine/Source/Editor/Persona/Private/SAnimSequenceCurveEditor.cpp
--- a/Engine/Source/Editor/Persona/Private/SAnimSequenceCurveEditor.cpp
+++ b/Engine/Source/Editor/Persona/Private/SAnimSequenceCurveEditor.cpp
@@ -18,6 +18,7 @@
 #include "Framework/MultiBox/MultiBoxBuilder.h"
 #include "Styling/AppStyle.h"
 #include "Animation/AnimData/AnimDataModel.h"
+#include "ToolMenus.h"
 
 #define LOCTEXT_NAMESPACE "SAnimSequenceCurveEditor"
 
@@ -385,7 +386,8 @@ void SAnimSequenceCurveEditor::Construct(const FArguments& InArgs, const TShared
 
 	AnimSequence->GetDataModel()->GetModifiedEvent().AddRaw(this, &SAnimSequenceCurveEditor::OnModelHasChanged);
 
-	CurveEditorTree = SNew(SCurveEditorTree, CurveEditor);
+	CurveEditorTree = SNew(SCurveEditorTree, CurveEditor)
+		.OnContextMenuOpening(this, &SAnimSequenceCurveEditor::OnContextMenuOpening);
 
 	TSharedRef<SCurveEditorPanel> CurveEditorPanel = SNew(SCurveEditorPanel, CurveEditor.ToSharedRef())
 		.GridLineTint(FLinearColor(0.f, 0.f, 0.f, 0.3f))
@@ -482,6 +484,50 @@ TSharedRef<SWidget> SAnimSequenceCurveEditor::MakeToolbar(TSharedRef<SCurveEdito
 	return ToolBarBuilder.MakeWidget();
 }
 
+TSharedPtr<SWidget> SAnimSequenceCurveEditor::OnContextMenuOpening()
+{
+	const TArray<FCurveEditorTreeItemID>& Selection = CurveEditorTree->GetSelectedItems();
+	if (Selection.Num())
+	{
+		UToolMenus* ToolMenus = UToolMenus::Get();
+		static const FName MenuName = "SAnimSequenceCurveEditor.CurveEditorTreeContextMenu";
+		if (!ToolMenus->IsMenuRegistered(MenuName))
+		{
+			ToolMenus->RegisterMenu(MenuName);
+		}
+
+		FToolMenuContext Context;
+		UToolMenu* Menu = ToolMenus->GenerateMenu(MenuName, Context);
+
+		FToolMenuSection& Section = Menu->AddSection("Selection", LOCTEXT("SelectionLablel", "Selection"));
+		Section.AddMenuEntry("RemoveSelectedCurves", LOCTEXT("RemoveCurveLabel", "Stop editing selected curve(s)"),
+			LOCTEXT("RemoveCurveTooltip", "Removes the currently selected curve(s) from editing"),
+			FSlateIcon(),
+			FToolUIActionChoice(FExecuteAction::CreateLambda([this]()
+				{
+					// Remove all selected tree items, and associated curves
+					TArray<FCurveModelID> ModelIDs;				
+					TArray<FCurveEditorTreeItemID> Selection = CurveEditorTree->GetSelectedItems();
+					for (const FCurveEditorTreeItemID& SelectedItem : Selection)
+					{
+						ModelIDs.Append(CurveEditor->GetTreeItem(SelectedItem).GetCurves());
+						CurveEditor->RemoveTreeItem(SelectedItem);
+					}
+					CurveEditorTree->ClearSelection();
+
+					for (const FCurveModelID& ID : ModelIDs)
+					{
+						CurveEditor->RemoveCurve(ID);
+					}
+				}))				
+			);
+
+		return ToolMenus->GenerateWidget(Menu);
+	}
+
+	return SNullWidget::NullWidget;
+}
+
 void SAnimSequenceCurveEditor::ResetCurves()
 {
 	CurveEditor->RemoveAllTreeItems();
@@ -490,6 +536,16 @@ void SAnimSequenceCurveEditor::ResetCurves()
 
 void SAnimSequenceCurveEditor::AddCurve(const FText& InCurveDisplayName, const FLinearColor& InCurveColor, const FSmartName& InName, ERawCurveTrackTypes InType, int32 InCurveIndex, FSimpleDelegate InOnCurveModified)
 {
+	// Ensure that curve is not already being edited
+	for (const TPair<FCurveModelID, TUniquePtr<FCurveModel>>& CurvePair : CurveEditor->GetCurves())
+	{
+		const FRichCurveEditorModelNamed* Model = static_cast<FRichCurveEditorModelNamed*>(CurvePair.Value.Get());
+		if(Model->Name == InName && Model->Type == InType && Model->CurveIndex == InCurveIndex)
+		{
+			return;
+		}
+	}
+	
 	FCurveEditorTreeItem* TreeItem = CurveEditor->AddTreeItem(FCurveEditorTreeItemID());
 	TreeItem->SetStrongItem(MakeShared<FAnimSequenceCurveEditorItem>(InName, InType, InCurveIndex, AnimSequence, InCurveDisplayName, InCurveColor, InOnCurveModified, TreeItem->GetID()));
 
diff --git a/Engine/Source/Editor/Persona/Private/SAnimSequenceCurveEditor.h b/Engine/Source/Editor/Persona/Private/SAnimSequenceCurveEditor.h
--- a/Engine/Source/Editor/Persona/Private/SAnimSequenceCurveEditor.h
+++ b/Engine/Source/Editor/Persona/Private/SAnimSequenceCurveEditor.h
@@ -77,6 +77,7 @@ class SAnimSequenceCurveEditor : public IAnimSequenceCurveEditor
 private:
 	// Build the toolbar for this curve editor
 	TSharedRef<SWidget> MakeToolbar(TSharedRef<SCurveEditorPanel> InEditorPanel);
+	TSharedPtr<SWidget> OnContextMenuOpening();
 
 private:
 	/** The actual curve editor */
