FNames containing asset paths are deprecated. FSoftObjectPath or FTopLevelAssetPath should be used instead.
Fixups for Persona.

diff --git a/Engine/Source/Editor/Persona/Private/PersonaAssetFamily.cpp b/Engine/Source/Editor/Persona/Private/PersonaAssetFamily.cpp
--- a/Engine/Source/Editor/Persona/Private/PersonaAssetFamily.cpp
+++ b/Engine/Source/Editor/Persona/Private/PersonaAssetFamily.cpp
@@ -150,7 +150,7 @@ FAssetData FPersonaAssetFamily::FindAssetOfType(UClass* InAssetClass) const
 				Filter.ClassPaths.Add(UPhysicsAsset::StaticClass()->GetClassPathName());
 				if(Mesh.IsValid())
 				{
-					Filter.TagsAndValues.Add(GET_MEMBER_NAME_CHECKED(UPhysicsAsset, PreviewSkeletalMesh), FAssetData(Mesh.Get()).ObjectPath.ToString());
+					Filter.TagsAndValues.Add(GET_MEMBER_NAME_CHECKED(UPhysicsAsset, PreviewSkeletalMesh), FSoftObjectPath(Mesh.Get()).ToString());
 				}
 
 				AssetRegistryModule.Get().GetAssets(Filter, Assets);
@@ -195,7 +195,7 @@ void FPersonaAssetFamily::FindAssetsOfType(UClass* InAssetClass, TArray<FAssetDa
 			Filter.ClassPaths.Add(UPhysicsAsset::StaticClass()->GetClassPathName());
 			if(Mesh != nullptr)
 			{
-				Filter.TagsAndValues.Add(GET_MEMBER_NAME_CHECKED(UPhysicsAsset, PreviewSkeletalMesh), FAssetData(Mesh.Get()).ObjectPath.ToString());
+				Filter.TagsAndValues.Add(GET_MEMBER_NAME_CHECKED(UPhysicsAsset, PreviewSkeletalMesh), FSoftObjectPath(Mesh.Get()).ToString());
 			}
 
 			AssetRegistryModule.Get().GetAssets(Filter, OutAssets);
@@ -342,7 +342,7 @@ bool FPersonaAssetFamily::IsAssetCompatible(const FAssetData& InAssetData) const
 			FAssetDataTagMapSharedView::FFindTagResult Result = InAssetData.TagsAndValues.FindTag(GET_MEMBER_NAME_CHECKED(UPhysicsAsset, PreviewSkeletalMesh));
 			if (Result.IsSet() && Mesh.IsValid())
 			{
-				return Result.GetValue() == FAssetData(Mesh.Get()).ObjectPath.ToString();
+				return Result.GetValue() == FSoftObjectPath(Mesh.Get()).ToString();
 			}
 		}
 	}
diff --git a/Engine/Source/Editor/Persona/Private/PersonaModule.cpp b/Engine/Source/Editor/Persona/Private/PersonaModule.cpp
--- a/Engine/Source/Editor/Persona/Private/PersonaModule.cpp
+++ b/Engine/Source/Editor/Persona/Private/PersonaModule.cpp
@@ -627,7 +627,7 @@ void FPersonaModule::TestSkeletonCurveNamesForUse(const TSharedRef<IEditableSkel
 							if (!Message.IsValid())
 							{
 								Message = CurveOutput.Warning();
-								Message->AddToken(FAssetNameToken::Create(Anim.ObjectPath.ToString(), FText::FromName(Anim.AssetName)));
+								Message->AddToken(FAssetNameToken::Create(Anim.GetObjectPathString(), FText::FromName(Anim.AssetName)));
 								Message->AddToken(FTextToken::Create(LOCTEXT("VerifyCurves_FoundAnimationsWithUnusedReferences", "References the following curves that are not used for either morph targets or material parameters and so may be unneeded")));
 							}
 							CurveOutput.Info(FText::FromName(Curve.Name.DisplayName));
diff --git a/Engine/Source/Editor/Persona/Private/PhysicsAssetRenderUtils.cpp b/Engine/Source/Editor/Persona/Private/PhysicsAssetRenderUtils.cpp
--- a/Engine/Source/Editor/Persona/Private/PhysicsAssetRenderUtils.cpp
+++ b/Engine/Source/Editor/Persona/Private/PhysicsAssetRenderUtils.cpp
@@ -566,7 +566,7 @@ void UPhysicsAssetRenderUtilities::OnAssetRenamed(FAssetData const& AssetInfo, c
 
 	if (IdToSettingsMap.RemoveAndCopyValue(OldPhysicsAssetPathNameHash, Temp))
 	{
-		const uint32 PhysicsAssetPathNameHash = GetPathNameHash(AssetInfo.ObjectPath.ToString());
+		const uint32 PhysicsAssetPathNameHash = GetPathNameHash(AssetInfo.GetObjectPathString());
 		IdToSettingsMap.Add(PhysicsAssetPathNameHash, Temp);
 		SaveConfig();
 	}
diff --git a/Engine/Source/Editor/Persona/Private/SAnimationSequenceBrowser.cpp b/Engine/Source/Editor/Persona/Private/SAnimationSequenceBrowser.cpp
--- a/Engine/Source/Editor/Persona/Private/SAnimationSequenceBrowser.cpp
+++ b/Engine/Source/Editor/Persona/Private/SAnimationSequenceBrowser.cpp
@@ -1077,7 +1077,7 @@ TSharedRef<SWidget> SAnimationSequenceBrowser::CreateHistoryMenu(bool bInBackHis
 			if(AssetData.IsValid())
 			{
 				const FText DisplayName = FText::FromName(AssetData.AssetName);
-				const FText Tooltip = FText::FromString( AssetData.ObjectPath.ToString() );
+				const FText Tooltip = FText::FromString( AssetData.GetObjectPathString() );
 
 				MenuBuilder.AddMenuEntry(DisplayName, Tooltip, FSlateIcon(), 
 					FUIAction(
@@ -1099,7 +1099,7 @@ TSharedRef<SWidget> SAnimationSequenceBrowser::CreateHistoryMenu(bool bInBackHis
 			if(AssetData.IsValid())
 			{
 				const FText DisplayName = FText::FromName(AssetData.AssetName);
-				const FText Tooltip = FText::FromString( AssetData.ObjectPath.ToString() );
+				const FText Tooltip = FText::FromString(AssetData.GetObjectPathString());
 
 				MenuBuilder.AddMenuEntry(DisplayName, Tooltip, FSlateIcon(), 
 					FUIAction(
@@ -1199,7 +1199,7 @@ void SAnimationSequenceBrowser::CacheOriginalAnimAssetHistory()
 			if(UObject* PreviewAsset = PersonaToolkitPtr.Pin()->GetPreviewScene()->GetPreviewAnimationAsset())
 			{
 				FAssetRegistryModule& AssetRegistryModule = FModuleManager::LoadModuleChecked<FAssetRegistryModule>("AssetRegistry");
-				FAssetData AssetData = AssetRegistryModule.Get().GetAssetByObjectPath(FName(*PreviewAsset->GetPathName()));
+				FAssetData AssetData = AssetRegistryModule.Get().GetAssetByObjectPath(FSoftObjectPath(PreviewAsset));
 				AssetHistory.Add(AssetData);
 				CurrentAssetHistoryIndex = AssetHistory.Num() - 1;
 			}
diff --git a/Engine/Source/Editor/Persona/Private/SAssetFamilyShortcutBar.cpp b/Engine/Source/Editor/Persona/Private/SAssetFamilyShortcutBar.cpp
--- a/Engine/Source/Editor/Persona/Private/SAssetFamilyShortcutBar.cpp
+++ b/Engine/Source/Editor/Persona/Private/SAssetFamilyShortcutBar.cpp
@@ -207,7 +207,7 @@ public:
 			}
 			else
 			{
-				UE_LOG(LogAnimation, Error, TEXT("Asset cannot be opened: %s"), *AssetData.ObjectPath.ToString());
+				UE_LOG(LogAnimation, Error, TEXT("Asset cannot be opened: %s"), *AssetData.GetObjectPathString());
 			}
 		}
 	}
@@ -240,7 +240,7 @@ public:
 			{
 				for (UObject* Object : *Objects)
 				{
-					if (Object->GetPathName().Compare(AssetData.ObjectPath.ToString(), ESearchCase::IgnoreCase) == 0)
+					if (Object->GetPathName().Compare(AssetData.GetObjectPathString(), ESearchCase::IgnoreCase) == 0)
 					{
 						return ECheckBoxState::Checked;
 					}
@@ -371,7 +371,7 @@ public:
 	{
 		if (AssetFamily->IsAssetCompatible(InAssetData))
 		{
-			if(InOldObjectPath == AssetData.ObjectPath.ToString())
+			if (InOldObjectPath == AssetData.GetObjectPathString())
 			{
 				AssetData = InAssetData;
 
diff --git a/Engine/Source/Editor/Persona/Private/SRigPicker.cpp b/Engine/Source/Editor/Persona/Private/SRigPicker.cpp
--- a/Engine/Source/Editor/Persona/Private/SRigPicker.cpp
+++ b/Engine/Source/Editor/Persona/Private/SRigPicker.cpp
@@ -105,7 +105,7 @@ void SRigPicker::OnAssetSelected(const FAssetData& AssetData)
 
 bool SRigPicker::OnShouldFilterAsset(const struct FAssetData& AssetData)
 {
-	if (EngineHumanoidRig && AssetData.ObjectPath == FName(*EngineHumanoidRig->GetPathName()))
+	if (EngineHumanoidRig && AssetData.GetSoftObjectPath() == FSoftObjectPath(EngineHumanoidRig))
 	{
 		return true;
 	}
