Fixed a crash when changing the name of the sampled region in the skel mesh editor

diff --git a/Engine/Source/Editor/Persona/Private/Customization/SkeletalMeshRegionCustomization.cpp b/Engine/Source/Editor/Persona/Private/Customization/SkeletalMeshRegionCustomization.cpp
--- a/Engine/Source/Editor/Persona/Private/Customization/SkeletalMeshRegionCustomization.cpp
+++ b/Engine/Source/Editor/Persona/Private/Customization/SkeletalMeshRegionCustomization.cpp
@@ -189,7 +189,10 @@ TSharedRef<SWidget> FNiagaraSkeletalMeshRegionMaterialFilterDetails::HandleMater
 
 void FNiagaraSkeletalMeshRegionMaterialFilterDetails::HandleMaterialNameComboBoxSelectionChanged(TSharedPtr<FName> StringItem, ESelectInfo::Type SelectInfo)
 {
-	MaterialNameHandle->SetValue(*StringItem);
+	if (StringItem)
+	{
+		MaterialNameHandle->SetValue(*StringItem);
+	}
 }
 
 void FNiagaraSkeletalMeshRegionMaterialFilterDetails::OnComboOpening()
