Adding "Preview Scene Settings" tabs to IK Rig & Retarget editors.

This change also added a delegate to FPersonaToolkitArgs with reference to a IDetailBuilderLayout so assets can customize the Preview Scene Settings.

diff --git a/Engine/Plugins/Animation/IKRig/Source/IKRigEditor/Private/RetargetEditor/IKRetargetApplicationMode.cpp b/Engine/Plugins/Animation/IKRig/Source/IKRigEditor/Private/RetargetEditor/IKRetargetApplicationMode.cpp
--- a/Engine/Plugins/Animation/IKRig/Source/IKRigEditor/Private/RetargetEditor/IKRetargetApplicationMode.cpp
+++ b/Engine/Plugins/Animation/IKRig/Source/IKRigEditor/Private/RetargetEditor/IKRetargetApplicationMode.cpp
@@ -35,6 +35,7 @@ FIKRetargetApplicationMode::FIKRetargetApplicationMode(
 	FPersonaModule& PersonaModule = FModuleManager::LoadModuleChecked<FPersonaModule>("Persona");
 	TabFactories.RegisterFactory(PersonaModule.CreatePersonaViewportTabFactory(InHostingApp, ViewportArgs));
 	TabFactories.RegisterFactory(PersonaModule.CreateDetailsTabFactory(InHostingApp, FOnDetailsCreated::CreateSP(&IKRetargetEditor.Get(), &FIKRetargetEditor::HandleDetailsCreated)));
+	TabFactories.RegisterFactory(PersonaModule.CreateAdvancedPreviewSceneTabFactory(InHostingApp, IKRetargetEditor->GetPersonaToolkit()->GetPreviewScene()));
 
 	// register custom tabs
 	TabFactories.RegisterFactory(MakeShared<FIKRetargetChainTabSummoner>(IKRetargetEditor));
@@ -43,7 +44,7 @@ FIKRetargetApplicationMode::FIKRetargetApplicationMode(
 	TabFactories.RegisterFactory(MakeShared<FIKRetargetHierarchyTabSummoner>(IKRetargetEditor));
 
 	// create tab layout
-	TabLayout = FTabManager::NewLayout("Standalone_IKRetargetEditor_Layout_v1.017")
+	TabLayout = FTabManager::NewLayout("Standalone_IKRetargetEditor_Layout_v1.018")
 		->AddArea
 		(
 			FTabManager::NewPrimaryArea()
@@ -95,6 +96,7 @@ FIKRetargetApplicationMode::FIKRetargetApplicationMode(
 					FTabManager::NewStack()
 						->SetSizeCoefficient(0.7f)
 						->AddTab(FPersonaTabs::DetailsID, ETabState::OpenedTab)
+						->AddTab(FPersonaTabs::AdvancedPreviewSceneSettingsID, ETabState::OpenedTab)
 						->SetForegroundTab(FPersonaTabs::DetailsID)
 					)
 					->Split
diff --git a/Engine/Plugins/Animation/IKRig/Source/IKRigEditor/Private/RetargetEditor/IKRetargetEditor.cpp b/Engine/Plugins/Animation/IKRig/Source/IKRigEditor/Private/RetargetEditor/IKRetargetEditor.cpp
--- a/Engine/Plugins/Animation/IKRig/Source/IKRigEditor/Private/RetargetEditor/IKRetargetEditor.cpp
+++ b/Engine/Plugins/Animation/IKRig/Source/IKRigEditor/Private/RetargetEditor/IKRetargetEditor.cpp
@@ -49,6 +49,7 @@ void FIKRetargetEditor::InitAssetEditor(
 	
 	FPersonaToolkitArgs PersonaToolkitArgs;
 	PersonaToolkitArgs.OnPreviewSceneCreated = FOnPreviewSceneCreated::FDelegate::CreateSP(this, &FIKRetargetEditor::HandlePreviewSceneCreated);
+	PersonaToolkitArgs.OnPreviewSceneSettingsCustomized = FOnPreviewSceneSettingsCustomized::FDelegate::CreateSP(this, &FIKRetargetEditor::HandleOnPreviewSceneSettingsCustomized);
 	
 	FPersonaModule& PersonaModule = FModuleManager::LoadModuleChecked<FPersonaModule>("Persona");
 	EditorController->PersonaToolkit = PersonaModule.CreatePersonaToolkit(InAsset, PersonaToolkitArgs);
@@ -383,6 +384,14 @@ void FIKRetargetEditor::SetupAnimInstance()
 	EditorController->TargetAnimInstance->InitializeAnimation();
 }
 
+void FIKRetargetEditor::HandleOnPreviewSceneSettingsCustomized(IDetailLayoutBuilder& DetailBuilder) const
+{
+	DetailBuilder.HideCategory("Additional Meshes");
+	DetailBuilder.HideCategory("Physics");
+	DetailBuilder.HideCategory("Mesh");
+	DetailBuilder.HideCategory("Animation Blueprint");
+}
+
 void FIKRetargetEditor::HandleDetailsCreated(const TSharedRef<class IDetailsView>& InDetailsView)
 {
 	InDetailsView->OnFinishedChangingProperties().AddSP(this, &FIKRetargetEditor::OnFinishedChangingDetails);
diff --git a/Engine/Plugins/Animation/IKRig/Source/IKRigEditor/Private/RigEditor/IKRigMode.cpp b/Engine/Plugins/Animation/IKRig/Source/IKRigEditor/Private/RigEditor/IKRigMode.cpp
--- a/Engine/Plugins/Animation/IKRig/Source/IKRigEditor/Private/RigEditor/IKRigMode.cpp
+++ b/Engine/Plugins/Animation/IKRig/Source/IKRigEditor/Private/RigEditor/IKRigMode.cpp
@@ -3,6 +3,7 @@
 #include "RigEditor/IKRigMode.h"
 #include "RigEditor/IKRigToolkit.h"
 #include "IPersonaPreviewScene.h"
+#include "IPersonaToolkit.h"
 #include "PersonaModule.h"
 #include "ISkeletonEditorModule.h"
 #include "Modules/ModuleManager.h"
@@ -34,6 +35,7 @@ FIKRigMode::FIKRigMode(
 	FPersonaModule& PersonaModule = FModuleManager::LoadModuleChecked<FPersonaModule>("Persona");
 	TabFactories.RegisterFactory(PersonaModule.CreatePersonaViewportTabFactory(InHostingApp, ViewportArgs));
 	TabFactories.RegisterFactory(PersonaModule.CreateDetailsTabFactory(InHostingApp, FOnDetailsCreated::CreateSP(&IKRigEditor.Get(), &FIKRigEditorToolkit::HandleDetailsCreated)));
+	TabFactories.RegisterFactory(PersonaModule.CreateAdvancedPreviewSceneTabFactory(InHostingApp, IKRigEditor->GetPersonaToolkit()->GetPreviewScene()));
 
 	// register custom tabs
 	TabFactories.RegisterFactory(MakeShared<FIKRigAssetBrowserTabSummoner>(IKRigEditor));
@@ -43,7 +45,7 @@ FIKRigMode::FIKRigMode(
 	TabFactories.RegisterFactory(MakeShared<FIKRigOutputLogTabSummoner>(IKRigEditor));
 
 	// create tab layout
-	TabLayout = FTabManager::NewLayout("Standalone_IKRigEditor_Layout_v1.125")
+	TabLayout = FTabManager::NewLayout("Standalone_IKRigEditor_Layout_v1.127")
 		->AddArea
 		(
 			FTabManager::NewPrimaryArea()
@@ -100,6 +102,7 @@ FIKRigMode::FIKRigMode(
 						FTabManager::NewStack()
 						->SetSizeCoefficient(0.6f)
 						->AddTab(FPersonaTabs::DetailsID, ETabState::OpenedTab)
+						->AddTab(FPersonaTabs::AdvancedPreviewSceneSettingsID, ETabState::OpenedTab)
 						->SetForegroundTab(FPersonaTabs::DetailsID)
                     )
                     ->Split
diff --git a/Engine/Plugins/Animation/IKRig/Source/IKRigEditor/Private/RigEditor/IKRigToolkit.cpp b/Engine/Plugins/Animation/IKRig/Source/IKRigEditor/Private/RigEditor/IKRigToolkit.cpp
--- a/Engine/Plugins/Animation/IKRig/Source/IKRigEditor/Private/RigEditor/IKRigToolkit.cpp
+++ b/Engine/Plugins/Animation/IKRig/Source/IKRigEditor/Private/RigEditor/IKRigToolkit.cpp
@@ -18,6 +18,7 @@
 
 #include "IKRigDefinition.h"
 #include "IPersonaViewport.h"
+#include "Persona/Private/PersonaPreviewSceneDescription.h"
 #include "RigEditor/IKRigAnimInstance.h"
 #include "RigEditor/IKRigCommands.h"
 #include "RigEditor/IKRigEditMode.h"
@@ -54,6 +55,7 @@ void FIKRigEditorToolkit::InitAssetEditor(
 	
 	FPersonaToolkitArgs PersonaToolkitArgs;
 	PersonaToolkitArgs.OnPreviewSceneCreated = FOnPreviewSceneCreated::FDelegate::CreateSP(this, &FIKRigEditorToolkit::HandlePreviewSceneCreated);
+	PersonaToolkitArgs.OnPreviewSceneSettingsCustomized = FOnPreviewSceneSettingsCustomized::FDelegate::CreateSP(this, &FIKRigEditorToolkit::HandleOnPreviewSceneSettingsCustomized);
 	
 	const FPersonaModule& PersonaModule = FModuleManager::LoadModuleChecked<FPersonaModule>("Persona");
 	PersonaToolkit = PersonaModule.CreatePersonaToolkit(IKRigAsset, PersonaToolkitArgs);
@@ -180,6 +182,14 @@ TStatId FIKRigEditorToolkit::GetStatId() const
 	RETURN_QUICK_DECLARE_CYCLE_STAT(FIKRigEditorToolkit, STATGROUP_Tickables);
 }
 
+void FIKRigEditorToolkit::HandleOnPreviewSceneSettingsCustomized(IDetailLayoutBuilder& DetailBuilder) const
+{
+	DetailBuilder.HideCategory("Additional Meshes");
+	DetailBuilder.HideCategory("Physics");
+	DetailBuilder.HideCategory("Mesh");
+	DetailBuilder.HideCategory("Animation Blueprint");
+}
+
 void FIKRigEditorToolkit::PostUndo(bool bSuccess)
 {
 	EditorController->AssetController->BroadcastNeedsReinitialized();
diff --git a/Engine/Plugins/Animation/IKRig/Source/IKRigEditor/Public/RetargetEditor/IKRetargetEditor.h b/Engine/Plugins/Animation/IKRig/Source/IKRigEditor/Public/RetargetEditor/IKRetargetEditor.h
--- a/Engine/Plugins/Animation/IKRig/Source/IKRigEditor/Public/RetargetEditor/IKRetargetEditor.h
+++ b/Engine/Plugins/Animation/IKRig/Source/IKRigEditor/Public/RetargetEditor/IKRetargetEditor.h
@@ -94,6 +94,7 @@ private:
 	void HandleDetailsCreated(const TSharedRef<class IDetailsView>& InDetailsView);
 	void OnFinishedChangingDetails(const FPropertyChangedEvent& PropertyChangedEvent);
 	void SetupAnimInstance();
+	void HandleOnPreviewSceneSettingsCustomized(IDetailLayoutBuilder& DetailBuilder) const;
 	/** END preview scene setup */
 	
 	/** centralized management across all views */
diff --git a/Engine/Plugins/Animation/IKRig/Source/IKRigEditor/Public/RigEditor/IKRigToolkit.h b/Engine/Plugins/Animation/IKRig/Source/IKRigEditor/Public/RigEditor/IKRigToolkit.h
--- a/Engine/Plugins/Animation/IKRig/Source/IKRigEditor/Public/RigEditor/IKRigToolkit.h
+++ b/Engine/Plugins/Animation/IKRig/Source/IKRigEditor/Public/RigEditor/IKRigToolkit.h
@@ -71,6 +71,7 @@ public:
 	void HandlePreviewSceneCreated(const TSharedRef<IPersonaPreviewScene>& InPersonaPreviewScene);
 	void HandleDetailsCreated(const TSharedRef<class IDetailsView>& InDetailsView) const;
 	void HandleViewportCreated(const TSharedRef<class IPersonaViewport>& InViewport);
+	void HandleOnPreviewSceneSettingsCustomized(IDetailLayoutBuilder& DetailBuilder) const;
 	/** END preview scene setup */
 
 private:
diff --git a/Engine/Source/Editor/AnimationBlueprintEditor/Private/AnimationBlueprintEditor.cpp b/Engine/Source/Editor/AnimationBlueprintEditor/Private/AnimationBlueprintEditor.cpp
--- a/Engine/Source/Editor/AnimationBlueprintEditor/Private/AnimationBlueprintEditor.cpp
+++ b/Engine/Source/Editor/AnimationBlueprintEditor/Private/AnimationBlueprintEditor.cpp
@@ -97,7 +97,10 @@
 #include "PersonaModule.h"
 #include "PersonaToolMenuContext.h"
 #include "PersonaUtils.h"
+#include "Persona/Private/AnimationEditorPreviewScene.h"
 // Hide related nodes feature
+#include "DetailLayoutBuilder.h"
+#include "DetailWidgetRow.h"
 #include "Preferences/AnimationBlueprintEditorOptions.h"
 #include "PropertyEditorDelegates.h"
 #include "SBlueprintEditorToolbar.h"
@@ -115,6 +118,8 @@
 #include "Templates/Casts.h"
 #include "Templates/SubclassOf.h"
 #include "ToolMenuContext.h"
+#include "Persona/Private/PersonaPreviewSceneDescription.h"
+#include "Preferences/PersonaOptions.h"
 #include "Toolkits/AssetEditorToolkit.h"
 #include "UObject/Object.h"
 #include "UObject/ObjectPtr.h"
@@ -126,6 +131,7 @@
 #include "Widgets/SBoxPanel.h"
 #include "Widgets/Text/STextBlock.h"
 #include "WorkflowOrientedApp/WorkflowTabManager.h"
+#include "PropertyCustomizationHelpers.h"
 
 class FEditorModeTools;
 class FToolBarBuilder;
@@ -377,8 +383,11 @@ void FAnimationBlueprintEditor::InitAnimationBlueprintEditor(const EToolkitMode:
 
 	GetToolkitCommands()->Append(FPlayWorldCommands::GlobalPlayWorldActions.ToSharedRef());
 
+	FPersonaToolkitArgs PersonaToolkitArgs;
+	PersonaToolkitArgs.OnPreviewSceneSettingsCustomized = FOnPreviewSceneSettingsCustomized::FDelegate::CreateSP(this, &FAnimationBlueprintEditor::HandleOnPreviewSceneSettingsCustomized);
+
 	FPersonaModule& PersonaModule = FModuleManager::GetModuleChecked<FPersonaModule>("Persona");
-	PersonaToolkit = PersonaModule.CreatePersonaToolkit(InAnimBlueprint);
+	PersonaToolkit = PersonaModule.CreatePersonaToolkit(InAnimBlueprint, PersonaToolkitArgs);
 
 	PersonaToolkit->GetPreviewScene()->SetDefaultAnimationMode(EPreviewSceneDefaultAnimationMode::AnimationBlueprint);
 	PersonaToolkit->GetPreviewScene()->RegisterOnPreviewMeshChanged(FOnPreviewMeshChanged::CreateSP(this, &FAnimationBlueprintEditor::HandlePreviewMeshChanged));
@@ -2081,6 +2090,129 @@ void FAnimationBlueprintEditor::HandlePoseWatchSelectedNodes()
 	}
 }
 
+void FAnimationBlueprintEditor::HandleOnPreviewSceneSettingsCustomized(IDetailLayoutBuilder& DetailBuilder)
+{
+	const TSharedRef<IPropertyHandle> PreviewAnimationBlueprintProperty = DetailBuilder.GetProperty(GET_MEMBER_NAME_CHECKED(UPersonaPreviewSceneDescription, PreviewAnimationBlueprint));
+	const TSharedRef<IPropertyHandle> ApplicationMethodProperty = DetailBuilder.GetProperty(GET_MEMBER_NAME_CHECKED(UPersonaPreviewSceneDescription, ApplicationMethod));
+	const TSharedRef<IPropertyHandle> LinkedAnimGraphTagProperty = DetailBuilder.GetProperty(GET_MEMBER_NAME_CHECKED(UPersonaPreviewSceneDescription, LinkedAnimGraphTag));
+	
+	// customize "Preview Animation Blueprint" for animation blueprint classes
+	DetailBuilder.EditCategory("Animation Blueprint")
+	.AddProperty(PreviewAnimationBlueprintProperty)
+	.CustomWidget()
+	.NameContent()
+	[
+		SNew(SVerticalBox)
+		+SVerticalBox::Slot()
+		.AutoHeight()
+		[
+			PreviewAnimationBlueprintProperty->CreatePropertyNameWidget()
+		]
+	]
+	.ValueContent()
+	.MaxDesiredWidth(250.0f)
+	.MinDesiredWidth(250.0f)
+	[
+		SNew(SObjectPropertyEntryBox)
+		.AllowedClass(UAnimBlueprint::StaticClass())
+		.PropertyHandle(PreviewAnimationBlueprintProperty)
+		.OnShouldFilterAsset(this, &FAnimationBlueprintEditor::HandleShouldFilterAsset, FName("TargetSkeleton"), false)
+		.OnObjectChanged(this, &FAnimationBlueprintEditor::HandlePreviewAnimBlueprintChanged)
+		.ThumbnailPool(DetailBuilder.GetThumbnailPool())
+	];
+
+	ApplicationMethodProperty->SetOnPropertyValueChanged(FSimpleDelegate::CreateLambda([this]()
+	{
+		FScopedTransaction Transaction(LOCTEXT("SetAnimationBlueprintApplicationMethod", "Set Application Method"));
+		
+		const TSharedRef<FAnimationEditorPreviewScene> LocalPreviewScene = StaticCastSharedRef<FAnimationEditorPreviewScene>(PersonaToolkit->GetPreviewScene());
+		const UPersonaPreviewSceneDescription* PersonaPreviewSceneDescription = LocalPreviewScene->GetPreviewSceneDescription();
+		PersonaToolkit->GetAnimBlueprint()->SetPreviewAnimationBlueprintApplicationMethod(PersonaPreviewSceneDescription->ApplicationMethod);
+		LocalPreviewScene->SetPreviewAnimationBlueprint(PersonaPreviewSceneDescription->PreviewAnimationBlueprint.Get(), PersonaToolkit->GetAnimBlueprint());
+	}));
+
+	DetailBuilder.EditCategory("Animation Blueprint")
+	.AddProperty(ApplicationMethodProperty)
+	.IsEnabled(MakeAttributeLambda([this]()
+	{
+		TSharedRef<FAnimationEditorPreviewScene> LocalPreviewScene = StaticCastSharedRef<FAnimationEditorPreviewScene>(PersonaToolkit->GetPreviewScene());
+		UPersonaPreviewSceneDescription* PersonaPreviewSceneDescription = LocalPreviewScene->GetPreviewSceneDescription();
+		
+		return PersonaPreviewSceneDescription->PreviewAnimationBlueprint.IsValid();
+	}));
+
+	LinkedAnimGraphTagProperty->SetOnPropertyValueChanged(FSimpleDelegate::CreateLambda([this]()
+	{
+		FScopedTransaction Transaction(LOCTEXT("SetAnimationBlueprintTag", "Set Linked Anim Graph Tag"));
+		
+		TSharedRef<FAnimationEditorPreviewScene> LocalPreviewScene = StaticCastSharedRef<FAnimationEditorPreviewScene>(PersonaToolkit->GetPreviewScene());
+		UPersonaPreviewSceneDescription* PersonaPreviewSceneDescription = LocalPreviewScene->GetPreviewSceneDescription();
+		PersonaToolkit->GetAnimBlueprint()->SetPreviewAnimationBlueprintTag(PersonaPreviewSceneDescription->LinkedAnimGraphTag);
+		LocalPreviewScene->SetPreviewAnimationBlueprint(PersonaPreviewSceneDescription->PreviewAnimationBlueprint.Get(), PersonaToolkit->GetAnimBlueprint());
+	}));
+
+	DetailBuilder.EditCategory("Animation Blueprint")
+	.AddProperty(LinkedAnimGraphTagProperty)
+	.IsEnabled(MakeAttributeLambda([this]()
+	{
+		TSharedRef<FAnimationEditorPreviewScene> LocalPreviewScene = StaticCastSharedRef<FAnimationEditorPreviewScene>(PersonaToolkit->GetPreviewScene());
+		UPersonaPreviewSceneDescription* PersonaPreviewSceneDescription = LocalPreviewScene->GetPreviewSceneDescription();
+		
+		return PersonaPreviewSceneDescription->PreviewAnimationBlueprint.IsValid() && PersonaPreviewSceneDescription->ApplicationMethod == EPreviewAnimationBlueprintApplicationMethod::LinkedAnimGraph;
+	}));
+	
+	if (PersonaToolkit->GetAnimBlueprint())
+	{
+		PersonaToolkit->GetAnimBlueprint()->OnCompiled().AddSP(this, &FAnimationBlueprintEditor::HandleAnimBlueprintCompiled);
+	}
+}
+
+void FAnimationBlueprintEditor::HandleAnimBlueprintCompiled(UBlueprint* Blueprint) const
+{
+	// Only re-initialize controller if we are not debugging an external instance.
+	// If we switch at this point then we will disconnect from the external instance
+	const TSharedRef<FAnimationEditorPreviewScene> AnimPreviewScene = StaticCastSharedRef<FAnimationEditorPreviewScene>(GetPreviewScene());
+	if(AnimPreviewScene->GetPreviewMeshComponent()->PreviewInstance == nullptr || AnimPreviewScene->GetPreviewMeshComponent()->PreviewInstance->GetDebugSkeletalMeshComponent() == nullptr)
+	{
+		UPersonaPreviewSceneDescription* PersonaPreviewSceneDescription = AnimPreviewScene->GetPreviewSceneDescription();
+		PersonaPreviewSceneDescription->PreviewControllerInstance->UninitializeView(PersonaPreviewSceneDescription, &GetPreviewScene().Get());
+		PersonaPreviewSceneDescription->PreviewControllerInstance->InitializeView(PersonaPreviewSceneDescription, &GetPreviewScene().Get());
+	}
+}
+
+void FAnimationBlueprintEditor::HandlePreviewAnimBlueprintChanged(const FAssetData& InAssetData) const   
+{
+	UAnimBlueprint* NewAnimBlueprint = Cast<UAnimBlueprint>(InAssetData.GetAsset());
+	PersonaToolkit->SetPreviewAnimationBlueprint(NewAnimBlueprint);
+}
+
+bool FAnimationBlueprintEditor::HandleShouldFilterAsset(
+	const FAssetData& InAssetData,
+	FName InTag,
+	bool bCanUseDifferentSkeleton) const
+{
+	if (bCanUseDifferentSkeleton && GetDefault<UPersonaOptions>()->bAllowPreviewMeshCollectionsToSelectFromDifferentSkeletons)
+	{
+		return false;
+	}
+
+	const TSharedRef<FAnimationEditorPreviewScene> AnimPreviewScene = StaticCastSharedRef<FAnimationEditorPreviewScene>(GetPreviewScene());
+	const UPersonaPreviewSceneDescription* PersonaPreviewSceneDescription = AnimPreviewScene->GetPreviewSceneDescription();
+	if(!PersonaPreviewSceneDescription->PreviewMesh.IsValid())
+	{
+		return false;
+	}
+	
+	const USkeleton* Skeleton = PersonaPreviewSceneDescription->PreviewMesh->GetSkeleton();
+	const FString SkeletonTag = InAssetData.GetTagValueRef<FString>(InTag);
+	if (Skeleton && Skeleton->IsCompatibleSkeletonByAssetString(SkeletonTag))
+	{
+		return false;
+	}
+
+	return true;
+}
+
 void FAnimationBlueprintEditor::RemoveAllSelectionPoseWatches()
 {
 	TSharedPtr<SGraphEditor> FocusedGraphEd = FocusedGraphEdPtr.Pin();
diff --git a/Engine/Source/Editor/AnimationBlueprintEditor/Private/AnimationBlueprintEditor.h b/Engine/Source/Editor/AnimationBlueprintEditor/Private/AnimationBlueprintEditor.h
--- a/Engine/Source/Editor/AnimationBlueprintEditor/Private/AnimationBlueprintEditor.h
+++ b/Engine/Source/Editor/AnimationBlueprintEditor/Private/AnimationBlueprintEditor.h
@@ -330,6 +330,12 @@ private:
 	/** Enable/disable pose watch on selected nodes */
 	void HandlePoseWatchSelectedNodes();
 
+	/** Customize Preview Scene Settings details panel */
+	void HandleOnPreviewSceneSettingsCustomized(IDetailLayoutBuilder& DetailBuilder);
+	void HandlePreviewAnimBlueprintChanged(const FAssetData& InAssetData) const;
+	bool HandleShouldFilterAsset(const FAssetData& InAssetData, FName InTag, bool bCanUseDifferentSkeleton) const;
+	void HandleAnimBlueprintCompiled(UBlueprint* Blueprint) const;
+
 	/** Removes all pose watches created by selection from the current view */
 	void RemoveAllSelectionPoseWatches();
 
diff --git a/Engine/Source/Editor/AnimationEditor/Private/AnimationEditor.cpp b/Engine/Source/Editor/AnimationEditor/Private/AnimationEditor.cpp
--- a/Engine/Source/Editor/AnimationEditor/Private/AnimationEditor.cpp
+++ b/Engine/Source/Editor/AnimationEditor/Private/AnimationEditor.cpp
@@ -19,6 +19,7 @@
 #include "AnimationEditorMode.h"
 #include "AnimationEditorUtils.h"
 #include "AnimationToolMenuContext.h"
+#include "DetailLayoutBuilder.h"
 #include "AssetRegistry/AssetData.h"
 #include "Curves/RichCurve.h"
 #include "Editor.h"
@@ -139,8 +140,11 @@ void FAnimationEditor::InitAnimationEditor(const EToolkitMode::Type Mode, const
 	FReimportManager::Instance()->OnPostReimport().AddRaw(this, &FAnimationEditor::HandlePostReimport);
 	GEditor->GetEditorSubsystem<UImportSubsystem>()->OnAssetPostImport.AddRaw(this, &FAnimationEditor::HandlePostImport);
 
+	FPersonaToolkitArgs PersonaToolkitArgs;
+	PersonaToolkitArgs.OnPreviewSceneSettingsCustomized = FOnPreviewSceneSettingsCustomized::FDelegate::CreateSP(this, &FAnimationEditor::HandleOnPreviewSceneSettingsCustomized);
+	
 	FPersonaModule& PersonaModule = FModuleManager::LoadModuleChecked<FPersonaModule>("Persona");
-	PersonaToolkit = PersonaModule.CreatePersonaToolkit(InAnimationAsset);
+	PersonaToolkit = PersonaModule.CreatePersonaToolkit(InAnimationAsset, PersonaToolkitArgs);
 
 	PersonaToolkit->GetPreviewScene()->SetDefaultAnimationMode(EPreviewSceneDefaultAnimationMode::Animation);
 
@@ -268,6 +272,11 @@ TSharedPtr<FAnimationEditor> FAnimationEditor::GetAnimationEditor(const FToolMen
 	return TSharedPtr<FAnimationEditor>();
 }
 
+void FAnimationEditor::HandleOnPreviewSceneSettingsCustomized(IDetailLayoutBuilder& DetailBuilder) const
+{
+	DetailBuilder.HideCategory("Animation Blueprint");
+}
+
 void FAnimationEditor::ExtendToolbar()
 {
 	FToolMenuOwnerScoped OwnerScoped(this);
diff --git a/Engine/Source/Editor/AnimationEditor/Private/AnimationEditor.h b/Engine/Source/Editor/AnimationEditor/Private/AnimationEditor.h
--- a/Engine/Source/Editor/AnimationEditor/Private/AnimationEditor.h
+++ b/Engine/Source/Editor/AnimationEditor/Private/AnimationEditor.h
@@ -19,6 +19,7 @@
 #include "UObject/GCObject.h"
 #include "UObject/NameTypes.h"
 
+class IDetailLayoutBuilder;
 class FExtender;
 class FMenuBuilder;
 class FReferenceCollector;
@@ -172,6 +173,9 @@ private:
 	bool RecordMeshToAnimation(USkeletalMeshComponent* PreviewComponent, UAnimSequence* NewAsset) const;
 
 	static TSharedPtr<FAnimationEditor> GetAnimationEditor(const FToolMenuContext& InMenuContext);
+
+	void HandleOnPreviewSceneSettingsCustomized(IDetailLayoutBuilder& DetailBuilder) const;
+	
 public:
 	/** Multicast delegate fired on global undo/redo */
 	FSimpleMulticastDelegate OnLODChanged;
diff --git a/Engine/Source/Editor/Persona/Private/PersonaModule.cpp b/Engine/Source/Editor/Persona/Private/PersonaModule.cpp
--- a/Engine/Source/Editor/Persona/Private/PersonaModule.cpp
+++ b/Engine/Source/Editor/Persona/Private/PersonaModule.cpp
@@ -184,77 +184,45 @@ void FPersonaModule::ShutdownModule()
 	}
 }
 
-static void SetupPersonaToolkit(const TSharedRef<FPersonaToolkit>& Toolkit, const FPersonaToolkitArgs& PersonaToolkitArgs)
-{
-	if (PersonaToolkitArgs.bCreatePreviewScene)
-	{
-		Toolkit->CreatePreviewScene(PersonaToolkitArgs);
-	}
-}
-
 TSharedRef<IPersonaToolkit> FPersonaModule::CreatePersonaToolkit(UObject* InAsset, const FPersonaToolkitArgs& PersonaToolkitArgs, USkeleton* InSkeleton) const
 {
 	TSharedRef<FPersonaToolkit> NewPersonaToolkit(new FPersonaToolkit());
-
-	NewPersonaToolkit->Initialize(InAsset, InSkeleton);
-
-	SetupPersonaToolkit(NewPersonaToolkit, PersonaToolkitArgs);
-
+	NewPersonaToolkit->Initialize(InAsset, PersonaToolkitArgs, InSkeleton);
 	return NewPersonaToolkit;
 }
 
 TSharedRef<IPersonaToolkit> FPersonaModule::CreatePersonaToolkit(USkeleton* InSkeleton, const FPersonaToolkitArgs& PersonaToolkitArgs) const
 {
 	TSharedRef<FPersonaToolkit> NewPersonaToolkit(new FPersonaToolkit());
-
-	NewPersonaToolkit->Initialize(InSkeleton);
-
-	SetupPersonaToolkit(NewPersonaToolkit, PersonaToolkitArgs);
-
+	NewPersonaToolkit->Initialize(InSkeleton, PersonaToolkitArgs);
 	return NewPersonaToolkit;
 }
 
 TSharedRef<IPersonaToolkit> FPersonaModule::CreatePersonaToolkit(UAnimationAsset* InAnimationAsset, const FPersonaToolkitArgs& PersonaToolkitArgs) const
 {
 	TSharedRef<FPersonaToolkit> NewPersonaToolkit(new FPersonaToolkit());
-
-	NewPersonaToolkit->Initialize(InAnimationAsset);
-
-	SetupPersonaToolkit(NewPersonaToolkit, PersonaToolkitArgs);
-
+	NewPersonaToolkit->Initialize(InAnimationAsset, PersonaToolkitArgs);
 	return NewPersonaToolkit;
 }
 
 TSharedRef<IPersonaToolkit> FPersonaModule::CreatePersonaToolkit(USkeletalMesh* InSkeletalMesh, const FPersonaToolkitArgs& PersonaToolkitArgs) const
 {
 	TSharedRef<FPersonaToolkit> NewPersonaToolkit(new FPersonaToolkit());
-
-	NewPersonaToolkit->Initialize(InSkeletalMesh);
-
-	SetupPersonaToolkit(NewPersonaToolkit, PersonaToolkitArgs);
-
+	NewPersonaToolkit->Initialize(InSkeletalMesh, PersonaToolkitArgs);
 	return NewPersonaToolkit;
 }
 
 TSharedRef<IPersonaToolkit> FPersonaModule::CreatePersonaToolkit(UAnimBlueprint* InAnimBlueprint, const FPersonaToolkitArgs& PersonaToolkitArgs) const
 {
 	TSharedRef<FPersonaToolkit> NewPersonaToolkit(new FPersonaToolkit());
-
-	NewPersonaToolkit->Initialize(InAnimBlueprint);
-
-	SetupPersonaToolkit(NewPersonaToolkit, PersonaToolkitArgs);
-
+	NewPersonaToolkit->Initialize(InAnimBlueprint, PersonaToolkitArgs);
 	return NewPersonaToolkit;
 }
 
 TSharedRef<IPersonaToolkit> FPersonaModule::CreatePersonaToolkit(UPhysicsAsset* InPhysicsAsset, const FPersonaToolkitArgs& PersonaToolkitArgs) const
 {
 	TSharedRef<FPersonaToolkit> NewPersonaToolkit(new FPersonaToolkit());
-
-	NewPersonaToolkit->Initialize(InPhysicsAsset);
-
-	SetupPersonaToolkit(NewPersonaToolkit, PersonaToolkitArgs);
-
+	NewPersonaToolkit->Initialize(InPhysicsAsset, PersonaToolkitArgs);
 	return NewPersonaToolkit;
 }
 
diff --git a/Engine/Source/Editor/Persona/Private/PersonaToolkit.cpp b/Engine/Source/Editor/Persona/Private/PersonaToolkit.cpp
--- a/Engine/Source/Editor/Persona/Private/PersonaToolkit.cpp
+++ b/Engine/Source/Editor/Persona/Private/PersonaToolkit.cpp
@@ -41,7 +41,7 @@ static void FindCounterpartAssets(const UObject* InAsset, TWeakObjectPtr<USkelet
 	OutMesh = const_cast<USkeletalMesh*>(CounterpartMesh);
 }
 
-void FPersonaToolkit::Initialize(UObject* InAsset, USkeleton* InSkeleton)
+void FPersonaToolkit::Initialize(UObject* InAsset, const FPersonaToolkitArgs& PersonaToolkitArgs, USkeleton* InSkeleton)
 {
 	Asset = InAsset;
 	InitialAssetClass = Asset->GetClass();
@@ -52,51 +52,73 @@ void FPersonaToolkit::Initialize(UObject* InAsset, USkeleton* InSkeleton)
 	}
 
 	Skeleton = InSkeleton;
+
+	CommonInitialSetup(PersonaToolkitArgs);
 }
 
-void FPersonaToolkit::Initialize(USkeleton* InSkeleton)
+void FPersonaToolkit::Initialize(USkeleton* InSkeleton, const FPersonaToolkitArgs& PersonaToolkitArgs)
 {
 	check(InSkeleton);
 	Skeleton = InSkeleton;
 	InitialAssetClass = USkeleton::StaticClass();
 
 	FindCounterpartAssets(InSkeleton, Skeleton, Mesh);
+
+	CommonInitialSetup(PersonaToolkitArgs);
 }
 
-void FPersonaToolkit::Initialize(UAnimationAsset* InAnimationAsset)
+void FPersonaToolkit::Initialize(UAnimationAsset* InAnimationAsset, const FPersonaToolkitArgs& PersonaToolkitArgs)
 {
 	check(InAnimationAsset);
 	AnimationAsset = InAnimationAsset;
 	InitialAssetClass = UAnimationAsset::StaticClass();
 
 	FindCounterpartAssets(InAnimationAsset, Skeleton, Mesh);
+
+	CommonInitialSetup(PersonaToolkitArgs);
 }
 
-void FPersonaToolkit::Initialize(USkeletalMesh* InSkeletalMesh)
+void FPersonaToolkit::Initialize(USkeletalMesh* InSkeletalMesh, const FPersonaToolkitArgs& PersonaToolkitArgs)
 {
 	check(InSkeletalMesh);
 	Mesh = InSkeletalMesh;
 	InitialAssetClass = USkeletalMesh::StaticClass();
 
 	FindCounterpartAssets(InSkeletalMesh, Skeleton, Mesh);
+
+	CommonInitialSetup(PersonaToolkitArgs);
 }
 
-void FPersonaToolkit::Initialize(UAnimBlueprint* InAnimBlueprint)
+void FPersonaToolkit::Initialize(UAnimBlueprint* InAnimBlueprint, const FPersonaToolkitArgs& PersonaToolkitArgs)
 {
 	check(InAnimBlueprint);
 	AnimBlueprint = InAnimBlueprint;
 	InitialAssetClass = UAnimBlueprint::StaticClass();
 
 	FindCounterpartAssets(InAnimBlueprint, Skeleton, Mesh);
+
+	CommonInitialSetup(PersonaToolkitArgs);
 }
 
-void FPersonaToolkit::Initialize(UPhysicsAsset* InPhysicsAsset)
+void FPersonaToolkit::Initialize(UPhysicsAsset* InPhysicsAsset, const FPersonaToolkitArgs& PersonaToolkitArgs)
 {
 	check(InPhysicsAsset);
 	PhysicsAsset = InPhysicsAsset;
 	InitialAssetClass = UPhysicsAsset::StaticClass();
 
 	FindCounterpartAssets(InPhysicsAsset, Skeleton, Mesh);
+
+	CommonInitialSetup(PersonaToolkitArgs);
+}
+
+void FPersonaToolkit::CommonInitialSetup(const FPersonaToolkitArgs& PersonaToolkitArgs)
+{
+	if (PersonaToolkitArgs.bCreatePreviewScene)
+	{
+		CreatePreviewScene(PersonaToolkitArgs);
+	}
+
+	OnPreviewSceneSettingsCustomized = PersonaToolkitArgs.OnPreviewSceneSettingsCustomized;
 }
 
 void FPersonaToolkit::CreatePreviewScene(const FPersonaToolkitArgs& PersonaToolkitArgs)
@@ -382,6 +404,11 @@ void FPersonaToolkit::SetCustomData(const int32 Key, const int32 CustomData)
 	CustomEditorData.FindOrAdd(Key) = CustomData;
 }
 
+void FPersonaToolkit::CustomizeSceneSettings(IDetailLayoutBuilder& DetailBuilder)
+{
+	OnPreviewSceneSettingsCustomized.ExecuteIfBound(DetailBuilder);
+}
+
 FName FPersonaToolkit::GetContext() const
 {
 	if (InitialAssetClass != nullptr)
diff --git a/Engine/Source/Editor/Persona/Private/PersonaToolkit.h b/Engine/Source/Editor/Persona/Private/PersonaToolkit.h
--- a/Engine/Source/Editor/Persona/Private/PersonaToolkit.h
+++ b/Engine/Source/Editor/Persona/Private/PersonaToolkit.h
@@ -6,6 +6,7 @@
 #include "IPersonaPreviewScene.h"
 #include "IEditableSkeleton.h"
 #include "IPersonaToolkit.h"
+#include "PersonaModule.h"
 
 class FAnimationEditorPreviewScene;
 class UAnimationAsset;
@@ -20,17 +21,20 @@ public:
 	FPersonaToolkit();
 	virtual ~FPersonaToolkit();
 
-	/** Initialize from a various sources */
-	void Initialize(UObject* InAsset, USkeleton* InSkeleton = nullptr);
-	void Initialize(USkeleton* InSkeleton);
-	void Initialize(UAnimationAsset* InAnimationAsset);
-	void Initialize(USkeletalMesh* InSkeletalMesh);
-	void Initialize(UAnimBlueprint* InAnimBlueprint);
-	void Initialize(UPhysicsAsset* InPhysicsAsset);
+	/** Initialize from various sources */
+	void Initialize(UObject* InAsset, const FPersonaToolkitArgs& PersonaToolkitArgs, USkeleton* InSkeleton = nullptr);
+	void Initialize(USkeleton* InSkeleton, const FPersonaToolkitArgs& PersonaToolkitArgs);
+	void Initialize(UAnimationAsset* InAnimationAsset, const FPersonaToolkitArgs& PersonaToolkitArgs);
+	void Initialize(USkeletalMesh* InSkeletalMesh, const FPersonaToolkitArgs& PersonaToolkitArgs);
+	void Initialize(UAnimBlueprint* InAnimBlueprint, const FPersonaToolkitArgs& PersonaToolkitArgs);
+	void Initialize(UPhysicsAsset* InPhysicsAsset, const FPersonaToolkitArgs& PersonaToolkitArgs);
 
 	/** Optionally create a preview scene - note: creates an editable skeleton */
 	void CreatePreviewScene(const FPersonaToolkitArgs& PersonaToolkitArgs);
 
+	/** Store the options for the Preview Scene Settings details panel layout. */
+	void SetPreviewSceneDetailsOptions(const FPersonaPreviewSceneDetailsOptions& Options);
+
 	/** IPersonaToolkit interface */
 	virtual class USkeleton* GetSkeleton() const override;
 	virtual TSharedPtr<class IEditableSkeleton> GetEditableSkeleton() const override;
@@ -47,9 +51,13 @@ public:
 	virtual UAnimBlueprint* GetPreviewAnimationBlueprint() const override;
 	virtual int32 GetCustomData(const int32 Key) const override;
 	virtual void SetCustomData(const int32 Key, const int32 CustomData) override;
+	virtual void CustomizeSceneSettings(IDetailLayoutBuilder& DetailBuilder) override;
 	virtual FName GetContext() const override;
 
 private:
+	/** Common initialization */
+	void CommonInitialSetup(const FPersonaToolkitArgs& PersonaToolkitArgs);
+	
 	/** The skeleton we are editing */
 	TWeakObjectPtr<USkeleton> Skeleton;
 
@@ -77,6 +85,9 @@ private:
 	/** Preview scene for the editor */
 	TSharedPtr<FAnimationEditorPreviewScene> PreviewScene;
 
+	/* Callback to customize details tab of Preview Scene Settings */
+	FOnPreviewSceneSettingsCustomized::FDelegate OnPreviewSceneSettingsCustomized;
+
 	/** The class of the initial asset we were created with */
 	UClass* InitialAssetClass;
 };
diff --git a/Engine/Source/Editor/Persona/Private/PreviewSceneCustomizations.cpp b/Engine/Source/Editor/Persona/Private/PreviewSceneCustomizations.cpp
--- a/Engine/Source/Editor/Persona/Private/PreviewSceneCustomizations.cpp
+++ b/Engine/Source/Editor/Persona/Private/PreviewSceneCustomizations.cpp
@@ -31,6 +31,7 @@
 #include "ScopedTransaction.h"
 #include "Features/IModularFeatures.h"
 #include "AnimPreviewInstance.h"
+#include "PersonaModule.h"
 
 #define LOCTEXT_NAMESPACE "PreviewSceneCustomizations"
 
@@ -80,102 +81,107 @@ FPreviewSceneDescriptionCustomization::~FPreviewSceneDescriptionCustomization()
 void FPreviewSceneDescriptionCustomization::CustomizeDetails(IDetailLayoutBuilder& DetailBuilder)
 {
 	MyDetailLayout = &DetailBuilder;
-	FPropertyEditorModule& PropertyEditorModule = FModuleManager::LoadModuleChecked<FPropertyEditorModule>("PropertyEditor");
-
-	TSharedRef<IPropertyHandle> PreviewControllerProperty = DetailBuilder.GetProperty(GET_MEMBER_NAME_CHECKED(UPersonaPreviewSceneDescription, PreviewController));
-	TSharedRef<IPropertyHandle> SkeletalMeshProperty = DetailBuilder.GetProperty(GET_MEMBER_NAME_CHECKED(UPersonaPreviewSceneDescription, PreviewMesh));
-
-	AdditionalMeshesProperty = DetailBuilder.GetProperty(GET_MEMBER_NAME_CHECKED(UPersonaPreviewSceneDescription, AdditionalMeshes));
-	AdditionalMeshesProperty->SetOnPropertyResetToDefault(FSimpleDelegate::CreateSP(this, &FPreviewSceneDescriptionCustomization::OnResetAdditionalMeshes));
+	
+	// allow customization by client asset editor
+	PersonaToolkit.Pin()->CustomizeSceneSettings(DetailBuilder);
 
-	TArray<UClass*> BuiltInPreviewControllers = { UPersonaPreviewSceneDefaultController::StaticClass(), UPersonaPreviewSceneRefPoseController::StaticClass(), UPersonaPreviewSceneAnimationController::StaticClass() };
+	// name label given to the context of this persona instance (usually the class name of the asset)
+	const FName PersonaContextName = PersonaToolkit.Pin()->GetContext();
 
-	TArray<UClass*> DynamicPreviewControllers;
-	
-	for (TObjectIterator<UClass> It; It; ++It)
+	//
+	// Preview Controller section...
+	//
 	{
-		UClass* CurrentClass = *It;
-		if (CurrentClass->IsChildOf(UPersonaPreviewSceneController::StaticClass()) &&
-			!(CurrentClass->HasAnyClassFlags(CLASS_Abstract)) &&
-			!BuiltInPreviewControllers.Contains(CurrentClass))
+		TSharedRef<IPropertyHandle> PreviewControllerProperty = DetailBuilder.GetProperty(GET_MEMBER_NAME_CHECKED(UPersonaPreviewSceneDescription, PreviewController));
+		PreviewControllerProperty->MarkHiddenByCustomization();
+		
+		TArray<UClass*> BuiltInPreviewControllers = {
+			UPersonaPreviewSceneDefaultController::StaticClass(),
+			UPersonaPreviewSceneRefPoseController::StaticClass(),
+			UPersonaPreviewSceneAnimationController::StaticClass() };
+		TArray<UClass*> DynamicPreviewControllers;
+		for (TObjectIterator<UClass> It; It; ++It)
 		{
-			DynamicPreviewControllers.Add(CurrentClass);
+			UClass* CurrentClass = *It;
+			if (CurrentClass->IsChildOf(UPersonaPreviewSceneController::StaticClass()) &&
+				!(CurrentClass->HasAnyClassFlags(CLASS_Abstract)) &&
+				!BuiltInPreviewControllers.Contains(CurrentClass))
+			{
+				DynamicPreviewControllers.Add(CurrentClass);
+			}
 		}
-	}
-
-	Algo::SortBy(DynamicPreviewControllers, [](UClass* Cls) { return Cls->GetName(); });
 
-	ControllerItems.Reset();
-
-	for (UClass* ControllerClass : BuiltInPreviewControllers)
-	{
-		ControllerItems.Add(MakeShared<FPersonaModeComboEntry>(ControllerClass));
-	}
-	for (UClass* ControllerClass : DynamicPreviewControllers)
-	{
-		ControllerItems.Add(MakeShared<FPersonaModeComboEntry>(ControllerClass));
-	}
+		Algo::SortBy(DynamicPreviewControllers, [](UClass* Cls) { return Cls->GetName(); });
 
-	PreviewControllerProperty->MarkHiddenByCustomization();
+		ControllerItems.Reset();
+		for (UClass* ControllerClass : BuiltInPreviewControllers)
+		{
+			ControllerItems.Add(MakeShared<FPersonaModeComboEntry>(ControllerClass));
+		}
+		for (UClass* ControllerClass : DynamicPreviewControllers)
+		{
+			ControllerItems.Add(MakeShared<FPersonaModeComboEntry>(ControllerClass));
+		}
 
-	IDetailCategoryBuilder& AnimCategory = DetailBuilder.EditCategory("Animation");
-	AnimCategory.AddCustomRow(PreviewControllerProperty->GetPropertyDisplayName())
-	.NameContent()
-	[
-		PreviewControllerProperty->CreatePropertyNameWidget()
-	]
-	.ValueContent()
-	.MinDesiredWidth(200.0f)
-	[
-		SNew(SComboBox<TSharedPtr<FPersonaModeComboEntry>>)
-		.OptionsSource(&ControllerItems)
-		.OnGenerateWidget(this, &FPreviewSceneDescriptionCustomization::MakeControllerComboEntryWidget)
-		.OnSelectionChanged(this, &FPreviewSceneDescriptionCustomization::OnComboSelectionChanged)
+		IDetailCategoryBuilder& AnimCategory = DetailBuilder.EditCategory("Animation");
+		AnimCategory.AddCustomRow(PreviewControllerProperty->GetPropertyDisplayName())
+		.NameContent()
 		[
-			SNew(STextBlock)
-			.Text(this, &FPreviewSceneDescriptionCustomization::GetCurrentPreviewControllerText)
+			PreviewControllerProperty->CreatePropertyNameWidget()
 		]
-	];
-
-	TSharedPtr<FAnimationEditorPreviewScene> PreviewScenePtr = PreviewScene.Pin();
-	UPersonaPreviewSceneDescription* PersonaPreviewSceneDescription = PreviewScenePtr->GetPreviewSceneDescription();
-
-	FSimpleDelegate PropertyChangedDelegate = FSimpleDelegate::CreateSP(this, &FPreviewSceneDescriptionCustomization::HandlePreviewControllerPropertyChanged);
+		.ValueContent()
+		.MinDesiredWidth(200.0f)
+		[
+			SNew(SComboBox<TSharedPtr<FPersonaModeComboEntry>>)
+			.OptionsSource(&ControllerItems)
+			.OnGenerateWidget(this, &FPreviewSceneDescriptionCustomization::MakeControllerComboEntryWidget)
+			.OnSelectionChanged(this, &FPreviewSceneDescriptionCustomization::OnComboSelectionChanged)
+			[
+				SNew(STextBlock)
+				.Text(this, &FPreviewSceneDescriptionCustomization::GetCurrentPreviewControllerText)
+			]
+		];
 
-	for (const FProperty* TestProperty : TFieldRange<FProperty>(PersonaPreviewSceneDescription->PreviewControllerInstance->GetClass()))
-	{
-		if (TestProperty->HasAnyPropertyFlags(CPF_Edit))
+		// register PropertyValueChanged callbacks for all properties in the preview controller
+		FSimpleDelegate PropertyChangedDelegate = FSimpleDelegate::CreateSP(this, &FPreviewSceneDescriptionCustomization::HandlePreviewControllerPropertyChanged);
+		UPersonaPreviewSceneDescription* PersonaPreviewSceneDescription = PreviewScene.Pin()->GetPreviewSceneDescription();
+		for (const FProperty* TestProperty : TFieldRange<FProperty>(PersonaPreviewSceneDescription->PreviewControllerInstance->GetClass()))
 		{
-			const bool bAdvancedDisplay = TestProperty->HasAnyPropertyFlags(CPF_AdvancedDisplay);
-			const EPropertyLocation::Type PropertyLocation = bAdvancedDisplay ? EPropertyLocation::Advanced : EPropertyLocation::Common;
-
-			IDetailPropertyRow* NewRow = PersonaPreviewSceneDescription->PreviewControllerInstance->AddPreviewControllerPropertyToDetails(PersonaToolkit.Pin().ToSharedRef(), DetailBuilder, AnimCategory, TestProperty, PropertyLocation);
-			if (NewRow)
+			if (TestProperty->HasAnyPropertyFlags(CPF_Edit))
 			{
-				NewRow->GetPropertyHandle()->SetOnPropertyValueChanged(PropertyChangedDelegate);
+				const bool bAdvancedDisplay = TestProperty->HasAnyPropertyFlags(CPF_AdvancedDisplay);
+				const EPropertyLocation::Type PropertyLocation = bAdvancedDisplay ? EPropertyLocation::Advanced : EPropertyLocation::Common;
+
+				IDetailPropertyRow* NewRow = PersonaPreviewSceneDescription->PreviewControllerInstance->AddPreviewControllerPropertyToDetails(PersonaToolkit.Pin().ToSharedRef(), DetailBuilder, AnimCategory, TestProperty, PropertyLocation);
+				if (NewRow)
+				{
+					NewRow->GetPropertyHandle()->SetOnPropertyValueChanged(PropertyChangedDelegate);
+				}
 			}
 		}
 	}
 
-	// if mesh editor, we hide preview mesh section and additional mesh section
-	// sometimes additional meshes are interfering with preview mesh, it is not a great experience
-	const bool bMeshEditor = PersonaToolkit.Pin()->GetContext() == USkeletalMesh::StaticClass()->GetFName();
-	if (!bMeshEditor)
+	//
+	// Preview Mesh section...
+	//
 	{
+		TSharedRef<IPropertyHandle> SkeletalMeshProperty = DetailBuilder.GetProperty(GET_MEMBER_NAME_CHECKED(UPersonaPreviewSceneDescription, PreviewMesh));
+		SkeletalMeshProperty->MarkHiddenByCustomization();
+		
 		FText PreviewMeshName;
-		if (PersonaToolkit.Pin()->GetContext() == UAnimationAsset::StaticClass()->GetFName())
+		if (PersonaContextName == UAnimationAsset::StaticClass()->GetFName())
 		{
 			PreviewMeshName = FText::Format(LOCTEXT("PreviewMeshAnimation", "{0}\n(Animation)"), SkeletalMeshProperty->GetPropertyDisplayName());
 		}
-		else if(PersonaToolkit.Pin()->GetContext() == UAnimBlueprint::StaticClass()->GetFName())
+		else if(PersonaContextName == UAnimBlueprint::StaticClass()->GetFName())
 		{
 			PreviewMeshName = FText::Format(LOCTEXT("PreviewMeshAnimBlueprint", "{0}\n(Animation Blueprint)"), SkeletalMeshProperty->GetPropertyDisplayName());
 		}
-		else if(PersonaToolkit.Pin()->GetContext() == UPhysicsAsset::StaticClass()->GetFName())
+		else if(PersonaContextName == UPhysicsAsset::StaticClass()->GetFName())
 		{
 			PreviewMeshName = FText::Format(LOCTEXT("PreviewMeshPhysicsAsset", "{0}\n(Physics Asset)"), SkeletalMeshProperty->GetPropertyDisplayName());
 		}
-		else if(PersonaToolkit.Pin()->GetContext() == USkeleton::StaticClass()->GetFName())
+		else if(PersonaContextName == USkeleton::StaticClass()->GetFName())
 		{
 			PreviewMeshName = FText::Format(LOCTEXT("PreviewMeshSkeleton", "{0}\n(Skeleton)"), SkeletalMeshProperty->GetPropertyDisplayName());
 		}
@@ -253,121 +259,52 @@ void FPreviewSceneDescriptionCustomization::CustomizeDetails(IDetailLayoutBuilde
 			.OnObjectChanged(this, &FPreviewSceneDescriptionCustomization::HandleMeshChanged)
 			.ThumbnailPool(DetailBuilder.GetThumbnailPool())
 		];
+	}
+	
 
-		// Customize animation blueprint preview
-		TSharedRef<IPropertyHandle> PreviewAnimationBlueprintProperty = DetailBuilder.GetProperty(GET_MEMBER_NAME_CHECKED(UPersonaPreviewSceneDescription, PreviewAnimationBlueprint));
-		TSharedRef<IPropertyHandle> ApplicationMethodProperty = DetailBuilder.GetProperty(GET_MEMBER_NAME_CHECKED(UPersonaPreviewSceneDescription, ApplicationMethod));
-		TSharedRef<IPropertyHandle> LinkedAnimGraphTagProperty = DetailBuilder.GetProperty(GET_MEMBER_NAME_CHECKED(UPersonaPreviewSceneDescription, LinkedAnimGraphTag));
-		
-		if (PersonaToolkit.Pin()->GetContext() == UAnimBlueprint::StaticClass()->GetFName())
-		{
-			DetailBuilder.EditCategory("Animation Blueprint")
-			.AddProperty(PreviewAnimationBlueprintProperty)
-			.CustomWidget()
-			.NameContent()
-			[
-				SNew(SVerticalBox)
-				+SVerticalBox::Slot()
-				.AutoHeight()
-				[
-					PreviewAnimationBlueprintProperty->CreatePropertyNameWidget()
-				]
-			]
-			.ValueContent()
-			.MaxDesiredWidth(250.0f)
-			.MinDesiredWidth(250.0f)
-			[
-				SNew(SObjectPropertyEntryBox)
-				.AllowedClass(UAnimBlueprint::StaticClass())
-				.PropertyHandle(PreviewAnimationBlueprintProperty)
-				.OnShouldFilterAsset(this, &FPreviewSceneDescriptionCustomization::HandleShouldFilterAsset, FName("TargetSkeleton"), false)
-				.OnObjectChanged(this, &FPreviewSceneDescriptionCustomization::HandlePreviewAnimBlueprintChanged)
-				.ThumbnailPool(DetailBuilder.GetThumbnailPool())
-			];
-
-			ApplicationMethodProperty->SetOnPropertyValueChanged(FSimpleDelegate::CreateLambda([this]()
-			{
-				FScopedTransaction Transaction(LOCTEXT("SetAnimationBlueprintApplicationMethod", "Set Application Method"));
-
-				TSharedPtr<IPersonaToolkit> PinnedPersonaToolkit = PersonaToolkit.Pin();
-				TSharedRef<FAnimationEditorPreviewScene> LocalPreviewScene = StaticCastSharedRef<FAnimationEditorPreviewScene>(PinnedPersonaToolkit->GetPreviewScene());
-				UPersonaPreviewSceneDescription* PersonaPreviewSceneDescription = LocalPreviewScene->GetPreviewSceneDescription();
-				PinnedPersonaToolkit->GetAnimBlueprint()->SetPreviewAnimationBlueprintApplicationMethod(PersonaPreviewSceneDescription->ApplicationMethod);
-				LocalPreviewScene->SetPreviewAnimationBlueprint(PersonaPreviewSceneDescription->PreviewAnimationBlueprint.Get(), PinnedPersonaToolkit->GetAnimBlueprint());
-			}));
-
-			DetailBuilder.EditCategory("Animation Blueprint")
-			.AddProperty(ApplicationMethodProperty)
-			.IsEnabled(MakeAttributeLambda([this]()
-			{
-				TSharedPtr<IPersonaToolkit> PinnedPersonaToolkit = PersonaToolkit.Pin();
-				TSharedRef<FAnimationEditorPreviewScene> LocalPreviewScene = StaticCastSharedRef<FAnimationEditorPreviewScene>(PinnedPersonaToolkit->GetPreviewScene());
-				UPersonaPreviewSceneDescription* PersonaPreviewSceneDescription = LocalPreviewScene->GetPreviewSceneDescription();
-				
-				return PersonaPreviewSceneDescription->PreviewAnimationBlueprint.IsValid();
-			}));
-		
-			LinkedAnimGraphTagProperty->SetOnPropertyValueChanged(FSimpleDelegate::CreateLambda([this]()
-			{
-				FScopedTransaction Transaction(LOCTEXT("SetAnimationBlueprintTag", "Set Linked Anim Graph Tag"));
-
-				TSharedPtr<IPersonaToolkit> PinnedPersonaToolkit = PersonaToolkit.Pin();
-				TSharedRef<FAnimationEditorPreviewScene> LocalPreviewScene = StaticCastSharedRef<FAnimationEditorPreviewScene>(PinnedPersonaToolkit->GetPreviewScene());
-				UPersonaPreviewSceneDescription* PersonaPreviewSceneDescription = LocalPreviewScene->GetPreviewSceneDescription();
-				PinnedPersonaToolkit->GetAnimBlueprint()->SetPreviewAnimationBlueprintTag(PersonaPreviewSceneDescription->LinkedAnimGraphTag);
-				LocalPreviewScene->SetPreviewAnimationBlueprint(PersonaPreviewSceneDescription->PreviewAnimationBlueprint.Get(), PinnedPersonaToolkit->GetAnimBlueprint());
-			}));
-
-			DetailBuilder.EditCategory("Animation Blueprint")
-			.AddProperty(LinkedAnimGraphTagProperty)
-			.IsEnabled(MakeAttributeLambda([this]()
-			{
-				TSharedPtr<IPersonaToolkit> PinnedPersonaToolkit = PersonaToolkit.Pin();
-				TSharedRef<FAnimationEditorPreviewScene> LocalPreviewScene = StaticCastSharedRef<FAnimationEditorPreviewScene>(PinnedPersonaToolkit->GetPreviewScene());
-				UPersonaPreviewSceneDescription* PersonaPreviewSceneDescription = LocalPreviewScene->GetPreviewSceneDescription();
-				
-				return PersonaPreviewSceneDescription->PreviewAnimationBlueprint.IsValid() && PersonaPreviewSceneDescription->ApplicationMethod == EPreviewAnimationBlueprintApplicationMethod::LinkedAnimGraph;
-			}));
-		}
-		else
-		{
-			PreviewAnimationBlueprintProperty->MarkHiddenByCustomization();
-			ApplicationMethodProperty->MarkHiddenByCustomization();
-			LinkedAnimGraphTagProperty->MarkHiddenByCustomization();
-		}
-
+	//
+	// Physics section...
+	//
 #if CHAOS_SIMULATION_DETAIL_VIEW_FACTORY_SELECTOR
-		// Physics settings
-		ClothSimulationFactoryList.Reset();
-		const TArray<IClothingSimulationFactoryClassProvider*> ClassProviders = IModularFeatures::Get().GetModularFeatureImplementations<IClothingSimulationFactoryClassProvider>(IClothingSimulationFactoryClassProvider::FeatureName);
-		for (const auto& ClassProvider : ClassProviders)
-		{
-			// Populate cloth factory list
-			ClothSimulationFactoryList.Add(MakeShared<TSubclassOf<class UClothingSimulationFactory>>(ClassProvider->GetClothingSimulationFactoryClass()));
-		}
+	// Physics settings
+	ClothSimulationFactoryList.Reset();
+	const TArray<IClothingSimulationFactoryClassProvider*> ClassProviders = IModularFeatures::Get().GetModularFeatureImplementations<IClothingSimulationFactoryClassProvider>(IClothingSimulationFactoryClassProvider::FeatureName);
+	for (const auto& ClassProvider : ClassProviders)
+	{
+		// Populate cloth factory list
+		ClothSimulationFactoryList.Add(MakeShared<TSubclassOf<class UClothingSimulationFactory>>(ClassProvider->GetClothingSimulationFactoryClass()));
+	}
 
-		DetailBuilder.EditCategory("Physics")
-		.AddCustomRow(LOCTEXT("PhysicsClothingSimulationFactory", "Clothing Simulation Factory Option"))
-		.NameContent()
+	DetailBuilder.EditCategory("Physics")
+	.AddCustomRow(LOCTEXT("PhysicsClothingSimulationFactory", "Clothing Simulation Factory Option"))
+	.NameContent()
+	[
+		SNew(STextBlock)
+		.Font(IDetailLayoutBuilder::GetDetailFont())
+		.Text(LOCTEXT("PhysicsClothingSimulationFactory_Text", "Clothing Simulation Factory"))
+		.ToolTipText(LOCTEXT("PhysicsClothingSimulationFactory_ToolTip", "Select the cloth simulation used to preview the scene."))
+	]
+	.ValueContent()
+	.MinDesiredWidth(200.0f)
+	[
+		SNew(SComboBox<TSharedPtr<TSubclassOf<class UClothingSimulationFactory>>>)
+		.OptionsSource(&ClothSimulationFactoryList)
+		.OnGenerateWidget(this, &FPreviewSceneDescriptionCustomization::MakeClothingSimulationFactoryWidget)
+		.OnSelectionChanged(this, &FPreviewSceneDescriptionCustomization::OnClothingSimulationFactorySelectionChanged)
 		[
 			SNew(STextBlock)
-			.Font(IDetailLayoutBuilder::GetDetailFont())
-			.Text(LOCTEXT("PhysicsClothingSimulationFactory_Text", "Clothing Simulation Factory"))
-			.ToolTipText(LOCTEXT("PhysicsClothingSimulationFactory_ToolTip", "Select the cloth simulation used to preview the scene."))
+			.Text(this, &FPreviewSceneDescriptionCustomization::GetCurrentClothingSimulationFactoryText)
 		]
-		.ValueContent()
-		.MinDesiredWidth(200.0f)
-		[
-			SNew(SComboBox<TSharedPtr<TSubclassOf<class UClothingSimulationFactory>>>)
-			.OptionsSource(&ClothSimulationFactoryList)
-			.OnGenerateWidget(this, &FPreviewSceneDescriptionCustomization::MakeClothingSimulationFactoryWidget)
-			.OnSelectionChanged(this, &FPreviewSceneDescriptionCustomization::OnClothingSimulationFactorySelectionChanged)
-			[
-				SNew(STextBlock)
-				.Text(this, &FPreviewSceneDescriptionCustomization::GetCurrentClothingSimulationFactoryText)
-			]
-		];
+	];
 #endif  // #if CHAOS_SIMULATION_DETAIL_VIEW_FACTORY_SELECTOR
+
+	//
+	// Additional Meshes section...
+	//
+	{
+		AdditionalMeshesProperty = DetailBuilder.GetProperty(GET_MEMBER_NAME_CHECKED(UPersonaPreviewSceneDescription, AdditionalMeshes));
+		AdditionalMeshesProperty->SetOnPropertyResetToDefault(FSimpleDelegate::CreateSP(this, &FPreviewSceneDescriptionCustomization::OnResetAdditionalMeshes));
+		
 		// set the skeleton to use in our factory as we shouldn't be picking one here
 		FactoryToUse->CurrentSkeleton = EditableSkeleton.IsValid() ? MakeWeakObjectPtr(const_cast<USkeleton*>(&EditableSkeleton.Pin()->GetSkeleton())) : nullptr;
 		TArray<UFactory*> FactoriesToUse({ FactoryToUse });
@@ -472,17 +409,6 @@ void FPreviewSceneDescriptionCustomization::CustomizeDetails(IDetailLayoutBuilde
 			}
 		}
 	}
-	else
-	{
-		DetailBuilder.HideProperty(SkeletalMeshProperty);
-		DetailBuilder.HideProperty(AdditionalMeshesProperty);
-	}
-
-	TSharedPtr<IPersonaToolkit> Toolkit = PersonaToolkit.Pin();
-	if (Toolkit && Toolkit->GetAnimBlueprint())
-	{
-		Toolkit->GetAnimBlueprint()->OnCompiled().AddSP(this, &FPreviewSceneDescriptionCustomization::HandleAnimBlueprintCompiled);
-	}
 }
 
 EVisibility FPreviewSceneDescriptionCustomization::GetSaveButtonVisibility(TSharedRef<IPropertyHandle> InAdditionalMeshesProperty) const
@@ -601,12 +527,6 @@ void FPreviewSceneDescriptionCustomization::HandleMeshChanged(const FAssetData&
 	PersonaToolkit.Pin()->SetPreviewMesh(NewPreviewMesh, false);
 }
 
-void FPreviewSceneDescriptionCustomization::HandlePreviewAnimBlueprintChanged(const FAssetData& InAssetData)   
-{
-	UAnimBlueprint* NewAnimBlueprint = Cast<UAnimBlueprint>(InAssetData.GetAsset());
-	PersonaToolkit.Pin()->SetPreviewAnimationBlueprint(NewAnimBlueprint);
-}
-
 void FPreviewSceneDescriptionCustomization::HandleAdditionalMeshesChanged(const FAssetData& InAssetData, IDetailLayoutBuilder* DetailLayoutBuilder)
 {
 	UDataAsset* MeshCollection = Cast<UDataAsset>(InAssetData.GetAsset());
@@ -651,17 +571,6 @@ ECheckBoxState FPreviewSceneDescriptionCustomization::HandleUseCustomAnimBPIsChe
 	return GetDefault<UPersonaOptions>()->bAllowPreviewMeshCollectionsToUseCustomAnimBP? ECheckBoxState::Checked : ECheckBoxState::Unchecked;
 }
 
-void FPreviewSceneDescriptionCustomization::HandleAnimBlueprintCompiled(UBlueprint*)
-{
-	// Only re-initialize controller if we are not debugging an external instance.
-	// If we switch at this point then we will disconnect from the external instance
-	TSharedPtr<FAnimationEditorPreviewScene> PreviewScenePtr = PreviewScene.Pin();
-	if(PreviewScenePtr->GetPreviewMeshComponent()->PreviewInstance == nullptr || PreviewScenePtr->GetPreviewMeshComponent()->PreviewInstance->GetDebugSkeletalMeshComponent() == nullptr)
-	{
-		ReinitializePreviewController();
-	}
-}
-
 void FPreviewSceneDescriptionCustomization::ReinitializePreviewController()
 {
 	TSharedPtr<FAnimationEditorPreviewScene> PreviewScenePtr = PreviewScene.Pin();
diff --git a/Engine/Source/Editor/Persona/Private/PreviewSceneCustomizations.h b/Engine/Source/Editor/Persona/Private/PreviewSceneCustomizations.h
--- a/Engine/Source/Editor/Persona/Private/PreviewSceneCustomizations.h
+++ b/Engine/Source/Editor/Persona/Private/PreviewSceneCustomizations.h
@@ -70,8 +70,6 @@ private:
 
 	void HandleMeshChanged(const FAssetData& InAssetData);
 
-	void HandlePreviewAnimBlueprintChanged(const FAssetData& InAssetData);
-
 	void HandleAdditionalMeshesChanged(const FAssetData& InAssetData, IDetailLayoutBuilder* DetailLayoutBuilder);
 
 	void HandleAllowDifferentSkeletonsCheckedStateChanged(ECheckBoxState CheckState);
@@ -82,9 +80,6 @@ private:
 
 	ECheckBoxState HandleUseCustomAnimBPIsChecked() const;
 
-	// Called when the anim blueprint being edited is compiled.
-	void HandleAnimBlueprintCompiled(UBlueprint* Blueprint);
-
 	// Reinitialize the preview controller in the preview scene.
 	void ReinitializePreviewController();
 
diff --git a/Engine/Source/Editor/Persona/Public/IPersonaToolkit.h b/Engine/Source/Editor/Persona/Public/IPersonaToolkit.h
--- a/Engine/Source/Editor/Persona/Public/IPersonaToolkit.h
+++ b/Engine/Source/Editor/Persona/Public/IPersonaToolkit.h
@@ -11,6 +11,8 @@ class USkeletalMesh;
 class UAnimBlueprint;
 class UAnimationAsset;
 class IPersonaPreviewScene;
+struct FPersonaPreviewSceneDetailsOptions;
+class IDetailLayoutBuilder;
 
 /**
  * Reusable 'Persona' features for asset editors concerned with USkeleton-related assets
@@ -81,6 +83,8 @@ public:
 	 */
 	virtual void SetCustomData(const int32 Key, const int32 CustomData) {}
 
+	/** Callback to customize the Preview Scene Settings details tab. Different asset types require unique options. */
+	virtual void CustomizeSceneSettings(IDetailLayoutBuilder& DetailBuilder) {};
 	
 	/** Get the context in which this toolkit is being used (usually the class name of the asset) */
 	virtual FName GetContext() const = 0;
diff --git a/Engine/Source/Editor/Persona/Public/PersonaModule.h b/Engine/Source/Editor/Persona/Public/PersonaModule.h
--- a/Engine/Source/Editor/Persona/Public/PersonaModule.h
+++ b/Engine/Source/Editor/Persona/Public/PersonaModule.h
@@ -26,6 +26,8 @@ class FWorkflowTabFactory;
 class UBlendSpace;
 class IAnimSequenceCurveEditor;
 class IAnimationEditor;
+class IDetailLayoutBuilder;
+class FPreviewSceneDescriptionCustomization;
 
 extern const FName PersonaAppName;
 
@@ -47,16 +49,19 @@ DECLARE_DELEGATE_TwoParams(FTickRecording, USkeletalMeshComponent* /*Component*/
 DECLARE_DELEGATE_OneParam(FOnViewportCreated, const TSharedRef<class IPersonaViewport>&);
 
 /** Called back when a details panel is created */
-DECLARE_DELEGATE_OneParam(FOnDetailsCreated, const TSharedRef<class IDetailsView>&);
+DECLARE_DELEGATE_OneParam(FOnDetailsCreated, const TSharedRef<IDetailsView>&);
 
 /** Called back when an anim sequence browser is created */
-DECLARE_DELEGATE_OneParam(FOnAnimationSequenceBrowserCreated, const TSharedRef<class IAnimationSequenceBrowser>&);
+DECLARE_DELEGATE_OneParam(FOnAnimationSequenceBrowserCreated, const TSharedRef<IAnimationSequenceBrowser>&);
 
 /** Called back when a Persona preview scene is created */
-DECLARE_MULTICAST_DELEGATE_OneParam(FOnPreviewSceneCreated, const TSharedRef<class IPersonaPreviewScene>&);
+DECLARE_MULTICAST_DELEGATE_OneParam(FOnPreviewSceneCreated, const TSharedRef<IPersonaPreviewScene>&);
+
+/** Called back when a Persona preview scene settings are customized */
+DECLARE_MULTICAST_DELEGATE_OneParam(FOnPreviewSceneSettingsCustomized, IDetailLayoutBuilder& DetailBuilder);
 
 /** Called back to register tabs */
-DECLARE_MULTICAST_DELEGATE_TwoParams(FOnRegisterTabs, FWorkflowAllowedTabSet&, TSharedPtr<class FAssetEditorToolkit>);
+DECLARE_MULTICAST_DELEGATE_TwoParams(FOnRegisterTabs, FWorkflowAllowedTabSet&, TSharedPtr<FAssetEditorToolkit>);
 
 /** Called back to register common layout extensions */
 DECLARE_MULTICAST_DELEGATE_OneParam(FOnRegisterLayoutExtensions, FLayoutExtender&);
@@ -71,12 +76,14 @@ struct FPersonaToolkitArgs
 	FOnPreviewSceneCreated::FDelegate OnPreviewSceneCreated;
 
 	/** Whether to create a preview scene */
-	bool bCreatePreviewScene;
+	bool bCreatePreviewScene = true;
 
-	FPersonaToolkitArgs()
-		: bCreatePreviewScene(true)
-	{
-	}
+	/** 
+	 * Delegate called when the preview scene settings are being customized, supplies the IDetailLayoutBuilder
+	 * for the user to customize the layout however they wish. */
+	FOnPreviewSceneSettingsCustomized::FDelegate OnPreviewSceneSettingsCustomized;
+
+	FPersonaToolkitArgs() = default;
 };
 
 struct FAnimDocumentArgs
@@ -249,12 +256,12 @@ public:
 	virtual void ShutdownModule();
 
 	/** Create a re-usable toolkit that multiple asset editors that are concerned with USkeleton-related data can use */
-	virtual TSharedRef<IPersonaToolkit> CreatePersonaToolkit(UObject* InAsset, const FPersonaToolkitArgs& PersonaToolkitArgs = FPersonaToolkitArgs(), USkeleton* InSkeleton = nullptr) const;
-	virtual TSharedRef<IPersonaToolkit> CreatePersonaToolkit(USkeleton* InSkeleton, const FPersonaToolkitArgs& PersonaToolkitArgs = FPersonaToolkitArgs()) const;
-	virtual TSharedRef<IPersonaToolkit> CreatePersonaToolkit(UAnimationAsset* InAnimationAsset, const FPersonaToolkitArgs& PersonaToolkitArgs = FPersonaToolkitArgs()) const;
-	virtual TSharedRef<IPersonaToolkit> CreatePersonaToolkit(USkeletalMesh* InSkeletalMesh, const FPersonaToolkitArgs& PersonaToolkitArgs = FPersonaToolkitArgs()) const;
-	virtual TSharedRef<IPersonaToolkit> CreatePersonaToolkit(UAnimBlueprint* InAnimBlueprint, const FPersonaToolkitArgs& PersonaToolkitArgs = FPersonaToolkitArgs()) const;
-	virtual TSharedRef<IPersonaToolkit> CreatePersonaToolkit(UPhysicsAsset* InPhysicsAsset, const FPersonaToolkitArgs& PersonaToolkitArgs = FPersonaToolkitArgs()) const;
+	virtual TSharedRef<IPersonaToolkit> CreatePersonaToolkit(UObject* InAsset, const FPersonaToolkitArgs& PersonaToolkitArgs, USkeleton* InSkeleton = nullptr) const;
+	virtual TSharedRef<IPersonaToolkit> CreatePersonaToolkit(USkeleton* InSkeleton, const FPersonaToolkitArgs& PersonaToolkitArgs) const;
+	virtual TSharedRef<IPersonaToolkit> CreatePersonaToolkit(UAnimationAsset* InAnimationAsset, const FPersonaToolkitArgs& PersonaToolkitArgs) const;
+	virtual TSharedRef<IPersonaToolkit> CreatePersonaToolkit(USkeletalMesh* InSkeletalMesh, const FPersonaToolkitArgs& PersonaToolkitArgs) const;
+	virtual TSharedRef<IPersonaToolkit> CreatePersonaToolkit(UAnimBlueprint* InAnimBlueprint, const FPersonaToolkitArgs& PersonaToolkitArgs) const;
+	virtual TSharedRef<IPersonaToolkit> CreatePersonaToolkit(UPhysicsAsset* InPhysicsAsset, const FPersonaToolkitArgs& PersonaToolkitArgs) const;
 
 	/** Create an asset family for the supplied persona asset */
 	virtual TSharedRef<IAssetFamily> CreatePersonaAssetFamily(const UObject* InAsset) const;
diff --git a/Engine/Source/Editor/PhysicsAssetEditor/Private/PhysicsAssetEditor.cpp b/Engine/Source/Editor/PhysicsAssetEditor/Private/PhysicsAssetEditor.cpp
--- a/Engine/Source/Editor/PhysicsAssetEditor/Private/PhysicsAssetEditor.cpp
+++ b/Engine/Source/Editor/PhysicsAssetEditor/Private/PhysicsAssetEditor.cpp
@@ -178,11 +178,11 @@ void FPhysicsAssetEditor::InitPhysicsAssetEditor(const EToolkitMode::Type Mode,
 
 	SharedData->CachePreviewMesh();
 
-	FPersonaModule& PersonaModule = FModuleManager::LoadModuleChecked<FPersonaModule>("Persona");
-
 	FPersonaToolkitArgs PersonaToolkitArgs;
 	PersonaToolkitArgs.OnPreviewSceneCreated = FOnPreviewSceneCreated::FDelegate::CreateSP(this, &FPhysicsAssetEditor::HandlePreviewSceneCreated);
+	PersonaToolkitArgs.OnPreviewSceneSettingsCustomized = FOnPreviewSceneSettingsCustomized::FDelegate::CreateSP(this, &FPhysicsAssetEditor::HandleOnPreviewSceneSettingsCustomized);
 
+	FPersonaModule& PersonaModule = FModuleManager::LoadModuleChecked<FPersonaModule>("Persona");
 	PersonaToolkit = PersonaModule.CreatePersonaToolkit(SharedData->PhysicsAsset, PersonaToolkitArgs);
 
 	TSharedRef<IAssetFamily> AssetFamily = PersonaModule.CreatePersonaAssetFamily(ObjectToEdit);
@@ -3553,6 +3553,11 @@ void FPhysicsAssetEditor::HandlePreviewSceneCreated(const TSharedRef<IPersonaPre
 	FloorMeshComponent->RecreatePhysicsState();
 }
 
+void FPhysicsAssetEditor::HandleOnPreviewSceneSettingsCustomized(IDetailLayoutBuilder& DetailBuilder) const
+{
+	DetailBuilder.HideCategory("Animation Blueprint");
+}
+
 void FPhysicsAssetEditor::HandleExtendContextMenu(FMenuBuilder& InMenuBuilder)
 {
 	TArray<TSharedPtr<ISkeletonTreeItem>> SelectedItems = SkeletonTree->GetSelectedItems();
diff --git a/Engine/Source/Editor/PhysicsAssetEditor/Private/PhysicsAssetEditor.h b/Engine/Source/Editor/PhysicsAssetEditor/Private/PhysicsAssetEditor.h
--- a/Engine/Source/Editor/PhysicsAssetEditor/Private/PhysicsAssetEditor.h
+++ b/Engine/Source/Editor/PhysicsAssetEditor/Private/PhysicsAssetEditor.h
@@ -22,6 +22,7 @@
 #include "Containers/ArrayView.h"
 #include "GraphEditor.h"
 
+class IDetailLayoutBuilder;
 struct FAssetData;
 class FPhysicsAssetEditorTreeInfo;
 class IDetailsView;
@@ -321,6 +322,9 @@ private:
 	/** Handle initial preview scene setup */
 	void HandlePreviewSceneCreated(const TSharedRef<IPersonaPreviewScene>& InPersonaPreviewScene);
 
+	/** Handle customization of Preview Scene Settings details */
+	void HandleOnPreviewSceneSettingsCustomized(IDetailLayoutBuilder& DetailBuilder) const;
+
 	/** Build context menu for tree items */
 	void HandleExtendContextMenu(FMenuBuilder& InMenuBuilder);
 
diff --git a/Engine/Source/Editor/SkeletalMeshEditor/Private/SkeletalMeshEditor.cpp b/Engine/Source/Editor/SkeletalMeshEditor/Private/SkeletalMeshEditor.cpp
--- a/Engine/Source/Editor/SkeletalMeshEditor/Private/SkeletalMeshEditor.cpp
+++ b/Engine/Source/Editor/SkeletalMeshEditor/Private/SkeletalMeshEditor.cpp
@@ -57,6 +57,7 @@
 #include "Toolkits/AssetEditorToolkitMenuContext.h"
 #include "MeshMergeModule.h"
 #include "AssetToolsModule.h"
+#include "DetailLayoutBuilder.h"
 #include "IAssetTools.h"
 #include "SkeletalMeshEditorContextMenuContext.h"
 #include "Styling/AppStyle.h"
@@ -231,8 +232,11 @@ void FSkeletalMeshEditor::InitSkeletalMeshEditor(const EToolkitMode::Type Mode,
 {
 	SkeletalMesh = InSkeletalMesh;
 
+	FPersonaToolkitArgs PersonaToolkitArgs;
+	PersonaToolkitArgs.OnPreviewSceneSettingsCustomized = FOnPreviewSceneSettingsCustomized::FDelegate::CreateSP(this, &FSkeletalMeshEditor::HandleOnPreviewSceneSettingsCustomized);
+	
 	FPersonaModule& PersonaModule = FModuleManager::LoadModuleChecked<FPersonaModule>("Persona");
-	PersonaToolkit = PersonaModule.CreatePersonaToolkit(InSkeletalMesh);
+	PersonaToolkit = PersonaModule.CreatePersonaToolkit(InSkeletalMesh, PersonaToolkitArgs);
 
 	PersonaToolkit->GetPreviewScene()->SetDefaultAnimationMode(EPreviewSceneDefaultAnimationMode::ReferencePose);
 
@@ -1462,6 +1466,15 @@ void FSkeletalMeshEditor::HandleReimportAllMeshWithNewFile(int32 SourceFileIndex
 	}
 }
 
+void FSkeletalMeshEditor::HandleOnPreviewSceneSettingsCustomized(IDetailLayoutBuilder& DetailBuilder)
+{
+	DetailBuilder.HideCategory("Mesh");
+	DetailBuilder.HideCategory("Physics");
+	// in mesh editor, we hide preview mesh section and additional mesh section
+	// sometimes additional meshes are interfering with preview mesh, it is not a great experience
+	DetailBuilder.HideCategory("Additional Meshes");
+}
+
 bool FSkeletalMeshEditor::IsMeshSectionSelectionChecked() const
 {
 	return GetPersonaToolkit()->GetPreviewScene()->AllowMeshHitProxies();
diff --git a/Engine/Source/Editor/SkeletalMeshEditor/Private/SkeletalMeshEditor.h b/Engine/Source/Editor/SkeletalMeshEditor/Private/SkeletalMeshEditor.h
--- a/Engine/Source/Editor/SkeletalMeshEditor/Private/SkeletalMeshEditor.h
+++ b/Engine/Source/Editor/SkeletalMeshEditor/Private/SkeletalMeshEditor.h
@@ -11,6 +11,7 @@
 #include "ISkeletalMeshEditor.h"
 #include "Containers/ArrayView.h"
 
+class IDetailLayoutBuilder;
 class IDetailsView;
 class IPersonaToolkit;
 class IPersonaViewport;
@@ -131,6 +132,8 @@ private:
 	void HandleReimportAllMesh(int32 SourceFileIndex = INDEX_NONE);
 	void HandleReimportAllMeshWithNewFile(int32 SourceFileIndex = INDEX_NONE);
 
+	void HandleOnPreviewSceneSettingsCustomized(IDetailLayoutBuilder& DetailBuilder);
+
 	/** Callback for toggling UV drawing in the viewport */
 	void ToggleMeshSectionSelection();
 
diff --git a/Engine/Source/Editor/SkeletonEditor/Private/SkeletonEditor.cpp b/Engine/Source/Editor/SkeletonEditor/Private/SkeletonEditor.cpp
--- a/Engine/Source/Editor/SkeletonEditor/Private/SkeletonEditor.cpp
+++ b/Engine/Source/Editor/SkeletonEditor/Private/SkeletonEditor.cpp
@@ -1,6 +1,8 @@
 // Copyright Epic Games, Inc. All Rights Reserved.
 
 #include "SkeletonEditor.h"
+
+#include "DetailLayoutBuilder.h"
 #include "Modules/ModuleManager.h"
 #include "Framework/MultiBox/MultiBoxBuilder.h"
 #include "Animation/DebugSkelMeshComponent.h"
@@ -88,9 +90,12 @@ void FSkeletonEditor::UnregisterTabSpawners(const TSharedRef<class FTabManager>&
 void FSkeletonEditor::InitSkeletonEditor(const EToolkitMode::Type Mode, const TSharedPtr<IToolkitHost>& InitToolkitHost, USkeleton* InSkeleton)
 {
 	Skeleton = InSkeleton;
+	
+	FPersonaToolkitArgs PersonaToolkitArgs;
+	PersonaToolkitArgs.OnPreviewSceneSettingsCustomized = FOnPreviewSceneSettingsCustomized::FDelegate::CreateSP(this, &FSkeletonEditor::HandleOnPreviewSceneSettingsCustomized);
 
 	FPersonaModule& PersonaModule = FModuleManager::LoadModuleChecked<FPersonaModule>("Persona");
-	PersonaToolkit = PersonaModule.CreatePersonaToolkit(InSkeleton);
+	PersonaToolkit = PersonaModule.CreatePersonaToolkit(InSkeleton, PersonaToolkitArgs);
 
 	PersonaToolkit->GetPreviewScene()->SetDefaultAnimationMode(EPreviewSceneDefaultAnimationMode::ReferencePose);
 
@@ -386,6 +391,11 @@ void FSkeletonEditor::OnImportAsset()
 	PersonaModule.ImportNewAsset(Skeleton, FBXIT_SkeletalMesh);
 }
 
+void FSkeletonEditor::HandleOnPreviewSceneSettingsCustomized(IDetailLayoutBuilder& DetailBuilder)
+{
+	DetailBuilder.HideCategory("Animation Blueprint");
+}
+
 void FSkeletonEditor::HandleDetailsCreated(const TSharedRef<IDetailsView>& InDetailsView)
 {
 	DetailsView = InDetailsView;
diff --git a/Engine/Source/Editor/SkeletonEditor/Private/SkeletonEditor.h b/Engine/Source/Editor/SkeletonEditor/Private/SkeletonEditor.h
--- a/Engine/Source/Editor/SkeletonEditor/Private/SkeletonEditor.h
+++ b/Engine/Source/Editor/SkeletonEditor/Private/SkeletonEditor.h
@@ -12,6 +12,7 @@
 #include "Containers/ArrayView.h"
 #include "IAnimationSequenceBrowser.h"
 
+class IDetailLayoutBuilder;
 class IPersonaToolkit;
 class IPersonaViewport;
 class ISkeletonTree;
@@ -118,6 +119,8 @@ private:
 
 	void OnImportAsset();
 
+	void HandleOnPreviewSceneSettingsCustomized(IDetailLayoutBuilder& DetailBuilder);
+
 	static TSharedPtr<FSkeletonEditor> GetSkeletonEditor(const FToolMenuContext& InMenuContext);
 public:
 	/** Multicast delegate fired on global undo/redo */
