Fixing issue, bone drawing modes don't immediately update if you click inside the radio button.

diff --git a/Engine/Source/Editor/Persona/Private/AnimationEditorViewportClient.cpp b/Engine/Source/Editor/Persona/Private/AnimationEditorViewportClient.cpp
--- a/Engine/Source/Editor/Persona/Private/AnimationEditorViewportClient.cpp
+++ b/Engine/Source/Editor/Persona/Private/AnimationEditorViewportClient.cpp
@@ -1426,6 +1426,7 @@ float FAnimationViewportClient::GetBoneDrawSize() const
 void FAnimationViewportClient::SetBoneDrawMode(EBoneDrawMode::Type AxesMode)
 {
 	ConfigOption->SetDefaultBoneDrawSelection(AxesMode);
+	RedrawRequested(Viewport);
 }
 
 bool FAnimationViewportClient::IsBoneDrawModeSet(EBoneDrawMode::Type AxesMode) const
