AnimNotifyPanel : Fixed Notify selection with custom windows display scaling

diff --git a/Engine/Source/Editor/Persona/Private/SAnimNotifyPanel.cpp b/Engine/Source/Editor/Persona/Private/SAnimNotifyPanel.cpp
--- a/Engine/Source/Editor/Persona/Private/SAnimNotifyPanel.cpp
+++ b/Engine/Source/Editor/Persona/Private/SAnimNotifyPanel.cpp
@@ -839,7 +839,7 @@ public:
 	void PasteSingleSyncMarker(FString& MarkerString, float PasteTime);
 
 	// Uses the given track space rect and marquee information to refresh selection information
-	void RefreshMarqueeSelectedNodes(FSlateRect& Rect, FNotifyMarqueeOperation& Marquee);
+	void RefreshMarqueeSelectedNodes(const FSlateRect& Rect, FNotifyMarqueeOperation& Marquee);
 
 	// Create new notifies
 	FAnimNotifyEvent& CreateNewBlueprintNotify(FString NewNotifyName, FString BlueprintPath, float StartTime);
@@ -3672,7 +3672,7 @@ void SAnimNotifyTrack::AppendSelectedNodeWidgetsToArray(TArray<TSharedPtr<SAnimN
 	}
 }
 
-void SAnimNotifyTrack::RefreshMarqueeSelectedNodes(FSlateRect& Rect, FNotifyMarqueeOperation& Marquee)
+void SAnimNotifyTrack::RefreshMarqueeSelectedNodes(const FSlateRect& Rect, FNotifyMarqueeOperation& Marquee)
 {
 	if(Marquee.Operation != FNotifyMarqueeOperation::Replace)
 	{
@@ -4749,7 +4749,9 @@ void SAnimNotifyPanel::RefreshMarqueeSelectedNodes(const FGeometry& PanelGeo)
 {
 	if(Marquee.IsValid())
 	{
-		FSlateRect MarqueeRect = Marquee.Rect.ToSlateRect();
+		const FSlateRect MarqueeRect = Marquee.Rect.ToSlateRect();
+		const FVector2D MarqueeTopLeftAbsolute = PanelGeo.LocalToAbsolute(MarqueeRect.GetTopLeft());
+
 		for(TSharedPtr<SAnimNotifyTrack> Track : NotifyAnimTracks)
 		{
 			if(Marquee.Operation == FNotifyMarqueeOperation::Replace || Marquee.OriginalSelection.Num() == 0)
@@ -4759,12 +4761,8 @@ void SAnimNotifyPanel::RefreshMarqueeSelectedNodes(const FGeometry& PanelGeo)
 
 			const FGeometry& TrackGeo = Track->GetCachedGeometry();
 
-			FSlateRect TrackClip = TrackGeo.GetLayoutBoundingRect();
-			FSlateRect PanelClip = PanelGeo.GetLayoutBoundingRect();
-			FVector2D PanelSpaceOrigin = TrackClip.GetTopLeft() - PanelClip.GetTopLeft();
-			FVector2D TrackSpaceOrigin = MarqueeRect.GetTopLeft() - PanelSpaceOrigin;
-			FSlateRect MarqueeTrackSpace(TrackSpaceOrigin, TrackSpaceOrigin + MarqueeRect.GetSize());
-
+			// Transform the Marquee Rect to Track Space
+			const FSlateRect MarqueeTrackSpace = FSlateRect::FromPointAndExtent(TrackGeo.AbsoluteToLocal(MarqueeTopLeftAbsolute), MarqueeRect.GetSize());
 			Track->RefreshMarqueeSelectedNodes(MarqueeTrackSpace, Marquee);
 		}
 	}
