Prevented anim node functions being bound in anim layer interfaces

Hide functions (and tag) categories in anim node interface contexts

diff --git a/Engine/Source/Editor/Persona/Private/AnimGraphNodeDetails.cpp b/Engine/Source/Editor/Persona/Private/AnimGraphNodeDetails.cpp
--- a/Engine/Source/Editor/Persona/Private/AnimGraphNodeDetails.cpp
+++ b/Engine/Source/Editor/Persona/Private/AnimGraphNodeDetails.cpp
@@ -115,6 +115,13 @@ void FAnimGraphNodeDetails::CustomizeDetails(class IDetailLayoutBuilder& DetailB
 	// customize anim graph node's own details if needed
 	AnimGraphNode->CustomizeDetails(DetailBuilder);
 
+	// Hide categories not relevant for interface BPs
+	if(AnimGraphNode->HasValidBlueprint() && AnimGraphNode->GetBlueprint()->BlueprintType == BPTYPE_Interface)
+	{
+		DetailBuilder.HideCategory("Functions");
+		DetailBuilder.HideCategory("Tag");
+	}
+	
 	// Hide the Node property as we are going to be adding its inner properties below
 	TSharedRef<IPropertyHandle> NodePropertyHandle = DetailBuilder.GetProperty(NodeProperty->GetFName(), AnimGraphNode->GetClass());
 	DetailBuilder.HideProperty(NodePropertyHandle);
