Fix for Undo/Redo not working on any BlendSpaceDetails Axis

diff --git a/Engine/Source/Editor/Persona/Private/Customization/BlendSpaceDetails.cpp b/Engine/Source/Editor/Persona/Private/Customization/BlendSpaceDetails.cpp
--- a/Engine/Source/Editor/Persona/Private/Customization/BlendSpaceDetails.cpp
+++ b/Engine/Source/Editor/Persona/Private/Customization/BlendSpaceDetails.cpp
@@ -29,10 +29,13 @@ FBlendSpaceDetails::FBlendSpaceDetails()
 	Builder = nullptr;
 	BlendSpace = nullptr;
 	BlendSpaceNode = nullptr;
+
+	GEditor->RegisterForUndo(this);
 }
 
 FBlendSpaceDetails::~FBlendSpaceDetails()
 {
+	GEditor->UnregisterForUndo(this);
 }
 
 //======================================================================================================================
@@ -121,6 +124,9 @@ static EVisibility GetAnalyzeButtonVisibility(
 // that if "None" was selected then the new value will be null
 void FBlendSpaceDetails::HandleAnalysisFunctionChanged(int32 AxisIndex, TSharedPtr<FString> NewFunctionName)
 {
+	const FScopedTransaction Transaction(LOCTEXT("BlendSpaceDetailsChangedAxisTransaction", "Changed Axis Function"));
+	BlendSpace->Modify();
+
 	UAnalysisProperties* NewAnalysisProperties = BlendSpaceAnalysis::MakeAnalysisProperties(BlendSpace, *NewFunctionName);
 	// Preserve values where possible
 	if (BlendSpace->AnalysisProperties[AxisIndex])
@@ -507,4 +513,14 @@ void FBlendSpaceDetails::CustomizeDetails(class IDetailLayoutBuilder& DetailBuil
 	}	
 }
 
+void FBlendSpaceDetails::PostUndo(bool bSuccess)
+{
+	GEditor->GetTimerManager()->SetTimerForNextTick(FTimerDelegate::CreateSP(this, &FBlendSpaceDetails::RefreshDetails));
+}
+
+void FBlendSpaceDetails::RefreshDetails()
+{
+	Builder->ForceRefreshDetails();
+}
+
 #undef LOCTEXT_NAMESPACE
diff --git a/Engine/Source/Editor/Persona/Private/Customization/BlendSpaceDetails.h b/Engine/Source/Editor/Persona/Private/Customization/BlendSpaceDetails.h
--- a/Engine/Source/Editor/Persona/Private/Customization/BlendSpaceDetails.h
+++ b/Engine/Source/Editor/Persona/Private/Customization/BlendSpaceDetails.h
@@ -7,12 +7,13 @@
 #include "Input/Reply.h"
 #include "UObject/WeakObjectPtr.h"
 #include "UObject/WeakObjectPtrTemplates.h"
+#include "EditorUndoClient.h"
 
 class IDetailLayoutBuilder;
 class UBlendSpace;
 class UAnimGraphNode_BlendSpaceGraphBase;
 
-class FBlendSpaceDetails : public IDetailCustomization
+class FBlendSpaceDetails : public IDetailCustomization, public FEditorUndoClient
 {
 public:
 	FBlendSpaceDetails();
@@ -25,10 +26,17 @@ public:
 
 	// IDetailCustomization interface
 	virtual void CustomizeDetails(class IDetailLayoutBuilder& DetailBuilder) override;
+
+	// FEditorUndoClient interface
+	virtual void PostUndo(bool bSuccess) override;
+	virtual void PostRedo(bool bSuccess) override { PostUndo(bSuccess); }
+	// End of FEditorUndoClient interface.
+
 private:
 	FReply HandleClearSamples();
 	FReply HandleAnalyzeSamples();
 	void HandleAnalysisFunctionChanged(int32 AxisIndex, TSharedPtr<FString> NewItem);
+	void RefreshDetails();
 
 	IDetailLayoutBuilder* Builder;
 	UBlendSpace* BlendSpace;
