Fix for Delete, Copy, and Paste shortcut commands not working for Anim Sync Markers.

diff --git a/Engine/Source/Editor/Persona/Private/AnimTimeline/AnimModel_AnimSequenceBase.cpp b/Engine/Source/Editor/Persona/Private/AnimTimeline/AnimModel_AnimSequenceBase.cpp
--- a/Engine/Source/Editor/Persona/Private/AnimTimeline/AnimModel_AnimSequenceBase.cpp
+++ b/Engine/Source/Editor/Persona/Private/AnimTimeline/AnimModel_AnimSequenceBase.cpp
@@ -156,7 +156,8 @@ void FAnimModel_AnimSequenceBase::Initialize()
 
 	CommandList->MapAction(
 		FGenericCommands::Get().Delete,
-		FExecuteAction::CreateSP(this, &FAnimModel_AnimSequenceBase::RemoveSelectedCurves));
+		FExecuteAction::CreateSP(this, &FAnimModel_AnimSequenceBase::RemoveSelectedCurves),
+		FCanExecuteAction::CreateSP(this, &FAnimModel_AnimSequenceBase::AreAnyCurvesSelected));
 }
 
 
@@ -535,9 +536,25 @@ bool FAnimModel_AnimSequenceBase::IsDisplaySecondaryChecked() const
 	return GetDefault<UPersonaOptions>()->bTimelineDisplayFormatSecondary;
 }
 
+bool FAnimModel_AnimSequenceBase::AreAnyCurvesSelected() const
+{
+	if (!SelectedTracks.IsEmpty())
+	{
+		for (const TSharedRef<FAnimTimelineTrack>& SelectedTrack : SelectedTracks)
+		{
+			if (SelectedTrack->IsA<FAnimTimelineTrack_Curve>())
+			{
+				return true;
+			}
+		}
+	}
+
+	return false;
+}
+
 void FAnimModel_AnimSequenceBase::CopyToClipboard() const
 {
-	if (SelectedTracks.Num() > 0)
+	if (!SelectedTracks.IsEmpty())
 	{
 		if (UAnimTimelineClipboardContent * ClipboardContent = UAnimTimelineClipboardContent::Create())
 		{
@@ -555,7 +572,7 @@ bool FAnimModel_AnimSequenceBase::CanCopyToClipboard()
 {
 	if (!SelectedTracks.IsEmpty())
 	{
-		for (const auto & Track : SelectedTracks)
+		for (const TSharedRef<FAnimTimelineTrack> & Track : SelectedTracks)
 		{
 			if (!Track->SupportsCopy())
 			{
@@ -563,6 +580,10 @@ bool FAnimModel_AnimSequenceBase::CanCopyToClipboard()
 			}
 		}
 	}
+	else
+	{
+		return false;
+	}
 	
 	return true;
 }
diff --git a/Engine/Source/Editor/Persona/Private/AnimTimeline/AnimModel_AnimSequenceBase.h b/Engine/Source/Editor/Persona/Private/AnimTimeline/AnimModel_AnimSequenceBase.h
--- a/Engine/Source/Editor/Persona/Private/AnimTimeline/AnimModel_AnimSequenceBase.h
+++ b/Engine/Source/Editor/Persona/Private/AnimTimeline/AnimModel_AnimSequenceBase.h
@@ -74,6 +74,7 @@ private:
 	bool IsDisplayPercentageChecked() const;
 	void ToggleDisplaySecondary();
 	bool IsDisplaySecondaryChecked() const;
+	bool AreAnyCurvesSelected() const;
 	
 	/** Copy selected curves to clipboard */
 	void CopyToClipboard() const;
diff --git a/Engine/Source/Editor/Persona/Private/AnimTimeline/AnimTimelineClipboard.cpp b/Engine/Source/Editor/Persona/Private/AnimTimeline/AnimTimelineClipboard.cpp
--- a/Engine/Source/Editor/Persona/Private/AnimTimeline/AnimTimelineClipboard.cpp
+++ b/Engine/Source/Editor/Persona/Private/AnimTimeline/AnimTimelineClipboard.cpp
@@ -140,7 +140,7 @@ bool FAnimTimelineClipboardUtilities::CanOverwriteSelectedCurveDataFromClipboard
 	// Validate clipboard
 	if (ClipboardContent && !ClipboardContent->IsEmpty())
 	{
-		// Exit on multiple curves selected
+		// Exit on multiple curves selected or no curves selected
 		if (SelectedTracks.Num() != 1)
 		{
 			return false;
diff --git a/Engine/Source/Editor/Persona/Private/SAnimNotifyPanel.cpp b/Engine/Source/Editor/Persona/Private/SAnimNotifyPanel.cpp
--- a/Engine/Source/Editor/Persona/Private/SAnimNotifyPanel.cpp
+++ b/Engine/Source/Editor/Persona/Private/SAnimNotifyPanel.cpp
@@ -3846,7 +3846,6 @@ void SAnimNotifyPanel::Construct(const FArguments& InArgs, const TSharedRef<FAni
 		.OnSetInputViewRange(InArgs._OnSetInputViewRange));
 
 	WeakModel = InModel;
-	WeakCommandList = InModel->GetCommandList();
 	Sequence = InArgs._Sequence;
 	OnInvokeTab = InArgs._OnInvokeTab;
 	OnNotifiesChanged = InArgs._OnNotifiesChanged;
@@ -4643,9 +4642,12 @@ void SAnimNotifyPanel::OnPropertyChanged(UObject* ChangedObject, FPropertyChange
 
 void SAnimNotifyPanel::BindCommands()
 {
-	TSharedRef<FUICommandList> CommandList = GetCommandList();
-	const FAnimNotifyPanelCommands& Commands = FAnimNotifyPanelCommands::Get();
+	// This should not be called twice on the same instance
+	check(!CommandList.IsValid());
+	CommandList = MakeShareable(new FUICommandList);
 
+	const FAnimNotifyPanelCommands& Commands = FAnimNotifyPanelCommands::Get();
+	
 	CommandList->MapAction(
 		Commands.DeleteNotify,
 		FExecuteAction::CreateSP(this, &SAnimNotifyPanel::OnDeletePressed));
@@ -4661,7 +4663,7 @@ void SAnimNotifyPanel::BindCommands()
 
 FReply SAnimNotifyPanel::OnKeyDown(const FGeometry& MyGeometry, const FKeyEvent& InKeyEvent)
 {
-	if(GetCommandList()->ProcessCommandBindings(InKeyEvent))
+	if(CommandList.IsValid() && CommandList->ProcessCommandBindings(InKeyEvent))
 	{
 		return FReply::Handled();
 	}
diff --git a/Engine/Source/Editor/Persona/Private/SAnimNotifyPanel.h b/Engine/Source/Editor/Persona/Private/SAnimNotifyPanel.h
--- a/Engine/Source/Editor/Persona/Private/SAnimNotifyPanel.h
+++ b/Engine/Source/Editor/Persona/Private/SAnimNotifyPanel.h
@@ -258,7 +258,7 @@ public:
 
 	void HandleObjectsSelected(const TArray<UObject*>& InObjects);
 
-	TSharedRef<FUICommandList> GetCommandList() const { return WeakCommandList.Pin().ToSharedRef(); }
+	TSharedRef<FUICommandList> GetCommandList() const { return CommandList.ToSharedRef(); }
 
 private:
 	friend struct FScopedSavedNotifySelection;
@@ -345,8 +345,8 @@ private:
 	FOnSnapPosition OnSnapPosition;
 
 	/** UI commands for this widget */
-	TWeakPtr<FUICommandList> WeakCommandList;
-
+	TSharedPtr<FUICommandList> CommandList;
+	
 	/** Classes that are known to be derived from blueprint notifies */
 	TArray<FString> NotifyClassNames;
 
