FAnimationViewportClient: Fix for crash closing the editor, due to memory use after delete at AssetEditorModeManager.cpp (preview was already deleted when trying to get preview world)

diff --git a/Engine/Source/Editor/Persona/Private/AnimationEditorViewportClient.cpp b/Engine/Source/Editor/Persona/Private/AnimationEditorViewportClient.cpp
--- a/Engine/Source/Editor/Persona/Private/AnimationEditorViewportClient.cpp
+++ b/Engine/Source/Editor/Persona/Private/AnimationEditorViewportClient.cpp
@@ -166,6 +166,11 @@ FAnimationViewportClient::~FAnimationViewportClient()
 {
 	CameraController = CachedDefaultCameraController;
 
+	if (ModeTools.IsValid())
+	{
+		((FAssetEditorModeManager*)ModeTools.Get())->SetPreviewScene(nullptr);
+	}
+
 	// Unregistering the callbacks is mandatory, else we get random crashes
 	if (PreviewScenePtr.IsValid())
 	{
