Skeleton compatibility:
- Allow picking preview meshes that use compatible skeletons.

diff --git a/Engine/Source/Editor/Persona/Private/AnimationEditorPreviewScene.cpp b/Engine/Source/Editor/Persona/Private/AnimationEditorPreviewScene.cpp
--- a/Engine/Source/Editor/Persona/Private/AnimationEditorPreviewScene.cpp
+++ b/Engine/Source/Editor/Persona/Private/AnimationEditorPreviewScene.cpp
@@ -159,26 +159,32 @@ void FAnimationEditorPreviewScene::SetPreviewMesh(USkeletalMesh* NewPreviewMesh,
 	if (NewPreviewMesh != nullptr && GetEditableSkeleton().IsValid() && !GetEditableSkeleton()->GetSkeleton().IsCompatibleMesh(NewPreviewMesh))
 	{
 		const USkeleton& Skeleton = GetEditableSkeleton()->GetSkeleton();
-
-		// message box, ask if they'd like to regenerate skeleton
-		if (FMessageDialog::Open(EAppMsgType::YesNo, LOCTEXT("RenerateSkeleton", "The preview mesh hierarchy doesn't match with Skeleton anymore. Would you like to regenerate skeleton?")) == EAppReturnType::Yes)
+		if (NewPreviewMesh->GetSkeleton() && Skeleton.IsCompatible(NewPreviewMesh->GetSkeleton()))
 		{
-			GetEditableSkeleton()->RecreateBoneTree(NewPreviewMesh);
 			SetPreviewMeshInternal(NewPreviewMesh);
-		}
+		}	
 		else
 		{
-			// Send a notification that the skeletal mesh cannot work with the skeleton
-			FFormatNamedArguments Args;
-			Args.Add(TEXT("PreviewMeshName"), FText::FromString(NewPreviewMesh->GetName()));
-			Args.Add(TEXT("TargetSkeletonName"), FText::FromString(Skeleton.GetName()));
-			FNotificationInfo Info(FText::Format(LOCTEXT("SkeletalMeshIncompatible", "Skeletal Mesh \"{PreviewMeshName}\" incompatible with Skeleton \"{TargetSkeletonName}\""), Args));
-			Info.ExpireDuration = 3.0f;
-			Info.bUseLargeFont = false;
-			TSharedPtr<SNotificationItem> Notification = FSlateNotificationManager::Get().AddNotification(Info);
-			if (Notification.IsValid())
+			// message box, ask if they'd like to regenerate skeleton
+			if (FMessageDialog::Open(EAppMsgType::YesNo, LOCTEXT("RenerateSkeleton", "The preview mesh hierarchy doesn't match with Skeleton anymore. Would you like to regenerate skeleton?")) == EAppReturnType::Yes)
+			{
+				GetEditableSkeleton()->RecreateBoneTree(NewPreviewMesh);
+				SetPreviewMeshInternal(NewPreviewMesh);
+			}
+			else
 			{
-				Notification->SetCompletionState(SNotificationItem::CS_Fail);
+				// Send a notification that the skeletal mesh cannot work with the skeleton
+				FFormatNamedArguments Args;
+				Args.Add(TEXT("PreviewMeshName"), FText::FromString(NewPreviewMesh->GetName()));
+				Args.Add(TEXT("TargetSkeletonName"), FText::FromString(Skeleton.GetName()));
+				FNotificationInfo Info(FText::Format(LOCTEXT("SkeletalMeshIncompatible", "Skeletal Mesh \"{PreviewMeshName}\" incompatible with Skeleton \"{TargetSkeletonName}\""), Args));
+				Info.ExpireDuration = 3.0f;
+				Info.bUseLargeFont = false;
+				TSharedPtr<SNotificationItem> Notification = FSlateNotificationManager::Get().AddNotification(Info);
+				if (Notification.IsValid())
+				{
+					Notification->SetCompletionState(SNotificationItem::CS_Fail);
+				}
 			}
 		}
 	}
