Fix crash when deleting a material slot on a mesh containing morph targets.
#rnx

diff --git a/Engine/Source/Editor/Persona/Private/AnimationEditorViewportClient.cpp b/Engine/Source/Editor/Persona/Private/AnimationEditorViewportClient.cpp
--- a/Engine/Source/Editor/Persona/Private/AnimationEditorViewportClient.cpp
+++ b/Engine/Source/Editor/Persona/Private/AnimationEditorViewportClient.cpp
@@ -927,6 +927,7 @@ FText FAnimationViewportClient::GetDisplayInfo(bool bDisplayAllInfo) const
 		TArray<UMaterial*> MaterialsThatNeedSaving;
 
 		const TIndirectArray<FSkeletalMeshLODModel>& LODModels = PreviewMeshComponent->GetSkeletalMeshAsset()->GetImportedModel()->LODModels;
+		const TArray<FSkeletalMaterial>& SkeletalMeshMaterials = PreviewMeshComponent->GetSkeletalMeshAsset()->GetMaterials();
 		int32 LodNumber = LODModels.Num();
 		TArray<UMaterialInterface*> MaterialUsingMorphTarget;
 		for (UMorphTarget *MorphTarget : PreviewMeshComponent->GetSkeletalMeshAsset()->GetMorphTargets())
@@ -941,10 +942,19 @@ FText FAnimationViewportClient::GetDisplayInfo(bool bDisplayAllInfo) const
 				{
 					for (int32 LodIdx = 0; LodIdx < LodNumber; LodIdx++)
 					{
+						const TArray<int32>& LODMaterialMap = PreviewMeshComponent->GetSkeletalMeshAsset()->GetLODInfo(LodIdx)->LODMaterialMap;
 						const FSkeletalMeshLODModel& LODModel = LODModels[LodIdx];
 						if (LODModel.Sections.IsValidIndex(SectionIndex))
 						{
-							MaterialUsingMorphTarget.AddUnique(PreviewMeshComponent->GetSkeletalMeshAsset()->GetMaterials()[LODModel.Sections[SectionIndex].MaterialIndex].MaterialInterface);
+							int32 SectionMaterialIndex = LODModel.Sections[SectionIndex].MaterialIndex;
+							if (LODMaterialMap.IsValidIndex(SectionIndex) && LODMaterialMap[SectionIndex] != INDEX_NONE)
+							{
+								SectionMaterialIndex = LODMaterialMap[SectionIndex];
+							}
+							if (SkeletalMeshMaterials.IsValidIndex(SectionMaterialIndex))
+							{
+								MaterialUsingMorphTarget.AddUnique(SkeletalMeshMaterials[SectionMaterialIndex].MaterialInterface);
+							}
 						}
 					}
 				}
