Animation BP toolbar trimming & permissions
* Tab filters are now applied before tab widgets before are created for animation blueprint editor instead of only when a user tries opening a new tab. This prevents disallowed tabs in default layouts being created.
* Event graph documents are not opened if event graph permissions are disabled when loading blueprint editor.
* Added new blueprint permission to optionally disable function overrides.

diff --git a/Engine/Source/Developer/MeshUtilities/Private/MeshUtilities.cpp b/Engine/Source/Developer/MeshUtilities/Private/MeshUtilities.cpp
--- a/Engine/Source/Developer/MeshUtilities/Private/MeshUtilities.cpp
+++ b/Engine/Source/Developer/MeshUtilities/Private/MeshUtilities.cpp
@@ -5619,7 +5619,7 @@ void FMeshUtilities::ShutdownModule()
 void FMeshUtilities::RegisterMenus()
 {
 	AddMakeStaticMeshEntryToToolMenu("AssetEditor.AnimationEditor.ToolBar");
-	AddMakeStaticMeshEntryToToolMenu("AssetEditor.AnimationBlueprintEditor.ToolBar");
+	AddMakeStaticMeshEntryToToolMenu("AssetEditor.AnimationBlueprintEditor.MainMenu.Tools");
 	AddMakeStaticMeshEntryToToolMenu("AssetEditor.SkeletalMeshEditor.ToolBar");
 	AddMakeStaticMeshEntryToToolMenu("AssetEditor.SkeletonEditor.ToolBar");
 }
@@ -6121,24 +6121,24 @@ void FMeshUtilities::AddMakeStaticMeshEntryToToolMenu(FName InToolMenuName)
 	FToolMenuOwnerScoped OwnerScoped(this);
 
 	UToolMenu* ToolMenu = UToolMenus::Get()->ExtendMenu(InToolMenuName);
-	ToolMenu->AddMenuEntry("SkeletalMesh",
-		FToolMenuEntry::InitToolBarButton(
-			"MakeStaticMesh",
-			FToolUIActionChoice(
-				FToolUIAction(FToolMenuExecuteAction::CreateLambda([this](const FToolMenuContext& InContext)
+	FToolMenuSection& SkeletalMeshSection = ToolMenu->AddSection("SkeletalMesh", LOCTEXT("SkeletalMeshSection", "Skeletal Mesh"));
+
+	SkeletalMeshSection.AddMenuEntry(
+		"MakeStaticMesh",
+		LOCTEXT("MakeStaticMesh", "Make Static Mesh"),
+		LOCTEXT("MakeStaticMeshTooltip", "Make a new static mesh out of the preview's current pose."),
+		FSlateIcon(FAppStyle::GetAppStyleSetName(), "Persona.ConvertToStaticMesh"),
+		FToolUIActionChoice(
+			FToolUIAction(FToolMenuExecuteAction::CreateLambda([this](const FToolMenuContext& InContext)
+			{
+				if(UPersonaToolMenuContext* MenuContext = InContext.FindContext<UPersonaToolMenuContext>())
 				{
-					if(UPersonaToolMenuContext* MenuContext = InContext.FindContext<UPersonaToolMenuContext>())
+					if(UDebugSkelMeshComponent* MeshComponent = MenuContext->GetPreviewMeshComponent())
 					{
-						if(UDebugSkelMeshComponent* MeshComponent = MenuContext->GetPreviewMeshComponent())
-						{
-							ConvertMeshesToStaticMesh(TArray<UMeshComponent*>({ MeshComponent }), MeshComponent->GetComponentToWorld());
-						}
+						ConvertMeshesToStaticMesh(TArray<UMeshComponent*>({ MeshComponent }), MeshComponent->GetComponentToWorld());
 					}
-				}))
-			),
-			LOCTEXT("MakeStaticMesh", "Make Static Mesh"),
-			LOCTEXT("MakeStaticMeshTooltip", "Make a new static mesh out of the preview's current pose."),
-			FSlateIcon(FAppStyle::GetAppStyleSetName(), "Persona.ConvertToStaticMesh")
+				}
+			}))
 		)
 	);
 }
diff --git a/Engine/Source/Editor/AnimationBlueprintEditor/Private/AnimationBlueprintEditor.cpp b/Engine/Source/Editor/AnimationBlueprintEditor/Private/AnimationBlueprintEditor.cpp
--- a/Engine/Source/Editor/AnimationBlueprintEditor/Private/AnimationBlueprintEditor.cpp
+++ b/Engine/Source/Editor/AnimationBlueprintEditor/Private/AnimationBlueprintEditor.cpp
@@ -131,6 +131,8 @@
 #include "Widgets/SBoxPanel.h"
 #include "Widgets/Text/STextBlock.h"
 #include "WorkflowOrientedApp/WorkflowTabManager.h"
+#include "AnimationBlueprintToolMenuContext.h"
+#include "ToolMenus.h"
 #include "PropertyCustomizationHelpers.h"
 
 class FEditorModeTools;
@@ -361,6 +363,34 @@ void FAnimationBlueprintEditor::ExtendMenu()
 	// add extensible menu if exists
 	FAnimationBlueprintEditorModule& AnimationBlueprintEditorModule = FModuleManager::LoadModuleChecked<FAnimationBlueprintEditorModule>("AnimationBlueprintEditor");
 	AddMenuExtender(AnimationBlueprintEditorModule.GetMenuExtensibilityManager()->GetAllExtenders(GetToolkitCommands(), GetEditingObjects()));
+
+	FToolMenuOwnerScoped OwnerScoped(this);
+
+	// Add in Editor Specific functionality
+	static const FName MenuName = GetToolMenuName();
+	static const FName ToolsMenuName = *(MenuName.ToString() + TEXT(".") + TEXT("Tools"));
+
+	UToolMenu* ToolsMenu = UToolMenus::Get()->ExtendMenu(ToolsMenuName);
+	const FToolMenuInsert SectionInsertLocation("Programming", EToolMenuInsertType::Before);
+
+	UAnimBlueprint* AnimBlueprint = PersonaToolkit->GetAnimBlueprint();
+	if (AnimBlueprint && AnimBlueprint->BlueprintType != BPTYPE_Interface && !AnimBlueprint->bIsTemplate)
+	{
+		ToolsMenu->AddDynamicSection("Persona", FNewToolMenuDelegate::CreateLambda([](UToolMenu* InToolMenu)
+		{
+			TSharedPtr<FAnimationBlueprintEditor> AnimationBlueprintEditor = GetAnimationBlueprintEditor(InToolMenu->Context);
+			if (AnimationBlueprintEditor.IsValid() && AnimationBlueprintEditor->PersonaToolkit.IsValid())
+			{
+				FPersonaModule& PersonaModule = FModuleManager::LoadModuleChecked<FPersonaModule>("Persona");
+				FPersonaModule::FCommonToolbarExtensionArgs Args;
+				Args.bPreviewAnimation = false;
+				Args.bPreviewMesh = true;
+				Args.bReferencePose = false;
+				Args.bCreateAsset = true;
+				PersonaModule.AddCommonMenuExtensions(InToolMenu, Args);
+			}
+		}), SectionInsertLocation);
+	}
 }
 
 void FAnimationBlueprintEditor::RegisterMenus()
@@ -514,6 +544,19 @@ void FAnimationBlueprintEditor::BindCommands()
 		FExecuteAction::CreateRaw(&GetPersonaToolkit()->GetPreviewScene().Get(), &IPersonaPreviewScene::TogglePlayback));
 }
 
+TSharedPtr<FAnimationBlueprintEditor> FAnimationBlueprintEditor::GetAnimationBlueprintEditor(const FToolMenuContext& InMenuContext)
+{
+	if (UAnimationBlueprintToolMenuContext* Context = InMenuContext.FindContext<UAnimationBlueprintToolMenuContext>())
+	{
+		if (Context->AnimationBlueprintEditor.IsValid())
+		{
+			return StaticCastSharedPtr<FAnimationBlueprintEditor>(Context->AnimationBlueprintEditor.Pin());
+		}
+	}
+
+	return TSharedPtr<FAnimationBlueprintEditor>();
+}
+
 void FAnimationBlueprintEditor::ExtendToolbar()
 {
 	// If the ToolbarExtender is valid, remove it before rebuilding it
@@ -550,10 +593,6 @@ void FAnimationBlueprintEditor::ExtendToolbar()
 			FToolBarExtensionDelegate::CreateLambda([this](FToolBarBuilder& ParentToolbarBuilder)
 			{
 				FPersonaModule& PersonaModule = FModuleManager::LoadModuleChecked<FPersonaModule>("Persona");
-				FPersonaModule::FCommonToolbarExtensionArgs Args;
-				Args.bPreviewAnimation = false;
-				PersonaModule.AddCommonToolbarExtensions(ParentToolbarBuilder, PersonaToolkit.ToSharedRef(), Args);
-
 				TSharedRef<class IAssetFamily> AssetFamily = PersonaModule.CreatePersonaAssetFamily(GetBlueprintObj());
 				AddToolbarWidget(PersonaModule.CreateAssetFamilyShortcutWidget(SharedThis(this), AssetFamily));
 			}
@@ -1475,6 +1514,11 @@ FName FAnimationBlueprintEditor::GetToolkitFName() const
 	return FName("AnimationBlueprintEditor");
 }
 
+FName FAnimationBlueprintEditor::GetToolkitContextFName() const
+{
+	return FName("AnimationBlueprintEditor");
+}
+
 FText FAnimationBlueprintEditor::GetBaseToolkitName() const
 {
 	return LOCTEXT("AppLabel", "Animation Blueprint Editor");
@@ -1500,6 +1544,10 @@ void FAnimationBlueprintEditor::InitToolMenuContext(FToolMenuContext& MenuContex
 {
 	IAnimationBlueprintEditor::InitToolMenuContext(MenuContext);
 
+	UAnimationBlueprintToolMenuContext* AnimationBlueprintToolMenuContext = NewObject<UAnimationBlueprintToolMenuContext>();
+	AnimationBlueprintToolMenuContext->AnimationBlueprintEditor = SharedThis(this);
+	MenuContext.AddObject(AnimationBlueprintToolMenuContext);
+
 	UPersonaToolMenuContext* Context = NewObject<UPersonaToolMenuContext>();
 	Context->SetToolkit(GetPersonaToolkit());
 
@@ -1619,6 +1667,17 @@ void FAnimationBlueprintEditor::GetCustomDebugObjects(TArray<FCustomDebugObject>
 	AnimationBlueprintEditorModule.OnGetCustomDebugObjects().Broadcast(*this, DebugList);
 }
 
+FString FAnimationBlueprintEditor::GetCustomDebugObjectLabel(UObject* ObjectBeingDebugged) const
+{
+	UAnimInstance* PreviewInstance = GetPreviewInstance();
+	if (PreviewInstance == ObjectBeingDebugged)
+	{
+		return LOCTEXT("PreviewObjectLabel", "Preview Instance").ToString();
+	}
+
+	return FString();
+}
+
 void FAnimationBlueprintEditor::CreateDefaultTabContents(const TArray<UBlueprint*>& InBlueprints)
 {
 	FBlueprintEditor::CreateDefaultTabContents(InBlueprints);
diff --git a/Engine/Source/Editor/AnimationBlueprintEditor/Private/AnimationBlueprintEditor.h b/Engine/Source/Editor/AnimationBlueprintEditor/Private/AnimationBlueprintEditor.h
--- a/Engine/Source/Editor/AnimationBlueprintEditor/Private/AnimationBlueprintEditor.h
+++ b/Engine/Source/Editor/AnimationBlueprintEditor/Private/AnimationBlueprintEditor.h
@@ -164,6 +164,7 @@ public:
 public:
 	//~ Begin IToolkit Interface
 	virtual FName GetToolkitFName() const override;
+	virtual FName GetToolkitContextFName() const override;
 	virtual FText GetBaseToolkitName() const override;
 	virtual FText GetToolkitToolTipText() const override;
 	virtual FString GetWorldCentricTabPrefix() const override;
@@ -210,6 +211,7 @@ protected:
 	virtual void OnGraphEditorBackgrounded(const TSharedRef<SGraphEditor>& InGraphEditor) override;
 	virtual bool IsInAScriptingMode() const override { return true; }
 	virtual void GetCustomDebugObjects(TArray<FCustomDebugObject>& DebugList) const override;
+	virtual FString GetCustomDebugObjectLabel(UObject* ObjectBeingDebugged) const override;
 	virtual void CreateDefaultTabContents(const TArray<UBlueprint*>& InBlueprints) override;
 	virtual FGraphAppearanceInfo GetGraphAppearance(class UEdGraph* InGraph) const override;
 	virtual bool IsEditable(UEdGraph* InGraph) const override;
@@ -291,6 +293,9 @@ private:
 	/** Extend toolbar */
 	void ExtendToolbar();
 
+	/** Get the anim BP editor referenced by the supplied tool menu context */
+	static TSharedPtr<FAnimationBlueprintEditor> GetAnimationBlueprintEditor(const FToolMenuContext& InMenuContext);
+
 	/** Called immediately prior to a blueprint compilation */
 	void OnBlueprintPreCompile(UBlueprint* BlueprintToCompile);
 
diff --git a/Engine/Source/Editor/AnimationBlueprintEditor/Private/AnimationBlueprintToolMenuContext.h b/Engine/Source/Editor/AnimationBlueprintEditor/Private/AnimationBlueprintToolMenuContext.h
new file mode 100644
--- /dev/null
+++ b/Engine/Source/Editor/AnimationBlueprintEditor/Private/AnimationBlueprintToolMenuContext.h
@@ -0,0 +1,17 @@
+// Copyright Epic Games, Inc. All Rights Reserved.
+
+#pragma once
+
+#include "CoreMinimal.h"
+#include "AnimationBlueprintToolMenuContext.generated.h"
+
+class FAnimationBlueprintEditor;
+
+UCLASS()
+class UAnimationBlueprintToolMenuContext : public UObject
+{
+	GENERATED_BODY()
+
+public:
+	TWeakPtr<FAnimationBlueprintEditor> AnimationBlueprintEditor;
+};
\ No newline at end of file
diff --git a/Engine/Source/Editor/AnimationEditor/Private/AnimationEditor.cpp b/Engine/Source/Editor/AnimationEditor/Private/AnimationEditor.cpp
--- a/Engine/Source/Editor/AnimationEditor/Private/AnimationEditor.cpp
+++ b/Engine/Source/Editor/AnimationEditor/Private/AnimationEditor.cpp
@@ -291,16 +291,16 @@ void FAnimationEditor::ExtendToolbar()
 	const FToolMenuInsert SectionInsertLocation("Asset", EToolMenuInsertType::After);
 
 	{
-		ToolMenu->AddDynamicSection("Persona", FNewToolBarDelegateLegacy::CreateLambda([](FToolBarBuilder& ToolbarBuilder, UToolMenu* InMenu)
+		ToolMenu->AddDynamicSection("Persona", FNewToolMenuDelegate::CreateLambda([](UToolMenu* InToolMenu)
 		{
-			TSharedPtr<FAnimationEditor> AnimationEditor = GetAnimationEditor(InMenu->Context);
+			TSharedPtr<FAnimationEditor> AnimationEditor = GetAnimationEditor(InToolMenu->Context);
 			if (AnimationEditor.IsValid() && AnimationEditor->PersonaToolkit.IsValid())
 			{
 				FPersonaModule& PersonaModule = FModuleManager::LoadModuleChecked<FPersonaModule>("Persona");
 				FPersonaModule::FCommonToolbarExtensionArgs Args;
 				Args.bPreviewAnimation = false;
 				Args.bReferencePose = false;
-				PersonaModule.AddCommonToolbarExtensions(ToolbarBuilder, AnimationEditor->PersonaToolkit.ToSharedRef(), Args);
+				PersonaModule.AddCommonToolbarExtensions(InToolMenu, Args);
 			}
 		}), SectionInsertLocation);
 	}
diff --git a/Engine/Source/Editor/Kismet/Private/BlueprintEditor.cpp b/Engine/Source/Editor/Kismet/Private/BlueprintEditor.cpp
--- a/Engine/Source/Editor/Kismet/Private/BlueprintEditor.cpp
+++ b/Engine/Source/Editor/Kismet/Private/BlueprintEditor.cpp
@@ -9816,7 +9816,10 @@ void FBlueprintEditor::RestoreEditedObjectState()
 			Blueprint->LastEditedDocuments.Add(FBlueprintEditorUtils::FindUserConstructionScript(Blueprint));
 		}
 
-		Blueprint->LastEditedDocuments.Add(FBlueprintEditorUtils::FindEventGraph(Blueprint));
+		if (Blueprint->SupportsEventGraphs())
+		{
+			Blueprint->LastEditedDocuments.Add(FBlueprintEditorUtils::FindEventGraph(Blueprint));
+		}
 	}
 
 	for (int32 i = 0; i < Blueprint->LastEditedDocuments.Num(); i++)
@@ -9825,6 +9828,11 @@ void FBlueprintEditor::RestoreEditedObjectState()
 		{
 			if(UEdGraph* Graph = Cast<UEdGraph>(Obj))
 			{
+				if (FBlueprintEditorUtils::IsEventGraph(Graph) && !Blueprint->SupportsEventGraphs())
+				{
+					continue;
+				}
+
 				struct LocalStruct
 				{
 					static TSharedPtr<SDockTab> OpenGraphTree(FBlueprintEditor* InBlueprintEditor, UEdGraph* InGraph)
diff --git a/Engine/Source/Editor/Kismet/Private/SMyBlueprint.cpp b/Engine/Source/Editor/Kismet/Private/SMyBlueprint.cpp
--- a/Engine/Source/Editor/Kismet/Private/SMyBlueprint.cpp
+++ b/Engine/Source/Editor/Kismet/Private/SMyBlueprint.cpp
@@ -1589,7 +1589,9 @@ void SMyBlueprint::CollectAllActions(FGraphActionListBuilderBase& OutAllActions)
 			 && !OverridableFunctionNames.Contains(FunctionName) 
 			 && !ImplementedFunctionCache.Contains(FunctionName) 
 			 && !FObjectEditorUtils::IsFunctionHiddenFromClass(Function, ParentClass)
-			 && !FBlueprintEditorUtils::FindOverrideForFunction(BlueprintObj, OuterClass, Function->GetFName()) )
+			 && !FBlueprintEditorUtils::FindOverrideForFunction(BlueprintObj, OuterClass, Function->GetFName())
+			 && Blueprint->AllowFunctionOverride(Function)
+		   )
 		{
 			FText FunctionTooltip = FText::FromString(UK2Node_CallFunction::GetDefaultTooltipForFunction(Function));
 			FText FunctionDesc = K2Schema->GetFriendlySignatureName(Function);
diff --git a/Engine/Source/Editor/Persona/Private/PersonaModule.cpp b/Engine/Source/Editor/Persona/Private/PersonaModule.cpp
--- a/Engine/Source/Editor/Persona/Private/PersonaModule.cpp
+++ b/Engine/Source/Editor/Persona/Private/PersonaModule.cpp
@@ -94,6 +94,8 @@
 #include "SAnimMontageSectionsPanel.h"
 #include "Animation/AnimSequenceHelpers.h"
 #include "SkeletalMeshReferenceSectionDetails.h"
+#include "PersonaToolMenuContext.h"
+#include "ToolMenu.h"
 
 IMPLEMENT_MODULE( FPersonaModule, Persona );
 
@@ -831,245 +833,248 @@ IPersonaEditorModeManager* FPersonaModule::CreatePersonaEditorModeManager()
 	return new FPersonaEditorModeManager();
 }
 
-void FPersonaModule::AddCommonToolbarExtensions(FToolBarBuilder& InToolbarBuilder, TSharedRef<IPersonaToolkit> PersonaToolkit, const FCommonToolbarExtensionArgs& InArgs)
+namespace UE::Persona::Private
 {
-	TWeakPtr<IPersonaToolkit> WeakPersonaToolkit = PersonaToolkit;
-
-	if(InArgs.bPreviewMesh)
+	struct FNotificationHandler : public TSharedFromThis<FNotificationHandler>
 	{
-		// Handler to hang notifications on
-		struct FNotificationHandler : public TSharedFromThis<FNotificationHandler>
+		static void HandleApplyPreviewMesh(TSharedPtr<FNotificationHandler> InNotificationHandler, TWeakPtr<IPersonaToolkit> InWeakPersonaToolkit)
 		{
-			static void HandleApplyPreviewMesh(TSharedPtr<FNotificationHandler> InNotificationHandler, TWeakPtr<IPersonaToolkit> InWeakPersonaToolkit)
+			TSharedPtr<IPersonaToolkit> PinnedPersonaToolkit = InWeakPersonaToolkit.Pin();
+			if (PinnedPersonaToolkit.IsValid())	// Toolkit can become invalid while the toast is open
 			{
-				TSharedPtr<IPersonaToolkit> PinnedPersonaToolkit = InWeakPersonaToolkit.Pin();
-				if(PinnedPersonaToolkit.IsValid())	// Toolkit can become invalid while the toast is open
+				PinnedPersonaToolkit->SetPreviewMesh(PinnedPersonaToolkit->GetPreviewScene()->GetPreviewMesh(), true);
+				if (InNotificationHandler->Notification.IsValid())
 				{
-					PinnedPersonaToolkit->SetPreviewMesh(PinnedPersonaToolkit->GetPreviewScene()->GetPreviewMesh(), true);
-					if(InNotificationHandler->Notification.IsValid())
-					{
-						InNotificationHandler->Notification->Fadeout();
-					}
+					InNotificationHandler->Notification->Fadeout();
 				}
 			}
+		}
 
-			TSharedPtr<SNotificationItem> Notification;
-		};
+		TSharedPtr<SNotificationItem> Notification;
+	};
 
+	static TSharedRef<SWidget> CreatePreviewMeshComboButtonContents(TWeakPtr<IPersonaToolkit> InWeakPersonaToolkit)
+	{
 		static TWeakPtr<SNotificationItem> WeakPendingApplyMeshNotification = nullptr;
-		auto CreatePreviewMeshComboButtonContents = [WeakPersonaToolkit]()
+
+		FMenuBuilder MenuBuilder(true, nullptr);
+
+		MenuBuilder.BeginSection(TEXT("ChoosePreviewMesh"), LOCTEXT("ChoosePreviewMesh", "Choose Preview Mesh"));
 		{
-			FMenuBuilder MenuBuilder(true, nullptr);
+			FAssetPickerConfig AssetPickerConfig;
 
-			MenuBuilder.BeginSection(TEXT("ChoosePreviewMesh"), LOCTEXT("ChoosePreviewMesh", "Choose Preview Mesh"));
+			auto HandleAssetSelected = [InWeakPersonaToolkit](const FAssetData& AssetData)
 			{
-				FAssetPickerConfig AssetPickerConfig;
+				if (InWeakPersonaToolkit.IsValid())
+				{
+					InWeakPersonaToolkit.Pin()->SetPreviewMesh(Cast<USkeletalMesh>(AssetData.GetAsset()), false);
+				}
 
-				auto HandleAssetSelected = [WeakPersonaToolkit](const FAssetData& AssetData)
+				if (InWeakPersonaToolkit.IsValid())	// SetPreviewMesh can invalidate the persona toolkit, so check it here before displaying toast
 				{
-					if (WeakPersonaToolkit.IsValid())
+					TSharedPtr<FNotificationHandler> NotificationHandler = MakeShared<FNotificationHandler>();
+
+					FNotificationInfo Info(LOCTEXT("PreviewMeshSetTemporarily", "Preview mesh set temporarily"));
+					Info.ExpireDuration = 10.0f;
+					Info.bUseLargeFont = true;
+					Info.ButtonDetails.Add(
+						FNotificationButtonInfo(
+							LOCTEXT("ApplyToAsset", "Apply To Asset"),
+							LOCTEXT("ApplyToAssetToolTip", "The preview mesh has changed, but it will not be able to be saved until it is applied to the asset. Click here to make the change to the preview mesh persistent."),
+							FSimpleDelegate::CreateStatic(&FNotificationHandler::HandleApplyPreviewMesh, NotificationHandler, InWeakPersonaToolkit),
+							SNotificationItem::CS_Success));
+
+					// Fade-out previously added SetPreviewMesh notification (if any)
+					const TSharedPtr<SNotificationItem> PendingNotification = WeakPendingApplyMeshNotification.Pin();
+					if(PendingNotification.IsValid())
 					{
-						WeakPersonaToolkit.Pin()->SetPreviewMesh(Cast<USkeletalMesh>(AssetData.GetAsset()), false);
+						PendingNotification->Fadeout();
 					}
 
-					if(WeakPersonaToolkit.IsValid())	// SetPreviewMesh can invalidate the persona toolkit, so check it here before displaying toast
+					NotificationHandler->Notification = FSlateNotificationManager::Get().AddNotification(Info);
+					WeakPendingApplyMeshNotification = NotificationHandler->Notification;
+					if (NotificationHandler->Notification.IsValid())
 					{
-						TSharedPtr<FNotificationHandler> NotificationHandler = MakeShared<FNotificationHandler>();
-
-						FNotificationInfo Info(LOCTEXT("PreviewMeshSetTemporarily", "Preview mesh set temporarily"));
-						Info.ExpireDuration = 10.0f;
-						Info.bUseLargeFont = true;
-						Info.ButtonDetails.Add(
-							FNotificationButtonInfo(
-								LOCTEXT("ApplyToAsset", "Apply To Asset"), 
-								LOCTEXT("ApplyToAssetToolTip", "The preview mesh has changed, but it will not be able to be saved until it is applied to the asset. Click here to make the change to the preview mesh persistent."),
-								FSimpleDelegate::CreateStatic(&FNotificationHandler::HandleApplyPreviewMesh, NotificationHandler, WeakPersonaToolkit),
-								SNotificationItem::CS_Success));
-
-						// Fade-out previously added SetPreviewMesh notification (if any)
-						const TSharedPtr<SNotificationItem> PendingNotification = WeakPendingApplyMeshNotification.Pin();
-						if(PendingNotification.IsValid())
-						{
-							PendingNotification->Fadeout();
-						}
+						NotificationHandler->Notification->SetCompletionState(SNotificationItem::CS_Success);
+					}
 
-						NotificationHandler->Notification = FSlateNotificationManager::Get().AddNotification(Info);
-						WeakPendingApplyMeshNotification = NotificationHandler->Notification;
-						if (NotificationHandler->Notification.IsValid())
-						{
-							NotificationHandler->Notification->SetCompletionState(SNotificationItem::CS_Success);
-						}
+					FSlateApplication::Get().DismissAllMenus();
+				}
+			};
 
-						FSlateApplication::Get().DismissAllMenus();
-					}
-				};
-				
-				AssetPickerConfig.OnAssetEnterPressed = FOnAssetEnterPressed::CreateLambda([HandleAssetSelected](const TArray<FAssetData>& SelectedAssetData)
+			AssetPickerConfig.OnAssetEnterPressed = FOnAssetEnterPressed::CreateLambda([HandleAssetSelected](const TArray<FAssetData>& SelectedAssetData)
+			{
+				if (SelectedAssetData.Num() == 1)
 				{
-					if (SelectedAssetData.Num() == 1)
-					{
-						HandleAssetSelected(SelectedAssetData[0]);
-					}
-				});
-				AssetPickerConfig.OnAssetSelected = FOnAssetSelected::CreateLambda(HandleAssetSelected);
-				AssetPickerConfig.bAllowNullSelection = false;
-				AssetPickerConfig.InitialAssetViewType = EAssetViewType::List;
-				AssetPickerConfig.Filter.bRecursiveClasses = false;
-				AssetPickerConfig.Filter.ClassPaths.Add(USkeletalMesh::StaticClass()->GetClassPathName());
-				AssetPickerConfig.OnShouldFilterAsset = FOnShouldFilterAsset::CreateLambda([WeakPersonaToolkit](const FAssetData& AssetData)
+					HandleAssetSelected(SelectedAssetData[0]);
+				}
+			});
+			AssetPickerConfig.OnAssetSelected = FOnAssetSelected::CreateLambda(HandleAssetSelected);
+			AssetPickerConfig.bAllowNullSelection = false;
+			AssetPickerConfig.InitialAssetViewType = EAssetViewType::List;
+			AssetPickerConfig.Filter.bRecursiveClasses = false;
+			AssetPickerConfig.Filter.ClassPaths.Add(USkeletalMesh::StaticClass()->GetClassPathName());
+			AssetPickerConfig.OnShouldFilterAsset = FOnShouldFilterAsset::CreateLambda([InWeakPersonaToolkit](const FAssetData& AssetData)
+			{
+				if (InWeakPersonaToolkit.IsValid())
 				{
-					if (WeakPersonaToolkit.IsValid())
+					if (InWeakPersonaToolkit.Pin()->GetContext() == UPhysicsAsset::StaticClass()->GetFName())
 					{
-						if(WeakPersonaToolkit.Pin()->GetContext() == UPhysicsAsset::StaticClass()->GetFName())
-						{
-							return false;
-						}
+						return false;
+					}
 
-						FString TagValue;
-						if (AssetData.GetTagValue("Skeleton", TagValue))
-						{
-							return !WeakPersonaToolkit.Pin()->GetSkeleton()->IsCompatibleSkeletonByAssetData(AssetData);
-						}
+					FString TagValue;
+					if (AssetData.GetTagValue("Skeleton", TagValue))
+					{
+						return !InWeakPersonaToolkit.Pin()->GetSkeleton()->IsCompatibleSkeletonByAssetData(AssetData);
 					}
-					return true;
-				});
-				if (WeakPersonaToolkit.IsValid())
-				{
-					AssetPickerConfig.InitialAssetSelection = FAssetData(WeakPersonaToolkit.Pin()->GetPreviewScene()->GetPreviewMesh());
 				}
+				return true;
+			});
+			if (InWeakPersonaToolkit.IsValid())
+			{
+				AssetPickerConfig.InitialAssetSelection = FAssetData(InWeakPersonaToolkit.Pin()->GetPreviewScene()->GetPreviewMesh());
+			}
 
-				FContentBrowserModule& ContentBrowserModule = FModuleManager::Get().LoadModuleChecked<FContentBrowserModule>(TEXT("ContentBrowser"));
+			FContentBrowserModule& ContentBrowserModule = FModuleManager::Get().LoadModuleChecked<FContentBrowserModule>(TEXT("ContentBrowser"));
 
-				TSharedPtr<SBox> MenuEntry = SNew(SBox)
-					.WidthOverride(300.0f)
-					.HeightOverride(300.0f)
-					[
-						ContentBrowserModule.Get().CreateAssetPicker(AssetPickerConfig)
-					];
+			TSharedPtr<SBox> MenuEntry = SNew(SBox)
+				.WidthOverride(300.0f)
+				.HeightOverride(300.0f)
+				[
+					ContentBrowserModule.Get().CreateAssetPicker(AssetPickerConfig)
+				];
 
-				MenuBuilder.AddWidget(MenuEntry.ToSharedRef(), FText::GetEmpty(), true);
-			}
-			MenuBuilder.EndSection();
+			MenuBuilder.AddWidget(MenuEntry.ToSharedRef(), FText::GetEmpty(), true);
+		}
+		MenuBuilder.EndSection();
 
-			return MenuBuilder.MakeWidget();
-		};
+		return MenuBuilder.MakeWidget();
+	};
 
-		InToolbarBuilder.AddComboButton(
-			FUIAction(),
-			FOnGetContent::CreateLambda(CreatePreviewMeshComboButtonContents),
-			LOCTEXT("SetPreviewMesh", "Preview Mesh"),
-			LOCTEXT("SetPreviewMeshTooltip", "Set a new preview skeletal mesh for the current asset (stored per-animation or per-skeleton)"),
-			FSlateIcon(FAppStyle::GetAppStyleSetName(), "Persona.TogglePreviewAsset", "Persona.TogglePreviewAsset.Small")
-			);
-	}
-
-	if(InArgs.bPreviewAnimation)
+	static TSharedRef<SWidget> CreatePreviewAnimationComboButtonContents(TWeakPtr<IPersonaToolkit> InWeakPersonaToolkit)
 	{
-		auto CreatePreviewAnimationComboButtonContents = [WeakPersonaToolkit]()
-		{
-			FMenuBuilder MenuBuilder(true, nullptr);
+		FMenuBuilder MenuBuilder(true, nullptr);
 
-			MenuBuilder.BeginSection(TEXT("ChoosePreviewAnimation"), LOCTEXT("ChoosePreviewAnimation", "Choose Preview Animation"));
+		MenuBuilder.BeginSection(TEXT("ChoosePreviewAnimation"), LOCTEXT("ChoosePreviewAnimation", "Choose Preview Animation"));
+		{
+			FAssetPickerConfig AssetPickerConfig;
+			AssetPickerConfig.OnAssetSelected = FOnAssetSelected::CreateLambda([InWeakPersonaToolkit](const FAssetData& AssetData)
 			{
-				FAssetPickerConfig AssetPickerConfig;
-				AssetPickerConfig.OnAssetSelected = FOnAssetSelected::CreateLambda([WeakPersonaToolkit](const FAssetData& AssetData)
+				if (InWeakPersonaToolkit.IsValid())
 				{
-					if (WeakPersonaToolkit.IsValid())
-					{
-						TSharedRef<FAnimationEditorPreviewScene> PreviewScene = StaticCastSharedRef<FAnimationEditorPreviewScene>(WeakPersonaToolkit.Pin()->GetPreviewScene());
-						PreviewScene->GetPreviewSceneDescription()->SetPreviewController(UPersonaPreviewSceneAnimationController::StaticClass(), &PreviewScene.Get());
+					TSharedRef<FAnimationEditorPreviewScene> PreviewScene = StaticCastSharedRef<FAnimationEditorPreviewScene>(InWeakPersonaToolkit.Pin()->GetPreviewScene());
+					PreviewScene->GetPreviewSceneDescription()->SetPreviewController(UPersonaPreviewSceneAnimationController::StaticClass(), &PreviewScene.Get());
 
-						UPersonaPreviewSceneAnimationController* AnimController = CastChecked<UPersonaPreviewSceneAnimationController>(PreviewScene->GetPreviewSceneDescription()->PreviewControllerInstance);
-						AnimController->Animation = AssetData.GetAsset();
-						AnimController->InitializeView(PreviewScene->GetPreviewSceneDescription(), &PreviewScene.Get());
+					UPersonaPreviewSceneAnimationController* AnimController = CastChecked<UPersonaPreviewSceneAnimationController>(PreviewScene->GetPreviewSceneDescription()->PreviewControllerInstance);
+					AnimController->Animation = AssetData.GetAsset();
+					AnimController->InitializeView(PreviewScene->GetPreviewSceneDescription(), &PreviewScene.Get());
 
-						// Make sure any settings views are updated with the new settings
-						UAssetViewerSettings::Get()->OnAssetViewerProfileAddRemoved().Broadcast();
-					}
+					// Make sure any settings views are updated with the new settings
+					UAssetViewerSettings::Get()->OnAssetViewerProfileAddRemoved().Broadcast();
+				}
 
-					FSlateApplication::Get().DismissAllMenus();
-				});
-				AssetPickerConfig.bAllowNullSelection = false;
-				AssetPickerConfig.InitialAssetViewType = EAssetViewType::List;
-				AssetPickerConfig.Filter.bRecursiveClasses = true;
-				AssetPickerConfig.Filter.ClassPaths.Add(UAnimationAsset::StaticClass()->GetClassPathName());
-				AssetPickerConfig.OnShouldFilterAsset = FOnShouldFilterAsset::CreateLambda([WeakPersonaToolkit](const FAssetData& AssetData)
+				FSlateApplication::Get().DismissAllMenus();
+			});
+			AssetPickerConfig.bAllowNullSelection = false;
+			AssetPickerConfig.InitialAssetViewType = EAssetViewType::List;
+			AssetPickerConfig.Filter.bRecursiveClasses = true;
+			AssetPickerConfig.Filter.ClassPaths.Add(UAnimationAsset::StaticClass()->GetClassPathName());
+			AssetPickerConfig.OnShouldFilterAsset = FOnShouldFilterAsset::CreateLambda([InWeakPersonaToolkit](const FAssetData& AssetData)
+			{
+				if (InWeakPersonaToolkit.IsValid())
 				{
-					if (WeakPersonaToolkit.IsValid())
+					FString TagValue;
+					if (AssetData.GetTagValue("Skeleton", TagValue))
 					{
-						FString TagValue;
-						if (AssetData.GetTagValue("Skeleton", TagValue))
-						{
-							return !WeakPersonaToolkit.Pin()->GetSkeleton()->IsCompatibleSkeletonByAssetData(AssetData);
-						}
+						return !InWeakPersonaToolkit.Pin()->GetSkeleton()->IsCompatibleSkeletonByAssetData(AssetData);
 					}
-					return true;
-				});
-				if (WeakPersonaToolkit.IsValid())
-				{
-					AssetPickerConfig.InitialAssetSelection = FAssetData(WeakPersonaToolkit.Pin()->GetPreviewScene()->GetPreviewAnimationAsset());
 				}
+				return true;
+			});
+			if (InWeakPersonaToolkit.IsValid())
+			{
+				AssetPickerConfig.InitialAssetSelection = FAssetData(InWeakPersonaToolkit.Pin()->GetPreviewScene()->GetPreviewAnimationAsset());
+			}
 
-				FContentBrowserModule& ContentBrowserModule = FModuleManager::Get().LoadModuleChecked<FContentBrowserModule>(TEXT("ContentBrowser"));
+			FContentBrowserModule& ContentBrowserModule = FModuleManager::Get().LoadModuleChecked<FContentBrowserModule>(TEXT("ContentBrowser"));
 
-				TSharedPtr<SBox> MenuEntry = SNew(SBox)
-					.WidthOverride(300.0f)
-					.HeightOverride(300.0f)
-					[
-						ContentBrowserModule.Get().CreateAssetPicker(AssetPickerConfig)
-					];
+			TSharedPtr<SBox> MenuEntry = SNew(SBox)
+				.WidthOverride(300.0f)
+				.HeightOverride(300.0f)
+				[
+					ContentBrowserModule.Get().CreateAssetPicker(AssetPickerConfig)
+				];
 
-				MenuBuilder.AddWidget(MenuEntry.ToSharedRef(), FText::GetEmpty(), true);
-			}
-			MenuBuilder.EndSection();
+			MenuBuilder.AddWidget(MenuEntry.ToSharedRef(), FText::GetEmpty(), true);
+		}
+		MenuBuilder.EndSection();
 
-			return MenuBuilder.MakeWidget();
-		};
+		return MenuBuilder.MakeWidget();
+	};
+
+	static void ShowReferencePose(TWeakPtr<IPersonaToolkit> InWeakPersonaToolkit)
+	{
+		if (InWeakPersonaToolkit.IsValid())
+		{
+			TSharedRef<FAnimationEditorPreviewScene> PreviewScene = StaticCastSharedRef<FAnimationEditorPreviewScene>(InWeakPersonaToolkit.Pin()->GetPreviewScene());
+			PreviewScene->GetPreviewSceneDescription()->SetPreviewController(UPersonaPreviewSceneRefPoseController::StaticClass(), &PreviewScene.Get());
 
+			UPersonaPreviewSceneRefPoseController* AnimController = CastChecked<UPersonaPreviewSceneRefPoseController>(PreviewScene->GetPreviewSceneDescription()->PreviewControllerInstance);
+			AnimController->bResetBoneTransforms = true;
+			AnimController->InitializeView(PreviewScene->GetPreviewSceneDescription(), &PreviewScene.Get());
+
+			// Reset this to false here as we dont want it to always reset bone transforms, only if they user picks it from the toolbar
+			AnimController->bResetBoneTransforms = false;
+
+			// Make sure any settings views are updated with the new settings
+			UAssetViewerSettings::Get()->OnAssetViewerProfileAddRemoved().Broadcast();
+		}
+	}
+}
+
+void FPersonaModule::AddCommonToolbarExtensions(FToolBarBuilder& InToolbarBuilder, TSharedRef<IPersonaToolkit> PersonaToolkit, const FCommonToolbarExtensionArgs& InArgs)
+{
+	TWeakPtr<IPersonaToolkit> WeakPersonaToolkit = PersonaToolkit;
+
+	if (InArgs.bPreviewMesh)
+	{
+		InToolbarBuilder.AddComboButton(
+			FUIAction(),
+			FOnGetContent::CreateStatic(&UE::Persona::Private::CreatePreviewMeshComboButtonContents, WeakPersonaToolkit),
+			LOCTEXT("SetPreviewMesh", "Preview Mesh"),
+			LOCTEXT("SetPreviewMeshTooltip", "Set a new preview skeletal mesh for the current asset (stored per-animation or per-skeleton)"),
+			FSlateIcon(FAppStyle::GetAppStyleSetName(), "Persona.TogglePreviewAsset", "Persona.TogglePreviewAsset.Small")
+		);
+	}
+
+	if (InArgs.bPreviewAnimation)
+	{
 		InToolbarBuilder.AddComboButton(
 			FUIAction(),
-			FOnGetContent::CreateLambda(CreatePreviewAnimationComboButtonContents),
+			FOnGetContent::CreateStatic(&UE::Persona::Private::CreatePreviewAnimationComboButtonContents, WeakPersonaToolkit),
 			LOCTEXT("SetPreviewAnimation", "Preview Animation"),
 			LOCTEXT("SetPreviewAnimationTooltip", "Setup the scene to use a preview animation. More advanced settings are available in Preview Scene Settings."),
 			FSlateIcon(FAppStyle::GetAppStyleSetName(), "Persona.TogglePreviewAnimation", "Persona.TogglePreviewAnimation.Small")
-			);
+		);
 	}
 
-	if(InArgs.bReferencePose)
+	if (InArgs.bReferencePose)
 	{
 		InToolbarBuilder.AddToolBarButton(
-			FUIAction(
-				FExecuteAction::CreateLambda([WeakPersonaToolkit]()
-				{
-					if (WeakPersonaToolkit.IsValid())
-					{
-						TSharedRef<FAnimationEditorPreviewScene> PreviewScene = StaticCastSharedRef<FAnimationEditorPreviewScene>(WeakPersonaToolkit.Pin()->GetPreviewScene());
-						PreviewScene->GetPreviewSceneDescription()->SetPreviewController(UPersonaPreviewSceneRefPoseController::StaticClass(), &PreviewScene.Get());
-
-						UPersonaPreviewSceneRefPoseController* AnimController = CastChecked<UPersonaPreviewSceneRefPoseController>(PreviewScene->GetPreviewSceneDescription()->PreviewControllerInstance);
-						AnimController->bResetBoneTransforms = true;
-						AnimController->InitializeView(PreviewScene->GetPreviewSceneDescription(), &PreviewScene.Get());
-
-						// Reset this to false here as we dont want it to always reset bone transforms, only if they user picks it from the toolbar
-						AnimController->bResetBoneTransforms = false;
-
-						// Make sure any settings views are updated with the new settings
-						UAssetViewerSettings::Get()->OnAssetViewerProfileAddRemoved().Broadcast();
-					}
-				})
-			),
+			FUIAction(FExecuteAction::CreateStatic(&UE::Persona::Private::ShowReferencePose, WeakPersonaToolkit)),
 			NAME_None,
 			LOCTEXT("ShowReferencePose", "Reference Pose"),
 			LOCTEXT("ShowReferencePoseTooltip", "Show the reference pose. Clears all bone modifications. More advanced settings are available in Preview Scene Settings."),
 			FSlateIcon(FAppStyle::GetAppStyleSetName(), "Persona.ToggleReferencePose", "Persona.ToggleReferencePose.Small")
-			);
+		);
 	}
 
-	if(InArgs.bCreateAsset)
+	if (InArgs.bCreateAsset)
 	{
 		InToolbarBuilder.AddComboButton(
 			FUIAction(),
-			FOnGetContent::CreateRaw(this, &FPersonaModule::GenerateCreateAssetMenu, WeakPersonaToolkit),
+			FOnGetContent::CreateStatic(&FPersonaModule::GenerateCreateAssetMenu, WeakPersonaToolkit),
 			LOCTEXT("CreateAsset_Label", "Create Asset"),
 			LOCTEXT("CreateAsset_ToolTip", "Create Assets for this skeleton."),
 			FSlateIcon(FAppStyle::GetAppStyleSetName(), "Persona.CreateAsset")
@@ -1077,7 +1082,128 @@ void FPersonaModule::AddCommonToolbarExtensions(FToolBarBuilder& InToolbarBuilde
 	}
 }
 
-TSharedRef< SWidget > FPersonaModule::GenerateCreateAssetMenu(TWeakPtr<IPersonaToolkit> InWeakPersonaToolkit) const
+void FPersonaModule::AddCommonMenuExtensions(UToolMenu* InToolMenu, const FCommonToolMenuExtensionArgs& InArgs)
+{
+	UPersonaToolMenuContext* PersonaToolMenuContext = InToolMenu->FindContext<UPersonaToolMenuContext>();
+	if (PersonaToolMenuContext == nullptr)
+	{
+		return;
+	}
+
+	TWeakPtr<IPersonaToolkit> WeakPersonaToolkit = PersonaToolMenuContext->GetToolkit();
+
+	FToolMenuSection& PersonaSection = InToolMenu->AddSection("Persona", LOCTEXT("PersonaSectionName", "Persona"));
+
+	if (InArgs.bPreviewMesh)
+	{
+		PersonaSection.AddSubMenu(
+			"PreviewMesh",
+			LOCTEXT("SetPreviewMesh", "Preview Mesh"),
+			LOCTEXT("SetPreviewMeshTooltip", "Set a new preview skeletal mesh for the current asset (stored per-animation or per-skeleton)"),
+			FOnGetContent::CreateStatic(&UE::Persona::Private::CreatePreviewMeshComboButtonContents, WeakPersonaToolkit),
+			false,
+			FSlateIcon(FAppStyle::GetAppStyleSetName(), "Persona.TogglePreviewAsset", "Persona.TogglePreviewAsset.Small")
+		);
+	}
+
+	if (InArgs.bPreviewAnimation)
+	{
+		PersonaSection.AddSubMenu(
+			"PreviewAnimation",
+			LOCTEXT("SetPreviewAnimation", "Preview Animation"),
+			LOCTEXT("SetPreviewAnimationTooltip", "Setup the scene to use a preview animation. More advanced settings are available in Preview Scene Settings."),
+			FOnGetContent::CreateStatic(&UE::Persona::Private::CreatePreviewAnimationComboButtonContents, WeakPersonaToolkit),
+			false,
+			FSlateIcon(FAppStyle::GetAppStyleSetName(), "Persona.TogglePreviewAnimation", "Persona.TogglePreviewAnimation.Small")
+		);
+	}
+
+	if (InArgs.bReferencePose)
+	{
+		PersonaSection.AddMenuEntry(
+			"ReferencePose",
+			LOCTEXT("ShowReferencePose", "Reference Pose"),
+			LOCTEXT("ShowReferencePoseTooltip", "Show the reference pose. Clears all bone modifications. More advanced settings are available in Preview Scene Settings."),
+			FSlateIcon(FAppStyle::GetAppStyleSetName(), "Persona.ToggleReferencePose", "Persona.ToggleReferencePose.Small"),
+			FUIAction(FExecuteAction::CreateStatic(&UE::Persona::Private::ShowReferencePose, WeakPersonaToolkit))
+		);
+	}
+
+	if (InArgs.bCreateAsset)
+	{
+		PersonaSection.AddSubMenu(
+			"CreateAsset",
+			LOCTEXT("CreateAsset_Label", "Create Asset"),
+			LOCTEXT("CreateAsset_ToolTip", "Create Assets for this skeleton."),
+			FOnGetContent::CreateStatic(&FPersonaModule::GenerateCreateAssetMenu, WeakPersonaToolkit),
+			false,
+			FSlateIcon(FAppStyle::GetAppStyleSetName(), "Persona.CreateAsset")
+		);
+	}
+}
+
+
+void FPersonaModule::AddCommonToolbarExtensions(UToolMenu* InToolMenu, const FCommonToolMenuExtensionArgs& InArgs)
+{
+	UPersonaToolMenuContext* PersonaToolMenuContext = InToolMenu->FindContext<UPersonaToolMenuContext>();
+	if (PersonaToolMenuContext == nullptr)
+	{
+		return;
+	}
+
+	TWeakPtr<IPersonaToolkit> WeakPersonaToolkit = PersonaToolMenuContext->GetToolkit();
+
+	FToolMenuSection& PersonaSection = InToolMenu->AddSection("Persona", LOCTEXT("PersonaSectionName", "Persona"));
+
+	if (InArgs.bPreviewMesh)
+	{
+		PersonaSection.AddEntry(FToolMenuEntry::InitComboButton(
+			"PreviewMesh",
+			FUIAction(),
+			FOnGetContent::CreateStatic(&UE::Persona::Private::CreatePreviewMeshComboButtonContents, WeakPersonaToolkit),
+			LOCTEXT("SetPreviewMesh", "Preview Mesh"),
+			LOCTEXT("SetPreviewMeshTooltip", "Set a new preview skeletal mesh for the current asset (stored per-animation or per-skeleton)"),
+			FSlateIcon(FAppStyle::GetAppStyleSetName(), "Persona.TogglePreviewAsset", "Persona.TogglePreviewAsset.Small"))
+		);
+	}
+
+	if (InArgs.bPreviewAnimation)
+	{
+		PersonaSection.AddEntry(FToolMenuEntry::InitComboButton(
+			"PreviewAnimation",
+			FUIAction(),
+			FOnGetContent::CreateStatic(&UE::Persona::Private::CreatePreviewAnimationComboButtonContents, WeakPersonaToolkit),
+			LOCTEXT("SetPreviewAnimation", "Preview Animation"),
+			LOCTEXT("SetPreviewAnimationTooltip", "Setup the scene to use a preview animation. More advanced settings are available in Preview Scene Settings."),
+			FSlateIcon(FAppStyle::GetAppStyleSetName(), "Persona.TogglePreviewAnimation", "Persona.TogglePreviewAnimation.Small"))
+		);
+	}
+
+	if (InArgs.bReferencePose)
+	{
+		PersonaSection.AddMenuEntry(
+			"ReferencePose",
+			LOCTEXT("ShowReferencePose", "Reference Pose"),
+			LOCTEXT("ShowReferencePoseTooltip", "Show the reference pose. Clears all bone modifications. More advanced settings are available in Preview Scene Settings."),
+			FSlateIcon(FAppStyle::GetAppStyleSetName(), "Persona.ToggleReferencePose", "Persona.ToggleReferencePose.Small"),
+			FUIAction(FExecuteAction::CreateStatic(&UE::Persona::Private::ShowReferencePose, WeakPersonaToolkit))
+		);
+	}
+
+	if (InArgs.bCreateAsset)
+	{
+		PersonaSection.AddEntry(FToolMenuEntry::InitComboButton(
+			"CreateAsset",
+			FUIAction(),
+			FOnGetContent::CreateStatic(&FPersonaModule::GenerateCreateAssetMenu, WeakPersonaToolkit),
+			LOCTEXT("CreateAsset_Label", "Create Asset"),
+			LOCTEXT("CreateAsset_ToolTip", "Create Assets for this skeleton."),
+			FSlateIcon(FAppStyle::GetAppStyleSetName(), "Persona.CreateAsset"))
+		);
+	}
+}
+
+TSharedRef< SWidget > FPersonaModule::GenerateCreateAssetMenu(TWeakPtr<IPersonaToolkit> InWeakPersonaToolkit)
 {
 	IAssetTools& AssetTools = FModuleManager::LoadModuleChecked<FAssetToolsModule>("AssetTools").Get();
 
@@ -1093,7 +1219,7 @@ TSharedRef< SWidget > FPersonaModule::GenerateCreateAssetMenu(TWeakPtr<IPersonaT
 			MenuBuilder.AddSubMenu(
 				LOCTEXT("CreateAnimationSubmenu", "Create Animation"),
 				LOCTEXT("CreateAnimationSubmenu_ToolTip", "Create Animation for this skeleton"),
-				FNewMenuDelegate::CreateRaw(this, &FPersonaModule::FillCreateAnimationMenu, InWeakPersonaToolkit),
+				FNewMenuDelegate::CreateStatic(&FPersonaModule::FillCreateAnimationMenu, InWeakPersonaToolkit),
 				false,
 				FSlateIcon(FAppStyle::GetAppStyleSetName(), "Persona.AssetActions.CreateAnimAsset")
 				);
@@ -1104,7 +1230,7 @@ TSharedRef< SWidget > FPersonaModule::GenerateCreateAssetMenu(TWeakPtr<IPersonaT
 			MenuBuilder.AddSubMenu(
 				LOCTEXT("CreatePoseAssetSubmenu", "Create PoseAsset"),
 				LOCTEXT("CreatePoseAsssetSubmenu_ToolTip", "Create PoseAsset for this skeleton"),
-				FNewMenuDelegate::CreateRaw(this, &FPersonaModule::FillCreatePoseAssetMenu, InWeakPersonaToolkit),
+				FNewMenuDelegate::CreateStatic(&FPersonaModule::FillCreatePoseAssetMenu, InWeakPersonaToolkit),
 				false,
 				FSlateIcon(FAppStyle::GetAppStyleSetName(), "ClassIcon.PoseAsset")
 			);
@@ -1123,12 +1249,12 @@ TSharedRef< SWidget > FPersonaModule::GenerateCreateAssetMenu(TWeakPtr<IPersonaT
 		Objects.Add(PersonaToolkit->GetSkeleton());
 	}
 
-	AnimationEditorUtils::FillCreateAssetMenu(MenuBuilder, Objects, FAnimAssetCreated::CreateRaw(const_cast<FPersonaModule*>(this), &FPersonaModule::HandleAssetCreated), false);
+	AnimationEditorUtils::FillCreateAssetMenu(MenuBuilder, Objects, FAnimAssetCreated::CreateStatic(&FPersonaModule::HandleAssetCreated), false);
 
 	return MenuBuilder.MakeWidget();
 }
 
-void FPersonaModule::FillCreateAnimationMenu(FMenuBuilder& MenuBuilder, TWeakPtr<IPersonaToolkit> InWeakPersonaToolkit) const
+void FPersonaModule::FillCreateAnimationMenu(FMenuBuilder& MenuBuilder, TWeakPtr<IPersonaToolkit> InWeakPersonaToolkit)
 {
 	TSharedRef<IPersonaToolkit> PersonaToolkit = InWeakPersonaToolkit.Pin().ToSharedRef();
 	TArray<TWeakObjectPtr<UObject>> Objects;
@@ -1152,7 +1278,7 @@ void FPersonaModule::FillCreateAnimationMenu(FMenuBuilder& MenuBuilder, TWeakPtr
 			LOCTEXT("CreateAnimation_RefPose_Tooltip", "Create Animation from reference pose."),
 			FSlateIcon(),
 			FUIAction(
-				FExecuteAction::CreateStatic(&AnimationEditorUtils::ExecuteNewAnimAsset<UAnimSequenceFactory, UAnimSequence>, Objects, FString("_Sequence"), FAnimAssetCreated::CreateRaw(const_cast<FPersonaModule*>(this), &FPersonaModule::CreateAnimation, EPoseSourceOption::ReferencePose, InWeakPersonaToolkit), bDoNotShowNameDialog, bAllowReplaceExisting),
+				FExecuteAction::CreateStatic(&AnimationEditorUtils::ExecuteNewAnimAsset<UAnimSequenceFactory, UAnimSequence>, Objects, FString("_Sequence"), FAnimAssetCreated::CreateStatic(&FPersonaModule::CreateAnimation, EPoseSourceOption::ReferencePose, InWeakPersonaToolkit), bDoNotShowNameDialog, bAllowReplaceExisting),
 				FCanExecuteAction()
 				)
 			);
@@ -1162,7 +1288,7 @@ void FPersonaModule::FillCreateAnimationMenu(FMenuBuilder& MenuBuilder, TWeakPtr
 			LOCTEXT("CreateAnimation_CurrentPose_Tooltip", "Create Animation from current pose."),
 			FSlateIcon(),
 			FUIAction(
-				FExecuteAction::CreateStatic(&AnimationEditorUtils::ExecuteNewAnimAsset<UAnimSequenceFactory, UAnimSequence>, Objects, FString("_Sequence"), FAnimAssetCreated::CreateRaw(const_cast<FPersonaModule*>(this), &FPersonaModule::CreateAnimation, EPoseSourceOption::CurrentPose, InWeakPersonaToolkit), bDoNotShowNameDialog, bAllowReplaceExisting),
+				FExecuteAction::CreateStatic(&AnimationEditorUtils::ExecuteNewAnimAsset<UAnimSequenceFactory, UAnimSequence>, Objects, FString("_Sequence"), FAnimAssetCreated::CreateStatic(&FPersonaModule::CreateAnimation, EPoseSourceOption::CurrentPose, InWeakPersonaToolkit), bDoNotShowNameDialog, bAllowReplaceExisting),
 				FCanExecuteAction()
 				)
 			);
@@ -1172,7 +1298,7 @@ void FPersonaModule::FillCreateAnimationMenu(FMenuBuilder& MenuBuilder, TWeakPtr
 			MenuBuilder.AddSubMenu(
 				LOCTEXT("CreateAnimation_CurrenAnimationSubMenu", "Current Animation"),
 				LOCTEXT("CreateAnimation_CurrenAnimationSubMenu_ToolTip", "Create Animation from current animation"),
-				FNewMenuDelegate::CreateRaw(this, &FPersonaModule::FillCreateAnimationFromCurrentAnimationMenu, InWeakPersonaToolkit),
+				FNewMenuDelegate::CreateStatic(&FPersonaModule::FillCreateAnimationFromCurrentAnimationMenu, InWeakPersonaToolkit),
 				false,
 				FSlateIcon(FAppStyle::GetAppStyleSetName(), "Persona.AssetActions.CreateAnimAsset")
 			);
@@ -1181,7 +1307,7 @@ void FPersonaModule::FillCreateAnimationMenu(FMenuBuilder& MenuBuilder, TWeakPtr
 	MenuBuilder.EndSection();
 }
 
-void FPersonaModule::FillCreateAnimationFromCurrentAnimationMenu(FMenuBuilder& MenuBuilder, TWeakPtr<IPersonaToolkit> InWeakPersonaToolkit) const
+void FPersonaModule::FillCreateAnimationFromCurrentAnimationMenu(FMenuBuilder& MenuBuilder, TWeakPtr<IPersonaToolkit> InWeakPersonaToolkit) 
 {
 	TSharedRef<IPersonaToolkit> PersonaToolkit = InWeakPersonaToolkit.Pin().ToSharedRef();
 	TArray<TWeakObjectPtr<UObject>> Objects;
@@ -1206,7 +1332,7 @@ void FPersonaModule::FillCreateAnimationFromCurrentAnimationMenu(FMenuBuilder& M
 			LOCTEXT("CreateAnimation_CurrentAnimation_AnimData_Tooltip", "Create Animation from Animation Source Data."),
 			FSlateIcon(),
 			FUIAction(
-				FExecuteAction::CreateStatic(&AnimationEditorUtils::ExecuteNewAnimAsset<UAnimSequenceFactory, UAnimSequence>, Objects, FString("_Sequence"), FAnimAssetCreated::CreateRaw(const_cast<FPersonaModule*>(this), &FPersonaModule::CreateAnimation, EPoseSourceOption::CurrentAnimation_AnimData, InWeakPersonaToolkit), bDoNotShowNameDialog, bAllowReplaceExisting)
+				FExecuteAction::CreateStatic(&AnimationEditorUtils::ExecuteNewAnimAsset<UAnimSequenceFactory, UAnimSequence>, Objects, FString("_Sequence"), FAnimAssetCreated::CreateStatic(&FPersonaModule::CreateAnimation, EPoseSourceOption::CurrentAnimation_AnimData, InWeakPersonaToolkit), bDoNotShowNameDialog, bAllowReplaceExisting)
 			)
 		);
 
@@ -1215,14 +1341,14 @@ void FPersonaModule::FillCreateAnimationFromCurrentAnimationMenu(FMenuBuilder& M
 			LOCTEXT("CreateAnimation_CurrentAnimation_PreviewMesh_Tooltip", "Create Animation by playing on the Current Preview Mesh, including Retargeting, Post Process Graph, or anything you see on the preview mesh."),
 			FSlateIcon(),
 			FUIAction(
-				FExecuteAction::CreateStatic(&AnimationEditorUtils::ExecuteNewAnimAsset<UAnimSequenceFactory, UAnimSequence>, Objects, FString("_Sequence"), FAnimAssetCreated::CreateRaw(const_cast<FPersonaModule*>(this), &FPersonaModule::CreateAnimation, EPoseSourceOption::CurrentAnimation_PreviewMesh, InWeakPersonaToolkit), bDoNotShowNameDialog, bAllowReplaceExisting)
+				FExecuteAction::CreateStatic(&AnimationEditorUtils::ExecuteNewAnimAsset<UAnimSequenceFactory, UAnimSequence>, Objects, FString("_Sequence"), FAnimAssetCreated::CreateStatic(&FPersonaModule::CreateAnimation, EPoseSourceOption::CurrentAnimation_PreviewMesh, InWeakPersonaToolkit), bDoNotShowNameDialog, bAllowReplaceExisting)
 			)
 		);
 	}
 	MenuBuilder.EndSection();
 }
 
-void FPersonaModule::FillCreatePoseAssetMenu(FMenuBuilder& MenuBuilder, TWeakPtr<IPersonaToolkit> InWeakPersonaToolkit) const
+void FPersonaModule::FillCreatePoseAssetMenu(FMenuBuilder& MenuBuilder, TWeakPtr<IPersonaToolkit> InWeakPersonaToolkit)
 {
 	TSharedRef<IPersonaToolkit> PersonaToolkit = InWeakPersonaToolkit.Pin().ToSharedRef();
 	TArray<TWeakObjectPtr<UObject>> Objects;
@@ -1247,7 +1373,7 @@ void FPersonaModule::FillCreatePoseAssetMenu(FMenuBuilder& MenuBuilder, TWeakPtr
 			LOCTEXT("CreatePoseAsset_CurrentPose_Tooltip", "Create PoseAsset from current pose."),
 			FSlateIcon(),
 			FUIAction(
-				FExecuteAction::CreateStatic(&AnimationEditorUtils::ExecuteNewAnimAsset<UPoseAssetFactory, UPoseAsset>, Objects, FString("_PoseAsset"), FAnimAssetCreated::CreateRaw(const_cast<FPersonaModule*>(this), &FPersonaModule::CreatePoseAsset, EPoseSourceOption::CurrentPose, InWeakPersonaToolkit), bDoNotShowNameDialog, bAllowReplaceExisting),
+				FExecuteAction::CreateStatic(&AnimationEditorUtils::ExecuteNewAnimAsset<UPoseAssetFactory, UPoseAsset>, Objects, FString("_PoseAsset"), FAnimAssetCreated::CreateStatic(&FPersonaModule::CreatePoseAsset, EPoseSourceOption::CurrentPose, InWeakPersonaToolkit), bDoNotShowNameDialog, bAllowReplaceExisting),
 				FCanExecuteAction()
 			)
 		);
@@ -1257,7 +1383,7 @@ void FPersonaModule::FillCreatePoseAssetMenu(FMenuBuilder& MenuBuilder, TWeakPtr
 			LOCTEXT("CreatePoseAsset_CurrentAnimation_Tooltip", "Create Animation from current animation."),
 			FSlateIcon(),
 			FUIAction(
-				FExecuteAction::CreateStatic(&AnimationEditorUtils::ExecuteNewAnimAsset<UPoseAssetFactory, UPoseAsset>, Objects, FString("_PoseAsset"), FAnimAssetCreated::CreateRaw(const_cast<FPersonaModule*>(this), &FPersonaModule::CreatePoseAsset, EPoseSourceOption::CurrentAnimation_AnimData, InWeakPersonaToolkit), bDoNotShowNameDialog, bAllowReplaceExisting),
+				FExecuteAction::CreateStatic(&AnimationEditorUtils::ExecuteNewAnimAsset<UPoseAssetFactory, UPoseAsset>, Objects, FString("_PoseAsset"), FAnimAssetCreated::CreateStatic(&FPersonaModule::CreatePoseAsset, EPoseSourceOption::CurrentAnimation_AnimData, InWeakPersonaToolkit), bDoNotShowNameDialog, bAllowReplaceExisting),
 				FCanExecuteAction()
 			)
 		);
@@ -1270,7 +1396,7 @@ void FPersonaModule::FillCreatePoseAssetMenu(FMenuBuilder& MenuBuilder, TWeakPtr
 		MenuBuilder.AddSubMenu(
 			LOCTEXT("InsertPoseSubmenu", "Insert Pose"),
 			LOCTEXT("InsertPoseSubmenu_ToolTip", "Insert current pose to selected PoseAsset"),
-			FNewMenuDelegate::CreateRaw(this, &FPersonaModule::FillInsertPoseMenu, InWeakPersonaToolkit),
+			FNewMenuDelegate::CreateStatic(&FPersonaModule::FillInsertPoseMenu, InWeakPersonaToolkit),
 			false,
 			FSlateIcon(FAppStyle::GetAppStyleSetName(), "ClassIcon.PoseAsset")
 		);
@@ -1278,7 +1404,7 @@ void FPersonaModule::FillCreatePoseAssetMenu(FMenuBuilder& MenuBuilder, TWeakPtr
 	MenuBuilder.EndSection();
 }
 
-void FPersonaModule::FillInsertPoseMenu(FMenuBuilder& MenuBuilder, TWeakPtr<IPersonaToolkit> InWeakPersonaToolkit) const
+void FPersonaModule::FillInsertPoseMenu(FMenuBuilder& MenuBuilder, TWeakPtr<IPersonaToolkit> InWeakPersonaToolkit)
 {
 	FAssetPickerConfig AssetPickerConfig;
 
@@ -1292,7 +1418,7 @@ void FPersonaModule::FillInsertPoseMenu(FMenuBuilder& MenuBuilder, TWeakPtr<IPer
 	AssetPickerConfig.Filter.TagsAndValues.Add(TEXT("Skeleton"), FAssetData(Skeleton).GetExportTextName());
 
 	/** The delegate that fires when an asset was selected */
-	AssetPickerConfig.OnAssetSelected = FOnAssetSelected::CreateRaw(const_cast<FPersonaModule*>(this), &FPersonaModule::InsertCurrentPoseToAsset, InWeakPersonaToolkit);
+	AssetPickerConfig.OnAssetSelected = FOnAssetSelected::CreateStatic(&FPersonaModule::InsertCurrentPoseToAsset, InWeakPersonaToolkit);
 
 	/** The default view mode should be a list view */
 	AssetPickerConfig.InitialAssetViewType = EAssetViewType::List;
diff --git a/Engine/Source/Editor/Persona/Public/PersonaModule.h b/Engine/Source/Editor/Persona/Public/PersonaModule.h
--- a/Engine/Source/Editor/Persona/Public/PersonaModule.h
+++ b/Engine/Source/Editor/Persona/Public/PersonaModule.h
@@ -386,9 +386,9 @@ public:
 	virtual FOnPreviewSceneCreated& OnPreviewSceneCreated() { return OnPreviewSceneCreatedDelegate; }
 
 	/** Settings for AddCommonToolbarExtensions */
-	struct FCommonToolbarExtensionArgs
+	struct FCommonToolMenuExtensionArgs
 	{
-		FCommonToolbarExtensionArgs()
+		FCommonToolMenuExtensionArgs()
 			: bPreviewMesh(true)
 			, bPreviewAnimation(true)
 			, bReferencePose(false)
@@ -408,7 +408,15 @@ public:
 		bool bCreateAsset;
 	};
 
-	/** Add common toobar extensions */
+	typedef FCommonToolMenuExtensionArgs FCommonToolbarExtensionArgs;
+
+	/** Add common menu extensions */
+	virtual void AddCommonMenuExtensions(UToolMenu* InToolMenu, const FCommonToolMenuExtensionArgs& InArgs = FCommonToolMenuExtensionArgs());
+
+	/** Add common toolbar extensions */
+	virtual void AddCommonToolbarExtensions(UToolMenu* InToolMenu, const FCommonToolMenuExtensionArgs& InArgs = FCommonToolMenuExtensionArgs());
+
+	/** Add common toobar extensions (legacy support) - DEPRECATED */
 	virtual void AddCommonToolbarExtensions(FToolBarBuilder& InToolbarBuilder, TSharedRef<IPersonaToolkit> PersonaToolkit, const FCommonToolbarExtensionArgs& InArgs = FCommonToolbarExtensionArgs());
 
 	/** Register common layout extensions */
@@ -434,23 +442,23 @@ private:
 		Max
 	};
 
-	TSharedRef< SWidget > GenerateCreateAssetMenu(TWeakPtr<IPersonaToolkit> InWeakPersonaToolkit) const;
+	static TSharedRef< SWidget > GenerateCreateAssetMenu(TWeakPtr<IPersonaToolkit> InWeakPersonaToolkit);
 
-	void FillCreateAnimationMenu(FMenuBuilder& MenuBuilder, TWeakPtr<IPersonaToolkit> InWeakPersonaToolkit) const;
+	static void FillCreateAnimationMenu(FMenuBuilder& MenuBuilder, TWeakPtr<IPersonaToolkit> InWeakPersonaToolkit);
 
-	void FillCreateAnimationFromCurrentAnimationMenu(FMenuBuilder& MenuBuilder, TWeakPtr<IPersonaToolkit> InWeakPersonaToolkit) const;
+	static void FillCreateAnimationFromCurrentAnimationMenu(FMenuBuilder& MenuBuilder, TWeakPtr<IPersonaToolkit> InWeakPersonaToolkit);
 
-	void FillCreatePoseAssetMenu(FMenuBuilder& MenuBuilder, TWeakPtr<IPersonaToolkit> InWeakPersonaToolkit) const;
+	static void FillCreatePoseAssetMenu(FMenuBuilder& MenuBuilder, TWeakPtr<IPersonaToolkit> InWeakPersonaToolkit);
 
-	void FillInsertPoseMenu(FMenuBuilder& MenuBuilder, TWeakPtr<IPersonaToolkit> InWeakPersonaToolkit) const;
+	static void FillInsertPoseMenu(FMenuBuilder& MenuBuilder, TWeakPtr<IPersonaToolkit> InWeakPersonaToolkit);
 
-	void InsertCurrentPoseToAsset(const FAssetData& NewPoseAssetData, TWeakPtr<IPersonaToolkit> InWeakPersonaToolkit);
+	static void InsertCurrentPoseToAsset(const FAssetData& NewPoseAssetData, TWeakPtr<IPersonaToolkit> InWeakPersonaToolkit);
 
-	bool CreateAnimation(const TArray<UObject*> NewAssets, const EPoseSourceOption Option, TWeakPtr<IPersonaToolkit> InWeakPersonaToolkit);
+	static bool CreateAnimation(const TArray<UObject*> NewAssets, const EPoseSourceOption Option, TWeakPtr<IPersonaToolkit> InWeakPersonaToolkit);
 
-	bool CreatePoseAsset(const TArray<UObject*> NewAssets, const EPoseSourceOption Option, TWeakPtr<IPersonaToolkit> InWeakPersonaToolkit);
+	static bool CreatePoseAsset(const TArray<UObject*> NewAssets, const EPoseSourceOption Option, TWeakPtr<IPersonaToolkit> InWeakPersonaToolkit);
 	
-	bool HandleAssetCreated(const TArray<UObject*> NewAssets);
+	static bool HandleAssetCreated(const TArray<UObject*> NewAssets);
 
 private:
 	TSharedPtr<FExtensibilityManager> MenuExtensibilityManager;
diff --git a/Engine/Source/Editor/Persona/Public/PersonaToolMenuContext.h b/Engine/Source/Editor/Persona/Public/PersonaToolMenuContext.h
--- a/Engine/Source/Editor/Persona/Public/PersonaToolMenuContext.h
+++ b/Engine/Source/Editor/Persona/Public/PersonaToolMenuContext.h
@@ -7,6 +7,7 @@
 #include "UObject/Object.h"
 #include "UObject/ObjectMacros.h"
 #include "UObject/UObjectGlobals.h"
+#include "PersonaModule.h"
 
 #include "PersonaToolMenuContext.generated.h"
 
@@ -44,8 +45,12 @@ public:
 	UFUNCTION(BlueprintCallable, Category = PersonaEditorExtensions)
 	UAnimationAsset* GetAnimationAsset() const;
 
+	/** Get the persona toolkit */
 	void SetToolkit(TSharedRef<IPersonaToolkit> InToolkit);
 
+	/** Get a weak ptr to the persona toolkit */
+	TWeakPtr<IPersonaToolkit> GetToolkit() const { return WeakToolkit; }
+
 protected:
 	bool HasValidToolkit() const;
 private:
diff --git a/Engine/Source/Editor/SkeletalMeshEditor/Private/SkeletalMeshEditor.cpp b/Engine/Source/Editor/SkeletalMeshEditor/Private/SkeletalMeshEditor.cpp
--- a/Engine/Source/Editor/SkeletalMeshEditor/Private/SkeletalMeshEditor.cpp
+++ b/Engine/Source/Editor/SkeletalMeshEditor/Private/SkeletalMeshEditor.cpp
@@ -616,15 +616,15 @@ void FSkeletalMeshEditor::ExtendToolbar()
 	const FToolMenuInsert SectionInsertLocation("Asset", EToolMenuInsertType::After);
 
 	{
-		ToolMenu->AddDynamicSection("Persona", FNewToolBarDelegateLegacy::CreateLambda([](FToolBarBuilder& ToolbarBuilder, UToolMenu* InMenu)
+		ToolMenu->AddDynamicSection("Persona", FNewToolMenuDelegate::CreateLambda([](UToolMenu* InToolMenu)
 		{
-			TSharedPtr<FSkeletalMeshEditor> SkeletalMeshEditor = GetSkeletalMeshEditor(InMenu->Context);
+			TSharedPtr<FSkeletalMeshEditor> SkeletalMeshEditor = GetSkeletalMeshEditor(InToolMenu->Context);
 			if (SkeletalMeshEditor.IsValid() && SkeletalMeshEditor->PersonaToolkit.IsValid())
 			{
 				FPersonaModule& PersonaModule = FModuleManager::LoadModuleChecked<FPersonaModule>("Persona");
 				FPersonaModule::FCommonToolbarExtensionArgs Args;
 				Args.bPreviewMesh = false;
-				PersonaModule.AddCommonToolbarExtensions(ToolbarBuilder, SkeletalMeshEditor->PersonaToolkit.ToSharedRef(), Args);
+				PersonaModule.AddCommonToolbarExtensions(InToolMenu, Args);
 			}
 		}), SectionInsertLocation);
 	}
diff --git a/Engine/Source/Editor/SkeletonEditor/Private/SkeletonEditor.cpp b/Engine/Source/Editor/SkeletonEditor/Private/SkeletonEditor.cpp
--- a/Engine/Source/Editor/SkeletonEditor/Private/SkeletonEditor.cpp
+++ b/Engine/Source/Editor/SkeletonEditor/Private/SkeletonEditor.cpp
@@ -214,13 +214,13 @@ void FSkeletonEditor::ExtendToolbar()
 	const FToolMenuInsert SectionInsertLocation("Asset", EToolMenuInsertType::After);
 
 	{
-		ToolMenu->AddDynamicSection("Persona", FNewToolBarDelegateLegacy::CreateLambda([](FToolBarBuilder& ToolbarBuilder, UToolMenu* InMenu)
+		ToolMenu->AddDynamicSection("Persona", FNewToolMenuDelegate::CreateLambda([](UToolMenu* InToolMenu)
 		{
-			TSharedPtr<FSkeletonEditor> SkeletonEditor = GetSkeletonEditor(InMenu->Context);
+			TSharedPtr<FSkeletonEditor> SkeletonEditor = GetSkeletonEditor(InToolMenu->Context);
 			if (SkeletonEditor.IsValid() && SkeletonEditor->PersonaToolkit.IsValid())
 			{
 				FPersonaModule& PersonaModule = FModuleManager::LoadModuleChecked<FPersonaModule>("Persona");
-				PersonaModule.AddCommonToolbarExtensions(ToolbarBuilder, SkeletonEditor->PersonaToolkit.ToSharedRef());
+				PersonaModule.AddCommonToolbarExtensions(InToolMenu);
 			}
 		}), SectionInsertLocation);
 	}
diff --git a/Engine/Source/Editor/UnrealEd/Classes/Settings/AnimBlueprintSettings.h b/Engine/Source/Editor/UnrealEd/Classes/Settings/AnimBlueprintSettings.h
--- a/Engine/Source/Editor/UnrealEd/Classes/Settings/AnimBlueprintSettings.h
+++ b/Engine/Source/Editor/UnrealEd/Classes/Settings/AnimBlueprintSettings.h
@@ -26,4 +26,12 @@ public:
 	/** Whether to allow delegates to be created/displayed in animation blueprints */
 	UPROPERTY()
 	bool bAllowDelegates = true;
+
+	/** Whether to allow restrict which base function overrides can created/displayed in animation blueprints */
+	UPROPERTY()
+	bool bRestrictBaseFunctionOverrides = false;
+
+	/** The set of allowed base functions if restricted */
+	UPROPERTY()
+	TArray<FName> BaseFunctionOverrideAllowList;
 };
\ No newline at end of file
diff --git a/Engine/Source/Runtime/Engine/Classes/Animation/AnimBlueprint.h b/Engine/Source/Runtime/Engine/Classes/Animation/AnimBlueprint.h
--- a/Engine/Source/Runtime/Engine/Classes/Animation/AnimBlueprint.h
+++ b/Engine/Source/Runtime/Engine/Classes/Animation/AnimBlueprint.h
@@ -180,6 +180,7 @@ class ENGINE_API UAnimBlueprint : public UBlueprint, public IInterface_PreviewMe
 	virtual bool SupportsEventGraphs() const override;
 	virtual bool SupportsDelegates() const override;
 	virtual bool SupportsMacros() const override;
+	virtual bool AllowFunctionOverride(const UFunction* const InFunction) const override;
 
 protected:
 	// Broadcast when an override is changed, allowing derived blueprints to be updated
diff --git a/Engine/Source/Runtime/Engine/Classes/Engine/Blueprint.h b/Engine/Source/Runtime/Engine/Classes/Engine/Blueprint.h
--- a/Engine/Source/Runtime/Engine/Classes/Engine/Blueprint.h
+++ b/Engine/Source/Runtime/Engine/Classes/Engine/Blueprint.h
@@ -1246,6 +1246,11 @@ public:
 	 */
 	virtual bool SupportsFunctions() const { return true; }
 
+	/**
+	 * Returns true if this blueprints allows the given function to be overridden
+	 */
+	virtual bool AllowFunctionOverride(const UFunction* const InFunction) const { return true; }
+		
 	/**
 	 * Returns true if this blueprint supports macros
 	 */
diff --git a/Engine/Source/Runtime/Engine/Private/AnimBlueprint.cpp b/Engine/Source/Runtime/Engine/Private/AnimBlueprint.cpp
--- a/Engine/Source/Runtime/Engine/Private/AnimBlueprint.cpp
+++ b/Engine/Source/Runtime/Engine/Private/AnimBlueprint.cpp
@@ -240,6 +240,26 @@ bool UAnimBlueprint::SupportsMacros() const
 	return GetDefault<UAnimBlueprintSettings>()->bAllowMacros;
 }
 
+bool UAnimBlueprint::AllowFunctionOverride(const UFunction* const InFunction) const
+{
+	check(InFunction);
+
+	if (!GetDefault<UAnimBlueprintSettings>()->bRestrictBaseFunctionOverrides)
+	{
+		return true;
+	}
+
+	UClass* OwnerClass = InFunction->GetOwnerClass();
+	if (OwnerClass == UAnimInstance::StaticClass())
+	{
+		return GetDefault<UAnimBlueprintSettings>()->BaseFunctionOverrideAllowList.Contains(InFunction->GetFName());
+	}
+	else
+	{
+		return true;
+	}
+}
+
 #endif
 
 USkeletalMesh* UAnimBlueprint::GetPreviewMesh(bool bFindIfNotSet/*=false*/)
