Interchange:
- Skeletal mesh is now import using the skeletalmesh compile system to make sure all properties are locked and will stall another thread trying to access them.
---- This fix multiple crash if a viewer display the skeletal mesh and use some properties.
- Remove animation pipeline ensure(no skeleton) if we import with force static mesh and there is some translated animation node
- Fix crash when importing static mesh and the get mesh payload return an empty mesh description
- Reimport of skeletalmesh is not calling anymore postedit change if we reimport with interchange, interchange will call it in a post edit step.
---- This fix multiple crash if a viewer display the bones of the re-import skeletal mesh.
General:
- Reimport handler now have a reimport async version so the caller can control if a reimport is done asynchronously or synchronously

#rnx

diff --git a/Engine/Plugins/Interchange/Runtime/Source/Import/Private/Mesh/InterchangeSkeletalMeshFactory.cpp b/Engine/Plugins/Interchange/Runtime/Source/Import/Private/Mesh/InterchangeSkeletalMeshFactory.cpp
--- a/Engine/Plugins/Interchange/Runtime/Source/Import/Private/Mesh/InterchangeSkeletalMeshFactory.cpp
+++ b/Engine/Plugins/Interchange/Runtime/Source/Import/Private/Mesh/InterchangeSkeletalMeshFactory.cpp
@@ -7,6 +7,7 @@
 #include "Components.h"
 #include "CoreGlobals.h"
 #include "Engine/SkeletalMesh.h"
+#include "Engine/SkinnedAssetAsyncCompileUtils.h"
 #include "GenericPlatform/GenericPlatformMisc.h"
 #include "GPUSkinPublicDefs.h"
 #include "InterchangeAssetImportData.h"
@@ -915,7 +916,14 @@ UObject* UInterchangeSkeletalMeshFactory::CreateEmptyAsset(const FCreateAssetPar
 	SkeletalMesh->PreEditChange(nullptr);
 	//Allocate the LODImport data in the main thread
 	SkeletalMesh->ReserveLODImportData(SkeletalMeshFactoryNode->GetLodDataCount());
-	
+
+	//Lock the skeletalmesh properties if the skeletal mesh already exist (re-import)
+	if (ExistingAsset)
+	{
+		SkeletalMeshLockPropertiesEvent = FPlatformProcess::GetSynchEventFromPool();
+		SkeletalMesh->LockPropertiesUntil(SkeletalMeshLockPropertiesEvent);
+	}
+
 	return SkeletalMesh;
 #endif //else !WITH_EDITOR || !WITH_EDITORONLY_DATA
 }
@@ -992,6 +1000,9 @@ UObject* UInterchangeSkeletalMeshFactory::CreateAsset(const FCreateAssetParams&
 		return nullptr;
 	}
 
+	//Make sure we can modify the skeletalmesh properties
+	FSkinnedAssetAsyncBuildScope AsyncBuildScope(SkeletalMesh);
+
 	//This is consider has a re-import if we have a reimport object or if the object exist and have some valid LOD
 	const bool bIsReImport = (Arguments.ReimportObject != nullptr) || (SkeletalMesh->GetLODNum() > 0);
 
@@ -1451,6 +1462,9 @@ UObject* UInterchangeSkeletalMeshFactory::CreateAsset(const FCreateAssetParams&
 					, NSLOCTEXT("InterchangeSkeletalMeshFactory", "SkeletonFailed_BoneMerge", "FAILED TO MERGE BONES:\n\n This could happen if significant hierarchical changes have been made\ne.g. inserting a bone between nodes.\nWould you like to regenerate the Skeleton from this mesh?\n\n***WARNING: THIS MAY INVALIDATE OR REQUIRE RECOMPRESSION OF ANIMATION DATA.***\n"));
 				if (MergeBonesChoice == EAppReturnType::Yes)
 				{
+					//Allow this thread scope to read and write skeletalmesh locked properties
+					FSkinnedAssetAsyncBuildScope AsyncBuildScope(SkeletalMeshPtr);
+
 					if (SkeletonPtr->RecreateBoneTree(SkeletalMeshPtr))
 					{
 						TArray<const USkeletalMesh*> OtherSkeletalMeshUsingSkeleton;
@@ -1572,6 +1586,16 @@ UObject* UInterchangeSkeletalMeshFactory::CreateAsset(const FCreateAssetParams&
 #endif //else !WITH_EDITOR || !WITH_EDITORONLY_DATA
 }
 
+void UInterchangeSkeletalMeshFactory::Cancel()
+{
+	if (SkeletalMeshLockPropertiesEvent)
+	{
+		SkeletalMeshLockPropertiesEvent->Trigger();
+		FPlatformProcess::ReturnSynchEventToPool(SkeletalMeshLockPropertiesEvent);
+		SkeletalMeshLockPropertiesEvent = nullptr;
+	}
+}
+
 /* This function is call in the completion task on the main thread, use it to call main thread post creation step for your assets*/
 void UInterchangeSkeletalMeshFactory::PreImportPreCompletedCallback(const FImportPreCompletedCallbackParams& Arguments)
 {
@@ -1585,6 +1609,13 @@ void UInterchangeSkeletalMeshFactory::PreImportPreCompletedCallback(const FImpor
 	{
 		//We must call the Update of the asset source file in the main thread because UAssetImportData::Update execute some delegate we do not control
 		USkeletalMesh* SkeletalMesh = CastChecked<USkeletalMesh>(Arguments.ImportedObject);
+		//Release the promise so we unlock the skeletalmesh properties
+		if (SkeletalMeshLockPropertiesEvent)
+		{
+			SkeletalMeshLockPropertiesEvent->Trigger();
+			FPlatformProcess::ReturnSynchEventToPool(SkeletalMeshLockPropertiesEvent);
+			SkeletalMeshLockPropertiesEvent = nullptr;
+		}
 
 		UAssetImportData* ImportDataPtr = SkeletalMesh->GetAssetImportData();
 		UE::Interchange::FFactoryCommon::FUpdateImportAssetDataParameters UpdateImportAssetDataParameters(SkeletalMesh
diff --git a/Engine/Plugins/Interchange/Runtime/Source/Import/Private/Mesh/InterchangeStaticMeshFactory.cpp b/Engine/Plugins/Interchange/Runtime/Source/Import/Private/Mesh/InterchangeStaticMeshFactory.cpp
--- a/Engine/Plugins/Interchange/Runtime/Source/Import/Private/Mesh/InterchangeStaticMeshFactory.cpp
+++ b/Engine/Plugins/Interchange/Runtime/Source/Import/Private/Mesh/InterchangeStaticMeshFactory.cpp
@@ -267,6 +267,10 @@ UObject* UInterchangeStaticMeshFactory::CreateAsset(const FCreateAssetParams& Ar
 			if (!bFirstValidMoved)
 			{
 				FMeshDescription& MeshDescription = const_cast<FMeshDescription&>(LodMeshPayload->MeshDescription);
+				if (MeshDescription.IsEmpty())
+				{
+					continue;
+				}
 				LodMeshDescription = MoveTemp(MeshDescription);
 				bFirstValidMoved = true;
 
@@ -278,6 +282,10 @@ UObject* UInterchangeStaticMeshFactory::CreateAsset(const FCreateAssetParams& Ar
 			}
 			else
 			{
+				if (LodMeshPayload->MeshDescription.IsEmpty())
+				{
+					continue;
+				}
 				// Bake the payload mesh, with the provided transform
 				AppendSettings.MeshTransform = MeshPayload.Transform;
 				FStaticMeshOperations::AppendMeshDescription(LodMeshPayload->MeshDescription, LodMeshDescription, AppendSettings);
diff --git a/Engine/Plugins/Interchange/Runtime/Source/Import/Public/Mesh/InterchangeSkeletalMeshFactory.h b/Engine/Plugins/Interchange/Runtime/Source/Import/Public/Mesh/InterchangeSkeletalMeshFactory.h
--- a/Engine/Plugins/Interchange/Runtime/Source/Import/Public/Mesh/InterchangeSkeletalMeshFactory.h
+++ b/Engine/Plugins/Interchange/Runtime/Source/Import/Public/Mesh/InterchangeSkeletalMeshFactory.h
@@ -24,6 +24,7 @@ public:
 	virtual EInterchangeFactoryAssetType GetFactoryAssetType() override { return EInterchangeFactoryAssetType::Meshes; }
 	virtual UObject* CreateEmptyAsset(const FCreateAssetParams& Arguments) override;
 	virtual UObject* CreateAsset(const FCreateAssetParams& Arguments) override;
+	virtual void Cancel() override;
 	virtual void PreImportPreCompletedCallback(const FImportPreCompletedCallbackParams& Arguments) override;
 	//virtual void PostImportPreCompletedCallback(const FImportPreCompletedCallbackParams& Arguments) const override;
 	virtual bool GetSourceFilenames(const UObject* Object, TArray<FString>& OutSourceFilenames) const override;
@@ -32,6 +33,8 @@ public:
 
 	// Interchange factory base interface end
 	//////////////////////////////////////////////////////////////////////////
+private:
+	FEvent* SkeletalMeshLockPropertiesEvent = nullptr;
 };
 
 
diff --git a/Engine/Plugins/Interchange/Runtime/Source/Pipelines/Private/InterchangeGenericAnimationPipeline.cpp b/Engine/Plugins/Interchange/Runtime/Source/Pipelines/Private/InterchangeGenericAnimationPipeline.cpp
--- a/Engine/Plugins/Interchange/Runtime/Source/Pipelines/Private/InterchangeGenericAnimationPipeline.cpp
+++ b/Engine/Plugins/Interchange/Runtime/Source/Pipelines/Private/InterchangeGenericAnimationPipeline.cpp
@@ -638,9 +638,9 @@ void UInterchangeGenericAnimationPipeline::CreateAnimSequenceFactoryNode(UInterc
 
 	const FString SkeletonFactoryNodeUid = UInterchangeFactoryBaseNode::BuildFactoryNodeUid(SkeletonNodeUid);
 	const UInterchangeSkeletonFactoryNode* SkeletonFactoryNode = Cast<const UInterchangeSkeletonFactoryNode>(BaseNodeContainer->GetFactoryNode(SkeletonFactoryNodeUid));
-	if (!ensure(SkeletonFactoryNode))
+	if (!SkeletonFactoryNode)
 	{
-		// TODO: Warn something wrong happened
+		// It can happen if we force static mesh import, in that case no skeleton will be create
 		return;
 	}
 
diff --git a/Engine/Source/Editor/AnimationEditor/Private/AnimationEditor.cpp b/Engine/Source/Editor/AnimationEditor/Private/AnimationEditor.cpp
--- a/Engine/Source/Editor/AnimationEditor/Private/AnimationEditor.cpp
+++ b/Engine/Source/Editor/AnimationEditor/Private/AnimationEditor.cpp
@@ -212,6 +212,11 @@ void FAnimationEditor::InitToolMenuContext(FToolMenuContext& MenuContext)
 
 void FAnimationEditor::Tick(float DeltaTime)
 {
+	//Do not tick the animation editor if we are compiling the skeletalmesh we edit
+	if (GetPersonaToolkit()->GetMesh() && GetPersonaToolkit()->GetMesh()->IsCompiling())
+	{
+		return;
+	}
 	GetPersonaToolkit()->GetPreviewScene()->InvalidateViews();
 }
 
@@ -639,7 +644,7 @@ void FAnimationEditor::OnReimportAnimation()
 	UAnimSequence* AnimSequence = Cast<UAnimSequence>(AnimationAsset);
 	if (AnimSequence)
 	{
-		FReimportManager::Instance()->Reimport(AnimSequence, true);
+		FReimportManager::Instance()->ReimportAsync(AnimSequence, true);
 	}
 }
 
diff --git a/Engine/Source/Editor/Persona/Private/AnimationEditorViewportClient.cpp b/Engine/Source/Editor/Persona/Private/AnimationEditorViewportClient.cpp
--- a/Engine/Source/Editor/Persona/Private/AnimationEditorViewportClient.cpp
+++ b/Engine/Source/Editor/Persona/Private/AnimationEditorViewportClient.cpp
@@ -514,7 +514,7 @@ void FAnimationViewportClient::Draw(const FSceneView* View, FPrimitiveDrawInterf
 	TArray<UDebugSkelMeshComponent*> PreviewMeshComponents = GetPreviewScene()->GetAllPreviewMeshComponents();
 	for (UDebugSkelMeshComponent* PreviewMeshComponent : PreviewMeshComponents)
 	{
-		if (!(PreviewMeshComponent && PreviewMeshComponent->GetSkeletalMeshAsset()))
+		if (!(PreviewMeshComponent && PreviewMeshComponent->GetSkeletalMeshAsset() && !PreviewMeshComponent->GetSkeletalMeshAsset()->IsCompiling()))
 		{
 			continue;
 		}
@@ -522,6 +522,17 @@ void FAnimationViewportClient::Draw(const FSceneView* View, FPrimitiveDrawInterf
 		// Can't have both bones of interest and sockets of interest set
 		check( !(GetAnimPreviewScene()->GetSelectedBoneIndex() != INDEX_NONE && GetAnimPreviewScene()->GetSelectedSocket().IsValid() ) )
 
+		const FReferenceSkeleton& RefSkeleton = PreviewMeshComponent->GetReferenceSkeleton();
+		const TArray<FBoneIndexType>& DrawBoneIndices = PreviewMeshComponent->GetDrawBoneIndices();
+		//If DrawBoneIndices count is different from the refskeleton, it mean the component must be re-register
+		//It can happen during an asynchronous re-import of a skeletal mesh where the properties lock is release
+		//before the postedit change scope have re-register the components
+		if (RefSkeleton.GetNum() != DrawBoneIndices.Num())
+		{
+			//Skip this draw call until the component cache fit the skeletalmesh data
+			continue;
+		}
+
 		// if we have BonesOfInterest, draw sub set of the bones only
 		if (GetAnimPreviewScene()->GetSelectedBoneIndex() != INDEX_NONE)
 		{
@@ -643,7 +654,7 @@ void FAnimationViewportClient::DrawCanvas( FViewport& InViewport, FSceneView& Vi
 
 void FAnimationViewportClient::DrawUVsForMesh(FViewport* InViewport, FCanvas* InCanvas, int32 InTextYPos, UDebugSkelMeshComponent* PreviewMeshComponent)
 {
-	if (!PreviewMeshComponent)
+	if (!PreviewMeshComponent || !PreviewMeshComponent->GetSkeletalMeshAsset() || PreviewMeshComponent->GetSkeletalMeshAsset()->IsCompiling())
 	{
 		return;
 	}
@@ -658,6 +669,14 @@ void FAnimationViewportClient::DrawUVsForMesh(FViewport* InViewport, FCanvas* In
 
 void FAnimationViewportClient::Tick(float DeltaSeconds) 
 {
+	//Avoid ticking the animation viewport if the skeletalmesh is compiling
+	if (GetAnimPreviewScene()->GetPreviewMeshComponent()
+		&& GetAnimPreviewScene()->GetPreviewMeshComponent()->GetSkeletalMeshAsset()
+		&& GetAnimPreviewScene()->GetPreviewMeshComponent()->GetSkeletalMeshAsset()->IsCompiling())
+	{
+		return;
+	}
+
 	FEditorViewportClient::Tick(DeltaSeconds);
 
 	GetAnimPreviewScene()->FlagTickable();
@@ -910,7 +929,7 @@ FText FAnimationViewportClient::GetDisplayInfo(bool bDisplayAllInfo) const
 
 	// if not valid skeletalmesh
 	UDebugSkelMeshComponent* PreviewMeshComponent = GetPreviewScene()->GetPreviewMeshComponent();
-	if (PreviewMeshComponent == nullptr || PreviewMeshComponent->GetSkeletalMeshAsset() == nullptr)
+	if (PreviewMeshComponent == nullptr || PreviewMeshComponent->GetSkeletalMeshAsset() == nullptr || PreviewMeshComponent->GetSkeletalMeshAsset()->IsCompiling())
 	{
 		return FText();
 	}
diff --git a/Engine/Source/Editor/Persona/Private/EditModes/SkeletonSelectionEditMode.cpp b/Engine/Source/Editor/Persona/Private/EditModes/SkeletonSelectionEditMode.cpp
--- a/Engine/Source/Editor/Persona/Private/EditModes/SkeletonSelectionEditMode.cpp
+++ b/Engine/Source/Editor/Persona/Private/EditModes/SkeletonSelectionEditMode.cpp
@@ -339,7 +339,7 @@ FIntPoint FSkeletonSelectionEditMode::GetDPIUnscaledSize(FViewport* Viewport, FV
 void FSkeletonSelectionEditMode::DrawHUD(FEditorViewportClient* ViewportClient, FViewport* Viewport, const FSceneView* View, FCanvas* Canvas)
 {
 	UDebugSkelMeshComponent* PreviewMeshComponent = GetAnimPreviewScene().GetPreviewMeshComponent();
-	if (PreviewMeshComponent == nullptr || PreviewMeshComponent->GetSkeletalMeshAsset() == nullptr)
+	if (PreviewMeshComponent == nullptr || PreviewMeshComponent->GetSkeletalMeshAsset() == nullptr || PreviewMeshComponent->GetSkeletalMeshAsset()->IsCompiling())
 	{
 		return;
 	}
@@ -348,7 +348,7 @@ void FSkeletonSelectionEditMode::DrawHUD(FEditorViewportClient* ViewportClient,
 	int32 BoneIndex = GetAnimPreviewScene().GetSelectedBoneIndex();
 
 	// Draw name of selected bone
-	if (IsSelectedBoneRequired() || RefSkeleton.GetRequiredVirtualBones().Contains(BoneIndex))
+	if (RefSkeleton.IsValidIndex(BoneIndex) && (IsSelectedBoneRequired() || RefSkeleton.GetRequiredVirtualBones().Contains(BoneIndex)))
 	{
 		const FIntPoint ViewPortSize = GetDPIUnscaledSize(Viewport, ViewportClient);
 		const int32 HalfX = ViewPortSize.X / 2;
@@ -446,7 +446,7 @@ bool FSkeletonSelectionEditMode::UsesTransformWidget(UE::Widget::EWidgetMode Che
 bool FSkeletonSelectionEditMode::GetCustomDrawingCoordinateSystem(FMatrix& InMatrix, void* InData)
 {
 	UDebugSkelMeshComponent* PreviewMeshComponent = GetAnimPreviewScene().GetPreviewMeshComponent();
-	if(PreviewMeshComponent && PreviewMeshComponent->GetSkeletalMeshAsset())
+	if(PreviewMeshComponent && PreviewMeshComponent->GetSkeletalMeshAsset() && !PreviewMeshComponent->GetSkeletalMeshAsset()->IsCompiling())
 	{
 		int32 BoneIndex = GetAnimPreviewScene().GetSelectedBoneIndex();
 		if (BoneIndex != INDEX_NONE && PreviewMeshComponent->GetSkeletalMeshAsset()->GetRefSkeleton().IsValidIndex(BoneIndex))
@@ -483,13 +483,17 @@ bool FSkeletonSelectionEditMode::GetCustomInputCoordinateSystem(FMatrix& InMatri
 FVector FSkeletonSelectionEditMode::GetWidgetLocation() const
 {
 	UDebugSkelMeshComponent* PreviewMeshComponent = GetAnimPreviewScene().GetPreviewMeshComponent();
+	USkeletalMesh* SkeletalMesh = PreviewMeshComponent ? PreviewMeshComponent->GetSkeletalMeshAsset() : nullptr;
+	const bool bAssetIsCompiling = SkeletalMesh ? SkeletalMesh->IsCompiling() : false;
 
 	int32 BoneIndex = GetAnimPreviewScene().GetSelectedBoneIndex();
 	if (BoneIndex != INDEX_NONE)
 	{
-		if (PreviewMeshComponent && PreviewMeshComponent->GetSkeletalMeshAsset() && PreviewMeshComponent->GetSkeletalMeshAsset()->GetRefSkeleton().IsValidIndex(BoneIndex))
+		if (!bAssetIsCompiling
+			&& SkeletalMesh
+			&& SkeletalMesh->GetRefSkeleton().IsValidIndex(BoneIndex))
 		{
-			const FName BoneName = PreviewMeshComponent->GetSkeletalMeshAsset()->GetRefSkeleton().GetBoneName(BoneIndex);
+			const FName BoneName = SkeletalMesh->GetRefSkeleton().GetBoneName(BoneIndex);
 
 			const FMatrix BoneMatrix = PreviewMeshComponent->GetBoneMatrix(BoneIndex);
 
@@ -498,12 +502,15 @@ FVector FSkeletonSelectionEditMode::GetWidgetLocation() const
 	}
 	else if (GetAnimPreviewScene().GetSelectedSocket().IsValid())
 	{
-		USkeletalMeshSocket* Socket = GetAnimPreviewScene().GetSelectedSocket().Socket;
+		if (!bAssetIsCompiling)
+		{
+			USkeletalMeshSocket* Socket = GetAnimPreviewScene().GetSelectedSocket().Socket;
 
-		FMatrix SocketMatrix;
-		Socket->GetSocketMatrix(SocketMatrix, PreviewMeshComponent);
+			FMatrix SocketMatrix;
+			Socket->GetSocketMatrix(SocketMatrix, PreviewMeshComponent);
 
-		return SocketMatrix.GetOrigin();
+			return SocketMatrix.GetOrigin();
+		}
 	}
 	else if (AActor* SelectedActor = GetAnimPreviewScene().GetSelectedActor())
 	{
diff --git a/Engine/Source/Editor/Persona/Private/PersonaMeshDetails.cpp b/Engine/Source/Editor/Persona/Private/PersonaMeshDetails.cpp
--- a/Engine/Source/Editor/Persona/Private/PersonaMeshDetails.cpp
+++ b/Engine/Source/Editor/Persona/Private/PersonaMeshDetails.cpp
@@ -210,7 +210,7 @@ private:
 		{
 			USkeletalMesh* SkelMesh = SharedToolkit->GetMesh();
 
-			if(!SkelMesh)
+			if(!SkelMesh || SkelMesh->IsCompiling())
 			{
 				return EActiveTimerReturnType::Continue;
 			}
@@ -236,6 +236,12 @@ private:
 
 		USkeletalMesh* SkelMesh = SharedToolkit->GetMesh();
 		check(SkelMesh);
+		
+		if (SkelMesh->IsCompiling())
+		{
+			return LOCTEXT("ReimportButton_NewFile_SkeletalMeshIsCompiling_ToolTip", "Cannot re-import, skeletal mesh is compiling...");
+		}
+		
 		if (!SkelMesh->IsValidLODIndex(LODIndex))
 		{
 			// Should be true for the button to exist except if we delete a LOD
@@ -266,7 +272,7 @@ private:
 		{
 			USkeletalMesh* SkeletalMesh = SharedToolkit->GetMesh();
 
-			if (!SkeletalMesh)
+			if (!SkeletalMesh || SkeletalMesh->IsCompiling())
 			{
 				return false;
 			}
@@ -2082,7 +2088,7 @@ bool FPersonaMeshDetails::OnCanCopySectionList(int32 LODIndex) const
 {
 	USkeletalMesh* Mesh = GetPersonaToolkit()->GetMesh();
 
-	if (Mesh != nullptr)
+	if (Mesh != nullptr && !Mesh->IsCompiling())
 	{
 		FSkeletalMeshModel* ImportedResource = Mesh->GetImportedModel();
 
@@ -2216,7 +2222,7 @@ bool FPersonaMeshDetails::OnCanCopySectionItem(int32 LODIndex, int32 SectionInde
 {
 	USkeletalMesh* Mesh = GetPersonaToolkit()->GetMesh();
 
-	if (Mesh != nullptr)
+	if (Mesh != nullptr && !Mesh->IsCompiling())
 	{
 		FSkeletalMeshModel* ImportedResource = Mesh->GetImportedModel();
 
@@ -2324,7 +2330,7 @@ bool FPersonaMeshDetails::OnCanCopyMaterialList() const
 {
 	USkeletalMesh* Mesh = GetPersonaToolkit()->GetMesh();
 
-	if (Mesh != nullptr)
+	if (Mesh != nullptr && !Mesh->IsCompiling())
 	{
 		return Mesh->GetMaterials().Num() > 0;
 	}
@@ -2403,7 +2409,7 @@ bool FPersonaMeshDetails::OnCanCopyMaterialItem(int32 CurrentSlot) const
 {
 	USkeletalMesh* Mesh = GetPersonaToolkit()->GetMesh();
 
-	if (Mesh != nullptr)
+	if (Mesh != nullptr && !Mesh->IsCompiling())
 	{
 		return Mesh->GetMaterials().IsValidIndex(CurrentSlot);
 	}
@@ -3038,7 +3044,7 @@ TOptional<int32> FPersonaMeshDetails::GetLodSliderMaxValue() const
 {
 	USkeletalMesh* SkelMesh = GetPersonaToolkit()->GetMesh();
 
-	if(SkelMesh)
+	if(SkelMesh && !SkelMesh->IsCompiling())
 	{
 		return SkelMesh->GetLODNum() + PersonaMeshDetailsConstants::LodSliderExtension;
 	}
@@ -3339,7 +3345,10 @@ void FPersonaMeshDetails::CustomizeLODSettingsCategories(IDetailLayoutBuilder& D
 				.ButtonStyle(FAppStyle::Get(), "HoverHintOnly")
 				.ToolTipText(LOCTEXT("QualityLevelMinLodToolTip", "Clear MinLOD conversion data"))
 				.ForegroundColor(FSlateColor::UseForeground())
-				.IsEnabled(TAttribute<bool>::CreateLambda([this]() { return GetMinLod().PerPlatform.Num() != 0 || GetMinLod().Default != 0; }))
+				.IsEnabled(TAttribute<bool>::CreateLambda([this]()
+					{
+						return GetMinLod().PerPlatform.Num() != 0 || GetMinLod().Default != 0;
+					}))
 				.Content()
 				[
 					SNew(SImage)
@@ -3463,7 +3472,7 @@ bool FPersonaMeshDetails::IsMinLodEnable() const
 bool FPersonaMeshDetails::IsLODInfoEditingEnabled(int32 LODIndex) const
 {
 	USkeletalMesh* SkelMesh = GetPersonaToolkit()->GetMesh();
-	if (SkelMesh)
+	if (SkelMesh && !SkelMesh->IsCompiling())
 	{
 		if (SkelMesh->GetLODSettings())
 		{
@@ -3582,6 +3591,10 @@ int32 FPersonaMeshDetails::GetNoRefStreamingLODBias(FName QualityLevel) const
 {
 	USkeletalMesh* SkelMesh = GetPersonaToolkit()->GetMesh();
 	check(SkelMesh);
+	if (SkelMesh->IsCompiling())
+	{
+		return 0;
+	}
 
 	int32 QLKey = QualityLevelProperty::FNameToQualityLevel(QualityLevel);
 	const int32* ValuePtr = (QualityLevel == NAME_None) ? nullptr : SkelMesh->GetNoRefStreamingLODBias().PerQuality.Find(QLKey);
@@ -3643,8 +3656,12 @@ TArray<FName> FPersonaMeshDetails::GetNoRefStreamingLODBiasOverrideNames() const
 {
 	USkeletalMesh* SkelMesh = GetPersonaToolkit()->GetMesh();
 	check(SkelMesh);
-
 	TArray<FName> OverrideNames;
+	if (SkelMesh->IsCompiling())
+	{
+		return OverrideNames;
+	}
+	
 	for (const TPair<int32, int32>& Pair : SkelMesh->GetNoRefStreamingLODBias().PerQuality)
 	{
 		OverrideNames.Add(QualityLevelProperty::QualityLevelToFName(Pair.Key));
@@ -3663,7 +3680,7 @@ FPerPlatformInt FPersonaMeshDetails::GetMinLod()
 	USkeletalMesh* SkelMesh = GetPersonaToolkit()->GetMesh();
 	check(SkelMesh);
 
-	return SkelMesh->GetMinLod();
+	return SkelMesh->IsCompiling() ? 0 : SkelMesh->GetMinLod();
 }
 
 int32 FPersonaMeshDetails::GetMinQualityLevelLod(FName QualityLevel) const
@@ -3671,6 +3688,11 @@ int32 FPersonaMeshDetails::GetMinQualityLevelLod(FName QualityLevel) const
 	USkeletalMesh* SkelMesh = GetPersonaToolkit()->GetMesh();
 	check(SkelMesh);
 
+	if (SkelMesh->IsCompiling())
+	{
+		return 0;
+	}
+
 	int32 QLKey = QualityLevelProperty::FNameToQualityLevel(QualityLevel);
 	const int32* ValuePtr = (QualityLevel == NAME_None) ? nullptr : SkelMesh->GetQualityLevelMinLod().PerQuality.Find(QLKey);
 	return (ValuePtr != nullptr) ? *ValuePtr : SkelMesh->GetQualityLevelMinLod().Default;
@@ -3766,8 +3788,12 @@ TArray<FName> FPersonaMeshDetails::GetMinQualityLevelLodOverrideNames() const
 {
 	USkeletalMesh* SkelMesh = GetPersonaToolkit()->GetMesh();
 	check(SkelMesh);
-
 	TArray<FName> OverrideNames;
+	if (SkelMesh->IsCompiling())
+	{
+		return OverrideNames;
+	}
+	
 	for (const TPair<int32, int32>& Pair : SkelMesh->GetQualityLevelMinLod().PerQuality)
 	{
 		OverrideNames.Add(QualityLevelProperty::QualityLevelToFName(Pair.Key));
@@ -4169,7 +4195,7 @@ bool FPersonaMeshDetails::IsApplyNeeded() const
 	USkeletalMesh* SkelMesh = GetPersonaToolkit()->GetMesh();
 	check(SkelMesh);
 
-	if (SkelMesh->GetLODNum() != LODCount)
+	if (!SkelMesh->IsCompiling() && SkelMesh->GetLODNum() != LODCount)
 	{
 		return true;
 	}
@@ -4190,7 +4216,7 @@ FText FPersonaMeshDetails::GetLODCountTooltip() const
 FText FPersonaMeshDetails::GetLODImportedText(int32 LODIndex) const
 {
 	USkeletalMesh* Mesh = GetPersonaToolkit()->GetMesh();
-	if (Mesh && Mesh->IsValidLODIndex(LODIndex))
+	if (Mesh && !Mesh->IsCompiling() && Mesh->IsValidLODIndex(LODIndex))
 	{
 		if (Mesh->GetLODInfo(LODIndex)->bHasBeenSimplified)
 		{
@@ -4204,7 +4230,7 @@ FText FPersonaMeshDetails::GetLODImportedText(int32 LODIndex) const
 FText FPersonaMeshDetails::GetMaterialSlotNameText(int32 MaterialIndex) const
 {
 	USkeletalMesh* Mesh = GetPersonaToolkit()->GetMesh();
-	if (Mesh && Mesh->GetMaterials().IsValidIndex(MaterialIndex))
+	if (Mesh && !Mesh->IsCompiling() && Mesh->GetMaterials().IsValidIndex(MaterialIndex))
 	{
 		return FText::FromName(Mesh->GetMaterials()[MaterialIndex].MaterialSlotName);
 	}
@@ -4264,7 +4290,7 @@ void FPersonaMeshDetails::CustomizeDetails( IDetailLayoutBuilder& DetailLayout )
 	PostProcessWarningRow
 	.Visibility(MakeAttributeLambda([this]()
 	{
-		if(SkeletalMeshPtr.IsValid())
+		if(SkeletalMeshPtr.IsValid() && !SkeletalMeshPtr->IsCompiling())
 		{
 			TSubclassOf<UAnimInstance> PostProcessAnimBlueprintClass = SkeletalMeshPtr->GetPostProcessAnimBlueprint();
 			if(PostProcessAnimBlueprintClass.Get())
@@ -4385,9 +4411,13 @@ void FPersonaMeshDetails::OnPostProcessBlueprintChanged(IDetailLayoutBuilder* De
 FString FPersonaMeshDetails::GetCurrentPostProcessBlueprintPath() const
 {
 	USkeletalMesh* SkelMesh = GetPersonaToolkit()->GetMesh();
-	if(UClass* PostProcessClass = *SkelMesh->GetPostProcessAnimBlueprint())
+	if (!SkelMesh->IsCompiling())
 	{
-		return PostProcessClass->GetPathName();
+		if (UClass* PostProcessClass = *SkelMesh->GetPostProcessAnimBlueprint())
+		{
+
+			return PostProcessClass->GetPathName();
+		}
 	}
 
 	return FString();
@@ -4397,8 +4427,11 @@ bool FPersonaMeshDetails::OnShouldFilterPostProcessBlueprint(const FAssetData& A
 {
 	if(USkeletalMesh* SkelMesh = GetPersonaToolkit()->GetMesh())
 	{
-		const FString SkeletonName = AssetData.GetTagValueRef<FString>("TargetSkeleton");
-		return !SkelMesh->GetSkeleton()->IsCompatibleSkeletonByAssetString(SkeletonName);
+		if (!SkelMesh->IsCompiling())
+		{
+			const FString SkeletonName = AssetData.GetTagValueRef<FString>("TargetSkeleton");
+			return !SkelMesh->GetSkeleton()->IsCompatibleSkeletonByAssetString(SkeletonName);
+		}
 	}
 
 	return true;
@@ -4498,7 +4531,7 @@ void FPersonaMeshDetails::OnGetMaterialsForArray(class IMaterialListBuilder& Out
 {
 	USkeletalMesh* SkelMesh = GetPersonaToolkit()->GetMesh();
 
-	if (!SkelMesh)
+	if (!SkelMesh || SkelMesh->IsCompiling())
 		return;
 
 	TArray<FSkeletalMaterial>& SkelMeshMaterials = SkelMesh->GetMaterials();
@@ -4607,7 +4640,7 @@ FText FPersonaMeshDetails::GetMaterialArrayText() const
 {
 	FString MaterialArrayText = TEXT(" Material Slots");
 	int32 SlotNumber = 0;
-	if (SkeletalMeshPtr.IsValid())
+	if (SkeletalMeshPtr.IsValid() && !SkeletalMeshPtr->IsCompiling())
 	{
 		SlotNumber = SkeletalMeshPtr->GetMaterials().Num();
 	}
@@ -4618,6 +4651,10 @@ FText FPersonaMeshDetails::GetMaterialArrayText() const
 void FPersonaMeshDetails::OnGetSectionsForView(ISectionListBuilder& OutSections, int32 LODIndex)
 {
 	USkeletalMesh* SkelMesh = GetPersonaToolkit()->GetMesh();
+	if (!SkelMesh || SkelMesh->IsCompiling())
+	{
+		return;
+	}
 
 	FSkeletalMeshModel* ImportedResource = SkelMesh->GetImportedModel();
 
@@ -4669,7 +4706,7 @@ void FPersonaMeshDetails::OnGetSectionsForView(ISectionListBuilder& OutSections,
 
 FText FPersonaMeshDetails::GetMaterialNameText(int32 MaterialIndex) const
 {
-	if (SkeletalMeshPtr.IsValid() && SkeletalMeshPtr->GetMaterials().IsValidIndex(MaterialIndex))
+	if (SkeletalMeshPtr.IsValid() && !SkeletalMeshPtr->IsCompiling() && SkeletalMeshPtr->GetMaterials().IsValidIndex(MaterialIndex))
 	{
 		return FText::FromName(SkeletalMeshPtr->GetMaterials()[MaterialIndex].MaterialSlotName);
 	}
@@ -4678,7 +4715,7 @@ FText FPersonaMeshDetails::GetMaterialNameText(int32 MaterialIndex) const
 
 FText FPersonaMeshDetails::GetOriginalImportMaterialNameText(int32 MaterialIndex) const
 {
-	if (SkeletalMeshPtr.IsValid() && SkeletalMeshPtr->GetMaterials().IsValidIndex(MaterialIndex))
+	if (SkeletalMeshPtr.IsValid() && !SkeletalMeshPtr->IsCompiling() && SkeletalMeshPtr->GetMaterials().IsValidIndex(MaterialIndex))
 	{
 		FString OriginalImportMaterialName;
 		SkeletalMeshPtr->GetMaterials()[MaterialIndex].ImportedMaterialSlotName.ToString(OriginalImportMaterialName);
@@ -4759,7 +4796,7 @@ TSharedRef<SWidget> FPersonaMeshDetails::OnGenerateCustomMaterialWidgetsForMater
 
 FText FPersonaMeshDetails::GetFirstMaterialSlotUsedBySection(int32 MaterialIndex) const
 {
-	if (SkeletalMeshPtr.IsValid() && MaterialUsedMap.Contains(MaterialIndex))
+	if (SkeletalMeshPtr.IsValid() && !SkeletalMeshPtr->IsCompiling() && MaterialUsedMap.Contains(MaterialIndex))
 	{
 		const TArray<FSectionLocalizer> *SectionLocalizers = MaterialUsedMap.Find(MaterialIndex);
 		if (SectionLocalizers->Num() > 0)
@@ -4795,7 +4832,7 @@ TSharedRef<SWidget> FPersonaMeshDetails::OnGetMaterialSlotUsedByMenuContent(int3
 
 bool FPersonaMeshDetails::CanDeleteMaterialSlot(int32 MaterialIndex) const
 {
-	if (!SkeletalMeshPtr.IsValid())
+	if (!SkeletalMeshPtr.IsValid() || SkeletalMeshPtr->IsCompiling())
 	{
 		return false;
 	}
@@ -4862,7 +4899,7 @@ bool FPersonaMeshDetails::OnMaterialListDirty()
 {
 	bool ForceMaterialListRefresh = false;
 	TMap<int32, TArray<FSectionLocalizer>> TempMaterialUsedMap;
-	if (SkeletalMeshPtr.IsValid())
+	if (SkeletalMeshPtr.IsValid() && !SkeletalMeshPtr->IsCompiling())
 	{
 		const int32 MaterialCount = SkeletalMeshPtr->GetMaterials().Num();
 		for (int32 MaterialIndex = 0; MaterialIndex < MaterialCount; ++MaterialIndex)
@@ -5159,7 +5196,7 @@ TSharedRef<SWidget> FPersonaMeshDetails::OnGenerateCustomSectionWidgetsForSectio
 
 bool FPersonaMeshDetails::IsSectionEnabled(int32 LodIndex, int32 SectionIndex) const
 {
-	if(SkeletalMeshPtr.IsValid())
+	if(SkeletalMeshPtr.IsValid() && !SkeletalMeshPtr->IsCompiling())
 	{
 		FSkeletalMeshModel* SourceModel = SkeletalMeshPtr->GetImportedModel();
 
@@ -5247,6 +5284,7 @@ void FPersonaMeshDetails::OnSectionEnabledChanged(int32 LodIndex, int32 SectionI
 TOptional<int8> FPersonaMeshDetails::GetSectionGenerateUpToValue(int32 LodIndex, int32 SectionIndex) const
 {
 	if (!SkeletalMeshPtr.IsValid() ||
+		SkeletalMeshPtr->IsCompiling() ||
 		!SkeletalMeshPtr->GetImportedModel()->LODModels.IsValidIndex(LodIndex) ||
 		!SkeletalMeshPtr->GetImportedModel()->LODModels[LodIndex].Sections.IsValidIndex(SectionIndex) )
 	{
@@ -5324,6 +5362,7 @@ void FPersonaMeshDetails::SetSectionGenerateUpToValueCommitted(int8 Value, EText
 EVisibility FPersonaMeshDetails::ShowSectionGenerateUpToSlider(int32 LodIndex, int32 SectionIndex) const
 {
 	if (!SkeletalMeshPtr.IsValid() ||
+		SkeletalMeshPtr->IsCompiling() ||
 		!SkeletalMeshPtr->GetImportedModel()->LODModels.IsValidIndex(LodIndex) ||
 		!SkeletalMeshPtr->GetImportedModel()->LODModels[LodIndex].Sections.IsValidIndex(SectionIndex))
 	{
@@ -5335,6 +5374,7 @@ EVisibility FPersonaMeshDetails::ShowSectionGenerateUpToSlider(int32 LodIndex, i
 ECheckBoxState FPersonaMeshDetails::IsGenerateUpToSectionEnabled(int32 LodIndex, int32 SectionIndex) const
 {
 	if (!SkeletalMeshPtr.IsValid() ||
+		SkeletalMeshPtr->IsCompiling() ||
 		!SkeletalMeshPtr->GetImportedModel()->LODModels.IsValidIndex(LodIndex) ||
 		!SkeletalMeshPtr->GetImportedModel()->LODModels[LodIndex].Sections.IsValidIndex(SectionIndex))
 	{
@@ -5452,7 +5492,7 @@ TSharedRef<SWidget> FPersonaMeshDetails::OnGenerateLodMenuForLodPicker()
 {
 	USkeletalMesh* SkelMesh = GetPersonaToolkit()->GetMesh();
 
-	if (SkelMesh == nullptr)
+	if (SkelMesh == nullptr || SkelMesh->IsCompiling())
 	{
 		return SNullWidget::NullWidget;
 	}
@@ -5634,7 +5674,7 @@ ECheckBoxState FPersonaMeshDetails::IsSectionShadowCastingEnabled(int32 LODIndex
 {
 	ECheckBoxState State = ECheckBoxState::Unchecked;
 	const USkeletalMesh* Mesh = GetPersonaToolkit()->GetMesh();
-	if (Mesh == nullptr)
+	if (Mesh == nullptr || Mesh->IsCompiling())
 		return State;
 	
 	check(Mesh->GetImportedModel());
@@ -5705,7 +5745,7 @@ ECheckBoxState FPersonaMeshDetails::IsSectionVisibleInRayTracingEnabled(int32 LO
 {
 	ECheckBoxState State = ECheckBoxState::Unchecked;
 	const USkeletalMesh* Mesh = GetPersonaToolkit()->GetMesh();
-	if (Mesh == nullptr)
+	if (Mesh == nullptr || Mesh->IsCompiling())
 		return State;
 
 	check(Mesh->GetImportedModel());
@@ -5783,7 +5823,7 @@ const FString RecomputeTangentsChannelNames[] = {
 FText FPersonaMeshDetails::GetCurrentRecomputeTangentsSetting(int32 LODIndex, int32 SectionIndex) const
 {
 	const USkeletalMesh* Mesh = GetPersonaToolkit()->GetMesh();
-	if (Mesh == nullptr)
+	if (Mesh == nullptr || Mesh->IsCompiling())
 		return FText::GetEmpty();
 
 	check(Mesh->GetImportedModel());
@@ -5902,7 +5942,7 @@ EVisibility FPersonaMeshDetails::GetOverrideUVDensityVisibililty() const
 ECheckBoxState FPersonaMeshDetails::IsUVDensityOverridden(int32 MaterialIndex) const
 {
 	USkeletalMesh* Mesh = SkeletalMeshPtr.Get();
-	if (!Mesh || !Mesh->GetMaterials().IsValidIndex(MaterialIndex))
+	if (!Mesh || Mesh->IsCompiling() || !Mesh->GetMaterials().IsValidIndex(MaterialIndex))
 	{
 		return ECheckBoxState::Undetermined;
 	}
@@ -5943,7 +5983,7 @@ EVisibility FPersonaMeshDetails::GetUVDensityVisibility(int32 MaterialIndex, int
 TOptional<float> FPersonaMeshDetails::GetUVDensityValue(int32 MaterialIndex, int32 UVChannelIndex) const
 {
 	USkeletalMesh* Mesh = SkeletalMeshPtr.Get();
-	if (Mesh && Mesh->GetMaterials().IsValidIndex(MaterialIndex))
+	if (Mesh && !Mesh->IsCompiling() && Mesh->GetMaterials().IsValidIndex(MaterialIndex))
 	{
 		float Value = Mesh->GetMaterials()[MaterialIndex].UVChannelData.LocalUVDensities[UVChannelIndex];
 		return FMath::RoundToFloat(Value * 4.f) * .25f;
@@ -5966,6 +6006,10 @@ int32 FPersonaMeshDetails::GetMaterialIndex(int32 LODIndex, int32 SectionIndex)
 	USkeletalMesh* SkelMesh = GetPersonaToolkit()->GetMesh();
 
 	check(LODIndex < SkelMesh->GetLODNum());
+	if (SkelMesh->IsCompiling())
+	{
+		return 0;
+	}
 	
 	FSkeletalMeshModel* ImportedResource = SkelMesh->GetImportedModel();
 	check(ImportedResource && ImportedResource->LODModels.IsValidIndex(LODIndex));
@@ -6413,7 +6457,7 @@ FText FPersonaMeshDetails::OnGetClothingComboText(int32 InLodIdx, int32 InSectio
 {
 	USkeletalMesh* Mesh = GetPersonaToolkit()->GetMesh();
 
-	if(Mesh)
+	if(Mesh && !Mesh->IsCompiling())
 	{
 		UClothingAssetCommon* ClothingAsset = Cast<UClothingAssetCommon>(Mesh->GetSectionClothingAsset(InLodIdx, InSectionIdx));
 
diff --git a/Engine/Source/Editor/Persona/Private/SAnimationEditorViewport.cpp b/Engine/Source/Editor/Persona/Private/SAnimationEditorViewport.cpp
--- a/Engine/Source/Editor/Persona/Private/SAnimationEditorViewport.cpp
+++ b/Engine/Source/Editor/Persona/Private/SAnimationEditorViewport.cpp
@@ -2545,7 +2545,7 @@ void SAnimationEditorViewportTabBody::AddMinLODNotification()
 
 	auto GetMinLODNotificationVisibility = [this]()
 	{
-		if (GetPreviewScene()->GetPreviewMesh() && GetPreviewScene()->GetPreviewMesh()->GetDefaultMinLod() != 0)
+		if (GetPreviewScene()->GetPreviewMesh() && !GetPreviewScene()->GetPreviewMesh()->IsCompiling() && GetPreviewScene()->GetPreviewMesh()->GetDefaultMinLod() != 0)
 		{
 			return EVisibility::Visible;
 		}
diff --git a/Engine/Source/Editor/SkeletalMeshEditor/Private/SkeletalMeshEditor.cpp b/Engine/Source/Editor/SkeletalMeshEditor/Private/SkeletalMeshEditor.cpp
--- a/Engine/Source/Editor/SkeletalMeshEditor/Private/SkeletalMeshEditor.cpp
+++ b/Engine/Source/Editor/SkeletalMeshEditor/Private/SkeletalMeshEditor.cpp
@@ -10,6 +10,7 @@
 #include "Algo/Transform.h"
 #include "Animation/DebugSkelMeshComponent.h"
 #include "AssetRegistry/AssetData.h"
+#include "Async/Async.h"
 #include "ClothingAsset.h"
 #include "ClothingSystemEditorInterfaceModule.h"
 #include "ComponentReregisterContext.h"
@@ -62,6 +63,7 @@
 #include "SkeletalMeshEditorContextMenuContext.h"
 #include "Styling/AppStyle.h"
 #include "InterchangeAssetImportData.h"
+#include "InterchangeManager.h"
 
 const FName SkeletalMeshEditorAppIdentifier = FName(TEXT("SkeletalMeshEditorApp"));
 
@@ -1348,14 +1350,59 @@ UObject* FSkeletalMeshEditor::HandleGetAsset()
 
 bool FSkeletalMeshEditor::HandleReimportMeshInternal(int32 SourceFileIndex /*= INDEX_NONE*/, bool bWithNewFile /*= false*/)
 {
-	FScopedSuspendAlternateSkinWeightPreview ScopedSuspendAlternateSkinnWeightPreview(SkeletalMesh);
 	bool bResult = false;
+	//Interchange reimport are asynchronous
+	if (UInterchangeAssetImportData* AssetImportData = Cast<UInterchangeAssetImportData>(SkeletalMesh->GetAssetImportData()))
 	{
-		FScopedSkeletalMeshPostEditChange ScopedPostEditChange(SkeletalMesh);
-		// Reimport the asset
-		bResult = FReimportManager::Instance()->Reimport(SkeletalMesh, true, true, TEXT(""), nullptr, SourceFileIndex, bWithNewFile);
-		// Refresh skeleton tree
-		SkeletonTree->Refresh();
+		UE::Interchange::FAssetImportResultRef Result = FReimportManager::Instance()->ReimportAsync(SkeletalMesh, true, true, TEXT(""), nullptr, SourceFileIndex, bWithNewFile);
+		
+		Result->OnDone([SkeletonTreePtr = SkeletonTree, WeakSkeletalMesh = TWeakObjectPtr<USkeletalMesh>(SkeletalMesh)](UE::Interchange::FImportResult& Result)
+			{
+				auto ResetComponent = [SkeletonTreePtr, WeakSkeletalMesh]()
+				{
+					FScopedSuspendAlternateSkinWeightPreview ScopedSuspendAlternateSkinnWeightPreview(WeakSkeletalMesh.Get());
+					{
+						constexpr bool bCallPostEditChange = false;
+						constexpr bool bReregisterComponents = true;
+						FScopedSkeletalMeshPostEditChange ScopedPostEditChange = FScopedSkeletalMeshPostEditChange(WeakSkeletalMesh.Get(), bCallPostEditChange, bReregisterComponents);
+						// Refresh skeleton tree
+						SkeletonTreePtr->Refresh();
+					}
+				};
+
+				if (IsInGameThread())
+				{
+					ResetComponent();
+				}
+				else
+				{
+					Async(EAsyncExecution::TaskGraphMainThread, MoveTemp(ResetComponent));
+				}
+			});
+		
+		if (Result->GetStatus() == UE::Interchange::FImportResult::EStatus::Done)
+		{
+			const TArray<UInterchangeResult*>& Results = Result->GetResults()->GetResults();
+			for (const UInterchangeResult* InterchangeResult : Results)
+			{
+				if (InterchangeResult->IsA<UInterchangeResultError_ReimportFail>())
+				{
+					return false;
+				}
+			}
+		}
+		return true;
+	}
+	else
+	{
+		FScopedSuspendAlternateSkinWeightPreview ScopedSuspendAlternateSkinnWeightPreview(SkeletalMesh);
+		{
+			FScopedSkeletalMeshPostEditChange ScopedPostEditChange(SkeletalMesh);
+			// Reimport the asset
+			bResult = FReimportManager::Instance()->Reimport(SkeletalMesh, true, true, TEXT(""), nullptr, SourceFileIndex, bWithNewFile);
+			// Refresh skeleton tree
+			SkeletonTree->Refresh();
+		}
 	}
 	return bResult;
 }
diff --git a/Engine/Source/Editor/SkeletonEditor/Private/BoneProxy.cpp b/Engine/Source/Editor/SkeletonEditor/Private/BoneProxy.cpp
--- a/Engine/Source/Editor/SkeletonEditor/Private/BoneProxy.cpp
+++ b/Engine/Source/Editor/SkeletonEditor/Private/BoneProxy.cpp
@@ -27,6 +27,12 @@ void UBoneProxy::Tick(float DeltaTime)
 	{
 		if (UDebugSkelMeshComponent* Component = SkelMeshComponent.Get())
 		{
+			if (Component->GetSkeletalMeshAsset() && Component->GetSkeletalMeshAsset()->IsCompiling())
+			{
+				//We do not want to tick if the skeletalmesh is inside a compilation
+				return;
+			}
+
 			TArray<FTransform> LocalBoneTransforms = Component->GetBoneSpaceTransforms();
 
 			int32 BoneIndex = Component->GetBoneIndex(BoneName);
diff --git a/Engine/Source/Editor/StaticMeshEditor/Private/StaticMeshEditor.cpp b/Engine/Source/Editor/StaticMeshEditor/Private/StaticMeshEditor.cpp
--- a/Engine/Source/Editor/StaticMeshEditor/Private/StaticMeshEditor.cpp
+++ b/Engine/Source/Editor/StaticMeshEditor/Private/StaticMeshEditor.cpp
@@ -1656,7 +1656,7 @@ void FStaticMeshEditor::HandleReimportMesh()
 	// Reimport the asset
 	if (StaticMesh)
 	{
-		FReimportManager::Instance()->Reimport(StaticMesh, true);
+		FReimportManager::Instance()->ReimportAsync(StaticMesh, true);
 	}
 }
 
diff --git a/Engine/Source/Editor/TextureEditor/Private/TextureEditorToolkit.cpp b/Engine/Source/Editor/TextureEditor/Private/TextureEditorToolkit.cpp
--- a/Engine/Source/Editor/TextureEditor/Private/TextureEditorToolkit.cpp
+++ b/Engine/Source/Editor/TextureEditor/Private/TextureEditorToolkit.cpp
@@ -2382,7 +2382,7 @@ bool FTextureEditorToolkit::HandleReimportActionCanExecute( ) const
 
 void FTextureEditorToolkit::HandleReimportActionExecute( )
 {
-	FReimportManager::Instance()->Reimport(Texture, /*bAskForNewFileIfMissing=*/true);
+	FReimportManager::Instance()->ReimportAsync(Texture, /*bAskForNewFileIfMissing=*/true);
 }
 
 
diff --git a/Engine/Source/Editor/UnrealEd/Private/Animation/DebugSkelMeshComponent.cpp b/Engine/Source/Editor/UnrealEd/Private/Animation/DebugSkelMeshComponent.cpp
--- a/Engine/Source/Editor/UnrealEd/Private/Animation/DebugSkelMeshComponent.cpp
+++ b/Engine/Source/Editor/UnrealEd/Private/Animation/DebugSkelMeshComponent.cpp
@@ -97,7 +97,7 @@ FBoxSphereBounds UDebugSkelMeshComponent::CalcBounds(const FTransform& LocalToWo
 			}			
 		}
 
-		if (GetSkeletalMeshAsset())
+		if (GetSkeletalMeshAsset() && !GetSkeletalMeshAsset()->IsCompiling())
 		{
 			Result = Result + GetSkeletalMeshAsset()->GetBounds();
 		}
@@ -700,6 +700,11 @@ bool UDebugSkelMeshComponent::ShouldRunClothTick() const
 void UDebugSkelMeshComponent::SendRenderDynamicData_Concurrent()
 {
 	Super::SendRenderDynamicData_Concurrent();
+	
+	if (GetSkeletalMeshAsset() && GetSkeletalMeshAsset()->IsCompiling())
+	{
+		return;
+	}
 
 	if (SceneProxy)
 	{
@@ -1080,6 +1085,12 @@ void UDebugSkelMeshComponent::CheckClothTeleport()
 
 void UDebugSkelMeshComponent::TickComponent(float DeltaTime, enum ELevelTick TickType, FActorComponentTickFunction *ThisTickFunction)
 {
+	//Do not tick a skeletalmesh component if the skeletalmesh is compiling
+	if (GetSkeletalMeshAsset() && GetSkeletalMeshAsset()->IsCompiling())
+	{
+		return;
+	}
+
 	if (TurnTableMode == EPersonaTurnTableMode::Playing)
 	{
 		FRotator Rotation = GetRelativeTransform().Rotator();
diff --git a/Engine/Source/Editor/UnrealEd/Private/Editor.cpp b/Engine/Source/Editor/UnrealEd/Private/Editor.cpp
--- a/Engine/Source/Editor/UnrealEd/Private/Editor.cpp
+++ b/Engine/Source/Editor/UnrealEd/Private/Editor.cpp
@@ -69,6 +69,7 @@
 #include "AssetToolsModule.h"
 
 #include "InterchangeManager.h"
+#include "InterchangeResultsContainer.h"
 
 #if WITH_EDITOR
 #include "Subsystems/AssetEditorSubsystem.h"
@@ -262,8 +263,24 @@ void FReimportManager::UpdateReimportPath(UObject* Obj, const FString& Filename,
 }
 
 
-bool FReimportManager::Reimport( UObject* Obj, bool bAskForNewFileIfMissing, bool bShowNotification, FString PreferredReimportFile, FReimportHandler* SpecifiedReimportHandler, int32 SourceFileIndex, bool bForceNewFile /*= false*/, bool bAutomated /*= false*/)
+bool FReimportManager::Reimport(UObject* Obj, bool bAskForNewFileIfMissing, bool bShowNotification, FString PreferredReimportFile, FReimportHandler* SpecifiedReimportHandler, int32 SourceFileIndex, bool bForceNewFile /*= false*/, bool bAutomated /*= false*/)
 {
+	UE::Interchange::FAssetImportResultRef ImportResult = ReimportAsync(Obj, bAskForNewFileIfMissing, bShowNotification, PreferredReimportFile, SpecifiedReimportHandler, SourceFileIndex, bForceNewFile, bAutomated);
+	ImportResult->WaitUntilDone();
+	const TArray<UInterchangeResult*>& Results = ImportResult->GetResults()->GetResults();
+	for (const UInterchangeResult* InterchangeResult : Results)
+	{
+		if (InterchangeResult->IsA<UInterchangeResultError_ReimportFail>())
+		{
+			return false;
+		}
+	}
+	return true;
+}
+
+UE::Interchange::FAssetImportResultRef FReimportManager::ReimportAsync(UObject* Obj, bool bAskForNewFileIfMissing, bool bShowNotification, FString PreferredReimportFile, FReimportHandler* SpecifiedReimportHandler, int32 SourceFileIndex, bool bForceNewFile /*= false*/, bool bAutomated /*= false*/)
+{
+	UE::Interchange::FAssetImportResultRef ImportResultSynchronous = MakeShared< UE::Interchange::FImportResult, ESPMode::ThreadSafe >();
 	// Warn that were about to reimport, so prep for it
 	PreReimport.Broadcast( Obj );
 
@@ -422,9 +439,9 @@ bool FReimportManager::Reimport( UObject* Obj, bool bAskForNewFileIfMissing, boo
 						ImportAssetParameters.bIsAutomated = GIsAutomationTesting || FApp::IsUnattended() || IsRunningCommandlet() || GIsRunningUnattendedScript;
 						ImportAssetParameters.ReimportAsset = Obj;
 						ImportAssetParameters.ReimportSourceIndex = SourceFileIndex;
-						InterchangeManager.ImportAsset(FString(), ScopedSourceData.GetSourceData(), ImportAssetParameters);
+						UE::Interchange::FAssetImportResultRef ImportResult = InterchangeManager.ImportAssetAsync(FString(), ScopedSourceData.GetSourceData(), ImportAssetParameters);
 						InterchangeManager.OnAssetPostImport.Remove(PostImportHandle);
-						return true;
+						return ImportResult;
 					}
 				}
 
@@ -514,7 +531,13 @@ bool FReimportManager::Reimport( UObject* Obj, bool bAskForNewFileIfMissing, boo
 
 	GEditor->RedrawAllViewports();
 
-	return bSuccess;
+	if (!bSuccess)
+	{
+		//Add a ReimportFail message
+		ImportResultSynchronous->GetResults()->Add<UInterchangeResultError_ReimportFail>();
+	}
+	ImportResultSynchronous->SetDone();
+	return ImportResultSynchronous;
 }
 
 void FReimportManager::ValidateAllSourceFileAndReimport(TArray<UObject*> &ToImportObjects, bool bShowNotification, int32 SourceFileIndex, bool bForceNewFile /*= false*/, bool bAutomated /*= false*/)
@@ -681,8 +704,21 @@ bool FReimportManager::ReimportMultiple(TArrayView<UObject*> Objects, bool bAskF
 			FText SingleTaskTest = FText::Format(LOCTEXT("BulkReimport_SingleItem", "Reimporting {0}"), FText::FromString(CurrentObject->GetName()));
 			FScopedSlowTask SingleObjectTask(1.0f, SingleTaskTest);
 			SingleObjectTask.EnterProgressFrame(1.0f);
-
-			bBulkSuccess = bBulkSuccess && Reimport(CurrentObject, bAskForNewFileIfMissing, bShowNotification, PreferredReimportFile, SpecifiedReimportHandler, SourceFileIndex, bForceNewFile, bAutomated);
+			UE::Interchange::FAssetImportResultRef ImportResult = ReimportAsync(CurrentObject, bAskForNewFileIfMissing, bShowNotification, PreferredReimportFile, SpecifiedReimportHandler, SourceFileIndex, bForceNewFile, bAutomated);
+			const bool bAsync = ImportResult->GetStatus() == UE::Interchange::FImportResult::EStatus::InProgress;
+			bool bResultSuccess = true;
+			if (!bAsync)
+			{
+				const TArray<UInterchangeResult*>& Results = ImportResult->GetResults()->GetResults();
+				for (const UInterchangeResult* InterchangeResult : Results)
+				{
+					if (InterchangeResult->IsA<UInterchangeResultError_ReimportFail>())
+					{
+						bResultSuccess = false;
+					}
+				}
+			}
+			bBulkSuccess = bBulkSuccess && bResultSuccess;
 		}
 
 		BulkReimportTask.EnterProgressFrame(1.0f);
diff --git a/Engine/Source/Editor/UnrealEd/Private/Toolkits/AssetEditorToolkit.cpp b/Engine/Source/Editor/UnrealEd/Private/Toolkits/AssetEditorToolkit.cpp
--- a/Engine/Source/Editor/UnrealEd/Private/Toolkits/AssetEditorToolkit.cpp
+++ b/Engine/Source/Editor/UnrealEd/Private/Toolkits/AssetEditorToolkit.cpp
@@ -34,6 +34,7 @@
 #include "Misc/Attribute.h"
 #include "Textures/SlateIcon.h"
 #include "WidgetDrawerConfig.h"
+#include "Interfaces/Interface_AsyncCompilation.h"
 
 #define LOCTEXT_NAMESPACE "AssetEditorToolkit"
 
@@ -194,6 +195,20 @@ void FAssetEditorToolkit::InitAssetEditor( const EToolkitMode::Type Mode, const
 		NewMajorTab->SetContent
 		( 
 			SAssignNew( NewStandaloneHost, SStandaloneAssetEditorToolkitHost, NewTabManager, AppIdentifier )
+			.IsEnabled_Lambda([ObjectsToEdit]()
+				{
+					for (const UObject* Object : ObjectsToEdit)
+					{
+						if (const IInterface_AsyncCompilation* AsyncAsset = Cast<IInterface_AsyncCompilation>(Object))
+						{
+							if (AsyncAsset->IsCompiling())
+							{
+								return false;
+							}
+						}
+					}
+					return true;
+				})
 			.OnRequestClose(this, &FAssetEditorToolkit::OnRequestClose)
 			.OnClose(this, &FAssetEditorToolkit::OnClose)
 		);
diff --git a/Engine/Source/Editor/UnrealEd/Public/EditorReimportHandler.h b/Engine/Source/Editor/UnrealEd/Public/EditorReimportHandler.h
--- a/Engine/Source/Editor/UnrealEd/Public/EditorReimportHandler.h
+++ b/Engine/Source/Editor/UnrealEd/Public/EditorReimportHandler.h
@@ -9,6 +9,13 @@
 
 class FReimportHandler;
 
+namespace UE::Interchange
+{
+	class FImportResult;
+	using FAssetImportResultRef = TSharedRef< FImportResult, ESPMode::ThreadSafe >;
+}
+
+
 /** Reimport manager for package resources with associated source files on disk. */
 class FReimportManager : FGCObject
 {
@@ -58,6 +65,21 @@ public:
 	 */
 	UNREALED_API virtual bool Reimport( UObject* Obj, bool bAskForNewFileIfMissing = false, bool bShowNotification = true, FString PreferredReimportFile = TEXT(""), FReimportHandler* SpecifiedReimportHandler = nullptr, int32 SourceFileIndex = INDEX_NONE, bool bForceNewFile = false, bool bAutomated = false);
 
+	/**
+	 * Attempt to reimport the specified object from its source by giving registered reimport
+	 * handlers a chance to try to reimport the object
+	 *
+	 * @param	Obj	Object to try reimporting
+	 * @param	bAskForNewFileIfMissing If the file is missing, open a dialog to ask for a new one
+	 * @param	bShowNotification True to show a notification when complete, false otherwise
+	 * @param	PreferredReimportFile if not empty, will be use in case the original file is missing and bAskForNewFileIfMissing is set to false
+	 * @param	SourceFileIndex		which source file index you want to reimport default is INDEX_NONE(the factory will choose)
+	 * @param	bAutomated		True to skip dialog prompts
+	 *
+	 * @return	FAssetImportResultRef
+	 */
+	UNREALED_API virtual UE::Interchange::FAssetImportResultRef ReimportAsync(UObject* Obj, bool bAskForNewFileIfMissing = false, bool bShowNotification = true, FString PreferredReimportFile = TEXT(""), FReimportHandler* SpecifiedReimportHandler = nullptr, int32 SourceFileIndex = INDEX_NONE, bool bForceNewFile = false, bool bAutomated = false);
+
 	/**
 	 * Attemp to reimport all specified objects. This function will verify that all source file exist and ask the user
 	 * to decide on how the system will handle asset with missing source file path.
diff --git a/Engine/Source/Runtime/Engine/Classes/Engine/SkeletalMesh.h b/Engine/Source/Runtime/Engine/Classes/Engine/SkeletalMesh.h
--- a/Engine/Source/Runtime/Engine/Classes/Engine/SkeletalMesh.h
+++ b/Engine/Source/Runtime/Engine/Classes/Engine/SkeletalMesh.h
@@ -2407,6 +2407,13 @@ public:
 	*/
 	void Build();
 
+	/**
+	* Lock the skeletalmesh properties until the caller trigger the Event parameter.
+	* 
+	* @Param Event - When the caller will trigger the event the skeletal mesh properties will be unlocked
+	*/
+	void LockPropertiesUntil(FEvent* Event);
+
 	virtual void PostEditChangeProperty(FPropertyChangedEvent& PropertyChangedEvent) override;
 
 	virtual void PostEditUndo() override;
@@ -2702,7 +2709,24 @@ private:
 	 * Copy build/load context result data to the skeletalmesh member on the game thread - Can't be done in parallel.
 	 */
 	void ApplyFinishBuildInternalData(FSkinnedAssetCompilationContext* ContextPtr);
-	
+
+
+
+	/**
+	 * Initial step for the building process - Can't be done in parallel. USkinnedAsset Interface.
+	 */
+	virtual void BeginAsyncTaskInternal(FSkinnedAsyncTaskContext& Context) override;
+
+	/**
+	 * Thread-safe part. USkinnedAsset Interface.
+	 */
+	virtual void ExecuteAsyncTaskInternal(FSkinnedAsyncTaskContext& Context) override;
+
+	/**
+	 * Complete the building process - Can't be done in parallel. USkinnedAsset Interface.
+	 */
+	virtual void FinishAsyncTaskInternal(FSkinnedAsyncTaskContext& Context) override;
+
 #endif
 
 	/** Utility function to help with building the combined socket list */
diff --git a/Engine/Source/Runtime/Engine/Classes/Engine/SkinnedAsset.h b/Engine/Source/Runtime/Engine/Classes/Engine/SkinnedAsset.h
--- a/Engine/Source/Runtime/Engine/Classes/Engine/SkinnedAsset.h
+++ b/Engine/Source/Runtime/Engine/Classes/Engine/SkinnedAsset.h
@@ -293,6 +293,13 @@ protected:
 	/** Complete the building process - Can't be done in parallel. */
 	virtual void FinishBuildInternal(FSkinnedAssetBuildContext& Context) {}
 
+	/** Initial step for the async task process - Can't be done in parallel. */
+	virtual void BeginAsyncTaskInternal(FSkinnedAsyncTaskContext& Context) {}
+	/** Thread-safe part. */
+	virtual void ExecuteAsyncTaskInternal(FSkinnedAsyncTaskContext& Context) {}
+	/** Complete the async task process - Can't be done in parallel. */
+	virtual void FinishAsyncTaskInternal(FSkinnedAsyncTaskContext& Context) {}
+
 	/** Holds the pointer to an async task if one exists. */
 	TUniquePtr<FSkinnedAssetAsyncBuildTask> AsyncTask;
 #endif // WITH_EDITOR
diff --git a/Engine/Source/Runtime/Engine/Classes/Engine/SkinnedAssetAsyncCompileUtils.h b/Engine/Source/Runtime/Engine/Classes/Engine/SkinnedAssetAsyncCompileUtils.h
--- a/Engine/Source/Runtime/Engine/Classes/Engine/SkinnedAssetAsyncCompileUtils.h
+++ b/Engine/Source/Runtime/Engine/Classes/Engine/SkinnedAssetAsyncCompileUtils.h
@@ -37,6 +37,19 @@ public:
 	TUniquePtr<class FSkinnedMeshComponentRecreateRenderStateContext> RecreateRenderStateContext;
 };
 
+class FSkinnedAsyncTaskContext : public FSkinnedAssetCompilationContext
+{
+public:
+	FSkinnedAsyncTaskContext(bool bInResetAsyncFlagOnFinish, TFunctionRef<void()> InAsyncTaskFunction)
+		: bResetAsyncFlagOnFinish(bInResetAsyncFlagOnFinish)
+		, AsyncTaskFunction(InAsyncTaskFunction)
+	{
+	}
+
+	bool bResetAsyncFlagOnFinish = true;
+	TFunctionRef<void()> AsyncTaskFunction;
+};
+
 #if WITH_EDITOR
 
 // Any thread implicated in the build must have a valid scope to be granted access to protected properties without causing any stalls.
@@ -62,6 +75,7 @@ public:
 	USkinnedAsset* SkinnedAsset;
 	TOptional<FSkinnedAssetPostLoadContext> PostLoadContext;
 	TOptional<FSkinnedAssetBuildContext> BuildContext;
+	TOptional<FSkinnedAsyncTaskContext> AsyncTaskContext;
 
 	/** Initialization constructor. */
 	FSkinnedAssetAsyncBuildWorker(
@@ -81,6 +95,14 @@ public:
 	{
 	}
 
+	FSkinnedAssetAsyncBuildWorker(
+		USkinnedAsset* InSkinnedAsset,
+		FSkinnedAsyncTaskContext&& InAsyncTaskContext)
+		: SkinnedAsset(InSkinnedAsset)
+		, AsyncTaskContext(MoveTemp(InAsyncTaskContext))
+	{
+	}
+
 	FORCEINLINE TStatId GetStatId() const
 	{
 		RETURN_QUICK_DECLARE_CYCLE_STAT(FSkinnedAssetAsyncBuildWorker, STATGROUP_ThreadPoolAsyncTasks);
@@ -107,6 +129,14 @@ struct FSkinnedAssetAsyncBuildTask : public FAsyncTask<FSkinnedAssetAsyncBuildWo
 	{
 	}
 
+	FSkinnedAssetAsyncBuildTask(
+		USkinnedAsset* InSkinnedAsset,
+		FSkinnedAsyncTaskContext&& InAsyncTaskContext)
+		: FAsyncTask<FSkinnedAssetAsyncBuildWorker>(InSkinnedAsset, MoveTemp(InAsyncTaskContext))
+		, SkinnedAsset(InSkinnedAsset)
+	{
+	}
+
 	const USkinnedAsset* SkinnedAsset;
 };
 
diff --git a/Engine/Source/Runtime/Engine/Private/Components/SkeletalMeshComponent.cpp b/Engine/Source/Runtime/Engine/Private/Components/SkeletalMeshComponent.cpp
--- a/Engine/Source/Runtime/Engine/Private/Components/SkeletalMeshComponent.cpp
+++ b/Engine/Source/Runtime/Engine/Private/Components/SkeletalMeshComponent.cpp
@@ -1231,7 +1231,11 @@ void USkeletalMeshComponent::TickAnimation(float DeltaTime, bool bNeedsValidRoot
 		RecalcRequiredCurves();
 	}
 
-	if (GetSkeletalMeshAsset() != nullptr)
+	bool bIsCompiling = false;
+#if WITH_EDITOR
+	bIsCompiling = GetSkeletalMeshAsset() && !GetSkeletalMeshAsset()->IsCompiling();
+#endif
+	if (GetSkeletalMeshAsset() != nullptr && !bIsCompiling)
 	{
 		// We're about to UpdateAnimation, this will potentially queue events that we'll need to dispatch.
 		bNeedsQueuedAnimEventsDispatched = true;
@@ -3954,6 +3958,11 @@ void USkeletalMeshComponent::SendRenderDynamicData_Concurrent()
 	Super::SendRenderDynamicData_Concurrent();
 
 #if WITH_EDITOR
+	if (GetSkeletalMeshAsset() && GetSkeletalMeshAsset()->IsCompiling())
+	{
+		return;
+	}
+
 	if (SceneProxy)
 	{
 		UpdatePoseWatches();
diff --git a/Engine/Source/Runtime/Engine/Private/Components/SkinnedMeshComponent.cpp b/Engine/Source/Runtime/Engine/Private/Components/SkinnedMeshComponent.cpp
--- a/Engine/Source/Runtime/Engine/Private/Components/SkinnedMeshComponent.cpp
+++ b/Engine/Source/Runtime/Engine/Private/Components/SkinnedMeshComponent.cpp
@@ -955,6 +955,13 @@ void USkinnedMeshComponent::SendRenderDynamicData_Concurrent()
 
 	Super::SendRenderDynamicData_Concurrent();
 
+#if WITH_EDITOR
+	if (GetSkinnedAsset() && GetSkinnedAsset()->IsCompiling())
+	{
+		return;
+	}
+#endif
+
 	// if we have not updated the transforms then no need to send them to the rendering thread
 	// @todo GIsEditor used to be bUpdateSkelWhenNotRendered. Look into it further to find out why it doesn't update animations in the AnimSetViewer, when a level is loaded in UED (like POC_Cover.gear).
 	if( MeshObject && GetSkinnedAsset() && (bForceMeshObjectUpdate || (bRecentlyRendered || VisibilityBasedAnimTickOption == EVisibilityBasedAnimTickOption::AlwaysTickPoseAndRefreshBones || GIsEditor || MeshObject->bHasBeenUpdatedAtLeastOnce == false)) )
@@ -1460,6 +1467,14 @@ void USkinnedMeshComponent::UpdateBounds()
 		LastStreamerUpdateBoundsRadius = float(Bounds.SphereRadius);
 	}
 
+#if WITH_EDITOR
+	//We cannot update bounds during compilation
+	if (GetSkinnedAsset() && GetSkinnedAsset()->IsCompiling())
+	{
+		return;
+	}
+#endif
+
 	Super::UpdateBounds();
 
 	// Only notify the streamer if the bounds radius has changed enough
diff --git a/Engine/Source/Runtime/Engine/Private/SkeletalMesh.cpp b/Engine/Source/Runtime/Engine/Private/SkeletalMesh.cpp
--- a/Engine/Source/Runtime/Engine/Private/SkeletalMesh.cpp
+++ b/Engine/Source/Runtime/Engine/Private/SkeletalMesh.cpp
@@ -325,7 +325,7 @@ FScopedSkeletalMeshPostEditChange::~FScopedSkeletalMeshPostEditChange()
 void FScopedSkeletalMeshPostEditChange::SetSkeletalMesh(USkeletalMesh* InSkeletalMesh)
 {
 	//Some parallel task may try to call post edit change, we must prevent it
-	if (!IsInGameThread())
+	if (!IsInGameThread() || InSkeletalMesh->IsCompiling())
 	{
 		return;
 	}
@@ -1971,6 +1971,59 @@ void USkeletalMesh::FinishBuildInternal(FSkinnedAssetBuildContext& Context)
 	PostMeshCached.Broadcast(this);
 }
 
+void USkeletalMesh::LockPropertiesUntil(FEvent* Event)
+{
+	TRACE_CPUPROFILER_EVENT_SCOPE(USkeletalMesh::Import);
+	if (IsCompiling())
+	{
+		FSkinnedAssetCompilingManager::Get().FinishCompilation({ this });
+	}
+
+	auto AsyncTaskFunction = [Event]()
+	{
+		Event->Wait();
+	};
+
+	//Use the async task compile to lock the properties
+	FSkinnedAsyncTaskContext Context(!HasAnyInternalFlags(EInternalObjectFlags::Async), AsyncTaskFunction);
+	BeginAsyncTaskInternal(Context);
+	PrepareForAsyncCompilation();
+	FQueuedThreadPool* SkeletalMeshThreadPool = FSkinnedAssetCompilingManager::Get().GetThreadPool();
+	EQueuedWorkPriority BasePriority = FSkinnedAssetCompilingManager::Get().GetBasePriority(this);
+	check(AsyncTask == nullptr);
+	AsyncTask = MakeUnique<FSkinnedAssetAsyncBuildTask>(this, MoveTemp(Context));
+	AsyncTask->StartBackgroundTask(SkeletalMeshThreadPool, BasePriority, EQueuedWorkFlags::DoNotRunInsideBusyWait);
+	FSkinnedAssetCompilingManager::Get().AddSkinnedAssets({ this });
+}
+
+void USkeletalMesh::BeginAsyncTaskInternal(FSkinnedAsyncTaskContext& Context)
+{
+	check(IsInGameThread());
+
+	if (Context.bResetAsyncFlagOnFinish)
+	{
+		SetInternalFlags(EInternalObjectFlags::Async);
+	}
+	AcquireAsyncProperty();
+	//Allow thumbnail data so content browser get refresh properly
+	ReleaseAsyncProperty((uint64)ESkeletalMeshAsyncProperties::ThumbnailInfo);
+}
+
+void USkeletalMesh::ExecuteAsyncTaskInternal(FSkinnedAsyncTaskContext& Context)
+{
+	Context.AsyncTaskFunction();
+}
+
+void USkeletalMesh::FinishAsyncTaskInternal(FSkinnedAsyncTaskContext& Context)
+{
+	check(IsInGameThread());
+	if (Context.bResetAsyncFlagOnFinish)
+	{
+		ClearInternalFlags(EInternalObjectFlags::Async);
+	}
+	ReleaseAsyncProperty();
+}
+
 #endif // #if WITH_EDITOR
 
 void USkeletalMesh::PreSave(const class ITargetPlatform* TargetPlatform)
diff --git a/Engine/Source/Runtime/Engine/Private/SkeletalMeshComponentPhysics.cpp b/Engine/Source/Runtime/Engine/Private/SkeletalMeshComponentPhysics.cpp
--- a/Engine/Source/Runtime/Engine/Private/SkeletalMeshComponentPhysics.cpp
+++ b/Engine/Source/Runtime/Engine/Private/SkeletalMeshComponentPhysics.cpp
@@ -3811,7 +3811,11 @@ bool USkeletalMeshComponent::GetClothSimulatedPosition_GameThread(const FGuid& A
 
 void USkeletalMeshComponent::TickClothing(float DeltaTime, FTickFunction& ThisTickFunction)
 {
-	if (GetSkeletalMeshAsset() == nullptr || !ClothingSimulation || CVarEnableClothPhysics.GetValueOnGameThread() == 0)
+	bool bIsCompiling = false;
+#if WITH_EDITOR
+	bIsCompiling = (GetSkeletalMeshAsset() && GetSkeletalMeshAsset()->IsCompiling());
+#endif
+	if (GetSkeletalMeshAsset() == nullptr || !ClothingSimulation || CVarEnableClothPhysics.GetValueOnGameThread() == 0 || bIsCompiling)
 	{
 		return;
 	}
diff --git a/Engine/Source/Runtime/Engine/Private/SkinnedAssetAsyncCompileUtils.cpp b/Engine/Source/Runtime/Engine/Private/SkinnedAssetAsyncCompileUtils.cpp
--- a/Engine/Source/Runtime/Engine/Private/SkinnedAssetAsyncCompileUtils.cpp
+++ b/Engine/Source/Runtime/Engine/Private/SkinnedAssetAsyncCompileUtils.cpp
@@ -43,6 +43,11 @@ void FSkinnedAssetAsyncBuildWorker::DoWork()
 	{
 		SkinnedAsset->ExecuteBuildInternal(*BuildContext);
 	}
+
+	if (AsyncTaskContext.IsSet())
+	{
+		SkinnedAsset->ExecuteAsyncTaskInternal(*AsyncTaskContext);
+	}
 }
 
 #endif //WITH_EDITOR
\ No newline at end of file
diff --git a/Engine/Source/Runtime/Engine/Private/SkinnedAssetCompiler.cpp b/Engine/Source/Runtime/Engine/Private/SkinnedAssetCompiler.cpp
--- a/Engine/Source/Runtime/Engine/Private/SkinnedAssetCompiler.cpp
+++ b/Engine/Source/Runtime/Engine/Private/SkinnedAssetCompiler.cpp
@@ -229,6 +229,13 @@ void FSkinnedAssetCompilingManager::PostCompilation(USkinnedAsset* SkinnedAsset)
 
 				LocalAsyncTask->GetTask().BuildContext.Reset();
 			}
+
+			if (LocalAsyncTask->GetTask().AsyncTaskContext.IsSet())
+			{
+				SkinnedAsset->FinishAsyncTaskInternal(*LocalAsyncTask->GetTask().AsyncTaskContext);
+
+				LocalAsyncTask->GetTask().AsyncTaskContext.Reset();
+			}
 		}
 
 		// Calling this delegate during app exit might be quite dangerous and lead to crash
diff --git a/Engine/Source/Runtime/Interchange/Core/Public/InterchangeFactoryBase.h b/Engine/Source/Runtime/Interchange/Core/Public/InterchangeFactoryBase.h
--- a/Engine/Source/Runtime/Interchange/Core/Public/InterchangeFactoryBase.h
+++ b/Engine/Source/Runtime/Interchange/Core/Public/InterchangeFactoryBase.h
@@ -141,6 +141,9 @@ public:
 		return true;
 	}
 
+	/** Call when the user cancel the operation. */
+	virtual void Cancel() {}
+
 	/**
 	 * Parameters to pass to CreateAsset function
 	 */
diff --git a/Engine/Source/Runtime/Interchange/Core/Public/InterchangeResult.h b/Engine/Source/Runtime/Interchange/Core/Public/InterchangeResult.h
--- a/Engine/Source/Runtime/Interchange/Core/Public/InterchangeResult.h
+++ b/Engine/Source/Runtime/Interchange/Core/Public/InterchangeResult.h
@@ -139,3 +139,13 @@ public:
 	UPROPERTY()
 	FText Text;
 };
+
+UCLASS()
+class INTERCHANGECORE_API UInterchangeResultError_ReimportFail : public UInterchangeResultError
+{
+	GENERATED_BODY()
+
+public:
+
+	virtual FText GetText() const override { return NSLOCTEXT("InterchangeResultNS", "UInterchangeResultError_ReimportFail_GetText", "Re-import Fail, see log for more detail."); }
+};
diff --git a/Engine/Source/Runtime/Interchange/Engine/Private/InterchangeMeshUtilities.cpp b/Engine/Source/Runtime/Interchange/Engine/Private/InterchangeMeshUtilities.cpp
--- a/Engine/Source/Runtime/Interchange/Engine/Private/InterchangeMeshUtilities.cpp
+++ b/Engine/Source/Runtime/Interchange/Engine/Private/InterchangeMeshUtilities.cpp
@@ -3,8 +3,9 @@
 #include "InterchangeMeshUtilities.h"
 
 #include "Async/Future.h"
-#include "Engine/SkeletalMesh.h"
 #include "Engine/StaticMesh.h"
+#include "Engine/SkeletalMesh.h"
+#include "Engine/SkinnedAssetAsyncCompileUtils.h"
 #include "InterchangeAssetImportData.h"
 #include "InterchangeFactoryBase.h"
 #include "InterchangeEngineLogPrivate.h"
@@ -143,7 +144,8 @@ TFuture<bool> UInterchangeMeshUtilities::InternalImportCustomLodAsync(TSharedPtr
 
 				if(SourceSkeletalMesh)
 				{
-
+					//Make sure we can modify the skeletalmesh properties
+					FSkinnedAssetAsyncBuildScope AsyncBuildScope(SkeletalMesh);
 					Promise->SetValue(FLODUtilities::SetCustomLOD(SkeletalMesh, SourceSkeletalMesh, LodIndex, SourceDataFilename));
 					SourceSkeletalMesh->ClearFlags(RF_Standalone);
 					SourceSkeletalMesh->ClearInternalFlags(EInternalObjectFlags::Async);
diff --git a/Engine/Source/Runtime/Interchange/Engine/Private/Tasks/InterchangeTaskCompletion.cpp b/Engine/Source/Runtime/Interchange/Engine/Private/Tasks/InterchangeTaskCompletion.cpp
--- a/Engine/Source/Runtime/Interchange/Engine/Private/Tasks/InterchangeTaskCompletion.cpp
+++ b/Engine/Source/Runtime/Interchange/Engine/Private/Tasks/InterchangeTaskCompletion.cpp
@@ -52,6 +52,11 @@ void UE::Interchange::FTaskPreCompletion::DoTask(ENamedThreads::Type CurrentThre
 			//Verify if the task was cancel
 			if (AsyncHelper->bCancel)
 			{
+				for (const FImportAsyncHelper::FImportedObjectInfo& ObjectInfo : ObjectInfosPerSourceIndexPair.Value)
+				{
+					//Cancel factories so they can do proper cleanup
+					ObjectInfo.Factory->Cancel();
+				}
 				break;
 			}
 
