ANIM-RT truncation fixes - Persona

diff --git a/Engine/Source/Editor/Persona/Private/AnimTimeline/AnimModel.cpp b/Engine/Source/Editor/Persona/Private/AnimTimeline/AnimModel.cpp
--- a/Engine/Source/Editor/Persona/Private/AnimTimeline/AnimModel.cpp
+++ b/Engine/Source/Editor/Persona/Private/AnimTimeline/AnimModel.cpp
@@ -115,7 +115,7 @@ void FAnimModel::SetScrubPosition(FFrameTime NewScrubPostion) const
 				PreviewMeshComponent->PreviewInstance->SetPlaying(false);
 			}
 			
-			PreviewMeshComponent->PreviewInstance->SetPosition(NewScrubPostion.AsDecimal() / (double)GetTickResolution());
+			PreviewMeshComponent->PreviewInstance->SetPosition(static_cast<float>(NewScrubPostion.AsDecimal() / static_cast<double>(GetTickResolution())));
 		}
 	}
 }
@@ -248,9 +248,9 @@ void FAnimModel::SetEditableTime(int32 TimeIndex, double Time, bool bIsDragging)
 
 bool FAnimModel::Snap(float& InOutTime, float InSnapMargin, TArrayView<const FName> InSkippedSnapTypes) const
 {
-	double DoubleTime = (double)InOutTime;
+	double DoubleTime = InOutTime;
 	bool bResult = Snap(DoubleTime, (double)InSnapMargin, InSkippedSnapTypes);
-	InOutTime = DoubleTime;
+	InOutTime = static_cast<float>(DoubleTime);
 	return bResult;
 }
 
diff --git a/Engine/Source/Editor/Persona/Private/AnimTimeline/AnimModel_AnimMontage.cpp b/Engine/Source/Editor/Persona/Private/AnimTimeline/AnimModel_AnimMontage.cpp
--- a/Engine/Source/Editor/Persona/Private/AnimTimeline/AnimModel_AnimMontage.cpp
+++ b/Engine/Source/Editor/Persona/Private/AnimTimeline/AnimModel_AnimMontage.cpp
@@ -134,8 +134,8 @@ void FAnimModel_AnimMontage::OnSetEditableTime(int32 TimeIndex, double Time, boo
 			AnimMontage->Modify();
 	
 			FCompositeSection& Section = AnimMontage->CompositeSections[TimeIndex];
-			Section.SetTime(Time);
-			Section.Link(AnimMontage, Time);
+			Section.SetTime(static_cast<float>(Time));
+			Section.Link(AnimMontage, static_cast<float>(Time));
 
 			SortSections();
 			RefreshNotifyTriggerOffsets();
@@ -146,7 +146,7 @@ void FAnimModel_AnimMontage::OnSetEditableTime(int32 TimeIndex, double Time, boo
 		}
 	}
 
-	OnSectionTimeDragged.ExecuteIfBound(TimeIndex, Time, bIsDragging);
+	OnSectionTimeDragged.ExecuteIfBound(TimeIndex, static_cast<float>(Time), bIsDragging);
 }
 
 void FAnimModel_AnimMontage::OnMontageModified()
diff --git a/Engine/Source/Editor/Persona/Private/AnimTimeline/AnimTimeSliderController.cpp b/Engine/Source/Editor/Persona/Private/AnimTimeline/AnimTimeSliderController.cpp
--- a/Engine/Source/Editor/Persona/Private/AnimTimeline/AnimTimeSliderController.cpp
+++ b/Engine/Source/Editor/Persona/Private/AnimTimeline/AnimTimeSliderController.cpp
@@ -44,10 +44,10 @@ struct FAnimTimeSliderController::FScrubRangeToScreen
 
 	FScrubRangeToScreen(const TRange<double>& InViewInput, const FVector2D& InWidgetSize )
 	{
-		float ViewInputRange = InViewInput.Size<double>();
+		const double ViewInputRange = InViewInput.Size<double>();
 
 		ViewStart      = InViewInput.GetLowerBoundValue();
-		PixelsPerInput = ViewInputRange > 0 ? ( InWidgetSize.X / ViewInputRange ) : 0;
+		PixelsPerInput = ViewInputRange > 0 ? static_cast<float>( InWidgetSize.X / ViewInputRange ) : 0;
 	}
 
 	/** Local Widget Space -> Curve Input domain. */
@@ -59,7 +59,7 @@ struct FAnimTimeSliderController::FScrubRangeToScreen
 	/** Curve Input domain -> local Widget Space */
 	float InputToLocalX(double Input) const
 	{
-		return (Input - ViewStart) * PixelsPerInput;
+		return static_cast<float>((Input - ViewStart) * PixelsPerInput);
 	}
 };
 
@@ -101,15 +101,15 @@ FAnimTimeSliderController::FAnimTimeSliderController( const FTimeSliderArgs& InA
 
 FFrameTime FAnimTimeSliderController::ComputeFrameTimeFromMouse(const FGeometry& Geometry, FVector2D ScreenSpacePosition, FScrubRangeToScreen RangeToScreen, bool CheckSnapping) const
 {
-	FVector2D CursorPos  = Geometry.AbsoluteToLocal( ScreenSpacePosition );
-	double    MouseValue = RangeToScreen.LocalXToInput( CursorPos.X );
+	const FVector2D CursorPos  = Geometry.AbsoluteToLocal( ScreenSpacePosition );
+	const double    MouseValue = RangeToScreen.LocalXToInput( static_cast<float>(CursorPos.X) );
 
 	return MouseValue * GetTickResolution();
 }
 
 FAnimTimeSliderController::FScrubPixelRange FAnimTimeSliderController::GetHitTestScrubberPixelRange(FFrameTime ScrubTime, const FScrubRangeToScreen& RangeToScreen) const
 {
-	static const float DragToleranceSlateUnits = 2.f, MouseTolerance = 2.f;
+	constexpr float DragToleranceSlateUnits = 2.f, MouseTolerance = 2.f;
 	return GetScrubberPixelRange(ScrubTime, GetTickResolution(), GetDisplayRate(), RangeToScreen, DragToleranceSlateUnits + MouseTolerance);
 }
 
@@ -126,7 +126,7 @@ FAnimTimeSliderController::FScrubPixelRange FAnimTimeSliderController::GetScrubb
 	float EndPixel   = RangeToScreen.InputToLocalX( (Frame+1) / Resolution);
 
 	{
-		float RoundedStartPixel = FMath::RoundToInt(StartPixel);
+		const float RoundedStartPixel = static_cast<float>(FMath::RoundToInt(StartPixel));
 		EndPixel -= (StartPixel - RoundedStartPixel);
 
 		StartPixel = RoundedStartPixel;
@@ -194,7 +194,7 @@ void FAnimTimeSliderController::DrawTicks( FSlateWindowElementList& OutDrawEleme
 
 	double MajorGridStep  = 0.0;
 	int32  MinorDivisions = 0;
-	if (!Timeline->GetGridMetrics(InArgs.AllottedGeometry.Size.X, MajorGridStep, MinorDivisions))
+	if (!Timeline->GetGridMetrics(static_cast<float>(InArgs.AllottedGeometry.Size.X), MajorGridStep, MinorDivisions))
 	{
 		return;
 	}
@@ -276,19 +276,18 @@ int32 FAnimTimeSliderController::OnPaintTimeSlider( bool bMirrorLabels, const FG
 	const bool bEnabled = bParentEnabled;
 	const ESlateDrawEffect DrawEffects = bEnabled ? ESlateDrawEffect::None : ESlateDrawEffect::DisabledEffect;
 
-	TRange<double> LocalViewRange      = TimeSliderArgs.ViewRange.Get();
-	const float    LocalViewRangeMin   = LocalViewRange.GetLowerBoundValue();
-	const float    LocalViewRangeMax   = LocalViewRange.GetUpperBoundValue();
+	const TRange<double> LocalViewRange      = TimeSliderArgs.ViewRange.Get();
+	const float    LocalViewRangeMin   = static_cast<float>(LocalViewRange.GetLowerBoundValue());
+	const float    LocalViewRangeMax   = static_cast<float>(LocalViewRange.GetUpperBoundValue());
 	const float    LocalSequenceLength = LocalViewRangeMax-LocalViewRangeMin;
-	TRange<FFrameNumber> LocalPlaybackRange = TimeSliderArgs.PlaybackRange.Get();
+	const TRange<FFrameNumber> LocalPlaybackRange = TimeSliderArgs.PlaybackRange.Get();
 
-	FVector2D Scale = FVector2D(1.0f,1.0f);
 	if ( LocalSequenceLength > 0)
 	{
 		FScrubRangeToScreen RangeToScreen( LocalViewRange, AllottedGeometry.Size );
 
 		// draw tick marks
-		const float MajorTickHeight = 9.0f;
+		constexpr float MajorTickHeight = 9.0f;
 	
 		FDrawTickArgs Args;
 		{
@@ -299,7 +298,7 @@ int32 FAnimTimeSliderController::OnPaintTimeSlider( bool bMirrorLabels, const FG
 			Args.CullingRect = MyCullingRect;
 			Args.DrawEffects = DrawEffects;
 			Args.StartLayer = LayerId;
-			Args.TickOffset = bMirrorLabels ? 0.0f : FMath::Abs( AllottedGeometry.Size.Y - MajorTickHeight );
+			Args.TickOffset = bMirrorLabels ? 0.0f : FMath::Abs( static_cast<float>(AllottedGeometry.Size.Y) - MajorTickHeight );
 			Args.MajorTickHeight = MajorTickHeight;
 		}
 
@@ -365,12 +364,12 @@ int32 FAnimTimeSliderController::OnPaintTimeSlider( bool bMirrorLabels, const FG
 			FSlateFontInfo SmallLayoutFont = FCoreStyle::GetDefaultFontStyle("Regular", 10);
 
 			const TSharedRef< FSlateFontMeasure > FontMeasureService = FSlateApplication::Get().GetRenderer()->GetFontMeasureService();
-			FVector2D TextSize = FontMeasureService->Measure(FrameString, SmallLayoutFont);
+			const FVector2D TextSize = FontMeasureService->Measure(FrameString, SmallLayoutFont);
 
 			// Flip the text position if getting near the end of the view range
-			static const float TextOffsetPx  = 2.f;
-			bool  bDrawLeft    = (AllottedGeometry.Size.X - HandleEnd) < (TextSize.X + 14.f) - TextOffsetPx;
-			float TextPosition = bDrawLeft ? HandleStart - TextSize.X - TextOffsetPx : HandleEnd + TextOffsetPx;
+			constexpr float TextOffsetPx  = 2.f;
+			const bool  bDrawLeft    = (static_cast<float>(AllottedGeometry.Size.X) - HandleEnd) < (TextSize.X + 14.f) - TextOffsetPx;
+			const float TextPosition = bDrawLeft ? HandleStart - static_cast<float>(TextSize.X) - TextOffsetPx : HandleEnd + TextOffsetPx;
 
 			FVector2D TextOffset( TextPosition, Args.bMirrorLabels ? TextSize.Y-6.f : Args.AllottedGeometry.Size.Y - (Args.MajorTickHeight+TextSize.Y) );
 
@@ -621,11 +620,11 @@ FReply FAnimTimeSliderController::OnMouseButtonUp( SWidget& WidgetOwner, const F
 
 			if(!MouseEvent.IsControlDown())
 			{
-				double SnapMargin = (ScrubConstants::SnapMarginInPixels / (double)RangeToScreen.PixelsPerInput);
+				const double SnapMargin = (ScrubConstants::SnapMarginInPixels / static_cast<double>(RangeToScreen.PixelsPerInput));
 				WeakModel.Pin()->Snap(Time, SnapMargin, { FName("MontageSection") });
 			}
 
-			SetEditableTime(DraggedTimeIndex, Time, false);
+			SetEditableTime(DraggedTimeIndex, static_cast<float>(Time), false);
 			DraggedTimeIndex = INDEX_NONE;
 		}
 		else
@@ -647,14 +646,14 @@ FReply FAnimTimeSliderController::OnMouseButtonUp( SWidget& WidgetOwner, const F
 
 FReply FAnimTimeSliderController::OnMouseMove( SWidget& WidgetOwner, const FGeometry& MyGeometry, const FPointerEvent& MouseEvent )
 {
-	bool bHandleLeftMouseButton  = MouseEvent.IsMouseButtonDown( EKeys::LeftMouseButton  );
-	bool bHandleRightMouseButton = MouseEvent.IsMouseButtonDown( EKeys::RightMouseButton ) && TimeSliderArgs.AllowZoom;
+	const bool bHandleLeftMouseButton  = MouseEvent.IsMouseButtonDown( EKeys::LeftMouseButton  );
+	const bool bHandleRightMouseButton = MouseEvent.IsMouseButtonDown( EKeys::RightMouseButton ) && TimeSliderArgs.AllowZoom;
 
 	if (bHandleRightMouseButton)
 	{
 		if (!bPanning)
 		{
-			DistanceDragged += FMath::Abs( MouseEvent.GetCursorDelta().X );
+			DistanceDragged += static_cast<float>(FMath::Abs( MouseEvent.GetCursorDelta().X ));
 			if ( DistanceDragged > FSlateApplication::Get().GetDragTriggerDistance() )
 			{
 				bPanning = true;
@@ -662,12 +661,12 @@ FReply FAnimTimeSliderController::OnMouseMove( SWidget& WidgetOwner, const FGeom
 		}
 		else
 		{
-			TRange<double> LocalViewRange = TimeSliderArgs.ViewRange.Get();
-			double LocalViewRangeMin = LocalViewRange.GetLowerBoundValue();
-			double LocalViewRangeMax = LocalViewRange.GetUpperBoundValue();
+			const TRange<double> LocalViewRange = TimeSliderArgs.ViewRange.Get();
+			const double LocalViewRangeMin = LocalViewRange.GetLowerBoundValue();
+			const double LocalViewRangeMax = LocalViewRange.GetUpperBoundValue();
 
-			FScrubRangeToScreen ScaleInfo( LocalViewRange, MyGeometry.Size );
-			FVector2D ScreenDelta = MouseEvent.GetCursorDelta();
+			const FScrubRangeToScreen ScaleInfo( LocalViewRange, MyGeometry.Size );
+			const FVector2D ScreenDelta = MouseEvent.GetCursorDelta();
 			FVector2D InputDelta;
 			InputDelta.X = ScreenDelta.X/ScaleInfo.PixelsPerInput;
 
@@ -680,18 +679,18 @@ FReply FAnimTimeSliderController::OnMouseMove( SWidget& WidgetOwner, const FGeom
 	}
 	else if (bHandleLeftMouseButton)
 	{
-		TRange<double> LocalViewRange = TimeSliderArgs.ViewRange.Get();
-		FScrubRangeToScreen RangeToScreen(LocalViewRange, MyGeometry.Size);
-		DistanceDragged += FMath::Abs( MouseEvent.GetCursorDelta().X );
+		const TRange<double> LocalViewRange = TimeSliderArgs.ViewRange.Get();
+		const FScrubRangeToScreen RangeToScreen(LocalViewRange, MyGeometry.Size);
+		DistanceDragged += static_cast<float>(FMath::Abs( MouseEvent.GetCursorDelta().X ));
 
 		if ( MouseDragType == DRAG_NONE )
 		{
 			if ( DistanceDragged > FSlateApplication::Get().GetDragTriggerDistance() )
 			{
 				UAnimMontage* AnimMontage = Cast<UAnimMontage>(WeakModel.Pin()->GetAnimSequenceBase());
-				bool bChildAnimMontage = AnimMontage && AnimMontage->HasParentAsset();
+				const bool bChildAnimMontage = AnimMontage && AnimMontage->HasParentAsset();
 
-				FFrameTime MouseDownFree = ComputeFrameTimeFromMouse(MyGeometry, MouseDownPosition[0], RangeToScreen, false);
+				const FFrameTime MouseDownFree = ComputeFrameTimeFromMouse(MyGeometry, MouseDownPosition[0], RangeToScreen, false);
 
 				const FFrameRate FrameResolution    = GetTickResolution();
 				const bool       bLockedPlayRange   = TimeSliderArgs.IsPlaybackRangeLocked.Get();
@@ -700,8 +699,8 @@ FReply FAnimTimeSliderController::OnMouseMove( SWidget& WidgetOwner, const FGeom
 				const int32      HitTimeIndex       = HitTestTimes(RangeToScreen, MouseDownPixel);
 				const bool       bHitTime           = !bChildAnimMontage && HitTimeIndex != INDEX_NONE;
 
-				TRange<double>   SelectionRange   = TimeSliderArgs.SelectionRange.Get() / FrameResolution;
-				TRange<double>   PlaybackRange    = TimeSliderArgs.PlaybackRange.Get()  / FrameResolution;
+				const TRange<double>   SelectionRange   = TimeSliderArgs.SelectionRange.Get() / FrameResolution;
+				const TRange<double>   PlaybackRange    = TimeSliderArgs.PlaybackRange.Get()  / FrameResolution;
 
 				// Disable selection range test if it's empty so that the playback range scrubbing gets priority
 				if (!SelectionRange.IsEmpty() && !bHitScrubber && HitTestRangeEnd(RangeToScreen, SelectionRange, MouseDownPixel))
@@ -782,11 +781,11 @@ FReply FAnimTimeSliderController::OnMouseMove( SWidget& WidgetOwner, const FGeom
 
 				if(!MouseEvent.IsControlDown())
 				{
-					double SnapMargin = (ScrubConstants::SnapMarginInPixels / (double)RangeToScreen.PixelsPerInput);
+					const double SnapMargin = (ScrubConstants::SnapMarginInPixels / static_cast<double>(RangeToScreen.PixelsPerInput));
 					WeakModel.Pin()->Snap(Time, SnapMargin, { FName("MontageSection") });
 				}
 
-				SetEditableTime(DraggedTimeIndex, Time, true);
+				SetEditableTime(DraggedTimeIndex, static_cast<float>(Time), true);
 			}
 		}
 	}
@@ -818,7 +817,7 @@ FReply FAnimTimeSliderController::OnMouseWheel( SWidget& WidgetOwner, const FGeo
 
 	if ( TimeSliderArgs.AllowZoom && MouseEvent.IsControlDown() )
 	{
-		float MouseFractionX = MyGeometry.AbsoluteToLocal(MouseEvent.GetScreenSpacePosition()).X / MyGeometry.GetLocalSize().X;
+		const float MouseFractionX = static_cast<float>(MyGeometry.AbsoluteToLocal(MouseEvent.GetScreenSpacePosition()).X / MyGeometry.GetLocalSize().X);
 
 		const float ZoomDelta = -0.2f * MouseEvent.GetWheelDelta();
 		if (ZoomByDelta(ZoomDelta, MouseFractionX))
@@ -837,19 +836,19 @@ FReply FAnimTimeSliderController::OnMouseWheel( SWidget& WidgetOwner, const FGeo
 
 FCursorReply FAnimTimeSliderController::OnCursorQuery( TSharedRef<const SWidget> WidgetOwner, const FGeometry& MyGeometry, const FPointerEvent& CursorEvent ) const
 {
-	FScrubRangeToScreen RangeToScreen(TimeSliderArgs.ViewRange.Get(), MyGeometry.Size);
+	const FScrubRangeToScreen RangeToScreen(TimeSliderArgs.ViewRange.Get(), MyGeometry.Size);
 
 	UAnimMontage* AnimMontage = Cast<UAnimMontage>(WeakModel.Pin()->GetAnimSequenceBase());
-	bool bChildAnimMontage = AnimMontage && AnimMontage->HasParentAsset();
+	const bool bChildAnimMontage = AnimMontage && AnimMontage->HasParentAsset();
 
 	const FFrameRate FrameResolution  = GetTickResolution();
 	const bool       bLockedPlayRange = TimeSliderArgs.IsPlaybackRangeLocked.Get();
-	const float      HitTestPixel     = MyGeometry.AbsoluteToLocal(CursorEvent.GetScreenSpacePosition()).X;
+	const float      HitTestPixel     = static_cast<float>(MyGeometry.AbsoluteToLocal(CursorEvent.GetScreenSpacePosition()).X);
 	const bool       bHitScrubber     = GetHitTestScrubberPixelRange(TimeSliderArgs.ScrubPosition.Get(), RangeToScreen).HandleRange.Contains(HitTestPixel);
 	const bool       bHitTime         = !bChildAnimMontage && (HitTestTimes(RangeToScreen, HitTestPixel) != INDEX_NONE);
 
-	TRange<double>   SelectionRange   = TimeSliderArgs.SelectionRange.Get() / FrameResolution;
-	TRange<double>   PlaybackRange    = TimeSliderArgs.PlaybackRange.Get()  / FrameResolution;
+	const TRange<double>   SelectionRange   = TimeSliderArgs.SelectionRange.Get() / FrameResolution;
+	const TRange<double>   PlaybackRange    = TimeSliderArgs.PlaybackRange.Get()  / FrameResolution;
 
 	if (MouseDragType == DRAG_SCRUBBING_TIME)
 	{
@@ -906,7 +905,7 @@ int32 FAnimTimeSliderController::OnPaintViewArea( const FGeometry& AllottedGeome
 			DrawTickArgs.StartLayer = LayerId-1;
 			// Draw the tick the entire height of the section area
 			DrawTickArgs.TickOffset = 0.0f;
-			DrawTickArgs.MajorTickHeight = AllottedGeometry.Size.Y;
+			DrawTickArgs.MajorTickHeight = static_cast<float>(AllottedGeometry.Size.Y);
 		}
 
 		DrawTicks( OutDrawElements, LocalViewRange, RangeToScreen, DrawTickArgs );
@@ -1125,13 +1124,13 @@ bool FAnimTimeSliderController::ZoomByDelta( float InDelta, float MousePositionF
 
 void FAnimTimeSliderController::PanByDelta( float InDelta )
 {
-	TRange<double> LocalViewRange = TimeSliderArgs.ViewRange.Get().GetAnimationTarget();
+	const TRange<double> LocalViewRange = TimeSliderArgs.ViewRange.Get().GetAnimationTarget();
 
-	double CurrentMin = LocalViewRange.GetLowerBoundValue();
-	double CurrentMax = LocalViewRange.GetUpperBoundValue();
+	const double CurrentMin = LocalViewRange.GetLowerBoundValue();
+	const double CurrentMax = LocalViewRange.GetUpperBoundValue();
 
 	// Adjust the delta to be a percentage of the current range
-	InDelta *= ScrubConstants::ScrollPanFraction * (CurrentMax - CurrentMin);
+	InDelta *= ScrubConstants::ScrollPanFraction * static_cast<float>(CurrentMax - CurrentMin);
 
 	double NewViewOutputMin = CurrentMin + InDelta;
 	double NewViewOutputMax = CurrentMax + InDelta;
diff --git a/Engine/Source/Editor/Persona/Private/AnimTimeline/AnimTimelineTrack.cpp b/Engine/Source/Editor/Persona/Private/AnimTimeline/AnimTimelineTrack.cpp
--- a/Engine/Source/Editor/Persona/Private/AnimTimeline/AnimTimelineTrack.cpp
+++ b/Engine/Source/Editor/Persona/Private/AnimTimeline/AnimTimelineTrack.cpp
@@ -174,18 +174,18 @@ float FAnimTimelineTrack::GetMaxInput() const
 
 float FAnimTimelineTrack::GetViewMinInput() const
 {
-	return GetModel()->GetViewRange().GetLowerBoundValue();
+	return static_cast<float>(GetModel()->GetViewRange().GetLowerBoundValue());
 }
 
 float FAnimTimelineTrack::GetViewMaxInput() const
 {
-	return GetModel()->GetViewRange().GetUpperBoundValue();
+	return static_cast<float>(GetModel()->GetViewRange().GetUpperBoundValue());
 }
 
 float FAnimTimelineTrack::GetScrubValue() const
 {
-	const int32 Resolution = FMath::RoundToInt((double)GetDefault<UPersonaOptions>()->TimelineScrubSnapValue * GetModel()->GetFrameRate());
-	return (float)((double)GetModel()->GetScrubPosition().Value / (double)Resolution);
+	const int64 Resolution = FMath::RoundToInt(static_cast<double>(GetDefault<UPersonaOptions>()->TimelineScrubSnapValue) * GetModel()->GetFrameRate());
+	return static_cast<float>(static_cast<double>(GetModel()->GetScrubPosition().Value) / static_cast<double>(Resolution));
 }
 
 void FAnimTimelineTrack::SelectObjects(const TArray<UObject*>& SelectedItems)
diff --git a/Engine/Source/Editor/Persona/Private/AnimTimeline/AnimTimelineTrack_Curve.cpp b/Engine/Source/Editor/Persona/Private/AnimTimeline/AnimTimelineTrack_Curve.cpp
--- a/Engine/Source/Editor/Persona/Private/AnimTimeline/AnimTimelineTrack_Curve.cpp
+++ b/Engine/Source/Editor/Persona/Private/AnimTimeline/AnimTimelineTrack_Curve.cpp
@@ -74,9 +74,9 @@ public:
 	virtual int32 OnPaint(const FPaintArgs& Args, const FGeometry& AllottedGeometry, const FSlateRect& MyCullingRect, FSlateWindowElementList& OutDrawElements, int32 LayerId, const FWidgetStyle& InWidgetStyle, bool bParentEnabled) const override
 	{
 		// Rendering info
-		const float LabelOffsetPx = 2.0f;
-		const float Width = AllottedGeometry.GetLocalSize().X;
-		const float Height = AllottedGeometry.GetLocalSize().Y;
+		constexpr float LabelOffsetPx = 2.0f;
+		const float Width = static_cast<float>(AllottedGeometry.GetLocalSize().X);
+		const float Height = static_cast<float>(AllottedGeometry.GetLocalSize().Y);
 		const FPaintGeometry PaintGeometry  = AllottedGeometry.ToPaintGeometry();
 		const FSlateFontInfo FontInfo = FCoreStyle::GetDefaultFontStyle("Regular", 8);
 		TSharedRef<FSlateFontMeasure> FontMeasureService = FSlateApplication::Get().GetRenderer()->GetFontMeasureService();
diff --git a/Engine/Source/Editor/Persona/Private/AnimTimeline/AnimTimelineTrack_Montage.cpp b/Engine/Source/Editor/Persona/Private/AnimTimeline/AnimTimelineTrack_Montage.cpp
--- a/Engine/Source/Editor/Persona/Private/AnimTimeline/AnimTimelineTrack_Montage.cpp
+++ b/Engine/Source/Editor/Persona/Private/AnimTimeline/AnimTimelineTrack_Montage.cpp
@@ -154,21 +154,21 @@ class SMontageSections : public SLeafWidget
 				NameText = FText::FromName(CompositeSection.SectionName);
 			}
 
-			FVector2D TextSize = FontMeasureService->Measure(NameText, LabelFont);
-			FVector2D TextBorderSize = TextSize + (MontageSectionsConstants::TextBorderMargin * 2.0f);
+			const FVector2D TextSize = FontMeasureService->Measure(NameText, LabelFont);
+			const FVector2D TextBorderSize = TextSize + (MontageSectionsConstants::TextBorderMargin * 2.0f);
 
-			FText IconText = GetSectionIconText(SectionIndex);
-			FVector2D TextIconSize = FontMeasureService->Measure(IconText, LabelFont);
-			FVector2D TextIconBorderSize = TextIconSize + (MontageSectionsConstants::TextBorderMargin * 2.0f);
+			const FText IconText = GetSectionIconText(SectionIndex);
+			const FVector2D TextIconSize = FontMeasureService->Measure(IconText, LabelFont);
+			const FVector2D TextIconBorderSize = TextIconSize + (MontageSectionsConstants::TextBorderMargin * 2.0f);
 
-			FVector2D TotalBorderSize(TextBorderSize.X + TextIconBorderSize.X, FMath::Max(TextBorderSize.Y, TextIconBorderSize.Y));
+			const FVector2D TotalBorderSize(TextBorderSize.X + TextIconBorderSize.X, FMath::Max(TextBorderSize.Y, TextIconBorderSize.Y));
 
-			float LabelPosX = ScaleInfo.InputToLocalX(DraggedSectionIndex == SectionIndex ? DraggedSectionTime : CompositeSection.GetTime());
-			float LabelPosY = (AllottedGeometry.GetLocalSize().Y * 0.5f) - (TextBorderSize.Y * 0.5f);
-			float IconPosY = (AllottedGeometry.GetLocalSize().Y * 0.5f) - (TextIconBorderSize.Y * 0.5f);
+			const float LabelPosX = ScaleInfo.InputToLocalX(DraggedSectionIndex == SectionIndex ? DraggedSectionTime : CompositeSection.GetTime());
+			const float LabelPosY = static_cast<float>((AllottedGeometry.GetLocalSize().Y * 0.5) - (TextBorderSize.Y * 0.5));
+			const float IconPosY = static_cast<float>((AllottedGeometry.GetLocalSize().Y * 0.5) - (TextIconBorderSize.Y * 0.5));
 
-			float RightEdgeToNotify = AllottedGeometry.Size.X - (LabelPosX + TextBorderSize.X);
-			bool bDrawLabelOnLeft = RightEdgeToNotify < 0.0f;
+			const float RightEdgeToNotify = static_cast<float>(AllottedGeometry.Size.X - (LabelPosX + TextBorderSize.X));
+			const bool bDrawLabelOnLeft = RightEdgeToNotify < 0.0f;
 
 			FSlateDrawElement::MakeBox( 
 				OutDrawElements,
@@ -214,21 +214,21 @@ class SMontageSections : public SLeafWidget
 		{
 			const FCompositeSection& CompositeSection = AnimMontage->CompositeSections[SectionIndex];
 
-			FText NameText = FText::FromName(CompositeSection.SectionName);
-			FVector2D TextSize = FontMeasureService->Measure(NameText, LabelFont);
-			FVector2D TextBorderSize = TextSize + (MontageSectionsConstants::TextBorderMargin * 2.0f);
+			const FText NameText = FText::FromName(CompositeSection.SectionName);
+			const FVector2D TextSize = FontMeasureService->Measure(NameText, LabelFont);
+			const FVector2D TextBorderSize = TextSize + (MontageSectionsConstants::TextBorderMargin * 2.0f);
 
-			FText IconText = GetSectionIconText(SectionIndex);
-			FVector2D TextIconSize = FontMeasureService->Measure(IconText, LabelFont);
-			FVector2D TextIconBorderSize = TextIconSize + (MontageSectionsConstants::TextBorderMargin * 2.0f);
+			const FText IconText = GetSectionIconText(SectionIndex);
+			const FVector2D TextIconSize = FontMeasureService->Measure(IconText, LabelFont);
+			const FVector2D TextIconBorderSize = TextIconSize + (MontageSectionsConstants::TextBorderMargin * 2.0f);
 
-			FVector2D TotalBorderSize(TextBorderSize.X + TextIconBorderSize.X, FMath::Max(TextBorderSize.Y, TextIconBorderSize.Y));
+			const FVector2D TotalBorderSize(TextBorderSize.X + TextIconBorderSize.X, FMath::Max(TextBorderSize.Y, TextIconBorderSize.Y));
 
-			float LabelPosX = ScaleInfo.InputToLocalX(CompositeSection.GetTime());
-			float LabelPosY = (MyGeometry.GetLocalSize().Y * 0.5f) - (TotalBorderSize.Y * 0.5f);
+			const double LabelPosX = ScaleInfo.InputToLocalX(CompositeSection.GetTime());
+			const double LabelPosY = (MyGeometry.GetLocalSize().Y * 0.5) - (TotalBorderSize.Y * 0.5);
 
-			float RightEdgeToNotify = MyGeometry.Size.X - (LabelPosX + TotalBorderSize.X);
-			bool bDrawLabelOnLeft = RightEdgeToNotify < 0.0f;
+			const double RightEdgeToNotify = MyGeometry.Size.X - (LabelPosX + TotalBorderSize.X);
+			const bool bDrawLabelOnLeft = RightEdgeToNotify < 0.0;
 
 			const FGeometry LabelGeometry = MyGeometry.MakeChild(FVector2D(bDrawLabelOnLeft ? LabelPosX - TotalBorderSize.X: LabelPosX, LabelPosY), TotalBorderSize);
 
@@ -277,7 +277,7 @@ class SMontageSections : public SLeafWidget
 				FMenuBuilder MenuBuilder(true, nullptr);
 
 				FTrackScaleInfo ScaleInfo(ViewInputMin.Get(), ViewInputMax.Get(), 0, 0, MyGeometry.GetLocalSize());
-				SummonTrackContextMenu(MenuBuilder, ScaleInfo.LocalXToInput(MyGeometry.AbsoluteToLocal(MouseEvent.GetScreenSpacePosition()).X), SelectedSectionIndex);
+				SummonTrackContextMenu(MenuBuilder, ScaleInfo.LocalXToInput(static_cast<float>(MyGeometry.AbsoluteToLocal(MouseEvent.GetScreenSpacePosition()).X)), SelectedSectionIndex);
 
 				FSlateApplication::Get().PushMenu(
 					AsShared(),
diff --git a/Engine/Source/Editor/Persona/Private/AnimTimeline/SAnimOutliner.cpp b/Engine/Source/Editor/Persona/Private/AnimTimeline/SAnimOutliner.cpp
--- a/Engine/Source/Editor/Persona/Private/AnimTimeline/SAnimOutliner.cpp
+++ b/Engine/Source/Editor/Persona/Private/AnimTimeline/SAnimOutliner.cpp
@@ -187,18 +187,18 @@ void SAnimOutliner::HandleTracksChanged()
 
 void SAnimOutliner::ReportChildRowGeometry(const TSharedRef<FAnimTimelineTrack>& InTrack, const FGeometry& InGeometry)
 {
-	float ChildOffset = TransformPoint(
+	const float ChildOffset = static_cast<float>(TransformPoint(
 		Concatenate(
 			InGeometry.GetAccumulatedLayoutTransform(),
 			GetCachedGeometry().GetAccumulatedLayoutTransform().Inverse()
 		),
 		FVector2D(0,0)
-	).Y;
+	).Y);
 
-	FCachedGeometry* ExistingGeometry = CachedTrackGeometry.Find(InTrack);
+	const FCachedGeometry* ExistingGeometry = CachedTrackGeometry.Find(InTrack);
 	if(ExistingGeometry == nullptr || (ExistingGeometry->Top != ChildOffset || ExistingGeometry->Height != InGeometry.Size.Y))
 	{
-		CachedTrackGeometry.Add(InTrack, FCachedGeometry(InTrack, ChildOffset, InGeometry.Size.Y));
+		CachedTrackGeometry.Add(InTrack, FCachedGeometry(InTrack, ChildOffset, static_cast<float>(InGeometry.Size.Y)));
 		bPhysicalTracksNeedUpdate = true;
 	}
 }
diff --git a/Engine/Source/Editor/Persona/Private/AnimTimeline/SAnimTimeline.cpp b/Engine/Source/Editor/Persona/Private/AnimTimeline/SAnimTimeline.cpp
--- a/Engine/Source/Editor/Persona/Private/AnimTimeline/SAnimTimeline.cpp
+++ b/Engine/Source/Editor/Persona/Private/AnimTimeline/SAnimTimeline.cpp
@@ -52,7 +52,7 @@ void SAnimTimeline::Construct(const FArguments& InArgs, const TSharedRef<FAnimMo
 	OnReceivedFocus = InArgs._OnReceivedFocus;
 
 	int32 TickResolutionValue = InModel->GetTickResolution();
-	int32 SequenceFrameRate = FMath::RoundToInt(InModel->GetFrameRate());
+	int32 SequenceFrameRate = static_cast<int32>(FMath::RoundToInt(InModel->GetFrameRate()));
 
 	if (InModel->GetPreviewScene()->GetPreviewMeshComponent()->PreviewInstance)
 	{
@@ -61,22 +61,22 @@ void SAnimTimeline::Construct(const FArguments& InArgs, const TSharedRef<FAnimMo
 
 	ViewRange = MakeAttributeLambda([WeakModel](){ return WeakModel.IsValid() ? WeakModel.Pin()->GetViewRange() : FAnimatedRange(0.0, 0.0); });
 
-	TAttribute<EFrameNumberDisplayFormats> DisplayFormat = MakeAttributeLambda([]()
+	const TAttribute<EFrameNumberDisplayFormats> DisplayFormat = MakeAttributeLambda([]()
 	{
 		return GetDefault<UPersonaOptions>()->TimelineDisplayFormat;
 	});
 
-	TAttribute<EFrameNumberDisplayFormats> DisplayFormatSecondary = MakeAttributeLambda([]()
+	const TAttribute<EFrameNumberDisplayFormats> DisplayFormatSecondary = MakeAttributeLambda([]()
 	{
 		return GetDefault<UPersonaOptions>()->TimelineDisplayFormat == EFrameNumberDisplayFormats::Frames ? EFrameNumberDisplayFormats::Seconds : EFrameNumberDisplayFormats::Frames;
 	});
 
-	TAttribute<FFrameRate> TickResolution = MakeAttributeLambda([TickResolutionValue]()
+	const TAttribute<FFrameRate> TickResolution = MakeAttributeLambda([TickResolutionValue]()
 	{
 		return FFrameRate(TickResolutionValue, 1);
 	});
 
-	TAttribute<FFrameRate> DisplayRate = MakeAttributeLambda([SequenceFrameRate]()
+	const TAttribute<FFrameRate> DisplayRate = MakeAttributeLambda([SequenceFrameRate]()
 	{
 		return FFrameRate(SequenceFrameRate, 1);
 	});
@@ -286,7 +286,7 @@ void SAnimTimeline::Construct(const FArguments& InArgs, const TSharedRef<FAnimMo
 						SNew( SBorder )
 						.BorderImage( FAppStyle::GetBrush("ToolPanel.GroupBorder") )
 						.BorderBackgroundColor( FLinearColor(.50f, .50f, .50f, 1.0f ) )
-						.Padding(0)
+						.Padding(0.f)
 						.Clipping(EWidgetClipping::ClipToBounds)
 						[
 							TopTimeSlider.ToSharedRef()
@@ -405,18 +405,18 @@ FReply SAnimTimeline::OnMouseButtonUp(const FGeometry& MyGeometry, const FPointe
 		}
 		MenuBuilder.EndSection();
 
-		UAnimSequence* AnimSequence = Cast<UAnimSequence>(Model.Pin()->GetAnimSequenceBase());
+		const UAnimSequence* AnimSequence = Cast<UAnimSequence>(Model.Pin()->GetAnimSequenceBase());
 		if( AnimSequence )
 		{
 			FFrameTime MouseTime = TimeSliderController->GetFrameTimeFromMouse(MyGeometry, MouseEvent.GetScreenSpacePosition());
-			float CurrentFrameTime = (float)((double)MouseTime.AsDecimal() / (double)Model.Pin()->GetTickResolution());
-			float SequenceLength = AnimSequence->GetPlayLength();
+			const float CurrentFrameTime = static_cast<float>(MouseTime.AsDecimal() / static_cast<double>(Model.Pin()->GetTickResolution()));
+			const float SequenceLength = AnimSequence->GetPlayLength();
 			const uint32 NumKeys = AnimSequence->GetNumberOfSampledKeys();
 
 			MenuBuilder.BeginSection("SequenceEditingContext", LOCTEXT("SequenceEditing", "Sequence Editing") );
 			{
 				float CurrentFrameFraction = CurrentFrameTime / SequenceLength;
-				int32 CurrentKeyIndex = CurrentFrameFraction * NumKeys;
+				int32 CurrentKeyIndex = static_cast<int32>(CurrentFrameFraction * NumKeys);
 
 				FUIAction Action;
 				FText Label;
@@ -425,7 +425,7 @@ FReply SAnimTimeline::OnMouseButtonUp(const FGeometry& MyGeometry, const FPointe
 				//Only show this option if the selected frame is greater than frame 1 (first frame)
 				if (CurrentKeyIndex > 0)
 				{
-					CurrentFrameFraction = (float)CurrentKeyIndex / (float)NumKeys;
+					CurrentFrameFraction = static_cast<float>(CurrentKeyIndex) / static_cast<float>(NumKeys);
 
 					//Corrected frame time based on selected frame number
 					float CorrectedFrameTime = CurrentFrameFraction * SequenceLength;
@@ -668,7 +668,7 @@ TSharedRef<INumericTypeInterface<double>> SAnimTimeline::GetNumericTypeInterface
 static bool ComputeGridSpacing(const FFrameRate& InFrameRate, float PixelsPerSecond, double& OutMajorInterval, int32& OutMinorDivisions, float MinTickPx, float DesiredMajorTickPx)
 {
 	// First try built-in spacing
-	bool bResult = InFrameRate.ComputeGridSpacing(PixelsPerSecond, OutMajorInterval, OutMinorDivisions, MinTickPx, DesiredMajorTickPx);
+	const bool bResult = InFrameRate.ComputeGridSpacing(PixelsPerSecond, OutMajorInterval, OutMinorDivisions, MinTickPx, DesiredMajorTickPx);
 	if(!bResult || OutMajorInterval == 1.0)
 	{
 		if (PixelsPerSecond <= 0.f)
@@ -676,7 +676,7 @@ static bool ComputeGridSpacing(const FFrameRate& InFrameRate, float PixelsPerSec
 			return false;
 		}
 
-		const int32 RoundedFPS = FMath::RoundToInt(InFrameRate.AsDecimal());
+		const int32 RoundedFPS = static_cast<int32>(FMath::RoundToInt(InFrameRate.AsDecimal()));
 
 		if (RoundedFPS > 0)
 		{
@@ -721,11 +721,11 @@ static bool ComputeGridSpacing(const FFrameRate& InFrameRate, float PixelsPerSec
 
 			Algo::Reverse(CommonBases);
 
-			const int32 Scale     = FMath::CeilToInt(DesiredMajorTickPx / PixelsPerSecond * InFrameRate.AsDecimal());
+			const int32 Scale     = static_cast<int32>(FMath::CeilToInt(DesiredMajorTickPx / PixelsPerSecond * InFrameRate.AsDecimal()));
 			const int32 BaseIndex = FMath::Min(Algo::LowerBound(CommonBases, Scale), CommonBases.Num()-1);
 			const int32 Base      = CommonBases[BaseIndex];
 
-			int32 MajorIntervalFrames = FMath::CeilToInt(Scale / float(Base)) * Base;
+			const int32 MajorIntervalFrames = FMath::CeilToInt(Scale / static_cast<float>(Base)) * Base;
 			OutMajorInterval  = MajorIntervalFrames * InFrameRate.AsInterval();
 
 			// Find the lowest number of divisions we can show that's larger than the minimum tick size
@@ -734,7 +734,7 @@ static bool ComputeGridSpacing(const FFrameRate& InFrameRate, float PixelsPerSec
 			{
 				if (Base % CommonBases[DivIndex] == 0)
 				{
-					int32 MinorDivisions = MajorIntervalFrames/CommonBases[DivIndex];
+					const int32 MinorDivisions = MajorIntervalFrames/CommonBases[DivIndex];
 					if (OutMajorInterval / MinorDivisions * PixelsPerSecond >= MinTickPx)
 					{
 						OutMinorDivisions = MinorDivisions;
@@ -750,24 +750,24 @@ static bool ComputeGridSpacing(const FFrameRate& InFrameRate, float PixelsPerSec
 
 bool SAnimTimeline::GetGridMetrics(float PhysicalWidth, double& OutMajorInterval, int32& OutMinorDivisions) const
 {
-	FSlateFontInfo SmallLayoutFont = FCoreStyle::GetDefaultFontStyle("Regular", 8);
-	TSharedRef<FSlateFontMeasure> FontMeasureService = FSlateApplication::Get().GetRenderer()->GetFontMeasureService();
+	const FSlateFontInfo SmallLayoutFont = FCoreStyle::GetDefaultFontStyle("Regular", 8);
+	const TSharedRef<FSlateFontMeasure> FontMeasureService = FSlateApplication::Get().GetRenderer()->GetFontMeasureService();
 
-	FFrameRate DisplayRate(FMath::Max(FMath::RoundToInt(Model.Pin()->GetFrameRate()), 1), 1);
-	double BiggestTime = ViewRange.Get().GetUpperBoundValue();
-	FString TickString = NumericTypeInterface->ToString((BiggestTime * DisplayRate).FrameNumber.Value);
-	FVector2D MaxTextSize = FontMeasureService->Measure(TickString, SmallLayoutFont);
+	const FFrameRate DisplayRate(static_cast<int32>(FMath::Max(FMath::RoundToInt(Model.Pin()->GetFrameRate()), 1)), 1);
+	const double BiggestTime = ViewRange.Get().GetUpperBoundValue();
+	const FString TickString = NumericTypeInterface->ToString((BiggestTime * DisplayRate).FrameNumber.Value);
+	const FVector2D MaxTextSize = FontMeasureService->Measure(TickString, SmallLayoutFont);
 
-	static float MajorTickMultiplier = 2.f;
+	constexpr float MajorTickMultiplier = 2.f;
 
-	float MinTickPx = MaxTextSize.X + 5.f;
-	float DesiredMajorTickPx = MaxTextSize.X * MajorTickMultiplier;
+	const float MinTickPx = static_cast<float>(MaxTextSize.X) + 5.f;
+	const float DesiredMajorTickPx = static_cast<float>(MaxTextSize.X) * MajorTickMultiplier;
 
 	if (PhysicalWidth > 0 && DisplayRate.AsDecimal() > 0)
 	{
 		return ComputeGridSpacing(
 			DisplayRate,
-			PhysicalWidth / ViewRange.Get().Size<double>(),
+			static_cast<float>(PhysicalWidth / ViewRange.Get().Size<double>()),
 			OutMajorInterval,
 			OutMinorDivisions,
 			MinTickPx,
@@ -805,7 +805,7 @@ class UAnimSingleNodeInstance* SAnimTimeline::GetPreviewInstance() const
 	return PreviewMeshComponent && PreviewMeshComponent->IsPreviewOn()? PreviewMeshComponent->PreviewInstance : nullptr;
 }
 
-void SAnimTimeline::HandleScrubPositionChanged(FFrameTime NewScrubPosition, bool bIsScrubbing, bool bEvaluate)
+void SAnimTimeline::HandleScrubPositionChanged(FFrameTime NewScrubPosition, bool bIsScrubbing, bool bEvaluate) const
 {
 	if (UAnimSingleNodeInstance* PreviewInstance = GetPreviewInstance())
 	{
@@ -820,7 +820,7 @@ void SAnimTimeline::HandleScrubPositionChanged(FFrameTime NewScrubPosition, bool
 
 double SAnimTimeline::GetSpinboxDelta() const
 {
-	return FFrameRate(Model.Pin()->GetTickResolution(), 1).AsDecimal() * FFrameRate(FMath::RoundToInt(Model.Pin()->GetFrameRate()), 1).AsInterval();
+	return FFrameRate(Model.Pin()->GetTickResolution(), 1).AsDecimal() * FFrameRate(static_cast<int32>(FMath::RoundToInt(Model.Pin()->GetFrameRate())), 1).AsInterval();
 }
 
 void SAnimTimeline::SetPlayTime(double InFrameTime)
@@ -828,7 +828,7 @@ void SAnimTimeline::SetPlayTime(double InFrameTime)
 	if (UAnimSingleNodeInstance* PreviewInstance = GetPreviewInstance())
 	{
 		PreviewInstance->SetPlaying(false);
-		PreviewInstance->SetPosition(InFrameTime / (double)Model.Pin()->GetTickResolution());
+		PreviewInstance->SetPosition(static_cast<float>(InFrameTime / static_cast<double>(Model.Pin()->GetTickResolution())));
 	}
 }
 
diff --git a/Engine/Source/Editor/Persona/Private/AnimTimeline/SAnimTimeline.h b/Engine/Source/Editor/Persona/Private/AnimTimeline/SAnimTimeline.h
--- a/Engine/Source/Editor/Persona/Private/AnimTimeline/SAnimTimeline.h
+++ b/Engine/Source/Editor/Persona/Private/AnimTimeline/SAnimTimeline.h
@@ -78,7 +78,7 @@ private:
 
 
 
-	void HandleScrubPositionChanged(FFrameTime NewScrubPosition, bool bIsScrubbing, bool bEvaluate);
+	void HandleScrubPositionChanged(FFrameTime NewScrubPosition, bool bIsScrubbing, bool bEvaluate) const;
 
 	void OnCropAnimSequence(bool bFromStart, float CurrentTime);
 
diff --git a/Engine/Source/Editor/Persona/Private/AnimTimeline/SAnimTimelineSplitterOverlay.cpp b/Engine/Source/Editor/Persona/Private/AnimTimeline/SAnimTimelineSplitterOverlay.cpp
--- a/Engine/Source/Editor/Persona/Private/AnimTimeline/SAnimTimelineSplitterOverlay.cpp
+++ b/Engine/Source/Editor/Persona/Private/AnimTimeline/SAnimTimelineSplitterOverlay.cpp
@@ -56,9 +56,9 @@ void SAnimTimelineSplitterOverlay::OnArrangeChildren( const FGeometry& AllottedG
 		if (Splitter->GetOrientation() == EOrientation::Orient_Horizontal)
 		{
 			SlotPadding.Add(FMargin(
-				ThisGeometry.Position.X + ThisGeometry.GetLocalSize().X,
+				ThisGeometry.Position.X + static_cast<float>(ThisGeometry.GetLocalSize().X),
 				0,
-				AllottedGeometry.Size.X - NextGeometry.Position.X,
+				static_cast<float>(AllottedGeometry.Size.X) - NextGeometry.Position.X,
 				0)
 			);
 		}
@@ -66,9 +66,9 @@ void SAnimTimelineSplitterOverlay::OnArrangeChildren( const FGeometry& AllottedG
 		{
 			SlotPadding.Add(FMargin(
 				0,
-				ThisGeometry.Position.Y + ThisGeometry.GetLocalSize().Y,
+				ThisGeometry.Position.Y + static_cast<float>(ThisGeometry.GetLocalSize().Y),
 				0,
-				AllottedGeometry.Size.Y - NextGeometry.Position.Y)
+				static_cast<float>(AllottedGeometry.Size.Y) - NextGeometry.Position.Y)
 			);
 		}
 	}
diff --git a/Engine/Source/Editor/Persona/Private/AnimTimeline/SAnimTimelineTransportControls.cpp b/Engine/Source/Editor/Persona/Private/AnimTimeline/SAnimTimelineTransportControls.cpp
--- a/Engine/Source/Editor/Persona/Private/AnimTimeline/SAnimTimelineTransportControls.cpp
+++ b/Engine/Source/Editor/Persona/Private/AnimTimeline/SAnimTimelineTransportControls.cpp
@@ -62,12 +62,12 @@ FReply SAnimTimelineTransportControls::OnClick_Forward_Step()
 	}
 	else if (SMC)
 	{
-		UAnimSequence* AnimSequence = Cast<UAnimSequence>(AnimSequenceBase);
+		const UAnimSequence* AnimSequence = Cast<UAnimSequence>(AnimSequenceBase);
 		const FFrameRate TargetFramerate = AnimSequence ? AnimSequence->GetSamplingFrameRate() : FFrameRate(30, 1);
 
 		// Advance a single frame, leaving it paused afterwards
 		SMC->GlobalAnimRateScale = 1.0f;
-		SMC->TickAnimation(TargetFramerate.AsInterval(), false);
+		SMC->TickAnimation(static_cast<float>(TargetFramerate.AsInterval()), false);
 		SMC->GlobalAnimRateScale = 0.0f;
 	}
 
diff --git a/Engine/Source/Editor/Persona/Private/AnimTimeline/SAnimTrackArea.cpp b/Engine/Source/Editor/Persona/Private/AnimTimeline/SAnimTrackArea.cpp
--- a/Engine/Source/Editor/Persona/Private/AnimTimeline/SAnimTrackArea.cpp
+++ b/Engine/Source/Editor/Persona/Private/AnimTimeline/SAnimTrackArea.cpp
@@ -76,8 +76,8 @@ void SAnimTrackArea::OnArrangeChildren(const FGeometry& AllottedGeometry, FArran
 
 		const FMargin Padding(0, CurChild.GetVerticalOffset(), 0, 0);
 
-		AlignmentArrangeResult XResult = AlignChild<Orient_Horizontal>(AllottedGeometry.GetLocalSize().X, CurChild, Padding, 1.0f, false);
-		AlignmentArrangeResult YResult = AlignChild<Orient_Vertical>(AllottedGeometry.GetLocalSize().Y, CurChild, Padding, 1.0f, false);
+		const AlignmentArrangeResult XResult = AlignChild<Orient_Horizontal>(static_cast<float>(AllottedGeometry.GetLocalSize().X), CurChild, Padding, 1.0f, false);
+		const AlignmentArrangeResult YResult = AlignChild<Orient_Vertical>(static_cast<float>(AllottedGeometry.GetLocalSize().Y), CurChild, Padding, 1.0f, false);
 
 		ArrangedChildren.AddWidget(ChildVisibility,
 			AllottedGeometry.MakeChild(
@@ -168,7 +168,7 @@ FReply SAnimTrackArea::OnMouseMove( const FGeometry& MyGeometry, const FPointerE
 {
 	UpdateHoverStates(MyGeometry, MouseEvent);
 
-	TSharedPtr<FAnimTimeSliderController> TimeSliderController = WeakTimeSliderController.Pin();
+	const TSharedPtr<FAnimTimeSliderController> TimeSliderController = WeakTimeSliderController.Pin();
 	if(TimeSliderController.IsValid())
 	{
 		FReply Reply = WeakTimeSliderController.Pin()->OnMouseMove(*this, MyGeometry, MouseEvent);
@@ -178,7 +178,7 @@ FReply SAnimTrackArea::OnMouseMove( const FGeometry& MyGeometry, const FPointerE
 		{
 			if (MouseEvent.IsMouseButtonDown(EKeys::RightMouseButton) && HasMouseCapture())
 			{
-				WeakOutliner.Pin()->ScrollByDelta(-MouseEvent.GetCursorDelta().Y);
+				WeakOutliner.Pin()->ScrollByDelta(static_cast<float>(-MouseEvent.GetCursorDelta().Y));
 			}
 		}
 
diff --git a/Engine/Source/Editor/Persona/Private/AnimTimeline/SAnimTrackResizeArea.cpp b/Engine/Source/Editor/Persona/Private/AnimTimeline/SAnimTrackResizeArea.cpp
--- a/Engine/Source/Editor/Persona/Private/AnimTimeline/SAnimTrackResizeArea.cpp
+++ b/Engine/Source/Editor/Persona/Private/AnimTimeline/SAnimTrackResizeArea.cpp
@@ -17,10 +17,10 @@ void SAnimTrackResizeArea::Construct(const FArguments& InArgs, TWeakPtr<FAnimTim
 
 FReply SAnimTrackResizeArea::OnMouseButtonDown(const FGeometry& MyGeometry, const FPointerEvent& MouseEvent)
 {
-	TSharedPtr<FAnimTimelineTrack> ResizeTarget = Track.Pin();
+	const TSharedPtr<FAnimTimelineTrack> ResizeTarget = Track.Pin();
 	if (ResizeTarget.IsValid() && MouseEvent.IsMouseButtonDown(EKeys::LeftMouseButton))
 	{
-		DragParameters = FDragParameters(ResizeTarget->GetHeight(), MouseEvent.GetScreenSpacePosition().Y);
+		DragParameters = FDragParameters(ResizeTarget->GetHeight(), static_cast<float>(MouseEvent.GetScreenSpacePosition().Y));
 		return FReply::Handled().CaptureMouse(AsShared());
 	}
 	return FReply::Handled();
@@ -36,9 +36,9 @@ FReply SAnimTrackResizeArea::OnMouseMove(const FGeometry& MyGeometry, const FPoi
 {
 	if (DragParameters.IsSet() && HasMouseCapture())
 	{
-		float NewHeight = DragParameters->OriginalHeight + (MouseEvent.GetScreenSpacePosition().Y - DragParameters->DragStartY);
+		const float NewHeight = DragParameters->OriginalHeight + (static_cast<float>(MouseEvent.GetScreenSpacePosition().Y) - DragParameters->DragStartY);
 
-		TSharedPtr<FAnimTimelineTrack> ResizeTarget = Track.Pin();
+		const TSharedPtr<FAnimTimelineTrack> ResizeTarget = Track.Pin();
 		if (ResizeTarget.IsValid() && FMath::RoundToInt(NewHeight) != FMath::RoundToInt(DragParameters->OriginalHeight))
 		{
 			ResizeTarget->SetHeight(NewHeight);
diff --git a/Engine/Source/Editor/Persona/Private/AnimationEditorViewportClient.cpp b/Engine/Source/Editor/Persona/Private/AnimationEditorViewportClient.cpp
--- a/Engine/Source/Editor/Persona/Private/AnimationEditorViewportClient.cpp
+++ b/Engine/Source/Editor/Persona/Private/AnimationEditorViewportClient.cpp
@@ -630,7 +630,7 @@ void FAnimationViewportClient::DrawCanvas( FViewport& InViewport, FSceneView& Vi
 
 		if (bDrawUVs)
 		{
-			DrawUVsForMesh(Viewport, &Canvas, 1.0f, PreviewMeshComponent);
+			DrawUVsForMesh(Viewport, &Canvas, 1, PreviewMeshComponent);
 		}
 
 		// Debug draw clothing texts
@@ -791,8 +791,8 @@ void FAnimationViewportClient::ShowBoneNames( FCanvas* Canvas, FSceneView* View,
 		return;
 	}
 
-	const int32 HalfX = Viewport->GetSizeXY().X/2 / GetDPIScale();
-	const int32 HalfY = Viewport->GetSizeXY().Y/2 / GetDPIScale();
+	const int32 HalfX = static_cast<int32>(Viewport->GetSizeXY().X / 2 / GetDPIScale());
+	const int32 HalfY = static_cast<int32>(Viewport->GetSizeXY().Y / 2 / GetDPIScale());
 
 	for (int32 i=0; i< LODData.RequiredBones.Num(); i++)
 	{
@@ -840,8 +840,8 @@ void FAnimationViewportClient::ShowBoneNames( FCanvas* Canvas, FSceneView* View,
 			const FPlane proj = View->Project(BonePos);
 			if (proj.W > 0.f)
 			{
-				const int32 XPos = HalfX + ( HalfX * proj.X );
-				const int32 YPos = HalfY + ( HalfY * (proj.Y * -1) );
+				const int32 XPos = static_cast<int32>(HalfX + ( HalfX * proj.X ));
+				const int32 YPos = static_cast<int32>(HalfY + ( HalfY * (proj.Y * -1) ));
 
 				const FName BoneName = ReferenceSkeleton.GetBoneName(BoneIndex);
 				const FString BoneString = FString::Printf( TEXT("%d: %s"), BoneIndex, *BoneName.ToString() );
@@ -860,8 +860,8 @@ void FAnimationViewportClient::ShowAttributeNames(FCanvas* Canvas, FSceneView* V
 		return;
 	}
 	
-	const int32 HalfX = Viewport->GetSizeXY().X / 2 / GetDPIScale();
-	const int32 HalfY = Viewport->GetSizeXY().Y / 2 / GetDPIScale();
+	const int32 HalfX = static_cast<int32>(Viewport->GetSizeXY().X / 2 / GetDPIScale());
+	const int32 HalfY = static_cast<int32>(Viewport->GetSizeXY().Y / 2 / GetDPIScale());
 
 	const UE::Anim::FMeshAttributeContainer& Attributes = MeshComponent->GetCustomAttributes();
 
@@ -884,8 +884,8 @@ void FAnimationViewportClient::ShowAttributeNames(FCanvas* Canvas, FSceneView* V
 				const FPlane proj = View->Project(AttributeTransform.GetLocation());
 				if (proj.W > 0.f)
 				{
-					const int32 XPos = HalfX + (HalfX * proj.X);
-					const int32 YPos = HalfY + (HalfY * (proj.Y * -1));
+					const int32 XPos = HalfX + static_cast<int32>(HalfX * proj.X);
+					const int32 YPos = HalfY + static_cast<int32>(HalfY * (proj.Y * -1));
 
 					FCanvasTextItem TextItem(FVector2D(XPos, YPos), FText::FromName(AttributeIdentifier.GetName()), GEngine->GetSmallFont(), FLinearColor(0.0f, 1.0f, 1.0f));
 					TextItem.EnableShadow(FLinearColor::Black);
@@ -1794,7 +1794,7 @@ void FAnimationViewportClient::DrawMeshSubsetBones(const UDebugSkelMeshComponent
 		if (bDrawBone)
 		{
 			//add to the list
-			RequiredBones.AddUnique(BoneIndex);
+			RequiredBones.AddUnique(static_cast<FBoneIndexType>(BoneIndex));
 			WorldTransforms[BoneIndex] = MeshComponent->GetDrawTransform(BoneIndex) * MeshComponent->GetComponentTransform();
 		}
 	}
@@ -1885,14 +1885,6 @@ void FAnimationViewportClient::DrawSockets(const UDebugSkelMeshComponent* InPrev
 				SocketColor = (bUseSkeletonSocketColor) ? FLinearColor::White : FLinearColor::Red;
 			}
 
-			static const float SphereRadius = 1.0f;
-			TArray<FVector> Verts;
-
-			//Calc cone size 
-			FVector EndToStart = (Start-End);
-			float ConeLength = EndToStart.Size();
-			float Angle = FMath::RadiansToDegrees(FMath::Atan(SphereRadius / ConeLength));
-
 			//Render Sphere for bone end point and a cone between it and its parent.
 			PDI->DrawLine( Start, End, SocketColor, SDPG_Foreground );
 			
@@ -2335,11 +2327,11 @@ void FAnimationViewportClient::SetupViewForRendering( FSceneViewFamily& ViewFami
 	}
 
 	// Cache screen size
-	UDebugSkelMeshComponent* PreviewMeshComponent = GetAnimPreviewScene()->GetPreviewMeshComponent();
+	const UDebugSkelMeshComponent* PreviewMeshComponent = GetAnimPreviewScene()->GetPreviewMeshComponent();
 	if (PreviewMeshComponent != nullptr && PreviewMeshComponent->MeshObject != nullptr)
 	{
 		const FBoxSphereBounds& SkelBounds = PreviewMeshComponent->Bounds;
-		CachedScreenSize = ComputeBoundsScreenSize(SkelBounds.Origin, SkelBounds.SphereRadius, View);
+		CachedScreenSize = ComputeBoundsScreenSize(SkelBounds.Origin, static_cast<float>(SkelBounds.SphereRadius), View);
 	}
 }
 
diff --git a/Engine/Source/Editor/Persona/Private/BoneMappingHelper.cpp b/Engine/Source/Editor/Persona/Private/BoneMappingHelper.cpp
--- a/Engine/Source/Editor/Persona/Private/BoneMappingHelper.cpp
+++ b/Engine/Source/Editor/Persona/Private/BoneMappingHelper.cpp
@@ -27,10 +27,10 @@ int32 FBoneDescription::GetBestIndex() const
 
 float FBoneDescription::CalculateNameScore(const FName& Name1, const FName& Name2) const
 {
-	FString String1 = Name1.ToString();
-	FString String2 = Name2.ToString();
-	int32 Distance = FAnimationRuntime::GetStringDistance(String1, String2);
-	float MaxLength = FMath::Max(String1.Len() , String2.Len());
+	const FString String1 = Name1.ToString();
+	const FString String2 = Name2.ToString();
+	const int32 Distance = FAnimationRuntime::GetStringDistance(String1, String2);
+	const float MaxLength = static_cast<float>(FMath::Max(String1.Len() , String2.Len()));
 	return (MaxLength) ? (MaxLength - Distance) / (MaxLength) : 1.f;
 }
 
@@ -42,18 +42,15 @@ float FBoneDescription::CalculateScore(const FBoneDescription& Other) const
 		return 1.f;
 	}
 
-	// calculating score is tricky
-	float Score = 0.f;
-
 	// each element will exit from [0, 1], and we'll apply weight to it
 
 	// check direction of facing [-1, 1]
-	float Cosine = FVector::DotProduct(DirFromParent, Other.DirFromParent);
+	float Cosine = static_cast<float>(FVector::DotProduct(DirFromParent, Other.DirFromParent));
 	float Score_DirFromParent = FMath::Max((Cosine-0.5f)*2.f, 0.f); // scale to only care for range of [0.5-1]
 	Score_DirFromParent = FMath::Clamp(Score_DirFromParent, 0.f, 1.f);
 
 	// check direction of facing [-1, 1]
-	Cosine = FVector::DotProduct(DirFromRoot, Other.DirFromRoot);
+	Cosine = static_cast<float>(FVector::DotProduct(DirFromRoot, Other.DirFromRoot));
 	float Score_DirFromRoot = FMath::Max((Cosine - 0.5f)*2.f, 0.f); // scale to only care for range of [0.5-1]
 	Score_DirFromRoot = FMath::Clamp(Score_DirFromRoot, 0.f, 1.f);
 
@@ -62,7 +59,7 @@ float FBoneDescription::CalculateScore(const FBoneDescription& Other) const
 	const int32 MaxNumChildren = FMath::Max(Other.NumChildren, NumChildren);
 	if (MaxNumChildren > 0)
 	{
-		Score_NumChildren = (1.f - ((float)(FMath::Abs(Other.NumChildren - NumChildren)) / MaxNumChildren));
+		Score_NumChildren = 1.f - static_cast<float>(FMath::Abs(Other.NumChildren - NumChildren) / MaxNumChildren);
 	}
 	else
 	{
@@ -86,9 +83,9 @@ float FBoneDescription::CalculateScore(const FBoneDescription& Other) const
 	Score_RatioFromParent = FMath::Clamp(Score_RatioFromParent, 0.f, 1.f);
 
 	// check normalized position - since this is normalized, it should stay within 1
-	FVector DiffNormalizedPosition = (Other.NormalizedPosition - NormalizedPosition).GetAbs();
-	const float MaxNormalizedPosition = 3.f; /* since 1^2+1^2+1^2 = 3*/
-	float Score_NormalizedPosition = (MaxNormalizedPosition - DiffNormalizedPosition.SizeSquared())/ MaxNormalizedPosition;
+	const FVector DiffNormalizedPosition = (Other.NormalizedPosition - NormalizedPosition).GetAbs();
+	constexpr float MaxNormalizedPosition = 3.f; /* since 1^2+1^2+1^2 = 3*/
+	float Score_NormalizedPosition = (MaxNormalizedPosition - static_cast<float>(DiffNormalizedPosition.SizeSquared()))/ MaxNormalizedPosition;
 	Score_NormalizedPosition = FMath::Clamp(Score_NormalizedPosition, 0.f, 1.f);
 
 	float Score_NameMatching = CalculateNameScore(BoneInfo.Name, Other.BoneInfo.Name);
@@ -102,7 +99,7 @@ float FBoneDescription::CalculateScore(const FBoneDescription& Other) const
 	static float Weight_NameMatching = 2.0f;
 	static float Weight_DirFromRoot = 0.f;
 
-	float FinalScore = (Score_DirFromParent * Weight_DirFromParent + Score_NumChildren * Weight_NumChildren + Score_NormalizedPosition * Weight_NormalizedPosition + Score_RatioFromParent * Weight_RatioFromParent + Score_NameMatching * Weight_NameMatching + Score_DirFromRoot * Weight_DirFromRoot) / (Weight_DirFromParent + Weight_NumChildren + Weight_NormalizedPosition + Weight_RatioFromParent + Weight_NameMatching + Weight_DirFromRoot);
+	const float FinalScore = (Score_DirFromParent * Weight_DirFromParent + Score_NumChildren * Weight_NumChildren + Score_NormalizedPosition * Weight_NormalizedPosition + Score_RatioFromParent * Weight_RatioFromParent + Score_NameMatching * Weight_NameMatching + Score_DirFromRoot * Weight_DirFromRoot) / (Weight_DirFromParent + Weight_NumChildren + Weight_NormalizedPosition + Weight_RatioFromParent + Weight_NameMatching + Weight_DirFromRoot);
 	UE_LOG(LogAnimation, Log, TEXT("Calculate Score - [%s] - [%s] (Score_DirFromParent(%0.2f), Score_NumChildren(%0.2f), Score_NormalizedPosition(%0.2f), Score_RatioFromParent(%0.2f), Score_NameMatching(%0.2f), Score_DirFromRoot(%0.2f) )"), *BoneInfo.Name.ToString(), *Other.BoneInfo.Name.ToString(), Score_DirFromParent, Score_NumChildren, Score_NormalizedPosition, Score_RatioFromParent, Score_NameMatching, Score_DirFromRoot);
 	return FinalScore;
 }
@@ -167,7 +164,7 @@ void FBoneMappingHelper::Initialize(int32 Index, const FReferenceSkeleton&  InRe
 
 					const FVector& ParentPosition = ComponentSpaceTransforms[ParentIndex].GetLocation();
 					FVector ToChild = (Position-ParentPosition);
-					BoneDesc.RatioFromParent = ToChild.Size() / MeshBoxSize.Size();
+					BoneDesc.RatioFromParent = static_cast<float>(ToChild.Size() / MeshBoxSize.Size());
 					BoneDesc.DirFromParent = ToChild.GetSafeNormal();
 					BoneDesc.DirFromRoot = (Position-RootPosition).GetSafeNormal(); // based on whole mesh size
 				}
diff --git a/Engine/Source/Editor/Persona/Private/Customization/AnimGraphNodeSlotDetails.cpp b/Engine/Source/Editor/Persona/Private/Customization/AnimGraphNodeSlotDetails.cpp
--- a/Engine/Source/Editor/Persona/Private/Customization/AnimGraphNodeSlotDetails.cpp
+++ b/Engine/Source/Editor/Persona/Private/Customization/AnimGraphNodeSlotDetails.cpp
@@ -73,13 +73,13 @@ void FAnimGraphNodeSlotDetails::CustomizeDetails(class IDetailLayoutBuilder& Det
 		DetailBuilder.HideProperty(SlotNodeNamePropertyHandle);
 		
 		// this is used for 2 things, to create name, but also for another pop-up box to show off, it's a bit hacky to use this, but I need widget to parent on
-		TSharedRef<SWidget> SlotNodePropertyNameWidget = SlotNodeNamePropertyHandle->CreatePropertyNameWidget();
+		const TSharedRef<SWidget> SlotNodePropertyNameWidget = SlotNodeNamePropertyHandle->CreatePropertyNameWidget();
 
 		// fill combo with groups and slots
 		RefreshComboLists();
 
-		int32 FoundIndex = SlotNameList.Find(SlotNameComboSelectedName);
-		TSharedPtr<FString> ComboBoxSelectedItem = SlotNameComboListItems[FoundIndex];
+		const int32 FoundIndex = SlotNameList.Find(SlotNameComboSelectedName);
+		const TSharedPtr<FString> ComboBoxSelectedItem = SlotNameComboListItems[FoundIndex];
 
 		// create combo box
 		IDetailCategoryBuilder& SlotNameGroup = DetailBuilder.EditCategory(TEXT("Settings"));
@@ -102,7 +102,7 @@ void FAnimGraphNodeSlotDetails::CustomizeDetails(class IDetailLayoutBuilder& Det
 				.OnSelectionChanged(this, &FAnimGraphNodeSlotDetails::OnSlotNameChanged)
 				.OnComboBoxOpening(this, &FAnimGraphNodeSlotDetails::OnSlotListOpening)
 				.InitiallySelectedItem(ComboBoxSelectedItem)
-				.ContentPadding(2)
+				.ContentPadding(2.f)
 				.ToolTipText(FText::FromString(*ComboBoxSelectedItem))
 			]
 
diff --git a/Engine/Source/Editor/Persona/Private/Customization/BlendSampleDetails.cpp b/Engine/Source/Editor/Persona/Private/Customization/BlendSampleDetails.cpp
--- a/Engine/Source/Editor/Persona/Private/Customization/BlendSampleDetails.cpp
+++ b/Engine/Source/Editor/Persona/Private/Customization/BlendSampleDetails.cpp
@@ -415,7 +415,7 @@ void FBlendSampleDetails::GenerateBlendSampleWidget(TFunction<FDetailWidgetRow&
 			{
 				if (BlendSpace)
 				{
-					return BlendSpace->IsValidBlendSampleIndex(SampleIndex) ? BlendSpace->GetBlendSample(SampleIndex).SampleValue[ParameterIndex] : 0.0f;
+					return BlendSpace->IsValidBlendSampleIndex(SampleIndex) ? static_cast<float>(BlendSpace->GetBlendSample(SampleIndex).SampleValue[ParameterIndex]) : 0.0f;
 				}
 
 				return 0.0f;
diff --git a/Engine/Source/Editor/Persona/Private/Customization/BlendSpaceDetails.cpp b/Engine/Source/Editor/Persona/Private/Customization/BlendSpaceDetails.cpp
--- a/Engine/Source/Editor/Persona/Private/Customization/BlendSpaceDetails.cpp
+++ b/Engine/Source/Editor/Persona/Private/Customization/BlendSpaceDetails.cpp
@@ -289,7 +289,7 @@ void FBlendSpaceDetails::CustomizeDetails(class IDetailLayoutBuilder& DetailBuil
 						.HAlign(EHorizontalAlignment::HAlign_Left)
 						.ToolTipText(LOCTEXT("AnalyzeAllSamples", "Analyze all samples"))
 						.OnClicked(this, &FBlendSpaceDetails::HandleAnalyzeSamples)
-						.ContentPadding(1)
+						.ContentPadding(1.f)
 						[
 							SNew(SImage)
 							.Image(FAppStyle::GetBrush("Icons.Refresh"))
@@ -427,7 +427,7 @@ void FBlendSpaceDetails::CustomizeDetails(class IDetailLayoutBuilder& DetailBuil
 						.HAlign(EHorizontalAlignment::HAlign_Left)
 						.ToolTipText(LOCTEXT("RemoveAllSamples", "Remove all samples"))
 						.OnClicked(this, &FBlendSpaceDetails::HandleClearSamples)
-						.ContentPadding(1)
+						.ContentPadding(1.f)
 						[
 							SNew(SImage)
 							.Image(FAppStyle::GetBrush("Icons.Delete"))
diff --git a/Engine/Source/Editor/Persona/Private/EditModes/SkeletonSelectionEditMode.cpp b/Engine/Source/Editor/Persona/Private/EditModes/SkeletonSelectionEditMode.cpp
--- a/Engine/Source/Editor/Persona/Private/EditModes/SkeletonSelectionEditMode.cpp
+++ b/Engine/Source/Editor/Persona/Private/EditModes/SkeletonSelectionEditMode.cpp
@@ -333,19 +333,19 @@ FIntPoint FSkeletonSelectionEditMode::GetDPIUnscaledSize(FViewport* Viewport, FV
 	const FIntPoint Size = Viewport->GetSizeXY();
 	const float DPIScale = Client->GetDPIScale();
 	// (FIntPoint / float) implicitly casts the float to an int if you try to divide it directly
-	return FIntPoint(Size.X / DPIScale, Size.Y / DPIScale);
+	return FIntPoint(static_cast<int32>(Size.X / DPIScale), static_cast<int32>(Size.Y / DPIScale));
 }
 
 void FSkeletonSelectionEditMode::DrawHUD(FEditorViewportClient* ViewportClient, FViewport* Viewport, const FSceneView* View, FCanvas* Canvas)
 {
-	UDebugSkelMeshComponent* PreviewMeshComponent = GetAnimPreviewScene().GetPreviewMeshComponent();
+	const UDebugSkelMeshComponent* PreviewMeshComponent = GetAnimPreviewScene().GetPreviewMeshComponent();
 	if (PreviewMeshComponent == nullptr || PreviewMeshComponent->GetSkeletalMeshAsset() == nullptr || PreviewMeshComponent->GetSkeletalMeshAsset()->IsCompiling())
 	{
 		return;
 	}
 
-	FReferenceSkeleton& RefSkeleton = PreviewMeshComponent->GetSkeletalMeshAsset()->GetRefSkeleton();
-	int32 BoneIndex = GetAnimPreviewScene().GetSelectedBoneIndex();
+	const FReferenceSkeleton& RefSkeleton = PreviewMeshComponent->GetSkeletalMeshAsset()->GetRefSkeleton();
+	const int32 BoneIndex = GetAnimPreviewScene().GetSelectedBoneIndex();
 
 	// Draw name of selected bone
 	if (RefSkeleton.IsValidIndex(BoneIndex) && (IsSelectedBoneRequired() || RefSkeleton.GetRequiredVirtualBones().Contains(BoneIndex)))
@@ -356,15 +356,13 @@ void FSkeletonSelectionEditMode::DrawHUD(FEditorViewportClient* ViewportClient,
 
 		const FName BoneName = PreviewMeshComponent->GetSkeletalMeshAsset()->GetRefSkeleton().GetBoneName(BoneIndex);
 
-		FMatrix BoneMatrix = PreviewMeshComponent->GetBoneMatrix(BoneIndex);
+		const FMatrix BoneMatrix = PreviewMeshComponent->GetBoneMatrix(BoneIndex);
 		const FPlane Proj = View->Project(BoneMatrix.GetOrigin());
 		if (Proj.W > 0.f)
 		{
-			const int32 XPos = HalfX + (HalfX * Proj.X);
-			const int32 YPos = HalfY + (HalfY * (Proj.Y * -1));
+			const int32 XPos = HalfX + static_cast<int32>(HalfX * Proj.X);
+			const int32 YPos = HalfY + static_cast<int32>(HalfY * Proj.Y * -1);
 
-			FQuat BoneQuat = PreviewMeshComponent->GetBoneQuaternion(BoneName);
-			FVector Loc = PreviewMeshComponent->GetBoneLocation(BoneName);
 			FCanvasTextItem TextItem(FVector2D(XPos, YPos), FText::FromString(BoneName.ToString()), GEngine->GetSmallFont(), FLinearColor::White);
 			TextItem.EnableShadow(FLinearColor::Black);
 			Canvas->DrawItem(TextItem);
@@ -374,7 +372,7 @@ void FSkeletonSelectionEditMode::DrawHUD(FEditorViewportClient* ViewportClient,
 	// Draw name of selected socket
 	if (GetAnimPreviewScene().GetSelectedSocket().IsValid())
 	{
-		USkeletalMeshSocket* Socket = GetAnimPreviewScene().GetSelectedSocket().Socket;
+		const USkeletalMeshSocket* Socket = GetAnimPreviewScene().GetSelectedSocket().Socket;
 
 		FMatrix SocketMatrix;
 		Socket->GetSocketMatrix(SocketMatrix, PreviewMeshComponent);
@@ -387,8 +385,8 @@ void FSkeletonSelectionEditMode::DrawHUD(FEditorViewportClient* ViewportClient,
 			const int32 HalfX = ViewPortSize.X / 2;
 			const int32 HalfY = ViewPortSize.Y / 2;
 
-			const int32 XPos = HalfX + (HalfX * Proj.X);
-			const int32 YPos = HalfY + (HalfY * (Proj.Y * -1));
+			const int32 XPos = HalfX + static_cast<int32>(HalfX * Proj.X);
+			const int32 YPos = HalfY + static_cast<int32>(HalfY * (Proj.Y * -1));
 			FCanvasTextItem TextItem(FVector2D(XPos, YPos), FText::FromString(Socket->SocketName.ToString()), GEngine->GetSmallFont(), FLinearColor::White);
 			TextItem.EnableShadow(FLinearColor::Black);
 			Canvas->DrawItem(TextItem);
@@ -403,19 +401,19 @@ bool FSkeletonSelectionEditMode::AllowWidgetMove()
 
 bool FSkeletonSelectionEditMode::IsSelectedBoneRequired() const
 {
-	UDebugSkelMeshComponent* PreviewMeshComponent = GetAnimPreviewScene().GetPreviewMeshComponent();
-	int32 SelectedBoneIndex = GetAnimPreviewScene().GetSelectedBoneIndex();
+	const UDebugSkelMeshComponent* PreviewMeshComponent = GetAnimPreviewScene().GetPreviewMeshComponent();
+	const int32 SelectedBoneIndex = GetAnimPreviewScene().GetSelectedBoneIndex();
 	if (SelectedBoneIndex != INDEX_NONE && PreviewMeshComponent->GetSkeletalMeshRenderData())
 	{
 		//Get current LOD
-		FSkeletalMeshRenderData* SkelMeshRenderData = PreviewMeshComponent->GetSkeletalMeshRenderData();
+		const FSkeletalMeshRenderData* SkelMeshRenderData = PreviewMeshComponent->GetSkeletalMeshRenderData();
 		if(SkelMeshRenderData->LODRenderData.Num() > 0)
 		{
 			const int32 LODIndex = FMath::Clamp(PreviewMeshComponent->GetPredictedLODLevel(), 0, SkelMeshRenderData->LODRenderData.Num() - 1);
-			FSkeletalMeshLODRenderData& LODData = SkelMeshRenderData->LODRenderData[LODIndex];
+			const FSkeletalMeshLODRenderData& LODData = SkelMeshRenderData->LODRenderData[LODIndex];
 
 			//Check whether the bone is vertex weighted
-			return LODData.RequiredBones.Find(SelectedBoneIndex) != INDEX_NONE;
+			return LODData.RequiredBones.Find(static_cast<FBoneIndexType>(SelectedBoneIndex)) != INDEX_NONE;
 		}
 	}
 
diff --git a/Engine/Source/Editor/Persona/Private/PersonaMeshDetails.cpp b/Engine/Source/Editor/Persona/Private/PersonaMeshDetails.cpp
--- a/Engine/Source/Editor/Persona/Private/PersonaMeshDetails.cpp
+++ b/Engine/Source/Editor/Persona/Private/PersonaMeshDetails.cpp
@@ -488,7 +488,7 @@ void FSkeletalMeshReductionSettingsLayout::GenerateChildContent(IDetailChildrenB
 			.IsEnabled(this, &FSkeletalMeshReductionSettingsLayout::IsReductionEnabled)
 			.OnGetMenuContent(this, &FSkeletalMeshReductionSettingsLayout::FillReductionMethodMenu)
 			.VAlign(VAlign_Center)
-			.ContentPadding(2)
+			.ContentPadding(2.f)
 			.ButtonContent()
 			[
 				SNew(STextBlock)
@@ -537,7 +537,7 @@ void FSkeletalMeshReductionSettingsLayout::GenerateChildContent(IDetailChildrenB
 				.IsEnabled(this, &FSkeletalMeshReductionSettingsLayout::IsReductionEnabled)
 				.OnGetMenuContent(this, &FSkeletalMeshReductionSettingsLayout::FillReductionImportanceMenu, ImportanceType)
 				.VAlign(VAlign_Center)
-				.ContentPadding(2)
+				.ContentPadding(2.f)
 				.ButtonContent()
 				[
 					SNew(STextBlock)
@@ -615,7 +615,7 @@ void FSkeletalMeshReductionSettingsLayout::GenerateChildContent(IDetailChildrenB
 			.IsEnabled(this, &FSkeletalMeshReductionSettingsLayout::IsReductionEnabled)
 			.OnGetMenuContent(this, &FSkeletalMeshReductionSettingsLayout::FillReductionTerminationCriterionMenu)
 			.VAlign(VAlign_Center)
-			.ContentPadding(2)
+			.ContentPadding(2.f)
 			.ButtonContent()
 			[
 				SNew(STextBlock)
@@ -3201,7 +3201,7 @@ void FPersonaMeshDetails::CustomizeLODSettingsCategories(IDetailLayoutBuilder& D
 	.ValueContent()
 	[
 		SNew(STextComboBox)
-		.ContentPadding(0)
+		.ContentPadding(0.f)
 		.OptionsSource(&LODNames)
 		.InitiallySelectedItem(LODNames[0])
 		.Font(IDetailLayoutBuilder::GetDetailFont())
@@ -3751,11 +3751,11 @@ bool FPersonaMeshDetails::AddMinLodQualityLevelOverride(FName QualityLevelName)
 	USkeletalMesh* SkelMesh = GetPersonaToolkit()->GetMesh();
 	check(SkelMesh);
 	SkelMesh->Modify();
-	int32 QLKey = QualityLevelProperty::FNameToQualityLevel(QualityLevelName);
+	const int32 QLKey = QualityLevelProperty::FNameToQualityLevel(QualityLevelName);
 	if (SkelMesh->GetQualityLevelMinLod().PerQuality.Find(QLKey) == nullptr)
 	{
 		FPerQualityLevelInt MinLOD = SkelMesh->GetQualityLevelMinLod();
-		float Value = MinLOD.Default;
+		const int32 Value = MinLOD.Default;
 		MinLOD.PerQuality.Add(QLKey, Value);
 		SkelMesh->SetQualityLevelMinLod(MoveTemp(MinLOD));
 		OnMinQualityLevelLodChanged(Value, QualityLevelName);
@@ -3773,10 +3773,10 @@ bool FPersonaMeshDetails::RemoveMinLodQualityLevelOverride(FName QualityLevelNam
 	SkelMesh->Modify();
 
 	FPerQualityLevelInt MinLOD = SkelMesh->GetQualityLevelMinLod();
-	int32 QL = QualityLevelProperty::FNameToQualityLevel(QualityLevelName);
+	const int32 QL = QualityLevelProperty::FNameToQualityLevel(QualityLevelName);
 	if (QL != INDEX_NONE && MinLOD.PerQuality.Remove(QL) != 0)
 	{
-		float Value = MinLOD.Default;
+		const int32 Value = MinLOD.Default;
 		SkelMesh->SetQualityLevelMinLod(MoveTemp(MinLOD));
 		OnMinQualityLevelLodChanged(Value, QualityLevelName);
 		return true;
@@ -5028,17 +5028,16 @@ TSharedRef<SWidget> FPersonaMeshDetails::OnGenerateCustomNameWidgetsForSection(i
 						]
 					]
 					+SHorizontalBox::Slot()
-					.Padding(5, 2, 5, 0)
+					.Padding(5.f, 2.f, 5.f, 0.f)
 					.AutoWidth()
 					[
 						SNew(SNumericEntryBox<int8>)
 						.Visibility(this, &FPersonaMeshDetails::ShowSectionGenerateUpToSlider, LodIndex, SectionIndex)
 						.Font(IDetailLayoutBuilder::GetDetailFont())
 						.MinDesiredValueWidth(40.0f)
-						.MinValue(LodIndex)
-						//.MaxValue(1)
-						.MinSliderValue(LodIndex)
-						.MaxSliderValue(FMath::Max(8, LODCount))
+						.MinValue(static_cast<int8>(LodIndex))
+						.MinSliderValue(static_cast<int8>(LodIndex))
+						.MaxSliderValue(static_cast<int8>(FMath::Max(8, LODCount)))
 						.AllowSpin(true)
 						.Value(this, &FPersonaMeshDetails::GetSectionGenerateUpToValue, LodIndex, SectionIndex)
 						.OnValueChanged(this, &FPersonaMeshDetails::SetSectionGenerateUpToValue, LodIndex, SectionIndex)
@@ -5141,7 +5140,7 @@ TSharedRef<SWidget> FPersonaMeshDetails::OnGenerateCustomSectionWidgetsForSectio
 			.IsEnabled(IsGPUSkinCacheAvailable(GMaxRHIShaderPlatform))
 			.OnGetMenuContent(this, &FPersonaMeshDetails::OnGenerateRecomputeTangentsSetting, LODIndex, SectionIndex)
 			.VAlign(VAlign_Center)
-			.ContentPadding(2)
+			.ContentPadding(2.f)
 			.ButtonContent()
 			[
 				SNew(STextBlock)
@@ -5290,7 +5289,7 @@ TOptional<int8> FPersonaMeshDetails::GetSectionGenerateUpToValue(int32 LodIndex,
 	{
 		return TOptional<int8>(-1);
 	}
-	int8 SpecifiedLodIndex = SkeletalMeshPtr->GetImportedModel()->LODModels[LodIndex].Sections[SectionIndex].GenerateUpToLodIndex;
+	const int8 SpecifiedLodIndex = static_cast<int8>(SkeletalMeshPtr->GetImportedModel()->LODModels[LodIndex].Sections[SectionIndex].GenerateUpToLodIndex);
 	check(SpecifiedLodIndex == -1 || SpecifiedLodIndex >= LodIndex);
 	return TOptional<int8>(SpecifiedLodIndex);
 }
@@ -5303,10 +5302,10 @@ void FPersonaMeshDetails::SetSectionGenerateUpToValue(int8 Value, int32 LodIndex
 	{
 		return;
 	}
-	int64 ValueKey = ((int64)LodIndex << 32) | (int64)SectionIndex;
+	const int64 ValueKey = ((int64)LodIndex << 32) | (int64)SectionIndex;
 	if (!OldGenerateUpToSliderValues.Contains(ValueKey))
 	{
-		OldGenerateUpToSliderValues.Add(ValueKey, SkeletalMeshPtr->GetImportedModel()->LODModels[LodIndex].Sections[SectionIndex].GenerateUpToLodIndex);
+		OldGenerateUpToSliderValues.Add(ValueKey, static_cast<int8>(SkeletalMeshPtr->GetImportedModel()->LODModels[LodIndex].Sections[SectionIndex].GenerateUpToLodIndex));
 	}
 	SkeletalMeshPtr->GetImportedModel()->LODModels[LodIndex].Sections[SectionIndex].GenerateUpToLodIndex = Value;
 }
@@ -5385,7 +5384,7 @@ ECheckBoxState FPersonaMeshDetails::IsGenerateUpToSectionEnabled(int32 LodIndex,
 
 void FPersonaMeshDetails::OnSectionGenerateUpToChanged(ECheckBoxState NewState, int32 LodIndex, int32 SectionIndex)
 {
-	SetSectionGenerateUpToValueCommitted(NewState == ECheckBoxState::Checked ? LodIndex : -1, ETextCommit::Type::Default , LodIndex, SectionIndex);
+	SetSectionGenerateUpToValueCommitted(NewState == ECheckBoxState::Checked ? static_cast<int8>(LodIndex) : INDEX_NONE, ETextCommit::Type::Default , LodIndex, SectionIndex);
 }
 
 void FPersonaMeshDetails::SetCurrentLOD(int32 NewLodIndex)
@@ -5462,7 +5461,7 @@ TSharedRef<SWidget> FPersonaMeshDetails::OnGenerateLodComboBoxForLodPicker()
 		.IsEnabled(this, &FPersonaMeshDetails::IsLodComboBoxEnabledForLodPicker)
 		.OnGetMenuContent(this, &FPersonaMeshDetails::OnGenerateLodMenuForLodPicker)
 		.VAlign(VAlign_Center)
-		.ContentPadding(2)
+		.ContentPadding(2.f)
 		.ButtonContent()
 		[
 			SNew(STextBlock)
@@ -6134,14 +6133,14 @@ void FPersonaMeshDetails::OnGenerateElementForClothingAsset( TSharedRef<IPropert
 		// remove button
 		+ SHorizontalBox::Slot()
 		.VAlign( VAlign_Center )
-		.Padding(2)
+		.Padding(2.f)
 		.AutoWidth()
 		[
 			SNew( SButton )
 			.Text( LOCTEXT("ClearButtonLabel", "Remove") )
 			.OnClicked( this, &FPersonaMeshDetails::OnRemoveClothingAssetClicked, ElementIndex, DetailLayout )
 			.IsFocusable( false )
-			.ContentPadding(0)
+			.ContentPadding(0.f)
 			.ForegroundColor( FSlateColor::UseForeground() )
 			.ButtonColorAndOpacity(FLinearColor(1.0f,1.0f,1.0f,0.0f))
 			.ToolTipText(LOCTEXT("RemoveClothingAssetTip", "Remove this clothing asset"))
@@ -6525,7 +6524,7 @@ void FPersonaMeshDetails::OnClothingSelectionChanged(TSharedPtr<FClothingEntry>
 				//Successful bind so set the SectionUserData
 				int32 AssetIndex = INDEX_NONE;
 				check(Mesh->GetMeshClothingAssets().Find(ClothingAsset, AssetIndex));
-				OriginalSectionData.CorrespondClothAssetIndex = AssetIndex;
+				OriginalSectionData.CorrespondClothAssetIndex = static_cast<int16>(AssetIndex);
 				OriginalSectionData.ClothingData.AssetGuid = ClothingAsset->GetAssetGuid();
 				OriginalSectionData.ClothingData.AssetLodIndex = InNewEntry->AssetLodIndex;
 			}
diff --git a/Engine/Source/Editor/Persona/Private/PersonaModule.cpp b/Engine/Source/Editor/Persona/Private/PersonaModule.cpp
--- a/Engine/Source/Editor/Persona/Private/PersonaModule.cpp
+++ b/Engine/Source/Editor/Persona/Private/PersonaModule.cpp
@@ -514,7 +514,7 @@ void FPersonaModule::TestSkeletonCurveNamesForUse(const TSharedRef<IEditableSkel
 		{
 			const FText LoadingStatusUpdate = FText::Format(LOCTEXT("VerifyCurves_LoadingAllAnimations", "Loading all animations for skeleton '{0}'"), FText::FromString(Skeleton.GetName()));
 			{
-				FScopedSlowTask LoadingAnimSlowTask(Animations.Num(), LoadingStatusUpdate);
+				FScopedSlowTask LoadingAnimSlowTask(static_cast<float>(Animations.Num()), LoadingStatusUpdate);
 				LoadingAnimSlowTask.MakeDialog();
 
 				// Loop through all animations to load then, this makes sure smart names are all up to date
@@ -531,7 +531,7 @@ void FPersonaModule::TestSkeletonCurveNamesForUse(const TSharedRef<IEditableSkel
 
 			const FText ProcessingStatusUpdate = FText::Format(LOCTEXT("VerifyCurves_ProcessingCurveUsage", "Looking at curve useage for each skeletal mesh of skeleton '{0}'"), FText::FromString(Skeleton.GetName()));
 			{
-				FScopedSlowTask LoadingSkelMeshSlowTask(SkeletalMeshes.Num(), ProcessingStatusUpdate);
+				FScopedSlowTask LoadingSkelMeshSlowTask(static_cast<float>(SkeletalMeshes.Num()), ProcessingStatusUpdate);
 				LoadingSkelMeshSlowTask.MakeDialog();
 
 				for (int32 MeshIdx = 0; MeshIdx < SkeletalMeshes.Num(); ++MeshIdx)
@@ -580,7 +580,7 @@ void FPersonaModule::TestSkeletonCurveNamesForUse(const TSharedRef<IEditableSkel
 
 			const FText ProcessingAnimStatusUpdate = FText::Format(LOCTEXT("FindUnusedCurves_ProcessingSkeletalMeshes", "Finding animations that reference unused curves on skeleton '{0}'"), FText::FromString(Skeleton.GetName()));
 			{
-				FScopedSlowTask ProcessingAnimationsSlowTask(Animations.Num(), ProcessingAnimStatusUpdate);
+				FScopedSlowTask ProcessingAnimationsSlowTask(static_cast<float>(Animations.Num()), ProcessingAnimStatusUpdate);
 				ProcessingAnimationsSlowTask.MakeDialog();
 
 				for (const FAssetData& Anim : Animations)
diff --git a/Engine/Source/Editor/Persona/Private/PhysicsAssetRenderUtils.cpp b/Engine/Source/Editor/Persona/Private/PhysicsAssetRenderUtils.cpp
--- a/Engine/Source/Editor/Persona/Private/PhysicsAssetRenderUtils.cpp
+++ b/Engine/Source/Editor/Persona/Private/PhysicsAssetRenderUtils.cpp
@@ -301,9 +301,7 @@ namespace PhysicsAssetRender
 			// The TM position is at the center of the objects, so can be used directly for sorting.
 			// Sorting by distance (rather than along the view direction) reduces flickering when
 			// the camera is rotated without moving it.
-			FVector ViewDir = InView->GetViewDirection();
-			//Distance = TM.GetTranslation() | ViewDir; // Would sort along the view direction
-			SortingMetric = (TM.GetTranslation() - InView->ViewLocation).SquaredLength(); // Sorts by Distance
+			SortingMetric = static_cast<float>((TM.GetTranslation() - InView->ViewLocation).SquaredLength()); // Sorts by Distance
 		}
 
 		UMaterialInterface* Material = 0;
@@ -390,7 +388,7 @@ namespace PhysicsAssetRender
 			if (BoneIndex != INDEX_NONE)
 			{
 				FTransform BoneTM = SkeletalMeshComponent->GetBoneTransform(BoneIndex);
-				const float Scale = BoneTM.GetScale3D().GetAbsMax();
+				const float Scale = static_cast<float>(BoneTM.GetScale3D().GetAbsMax());
 				BoneTM.RemoveScaling();
 
 				FKAggregateGeom* const AggGeom = &PhysicsAsset->SkeletalBodySetups[i]->AggGeom;
diff --git a/Engine/Source/Editor/Persona/Private/PoseWatchManagerColumnColor.cpp b/Engine/Source/Editor/Persona/Private/PoseWatchManagerColumnColor.cpp
--- a/Engine/Source/Editor/Persona/Private/PoseWatchManagerColumnColor.cpp
+++ b/Engine/Source/Editor/Persona/Private/PoseWatchManagerColumnColor.cpp
@@ -33,7 +33,7 @@ public:
 			.VAlign(VAlign_Center)
 			[
 				SAssignNew(ColorWidgetBackgroundBorder, SBorder)
-				.Padding(1)
+				.Padding(1.f)
 				.BorderImage(FAppStyle::Get().GetBrush("ColorPicker.RoundedSolidBackground"))
 				.BorderBackgroundColor(this, &SColorBoxWidget::GetColorWidgetBorderColor)
 				.VAlign(VAlign_Center)
diff --git a/Engine/Source/Editor/Persona/Private/PoseWatchManagerPublicTypes.h b/Engine/Source/Editor/Persona/Private/PoseWatchManagerPublicTypes.h
--- a/Engine/Source/Editor/Persona/Private/PoseWatchManagerPublicTypes.h
+++ b/Engine/Source/Editor/Persona/Private/PoseWatchManagerPublicTypes.h
@@ -43,7 +43,7 @@ struct FPoseWatchManagerInitializationOptions
 
 struct FPoseWatchManagerDefaultTreeItemMetrics
 {
-	static int32	RowHeight() { return 20; };
-	static int32	IconSize() { return 16; };
+	static float	RowHeight() { return 20.f; };
+	static float	IconSize() { return 16; };
 	static FMargin	IconPadding() { return FMargin(0.f, 1.f, 6.f, 1.f); };
 };
\ No newline at end of file
diff --git a/Engine/Source/Editor/Persona/Private/SAnimCurveEd.cpp b/Engine/Source/Editor/Persona/Private/SAnimCurveEd.cpp
--- a/Engine/Source/Editor/Persona/Private/SAnimCurveEd.cpp
+++ b/Engine/Source/Editor/Persona/Private/SAnimCurveEd.cpp
@@ -78,17 +78,15 @@ FReply SAnimCurveEd::OnMouseWheel(const FGeometry& MyGeometry, const FPointerEve
 	const float ZoomDelta = -0.1f * MouseEvent.GetWheelDelta();
 
 	const FVector2D WidgetSpace = MyGeometry.AbsoluteToLocal(MouseEvent.GetScreenSpacePosition());
-	const float ZoomRatio = FMath::Clamp((WidgetSpace.X / MyGeometry.GetLocalSize().X), 0.f, 1.f);
+	const float ZoomRatio = FMath::Clamp(static_cast<float>(WidgetSpace.X / MyGeometry.GetLocalSize().X), 0.f, 1.f);
 
-	{
-		const float InputViewSize = ViewMaxInput.Get() - ViewMinInput.Get();
-		const float InputChange = InputViewSize * ZoomDelta;
+	const float InputViewSize = ViewMaxInput.Get() - ViewMinInput.Get();
+	const float InputChange = InputViewSize * ZoomDelta;
 
-		float NewViewMinInput = ViewMinInput.Get() - (InputChange * ZoomRatio);
-		float NewViewMaxInput = ViewMaxInput.Get() + (InputChange * (1.f - ZoomRatio));
+	const float NewViewMinInput = ViewMinInput.Get() - (InputChange * ZoomRatio);
+	const float NewViewMaxInput = ViewMaxInput.Get() + (InputChange * (1.f - ZoomRatio));
 
-		SetInputMinMax(NewViewMinInput, NewViewMaxInput);
-	}
+	SetInputMinMax(NewViewMinInput, NewViewMaxInput);
 
 	return FReply::Handled();
 }
diff --git a/Engine/Source/Editor/Persona/Private/SAnimCurveViewer.cpp b/Engine/Source/Editor/Persona/Private/SAnimCurveViewer.cpp
--- a/Engine/Source/Editor/Persona/Private/SAnimCurveViewer.cpp
+++ b/Engine/Source/Editor/Persona/Private/SAnimCurveViewer.cpp
@@ -577,7 +577,7 @@ void SAnimCurveViewer::Construct(const FArguments& InArgs, const TSharedRef<clas
 		.AutoHeight()		// This is required to make the scrollbar work, as content overflows Slate containers by default
 		[
 			SNew(SBox)
-			.WidthOverride(150)
+			.WidthOverride(150.f)
 			.Content()
 			[
 				AnimTypeBoxContainer.ToSharedRef()
@@ -966,7 +966,7 @@ void SAnimCurveViewer::CreateAnimCurveTypeList(TSharedRef<SHorizontalBox> Horizo
 	// Add toggle button for 'all curves'
 	HorizontalBox->AddSlot()
 	.AutoWidth()
-	.Padding(3, 1)
+	.Padding(3.f, 1.f)
 	[
 		SNew(SCheckBox)
 		.Style(FAppStyle::Get(), "ToggleButtonCheckbox")
@@ -974,7 +974,7 @@ void SAnimCurveViewer::CreateAnimCurveTypeList(TSharedRef<SHorizontalBox> Horizo
 		.Type(ESlateCheckBoxType::ToggleButton)
 		.IsChecked(this, &SAnimCurveViewer::IsShowingAllCurves)
 		.OnCheckStateChanged(this, &SAnimCurveViewer::OnToggleShowingAllCurves)
-		.Padding(4)
+		.Padding(4.f)
 		.Content()
 		[
 			SNew(STextBlock)
@@ -992,7 +992,7 @@ void SAnimCurveViewer::CreateAnimCurveTypeList(TSharedRef<SHorizontalBox> Horizo
 	{
 		HorizontalBox->AddSlot()
 		.AutoWidth()
-		.Padding(3, 1)
+		.Padding(3.f, 1.f)
 		[
 			SNew(SAnimCurveTypeList)
 			.IsEnabled(this, &SAnimCurveViewer::IsCurveFilterEnabled)
@@ -1133,7 +1133,7 @@ void SAnimCurveViewer::OnSelectionChanged(TSharedPtr<FDisplayedAnimCurveInfo> In
 			TArray<FBoneReference> BoneLinks;
 			FSmartName CurrentName = RowItem->SmartName;
 			const FCurveMetaData* CurveMetaData = EditableSkeletonPtr.Pin()->GetSkeleton().GetCurveMetaData(CurrentName);
-			uint32 MaxLOD = 0xFF;
+			uint8 MaxLOD = 0xFF;
 			if (CurveMetaData)
 			{
 				BoneLinks = CurveMetaData->LinkedBones;
diff --git a/Engine/Source/Editor/Persona/Private/SAnimNotifyPanel.cpp b/Engine/Source/Editor/Persona/Private/SAnimNotifyPanel.cpp
--- a/Engine/Source/Editor/Persona/Private/SAnimNotifyPanel.cpp
+++ b/Engine/Source/Editor/Persona/Private/SAnimNotifyPanel.cpp
@@ -691,7 +691,7 @@ void SAnimNotifyPair::Construct(const FArguments& InArgs)
 
 float SAnimNotifyPair::GetWidgetPaddingLeft()
 {
-	return NodePtr->GetWidgetPosition().X - PairedWidget->GetDesiredSize().X;
+	return static_cast<float>(NodePtr->GetWidgetPosition().X - PairedWidget->GetDesiredSize().X);
 }
 
 //////////////////////////////////////////////////////////////////////////
@@ -962,16 +962,16 @@ private:
 	// Returns the padding needed to render the notify in the correct track position
 	FMargin GetNotifyTrackPadding(int32 NotifyIndex) const
 	{
-		float LeftMargin = NotifyPairs[NotifyIndex]->GetWidgetPaddingLeft();
-		float RightMargin = CachedGeometry.GetLocalSize().X - NotifyNodes[NotifyIndex]->GetWidgetPosition().X - NotifyNodes[NotifyIndex]->GetSize().X;
+		const float LeftMargin = NotifyPairs[NotifyIndex]->GetWidgetPaddingLeft();
+		const float RightMargin = static_cast<float>(CachedGeometry.GetLocalSize().X - NotifyNodes[NotifyIndex]->GetWidgetPosition().X - NotifyNodes[NotifyIndex]->GetSize().X);
 		return FMargin(LeftMargin, 0, RightMargin, 0);
 	}
 
 	// Returns the padding needed to render the notify in the correct track position
 	FMargin GetSyncMarkerTrackPadding(int32 SyncMarkerIndex) const
 	{
-		float LeftMargin = NotifyNodes[SyncMarkerIndex]->GetWidgetPosition().X;
-		float RightMargin = CachedGeometry.GetLocalSize().X - NotifyNodes[SyncMarkerIndex]->GetWidgetPosition().X - NotifyNodes[SyncMarkerIndex]->GetSize().X;
+		const float LeftMargin = static_cast<float>(NotifyNodes[SyncMarkerIndex]->GetWidgetPosition().X);
+		const float RightMargin = static_cast<float>(CachedGeometry.GetLocalSize().X - NotifyNodes[SyncMarkerIndex]->GetWidgetPosition().X - NotifyNodes[SyncMarkerIndex]->GetSize().X);
 		return FMargin(LeftMargin, 0, RightMargin, 0);
 	}
 
@@ -1196,16 +1196,16 @@ public:
 		// Tracks the movement amount to apply to the selection due to a snap.
 		float SnapMovement = 0.0f;
 		// Clamp the selection into the track
-		float SelectionBeginLocalPositionX = TrackGeom.AbsoluteToLocal(SelectionBeginPosition).X;
-		const float ClampedEnd = FMath::Clamp(SelectionBeginLocalPositionX + NodeGroupSize.X, LocalTrackMin, LocalTrackMax);
+		float SelectionBeginLocalPositionX = static_cast<float>(TrackGeom.AbsoluteToLocal(SelectionBeginPosition).X);
+		const float ClampedEnd = FMath::Clamp(SelectionBeginLocalPositionX + static_cast<float>(NodeGroupSize.X), LocalTrackMin, LocalTrackMax);
 		const float ClampedBegin = FMath::Clamp(SelectionBeginLocalPositionX, LocalTrackMin, LocalTrackMax);
 		if(ClampedBegin > SelectionBeginLocalPositionX)
 		{
 			SelectionBeginLocalPositionX = ClampedBegin;
 		}
-		else if(ClampedEnd < SelectionBeginLocalPositionX + NodeGroupSize.X)
+		else if(ClampedEnd < SelectionBeginLocalPositionX + static_cast<float>(NodeGroupSize.X))
 		{
-			SelectionBeginLocalPositionX = ClampedEnd - NodeGroupSize.X;
+			SelectionBeginLocalPositionX = ClampedEnd - static_cast<float>(NodeGroupSize.X);
 		}
 
 		SelectionBeginPosition.X = TrackGeom.LocalToAbsolute(FVector2D(SelectionBeginLocalPositionX, 0.0f)).X;
@@ -1231,7 +1231,7 @@ public:
 			// Always clamp the Y to the current track
 			SelectionBeginPosition.Y = SelectionPositionClampInfo->TrackPos - 1.0f;
 
-			float SnapX = GetSnapPosition(NodeClamp, TrackNodePos.X, bSnapped);
+			float SnapX = GetSnapPosition(NodeClamp, static_cast<float>(TrackNodePos.X), bSnapped);
 			if (FAnimNotifyEvent* CurrentEvent = CurrentNode->NodeObjectInterface->GetNotifyEvent())
 			{
 				if (bSnapped)
@@ -1251,7 +1251,7 @@ public:
 
 					if (SnapMovement == 0.0f)
 					{
-						SnapMovement = SnapX - TrackNodePos.X;
+						SnapMovement = SnapX - static_cast<float>(TrackNodePos.X);
 						TrackNodePos.X = SnapX;
 						SnapTime = TrackScaleInfo.LocalXToInput(SnapX);
 						SnappedNode = CurrentNode;
@@ -1268,8 +1268,8 @@ public:
 					// If we didn't snap the beginning of the node, attempt to snap the end
 					if (!bSnapped)
 					{
-						FVector2D TrackNodeEndPos = TrackNodePos + CurrentNode->GetDurationSize();
-						SnapX = GetSnapPosition(*SelectionPositionClampInfo, TrackNodeEndPos.X, bSnapped);
+						const FVector2D TrackNodeEndPos = TrackNodePos + CurrentNode->GetDurationSize();
+						SnapX = GetSnapPosition(*SelectionPositionClampInfo, static_cast<float>(TrackNodeEndPos.X), bSnapped);
 
 						// Only attempt to snap if the node will fit on the track
 						if (SnapX >= CurrentNode->GetDurationSize())
@@ -1288,7 +1288,7 @@ public:
 
 							if (SnapMovement == 0.0f)
 							{
-								SnapMovement = SnapX - TrackNodeEndPos.X;
+								SnapMovement = SnapX - static_cast<float>(TrackNodeEndPos.X);
 								SnapTime = TrackScaleInfo.LocalXToInput(SnapX) - CurrentEvent->GetDuration();
 								CurrentNode->SetLastSnappedTime(SnapTime);
 								SnappedNode = CurrentNode;
@@ -1306,24 +1306,24 @@ public:
 
 		SelectionBeginPosition.X += SnapMovement;
 
-		CurrentDragXPosition = TrackGeom.AbsoluteToLocal(FVector2D(SelectionBeginPosition.X,0.0f)).X;
+		CurrentDragXPosition = static_cast<float>(TrackGeom.AbsoluteToLocal(FVector2D(SelectionBeginPosition.X,0.0f)).X);
 
 		CursorDecoratorWindow->MoveWindowTo(TrackGeom.LocalToAbsolute(TrackGeom.AbsoluteToLocal(SelectionBeginPosition) - SelectedNodes[0]->GetNotifyPositionOffset()));
 		NodeGroupPosition = SelectionBeginPosition;
 
 		//scroll view
-		float LocalMouseXPos = TrackGeom.AbsoluteToLocal(DragDropEvent.GetScreenSpacePosition()).X;
-		float LocalViewportMin = 0.0f;
-		float LocalViewportMax = TrackGeom.GetLocalSize().X;
+		const float LocalMouseXPos = static_cast<float>(TrackGeom.AbsoluteToLocal(DragDropEvent.GetScreenSpacePosition()).X);
+		constexpr float LocalViewportMin = 0.0f;
+		const float LocalViewportMax = static_cast<float>(TrackGeom.GetLocalSize().X);
 		if(LocalMouseXPos < LocalViewportMin && LocalViewportMin > LocalTrackMin - 10.0f)
 		{
-			float ScreenDelta = FMath::Max(LocalMouseXPos - LocalViewportMin, -10.0f);
-			RequestTrackPan.Execute(ScreenDelta, FVector2D(LocalTrackWidth, 1.f));
+			const float ScreenDelta = FMath::Max(LocalMouseXPos - LocalViewportMin, -10.0f);
+			RequestTrackPan.Execute(static_cast<int32>(ScreenDelta), FVector2D(LocalTrackWidth, 1.f));
 		}
 		else if(LocalMouseXPos > LocalViewportMax && LocalViewportMax < LocalTrackMax + 10.0f)
 		{
-			float ScreenDelta =  FMath::Max(LocalMouseXPos - LocalViewportMax, 10.0f);
-			RequestTrackPan.Execute(ScreenDelta, FVector2D(LocalTrackWidth, 1.f));
+			const float ScreenDelta =  FMath::Max(LocalMouseXPos - LocalViewportMax, 10.0f);
+			RequestTrackPan.Execute(static_cast<int32>(ScreenDelta), FVector2D(LocalTrackWidth, 1.f));
 		}
 
 		OnNodesBeingDragged.ExecuteIfBound(SelectedNodes, DragDropEvent, CurrentDragXPosition, TrackScaleInfo.LocalXToInput(CurrentDragXPosition));
@@ -1366,10 +1366,10 @@ public:
 	FTrackClampInfo& GetTrackClampInfo(const FVector2D NodePos)
 	{
 		int32 ClampInfoIndex = 0;
-		int32 SmallestNodeTrackDist = FMath::Abs(ClampInfos[0].TrackSnapTestPos - NodePos.Y);
+		int32 SmallestNodeTrackDist = FMath::Abs(ClampInfos[0].TrackSnapTestPos - static_cast<int32>(NodePos.Y));
 		for(int32 i = 0; i < ClampInfos.Num(); ++i)
 		{
-			int32 Dist = FMath::Abs(ClampInfos[i].TrackSnapTestPos - NodePos.Y);
+			const int32 Dist = FMath::Abs(ClampInfos[i].TrackSnapTestPos - static_cast<int32>(NodePos.Y));
 			if(Dist < SmallestNodeTrackDist)
 			{
 				SmallestNodeTrackDist = Dist;
@@ -1448,7 +1448,7 @@ public:
 			Node->ClearLastSnappedTime();
 			Operation->NodeTimeOffsets.Add(NotifyTime - BeginTime);
 			Operation->NodeTimes.Add(NotifyTime);
-			Operation->NodeXOffsets.Add(Node->GetNotifyPositionOffset().X);
+			Operation->NodeXOffsets.Add(static_cast<float>(Node->GetNotifyPositionOffset().X));
 
 			// Calculate the time length of the selection. Because it is possible to have states
 			// with arbitrary durations we need to search all of the nodes and find the furthest
@@ -1463,8 +1463,8 @@ public:
 			FTrackClampInfo Info;
 			Info.NotifyTrack = NotifyTracks[i];
 			const FGeometry& CachedGeometry = Info.NotifyTrack->GetCachedGeometry();
-			Info.TrackPos = CachedGeometry.AbsolutePosition.Y;
-			Info.TrackSnapTestPos = Info.TrackPos + (CachedGeometry.Size.Y / 2);
+			Info.TrackPos = static_cast<int32>(CachedGeometry.AbsolutePosition.Y);
+			Info.TrackSnapTestPos = Info.TrackPos + static_cast<int32>(CachedGeometry.Size.Y / 2);
 			Operation->ClampInfos.Add(Info);
 		}
 
@@ -1634,13 +1634,13 @@ ENotifyStateHandleHit::Type SAnimNotifyNode::DurationHandleHitTest(const FVector
 	if(NotifyDurationSizeX > 0.0f)
 	{
 		// Test for mouse inside duration box with handles included
-		float ScrubHandleHalfWidth = ScrubHandleSize.X / 2.0f;
+		const double ScrubHandleHalfWidth = ScrubHandleSize.X / 2.0f;
 
 		// Position and size of the notify node including the scrub handles
-		FVector2D NotifyNodePosition(NotifyScrubHandleCentre - ScrubHandleHalfWidth, 0.0f);
-		FVector2D NotifyNodeSize(NotifyDurationSizeX + ScrubHandleHalfWidth * 2.0f, NotifyHeight);
+		const FVector2D NotifyNodePosition(NotifyScrubHandleCentre - ScrubHandleHalfWidth, 0.0);
+		const FVector2D NotifyNodeSize(NotifyDurationSizeX + ScrubHandleHalfWidth * 2.0, NotifyHeight);
 
-		FVector2D MouseRelativePosition(CursorTrackPosition - GetWidgetPosition());
+		const FVector2D MouseRelativePosition(CursorTrackPosition - GetWidgetPosition());
 
 		if(MouseRelativePosition.ComponentwiseAllGreaterThan(NotifyNodePosition) &&
 			MouseRelativePosition.ComponentwiseAllLessThan(NotifyNodePosition + NotifyNodeSize))
@@ -1674,22 +1674,21 @@ void SAnimNotifyNode::UpdateSizeAndPosition(const FGeometry& AllottedGeometry)
 
 	const TSharedRef< FSlateFontMeasure > FontMeasureService = FSlateApplication::Get().GetRenderer()->GetFontMeasureService();
 	TextSize = FontMeasureService->Measure( GetNotifyText(), Font );
-	LabelWidth = TextSize.X + (TextBorderSize.X * 2.f) + (ScrubHandleSize.X / 2.f);
+	LabelWidth = static_cast<float>(TextSize.X + (TextBorderSize.X * 2.0) + (ScrubHandleSize.X / 2.0));
 
-	bool bDrawBranchingPoint = NodeObjectInterface->IsBranchingPoint();
+	const bool bDrawBranchingPoint = NodeObjectInterface->IsBranchingPoint();
 	BranchingPointIconSize = FVector2D(TextSize.Y, TextSize.Y);
 	if (bDrawBranchingPoint)
 	{
-		LabelWidth += BranchingPointIconSize.X + TextBorderSize.X * 2.f;
+		LabelWidth += static_cast<float>(BranchingPointIconSize.X + TextBorderSize.X) * 2.f;
 	}
 
 	//Calculate scrub handle box size (the notional box around the scrub handle and the alignment marker)
-	float NotifyHandleBoxWidth = FMath::Max(ScrubHandleSize.X, AlignmentMarkerSize.X * 2);
+	const float NotifyHandleBoxWidth = static_cast<float>(FMath::Max(ScrubHandleSize.X, AlignmentMarkerSize.X * 2));
 
 	// Work out where we will have to draw the tool tip
-	FVector2D Size = GetSize();
-	float LeftEdgeToNotify = NotifyTimePositionX;
-	float RightEdgeToNotify = AllottedGeometry.Size.X - NotifyTimePositionX;
+	const float LeftEdgeToNotify = NotifyTimePositionX;
+	const float RightEdgeToNotify = static_cast<float>(AllottedGeometry.Size.X) - NotifyTimePositionX;
 	bDrawTooltipToRight = NotifyDurationSizeX > 0.0f || ((RightEdgeToNotify > LabelWidth) || (RightEdgeToNotify > LeftEdgeToNotify));
 
 	// Calculate widget width/position based on where we are drawing the tool tip
@@ -1699,7 +1698,7 @@ void SAnimNotifyNode::UpdateSizeAndPosition(const FGeometry& AllottedGeometry)
 	
 	if(EndMarkerNodeOverlay.IsValid())
 	{
-		FVector2D OverlaySize = EndMarkerNodeOverlay->GetDesiredSize();
+		const FVector2D OverlaySize = EndMarkerNodeOverlay->GetDesiredSize();
 		WidgetSize.X += OverlaySize.X;
 	}
 
@@ -1751,7 +1750,7 @@ int32 SAnimNotifyNode::OnPaint(const FPaintArgs& Args, const FGeometry& Allotted
 	FLinearColor NodeColor = SAnimNotifyNode::GetNotifyColor();
 	FLinearColor BoxColor = bSelected ? FAppStyle::GetSlateColor("SelectionColor").GetSpecifiedColor() : SAnimNotifyNode::GetNotifyColor();
 
-	float HalfScrubHandleWidth = ScrubHandleSize.X / 2.0f;
+	const float HalfScrubHandleWidth = static_cast<float>(ScrubHandleSize.X) / 2.0f;
 
 	// Show duration of AnimNotifyState
 	if( NotifyDurationSizeX > 0.f )
@@ -1766,16 +1765,16 @@ int32 SAnimNotifyNode::OnPaint(const FPaintArgs& Args, const FGeometry& Allotted
 			ESlateDrawEffect::None,
 			BoxColor);
 
-		DrawScrubHandle(DurationBoxPosition.X + DurationBoxSize.X, OutDrawElements, ScrubHandleID, AllottedGeometry, MyCullingRect, NodeColor);
+		DrawScrubHandle(static_cast<float>(DurationBoxPosition.X + DurationBoxSize.X), OutDrawElements, ScrubHandleID, AllottedGeometry, MyCullingRect, NodeColor);
 		
 		// Render offsets if necessary
 		if(AnimNotifyEvent && AnimNotifyEvent->EndTriggerTimeOffset != 0.f) //Do we have an offset to render?
 		{
-			float EndTime = AnimNotifyEvent->GetTime() + AnimNotifyEvent->GetDuration();
+			const float EndTime = AnimNotifyEvent->GetTime() + AnimNotifyEvent->GetDuration();
 			if(EndTime != Sequence->GetPlayLength()) //Don't render offset when we are at the end of the sequence, doesnt help the user
 			{
 				// ScrubHandle
-				float HandleCentre = NotifyDurationSizeX + (ScrubHandleSize.X - 2.0f);
+				const float HandleCentre = NotifyDurationSizeX + (static_cast<float>(ScrubHandleSize.X) - 2.0f);
 				DrawHandleOffset(AnimNotifyEvent->EndTriggerTimeOffset, HandleCentre, OutDrawElements, MarkerLayer, AllottedGeometry, MyCullingRect, NodeColor);
 			}
 		}
@@ -1881,10 +1880,10 @@ FReply SAnimNotifyNode::OnMouseMove( const FGeometry& MyGeometry, const FPointer
 	
 	FTrackScaleInfo ScaleInfo(ViewInputMin.Get(), ViewInputMax.Get(), 0, 0, CachedAllotedGeometrySize);
 	
-	float XPositionInTrack = MyGeometry.AbsolutePosition.X - CachedTrackGeometry.AbsolutePosition.X;
-	float TrackScreenSpaceXPosition = MyGeometry.AbsolutePosition.X - XPositionInTrack;
-	float TrackScreenSpaceOrigin = CachedTrackGeometry.LocalToAbsolute(FVector2D(ScaleInfo.InputToLocalX(0.0f), 0.0f)).X;
-	float TrackScreenSpaceLimit = CachedTrackGeometry.LocalToAbsolute(FVector2D(ScaleInfo.InputToLocalX(Sequence->GetPlayLength()), 0.0f)).X;
+	const float XPositionInTrack = MyGeometry.AbsolutePosition.X - CachedTrackGeometry.AbsolutePosition.X;
+	const float TrackScreenSpaceXPosition = MyGeometry.AbsolutePosition.X - XPositionInTrack;
+	const float TrackScreenSpaceOrigin = static_cast<float>(CachedTrackGeometry.LocalToAbsolute(FVector2D(ScaleInfo.InputToLocalX(0.0f), 0.0f)).X);
+	const float TrackScreenSpaceLimit = static_cast<float>(CachedTrackGeometry.LocalToAbsolute(FVector2D(ScaleInfo.InputToLocalX(Sequence->GetPlayLength()), 0.0f)).X);
 
 	if(CurrentDragHandle == ENotifyStateHandleHit::Start)
 	{
@@ -1894,7 +1893,7 @@ FReply SAnimNotifyNode::OnMouseMove( const FGeometry& MyGeometry, const FPointer
 		if(MouseEvent.GetScreenSpacePosition().X >= TrackScreenSpaceXPosition && MouseEvent.GetScreenSpacePosition().X <= TrackScreenSpaceXPosition + CachedAllotedGeometrySize.X)
 		{
 			float NewDisplayTime = ScaleInfo.LocalXToInput((FVector2f(MouseEvent.GetScreenSpacePosition()) - MyGeometry.AbsolutePosition + XPositionInTrack).X);	// LWC_TODO: Precision loss
-			float NewDuration = NodeObjectInterface->GetDuration() + OldDisplayTime - NewDisplayTime;
+			const float NewDuration = NodeObjectInterface->GetDuration() + OldDisplayTime - NewDisplayTime;
 
 			// Check to make sure the duration is not less than the minimum allowed
 			if(NewDuration < MinimumStateDuration)
@@ -2059,15 +2058,15 @@ float SAnimNotifyNode::HandleOverflowPan( const FVector2D &ScreenCursorPos, floa
 	if(ScreenCursorPos.X < TrackScreenSpaceXPosition && TrackScreenSpaceXPosition > TrackScreenSpaceMin - 10.0f)
 	{
 		// Overflow left edge
-		Overflow = FMath::Min(ScreenCursorPos.X - TrackScreenSpaceXPosition, -10.0f);
+		Overflow = FMath::Min(static_cast<float>(ScreenCursorPos.X) - TrackScreenSpaceXPosition, -10.0f);
 	}
 	else if(ScreenCursorPos.X > CachedAllotedGeometrySize.X && (TrackScreenSpaceXPosition + CachedAllotedGeometrySize.X) < TrackScreenSpaceMax + 10.0f)
 	{
 		// Overflow right edge
-		Overflow = FMath::Max(ScreenCursorPos.X - (TrackScreenSpaceXPosition + CachedAllotedGeometrySize.X), 10.0f);
+		Overflow = FMath::Max(static_cast<float>(ScreenCursorPos.X) - (TrackScreenSpaceXPosition + static_cast<float>(CachedAllotedGeometrySize.X)), 10.0f);
 	}
 
-	PanTrackRequest.ExecuteIfBound(Overflow, CachedAllotedGeometrySize);
+	PanTrackRequest.ExecuteIfBound(static_cast<int32>(Overflow), CachedAllotedGeometrySize);
 
 	return Overflow;
 }
@@ -2149,11 +2148,11 @@ FCursorReply SAnimNotifyNode::OnCursorQuery(const FGeometry& MyGeometry, const F
 	// Show resize cursor if the cursor is hoverring over either of the scrub handles of a notify state node
 	if(IsHovered() && GetDurationSize() > 0.0f)
 	{
-		FVector2D RelMouseLocation = MyGeometry.AbsoluteToLocal(CursorEvent.GetScreenSpacePosition());
+		const FVector2D RelMouseLocation = MyGeometry.AbsoluteToLocal(CursorEvent.GetScreenSpacePosition());
 
-		const float HandleHalfWidth = ScrubHandleSize.X / 2.0f;
-		const float DistFromFirstHandle = FMath::Abs(RelMouseLocation.X - NotifyScrubHandleCentre);
-		const float DistFromSecondHandle = FMath::Abs(RelMouseLocation.X - (NotifyScrubHandleCentre + NotifyDurationSizeX));
+		const float HandleHalfWidth = static_cast<float>(ScrubHandleSize.X / 2.0);
+		const float DistFromFirstHandle = FMath::Abs(static_cast<float>(RelMouseLocation.X) - NotifyScrubHandleCentre);
+		const float DistFromSecondHandle = FMath::Abs(static_cast<float>(RelMouseLocation.X) - (NotifyScrubHandleCentre + NotifyDurationSizeX));
 
 		if(DistFromFirstHandle < HandleHalfWidth || DistFromSecondHandle < HandleHalfWidth || CurrentDragHandle != ENotifyStateHandleHit::None)
 		{
@@ -3409,8 +3408,8 @@ FReply SAnimNotifyTrack::OnNotifyNodeDragStarted(TSharedRef<SAnimNotifyNode> Not
 	// Sort our nodes so we're acessing them in time order
 	SelectedNodeIndices.Sort([this](const int32& A, const int32& B)
 	{
-		float TimeA = NotifyNodes[A]->NodeObjectInterface->GetTime();
-		float TimeB = NotifyNodes[B]->NodeObjectInterface->GetTime();
+		const double TimeA = NotifyNodes[A]->NodeObjectInterface->GetTime();
+		const double TimeB = NotifyNodes[B]->NodeObjectInterface->GetTime();
 		return TimeA < TimeB;
 	});
 
@@ -3418,12 +3417,10 @@ FReply SAnimNotifyTrack::OnNotifyNodeDragStarted(TSharedRef<SAnimNotifyNode> Not
 	if (!bDragOnMarker)
 	{
 		TArray<TSharedPtr<SAnimNotifyNode>> NodesToDrag;
-		const float FirstNodeX = NotifyNodes[SelectedNodeIndices[0]]->GetWidgetPosition().X;
-
-		TSharedRef<SOverlay> DragBox = SNew(SOverlay);
+		const TSharedRef<SOverlay> DragBox = SNew(SOverlay);
 		for (auto Iter = SelectedNodeIndices.CreateIterator(); Iter; ++Iter)
 		{
-			TSharedPtr<SAnimNotifyNode> Node = NotifyNodes[*Iter];
+			const TSharedPtr<SAnimNotifyNode> Node = NotifyNodes[*Iter];
 			NodesToDrag.Add(Node);
 		}
 
@@ -3471,7 +3468,7 @@ float SAnimNotifyTrack::CalculateTime(const FGeometry& MyGeometry, FVector2D Nod
 		NodePos = MyGeometry.AbsoluteToLocal(NodePos);
 	}
 	FTrackScaleInfo ScaleInfo(ViewInputMin.Get(), ViewInputMax.Get(), 0, 0, MyGeometry.Size);
-	return FMath::Clamp<float>(ScaleInfo.LocalXToInput(NodePos.X), 0.f, Sequence->GetPlayLength());
+	return FMath::Clamp<float>(ScaleInfo.LocalXToInput(static_cast<float>(NodePos.X)), 0.f, Sequence->GetPlayLength());
 }
 
 FReply SAnimNotifyTrack::OnDrop(const FGeometry& MyGeometry, const FDragDropEvent& DragDropEvent)
@@ -3488,9 +3485,9 @@ void SAnimNotifyTrack::HandleNodeDrop(TSharedPtr<SAnimNotifyNode> Node, float Of
 		UBlendSpace::UpdateBlendSpacesUsingAnimSequence(Sequence);
 	}
 
-	float LocalX = GetCachedGeometry().AbsoluteToLocal(Node->GetScreenPosition() + Offset).X;
-	float SnapTime = Node->GetLastSnappedTime();
-	float Time = SnapTime != -1.0f ? SnapTime : GetCachedScaleInfo().LocalXToInput(LocalX);
+	const float LocalX = static_cast<float>(GetCachedGeometry().AbsoluteToLocal(Node->GetScreenPosition() + Offset).X);
+	const float SnapTime = Node->GetLastSnappedTime();
+	const float Time = SnapTime != -1.0f ? SnapTime : GetCachedScaleInfo().LocalXToInput(LocalX);
 	Node->NodeObjectInterface->HandleDrop(Sequence, Time, TrackIndex);
 }
 
@@ -3501,14 +3498,12 @@ void SAnimNotifyTrack::DisconnectSelectedNodesForDrag(TArray<TSharedPtr<SAnimNot
 		return;
 	}
 
-	const float FirstNodeX = NotifyNodes[SelectedNodeIndices[0]]->GetWidgetPosition().X;
-
 	for(auto Iter = SelectedNodeIndices.CreateIterator(); Iter; ++Iter)
 	{
-		TSharedPtr<SAnimNotifyNode> Node = NotifyNodes[*Iter];
+		const TSharedPtr<SAnimNotifyNode> Node = NotifyNodes[*Iter];
 		if (Node->NodeObjectInterface->GetNotifyEvent())
 		{
-			TSharedPtr<SAnimNotifyPair> Pair = NotifyPairs[*Iter];
+			const TSharedPtr<SAnimNotifyPair> Pair = NotifyPairs[*Iter];
 			NodeSlots->RemoveSlot(Pair->AsShared());
 		}
 		else
@@ -4185,15 +4180,15 @@ FReply SAnimNotifyPanel::OnNotifyNodeDragStarted(TArray<TSharedPtr<SAnimNotifyNo
 		OverlayBounds += FBox2D(NodePosition, NodePosition + FVector2D(NodeDuration, 0.0f));
 	}
 
-	FVector2D OverlayOrigin = OverlayBounds.Min;
-	FVector2D OverlayExtents = OverlayBounds.GetSize();
+	const FVector2D OverlayOrigin = OverlayBounds.Min;
+	const FVector2D OverlayExtents = OverlayBounds.GetSize();
 
-	for(TSharedPtr<SAnimNotifyNode> Node : Nodes)
+	for(const TSharedPtr<SAnimNotifyNode>& Node : Nodes)
 	{
-		FVector2D OffsetFromFirst(Node->GetScreenPosition() - OverlayOrigin);
+		const FVector2D OffsetFromFirst(Node->GetScreenPosition() - OverlayOrigin);
 
 		NodeDragDecoratorOverlay->AddSlot()
-			.Padding(FMargin(OffsetFromFirst.X, OffsetFromFirst.Y, 0.0f, 0.0f))
+			.Padding(FMargin(static_cast<float>(OffsetFromFirst.X), static_cast<float>(OffsetFromFirst.Y), 0.0f, 0.0f))
 			[
 				Node->AsShared()
 			];
diff --git a/Engine/Source/Editor/Persona/Private/SAnimTimingPanel.cpp b/Engine/Source/Editor/Persona/Private/SAnimTimingPanel.cpp
--- a/Engine/Source/Editor/Persona/Private/SAnimTimingPanel.cpp
+++ b/Engine/Source/Editor/Persona/Private/SAnimTimingPanel.cpp
@@ -185,13 +185,13 @@ FVector2D SAnimTimingNode::ComputeDesiredSize(float) const
 	// Desired height is always the same (a little less than the track height) but the width depends on the text we display
 	const TSharedRef< FSlateFontMeasure > FontMeasureService = FSlateApplication::Get().GetRenderer()->GetFontMeasureService();
 	static FSlateFontInfo LabelFont = FCoreStyle::GetDefaultFontStyle("Regular", AnimTimingConstants::FontSize);
-	float TextWidth = FontMeasureService->Measure(FString::FromInt(Element->TriggerIdx), LabelFont).X;
+	const float TextWidth = static_cast<float>(FontMeasureService->Measure(FString::FromInt(Element->TriggerIdx), LabelFont).X);
 	return FVector2D(FMath::Max(AnimTimingConstants::DefaultNodeSize, TextWidth), AnimTimingConstants::DefaultNodeSize);
 }
 
 void SAnimTimingTrackNode::Construct(const FArguments& InArgs)
 {
-	TAttribute<float> TimeAttr = TAttribute<float>::Create(TAttribute<float>::FGetter::CreateSP(InArgs._Element.ToSharedRef(), &FTimingRelevantElementBase::GetElementTime));
+	const TAttribute<float> TimeAttr = TAttribute<float>::Create(TAttribute<float>::FGetter::CreateSP(InArgs._Element.ToSharedRef(), &FTimingRelevantElementBase::GetElementTime));
 
 	STrackNode::Construct(STrackNode::FArguments()
 		.ViewInputMin(InArgs._ViewInputMin)
diff --git a/Engine/Source/Editor/Persona/Private/SAnimViewportToolBar.cpp b/Engine/Source/Editor/Persona/Private/SAnimViewportToolBar.cpp
--- a/Engine/Source/Editor/Persona/Private/SAnimViewportToolBar.cpp
+++ b/Engine/Source/Editor/Persona/Private/SAnimViewportToolBar.cpp
@@ -84,10 +84,10 @@ public:
 					SNew(SNumericEntryBox<float>)
 					.Font(FAppStyle::GetFontStyle(TEXT("MenuItem.Font")))
 					.ToolTipText(LOCTEXT("WindStrength_ToolTip", "Change wind strength"))
-					.MinValue(0)
+					.MinValue(0.f)
 					.AllowSpin(true)
-					.MinSliderValue(0)
-					.MaxSliderValue(10)
+					.MinSliderValue(0.f)
+					.MaxSliderValue(10.f)
 					.Value(AnimViewportPtr.Pin().ToSharedRef(), &SAnimationEditorViewportTabBody::GetWindStrengthSliderValue)
 					.OnValueChanged(SSpinBox<float>::FOnValueChanged::CreateSP(AnimViewportPtr.Pin().ToSharedRef(), &SAnimationEditorViewportTabBody::SetWindStrength))
 				]
@@ -135,8 +135,8 @@ public:
 					SNew(SSpinBox<float>)
 					.Font(FAppStyle::GetFontStyle(TEXT("MenuItem.Font")))
 					.ToolTipText(LOCTEXT("GravityScale_ToolTip", "Change gravity scale"))
-					.MinValue(0)
-					.MaxValue(4)
+					.MinValue(0.f)
+					.MaxValue(4.f)
 					.Value(AnimViewportPtr.Pin().ToSharedRef(), &SAnimationEditorViewportTabBody::GetGravityScaleSliderValue)
 					.OnValueChanged(SSpinBox<float>::FOnValueChanged::CreateSP(AnimViewportPtr.Pin().ToSharedRef(), &SAnimationEditorViewportTabBody::SetGravityScale))
 				]
@@ -189,8 +189,8 @@ public:
 					SNew(SSpinBox<float>)
 					.Font(FAppStyle::GetFontStyle(TEXT("MenuItem.Font")))
 					.ToolTipText(LOCTEXT("BoneDrawSize_ToolTip", "Change bone size in viewport."))
-					.MinValue(0)
-					.MaxSliderValue(10)
+					.MinValue(0.f)
+					.MaxSliderValue(10.f)
 					.SupportDynamicSliderMaxValue(true)
 					.Value(AnimViewportPtr.Pin().ToSharedRef(), &SAnimationEditorViewportTabBody::GetBoneDrawSize)
 					.OnValueChanged(SSpinBox<float>::FOnValueChanged::CreateSP(AnimViewportPtr.Pin().ToSharedRef(), &SAnimationEditorViewportTabBody::SetBoneDrawSize))
diff --git a/Engine/Source/Editor/Persona/Private/SAnimationBlendSpaceGridWidget.cpp b/Engine/Source/Editor/Persona/Private/SAnimationBlendSpaceGridWidget.cpp
--- a/Engine/Source/Editor/Persona/Private/SAnimationBlendSpaceGridWidget.cpp
+++ b/Engine/Source/Editor/Persona/Private/SAnimationBlendSpaceGridWidget.cpp
@@ -125,7 +125,7 @@ static void PaintPolygon(
 		NewVert.Color = FillColor.ToFColor(false);
 		MidVertex.Position += NewVert.Position;
 	}
-	MidVertex.Position /= Points.Num();
+	MidVertex.Position /= static_cast<float>(Points.Num());
 	MidVertex.Color = FillColor.ToFColor(false);
 
 	// Make sure the points all wind correctly relative to the mid point
@@ -134,10 +134,10 @@ static void PaintPolygon(
 		FComparePoints(const FSlateVertex& Mid) : MidPoint(Mid) {}
 		bool operator()(const FSlateVertex& A, const FSlateVertex& B) const
 		{
-			FVector2D DeltaA = FVector2D(A.Position) - FVector2D(MidPoint.Position);
-			FVector2D DeltaB = FVector2D(B.Position) - FVector2D(MidPoint.Position);
-			float AngleA = FMath::Atan2(DeltaA.Y, DeltaA.X);
-			float AngleB = FMath::Atan2(DeltaB.Y, DeltaB.X);
+			const FVector2D DeltaA = FVector2D(A.Position) - FVector2D(MidPoint.Position);
+			const FVector2D DeltaB = FVector2D(B.Position) - FVector2D(MidPoint.Position);
+			const double AngleA = FMath::Atan2(DeltaA.Y, DeltaA.X);
+			const double AngleB = FMath::Atan2(DeltaB.Y, DeltaB.X);
 			return AngleA < AngleB;
 		}
 		FSlateVertex MidPoint;
@@ -307,7 +307,7 @@ void SBlendSpaceGridWidget::Construct(const FArguments& InArgs)
 					.HAlign(HAlign_Left)
 					.BorderImage(FAppStyle::GetBrush("NoBorder"))
 					.DesiredSizeScale(FVector2D(1.0f, 1.0f))
-					.Padding_Lambda([&]() { return FMargin(GridMargin.Left + 6, 0, 0, 0) + GridRatioMargin; })
+					.Padding_Lambda([&]() { return FMargin(GridMargin.Left + 6.f, 0, 0, 0) + GridRatioMargin; })
 					[
 						SNew(SVerticalBox)
 						+ SVerticalBox::Slot()
@@ -326,7 +326,7 @@ void SBlendSpaceGridWidget::Construct(const FArguments& InArgs)
 									.ToolTipText(LOCTEXT("ShowTriangulation", "Show Triangulation"))
 									.OnClicked(this, &SBlendSpaceGridWidget::ToggleTriangulationVisibility)
 									.ButtonColorAndOpacity_Lambda([this]() -> FLinearColor { return bShowTriangulation ? FAppStyle::GetSlateColor("SelectionColor").GetSpecifiedColor() : FLinearColor::White; })
-									.ContentPadding(1)
+									.ContentPadding(1.f)
 									[
 										SNew(SImage)
 										.Image(FAppStyle::GetBrush("BlendSpaceEditor.ToggleTriangulation"))
@@ -347,7 +347,7 @@ void SBlendSpaceGridWidget::Construct(const FArguments& InArgs)
 									.ToolTipText(LOCTEXT("ShowAnimationNames", "Show Sample Names"))
 									.OnClicked(this, &SBlendSpaceGridWidget::ToggleShowAnimationNames)
 									.ButtonColorAndOpacity_Lambda([this]() -> FLinearColor { return bShowAnimationNames ? FAppStyle::GetSlateColor("SelectionColor").GetSpecifiedColor() : FLinearColor::White; })
-									.ContentPadding(1)
+									.ContentPadding(1.f)
 									[
 										SNew(SImage)
 										.Image(FAppStyle::GetBrush("BlendSpaceEditor.ToggleLabels"))
@@ -367,7 +367,7 @@ void SBlendSpaceGridWidget::Construct(const FArguments& InArgs)
 									SNew(SButton)
 									.ToolTipText(this, &SBlendSpaceGridWidget::GetFittingTypeButtonToolTipText)
 									.OnClicked(this, &SBlendSpaceGridWidget::ToggleFittingType)
-									.ContentPadding(1)
+									.ContentPadding(1.f)
 									.ButtonColorAndOpacity_Lambda([this]() -> FLinearColor { return bStretchToFit ? FAppStyle::GetSlateColor("SelectionColor").GetSpecifiedColor() : FLinearColor::White; })
 									[
 										SNew(SImage)
@@ -543,7 +543,7 @@ void SBlendSpaceGridWidget::PaintBackgroundAndGrid(const FGeometry& AllottedGeom
 		for (uint32 ParameterIndex = 0; ParameterIndex < BlendParametersToDraw; ++ParameterIndex)
 		{
 			const FBlendParameter& BlendParameter = BlendSpace->GetBlendParameter(ParameterIndex);
-			const float Steps = GridSize[ParameterIndex] / ( BlendParameter.GridNum);
+			const float Steps = static_cast<float>(GridSize[ParameterIndex] / ( BlendParameter.GridNum));
 
 			for (int32 Index = 1; Index < BlendParameter.GridNum; ++Index)
 			{			
@@ -648,7 +648,7 @@ void SBlendSpaceGridWidget::PaintSampleKeys(
 				OutDrawElements, SampleLayer, AllottedGeometry.ToPaintGeometry(GridPosition, KeySize), 
 				KeyBrush, ESlateDrawEffect::None, DrawColor );
 
-			float SampleLookupWeight = GetSampleLookupWeight(SampleIndex);
+			const float SampleLookupWeight = GetSampleLookupWeight(SampleIndex);
 			if (SampleLookupWeight <= SampleLookupWeightThreshold)
 			{
 				const FVector2D CirclePosition = SampleValueToScreenPosition(Sample.SampleValue);
@@ -664,7 +664,7 @@ void SBlendSpaceGridWidget::PaintSampleKeys(
 
 		// Always draw the filtered position which comes back from whatever is running
 		{
-			FVector2D GridPosition = SampleValueToScreenPosition(PreviewFilteredPosition) - (PreviewSize * .5f);
+			const FVector2D GridPosition = SampleValueToScreenPosition(PreviewFilteredPosition) - (PreviewSize * .5f);
 			FSlateDrawElement::MakeBox(
 				OutDrawElements, FilteredPositionLayer, AllottedGeometry.ToPaintGeometry(GridPosition, PreviewSize), 
 				PreviewBrush, ESlateDrawEffect::None, PreviewKeyColor.GetSpecifiedColor() * 0.7f);
@@ -672,7 +672,7 @@ void SBlendSpaceGridWidget::PaintSampleKeys(
 
 		// Always draw the preview position
 		{
-			FVector2D GridPosition = SampleValueToScreenPosition(PreviewPosition) - (PreviewSize * .5f);
+			const FVector2D GridPosition = SampleValueToScreenPosition(PreviewPosition) - (PreviewSize * .5f);
 			FSlateDrawElement::MakeBox(
 				OutDrawElements, PreviewPositionLayer, AllottedGeometry.ToPaintGeometry(GridPosition, PreviewSize), 
 				PreviewBrush, ESlateDrawEffect::None, PreviewKeyColor.GetSpecifiedColor());
@@ -882,8 +882,8 @@ void SBlendSpaceGridWidget::PaintTriangulation(
 				{
 					FVector2D A = NormalizedPositions[(Index + 2) % 3] - NormalizedPositions[(Index + 1) % 3];
 					FVector2D B = NormalizedPositions[Index] - NormalizedPositions[(Index + 1) % 3];
-					float Dot = A.GetSafeNormal() | B.GetSafeNormal();
-					float Area = 0.5f * FMath::Abs(B ^ A);
+					double Dot = A.GetSafeNormal() | B.GetSafeNormal();
+					double Area = 0.5f * FMath::Abs(B ^ A);
 					if (Dot > CriticalDot || Area < CriticalTriangulationArea)
 					{
 						TriangleFillColor = FLinearColor::Red;
@@ -1309,7 +1309,7 @@ FReply SBlendSpaceGridWidget::ProcessClick(const FGeometry& MyGeometry, const FP
 					// customization into account and return an incorrect (small) size
 					const FVector2D ExpectedSize(300, 300);
 					const FVector2D MenuPosition = FSlateApplication::Get().CalculatePopupWindowPosition(
-						FSlateRect(MousePosition.X, MousePosition.Y, MousePosition.X, MousePosition.Y), ExpectedSize, false);
+						FSlateRect(static_cast<float>(MousePosition.X), static_cast<float>(MousePosition.Y), static_cast<float>(MousePosition.X), static_cast<float>(MousePosition.Y)), ExpectedSize, false);
 
 					FSlateApplication::Get().PushMenu(
 						AsShared(),
@@ -1635,10 +1635,10 @@ void SBlendSpaceGridWidget::CalculateGridPoints()
 const FVector2D SBlendSpaceGridWidget::SnapScreenPositionToGrid(const FVector2D& InPosition, bool bForceSnap) const
 {
 	const int32 GridPointIndex = FindClosestGridPointIndexFromScreenPosition(InPosition);
-	FVector2D GridPoint = CachedGridPoints[GridPointIndex];
+	const FVector2D& GridPoint = CachedGridPoints[GridPointIndex];
 
-	bool bSnapX = bForceSnap || bSampleSnapToGrid[0];
-	bool bSnapY = bForceSnap || bSampleSnapToGrid[1];
+	const bool bSnapX = bForceSnap || bSampleSnapToGrid[0];
+	const bool bSnapY = bForceSnap || bSampleSnapToGrid[1];
 
 	return FVector2D
 	{
@@ -1652,15 +1652,15 @@ const FVector SBlendSpaceGridWidget::ScreenPositionToSampleValueWithSnapping(con
 	FVector SampleValue = ScreenPositionToSampleValue(InPosition, true);
 
 	const int32 GridPointIndex = FindClosestGridPointIndexFromScreenPosition(InPosition);
-	FVector GridPos = CachedSamplePoints[GridPointIndex];
+	const FVector GridPos = CachedSamplePoints[GridPointIndex];
 
-	bool bSnapX = bForceSnap || bSampleSnapToGrid[0];
-	bool bSnapY = bForceSnap || bSampleSnapToGrid[1];
+	const bool bSnapX = bForceSnap || bSampleSnapToGrid[0];
+	const bool bSnapY = bForceSnap || bSampleSnapToGrid[1];
 	return FVector
 	{
 		bSnapX ? GridPos.X : SampleValue.X,
 		bSnapY ? GridPos.Y : SampleValue.Y,
-		GridPos.Z,
+		0.f
 	};
 }
 
@@ -1671,12 +1671,12 @@ int32 SBlendSpaceGridWidget::FindClosestGridPointIndexFromScreenPosition(const F
 		FMath::Clamp(InPosition.X, CachedGridRectangle.Left, CachedGridRectangle.Right),
 		FMath::Clamp(InPosition.Y, CachedGridRectangle.Top, CachedGridRectangle.Bottom));
 	// Find the closest grid point
-	float Distance = FLT_MAX;
+	double Distance = TNumericLimits<double>::Max();
 	int32 GridPointIndex = INDEX_NONE;
 	for (int32 Index = 0; Index < CachedGridPoints.Num(); ++Index)
 	{
 		const FVector2D& GridPoint = CachedGridPoints[Index];
-		const float DistanceToGrid = FVector2D::DistSquared(GridPosition, GridPoint);
+		const double DistanceToGrid = FVector2D::DistSquared(GridPosition, GridPoint);
 		if (DistanceToGrid < Distance)
 		{
 			Distance = DistanceToGrid;
@@ -1700,10 +1700,10 @@ const FVector2D SBlendSpaceGridWidget::SampleValueToNormalizedPosition(const FVe
 
 const FVector2D SBlendSpaceGridWidget::SampleValueToScreenPosition(const FVector& SampleValue) const
 {
-	FVector2D NormalizedPosition = SampleValueToNormalizedPosition(SampleValue);
+	const FVector2D NormalizedPosition = SampleValueToNormalizedPosition(SampleValue);
 	const FVector2D GridSize = CachedGridRectangle.GetSize();
 	const FVector2D GridCorner = CachedGridRectangle.GetCenter() - 0.5f * GridSize;
-	FVector2D ScreenPosition = GridCorner + NormalizedPosition * CachedGridRectangle.GetSize();
+	const FVector2D ScreenPosition = GridCorner + NormalizedPosition * CachedGridRectangle.GetSize();
 	return ScreenPosition;	
 }
 
@@ -1734,7 +1734,7 @@ const FVector SBlendSpaceGridWidget::ScreenPositionToSampleValue(const FVector2D
 const FSlateRect SBlendSpaceGridWidget::GetGridRectangleFromGeometry(const FGeometry& MyGeometry)
 {
 	const float TopOffset = bReadOnly ? 0.0f : 20.0f; // Ideally we'd get the size of the buttons (showing the label/triangulation etc)
-	FSlateRect WindowRect = FSlateRect(0, TopOffset, MyGeometry.GetLocalSize().X, MyGeometry.GetLocalSize().Y);
+	FSlateRect WindowRect = FSlateRect(0, TopOffset, static_cast<float>(MyGeometry.GetLocalSize().X), static_cast<float>(MyGeometry.GetLocalSize().Y));
 	if (!bStretchToFit)
 	{
 		UpdateGridRatioMargin(WindowRect.GetSize());
@@ -1746,7 +1746,7 @@ const FSlateRect SBlendSpaceGridWidget::GetGridRectangleFromGeometry(const FGeom
 bool SBlendSpaceGridWidget::IsSampleValueWithinMouseRange(const FVector& SampleValue, float& OutDistance) const
 {
 	const FVector2D GridPosition = SampleValueToScreenPosition(SampleValue);
-	OutDistance = FVector2D::Distance(LocalMousePosition, GridPosition);
+	OutDistance = static_cast<float>(FVector2D::Distance(LocalMousePosition, GridPosition));
 	return (OutDistance < ClickAndHighlightThreshold);
 }
 
@@ -2152,7 +2152,7 @@ TOptional<float> SBlendSpaceGridWidget::GetInputBoxValue(const int32 ParameterIn
 		if (SelectedSampleIndex != INDEX_NONE && SelectedSampleIndex < BlendSpace->GetNumberOfBlendSamples())
 		{
 			const FBlendSample& BlendSample = BlendSpace->GetBlendSample(SelectedSampleIndex);
-			ReturnValue = BlendSample.SampleValue[ParameterIndex];
+			ReturnValue = static_cast<float>(BlendSample.SampleValue[ParameterIndex]);
 		}
 	}
 	return ReturnValue;
@@ -2161,19 +2161,19 @@ TOptional<float> SBlendSpaceGridWidget::GetInputBoxValue(const int32 ParameterIn
 TOptional<float> SBlendSpaceGridWidget::GetInputBoxMinValue(const int32 ParameterIndex) const
 {
 	checkf(ParameterIndex < 3, TEXT("Invalid parameter index, suppose to be within FVector array range"));
-	return SampleValueMin[ParameterIndex];
+	return static_cast<float>(SampleValueMin[ParameterIndex]);
 }
 
 TOptional<float> SBlendSpaceGridWidget::GetInputBoxMaxValue(const int32 ParameterIndex) const
 {
 	checkf(ParameterIndex < 3, TEXT("Invalid parameter index, suppose to be within FVector array range"));
-	return SampleValueMax[ParameterIndex];
+	return static_cast<float>(SampleValueMax[ParameterIndex]);
 }
 
 float SBlendSpaceGridWidget::GetInputBoxDelta(const int32 ParameterIndex) const
 {
 	checkf(ParameterIndex < 3, TEXT("Invalid parameter index, suppose to be within FVector array range"));
-	return SampleGridDelta[ParameterIndex];
+	return static_cast<float>(SampleGridDelta[ParameterIndex]);
 }
 
 void SBlendSpaceGridWidget::OnInputBoxValueCommited(const float NewValue, ETextCommit::Type CommitType, const int32 ParameterIndex)
@@ -2200,11 +2200,11 @@ void SBlendSpaceGridWidget::OnInputBoxValueChanged(const float NewValue, const i
 		// Calculate snapped value
 		if (bSampleSnapToGrid[ParameterIndex])
 		{
-			const float MinOffset = NewValue - SampleValueMin[ParameterIndex];
-			float GridSteps = MinOffset / SampleGridDelta[ParameterIndex];
-			int32 FlooredSteps = FMath::FloorToInt(GridSteps);
+			const double MinOffset = NewValue - SampleValueMin[ParameterIndex];
+			double GridSteps = MinOffset / SampleGridDelta[ParameterIndex];
+			int64 FlooredSteps = FMath::FloorToInt(GridSteps);
 			GridSteps -= FlooredSteps;
-			FlooredSteps = (GridSteps > .5f) ? FlooredSteps + 1 : FlooredSteps;
+			FlooredSteps = (GridSteps > .5) ? FlooredSteps + 1 : FlooredSteps;
 
 			// Temporary snap this value to closest point on grid (since the spin box delta does not provide the desired functionality)
 			SampleValue[ParameterIndex] = SampleValueMin[ParameterIndex] + (FlooredSteps * SampleGridDelta[ParameterIndex]);
@@ -2289,13 +2289,13 @@ void SBlendSpaceGridWidget::UpdateGridRatioMargin(const FVector2D& GeometrySize)
 		GridRatioMargin.Top = GridRatioMargin.Bottom = GridRatioMargin.Left = GridRatioMargin.Right = 0.0f;
 		if (GeometrySize.Y > GeometrySize.X)
 		{
-			const float Difference = GeometrySize.Y - GeometrySize.X;
-			GridRatioMargin.Top = GridRatioMargin.Bottom = Difference * 0.5f;
+			const double Difference = GeometrySize.Y - GeometrySize.X;
+			GridRatioMargin.Top = GridRatioMargin.Bottom = static_cast<float>(Difference) * 0.5f;
 		}
 		else if (GeometrySize.X > GeometrySize.Y)
 		{
-			const float Difference = GeometrySize.X - GeometrySize.Y;
-			GridRatioMargin.Left = GridRatioMargin.Right = Difference * 0.5f;
+			const double Difference = GeometrySize.X - GeometrySize.Y;
+			GridRatioMargin.Left = GridRatioMargin.Right = static_cast<float>(Difference) * 0.5f;
 		}
 	}
 }
@@ -2343,23 +2343,23 @@ void SBlendSpaceGridWidget::UpdateCachedBlendParameterData()
 		const TSharedRef< FSlateFontMeasure > FontMeasure = FSlateApplication::Get().GetRenderer()->GetFontMeasureService();
 		MaxVerticalAxisTextWidth = HorizontalAxisMaxTextWidth = MaxHorizontalAxisTextHeight = 0.0f;
 		FVector2D TextSize = FontMeasure->Measure(ParameterYName, FontInfo);	
-		MaxVerticalAxisTextWidth = FMath::Max(MaxVerticalAxisTextWidth, TextSize.X);
+		MaxVerticalAxisTextWidth = FMath::Max(MaxVerticalAxisTextWidth, static_cast<float>(TextSize.X));
 
 		TextSize = FontMeasure->Measure(FString::SanitizeFloat(SampleValueMin.Y), FontInfo);
-		MaxVerticalAxisTextWidth = FMath::Max(MaxVerticalAxisTextWidth, TextSize.X);
+		MaxVerticalAxisTextWidth = FMath::Max(MaxVerticalAxisTextWidth, static_cast<float>(TextSize.X));
 
 		TextSize = FontMeasure->Measure(FString::SanitizeFloat(SampleValueMax.Y), FontInfo);
-		MaxVerticalAxisTextWidth = FMath::Max(MaxVerticalAxisTextWidth, TextSize.X);
+		MaxVerticalAxisTextWidth = FMath::Max(MaxVerticalAxisTextWidth, static_cast<float>(TextSize.X));
 	
 		TextSize = FontMeasure->Measure(ParameterXName, FontInfo);
-		MaxHorizontalAxisTextHeight = FMath::Max(MaxHorizontalAxisTextHeight, TextSize.Y);
+		MaxHorizontalAxisTextHeight = FMath::Max(MaxHorizontalAxisTextHeight, static_cast<float>(TextSize.Y));
 
 		TextSize = FontMeasure->Measure(FString::SanitizeFloat(SampleValueMin.X), FontInfo);
-		MaxHorizontalAxisTextHeight = FMath::Max(MaxHorizontalAxisTextHeight, TextSize.Y);
+		MaxHorizontalAxisTextHeight = FMath::Max(MaxHorizontalAxisTextHeight, static_cast<float>(TextSize.Y));
 
 		TextSize = FontMeasure->Measure(FString::SanitizeFloat(SampleValueMax.X), FontInfo);
-		MaxHorizontalAxisTextHeight = FMath::Max(MaxHorizontalAxisTextHeight, TextSize.Y);
-		HorizontalAxisMaxTextWidth = TextSize.X;
+		MaxHorizontalAxisTextHeight = FMath::Max(MaxHorizontalAxisTextHeight, static_cast<float>(TextSize.Y));
+		HorizontalAxisMaxTextWidth = static_cast<float>(TextSize.X);
 	}
 }
 
diff --git a/Engine/Source/Editor/Persona/Private/SAnimationEditorViewport.cpp b/Engine/Source/Editor/Persona/Private/SAnimationEditorViewport.cpp
--- a/Engine/Source/Editor/Persona/Private/SAnimationEditorViewport.cpp
+++ b/Engine/Source/Editor/Persona/Private/SAnimationEditorViewport.cpp
@@ -1837,7 +1837,7 @@ float SAnimationEditorViewportTabBody::GetViewMaxInput() const
 		}
 		else if (PreviewComponent->GetAnimInstance() != NULL)
 		{
-			return PreviewComponent->GetAnimInstance()->LifeTimer;
+			return static_cast<float>(PreviewComponent->GetAnimInstance()->LifeTimer);
 		}
 	}
 
diff --git a/Engine/Source/Editor/Persona/Private/SAnimationScrubPanel.cpp b/Engine/Source/Editor/Persona/Private/SAnimationScrubPanel.cpp
--- a/Engine/Source/Editor/Persona/Private/SAnimationScrubPanel.cpp
+++ b/Engine/Source/Editor/Persona/Private/SAnimationScrubPanel.cpp
@@ -362,7 +362,7 @@ float SAnimationScrubPanel::GetSequenceLength() const
 		FAnimBlueprintDebugData* DebugData;
 		if (GetAnimBlueprintDebugData(/*out*/ Instance, /*out*/ DebugData))
 		{
-			return Instance->LifeTimer;
+			return static_cast<float>(Instance->LifeTimer);
 		}
 	}
 
@@ -406,7 +406,7 @@ float SAnimationScrubPanel::GetScrubValue() const
 		FAnimBlueprintDebugData* DebugData;
 		if (GetAnimBlueprintDebugData(/*out*/ Instance, /*out*/ DebugData))
 		{
-			return Instance->CurrentLifeTimerScrubPosition;
+			return static_cast<float>(Instance->CurrentLifeTimerScrubPosition);
 		}
 	}
 
diff --git a/Engine/Source/Editor/Persona/Private/SAnimationSequenceBrowser.cpp b/Engine/Source/Editor/Persona/Private/SAnimationSequenceBrowser.cpp
--- a/Engine/Source/Editor/Persona/Private/SAnimationSequenceBrowser.cpp
+++ b/Engine/Source/Editor/Persona/Private/SAnimationSequenceBrowser.cpp
@@ -1324,11 +1324,11 @@ TSharedRef<SToolTip> SAnimationSequenceBrowser::CreateCustomAssetToolTip(FAssetD
 
 	TSharedPtr<SHorizontalBox> ContentBox = nullptr;
 	TSharedRef<SToolTip> ToolTipWidget = SNew(SToolTip)
-	.TextMargin(1)
+	.TextMargin(1.f)
 	.BorderImage(FAppStyle::GetBrush("ContentBrowser.TileViewTooltip.ToolTipBorder"))
 	[
 		SNew(SBorder)
-		.Padding(6)
+		.Padding(6.f)
 		.BorderImage(FAppStyle::GetBrush("ContentBrowser.TileViewTooltip.NonContentBorder"))
 		[
 			SNew(SVerticalBox)
@@ -1337,7 +1337,7 @@ TSharedRef<SToolTip> SAnimationSequenceBrowser::CreateCustomAssetToolTip(FAssetD
 			.Padding(0,0,0,4)
 			[
 				SNew(SBorder)
-				.Padding(6)
+				.Padding(6.f)
 				.BorderImage(FAppStyle::GetBrush("ContentBrowser.TileViewTooltip.ContentBorder"))
 				[
 					SNew(SBox)
@@ -1357,7 +1357,7 @@ TSharedRef<SToolTip> SAnimationSequenceBrowser::CreateCustomAssetToolTip(FAssetD
 				.AutoWidth()
 				[
 					SNew(SBorder)
-					.Padding(6)
+					.Padding(6.f)
 					.Visibility(AssetClass->IsChildOf<UAnimationAsset>() ? EVisibility::Visible : EVisibility::Collapsed)
 					.BorderImage(FAppStyle::GetBrush("ContentBrowser.TileViewTooltip.ContentBorder"))
 					[
@@ -1384,10 +1384,10 @@ TSharedRef<SToolTip> SAnimationSequenceBrowser::CreateCustomAssetToolTip(FAssetD
 
 	// add an extra section to the tooltip for it.
 	ContentBox->AddSlot()
-	.Padding(AssetClass->IsChildOf<UAnimationAsset>() ? 4 : 0, 0, 0, 0)
+	.Padding(AssetClass->IsChildOf<UAnimationAsset>() ? 4.f : 0, 0, 0, 0)
 	[
 		SNew(SBorder)
-		.Padding(6)
+		.Padding(6.f)
 		.BorderImage(FAppStyle::GetBrush("ContentBrowser.TileViewTooltip.ContentBorder"))
 		[
 			DescriptionBox
@@ -1460,9 +1460,9 @@ bool SAnimationSequenceBrowser::OnVisualizeAssetToolTip(const TSharedPtr<SWidget
 			PreviewComponent->EnablePreview(true, Asset);
 			PreviewComponent->PreviewInstance->PlayAnim(true);
 
-			FBoxSphereBounds MeshImportedBounds = MeshToUse->GetImportedBounds();
-			float HalfFov = FMath::DegreesToRadians(ViewportClient->ViewFOV) / 2.0f;
-			float TargetDist = MeshImportedBounds.SphereRadius / FMath::Tan(HalfFov);
+			const FBoxSphereBounds MeshImportedBounds = MeshToUse->GetImportedBounds();
+			const float HalfFov = FMath::DegreesToRadians(ViewportClient->ViewFOV) / 2.0f;
+			const float TargetDist = static_cast<float>(MeshImportedBounds.SphereRadius / FMath::Tan(HalfFov));
 
 			ViewportClient->SetViewRotation(FRotator(0.0f, -45.0f, 0.0f));
 			ViewportClient->SetViewLocationForOrbiting(FVector(0.0f, 0.0f, MeshImportedBounds.BoxExtent.Z / 2.0f), TargetDist);
diff --git a/Engine/Source/Editor/Persona/Private/SAssetFamilyShortcutBar.cpp b/Engine/Source/Editor/Persona/Private/SAssetFamilyShortcutBar.cpp
--- a/Engine/Source/Editor/Persona/Private/SAssetFamilyShortcutBar.cpp
+++ b/Engine/Source/Editor/Persona/Private/SAssetFamilyShortcutBar.cpp
@@ -87,7 +87,7 @@ public:
 				.IsChecked(this, &SAssetShortcut::GetCheckState)
 				.Visibility(this, &SAssetShortcut::GetSoloButtonVisibility)
 				.ToolTipText(this, &SAssetShortcut::GetButtonTooltip)
-				.Padding(0.0)
+				.Padding(0.0f)
 				[
 					SNew(SOverlay)
 
@@ -123,7 +123,7 @@ public:
 				.IsChecked(this, &SAssetShortcut::GetCheckState)
 				.Visibility(this, &SAssetShortcut::GetComboButtonVisibility)
 				.ToolTipText(this, &SAssetShortcut::GetButtonTooltip)
-				.Padding(0.0)
+				.Padding(0.0f)
 				[
 					SNew(SOverlay)
 
@@ -154,7 +154,7 @@ public:
 			[
 				SNew(SSeparator)
 				.Visibility(this, &SAssetShortcut::GetComboVisibility)
-				.Thickness(1.0)
+				.Thickness(1.0f)
 				.Orientation(EOrientation::Orient_Vertical)
 			]
 			+SHorizontalBox::Slot()
@@ -296,8 +296,8 @@ public:
 
 				MenuBuilder.AddWidget(
 					SNew(SBox)
-					.WidthOverride(300)
-					.HeightOverride(600)
+					.WidthOverride(300.f)
+					.HeightOverride(600.f)
 					[
 						ContentBrowserModule.Get().CreateAssetPicker(AssetPickerConfig)
 					],
diff --git a/Engine/Source/Editor/Persona/Private/SCreateClothingSettingsPanel.cpp b/Engine/Source/Editor/Persona/Private/SCreateClothingSettingsPanel.cpp
--- a/Engine/Source/Editor/Persona/Private/SCreateClothingSettingsPanel.cpp
+++ b/Engine/Source/Editor/Persona/Private/SCreateClothingSettingsPanel.cpp
@@ -76,18 +76,18 @@ void SCreateClothingSettingsPanel::Construct(const FArguments& InArgs)
 
 			+SVerticalBox::Slot()
 			.MaxHeight(500.0f)
-			.Padding(2)
+			.Padding(2.f)
 			[
 				StructureDetailsView->GetWidget()->AsShared()
 			]
 
 			+SVerticalBox::Slot()
 			.AutoHeight()
-			.Padding(2)
+			.Padding(2.f)
 			.HAlign(HAlign_Right)
 			[
 				SNew(SUniformGridPanel)
-				.SlotPadding(2)
+				.SlotPadding(2.f)
 				+SUniformGridPanel::Slot(0, 0)
 				[
 					SNew(SButton)
@@ -334,7 +334,7 @@ void FClothCreateSettingsCustomization::CustomizeDetails(IDetailLayoutBuilder& D
 		[
 			SNew(SComboButton)
 			.OnGetMenuContent(this, &FClothCreateSettingsCustomization::OnGetTargetAssetMenu)
-			.ContentPadding(2)
+			.ContentPadding(2.f)
 			.ButtonStyle(FAppStyle::Get(), "PropertyEditor.AssetComboStyle")
 			.ForegroundColor(FAppStyle::GetColor("PropertyEditor.AssetName.ColorAndOpacity"))
 			.CollapseMenuOnParentFocus(true)
@@ -358,7 +358,7 @@ void FClothCreateSettingsCustomization::CustomizeDetails(IDetailLayoutBuilder& D
 			SNew(SComboButton)
 			.IsEnabled(this, &FClothCreateSettingsCustomization::CanSelectLod)
 			.OnGetMenuContent(this, &FClothCreateSettingsCustomization::OnGetTargetLodMenu)
-			.ContentPadding(2)
+			.ContentPadding(2.f)
 			.CollapseMenuOnParentFocus(true)
 			.ButtonContent()
 			[
diff --git a/Engine/Source/Editor/Persona/Private/SPoseAssetNameWidget.cpp b/Engine/Source/Editor/Persona/Private/SPoseAssetNameWidget.cpp
--- a/Engine/Source/Editor/Persona/Private/SPoseAssetNameWidget.cpp
+++ b/Engine/Source/Editor/Persona/Private/SPoseAssetNameWidget.cpp
@@ -17,7 +17,7 @@ void SPoseAssetNameWidget::Construct(const FArguments& InArgs)
 			.OnGenerateWidget(this, &SPoseAssetNameWidget::MakeBasePoseComboWidget)
 			.OnSelectionChanged(this, &SPoseAssetNameWidget::SelectionChanged)
 			.OnComboBoxOpening(this, &SPoseAssetNameWidget::OnBasePoseComboOpening)
-			.ContentPadding(3)
+			.ContentPadding(3.f)
 			.Content()
 			[
 				SNew(STextBlock)
diff --git a/Engine/Source/Editor/Persona/Private/SPoseWatchManagerTreeView.cpp b/Engine/Source/Editor/Persona/Private/SPoseWatchManagerTreeView.cpp
--- a/Engine/Source/Editor/Persona/Private/SPoseWatchManagerTreeView.cpp
+++ b/Engine/Source/Editor/Persona/Private/SPoseWatchManagerTreeView.cpp
@@ -245,7 +245,7 @@ TSharedRef<SWidget> SPoseWatchManagerTreeRow::GenerateWidgetForColumn(const FNam
 
 				+ SHorizontalBox::Slot()
 				.AutoWidth()
-				.Padding(6, 0, 0, 0)
+				.Padding(6.f, 0.f, 0.f, 0.f)
 				[
 					SNew(SExpanderArrow, SharedThis(this)).IndentAmount(12)
 				]
@@ -302,12 +302,12 @@ void SPoseWatchManagerTreeRow::Tick(const FGeometry& AllottedGeometry, const dou
 	// Update highlight 'target' effect
 	{
 		const float HighlightLeftX = HighlightRectLeftOffset;
-		const float HighlightRightX = HighlightRectRightOffset + AllottedGeometry.GetLocalSize().X;
+		const float HighlightRightX = HighlightRectRightOffset + static_cast<float>(AllottedGeometry.GetLocalSize().X);
 
 		HighlightTargetLeftSpring.SetTarget(HighlightLeftX);
 		HighlightTargetRightSpring.SetTarget(HighlightRightX);
 
-		float TimeSinceHighlightInteraction = (float)(InCurrentTime - LastHighlightInteractionTime);
+		const float TimeSinceHighlightInteraction = static_cast<float>(InCurrentTime - LastHighlightInteractionTime);
 		if (TimeSinceHighlightInteraction <= HighlightTargetEffectDuration || bShouldAppearFocused)
 		{
 			HighlightTargetLeftSpring.Tick(InDeltaTime);
@@ -349,7 +349,7 @@ int32 SPoseWatchManagerTreeRow::OnPaint(const FPaintArgs& Args, const FGeometry&
 
 		// Compute the bounds offset of the highlight target from where the highlight target spring
 		// extents currently lie.  This is used to "grow" or "shrink" the highlight as needed.
-		const float LabelChangedAnimOffset = LabelChangedAnimOffsetPercent * AllottedGeometry.GetLocalSize().Y;
+		const float LabelChangedAnimOffset = LabelChangedAnimOffsetPercent * static_cast<float>(AllottedGeometry.GetLocalSize().Y);
 
 		// Choose an offset amount depending on whether we're highlighting, or clearing highlight
 		const float EffectOffset = EffectAlpha * LabelChangedAnimOffset;
@@ -357,7 +357,7 @@ int32 SPoseWatchManagerTreeRow::OnPaint(const FPaintArgs& Args, const FGeometry&
 		const float HighlightLeftX = HighlightTargetLeftSpring.GetPosition() - EffectOffset;
 		const float HighlightRightX = HighlightTargetRightSpring.GetPosition() + EffectOffset;
 		const float HighlightTopY = 0.0f - LabelChangedAnimOffset;
-		const float HighlightBottomY = AllottedGeometry.GetLocalSize().Y + EffectOffset;
+		const float HighlightBottomY = static_cast<float>(AllottedGeometry.GetLocalSize().Y) + EffectOffset;
 
 		const FVector2D DrawPosition = FVector2D(HighlightLeftX, HighlightTopY);
 		const FVector2D DrawSize = FVector2D(HighlightRightX - HighlightLeftX, HighlightBottomY - HighlightTopY);
diff --git a/Engine/Source/Editor/Persona/Private/SRetargetSourceWindow.cpp b/Engine/Source/Editor/Persona/Private/SRetargetSourceWindow.cpp
--- a/Engine/Source/Editor/Persona/Private/SRetargetSourceWindow.cpp
+++ b/Engine/Source/Editor/Persona/Private/SRetargetSourceWindow.cpp
@@ -442,16 +442,16 @@ void SRetargetSourceWindow::OnAddRetargetSource()
 	AssetPickerConfig.Filter.TagsAndValues.Add(USkeletalMesh::GetSkeletonMemberName(), SkeletonString);
 
 	TSharedRef<SWidget> Widget = SNew(SBox)
-		.WidthOverride(384)
-		.HeightOverride(768)
+		.WidthOverride(384.f)
+		.HeightOverride(768.f)
 		[
 			SNew(SBorder)
 			.BorderBackgroundColor(FLinearColor(0.25f, 0.25f, 0.25f, 1.f))
-			.Padding( 2 )
+			.Padding( 2.f )
 			[
 				SNew(SBorder)
 				.BorderImage( FAppStyle::GetBrush("ToolPanel.GroupBorder") )
-				.Padding( 8 )
+				.Padding( 8.f )
 				[
 					ContentBrowserModule.Get().CreateAssetPicker(AssetPickerConfig)
 				]
@@ -680,16 +680,16 @@ FReply SCompatibleSkeletons::OnAddSkeletonClicked()
 	AssetPickerConfig.InitialAssetViewType = EAssetViewType::Tile;
 
 	const TSharedRef<SWidget> Widget = SNew(SBox)
-		.WidthOverride(384)
-		.HeightOverride(768)
+		.WidthOverride(384.f)
+		.HeightOverride(768.f)
 		[
 			SNew(SBorder)
 			.BorderBackgroundColor(FLinearColor(0.25f, 0.25f, 0.25f, 1.f))
-			.Padding( 2 )
+			.Padding( 2.f )
 			[
 				SNew(SBorder)
 				.BorderImage( FAppStyle::GetBrush("ToolPanel.GroupBorder") )
-				.Padding( 8 )
+				.Padding( 8.f )
 				[
 					ContentBrowserModule.Get().CreateAssetPicker(AssetPickerConfig)
 				]
diff --git a/Engine/Source/Editor/Persona/Private/SRigPicker.cpp b/Engine/Source/Editor/Persona/Private/SRigPicker.cpp
--- a/Engine/Source/Editor/Persona/Private/SRigPicker.cpp
+++ b/Engine/Source/Editor/Persona/Private/SRigPicker.cpp
@@ -10,7 +10,7 @@
 
 #define LOCTEXT_NAMESPACE "RigPicker"
 
-static const FVector2D ContentBrowserWindowSize(300.0f, 300.0f);
+static const FVector2f ContentBrowserWindowSize(300.0f, 300.0f);
 
 URig* SRigPicker::EngineHumanoidRig = NULL;
 
diff --git a/Engine/Source/Editor/Persona/Private/SSkinWeightProfileImportOptions.cpp b/Engine/Source/Editor/Persona/Private/SSkinWeightProfileImportOptions.cpp
--- a/Engine/Source/Editor/Persona/Private/SSkinWeightProfileImportOptions.cpp
+++ b/Engine/Source/Editor/Persona/Private/SSkinWeightProfileImportOptions.cpp
@@ -134,10 +134,10 @@ void SSkinWeightProfileImportOptions::Construct(const FArguments& InArgs)
 		+ SVerticalBox::Slot()
 		.AutoHeight()
 		.HAlign(HAlign_Right)
-		.Padding(2)
+		.Padding(2.f)
 		[
 			SNew(SUniformGridPanel)
-			.SlotPadding(2)
+			.SlotPadding(2.f)
 			+ SUniformGridPanel::Slot(0, 0)
 			[
 				SAssignNew(ImportButton, SButton)
diff --git a/Engine/Source/Editor/Persona/Private/STimingTrack.cpp b/Engine/Source/Editor/Persona/Private/STimingTrack.cpp
--- a/Engine/Source/Editor/Persona/Private/STimingTrack.cpp
+++ b/Engine/Source/Editor/Persona/Private/STimingTrack.cpp
@@ -133,33 +133,33 @@ void STimingTrack::OnArrangeChildren(const FGeometry& AllottedGeometry, FArrange
 		// Separation of the nodes on the track
 		const float NodeSeparation = 3.0f;
 
-		for(NodeIsland& Island : Islands)
+		for(const NodeIsland& Island : Islands)
 		{
 			bResolved = true;
 
 			// Island resolution
-			int32 NumIslandNodes = Island.Nodes.Num();
-			float Width = FMath::Max<float>((NumIslandNodes - 1), 0.0f) * NodeSeparation;
+			const int32 NumIslandNodes = Island.Nodes.Num();
+			float Width = FMath::Max<float>(static_cast<float>(NumIslandNodes - 1), 0.0f) * NodeSeparation;
 			float Centre = 0.0f;
 
-			for(NodeData* Node : Island.Nodes)
+			for(const NodeData* Node : Island.Nodes)
 			{
-				Width += Node->ActualRect.GetSize().X;
-				Centre += Node->ActualRect.GetCenter().X;
+				Width += static_cast<float>(Node->ActualRect.GetSize().X);
+				Centre += static_cast<float>(Node->ActualRect.GetCenter().X);
 			}
 			Centre /= NumIslandNodes;
 
-			float Begin = Centre - Width / 2.0f;
+			const float Begin = Centre - Width / 2.0f;
 			for(int32 NodeIdx = 0 ; NodeIdx < NumIslandNodes; ++NodeIdx)
 			{
 				FBox2D& NodeBox = Island.Nodes[NodeIdx]->ActualRect;
-				float NodeWidth = NodeBox.GetSize().X;
-				float SeparationOffset = NodeIdx * NodeSeparation;
+				const float NodeWidth = static_cast<float>(NodeBox.GetSize().X);
+				const float SeparationOffset = NodeIdx * NodeSeparation;
 				float PositionOffset = 0.0f;
 			
 				for(int32 PositionNodeIdx = 0 ; PositionNodeIdx < NodeIdx ; ++PositionNodeIdx)
 				{
-					PositionOffset += Island.Nodes[PositionNodeIdx]->ActualRect.GetSize().X;
+					PositionOffset += static_cast<float>(Island.Nodes[PositionNodeIdx]->ActualRect.GetSize().X);
 				}
 			
 				FBox2D& OriginalRect = Island.Nodes[NodeIdx]->ActualRect;
diff --git a/Engine/Source/Editor/Persona/Private/STrack.cpp b/Engine/Source/Editor/Persona/Private/STrack.cpp
--- a/Engine/Source/Editor/Persona/Private/STrack.cpp
+++ b/Engine/Source/Editor/Persona/Private/STrack.cpp
@@ -513,16 +513,16 @@ int32 STrack::OnPaint( const FPaintArgs& Args, const FGeometry& AllottedGeometry
 // drag drop relationship
 FReply STrack::OnDrop( const FGeometry& MyGeometry, const FDragDropEvent& DragDropEvent )
 {
-	FVector2D CursorPos = MyGeometry.AbsoluteToLocal(DragDropEvent.GetScreenSpacePosition());
-	float CusorDataPos = LocalToDataX(CursorPos.X, MyGeometry);
+	const FVector2D CursorPos = MyGeometry.AbsoluteToLocal(DragDropEvent.GetScreenSpacePosition());
+	const float CusorDataPos = LocalToDataX(static_cast<float>(CursorPos.X), MyGeometry);
 
 	// Handle TrackNodes that were dropped
-	TSharedPtr<FTrackNodeDragDropOp> DragDropOp = DragDropEvent.GetOperationAs<FTrackNodeDragDropOp>();
+	const TSharedPtr<FTrackNodeDragDropOp> DragDropOp = DragDropEvent.GetOperationAs<FTrackNodeDragDropOp>();
 	if (DragDropOp.IsValid())
 	{
-		TSharedPtr<STrackNode> TrackNode = DragDropOp->OriginalTrackNode.Pin();
+		const TSharedPtr<STrackNode> TrackNode = DragDropOp->OriginalTrackNode.Pin();
 
-		float DataPos = GetNodeDragDropDataPos(MyGeometry, DragDropEvent);
+		const float DataPos = GetNodeDragDropDataPos(MyGeometry, DragDropEvent);
 		TrackNode->OnTrackNodeDragged.ExecuteIfBound( DataPos );
 		TrackNode->OnTrackNodeDropped.ExecuteIfBound();
 	}
@@ -558,21 +558,21 @@ FReply STrack::OnDragDetected(const FGeometry& MyGeometry, const FPointerEvent&
 	return FReply::Unhandled();
 }
 
-float STrack::GetNodeDragDropDataPos( const FGeometry& MyGeometry, const FDragDropEvent& DragDropEvent )
+float STrack::GetNodeDragDropDataPos( const FGeometry& MyGeometry, const FDragDropEvent& DragDropEvent ) const
 {
 	float DataPos = 0.f;
-	TSharedPtr<FTrackNodeDragDropOp> DragDropOp = StaticCastSharedPtr<FTrackNodeDragDropOp>(DragDropEvent.GetOperation());
+	const TSharedPtr<FTrackNodeDragDropOp> DragDropOp = StaticCastSharedPtr<FTrackNodeDragDropOp>(DragDropEvent.GetOperation());
 	if(DragDropOp.IsValid())
 	{
-		TSharedPtr<STrackNode> TrackNode = DragDropOp->OriginalTrackNode.Pin();
+		const TSharedPtr<STrackNode> TrackNode = DragDropOp->OriginalTrackNode.Pin();
 		if(TrackNode.IsValid())
 		{
-			FVector2D CursorPos = MyGeometry.AbsoluteToLocal(TrackNode->GetDragDropScreenSpacePosition(MyGeometry, DragDropEvent));
-			DataPos = LocalToDataX(CursorPos.X, MyGeometry);
+			const FVector2D CursorPos = MyGeometry.AbsoluteToLocal(TrackNode->GetDragDropScreenSpacePosition(MyGeometry, DragDropEvent));
+			DataPos = LocalToDataX(static_cast<float>(CursorPos.X), MyGeometry);
 			if(TrackNode->SnapToDragBars())
 			{
-				float OriginalX = DataPos;
-				DataPos = GetSnappedPosForLocalPos(MyGeometry, CursorPos.X);
+				const float OriginalX = DataPos;
+				DataPos = GetSnappedPosForLocalPos(MyGeometry, static_cast<float>(CursorPos.X));
 				TrackNode->OnSnapNodeDataPosition(OriginalX, DataPos);
 			}
 		}
@@ -624,8 +624,8 @@ FReply STrack::OnMouseMove( const FGeometry& MyGeometry, const FPointerEvent& Mo
 	if (bDraggingBar && OnBarDrag.IsBound())
 	{
 		/** Update drag bar position if we are dragging */
-		FVector2D CursorPos = MyGeometry.AbsoluteToLocal( MouseEvent.GetScreenSpacePosition() );
-		float NewDataPos = FMath::Clamp( LocalToDataX(CursorPos.X, MyGeometry), TrackMinValue.Get(), TrackMaxValue.Get() );
+		const FVector2D CursorPos = MyGeometry.AbsoluteToLocal( MouseEvent.GetScreenSpacePosition() );
+		const float NewDataPos = FMath::Clamp( static_cast<float>(LocalToDataX(static_cast<float>(CursorPos.X), MyGeometry)), TrackMinValue.Get(), TrackMaxValue.Get() );
 		OnBarDrag.Execute(DraggableBarIndex, NewDataPos);
 
 		// Update details panel
@@ -701,14 +701,14 @@ TSharedPtr<SWidget> STrack::SummonContextMenu(const FGeometry& MyGeometry, const
 {
 	bool SummonedContextMenu = false;
 
-	const bool bCloseWindowAfterMenuSelection = true;
+	constexpr bool bCloseWindowAfterMenuSelection = true;
 	FMenuBuilder MenuBuilder( bCloseWindowAfterMenuSelection, EditorActions );
 
-	FVector2D CursorPos = MouseEvent.GetScreenSpacePosition();
-	float DataPos = LocalToDataX( MyGeometry.AbsoluteToLocal(CursorPos).X, MyGeometry );
+	const FVector2D CursorPos = MouseEvent.GetScreenSpacePosition();
+	const float DataPos = LocalToDataX( static_cast<float>(MyGeometry.AbsoluteToLocal(CursorPos).X), MyGeometry );
 
 	// Context menu for a node
-	int NotifyIndex = GetHitNode(MyGeometry, MyGeometry.AbsoluteToLocal(CursorPos));
+	const int32 NotifyIndex = GetHitNode(MyGeometry, MyGeometry.AbsoluteToLocal(CursorPos));
 	if(NotifyIndex != INDEX_NONE)
 	{
 		if(TrackNodes[NotifyIndex]->OnNodeRightClickContextMenu.IsBound())
diff --git a/Engine/Source/Editor/Persona/Private/STrack.h b/Engine/Source/Editor/Persona/Private/STrack.h
--- a/Engine/Source/Editor/Persona/Private/STrack.h
+++ b/Engine/Source/Editor/Persona/Private/STrack.h
@@ -303,7 +303,7 @@ protected:
 	
 	bool									GetDraggableBarSnapPosition(const FGeometry& MyGeometry, float &OutPosition) const;
 	virtual TSharedPtr<SWidget>				SummonContextMenu(const FGeometry& MyGeometry, const FPointerEvent& MouseEvent);
-	float									GetNodeDragDropDataPos( const FGeometry& MyGeometry, const FDragDropEvent& DragDropEvent );
+	float									GetNodeDragDropDataPos( const FGeometry& MyGeometry, const FDragDropEvent& DragDropEvent ) const;
 	float									GetSnappedPosForLocalPos( const FGeometry& MyGeometry, float TrackPos) const;
 	void									UpdateDraggableBarIndex( const FGeometry& MyGeometry, FVector2D CursorScreenPos );
 	float									DataToLocalX( float Data, const FGeometry& MyGeometry ) const;
diff --git a/Engine/Source/Editor/Persona/Private/Shared/BoneSelectionWidget.cpp b/Engine/Source/Editor/Persona/Private/Shared/BoneSelectionWidget.cpp
--- a/Engine/Source/Editor/Persona/Private/Shared/BoneSelectionWidget.cpp
+++ b/Engine/Source/Editor/Persona/Private/Shared/BoneSelectionWidget.cpp
@@ -35,13 +35,13 @@ void SBoneTreeMenu::Construct(const FArguments& InArgs)
 	ChildSlot
 	[
 		SNew(SBorder)
-		.Padding(6)
+		.Padding(6.f)
 		.BorderImage(FAppStyle::GetBrush("NoBorder"))
 		.Content()
 		[
 			SNew(SBox)
-			.WidthOverride(300)
-			.HeightOverride(512)
+			.WidthOverride(300.f)
+			.HeightOverride(512.f)
 			.Content()
 			[
 				SNew(SVerticalBox)
diff --git a/Engine/Source/Editor/Persona/Private/Shared/ReferenceSectionSelectionWidget.cpp b/Engine/Source/Editor/Persona/Private/Shared/ReferenceSectionSelectionWidget.cpp
--- a/Engine/Source/Editor/Persona/Private/Shared/ReferenceSectionSelectionWidget.cpp
+++ b/Engine/Source/Editor/Persona/Private/Shared/ReferenceSectionSelectionWidget.cpp
@@ -90,16 +90,16 @@ TSharedRef<SWidget> SReferenceSectionSelectionWidget::CreateSectionListWidgetMen
 		.OnSelectionChanged(this, &SReferenceSectionSelectionWidget::OnSelectionChanged)
 		.HeaderRow(PlatformColumnHeader.ToSharedRef());
 
-	FText TitleText = FText(LOCTEXT("ReferenceListWidgetMenuTitle", "Choose a Section"));
+	const FText TitleText = FText(LOCTEXT("ReferenceListWidgetMenuTitle", "Choose a Section"));
 	TSharedPtr<SBorder> BorderWidget;
 	SAssignNew(BorderWidget, SBorder)
-	.Padding(6)
+	.Padding(6.f)
 	.BorderImage(FAppStyle::GetBrush("NoBorder"))
 	.Content()
 	[
 		SNew(SBox)
-		.WidthOverride(300)
-		.HeightOverride(512)
+		.WidthOverride(300.f)
+		.HeightOverride(512.f)
 		.Content()
 		[
 			SNew(SVerticalBox)
diff --git a/Engine/Source/Editor/Persona/Public/BlendSpaceAnalysis.h b/Engine/Source/Editor/Persona/Public/BlendSpaceAnalysis.h
--- a/Engine/Source/Editor/Persona/Public/BlendSpaceAnalysis.h
+++ b/Engine/Source/Editor/Persona/Public/BlendSpaceAnalysis.h
@@ -262,7 +262,7 @@ static bool CalculateComponentSampleValue(
 	int32 ComponentIndex = (int32) AnalysisProperties->FunctionAxis;
 	if (Fn(Value, BlendSpace, AnalysisProperties, Animation, RateScale))
 	{
-		Result = Value[ComponentIndex];
+		Result = static_cast<float>(Value[ComponentIndex]);
 		return true;
 	}
 	return false;
@@ -444,32 +444,31 @@ void CalculateBoneOrientation(
 	const FVector&             FrameRightDir, 
 	const FVector&             FrameUpDir)
 {
-	FTransform BoneTM = GetBoneTransform(Animation, Key, BoneName);
+	const FTransform BoneTM = GetBoneTransform(Animation, Key, BoneName);
 
-	FTransform TM = BoneOffset * BoneTM;
-	FQuat AimQuat = TM.GetRotation();
-	FVector AimForwardDir = BlendSpaceAnalysis::GetAxisFromTM(TM, AnalysisProperties->BoneFacingAxis);
-	FVector AimRightDir = BlendSpaceAnalysis::GetAxisFromTM(TM, AnalysisProperties->BoneRightAxis);
+	const FTransform TM = BoneOffset * BoneTM;
+	const FVector AimForwardDir = BlendSpaceAnalysis::GetAxisFromTM(TM, AnalysisProperties->BoneFacingAxis);
+	const FVector AimRightDir = BlendSpaceAnalysis::GetAxisFromTM(TM, AnalysisProperties->BoneRightAxis);
 
 	// Note that Yaw is best taken from the AimRightDir - this is to avoid problems when the gun is pointing
 	// up or down - especially if it goes beyond 90 degrees in pitch.
-	float Yaw = FMath::RadiansToDegrees(FMath::Atan2(
+	const double Yaw = FMath::RadiansToDegrees(FMath::Atan2(
 		FVector::DotProduct(AimRightDir, -FrameFacingDir), FVector::DotProduct(AimRightDir, FrameRightDir)));
 
 	// Undo the yaw to get pitch
 	const FQuat YawQuat(FrameUpDir, FMath::DegreesToRadians(Yaw));
-	FVector UnYawedAimForwardDir = YawQuat.UnrotateVector(AimForwardDir);
-	float Up = UnYawedAimForwardDir | FrameUpDir;
-	float Forward = UnYawedAimForwardDir | FrameFacingDir;
-	float Pitch = FMath::RadiansToDegrees(FMath::Atan2(Up, Forward));
+	const FVector UnYawedAimForwardDir = YawQuat.UnrotateVector(AimForwardDir);
+	const double Up = UnYawedAimForwardDir | FrameUpDir;
+	const double Forward = UnYawedAimForwardDir | FrameFacingDir;
+	const double Pitch = FMath::RadiansToDegrees(FMath::Atan2(Up, Forward));
 
 	// Undo the pitch to get roll
-	FVector UnYawedAimRightDir = YawQuat.UnrotateVector(AimRightDir);
+	const FVector UnYawedAimRightDir = YawQuat.UnrotateVector(AimRightDir);
 	const FQuat PitchQuat(FrameRightDir, -FMath::DegreesToRadians(Pitch));
 
-	FVector UnYawedUnPitchedAimRightDir = PitchQuat.UnrotateVector(UnYawedAimRightDir);
+	const FVector UnYawedUnPitchedAimRightDir = PitchQuat.UnrotateVector(UnYawedAimRightDir);
 
-	float Roll = FMath::RadiansToDegrees(FMath::Atan2(
+	const double Roll = FMath::RadiansToDegrees(FMath::Atan2(
 		FVector::DotProduct(UnYawedUnPitchedAimRightDir, -FrameUpDir), 
 		FVector::DotProduct(UnYawedUnPitchedAimRightDir, FrameRightDir)));
 
