Fix for Anim Notifiers in animation sequences not triggering if the end-user drags it all the way to the left

diff --git a/Engine/Source/Editor/Persona/Private/SAnimNotifyPanel.cpp b/Engine/Source/Editor/Persona/Private/SAnimNotifyPanel.cpp
--- a/Engine/Source/Editor/Persona/Private/SAnimNotifyPanel.cpp
+++ b/Engine/Source/Editor/Persona/Private/SAnimNotifyPanel.cpp
@@ -1226,7 +1226,8 @@ public:
 			// Look for a snap on the first scrub handle
 			FVector2D TrackNodePos = TrackGeom.AbsoluteToLocal(EventPosition);
 			const FVector2D OriginalNodePosition = TrackNodePos;
-			float SequenceEnd = TrackScaleInfo.InputToLocalX(Sequence->GetPlayLength());
+			const float SequenceStart = TrackScaleInfo.InputToLocalX(0.f);
+			const float SequenceEnd = TrackScaleInfo.InputToLocalX(Sequence->GetPlayLength());
 
 			// Always clamp the Y to the current track
 			SelectionBeginPosition.Y = SelectionPositionClampInfo->TrackPos - 1.0f;
@@ -1237,9 +1238,9 @@ public:
 				if (bSnapped)
 				{
 					EAnimEventTriggerOffsets::Type Offset = EAnimEventTriggerOffsets::NoOffset;
-					if (SnapX == 0.0f || SnapX == SequenceEnd)
+					if (SnapX == SequenceStart || SnapX == SequenceEnd)
 					{
-						Offset = SnapX > 0.0f ? EAnimEventTriggerOffsets::OffsetBefore : EAnimEventTriggerOffsets::OffsetAfter;
+						Offset = SnapX > SequenceStart ? EAnimEventTriggerOffsets::OffsetBefore : EAnimEventTriggerOffsets::OffsetAfter;
 					}
 					else
 					{
