Added support for copy/paste for curve viewer

diff --git a/Engine/Source/Editor/Persona/Private/SAnimCurveViewer.cpp b/Engine/Source/Editor/Persona/Private/SAnimCurveViewer.cpp
--- a/Engine/Source/Editor/Persona/Private/SAnimCurveViewer.cpp
+++ b/Engine/Source/Editor/Persona/Private/SAnimCurveViewer.cpp
@@ -22,6 +22,8 @@
 #include "Framework/Commands/GenericCommands.h"
 #include "CurveViewerCommands.h"
 #include "Animation/EditorAnimCurveBoneLinks.h"
+#include "HAL/PlatformApplicationMisc.h"
+#include "ScopedTransaction.h"
 
 #define LOCTEXT_NAMESPACE "SAnimCurveViewer"
 
@@ -687,6 +689,16 @@ void SAnimCurveViewer::BindCommands()
 		FExecuteAction::CreateSP(this, &SAnimCurveViewer::OnDeleteNameClicked),
 		FCanExecuteAction::CreateSP(this, &SAnimCurveViewer::CanDelete));
 
+	CommandList.MapAction(
+		FGenericCommands::Get().Copy,
+		FExecuteAction::CreateSP(this, &SAnimCurveViewer::OnCopyClicked),
+		FCanExecuteAction::CreateSP(this, &SAnimCurveViewer::CanCopy));
+
+	CommandList.MapAction(
+		FGenericCommands::Get().Paste,
+		FExecuteAction::CreateSP(this, &SAnimCurveViewer::OnPasteClicked),
+		FCanExecuteAction::CreateSP(this, &SAnimCurveViewer::CanPaste));
+	
 	CommandList.MapAction(
 		MenuActions.AddCurve,
 		FExecuteAction::CreateSP(this, &SAnimCurveViewer::OnAddClicked),
@@ -737,6 +749,8 @@ TSharedPtr<SWidget> SAnimCurveViewer::OnGetContextMenuContent() const
 
 	MenuBuilder.AddMenuEntry(FGenericCommands::Get().Rename, NAME_None, LOCTEXT("RenameSmartNameLabel", "Rename Curve"), LOCTEXT("RenameSmartNameToolTip", "Rename the selected curve"));
 	MenuBuilder.AddMenuEntry(FGenericCommands::Get().Delete, NAME_None, LOCTEXT("DeleteSmartNameLabel", "Delete Curve"), LOCTEXT("DeleteSmartNameToolTip", "Delete the selected curve"));
+	MenuBuilder.AddMenuEntry(FGenericCommands::Get().Copy);
+	MenuBuilder.AddMenuEntry(FGenericCommands::Get().Paste);
 	MenuBuilder.AddMenuEntry(Actions.AddCurve);
 
 	MenuBuilder.EndSection();
@@ -1117,6 +1131,70 @@ bool SAnimCurveViewer::CanDelete()
 	return AnimCurveListView->GetNumItemsSelected() > 0;
 }
 
+static const TCHAR* ClipboardHeader = TEXT("AnimCurveViewer");
+
+void SAnimCurveViewer::OnCopyClicked()
+{
+	TArray<TSharedPtr<FDisplayedAnimCurveInfo>> SelectedItems = AnimCurveListView->GetSelectedItems();
+
+	FAnimCurveViewerClipboard Clipboard;
+
+	for (TSharedPtr<FDisplayedAnimCurveInfo> Item : SelectedItems)
+	{
+		FAnimCurveViewerClipboardEntry Entry;
+		Entry.CurveName = Item->SmartName.DisplayName;
+
+		if(const FCurveMetaData* CurveMetaData = EditableSkeletonPtr.Pin()->GetSkeleton().GetCurveMetaData(Item->SmartName.DisplayName))
+		{
+			Entry.MetaData = *CurveMetaData;
+		}
+
+		Clipboard.Entries.Add(Entry);
+	}
+
+	FString ClipboardString = ClipboardHeader;
+	FAnimCurveViewerClipboard::StaticStruct()->ExportText(ClipboardString, &Clipboard, nullptr, nullptr, PPF_None, nullptr);
+	FPlatformApplicationMisc::ClipboardCopy(*ClipboardString);
+}
+
+bool SAnimCurveViewer::CanCopy() const
+{
+	return AnimCurveListView->GetNumItemsSelected() > 0;
+}
+
+void SAnimCurveViewer::OnPasteClicked()
+{
+	if (const FSmartNameMapping* NameMapping = GetAnimCurveMapping())
+	{
+		FScopedTransaction Transaction(LOCTEXT("PasteCurves", "Paste Curves"));
+		
+		FString TextToImport;
+		FPlatformApplicationMisc::ClipboardPaste(TextToImport);
+		TextToImport.RemoveFromStart(ClipboardHeader);
+
+		FAnimCurveViewerClipboard Clipboard;
+		FAnimCurveViewerClipboard::StaticStruct()->ImportText(*TextToImport, &Clipboard, nullptr, PPF_None, GLog, FAnimCurveViewerClipboard::StaticStruct()->GetName());
+
+		for(const FAnimCurveViewerClipboardEntry& Entry : Clipboard.Entries)
+		{
+			FSmartName NewCurveName;
+			EditableSkeletonPtr.Pin()->AddSmartname(ContainerName, Entry.CurveName, NewCurveName);
+			EditableSkeletonPtr.Pin()->SetCurveMetaBoneLinks(NewCurveName, Entry.MetaData.LinkedBones, Entry.MetaData.MaxLOD);
+			EditableSkeletonPtr.Pin()->SetCurveMetaDataMaterial(NewCurveName, Entry.MetaData.Type.bMaterial);
+			EditableSkeletonPtr.Pin()->SetCurveMetaDataMorphTarget(NewCurveName, Entry.MetaData.Type.bMorphtarget);
+		}
+
+		RefreshCurveList(true);
+	}
+}
+
+bool SAnimCurveViewer::CanPaste() const
+{
+	FString TextToImport;
+	FPlatformApplicationMisc::ClipboardPaste(TextToImport);
+	return TextToImport.StartsWith(ClipboardHeader);
+}
+
 void SAnimCurveViewer::OnSelectionChanged(TSharedPtr<FDisplayedAnimCurveInfo> InItem, ESelectInfo::Type SelectInfo)
 {
 	// make sure the currently selected ones are refreshed if it's first time
diff --git a/Engine/Source/Editor/Persona/Private/SAnimCurveViewer.h b/Engine/Source/Editor/Persona/Private/SAnimCurveViewer.h
--- a/Engine/Source/Editor/Persona/Private/SAnimCurveViewer.h
+++ b/Engine/Source/Editor/Persona/Private/SAnimCurveViewer.h
@@ -19,11 +19,34 @@
 #include "EditorObjectsTracker.h"
 #include "PersonaDelegates.h"
 #include "EditorUndoClient.h"
+#include "SAnimCurveViewer.generated.h"
 
 class FUICommandList;
 class IEditableSkeleton;
 class SAnimCurveViewer;
 
+/** Clipboard contents for anim curves */
+USTRUCT()
+struct FAnimCurveViewerClipboardEntry
+{
+	GENERATED_BODY()
+	
+	UPROPERTY()
+	FName CurveName;
+
+	UPROPERTY()
+	FCurveMetaData MetaData;
+};
+
+USTRUCT()
+struct FAnimCurveViewerClipboard
+{
+	GENERATED_BODY()
+	
+	UPROPERTY()
+	TArray<FAnimCurveViewerClipboardEntry> Entries;
+};
+
 //////////////////////////////////////////////////////////////////////////
 // FDisplayedAnimCurveInfo
 
@@ -278,6 +301,12 @@ private:
 	void OnRenameClicked();
 	bool CanRename();
 
+	void OnCopyClicked();
+	bool CanCopy() const;
+
+	void OnPasteClicked();
+	bool CanPaste() const;
+
 	void OnAddClicked();
 
 	ECheckBoxState IsShowingAllCurves() const;
diff --git a/Engine/Source/Editor/SkeletonEditor/Private/EditableSkeleton.cpp b/Engine/Source/Editor/SkeletonEditor/Private/EditableSkeleton.cpp
--- a/Engine/Source/Editor/SkeletonEditor/Private/EditableSkeleton.cpp
+++ b/Engine/Source/Editor/SkeletonEditor/Private/EditableSkeleton.cpp
@@ -532,7 +532,18 @@ void FEditableSkeleton::SetCurveMetaDataMaterial(const FSmartName& CurveName, bo
 	}
 }
 
-void FEditableSkeleton::SetCurveMetaBoneLinks(const FSmartName& CurveName, TArray<FBoneReference>& BoneLinks, uint8 InMaxLOD)
+void FEditableSkeleton::SetCurveMetaDataMorphTarget(const FSmartName& CurveName, bool bOverrideMorphTarget)
+{
+	Skeleton->Modify();
+	FCurveMetaData* CurveMetaData = Skeleton->GetCurveMetaData(CurveName);
+	if (CurveMetaData)
+	{
+		// override curve data
+		CurveMetaData->Type.bMorphtarget = !!bOverrideMorphTarget;
+	}
+}
+
+void FEditableSkeleton::SetCurveMetaBoneLinks(const FSmartName& CurveName, const TArray<FBoneReference>& BoneLinks, uint8 InMaxLOD)
 {
 	Skeleton->Modify();
 	FCurveMetaData* CurveMetaData = Skeleton->GetCurveMetaData(CurveName);
diff --git a/Engine/Source/Editor/SkeletonEditor/Private/EditableSkeleton.h b/Engine/Source/Editor/SkeletonEditor/Private/EditableSkeleton.h
--- a/Engine/Source/Editor/SkeletonEditor/Private/EditableSkeleton.h
+++ b/Engine/Source/Editor/SkeletonEditor/Private/EditableSkeleton.h
@@ -54,7 +54,8 @@ public:
 	virtual void RenameSmartname(const FName InContainerName, SmartName::UID_Type InNameUid, const FName InNewName) override;
 	virtual void RemoveSmartnamesAndFixupAnimations(const FName& InContainerName, const TArray<FName>& InNames) override;
 	virtual void SetCurveMetaDataMaterial(const FSmartName& CurveName, bool bOverrideMaterial) override;
-	virtual void SetCurveMetaBoneLinks(const FSmartName& CurveName, TArray<FBoneReference>& BoneLinks, uint8 InMaxLOD) override;
+	virtual void SetCurveMetaDataMorphTarget(const FSmartName& CurveName, bool bOverrideMorphTarget) override;
+	virtual void SetCurveMetaBoneLinks(const FSmartName& CurveName, const TArray<FBoneReference>& BoneLinks, uint8 InMaxLOD) override;
 	virtual void SetPreviewMesh(class USkeletalMesh* InSkeletalMesh) override;
 	virtual void LoadAdditionalPreviewSkeletalMeshes() override;
 	virtual void SetAdditionalPreviewSkeletalMeshes(class UDataAsset* InPreviewCollectionAsset) override;
diff --git a/Engine/Source/Editor/SkeletonEditor/Public/IEditableSkeleton.h b/Engine/Source/Editor/SkeletonEditor/Public/IEditableSkeleton.h
--- a/Engine/Source/Editor/SkeletonEditor/Public/IEditableSkeleton.h
+++ b/Engine/Source/Editor/SkeletonEditor/Public/IEditableSkeleton.h
@@ -75,8 +75,11 @@ public:
 	/** Sets Material Meta Data for the curve */
 	virtual void SetCurveMetaDataMaterial(const FSmartName& CurveName, bool bOverrideMaterial) = 0;
 
+	/** Sets Morph Target Meta Data for the curve */
+	virtual void SetCurveMetaDataMorphTarget(const FSmartName& CurveName, bool bOverrideMorphTarget) = 0;
+
 	/** Sets Bone Links per curve */
-	virtual void SetCurveMetaBoneLinks(const FSmartName& CurveName, TArray<FBoneReference>& BoneLinks, uint8 InMaxLOD) = 0;
+	virtual void SetCurveMetaBoneLinks(const FSmartName& CurveName, const TArray<FBoneReference>& BoneLinks, uint8 InMaxLOD) = 0;
 
 	/**
 	 * Makes sure all attached objects are valid and removes any that aren't.
diff --git a/Engine/Source/Runtime/Engine/Classes/Animation/SmartName.h b/Engine/Source/Runtime/Engine/Classes/Animation/SmartName.h
--- a/Engine/Source/Runtime/Engine/Classes/Animation/SmartName.h
+++ b/Engine/Source/Runtime/Engine/Classes/Animation/SmartName.h
@@ -23,10 +23,15 @@ struct FCurveMetaData
 	GENERATED_USTRUCT_BODY()
 
 	/** connected bones to this meta data */
-	TArray<struct FBoneReference> LinkedBones;
+	UPROPERTY()
+	TArray<FBoneReference> LinkedBones;
+	
 	/* max LOD (lowest LOD) to evaluate this. -1 means it will evaluate all the time. */
+	UPROPERTY()
 	uint8 MaxLOD;
-	struct FAnimCurveType Type;
+	
+	UPROPERTY()
+	FAnimCurveType Type;
 
 	friend FArchive& operator<<(FArchive& Ar, FCurveMetaData& B)
 	{
diff --git a/Engine/Source/Runtime/Engine/Public/BoneContainer.h b/Engine/Source/Runtime/Engine/Public/BoneContainer.h
--- a/Engine/Source/Runtime/Engine/Public/BoneContainer.h
+++ b/Engine/Source/Runtime/Engine/Public/BoneContainer.h
@@ -34,7 +34,10 @@ struct ENGINE_API FAnimCurveType
 {
 	GENERATED_USTRUCT_BODY()
 
+	UPROPERTY()
 	bool bMaterial;
+
+	UPROPERTY()
 	bool bMorphtarget;
 
 	FAnimCurveType(bool bInMorphtarget = false, bool bInMaterial = false)
