Fixed crash deleting a curve then validating its asset post-save

Delegates from anim sequences can be fired during PostLoad and potentially execute BP bytecode, which is not allowed.
Switched to a deferred update in the notify panel to ensure that we only update once delegates have been fired.

diff --git a/Engine/Source/Editor/Persona/Private/SAnimNotifyPanel.cpp b/Engine/Source/Editor/Persona/Private/SAnimNotifyPanel.cpp
--- a/Engine/Source/Editor/Persona/Private/SAnimNotifyPanel.cpp
+++ b/Engine/Source/Editor/Persona/Private/SAnimNotifyPanel.cpp
@@ -3850,16 +3850,18 @@ void SAnimNotifyPanel::Construct(const FArguments& InArgs, const TSharedRef<FAni
 	OnNotifyNodesBeingDragged = InArgs._OnNotifyNodesBeingDragged;
 	bIsSelecting = false;
 	bIsUpdating = false;
+	bUpdateRequested = false;
+	bRefreshRequested = false;
 
 	InModel->OnHandleObjectsSelected().AddSP(this, &SAnimNotifyPanel::HandleObjectsSelected);
 
 	FAnimNotifyPanelCommands::Register();
 	BindCommands();
 
-	Sequence->RegisterOnNotifyChanged(UAnimSequenceBase::FOnNotifyChanged::CreateSP(this, &SAnimNotifyPanel::RefreshNotifyTracks ));
+	Sequence->RegisterOnNotifyChanged(UAnimSequenceBase::FOnNotifyChanged::CreateSP(this, &SAnimNotifyPanel::RequestRefresh ));
 
-	InModel->GetEditableSkeleton()->RegisterOnNotifiesChanged(FSimpleDelegate::CreateSP(this, &SAnimNotifyPanel::RefreshNotifyTracks));
-	InModel->OnTracksChanged().Add(FSimpleDelegate::CreateSP(this, &SAnimNotifyPanel::RefreshNotifyTracks));
+	InModel->GetEditableSkeleton()->RegisterOnNotifiesChanged(FSimpleDelegate::CreateSP(this, &SAnimNotifyPanel::RequestRefresh));
+	InModel->OnTracksChanged().Add(FSimpleDelegate::CreateSP(this, &SAnimNotifyPanel::RequestRefresh));
 
 	if(GEditor)
 	{
@@ -3893,7 +3895,7 @@ void SAnimNotifyPanel::Construct(const FArguments& InArgs, const TSharedRef<FAni
 	PopulateNotifyBlueprintClasses(NotifyClassNames);
 	PopulateNotifyBlueprintClasses(NotifyStateClassNames);
 
-	Update();
+	RequestUpdate();
 }
 
 SAnimNotifyPanel::~SAnimNotifyPanel()
@@ -3959,7 +3961,7 @@ FReply SAnimNotifyPanel::InsertTrack(int32 TrackIndexToInsert)
 	Sequence->PostEditChange();
 	Sequence->MarkPackageDirty();
 
-	Update();
+	RequestUpdate();
 
 	return FReply::Handled();
 }
@@ -3973,7 +3975,7 @@ FReply SAnimNotifyPanel::AddTrack()
 	Sequence->AnimNotifyTracks.Add(NewItem);
 	Sequence->MarkPackageDirty();
 
-	Update();
+	RequestUpdate();
 
 	return FReply::Handled();
 }
@@ -4006,7 +4008,7 @@ FReply SAnimNotifyPanel::DeleteTrack(int32 TrackIndexToDelete)
 			Sequence->AnimNotifyTracks.RemoveAt(TrackIndexToDelete);
 			Sequence->PostEditChange();
 			Sequence->MarkPackageDirty();
-			Update();
+			RequestUpdate();
 		}
 	}
 	return FReply::Handled();
@@ -4051,6 +4053,16 @@ void SAnimNotifyPanel::Update()
 	}
 }
 
+void SAnimNotifyPanel::RequestUpdate()
+{
+	bUpdateRequested = true;
+}
+
+void SAnimNotifyPanel::RequestRefresh()
+{
+	bRefreshRequested = true;
+}
+
 // Helper to save/restore selection state when widgets are recreated
 struct FScopedSavedNotifySelection
 {
@@ -4122,7 +4134,7 @@ void SAnimNotifyPanel::RefreshNotifyTracks()
 				.ViewInputMax(ViewInputMax)
 				.OnGetScrubValue(OnGetScrubValue)
 				.OnGetDraggedNodePos(this, &SAnimNotifyPanel::CalculateDraggedNodePos)
-				.OnUpdatePanel(this, &SAnimNotifyPanel::Update)
+				.OnUpdatePanel(this, &SAnimNotifyPanel::RequestUpdate)
 				.OnGetNotifyBlueprintData(this, &SAnimNotifyPanel::OnGetNotifyBlueprintData, &NotifyClassNames)
 				.OnGetNotifyStateBlueprintData(this, &SAnimNotifyPanel::OnGetNotifyBlueprintData, &NotifyStateClassNames)
 				.OnGetNotifyNativeClasses(this, &SAnimNotifyPanel::OnGetNativeNotifyData, UAnimNotify::StaticClass(), &NotifyClassNames)
@@ -4196,7 +4208,7 @@ FReply SAnimNotifyPanel::OnNotifyNodeDragStarted(TArray<TSharedPtr<SAnimNotifyNo
 	}
 
 	FPanTrackRequest PanRequestDelegate = FPanTrackRequest::CreateSP(this, &SAnimNotifyPanel::PanInputViewRange);
-	FOnUpdatePanel UpdateDelegate = FOnUpdatePanel::CreateSP(this, &SAnimNotifyPanel::Update);
+	FOnUpdatePanel UpdateDelegate = FOnUpdatePanel::CreateSP(this, &SAnimNotifyPanel::RequestUpdate);
 	return FReply::Handled().BeginDragDrop(FNotifyDragDropOp::New(Nodes, NodeDragDecorator, NotifyAnimTracks, Sequence, ScreenCursorPos, OverlayOrigin, OverlayExtents, CurrentDragXPosition, PanRequestDelegate, OnSnapPosition, UpdateDelegate, OnNotifyNodesBeingDragged));
 }
 
@@ -4266,7 +4278,7 @@ void SAnimNotifyPanel::DeleteSelectedNodeObjects()
 	TArray<UObject*> Objects;
 	OnSelectionChanged.ExecuteIfBound(Objects);
 
-	Update();
+	RequestUpdate();
 }
 
 void SAnimNotifyPanel::SetSequence(class UAnimSequenceBase*	InSequence)
@@ -4274,7 +4286,7 @@ void SAnimNotifyPanel::SetSequence(class UAnimSequenceBase*	InSequence)
 	if (InSequence != Sequence)
 	{
 		Sequence = InSequence;
-		Update();
+		RequestUpdate();
 	}
 }
 
@@ -4492,7 +4504,7 @@ void SAnimNotifyPanel::OnReplaceSelectedWithNotify(FString NewNotifyName, UClass
 	Sequence->PostEditChange();
 	Sequence->MarkPackageDirty();
 
-	Update();
+	RequestUpdate();
 }
 
 void SAnimNotifyPanel::OnReplaceSelectedWithNotifyBlueprint(FString NewBlueprintNotifyName, FString NewBlueprintNotifyClass)
@@ -4630,7 +4642,7 @@ void SAnimNotifyPanel::OnPropertyChanged(UObject* ChangedObject, FPropertyChange
 			if(Event.Notify == ChangedObject || Event.NotifyStateClass == ChangedObject)
 			{
 				// If we've changed a notify present in the sequence, refresh our tracks.
-				Update();
+				RequestUpdate();
 			}
 		}
 	}
@@ -4743,6 +4755,21 @@ int32 SAnimNotifyPanel::OnPaint(const FPaintArgs& Args, const FGeometry& Allotte
 	return LayerId;
 }
 
+void SAnimNotifyPanel::Tick( const FGeometry& AllottedGeometry, const double InCurrentTime, const float InDeltaTime )
+{
+	if(bUpdateRequested)
+	{
+		Update();
+		bUpdateRequested = false;
+		bRefreshRequested = false;
+	}
+	if(bRefreshRequested)
+	{
+		RefreshNotifyTracks();
+		bRefreshRequested = false;
+	}
+}
+
 void SAnimNotifyPanel::RefreshMarqueeSelectedNodes(const FGeometry& PanelGeo)
 {
 	if(Marquee.IsValid())
diff --git a/Engine/Source/Editor/Persona/Private/SAnimNotifyPanel.h b/Engine/Source/Editor/Persona/Private/SAnimNotifyPanel.h
--- a/Engine/Source/Editor/Persona/Private/SAnimNotifyPanel.h
+++ b/Engine/Source/Editor/Persona/Private/SAnimNotifyPanel.h
@@ -216,6 +216,9 @@ public:
 	// Handler function for renaming a notify track
 	void OnCommitTrackName(const FText& InText, ETextCommit::Type CommitInfo, int32 TrackIndexToName);
 
+	// Request a deferred Update call
+	void RequestUpdate();
+	
 	void Update();
 
 	/** Returns the position of the notify node currently being dragged. Returns -1 if no node is being dragged */
@@ -241,6 +244,7 @@ public:
 	virtual FReply OnDragDetected(const FGeometry& MyGeometry, const FPointerEvent& MouseEvent) override;
 	virtual void OnFocusLost(const FFocusEvent& InFocusEvent) override;
 	virtual int32 OnPaint(const FPaintArgs& Args, const FGeometry& AllottedGeometry, const FSlateRect& MyCullingRect, FSlateWindowElementList& OutDrawElements, int32 LayerId, const FWidgetStyle& InWidgetStyle, bool bParentEnabled) const override;
+	virtual void Tick( const FGeometry& AllottedGeometry, const double InCurrentTime, const float InDeltaTime ) override;
 	/** End SWidget Interface */
 
 	void RefreshMarqueeSelectedNodes(const FGeometry& PanelGeo);
@@ -288,6 +292,9 @@ private:
 	/** Cached list of Notify editor tracks */
 	TArray<TSharedPtr<SNotifyEdTrack>> NotifyEditorTracks;
 
+	// Request a deferred RefreshNotifyTracks call
+	void RequestRefresh();
+
 	// this just refresh notify tracks - UI purpose only
 	// do not call this from here. This gets called by asset. 
 	void RefreshNotifyTracks();
@@ -373,4 +380,8 @@ private:
 
 	/** Recursion guard for updating */
 	bool bIsUpdating;
+
+	/** Flags to handle deferred updates */
+	bool bUpdateRequested;
+	bool bRefreshRequested;
 };
