Provide sound library contexts the ability to easily select bone names inside the persona animation editor

#rnx

diff --git a/Engine/Source/Editor/Persona/Private/AnimNotifyDetails.cpp b/Engine/Source/Editor/Persona/Private/AnimNotifyDetails.cpp
--- a/Engine/Source/Editor/Persona/Private/AnimNotifyDetails.cpp
+++ b/Engine/Source/Editor/Persona/Private/AnimNotifyDetails.cpp
@@ -326,6 +326,43 @@ void FAnimNotifyDetails::HideLinkProperties(IDetailLayoutBuilder& Builder, TShar
 
 bool FAnimNotifyDetails::CustomizeProperty(IDetailCategoryBuilder& CategoryBuilder, UObject* Notify, TSharedPtr<IPropertyHandle> Property)
 {
+	TFunction<void(const TSharedPtr<IPropertyHandle>)> FixBoneNamePropertyRecurse;
+	FixBoneNamePropertyRecurse = [this, &FixBoneNamePropertyRecurse, &CategoryBuilder, &Notify](const TSharedPtr<IPropertyHandle>& InPropertyHandle)
+	{
+		const bool bHasExpandMeta = InPropertyHandle->GetBoolMetaData(TEXT("AnimNotifyExpand"));
+		bool bParentIsObjectPtr = false;
+
+		TSharedPtr<IPropertyHandle> ParentProp = InPropertyHandle->GetParentHandle();
+		if (ParentProp.IsValid() && ParentProp->IsValidHandle())
+		{
+			bParentIsObjectPtr = (CastField<FObjectPropertyBase>(ParentProp->GetProperty()) != nullptr);
+		}
+
+		// Recurse into Object Ptrs or properties with AnimNotifyExpand
+		if (bParentIsObjectPtr || bHasExpandMeta)
+		{
+			IDetailLayoutBuilder& LayoutBuilder = CategoryBuilder.GetParentLayout();
+			LayoutBuilder.HideProperty(InPropertyHandle);
+
+			uint32 NumChildren = 0;
+			InPropertyHandle->GetNumChildren(NumChildren);
+			for (uint32 i = 0; i < NumChildren; ++i)
+			{
+				TSharedPtr<IPropertyHandle> ChildHandle = InPropertyHandle->GetChildHandle(i);
+				FixBoneNamePropertyRecurse(ChildHandle);
+			}
+		}
+		else if (InPropertyHandle->GetBoolMetaData(TEXT("AnimNotifyBoneName")))
+		{
+			// Convert this property to a bone name property
+			AddBoneNameProperty(CategoryBuilder, Notify, InPropertyHandle);
+		}
+		else
+		{
+			CategoryBuilder.AddProperty(InPropertyHandle);
+		}
+	};
+
 	if(Notify && Notify->GetClass() && Property->IsValidHandle())
 	{
 		FString ClassName = Notify->GetClass()->GetName();
@@ -347,6 +384,12 @@ bool FAnimNotifyDetails::CustomizeProperty(IDetailCategoryBuilder& CategoryBuild
 			AddBoneNameProperty(CategoryBuilder, Notify, Property);
 			return true;
 		}
+		else if (ClassName.Find(TEXT("_SoundLibrary")) != INDEX_NONE && PropertyName == TEXT("SoundContext"))
+		{
+			CategoryBuilder.AddProperty(Property);
+			FixBoneNamePropertyRecurse(Property);
+			return true;
+		}
 		else if (ClassName.Find(TEXT("AnimNotifyState_Trail")) != INDEX_NONE)
 		{
 			if(PropertyName == TEXT("FirstSocketName") || PropertyName == TEXT("SecondSocketName"))
