Fixed crash accessing out of bounds notify index when deleting an anim notify in a montage

Because of delayed delegate broadcasts (fallout from CL 22837509), out-of-bounds notify indices can hang around in the timing panel for a frame.
Fix just does a bounds check (similar to the one already present for montage sections) to prevent out of bounds accesses.

diff --git a/Engine/Source/Editor/Persona/Private/SAnimTimingPanel.cpp b/Engine/Source/Editor/Persona/Private/SAnimTimingPanel.cpp
--- a/Engine/Source/Editor/Persona/Private/SAnimTimingPanel.cpp
+++ b/Engine/Source/Editor/Persona/Private/SAnimTimingPanel.cpp
@@ -401,19 +401,22 @@ ETimingElementType::Type FTimingRelevantElement_Section::GetType()
 void FTimingRelevantElement_Section::GetDescriptionItems(TMap<FString, FText>& Items)
 {
 	check(Montage);
-	FCompositeSection& Section = Montage->CompositeSections[SectionIdx];
-
-	Items.Add(LOCTEXT("SectionName", "Name").ToString(), FText::FromName(Section.SectionName));
-	if(GetDefault<UPersonaOptions>()->TimelineDisplayFormat == EFrameNumberDisplayFormats::Frames)
+	if(Montage->CompositeSections.IsValidIndex(SectionIdx))
 	{
-		Items.Add(LOCTEXT("SectionTriggerFrame", "Trigger Frame").ToString(), FText::Format(LOCTEXT("SectionTriggerFrameValue", "{0}"), FText::AsNumber(Montage->GetFrameAtTime(Section.GetTime()))));
+		FCompositeSection& Section = Montage->CompositeSections[SectionIdx];
+
+		Items.Add(LOCTEXT("SectionName", "Name").ToString(), FText::FromName(Section.SectionName));
+		if(GetDefault<UPersonaOptions>()->TimelineDisplayFormat == EFrameNumberDisplayFormats::Frames)
+		{
+			Items.Add(LOCTEXT("SectionTriggerFrame", "Trigger Frame").ToString(), FText::Format(LOCTEXT("SectionTriggerFrameValue", "{0}"), FText::AsNumber(Montage->GetFrameAtTime(Section.GetTime()))));
+		}
+		else
+		{
+			FNumberFormattingOptions NumberOptions;
+			NumberOptions.MinimumFractionalDigits = 3;
+			Items.Add(LOCTEXT("SectionTriggerTime", "Trigger Time").ToString(), FText::Format(LOCTEXT("SectionTriggerTimeValue", "{0}s"), FText::AsNumber(Section.GetTime(), &NumberOptions)));
+		}	
 	}
-	else
-	{
-		FNumberFormattingOptions NumberOptions;
-		NumberOptions.MinimumFractionalDigits = 3;
-		Items.Add(LOCTEXT("SectionTriggerTime", "Trigger Time").ToString(), FText::Format(LOCTEXT("SectionTriggerTimeValue", "{0}s"), FText::AsNumber(Section.GetTime(), &NumberOptions)));
-	}	
 }
 
 FName FTimingRelevantElement_Notify::GetTypeName()
@@ -447,58 +450,70 @@ float FTimingRelevantElement_Notify::GetElementTime() const
 ETimingElementType::Type FTimingRelevantElement_Notify::GetType()
 {
 	check(Sequence);
-	FAnimNotifyEvent& Event = Sequence->Notifies[NotifyIndex];
-
-	if(Event.IsBranchingPoint())
-	{
-		return ETimingElementType::BranchPointNotify;
-	}
-	else
+	if(Sequence->Notifies.IsValidIndex(NotifyIndex))
 	{
-		if(Event.NotifyStateClass)
+		FAnimNotifyEvent& Event = Sequence->Notifies[NotifyIndex];
+
+		if(Event.IsBranchingPoint())
 		{
-			return ETimingElementType::NotifyStateBegin;
+			return ETimingElementType::BranchPointNotify;
 		}
 		else
 		{
-			return ETimingElementType::QueuedNotify;
+			if(Event.NotifyStateClass)
+			{
+				return ETimingElementType::NotifyStateBegin;
+			}
+			else
+			{
+				return ETimingElementType::QueuedNotify;
+			}
 		}
 	}
+
+	return ETimingElementType::Max;
 }
 
 void FTimingRelevantElement_Notify::GetDescriptionItems(TMap<FString, FText>& Items)
 {
 	check(Sequence);
-	FAnimNotifyEvent& Event = Sequence->Notifies[NotifyIndex];
+	if(Sequence->Notifies.IsValidIndex(NotifyIndex))
+	{
+		FAnimNotifyEvent& Event = Sequence->Notifies[NotifyIndex];
 
-	FNumberFormattingOptions NumberOptions;
-	NumberOptions.MinimumFractionalDigits = 3;
+		FNumberFormattingOptions NumberOptions;
+		NumberOptions.MinimumFractionalDigits = 3;
 
-	Items.Add(LOCTEXT("NotifyName", "Name").ToString(), FText::FromName(Event.NotifyName));
+		Items.Add(LOCTEXT("NotifyName", "Name").ToString(), FText::FromName(Event.NotifyName));
 
-	if(GetDefault<UPersonaOptions>()->TimelineDisplayFormat == EFrameNumberDisplayFormats::Frames)
-	{
-		Items.Add(LOCTEXT("NotifyTriggerFrame", "Trigger Frame").ToString(), FText::Format(LOCTEXT("NotifyTriggerFrame_Val", "{0}"), FText::AsNumber(Sequence->GetFrameAtTime(Event.GetTime()))));
-	}
-	else
-	{
-		Items.Add(LOCTEXT("NotifyTriggerTime", "Trigger Time").ToString(), FText::Format(LOCTEXT("NotifyTriggerTime_Val", "{0}s"), FText::AsNumber(Event.GetTime(), &NumberOptions)));
-	}	
+		if(GetDefault<UPersonaOptions>()->TimelineDisplayFormat == EFrameNumberDisplayFormats::Frames)
+		{
+			Items.Add(LOCTEXT("NotifyTriggerFrame", "Trigger Frame").ToString(), FText::Format(LOCTEXT("NotifyTriggerFrame_Val", "{0}"), FText::AsNumber(Sequence->GetFrameAtTime(Event.GetTime()))));
+		}
+		else
+		{
+			Items.Add(LOCTEXT("NotifyTriggerTime", "Trigger Time").ToString(), FText::Format(LOCTEXT("NotifyTriggerTime_Val", "{0}s"), FText::AsNumber(Event.GetTime(), &NumberOptions)));
+		}	
 
-	// +1 as we start at 1 when showing tracks to the user
-	Items.Add(LOCTEXT("TrackIdx", "Track").ToString(), FText::AsNumber(Event.TrackIndex + 1));
+		// +1 as we start at 1 when showing tracks to the user
+		Items.Add(LOCTEXT("TrackIdx", "Track").ToString(), FText::AsNumber(Event.TrackIndex + 1));
 
-	if(Event.NotifyStateClass)
-	{
-		Items.Add(LOCTEXT("NotifyDuration", "Duration").ToString(), FText::Format(LOCTEXT("NotifyDuration_Val", "{0}s"), FText::AsNumber(Event.GetDuration(), &NumberOptions)));
+		if(Event.NotifyStateClass)
+		{
+			Items.Add(LOCTEXT("NotifyDuration", "Duration").ToString(), FText::Format(LOCTEXT("NotifyDuration_Val", "{0}s"), FText::AsNumber(Event.GetDuration(), &NumberOptions)));
+		}
 	}
 }
 
 int32 FTimingRelevantElement_Notify::GetElementSortPriority() const
 {
 	check(Sequence);
-	FAnimNotifyEvent& Event = Sequence->Notifies[NotifyIndex];
-	return Event.TrackIndex;
+	if(Sequence->Notifies.IsValidIndex(NotifyIndex))
+	{
+		FAnimNotifyEvent& Event = Sequence->Notifies[NotifyIndex];
+		return Event.TrackIndex;
+	}
+	return 0;
 }
 
 FName FTimingRelevantElement_NotifyStateEnd::GetTypeName()
@@ -509,9 +524,13 @@ FName FTimingRelevantElement_NotifyStateEnd::GetTypeName()
 float FTimingRelevantElement_NotifyStateEnd::GetElementTime() const
 {
 	check(Sequence);
-	FAnimNotifyEvent& Event = Sequence->Notifies[NotifyIndex];
-	check(Event.NotifyStateClass);
-	return Event.GetEndTriggerTime();
+	if(Sequence->Notifies.IsValidIndex(NotifyIndex))
+	{
+		FAnimNotifyEvent& Event = Sequence->Notifies[NotifyIndex];
+		check(Event.NotifyStateClass);
+		return Event.GetEndTriggerTime();
+	}
+	return -1.0f;
 }
 
 ETimingElementType::Type FTimingRelevantElement_NotifyStateEnd::GetType()
