Fixed check() when deleting certain configurations of notifies/notify states

Re-indexing for a frame when edits are made to the underlying data can make timing elements point at incorrect types, causing a check() to fire when a notify state no longer indexes correctly.

diff --git a/Engine/Source/Editor/Persona/Private/SAnimTimingPanel.cpp b/Engine/Source/Editor/Persona/Private/SAnimTimingPanel.cpp
--- a/Engine/Source/Editor/Persona/Private/SAnimTimingPanel.cpp
+++ b/Engine/Source/Editor/Persona/Private/SAnimTimingPanel.cpp
@@ -527,8 +527,10 @@ float FTimingRelevantElement_NotifyStateEnd::GetElementTime() const
 	if(Sequence->Notifies.IsValidIndex(NotifyIndex))
 	{
 		FAnimNotifyEvent& Event = Sequence->Notifies[NotifyIndex];
-		check(Event.NotifyStateClass);
-		return Event.GetEndTriggerTime();
+		if(Event.NotifyStateClass)
+		{
+			return Event.GetEndTriggerTime();
+		}
 	}
 	return -1.0f;
 }
