Skip updating compatible assets in the asset family shortcut bar if the asset registry is still loading assets

Prevents large hitches

diff --git a/Engine/Source/Editor/Persona/Private/SAssetFamilyShortcutBar.cpp b/Engine/Source/Editor/Persona/Private/SAssetFamilyShortcutBar.cpp
--- a/Engine/Source/Editor/Persona/Private/SAssetFamilyShortcutBar.cpp
+++ b/Engine/Source/Editor/Persona/Private/SAssetFamilyShortcutBar.cpp
@@ -383,11 +383,15 @@ public:
 
 	void HandleAssetAdded(const FAssetData& InAssetData)
 	{
-		if (AssetFamily->IsAssetCompatible(InAssetData))
+		const IAssetRegistry& AssetRegistry = FModuleManager::LoadModuleChecked<FAssetRegistryModule>("AssetRegistry").Get();
+		if(!AssetRegistry.IsLoadingAssets())
 		{
-			TArray<FAssetData> Assets;
-			AssetFamily->FindAssetsOfType(AssetData.GetClass(), Assets);
-			bMultipleAssetsExist = Assets.Num() > 1;
+			if (AssetFamily->IsAssetCompatible(InAssetData))
+			{
+				TArray<FAssetData> Assets;
+				AssetFamily->FindAssetsOfType(AssetData.GetClass(), Assets);
+				bMultipleAssetsExist = Assets.Num() > 1;
+			}
 		}
 	}
 
