Prefer using the physics asset of the currently displayed mesh instead of the one who's preview mesh matches the current mesh

diff --git a/Engine/Source/Editor/Persona/Private/PersonaAssetFamily.cpp b/Engine/Source/Editor/Persona/Private/PersonaAssetFamily.cpp
--- a/Engine/Source/Editor/Persona/Private/PersonaAssetFamily.cpp
+++ b/Engine/Source/Editor/Persona/Private/PersonaAssetFamily.cpp
@@ -136,13 +136,24 @@ FAssetData FPersonaAssetFamily::FindAssetOfType(UClass* InAssetClass) const
 			}
 			else
 			{
+				// If we have a mesh and it has a physics asset, use it
+				if (Mesh.IsValid())
+				{
+					if (UPhysicsAsset* MeshPhysicsAsset = Mesh->GetPhysicsAsset())
+					{
+						return FAssetData(MeshPhysicsAsset);
+					}
+				}
+
 				TArray<FAssetData> Assets;
 
 				FAssetRegistryModule& AssetRegistryModule = FModuleManager::LoadModuleChecked<FAssetRegistryModule>("AssetRegistry");
 				FARFilter Filter;
 				Filter.bRecursiveClasses = true;
 				Filter.ClassPaths.Add(UPhysicsAsset::StaticClass()->GetClassPathName());
-				if(Mesh.IsValid())
+
+				// If we have a mesh, look for a physics asset that has that mesh as its preview mesh
+				if (Mesh.IsValid())
 				{
 					Filter.TagsAndValues.Add(GET_MEMBER_NAME_CHECKED(UPhysicsAsset, PreviewSkeletalMesh), FSoftObjectPath(Mesh.Get()).ToString());
 				}
@@ -186,12 +197,23 @@ void FPersonaAssetFamily::FindAssetsOfType(UClass* InAssetClass, TArray<FAssetDa
 			FARFilter Filter;
 			Filter.bRecursiveClasses = true;
 			Filter.ClassPaths.Add(UPhysicsAsset::StaticClass()->GetClassPathName());
-			if(Mesh != nullptr)
+
+			// If we have a mesh, look for a physics asset that has that mesh as its preview mesh
+			if (Mesh.IsValid())
 			{
 				Filter.TagsAndValues.Add(GET_MEMBER_NAME_CHECKED(UPhysicsAsset, PreviewSkeletalMesh), FSoftObjectPath(Mesh.Get()).ToString());
 			}
 
 			AssetRegistryModule.Get().GetAssets(Filter, OutAssets);
+
+			// If we have a mesh and it has a physics asset, use it but only if its different from the one on the preview mesh
+			if (Mesh.IsValid())
+			{
+				if (UPhysicsAsset* MeshPhysicsAsset = Mesh->GetPhysicsAsset())
+				{
+					OutAssets.AddUnique(FAssetData(MeshPhysicsAsset));
+				}
+			}
 		}
 	}
 }
@@ -332,6 +354,13 @@ bool FPersonaAssetFamily::IsAssetCompatible(const FAssetData& InAssetData) const
 		}
 		else if (Class->IsChildOf<UPhysicsAsset>())
 		{
+			// If our mesh is valid and this is the physics asset used on it, we are compatible
+			if (Mesh.IsValid() && InAssetData.GetSoftObjectPath() == FSoftObjectPath(Mesh.Get()->GetPhysicsAsset()))
+			{
+				return true;
+			}
+
+			// Otherwise check if our mesh is the preview mesh of the physics asset
 			FAssetDataTagMapSharedView::FFindTagResult Result = InAssetData.TagsAndValues.FindTag(GET_MEMBER_NAME_CHECKED(UPhysicsAsset, PreviewSkeletalMesh));
 			if (Result.IsSet() && Mesh.IsValid())
 			{
