Fixed an issue where LOD bones did not display correctly in the Animation Editor Viewport

diff --git a/Engine/Source/Editor/Persona/Private/AnimationEditorViewportClient.cpp b/Engine/Source/Editor/Persona/Private/AnimationEditorViewportClient.cpp
--- a/Engine/Source/Editor/Persona/Private/AnimationEditorViewportClient.cpp
+++ b/Engine/Source/Editor/Persona/Private/AnimationEditorViewportClient.cpp
@@ -525,14 +525,6 @@ void FAnimationViewportClient::Draw(const FSceneView* View, FPrimitiveDrawInterf
 
 		const FReferenceSkeleton& RefSkeleton = PreviewMeshComponent->GetReferenceSkeleton();
 		const TArray<FBoneIndexType>& DrawBoneIndices = PreviewMeshComponent->GetDrawBoneIndices();
-		//If DrawBoneIndices count is different from the refskeleton, it mean the component must be re-register
-		//It can happen during an asynchronous re-import of a skeletal mesh where the properties lock is release
-		//before the postedit change scope have re-register the components
-		if (RefSkeleton.GetNum() != DrawBoneIndices.Num())
-		{
-			//Skip this draw call until the component cache fit the skeletalmesh data
-			continue;
-		}
 
 		// if we have BonesOfInterest, draw sub set of the bones only
 		if (GetAnimPreviewScene()->GetSelectedBoneIndex() != INDEX_NONE)
