Persona: Speculative fix

diff --git a/Engine/Source/Editor/Persona/Private/PreviewSceneCustomizations.cpp b/Engine/Source/Editor/Persona/Private/PreviewSceneCustomizations.cpp
--- a/Engine/Source/Editor/Persona/Private/PreviewSceneCustomizations.cpp
+++ b/Engine/Source/Editor/Persona/Private/PreviewSceneCustomizations.cpp
@@ -198,19 +198,19 @@ void FPreviewSceneDescriptionCustomization::CustomizeDetails(IDetailLayoutBuilde
 		.AddProperty(SkeletalMeshProperty)
 		.CustomWidget()
 		.OverrideResetToDefault(FResetToDefaultOverride::Create(
-			TAttribute<bool>::CreateLambda([this]()
+			TAttribute<bool>::CreateLambda([PreviewSceneWeakPtr = TWeakPtr<class FAnimationEditorPreviewScene>(PreviewScene)]()
 			{
-				if (PreviewScene.IsValid())
+				if (PreviewSceneWeakPtr.IsValid())
 				{
-					return PreviewScene.Pin()->GetPreviewMesh() != nullptr;
+					return PreviewSceneWeakPtr.Pin()->GetPreviewMesh() != nullptr;
 				}
 				return false;
 			}),
-			FSimpleDelegate::CreateLambda([this]() 
+			FSimpleDelegate::CreateLambda([PreviewSceneWeakPtr = TWeakPtr<class FAnimationEditorPreviewScene>(PreviewScene)]()
 			{
-				if (PreviewScene.IsValid())
+				if (PreviewSceneWeakPtr.IsValid())
 				{
-					PreviewScene.Pin()->SetPreviewMesh(nullptr, false);
+					PreviewSceneWeakPtr.Pin()->SetPreviewMesh(nullptr, false);
 				}
 			}))
 		)
@@ -229,21 +229,24 @@ void FPreviewSceneDescriptionCustomization::CustomizeDetails(IDetailLayoutBuilde
 				SNew(SButton)
 				.Text(LOCTEXT("ApplyToAsset", "Apply To Asset"))
 				.ToolTipText(LOCTEXT("ApplyToAssetToolTip", "The preview mesh has changed, but it will not be able to be saved until it is applied to the asset. Click here to make the change to the preview mesh persistent."))
-				.Visibility_Lambda([this]()
+				.Visibility_Lambda([PersonaToolkitWeakPtr = TWeakPtr<class IPersonaToolkit>(PersonaToolkit)]()
 				{
-					if (PersonaToolkit.IsValid() && PreviewScene.IsValid())
+					if (PersonaToolkitWeakPtr.IsValid())
 					{
-						const TSharedPtr<IPersonaToolkit> PinnedPersonaToolkit = PersonaToolkit.Pin();
+						const TSharedPtr<IPersonaToolkit> PinnedPersonaToolkit = PersonaToolkitWeakPtr.Pin();
 						USkeletalMesh* SkeletalMesh = PinnedPersonaToolkit->GetPreviewMesh();
 						return (SkeletalMesh != PinnedPersonaToolkit->GetPreviewScene()->GetPreviewMesh()) ? EVisibility::Visible : EVisibility::Collapsed;
 					}
 
 					return EVisibility::Collapsed;
 				})
-				.OnClicked_Lambda([this]() 
+				.OnClicked_Lambda([PersonaToolkitWeakPtr = TWeakPtr<class IPersonaToolkit>(PersonaToolkit)]()
 				{
-					TSharedPtr<IPersonaToolkit> PinnedPersonaToolkit = PersonaToolkit.Pin();
-					PinnedPersonaToolkit->SetPreviewMesh(PinnedPersonaToolkit->GetPreviewScene()->GetPreviewMesh(), true);
+					if (PersonaToolkitWeakPtr.IsValid())
+					{
+						TSharedPtr<IPersonaToolkit> PinnedPersonaToolkit = PersonaToolkitWeakPtr.Pin();
+						PinnedPersonaToolkit->SetPreviewMesh(PinnedPersonaToolkit->GetPreviewScene()->GetPreviewMesh(), true);
+					}
 					return FReply::Handled();
 				})
 			]
