AnimationPreviewScene: Fix crash in Control Rig when opening up PreviewScene

diff --git a/Engine/Source/Editor/Persona/Private/AnimationEditorPreviewScene.cpp b/Engine/Source/Editor/Persona/Private/AnimationEditorPreviewScene.cpp
--- a/Engine/Source/Editor/Persona/Private/AnimationEditorPreviewScene.cpp
+++ b/Engine/Source/Editor/Persona/Private/AnimationEditorPreviewScene.cpp
@@ -300,15 +300,18 @@ void FAnimationEditorPreviewScene::SetPreviewMeshInternal(USkeletalMesh* NewPrev
 	// Setting the skeletal mesh to in the PreviewScene can change AnimScriptInstance so we must re register it
 	// with the AnimBlueprint
 	UAnimBlueprint* SourceBlueprint = PersonaToolkit.Pin()->GetAnimBlueprint();
-	if (DebuggedSkeletalMeshComponent && DebuggedSkeletalMeshComponent->GetAnimInstance() && DebuggedSkeletalMeshComponent->GetAnimInstance()->IsA(SourceBlueprint->GeneratedClass))
+	if (SourceBlueprint)
 	{
-		PersonaUtils::SetObjectBeingDebugged(SourceBlueprint, DebuggedSkeletalMeshComponent->GetAnimInstance());
-	}
+		if (DebuggedSkeletalMeshComponent && DebuggedSkeletalMeshComponent->GetAnimInstance() && DebuggedSkeletalMeshComponent->GetAnimInstance()->IsA(SourceBlueprint->GeneratedClass))
+		{
+			PersonaUtils::SetObjectBeingDebugged(SourceBlueprint, DebuggedSkeletalMeshComponent->GetAnimInstance());
+		}
 
-	// If we didn't have a preview mesh before and we select one now, set it up as the object being debugged
-	if (DebuggedSkeletalMeshComponent == nullptr && NewPreviewMesh != nullptr && SkeletalMeshComponent->GetAnimInstance() && SkeletalMeshComponent->GetAnimInstance()->IsA(SourceBlueprint->GeneratedClass))
-	{
-		PersonaUtils::SetObjectBeingDebugged(SourceBlueprint, SkeletalMeshComponent->GetAnimInstance());
+		// If we didn't have a preview mesh before and we select one now, set it up as the object being debugged
+		if (DebuggedSkeletalMeshComponent == nullptr && NewPreviewMesh != nullptr && SkeletalMeshComponent->GetAnimInstance() && SkeletalMeshComponent->GetAnimInstance()->IsA(SourceBlueprint->GeneratedClass))
+		{
+			PersonaUtils::SetObjectBeingDebugged(SourceBlueprint, SkeletalMeshComponent->GetAnimInstance());
+		}
 	}
 
 	OnPreviewMeshChanged.Broadcast(OldPreviewMesh, NewPreviewMesh);
