Contributed by [REDACTED_PERSON] from The Coalition:
Tree View for Animation Curves
New settings in AnimationEditor/Timeline : UseTreeViewForAnimationCurves, AnimationCurveGroupingDelimiters to enable tree view for curves in Animation Sequence Editor
Checklist:
- Support multiple delimiters, default to "._/|\\"
- Support path compression, "A.B.C" maybe preseted as one node instead of three if B is the only child of A.
- Support name reused for both curve & path, curve "A.B" and "A.B.C" can exist the same time.
- Support name collision with different delimiters, "A.B", "A/B" and "A_B" can exist the same time.
- Context menu actions Remove Curves & Edit Curves is working as expected.

diff --git a/Engine/Source/Editor/Persona/Private/AnimTimeline/AnimModel_AnimSequenceBase.cpp b/Engine/Source/Editor/Persona/Private/AnimTimeline/AnimModel_AnimSequenceBase.cpp
--- a/Engine/Source/Editor/Persona/Private/AnimTimeline/AnimModel_AnimSequenceBase.cpp
+++ b/Engine/Source/Editor/Persona/Private/AnimTimeline/AnimModel_AnimSequenceBase.cpp
@@ -7,6 +7,7 @@
 #include "AnimTimeline/AnimTimelineTrack_NotifiesPanel.h"
 #include "AnimTimeline/AnimTimelineTrack_Curves.h"
 #include "AnimTimeline/AnimTimelineTrack_Curve.h"
+#include "AnimTimeline/AnimTimelineTrack_CompoundCurve.h"
 #include "AnimTimeline/AnimTimelineTrack_FloatCurve.h"
 #include "AnimTimeline/AnimTimelineTrack_VectorCurve.h"
 #include "AnimTimeline/AnimTimelineTrack_TransformCurve.h"
@@ -224,9 +225,17 @@ void FAnimModel_AnimSequenceBase::RefreshCurveTracks()
 
 	// Next add a track for each float curve
 	const FAnimationCurveData& AnimationModelCurveData = AnimSequenceBase->GetDataModel()->GetCurveData();
-	for(const FFloatCurve& FloatCurve : AnimationModelCurveData.FloatCurves)
+	const UPersonaOptions* PersonaOptions = GetDefault<UPersonaOptions>();
+	if (!PersonaOptions || !PersonaOptions->bUseTreeViewForAnimationCurves)
 	{
-		CurveRoot->AddChild(MakeShared<FAnimTimelineTrack_FloatCurve>(&FloatCurve, SharedThis(this)));
+		for (const FFloatCurve& FloatCurve : AnimationModelCurveData.FloatCurves)
+		{
+			CurveRoot->AddChild(MakeShared<FAnimTimelineTrack_FloatCurve>(&FloatCurve, SharedThis(this)));
+		}
+	}
+	else
+	{
+		FAnimTimelineTrack_CompoundCurve::AddGroupedCurveTracks(AnimationModelCurveData.FloatCurves, *CurveRoot, SharedThis(this), PersonaOptions->AnimationCurveGroupingDelimiters);
 	}
 
 	UAnimSequence* AnimSequence = Cast<UAnimSequence>(AnimSequenceBase);
@@ -466,6 +475,26 @@ void FAnimModel_AnimSequenceBase::RemoveSelectedCurves()
 				}
 			}	
 		}
+		else if (SelectedTrack->IsA<FAnimTimelineTrack_CompoundCurve>())
+		{
+			// Lowest priority for base class delete, ERawCurveTrackTypes::RCT_Transform must take priority
+			TSharedRef<FAnimTimelineTrack_CompoundCurve> CurveTrack = StaticCastSharedRef<FAnimTimelineTrack_CompoundCurve>(SelectedTrack);
+
+			// Find all editable curves in this track
+			for (const FSmartName CurveName : CurveTrack->GetCurveNames())
+			{
+				const FAnimationCurveIdentifier CurveId(CurveName, ERawCurveTrackTypes::RCT_Float);
+				if (AnimSequenceBase->GetDataModel()->FindCurve(CurveId))
+				{
+					FSmartName TrackName;
+					if (AnimSequenceBase->GetSkeleton()->GetSmartNameByUID(USkeleton::AnimCurveMappingName, CurveName.UID, TrackName))
+					{
+						Controller.RemoveCurve(CurveId);
+						bDeletedCurve = true;
+					}
+				}
+			}
+		}
 	}
 
 	Controller.CloseBracket();
diff --git a/Engine/Source/Editor/Persona/Private/AnimTimeline/AnimTimelineTrack_CompoundCurve.cpp b/Engine/Source/Editor/Persona/Private/AnimTimeline/AnimTimelineTrack_CompoundCurve.cpp
new file mode 100644
--- /dev/null
+++ b/Engine/Source/Editor/Persona/Private/AnimTimeline/AnimTimelineTrack_CompoundCurve.cpp
@@ -0,0 +1,257 @@
+// Copyright Epic Games, Inc. All Rights Reserved.
+
+#include "AnimTimelineTrack_CompoundCurve.h"
+#include "AnimTimelineTrack_FloatCurve.h"
+#include "Containers/BitArray.h"
+
+ANIMTIMELINE_IMPLEMENT_TRACK(FAnimTimelineTrack_CompoundCurve);
+ANIMTIMELINE_IMPLEMENT_TRACK(FAnimTimelineTrack_FloatCurveWithDisplayName);
+
+FAnimTimelineTrack_FloatCurveWithDisplayName::FAnimTimelineTrack_FloatCurveWithDisplayName(const FFloatCurve* InCurve, FText InDisplayName, const TSharedRef<FAnimModel>& InModel)
+	: FAnimTimelineTrack_FloatCurve(InCurve, InModel)
+{
+	DisplayName = InDisplayName;
+}
+
+FText FAnimTimelineTrack_FloatCurveWithDisplayName::GetLabel() const
+{
+	// FAnimTimelineTrack_FloatCurve forces the label to its underlying curve name
+	// We want to use the shorter alias for it
+	return FAnimTimelineTrack::GetLabel();
+}
+
+bool FAnimTimelineTrack_FloatCurveWithDisplayName::ShowCurves() const
+{
+	// Always show curves even if expanded
+	return true;
+}
+
+TArray<const FRichCurve*> FAnimTimelineTrack_CompoundCurve::ToRichCurves(TArrayView<FFloatCurve const* const> InCurves)
+{
+	TArray<const FRichCurve*> RichCurves;
+	Algo::Transform(InCurves, RichCurves, [](const FFloatCurve* Curve) { return &Curve->FloatCurve; });
+	return RichCurves;
+}
+
+FAnimTimelineTrack_CompoundCurve::FAnimTimelineTrack_CompoundCurve(TArray<const FFloatCurve*> InCurves, const FText& InCurveName, const FText& InFullCurveName, const FLinearColor& InBackgroundColor, const TSharedRef<FAnimModel>& InModel)
+	: FAnimTimelineTrack_Curve(ToRichCurves(InCurves), InCurveName, InFullCurveName, FLinearColor{}, InBackgroundColor, InModel)
+	, Curves(MoveTemp(InCurves))
+{
+	OuterCurveIndex = 0;
+	Algo::Transform(Curves, CurveNames, [](const FFloatCurve* Curve) { return Curve->Name; });
+}
+
+/** FAnimTimelineTrack_Curve interface */
+FLinearColor FAnimTimelineTrack_CompoundCurve::GetCurveColor(int32 InCurveIndex) const
+{
+	return Curves[InCurveIndex]->Color;
+}
+
+FText FAnimTimelineTrack_CompoundCurve::GetFullCurveName(int32 InCurveIndex) const
+{
+	return FText::FromName(Curves[InCurveIndex]->Name.DisplayName);
+}
+
+bool FAnimTimelineTrack_CompoundCurve::CanEditCurve(int32 InCurveIndex) const
+{
+	return true;
+}
+
+void FAnimTimelineTrack_CompoundCurve::GetCurveEditInfo(int32 InCurveIndex, FSmartName& OutName, ERawCurveTrackTypes& OutType, int32& OutCurveIndex) const
+{
+	OutName = Curves[InCurveIndex]->Name;
+	OutType = ERawCurveTrackTypes::RCT_Float;
+	OutCurveIndex = 0;
+}
+
+struct FAnimTimelineTrack_CompoundCurve::FCurveGroup
+{
+	using CharType = FString::ElementType;
+
+	FString FullName; // Full path name of this curve group, e.g. A.B or A.B.C
+	int		NameLength = 0; // Length of the last part of the full name
+
+	const FFloatCurve*	Curve = nullptr;	// Curve of FullName
+	TArray<FCurveGroup>	Subgroups;			// Children of this curve group
+
+	// \brief Get the last part of the curve name
+	FStringView Name() const
+	{
+		return FStringView{FullName}.Right(NameLength);
+	}
+
+	//! \brief Get the name for lookup i.e. ".{Name}", Delimiter+Name
+	FStringView LookupName() const
+	{
+		return FStringView{FullName}.Right(NameLength + 1);
+	}
+
+	FCurveGroup() = default;
+	FCurveGroup(FStringView InLookupName, FStringView ParentName)
+	: FullName{ParentName, InLookupName.Len()}
+	{
+		FullName += InLookupName;
+		NameLength = InLookupName.Len() - !ParentName.IsEmpty();
+	}
+
+	bool operator==(FStringView InLookupName) const
+	{
+		return LookupName() == InLookupName;
+	}
+
+	bool IsRoot() const
+	{
+		return FullName.IsEmpty();
+	}
+
+	FCurveGroup& FindOrEmplace(FStringView SubgroupLookupName)
+	{
+		FCurveGroup* Subgroup = Subgroups.FindByKey(SubgroupLookupName);
+		if (!Subgroup)
+		{
+			Subgroup = &Subgroups.Emplace_GetRef(SubgroupLookupName, FullName);
+		}
+		return *Subgroup;
+	}
+
+	// Combine non-curve groups only have one child
+	void Compress()
+	{
+		TArray<FCurveGroup> Temp;
+		while (Subgroups.Num() == 1 && Curve == nullptr)
+		{
+			FCurveGroup& Subgroup = Subgroups[0];
+			NameLength = NameLength + Subgroup.NameLength + 1;
+			FullName = MoveTemp(Subgroup.FullName);
+			Curve = Subgroup.Curve;
+			// Moving Subgroups must avoid aliasing (this.Subgroup.Subgroups in a data member of this.Subgroup)
+			Temp = MoveTemp(Subgroup.Subgroups);
+			Subgroups = MoveTemp(Temp);
+		}
+		for (auto& Subgroup : Subgroups)
+		{
+			Subgroup.Compress();
+		}
+	}
+
+	static FCurveGroup MakeCurveGroup(TArrayView<const FFloatCurve> FloatCurves, FStringView Delimiters)
+	{
+		static constexpr int BitSetSize = 256; // Char set for 256 ASCII characters, should be enough to define delimiters
+		TBitArray<TInlineAllocator<BitSetSize/FBitSet::BitsPerWord>> DelimiterSet;
+		DelimiterSet.Init(false, 256);
+		check(DelimiterSet.Num() >= BitSetSize);
+		for (CharType Delimiter : Delimiters)
+		{
+			if (static_cast<int>(Delimiter) < BitSetSize)
+			{
+				DelimiterSet[Delimiter] = true;
+			}
+		}
+		auto IsDelimiter = [&](CharType Char) -> bool {
+			if (Char < BitSetSize)
+			{
+				return DelimiterSet[Char];
+			}
+			else
+			{
+				// Should be optimized on Windows where TCHAR is one byte
+				return UE::String::FindFirstChar(Delimiters, Char) != INDEX_NONE;
+			}
+		};
+
+		FCurveGroup RootGroup;
+		for (auto& FloatCurve : FloatCurves)
+		{
+			FString CurveNameStr = FloatCurve.Name.DisplayName.ToString();
+			FStringView CurveName = CurveNameStr;
+
+			// We need c++ 20 std::split_view here
+			// for (auto GroupName : std::ranges::views::split(CurveName, TEXT(".")))
+			int GroupBegin = 0;
+			FCurveGroup* CurrentGroup = &RootGroup;
+			for (int I = 0; I < CurveName.Len(); I++)
+			{
+				bool IsLastGroup = I == CurveName.Len() - 1;
+				if (IsLastGroup || IsDelimiter(CurveName[I]))
+				{
+					int GroupEnd = I + IsLastGroup;
+					FStringView GroupName = CurveName.SubStr(GroupBegin, GroupEnd - GroupBegin);
+
+					// GroupName maybe empty, we may skip empty groups
+					// Which could lead to trouble for this tricky case ["A.B", "A..B", "A...B"]
+					CurrentGroup = &CurrentGroup->FindOrEmplace(GroupName);
+
+					GroupBegin = I;
+				}
+			}
+
+			// Set curve on the leaf group
+			ensure(CurrentGroup && !CurrentGroup->Curve);
+			CurrentGroup->Curve = &FloatCurve;
+		}
+
+		RootGroup.Compress();
+		return RootGroup;
+	}
+
+	//! \brief Collect descendant curves under this group
+	//! \note This recursive function is O(Curves), make the tree view creation O(Curves^2)
+	//!       Which should be OK given the small number of curves we have, add cache later when the demand raises
+	void GatherCurves(TArray<const FFloatCurve*>& InOutCurves)
+	{
+		if (Curve)
+		{
+			InOutCurves.Add(Curve);
+		}
+		else
+		{
+			for (auto& Subgroup : Subgroups)
+			{
+				Subgroup.GatherCurves(InOutCurves);
+			}
+		}
+	}
+
+	// Recursive add curve groups to ParentTrack
+	void AddCurvesToTrack(const TSharedRef<FAnimModel>& InModel, FAnimTimelineTrack* ParentTrack)
+	{
+		TSharedPtr<FAnimTimelineTrack> CurveTrack;
+
+		// If this curve group is not root, create a node for it
+		if (!IsRoot())
+		{
+			if (Curve != nullptr)
+			{
+				// If the current group have a curve, create a single curve track for it
+				// Such track can still have child tracks, those child curve will not show up when editing this curve, e.g. Curves ["Left.Foot", "Left.Foot.IK"]
+				CurveTrack = MakeShared<FAnimTimelineTrack_FloatCurveWithDisplayName>(Curve, FText::FromStringView(Name()), InModel);
+			}
+			else
+			{
+				// Create a compound track for all descend curves in this group
+				TArray<const FFloatCurve*> CurvesInGroup;
+				GatherCurves(CurvesInGroup);
+				ensure(!CurvesInGroup.IsEmpty());
+				FLinearColor CurveBackgroundColor = FLinearColor::Transparent;
+				CurveTrack = MakeShared<FAnimTimelineTrack_CompoundCurve>(MoveTemp(CurvesInGroup), FText::FromStringView(Name()), FText::FromString(FullName), CurveBackgroundColor, InModel);
+				CurveTrack->SetExpanded(false);
+			}
+
+			if (CurveTrack) {
+				ParentTrack->AddChild(CurveTrack.ToSharedRef());
+				ParentTrack = CurveTrack.Get();
+			}
+		}
+
+		for (auto& Subgroup : Subgroups)
+		{
+			Subgroup.AddCurvesToTrack(InModel, ParentTrack);
+		}
+	}
+};
+
+void FAnimTimelineTrack_CompoundCurve::AddGroupedCurveTracks(TArrayView<const FFloatCurve> FloatCurves, FAnimTimelineTrack& ParentTrack, const TSharedRef<FAnimModel>& InModel, FStringView Delimiters)
+{
+	FCurveGroup CurveGroups = FCurveGroup::MakeCurveGroup(FloatCurves, Delimiters);
+	CurveGroups.AddCurvesToTrack(InModel, &ParentTrack);
+}
diff --git a/Engine/Source/Editor/Persona/Private/AnimTimeline/AnimTimelineTrack_CompoundCurve.h b/Engine/Source/Editor/Persona/Private/AnimTimeline/AnimTimelineTrack_CompoundCurve.h
new file mode 100644
--- /dev/null
+++ b/Engine/Source/Editor/Persona/Private/AnimTimeline/AnimTimelineTrack_CompoundCurve.h
@@ -0,0 +1,51 @@
+// Copyright Epic Games, Inc. All Rights Reserved.
+
+#pragma once
+#include "AnimTimelineTrack_Curve.h"
+#include "AnimTimelineTrack_FloatCurve.h"
+#include "Containers/Array.h"
+#include "Containers/ArrayView.h"
+
+struct FFloatCurve;
+
+//! \brief Tree view of a group of float curves, use AddGroupedCurveTracks(FloatCurves, ...) to create a tree view for curves
+class FAnimTimelineTrack_CompoundCurve : public FAnimTimelineTrack_Curve
+{
+	ANIMTIMELINE_DECLARE_TRACK(FAnimTimelineTrack_CompoundCurve, FAnimTimelineTrack_Curve);
+
+public:
+	FAnimTimelineTrack_CompoundCurve(TArray<const FFloatCurve*> InCurves, const FText& InCurveName, const FText& InFullCurveName, const FLinearColor& InBackgroundColor, const TSharedRef<FAnimModel>& InModel);
+
+	/** FAnimTimelineTrack_Curve interface */
+	virtual FLinearColor GetCurveColor(int32 InCurveIndex) const override;
+	virtual FText GetFullCurveName(int32 InCurveIndex) const override;
+	virtual bool CanEditCurve(int32 InCurveIndex) const override;
+	virtual void GetCurveEditInfo(int32 InCurveIndex, FSmartName& OutName, ERawCurveTrackTypes& OutType, int32& OutCurveIndex) const override;
+	/** End FAnimTimelineTrack_Curve interface */
+
+	static constexpr auto DefaultDelimiters = TEXT("._/\\|");
+	//! \brief Add grouped view curve tracks for a list of curves named like 'A.B.C'
+	//! \param[in] FloatCurves The curves to display
+	//! \param[inout] ParentTrack The parent track to add the curves, typically a FAnimTimelineTrack_Curve
+	//! \param[in] InModel FAnimModel to display the track
+	static void AddGroupedCurveTracks(TArrayView<const FFloatCurve> FloatCurves, FAnimTimelineTrack& ParentTrack, const TSharedRef<FAnimModel>& InModel, FStringView Delimiters = DefaultDelimiters);
+
+	const TArray<FSmartName>& GetCurveNames() const { return CurveNames; }
+private:
+	TArray<const FFloatCurve*> Curves;
+	TArray<FSmartName> CurveNames; // Curve smart names used when removing curves (in case the FFloatCurve was removed)
+
+	static TArray<const FRichCurve*> ToRichCurves(TArrayView<FFloatCurve const* const> InCurves);
+	struct FCurveGroup;
+};
+
+// FloatCurve with customizable display name, e.g. showing "C" for curve "A.B.C"
+class FAnimTimelineTrack_FloatCurveWithDisplayName : public FAnimTimelineTrack_FloatCurve
+{
+	ANIMTIMELINE_DECLARE_TRACK(FAnimTimelineTrack_FloatCurveWithDisplayName, FAnimTimelineTrack_FloatCurve);
+
+public:
+	FAnimTimelineTrack_FloatCurveWithDisplayName(const FFloatCurve* InCurve, FText InDisplayName, const TSharedRef<FAnimModel>& InModel);
+	virtual FText GetLabel() const override;
+	virtual bool ShowCurves() const override;
+};
diff --git a/Engine/Source/Editor/Persona/Private/AnimTimeline/AnimTimelineTrack_Curve.cpp b/Engine/Source/Editor/Persona/Private/AnimTimeline/AnimTimelineTrack_Curve.cpp
--- a/Engine/Source/Editor/Persona/Private/AnimTimeline/AnimTimelineTrack_Curve.cpp
+++ b/Engine/Source/Editor/Persona/Private/AnimTimeline/AnimTimelineTrack_Curve.cpp
@@ -254,7 +254,7 @@ TSharedRef<SWidget> FAnimTimelineTrack_Curve::MakeCurveWidget()
 			.Visibility_Lambda([this]()
 			{  
 				// Dont show curves in parent tracks when children are expanded
-				return !IsExpanded() || GetChildren().Num() == 0 ? EVisibility::Visible : EVisibility::Hidden;
+				return ShowCurves() ? EVisibility::Visible : EVisibility::Hidden;
 			})
 			.CurveThickness_Lambda([this]()
 			{
@@ -268,7 +268,7 @@ TSharedRef<SWidget> FAnimTimelineTrack_Curve::MakeCurveWidget()
 			.Visibility_Lambda([this]()
 			{  
 				// Dont show curves in parent tracks when children are expanded
-				return (!IsExpanded() || GetChildren().Num() == 0) && IsHovered() ? EVisibility::Visible : EVisibility::Hidden;
+				return ShowCurves() && IsHovered() ? EVisibility::Visible : EVisibility::Hidden;
 			})
 		];
 }
@@ -301,6 +301,12 @@ void FAnimTimelineTrack_Curve::AddCurveTrackButton(TSharedPtr<SHorizontalBox> In
 	];
 }
 
+bool FAnimTimelineTrack_Curve::ShowCurves() const
+{
+	// Dont show curves in parent tracks when children are expanded
+	return !IsExpanded() || Children.Num() == 0;
+}
+
 TSharedRef<SWidget> FAnimTimelineTrack_Curve::BuildCurveTrackMenu()
 {
 	FMenuBuilder MenuBuilder(true, GetModel()->GetCommandList());
diff --git a/Engine/Source/Editor/Persona/Private/AnimTimeline/AnimTimelineTrack_Curve.h b/Engine/Source/Editor/Persona/Private/AnimTimeline/AnimTimelineTrack_Curve.h
--- a/Engine/Source/Editor/Persona/Private/AnimTimeline/AnimTimelineTrack_Curve.h
+++ b/Engine/Source/Editor/Persona/Private/AnimTimeline/AnimTimelineTrack_Curve.h
@@ -66,6 +66,9 @@ protected:
 	/** Helper function for building outliner widget */
 	virtual void AddCurveTrackButton(TSharedPtr<SHorizontalBox> InnerHorizontalBox);
 
+	/** Should this track draw its curves, by default, hide curves in parent tracks when children are expanded. */
+	virtual bool ShowCurves() const;
+
 	/** Edit this curve when double clicked */
 	FReply HandleDoubleClicked(const FGeometry& InGeometry, const FPointerEvent& InPointerEvent);
 
diff --git a/Engine/Source/Editor/UnrealEd/Classes/Preferences/PersonaOptions.h b/Engine/Source/Editor/UnrealEd/Classes/Preferences/PersonaOptions.h
--- a/Engine/Source/Editor/UnrealEd/Classes/Preferences/PersonaOptions.h
+++ b/Engine/Source/Editor/UnrealEd/Classes/Preferences/PersonaOptions.h
@@ -243,6 +243,13 @@ class UNREALED_API UPersonaOptions : public UObject
 	UPROPERTY(EditAnywhere, config, Category = "Assets")
 	bool bAllowIncompatibleSkeletonSelection;
 
+	/** Whether to use tree view for animation curves*/
+	UPROPERTY(EditAnywhere, config, Category = "Timeline")
+	bool bUseTreeViewForAnimationCurves = false;
+
+	/** Delimiters to split animation curve names for grouping*/
+	UPROPERTY(EditAnywhere, config, Category = "Timeline")
+	FString AnimationCurveGroupingDelimiters = TEXT("._/|\\");
 public:
 	void SetShowGrid( bool bInShowGrid );
 	void SetHighlightOrigin( bool bInHighlightOrigin );
