//UE5/Release-5.1 - Build ContentExamples Win64 - Warning: PoseAsset is out-of-date with its source animation
   - Use SmartName display name rather than UID when generating data model GUID
   - Defer PoseAsset out-of-date warning until next resave

diff --git a/Engine/Source/Editor/Persona/Private/AnimationEditorViewportClient.cpp b/Engine/Source/Editor/Persona/Private/AnimationEditorViewportClient.cpp
--- a/Engine/Source/Editor/Persona/Private/AnimationEditorViewportClient.cpp
+++ b/Engine/Source/Editor/Persona/Private/AnimationEditorViewportClient.cpp
@@ -43,6 +43,7 @@
 #include "Rendering/SkeletalMeshModel.h"
 #include "UnrealWidget.h"
 #include "Engine/PoseWatchRenderData.h"
+#include "UObject/UE5MainStreamObjectVersion.h"
 
 namespace {
 	static const float AnimationEditorViewport_RotateSpeed = 0.02f;
@@ -1254,7 +1255,7 @@ FText FAnimationViewportClient::GetDisplayInfo(bool bDisplayAllInfo) const
 
 	if (const UPoseAsset* PoseAsset = Cast<UPoseAsset>(GetAnimPreviewScene()->GetPreviewAnimationAsset()))
 	{
-		if (PoseAsset->SourceAnimation && PoseAsset->SourceAnimation->GetDataModel()->GenerateGuid() != PoseAsset->SourceAnimationRawDataGUID)
+		if (PoseAsset->GetLinkerCustomVersion(FUE5MainStreamObjectVersion::GUID) >= FUE5MainStreamObjectVersion::PoseAssetRawDataGUIDUpdate && PoseAsset->SourceAnimation && PoseAsset->SourceAnimation->GetDataModel()->GenerateGuid() != PoseAsset->SourceAnimationRawDataGUID)
 		{
 			TextValue = ConcatenateLine(TextValue, LOCTEXT("PoseAssetOutOfDateWarning", "<AnimViewport.WarningText>Poses are out-of-sync with the source animation. To update them click \"Update Source\"</>"));
 		}
diff --git a/Engine/Source/Runtime/Core/Public/UObject/UE5MainStreamObjectVersion.h b/Engine/Source/Runtime/Core/Public/UObject/UE5MainStreamObjectVersion.h
--- a/Engine/Source/Runtime/Core/Public/UObject/UE5MainStreamObjectVersion.h
+++ b/Engine/Source/Runtime/Core/Public/UObject/UE5MainStreamObjectVersion.h
@@ -298,6 +298,9 @@ struct CORE_API FUE5MainStreamObjectVersion
 		// Added support for the RigVM to run branches lazily
 		RigVMLazyEvaluation,
 
+		// Adding additional object version to defer out-of-date pose asset warning until next resaves
+		PoseAssetRawDataGUIDUpdate,
+		
 		// -----<new versions can be added above this line>-------------------------------------------------
 		VersionPlusOne,
 		LatestVersion = VersionPlusOne - 1
diff --git a/Engine/Source/Runtime/Engine/Private/Animation/AnimDataModel.cpp b/Engine/Source/Runtime/Engine/Private/Animation/AnimDataModel.cpp
--- a/Engine/Source/Runtime/Engine/Private/Animation/AnimDataModel.cpp
+++ b/Engine/Source/Runtime/Engine/Private/Animation/AnimDataModel.cpp
@@ -425,7 +425,8 @@ FGuid UAnimDataModel::GenerateGuid() const
 
 	for (const FFloatCurve& Curve : CurveData.FloatCurves)
 	{
-		UpdateWithData(Curve.Name.UID);
+		const FString CurveName = Curve.Name.DisplayName.ToString();
+		UpdateSHAWithArray(CurveName.GetCharArray()),
 		UpdateWithFloatCurve(Curve.FloatCurve);
 	}
 
diff --git a/Engine/Source/Runtime/Engine/Private/Animation/PoseAsset.cpp b/Engine/Source/Runtime/Engine/Private/Animation/PoseAsset.cpp
--- a/Engine/Source/Runtime/Engine/Private/Animation/PoseAsset.cpp
+++ b/Engine/Source/Runtime/Engine/Private/Animation/PoseAsset.cpp
@@ -16,6 +16,7 @@
 #include "UObject/LinkerLoad.h"
 #include "UObject/ObjectSaveContext.h"
 #include "UObject/Package.h"
+#include "UObject/UE5MainStreamObjectVersion.h"
 #include "UObject/UE5ReleaseStreamObjectVersion.h"
 #include "Widgets/Notifications/SNotificationList.h"
 
@@ -893,7 +894,7 @@ void UPoseAsset::PostLoad()
 		PostProcessData();
     }
 
-	if (GetLinkerCustomVersion(FUE5ReleaseStreamObjectVersion::GUID) >= FUE5ReleaseStreamObjectVersion::PoseAssetRawDataGUID)
+	if (GetLinkerCustomVersion(FUE5MainStreamObjectVersion::GUID) >= FUE5MainStreamObjectVersion::PoseAssetRawDataGUIDUpdate)
 	{
 		if (SourceAnimation)
 		{
