Static analysis fixes

diff --git a/Engine/Source/Editor/AnimationModifiers/Private/SAnimationModifiersTab.cpp b/Engine/Source/Editor/AnimationModifiers/Private/SAnimationModifiersTab.cpp
--- a/Engine/Source/Editor/AnimationModifiers/Private/SAnimationModifiersTab.cpp
+++ b/Engine/Source/Editor/AnimationModifiers/Private/SAnimationModifiersTab.cpp
@@ -661,7 +661,7 @@ void SAnimationModifiersTab::RetrieveModifierData()
 					return true;
 				});
 				// Swap the double buffer
-				Swap(NextModifiersToCheck, NextModifiersToCheck);
+				Swap(ModifiersToCheck, NextModifiersToCheck);
 
 				// Break if all modifiers are out of date
 				return !ModifiersToCheck.IsEmpty();
diff --git a/Engine/Source/Editor/Persona/Private/AnimTimeline/AnimTimelineTrack_CompoundCurve.cpp b/Engine/Source/Editor/Persona/Private/AnimTimeline/AnimTimelineTrack_CompoundCurve.cpp
--- a/Engine/Source/Editor/Persona/Private/AnimTimeline/AnimTimelineTrack_CompoundCurve.cpp
+++ b/Engine/Source/Editor/Persona/Private/AnimTimeline/AnimTimelineTrack_CompoundCurve.cpp
@@ -128,7 +128,7 @@ struct FAnimTimelineTrack_CompoundCurve::FCurveGroup
 			Temp = MoveTemp(Subgroup.Subgroups);
 			Subgroups = MoveTemp(Temp);
 		}
-		for (auto& Subgroup : Subgroups)
+		for (FCurveGroup& Subgroup : Subgroups)
 		{
 			Subgroup.Compress();
 		}
@@ -142,7 +142,7 @@ struct FAnimTimelineTrack_CompoundCurve::FCurveGroup
 		check(DelimiterSet.Num() >= BitSetSize);
 		for (CharType Delimiter : Delimiters)
 		{
-			if (static_cast<int>(Delimiter) < BitSetSize)
+			if (static_cast<int32>(Delimiter) < BitSetSize)
 			{
 				DelimiterSet[Delimiter] = true;
 			}
@@ -160,22 +160,22 @@ struct FAnimTimelineTrack_CompoundCurve::FCurveGroup
 		};
 
 		FCurveGroup RootGroup;
-		for (auto& FloatCurve : FloatCurves)
+		for (const FFloatCurve& FloatCurve : FloatCurves)
 		{
 			FString CurveNameStr = FloatCurve.Name.DisplayName.ToString();
 			FStringView CurveName = CurveNameStr;
 
 			// We need c++ 20 std::split_view here
 			// for (auto GroupName : std::ranges::views::split(CurveName, TEXT(".")))
-			int GroupBegin = 0;
+			int32 GroupBegin = 0;
 			FCurveGroup* CurrentGroup = &RootGroup;
-			for (int I = 0; I < CurveName.Len(); I++)
+			for (int32 I = 0; I < CurveName.Len(); I++)
 			{
-				bool IsLastGroup = I == CurveName.Len() - 1;
-				if (IsLastGroup || IsDelimiter(CurveName[I]))
+				const bool bIsLastGroup = I == CurveName.Len() - 1;
+				if (bIsLastGroup || IsDelimiter(CurveName[I]))
 				{
-					int GroupEnd = I + IsLastGroup;
-					FStringView GroupName = CurveName.SubStr(GroupBegin, GroupEnd - GroupBegin);
+					const int32 GroupEnd = I + static_cast<int32>(bIsLastGroup);
+					const FStringView GroupName = CurveName.SubStr(GroupBegin, GroupEnd - GroupBegin);
 
 					// GroupName maybe empty, we may skip empty groups
 					// Which could lead to trouble for this tricky case ["A.B", "A..B", "A...B"]
@@ -205,7 +205,7 @@ struct FAnimTimelineTrack_CompoundCurve::FCurveGroup
 		}
 		else
 		{
-			for (auto& Subgroup : Subgroups)
+			for (FCurveGroup& Subgroup : Subgroups)
 			{
 				Subgroup.GatherCurves(InOutCurves);
 			}
@@ -243,7 +243,7 @@ struct FAnimTimelineTrack_CompoundCurve::FCurveGroup
 			}
 		}
 
-		for (auto& Subgroup : Subgroups)
+		for (FCurveGroup& Subgroup : Subgroups)
 		{
 			Subgroup.AddCurvesToTrack(InModel, ParentTrack);
 		}
