Editor Fix : Persona PreviewMesh crash when a preview skeletal mesh executes a raw delegate on an already destroyed object

diff --git a/Engine/Source/Editor/Persona/Private/AnimationEditorViewportClient.cpp b/Engine/Source/Editor/Persona/Private/AnimationEditorViewportClient.cpp
--- a/Engine/Source/Editor/Persona/Private/AnimationEditorViewportClient.cpp
+++ b/Engine/Source/Editor/Persona/Private/AnimationEditorViewportClient.cpp
@@ -84,8 +84,6 @@ FAnimationViewportClient::FAnimationViewportClient(const TSharedRef<IPersonaPrev
 
 	OnCameraControllerChanged();
 
-	InPreviewScene->RegisterOnCameraOverrideChanged(FSimpleDelegate::CreateRaw(this, &FAnimationViewportClient::OnCameraControllerChanged));
-
 	Widget->SetUsesEditorModeTools(ModeTools.Get());
 	((FAssetEditorModeManager*)ModeTools.Get())->SetPreviewScene(&InPreviewScene.Get());
 	ModeTools->SetDefaultMode(FPersonaEditModes::SkeletonSelection);
@@ -147,21 +145,6 @@ FAnimationViewportClient::FAnimationViewportClient(const TSharedRef<IPersonaPrev
 			AudioDevice->SetUseAttenuationForNonGameWorlds(ConfigOption->bUseAudioAttenuation);
 		}
 	}
-
-	InPreviewScene->RegisterOnPreviewMeshChanged(FOnPreviewMeshChanged::CreateRaw(this, &FAnimationViewportClient::HandleSkeletalMeshChanged));
-	if (InPreviewScene->GetPreviewMeshComponent())
-	{
-		HandleSkeletalMeshChanged(nullptr, InPreviewScene->GetPreviewMeshComponent()->GetSkeletalMeshAsset());
-	}
-	InPreviewScene->RegisterOnInvalidateViews(FSimpleDelegate::CreateRaw(this, &FAnimationViewportClient::HandleInvalidateViews));
-	InPreviewScene->RegisterOnFocusViews(FSimpleDelegate::CreateRaw(this, &FAnimationViewportClient::HandleFocusViews));
-	InPreviewScene->RegisterOnPreTick(FSimpleDelegate::CreateRaw(this, &FAnimationViewportClient::HandlePreviewScenePreTick));
-	InPreviewScene->RegisterOnPostTick(FSimpleDelegate::CreateRaw(this, &FAnimationViewportClient::HandlePreviewScenePostTick));
-
-	// Register delegate to update the show flags when the post processing is turned on or off
-	UAssetViewerSettings::Get()->OnAssetViewerSettingsChanged().AddRaw(this, &FAnimationViewportClient::OnAssetViewerSettingsChanged);
-	// Set correct flags according to current profile settings
-	SetAdvancedShowFlagsForScene(UAssetViewerSettings::Get()->Profiles[GetMutableDefault<UEditorPerProjectUserSettings>()->AssetViewerProfileIndex].bPostProcessingEnabled);
 }
 
 FAnimationViewportClient::~FAnimationViewportClient()
@@ -174,15 +157,15 @@ FAnimationViewportClient::~FAnimationViewportClient()
 	}
 
 	// Unregistering the callbacks is mandatory, else we get random crashes
-	if (PreviewScenePtr.IsValid())
+	if (FAnimationEditorPreviewScene* AnimationEditorPreviewScene = GetAnimPreviewScenePtr().Get())
 	{
-		PreviewScenePtr->UnregisterOnPreviewMeshChanged(this);
-		PreviewScenePtr->UnregisterOnInvalidateViews(this);
-		PreviewScenePtr->UnregisterOnCameraOverrideChanged(this);
-		PreviewScenePtr->UnregisterOnPreTick(this);
-		PreviewScenePtr->UnregisterOnPostTick(this);
+		AnimationEditorPreviewScene->UnregisterOnPreviewMeshChanged(this);
+		AnimationEditorPreviewScene->UnregisterOnInvalidateViews(this);
+		AnimationEditorPreviewScene->UnregisterOnCameraOverrideChanged(this);
+		AnimationEditorPreviewScene->UnregisterOnPreTick(this);
+		AnimationEditorPreviewScene->UnregisterOnPostTick(this);
 
-		if (UDebugSkelMeshComponent* PreviewMeshComponent = GetAnimPreviewScene()->GetPreviewMeshComponent())
+		if (UDebugSkelMeshComponent* PreviewMeshComponent = AnimationEditorPreviewScene->GetPreviewMeshComponent())
 		{
 			if (OnPhysicsCreatedDelegateHandle.IsValid())
 			{
@@ -204,6 +187,29 @@ FAnimationViewportClient::~FAnimationViewportClient()
 	UAssetViewerSettings::Get()->OnAssetViewerSettingsChanged().RemoveAll(this);
 }
 
+void FAnimationViewportClient::Initialize()
+{
+	if (FAnimationEditorPreviewScene* AnimationEditorPreviewScene = GetAnimPreviewScenePtr().Get())
+	{
+		AnimationEditorPreviewScene->RegisterOnCameraOverrideChanged(FSimpleDelegate::CreateSP(this, &FAnimationViewportClient::OnCameraControllerChanged));
+		AnimationEditorPreviewScene->RegisterOnPreviewMeshChanged(FOnPreviewMeshChanged::CreateSP(this, &FAnimationViewportClient::HandleSkeletalMeshChanged));
+		if (AnimationEditorPreviewScene->GetPreviewMeshComponent())
+		{
+			HandleSkeletalMeshChanged(nullptr, AnimationEditorPreviewScene->GetPreviewMeshComponent()->GetSkeletalMeshAsset());
+		}
+		AnimationEditorPreviewScene->RegisterOnInvalidateViews(FSimpleDelegate::CreateSP(this, &FAnimationViewportClient::HandleInvalidateViews));
+		AnimationEditorPreviewScene->RegisterOnFocusViews(FSimpleDelegate::CreateSP(this, &FAnimationViewportClient::HandleFocusViews));
+		AnimationEditorPreviewScene->RegisterOnPreTick(FSimpleDelegate::CreateSP(this, &FAnimationViewportClient::HandlePreviewScenePreTick));
+		AnimationEditorPreviewScene->RegisterOnPostTick(FSimpleDelegate::CreateSP(this, &FAnimationViewportClient::HandlePreviewScenePostTick));
+	}
+
+	// Register delegate to update the show flags when the post processing is turned on or off
+	UAssetViewerSettings::Get()->OnAssetViewerSettingsChanged().AddSP(this, &FAnimationViewportClient::OnAssetViewerSettingsChanged);
+	// Set correct flags according to current profile settings
+	SetAdvancedShowFlagsForScene(UAssetViewerSettings::Get()->Profiles[GetMutableDefault<UEditorPerProjectUserSettings>()->AssetViewerProfileIndex].bPostProcessingEnabled);
+
+}
+
 void FAnimationViewportClient::OnToggleAutoAlignFloor()
 {
 	bAutoAlignFloor = !bAutoAlignFloor;
@@ -454,11 +460,7 @@ void FAnimationViewportClient::HandleSkeletalMeshChanged(USkeletalMesh* OldSkele
 
 	if (NewSkeletalMesh)
 	{
-		OnMeshChangedDelegateHandle = NewSkeletalMesh->GetOnMeshChanged().AddLambda([this]()
-		{
-			UpdateCameraSetup();
-			Invalidate();
-		});
+		OnMeshChangedDelegateHandle = NewSkeletalMesh->GetOnMeshChanged().AddSP(this, &FAnimationViewportClient::HandleOnMeshChanged);
 	}
 
 	if (OldSkeletalMesh != NewSkeletalMesh || NewSkeletalMesh == nullptr)
@@ -486,14 +488,7 @@ void FAnimationViewportClient::HandleSkeletalMeshChanged(USkeletalMesh* OldSkele
 		}
 		// we need to make sure we monitor any change to the PhysicsState being recreated, as this can happen from path that is external to this class
 		// (example: setting a property on a body that is type "simulated" will recreate the state from USkeletalBodySetup::PostEditChangeProperty and let the body simulating (UE-107308)
-		OnPhysicsCreatedDelegateHandle = PreviewMeshComponent->RegisterOnPhysicsCreatedDelegate(FOnSkelMeshPhysicsCreated::CreateLambda([this]()
-			{
-				UDebugSkelMeshComponent* PreviewMeshComponent = GetAnimPreviewScene()->GetPreviewMeshComponent();
-				// let's make sure nothing is simulating and that all necessary state are in proper order
-				PreviewMeshComponent->SetPhysicsBlendWeight(0.f);
-				PreviewMeshComponent->SetSimulatePhysics(false);
-				DisableAllBodiesSimulatePhysics(PreviewMeshComponent);
-			}));
+		OnPhysicsCreatedDelegateHandle = PreviewMeshComponent->RegisterOnPhysicsCreatedDelegate(FOnSkelMeshPhysicsCreated::CreateSP(this, &FAnimationViewportClient::HandleOnSkelMeshPhysicsCreated));
 
 		PreviewMeshComponent->TermArticulated();
 		PreviewMeshComponent->InitArticulated(GetWorld()->GetPhysicsScene());
@@ -508,6 +503,21 @@ void FAnimationViewportClient::HandleSkeletalMeshChanged(USkeletalMesh* OldSkele
 	Invalidate();
 }
 
+void FAnimationViewportClient::HandleOnMeshChanged()
+{
+	UpdateCameraSetup();
+	Invalidate();
+}
+
+void FAnimationViewportClient::HandleOnSkelMeshPhysicsCreated()
+{
+	UDebugSkelMeshComponent* PreviewMeshComponent = GetAnimPreviewScene()->GetPreviewMeshComponent();
+	// let's make sure nothing is simulating and that all necessary state are in proper order
+	PreviewMeshComponent->SetPhysicsBlendWeight(0.f);
+	PreviewMeshComponent->SetSimulatePhysics(false);
+	DisableAllBodiesSimulatePhysics(PreviewMeshComponent);
+}
+
 void FAnimationViewportClient::Draw(const FSceneView* View, FPrimitiveDrawInterface* PDI)
 {
 	FEditorViewportClient::Draw(View, PDI);
@@ -2266,6 +2276,11 @@ EAnimationPlaybackSpeeds::Type FAnimationViewportClient::GetPlaybackSpeedMode()
 	return AnimationPlaybackSpeedMode;
 }
 
+TSharedPtr<class FAnimationEditorPreviewScene> FAnimationViewportClient::GetAnimPreviewScenePtr() const
+{
+	return StaticCastSharedPtr<FAnimationEditorPreviewScene>(PreviewScenePtr);
+}
+
 TSharedRef<FAnimationEditorPreviewScene> FAnimationViewportClient::GetAnimPreviewScene() const
 {
 	return StaticCastSharedRef<FAnimationEditorPreviewScene>(GetPreviewScene());
diff --git a/Engine/Source/Editor/Persona/Private/AnimationEditorViewportClient.h b/Engine/Source/Editor/Persona/Private/AnimationEditorViewportClient.h
--- a/Engine/Source/Editor/Persona/Private/AnimationEditorViewportClient.h
+++ b/Engine/Source/Editor/Persona/Private/AnimationEditorViewportClient.h
@@ -73,7 +73,7 @@ namespace EAnimationPlaybackSpeeds
 /////////////////////////////////////////////////////////////////////////
 // FAnimationViewportClient
 
-class PERSONA_API FAnimationViewportClient : public FEditorViewportClient
+class PERSONA_API FAnimationViewportClient : public FEditorViewportClient, public TSharedFromThis<FAnimationViewportClient>
 {
 protected:
 
@@ -90,6 +90,8 @@ public:
 	FAnimationViewportClient(const TSharedRef<class IPersonaPreviewScene>& InPreviewScene, const TSharedRef<class SAnimationEditorViewport>& InAnimationEditorViewport, const TSharedRef<class FAssetEditorToolkit>& InAssetEditorToolkit, int32 InViewportIndex, bool bInShowStats);
 	virtual ~FAnimationViewportClient();
 
+	void Initialize();
+
 	// FEditorViewportClient interface
 	virtual void Tick(float DeltaSeconds) override;
 	virtual void Draw(const FSceneView* View,FPrimitiveDrawInterface* PDI) override;
@@ -156,6 +158,12 @@ public:
 	/** Handle the skeletal mesh mesh component being used for preview changing */
 	void HandleSkeletalMeshChanged(class USkeletalMesh* OldSkeletalMesh, class USkeletalMesh* NewSkeletalMesh);
 
+	/** Handle a change in the skeletal mesh mesh component being used for preview changing */
+	void HandleOnMeshChanged();
+
+	/** Handle a change in the skeletal mesh phusics component being used for preview changing */
+	void HandleOnSkelMeshPhysicsCreated();
+
 	/** Function to display bone names*/
 	void ShowBoneNames(FViewport* Viewport, FCanvas* Canvas, UDebugSkelMeshComponent* MeshComponent);
 
@@ -430,6 +438,9 @@ private:
 	/** Draws bones from watched poses*/
 	void DrawWatchedPoses(UDebugSkelMeshComponent * MeshComponent, FPrimitiveDrawInterface* PDI);
 
+	/** Get the typed anim preview scene shared ptr*/
+	TSharedPtr<class FAnimationEditorPreviewScene> GetAnimPreviewScenePtr() const;
+
 	/** Get the typed anim preview scene */
 	TSharedRef<class FAnimationEditorPreviewScene> GetAnimPreviewScene() const;
 
diff --git a/Engine/Source/Editor/Persona/Private/SAnimationEditorViewport.cpp b/Engine/Source/Editor/Persona/Private/SAnimationEditorViewport.cpp
--- a/Engine/Source/Editor/Persona/Private/SAnimationEditorViewport.cpp
+++ b/Engine/Source/Editor/Persona/Private/SAnimationEditorViewport.cpp
@@ -123,6 +123,9 @@ TSharedRef<FEditorViewportClient> SAnimationEditorViewport::MakeEditorViewportCl
 	// Create an animation viewport client
 	LevelViewportClient = MakeShareable(new FAnimationViewportClient(PreviewScenePtr.Pin().ToSharedRef(), SharedThis(this), AssetEditorToolkitPtr.Pin().ToSharedRef(), ViewportIndex, bShowStats));
 
+	// Done after constructor, as the delegates require the shared pointer to be assigned
+	LevelViewportClient->Initialize();
+
 	LevelViewportClient->ViewportType = LVT_Perspective;
 	LevelViewportClient->bSetListenerPosition = false;
 	LevelViewportClient->SetViewLocation(EditorViewportDefs::DefaultPerspectiveViewLocation);
