Fix re-import of skeletalmesh, the PreferredReimportPath should not be test if empty.
Add reimport callback to refresh the UI when we reimport a skeletalmesh

#rnx

diff --git a/Engine/Source/Editor/Persona/Private/PersonaMeshDetails.cpp b/Engine/Source/Editor/Persona/Private/PersonaMeshDetails.cpp
--- a/Engine/Source/Editor/Persona/Private/PersonaMeshDetails.cpp
+++ b/Engine/Source/Editor/Persona/Private/PersonaMeshDetails.cpp
@@ -2118,6 +2118,7 @@ FPersonaMeshDetails::FPersonaMeshDetails(TSharedRef<class IPersonaToolkit> InPer
 	bDeleteWarningConsumed = false;
 
 	GEditor->GetEditorSubsystem<UImportSubsystem>()->OnAssetPostLODImport.AddRaw(this, &FPersonaMeshDetails::OnAssetPostLODImported);
+	GEditor->GetEditorSubsystem<UImportSubsystem>()->OnAssetReimport.AddRaw(this, &FPersonaMeshDetails::OnAssetReimport);
 }
 
 /**
@@ -2132,6 +2133,7 @@ FPersonaMeshDetails::~FPersonaMeshDetails()
 	}
 
 	GEditor->GetEditorSubsystem<UImportSubsystem>()->OnAssetPostLODImport.RemoveAll(this);
+	GEditor->GetEditorSubsystem<UImportSubsystem>()->OnAssetReimport.RemoveAll(this);
 }
 
 TSharedRef<IDetailCustomization> FPersonaMeshDetails::MakeInstance(TWeakPtr<class IPersonaToolkit> InPersonaToolkit)
@@ -3613,6 +3615,14 @@ void FPersonaMeshDetails::OnAssetPostLODImported(UObject* InObject, int32 InLODI
 	}
 }
 
+void FPersonaMeshDetails::OnAssetReimport(UObject* InObject)
+{
+	if (InObject == GetPersonaToolkit()->GetMesh())
+	{
+		RefreshMeshDetailLayout();
+	}
+}
+
 void FPersonaMeshDetails::OnImportLOD(TSharedPtr<FString> NewValue, ESelectInfo::Type SelectInfo, IDetailLayoutBuilder* DetailLayout)
 {
 	int32 LODIndex = 0;
diff --git a/Engine/Source/Editor/Persona/Private/PersonaMeshDetails.h b/Engine/Source/Editor/Persona/Private/PersonaMeshDetails.h
--- a/Engine/Source/Editor/Persona/Private/PersonaMeshDetails.h
+++ b/Engine/Source/Editor/Persona/Private/PersonaMeshDetails.h
@@ -592,6 +592,7 @@ private:
 
 	/** Called when a LOD is imported. Refreshes the UI. */
 	void OnAssetPostLODImported(UObject* InObject, int32 InLODIndex);
+	void OnAssetReimport(UObject* InObject);
 	/** Called from the PersonalMeshDetails UI to import a LOD. */
 	void OnImportLOD(TSharedPtr<FString> NewValue, ESelectInfo::Type SelectInfo, IDetailLayoutBuilder* DetailLayout);
 	void UpdateLODNames();
diff --git a/Engine/Source/Editor/UnrealEd/Private/Factories/EditorFactories.cpp b/Engine/Source/Editor/UnrealEd/Private/Factories/EditorFactories.cpp
--- a/Engine/Source/Editor/UnrealEd/Private/Factories/EditorFactories.cpp
+++ b/Engine/Source/Editor/UnrealEd/Private/Factories/EditorFactories.cpp
@@ -6330,7 +6330,7 @@ bool UReimportFbxSkeletalMeshFactory::CanReimport( UObject* Obj, TArray<FString>
 				//Fbx factory only support fbx and obj files
 				return false;
 			}
-			else if (!FactoryExtensions.Contains(FPaths::GetExtension(this->PreferredReimportPath)))
+			else if (!PreferredReimportPath.IsEmpty() && !FactoryExtensions.Contains(FPaths::GetExtension(PreferredReimportPath)))
 			{
 				// Unsupported extensions shouldnt be considered for reimport
 				return false;
