Slate: Initial phase of FVector2D deprecation

diff --git a/Engine/Source/Developer/ToolWidgets/Private/SSimpleTimeSlider.cpp b/Engine/Source/Developer/ToolWidgets/Private/SSimpleTimeSlider.cpp
--- a/Engine/Source/Developer/ToolWidgets/Private/SSimpleTimeSlider.cpp
+++ b/Engine/Source/Developer/ToolWidgets/Private/SSimpleTimeSlider.cpp
@@ -132,11 +132,11 @@ void SSimpleTimeSlider::DrawTicks( FSlateWindowElementList& OutDrawElements, con
 
 		if ( AbsOffsetNum % Divider == 0 )
 		{
-			FVector2D Offset( XPos, InArgs.TickOffset );
-			FVector2D TickSize( 1.0f, InArgs.MajorTickHeight );
+			FVector2f Offset( XPos, InArgs.TickOffset );
+			FVector2f TickSize( 1.0f, InArgs.MajorTickHeight );
 
 			LinePoints[0] = FVector2D(1.0f,1.0f);
-			LinePoints[1] = TickSize;
+			LinePoints[1] = FVector2D(TickSize);
 
 			// lines should not need anti-aliasing
 			const bool bAntiAliasLines = false;
@@ -145,7 +145,7 @@ void SSimpleTimeSlider::DrawTicks( FSlateWindowElementList& OutDrawElements, con
 			FSlateDrawElement::MakeLines(
 				OutDrawElements,
 				InArgs.StartLayer,
-				InArgs.AllottedGeometry.ToPaintGeometry( Offset, TickSize ),
+				InArgs.AllottedGeometry.ToPaintGeometry( TickSize, FSlateLayoutTransform(Offset) ),
 				LinePoints,
 				InArgs.DrawEffects,
 				InArgs.TickColor,
@@ -158,13 +158,13 @@ void SSimpleTimeSlider::DrawTicks( FSlateWindowElementList& OutDrawElements, con
 
 				// Space the text between the tick mark but slightly above
 				const TSharedRef< FSlateFontMeasure > FontMeasureService = FSlateApplication::Get().GetRenderer()->GetFontMeasureService();
-				FVector2D TextSize = FontMeasureService->Measure(FrameString, SmallLayoutFont);
-				FVector2D TextOffset( XPos-(TextSize.X*0.5f), InArgs.bMirrorLabels ? TextSize.Y :  FMath::Abs( InArgs.AllottedGeometry.GetLocalSize().Y - (InArgs.MajorTickHeight+TextSize.Y) ) );
+				FVector2f TextSize = UE::Slate::CastToVector2f(FontMeasureService->Measure(FrameString, SmallLayoutFont));
+				FVector2f TextOffset( XPos-(TextSize.X*0.5f), InArgs.bMirrorLabels ? TextSize.Y :  FMath::Abs( InArgs.AllottedGeometry.GetLocalSize().Y - (InArgs.MajorTickHeight+TextSize.Y) ) );
 
 				FSlateDrawElement::MakeText(
 					OutDrawElements,
 					InArgs.StartLayer+1, 
-					InArgs.AllottedGeometry.ToPaintGeometry( TextOffset, TextSize ), 
+					InArgs.AllottedGeometry.ToPaintGeometry( TextSize, FSlateLayoutTransform(TextOffset) ), 
 					FrameString, 
 					SmallLayoutFont, 
 					InArgs.DrawEffects,
@@ -177,18 +177,18 @@ void SSimpleTimeSlider::DrawTicks( FSlateWindowElementList& OutDrawElements, con
 			// Compute the size of each tick mark.  If we are half way between to visible values display a slightly larger tick mark
 			const float MinorTickHeight = AbsOffsetNum % HalfDivider == 0 ? 7.0f : 4.0f;
 
-			FVector2D Offset(XPos, InArgs.bMirrorLabels ? 0.0f : FMath::Abs( InArgs.AllottedGeometry.GetLocalSize().Y - MinorTickHeight ) );
-			FVector2D TickSize(1, MinorTickHeight);
+			FVector2f Offset(XPos, InArgs.bMirrorLabels ? 0.0f : FMath::Abs( InArgs.AllottedGeometry.GetLocalSize().Y - MinorTickHeight ) );
+			FVector2f TickSize(1.f, MinorTickHeight);
 
 			LinePoints[0] = FVector2D(1.0f,1.0f);
-			LinePoints[1] = TickSize;
+			LinePoints[1] = FVector2D(TickSize);
 
 			const bool bAntiAlias = false;
 			// Draw each sub mark
 			FSlateDrawElement::MakeLines(
 				OutDrawElements,
 				InArgs.StartLayer,
-				InArgs.AllottedGeometry.ToPaintGeometry( Offset, TickSize ),
+				InArgs.AllottedGeometry.ToPaintGeometry( TickSize, FSlateLayoutTransform(Offset) ),
 				LinePoints,
 				InArgs.DrawEffects,
 				InArgs.TickColor,
@@ -224,11 +224,11 @@ int32 SSimpleTimeSlider::OnPaintTimeSlider( bool bMirrorLabels, const FGeometry&
 		FPaintGeometry RangeGeometry;
 		if (bMirrorLabels)
 		{
-			RangeGeometry = AllottedGeometry.ToPaintGeometry(FVector2D(LeftClamp, 0), FVector2D(RightClamp-LeftClamp, Height));
+			RangeGeometry = AllottedGeometry.ToPaintGeometry(FVector2f(RightClamp-LeftClamp, Height), FSlateLayoutTransform(FVector2f(LeftClamp, 0)));
 		}
 		else
 		{
-			RangeGeometry = AllottedGeometry.ToPaintGeometry(FVector2D(LeftClamp, AllottedGeometry.GetLocalSize().Y - Height), FVector2D(RightClamp-LeftClamp, AllottedGeometry.GetLocalSize().Y ));
+			RangeGeometry = AllottedGeometry.ToPaintGeometry(FVector2f(RightClamp-LeftClamp, AllottedGeometry.GetLocalSize().Y ), FSlateLayoutTransform(FVector2f(LeftClamp, AllottedGeometry.GetLocalSize().Y - Height)));
 		}
 
 		FSlateDrawElement::MakeBox(
@@ -265,7 +265,7 @@ int32 SSimpleTimeSlider::OnPaintTimeSlider( bool bMirrorLabels, const FGeometry&
 		const float CursorHalfSize = CursorSize.Get() * 0.5f;
 		const int32 CursorLayer = LayerId + 2;
 		const float CursorHalfLength = AllottedGeometry.GetLocalSize().X * CursorHalfSize;
-		FPaintGeometry CursorGeometry = AllottedGeometry.ToPaintGeometry(FVector2D(XPos - CursorHalfLength, 0), FVector2D(2 * CursorHalfLength, AllottedGeometry.GetLocalSize().Y));
+		FPaintGeometry CursorGeometry = AllottedGeometry.ToPaintGeometry(FVector2f(2 * CursorHalfLength, AllottedGeometry.GetLocalSize().Y), FSlateLayoutTransform(FVector2f(XPos - CursorHalfLength, 0)));
 
 		FLinearColor CursorColor = InWidgetStyle.GetColorAndOpacityTint();
 		CursorColor.A = CursorColor.A*0.08f;
@@ -284,7 +284,7 @@ int32 SSimpleTimeSlider::OnPaintTimeSlider( bool bMirrorLabels, const FGeometry&
 
 		// Should draw above the text
 		const int32 ArrowLayer = LayerId + 3;
-		FPaintGeometry MyGeometry =	AllottedGeometry.ToPaintGeometry( FVector2D( XPos-HalfSize, 0 ), FVector2D( HandleSize, AllottedGeometry.GetLocalSize().Y ) );
+		FPaintGeometry MyGeometry =	AllottedGeometry.ToPaintGeometry( FVector2f( HandleSize, AllottedGeometry.GetLocalSize().Y ), FSlateLayoutTransform(FVector2f( XPos-HalfSize, 0 )) );
 		FLinearColor ScrubColor = InWidgetStyle.GetColorAndOpacityTint();
 
 		// @todo Sequencer this color should be specified in the style
diff --git a/Engine/Source/Developer/ToolWidgets/Private/SWidgetDrawer.cpp b/Engine/Source/Developer/ToolWidgets/Private/SWidgetDrawer.cpp
--- a/Engine/Source/Developer/ToolWidgets/Private/SWidgetDrawer.cpp
+++ b/Engine/Source/Developer/ToolWidgets/Private/SWidgetDrawer.cpp
@@ -218,7 +218,7 @@ class SDrawerOverlay : public SCompoundWidget
 		FSlateDrawElement::MakeBox(
 			OutDrawElements,
 			LayerId,
-			RenderTransformedChildGeometry.ToPaintGeometry(ShadowOffset, FVector2D(AllottedGeometry.GetLocalSize().X - (ShadowOffset.X * 2), TargetHeight)),
+			RenderTransformedChildGeometry.ToPaintGeometry( FVector2f(AllottedGeometry.GetLocalSize().X - (ShadowOffset.X * 2), TargetHeight), FSlateLayoutTransform(ShadowOffset) ),
 			BackgroundBrush,
 			ESlateDrawEffect::None,
 			BackgroundBrush->GetTint(InWidgetStyle));
@@ -229,7 +229,7 @@ class SDrawerOverlay : public SCompoundWidget
 		FSlateDrawElement::MakeBox(
 			OutDrawElements,
 			OutLayerId,
-			AllottedGeometry.ToPaintGeometry(FVector2D(0.0f, AllottedGeometry.GetLocalSize().Y - ShadowOffset.Y), FVector2D(AllottedGeometry.GetLocalSize().X, ShadowOffset.Y)),
+			AllottedGeometry.ToPaintGeometry(FVector2f(AllottedGeometry.GetLocalSize().X, ShadowOffset.Y), FSlateLayoutTransform(FVector2f(0.0f, AllottedGeometry.GetLocalSize().Y - ShadowOffset.Y))),
 			ShadowBrush,
 			ESlateDrawEffect::None,
 			ShadowBrush->GetTint(InWidgetStyle));
@@ -238,7 +238,7 @@ class SDrawerOverlay : public SCompoundWidget
 		FSlateDrawElement::MakeBox(
 			OutDrawElements,
 			OutLayerId+1,
-			RenderTransformedChildGeometry.ToPaintGeometry(ShadowOffset, FVector2D(AllottedGeometry.GetLocalSize().X - (ShadowOffset.X * 2), TargetHeight)),
+			RenderTransformedChildGeometry.ToPaintGeometry(FVector2f(AllottedGeometry.GetLocalSize().X - (ShadowOffset.X * 2), TargetHeight), FSlateLayoutTransform(ShadowOffset)),
 			BorderBrush,
 			ESlateDrawEffect::None,
 			BorderBrush->GetTint(InWidgetStyle));
@@ -279,7 +279,10 @@ private:
 
 	FGeometry GetResizeHandleGeometry(const FGeometry& AllottedGeometry) const
 	{
-		return GetRenderTransformedGeometry(AllottedGeometry).MakeChild(ShadowOffset - FVector2D(0.0f, ExpanderSize), FVector2D(AllottedGeometry.GetLocalSize().X-ShadowOffset.X*2, ExpanderSize));
+		return GetRenderTransformedGeometry(AllottedGeometry).MakeChild(
+			FVector2D(AllottedGeometry.GetLocalSize().X-ShadowOffset.X*2, ExpanderSize),
+			FSlateLayoutTransform(ShadowOffset - FVector2D(0.0f, ExpanderSize))
+		);
 	}
 
 	void SetHeight(float NewHeight)
diff --git a/Engine/Source/Developer/TreeMap/STreeMap.cpp b/Engine/Source/Developer/TreeMap/STreeMap.cpp
--- a/Engine/Source/Developer/TreeMap/STreeMap.cpp
+++ b/Engine/Source/Developer/TreeMap/STreeMap.cpp
@@ -261,7 +261,7 @@ int32 STreeMap::OnPaint( const FPaintArgs& Args, const FGeometry& AllottedGeomet
 					// Draw the visual's background box
 					const FVector2D VisualPosition = BlendedVisual.Position;
 					const FSlateRect VisualClippingRect = TransformRect( AllottedGeometry.GetAccumulatedLayoutTransform(), FSlateRect( VisualPosition, VisualPosition + BlendedVisual.Size ) );
-					const auto VisualPaintGeometry = AllottedGeometry.ToPaintGeometry( VisualPosition, BlendedVisual.Size );
+					const auto VisualPaintGeometry = AllottedGeometry.ToPaintGeometry( BlendedVisual.Size, FSlateLayoutTransform(VisualPosition) );
 					auto DrawColor = InWidgetStyle.GetColorAndOpacityTint() * ThisNodeBackground->TintColor.GetColor( InWidgetStyle ) * BlendedVisual.Color;
 
 					OutDrawElements.PushClip(FSlateClippingZone(VisualClippingRect));
@@ -291,7 +291,7 @@ int32 STreeMap::OnPaint( const FPaintArgs& Args, const FGeometry& AllottedGeomet
 							const float SizeDifference = LargestSize - BlendedVisual.Size.X;
 							BackgroundPosition.X -= SizeDifference * 0.5f;
 						}
-						const auto BackgroundPaintGeometry = AllottedGeometry.ToPaintGeometry( BackgroundPosition, BackgroundSize );
+						const auto BackgroundPaintGeometry = AllottedGeometry.ToPaintGeometry( BackgroundSize, FSlateLayoutTransform(BackgroundPosition) );
 
 						const FSlateRect BackgroundClippingRect = VisualClippingRect.InsetBy( FMargin( 1 ) );
 
@@ -379,7 +379,7 @@ int32 STreeMap::OnPaint( const FPaintArgs& Args, const FGeometry& AllottedGeomet
 						const FVector2D ScreenSpaceVisualSize( AllottedGeometry.Scale * BlendedVisual.Size );
 						if( ScreenSpaceVisualSize.X > 20 )
 						{
-							FGeometry VisualGeometry = AllottedGeometry.MakeChild(VisualPosition, BlendedVisual.Size);
+							FGeometry VisualGeometry = AllottedGeometry.MakeChild(BlendedVisual.Size, FSlateLayoutTransform(VisualPosition));
 							const FPaintGeometry VisualPaintGeometry = VisualGeometry.ToPaintGeometry();
 
 							
@@ -491,7 +491,7 @@ int32 STreeMap::OnPaint( const FPaintArgs& Args, const FGeometry& AllottedGeomet
 			FSlateDrawElement::MakeSpline(
 				OutDrawElements,
 				LayerId,
-				AllottedGeometry.ToPaintGeometry(  ShadowOffset, FVector2D( 1.0, 1.0f ) ),
+				AllottedGeometry.ToPaintGeometry( FVector2f( 1.0, 1.0f ), FSlateLayoutTransform(ShadowOffset) ),
 				SplineStart,
 				SplineStartDir,
 				SplineEnd,
diff --git a/Engine/Source/Editor/CurveEditor/Private/DragOperations/CurveEditorDragOperation_Marquee.cpp b/Engine/Source/Editor/CurveEditor/Private/DragOperations/CurveEditorDragOperation_Marquee.cpp
--- a/Engine/Source/Editor/CurveEditor/Private/DragOperations/CurveEditorDragOperation_Marquee.cpp
+++ b/Engine/Source/Editor/CurveEditor/Private/DragOperations/CurveEditorDragOperation_Marquee.cpp
@@ -183,7 +183,7 @@ void FCurveEditorDragOperation_Marquee::OnPaint(const FGeometry& AllottedGeometr
 	FSlateDrawElement::MakeBox(
 		OutDrawElements,
 		PaintOnLayerId,
-		AllottedGeometry.ToPaintGeometry(Marquee.GetTopLeft(), Marquee.GetBottomRight() - Marquee.GetTopLeft()),
+		AllottedGeometry.ToPaintGeometry(Marquee.GetBottomRight() - Marquee.GetTopLeft(), FSlateLayoutTransform(Marquee.GetTopLeft())),
 		FAppStyle::GetBrush(TEXT("MarqueeSelection"))
 		);
 }
diff --git a/Engine/Source/Editor/CurveEditor/Private/SCurveViewerPanel.cpp b/Engine/Source/Editor/CurveEditor/Private/SCurveViewerPanel.cpp
--- a/Engine/Source/Editor/CurveEditor/Private/SCurveViewerPanel.cpp
+++ b/Engine/Source/Editor/CurveEditor/Private/SCurveViewerPanel.cpp
@@ -97,8 +97,8 @@ void SCurveViewerPanel::DrawCurves(const FGeometry& AllottedGeometry, FSlateWind
 				PointTint = HSV.HSVToLinearRGB();
 
 				FPaintGeometry PointGeometry = AllottedGeometry.ToPaintGeometry(
-					Point.ScreenPosition - (PointDrawInfo.ScreenSize * 0.5f),
-					PointDrawInfo.ScreenSize
+					PointDrawInfo.ScreenSize,
+					FSlateLayoutTransform(Point.ScreenPosition - (PointDrawInfo.ScreenSize * 0.5f))
 				);
 
 				FSlateDrawElement::MakeBox(OutDrawElements, KeyLayerId, PointGeometry, PointDrawInfo.Brush, DrawEffects, PointTint );
diff --git a/Engine/Source/Editor/CurveEditor/Private/Views/SInteractiveCurveEditorView.cpp b/Engine/Source/Editor/CurveEditor/Private/Views/SInteractiveCurveEditorView.cpp
--- a/Engine/Source/Editor/CurveEditor/Private/Views/SInteractiveCurveEditorView.cpp
+++ b/Engine/Source/Editor/CurveEditor/Private/Views/SInteractiveCurveEditorView.cpp
@@ -504,8 +504,8 @@ void SInteractiveCurveEditorView::DrawCurves(TSharedRef<FCurveEditor> CurveEdito
 				}
 
 				FPaintGeometry PointGeometry = AllottedGeometry.ToPaintGeometry(
-					Point.ScreenPosition - (PointDrawInfo.ScreenSize * 0.5f),
-					PointDrawInfo.ScreenSize
+					PointDrawInfo.ScreenSize,
+					FSlateLayoutTransform(Point.ScreenPosition - (PointDrawInfo.ScreenSize * 0.5f))
 				);
 
 				FSlateDrawElement::MakeBox(OutDrawElements, KeyLayerId, PointGeometry, PointDrawInfo.Brush, DrawEffects, PointTint);
diff --git a/Engine/Source/Editor/EditorWidgets/Private/SDropTarget.cpp b/Engine/Source/Editor/EditorWidgets/Private/SDropTarget.cpp
--- a/Engine/Source/Editor/EditorWidgets/Private/SDropTarget.cpp
+++ b/Engine/Source/Editor/EditorWidgets/Private/SDropTarget.cpp
@@ -180,7 +180,7 @@ int32 SDropTarget::OnPaint(const FPaintArgs& Args, const FGeometry& AllottedGeom
 			FSlateDrawElement::MakeBox(
 				OutDrawElements,
 				DashLayer,
-				AllottedGeometry.ToPaintGeometry(FVector2D(Inset, 0), FVector2D(AllottedGeometry.GetLocalSize().X-Inset*2, HorizontalImage->ImageSize.Y)),
+				AllottedGeometry.ToPaintGeometry(FVector2f(AllottedGeometry.GetLocalSize().X-Inset*2, HorizontalImage->ImageSize.Y), FSlateLayoutTransform(FVector2f(Inset, 0))),
 				HorizontalImage,
 				ESlateDrawEffect::None,
 				DashColor.GetColor(InWidgetStyle));
@@ -189,7 +189,7 @@ int32 SDropTarget::OnPaint(const FPaintArgs& Args, const FGeometry& AllottedGeom
 			FSlateDrawElement::MakeBox(
 				OutDrawElements,
 				DashLayer,
-				AllottedGeometry.ToPaintGeometry(FVector2D(Inset, AllottedGeometry.GetLocalSize().Y - HorizontalImage->ImageSize.Y), FVector2D(AllottedGeometry.Size.X-Inset * 2, HorizontalImage->ImageSize.Y)),
+				AllottedGeometry.ToPaintGeometry(FVector2f(AllottedGeometry.Size.X-Inset * 2, HorizontalImage->ImageSize.Y), FSlateLayoutTransform(FVector2f(Inset, AllottedGeometry.GetLocalSize().Y - HorizontalImage->ImageSize.Y))),
 				HorizontalImage,
 				ESlateDrawEffect::None,
 				DashColor.GetColor(InWidgetStyle));
@@ -198,7 +198,7 @@ int32 SDropTarget::OnPaint(const FPaintArgs& Args, const FGeometry& AllottedGeom
 			FSlateDrawElement::MakeBox(
 				OutDrawElements,
 				DashLayer,
-				AllottedGeometry.ToPaintGeometry(FVector2D(0, Inset), FVector2D(VerticalImage->ImageSize.X, AllottedGeometry.GetLocalSize().Y-Inset * 2)),
+				AllottedGeometry.ToPaintGeometry(FVector2f(VerticalImage->ImageSize.X, AllottedGeometry.GetLocalSize().Y-Inset * 2), FSlateLayoutTransform(FVector2f(0, Inset))),
 				VerticalImage,
 				ESlateDrawEffect::None,
 				DashColor.GetColor(InWidgetStyle));
@@ -207,7 +207,7 @@ int32 SDropTarget::OnPaint(const FPaintArgs& Args, const FGeometry& AllottedGeom
 			FSlateDrawElement::MakeBox(
 				OutDrawElements,
 				DashLayer,
-				AllottedGeometry.ToPaintGeometry(FVector2D(AllottedGeometry.GetLocalSize().X - VerticalImage->ImageSize.X, Inset), FVector2D(VerticalImage->ImageSize.X, AllottedGeometry.GetLocalSize().Y-Inset * 2)),
+				AllottedGeometry.ToPaintGeometry(FVector2f(VerticalImage->ImageSize.X, AllottedGeometry.GetLocalSize().Y-Inset * 2), FSlateLayoutTransform(FVector2f(AllottedGeometry.GetLocalSize().X - VerticalImage->ImageSize.X, Inset))),
 				VerticalImage,
 				ESlateDrawEffect::None,
 				DashColor.GetColor(InWidgetStyle));
diff --git a/Engine/Source/Editor/EditorWidgets/Private/SObjectNameEditableTextBox.cpp b/Engine/Source/Editor/EditorWidgets/Private/SObjectNameEditableTextBox.cpp
--- a/Engine/Source/Editor/EditorWidgets/Private/SObjectNameEditableTextBox.cpp
+++ b/Engine/Source/Editor/EditorWidgets/Private/SObjectNameEditableTextBox.cpp
@@ -137,7 +137,7 @@ int32 SObjectNameEditableTextBox::OnPaint( const FPaintArgs& Args, const FGeomet
 		FSlateDrawElement::MakeBox(
 			OutDrawElements,
 			LayerId + TextLayer,
-			AllottedGeometry.ToPaintGeometry( DrawPosition, DrawSize ),	// Position, Size, Scale
+			AllottedGeometry.ToPaintGeometry( DrawSize, FSlateLayoutTransform(DrawPosition) ),	// Position, Size, Scale
 			StyleInfo,													// Style
 			DrawEffects,												// Effects to use
 			HighlightTargetColorAndOpacity );							// Color
diff --git a/Engine/Source/Editor/GraphEditor/Private/SGraphPanel.cpp b/Engine/Source/Editor/GraphEditor/Private/SGraphPanel.cpp
--- a/Engine/Source/Editor/GraphEditor/Private/SGraphPanel.cpp
+++ b/Engine/Source/Editor/GraphEditor/Private/SGraphPanel.cpp
@@ -415,7 +415,7 @@ int32 SGraphPanel::OnPaint( const FPaintArgs& Args, const FGeometry& AllottedGeo
 							const FSlateBrush* OverlayBrush = OverlayInfo.Brush;
 							if (OverlayBrush != nullptr)
 							{
-								FPaintGeometry BouncedGeometry = CurWidget.Geometry.ToPaintGeometry(OverlayInfo.OverlayOffset, OverlayBrush->ImageSize, 1.f);
+								FPaintGeometry BouncedGeometry = CurWidget.Geometry.ToPaintGeometry(OverlayBrush->ImageSize, FSlateLayoutTransform(OverlayInfo.OverlayOffset));
 
 								// Handle bouncing during PIE
 								const float BounceValue = FMath::Sin(2.0f * PI * BounceCurve.GetLerp());
@@ -452,7 +452,7 @@ int32 SGraphPanel::OnPaint( const FPaintArgs& Args, const FGeometry& AllottedGeo
 									// call SlatePrepass as these widgets are not in the 'normal' child hierarchy
 									Widget->SlatePrepass(AllottedGeometry.GetAccumulatedLayoutTransform().GetScale());
 
-									const FGeometry WidgetGeometry = CurWidget.Geometry.MakeChild(OverlayInfo.OverlayOffset, Widget->GetDesiredSize());
+									const FGeometry WidgetGeometry = CurWidget.Geometry.MakeChild(Widget->GetDesiredSize(), FSlateLayoutTransform(OverlayInfo.OverlayOffset));
 
 									Widget->Paint(NewArgs, WidgetGeometry, MyCullingRect, OutDrawElements, CurWidgetsMaxLayerId, InWidgetStyle, bParentEnabled);
 								}
diff --git a/Engine/Source/Editor/GraphEditor/Private/SNodePanel.cpp b/Engine/Source/Editor/GraphEditor/Private/SNodePanel.cpp
--- a/Engine/Source/Editor/GraphEditor/Private/SNodePanel.cpp
+++ b/Engine/Source/Editor/GraphEditor/Private/SNodePanel.cpp
@@ -1467,7 +1467,7 @@ void SNodePanel::PaintMarquee(const FGeometry& AllottedGeometry, const FSlateRec
 		FSlateDrawElement::MakeBox(
 			OutDrawElements,
 			DrawLayerId,
-			AllottedGeometry.ToPaintGeometry( GraphCoordToPanelCoord(Marquee.Rect.GetUpperLeft()), Marquee.Rect.GetSize()*GetZoomAmount() ),
+			AllottedGeometry.ToPaintGeometry( Marquee.Rect.GetSize()*GetZoomAmount(), FSlateLayoutTransform(GraphCoordToPanelCoord(Marquee.Rect.GetUpperLeft())) ),
 			FAppStyle::GetBrush(TEXT("MarqueeSelection"))
 		);
 	}
@@ -1486,7 +1486,7 @@ void SNodePanel::PaintSoftwareCursor(const FGeometry& AllottedGeometry, const FS
 	FSlateDrawElement::MakeBox(
 		OutDrawElements,
 		DrawLayerId,
-		AllottedGeometry.ToPaintGeometry( GraphCoordToPanelCoord( SoftwareCursorPosition ) - ( Brush->ImageSize / 2 ), Brush->ImageSize ),
+		AllottedGeometry.ToPaintGeometry( Brush->ImageSize, FSlateLayoutTransform(GraphCoordToPanelCoord( SoftwareCursorPosition ) - ( Brush->ImageSize / 2 )) ),
 		Brush
 	);
 }
@@ -1514,7 +1514,7 @@ void SNodePanel::PaintComment(const FString& CommentText, const FGeometry& Allot
 	FSlateDrawElement::MakeBox(
 		OutDrawElements,
 		DrawLayerId-1,
-		AllottedGeometry.ToPaintGeometry(CommentBubbleOffset, CommentTextSize),
+		AllottedGeometry.ToPaintGeometry(CommentTextSize, FSlateLayoutTransform(CommentBubbleOffset)),
 		CommentCalloutBubble,
 		ESlateDrawEffect::None,
 		CommentTinting
@@ -1523,7 +1523,7 @@ void SNodePanel::PaintComment(const FString& CommentText, const FGeometry& Allot
 	FSlateDrawElement::MakeBox(
 		OutDrawElements,
 		DrawLayerId-1,
-		AllottedGeometry.ToPaintGeometry( CommentBubbleArrowOffset, CommentCalloutArrow->ImageSize ),
+		AllottedGeometry.ToPaintGeometry( CommentCalloutArrow->ImageSize, FSlateLayoutTransform(CommentBubbleArrowOffset) ),
 		CommentCalloutArrow,
 		ESlateDrawEffect::None,
 		CommentTinting
@@ -1533,7 +1533,7 @@ void SNodePanel::PaintComment(const FString& CommentText, const FGeometry& Allot
 	FSlateDrawElement::MakeText(
 		OutDrawElements,
 		DrawLayerId,
-		AllottedGeometry.ToPaintGeometry( CommentBubbleOffset + CommentBubblePadding, CommentTextSize ),
+		AllottedGeometry.ToPaintGeometry( CommentTextSize, FSlateLayoutTransform(CommentBubbleOffset + CommentBubblePadding) ),
 		CommentText,
 		CommentFont,
 		ESlateDrawEffect::None,
diff --git a/Engine/Source/Editor/KismetWidgets/Private/SScrubWidget.cpp b/Engine/Source/Editor/KismetWidgets/Private/SScrubWidget.cpp
--- a/Engine/Source/Editor/KismetWidgets/Private/SScrubWidget.cpp
+++ b/Engine/Source/Editor/KismetWidgets/Private/SScrubWidget.cpp
@@ -137,7 +137,7 @@ int32 SScrubWidget::OnPaint( const FPaintArgs& Args, const FGeometry& AllottedGe
 				FSlateDrawElement::MakeBox(
 					OutDrawElements,
 					BackgroundLayer,
-					AllottedGeometry.ToPaintGeometry(Offset, Size),
+					AllottedGeometry.ToPaintGeometry(Size, FSlateLayoutTransform(Offset)),
 					StyleInfo,
 					DrawEffects,
 					InWidgetStyle.GetColorAndOpacityTint()
@@ -153,7 +153,7 @@ int32 SScrubWidget::OnPaint( const FPaintArgs& Args, const FGeometry& AllottedGe
 				FSlateDrawElement::MakeText(
 					OutDrawElements,
 					TextLayer, 
-					AllottedGeometry.ToPaintGeometry(TextOffset, TextSize), 
+					AllottedGeometry.ToPaintGeometry(TextSize, FSlateLayoutTransform(TextOffset)), 
 					FrameString, 
 					SmallLayoutFont, 
 					DrawEffects);
@@ -167,7 +167,7 @@ int32 SScrubWidget::OnPaint( const FPaintArgs& Args, const FGeometry& AllottedGe
  				FSlateDrawElement::MakeBox(
  					OutDrawElements,
  					BackgroundLayer,
-					AllottedGeometry.ToPaintGeometry(Offset, Size),
+					AllottedGeometry.ToPaintGeometry(Size, FSlateLayoutTransform(Offset)),
  					StyleInfo,
  					DrawEffects,
  					InWidgetStyle.GetColorAndOpacityTint()
@@ -185,7 +185,7 @@ int32 SScrubWidget::OnPaint( const FPaintArgs& Args, const FGeometry& AllottedGe
 			const double Height = AllottedGeometry.GetLocalSize().Y;
 			const FVector2D Offset( XPos - Height*0.25f, 0.f );
 
-			FPaintGeometry MyGeometry =	AllottedGeometry.ToPaintGeometry( Offset, FVector2D(Height * 0.5f, Height) );
+			FPaintGeometry MyGeometry =	AllottedGeometry.ToPaintGeometry( FVector2f(Height * 0.5f, Height), FSlateLayoutTransform(Offset) );
 			FLinearColor ScrubColor = InWidgetStyle.GetColorAndOpacityTint();
 			ScrubColor.A = ScrubColor.A*0.5f;
 			ScrubColor.B *= 0.1f;
@@ -216,7 +216,7 @@ int32 SScrubWidget::OnPaint( const FPaintArgs& Args, const FGeometry& AllottedGe
 				FSlateDrawElement::MakeBox(
 					OutDrawElements,
 					ArrowLayer+1,
-					AllottedGeometry.ToPaintGeometry(BarOffset, Size),
+					AllottedGeometry.ToPaintGeometry(Size, FSlateLayoutTransform(BarOffset)),
 					StyleInfo,
 					DrawEffects,
 					BarColor
diff --git a/Engine/Source/Editor/MovieSceneTools/Private/CommonMovieSceneTools.cpp b/Engine/Source/Editor/MovieSceneTools/Private/CommonMovieSceneTools.cpp
--- a/Engine/Source/Editor/MovieSceneTools/Private/CommonMovieSceneTools.cpp
+++ b/Engine/Source/Editor/MovieSceneTools/Private/CommonMovieSceneTools.cpp
@@ -55,7 +55,7 @@ void DrawFrameTimeHint(FSequencerSectionPainter& InPainter, const FFrameTime& Cu
 	FSlateDrawElement::MakeBox(
 		InPainter.DrawElements,
 		InPainter.LayerId + 5,
-		InPainter.SectionGeometry.ToPaintGeometry(TextOffset - BoxPadding, TextSize + 2.0f * BoxPadding),
+		InPainter.SectionGeometry.ToPaintGeometry(TextSize + 2.0f * BoxPadding, FSlateLayoutTransform(TextOffset - BoxPadding)),
 		FAppStyle::GetBrush("WhiteBrush"),
 		ESlateDrawEffect::None,
 		FLinearColor::Black.CopyWithNewOpacity(0.5f * InPainter.GhostAlpha)
@@ -66,7 +66,7 @@ void DrawFrameTimeHint(FSequencerSectionPainter& InPainter, const FFrameTime& Cu
 	FSlateDrawElement::MakeText(
 		InPainter.DrawElements,
 		InPainter.LayerId + 6,
-		InPainter.SectionGeometry.ToPaintGeometry(TextOffset, TextSize),
+		InPainter.SectionGeometry.ToPaintGeometry(TextSize, FSlateLayoutTransform(TextOffset)),
 		FrameTimeString,
 		SmallLayoutFont,
 		DrawEffects,
diff --git a/Engine/Source/Editor/MovieSceneTools/Private/Sections/BoolPropertySection.cpp b/Engine/Source/Editor/MovieSceneTools/Private/Sections/BoolPropertySection.cpp
--- a/Engine/Source/Editor/MovieSceneTools/Private/Sections/BoolPropertySection.cpp
+++ b/Engine/Source/Editor/MovieSceneTools/Private/Sections/BoolPropertySection.cpp
@@ -83,7 +83,7 @@ int32 FBoolPropertySection::OnPaintSection( FSequencerSectionPainter& Painter )
 		FSlateDrawElement::MakeBox(
 			Painter.DrawElements,
 			LayerId + 1,
-			Painter.SectionGeometry.ToPaintGeometry(StartPos, Size),
+			Painter.SectionGeometry.ToPaintGeometry(Size, FSlateLayoutTransform(StartPos)),
 			BoolOverlayBrush,
 			DrawEffects,
 			Color
diff --git a/Engine/Source/Editor/MovieSceneTools/Private/Sections/CameraCutSection.cpp b/Engine/Source/Editor/MovieSceneTools/Private/Sections/CameraCutSection.cpp
--- a/Engine/Source/Editor/MovieSceneTools/Private/Sections/CameraCutSection.cpp
+++ b/Engine/Source/Editor/MovieSceneTools/Private/Sections/CameraCutSection.cpp
@@ -191,7 +191,7 @@ int32 FCameraCutSection::OnPaintSection(FSequencerSectionPainter& InPainter) con
 		FSlateDrawElement::MakeBox(
 			InPainter.DrawElements,
 			InPainter.LayerId,
-			InPainter.SectionGeometry.ToPaintGeometry(FVector2D(1, 1), InPainter.SectionGeometry.GetLocalSize() - FVector2D(1, 1)),
+			InPainter.SectionGeometry.ToPaintGeometry(InPainter.SectionGeometry.GetLocalSize() - FVector2D(1.f, 1.f), FSlateLayoutTransform(FVector2D(1.f, 1.f))),
 			ErroredSectionOverlay,
 			DrawEffects,
 			ErrorColor.CopyWithNewOpacity(0.8f)
diff --git a/Engine/Source/Editor/MovieSceneTools/Private/Sections/ColorPropertySection.cpp b/Engine/Source/Editor/MovieSceneTools/Private/Sections/ColorPropertySection.cpp
--- a/Engine/Source/Editor/MovieSceneTools/Private/Sections/ColorPropertySection.cpp
+++ b/Engine/Source/Editor/MovieSceneTools/Private/Sections/ColorPropertySection.cpp
@@ -80,7 +80,7 @@ int32 FColorPropertySection::OnPaintSection( FSequencerSectionPainter& Painter )
 	FVector2D GradientSize = FVector2D( Painter.SectionGeometry.Size.X - 2.f, (Painter.SectionGeometry.Size.Y / 4) - 3.0f );
 	if ( GradientSize.X >= 1.f )
 	{
-		FPaintGeometry PaintGeometry = Painter.SectionGeometry.ToPaintGeometry( FVector2D( 1.f, 1.f ), GradientSize );
+		FPaintGeometry PaintGeometry = Painter.SectionGeometry.ToPaintGeometry( GradientSize, FSlateLayoutTransform(FVector2D( 1.f, 1.f )) );
 
 		// If we are showing a background pattern and the colors is transparent, draw a checker pattern
 		FSlateDrawElement::MakeBox(
diff --git a/Engine/Source/Editor/MovieSceneTools/Private/Sections/EventSection.cpp b/Engine/Source/Editor/MovieSceneTools/Private/Sections/EventSection.cpp
--- a/Engine/Source/Editor/MovieSceneTools/Private/Sections/EventSection.cpp
+++ b/Engine/Source/Editor/MovieSceneTools/Private/Sections/EventSection.cpp
@@ -75,7 +75,7 @@ void FEventSectionBase::PaintEventName(FSequencerSectionPainter& Painter, int32
 	FSlateDrawElement::MakeBox(
 		Painter.DrawElements,
 		LayerId + 1,
-		Painter.SectionGeometry.ToPaintGeometry(BoxOffset, BoxSize),
+		Painter.SectionGeometry.ToPaintGeometry(BoxSize, FSlateLayoutTransform(BoxOffset)),
 		FAppStyle::GetBrush("WhiteBrush"),
 		ESlateDrawEffect::None,
 		FLinearColor::Black.CopyWithNewOpacity(0.5f)
@@ -87,7 +87,7 @@ void FEventSectionBase::PaintEventName(FSequencerSectionPainter& Painter, int32
 		FSlateDrawElement::MakeText(
 			Painter.DrawElements,
 			LayerId + 2,
-			Painter.SectionGeometry.ToPaintGeometry(BoxOffset + IconOffset, IconSize),
+			Painter.SectionGeometry.ToPaintGeometry(IconSize, FSlateLayoutTransform(BoxOffset + IconOffset)),
 			WarningString,
 			FontAwesomeFont,
 			Painter.bParentEnabled ? ESlateDrawEffect::None : ESlateDrawEffect::DisabledEffect,
@@ -98,7 +98,7 @@ void FEventSectionBase::PaintEventName(FSequencerSectionPainter& Painter, int32
 	FSlateDrawElement::MakeText(
 		Painter.DrawElements,
 		LayerId + 2,
-		Painter.SectionGeometry.ToPaintGeometry(BoxOffset + TextOffset, TextSize),
+		Painter.SectionGeometry.ToPaintGeometry(TextSize, FSlateLayoutTransform(BoxOffset + TextOffset)),
 		InEventString,
 		SmallLayoutFont,
 		Painter.bParentEnabled ? ESlateDrawEffect::None : ESlateDrawEffect::DisabledEffect,
diff --git a/Engine/Source/Editor/MovieSceneTools/Private/Sections/ParameterSection.cpp b/Engine/Source/Editor/MovieSceneTools/Private/Sections/ParameterSection.cpp
--- a/Engine/Source/Editor/MovieSceneTools/Private/Sections/ParameterSection.cpp
+++ b/Engine/Source/Editor/MovieSceneTools/Private/Sections/ParameterSection.cpp
@@ -121,7 +121,7 @@ namespace UE::Sequencer
 				FVector2D GradientSize = AllottedGeometry.GetLocalSize() - FVector2D(2.f, 0.f);
 				if ( GradientSize.X >= 1.f )
 				{
-					FPaintGeometry PaintGeometry = AllottedGeometry.ToPaintGeometry( FVector2D( 1.f, 1.f ), GradientSize );
+					FPaintGeometry PaintGeometry = AllottedGeometry.ToPaintGeometry( GradientSize, FSlateLayoutTransform(FVector2D( 1.f, 1.f )) );
 
 					// If we are showing a background pattern and the colors is transparent, draw a checker pattern
 					FSlateDrawElement::MakeBox(
diff --git a/Engine/Source/Editor/MovieSceneTools/Private/TrackEditors/AudioTrackEditor.cpp b/Engine/Source/Editor/MovieSceneTools/Private/TrackEditors/AudioTrackEditor.cpp
--- a/Engine/Source/Editor/MovieSceneTools/Private/TrackEditors/AudioTrackEditor.cpp
+++ b/Engine/Source/Editor/MovieSceneTools/Private/TrackEditors/AudioTrackEditor.cpp
@@ -710,7 +710,7 @@ int32 FAudioSection::OnPaintSection( FSequencerSectionPainter& Painter ) const
 		FSlateDrawElement::MakeViewport(
 			Painter.DrawElements,
 			++LayerId,
-			Painter.SectionGeometry.ToPaintGeometry(FVector2D(StoredXOffset, 0), FVector2D(StoredXSize, GetSectionHeight() + 8.f)),
+			Painter.SectionGeometry.ToPaintGeometry(FVector2f(StoredXSize, GetSectionHeight() + 8.f), FSlateLayoutTransform(FVector2f(StoredXOffset, 0))),
 			WaveformThumbnail,
 			(Painter.bParentEnabled ? ESlateDrawEffect::None : ESlateDrawEffect::DisabledEffect) | ESlateDrawEffect::NoGamma,
 			FLinearColor::White
diff --git a/Engine/Source/Editor/MovieSceneTools/Private/TrackEditors/CameraShakeSourceShakeTrackEditor.cpp b/Engine/Source/Editor/MovieSceneTools/Private/TrackEditors/CameraShakeSourceShakeTrackEditor.cpp
--- a/Engine/Source/Editor/MovieSceneTools/Private/TrackEditors/CameraShakeSourceShakeTrackEditor.cpp
+++ b/Engine/Source/Editor/MovieSceneTools/Private/TrackEditors/CameraShakeSourceShakeTrackEditor.cpp
@@ -168,7 +168,7 @@ void FCameraShakeSourceTriggerSection::PaintShakeName(FSequencerSectionPainter&
 	TSharedRef<FSlateFontMeasure> FontMeasureService = FSlateApplication::Get().GetRenderer()->GetFontMeasureService();
 
 	const  FMargin   BoxPadding     = FMargin(4.0f, 2.0f);
-	const FVector2D  TextSize       = FontMeasureService->Measure(ShakeText, SmallLayoutFont);
+	const FVector2f  TextSize       = UE::Slate::CastToVector2f(FontMeasureService->Measure(ShakeText, SmallLayoutFont));
 
 	// Flip the text position if getting near the end of the view range
 	const bool bDrawLeft = (Painter.SectionGeometry.Size.X - PixelPos) < (TextSize.X + 22.f) - BoxOffsetPx;
@@ -176,14 +176,14 @@ void FCameraShakeSourceTriggerSection::PaintShakeName(FSequencerSectionPainter&
 			bDrawLeft ? PixelPos - TextSize.X - BoxOffsetPx : PixelPos + BoxOffsetPx,
 			0.f);
 
-	const FVector2D BoxOffset  = FVector2D(BoxPositionX, Painter.SectionGeometry.Size.Y*.5f - TextSize.Y*.5f);
-	const FVector2D TextOffset = FVector2D(BoxPadding.Left, 0);
+	const FVector2f BoxOffset  = FVector2f(BoxPositionX, Painter.SectionGeometry.Size.Y*.5f - TextSize.Y*.5f);
+	const FVector2f TextOffset = FVector2f(BoxPadding.Left, 0);
 
 	// Draw the background box.
 	FSlateDrawElement::MakeBox(
 		Painter.DrawElements,
 		LayerId + 1,
-		Painter.SectionGeometry.ToPaintGeometry(BoxOffset, TextSize),
+		Painter.SectionGeometry.ToPaintGeometry(TextSize, FSlateLayoutTransform(BoxOffset)),
 		FAppStyle::GetBrush("WhiteBrush"),
 		ESlateDrawEffect::None,
 		FLinearColor::Black.CopyWithNewOpacity(0.5f)
@@ -193,7 +193,7 @@ void FCameraShakeSourceTriggerSection::PaintShakeName(FSequencerSectionPainter&
 	FSlateDrawElement::MakeText(
 		Painter.DrawElements,
 		LayerId + 2,
-		Painter.SectionGeometry.ToPaintGeometry(BoxOffset + TextOffset, TextSize),
+		Painter.SectionGeometry.ToPaintGeometry(TextSize, FSlateLayoutTransform(BoxOffset + TextOffset)),
 		ShakeText,
 		SmallLayoutFont,
 		Painter.bParentEnabled ? ESlateDrawEffect::None : ESlateDrawEffect::DisabledEffect,
diff --git a/Engine/Source/Editor/MovieSceneTools/Private/TrackEditors/FadeTrackEditor.cpp b/Engine/Source/Editor/MovieSceneTools/Private/TrackEditors/FadeTrackEditor.cpp
--- a/Engine/Source/Editor/MovieSceneTools/Private/TrackEditors/FadeTrackEditor.cpp
+++ b/Engine/Source/Editor/MovieSceneTools/Private/TrackEditors/FadeTrackEditor.cpp
@@ -66,7 +66,7 @@ public:
 		const ESlateDrawEffect DrawEffects = Painter.bParentEnabled ? ESlateDrawEffect::None : ESlateDrawEffect::DisabledEffect;
 
 		FVector2D GradientSize = FVector2D( Painter.SectionGeometry.Size.X - 2.f, Painter.SectionGeometry.Size.Y - 3.0f );
-		FPaintGeometry PaintGeometry = Painter.SectionGeometry.ToPaintGeometry( FVector2D( 1.f, 3.f ), GradientSize );
+		FPaintGeometry PaintGeometry = Painter.SectionGeometry.ToPaintGeometry( GradientSize, FSlateLayoutTransform(FVector2f( 1.f, 3.f )) );
 
 		const UMovieSceneFadeSection* FadeSection = Cast<const UMovieSceneFadeSection>( WeakSection.Get() );
 
diff --git a/Engine/Source/Editor/MovieSceneTools/Private/TrackEditors/ParticleTrackEditor.cpp b/Engine/Source/Editor/MovieSceneTools/Private/TrackEditors/ParticleTrackEditor.cpp
--- a/Engine/Source/Editor/MovieSceneTools/Private/TrackEditors/ParticleTrackEditor.cpp
+++ b/Engine/Source/Editor/MovieSceneTools/Private/TrackEditors/ParticleTrackEditor.cpp
@@ -191,14 +191,14 @@ int32 FParticleSection::OnPaintSection( FSequencerSectionPainter& InPainter ) co
 		FSlateDrawElement::MakeBox(
 			InPainter.DrawElements,
 			InPainter.LayerId,
-			InPainter.SectionGeometry.ToPaintGeometry( FVector2D( XOffset, (InPainter.SectionGeometry.GetLocalSize().Y - SequencerSectionConstants::KeySize.Y) / 2 ), FVector2D( XSize, SequencerSectionConstants::KeySize.Y ) ),
+			InPainter.SectionGeometry.ToPaintGeometry( FVector2f( XSize, SequencerSectionConstants::KeySize.Y ), FSlateLayoutTransform(FVector2f( XOffset, (InPainter.SectionGeometry.GetLocalSize().Y - SequencerSectionConstants::KeySize.Y) / 2.f )) ),
 			FAppStyle::GetBrush( "Sequencer.Section.Background" ),
 			DrawEffects
 			);
 		FSlateDrawElement::MakeBox(
 			InPainter.DrawElements,
 			InPainter.LayerId,
-			InPainter.SectionGeometry.ToPaintGeometry( FVector2D( XOffset, (InPainter.SectionGeometry.GetLocalSize().Y - SequencerSectionConstants::KeySize.Y) / 2 ), FVector2D( XSize, SequencerSectionConstants::KeySize.Y ) ),
+			InPainter.SectionGeometry.ToPaintGeometry( FVector2f( XSize, SequencerSectionConstants::KeySize.Y ), FSlateLayoutTransform(FVector2f( XOffset, (InPainter.SectionGeometry.GetLocalSize().Y - SequencerSectionConstants::KeySize.Y) / 2.f )) ),
 			FAppStyle::GetBrush( "Sequencer.Section.BackgroundTint" ),
 			DrawEffects,
 			TrackColor
diff --git a/Engine/Source/Editor/MovieSceneTools/Private/TrackEditors/SubTrackEditorBase.cpp b/Engine/Source/Editor/MovieSceneTools/Private/TrackEditors/SubTrackEditorBase.cpp
--- a/Engine/Source/Editor/MovieSceneTools/Private/TrackEditors/SubTrackEditorBase.cpp
+++ b/Engine/Source/Editor/MovieSceneTools/Private/TrackEditors/SubTrackEditorBase.cpp
@@ -166,8 +166,8 @@ void FSubSectionPainterUtil::DoPaintNonLoopingSection(const UMovieSceneSubSectio
                 InPainter.DrawElements,
                 InPainter.LayerId++,
                 InPainter.SectionGeometry.ToPaintGeometry(
-                    FVector2D(0.0f, 0.f),
-                    FVector2D(StartOffset * PixelsPerFrame, InPainter.SectionGeometry.Size.Y)
+                    FVector2f(StartOffset * PixelsPerFrame, InPainter.SectionGeometry.Size.Y),
+                    FSlateLayoutTransform()
                     ),
                 FAppStyle::GetBrush("WhiteBrush"),
                 DrawEffects,
@@ -179,8 +179,8 @@ void FSubSectionPainterUtil::DoPaintNonLoopingSection(const UMovieSceneSubSectio
                 InPainter.DrawElements,
                 InPainter.LayerId++,
                 InPainter.SectionGeometry.ToPaintGeometry(
-                    FVector2D(StartOffset * PixelsPerFrame, 0.f),
-                    FVector2D(1.0f, InPainter.SectionGeometry.Size.Y)
+                    FVector2f(1.0f, InPainter.SectionGeometry.Size.Y),
+                    FSlateLayoutTransform(FVector2f(StartOffset * PixelsPerFrame, 0.f))
                     ),
                 FAppStyle::GetBrush("WhiteBrush"),
                 DrawEffects,
@@ -197,8 +197,8 @@ void FSubSectionPainterUtil::DoPaintNonLoopingSection(const UMovieSceneSubSectio
                 InPainter.DrawElements,
                 InPainter.LayerId++,
                 InPainter.SectionGeometry.ToPaintGeometry(
-                    FVector2D(EndOffset * PixelsPerFrame, 0.f),
-                    FVector2D((SectionSize - EndOffset) * PixelsPerFrame, InPainter.SectionGeometry.Size.Y)
+                    FVector2f((SectionSize - EndOffset) * PixelsPerFrame, InPainter.SectionGeometry.Size.Y),
+                    FSlateLayoutTransform(FVector2f(EndOffset * PixelsPerFrame, 0.f))
                     ),
                 FAppStyle::GetBrush("WhiteBrush"),
                 DrawEffects,
@@ -211,8 +211,8 @@ void FSubSectionPainterUtil::DoPaintNonLoopingSection(const UMovieSceneSubSectio
                 InPainter.DrawElements,
                 InPainter.LayerId++,
                 InPainter.SectionGeometry.ToPaintGeometry(
-                    FVector2D(EndOffset * PixelsPerFrame, 0.f),
-                    FVector2D(1.0f, InPainter.SectionGeometry.Size.Y)
+                    FVector2f(1.0f, InPainter.SectionGeometry.Size.Y),
+                    FSlateLayoutTransform(FVector2f(EndOffset * PixelsPerFrame, 0.f))
                     ),
                 FAppStyle::GetBrush("WhiteBrush"),
                 DrawEffects,
@@ -257,8 +257,8 @@ void FSubSectionPainterUtil::DoPaintLoopingSection(const UMovieSceneSubSection&
             InPainter.DrawElements,
             InPainter.LayerId++,
             InPainter.SectionGeometry.ToPaintGeometry(
-                FVector2D(CurOffset * PixelsPerFrame, 0.f),
-                FVector2D(1.0f, InPainter.SectionGeometry.Size.Y)
+                FVector2f(1.0f, InPainter.SectionGeometry.Size.Y),
+                FSlateLayoutTransform(FVector2f(CurOffset * PixelsPerFrame, 0.f))
             ),
             FAppStyle::GetBrush("WhiteBrush"),
             DrawEffects,
@@ -270,8 +270,8 @@ void FSubSectionPainterUtil::DoPaintLoopingSection(const UMovieSceneSubSection&
                 InPainter.DrawElements,
                 InPainter.LayerId++,
                 InPainter.SectionGeometry.ToPaintGeometry(
-                    FVector2D(CurOffset * PixelsPerFrame - 1.f, 0.f),
-                    FVector2D(1.0f, InPainter.SectionGeometry.Size.Y)
+                    FVector2f(1.0f, InPainter.SectionGeometry.Size.Y),
+                    FSlateLayoutTransform(FVector2f(CurOffset * PixelsPerFrame - 1.f, 0.f))
                 ),
                 FAppStyle::GetBrush("WhiteBrush"),
                 DrawEffects,
diff --git a/Engine/Source/Editor/Persona/Private/AnimTimeline/AnimTimeSliderController.cpp b/Engine/Source/Editor/Persona/Private/AnimTimeline/AnimTimeSliderController.cpp
--- a/Engine/Source/Editor/Persona/Private/AnimTimeline/AnimTimeSliderController.cpp
+++ b/Engine/Source/Editor/Persona/Private/AnimTimeline/AnimTimeSliderController.cpp
@@ -239,7 +239,7 @@ void FAnimTimeSliderController::DrawTicks( FSlateWindowElementList& OutDrawEleme
 			FSlateDrawElement::MakeText(
 				OutDrawElements,
 				InArgs.StartLayer+1, 
-				InArgs.AllottedGeometry.ToPaintGeometry( TextOffset, InArgs.AllottedGeometry.Size ), 
+				InArgs.AllottedGeometry.ToPaintGeometry( InArgs.AllottedGeometry.Size, FSlateLayoutTransform(TextOffset) ), 
 				FrameString, 
 				SmallLayoutFont,
 				InArgs.DrawEffects,
@@ -323,7 +323,7 @@ int32 FAnimTimeSliderController::OnPaintTimeSlider( bool bMirrorLabels, const FG
 		const float      HandleEnd          = HandleStart + 13.0f;
 
 		const int32 ArrowLayer = LayerId + 2;
-		FPaintGeometry MyGeometry =	AllottedGeometry.ToPaintGeometry( FVector2D( HandleStart, 0 ), FVector2D( HandleEnd - HandleStart, AllottedGeometry.Size.Y ) );
+		FPaintGeometry MyGeometry =	AllottedGeometry.ToPaintGeometry( FVector2f( HandleEnd - HandleStart, AllottedGeometry.Size.Y ), FSlateLayoutTransform(FVector2f( HandleStart, 0.f )) );
 		FLinearColor ScrubColor = InWidgetStyle.GetColorAndOpacityTint();
 		{
 			ScrubColor.A = ScrubColor.A * 0.75f;
@@ -376,7 +376,7 @@ int32 FAnimTimeSliderController::OnPaintTimeSlider( bool bMirrorLabels, const FG
 			FSlateDrawElement::MakeText(
 				OutDrawElements,
 				Args.StartLayer+1, 
-				Args.AllottedGeometry.ToPaintGeometry( TextOffset, TextSize ), 
+				Args.AllottedGeometry.ToPaintGeometry( TextSize, FSlateLayoutTransform(TextOffset) ), 
 				FrameString, 
 				SmallLayoutFont,
 				Args.DrawEffects,
@@ -402,7 +402,7 @@ int32 FAnimTimeSliderController::OnPaintTimeSlider( bool bMirrorLabels, const FG
 			FSlateDrawElement::MakeBox(
 				OutDrawElements,
 				LayerId+1,
-				AllottedGeometry.ToPaintGeometry( FVector2D(RangePosX, 0.f), FVector2D(RangeSizeX, AllottedGeometry.Size.Y) ),
+				AllottedGeometry.ToPaintGeometry( FVector2f(RangeSizeX, AllottedGeometry.Size.Y), FSlateLayoutTransform(FVector2f(RangePosX, 0.f)) ),
 				bMirrorLabels ? ScrubHandleDownBrush : ScrubHandleUpBrush,
 				DrawEffects,
 				MouseStartPosX < MouseEndPosX ? FLinearColor(0.5f, 0.5f, 0.5f) : FLinearColor(0.25f, 0.3f, 0.3f)
@@ -431,7 +431,7 @@ int32 FAnimTimeSliderController::DrawSelectionRange(const FGeometry& AllottedGeo
 			FSlateDrawElement::MakeBox(
 				OutDrawElements,
 				LayerId + 1,
-				AllottedGeometry.ToPaintGeometry(FVector2D(SelectionRangeL, 0.f), FVector2D(SelectionRangeR - SelectionRangeL, AllottedGeometry.Size.Y)),
+				AllottedGeometry.ToPaintGeometry(FVector2f(SelectionRangeR - SelectionRangeL, AllottedGeometry.Size.Y), FSlateLayoutTransform(FVector2D(SelectionRangeL, 0.f))),
 				FAppStyle::GetBrush("WhiteBrush"),
 				ESlateDrawEffect::None,
 				DrawColor.CopyWithNewOpacity(Args.SolidFillOpacity)
@@ -441,7 +441,7 @@ int32 FAnimTimeSliderController::DrawSelectionRange(const FGeometry& AllottedGeo
 		FSlateDrawElement::MakeBox(
 			OutDrawElements,
 			LayerId + 1,
-			AllottedGeometry.ToPaintGeometry(FVector2D(SelectionRangeL, 0.f), FVector2D(Args.BrushWidth, AllottedGeometry.Size.Y)),
+			AllottedGeometry.ToPaintGeometry(FVector2f(Args.BrushWidth, AllottedGeometry.Size.Y), FSlateLayoutTransform(FVector2D(SelectionRangeL, 0.f))),
 			Args.StartBrush,
 			ESlateDrawEffect::None,
 			DrawColor
@@ -450,7 +450,7 @@ int32 FAnimTimeSliderController::DrawSelectionRange(const FGeometry& AllottedGeo
 		FSlateDrawElement::MakeBox(
 			OutDrawElements,
 			LayerId + 1,
-			AllottedGeometry.ToPaintGeometry(FVector2D(SelectionRangeR - Args.BrushWidth, 0.f), FVector2D(Args.BrushWidth, AllottedGeometry.Size.Y)),
+			AllottedGeometry.ToPaintGeometry(FVector2f(Args.BrushWidth, AllottedGeometry.Size.Y), FSlateLayoutTransform(FVector2D(SelectionRangeR - Args.BrushWidth, 0.f))),
 			Args.EndBrush,
 			ESlateDrawEffect::None,
 			DrawColor
@@ -478,7 +478,7 @@ int32 FAnimTimeSliderController::DrawPlaybackRange(const FGeometry& AllottedGeom
 	FSlateDrawElement::MakeBox(
 		OutDrawElements,
 		LayerId+1,
-		AllottedGeometry.ToPaintGeometry(FVector2D(PlaybackRangeL, 0.f), FVector2D(Args.BrushWidth, AllottedGeometry.Size.Y)),
+		AllottedGeometry.ToPaintGeometry(FVector2f(Args.BrushWidth, AllottedGeometry.Size.Y), FSlateLayoutTransform(FVector2D(PlaybackRangeL, 0.f))),
 		Args.StartBrush,
 		ESlateDrawEffect::None,
 		FColor(32, 128, 32, OpacityBlend)	// 120, 75, 50 (HSV)
@@ -487,7 +487,7 @@ int32 FAnimTimeSliderController::DrawPlaybackRange(const FGeometry& AllottedGeom
 	FSlateDrawElement::MakeBox(
 		OutDrawElements,
 		LayerId+1,
-		AllottedGeometry.ToPaintGeometry(FVector2D(PlaybackRangeR - Args.BrushWidth, 0.f), FVector2D(Args.BrushWidth, AllottedGeometry.Size.Y)),
+		AllottedGeometry.ToPaintGeometry(FVector2f(Args.BrushWidth, AllottedGeometry.Size.Y), FSlateLayoutTransform(FVector2D(PlaybackRangeR - Args.BrushWidth, 0.f))),
 		Args.EndBrush,
 		ESlateDrawEffect::None,
 		FColor(128, 32, 32, OpacityBlend)	// 0, 75, 50 (HSV)
@@ -497,7 +497,7 @@ int32 FAnimTimeSliderController::DrawPlaybackRange(const FGeometry& AllottedGeom
 	FSlateDrawElement::MakeBox(
 		OutDrawElements,
 		LayerId+1,
-		AllottedGeometry.ToPaintGeometry(FVector2D(0.f, 0.f), FVector2D(PlaybackRangeL, AllottedGeometry.Size.Y)),
+		AllottedGeometry.ToPaintGeometry(FVector2f(PlaybackRangeL, AllottedGeometry.Size.Y), FSlateLayoutTransform(FVector2D(0.f, 0.f))),
 		FAppStyle::GetBrush("WhiteBrush"),
 		ESlateDrawEffect::None,
 		FLinearColor::Black.CopyWithNewOpacity(0.3f * OpacityBlend / 255.f)
@@ -506,7 +506,7 @@ int32 FAnimTimeSliderController::DrawPlaybackRange(const FGeometry& AllottedGeom
 	FSlateDrawElement::MakeBox(
 		OutDrawElements,
 		LayerId+1,
-		AllottedGeometry.ToPaintGeometry(FVector2D(PlaybackRangeR, 0.f), FVector2D(AllottedGeometry.Size.X - PlaybackRangeR, AllottedGeometry.Size.Y)),
+		AllottedGeometry.ToPaintGeometry(FVector2f(AllottedGeometry.Size.X - PlaybackRangeR, AllottedGeometry.Size.Y), FSlateLayoutTransform(FVector2D(PlaybackRangeR, 0.f))),
 		FAppStyle::GetBrush("WhiteBrush"),
 		ESlateDrawEffect::None,
 		FLinearColor::Black.CopyWithNewOpacity(0.3f * OpacityBlend / 255.f)
@@ -527,7 +527,7 @@ int32 FAnimTimeSliderController::DrawEditableTimes(const FGeometry& AllottedGeom
 		FSlateDrawElement::MakeBox(
 			OutDrawElements,
 			LayerId + 1,
-			AllottedGeometry.ToPaintGeometry(FVector2D(LinePos - 6.0f, AllottedGeometry.Size.Y - 12.0f), FVector2D(11.0f, 12.0f)),
+			AllottedGeometry.ToPaintGeometry(FVector2f(11.0f, 12.0f), FSlateLayoutTransform(FVector2f(LinePos - 6.0f, AllottedGeometry.Size.Y - 12.0f))),
 			EditableTimeBrush,
 			ESlateDrawEffect::None,
 			TimeColor
@@ -926,7 +926,7 @@ int32 FAnimTimeSliderController::OnPaintViewArea( const FGeometry& AllottedGeome
 		FSlateDrawElement::MakeLines(
 			OutDrawElements,
 			LayerId+1,
-			AllottedGeometry.ToPaintGeometry( FVector2D(LinePos, 0.0f ), FVector2D(1.0f,1.0f) ),
+			AllottedGeometry.ToPaintGeometry( FVector2f(1.0f,1.0f), FSlateLayoutTransform(FVector2f(LinePos, 0.0f )) ),
 			LinePoints,
 			DrawEffects,
 			FLinearColor(1.f, 1.f, 1.f, .5f),
@@ -954,7 +954,7 @@ int32 FAnimTimeSliderController::OnPaintViewArea( const FGeometry& AllottedGeome
 			FSlateDrawElement::MakeLines(
 				OutDrawElements,
 				LayerId+1,
-				AllottedGeometry.ToPaintGeometry( FVector2D(LinePos, 0.0f ), FVector2D(1.0f,1.0f) ),
+				AllottedGeometry.ToPaintGeometry( FVector2f(1.0f,1.0f), FSlateLayoutTransform(FVector2f(LinePos, 0.0f )) ),
 				LinePoints,
 				DrawEffects,
 				LineColor,
diff --git a/Engine/Source/Editor/Persona/Private/AnimTimeline/AnimTimelineTrack_Montage.cpp b/Engine/Source/Editor/Persona/Private/AnimTimeline/AnimTimelineTrack_Montage.cpp
--- a/Engine/Source/Editor/Persona/Private/AnimTimeline/AnimTimelineTrack_Montage.cpp
+++ b/Engine/Source/Editor/Persona/Private/AnimTimeline/AnimTimelineTrack_Montage.cpp
@@ -173,7 +173,7 @@ class SMontageSections : public SLeafWidget
 			FSlateDrawElement::MakeBox( 
 				OutDrawElements,
 				++LayerId,
-				AllottedGeometry.ToPaintGeometry(FVector2D(bDrawLabelOnLeft ? LabelPosX - TotalBorderSize.X : LabelPosX, LabelPosY), TotalBorderSize),
+				AllottedGeometry.ToPaintGeometry(TotalBorderSize, FSlateLayoutTransform(FVector2D(bDrawLabelOnLeft ? LabelPosX - TotalBorderSize.X : LabelPosX, LabelPosY))),
 				BorderBrush,
 				ESlateDrawEffect::None,
 				SelectedSectionIndex == SectionIndex ? SelectedColor : MontageColor);
@@ -181,7 +181,7 @@ class SMontageSections : public SLeafWidget
 			FSlateDrawElement::MakeText( 
 				OutDrawElements,
 				++LayerId,
-				AllottedGeometry.ToPaintGeometry(FVector2D(bDrawLabelOnLeft ? (LabelPosX - TotalBorderSize.X) + MontageSectionsConstants::TextOffset.X : LabelPosX + MontageSectionsConstants::TextOffset.X, LabelPosY + MontageSectionsConstants::TextOffset.Y), TextSize),
+				AllottedGeometry.ToPaintGeometry(TextSize, FSlateLayoutTransform(FVector2D(bDrawLabelOnLeft ? (LabelPosX - TotalBorderSize.X) + MontageSectionsConstants::TextOffset.X : LabelPosX + MontageSectionsConstants::TextOffset.X, LabelPosY + MontageSectionsConstants::TextOffset.Y))),
 				NameText,
 				LabelFont,
 				ESlateDrawEffect::None,
@@ -190,7 +190,7 @@ class SMontageSections : public SLeafWidget
 			FSlateDrawElement::MakeText( 
 				OutDrawElements,
 				++LayerId,
-				AllottedGeometry.ToPaintGeometry(FVector2D(bDrawLabelOnLeft ? (LabelPosX - TotalBorderSize.X) + TextBorderSize.X + (MontageSectionsConstants::TextBorderMargin.X * 2.0f) + MontageSectionsConstants::TextOffset.X : LabelPosX + MontageSectionsConstants::TextOffset.X + (MontageSectionsConstants::TextBorderMargin.X * 2.0f) + TextBorderSize.X, IconPosY + MontageSectionsConstants::TextOffset.Y), TextIconSize),
+				AllottedGeometry.ToPaintGeometry(TextIconSize, FSlateLayoutTransform(FVector2D(bDrawLabelOnLeft ? (LabelPosX - TotalBorderSize.X) + TextBorderSize.X + (MontageSectionsConstants::TextBorderMargin.X * 2.0f) + MontageSectionsConstants::TextOffset.X : LabelPosX + MontageSectionsConstants::TextOffset.X + (MontageSectionsConstants::TextBorderMargin.X * 2.0f) + TextBorderSize.X, IconPosY + MontageSectionsConstants::TextOffset.Y))),
 				IconText,
 				IconFont,
 				ESlateDrawEffect::None,
@@ -230,7 +230,10 @@ class SMontageSections : public SLeafWidget
 			const double RightEdgeToNotify = MyGeometry.Size.X - (LabelPosX + TotalBorderSize.X);
 			const bool bDrawLabelOnLeft = RightEdgeToNotify < 0.0;
 
-			const FGeometry LabelGeometry = MyGeometry.MakeChild(FVector2D(bDrawLabelOnLeft ? LabelPosX - TotalBorderSize.X: LabelPosX, LabelPosY), TotalBorderSize);
+			const FGeometry LabelGeometry = MyGeometry.MakeChild(
+				TotalBorderSize,
+				FSlateLayoutTransform(FVector2D(bDrawLabelOnLeft ? LabelPosX - TotalBorderSize.X: LabelPosX, LabelPosY))
+			);
 
 			if(LabelGeometry.IsUnderLocation(MouseEvent.GetScreenSpacePosition()))
 			{
diff --git a/Engine/Source/Editor/Persona/Private/AnimTimeline/SAnimTrack.cpp b/Engine/Source/Editor/Persona/Private/AnimTimeline/SAnimTrack.cpp
--- a/Engine/Source/Editor/Persona/Private/AnimTimeline/SAnimTrack.cpp
+++ b/Engine/Source/Editor/Persona/Private/AnimTimeline/SAnimTrack.cpp
@@ -44,8 +44,8 @@ int32 SAnimTrack::OnPaint(const FPaintArgs& Args, const FGeometry& AllottedGeome
 					OutDrawElements,
 					LayerId++,
 					AllottedGeometry.ToPaintGeometry(
-						FVector2D(0, 0),
-						FVector2D(AllottedGeometry.GetLocalSize().X, TotalNodeHeight)
+						FVector2D(AllottedGeometry.GetLocalSize().X, TotalNodeHeight),
+						FSlateLayoutTransform()
 					),
 					FAppStyle::GetBrush(BorderName),
 					ESlateDrawEffect::None,
diff --git a/Engine/Source/Editor/Persona/Private/SAnimNotifyPanel.cpp b/Engine/Source/Editor/Persona/Private/SAnimNotifyPanel.cpp
--- a/Engine/Source/Editor/Persona/Private/SAnimNotifyPanel.cpp
+++ b/Engine/Source/Editor/Persona/Private/SAnimNotifyPanel.cpp
@@ -1742,7 +1742,7 @@ int32 SAnimNotifyNode::OnPaint(const FPaintArgs& Args, const FGeometry& Allotted
 	{
 		FVector2D MarkerSize = EndMarkerNodeOverlay->GetDesiredSize();
 		FVector2D MarkerOffset(NotifyDurationSizeX + MarkerSize.X * 0.5f + 5.0f, (NotifyHeight - MarkerSize.Y) * 0.5f);
-		EndMarkerNodeOverlay->Paint(Args.WithNewParent(this), AllottedGeometry.MakeChild(MarkerOffset, MarkerSize, 1.0f), MyCullingRect, OutDrawElements, LayerId, InWidgetStyle, bParentEnabled);
+		EndMarkerNodeOverlay->Paint(Args.WithNewParent(this), AllottedGeometry.MakeChild(MarkerSize, FSlateLayoutTransform(MarkerOffset)), MyCullingRect, OutDrawElements, LayerId, InWidgetStyle, bParentEnabled);
 	}
 
 	const FSlateBrush* StyleInfo = FAppStyle::GetBrush( TEXT("SpecialEditableTextImageNormal") );
@@ -1761,7 +1761,7 @@ int32 SAnimNotifyNode::OnPaint(const FPaintArgs& Args, const FGeometry& Allotted
 		FSlateDrawElement::MakeBox( 
 			OutDrawElements,
 			LayerId, 
-			AllottedGeometry.ToPaintGeometry(DurationBoxPosition, DurationBoxSize), 
+			AllottedGeometry.ToPaintGeometry(DurationBoxSize, FSlateLayoutTransform(DurationBoxPosition)), 
 			StyleInfo,
 			ESlateDrawEffect::None,
 			BoxColor);
@@ -1795,7 +1795,7 @@ int32 SAnimNotifyNode::OnPaint(const FPaintArgs& Args, const FGeometry& Allotted
 		FSlateDrawElement::MakeBox( 
 			OutDrawElements,
 			LayerId, 
-			AllottedGeometry.ToPaintGeometry(LabelPosition, LabelSize), 
+			AllottedGeometry.ToPaintGeometry(LabelSize, FSlateLayoutTransform(LabelPosition)),
 			StyleInfo,
 			ESlateDrawEffect::None,
 			BoxColor);
@@ -1817,7 +1817,7 @@ int32 SAnimNotifyNode::OnPaint(const FPaintArgs& Args, const FGeometry& Allotted
 		TextPosition.X += BranchingPointIconSize.X;
 	}
 
-	FPaintGeometry TextGeometry = AllottedGeometry.ToPaintGeometry(TextPosition, DrawTextSize);
+	FPaintGeometry TextGeometry = AllottedGeometry.ToPaintGeometry(DrawTextSize, FSlateLayoutTransform(TextPosition));
 	OutDrawElements.PushClip(FSlateClippingZone(TextGeometry));
 
 	FSlateDrawElement::MakeText( 
@@ -1843,7 +1843,7 @@ int32 SAnimNotifyNode::OnPaint(const FPaintArgs& Args, const FGeometry& Allotted
 		FSlateDrawElement::MakeBox(
 			OutDrawElements,
 			BranchPointLayerID,
-			AllottedGeometry.ToPaintGeometry(BranchPointIconPos, BranchingPointIconSize),
+			AllottedGeometry.ToPaintGeometry(BranchingPointIconSize, FSlateLayoutTransform(BranchPointIconPos)),
 			FAppStyle::GetBrush(TEXT("AnimNotifyEditor.BranchingPoint")),
 			ESlateDrawEffect::None,
 			FLinearColor::White
@@ -2078,7 +2078,7 @@ void SAnimNotifyNode::DrawScrubHandle( float ScrubHandleCentre, FSlateWindowElem
 	FSlateDrawElement::MakeBox( 
 		OutDrawElements,
 		ScrubHandleID, 
-		AllottedGeometry.ToPaintGeometry(ScrubHandlePosition, ScrubHandleSize), 
+		AllottedGeometry.ToPaintGeometry(ScrubHandleSize, FSlateLayoutTransform(ScrubHandlePosition)), 
 		FAppStyle::GetBrush( TEXT( "Sequencer.KeyDiamond" ) ),
 		ESlateDrawEffect::None,
 		NodeColour
@@ -2087,7 +2087,7 @@ void SAnimNotifyNode::DrawScrubHandle( float ScrubHandleCentre, FSlateWindowElem
 	FSlateDrawElement::MakeBox( 
 		OutDrawElements,
 		ScrubHandleID, 
-		AllottedGeometry.ToPaintGeometry(ScrubHandlePosition, ScrubHandleSize), 
+		AllottedGeometry.ToPaintGeometry(ScrubHandleSize, FSlateLayoutTransform(ScrubHandlePosition)), 
 		FAppStyle::GetBrush( TEXT( "Sequencer.KeyDiamondBorder" ) ),
 		ESlateDrawEffect::None,
 		bSelected ? FAppStyle::GetSlateColor("SelectionColor").GetSpecifiedColor() : FLinearColor::Black
@@ -2112,7 +2112,7 @@ void SAnimNotifyNode::DrawHandleOffset( const float& Offset, const float& Handle
 	FSlateDrawElement::MakeBox( 
 		OutDrawElements,
 		MarkerLayer, 
-		AllottedGeometry.ToPaintGeometry(MarkerPosition, MarkerSize), 
+		AllottedGeometry.ToPaintGeometry(MarkerSize, FSlateLayoutTransform(MarkerPosition)), 
 		FAppStyle::GetBrush( TEXT( "Sequencer.Timeline.NotifyAlignmentMarker" ) ),
 		ESlateDrawEffect::None,
 		NodeColor
@@ -4747,7 +4747,7 @@ int32 SAnimNotifyPanel::OnPaint(const FPaintArgs& Args, const FGeometry& Allotte
 		FSlateDrawElement::MakeBox(
 			OutDrawElements,
 			LayerId++,
-			AllottedGeometry.ToPaintGeometry(Marquee.Rect.GetUpperLeft(), Marquee.Rect.GetSize()),
+			AllottedGeometry.ToPaintGeometry(Marquee.Rect.GetSize(), FSlateLayoutTransform(Marquee.Rect.GetUpperLeft())),
 			FAppStyle::GetBrush(TEXT("MarqueeSelection"))
 			);
 	}
diff --git a/Engine/Source/Editor/Persona/Private/SAnimationBlendSpaceGridWidget.cpp b/Engine/Source/Editor/Persona/Private/SAnimationBlendSpaceGridWidget.cpp
--- a/Engine/Source/Editor/Persona/Private/SAnimationBlendSpaceGridWidget.cpp
+++ b/Engine/Source/Editor/Persona/Private/SAnimationBlendSpaceGridWidget.cpp
@@ -538,7 +538,7 @@ void SBlendSpaceGridWidget::PaintBackgroundAndGrid(const FGeometry& AllottedGeom
 		const FVector2D GridOffset = CachedGridRectangle.GetTopLeft();
 
 		// Fill the background of the grid
-		FSlateDrawElement::MakeBox( OutDrawElements, DrawLayerId + 1, AllottedGeometry.ToPaintGeometry(GridOffset, GridSize), BackgroundImage );
+		FSlateDrawElement::MakeBox( OutDrawElements, DrawLayerId + 1, AllottedGeometry.ToPaintGeometry(GridSize, FSlateLayoutTransform(GridOffset)), BackgroundImage );
 		
 		TArray<FVector2D> LinePoints;
 
@@ -651,7 +651,7 @@ void SBlendSpaceGridWidget::PaintSampleKeys(
 
 			const FVector2D GridPosition = SampleValueToScreenPosition(Sample.SampleValue) - (KeySize * 0.5f);
 			FSlateDrawElement::MakeBox(
-				OutDrawElements, SampleLayer, AllottedGeometry.ToPaintGeometry(GridPosition, KeySize), 
+				OutDrawElements, SampleLayer, AllottedGeometry.ToPaintGeometry(KeySize, FSlateLayoutTransform(GridPosition)), 
 				KeyBrush, ESlateDrawEffect::None, DrawColor );
 
 			const float SampleLookupWeight = GetSampleLookupWeight(SampleIndex);
@@ -672,7 +672,7 @@ void SBlendSpaceGridWidget::PaintSampleKeys(
 		{
 			const FVector2D GridPosition = SampleValueToScreenPosition(PreviewFilteredPosition) - (PreviewSize * .5f);
 			FSlateDrawElement::MakeBox(
-				OutDrawElements, FilteredPositionLayer, AllottedGeometry.ToPaintGeometry(GridPosition, PreviewSize), 
+				OutDrawElements, FilteredPositionLayer, AllottedGeometry.ToPaintGeometry(PreviewSize, FSlateLayoutTransform(GridPosition)), 
 				PreviewBrush, ESlateDrawEffect::None, PreviewKeyColor.GetSpecifiedColor() * 0.7f);
 		}
 
@@ -680,7 +680,7 @@ void SBlendSpaceGridWidget::PaintSampleKeys(
 		{
 			const FVector2D GridPosition = SampleValueToScreenPosition(PreviewPosition) - (PreviewSize * .5f);
 			FSlateDrawElement::MakeBox(
-				OutDrawElements, PreviewPositionLayer, AllottedGeometry.ToPaintGeometry(GridPosition, PreviewSize), 
+				OutDrawElements, PreviewPositionLayer, AllottedGeometry.ToPaintGeometry(PreviewSize, FSlateLayoutTransform(GridPosition)), 
 				PreviewBrush, ESlateDrawEffect::None, PreviewKeyColor.GetSpecifiedColor());
 		}
 
@@ -689,7 +689,7 @@ void SBlendSpaceGridWidget::PaintSampleKeys(
 			const FVector2D GridPoint = SnapScreenPositionToGrid(
 				LocalMousePosition, FSlateApplication::Get().GetModifierKeys().IsShiftDown()) - (KeySize * .5f);
 			FSlateDrawElement::MakeBox(
-				OutDrawElements, SampleLayer, AllottedGeometry.ToPaintGeometry(GridPoint, KeySize), KeyBrush, ESlateDrawEffect::None,
+				OutDrawElements, SampleLayer, AllottedGeometry.ToPaintGeometry(KeySize, FSlateLayoutTransform(GridPoint)), KeyBrush, ESlateDrawEffect::None,
 				(DragState == EDragState::DragDrop) ? DropKeyColor.GetSpecifiedColor() : InvalidColor.GetSpecifiedColor() );
 		}
 
@@ -704,20 +704,26 @@ void SBlendSpaceGridWidget::PaintSampleKeys(
 
 				float MaxWeightWidth = 48;
 				float WeightHeight = 6;
+				float Border = 1.0f;
 
 				Point.Y -= KeySize.Y / 2 + WeightHeight * 1.25;
 				Point.X -= MaxWeightWidth * 0.5f;
 
+				FGeometry FillGeometry = AllottedGeometry.MakeChild(
+					FVector2D(MaxWeightWidth, WeightHeight),
+					FSlateLayoutTransform(FVector2D(Point.X, Point.Y))
+				);
+				FGeometry BorderGeometry = AllottedGeometry.MakeChild(
+					FVector2D(Weight * (MaxWeightWidth - 2 * Border), WeightHeight - 2 * Border),
+					FSlateLayoutTransform(FVector2D(Point.X + Border, Point.Y + Border))
+				);
+
 				FSlateDrawElement::MakeBox(
-					OutDrawElements, DrawLayerId + 1, AllottedGeometry.MakeChild(
-						FVector2D(Point.X, Point.Y), FVector2D(MaxWeightWidth, WeightHeight)).ToPaintGeometry(),
+					OutDrawElements, DrawLayerId + 1, FillGeometry.ToPaintGeometry(),
 					LabelBrush, ESlateDrawEffect::None, FLinearColor::Black);
 
-				float Border = 1.0f;
 				FSlateDrawElement::MakeBox(
-					OutDrawElements, DrawLayerId + 2, AllottedGeometry.MakeChild(
-						FVector2D(Point.X + Border, Point.Y + Border), 
-						FVector2D(Weight * (MaxWeightWidth - 2 * Border), WeightHeight - 2 * Border)).ToPaintGeometry(),
+					OutDrawElements, DrawLayerId + 2, BorderGeometry.ToPaintGeometry(),
 					LabelBrush, ESlateDrawEffect::None, FLinearColor::Gray);
 			}
 		}
@@ -746,36 +752,38 @@ void SBlendSpaceGridWidget::PaintAxisText(
 	FVector2D TextPosition = FVector2D(GridCenter.X - (TextSize.X * .5f), CachedGridRectangle.Bottom + TextMargin + (ArrowSize.Y * .25f));
 	FVector2D ArrowPosition = FVector2D(TextPosition.X - ArrowSize.X - 10.f/* give padding*/, TextPosition.Y);
 	FSlateDrawElement::MakeBox(
-		OutDrawElements, DrawLayerId + 1, AllottedGeometry.ToPaintGeometry(ArrowPosition, ArrowSize), 
+		OutDrawElements, DrawLayerId + 1, AllottedGeometry.ToPaintGeometry(ArrowSize, FSlateLayoutTransform(ArrowPosition)), 
 		ArrowBrushes[(uint8)EArrowDirection::Left], ESlateDrawEffect::None, FLinearColor::White);
 
 	// Label
 	FSlateDrawElement::MakeText(
-		OutDrawElements, DrawLayerId + 1, AllottedGeometry.MakeChild(TextPosition, FVector2D(1.0f, 1.0f)).ToPaintGeometry(), 
+		OutDrawElements, DrawLayerId + 1, AllottedGeometry.MakeChild(FVector2D(1.0f, 1.0f), FSlateLayoutTransform(TextPosition)).ToPaintGeometry(), 
 		Text, FontInfo, ESlateDrawEffect::None, FLinearColor::White);
 
 	// arrow right
 	ArrowSize = ArrowBrushes[(uint8)EArrowDirection::Right]->GetImageSize();
 	ArrowPosition = FVector2D(TextPosition.X + TextSize.X + 10.f/* give padding*/, TextPosition.Y);
 	FSlateDrawElement::MakeBox(
-		OutDrawElements, DrawLayerId + 1, AllottedGeometry.ToPaintGeometry(ArrowPosition, ArrowSize), 
+		OutDrawElements, DrawLayerId + 1, AllottedGeometry.ToPaintGeometry(ArrowSize, FSlateLayoutTransform(ArrowPosition)), 
 		ArrowBrushes[(uint8)EArrowDirection::Right], ESlateDrawEffect::None, FLinearColor::White);
 
 	Text = FString::SanitizeFloat(SampleValueMin.X);
 	TextSize = FontMeasure->Measure(Text, FontInfo);
 
 	// Minimum value
-	FSlateDrawElement::MakeText(OutDrawElements, DrawLayerId + 1, AllottedGeometry.MakeChild(FVector2D(
-		CachedGridRectangle.Left - (TextSize.X * .5f), CachedGridRectangle.Bottom + TextMargin + (TextSize.Y * .25f)), 
-		FVector2D(1.0f, 1.0f)).ToPaintGeometry(), Text, FontInfo, ESlateDrawEffect::None, FLinearColor::White);
+	FVector2D MinTextPosition(CachedGridRectangle.Left - (TextSize.X * .5f), CachedGridRectangle.Bottom + TextMargin + (TextSize.Y * .25f));
+	FSlateDrawElement::MakeText(
+		OutDrawElements, DrawLayerId + 1, AllottedGeometry.MakeChild(FVector2D(1.0f, 1.0f), FSlateLayoutTransform(MinTextPosition)).ToPaintGeometry(),
+		Text, FontInfo, ESlateDrawEffect::None, FLinearColor::White);
 
 	Text = FString::SanitizeFloat(SampleValueMax.X);
 	TextSize = FontMeasure->Measure(Text, FontInfo);
 
 	// Maximum value
-	FSlateDrawElement::MakeText(OutDrawElements, DrawLayerId + 1, AllottedGeometry.MakeChild(FVector2D(
-		CachedGridRectangle.Right - (TextSize.X * .5f), CachedGridRectangle.Bottom + TextMargin + (TextSize.Y * .25f)), 
-		FVector2D(1.0f, 1.0f)).ToPaintGeometry(), Text, FontInfo, ESlateDrawEffect::None, FLinearColor::White);
+	FVector2D MaxTextPosition(CachedGridRectangle.Right - (TextSize.X * .5f), CachedGridRectangle.Bottom + TextMargin + (TextSize.Y * .25f));
+	FSlateDrawElement::MakeText(
+		OutDrawElements, DrawLayerId + 1, AllottedGeometry.MakeChild(FVector2D(1.0f, 1.0f), FSlateLayoutTransform(MaxTextPosition)).ToPaintGeometry(),
+		Text, FontInfo, ESlateDrawEffect::None, FLinearColor::White);
 
 	// Only draw Y axis labels if this is a 2D grid
 	if (GridType == EGridType::TwoAxis)
@@ -789,38 +797,39 @@ void SBlendSpaceGridWidget::PaintAxisText(
 		TextPosition = FVector2D(((GridMargin.Left - TextSize.X) * 0.5f - (ArrowSize.X * .25f)) + GridRatioMargin.Left, GridCenter.Y - (TextSize.Y * .5f));
 		ArrowPosition = FVector2D(TextPosition.X + TextSize.X * 0.5f - ArrowSize.X * 0.5f, TextPosition.Y - ArrowSize.Y - 10.f/* give padding*/);
 		FSlateDrawElement::MakeBox(
-			OutDrawElements, DrawLayerId + 1, AllottedGeometry.ToPaintGeometry(ArrowPosition, ArrowSize), 
+			OutDrawElements, DrawLayerId + 1, AllottedGeometry.ToPaintGeometry(ArrowSize, FSlateLayoutTransform(ArrowPosition)), 
 			ArrowBrushes[(uint8)EArrowDirection::Up], ESlateDrawEffect::None, FLinearColor::White);
 
 		// Label
 		FSlateDrawElement::MakeText(
-			OutDrawElements, DrawLayerId + 1, AllottedGeometry.MakeChild(TextPosition, FVector2D(1.0f, 1.0f)).ToPaintGeometry(),
+			OutDrawElements, DrawLayerId + 1, AllottedGeometry.MakeChild(FVector2D(1.0f, 1.0f), FSlateLayoutTransform(TextPosition)).ToPaintGeometry(),
 			Text,	FontInfo, ESlateDrawEffect::None, FLinearColor::White);
 
 		// arrow down
 		ArrowSize = ArrowBrushes[(uint8)EArrowDirection::Down]->GetImageSize();
 		ArrowPosition = FVector2D(TextPosition.X + TextSize.X * 0.5f - ArrowSize.X * 0.5f, TextPosition.Y + TextSize.Y + 10.f/* give padding*/);
 		FSlateDrawElement::MakeBox(
-			OutDrawElements, DrawLayerId + 1, AllottedGeometry.ToPaintGeometry(ArrowPosition, ArrowSize), 
+			OutDrawElements, DrawLayerId + 1, AllottedGeometry.ToPaintGeometry(ArrowSize, FSlateLayoutTransform(ArrowPosition)), 
 			ArrowBrushes[(uint8)EArrowDirection::Down], ESlateDrawEffect::None, FLinearColor::White);
 
 		Text = FString::SanitizeFloat(SampleValueMin.Y);
 		TextSize = FontMeasure->Measure(Text, FontInfo);
 
 		// Minimum value
-		FSlateDrawElement::MakeText(OutDrawElements, DrawLayerId + 1, AllottedGeometry.MakeChild(FVector2D(
-			((GridMargin.Left - TextSize.X) * 0.5f - (TextSize.X * .25f)) + GridRatioMargin.Left, 
-			CachedGridRectangle.Bottom - (TextSize.Y * .5f)), FVector2D(1.0f, 1.0f)).ToPaintGeometry(), 
+		FVector2D MinValuePosition(((GridMargin.Left - TextSize.X) * 0.5f - (TextSize.X * .25f)) + GridRatioMargin.Left, CachedGridRectangle.Bottom - (TextSize.Y * .5f));
+		FSlateDrawElement::MakeText(OutDrawElements, DrawLayerId + 1, AllottedGeometry.MakeChild(
+				FVector2D(1.0f, 1.0f),
+				FSlateLayoutTransform(MinValuePosition)
+			).ToPaintGeometry(), 
 			Text, FontInfo, ESlateDrawEffect::None, FLinearColor::White);
 
 		Text = FString::SanitizeFloat(SampleValueMax.Y);
 		TextSize = FontMeasure->Measure(Text, FontInfo);
 
 		// Maximum value
-		FSlateDrawElement::MakeText(OutDrawElements, DrawLayerId + 1, AllottedGeometry.MakeChild(FVector2D(
-			((GridMargin.Left - TextSize.X) * 0.5f - (TextSize.X * .25f) ) + GridRatioMargin.Left, 
-			CachedGridRectangle.Top - (TextSize.Y * .5f)), 
-			FVector2D(1.0f, 1.0f)).ToPaintGeometry(), Text, FontInfo, ESlateDrawEffect::None, FLinearColor::White);
+		FSlateDrawElement::MakeText(OutDrawElements, DrawLayerId + 1, AllottedGeometry.MakeChild(
+			FVector2D(1.0f, 1.0f), FSlateLayoutTransform( FVector2D( ((GridMargin.Left - TextSize.X) * 0.5f - (TextSize.X * .25f) ) + GridRatioMargin.Left,  CachedGridRectangle.Top - (TextSize.Y * .5f)) )).ToPaintGeometry(),
+			Text, FontInfo, ESlateDrawEffect::None, FLinearColor::White);
 	}
 
 	DrawLayerId += 1;
@@ -880,8 +889,8 @@ void SBlendSpaceGridWidget::PaintGridSampleWeights(
 			GridPosition.Y += Padding.Y / 2;
 
 			FSlateDrawElement::MakeText(
-				OutDrawElements, DrawLayerId + 2, AllottedGeometry.MakeChild(GridPosition + Padding / 2,
-					FVector2D(1.0f, 1.0f)).ToPaintGeometry(), Name, FontInfo, ESlateDrawEffect::None, FLinearColor::White);
+				OutDrawElements, DrawLayerId + 2, AllottedGeometry.MakeChild(FVector2f(1.0f, 1.0f),
+					FSlateLayoutTransform(GridPosition + Padding / 2)).ToPaintGeometry(), Name, FontInfo, ESlateDrawEffect::None, FLinearColor::White);
 		}
 	}
 }
@@ -1085,11 +1094,11 @@ void SBlendSpaceGridWidget::PaintAnimationNames(const FGeometry& AllottedGeometr
 			GridPosition.Y += Padding.Y / 2;
 
 			FSlateDrawElement::MakeBox(
-				OutDrawElements, DrawLayerId + 1, AllottedGeometry.MakeChild(GridPosition, 
-				TextSize + Padding).ToPaintGeometry(), LabelBrush, ESlateDrawEffect::None, FLinearColor::Black);
+				OutDrawElements, DrawLayerId + 1, AllottedGeometry.MakeChild(TextSize + Padding, FSlateLayoutTransform(GridPosition)).ToPaintGeometry(),
+				LabelBrush, ESlateDrawEffect::None, FLinearColor::Black);
 			FSlateDrawElement::MakeText(
-				OutDrawElements, DrawLayerId + 2, AllottedGeometry.MakeChild(GridPosition + Padding/2,
-					FVector2D(1.0f, 1.0f)).ToPaintGeometry(), Name, FontInfo, ESlateDrawEffect::None, FLinearColor::White);
+				OutDrawElements, DrawLayerId + 2, AllottedGeometry.MakeChild(FVector2D(1.0f, 1.0f), FSlateLayoutTransform(GridPosition + Padding/2)).ToPaintGeometry(),
+				Name, FontInfo, ESlateDrawEffect::None, FLinearColor::White);
 		}
 	}
 
diff --git a/Engine/Source/Editor/Persona/Private/SPoseWatchManagerTreeView.cpp b/Engine/Source/Editor/Persona/Private/SPoseWatchManagerTreeView.cpp
--- a/Engine/Source/Editor/Persona/Private/SPoseWatchManagerTreeView.cpp
+++ b/Engine/Source/Editor/Persona/Private/SPoseWatchManagerTreeView.cpp
@@ -368,7 +368,7 @@ int32 SPoseWatchManagerTreeRow::OnPaint(const FPaintArgs& Args, const FGeometry&
 		FSlateDrawElement::MakeBox(
 			OutDrawElements,
 			LayerId + TextLayer,
-			AllottedGeometry.ToPaintGeometry(DrawPosition, DrawSize),	// Position, Size, Scale
+			AllottedGeometry.ToPaintGeometry(DrawSize, FSlateLayoutTransform(DrawPosition)),	// Position, Size, Scale
 			StyleInfo,													// Style
 			DrawEffects,												// Effects to use
 			HighlightTargetColorAndOpacity);							// Color
diff --git a/Engine/Source/Editor/Persona/Private/STrack.cpp b/Engine/Source/Editor/Persona/Private/STrack.cpp
--- a/Engine/Source/Editor/Persona/Private/STrack.cpp
+++ b/Engine/Source/Editor/Persona/Private/STrack.cpp
@@ -475,7 +475,7 @@ int32 STrack::OnPaint( const FPaintArgs& Args, const FGeometry& AllottedGeometry
 		{
 			//float EndPos = (DraggableBarIndex.Get().IsValidIndex(I+1) ? DraggableBarIndex.Get()(I+1)
 
-			FPaintGeometry TextGeometry = AllottedGeometry.ToPaintGeometry( FVector2D(XPos + 15.f, 5.f), AllottedGeometry.GetDrawSize() );
+			FPaintGeometry TextGeometry = AllottedGeometry.ToPaintGeometry( AllottedGeometry.GetDrawSize(), FSlateLayoutTransform(FVector2f(XPos + 15.f, 5.f)) );
 			FSlateDrawElement::MakeText( 
 				OutDrawElements,
 				CustomLayerId,
diff --git a/Engine/Source/Editor/PropertyEditor/Private/SDetailRowIndent.cpp b/Engine/Source/Editor/PropertyEditor/Private/SDetailRowIndent.cpp
--- a/Engine/Source/Editor/PropertyEditor/Private/SDetailRowIndent.cpp
+++ b/Engine/Source/Editor/PropertyEditor/Private/SDetailRowIndent.cpp
@@ -36,7 +36,7 @@ int32 SDetailRowIndent::OnPaint(const FPaintArgs& Args, const FGeometry& Allotte
 		FSlateDrawElement::MakeBox(
 			OutDrawElements,
 			LayerId,
-			AllottedGeometry.ToPaintGeometry(FVector2D(16 * i, 0), FVector2D(16, AllottedGeometry.GetLocalSize().Y)),
+			AllottedGeometry.ToPaintGeometry(FVector2f(16.f, AllottedGeometry.GetLocalSize().Y), FSlateLayoutTransform(FVector2f(16.f * i, 0.f))),
 			BackgroundBrush,
 			ESlateDrawEffect::None,
 			BackgroundColor.GetColor(InWidgetStyle)
@@ -45,7 +45,7 @@ int32 SDetailRowIndent::OnPaint(const FPaintArgs& Args, const FGeometry& Allotte
 		FSlateDrawElement::MakeBox(
 			OutDrawElements,
 			LayerId + 1,
-			AllottedGeometry.ToPaintGeometry(FVector2D(16 * i, 0), FVector2D(16, AllottedGeometry.GetLocalSize().Y)),
+			AllottedGeometry.ToPaintGeometry(FVector2f(16.f, AllottedGeometry.GetLocalSize().Y), FSlateLayoutTransform(FVector2f(16.f * i, 0.f))),
 			DropShadowBrush
 		);
 	}
diff --git a/Engine/Source/Editor/SceneOutliner/Private/SOutlinerTreeView.cpp b/Engine/Source/Editor/SceneOutliner/Private/SOutlinerTreeView.cpp
--- a/Engine/Source/Editor/SceneOutliner/Private/SOutlinerTreeView.cpp
+++ b/Engine/Source/Editor/SceneOutliner/Private/SOutlinerTreeView.cpp
@@ -419,7 +419,7 @@ int32 SSceneOutlinerTreeRow::OnPaint( const FPaintArgs& Args, const FGeometry& A
 		FSlateDrawElement::MakeBox(
 			OutDrawElements,
 			LayerId + TextLayer,
-			AllottedGeometry.ToPaintGeometry( DrawPosition, DrawSize ),	// Position, Size, Scale
+			AllottedGeometry.ToPaintGeometry( DrawSize, FSlateLayoutTransform(DrawPosition) ),	// Position, Size, Scale
 			StyleInfo,													// Style
 			DrawEffects,												// Effects to use
 			HighlightTargetColorAndOpacity );							// Color
diff --git a/Engine/Source/Editor/Sequencer/Private/SSequencerDebugVisualizer.cpp b/Engine/Source/Editor/Sequencer/Private/SSequencerDebugVisualizer.cpp
--- a/Engine/Source/Editor/Sequencer/Private/SSequencerDebugVisualizer.cpp
+++ b/Engine/Source/Editor/Sequencer/Private/SSequencerDebugVisualizer.cpp
@@ -163,7 +163,7 @@ FGeometry SSequencerEvaluationTemplateDebugVisualizer::GetSegmentGeometry(const
 	const FMovieSceneEvaluationField* EvaluationField = GetEvaluationField();
 	if (EvaluationField == nullptr || EvaluationField->GetSignature() != CachedSignature)
 	{
-		return AllottedGeometry.MakeChild(FVector2D(0,0), FVector2D(0,0));
+		return AllottedGeometry.MakeChild(FSlateRenderTransform());
 	}
 
 	TRange<FFrameNumber> SegmentRange = EvaluationField->GetRange(Slot.GetSegmentIndex());
@@ -175,8 +175,8 @@ FGeometry SSequencerEvaluationTemplateDebugVisualizer::GetSegmentGeometry(const
 	float SectionLength = FMath::Max(MinSectionWidth, PixelEndX - PixelStartX);
 
 	return AllottedGeometry.MakeChild(
-		FVector2D(PixelStartX, 0),
-		FVector2D(SectionLength, FMath::Max(Slot.GetDesiredSize().Y, 20.f))
+		FVector2D(SectionLength, FMath::Max(Slot.GetDesiredSize().Y, 20.f)),
+		FSlateLayoutTransform(FVector2D(PixelStartX, 0))
 		);
 }
 
@@ -555,8 +555,8 @@ FGeometry SSequencerEntityComponentSystemDebugVisualizer::GetSegmentGeometry(con
 	float SectionLength = FMath::Max(MinSectionWidth, PixelEndX - PixelStartX);
 
 	return AllottedGeometry.MakeChild(
-		FVector2D(PixelStartX, 0),
-		FVector2D(SectionLength, FMath::Max(Slot.GetDesiredSize().Y, 20.f))
+		FVector2D(SectionLength, FMath::Max(Slot.GetDesiredSize().Y, 20.f)),
+		FSlateLayoutTransform(FVector2D(PixelStartX, 0))
 		);
 }
 
diff --git a/Engine/Source/Editor/Sequencer/Private/SSequencerSection.cpp b/Engine/Source/Editor/Sequencer/Private/SSequencerSection.cpp
--- a/Engine/Source/Editor/Sequencer/Private/SSequencerSection.cpp
+++ b/Engine/Source/Editor/Sequencer/Private/SSequencerSection.cpp
@@ -226,7 +226,7 @@ struct FSequencerSectionPainterImpl : FSequencerSectionPainter
 				FSlateDrawElement::MakeBox(
 					DrawElements,
 					++LayerId,
-					ExpandedSectionGeometry.ToPaintGeometry(FVector2D(1, 1), ExpandedSectionGeometry.GetLocalSize() - FVector2D(2,2)),
+					ExpandedSectionGeometry.ToPaintGeometry(ExpandedSectionGeometry.GetLocalSize() - FVector2f(2.f,2.f), FSlateLayoutTransform(FVector2f(1.f, 1.f))),
 					CollapsedSelectedSectionOverlay,
 					DrawEffects,
 					SelectionColor.GetValue().CopyWithNewOpacity(0.8f)
@@ -274,7 +274,7 @@ struct FSequencerSectionPainterImpl : FSequencerSectionPainter
 				FSlateDrawElement::MakeBox(
 					DrawElements,
 					++LayerId,
-					HeaderGeometry.ToPaintGeometry(FVector2D(1, 1), HeaderGeometry.GetLocalSize() - FVector2D(2, 2)),
+					HeaderGeometry.ToPaintGeometry(HeaderGeometry.GetLocalSize() - FVector2f(2.f, 2.f), FSlateLayoutTransform(FVector2f(1.f, 1.f))),
 					SectionHeaderSelectedSectionOverlay,
 					DrawEffects,
 					SelectionColor.GetValue().CopyWithNewOpacity(0.8f)
@@ -286,7 +286,7 @@ struct FSequencerSectionPainterImpl : FSequencerSectionPainter
 				FSlateDrawElement::MakeBox(
 					DrawElements,
 					++LayerId,
-					ContentsGeometry.ToPaintGeometry(FVector2D(1, 1), ContentsGeometry.GetLocalSize() - FVector2D(2, 2)),
+					ContentsGeometry.ToPaintGeometry(ContentsGeometry.GetLocalSize() - FVector2f(2.f, 2.f), FSlateLayoutTransform(FVector2f(1.f, 1.f))),
 					SectionHeaderSelectedSectionOverlay,
 					DrawEffects,
 					FillSelectionColor.HSVToLinearRGB().CopyWithNewOpacity(0.8f)
@@ -1178,8 +1178,8 @@ bool SSequencerSection::CheckForEdgeInteraction( const FPointerEvent& MouseEvent
 		{
 			// Make areas to the left and right of the geometry.  We will use these areas to determine if someone dragged the left or right edge of a section
 			FGeometry SectionRectLeft = SectionGeometryWithoutHandles.MakeChild(
-				FVector2D( TimeToPixelConverter.FrameToPixel(UnderlappingSectionObj->GetInclusiveStartFrame()) - ThisHandleOffset, 0.f ),
-				GripSize
+				GripSize,
+				FSlateLayoutTransform(FVector2D( TimeToPixelConverter.FrameToPixel(UnderlappingSectionObj->GetInclusiveStartFrame()) - ThisHandleOffset, 0.f ))
 			);
 
 			if( SectionRectLeft.IsUnderLocation( MouseEvent.GetScreenSpacePosition() ) )
@@ -1192,8 +1192,8 @@ bool SSequencerSection::CheckForEdgeInteraction( const FPointerEvent& MouseEvent
 		if (UnderlappingSectionObj->HasEndFrame())
 		{
 			FGeometry SectionRectRight = SectionGeometryWithoutHandles.MakeChild(
-				FVector2D( TimeToPixelConverter.FrameToPixel(UnderlappingSectionObj->GetExclusiveEndFrame()) - UnderlappingSectionInterface->GetSectionGripSize() + ThisHandleOffset, 0 ), 
-				GripSize
+				GripSize,
+				FSlateLayoutTransform(FVector2D( TimeToPixelConverter.FrameToPixel(UnderlappingSectionObj->GetExclusiveEndFrame()) - UnderlappingSectionInterface->GetSectionGripSize() + ThisHandleOffset, 0 ))
 			);
 
 			if( SectionRectRight.IsUnderLocation( MouseEvent.GetScreenSpacePosition() ) )
@@ -1499,15 +1499,15 @@ void SSequencerSection::PaintEasingHandles( FSequencerSectionPainter& InPainter,
 			TRange<FFrameNumber> EaseInRange = UnderlappingSectionObj->GetEaseInRange();
 			// Always draw handles if the section is highlighted, even if there is no range (to allow manual adjustment)
 			FFrameNumber HandleFrame = EaseInRange.IsEmpty() ? UnderlappingSectionObj->GetInclusiveStartFrame() : UE::MovieScene::DiscreteExclusiveUpper(EaseInRange);
-			FVector2D HandlePos(TimeToPixelConverter.FrameToPixel(HandleFrame), 0.f);
+			FVector2f HandlePos(TimeToPixelConverter.FrameToPixel(HandleFrame), 0.f);
 			FSlateDrawElement::MakeBox(
 				InPainter.DrawElements,
 				// always draw selected keys on top of other keys
 				InPainter.LayerId,
 				// Center the key along X.  Ensure the middle of the key is at the actual key time
 				InPainter.SectionGeometry.ToPaintGeometry(
-					HandlePos - FVector2D(HandleSize.X*0.5f, 0.f),
-					HandleSize
+					HandleSize,
+					FSlateLayoutTransform(HandlePos - FVector2f(HandleSize.X*0.5f, 0.f))
 				),
 				EasingHandle,
 				DrawEffects,
@@ -1521,7 +1521,7 @@ void SSequencerSection::PaintEasingHandles( FSequencerSectionPainter& InPainter,
 
 			// Always draw handles if the section is highlighted, even if there is no range (to allow manual adjustment)
 			FFrameNumber HandleFrame = EaseOutRange.IsEmpty() ? UnderlappingSectionObj->GetExclusiveEndFrame() : UE::MovieScene::DiscreteInclusiveLower(EaseOutRange);
-			FVector2D    HandlePos   = FVector2D(TimeToPixelConverter.FrameToPixel(HandleFrame), 0.f);
+			FVector2f    HandlePos   = FVector2f(TimeToPixelConverter.FrameToPixel(HandleFrame), 0.f);
 
 			FSlateDrawElement::MakeBox(
 				InPainter.DrawElements,
@@ -1529,8 +1529,8 @@ void SSequencerSection::PaintEasingHandles( FSequencerSectionPainter& InPainter,
 				InPainter.LayerId,
 				// Center the key along X.  Ensure the middle of the key is at the actual key time
 				InPainter.SectionGeometry.ToPaintGeometry(
-					HandlePos - FVector2D(HandleSize.X*0.5f, 0.f),
-					HandleSize
+					HandleSize,
+					FSlateLayoutTransform(HandlePos - FVector2f(HandleSize.X*0.5f, 0.f))
 				),
 				EasingHandle,
 				DrawEffects,
@@ -1632,8 +1632,8 @@ void SSequencerSection::DrawSectionHandles( const FGeometry& AllottedGeometry, F
 		if (UnderlappingSectionObj->HasStartFrame())
 		{
 			FGeometry SectionRectLeft = SectionGeometryWithoutHandles.MakeChild(
-				FVector2D( TimeToPixelConverter.FrameToPixel(UnderlappingSectionObj->GetInclusiveStartFrame()) - ThisHandleOffset, 0.f ),
-				GripSize
+				GripSize,
+				FSlateLayoutTransform(FVector2D( TimeToPixelConverter.FrameToPixel(UnderlappingSectionObj->GetInclusiveStartFrame()) - ThisHandleOffset, 0.f ))
 			);
 			FSlateDrawElement::MakeBox
 			(
@@ -1650,8 +1650,8 @@ void SSequencerSection::DrawSectionHandles( const FGeometry& AllottedGeometry, F
 		if (UnderlappingSectionObj->HasEndFrame())
 		{
 			FGeometry SectionRectRight = SectionGeometryWithoutHandles.MakeChild(
-				FVector2D( TimeToPixelConverter.FrameToPixel(UnderlappingSectionObj->GetExclusiveEndFrame()) - UnderlappingSection->GetSectionGripSize() + ThisHandleOffset, 0 ), 
-				GripSize
+				GripSize,
+				FSlateLayoutTransform(FVector2D( TimeToPixelConverter.FrameToPixel(UnderlappingSectionObj->GetExclusiveEndFrame()) - UnderlappingSection->GetSectionGripSize() + ThisHandleOffset, 0 ))
 			);
 			FSlateDrawElement::MakeBox
 			(
diff --git a/Engine/Source/Editor/Sequencer/Private/SSequencerShotFilterOverlay.cpp b/Engine/Source/Editor/Sequencer/Private/SSequencerShotFilterOverlay.cpp
--- a/Engine/Source/Editor/Sequencer/Private/SSequencerShotFilterOverlay.cpp
+++ b/Engine/Source/Editor/Sequencer/Private/SSequencerShotFilterOverlay.cpp
@@ -42,7 +42,7 @@ int32 SSequencerShotFilterOverlay::OnPaint( const FPaintArgs& Args, const FGeome
 			FSlateDrawElement::MakeBox(
 				OutDrawElements,
 				LayerId,
-				AllottedGeometry.ToPaintGeometry(FVector2D(LowerBound, 0), FVector2D(UpperBound - LowerBound, AllottedGeometry.GetLocalSize().Y)),
+				AllottedGeometry.ToPaintGeometry(FVector2f(UpperBound - LowerBound, AllottedGeometry.GetLocalSize().Y), FSlateLayoutTransform(FVector2f(LowerBound, 0.f))),
 				FAppStyle::GetBrush("Sequencer.ShotFilter"),
 				ESlateDrawEffect::None,
 				FLinearColor(1.f, 1.f, 1.f, Alpha)
diff --git a/Engine/Source/Editor/Sequencer/Private/SequencerTimeSliderController.cpp b/Engine/Source/Editor/Sequencer/Private/SequencerTimeSliderController.cpp
--- a/Engine/Source/Editor/Sequencer/Private/SequencerTimeSliderController.cpp
+++ b/Engine/Source/Editor/Sequencer/Private/SequencerTimeSliderController.cpp
@@ -284,7 +284,7 @@ void FSequencerTimeSliderController::DrawTicks( FSlateWindowElementList& OutDraw
 			FSlateDrawElement::MakeText(
 				OutDrawElements,
 				InArgs.StartLayer+1, 
-				InArgs.AllottedGeometry.ToPaintGeometry( TextOffset, InArgs.AllottedGeometry.Size ), 
+				InArgs.AllottedGeometry.ToPaintGeometry( InArgs.AllottedGeometry.Size, FSlateLayoutTransform(TextOffset) ), 
 				FrameString, 
 				SmallLayoutFont,
 				InArgs.DrawEffects,
@@ -366,7 +366,7 @@ int32 FSequencerTimeSliderController::DrawMarkedFrames( const FGeometry& Allotte
 				FSlateDrawElement::MakeText(
 					DrawElements,
 					LayerId + 1,
-					AllottedGeometry.ToPaintGeometry(FVector2D(TextPosition, 0.f), TextSize),
+					AllottedGeometry.ToPaintGeometry(TextSize, FSlateLayoutTransform(FVector2f(TextPosition, 0.f))),
 					LabelString,
 					SmallLayoutFont,
 					DrawEffects,
@@ -474,7 +474,7 @@ int32 FSequencerTimeSliderController::OnPaintTimeSlider( bool bMirrorLabels, con
 		const float         HandleEnd     = ScrubMetrics.HandleRangePx.GetUpperBoundValue();
 
 		const int32 ArrowLayer = LayerId + 2;
-		FPaintGeometry MyGeometry =	AllottedGeometry.ToPaintGeometry( FVector2D( HandleStart, 0 ), FVector2D( HandleEnd - HandleStart, AllottedGeometry.Size.Y ) );
+		FPaintGeometry MyGeometry =	AllottedGeometry.ToPaintGeometry( FVector2f( HandleEnd - HandleStart, AllottedGeometry.Size.Y ), FSlateLayoutTransform(FVector2f( HandleStart, 0.f )) );
 		FLinearColor ScrubColor = InWidgetStyle.GetColorAndOpacityTint();
 		if(bIsEvaluating)
 		{
@@ -539,7 +539,7 @@ int32 FSequencerTimeSliderController::OnPaintTimeSlider( bool bMirrorLabels, con
 			FSlateDrawElement::MakeText(
 				OutDrawElements,
 				Args.StartLayer+1, 
-				Args.AllottedGeometry.ToPaintGeometry( TextOffset, TextSize ), 
+				Args.AllottedGeometry.ToPaintGeometry( TextSize, FSlateLayoutTransform(TextOffset) ), 
 				FrameString, 
 				SmallLayoutFont,
 				Args.DrawEffects,
@@ -565,7 +565,7 @@ int32 FSequencerTimeSliderController::OnPaintTimeSlider( bool bMirrorLabels, con
 			FSlateDrawElement::MakeBox(
 				OutDrawElements,
 				LayerId+1,
-				AllottedGeometry.ToPaintGeometry( FVector2D(RangePosX, 0.f), FVector2D(RangeSizeX, AllottedGeometry.Size.Y) ),
+				AllottedGeometry.ToPaintGeometry( FVector2f(RangeSizeX, AllottedGeometry.Size.Y), FSlateLayoutTransform(FVector2f(RangePosX, 0.f)) ),
 				bMirrorLabels ? VanillaScrubHandleDownBrush : VanillaScrubHandleUpBrush,
 				DrawEffects,
 				MouseStartPosX < MouseEndPosX ? FLinearColor(0.5f, 0.5f, 0.5f) : FLinearColor(0.25f, 0.3f, 0.3f)
@@ -600,7 +600,7 @@ int32 FSequencerTimeSliderController::DrawSelectionRange(const FGeometry& Allott
 			FSlateDrawElement::MakeBox(
 				OutDrawElements,
 				LayerId + 1,
-				AllottedGeometry.ToPaintGeometry(FVector2D(SelectionRangeL, 0.f), FVector2D(SelectionRangeR - SelectionRangeL, AllottedGeometry.Size.Y)),
+				AllottedGeometry.ToPaintGeometry(FVector2f(SelectionRangeR - SelectionRangeL, AllottedGeometry.Size.Y), FSlateLayoutTransform(FVector2f(SelectionRangeL, 0.f))),
 				FAppStyle::GetBrush("WhiteBrush"),
 				ESlateDrawEffect::None,
 				DrawColor.CopyWithNewOpacity(Args.SolidFillOpacity)
@@ -610,7 +610,7 @@ int32 FSequencerTimeSliderController::DrawSelectionRange(const FGeometry& Allott
 		FSlateDrawElement::MakeBox(
 			OutDrawElements,
 			LayerId + 1,
-			AllottedGeometry.ToPaintGeometry(FVector2D(SelectionRangeL, 0.f), FVector2D(Args.BrushWidth, AllottedGeometry.Size.Y)),
+			AllottedGeometry.ToPaintGeometry(FVector2f(Args.BrushWidth, AllottedGeometry.Size.Y), FSlateLayoutTransform(FVector2f(SelectionRangeL, 0.f))),
 			Args.StartBrush,
 			ESlateDrawEffect::None,
 			DrawColor
@@ -619,7 +619,7 @@ int32 FSequencerTimeSliderController::DrawSelectionRange(const FGeometry& Allott
 		FSlateDrawElement::MakeBox(
 			OutDrawElements,
 			LayerId + 1,
-			AllottedGeometry.ToPaintGeometry(FVector2D(SelectionRangeR - Args.BrushWidth, 0.f), FVector2D(Args.BrushWidth, AllottedGeometry.Size.Y)),
+			AllottedGeometry.ToPaintGeometry(FVector2f(Args.BrushWidth, AllottedGeometry.Size.Y), FSlateLayoutTransform(FVector2f(SelectionRangeR - Args.BrushWidth, 0.f))),
 			Args.EndBrush,
 			ESlateDrawEffect::None,
 			DrawColor
@@ -653,7 +653,7 @@ int32 FSequencerTimeSliderController::DrawPlaybackRange(const FGeometry& Allotte
 	FSlateDrawElement::MakeBox(
 		OutDrawElements,
 		LayerId+1,
-		AllottedGeometry.ToPaintGeometry(FVector2D(PlaybackRangeL, 0.f), FVector2D(Args.BrushWidth, AllottedGeometry.Size.Y)),
+		AllottedGeometry.ToPaintGeometry(FVector2f(Args.BrushWidth, AllottedGeometry.Size.Y), FSlateLayoutTransform(FVector2f(PlaybackRangeL, 0.f))),
 		Args.StartBrush,
 		ESlateDrawEffect::None,
 		FColor(32, 128, 32, OpacityBlend)	// 120, 75, 50 (HSV)
@@ -662,7 +662,7 @@ int32 FSequencerTimeSliderController::DrawPlaybackRange(const FGeometry& Allotte
 	FSlateDrawElement::MakeBox(
 		OutDrawElements,
 		LayerId+1,
-		AllottedGeometry.ToPaintGeometry(FVector2D(PlaybackRangeR - Args.BrushWidth, 0.f), FVector2D(Args.BrushWidth, AllottedGeometry.Size.Y)),
+		AllottedGeometry.ToPaintGeometry(FVector2f(Args.BrushWidth, AllottedGeometry.Size.Y), FSlateLayoutTransform(FVector2f(PlaybackRangeR - Args.BrushWidth, 0.f))),
 		Args.EndBrush,
 		ESlateDrawEffect::None,
 		FColor(128, 32, 32, OpacityBlend)	// 0, 75, 50 (HSV)
@@ -672,7 +672,7 @@ int32 FSequencerTimeSliderController::DrawPlaybackRange(const FGeometry& Allotte
 	FSlateDrawElement::MakeBox(
 		OutDrawElements,
 		LayerId+1,
-		AllottedGeometry.ToPaintGeometry(FVector2D(0.f, 0.f), FVector2D(PlaybackRangeL, AllottedGeometry.Size.Y)),
+		AllottedGeometry.ToPaintGeometry(FVector2f(PlaybackRangeL, AllottedGeometry.Size.Y), FSlateLayoutTransform()),
 		FAppStyle::GetBrush("WhiteBrush"),
 		ESlateDrawEffect::None,
 		FLinearColor::Black.CopyWithNewOpacity(0.3f * OpacityBlend / 255.f)
@@ -681,7 +681,7 @@ int32 FSequencerTimeSliderController::DrawPlaybackRange(const FGeometry& Allotte
 	FSlateDrawElement::MakeBox(
 		OutDrawElements,
 		LayerId+1,
-		AllottedGeometry.ToPaintGeometry(FVector2D(PlaybackRangeR, 0.f), FVector2D(AllottedGeometry.Size.X - PlaybackRangeR, AllottedGeometry.Size.Y)),
+		AllottedGeometry.ToPaintGeometry(FVector2f(AllottedGeometry.Size.X - PlaybackRangeR, AllottedGeometry.Size.Y), FSlateLayoutTransform(FVector2f(PlaybackRangeR, 0.f))),
 		FAppStyle::GetBrush("WhiteBrush"),
 		ESlateDrawEffect::None,
 		FLinearColor::Black.CopyWithNewOpacity(0.3f * OpacityBlend / 255.f)
@@ -720,7 +720,7 @@ int32 FSequencerTimeSliderController::DrawSubSequenceRange(const FGeometry& Allo
 	FSlateDrawElement::MakeBox(
 		OutDrawElements,
 		LayerId+1,
-		AllottedGeometry.ToPaintGeometry(FVector2D(SubSequenceRangeL, 0.f), FVector2D(Args.BrushWidth, AllottedGeometry.Size.Y)),
+		AllottedGeometry.ToPaintGeometry(FVector2f(Args.BrushWidth, AllottedGeometry.Size.Y), FSlateLayoutTransform(FVector2f(SubSequenceRangeL, 0.f))),
 		LineBrushL,
 		ESlateDrawEffect::None,
 		GreenTint
@@ -730,7 +730,7 @@ int32 FSequencerTimeSliderController::DrawSubSequenceRange(const FGeometry& Allo
 	FSlateDrawElement::MakeBox(
 		OutDrawElements,
 		LayerId+1,
-		AllottedGeometry.ToPaintGeometry(FVector2D(SubSequenceRangeR - Args.BrushWidth, 0.f), FVector2D(Args.BrushWidth, AllottedGeometry.Size.Y)),
+		AllottedGeometry.ToPaintGeometry(FVector2f(Args.BrushWidth, AllottedGeometry.Size.Y), FSlateLayoutTransform(FVector2f(SubSequenceRangeR - Args.BrushWidth, 0.f))),
 		LineBrushR,
 		ESlateDrawEffect::None,
 		RedTint
@@ -740,7 +740,7 @@ int32 FSequencerTimeSliderController::DrawSubSequenceRange(const FGeometry& Allo
 	FSlateDrawElement::MakeBox(
 		OutDrawElements,
 		LayerId+1,
-		AllottedGeometry.ToPaintGeometry(FVector2D(0.f, 0.f), FVector2D(SubSequenceRangeL, AllottedGeometry.Size.Y)),
+		AllottedGeometry.ToPaintGeometry(FVector2f(SubSequenceRangeL, AllottedGeometry.Size.Y), FSlateLayoutTransform()),
 		FAppStyle::GetBrush("WhiteBrush"),
 		ESlateDrawEffect::None,
 		FLinearColor::Black.CopyWithNewOpacity(0.3f)
@@ -749,7 +749,7 @@ int32 FSequencerTimeSliderController::DrawSubSequenceRange(const FGeometry& Allo
 	FSlateDrawElement::MakeBox(
 		OutDrawElements,
 		LayerId+1,
-		AllottedGeometry.ToPaintGeometry(FVector2D(SubSequenceRangeR, 0.f), FVector2D(AllottedGeometry.Size.X - SubSequenceRangeR, AllottedGeometry.Size.Y)),
+		AllottedGeometry.ToPaintGeometry(FVector2f(AllottedGeometry.Size.X - SubSequenceRangeR, AllottedGeometry.Size.Y), FSlateLayoutTransform(FVector2f(SubSequenceRangeR, 0.f))),
 		FAppStyle::GetBrush("WhiteBrush"),
 		ESlateDrawEffect::None,
 		FLinearColor::Black.CopyWithNewOpacity(0.3f)
@@ -759,7 +759,7 @@ int32 FSequencerTimeSliderController::DrawSubSequenceRange(const FGeometry& Allo
 	FSlateDrawElement::MakeBox(
 		OutDrawElements,
 		LayerId+1,
-		AllottedGeometry.ToPaintGeometry(FVector2D(SubSequenceRangeL - 16.f, 0.f), FVector2D(16.f, AllottedGeometry.Size.Y)),
+		AllottedGeometry.ToPaintGeometry(FVector2f(16.f, AllottedGeometry.Size.Y), FSlateLayoutTransform(FVector2f(SubSequenceRangeL - 16.f, 0.f))),
 		FAppStyle::GetBrush("Sequencer.Timeline.SubSequenceRangeHashL"),
 		ESlateDrawEffect::None,
 		GreenTint
@@ -768,7 +768,7 @@ int32 FSequencerTimeSliderController::DrawSubSequenceRange(const FGeometry& Allo
 	FSlateDrawElement::MakeBox(
 		OutDrawElements,
 		LayerId+1,
-		AllottedGeometry.ToPaintGeometry(FVector2D(SubSequenceRangeR, 0.f), FVector2D(16.f, AllottedGeometry.Size.Y)),
+		AllottedGeometry.ToPaintGeometry(FVector2f(16.f, AllottedGeometry.Size.Y), FSlateLayoutTransform(FVector2f(SubSequenceRangeR, 0.f))),
 		FAppStyle::GetBrush("Sequencer.Timeline.SubSequenceRangeHashR"),
 		ESlateDrawEffect::None,
 		RedTint
@@ -1307,7 +1307,7 @@ int32 FSequencerTimeSliderController::OnPaintViewArea( const FGeometry& Allotted
 			FSlateDrawElement::MakeBox(
 				OutDrawElements,
 				LayerId + 1,
-				AllottedGeometry.ToPaintGeometry(FVector2D(ScrubMetrics.FrameExtentsPx.GetLowerBoundValue(), 0.0f), FVector2D(ScrubMetrics.FrameExtentsPx.Size<float>(), AllottedGeometry.Size.Y)),
+				AllottedGeometry.ToPaintGeometry(FVector2f(ScrubMetrics.FrameExtentsPx.Size<float>(), AllottedGeometry.Size.Y), FSlateLayoutTransform(FVector2f(ScrubMetrics.FrameExtentsPx.GetLowerBoundValue(), 0.0f))),
 				ScrubFillBrush,
 				DrawEffects,
 				FLinearColor::White.CopyWithNewOpacity(0.5f)
diff --git a/Engine/Source/Editor/Sequencer/Private/Tools/SequencerEditTool_Movement.cpp b/Engine/Source/Editor/Sequencer/Private/Tools/SequencerEditTool_Movement.cpp
--- a/Engine/Source/Editor/Sequencer/Private/Tools/SequencerEditTool_Movement.cpp
+++ b/Engine/Source/Editor/Sequencer/Private/Tools/SequencerEditTool_Movement.cpp
@@ -89,7 +89,7 @@ FReply FSequencerEditTool_Movement::OnMouseMove(SWidget& OwnerWidget, const FGeo
 				double CurrentTime = VirtualTrackArea.PixelToSeconds(DragPosition.X);
 				Sequencer.UpdateAutoScroll(CurrentTime);
 
-				DragOperation->OnDrag(MouseEvent, DragPosition, VirtualTrackArea);
+				DragOperation->OnDrag(MouseEvent, FVector2D(DragPosition), VirtualTrackArea);
 			}
 		}
 		// Otherwise we can attempt a new drag
@@ -374,7 +374,7 @@ int32 FSequencerEditTool_Movement::OnPaint(const FGeometry& AllottedGeometry, co
 		FSlateDrawElement::MakeBox(
 			OutDrawElements,
 			++LayerId,
-			AllottedGeometry.ToPaintGeometry(DragPosition + FVector2D(5, -25), CursorDecorator->ImageSize),
+			AllottedGeometry.ToPaintGeometry(CursorDecorator->ImageSize, FSlateLayoutTransform(DragPosition + FVector2f(5.f, -25.f))),
 			CursorDecorator
 			);
 	}
@@ -416,7 +416,7 @@ int32 FSequencerEditTool_Movement::OnPaint(const FGeometry& AllottedGeometry, co
 				FSlateDrawElement::MakeLines(
 					OutDrawElements,
 					LayerId + 1,
-					AllottedGeometry.ToPaintGeometry(FVector2D(OldPos.X, 0.0f), FVector2D(1.0f, 1.0f)),
+					AllottedGeometry.ToPaintGeometry(FVector2f(1.0f, 1.0f), FSlateLayoutTransform(FVector2f(OldPos.X, 0.0f))),
 					LinePoints,
 					ESlateDrawEffect::None,
 					FLinearColor::White.CopyWithNewOpacity(0.5f),
@@ -427,7 +427,7 @@ int32 FSequencerEditTool_Movement::OnPaint(const FGeometry& AllottedGeometry, co
 				FSlateDrawElement::MakeLines(
 					OutDrawElements,
 					LayerId + 1,
-					AllottedGeometry.ToPaintGeometry(FVector2D(NewPos.X, 0.0f), FVector2D(1.0f, 1.0f)),
+					AllottedGeometry.ToPaintGeometry(FVector2f(1.0f, 1.0f), FSlateLayoutTransform(FVector2f(NewPos.X, 0.0f))),
 					LinePoints,
 					ESlateDrawEffect::None,
 					DrawColor,
@@ -442,7 +442,7 @@ int32 FSequencerEditTool_Movement::OnPaint(const FGeometry& AllottedGeometry, co
 				FSlateDrawElement::MakeBox( 
 					OutDrawElements,
 					LayerId + 2, 
-					AllottedGeometry.ToPaintGeometry(TimePos - BoxPadding, TimeStringSize + 2.0f * BoxPadding),
+					AllottedGeometry.ToPaintGeometry(TimeStringSize + 2.0f * BoxPadding, FSlateLayoutTransform(TimePos - BoxPadding)),
 					FAppStyle::GetBrush("WhiteBrush"),
 					ESlateDrawEffect::None, 
 					FLinearColor::Black.CopyWithNewOpacity(0.5f)
@@ -451,7 +451,7 @@ int32 FSequencerEditTool_Movement::OnPaint(const FGeometry& AllottedGeometry, co
 				FSlateDrawElement::MakeText(
 					OutDrawElements,
 					LayerId + 3,
-					AllottedGeometry.ToPaintGeometry(TimePos, TimeStringSize),
+					AllottedGeometry.ToPaintGeometry(TimeStringSize, FSlateLayoutTransform(TimePos)),
 					TimeString,
 					SmallLayoutFont,
 					ESlateDrawEffect::None,
@@ -467,7 +467,7 @@ int32 FSequencerEditTool_Movement::OnPaint(const FGeometry& AllottedGeometry, co
 				FSlateDrawElement::MakeBox( 
 					OutDrawElements,
 					LayerId + 2, 
-					AllottedGeometry.ToPaintGeometry(OffsetPos - BoxPadding, OffsetStringSize + 2.0f * BoxPadding),
+					AllottedGeometry.ToPaintGeometry(OffsetStringSize + 2.0f * BoxPadding, FSlateLayoutTransform(OffsetPos - BoxPadding)),
 					FAppStyle::GetBrush("WhiteBrush"),
 					ESlateDrawEffect::None, 
 					FLinearColor::Black.CopyWithNewOpacity(0.5f)
@@ -476,7 +476,7 @@ int32 FSequencerEditTool_Movement::OnPaint(const FGeometry& AllottedGeometry, co
 				FSlateDrawElement::MakeText(
 					OutDrawElements,
 					LayerId + 3,
-					AllottedGeometry.ToPaintGeometry(OffsetPos, TimeStringSize),
+					AllottedGeometry.ToPaintGeometry(TimeStringSize, FSlateLayoutTransform(OffsetPos)),
 					OffsetString,
 					SmallLayoutFont,
 					ESlateDrawEffect::None,
diff --git a/Engine/Source/Editor/Sequencer/Private/Tools/SequencerEditTool_Movement.h b/Engine/Source/Editor/Sequencer/Private/Tools/SequencerEditTool_Movement.h
--- a/Engine/Source/Editor/Sequencer/Private/Tools/SequencerEditTool_Movement.h
+++ b/Engine/Source/Editor/Sequencer/Private/Tools/SequencerEditTool_Movement.h
@@ -54,7 +54,7 @@ private:
 	TSharedPtr<UE::Sequencer::ISequencerEditToolDragOperation> DragOperation;
 
 	/** Current local position the mouse is dragged to. */
-	FVector2D DragPosition;
+	FVector2f DragPosition;
 
 	/** The hotspot's time before dragging started. */
 	FFrameTime OriginalHotspotTime;
diff --git a/Engine/Source/Editor/Sequencer/Private/Tools/SequencerEditTool_Selection.cpp b/Engine/Source/Editor/Sequencer/Private/Tools/SequencerEditTool_Selection.cpp
--- a/Engine/Source/Editor/Sequencer/Private/Tools/SequencerEditTool_Selection.cpp
+++ b/Engine/Source/Editor/Sequencer/Private/Tools/SequencerEditTool_Selection.cpp
@@ -266,7 +266,7 @@ public:
 		FSlateDrawElement::MakeBox(
 			OutDrawElements,
 			LayerId,
-			AllottedGeometry.ToPaintGeometry(SelectionTopLeft, SelectionBottomRight - SelectionTopLeft),
+			AllottedGeometry.ToPaintGeometry(SelectionBottomRight - SelectionTopLeft, FSlateLayoutTransform(SelectionTopLeft)),
 			FAppStyle::GetBrush(TEXT("MarqueeSelection"))
 			);
 
@@ -340,7 +340,7 @@ int32 FSequencerEditTool_Selection::OnPaint(const FGeometry& AllottedGeometry, c
 		FSlateDrawElement::MakeBox(
 			OutDrawElements,
 			++LayerId,
-			AllottedGeometry.ToPaintGeometry(MousePosition + FVector2D(5, 5), CursorDecorator->ImageSize),
+			AllottedGeometry.ToPaintGeometry(CursorDecorator->ImageSize, FSlateLayoutTransform(MousePosition + FVector2D(5.f, 5.f))),
 			CursorDecorator
 			);
 	}
diff --git a/Engine/Source/Editor/Sequencer/Private/Widgets/SSequencerLayerBar.cpp b/Engine/Source/Editor/Sequencer/Private/Widgets/SSequencerLayerBar.cpp
--- a/Engine/Source/Editor/Sequencer/Private/Widgets/SSequencerLayerBar.cpp
+++ b/Engine/Source/Editor/Sequencer/Private/Widgets/SSequencerLayerBar.cpp
@@ -328,7 +328,7 @@ int32 SSequencerLayerBar::OnPaint(const FPaintArgs& Args, const FGeometry& Allot
 		FSlateDrawElement::MakeBox(
 			OutDrawElements,
 			LayerId,
-			AllottedGeometry.ToPaintGeometry(FVector2D(1, 1), AllottedGeometry.GetLocalSize() - FVector2D(2,3)),
+			AllottedGeometry.ToPaintGeometry(AllottedGeometry.GetLocalSize() - FVector2f(2.f,3.f), FSlateLayoutTransform(FVector2f(1.f, 1.f))),
 			SelectedSectionOverlay,
 			DrawEffects,
 			SelectionColor.CopyWithNewOpacity(0.8f)
