AnimationAttributetab: added a missing scrollbar

diff --git a/Engine/Source/Editor/Persona/Private/SAnimAttributeView.cpp b/Engine/Source/Editor/Persona/Private/SAnimAttributeView.cpp
--- a/Engine/Source/Editor/Persona/Private/SAnimAttributeView.cpp
+++ b/Engine/Source/Editor/Persona/Private/SAnimAttributeView.cpp
@@ -12,6 +12,7 @@
 #include "Widgets/Views/SListView.h"
 #include "Framework/MultiBox/MultiBoxBuilder.h"
 #include "Styling/AppStyle.h"
+#include "Widgets/Layout/SScrollBox.h"
 
 #define LOCTEXT_NAMESPACE "SAnimAttributeView"
 
@@ -256,7 +257,7 @@ void SAnimAttributeView::Construct(const FArguments& InArgs)
 		.Value(0.4)
 		.MinSize(30.0f)
 		[
-			SAssignNew(ValueViewBox, SVerticalBox)
+			SAssignNew(ValueViewBox, SScrollBox)
 		]	
 	];
 }
@@ -606,7 +607,6 @@ void SAnimAttributeView::RefreshValueView()
 		const FAttributeValueView& ValueView = SelectedAttributeSnapshotValueViews[ViewIndex];
 		// slots are added in reverse order
 		ValueViewBox->AddSlot()
-		.AutoHeight()
 		[
 			ValueView.ViewWidget->GetWidget().ToSharedRef()
 		];	
diff --git a/Engine/Source/Editor/Persona/Public/SAnimAttributeView.h b/Engine/Source/Editor/Persona/Public/SAnimAttributeView.h
--- a/Engine/Source/Editor/Persona/Public/SAnimAttributeView.h
+++ b/Engine/Source/Editor/Persona/Public/SAnimAttributeView.h
@@ -14,7 +14,7 @@ class STableViewBase;
 class IStructureDetailsView;
 class FStructOnScope;
 class SHeaderRow;
-class SVerticalBox;
+class SScrollBox;
 
 namespace EColumnSortMode
 {
@@ -179,7 +179,7 @@ private:
 
 	
 	/** value view */
-	TSharedPtr<SVerticalBox> ValueViewBox;
+	TSharedPtr<SScrollBox> ValueViewBox;
 	bool bShouldRefreshValueView;
 	TOptional<FAnimAttributeEntry> SelectedAttribute;
 
