Fix crash when selecting curve after removing transform curves in an animation sequence

Curve editor only holds the currently-edited curves, so iterating over just those curves when removing doesnt remove all the curves that are referenced. To check the complete set, the tree also needs to be iterated over seperately.

diff --git a/Engine/Source/Editor/Persona/Private/SAnimSequenceCurveEditor.cpp b/Engine/Source/Editor/Persona/Private/SAnimSequenceCurveEditor.cpp
--- a/Engine/Source/Editor/Persona/Private/SAnimSequenceCurveEditor.cpp
+++ b/Engine/Source/Editor/Persona/Private/SAnimSequenceCurveEditor.cpp
@@ -570,15 +570,24 @@ void SAnimSequenceCurveEditor::AddCurve(const FText& InCurveDisplayName, const F
 
 void SAnimSequenceCurveEditor::RemoveCurve(const FSmartName& InName, ERawCurveTrackTypes InType, int32 InCurveIndex)
 {
+	for(const FCurveEditorTreeItemID& TreeItemID : CurveEditor->GetRootTreeItems())
+	{
+		FCurveEditorTreeItem& TreeItem = CurveEditor->GetTreeItem(TreeItemID);
+		TSharedPtr<FAnimSequenceCurveEditorItem> CurveItem = StaticCastSharedPtr<FAnimSequenceCurveEditorItem>(TreeItem.GetItem());
+		if(CurveItem->Name == InName && CurveItem->Type == InType && CurveItem->CurveIndex == InCurveIndex)
+		{
+			CurveEditor->RemoveTreeItem(TreeItemID);
+			break;
+		}
+	}
+	
 	for(const auto& CurvePair : CurveEditor->GetCurves())
 	{
 		FRichCurveEditorModelNamed* Model = static_cast<FRichCurveEditorModelNamed*>(CurvePair.Value.Get());
 		if(Model->Name == InName && Model->Type == InType && Model->CurveIndex == InCurveIndex)
 		{
 			// Cache ID to prevent use after release
-			const FCurveEditorTreeItemID TreeId = Model->TreeId; 
 			CurveEditor->RemoveCurve(CurvePair.Key);
-			CurveEditor->RemoveTreeItem(TreeId);
 			break;
 		}
 	}
