Fix crash when deleting morph target, also delete the raw data so the morph target get deleted properly

#rnx

diff --git a/Engine/Source/Editor/Persona/Private/SMorphTargetViewer.cpp b/Engine/Source/Editor/Persona/Private/SMorphTargetViewer.cpp
--- a/Engine/Source/Editor/Persona/Private/SMorphTargetViewer.cpp
+++ b/Engine/Source/Editor/Persona/Private/SMorphTargetViewer.cpp
@@ -16,6 +16,7 @@
 #include "GPUSkinCache.h"
 #include "Engine/RendererSettings.h"
 #include "IPersonaPreviewScene.h"
+#include "Rendering/SkeletalMeshLODImporterData.h"
 
 #define LOCTEXT_NAMESPACE "SMorphTargetViewer"
 
@@ -532,6 +533,27 @@ void SMorphTargetViewer::OnDeleteMorphTargets()
 			//Clean up override usage
 			AddMorphTargetOverride(SelectedRows[RowIndex]->Name, 0.0f, true);
 
+			//Remove the morph target from the raw import data
+			FSkeletalMeshImportData SkelMeshImportData;
+			SkeletalMesh->LoadLODImportedData(0, SkelMeshImportData);
+			int32 ToDeleteIndex = INDEX_NONE;
+			for (int32 MoprhTargetIndex = 0; MoprhTargetIndex < SkelMeshImportData.MorphTargetNames.Num(); ++MoprhTargetIndex)
+			{
+				if (SkelMeshImportData.MorphTargetNames[MoprhTargetIndex].Equals(SelectedRows[RowIndex]->Name.ToString()))
+				{
+					ToDeleteIndex = MoprhTargetIndex;
+					break;
+				}
+			}
+
+			if (ToDeleteIndex != INDEX_NONE)
+			{
+				SkelMeshImportData.MorphTargetNames.RemoveAt(ToDeleteIndex);
+				SkelMeshImportData.MorphTargetModifiedPoints.RemoveAt(ToDeleteIndex);
+				SkelMeshImportData.MorphTargets.RemoveAt(ToDeleteIndex);
+				SkeletalMesh->SaveLODImportedData(0, SkelMeshImportData);
+			}
+
 			SkeletalMesh->UnregisterMorphTarget(MorphTarget);
 		}
 	}
diff --git a/Engine/Source/Runtime/Engine/Private/MorphTools.cpp b/Engine/Source/Runtime/Engine/Private/MorphTools.cpp
--- a/Engine/Source/Runtime/Engine/Private/MorphTools.cpp
+++ b/Engine/Source/Runtime/Engine/Private/MorphTools.cpp
@@ -268,7 +268,28 @@ void FFinishBuildMorphTargetData::ApplyEditorData(USkeletalMesh * SkeletalMesh,
 			//which happen before the serialization of that cook skeletalmesh
 			if (!bIsSerializeSaving)
 			{
-				MorphTarget = NewObject<UMorphTarget>(SkeletalMesh, MorphTargetName);
+				//Avoid recycling morphtarget with NewObject it cannot be done asynchronously
+				//Find the UMorphTarget and simply clear the data if it exist.
+				TArray<UObject*> SubObjects;
+				GetObjectsWithOuter(SkeletalMesh, SubObjects, true);
+				for (UObject* SubObject : SubObjects)
+				{
+					if (SubObject->GetFName() == MorphTargetName)
+					{
+						if (UMorphTarget* SubMorphTarget = Cast<UMorphTarget>(SubObject))
+						{
+							MorphTarget = SubMorphTarget;
+							MorphTarget->EmptyMorphLODModels();
+							MorphTarget->ClearGarbage();
+							break;
+						}
+					}
+				}
+				//Create a new morph target, if the object do not exist (creating a new uobject is ok to do asynchronously)
+				if (!MorphTarget)
+				{
+					MorphTarget = NewObject<UMorphTarget>(SkeletalMesh, MorphTargetName);
+				}
 				check(MorphTarget);
 			}
 			else
