Persona: Remove magic code from preview customization and allow the various editors to explicitly disable skeleton match for preview meshes.

#rnx

diff --git a/Engine/Plugins/Animation/ControlRig/Source/ControlRigEditor/Private/Editor/ControlRigEditor.cpp b/Engine/Plugins/Animation/ControlRig/Source/ControlRigEditor/Private/Editor/ControlRigEditor.cpp
--- a/Engine/Plugins/Animation/ControlRig/Source/ControlRigEditor/Private/Editor/ControlRigEditor.cpp
+++ b/Engine/Plugins/Animation/ControlRig/Source/ControlRigEditor/Private/Editor/ControlRigEditor.cpp
@@ -335,6 +335,7 @@ void FControlRigEditor::InitControlRigEditor(const EToolkitMode::Type Mode, cons
 
 	FPersonaToolkitArgs PersonaToolkitArgs;
 	PersonaToolkitArgs.OnPreviewSceneCreated = FOnPreviewSceneCreated::FDelegate::CreateSP(this, &FControlRigEditor::HandlePreviewSceneCreated);
+	PersonaToolkitArgs.bPreviewMeshCanUseDifferentSkeleton = true;
 	PersonaToolkit = PersonaModule.CreatePersonaToolkit(InControlRigBlueprint, PersonaToolkitArgs);
 
 	// set delegate prior to setting mesh
diff --git a/Engine/Plugins/Animation/DeformerGraph/Source/OptimusEditor/Private/OptimusEditor.cpp b/Engine/Plugins/Animation/DeformerGraph/Source/OptimusEditor/Private/OptimusEditor.cpp
--- a/Engine/Plugins/Animation/DeformerGraph/Source/OptimusEditor/Private/OptimusEditor.cpp
+++ b/Engine/Plugins/Animation/DeformerGraph/Source/OptimusEditor/Private/OptimusEditor.cpp
@@ -89,6 +89,7 @@ void FOptimusEditor::Construct(
 	
 	FPersonaToolkitArgs PersonaToolkitArgs;
 	PersonaToolkitArgs.OnPreviewSceneCreated = FOnPreviewSceneCreated::FDelegate::CreateSP(this, &FOptimusEditor::HandlePreviewSceneCreated);
+	PersonaToolkitArgs.bPreviewMeshCanUseDifferentSkeleton = true;
 	PersonaToolkit = PersonaModule.CreatePersonaToolkit(InDeformerObject, PersonaToolkitArgs);
 
 	PersonaToolkit->GetPreviewScene()->SetDefaultAnimationMode(EPreviewSceneDefaultAnimationMode::Animation);
diff --git a/Engine/Source/Editor/Persona/Private/PersonaToolkit.cpp b/Engine/Source/Editor/Persona/Private/PersonaToolkit.cpp
--- a/Engine/Source/Editor/Persona/Private/PersonaToolkit.cpp
+++ b/Engine/Source/Editor/Persona/Private/PersonaToolkit.cpp
@@ -98,6 +98,11 @@ void FPersonaToolkit::Initialize(UAnimBlueprint* InAnimBlueprint, const FPersona
 	FindCounterpartAssets(InAnimBlueprint, Skeleton, Mesh);
 
 	CommonInitialSetup(PersonaToolkitArgs);
+
+	if (InAnimBlueprint->bIsTemplate)
+	{
+		bPreviewMeshCanUseDifferentSkeleton = true;
+	}
 }
 
 void FPersonaToolkit::Initialize(UPhysicsAsset* InPhysicsAsset, const FPersonaToolkitArgs& PersonaToolkitArgs)
@@ -119,6 +124,7 @@ void FPersonaToolkit::CommonInitialSetup(const FPersonaToolkitArgs& PersonaToolk
 	}
 
 	OnPreviewSceneSettingsCustomized = PersonaToolkitArgs.OnPreviewSceneSettingsCustomized;
+	bPreviewMeshCanUseDifferentSkeleton = PersonaToolkitArgs.bPreviewMeshCanUseDifferentSkeleton;
 }
 
 void FPersonaToolkit::CreatePreviewScene(const FPersonaToolkitArgs& PersonaToolkitArgs)
@@ -418,3 +424,9 @@ FName FPersonaToolkit::GetContext() const
 
 	return NAME_None;
 }
+
+
+bool FPersonaToolkit::CanPreviewMeshUseDifferentSkeleton() const
+{
+	return bPreviewMeshCanUseDifferentSkeleton;
+}
diff --git a/Engine/Source/Editor/Persona/Private/PersonaToolkit.h b/Engine/Source/Editor/Persona/Private/PersonaToolkit.h
--- a/Engine/Source/Editor/Persona/Private/PersonaToolkit.h
+++ b/Engine/Source/Editor/Persona/Private/PersonaToolkit.h
@@ -53,6 +53,7 @@ public:
 	virtual void SetCustomData(const int32 Key, const int32 CustomData) override;
 	virtual void CustomizeSceneSettings(IDetailLayoutBuilder& DetailBuilder) override;
 	virtual FName GetContext() const override;
+	virtual bool CanPreviewMeshUseDifferentSkeleton() const override;
 
 private:
 	/** Common initialization */
@@ -90,4 +91,7 @@ private:
 
 	/** The class of the initial asset we were created with */
 	UClass* InitialAssetClass;
+
+	/** A flag to indicate whether the preview mesh's skeleton can be incompatible with the editing skeleton */
+	bool bPreviewMeshCanUseDifferentSkeleton = false;
 };
diff --git a/Engine/Source/Editor/Persona/Private/PreviewSceneCustomizations.cpp b/Engine/Source/Editor/Persona/Private/PreviewSceneCustomizations.cpp
--- a/Engine/Source/Editor/Persona/Private/PreviewSceneCustomizations.cpp
+++ b/Engine/Source/Editor/Persona/Private/PreviewSceneCustomizations.cpp
@@ -190,10 +190,7 @@ void FPreviewSceneDescriptionCustomization::CustomizeDetails(IDetailLayoutBuilde
 			PreviewMeshName = SkeletalMeshProperty->GetPropertyDisplayName();
 		}
 
-		const bool bCanUseDifferentSkeleton =
-			(PersonaToolkit.Pin()->GetContext() == UPhysicsAsset::StaticClass()->GetFName()) ||
-			(PersonaToolkit.Pin()->GetContext() == TEXT("ControlRigBlueprint")) ||
-			(PersonaToolkit.Pin()->GetAnimBlueprint() != nullptr && PersonaToolkit.Pin()->GetAnimBlueprint()->bIsTemplate);
+		const bool bCanUseDifferentSkeleton = PersonaToolkit.Pin()->CanPreviewMeshUseDifferentSkeleton();
 
 		DetailBuilder.EditCategory("Mesh")
 		.AddProperty(SkeletalMeshProperty)
diff --git a/Engine/Source/Editor/Persona/Public/IPersonaToolkit.h b/Engine/Source/Editor/Persona/Public/IPersonaToolkit.h
--- a/Engine/Source/Editor/Persona/Public/IPersonaToolkit.h
+++ b/Engine/Source/Editor/Persona/Public/IPersonaToolkit.h
@@ -88,4 +88,7 @@ public:
 	
 	/** Get the context in which this toolkit is being used (usually the class name of the asset) */
 	virtual FName GetContext() const = 0;
+
+	/** Returns true if the preview mesh can use skeletal meshes that don't share the same skeleton as the one being edited. */
+	virtual bool CanPreviewMeshUseDifferentSkeleton() const = 0;
 };
diff --git a/Engine/Source/Editor/Persona/Public/PersonaModule.h b/Engine/Source/Editor/Persona/Public/PersonaModule.h
--- a/Engine/Source/Editor/Persona/Public/PersonaModule.h
+++ b/Engine/Source/Editor/Persona/Public/PersonaModule.h
@@ -83,6 +83,12 @@ struct FPersonaToolkitArgs
 	 * for the user to customize the layout however they wish. */
 	FOnPreviewSceneSettingsCustomized::FDelegate OnPreviewSceneSettingsCustomized;
 
+	/**
+	 * Set to true if the preview mesh can be associated with a skeleton different from the one being inspected
+	 * by Persona. Used for editors that are mostly skeleton agnostic.
+	 */
+	bool bPreviewMeshCanUseDifferentSkeleton = false;
+
 	FPersonaToolkitArgs() = default;
 };
 
diff --git a/Engine/Source/Editor/PhysicsAssetEditor/Private/PhysicsAssetEditor.cpp b/Engine/Source/Editor/PhysicsAssetEditor/Private/PhysicsAssetEditor.cpp
--- a/Engine/Source/Editor/PhysicsAssetEditor/Private/PhysicsAssetEditor.cpp
+++ b/Engine/Source/Editor/PhysicsAssetEditor/Private/PhysicsAssetEditor.cpp
@@ -183,6 +183,7 @@ void FPhysicsAssetEditor::InitPhysicsAssetEditor(const EToolkitMode::Type Mode,
 	FPersonaToolkitArgs PersonaToolkitArgs;
 	PersonaToolkitArgs.OnPreviewSceneCreated = FOnPreviewSceneCreated::FDelegate::CreateSP(this, &FPhysicsAssetEditor::HandlePreviewSceneCreated);
 	PersonaToolkitArgs.OnPreviewSceneSettingsCustomized = FOnPreviewSceneSettingsCustomized::FDelegate::CreateSP(this, &FPhysicsAssetEditor::HandleOnPreviewSceneSettingsCustomized);
+	PersonaToolkitArgs.bPreviewMeshCanUseDifferentSkeleton = true;
 
 	FPersonaModule& PersonaModule = FModuleManager::LoadModuleChecked<FPersonaModule>("Persona");
 	PersonaToolkit = PersonaModule.CreatePersonaToolkit(SharedData->PhysicsAsset, PersonaToolkitArgs);
