Make sure we can delete morph target on asset imported before 4.24

#rnx

diff --git a/Engine/Source/Editor/Persona/Private/SMorphTargetViewer.cpp b/Engine/Source/Editor/Persona/Private/SMorphTargetViewer.cpp
--- a/Engine/Source/Editor/Persona/Private/SMorphTargetViewer.cpp
+++ b/Engine/Source/Editor/Persona/Private/SMorphTargetViewer.cpp
@@ -533,25 +533,33 @@ void SMorphTargetViewer::OnDeleteMorphTargets()
 			//Clean up override usage
 			AddMorphTargetOverride(SelectedRows[RowIndex]->Name, 0.0f, true);
 
-			//Remove the morph target from the raw import data
-			FSkeletalMeshImportData SkelMeshImportData;
-			SkeletalMesh->LoadLODImportedData(0, SkelMeshImportData);
-			int32 ToDeleteIndex = INDEX_NONE;
-			for (int32 MoprhTargetIndex = 0; MoprhTargetIndex < SkelMeshImportData.MorphTargetNames.Num(); ++MoprhTargetIndex)
+			if (!SkeletalMesh->IsLODImportedDataEmpty(0) && SkeletalMesh->IsLODImportedDataBuildAvailable(0))
 			{
-				if (SkelMeshImportData.MorphTargetNames[MoprhTargetIndex].Equals(SelectedRows[RowIndex]->Name.ToString()))
+				//Remove the morph target from the raw import data
+				FSkeletalMeshImportData SkelMeshImportData;
+				SkeletalMesh->LoadLODImportedData(0, SkelMeshImportData);
+				int32 ToDeleteIndex = INDEX_NONE;
+				for (int32 MoprhTargetIndex = 0; MoprhTargetIndex < SkelMeshImportData.MorphTargetNames.Num(); ++MoprhTargetIndex)
 				{
-					ToDeleteIndex = MoprhTargetIndex;
-					break;
+					if (SkelMeshImportData.MorphTargetNames[MoprhTargetIndex].Equals(SelectedRows[RowIndex]->Name.ToString()))
+					{
+						ToDeleteIndex = MoprhTargetIndex;
+						break;
+					}
 				}
-			}
 
-			if (ToDeleteIndex != INDEX_NONE)
+				if (ToDeleteIndex != INDEX_NONE)
+				{
+					SkelMeshImportData.MorphTargetNames.RemoveAt(ToDeleteIndex);
+					SkelMeshImportData.MorphTargetModifiedPoints.RemoveAt(ToDeleteIndex);
+					SkelMeshImportData.MorphTargets.RemoveAt(ToDeleteIndex);
+					SkeletalMesh->SaveLODImportedData(0, SkelMeshImportData);
+				}
+			}
+			else
 			{
-				SkelMeshImportData.MorphTargetNames.RemoveAt(ToDeleteIndex);
-				SkelMeshImportData.MorphTargetModifiedPoints.RemoveAt(ToDeleteIndex);
-				SkelMeshImportData.MorphTargets.RemoveAt(ToDeleteIndex);
-				SkeletalMesh->SaveLODImportedData(0, SkelMeshImportData);
+				//If we deal with an old asset (pre 4.24) and we do not have some valid import data, we need to dirty the ddc key so it wont take the ddc with the old morph target we just delete
+				SkeletalMesh->InvalidateDeriveDataCacheGUID();
 			}
 
 			SkeletalMesh->UnregisterMorphTarget(MorphTarget);
