Make ScreenSize update in MobilePreview for the selected LOD

diff --git a/Engine/Source/Editor/Persona/Private/AnimationEditorViewportClient.cpp b/Engine/Source/Editor/Persona/Private/AnimationEditorViewportClient.cpp
--- a/Engine/Source/Editor/Persona/Private/AnimationEditorViewportClient.cpp
+++ b/Engine/Source/Editor/Persona/Private/AnimationEditorViewportClient.cpp
@@ -2340,7 +2340,13 @@ void FAnimationViewportClient::SetupViewForRendering( FSceneViewFamily& ViewFami
 	if (PreviewMeshComponent != nullptr && PreviewMeshComponent->MeshObject != nullptr)
 	{
 		const FBoxSphereBounds& SkelBounds = PreviewMeshComponent->Bounds;
-		CachedScreenSize = ComputeBoundsScreenSize(SkelBounds.Origin, static_cast<float>(SkelBounds.SphereRadius), View);
+
+		const int32 LODIndex = PreviewMeshComponent->GetPredictedLODLevel();
+		USkinnedAsset* SkinnedAsset = PreviewMeshComponent->GetSkinnedAsset();
+		check(SkinnedAsset);
+
+		FPerPlatformFloat ScreenSize = SkinnedAsset->GetLODInfo(LODIndex)->ScreenSize;
+		CachedScreenSize = SkinnedAsset->IsValidLODIndex(LODIndex) ? ScreenSize.GetValue(): CachedScreenSize;
 	}
 }
 
