Avoid copying TMap every frame

diff --git a/Engine/Source/Editor/Persona/Private/AnimationEditorViewportClient.cpp b/Engine/Source/Editor/Persona/Private/AnimationEditorViewportClient.cpp
--- a/Engine/Source/Editor/Persona/Private/AnimationEditorViewportClient.cpp
+++ b/Engine/Source/Editor/Persona/Private/AnimationEditorViewportClient.cpp
@@ -2345,7 +2345,7 @@ void FAnimationViewportClient::SetupViewForRendering( FSceneViewFamily& ViewFami
 		USkinnedAsset* SkinnedAsset = PreviewMeshComponent->GetSkinnedAsset();
 		check(SkinnedAsset);
 
-		FPerPlatformFloat ScreenSize = SkinnedAsset->GetLODInfo(LODIndex)->ScreenSize;
+		FPerPlatformFloat& ScreenSize = SkinnedAsset->GetLODInfo(LODIndex)->ScreenSize;
 		CachedScreenSize = SkinnedAsset->IsValidLODIndex(LODIndex) ? ScreenSize.GetValue(): CachedScreenSize;
 	}
 }
