"Allow incompatible skeletons" checkbox option doesn't retain after UE reopening
#fix add set-ter which ensures config is saved out to disk

diff --git a/Engine/Source/Editor/Persona/Private/PersonaModule.cpp b/Engine/Source/Editor/Persona/Private/PersonaModule.cpp
--- a/Engine/Source/Editor/Persona/Private/PersonaModule.cpp
+++ b/Engine/Source/Editor/Persona/Private/PersonaModule.cpp
@@ -1712,7 +1712,7 @@ void FPersonaModule::RegisterToolMenuExtensions()
 								FExecuteAction::CreateLambda([WeakAssetView = TWeakPtr<SAssetView>(AssetView)]()
 								{
 									UPersonaOptions* PersonaOptions = GetMutableDefault<UPersonaOptions>();
-									PersonaOptions->bAllowIncompatibleSkeletonSelection = !PersonaOptions->bAllowIncompatibleSkeletonSelection;
+									PersonaOptions->SetAllowIncompatibleSkeletonSelection(!PersonaOptions->GetAllowIncompatibleSkeletonSelection());
 									if(TSharedPtr<SAssetView> AssetView = WeakAssetView.Pin())
 									{
 										AssetView->RequestSlowFullListRefresh();
diff --git a/Engine/Source/Editor/UnrealEd/Classes/Preferences/PersonaOptions.h b/Engine/Source/Editor/UnrealEd/Classes/Preferences/PersonaOptions.h
--- a/Engine/Source/Editor/UnrealEd/Classes/Preferences/PersonaOptions.h
+++ b/Engine/Source/Editor/UnrealEd/Classes/Preferences/PersonaOptions.h
@@ -269,6 +269,7 @@ public:
 	void SetBranchingPointTimingNodeColor(const FLinearColor& InColor);
 	FAssetEditorOptions& GetAssetEditorOptions(const FName& InContext);
 	bool GetAllowIncompatibleSkeletonSelection() const;
+	void SetAllowIncompatibleSkeletonSelection(bool bState);
 
 	DECLARE_MULTICAST_DELEGATE_TwoParams(FOnUpdateSettingsMulticaster, const UPersonaOptions*, EPropertyChangeType::Type);
 	FOnUpdateSettingsMulticaster OnSettingsChange;
diff --git a/Engine/Source/Editor/UnrealEd/Private/PreferenceStubs.cpp b/Engine/Source/Editor/UnrealEd/Private/PreferenceStubs.cpp
--- a/Engine/Source/Editor/UnrealEd/Private/PreferenceStubs.cpp
+++ b/Engine/Source/Editor/UnrealEd/Private/PreferenceStubs.cpp
@@ -303,4 +303,10 @@ void UPersonaOptions::PostEditChangeProperty(struct FPropertyChangedEvent& Prope
 bool UPersonaOptions::GetAllowIncompatibleSkeletonSelection() const
 {
 	return bAllowIncompatibleSkeletonSelection;
-}
\ No newline at end of file
+}
+
+void UPersonaOptions::SetAllowIncompatibleSkeletonSelection(bool bState)
+{
+	bAllowIncompatibleSkeletonSelection = bState;
+	SaveConfig();
+}
