Guard against use of invalid actor pointer in Persona

diff --git a/Engine/Source/Editor/Persona/Private/AnimationEditorPreviewScene.cpp b/Engine/Source/Editor/Persona/Private/AnimationEditorPreviewScene.cpp
--- a/Engine/Source/Editor/Persona/Private/AnimationEditorPreviewScene.cpp
+++ b/Engine/Source/Editor/Persona/Private/AnimationEditorPreviewScene.cpp
@@ -129,17 +129,25 @@ FAnimationEditorPreviewScene::~FAnimationEditorPreviewScene()
 TArray<UDebugSkelMeshComponent*> FAnimationEditorPreviewScene::GetAllPreviewMeshComponents() const
 {
 	TArray<UDebugSkelMeshComponent*> PreviewMeshComponents;
-	GetActor()->GetComponents(PreviewMeshComponents, true);
+	const AActor* MyActor = GetActor();
+	if (ensure(MyActor))
+	{
+		MyActor->GetComponents(PreviewMeshComponents, true);
+	}
 	return PreviewMeshComponents;
 }
 
 void FAnimationEditorPreviewScene::ForEachPreviewMesh(TFunction<void (UDebugSkelMeshComponent*)> PerMeshFunction)
 {
 	TArray<UDebugSkelMeshComponent*> PreviewMeshComponents;
-	GetActor()->GetComponents(PreviewMeshComponents, true);
-	for (UDebugSkelMeshComponent* PreviewMesh : PreviewMeshComponents)
+	const AActor* MyActor = GetActor();
+	if (ensure(MyActor))
 	{
-		PerMeshFunction(PreviewMesh);
+		MyActor->GetComponents(PreviewMeshComponents, true);
+		for (UDebugSkelMeshComponent* PreviewMesh : PreviewMeshComponents)
+		{
+			PerMeshFunction(PreviewMesh);
+		}
 	}
 }
 
