Removed AnimTimelineTrack_MontageSections.cpp/h since they are not being used/referenced

#rnx

diff --git a/Engine/Source/Editor/Persona/Private/AnimTimeline/AnimModel_AnimMontage.cpp b/Engine/Source/Editor/Persona/Private/AnimTimeline/AnimModel_AnimMontage.cpp
--- a/Engine/Source/Editor/Persona/Private/AnimTimeline/AnimModel_AnimMontage.cpp
+++ b/Engine/Source/Editor/Persona/Private/AnimTimeline/AnimModel_AnimMontage.cpp
@@ -12,7 +12,6 @@
 #include "IPersonaPreviewScene.h"
 #include "Animation/DebugSkelMeshComponent.h"
 #include "AnimPreviewInstance.h"
-#include "AnimTimeline/AnimTimelineTrack_MontageSections.h"
 #include "AnimTimeline/AnimTimelineTrack_Montage.h"
 #include "SAnimMontagePanel.h"
 #include "IEditableSkeleton.h"
diff --git a/Engine/Source/Editor/Persona/Private/AnimTimeline/AnimTimelineTrack_MontageSections.cpp b/Engine/Source/Editor/Persona/Private/AnimTimeline/AnimTimelineTrack_MontageSections.cpp
deleted file mode 100644
diff --git a/Engine/Source/Editor/Persona/Private/AnimTimeline/AnimTimelineTrack_MontageSections.h b/Engine/Source/Editor/Persona/Private/AnimTimeline/AnimTimelineTrack_MontageSections.h
deleted file mode 100644
