Show the correct preview mesh in the animation editor when switching animations

diff --git a/Engine/Source/Editor/Persona/Private/PersonaToolkit.cpp b/Engine/Source/Editor/Persona/Private/PersonaToolkit.cpp
--- a/Engine/Source/Editor/Persona/Private/PersonaToolkit.cpp
+++ b/Engine/Source/Editor/Persona/Private/PersonaToolkit.cpp
@@ -244,12 +244,32 @@ UAnimationAsset* FPersonaToolkit::GetAnimationAsset() const
 
 void FPersonaToolkit::SetAnimationAsset(class UAnimationAsset* InAnimationAsset)
 {
+	USkeleton* PreviousAnimSkeleton = (AnimationAsset != nullptr) ? AnimationAsset->GetSkeleton() : nullptr;
+
 	if (InAnimationAsset != nullptr)
 	{
 		ensure(Skeleton->IsCompatibleForEditor(InAnimationAsset->GetSkeleton()));
 	}
 
 	AnimationAsset = InAnimationAsset;
+	
+	if(AnimationAsset)
+	{
+		check(InitialAssetClass == UAnimationAsset::StaticClass());
+
+		USkeletalMesh* NewPreviewMesh = GetPreviewMesh();
+		const USkeleton* CurrentAnimSkeleton = AnimationAsset->GetSkeleton();
+
+		if (NewPreviewMesh == nullptr && PreviousAnimSkeleton != CurrentAnimSkeleton)
+		{
+			NewPreviewMesh = CurrentAnimSkeleton->GetPreviewMesh();
+		}
+
+		if (NewPreviewMesh)
+		{
+			GetPreviewScene()->SetPreviewMesh(NewPreviewMesh, false);
+		}
+	}
 }
 
 TSharedRef<IPersonaPreviewScene> FPersonaToolkit::GetPreviewScene() const
