Remove mentions of legacy style "Sequencer.Toolbar"

diff --git a/Engine/Source/Editor/CurveTableEditor/Private/CurveTableEditor.cpp b/Engine/Source/Editor/CurveTableEditor/Private/CurveTableEditor.cpp
--- a/Engine/Source/Editor/CurveTableEditor/Private/CurveTableEditor.cpp
+++ b/Engine/Source/Editor/CurveTableEditor/Private/CurveTableEditor.cpp
@@ -980,7 +980,6 @@ TSharedRef<SWidget> FCurveTableEditor::MakeToolbar(TSharedRef<SCurveEditorPanel>
 {
 
 	FToolBarBuilder ToolBarBuilder(InEditorPanel->GetCommands(), FMultiBoxCustomization::None, InEditorPanel->GetToolbarExtender(), true);
-	ToolBarBuilder.SetStyle(&FAppStyle::Get(), "Sequencer.ToolBar");
 	ToolBarBuilder.BeginSection("Asset");
 	ToolBarBuilder.EndSection();
 	// We just use all of the extenders as our toolbar, we don't have a need to create a separate toolbar.
diff --git a/Engine/Source/Editor/Persona/Private/SAnimSequenceCurveEditor.cpp b/Engine/Source/Editor/Persona/Private/SAnimSequenceCurveEditor.cpp
--- a/Engine/Source/Editor/Persona/Private/SAnimSequenceCurveEditor.cpp
+++ b/Engine/Source/Editor/Persona/Private/SAnimSequenceCurveEditor.cpp
@@ -483,7 +483,6 @@ void SAnimSequenceCurveEditor::OnModelHasChanged(const EAnimDataModelNotifyType&
 TSharedRef<SWidget> SAnimSequenceCurveEditor::MakeToolbar(TSharedRef<SCurveEditorPanel> InEditorPanel)
 {
 	FToolBarBuilder ToolBarBuilder(InEditorPanel->GetCommands(), FMultiBoxCustomization::None, InEditorPanel->GetToolbarExtender(), true);
-	ToolBarBuilder.SetStyle(&FAppStyle::Get(), "Sequencer.ToolBar");
 	ToolBarBuilder.BeginSection("Asset");
 	ToolBarBuilder.EndSection();
 	// We just use all of the extenders as our toolbar, we don't have a need to create a separate toolbar.
