Persona - No option to view Just the Skeleton Reference Pose
#feature when previewing an animation asset with no preview Skeletal Mesh set + having Show All hierarchy enabled we now show the skeleton reference pose instead of an empty viewport

diff --git a/Engine/Source/Editor/Persona/Private/AnimationEditorViewportClient.cpp b/Engine/Source/Editor/Persona/Private/AnimationEditorViewportClient.cpp
--- a/Engine/Source/Editor/Persona/Private/AnimationEditorViewportClient.cpp
+++ b/Engine/Source/Editor/Persona/Private/AnimationEditorViewportClient.cpp
@@ -39,6 +39,7 @@
 #include "CameraController.h"
 #include "ContextObjectStore.h"
 #include "EdModeInteractiveToolsContext.h"
+#include "IPersonaToolkit.h"
 #include "Animation/MorphTarget.h"
 #include "Rendering/SkeletalMeshModel.h"
 #include "UnrealWidget.h"
@@ -526,7 +527,17 @@ void FAnimationViewportClient::Draw(const FSceneView* View, FPrimitiveDrawInterf
 	TArray<UDebugSkelMeshComponent*> PreviewMeshComponents = GetPreviewScene()->GetAllPreviewMeshComponents();
 	for (UDebugSkelMeshComponent* PreviewMeshComponent : PreviewMeshComponents)
 	{
-		if (!(PreviewMeshComponent && PreviewMeshComponent->GetSkeletalMeshAsset() && !PreviewMeshComponent->GetSkeletalMeshAsset()->IsCompiling()))
+		const bool bValidComponent = PreviewMeshComponent != nullptr;
+		const bool bValidSkeletalMesh = PreviewMeshComponent->GetSkeletalMeshAsset() != nullptr;
+
+		if (bValidComponent && !bValidSkeletalMesh && GetBoneDrawMode() == EBoneDrawMode::All)
+		{
+			if (const USkeleton* Skeleton = GetPreviewScene()->GetPersonaToolkit()->GetSkeleton())
+			{
+				DrawBonesFromSkeleton(Skeleton, PDI);
+			}
+		}
+		else if (!(bValidComponent && bValidSkeletalMesh && !PreviewMeshComponent->GetSkeletalMeshAsset()->IsCompiling()))
 		{
 			continue;
 		}
@@ -1643,6 +1654,57 @@ void FAnimationViewportClient::DrawMeshBonesBakedAnimation(UDebugSkelMeshCompone
 	}
 }
 
+void FAnimationViewportClient::DrawBonesFromSkeleton(const USkeleton* Skeleton, FPrimitiveDrawInterface* PDI) const
+{
+	check(Skeleton);
+
+	// Draw from skeleton ref pose
+	const TArray<FTransform>& SkeletonRefPose = Skeleton->GetRefLocalPoses(NAME_None);
+	TArray<FTransform> WorldTransforms;
+	WorldTransforms.AddUninitialized(SkeletonRefPose.Num());
+
+	TArray<FLinearColor> BoneColours;
+	BoneColours.AddUninitialized(SkeletonRefPose.Num());
+
+	TArray<FBoneIndexType> RequiredBones;
+
+	const FReferenceSkeleton& RefSkeleton = Skeleton->GetReferenceSkeleton();
+
+	const UPersonaOptions* PersonaOptions = GetDefault<UPersonaOptions>();
+	for (FBoneIndexType BoneIndex = 0; BoneIndex < SkeletonRefPose.Num(); ++BoneIndex)
+	{
+		const int32 ParentIndex = RefSkeleton.GetParentIndex(BoneIndex);
+
+		//add to the list
+		RequiredBones.AddUnique(BoneIndex);
+
+		if (ParentIndex >= 0)
+		{
+			WorldTransforms[BoneIndex] = SkeletonRefPose[BoneIndex] * WorldTransforms[ParentIndex];
+		}
+		else
+		{
+			WorldTransforms[BoneIndex] = SkeletonRefPose[BoneIndex];
+		}
+
+		BoneColours[BoneIndex] = PersonaOptions->DefaultBoneColor;
+	}
+
+	constexpr bool bForceDraw = true;
+	constexpr bool bAddHitProxy = false;
+	const TArray<int32> BonesOfInterest;
+	DrawBones(
+		FVector::ZeroVector,
+		RequiredBones,
+		RefSkeleton,
+		WorldTransforms,
+		BonesOfInterest,
+		BoneColours,
+		PDI,
+		bForceDraw,
+		bAddHitProxy);
+}
+
 void FAnimationViewportClient::DrawMeshBones(UDebugSkelMeshComponent* MeshComponent, FPrimitiveDrawInterface* PDI) const
 {
 	if (!MeshComponent ||
diff --git a/Engine/Source/Editor/Persona/Private/AnimationEditorViewportClient.h b/Engine/Source/Editor/Persona/Private/AnimationEditorViewportClient.h
--- a/Engine/Source/Editor/Persona/Private/AnimationEditorViewportClient.h
+++ b/Engine/Source/Editor/Persona/Private/AnimationEditorViewportClient.h
@@ -418,6 +418,8 @@ private:
 	void DrawMeshBonesSourceRawAnimation(UDebugSkelMeshComponent * MeshComponent, FPrimitiveDrawInterface* PDI) const;
 	/** Draw Bones for non retargeted animation. */
 	void DrawMeshBonesBakedAnimation(UDebugSkelMeshComponent * MeshComponent, FPrimitiveDrawInterface* PDI) const;
+	/** Draw Bones from skeleton reference pose. */
+	void DrawBonesFromSkeleton(const USkeleton* Skeleton, FPrimitiveDrawInterface* PDI) const;
 	/** Draws Bones for RequiredBones with WorldTransform **/
 	void DrawBones(
 		const FVector& ComponentOrigin,
