Rejig viewport bone drawsing code to satisfy static analysis and simplify flow

diff --git a/Engine/Source/Editor/Persona/Private/AnimationEditorViewportClient.cpp b/Engine/Source/Editor/Persona/Private/AnimationEditorViewportClient.cpp
--- a/Engine/Source/Editor/Persona/Private/AnimationEditorViewportClient.cpp
+++ b/Engine/Source/Editor/Persona/Private/AnimationEditorViewportClient.cpp
@@ -528,78 +528,76 @@ void FAnimationViewportClient::Draw(const FSceneView* View, FPrimitiveDrawInterf
 	for (UDebugSkelMeshComponent* PreviewMeshComponent : PreviewMeshComponents)
 	{
 		const bool bValidComponent = PreviewMeshComponent != nullptr;
-		const bool bValidSkeletalMesh = PreviewMeshComponent->GetSkeletalMeshAsset() != nullptr;
+		const bool bValidSkeletalMesh = bValidComponent && PreviewMeshComponent->GetSkeletalMeshAsset() != nullptr;
 
-		if (bValidComponent && !bValidSkeletalMesh && GetBoneDrawMode() == EBoneDrawMode::All)
+		if (bValidComponent && bValidSkeletalMesh && !PreviewMeshComponent->GetSkeletalMeshAsset()->IsCompiling())
 		{
-			if (const USkeleton* Skeleton = GetPreviewScene()->GetPersonaToolkit()->GetSkeleton())
-			{
-				DrawBonesFromSkeleton(Skeleton, PDI);
-			}
-		}
-		else if (!(bValidComponent && bValidSkeletalMesh && !PreviewMeshComponent->GetSkeletalMeshAsset()->IsCompiling()))
-		{
-			continue;
-		}
-
-		// Can't have both bones of interest and sockets of interest set
-		check( !(GetAnimPreviewScene()->GetSelectedBoneIndex() != INDEX_NONE && GetAnimPreviewScene()->GetSelectedSocket().IsValid() ) )
+			// Can't have both bones of interest and sockets of interest set
+			check( !(GetAnimPreviewScene()->GetSelectedBoneIndex() != INDEX_NONE && GetAnimPreviewScene()->GetSelectedSocket().IsValid() ) )
 
-		const FReferenceSkeleton& RefSkeleton = PreviewMeshComponent->GetReferenceSkeleton();
-		const TArray<FBoneIndexType>& DrawBoneIndices = PreviewMeshComponent->GetDrawBoneIndices();
+			const FReferenceSkeleton& RefSkeleton = PreviewMeshComponent->GetReferenceSkeleton();
+			const TArray<FBoneIndexType>& DrawBoneIndices = PreviewMeshComponent->GetDrawBoneIndices();
 
-		// if we have BonesOfInterest, draw sub set of the bones only
-		if (GetAnimPreviewScene()->GetSelectedBoneIndex() != INDEX_NONE)
-		{
-			DrawMeshSubsetBones(PreviewMeshComponent, PreviewMeshComponent->BonesOfInterest, PDI);
-		}
-		// otherwise, if we display bones, display
-		if ( GetBoneDrawMode() != EBoneDrawMode::None )
-		{
-			DrawMeshBones(PreviewMeshComponent, PDI);
-		}
-		if (PreviewMeshComponent->bDisplayRawAnimation )
-		{
-			DrawMeshBonesUncompressedAnimation(PreviewMeshComponent, PDI);
-		}
-		if (PreviewMeshComponent->NonRetargetedSpaceBases.Num() > 0 )
-		{
-			DrawMeshBonesNonRetargetedAnimation(PreviewMeshComponent, PDI);
-		}
-		if(PreviewMeshComponent->bDisplayAdditiveBasePose )
-		{
-			DrawMeshBonesAdditiveBasePose(PreviewMeshComponent, PDI);
-		}
-		if(PreviewMeshComponent->bDisplayBakedAnimation)
-		{
-			DrawMeshBonesBakedAnimation(PreviewMeshComponent, PDI);
-		}
-		if(PreviewMeshComponent->bDisplaySourceAnimation)
-		{
-			DrawMeshBonesSourceRawAnimation(PreviewMeshComponent, PDI);
-		}
-		
-		DrawWatchedPoses(PreviewMeshComponent, PDI);
+			// if we have BonesOfInterest, draw sub set of the bones only
+			if (GetAnimPreviewScene()->GetSelectedBoneIndex() != INDEX_NONE)
+			{
+				DrawMeshSubsetBones(PreviewMeshComponent, PreviewMeshComponent->BonesOfInterest, PDI);
+			}
+			// otherwise, if we display bones, display
+			if ( GetBoneDrawMode() != EBoneDrawMode::None )
+			{
+				DrawMeshBones(PreviewMeshComponent, PDI);
+			}
+			if (PreviewMeshComponent->bDisplayRawAnimation )
+			{
+				DrawMeshBonesUncompressedAnimation(PreviewMeshComponent, PDI);
+			}
+			if (PreviewMeshComponent->NonRetargetedSpaceBases.Num() > 0 )
+			{
+				DrawMeshBonesNonRetargetedAnimation(PreviewMeshComponent, PDI);
+			}
+			if(PreviewMeshComponent->bDisplayAdditiveBasePose )
+			{
+				DrawMeshBonesAdditiveBasePose(PreviewMeshComponent, PDI);
+			}
+			if(PreviewMeshComponent->bDisplayBakedAnimation)
+			{
+				DrawMeshBonesBakedAnimation(PreviewMeshComponent, PDI);
+			}
+			if(PreviewMeshComponent->bDisplaySourceAnimation)
+			{
+				DrawMeshBonesSourceRawAnimation(PreviewMeshComponent, PDI);
+			}
+			
+			DrawWatchedPoses(PreviewMeshComponent, PDI);
 
-		PreviewMeshComponent->DebugDrawClothing(PDI);
-		
-		// Display socket hit points
-		if (PreviewMeshComponent->bDrawSockets )
-		{
-			if (PreviewMeshComponent->bSkeletonSocketsVisible && PreviewMeshComponent->GetSkeletalMeshAsset()->GetSkeleton() )
+			PreviewMeshComponent->DebugDrawClothing(PDI);
+			
+			// Display socket hit points
+			if (PreviewMeshComponent->bDrawSockets )
 			{
-				DrawSockets(PreviewMeshComponent, PreviewMeshComponent->GetSkeletalMeshAsset()->GetSkeleton()->Sockets, GetAnimPreviewScene()->GetSelectedSocket(), PDI, true);
+				if (PreviewMeshComponent->bSkeletonSocketsVisible && PreviewMeshComponent->GetSkeletalMeshAsset()->GetSkeleton() )
+				{
+					DrawSockets(PreviewMeshComponent, PreviewMeshComponent->GetSkeletalMeshAsset()->GetSkeleton()->Sockets, GetAnimPreviewScene()->GetSelectedSocket(), PDI, true);
+				}
+
+				if ( PreviewMeshComponent->bMeshSocketsVisible )
+				{
+					DrawSockets(PreviewMeshComponent, PreviewMeshComponent->GetSkeletalMeshAsset()->GetMeshOnlySocketList(), GetAnimPreviewScene()->GetSelectedSocket(), PDI, false);
+				}
 			}
 
-			if ( PreviewMeshComponent->bMeshSocketsVisible )
+			if (PreviewMeshComponent->bDrawAttributes)
 			{
-				DrawSockets(PreviewMeshComponent, PreviewMeshComponent->GetSkeletalMeshAsset()->GetMeshOnlySocketList(), GetAnimPreviewScene()->GetSelectedSocket(), PDI, false);
+				DrawAttributes(PreviewMeshComponent, PDI);
 			}
 		}
-
-		if (PreviewMeshComponent->bDrawAttributes)
+		else if (bValidComponent && !bValidSkeletalMesh && GetBoneDrawMode() == EBoneDrawMode::All)
 		{
-			DrawAttributes(PreviewMeshComponent, PDI);
+			if (const USkeleton* Skeleton = GetPreviewScene()->GetPersonaToolkit()->GetSkeleton())
+			{
+				DrawBonesFromSkeleton(Skeleton, PDI);
+			}
 		}
 	}
 
