Animation Curve Runtime & Editor Improvements

diff --git a/Engine/Source/Editor/DetailCustomizations/Private/PoseAssetDetails.cpp b/Engine/Source/Editor/DetailCustomizations/Private/PoseAssetDetails.cpp
--- a/Engine/Source/Editor/DetailCustomizations/Private/PoseAssetDetails.cpp
+++ b/Engine/Source/Editor/DetailCustomizations/Private/PoseAssetDetails.cpp
@@ -179,8 +179,8 @@ void FPoseAssetDetails::CustomizeDetails(IDetailLayoutBuilder& DetailBuilder)
 	// add ref pose
 	BasePoseComboList.Add(MakeShareable(new FString(REFERENCE_BASE_POSE_NAME)));
 
-	TArray<FSmartName> PoseNames = PoseAsset->GetPoseNames();
-	FSmartName BasePoseName;
+	TArray<FName> PoseNames = PoseAsset->GetPoseFNames();
+	FName BasePoseName;
 
 	if (PoseNames.IsValidIndex(CachedBasePoseIndex))
 	{
@@ -195,7 +195,7 @@ void FPoseAssetDetails::CustomizeDetails(IDetailLayoutBuilder& DetailBuilder)
 		// go through profile and see if it has mine
 		for (const auto& PoseName : PoseNames)
 		{
-			BasePoseComboList.Add(MakeShareable(new FString(PoseName.DisplayName.ToString())));
+			BasePoseComboList.Add(MakeShareable(new FString(PoseName.ToString())));
 
 			if (PoseName == BasePoseName)
 			{
@@ -386,7 +386,7 @@ void FPoseAssetDetails::RefreshBasePoseChanged()
 			CachedBasePoseIndex = -1;
 		}
 
-		TArray<FSmartName> PoseNames = PoseAsset->GetPoseNames();
+		TArray<FName> PoseNames = PoseAsset->GetPoseFNames();
 		// add ref pose
 		BasePoseComboList.Add(MakeShareable(new FString(REFERENCE_BASE_POSE_NAME)));
 
@@ -395,7 +395,7 @@ void FPoseAssetDetails::RefreshBasePoseChanged()
 			// go through profile and see if it has mine
 			for (const auto& PoseName : PoseNames)
 			{
-				BasePoseComboList.Add(MakeShareable(new FString(PoseName.DisplayName.ToString())));
+				BasePoseComboList.Add(MakeShareable(new FString(PoseName.ToString())));
 			}
 		}
 
diff --git a/Engine/Source/Editor/EditorStyle/Private/StarshipStyle.cpp b/Engine/Source/Editor/EditorStyle/Private/StarshipStyle.cpp
--- a/Engine/Source/Editor/EditorStyle/Private/StarshipStyle.cpp
+++ b/Engine/Source/Editor/EditorStyle/Private/StarshipStyle.cpp
@@ -4861,7 +4861,8 @@ void FStarshipEditorStyle::FStyle::SetupPersonaStyle()
 		// Persona-specific tabs
 		Set("Persona.Tabs.SkeletonTree", new IMAGE_BRUSH(TEXT("Icons/icon_Persona_Skeleton_Tree_16x"), Icon16x16));
 		Set("Persona.Tabs.MorphTargetPreviewer", new IMAGE_BRUSH(TEXT("Icons/icon_Persona_Morph_Target_Previewer_16x"), Icon16x16));
-		Set("Persona.Tabs.AnimCurvePreviewer", new IMAGE_BRUSH(TEXT("Icons/icon_Persona_AnimCurve_Previewer_16x"), Icon16x16));
+		Set("Persona.Tabs.AnimCurveDebugger", new IMAGE_BRUSH(TEXT("Icons/icon_Persona_AnimCurve_Previewer_16x"), Icon16x16));
+		Set("Persona.Tabs.AnimCurveMetadataEditor", new IMAGE_BRUSH(TEXT("Icons/icon_Persona_AnimCurve_Previewer_16x"), Icon16x16));
 		Set("Persona.Tabs.AnimationNotifies", new IMAGE_BRUSH(TEXT("Icons/icon_Persona_Animation_Notifies_16x"), Icon16x16));
 		Set("Persona.Tabs.RetargetManager", new IMAGE_BRUSH_SVG("Starship/Persona/RetargetManager", Icon16x16));
 		Set("Persona.Tabs.AnimSlotManager", new IMAGE_BRUSH(TEXT("Icons/icon_Persona_Anim_Slot_Manager_16x"), Icon16x16));
@@ -5100,6 +5101,7 @@ void FStarshipEditorStyle::FStyle::SetupPersonaStyle()
 		Set("AnimCurveViewer.MaterialOff", new IMAGE_BRUSH(TEXT("Persona/AnimCurveViewer/Material_Off"), Icon16x16));
 		Set("AnimCurveViewer.MorphTargetHover", new IMAGE_BRUSH(TEXT("Persona/AnimCurveViewer/MorphTarget_On"), Icon16x16));
 		Set("AnimCurveViewer.MaterialHover", new IMAGE_BRUSH(TEXT("Persona/AnimCurveViewer/Material_On"), Icon16x16));
+		Set("AnimCurveViewer.ActiveCurveFont", DEFAULT_FONT("Bold", 8));
 
 		// blend space
 		Set("BlendSpaceEditor.ToggleTriangulation", new IMAGE_BRUSH(TEXT("Persona/BlendSpace/triangulation_16"), Icon16x16));
diff --git a/Engine/Source/Editor/MovieSceneTools/Private/TrackEditors/SkeletalAnimationTrackEditor.cpp b/Engine/Source/Editor/MovieSceneTools/Private/TrackEditors/SkeletalAnimationTrackEditor.cpp
--- a/Engine/Source/Editor/MovieSceneTools/Private/TrackEditors/SkeletalAnimationTrackEditor.cpp
+++ b/Engine/Source/Editor/MovieSceneTools/Private/TrackEditors/SkeletalAnimationTrackEditor.cpp
@@ -1001,7 +1001,7 @@ bool FSkeletalAnimationTrackEditor::CreatePoseAsset(const TArray<UObject*> NewAs
 			UPoseAsset* NewPoseAsset = Cast<UPoseAsset>(NewAsset);
 			if (NewPoseAsset)
 			{
-				NewPoseAsset->AddOrUpdatePoseWithUniqueName(SkeletalMeshComponent);
+				NewPoseAsset->AddPoseWithUniqueName(SkeletalMeshComponent);
 				bResult = true;
 			}
 		}
diff --git a/Engine/Source/Editor/Persona/Private/AnimGraphNodeDetails.cpp b/Engine/Source/Editor/Persona/Private/AnimGraphNodeDetails.cpp
--- a/Engine/Source/Editor/Persona/Private/AnimGraphNodeDetails.cpp
+++ b/Engine/Source/Editor/Persona/Private/AnimGraphNodeDetails.cpp
@@ -610,7 +610,14 @@ void FBoneReferenceCustomization::SetEditableSkeleton(TSharedRef<IPropertyHandle
 			// this is editor object that only exists for editor
 			if (UEditorAnimCurveBoneLinks* AnimCurveObj = Cast<UEditorAnimCurveBoneLinks>(InObject))
 			{
-				EditableSkeleton = AnimCurveObj->EditableSkeleton.Pin();
+				if(USkeleton* Skeleton = Cast<USkeleton>(AnimCurveObj->AnimCurveMetaData->GetOuter()))
+				{
+					TargetSkeleton = Skeleton;
+				}
+				else if(USkeletalMesh* SkeletalMesh = Cast<USkeletalMesh>(AnimCurveObj->AnimCurveMetaData->GetOuter()))
+				{
+					TargetSkeleton = SkeletalMesh->GetSkeleton();
+				}
 				break;
 			}
 
diff --git a/Engine/Source/Editor/Persona/Private/AnimTimeline/AnimModel_AnimSequenceBase.cpp b/Engine/Source/Editor/Persona/Private/AnimTimeline/AnimModel_AnimSequenceBase.cpp
--- a/Engine/Source/Editor/Persona/Private/AnimTimeline/AnimModel_AnimSequenceBase.cpp
+++ b/Engine/Source/Editor/Persona/Private/AnimTimeline/AnimModel_AnimSequenceBase.cpp
@@ -257,7 +257,7 @@ void FAnimModel_AnimSequenceBase::RefreshCurveTracks()
 			TransformCurveTrack->SetExpanded(false);
 			AdditiveRoot->AddChild(TransformCurveTrack);
 
-			FText TransformName = FAnimTimelineTrack_TransformCurve::GetTransformCurveName(AsShared(), TransformCurve.Name);
+			FText TransformName = FText::FromName(TransformCurve.GetName());
 			FLinearColor TransformColor = TransformCurve.GetColor();
 			FLinearColor XColor = FLinearColor::Red;
 			FLinearColor YColor = FLinearColor::Green;
@@ -268,33 +268,33 @@ void FAnimModel_AnimSequenceBase::RefreshCurveTracks()
 			
 			FText VectorFormat = LOCTEXT("TransformVectorFormat", "{0}.{1}");
 			FText TranslationName = LOCTEXT("TransformTranslationTrackName", "Translation");
-			TSharedRef<FAnimTimelineTrack_VectorCurve> TranslationCurveTrack = MakeShared<FAnimTimelineTrack_VectorCurve>(&TransformCurve.TranslationCurve, TransformCurve.Name, 0, ERawCurveTrackTypes::RCT_Transform, TranslationName, FText::Format(VectorFormat, TransformName, TranslationName), TransformColor, SharedThis(this));
+			TSharedRef<FAnimTimelineTrack_VectorCurve> TranslationCurveTrack = MakeShared<FAnimTimelineTrack_VectorCurve>(&TransformCurve.TranslationCurve, TransformCurve.GetName(), 0, ERawCurveTrackTypes::RCT_Transform, TranslationName, FText::Format(VectorFormat, TransformName, TranslationName), TransformColor, SharedThis(this));
 			TranslationCurveTrack->SetExpanded(false);
 			TransformCurveTrack->AddChild(TranslationCurveTrack);			
 
 			FText ComponentFormat = LOCTEXT("TransformComponentFormat", "{0}.{1}.{2}");
-			TranslationCurveTrack->AddChild(MakeShared<FAnimTimelineTrack_Curve>(&TransformCurve.TranslationCurve.FloatCurves[0], TransformCurve.Name, 0, ERawCurveTrackTypes::RCT_Transform, XName, FText::Format(ComponentFormat, TransformName, TranslationName, XName), XColor, XColor, SharedThis(this)));
-			TranslationCurveTrack->AddChild(MakeShared<FAnimTimelineTrack_Curve>(&TransformCurve.TranslationCurve.FloatCurves[1], TransformCurve.Name, 1, ERawCurveTrackTypes::RCT_Transform, YName, FText::Format(ComponentFormat, TransformName, TranslationName, YName), YColor, YColor, SharedThis(this)));
-			TranslationCurveTrack->AddChild(MakeShared<FAnimTimelineTrack_Curve>(&TransformCurve.TranslationCurve.FloatCurves[2], TransformCurve.Name, 2, ERawCurveTrackTypes::RCT_Transform, ZName, FText::Format(ComponentFormat, TransformName, TranslationName, ZName), ZColor, ZColor, SharedThis(this)));
+			TranslationCurveTrack->AddChild(MakeShared<FAnimTimelineTrack_Curve>(&TransformCurve.TranslationCurve.FloatCurves[0], TransformCurve.GetName(), 0, ERawCurveTrackTypes::RCT_Transform, XName, FText::Format(ComponentFormat, TransformName, TranslationName, XName), XColor, XColor, SharedThis(this)));
+			TranslationCurveTrack->AddChild(MakeShared<FAnimTimelineTrack_Curve>(&TransformCurve.TranslationCurve.FloatCurves[1], TransformCurve.GetName(), 1, ERawCurveTrackTypes::RCT_Transform, YName, FText::Format(ComponentFormat, TransformName, TranslationName, YName), YColor, YColor, SharedThis(this)));
+			TranslationCurveTrack->AddChild(MakeShared<FAnimTimelineTrack_Curve>(&TransformCurve.TranslationCurve.FloatCurves[2], TransformCurve.GetName(), 2, ERawCurveTrackTypes::RCT_Transform, ZName, FText::Format(ComponentFormat, TransformName, TranslationName, ZName), ZColor, ZColor, SharedThis(this)));
 
 			FText RollName = LOCTEXT("RotationRollTrackName", "Roll");
 			FText PitchName = LOCTEXT("RotationPitchTrackName", "Pitch");
 			FText YawName = LOCTEXT("RotationYawTrackName", "Yaw");
 			FText RotationName = LOCTEXT("TransformRotationTrackName", "Rotation");
-			TSharedRef<FAnimTimelineTrack_VectorCurve> RotationCurveTrack = MakeShared<FAnimTimelineTrack_VectorCurve>(&TransformCurve.RotationCurve, TransformCurve.Name, 3, ERawCurveTrackTypes::RCT_Transform, RotationName, FText::Format(VectorFormat, TransformName, RotationName), TransformColor, SharedThis(this));
+			TSharedRef<FAnimTimelineTrack_VectorCurve> RotationCurveTrack = MakeShared<FAnimTimelineTrack_VectorCurve>(&TransformCurve.RotationCurve, TransformCurve.GetName(), 3, ERawCurveTrackTypes::RCT_Transform, RotationName, FText::Format(VectorFormat, TransformName, RotationName), TransformColor, SharedThis(this));
 			RotationCurveTrack->SetExpanded(false);
 			TransformCurveTrack->AddChild(RotationCurveTrack);
-			RotationCurveTrack->AddChild(MakeShared<FAnimTimelineTrack_Curve>(&TransformCurve.RotationCurve.FloatCurves[0], TransformCurve.Name, 3, ERawCurveTrackTypes::RCT_Transform, RollName, FText::Format(ComponentFormat, TransformName, RotationName, RollName), XColor, XColor, SharedThis(this)));
-			RotationCurveTrack->AddChild(MakeShared<FAnimTimelineTrack_Curve>(&TransformCurve.RotationCurve.FloatCurves[1], TransformCurve.Name, 4, ERawCurveTrackTypes::RCT_Transform, PitchName, FText::Format(ComponentFormat, TransformName, RotationName, PitchName), YColor, YColor, SharedThis(this)));
-			RotationCurveTrack->AddChild(MakeShared<FAnimTimelineTrack_Curve>(&TransformCurve.RotationCurve.FloatCurves[2], TransformCurve.Name, 5, ERawCurveTrackTypes::RCT_Transform, YawName, FText::Format(ComponentFormat, TransformName, RotationName, YawName), ZColor, ZColor, SharedThis(this)));
+			RotationCurveTrack->AddChild(MakeShared<FAnimTimelineTrack_Curve>(&TransformCurve.RotationCurve.FloatCurves[0], TransformCurve.GetName(), 3, ERawCurveTrackTypes::RCT_Transform, RollName, FText::Format(ComponentFormat, TransformName, RotationName, RollName), XColor, XColor, SharedThis(this)));
+			RotationCurveTrack->AddChild(MakeShared<FAnimTimelineTrack_Curve>(&TransformCurve.RotationCurve.FloatCurves[1], TransformCurve.GetName(), 4, ERawCurveTrackTypes::RCT_Transform, PitchName, FText::Format(ComponentFormat, TransformName, RotationName, PitchName), YColor, YColor, SharedThis(this)));
+			RotationCurveTrack->AddChild(MakeShared<FAnimTimelineTrack_Curve>(&TransformCurve.RotationCurve.FloatCurves[2], TransformCurve.GetName(), 5, ERawCurveTrackTypes::RCT_Transform, YawName, FText::Format(ComponentFormat, TransformName, RotationName, YawName), ZColor, ZColor, SharedThis(this)));
 
 			FText ScaleName = LOCTEXT("TransformScaleTrackName", "Scale");
-			TSharedRef<FAnimTimelineTrack_VectorCurve> ScaleCurveTrack = MakeShared<FAnimTimelineTrack_VectorCurve>(&TransformCurve.ScaleCurve, TransformCurve.Name, 6, ERawCurveTrackTypes::RCT_Transform, ScaleName, FText::Format(VectorFormat, TransformName, ScaleName), TransformColor, SharedThis(this));
+			TSharedRef<FAnimTimelineTrack_VectorCurve> ScaleCurveTrack = MakeShared<FAnimTimelineTrack_VectorCurve>(&TransformCurve.ScaleCurve, TransformCurve.GetName(), 6, ERawCurveTrackTypes::RCT_Transform, ScaleName, FText::Format(VectorFormat, TransformName, ScaleName), TransformColor, SharedThis(this));
 			ScaleCurveTrack->SetExpanded(false);
 			TransformCurveTrack->AddChild(ScaleCurveTrack);
-			ScaleCurveTrack->AddChild(MakeShared<FAnimTimelineTrack_Curve>(&TransformCurve.ScaleCurve.FloatCurves[0], TransformCurve.Name, 6, ERawCurveTrackTypes::RCT_Transform, XName, FText::Format(ComponentFormat, TransformName, ScaleName, XName), XColor, XColor, SharedThis(this)));
-			ScaleCurveTrack->AddChild(MakeShared<FAnimTimelineTrack_Curve>(&TransformCurve.ScaleCurve.FloatCurves[1], TransformCurve.Name, 7, ERawCurveTrackTypes::RCT_Transform, YName, FText::Format(ComponentFormat, TransformName, ScaleName, YName), YColor, YColor, SharedThis(this)));
-			ScaleCurveTrack->AddChild(MakeShared<FAnimTimelineTrack_Curve>(&TransformCurve.ScaleCurve.FloatCurves[2], TransformCurve.Name, 8, ERawCurveTrackTypes::RCT_Transform, ZName, FText::Format(ComponentFormat, TransformName, ScaleName, ZName), ZColor, ZColor, SharedThis(this)));
+			ScaleCurveTrack->AddChild(MakeShared<FAnimTimelineTrack_Curve>(&TransformCurve.ScaleCurve.FloatCurves[0], TransformCurve.GetName(), 6, ERawCurveTrackTypes::RCT_Transform, XName, FText::Format(ComponentFormat, TransformName, ScaleName, XName), XColor, XColor, SharedThis(this)));
+			ScaleCurveTrack->AddChild(MakeShared<FAnimTimelineTrack_Curve>(&TransformCurve.ScaleCurve.FloatCurves[1], TransformCurve.GetName(), 7, ERawCurveTrackTypes::RCT_Transform, YName, FText::Format(ComponentFormat, TransformName, ScaleName, YName), YColor, YColor, SharedThis(this)));
+			ScaleCurveTrack->AddChild(MakeShared<FAnimTimelineTrack_Curve>(&TransformCurve.ScaleCurve.FloatCurves[2], TransformCurve.GetName(), 8, ERawCurveTrackTypes::RCT_Transform, ZName, FText::Format(ComponentFormat, TransformName, ScaleName, ZName), ZColor, ZColor, SharedThis(this)));
 		}		
 	}
 }
@@ -393,7 +393,7 @@ void FAnimModel_AnimSequenceBase::EditSelectedCurves()
 				{
 					FText FullName = CurveTrack->GetFullCurveName(CurveIndex);
 					FLinearColor Color = CurveTrack->GetCurveColor(CurveIndex);
-					FSmartName Name;
+					FName Name;
 					ERawCurveTrackTypes Type;
 					int32 EditCurveIndex;
 					CurveTrack->GetCurveEditInfo(CurveIndex, Name, Type, EditCurveIndex);
@@ -443,37 +443,26 @@ void FAnimModel_AnimSequenceBase::RemoveSelectedCurves()
 		if(SelectedTrack->IsA<FAnimTimelineTrack_FloatCurve>())
 		{
 			TSharedRef<FAnimTimelineTrack_FloatCurve> FloatCurveTrack = StaticCastSharedRef<FAnimTimelineTrack_FloatCurve>(SelectedTrack);
-
-			FSmartName CurveName = FloatCurveTrack->GetName();
+			FName CurveName = FloatCurveTrack->GetFName();
+			
 			const FAnimationCurveIdentifier FloatCurveId(CurveName, ERawCurveTrackTypes::RCT_Float);
 			if (AnimSequenceBase->GetDataModel()->FindCurve(FloatCurveId))
 			{
-				FSmartName TrackName;
-				if (AnimSequenceBase->GetSkeleton()->GetSmartNameByUID(USkeleton::AnimCurveMappingName, CurveName.UID, TrackName))
-				{
-					Controller.RemoveCurve(FloatCurveId);
-					bDeletedCurve = true;
-				}
+				Controller.RemoveCurve(FloatCurveId);
+				bDeletedCurve = true;
 			}
 		}
 		else if(SelectedTrack->IsA<FAnimTimelineTrack_TransformCurve>())
 		{
 			TSharedRef<FAnimTimelineTrack_TransformCurve> TransformCurveTrack = StaticCastSharedRef<FAnimTimelineTrack_TransformCurve>(SelectedTrack);
-
-			const FTransformCurve& TransformCurve = TransformCurveTrack->GetTransformCurve();
-			FSmartName CurveName = TransformCurveTrack->GetName();
+			FName CurveName = TransformCurveTrack->GetFName();
 
 			const FAnimationCurveIdentifier TransformCurveId(CurveName, ERawCurveTrackTypes::RCT_Transform);
-
 			if (AnimSequenceBase->GetDataModel()->FindTransformCurve(TransformCurveId))
 			{
-				FSmartName CurveToDelete;
-				if (AnimSequenceBase->GetSkeleton()->GetSmartNameByUID(USkeleton::AnimTrackCurveMappingName, CurveName.UID, CurveToDelete))
-				{
-					Controller.RemoveCurve(TransformCurveId);
-					bDeletedCurve = true;
-				}
-			}	
+				Controller.RemoveCurve(TransformCurveId);
+				bDeletedCurve = true;
+			}
 		}
 		else if (SelectedTrack->IsA<FAnimTimelineTrack_CompoundCurve>())
 		{
@@ -481,17 +470,13 @@ void FAnimModel_AnimSequenceBase::RemoveSelectedCurves()
 			TSharedRef<FAnimTimelineTrack_CompoundCurve> CurveTrack = StaticCastSharedRef<FAnimTimelineTrack_CompoundCurve>(SelectedTrack);
 
 			// Find all editable curves in this track
-			for (const FSmartName CurveName : CurveTrack->GetCurveNames())
+			for (const FName TrackCurveName : CurveTrack->GetCurveNames())
 			{
-				const FAnimationCurveIdentifier CurveId(CurveName, ERawCurveTrackTypes::RCT_Float);
+				const FAnimationCurveIdentifier CurveId(TrackCurveName, ERawCurveTrackTypes::RCT_Float);
 				if (AnimSequenceBase->GetDataModel()->FindCurve(CurveId))
 				{
-					FSmartName TrackName;
-					if (AnimSequenceBase->GetSkeleton()->GetSmartNameByUID(USkeleton::AnimCurveMappingName, CurveName.UID, TrackName))
-					{
-						Controller.RemoveCurve(CurveId);
-						bDeletedCurve = true;
-					}
+					Controller.RemoveCurve(CurveId);
+					bDeletedCurve = true;
 				}
 			}
 		}
@@ -519,12 +504,12 @@ void FAnimModel_AnimSequenceBase::CopySelectedCurveNamesToClipboard()
 		if(SelectedTrack->IsA<FAnimTimelineTrack_FloatCurve>())
 		{
 			TSharedRef<FAnimTimelineTrack_FloatCurve> FloatCurveTrack = StaticCastSharedRef<FAnimTimelineTrack_FloatCurve>(SelectedTrack);
-			TrackNames.Add(FloatCurveTrack->GetName().DisplayName.ToString());
+			TrackNames.Add(FloatCurveTrack->GetFName().ToString());
 		}
 		else if(SelectedTrack->IsA<FAnimTimelineTrack_TransformCurve>())
 		{
 			TSharedRef<FAnimTimelineTrack_TransformCurve> TransformCurveTrack = StaticCastSharedRef<FAnimTimelineTrack_TransformCurve>(SelectedTrack);
-			TrackNames.Add(TransformCurveTrack->GetName().DisplayName.ToString());
+			TrackNames.Add(TransformCurveTrack->GetFName().ToString());
 		}
 
 	}
diff --git a/Engine/Source/Editor/Persona/Private/AnimTimeline/AnimTimelineClipboard.cpp b/Engine/Source/Editor/Persona/Private/AnimTimeline/AnimTimelineClipboard.cpp
--- a/Engine/Source/Editor/Persona/Private/AnimTimeline/AnimTimelineClipboard.cpp
+++ b/Engine/Source/Editor/Persona/Private/AnimTimeline/AnimTimelineClipboard.cpp
@@ -17,8 +17,7 @@ FAnimationCurveIdentifier UAnimCurveBaseCopyObject::GetAnimationCurveIdentifier(
 	
 	OutAnimationCurveIdentifier.Axis = Axis;
 	OutAnimationCurveIdentifier.Channel = Channel;
-	OutAnimationCurveIdentifier.InternalName.DisplayName = DisplayName;
-	OutAnimationCurveIdentifier.InternalName.UID = UID;
+	OutAnimationCurveIdentifier.CurveName = CurveName;
 	OutAnimationCurveIdentifier.CurveType = CurveType;
 	
 	return OutAnimationCurveIdentifier;
@@ -115,7 +114,7 @@ const UAnimTimelineClipboardContent* FAnimTimelineClipboardUtilities::GetContent
 	return nullptr;
 }
 
-UAnimCurveBaseCopyObject::UAnimCurveBaseCopyObject() : DisplayName(NAME_None), UID(SmartName::MaxUID), CurveType(ERawCurveTrackTypes::RCT_MAX), Channel(ETransformCurveChannel::Invalid), Axis(EVectorCurveChannel::Invalid), OriginName(NAME_None)
+UAnimCurveBaseCopyObject::UAnimCurveBaseCopyObject() : CurveName(NAME_None), CurveType(ERawCurveTrackTypes::RCT_MAX), Channel(ETransformCurveChannel::Invalid), Axis(EVectorCurveChannel::Invalid), OriginName(NAME_None)
 {
 }
 
@@ -153,7 +152,7 @@ bool FAnimTimelineClipboardUtilities::CanOverwriteSelectedCurveDataFromClipboard
 			const FAnimTimelineTrack_Curve & CurveTrack = Track->As<FAnimTimelineTrack_Curve>();
 
 			// Get current track information
-			FSmartName CurveName;
+			FName CurveName;
 			ERawCurveTrackTypes CurveType;
 			int32 CurveIndex;
 			CurveTrack.GetCurveEditInfo(0 /* Unused */, CurveName, CurveType, CurveIndex);
@@ -185,7 +184,7 @@ void FAnimTimelineClipboardUtilities::OverwriteSelectedCurveDataFromClipboard(co
 		const FAnimTimelineTrack_Curve & CurveTrack = Track->As<FAnimTimelineTrack_Curve>();
 
 		// Get selected curve track information
-		FSmartName SelectedCurveTrackName;
+		FName SelectedCurveTrackName;
 		ERawCurveTrackTypes SelectedCurveTrackType;	// The value is supposed to always be RCT_TRANSFORM as mentioned in FAnimTimelineTrack_Curve but this is not true.
 		int32 SelectedCurveTrackIndex;
 		CurveTrack.GetCurveEditInfo(0 /* Unused */, SelectedCurveTrackName, SelectedCurveTrackType, SelectedCurveTrackIndex);
@@ -263,7 +262,7 @@ void FAnimTimelineClipboardUtilities::OverwriteOrAddCurvesFromClipboardContent(c
 		// Exit on types not supported by animation controller
 		if (InCurveCopyObj->CurveType != ERawCurveTrackTypes::RCT_Float && InCurveCopyObj->CurveType != ERawCurveTrackTypes::RCT_Transform)
 		{
-			UE_LOG(LogAnimation, Warning, TEXT("Attempting to paste curve { %s } of unsupported type { %s } by animation controller. "), *InCurveCopyObj->DisplayName.ToString(), *UEnum::GetValueAsString(InCurveCopyObj->CurveType));
+			UE_LOG(LogAnimation, Warning, TEXT("Attempting to paste curve { %s } of unsupported type { %s } by animation controller. "), *InCurveCopyObj->CurveName.ToString(), *UEnum::GetValueAsString(InCurveCopyObj->CurveType));
 			continue;
 		}
 		
@@ -271,61 +270,28 @@ void FAnimTimelineClipboardUtilities::OverwriteOrAddCurvesFromClipboardContent(c
 		if (InCurveCopyObj->OriginName.Compare(Controller.GetModelInterface().GetObject()->GetOuter()->GetFName()) == 0 && Controller.GetModel()->FindCurve(InCurveCopyObj->GetAnimationCurveIdentifier()))
 		{
 			// Do nothing and exit.
-			UE_LOG(LogAnimation, Warning, TEXT("Attempting to paste already existing curve: %s"), *InCurveCopyObj->DisplayName.ToString());
+			UE_LOG(LogAnimation, Warning, TEXT("Attempting to paste already existing curve: %s"), *InCurveCopyObj->CurveName.ToString());
 		}
 		else
 		{
-			FSmartName CurveTrackName(InCurveCopyObj->DisplayName, SmartName::MaxUID);
 			FAnimationCurveIdentifier CurveIdentifier = InCurveCopyObj->GetAnimationCurveIdentifier();
-			const FName & CurveContainerName = InCurveCopyObj->CurveType == ERawCurveTrackTypes::RCT_Float ? USkeleton::AnimCurveMappingName : USkeleton::AnimTrackCurveMappingName;
-			
-			const SmartName::UID_Type FoundCurveUID = Skeleton->GetUIDByName(CurveContainerName, InCurveCopyObj->DisplayName);
-			if (FoundCurveUID != SmartName::MaxUID)
-			{
-				// Update identifier since skeleton does has curve name in the mapping container
-				CurveTrackName.UID = FoundCurveUID;
-				CurveIdentifier.InternalName.UID = FoundCurveUID;
-			}
-			else
+			if (InCurveCopyObj->CurveType == ERawCurveTrackTypes::RCT_Float)
 			{
-				// Add new curve name based on clipboard information
-				Skeleton->AddSmartNameAndModify(CurveContainerName, InCurveCopyObj->DisplayName, CurveTrackName);
-					
-				// Ensure that curve smart name exists
-				FSmartName NewName;
-				ensureAlways(Skeleton->GetSmartNameByUID(CurveContainerName, CurveTrackName.UID, NewName));
-
-				// Update identifier since skeleton does has curve name in the mapping container
-				CurveIdentifier.InternalName = NewName;
+				const UFloatCurveCopyObject* FloatCopyObject = Cast<UFloatCurveCopyObject>(InCurveCopyObj);
+				
+				Controller.SetCurveKeys(CurveIdentifier, FloatCopyObject->Curve.FloatCurve.GetConstRefOfKeys());
+				Controller.SetCurveFlags(CurveIdentifier, FloatCopyObject->Curve.GetCurveTypeFlags());
 			}
-			
-			// Create curve with clipboard information if needed
-			if (!Controller.GetModel()->FindCurve(CurveIdentifier))
+			else if (InCurveCopyObj->CurveType == ERawCurveTrackTypes::RCT_Transform)
 			{
-				Controller.AddCurve(CurveIdentifier);
-			}
+				const UTransformCurveCopyObject* TransformCopyObject = Cast<UTransformCurveCopyObject>(InCurveCopyObj);
 				
-			// Update data for the curve
-			if (CurveTrackName.IsValid()) 
-			{
-				if (InCurveCopyObj->CurveType == ERawCurveTrackTypes::RCT_Float)
-				{
-					const UFloatCurveCopyObject* FloatCopyObject = Cast<UFloatCurveCopyObject>(InCurveCopyObj);
-					
-					Controller.SetCurveKeys(CurveIdentifier, FloatCopyObject->Curve.FloatCurve.GetConstRefOfKeys());
-					Controller.SetCurveFlags(CurveIdentifier, FloatCopyObject->Curve.GetCurveTypeFlags());
-				}
-				else if (InCurveCopyObj->CurveType == ERawCurveTrackTypes::RCT_Transform)
-				{
-					const UTransformCurveCopyObject* TransformCopyObject = Cast<UTransformCurveCopyObject>(InCurveCopyObj);
-					
-					TArray<float> TimeKeys;
-					TArray<FTransform> TransformValues;
-					TransformCopyObject->Curve.GetKeys(TimeKeys, TransformValues); // TODO: Optimize this if possible!
-		
-					Controller.SetTransformCurveKeys(CurveIdentifier, TransformValues, TimeKeys);
-					Controller.SetCurveFlags(CurveIdentifier, TransformCopyObject->Curve.GetCurveTypeFlags());
-				}
+				TArray<float> TimeKeys;
+				TArray<FTransform> TransformValues;
+				TransformCopyObject->Curve.GetKeys(TimeKeys, TransformValues); // TODO: Optimize this if possible!
+	
+				Controller.SetTransformCurveKeys(CurveIdentifier, TransformValues, TimeKeys);
+				Controller.SetCurveFlags(CurveIdentifier, TransformCopyObject->Curve.GetCurveTypeFlags());
 			}
 		}
 	}
diff --git a/Engine/Source/Editor/Persona/Private/AnimTimeline/AnimTimelineClipboard.h b/Engine/Source/Editor/Persona/Private/AnimTimeline/AnimTimelineClipboard.h
--- a/Engine/Source/Editor/Persona/Private/AnimTimeline/AnimTimelineClipboard.h
+++ b/Engine/Source/Editor/Persona/Private/AnimTimeline/AnimTimelineClipboard.h
@@ -27,19 +27,8 @@ public:
 
 	// Store information for FAnimationCurveIdentifier.
 
-	/** We declare individual properties of FSmartName since only FSmartName.DisplayName is serializable (has UPROPERTY macro)
-	 *  and FSmartName.UID is not.
-	 */
-	
-	// SmartName Start
-	
 	UPROPERTY()
-	FName DisplayName;
-
-	UPROPERTY()
-	uint16	UID;	// We can't use SmartName::UID_Type typedef since the SmartName struct is not properly prefixed.
-	
-	// SmartName End
+	FName CurveName;
 	
 	UPROPERTY()
 	ERawCurveTrackTypes CurveType;
diff --git a/Engine/Source/Editor/Persona/Private/AnimTimeline/AnimTimelineTrack_CompoundCurve.cpp b/Engine/Source/Editor/Persona/Private/AnimTimeline/AnimTimelineTrack_CompoundCurve.cpp
--- a/Engine/Source/Editor/Persona/Private/AnimTimeline/AnimTimelineTrack_CompoundCurve.cpp
+++ b/Engine/Source/Editor/Persona/Private/AnimTimeline/AnimTimelineTrack_CompoundCurve.cpp
@@ -38,7 +38,7 @@ FAnimTimelineTrack_CompoundCurve::FAnimTimelineTrack_CompoundCurve(TArray<const
 	, Curves(MoveTemp(InCurves))
 {
 	OuterCurveIndex = 0;
-	Algo::Transform(Curves, CurveNames, [](const FFloatCurve* Curve) { return Curve->Name; });
+	Algo::Transform(Curves, CurveNames, [](const FFloatCurve* Curve) { return Curve->GetName(); });
 }
 
 /** FAnimTimelineTrack_Curve interface */
@@ -49,7 +49,7 @@ FLinearColor FAnimTimelineTrack_CompoundCurve::GetCurveColor(int32 InCurveIndex)
 
 FText FAnimTimelineTrack_CompoundCurve::GetFullCurveName(int32 InCurveIndex) const
 {
-	return FText::FromName(Curves[InCurveIndex]->Name.DisplayName);
+	return FText::FromName(Curves[InCurveIndex]->GetName());
 }
 
 bool FAnimTimelineTrack_CompoundCurve::CanEditCurve(int32 InCurveIndex) const
@@ -57,9 +57,9 @@ bool FAnimTimelineTrack_CompoundCurve::CanEditCurve(int32 InCurveIndex) const
 	return true;
 }
 
-void FAnimTimelineTrack_CompoundCurve::GetCurveEditInfo(int32 InCurveIndex, FSmartName& OutName, ERawCurveTrackTypes& OutType, int32& OutCurveIndex) const
+void FAnimTimelineTrack_CompoundCurve::GetCurveEditInfo(int32 InCurveIndex, FName& OutName, ERawCurveTrackTypes& OutType, int32& OutCurveIndex) const
 {
-	OutName = Curves[InCurveIndex]->Name;
+	OutName = Curves[InCurveIndex]->GetName();
 	OutType = ERawCurveTrackTypes::RCT_Float;
 	OutCurveIndex = 0;
 }
@@ -162,7 +162,7 @@ struct FAnimTimelineTrack_CompoundCurve::FCurveGroup
 		FCurveGroup RootGroup;
 		for (const FFloatCurve& FloatCurve : FloatCurves)
 		{
-			FString CurveNameStr = FloatCurve.Name.DisplayName.ToString();
+			FString CurveNameStr = FloatCurve.GetName().ToString();
 			FStringView CurveName = CurveNameStr;
 
 			// We need c++ 20 std::split_view here
diff --git a/Engine/Source/Editor/Persona/Private/AnimTimeline/AnimTimelineTrack_CompoundCurve.h b/Engine/Source/Editor/Persona/Private/AnimTimeline/AnimTimelineTrack_CompoundCurve.h
--- a/Engine/Source/Editor/Persona/Private/AnimTimeline/AnimTimelineTrack_CompoundCurve.h
+++ b/Engine/Source/Editor/Persona/Private/AnimTimeline/AnimTimelineTrack_CompoundCurve.h
@@ -20,7 +20,7 @@ public:
 	virtual FLinearColor GetCurveColor(int32 InCurveIndex) const override;
 	virtual FText GetFullCurveName(int32 InCurveIndex) const override;
 	virtual bool CanEditCurve(int32 InCurveIndex) const override;
-	virtual void GetCurveEditInfo(int32 InCurveIndex, FSmartName& OutName, ERawCurveTrackTypes& OutType, int32& OutCurveIndex) const override;
+	virtual void GetCurveEditInfo(int32 InCurveIndex, FName& OutName, ERawCurveTrackTypes& OutType, int32& OutCurveIndex) const override;
 	/** End FAnimTimelineTrack_Curve interface */
 
 	static constexpr auto DefaultDelimiters = TEXT("._/\\|");
@@ -30,10 +30,10 @@ public:
 	//! \param[in] InModel FAnimModel to display the track
 	static void AddGroupedCurveTracks(TArrayView<const FFloatCurve> FloatCurves, FAnimTimelineTrack& ParentTrack, const TSharedRef<FAnimModel>& InModel, FStringView Delimiters = DefaultDelimiters);
 
-	const TArray<FSmartName>& GetCurveNames() const { return CurveNames; }
+	const TArray<FName>& GetCurveNames() const { return CurveNames; }
 private:
 	TArray<const FFloatCurve*> Curves;
-	TArray<FSmartName> CurveNames; // Curve smart names used when removing curves (in case the FFloatCurve was removed)
+	TArray<FName> CurveNames; // Curve names used when removing curves (in case the FFloatCurve was removed)
 
 	static TArray<const FRichCurve*> ToRichCurves(TArrayView<FFloatCurve const* const> InCurves);
 	struct FCurveGroup;
diff --git a/Engine/Source/Editor/Persona/Private/AnimTimeline/AnimTimelineTrack_Curve.cpp b/Engine/Source/Editor/Persona/Private/AnimTimeline/AnimTimelineTrack_Curve.cpp
--- a/Engine/Source/Editor/Persona/Private/AnimTimeline/AnimTimelineTrack_Curve.cpp
+++ b/Engine/Source/Editor/Persona/Private/AnimTimeline/AnimTimelineTrack_Curve.cpp
@@ -162,6 +162,19 @@ FAnimTimelineTrack_Curve::FAnimTimelineTrack_Curve(const FText& InCurveName, con
 }
 
 FAnimTimelineTrack_Curve::FAnimTimelineTrack_Curve(const FRichCurve* InCurve, const FSmartName& InName, int32 InCurveIndex, ERawCurveTrackTypes InType, const FText& InCurveName, const FText& InFullCurveName, const FLinearColor& InColor, const FLinearColor& InBackgroundColor, const TSharedRef<FAnimModel>& InModel)
+	: FAnimTimelineTrack(InCurveName, InCurveName, InModel)
+	, Color(InColor)
+	, BackgroundColor(InBackgroundColor)
+	, FullCurveName(InFullCurveName)
+	, OuterCurveName(InName.DisplayName)
+	, OuterCurveIndex(InCurveIndex)
+	, OuterType(InType)
+{
+	Curves.Add(InCurve);
+	SetHeight(32.0f);
+}
+
+FAnimTimelineTrack_Curve::FAnimTimelineTrack_Curve(const FRichCurve* InCurve, const FName& InName, int32 InCurveIndex, ERawCurveTrackTypes InType, const FText& InCurveName, const FText& InFullCurveName, const FLinearColor& InColor, const FLinearColor& InBackgroundColor, const TSharedRef<FAnimModel>& InModel)
 	: FAnimTimelineTrack(InCurveName, InCurveName, InModel)
 	, Color(InColor)
 	, BackgroundColor(InBackgroundColor)
@@ -197,13 +210,13 @@ TSharedRef<SWidget> FAnimTimelineTrack_Curve::GenerateContainerWidgetForTimeline
 	{
 		const FRichCurve* Curve = Curves[CurveIndex];
 
-		FSmartName Name;
+		FName Name;
 		ERawCurveTrackTypes Type;
 		int32 EditIndex;
 		GetCurveEditInfo(CurveIndex, Name, Type, EditIndex);
 
 		TUniquePtr<FRichCurveEditorModelNamed> NewCurveModel = MakeUnique<FRichCurveEditorModelNamed>(Name, Type, EditIndex, GetModel()->GetAnimSequenceBase());
-		NewCurveModel->SetColor(GetCurveColor(CurveIndex));
+		NewCurveModel->SetColor(GetCurveColor(CurveIndex), false);
 		NewCurveModel->SetIsKeyDrawEnabled(MakeAttributeLambda([](){ return GetDefault<UPersonaOptions>()->bTimelineDisplayCurveKeys; }));
 		CurveEditor->AddCurve(MoveTemp(NewCurveModel));
 	}
@@ -362,13 +375,12 @@ void FAnimTimelineTrack_Curve::Copy(UAnimTimelineClipboardContent* InOutClipboar
 		const FRichCurve* InCurve = Curves[0];
 
 		// Copy raw curve data
-		CopyableCurve->Curve.Name = { FName(FullCurveName.ToString()), SmartName::MaxUID };
+		CopyableCurve->Curve.SetName(FName(FullCurveName.ToString()));
 		CopyableCurve->Curve.FloatCurve = *InCurve;
 		CopyableCurve->Curve.SetCurveTypeFlags(AACF_Editable);
 
 		// Copy curve identifier data
-		CopyableCurve->DisplayName = CopyableCurve->Curve.Name.DisplayName;
-		CopyableCurve->UID = CopyableCurve->Curve.Name.UID;
+		CopyableCurve->CurveName = CopyableCurve->Curve.GetName();
 		CopyableCurve->CurveType = ERawCurveTrackTypes::RCT_Float;
 
 		// Origin data
@@ -426,7 +438,7 @@ void FAnimTimelineTrack_Curve::HendleEditCurve()
 	TArray<IAnimationEditor::FCurveEditInfo> EditCurveInfo;
 	for(int32 CurveIndex = 0; CurveIndex < Curves.Num(); ++CurveIndex)
 	{
-		FSmartName Name;
+		FName Name;
 		ERawCurveTrackTypes Type;
 		int32 EditCurveIndex;
 		GetCurveEditInfo(CurveIndex, Name, Type, EditCurveIndex);
@@ -436,7 +448,7 @@ void FAnimTimelineTrack_Curve::HendleEditCurve()
 	StaticCastSharedRef<FAnimModel_AnimSequenceBase>(GetModel())->OnEditCurves.ExecuteIfBound(GetModel()->GetAnimSequenceBase(), EditCurveInfo, nullptr);
 }
 
-void FAnimTimelineTrack_Curve::GetCurveEditInfo(int32 InCurveIndex, FSmartName& OutName, ERawCurveTrackTypes& OutType, int32& OutCurveIndex) const
+void FAnimTimelineTrack_Curve::GetCurveEditInfo(int32 InCurveIndex, FName& OutName, ERawCurveTrackTypes& OutType, int32& OutCurveIndex) const
 {
 	OutName = OuterCurveName;
 	OutType = OuterType;
diff --git a/Engine/Source/Editor/Persona/Private/AnimTimeline/AnimTimelineTrack_Curve.h b/Engine/Source/Editor/Persona/Private/AnimTimeline/AnimTimelineTrack_Curve.h
--- a/Engine/Source/Editor/Persona/Private/AnimTimeline/AnimTimelineTrack_Curve.h
+++ b/Engine/Source/Editor/Persona/Private/AnimTimeline/AnimTimelineTrack_Curve.h
@@ -20,7 +20,9 @@ class FAnimTimelineTrack_Curve : public FAnimTimelineTrack, public FSelfRegister
 
 public:
 	FAnimTimelineTrack_Curve(const FText& InCurveName, const FText& InFullCurveName, const FLinearColor& InColor, const FLinearColor& InBackgroundColor, const TSharedRef<FAnimModel>& InModel);
+	UE_DEPRECATED(5.3, "Please use the constructor that takes a FName.")
 	FAnimTimelineTrack_Curve(const FRichCurve* InCurve, const FSmartName& InName, int32 InCurveIndex, ERawCurveTrackTypes InType, const FText& InCurveName, const FText& InFullCurveName, const FLinearColor& InColor, const FLinearColor& InBackgroundColor, const TSharedRef<FAnimModel>& InModel);
+	FAnimTimelineTrack_Curve(const FRichCurve* InCurve, const FName& InName, int32 InCurveIndex, ERawCurveTrackTypes InType, const FText& InCurveName, const FText& InFullCurveName, const FLinearColor& InColor, const FLinearColor& InBackgroundColor, const TSharedRef<FAnimModel>& InModel);
 	FAnimTimelineTrack_Curve(const TArray<const FRichCurve*>& InCurves, const FText& InCurveName, const FText& InFullCurveName, const FLinearColor& InColor, const FLinearColor& InBackgroundColor, const TSharedRef<FAnimModel>& InModel);
 
 	/** FAnimTimelineTrack interface */
@@ -48,7 +50,10 @@ public:
 	virtual FText GetFullCurveName(int32 InCurveIndex) const { return FullCurveName; }
 
 	/** Get the information needed to reference this curve for editing */
-	virtual void GetCurveEditInfo(int32 InCurveIndex, FSmartName& OutName, ERawCurveTrackTypes& OutType, int32& OutCurveIndex) const;
+	virtual void GetCurveEditInfo(int32 InCurveIndex, FName& OutName, ERawCurveTrackTypes& OutType, int32& OutCurveIndex) const;
+
+	UE_DEPRECATED(5.3, "Please use GetCurveEditInfo that takes a FName.")
+	virtual void GetCurveEditInfo(int32 InCurveIndex, FSmartName& OutName, ERawCurveTrackTypes& OutType, int32& OutCurveIndex) const {}
 
 	/** Delegate handing curve changing externally */
 	void HandleCurveChanged();
@@ -98,7 +103,7 @@ protected:
 	FText FullCurveName;
 
 	/** The name of the outer curve */
-	FSmartName OuterCurveName;
+	FName OuterCurveName;
 
 	/** The index into the outer curve */
 	int32 OuterCurveIndex;
diff --git a/Engine/Source/Editor/Persona/Private/AnimTimeline/AnimTimelineTrack_Curves.cpp b/Engine/Source/Editor/Persona/Private/AnimTimeline/AnimTimelineTrack_Curves.cpp
--- a/Engine/Source/Editor/Persona/Private/AnimTimeline/AnimTimelineTrack_Curves.cpp
+++ b/Engine/Source/Editor/Persona/Private/AnimTimeline/AnimTimelineTrack_Curves.cpp
@@ -6,6 +6,7 @@
 #include "Widgets/Layout/SScrollBox.h"
 #include "Widgets/Input/SSearchBox.h"
 #include "AnimSequenceTimelineCommands.h"
+#include "IEditableSkeleton.h"
 #include "SAnimCurvePicker.h"
 #include "Framework/MultiBox/MultiBoxBuilder.h"
 #include "Animation/AnimSequenceBase.h"
@@ -153,9 +154,9 @@ void FAnimTimelineTrack_Curves::FillMetadataEntryMenu(FMenuBuilder& Builder)
 
 	// Add existing curve to timeline using curve picker
 	{
-		const TSharedRef<SWidget> CurvePickerWidget = SNew(SAnimCurvePicker, GetModel()->GetEditableSkeleton())
-		.OnCurveNamePicked(this, &FAnimTimelineTrack_Curves::OnMetadataCurveNamePicked)
-		.IsCurveMarkedForExclusion(this, &FAnimTimelineTrack_Curves::IsCurveMarkedForExclusion);
+		const TSharedRef<SWidget> CurvePickerWidget = SNew(SAnimCurvePicker, &GetModel()->GetEditableSkeleton()->GetSkeleton())
+		.OnCurvePicked(this, &FAnimTimelineTrack_Curves::OnMetadataCurveNamePicked)
+		.IsCurveNameMarkedForExclusion(this, &FAnimTimelineTrack_Curves::IsCurveMarkedForExclusion);
 		Builder.AddWidget(CurvePickerWidget, FText::GetEmpty(), true);
 	}
 	
@@ -177,9 +178,9 @@ void FAnimTimelineTrack_Curves::FillVariableCurveMenu(FMenuBuilder& Builder)
 
 	// Add existing curve to timeline using curve picker
 	{
-		const TSharedRef<SWidget> CurvePickerWidget = SNew(SAnimCurvePicker, GetModel()->GetEditableSkeleton())
-		.OnCurveNamePicked(this, &FAnimTimelineTrack_Curves::OnVariableCurveNamePicked)
-		.IsCurveMarkedForExclusion(this, &FAnimTimelineTrack_Curves::IsCurveMarkedForExclusion);
+		const TSharedRef<SWidget> CurvePickerWidget = SNew(SAnimCurvePicker, &GetModel()->GetEditableSkeleton()->GetSkeleton())
+		.OnCurvePicked(this, &FAnimTimelineTrack_Curves::OnVariableCurveNamePicked)
+		.IsCurveNameMarkedForExclusion(this, &FAnimTimelineTrack_Curves::IsCurveMarkedForExclusion);
 		
 		Builder.AddWidget(CurvePickerWidget, FText::GetEmpty(), true);
 	}
@@ -187,18 +188,16 @@ void FAnimTimelineTrack_Curves::FillVariableCurveMenu(FMenuBuilder& Builder)
 	Builder.EndSection();
 }
 
-void FAnimTimelineTrack_Curves::AddMetadataEntry(USkeleton::AnimCurveUID Uid)
+void FAnimTimelineTrack_Curves::AddMetadataEntry(const FName& InCurveName)
 {
-	FSmartName NewName;
 	UAnimSequenceBase* AnimSequenceBase = GetModel()->GetAnimSequenceBase();
-	ensureAlways(AnimSequenceBase->GetSkeleton()->GetSmartNameByUID(USkeleton::AnimCurveMappingName, Uid, NewName));
 
 	IAnimationDataController& Controller = AnimSequenceBase->GetController();
 	IAnimationDataController::FScopedBracket ScopedBracket(Controller, LOCTEXT("AddCurveMetadata", "Add Curve Metadata"));
 
-	const FAnimationCurveIdentifier MetadataCurveId(NewName, ERawCurveTrackTypes::RCT_Float);
+	const FAnimationCurveIdentifier MetadataCurveId(InCurveName, ERawCurveTrackTypes::RCT_Float);
 	Controller.AddCurve(MetadataCurveId, AACF_Metadata);
-	Controller.SetCurveKeys(MetadataCurveId, { FRichCurveKey(0.f, 1.f) });	
+	Controller.SetCurveKeys(MetadataCurveId, { FRichCurveKey(0.f, 1.f) });
 }
 
 void FAnimTimelineTrack_Curves::CreateNewMetadataEntryClicked()
@@ -228,12 +227,7 @@ void FAnimTimelineTrack_Curves::CreateNewMetadataEntry(const FText& CommittedTex
 		USkeleton* Skeleton = AnimSequenceBase->GetSkeleton();
 		if(Skeleton && !CommittedText.IsEmpty())
 		{
-			FSmartName CurveName;
-
-			if(Skeleton->AddSmartNameAndModify(USkeleton::AnimCurveMappingName, FName(*CommittedText.ToString()), CurveName))
-			{
-				AddMetadataEntry(CurveName.UID);
-			}
+			AddMetadataEntry(FName(*CommittedText.ToString()));
 		}
 	}
 }
@@ -264,32 +258,23 @@ void FAnimTimelineTrack_Curves::CreateTrack(const FText& ComittedText, ETextComm
 		if(Skeleton && !ComittedText.IsEmpty())
 		{
 			const FScopedTransaction Transaction(LOCTEXT("AnimCurve_AddTrack", "Add New Curve"));
-			FSmartName NewTrackName;
 
-			Skeleton->AddSmartNameAndModify(USkeleton::AnimCurveMappingName, FName(*ComittedText.ToString()), NewTrackName);
-			if ( NewTrackName.IsValid() )
-			{
-				AddVariableCurve(NewTrackName.UID);
-			}
+			AddVariableCurve(FName(*ComittedText.ToString()));
 		}
 
 		FSlateApplication::Get().DismissAllMenus();
 	}
 }
 
-void FAnimTimelineTrack_Curves::AddVariableCurve(USkeleton::AnimCurveUID CurveUid)
+void FAnimTimelineTrack_Curves::AddVariableCurve(const FName& InCurveName)
 {
 	FScopedTransaction Transaction(LOCTEXT("AddCurve", "Add Curve"));
 
 	UAnimSequenceBase* AnimSequenceBase = GetModel()->GetAnimSequenceBase();
 	AnimSequenceBase->Modify();
-	
-	USkeleton* Skeleton = AnimSequenceBase->GetSkeleton();
-	FSmartName NewName;
-	ensureAlways(Skeleton->GetSmartNameByUID(USkeleton::AnimCurveMappingName, CurveUid, NewName));
 
 	IAnimationDataController& Controller = AnimSequenceBase->GetController();
-	const FAnimationCurveIdentifier FloatCurveId(NewName, ERawCurveTrackTypes::RCT_Float);
+	const FAnimationCurveIdentifier FloatCurveId(InCurveName, ERawCurveTrackTypes::RCT_Float);
 	Controller.AddCurve(FloatCurveId);
 }
 
@@ -303,28 +288,28 @@ bool FAnimTimelineTrack_Curves::IsShowCurvePointsEnabled() const
 	return GetDefault<UPersonaOptions>()->bTimelineDisplayCurveKeys;
 }
 
-void FAnimTimelineTrack_Curves::OnMetadataCurveNamePicked(const FSmartName& InCurveSmartName)
+void FAnimTimelineTrack_Curves::OnMetadataCurveNamePicked(const FName& InCurveName)
 {
 	FSlateApplication::Get().DismissAllMenus();
 
-	if (InCurveSmartName.IsValid())
+	if(InCurveName != NAME_None)
 	{
-		AddMetadataEntry(InCurveSmartName.UID);
+		AddMetadataEntry(InCurveName);
 	}
 }
 
-void FAnimTimelineTrack_Curves::OnVariableCurveNamePicked(const FSmartName& InCurveSmartName)
+void FAnimTimelineTrack_Curves::OnVariableCurveNamePicked(const FName& InCurveName)
 {
 	FSlateApplication::Get().DismissAllMenus();
-
-	if (InCurveSmartName.IsValid())
+	
+	if(InCurveName != NAME_None)
 	{
-		AddVariableCurve(InCurveSmartName.UID);
+		AddMetadataEntry(InCurveName);
 	}
 }
 
-bool FAnimTimelineTrack_Curves::IsCurveMarkedForExclusion(const FSmartName& InCurveSmartName)
+bool FAnimTimelineTrack_Curves::IsCurveMarkedForExclusion(const FName& InCurveName)
 {
-	return GetModel()->GetAnimSequenceBase()->GetDataModel()->FindFloatCurve(FAnimationCurveIdentifier(InCurveSmartName, ERawCurveTrackTypes::RCT_Float)) != nullptr;
+	return GetModel()->GetAnimSequenceBase()->GetDataModel()->FindFloatCurve(FAnimationCurveIdentifier(InCurveName, ERawCurveTrackTypes::RCT_Float)) != nullptr;
 }
 #undef LOCTEXT_NAMESPACE
diff --git a/Engine/Source/Editor/Persona/Private/AnimTimeline/AnimTimelineTrack_Curves.h b/Engine/Source/Editor/Persona/Private/AnimTimeline/AnimTimelineTrack_Curves.h
--- a/Engine/Source/Editor/Persona/Private/AnimTimeline/AnimTimelineTrack_Curves.h
+++ b/Engine/Source/Editor/Persona/Private/AnimTimeline/AnimTimelineTrack_Curves.h
@@ -21,12 +21,12 @@ private:
 	TSharedRef<SWidget> BuildCurvesSubMenu();
 	void FillMetadataEntryMenu(FMenuBuilder& Builder);
 	void FillVariableCurveMenu(FMenuBuilder& Builder);
-	void AddMetadataEntry(USkeleton::AnimCurveUID Uid);
+	void AddMetadataEntry(const FName& InCurveName);
 	void CreateNewMetadataEntryClicked();
 	void CreateNewMetadataEntry(const FText& CommittedText, ETextCommit::Type CommitType);
 	void CreateNewCurveClicked();
 	void CreateTrack(const FText& ComittedText, ETextCommit::Type CommitInfo);
-	void AddVariableCurve(USkeleton::AnimCurveUID CurveUid);
+	void AddVariableCurve(const FName& InCurveName);
 	void DeleteAllCurves();
 
 	/** Handlers for showing curve points */
@@ -34,9 +34,9 @@ private:
 	bool IsShowCurvePointsEnabled() const;
 
 	/** Curve Picker Callbacks */
-	void OnMetadataCurveNamePicked(const FSmartName & InCurveSmartName);
-	void OnVariableCurveNamePicked(const FSmartName & InCurveSmartName);
-	bool IsCurveMarkedForExclusion(const FSmartName & InCurveSmartName);
+	void OnMetadataCurveNamePicked(const FName& InCurveName);
+	void OnVariableCurveNamePicked(const FName& InCurveName);
+	bool IsCurveMarkedForExclusion(const FName& InCurveName);
 	
 private:
 	TSharedPtr<SWidget>	OutlinerWidget;
diff --git a/Engine/Source/Editor/Persona/Private/AnimTimeline/AnimTimelineTrack_FloatCurve.cpp b/Engine/Source/Editor/Persona/Private/AnimTimeline/AnimTimelineTrack_FloatCurve.cpp
--- a/Engine/Source/Editor/Persona/Private/AnimTimeline/AnimTimelineTrack_FloatCurve.cpp
+++ b/Engine/Source/Editor/Persona/Private/AnimTimeline/AnimTimelineTrack_FloatCurve.cpp
@@ -28,10 +28,10 @@
 ANIMTIMELINE_IMPLEMENT_TRACK(FAnimTimelineTrack_FloatCurve);
 
 FAnimTimelineTrack_FloatCurve::FAnimTimelineTrack_FloatCurve(const FFloatCurve* InCurve, const TSharedRef<FAnimModel>& InModel)
-	: FAnimTimelineTrack_Curve(&InCurve->FloatCurve, InCurve->Name, 0, ERawCurveTrackTypes::RCT_Float, FText::FromName(InCurve->Name.DisplayName), FText::FromName(InCurve->Name.DisplayName), InCurve->GetColor(), InCurve->GetColor(), InModel)
+	: FAnimTimelineTrack_Curve(&InCurve->FloatCurve, InCurve->GetName(), 0, ERawCurveTrackTypes::RCT_Float, FText::FromName(InCurve->GetName()), FText::FromName(InCurve->GetName()), InCurve->GetColor(), InCurve->GetColor(), InModel)
 	, FloatCurve(InCurve)
-	, CurveName(InCurve->Name)
-	, CurveId(FAnimationCurveIdentifier(InCurve->Name, ERawCurveTrackTypes::RCT_Float))
+	, CurveName(InCurve->GetName())
+	, CurveId(FAnimationCurveIdentifier(InCurve->GetName(), ERawCurveTrackTypes::RCT_Float))
 {
 	SetHeight(32.0f);
 }
@@ -175,94 +175,64 @@ void FAnimTimelineTrack_FloatCurve::RemoveCurve()
 {
 	UAnimSequenceBase* AnimSequenceBase = GetModel()->GetAnimSequenceBase();
 
-	if(AnimSequenceBase->GetDataModel()->FindFloatCurve(FAnimationCurveIdentifier(FloatCurve->Name, ERawCurveTrackTypes::RCT_Float)))
+	if(AnimSequenceBase->GetDataModel()->FindFloatCurve(FAnimationCurveIdentifier(FloatCurve->GetName(), ERawCurveTrackTypes::RCT_Float)))
 	{
-		FSmartName TrackName;
-		if (AnimSequenceBase->GetSkeleton()->GetSmartNameByUID(USkeleton::AnimCurveMappingName, FloatCurve->Name.UID, TrackName))
-		{
-			// Stop editing this curve in the external editor window
-			IAnimationEditor::FCurveEditInfo EditInfo(CurveName, ERawCurveTrackTypes::RCT_Float, 0);
-			
-			AnimSequenceBase->Modify(true);
+		// Stop editing this curve in the external editor window
+		IAnimationEditor::FCurveEditInfo EditInfo(CurveName, ERawCurveTrackTypes::RCT_Float, 0);
 
-			IAnimationDataController& Controller = AnimSequenceBase->GetController();
-			Controller.RemoveCurve(CurveId);
+		AnimSequenceBase->Modify(true);
 
-			AnimSequenceBase->PostEditChange();
-		}
+		IAnimationDataController& Controller = AnimSequenceBase->GetController();
+		Controller.RemoveCurve(CurveId);
+
+		AnimSequenceBase->PostEditChange();
 	}
 }
 
 void FAnimTimelineTrack_FloatCurve::OnCommitCurveName(const FText& InText, ETextCommit::Type CommitInfo)
 {
 	UAnimSequenceBase* AnimSequenceBase = GetModel()->GetAnimSequenceBase();
-
-	if (USkeleton* Skeleton = AnimSequenceBase->GetSkeleton())
+	
+	// only do this if the name isn't same
+	FText CurrentCurveName = GetLabel();
+	if (!CurrentCurveName.EqualToCaseIgnored(InText))
 	{
-		// only do this if the name isn't same
-		FText CurrentCurveName = GetLabel();
-		if (!CurrentCurveName.EqualToCaseIgnored(InText))
+		// Check that the name doesn't already exist
+		const FName RequestedName = FName(*InText.ToString());
+		
+		const TArray<FFloatCurve>& FloatCurves = AnimSequenceBase->GetDataModel()->GetFloatCurves();
+		if (!FloatCurves.ContainsByPredicate([RequestedName](const FFloatCurve& Curve)
 		{
-			// Check that the name doesn't already exist
-			const FName RequestedName = FName(*InText.ToString());
+			return Curve.GetName() == RequestedName;
+		}))
+		{
+			FFormatNamedArguments Args;
+			Args.Add(TEXT("InvalidName"), FText::FromName(RequestedName));
+			FNotificationInfo Info(FText::Format(LOCTEXT("AnimCurveRenamedInUse", "The name \"{InvalidName}\" is already used."), Args));
 
-			const FSmartNameMapping* NameMapping = Skeleton->GetSmartNameContainer(USkeleton::AnimCurveMappingName);
+			Info.bUseLargeFont = false;
+			Info.ExpireDuration = 5.0f;
 
-			FSmartName NewSmartName;
-			if (NameMapping->FindSmartName(RequestedName, NewSmartName))
-			{
-				// Already in use in this sequence, and if it's not my UID
-				const TArray<FFloatCurve>& FloatCurves = AnimSequenceBase->GetDataModel()->GetFloatCurves();
-				if (NewSmartName.UID != FloatCurve->Name.UID && !FloatCurves.ContainsByPredicate([NewSmartName](FFloatCurve& Curve)
-				{
-					return Curve.Name == NewSmartName;
-				}))
-				{
-					FFormatNamedArguments Args;
-					Args.Add(TEXT("InvalidName"), FText::FromName(RequestedName));
-					FNotificationInfo Info(FText::Format(LOCTEXT("AnimCurveRenamedInUse", "The name \"{InvalidName}\" is already used."), Args));
-
-					Info.bUseLargeFont = false;
-					Info.ExpireDuration = 5.0f;
-
-					TSharedPtr<SNotificationItem> Notification = FSlateNotificationManager::Get().AddNotification(Info);
-					if (Notification.IsValid())
-					{
-						Notification->SetCompletionState(SNotificationItem::CS_Fail);
-					}
-					return;
-				}
-			}
-			else
+			TSharedPtr<SNotificationItem> Notification = FSlateNotificationManager::Get().AddNotification(Info);
+			if (Notification.IsValid())
 			{
-				if(!Skeleton->AddSmartNameAndModify(USkeleton::AnimCurveMappingName, RequestedName, NewSmartName))
-				{
-					FNotificationInfo Info(LOCTEXT("AnimCurveRenamedError", "Failed to rename curve smart name, check the log for errors."));
-
-					Info.bUseLargeFont = false;
-					Info.ExpireDuration = 5.0f;
-
-					TSharedPtr<SNotificationItem> Notification = FSlateNotificationManager::Get().AddNotification(Info);
-					if (Notification.IsValid())
-					{
-						Notification->SetCompletionState(SNotificationItem::CS_Fail);
-					}
-					return;
-				}
+				Notification->SetCompletionState(SNotificationItem::CS_Fail);
 			}
-
-			IAnimationDataController& Controller = AnimSequenceBase->GetController();
-			IAnimationDataController::FScopedBracket ScopedBracket(Controller, LOCTEXT("CurveEditor_RenameCurve", "Renaming Curve"));
-          
-			FAnimationCurveIdentifier NewCurveId(NewSmartName, ERawCurveTrackTypes::RCT_Float);
-			Controller.RenameCurve(CurveId, NewCurveId);
+			return;
 		}
+	
+
+		IAnimationDataController& Controller = AnimSequenceBase->GetController();
+		IAnimationDataController::FScopedBracket ScopedBracket(Controller, LOCTEXT("CurveEditor_RenameCurve", "Renaming Curve"));
+      
+		FAnimationCurveIdentifier NewCurveId(RequestedName, ERawCurveTrackTypes::RCT_Float);
+		Controller.RenameCurve(CurveId, NewCurveId);
 	}
 }
 
 FText FAnimTimelineTrack_FloatCurve::GetLabel() const
 {
-	return FAnimTimelineTrack_FloatCurve::GetFloatCurveName(GetModel(), CurveName);
+	return FText::FromName(FloatCurve->GetName());
 }
 
 void FAnimTimelineTrack_FloatCurve::Copy(UAnimTimelineClipboardContent* InOutClipboard) const
@@ -272,13 +242,12 @@ void FAnimTimelineTrack_FloatCurve::Copy(UAnimTimelineClipboardContent* InOutCli
 	UFloatCurveCopyObject * CopyableCurve = UAnimCurveBaseCopyObject::Create<UFloatCurveCopyObject>();
 
 	// Copy raw curve data
-	CopyableCurve->Curve.Name = FloatCurve->Name;
+	CopyableCurve->Curve.SetName(FloatCurve->GetName());
 	CopyableCurve->Curve.SetCurveTypeFlags(FloatCurve->GetCurveTypeFlags());
 	CopyableCurve->Curve.CopyCurve(*FloatCurve);
 
 	// Copy curve identifier data
-	CopyableCurve->DisplayName = CurveName.DisplayName;
-	CopyableCurve->UID = CurveName.UID;
+	CopyableCurve->CurveName = CurveName;
 	CopyableCurve->CurveType = ERawCurveTrackTypes::RCT_Float;
 	CopyableCurve->Channel = CurveId.Channel;
 	CopyableCurve->Axis = CurveId.Axis;
@@ -289,21 +258,6 @@ void FAnimTimelineTrack_FloatCurve::Copy(UAnimTimelineClipboardContent* InOutCli
 	InOutClipboard->Curves.Add(CopyableCurve);
 }
 
-FText FAnimTimelineTrack_FloatCurve::GetFloatCurveName(const TSharedRef<FAnimModel>& InModel, const FSmartName& InSmartName)
-{
-	const FSmartNameMapping* NameMapping = InModel->GetAnimSequenceBase()->GetSkeleton()->GetSmartNameContainer(USkeleton::AnimCurveMappingName);
-	if(NameMapping)
-	{
-		FName CurveName;
-		if(NameMapping->GetName(InSmartName.UID, CurveName))
-		{
-			return FText::FromName(CurveName);
-		}
-	}
-
-	return FText::FromName(InSmartName.DisplayName);
-}
-
 bool FAnimTimelineTrack_FloatCurve::CanEditCurve(int32 InCurveIndex) const
 {
 	return !FloatCurve->GetCurveTypeFlag(AACF_Metadata);
@@ -391,7 +345,7 @@ FLinearColor FAnimTimelineTrack_FloatCurve::GetCurveColor(int32 InCurveIndex) co
 	return FloatCurve->Color; 
 }
 
-void FAnimTimelineTrack_FloatCurve::GetCurveEditInfo(int32 InCurveIndex, FSmartName& OutName, ERawCurveTrackTypes& OutType, int32& OutCurveIndex) const
+void FAnimTimelineTrack_FloatCurve::GetCurveEditInfo(int32 InCurveIndex, FName& OutName, ERawCurveTrackTypes& OutType, int32& OutCurveIndex) const
 {
 	OutName = CurveName;
 	OutType = ERawCurveTrackTypes::RCT_Float;
diff --git a/Engine/Source/Editor/Persona/Private/AnimTimeline/AnimTimelineTrack_FloatCurve.h b/Engine/Source/Editor/Persona/Private/AnimTimeline/AnimTimelineTrack_FloatCurve.h
--- a/Engine/Source/Editor/Persona/Private/AnimTimeline/AnimTimelineTrack_FloatCurve.h
+++ b/Engine/Source/Editor/Persona/Private/AnimTimeline/AnimTimelineTrack_FloatCurve.h
@@ -25,18 +25,26 @@ public:
 	virtual void RequestRename() override;
 	virtual void AddCurveTrackButton(TSharedPtr<SHorizontalBox> InnerHorizontalBox) override;
 	virtual FLinearColor GetCurveColor(int32 InCurveIndex) const override;
-	virtual void GetCurveEditInfo(int32 InCurveIndex, FSmartName& OutName, ERawCurveTrackTypes& OutType, int32& OutCurveIndex) const override;
+	virtual void GetCurveEditInfo(int32 InCurveIndex, FName& OutName, ERawCurveTrackTypes& OutType, int32& OutCurveIndex) const override;
 	virtual bool SupportsCopy() const override { return true; }
 	virtual void Copy(UAnimTimelineClipboardContent* InOutClipboard) const override;
 	
 	/** Access the curve we are editing */
 	const FFloatCurve* GetFloatCurve() { return FloatCurve; }
 
-	/** Helper function used to get a smart name for a curve */
-	static FText GetFloatCurveName(const TSharedRef<FAnimModel>& InModel, const FSmartName& InSmartName);
+	UE_DEPRECATED(5.3, "This function is longer used")
+	static FText GetFloatCurveName(const TSharedRef<FAnimModel>& InModel, const FSmartName& InSmartName) { return FText::GetEmpty(); }
 
 	/** Get this curves name */
-	FSmartName GetName() const { return CurveName; }
+	FName GetFName() const { return CurveName; }
+	
+	UE_DEPRECATED(5.3, "Please use GetFName instead.")
+	FSmartName GetName() const
+	{
+		PRAGMA_DISABLE_DEPRECATION_WARNINGS
+		return FSmartName(CurveName, 0);
+		PRAGMA_ENABLE_DEPRECATION_WARNINGS
+	}
 
 private:
 	void ConvertCurveToMetaData();
@@ -49,7 +57,7 @@ private:
 	const FFloatCurve* FloatCurve;
 
 	/** The curve name and identifier */
-	FSmartName CurveName;
+	FName CurveName;
 	FAnimationCurveIdentifier CurveId;
 
 	/** Label we can edit */
diff --git a/Engine/Source/Editor/Persona/Private/AnimTimeline/AnimTimelineTrack_TransformCurve.cpp b/Engine/Source/Editor/Persona/Private/AnimTimeline/AnimTimelineTrack_TransformCurve.cpp
--- a/Engine/Source/Editor/Persona/Private/AnimTimeline/AnimTimelineTrack_TransformCurve.cpp
+++ b/Engine/Source/Editor/Persona/Private/AnimTimeline/AnimTimelineTrack_TransformCurve.cpp
@@ -19,10 +19,10 @@
 ANIMTIMELINE_IMPLEMENT_TRACK(FAnimTimelineTrack_TransformCurve);
 
 FAnimTimelineTrack_TransformCurve::FAnimTimelineTrack_TransformCurve(const FTransformCurve* InCurve, const TSharedRef<FAnimModel>& InModel)
-	: FAnimTimelineTrack_Curve(FAnimTimelineTrack_TransformCurve::GetTransformCurveName(InModel, InCurve->Name), FAnimTimelineTrack_TransformCurve::GetTransformCurveName(InModel, InCurve->Name), InCurve->GetColor(), InCurve->GetColor(), InModel)
+	: FAnimTimelineTrack_Curve(FText::FromName(InCurve->GetName()), FText::FromName(InCurve->GetName()), InCurve->GetColor(), InCurve->GetColor(), InModel)
 	, TransformCurve(InCurve)
-	, CurveName(InCurve->Name)
-	, CurveId(InCurve->Name, ERawCurveTrackTypes::RCT_Transform)
+	, CurveName(InCurve->GetName())
+	, CurveId(InCurve->GetName(), ERawCurveTrackTypes::RCT_Transform)
 {
 	Curves.Add(&InCurve->TranslationCurve.FloatCurves[0]);
 	Curves.Add(&InCurve->TranslationCurve.FloatCurves[1]);
@@ -74,13 +74,12 @@ void FAnimTimelineTrack_TransformCurve::Copy(UAnimTimelineClipboardContent* InOu
 	UTransformCurveCopyObject * CopyableCurve = UAnimCurveBaseCopyObject::Create<UTransformCurveCopyObject>();
 
 	// Copy raw curve data
-	CopyableCurve->Curve.Name = TransformCurve->Name;
+	CopyableCurve->Curve.SetName(TransformCurve->GetName());
 	CopyableCurve->Curve.SetCurveTypeFlags(TransformCurve->GetCurveTypeFlags());
 	CopyableCurve->Curve.CopyCurve(*TransformCurve);
 
 	// Copy curve identifier data
-	CopyableCurve->DisplayName = CurveName.DisplayName;
-	CopyableCurve->UID = CurveName.UID;
+	CopyableCurve->CurveName = CurveName;
 	CopyableCurve->CurveType = ERawCurveTrackTypes::RCT_Transform;
 	CopyableCurve->Channel = CurveId.Channel;
 	CopyableCurve->Axis = CurveId.Axis;
@@ -91,21 +90,6 @@ void FAnimTimelineTrack_TransformCurve::Copy(UAnimTimelineClipboardContent* InOu
 	InOutClipboard->Curves.Add(CopyableCurve);
 }
 
-FText FAnimTimelineTrack_TransformCurve::GetTransformCurveName(const TSharedRef<FAnimModel>& InModel, const FSmartName& InSmartName)
-{
-	const FSmartNameMapping* NameMapping = InModel->GetAnimSequenceBase()->GetSkeleton()->GetSmartNameContainer(USkeleton::AnimTrackCurveMappingName);
-	if(NameMapping)
-	{
-		FName CurveName;
-		if(NameMapping->GetName(InSmartName.UID, CurveName))
-		{
-			return FText::FromName(CurveName);
-		}
-	}
-
-	return FText::FromName(InSmartName.DisplayName);
-}
-
 TSharedRef<SWidget> FAnimTimelineTrack_TransformCurve::BuildCurveTrackMenu()
 {
 	FMenuBuilder MenuBuilder(true, GetModel()->GetCommandList());
@@ -149,16 +133,13 @@ void FAnimTimelineTrack_TransformCurve::DeleteTrack()
 	if(AnimSequenceBase->GetDataModel()->FindTransformCurve(CurveId))
 	{
 		const FScopedTransaction Transaction(LOCTEXT("AnimCurve_DeleteTrack", "Delete Curve"));
-		FSmartName CurveToDelete;
-		if (AnimSequenceBase->GetSkeleton()->GetSmartNameByUID(USkeleton::AnimTrackCurveMappingName, TransformCurve->Name.UID, CurveToDelete))
-		{
-			IAnimationDataController& Controller = AnimSequenceBase->GetController();
-			Controller.RemoveCurve(CurveId);
 
-			if (GetModel()->GetPreviewScene()->GetPreviewMeshComponent()->PreviewInstance != nullptr)
-			{
-				GetModel()->GetPreviewScene()->GetPreviewMeshComponent()->PreviewInstance->RefreshCurveBoneControllers();
-			}
+		IAnimationDataController& Controller = AnimSequenceBase->GetController();
+		Controller.RemoveCurve(CurveId);
+
+		if (GetModel()->GetPreviewScene()->GetPreviewMeshComponent()->PreviewInstance != nullptr)
+		{
+			GetModel()->GetPreviewScene()->GetPreviewMeshComponent()->PreviewInstance->RefreshCurveBoneControllers();
 		}
 	}
 }
@@ -184,9 +165,9 @@ void FAnimTimelineTrack_TransformCurve::ToggleEnabled()
 	}
 }
 
-void FAnimTimelineTrack_TransformCurve::GetCurveEditInfo(int32 InCurveIndex, FSmartName& OutName, ERawCurveTrackTypes& OutType, int32& OutCurveIndex) const
+void FAnimTimelineTrack_TransformCurve::GetCurveEditInfo(int32 InCurveIndex, FName& OutName, ERawCurveTrackTypes& OutType, int32& OutCurveIndex) const
 {
-	OutName = TransformCurve->Name;
+	OutName = TransformCurve->GetName();
 	OutType = ERawCurveTrackTypes::RCT_Transform;
 	OutCurveIndex = InCurveIndex;
 }
diff --git a/Engine/Source/Editor/Persona/Private/AnimTimeline/AnimTimelineTrack_TransformCurve.h b/Engine/Source/Editor/Persona/Private/AnimTimeline/AnimTimelineTrack_TransformCurve.h
--- a/Engine/Source/Editor/Persona/Private/AnimTimeline/AnimTimelineTrack_TransformCurve.h
+++ b/Engine/Source/Editor/Persona/Private/AnimTimeline/AnimTimelineTrack_TransformCurve.h
@@ -20,18 +20,25 @@ public:
 	virtual FText GetFullCurveName(int32 InCurveIndex) const override;
 	virtual TSharedRef<SWidget> BuildCurveTrackMenu() override;
 	virtual bool CanEditCurve(int32 InCurveIndex) const override { return true; }
-	virtual void GetCurveEditInfo(int32 InCurveIndex, FSmartName& OutName, ERawCurveTrackTypes& OutType, int32& OutCurveIndex) const override;
+	virtual void GetCurveEditInfo(int32 InCurveIndex, FName& OutName, ERawCurveTrackTypes& OutType, int32& OutCurveIndex) const override;
 	virtual bool SupportsCopy() const override { return true; }
 	virtual void Copy(UAnimTimelineClipboardContent* InOutClipboard) const override;
 	
 	/** Access the curve we are editing */
 	const FTransformCurve& GetTransformCurve() { return *TransformCurve; }
 
-	/** Helper function used to get a smart name for a curve */
-	static FText GetTransformCurveName(const TSharedRef<FAnimModel>& InModel, const FSmartName& InSmartName);
+	UE_DEPRECATED(5.3, "This function is no longer used")
+	static FText GetTransformCurveName(const TSharedRef<FAnimModel>& InModel, const FSmartName& InSmartName) { return FText::GetEmpty(); }
 
-	/** Get this curves name */
-	FSmartName GetName() const { return CurveName; }
+	UE_DEPRECATED(5.3, "Please use GetFName instead.")
+	FSmartName GetName() const
+	{
+		PRAGMA_DISABLE_DEPRECATION_WARNINGS
+		return FSmartName(CurveName, 0);
+		PRAGMA_ENABLE_DEPRECATION_WARNINGS
+	}
+	
+	FName GetFName() const { return CurveName; }
 
 private:
 	/** Delete this track via the track menu */
@@ -48,6 +55,6 @@ private:
 	const FTransformCurve* TransformCurve;
 
 	/** The curve name and identifier */
-	FSmartName CurveName;
+	FName CurveName;
 	FAnimationCurveIdentifier CurveId;
 };
diff --git a/Engine/Source/Editor/Persona/Private/AnimTimeline/AnimTimelineTrack_VectorCurve.cpp b/Engine/Source/Editor/Persona/Private/AnimTimeline/AnimTimelineTrack_VectorCurve.cpp
--- a/Engine/Source/Editor/Persona/Private/AnimTimeline/AnimTimelineTrack_VectorCurve.cpp
+++ b/Engine/Source/Editor/Persona/Private/AnimTimeline/AnimTimelineTrack_VectorCurve.cpp
@@ -9,7 +9,7 @@
 
 ANIMTIMELINE_IMPLEMENT_TRACK(FAnimTimelineTrack_VectorCurve);
 
-FAnimTimelineTrack_VectorCurve::FAnimTimelineTrack_VectorCurve(const FVectorCurve* InCurve, const FSmartName& InName, int32 InCurveIndex, ERawCurveTrackTypes InType, const FText& InCurveName, const FText& InFullCurveName, const FLinearColor& InColor, const TSharedRef<FAnimModel>& InModel)
+FAnimTimelineTrack_VectorCurve::FAnimTimelineTrack_VectorCurve(const FVectorCurve* InCurve, const FName& InName, int32 InCurveIndex, ERawCurveTrackTypes InType, const FText& InCurveName, const FText& InFullCurveName, const FLinearColor& InColor, const TSharedRef<FAnimModel>& InModel)
 	: FAnimTimelineTrack_Curve(InCurveName, InFullCurveName, InColor, InColor, InModel)
 	, VectorCurve(InCurve)
 {
@@ -22,6 +22,19 @@ FAnimTimelineTrack_VectorCurve::FAnimTimelineTrack_VectorCurve(const FVectorCurv
 	Curves.Add(&InCurve->FloatCurves[2]);
 }
 
+FAnimTimelineTrack_VectorCurve::FAnimTimelineTrack_VectorCurve(const FVectorCurve* InCurve, const FSmartName& InName, int32 InCurveIndex, ERawCurveTrackTypes InType, const FText& InCurveName, const FText& InFullCurveName, const FLinearColor& InColor, const TSharedRef<FAnimModel>& InModel)
+	: FAnimTimelineTrack_Curve(InCurveName, InFullCurveName, InColor, InColor, InModel)
+	, VectorCurve(InCurve)
+{
+	OuterCurveName = InName.DisplayName;
+	OuterCurveIndex = InCurveIndex;
+	OuterType = InType;
+
+	Curves.Add(&InCurve->FloatCurves[0]);
+	Curves.Add(&InCurve->FloatCurves[1]);
+	Curves.Add(&InCurve->FloatCurves[2]);
+}
+
 FLinearColor FAnimTimelineTrack_VectorCurve::GetCurveColor(int32 InCurveIndex) const
 {
 	static const FLinearColor Colors[3] =
@@ -48,7 +61,7 @@ FText FAnimTimelineTrack_VectorCurve::GetFullCurveName(int32 InCurveIndex) const
 	return FText::Format(LOCTEXT("TransformVectorFormat", "{0}.{1}"), FullCurveName, TrackNames[InCurveIndex]);
 }
 
-void FAnimTimelineTrack_VectorCurve::GetCurveEditInfo(int32 InCurveIndex, FSmartName& OutName, ERawCurveTrackTypes& OutType, int32& OutCurveIndex) const
+void FAnimTimelineTrack_VectorCurve::GetCurveEditInfo(int32 InCurveIndex, FName& OutName, ERawCurveTrackTypes& OutType, int32& OutCurveIndex) const
 {
 	OutName = OuterCurveName;
 	OutType = OuterType;
diff --git a/Engine/Source/Editor/Persona/Private/AnimTimeline/AnimTimelineTrack_VectorCurve.h b/Engine/Source/Editor/Persona/Private/AnimTimeline/AnimTimelineTrack_VectorCurve.h
--- a/Engine/Source/Editor/Persona/Private/AnimTimeline/AnimTimelineTrack_VectorCurve.h
+++ b/Engine/Source/Editor/Persona/Private/AnimTimeline/AnimTimelineTrack_VectorCurve.h
@@ -13,13 +13,16 @@ class FAnimTimelineTrack_VectorCurve : public FAnimTimelineTrack_Curve
 	ANIMTIMELINE_DECLARE_TRACK(FAnimTimelineTrack_VectorCurve, FAnimTimelineTrack_Curve);
 
 public:
+	FAnimTimelineTrack_VectorCurve(const FVectorCurve* InCurve, const FName& InName, int32 InCurveIndex, ERawCurveTrackTypes InType, const FText& InCurveName, const FText& InFullCurveName, const FLinearColor& InColor, const TSharedRef<FAnimModel>& InModel);
+
+	UE_DEPRECATED(5.3, "Please use the construtor that takes an FName.")
 	FAnimTimelineTrack_VectorCurve(const FVectorCurve* InCurve, const FSmartName& InName, int32 InCurveIndex, ERawCurveTrackTypes InType, const FText& InCurveName, const FText& InFullCurveName, const FLinearColor& InColor, const TSharedRef<FAnimModel>& InModel);
 
 	/** FAnimTimelineTrack_Curve interface */
 	virtual FLinearColor GetCurveColor(int32 InCurveIndex) const override;
 	virtual FText GetFullCurveName(int32 InCurveIndex) const override;
 	virtual bool CanEditCurve(int32 InCurveIndex) const override { return true; }
-	virtual void GetCurveEditInfo(int32 InCurveIndex, FSmartName& OutName, ERawCurveTrackTypes& OutType, int32& OutCurveIndex) const override;
+	virtual void GetCurveEditInfo(int32 InCurveIndex, FName& OutName, ERawCurveTrackTypes& OutType, int32& OutCurveIndex) const override;
 
 	/** Access the curve we are editing */
 	const FVectorCurve& GetVectorCurve() { return *VectorCurve; }
diff --git a/Engine/Source/Editor/Persona/Private/AssetSearchBoxUtilPersona.cpp b/Engine/Source/Editor/Persona/Private/AssetSearchBoxUtilPersona.cpp
--- a/Engine/Source/Editor/Persona/Private/AssetSearchBoxUtilPersona.cpp
+++ b/Engine/Source/Editor/Persona/Private/AssetSearchBoxUtilPersona.cpp
@@ -6,6 +6,8 @@
 #include "Engine/SkeletalMesh.h"
 #include "Engine/SkeletalMeshSocket.h"
 #include "Editor.h"
+#include "AssetRegistry/AssetRegistryModule.h"
+#include "Animation/AnimSequence.h"
 
 SAssetSearchBoxForBones::SAssetSearchBoxForBones()
 {
@@ -156,13 +158,30 @@ TArray<FAssetSearchBoxSuggestion> SAssetSearchBoxForCurves::GetCurveSearchSugges
 	TArray<FAssetSearchBoxSuggestion> PossibleSuggestions;
 	if (USkeleton* Skel = Skeleton.Get())
 	{
-		if (const FSmartNameMapping* Mapping = Skel->GetSmartNameContainer(USkeleton::AnimCurveMappingName))
+		// We use the asset registry to query all assets with the supplied skeleton, and accumulate their curve names
+		const FAssetRegistryModule& AssetRegistryModule = FModuleManager::LoadModuleChecked<FAssetRegistryModule>("AssetRegistry");
+
+		FARFilter Filter;
+		Filter.bRecursiveClasses = true;
+		Filter.ClassPaths.Add(UAnimSequence::StaticClass()->GetClassPathName());
+		Filter.TagsAndValues.Add(TEXT("Skeleton"), FAssetData(Skel).GetExportTextName());
+
+		TArray<FAssetData> Assets;
+		AssetRegistryModule.Get().GetAssets(Filter, Assets);
+		
+		for(const FAssetData& AssetData : Assets)
 		{
-			TArray<FName> Names; 
-			Mapping->FillNameArray(Names);
-			for (const FName& Name : Names)
+			const FString TagValue = AssetData.GetTagValueRef<FString>(USkeleton::CurveNameTag);
+			if (!TagValue.IsEmpty())
 			{
-				PossibleSuggestions.Add(FAssetSearchBoxSuggestion::MakeSimpleSuggestion(Name.ToString()));
+				TArray<FString> ParsedCurveNames;
+				if(TagValue.ParseIntoArray(ParsedCurveNames, *USkeleton::CurveTagDelimiter, true))
+				{
+					for (const FString& CurveString : ParsedCurveNames)
+					{
+						PossibleSuggestions.Add(FAssetSearchBoxSuggestion::MakeSimpleSuggestion(CurveString));
+					}
+				}
 			}
 		}
 	}
diff --git a/Engine/Source/Editor/Persona/Private/PersonaModule.cpp b/Engine/Source/Editor/Persona/Private/PersonaModule.cpp
--- a/Engine/Source/Editor/Persona/Private/PersonaModule.cpp
+++ b/Engine/Source/Editor/Persona/Private/PersonaModule.cpp
@@ -92,6 +92,7 @@
 #include "SAnimationBlendSpaceGridWidget.h"
 #include "SAnimSequenceCurveEditor.h"
 #include "AnimSequenceTimelineCommands.h"
+#include "SAnimCurvePicker.h"
 #include "SAnimMontageSectionsPanel.h"
 #include "Animation/AnimSequenceHelpers.h"
 #include "SkeletalMeshReferenceSectionDetails.h"
@@ -100,6 +101,7 @@
 #include "ContentBrowserMenuContexts.h"
 #include "SAssetView.h"
 #include "ToolMenus.h"
+#include "SAnimAssetFindReplace.h"
 
 IMPLEMENT_MODULE( FPersonaModule, Persona );
 
@@ -287,6 +289,11 @@ TSharedRef<class FWorkflowTabFactory> FPersonaModule::CreateCurveViewerTabFactor
 	return MakeShareable(new FAnimCurveViewerTabSummoner(InHostingApp, InEditableSkeleton, InPreviewScene, InOnObjectsSelected));
 }
 
+TSharedRef<class FWorkflowTabFactory> FPersonaModule::CreateCurveMetadataEditorTabFactory(const TSharedRef<class FWorkflowCentricApplication>& InHostingApp, UObject* InMetadataHost, const TSharedRef<IPersonaPreviewScene>& InPreviewScene, FOnObjectsSelected InOnObjectsSelected) const
+{
+	return MakeShared<FAnimCurveMetadataEditorTabSummoner>(InHostingApp, InMetadataHost, InPreviewScene, InOnObjectsSelected);
+}
+
 TSharedRef<class FWorkflowTabFactory> FPersonaModule::CreateRetargetSourcesTabFactory(const TSharedRef<class FWorkflowCentricApplication>& InHostingApp, const TSharedRef<class IEditableSkeleton>& InEditableSkeleton, const TSharedRef<IPersonaPreviewScene>& InPreviewScene, FSimpleMulticastDelegate& InOnPostUndo) const
 {
 	return MakeShareable(new FRetargetSourcesTabSummoner(InHostingApp, InEditableSkeleton, InPreviewScene, InOnPostUndo));
@@ -368,6 +375,11 @@ TSharedRef<class FWorkflowTabFactory> FPersonaModule::CreateAnimMontageSectionsT
 	return MakeShareable(new FAnimMontageSectionsSummoner(InHostingApp, InPersonaToolkit, InOnSectionsChanged));
 }
 
+TSharedRef<FWorkflowTabFactory> FPersonaModule::CreateAnimAssetFindReplaceTabFactory(const TSharedRef<FWorkflowCentricApplication>& InHostingApp, const FAnimAssetFindReplaceConfig& InConfig) const
+{
+	return MakeShared<FAnimAssetFindReplaceSummoner>(InHostingApp, InConfig);
+}
+
 TSharedRef<SWidget> FPersonaModule::CreateEditorWidgetForAnimDocument(const TSharedRef<IAnimationEditor>& InHostingApp, UObject* InAnimAsset, const FAnimDocumentArgs& InArgs, FString& OutDocumentLink)
 {
 	TSharedPtr<SWidget> Result = SNullWidget::NullWidget;
@@ -511,121 +523,117 @@ void PopulateWithAssets(const FTopLevelAssetPath& ClassName, FName SkeletonMembe
 	AssetRegistryModule.Get().GetAssets(Filter, OutAssets);
 }
 
-void FPersonaModule::TestSkeletonCurveNamesForUse(const TSharedRef<IEditableSkeleton>& InEditableSkeleton) const
+void FPersonaModule::TestSkeletonCurveMetaDataForUse(const TSharedRef<IEditableSkeleton>& InEditableSkeleton) const
 {
 	const USkeleton& Skeleton = InEditableSkeleton->GetSkeleton();
+	const FString SkeletonString = FAssetData(&Skeleton).GetExportTextName();
 
-	if (const FSmartNameMapping* Mapping = Skeleton.GetSmartNameContainer(USkeleton::AnimCurveMappingName))
-	{
-		const FString SkeletonString = FAssetData(&Skeleton).GetExportTextName();
-
-		TArray<FAssetData> SkeletalMeshes;
-		PopulateWithAssets(USkeletalMesh::StaticClass()->GetClassPathName(), USkeletalMesh::GetSkeletonMemberName(), SkeletonString, SkeletalMeshes);
-		TArray<FAssetData> Animations;
-		PopulateWithAssets(UAnimSequence::StaticClass()->GetClassPathName(), USkeletalMesh::GetSkeletonMemberName(), SkeletonString, Animations);
+	TArray<FAssetData> SkeletalMeshes;
+	PopulateWithAssets(USkeletalMesh::StaticClass()->GetClassPathName(), USkeletalMesh::GetSkeletonMemberName(), SkeletonString, SkeletalMeshes);
+	TArray<FAssetData> Animations;
+	PopulateWithAssets(UAnimSequence::StaticClass()->GetClassPathName(), USkeletalMesh::GetSkeletonMemberName(), SkeletonString, Animations);
 
-		FText TimeTakenMessage = FText::Format(LOCTEXT("TimeTakenWarning", "In order to verify curve usage all Skeletal Meshes and Animations that use this skeleton will be loaded, this may take some time.\n\nProceed?\n\nNumber of Meshes: {0}\nNumber of Animations: {1}"), FText::AsNumber(SkeletalMeshes.Num()), FText::AsNumber(Animations.Num()));
+	FText TimeTakenMessage = FText::Format(LOCTEXT("TimeTakenWarning", "In order to verify curve usage all Skeletal Meshes and Animations that use this skeleton will be loaded, this may take some time.\n\nProceed?\n\nNumber of Meshes: {0}\nNumber of Animations: {1}"), FText::AsNumber(SkeletalMeshes.Num()), FText::AsNumber(Animations.Num()));
 
-		if (FMessageDialog::Open(EAppMsgType::YesNo, TimeTakenMessage) == EAppReturnType::Yes)
+	if (FMessageDialog::Open(EAppMsgType::YesNo, TimeTakenMessage) == EAppReturnType::Yes)
+	{
+		const FText LoadingStatusUpdate = FText::Format(LOCTEXT("VerifyCurves_LoadingAllAnimations", "Loading all animations for skeleton '{0}'"), FText::FromString(Skeleton.GetName()));
 		{
-			const FText LoadingStatusUpdate = FText::Format(LOCTEXT("VerifyCurves_LoadingAllAnimations", "Loading all animations for skeleton '{0}'"), FText::FromString(Skeleton.GetName()));
-			{
-				FScopedSlowTask LoadingAnimSlowTask(static_cast<float>(Animations.Num()), LoadingStatusUpdate);
-				LoadingAnimSlowTask.MakeDialog();
+			FScopedSlowTask LoadingAnimSlowTask(static_cast<float>(Animations.Num()), LoadingStatusUpdate);
+			LoadingAnimSlowTask.MakeDialog();
 
-				// Loop through all animations to load then, this makes sure smart names are all up to date
-				for (const FAssetData& Anim : Animations)
-				{
-					LoadingAnimSlowTask.EnterProgressFrame();
-					UAnimSequence* Seq = Cast<UAnimSequence>(Anim.GetAsset());
-				}
+			// Loop through all animations to load then, this makes sure smart names are all up to date
+			for (const FAssetData& Anim : Animations)
+			{
+				LoadingAnimSlowTask.EnterProgressFrame();
+				UAnimSequence* Seq = Cast<UAnimSequence>(Anim.GetAsset());
 			}
+		}
+
+		// Grab all curve names for this skeleton
+		TArray<FName> UnusedNames;
+		Skeleton.GetCurveMetaDataNames(UnusedNames);
 
-			// Grab all curve names for this skeleton
-			TArray<FName> UnusedNames;
-			Mapping->FillNameArray(UnusedNames);
+		const FText ProcessingStatusUpdate = FText::Format(LOCTEXT("VerifyCurves_ProcessingCurveUsage", "Looking at curve useage for each skeletal mesh of skeleton '{0}'"), FText::FromString(Skeleton.GetName()));
+		{
+			FScopedSlowTask LoadingSkelMeshSlowTask(static_cast<float>(SkeletalMeshes.Num()), ProcessingStatusUpdate);
+			LoadingSkelMeshSlowTask.MakeDialog();
 
-			const FText ProcessingStatusUpdate = FText::Format(LOCTEXT("VerifyCurves_ProcessingCurveUsage", "Looking at curve useage for each skeletal mesh of skeleton '{0}'"), FText::FromString(Skeleton.GetName()));
+			for (int32 MeshIdx = 0; MeshIdx < SkeletalMeshes.Num(); ++MeshIdx)
 			{
-				FScopedSlowTask LoadingSkelMeshSlowTask(static_cast<float>(SkeletalMeshes.Num()), ProcessingStatusUpdate);
-				LoadingSkelMeshSlowTask.MakeDialog();
+				LoadingSkelMeshSlowTask.EnterProgressFrame();
 
-				for (int32 MeshIdx = 0; MeshIdx < SkeletalMeshes.Num(); ++MeshIdx)
-				{
-					LoadingSkelMeshSlowTask.EnterProgressFrame();
+				const USkeletalMesh* Mesh = Cast<USkeletalMesh>(SkeletalMeshes[MeshIdx].GetAsset());
 
-					const USkeletalMesh* Mesh = Cast<USkeletalMesh>(SkeletalMeshes[MeshIdx].GetAsset());
+				// Filter morph targets from curves
+				const TArray<UMorphTarget*>& MorphTargets = Mesh->GetMorphTargets();
+				for (int32 I = 0; I < MorphTargets.Num(); ++I)
+				{
+					const int32 CurveIndex = UnusedNames.RemoveSingleSwap(MorphTargets[I]->GetFName(), false);
+				}
 
-					// Filter morph targets from curves
-					const TArray<UMorphTarget*>& MorphTargets = Mesh->GetMorphTargets();
-					for (int32 I = 0; I < MorphTargets.Num(); ++I)
+				// Filter material params from curves
+				for (const FSkeletalMaterial& Mat : Mesh->GetMaterials())
+				{
+					if (UnusedNames.Num() == 0)
 					{
-						const int32 CurveIndex = UnusedNames.RemoveSingleSwap(MorphTargets[I]->GetFName(), false);
+						break; // Done
 					}
 
-					// Filter material params from curves
-					for (const FSkeletalMaterial& Mat : Mesh->GetMaterials())
+					UMaterial* Material = (Mat.MaterialInterface != nullptr) ? Mat.MaterialInterface->GetMaterial() : nullptr;
+					if (Material)
 					{
-						if (UnusedNames.Num() == 0)
-						{
-							break; // Done
-						}
-
-						UMaterial* Material = (Mat.MaterialInterface != nullptr) ? Mat.MaterialInterface->GetMaterial() : nullptr;
-						if (Material)
-						{
-							TArray<FMaterialParameterInfo> OutParameterInfo;
-							TArray<FGuid> OutParameterIds;
+						TArray<FMaterialParameterInfo> OutParameterInfo;
+						TArray<FGuid> OutParameterIds;
 
-							// Retrieve all scalar parameter names from the material
-							Mat.MaterialInterface->GetAllScalarParameterInfo(OutParameterInfo, OutParameterIds);
+						// Retrieve all scalar parameter names from the material
+						Mat.MaterialInterface->GetAllScalarParameterInfo(OutParameterInfo, OutParameterIds);
 
-							for (FMaterialParameterInfo SPInfo : OutParameterInfo)
-							{
-								UnusedNames.RemoveSingleSwap(SPInfo.Name);
-							}
+						for (FMaterialParameterInfo SPInfo : OutParameterInfo)
+						{
+							UnusedNames.RemoveSingleSwap(SPInfo.Name);
 						}
 					}
 				}
 			}
+		}
+
+		FMessageLog CurveOutput("Persona");
+		CurveOutput.NewPage(LOCTEXT("PersonaMessageLogName", "Persona"));
 
-			FMessageLog CurveOutput("Persona");
-			CurveOutput.NewPage(LOCTEXT("PersonaMessageLogName", "Persona"));
+		bool bFoundIssue = false;
 
-			bool bFoundIssue = false;
+		const FText ProcessingAnimStatusUpdate = FText::Format(LOCTEXT("FindUnusedCurves_ProcessingSkeletalMeshes", "Finding animations that reference unused curves on skeleton '{0}'"), FText::FromString(Skeleton.GetName()));
+		{
+			FScopedSlowTask ProcessingAnimationsSlowTask(static_cast<float>(Animations.Num()), ProcessingAnimStatusUpdate);
+			ProcessingAnimationsSlowTask.MakeDialog();
 
-			const FText ProcessingAnimStatusUpdate = FText::Format(LOCTEXT("FindUnusedCurves_ProcessingSkeletalMeshes", "Finding animations that reference unused curves on skeleton '{0}'"), FText::FromString(Skeleton.GetName()));
+			for (const FAssetData& Anim : Animations)
 			{
-				FScopedSlowTask ProcessingAnimationsSlowTask(static_cast<float>(Animations.Num()), ProcessingAnimStatusUpdate);
-				ProcessingAnimationsSlowTask.MakeDialog();
+				ProcessingAnimationsSlowTask.EnterProgressFrame();
+				UAnimSequence* Seq = Cast<UAnimSequence>(Anim.GetAsset());
 
-				for (const FAssetData& Anim : Animations)
+				TSharedPtr<FTokenizedMessage> Message;
+				for (const FFloatCurve& Curve : Seq->GetDataModel()->GetCurveData().FloatCurves)
 				{
-					ProcessingAnimationsSlowTask.EnterProgressFrame();
-					UAnimSequence* Seq = Cast<UAnimSequence>(Anim.GetAsset());
-
-					TSharedPtr<FTokenizedMessage> Message;
-					for (const FFloatCurve& Curve : Seq->GetDataModel()->GetCurveData().FloatCurves)
+					if (UnusedNames.Contains(Curve.GetName()))
 					{
-						if (UnusedNames.Contains(Curve.Name.DisplayName))
+						bFoundIssue = true;
+						if (!Message.IsValid())
 						{
-							bFoundIssue = true;
-							if (!Message.IsValid())
-							{
-								Message = CurveOutput.Warning();
-								Message->AddToken(FAssetNameToken::Create(Anim.GetObjectPathString(), FText::FromName(Anim.AssetName)));
-								Message->AddToken(FTextToken::Create(LOCTEXT("VerifyCurves_FoundAnimationsWithUnusedReferences", "References the following curves that are not used for either morph targets or material parameters and so may be unneeded")));
-							}
-							CurveOutput.Info(FText::FromName(Curve.Name.DisplayName));
+							Message = CurveOutput.Warning();
+							Message->AddToken(FAssetNameToken::Create(Anim.GetObjectPathString(), FText::FromName(Anim.AssetName)));
+							Message->AddToken(FTextToken::Create(LOCTEXT("VerifyCurves_FoundAnimationsWithUnusedReferences", "References the following curves that are not used for either morph targets or material parameters and so may be unneeded")));
 						}
+						CurveOutput.Info(FText::FromName(Curve.GetName()));
 					}
 				}
 			}
+		}
 
-			if (bFoundIssue)
-			{
-				CurveOutput.Notify();
-			}
+		if (bFoundIssue)
+		{
+			CurveOutput.Notify();
 		}
 	}
 }
@@ -1098,6 +1106,20 @@ void FPersonaModule::AddCommonToolbarExtensions(FToolBarBuilder& InToolbarBuilde
 	}
 }
 
+TSharedRef<SWidget> FPersonaModule::CreateCurvePicker(TSharedRef<IEditableSkeleton> InEditableSkeleton, FOnCurvePicked InOnCurvePicked, FIsCurveNameMarkedForExclusion InIsCurveNameMarkedForExclusion)
+{
+	return SNew(SAnimCurvePicker, &InEditableSkeleton->GetSkeleton())
+		.OnCurvePicked(InOnCurvePicked)
+		.IsCurveNameMarkedForExclusion(InIsCurveNameMarkedForExclusion);
+}
+
+TSharedRef<SWidget> FPersonaModule::CreateCurvePicker(const USkeleton* InSkeleton, FOnCurvePicked InOnCurvePicked, FIsCurveNameMarkedForExclusion InIsCurveNameMarkedForExclusion)
+{
+	return SNew(SAnimCurvePicker, InSkeleton)
+		.OnCurvePicked(InOnCurvePicked)
+		.IsCurveNameMarkedForExclusion(InIsCurveNameMarkedForExclusion);
+}
+
 void FPersonaModule::AddCommonMenuExtensions(UToolMenu* InToolMenu, const FCommonToolMenuExtensionArgs& InArgs)
 {
 	UPersonaToolMenuContext* PersonaToolMenuContext = InToolMenu->FindContext<UPersonaToolMenuContext>();
@@ -1460,36 +1482,17 @@ void FPersonaModule::InsertCurrentPoseToAsset(const FAssetData& NewPoseAssetData
 		UDebugSkelMeshComponent* PreviewMeshComponent = PersonaToolkit->GetPreviewMeshComponent();
 		if (PreviewMeshComponent)
 		{
-			FSmartName NewPoseName;
-
-			bool bSuccess = PoseAsset->AddOrUpdatePoseWithUniqueName(PreviewMeshComponent, &NewPoseName);
-
-			if (bSuccess)
-			{
-				FFormatNamedArguments Args;
-				Args.Add(TEXT("PoseAsset"), FText::FromString(PoseAsset->GetName()));
-				Args.Add(TEXT("PoseName"), FText::FromName(NewPoseName.DisplayName));
-				FNotificationInfo Info(FText::Format(LOCTEXT("InsertPoseSucceeded", "The current pose has inserted to {PoseAsset} with {PoseName}"), Args));
-				Info.ExpireDuration = 7.0f;
-				Info.bUseLargeFont = false;
-				TSharedPtr<SNotificationItem> Notification = FSlateNotificationManager::Get().AddNotification(Info);
-				if (Notification.IsValid())
-				{
-					Notification->SetCompletionState(SNotificationItem::CS_Success);
-				}
-			}
-			else
+			FName NewPoseName = PoseAsset->AddPoseWithUniqueName(PreviewMeshComponent);
+			FFormatNamedArguments Args;
+			Args.Add(TEXT("PoseAsset"), FText::FromString(PoseAsset->GetName()));
+			Args.Add(TEXT("PoseName"), FText::FromName(NewPoseName));
+			FNotificationInfo Info(FText::Format(LOCTEXT("InsertPoseSucceeded", "The current pose has inserted to {PoseAsset} with {PoseName}"), Args));
+			Info.ExpireDuration = 7.0f;
+			Info.bUseLargeFont = false;
+			TSharedPtr<SNotificationItem> Notification = FSlateNotificationManager::Get().AddNotification(Info);
+			if (Notification.IsValid())
 			{
-				FFormatNamedArguments Args;
-				Args.Add(TEXT("PoseAsset"), FText::FromString(PoseAsset->GetName()));
-				FNotificationInfo Info(FText::Format(LOCTEXT("InsertPoseFailed", "Inserting pose to asset {PoseAsset} has failed"), Args));
-				Info.ExpireDuration = 7.0f;
-				Info.bUseLargeFont = false;
-				TSharedPtr<SNotificationItem> Notification = FSlateNotificationManager::Get().AddNotification(Info);
-				if (Notification.IsValid())
-				{
-					Notification->SetCompletionState(SNotificationItem::CS_Fail);
-				}
+				Notification->SetCompletionState(SNotificationItem::CS_Success);
 			}
 		}
 	}
@@ -1571,7 +1574,7 @@ bool FPersonaModule::CreatePoseAsset(const TArray<UObject*> NewAssets, const EPo
 				switch (Option)
 				{
 				case EPoseSourceOption::CurrentPose:
-					NewPoseAsset->AddOrUpdatePoseWithUniqueName(PreviewComponent);
+					NewPoseAsset->AddPoseWithUniqueName(PreviewComponent);
 					bResult = true;
 					break;
 				case EPoseSourceOption::CurrentAnimation_AnimData:
