Animation Curve Runtime & Editor Improvements

diff --git a/Engine/Source/Editor/Persona/Private/SAnimAssetFindReplace.cpp b/Engine/Source/Editor/Persona/Private/SAnimAssetFindReplace.cpp
new file mode 100644
--- /dev/null
+++ b/Engine/Source/Editor/Persona/Private/SAnimAssetFindReplace.cpp
@@ -0,0 +1,1655 @@
+﻿// Copyright Epic Games, Inc. All Rights Reserved.
+
+#include "SAnimAssetFindReplace.h"
+
+#include "ContentBrowserDataSource.h"
+#include "ContentBrowserModule.h"
+#include "IContentBrowserSingleton.h"
+#include "PropertyCustomizationHelpers.h"
+#include "Widgets/Layout/SBox.h"
+#include "Widgets/SBoxPanel.h"
+#include "Widgets/Layout/SUniformGridPanel.h"
+#include "Widgets/Layout/SGridPanel.h"
+#include "Widgets/Input/SButton.h"
+#include "Widgets/Text/STextBlock.h"
+#include "Animation/AnimationAsset.h"
+#include "Animation/Skeleton.h"
+#include "Framework/Application/SlateApplication.h"
+#include "Widgets/Input/SEditableTextBox.h"
+#include "Widgets/Input/SSearchBox.h"
+#include "Animation/AnimSequence.h"
+#include "Engine/SkeletalMesh.h"
+#include "ToolMenus.h"
+#include "ToolMenuSection.h"
+#include "AssetRegistry/AssetRegistryModule.h"
+#include "Misc/ScopedSlowTask.h"
+#include "ScopedTransaction.h"
+#include "Animation/PoseAsset.h"
+#include "Widgets/Views/SListView.h"
+#include "UObject/AnimPhysObjectVersion.h"
+#include "UObject/ObjectVersion.h"
+#include "Widgets/Input/SHyperlink.h"
+
+#define LOCTEXT_NAMESPACE "SAnimAssetFindReplace"
+
+FAnimAssetFindReplaceSummoner::FAnimAssetFindReplaceSummoner(TSharedPtr<FAssetEditorToolkit> InHostingApp, const FAnimAssetFindReplaceConfig& InConfig)
+	: FWorkflowTabFactory(FPersonaTabs::FindReplaceID, InHostingApp)
+	, Config(InConfig)
+{
+	TabLabel = LOCTEXT("AnimAssetFindReplaceTabLabel", "Find/Replace");
+	TabIcon = FSlateIcon(FAppStyle::GetAppStyleSetName(), "Kismet.Tabs.FindResults");
+}
+
+TSharedRef<SWidget> FAnimAssetFindReplaceSummoner::CreateTabBody(const FWorkflowTabSpawnInfo& Info) const
+{
+	return SNew(SAnimAssetFindReplace)
+		.Config(Config);
+}
+
+TSharedPtr<SToolTip> FAnimAssetFindReplaceSummoner::CreateTabToolTipWidget(const FWorkflowTabSpawnInfo& Info) const
+{
+	return IDocumentation::Get()->CreateToolTip(LOCTEXT("WindowTooltip", "This tab lets you search and replace curve and notify names across multiple assets"), nullptr, TEXT("Shared/Editors/Persona"), TEXT("AnimationFindReplace_Window"));
+}
+
+namespace AnimAssetFindReplacePrivate
+{
+
+static TSharedPtr<SAnimAssetFindReplace> GetWidgetFromContext(const FToolMenuContext& InContext)
+{
+	if(UAnimAssetFindReplaceContext* Context = InContext.FindContext<UAnimAssetFindReplaceContext>())
+	{
+		return Context->Widget.Pin();
+	}
+
+	return nullptr;
+}
+
+}
+
+class SAutoCompleteSearchBox : public SCompoundWidget
+{
+public:
+	SLATE_BEGIN_ARGS(SAutoCompleteSearchBox) {}
+
+	/** The text displayed in the SearchBox when no text has been entered */
+	SLATE_ATTRIBUTE(FText, HintText)
+
+	/** The text displayed in the SearchBox when it's created */
+	SLATE_ATTRIBUTE(FText, InitialText)
+
+	/** Invoked whenever the text changes */
+	SLATE_EVENT(FOnTextChanged, OnTextChanged)
+
+	/** Items to show in the autocomplete popup */
+	SLATE_ARGUMENT(TSharedPtr<TArray<TSharedPtr<FString>>>, AutoCompleteItems)
+	
+	SLATE_END_ARGS()
+
+	void Construct(const FArguments& InArgs)
+	{
+		AutoCompleteItems = InArgs._AutoCompleteItems;
+
+		ChildSlot
+		[
+			SAssignNew(MenuAnchor, SMenuAnchor)
+			.Method(EPopupMethod::CreateNewWindow)
+			.Placement(EMenuPlacement::MenuPlacement_BelowAnchor)
+			.MenuContent
+			(
+				SNew(SBox)
+				.MaxDesiredHeight(200.0f)
+				.MinDesiredWidth(200.0f)
+				[
+					SAssignNew(AutoCompleteList, SListView<TSharedPtr<FString>>)
+					.SelectionMode(ESelectionMode::Single)
+					.ListItemsSource(&FilteredAutoCompleteItems)
+					.OnSelectionChanged_Lambda([this](TSharedPtr<FString> InString, ESelectInfo::Type InSelectInfo)
+					{
+						if(InString.IsValid() && InSelectInfo == ESelectInfo::OnMouseClick)
+						{
+							TGuardValue<bool> GuardValue(bSettingTextFromSearchItem, true);
+							SearchBox->SetText(FText::FromString(*InString.Get()));
+							MenuAnchor->SetIsOpen(false);
+						}
+					})
+					.OnGenerateRow_Lambda([this](TSharedPtr<FString> InString, const TSharedRef<STableViewBase>& InTableView)
+					{
+						if(InString.IsValid())
+						{
+							return
+								SNew(STableRow<TSharedPtr<FString>>, InTableView)
+								.Content()
+								[
+									SNew(STextBlock)
+									.Text(FText::FromString(*InString.Get()))
+									.HighlightText_Lambda([this]()
+									{
+										return SearchBox->GetText();
+									})
+								];
+						}
+						return SNew(STableRow<TSharedPtr<FString>>, InTableView);
+					})
+					.OnKeyDownHandler_Lambda([this](const FGeometry& InGeometry, const FKeyEvent& InKeyEvent)
+					{
+						if(InKeyEvent.GetKey() == EKeys::Enter)
+						{
+							TArray<TSharedPtr<FString>> SelectedItems;
+							AutoCompleteList->GetSelectedItems(SelectedItems);
+							if(SelectedItems.Num() > 0 && SelectedItems[0].IsValid())
+							{
+								TGuardValue<bool> GuardValue(bSettingTextFromSearchItem, true);
+								SearchBox->SetText(FText::FromString(*SelectedItems[0].Get()));
+								MenuAnchor->SetIsOpen(false);
+								FReply::Handled();
+							}
+						}
+						return FReply::Unhandled();
+					})
+				]
+			)
+			[
+				SAssignNew(SearchBox, SSearchBox)
+				.HintText(InArgs._HintText)
+				.InitialText(InArgs._InitialText)
+				.OnTextChanged_Lambda([this, OnTextChanged = InArgs._OnTextChanged](const FText& InText)
+				{
+					FilterItems(InText);
+
+					if(!bSettingTextFromSearchItem)
+					{
+						MenuAnchor->SetIsOpen(FilteredAutoCompleteItems.Num() > 0, false);
+					}
+
+					OnTextChanged.ExecuteIfBound(InText);
+				})
+				.OnVerifyTextChanged_Lambda([](const FText& InText, FText& OutErrorMessage)
+				{
+					return FName::IsValidXName(InText.ToString(), FString(INVALID_NAME_CHARACTERS), &OutErrorMessage);
+				})
+			]
+		];
+		
+		FilterItems(InArgs._InitialText.Get());
+	}
+
+	void FilterItems(const FText& InText)
+	{
+		FilteredAutoCompleteItems.Empty();
+
+		for(TSharedPtr<FString> String : *AutoCompleteItems.Get())
+		{
+			if(String.Get()->Contains(InText.ToString()))
+			{
+				FilteredAutoCompleteItems.Add(String);
+			}
+		}
+
+		RefreshAutoCompleteItems();
+	}
+
+	void RefreshAutoCompleteItems()
+	{
+		AutoCompleteList->RequestListRefresh();
+	}
+
+	TSharedRef<SSearchBox> GetSearchBox() const
+	{
+		return SearchBox.ToSharedRef();
+	}
+	
+	virtual void OnFocusChanging(const FWeakWidgetPath& PreviousFocusPath, const FWidgetPath& NewWidgetPath, const FFocusEvent& InFocusEvent) override
+	{
+		if(PreviousFocusPath.ContainsWidget(SearchBox.Get()) && !NewWidgetPath.ContainsWidget(MenuAnchor->GetMenuWindow().Get()))
+		{
+			MenuAnchor->SetIsOpen(false);
+		}
+		SCompoundWidget::OnFocusChanging(PreviousFocusPath, NewWidgetPath, InFocusEvent);
+	}
+
+	virtual FReply OnPreviewKeyDown(const FGeometry& MyGeometry, const FKeyEvent& InKeyEvent) override
+	{
+		if(InKeyEvent.GetKey() == EKeys::Down && MenuAnchor->IsOpen())
+		{
+			// switch focus to the drop-down autocomplete list
+			return FReply::Handled().SetUserFocus(AutoCompleteList.ToSharedRef(), EFocusCause::Navigation);
+		}
+		return SCompoundWidget::OnPreviewKeyDown(MyGeometry, InKeyEvent); 
+	}
+	
+private:
+	TArray<TSharedPtr<FString>> FilteredAutoCompleteItems;
+	TSharedPtr<TArray<TSharedPtr<FString>>> AutoCompleteItems;
+	TSharedPtr<SListView<TSharedPtr<FString>>> AutoCompleteList;
+	TSharedPtr<SMenuAnchor> MenuAnchor;
+	TSharedPtr<SSearchBox> SearchBox;
+	bool bSettingTextFromSearchItem = false;
+};
+
+void SAnimAssetFindReplace::Construct(const FArguments& InArgs)
+{
+	Mode = InArgs._Config.Mode;
+	Type = InArgs._Config.Type;
+	FindString = InArgs._Config.FindString;
+	ReplaceString = InArgs._Config.ReplaceString;
+	SkeletonFilter = InArgs._Config.SkeletonFilter;
+
+	const FContentBrowserModule& ContentBrowserModule = FModuleManager::LoadModuleChecked<FContentBrowserModule>("ContentBrowser");
+
+	AssetPickerConfig.SelectionMode = ESelectionMode::Multi;
+	AssetPickerConfig.InitialAssetViewType = EAssetViewType::Column;
+	AssetPickerConfig.Filter = MakeARFilter();
+	AssetPickerConfig.OnShouldFilterAsset = FOnShouldFilterAsset::CreateSP(this, &SAnimAssetFindReplace::HandleFilterAsset);
+	AssetPickerConfig.RefreshAssetViewDelegates.Add(&RefreshAssetViewDelegate);
+	AssetPickerConfig.GetCurrentSelectionDelegates.Add(&GetCurrentSelectionDelegate);
+	AssetPickerConfig.SetFilterDelegates.Add(&SetARFilterDelegate);
+	AssetPickerConfig.OnAssetSelected = FOnAssetSelected::CreateLambda([this](const FAssetData& InAssetData)
+	{
+		bAssetsSelected = InAssetData.IsValid();
+	});
+	AssetPickerConfig.OnAssetDoubleClicked = FOnAssetDoubleClicked::CreateLambda([this](const FAssetData& InAssetData)
+	{
+		if (UAssetEditorSubsystem* EditorSubsystem = GEditor->GetEditorSubsystem<UAssetEditorSubsystem>())
+		{
+			EditorSubsystem->OpenEditorForAsset(InAssetData.ToSoftObjectPath());
+		}
+	});
+	AssetPickerConfig.CustomColumns.Add(
+		FAssetViewCustomColumn(
+				"AssetResults",
+				LOCTEXT("ResultsColumnLabel", "Results"),
+				LOCTEXT("ResultsColumnTooltip", "The matching results that are in each asset"),
+				UObject::FAssetRegistryTag::TT_Alphabetical,
+				FOnGetCustomAssetColumnData::CreateLambda([this](FAssetData& InAssetData, FName InColumnName)
+				{
+					TStringBuilder<128> Builder;
+
+					if(Type == EAnimAssetFindReplaceType::Curves)
+					{
+						TArray<FString> CurveNames;
+						GetMatchingCurveNamesForAsset(InAssetData, CurveNames);
+						if(CurveNames.Num() > 0)
+						{
+							for(int32 NameIndex = 0; NameIndex < CurveNames.Num(); ++NameIndex)
+							{
+								Builder.Append(CurveNames[NameIndex]);
+								if(NameIndex != CurveNames.Num() - 1)
+								{
+									Builder.Append(TEXT(", "));
+								}
+							}
+						}
+					}
+					else if(Type == EAnimAssetFindReplaceType::Notifies)
+					{
+						TArray<FString> NotifyNames;
+						GetMatchingNotifyNamesForAsset(InAssetData, NotifyNames);
+						if(NotifyNames.Num() > 0)
+						{
+							for(int32 NameIndex = 0; NameIndex < NotifyNames.Num(); ++NameIndex)
+							{
+								Builder.Append(NotifyNames[NameIndex]);
+								if(NameIndex != NotifyNames.Num() - 1)
+								{
+									Builder.Append(TEXT(", "));
+								}
+							}
+						}
+					}
+
+					return FString(Builder.ToString());
+				})
+			));
+
+	TArray<UClass*> ClassesWithAssetRegistryTags = { UAnimSequenceBase::StaticClass(), USkeleton::StaticClass(), USkeletalMesh::StaticClass() };
+	GetDerivedClasses(UAnimSequenceBase::StaticClass(), ClassesWithAssetRegistryTags);
+	GetDerivedClasses(USkeleton::StaticClass(), ClassesWithAssetRegistryTags);
+	GetDerivedClasses(USkeletalMesh::StaticClass(), ClassesWithAssetRegistryTags);
+
+	TArray<UObject::FAssetRegistryTag> AssetRegistryTags;
+	for(UClass* Class : ClassesWithAssetRegistryTags)
+	{
+		Class->GetDefaultObject()->GetAssetRegistryTags(AssetRegistryTags);
+		for(UObject::FAssetRegistryTag& AssetRegistryTag : AssetRegistryTags)
+		{
+			if(AssetRegistryTag.Type != UObject::FAssetRegistryTag::TT_Hidden)
+			{
+				AssetPickerConfig.HiddenColumnNames.AddUnique(AssetRegistryTag.Name.ToString());
+			}
+		}
+	}
+
+	AssetPickerConfig.HiddenColumnNames.Add(ContentBrowserItemAttributes::ItemDiskSize.ToString());
+	AssetPickerConfig.HiddenColumnNames.Add(ContentBrowserItemAttributes::VirtualizedData.ToString());
+	AssetPickerConfig.HiddenColumnNames.Add(TEXT("Path"));
+	AssetPickerConfig.HiddenColumnNames.Add(TEXT("Class"));
+	AssetPickerConfig.HiddenColumnNames.Add(TEXT("RevisionControl"));
+	AssetPickerConfig.bShowPathInColumnView = true;
+	AssetPickerConfig.bShowTypeInColumnView = true;
+	AssetPickerConfig.bSortByPathInColumnView = false;
+	AssetPickerConfig.bFocusSearchBoxWhenOpened = false;
+
+	UToolMenu* Toolbar = UToolMenus::Get()->FindMenu("AnimAssetFindReplaceToolbar");
+	if(Toolbar == nullptr)
+	{
+		Toolbar = UToolMenus::Get()->RegisterMenu("AnimAssetFindReplaceToolbar", NAME_None, EMultiBoxType::SlimHorizontalToolBar);
+		Toolbar->StyleName = "CalloutToolbar";	// This style displays button text
+
+		{
+			FToolMenuSection& Section = Toolbar->AddSection("FindReplaceActions");
+			
+			FToolUIAction RefreshButton;
+			RefreshButton.ExecuteAction = FToolMenuExecuteAction::CreateLambda([](const FToolMenuContext& InContext)
+			{
+				if(TSharedPtr<SAnimAssetFindReplace> Widget = AnimAssetFindReplacePrivate::GetWidgetFromContext(InContext))
+				{
+					Widget->RefreshAutoCompleteItems();
+					Widget->RefreshSearchResults();
+				}
+			});
+
+			Section.AddEntry(
+				FToolMenuEntry::InitToolBarButton(
+					"Refresh",
+					RefreshButton,
+					LOCTEXT("RefreshRadioLabel", "Refresh"),
+					LOCTEXT("RefreshRadioTooltip", "Refresh search results."),
+					FSlateIcon(),
+					EUserInterfaceActionType::RadioButton));
+		}
+
+		{
+			FToolMenuSection& Section = Toolbar->AddSection("FindReplaceOptions");
+			
+			FToolUIAction CurvesRadio;
+			CurvesRadio.ExecuteAction = FToolMenuExecuteAction::CreateLambda([](const FToolMenuContext& InContext)
+			{
+				if(TSharedPtr<SAnimAssetFindReplace> Widget = AnimAssetFindReplacePrivate::GetWidgetFromContext(InContext))
+				{
+					Widget->SetFindReplaceType(EAnimAssetFindReplaceType::Curves);
+				}
+			});
+
+			CurvesRadio.GetActionCheckState = FToolMenuGetActionCheckState::CreateLambda([](const FToolMenuContext& InContext)
+			{
+				if(TSharedPtr<SAnimAssetFindReplace> Widget = AnimAssetFindReplacePrivate::GetWidgetFromContext(InContext))
+				{
+					return Widget->Type == EAnimAssetFindReplaceType::Curves ? ECheckBoxState::Checked : ECheckBoxState::Unchecked;
+				}
+				return ECheckBoxState::Undetermined;
+			});
+			
+			Section.AddEntry(
+				FToolMenuEntry::InitToolBarButton(
+					"Curves",
+					CurvesRadio,
+					LOCTEXT("CurvesRadioLabel", "Curves"),
+					LOCTEXT("CurvesRadioTooltip", "Search for Curves in Animation Assets."),
+					FSlateIcon(),
+					EUserInterfaceActionType::RadioButton));
+
+			FToolUIAction NotifiesRadio;
+			NotifiesRadio.ExecuteAction = FToolMenuExecuteAction::CreateLambda([](const FToolMenuContext& InContext)
+			{
+				if(TSharedPtr<SAnimAssetFindReplace> Widget = AnimAssetFindReplacePrivate::GetWidgetFromContext(InContext))
+				{
+					Widget->SetFindReplaceType(EAnimAssetFindReplaceType::Notifies);
+				}
+			});
+
+			NotifiesRadio.GetActionCheckState = FToolMenuGetActionCheckState::CreateLambda([](const FToolMenuContext& InContext)
+			{
+				if(TSharedPtr<SAnimAssetFindReplace> Widget = AnimAssetFindReplacePrivate::GetWidgetFromContext(InContext))
+				{
+					return Widget->Type == EAnimAssetFindReplaceType::Notifies ? ECheckBoxState::Checked : ECheckBoxState::Unchecked;
+				}
+				return ECheckBoxState::Undetermined;
+			});
+			
+			Section.AddEntry(
+				FToolMenuEntry::InitToolBarButton(
+					"Notifies", NotifiesRadio,
+					LOCTEXT("NotifiesRadioLabel", "Notifies"),
+					LOCTEXT("NotifiesRadioTooltip", "Search for Named Notifies in Animation Assets."),
+					FSlateIcon(),
+					EUserInterfaceActionType::RadioButton));
+
+			FToolUIAction MatchCaseCheckbox;
+			MatchCaseCheckbox.ExecuteAction = FToolMenuExecuteAction::CreateLambda([](const FToolMenuContext& InContext)
+			{
+				if(TSharedPtr<SAnimAssetFindReplace> Widget = AnimAssetFindReplacePrivate::GetWidgetFromContext(InContext))
+				{
+					Widget->SearchCase = Widget->SearchCase == ESearchCase::CaseSensitive ? ESearchCase::IgnoreCase : ESearchCase::CaseSensitive;
+					Widget->RefreshSearchResults();
+				}
+			});
+
+			MatchCaseCheckbox.GetActionCheckState = FToolMenuGetActionCheckState::CreateLambda([](const FToolMenuContext& InContext)
+			{
+				if(TSharedPtr<SAnimAssetFindReplace> Widget = AnimAssetFindReplacePrivate::GetWidgetFromContext(InContext))
+				{
+					return Widget->SearchCase == ESearchCase::CaseSensitive ? ECheckBoxState::Checked : ECheckBoxState::Unchecked;
+				}
+				return ECheckBoxState::Undetermined;
+			});
+			
+			Section.AddEntry(
+				FToolMenuEntry::InitToolBarButton(
+					"MatchCase",
+					MatchCaseCheckbox,
+					LOCTEXT("MatchCaseCheckboxLabel", "Match Case"),
+					LOCTEXT("MatchCaseCheckboxTooltip", "Whether to match case when searching."),
+					FSlateIcon(),
+					EUserInterfaceActionType::ToggleButton));
+
+			FToolUIAction MatchWholeWordCheckbox;
+			MatchWholeWordCheckbox.ExecuteAction = FToolMenuExecuteAction::CreateLambda([](const FToolMenuContext& InContext)
+			{
+				if(TSharedPtr<SAnimAssetFindReplace> Widget = AnimAssetFindReplacePrivate::GetWidgetFromContext(InContext))
+				{
+					Widget->bFindWholeWord = !Widget->bFindWholeWord;
+					Widget->RefreshSearchResults();
+				}
+			});
+
+			MatchWholeWordCheckbox.GetActionCheckState = FToolMenuGetActionCheckState::CreateLambda([](const FToolMenuContext& InContext)
+			{
+				if(TSharedPtr<SAnimAssetFindReplace> Widget = AnimAssetFindReplacePrivate::GetWidgetFromContext(InContext))
+				{
+					return Widget->bFindWholeWord ? ECheckBoxState::Checked : ECheckBoxState::Unchecked;
+				}
+				return ECheckBoxState::Undetermined;
+			});
+			
+			Section.AddEntry(
+				FToolMenuEntry::InitToolBarButton(
+					"MatchWholeWord",
+					MatchWholeWordCheckbox,
+					LOCTEXT("MatchWholeWordCheckboxLabel", "Match Whole Word"),
+					LOCTEXT("MatchWholeWordCheckboxTooltip", "Whether to match the whole word or just part of the word when searching."),
+					FSlateIcon(),
+					EUserInterfaceActionType::ToggleButton));
+
+			Section.AddDynamicEntry("SkeletonFilter", FNewToolMenuSectionDelegate::CreateLambda([](FToolMenuSection& InSection)
+			{
+				if(TSharedPtr<SAnimAssetFindReplace> Widget = AnimAssetFindReplacePrivate::GetWidgetFromContext(InSection.Context))
+				{
+					InSection.AddEntry(
+					FToolMenuEntry::InitWidget(
+						"SkeletonFilterWidget",
+						SNew(SHorizontalBox)
+						.ToolTipText(LOCTEXT("SkeletonFilterTooltip", "Choose a Skeleton asset to filter results by."))
+						+SHorizontalBox::Slot()
+						.AutoWidth()
+						.VAlign(VAlign_Center)
+						.Padding(5.0f, 0.0f)
+						[
+							SNew(STextBlock)
+							.Text(LOCTEXT("SkeletonFilterLabel", "Skeleton"))
+						]
+						+SHorizontalBox::Slot()
+						.AutoWidth()
+						.VAlign(VAlign_Center)
+						[
+							SNew(SObjectPropertyEntryBox)
+							.ObjectPath_Lambda([WeakWidget = TWeakPtr<SAnimAssetFindReplace>(Widget)]()
+							{
+								if(TSharedPtr<SAnimAssetFindReplace> PinnedWidget = WeakWidget.Pin())
+								{
+									return PinnedWidget->SkeletonFilter.GetObjectPathString();
+								}
+								return FString();
+							})
+							.OnObjectChanged_Lambda([WeakWidget = TWeakPtr<SAnimAssetFindReplace>(Widget)](const FAssetData& InAssetData)
+							{
+								if(TSharedPtr<SAnimAssetFindReplace> PinnedWidget = WeakWidget.Pin())
+								{
+									PinnedWidget->SkeletonFilter = InAssetData;
+									PinnedWidget->RefreshSearchResults();
+								}
+							})
+							.AllowedClass(USkeleton::StaticClass())
+						],
+					FText::GetEmpty(),
+					true, true, true));
+				}
+			}));
+		}
+	}
+
+	ToolbarContext = TStrongObjectPtr(NewObject<UAnimAssetFindReplaceContext>());
+	ToolbarContext->Widget = SharedThis(this);
+
+	AutoCompleteItems = MakeShared<TArray<TSharedPtr<FString>>>();
+
+	ChildSlot
+	[
+		SNew(SVerticalBox)
+		+SVerticalBox::Slot()
+		.AutoHeight()
+		[
+			UToolMenus::Get()->GenerateWidget("AnimAssetFindReplaceToolbar", FToolMenuContext(ToolbarContext.Get()))
+		]
+		+SVerticalBox::Slot()
+		.AutoHeight()
+		.Padding(5.0f, 10.0f)
+		[
+			SNew(SVerticalBox)
+			+SVerticalBox::Slot()
+			.Padding(6.0f, 2.0f)
+			[
+				SAssignNew(FindSearchBox, SAutoCompleteSearchBox)
+				.AutoCompleteItems(AutoCompleteItems)
+				.HintText(LOCTEXT("FindLabel", "Find"))
+				.ToolTipText(LOCTEXT("FindLabel", "Find"))
+				.InitialText_Lambda([this](){ return FText::FromString(FindString); })
+				.OnTextChanged_Lambda([this](const FText& InText)
+				{
+					FindString = InText.ToString();
+					RefreshSearchResults();
+				})
+			]
+			+SVerticalBox::Slot()
+			.Padding(6.0f, 2.0f)
+			[
+				SAssignNew(ReplaceSearchBox, SAutoCompleteSearchBox)
+				.AutoCompleteItems(AutoCompleteItems)
+				.HintText(LOCTEXT("ReplaceLabel", "Replace With"))
+				.ToolTipText(LOCTEXT("ReplaceLabel", "Replace With"))
+				.InitialText_Lambda([this](){ return FText::FromString(ReplaceString); })
+				.OnTextChanged_Lambda([this](const FText& InText)
+				{
+					ReplaceString = InText.ToString();
+				})
+			]
+		]
+		+SVerticalBox::Slot()
+		.FillHeight(1.0f)
+		.Padding(5.0f, 10.0f)
+		[
+			ContentBrowserModule.Get().CreateAssetPicker(AssetPickerConfig)
+		]
+		+SVerticalBox::Slot()
+		.AutoHeight()
+		.Padding(10.0f)
+		[
+			SNew(SHorizontalBox)
+			+SHorizontalBox::Slot()
+			.AutoWidth()
+			.HAlign(HAlign_Left)
+			.VAlign(VAlign_Center)
+			[
+				SNew(SHyperlink)
+				.Visibility_Lambda([this]()
+				{
+					return OldAssets.Num() > 0 ? EVisibility::Visible : EVisibility::Collapsed;
+				})
+				.Text_Lambda([this]()
+				{
+					return FText::Format(LOCTEXT("UnindexedAssetWarningFormat", "{0} assets could not be indexed, load them now?"), FText::AsNumber(OldAssets.Num()));
+				})
+				.OnNavigate_Lambda([this]()
+				{
+					// Load all old unindexed assets
+					FScopedSlowTask SlowTask(OldAssets.Num(), FText::Format(LOCTEXT("LoadingUnindexedAssetsFormat", "Loading {0} Unindexed Assets..."), FText::AsNumber(OldAssets.Num())));
+					SlowTask.MakeDialog(true);
+					
+					for(const FAssetData& AssetData : OldAssets)
+					{
+						SlowTask.EnterProgressFrame();
+						
+						AssetData.GetAsset();
+
+						if(SlowTask.ShouldCancel())
+						{
+							break;
+						}
+					}
+
+					RefreshSearchResults();
+				})
+			]
+			+SHorizontalBox::Slot()
+			.FillWidth(1.0f)
+			.HAlign(HAlign_Right)
+			[
+				SNew(SUniformGridPanel)
+				.SlotPadding(FAppStyle::GetMargin("StandardDialog.SlotPadding"))
+				.MinDesiredSlotWidth(FAppStyle::GetFloat("StandardDialog.MinDesiredSlotWidth"))
+				.MinDesiredSlotHeight(FAppStyle::GetFloat("StandardDialog.MinDesiredSlotHeight"))
+				+SUniformGridPanel::Slot(0, 0)
+				[
+					SNew(SButton)
+					.HAlign(HAlign_Center)
+					.VAlign(VAlign_Center)
+					.Text(LOCTEXT("RemoveButton", "Remove"))
+					.ToolTipText(LOCTEXT("RemoveButtonTooltip", "Remove selected items"))
+					.IsEnabled_Lambda([this]()
+					{
+						return bAssetsSelected;
+					})
+					.OnClicked(this, &SAnimAssetFindReplace::HandleRemove)
+				]
+				+SUniformGridPanel::Slot(1, 0)
+				[
+					SNew(SButton)
+					.HAlign(HAlign_Center)
+					.VAlign(VAlign_Center)
+					.Text(LOCTEXT("RemoveAllButton", "Remove All"))
+					.ToolTipText(LOCTEXT("RemoveAllButtonTooltip", "Remove all matching items"))
+					.IsEnabled_Lambda([this]()
+					{
+						return bFoundAssets;
+					})
+					.OnClicked(this, &SAnimAssetFindReplace::HandleRemoveAll)
+				]
+				+SUniformGridPanel::Slot(2, 0)
+				[
+					SNew(SButton)
+					.HAlign(HAlign_Center)
+					.VAlign(VAlign_Center)
+					.Text(LOCTEXT("ReplaceButton", "Replace"))
+					.ToolTipText(LOCTEXT("ReplaceButtonTooltip", "Replace selected items"))
+					.IsEnabled_Lambda([this]()
+					{
+						return bAssetsSelected && !ReplaceString.IsEmpty();
+					})
+					.OnClicked(this, &SAnimAssetFindReplace::HandleReplace)
+				]
+				+SUniformGridPanel::Slot(3, 0)
+				[
+					SNew(SButton)
+					.HAlign(HAlign_Center)
+					.VAlign(VAlign_Center)
+					.Text(LOCTEXT("ReplaceAllButton", "Replace All"))
+					.ToolTipText(LOCTEXT("ReplaceAllButtonTooltip", "Replace all matching items"))
+					.IsEnabled_Lambda([this]()
+					{
+						return bFoundAssets && !ReplaceString.IsEmpty();
+					})
+					.OnClicked(this, &SAnimAssetFindReplace::HandleReplaceAll)
+				]
+			]
+		]
+	];
+	
+	RefreshAutoCompleteItems();
+	
+	RegisterActiveTimer(0.0f, FWidgetActiveTimerDelegate::CreateLambda([FindWidgetPtr = TWeakPtr<SAutoCompleteSearchBox>(FindSearchBox)](double InCurrentTime, float InDeltaTime)
+	{
+		if (FindWidgetPtr.IsValid())
+		{
+			FWidgetPath WidgetToFocusPath;
+			FSlateApplication::Get().GeneratePathToWidgetUnchecked(FindWidgetPtr.Pin()->GetSearchBox(), WidgetToFocusPath);
+			FSlateApplication::Get().SetKeyboardFocus(WidgetToFocusPath, EFocusCause::SetDirectly);
+			WidgetToFocusPath.GetWindow()->SetWidgetToFocusOnActivate(FindWidgetPtr.Pin()->GetSearchBox());
+		}
+
+		return EActiveTimerReturnType::Stop;
+	}));
+}
+
+void SAnimAssetFindReplace::SetFindReplaceType(EAnimAssetFindReplaceType InType)
+{
+	Type = InType;
+	RefreshAutoCompleteItems();
+	RefreshSearchResults();
+}
+
+FARFilter SAnimAssetFindReplace::MakeARFilter() const
+{
+	FARFilter Filter;
+	Filter.ClassPaths =
+	{
+		UAnimationAsset::StaticClass()->GetClassPathName(),
+		USkeleton::StaticClass()->GetClassPathName(),
+		USkeletalMesh::StaticClass()->GetClassPathName()
+	};
+	Filter.bRecursiveClasses = true;
+
+	return Filter;
+}
+
+void SAnimAssetFindReplace::RefreshSearchResults()
+{
+	bFoundAssets = false;
+	OldAssets.Empty();
+	SetARFilterDelegate.ExecuteIfBound(MakeARFilter());
+	RefreshAssetViewDelegate.ExecuteIfBound(true);
+}
+
+bool SAnimAssetFindReplace::ShouldFilterOutAsset(const FAssetData& InAssetData, bool& bOutIsOldAsset) const
+{
+	switch(Type)
+	{
+	case EAnimAssetFindReplaceType::Curves:
+		return FilterByCurve(InAssetData, bOutIsOldAsset);
+	case EAnimAssetFindReplaceType::Notifies:
+		return FilterByNotify(InAssetData, bOutIsOldAsset);
+	default:
+		break;
+	}
+
+	return true;
+}
+
+bool SAnimAssetFindReplace::HandleFilterAsset(const FAssetData& InAssetData)
+{
+	bool bIsOldAsset = false;
+	const bool bShouldFilterOut = ShouldFilterOutAsset(InAssetData, bIsOldAsset);
+	bFoundAssets |= !bShouldFilterOut;
+	if(bIsOldAsset)
+	{
+		OldAssets.Add(InAssetData);
+	}
+	return bShouldFilterOut;
+}
+
+void SAnimAssetFindReplace::GetMatchingCurveNamesForAsset(const FAssetData& InAssetData, TArray<FString>& OutCurveNames) const
+{
+	const FString TagValue = InAssetData.GetTagValueRef<FString>(USkeleton::CurveNameTag);
+	if (!TagValue.IsEmpty())
+	{
+		check(IsInGameThread());
+		static TArray<FString> CurveNames;
+		CurveNames.Reset();
+		if (TagValue.ParseIntoArray(CurveNames, *USkeleton::CurveTagDelimiter, true) > 0)
+		{
+			for (const FString& CurveNameString : CurveNames)
+			{
+				if(bFindWholeWord)
+				{
+					if(CurveNameString.Compare(FindString, SearchCase) == 0)
+					{
+						OutCurveNames.Add(CurveNameString);
+					}
+				}
+				else
+				{
+					if(CurveNameString.Contains(FindString, SearchCase))
+					{
+						OutCurveNames.Add(CurveNameString);
+					}
+				}
+			}
+		}
+	}
+}
+
+void SAnimAssetFindReplace::GetMatchingNotifyNamesForAsset(const FAssetData& InAssetData, TArray<FString>& OutNotifyNames) const
+{
+	const FString TagValue = InAssetData.GetTagValueRef<FString>(USkeleton::AnimNotifyTag);
+	if (!TagValue.IsEmpty())
+	{
+		check(IsInGameThread());
+		static TArray<FString> NotifyNames;
+		NotifyNames.Reset();
+		if (TagValue.ParseIntoArray(NotifyNames, *USkeleton::AnimNotifyTagDelimiter, true) > 0)
+		{
+			for (const FString& NotifyNameString : NotifyNames)
+			{
+				if(bFindWholeWord)
+				{
+					if(NotifyNameString.Compare(FindString, SearchCase) == 0)
+					{
+						OutNotifyNames.Add(NotifyNameString);
+					}
+				}
+				else
+				{
+					if(NotifyNameString.Contains(FindString, SearchCase))
+					{
+						OutNotifyNames.Add(NotifyNameString);
+					}
+				}
+			}
+		}
+	}
+}
+
+bool SAnimAssetFindReplace::NameMatches(const FString& InNameString) const
+{
+	if(bFindWholeWord)
+	{
+		if(InNameString.Compare(FindString, SearchCase) == 0)
+		{
+			return true;
+		}
+	}
+	else
+	{
+		if(InNameString.Contains(FindString, SearchCase))
+		{
+			return true;
+		}
+	}
+
+	return false;
+}
+
+bool SAnimAssetFindReplace::FilterByCurve(const FAssetData& InAssetData, bool& bOutIsOldAsset) const
+{
+	FString TagValue;
+	if(InAssetData.GetTagValue<FString>(USkeleton::CurveNameTag, TagValue))
+	{
+		bOutIsOldAsset = false;
+
+		if(FindString.IsEmpty())
+		{
+			return true;
+		}
+		
+		if(SkeletonFilter.IsValid() )
+		{
+			if(InAssetData.GetClass() != USkeleton::StaticClass())
+			{
+				FString SkeletonPath;
+				if(InAssetData.GetTagValue<FString>(TEXT("Skeleton"), SkeletonPath))
+				{
+					if(SkeletonPath != SkeletonFilter.GetExportTextName())
+					{
+						return true;
+					}
+				}
+			}
+			else
+			{
+				if(InAssetData.ToSoftObjectPath() != SkeletonFilter.ToSoftObjectPath())
+				{
+					return true;
+				}
+			}
+		}
+
+		check(IsInGameThread());
+		static TArray<FString> CurveNames;
+		CurveNames.Reset();
+		if (TagValue.ParseIntoArray(CurveNames, *USkeleton::CurveTagDelimiter, true) > 0)
+		{
+			for (const FString& CurveNameString : CurveNames)
+			{
+				if(NameMatches(CurveNameString))
+				{
+					return false;
+				}
+			}
+		}
+	}
+	else
+	{
+		bOutIsOldAsset = IsAssetWithoutTagOldAsset(USkeleton::CurveNameTag, InAssetData);
+	}
+
+	return true;
+}
+
+bool SAnimAssetFindReplace::FilterByNotify(const FAssetData& InAssetData, bool& bOutIsOldAsset) const
+{
+	FString TagValue;
+	if(InAssetData.GetTagValue<FString>(USkeleton::AnimNotifyTag, TagValue))
+	{
+		bOutIsOldAsset = false;
+
+		if(FindString.IsEmpty())
+		{
+			return true;
+		}
+
+		if(SkeletonFilter.IsValid() )
+		{
+			if(InAssetData.GetClass() != USkeleton::StaticClass())
+			{
+				FString SkeletonPath;
+				if(InAssetData.GetTagValue<FString>(TEXT("Skeleton"), SkeletonPath))
+				{
+					if(SkeletonPath != SkeletonFilter.GetExportTextName())
+					{
+						return true;
+					}
+				}
+			}
+			else
+			{
+				if(InAssetData.ToSoftObjectPath() != SkeletonFilter.ToSoftObjectPath())
+				{
+					return true;
+				}
+			}
+		}
+		
+		check(IsInGameThread());
+		static TArray<FString> NotifyNames;
+		NotifyNames.Reset();
+		if (TagValue.ParseIntoArray(NotifyNames, *USkeleton::AnimNotifyTagDelimiter, true) > 0)
+		{
+			for (const FString& NotifyNameString : NotifyNames)
+			{
+				if(NameMatches(NotifyNameString))
+				{
+					return false;
+				}
+			}
+		}
+	}
+	else
+	{
+		bOutIsOldAsset = IsAssetWithoutTagOldAsset(USkeleton::AnimNotifyTag, InAssetData);
+	}
+
+	return true;
+}
+
+bool SAnimAssetFindReplace::IsAssetWithoutTagOldAsset(FName InTag, const FAssetData& InAssetData) const
+{
+	FAssetRegistryModule& AssetRegistryModule = FModuleManager::LoadModuleChecked<FAssetRegistryModule>("AssetRegistry");
+	UClass* AssetClass = InAssetData.GetClass();
+	if(AssetClass->IsChildOf(UAnimSequenceBase::StaticClass()))
+	{
+		if(InTag == USkeleton::CurveNameTag)
+		{
+			// Check the package object version - the asset was saving empty tags for curves, so the absence of curves is
+			// not the same as an empty value
+			FAssetPackageData PackageData;
+			const UE::AssetRegistry::EExists PackageExists = AssetRegistryModule.Get().TryGetAssetPackageData(InAssetData.PackageName, PackageData);
+			if (PackageExists == UE::AssetRegistry::EExists::Exists)
+			{
+				return PackageData.FileVersionUE < VER_UE4_SKELETON_ADD_SMARTNAMES;
+			}
+			else
+			{
+				// Does not exist or unknown - treat it as 'old'
+				return true;
+			}
+		}
+		else if(InTag == USkeleton::AnimNotifyTag)
+		{
+			return true;
+		}
+	}
+	else if(AssetClass->IsChildOf(UPoseAsset::StaticClass()))
+	{
+		if(InTag == USkeleton::CurveNameTag)
+		{
+			// Check the package custom version - the asset was saving empty tags for curves, so the absence of curves is
+			// not the same as an empty value
+			FAssetPackageData PackageData;
+			const UE::AssetRegistry::EExists PackageExists = AssetRegistryModule.Get().TryGetAssetPackageData(InAssetData.PackageName, PackageData);
+			if (PackageExists == UE::AssetRegistry::EExists::Exists)
+			{
+				for(const UE::AssetRegistry::FPackageCustomVersion& CustomVersion : PackageData.GetCustomVersions())
+				{
+					if(CustomVersion.Key == FAnimPhysObjectVersion::GUID)
+					{
+						return CustomVersion.Version < FAnimPhysObjectVersion::SmartNameRefactorForDeterministicCooking;
+					}
+				}
+
+				// No FAnimPhysObjectVersion, treat as old
+				return true;
+			}
+			else
+			{
+				// Does not exist or unknown - treat it as 'old'
+				return true;
+			}
+		}
+	}
+	else if(AssetClass->IsChildOf(USkeleton::StaticClass()))
+	{
+		return true;
+	}
+	else if(AssetClass->IsChildOf(USkeletalMesh::StaticClass()))
+	{
+		// Skeletal meshes didn't have curves before, so cant be 'old' 
+		return false;
+	}
+
+	// Assume unknown assets are not 'old'
+	return false;
+}
+
+FReply SAnimAssetFindReplace::HandleReplace()
+{
+	if(GetCurrentSelectionDelegate.IsBound())
+	{
+		TArray<FAssetData> SelectedAssets = GetCurrentSelectionDelegate.Execute();
+		if(SelectedAssets.Num() > 0)
+		{
+			ReplaceInAssets(SelectedAssets);
+		}
+	}
+
+	return FReply::Handled();
+}
+
+FReply SAnimAssetFindReplace::HandleReplaceAll()
+{
+	// Apply current filter
+	FAssetRegistryModule& AssetRegistryModule = FModuleManager::LoadModuleChecked<FAssetRegistryModule>("AssetRegistry");
+
+	TArray<FAssetData> FilteredAssets;
+	AssetRegistryModule.Get().GetAssets(AssetPickerConfig.Filter,FilteredAssets);
+
+	FilteredAssets.RemoveAll([this](const FAssetData& InAssetData)
+	{
+		bool bIsOldAsset = false;
+		return ShouldFilterOutAsset(InAssetData, bIsOldAsset);
+	});
+
+	ReplaceInAssets(FilteredAssets);
+
+	return FReply::Handled();
+}
+
+void SAnimAssetFindReplace::ReplaceInAssets(const TArray<FAssetData>& InAssetDatas)
+{
+	FText TypeName;
+	switch(Type)
+	{
+	case EAnimAssetFindReplaceType::Curves:
+		TypeName = LOCTEXT("FindReplaceTypeCurves", "Curves");
+		break;
+	case EAnimAssetFindReplaceType::Notifies:
+		TypeName = LOCTEXT("FindReplaceTypeNotifies", "Notifies");
+		break;
+	default:
+		break;
+	}
+
+	{
+		const FText MessageText = FText::Format(LOCTEXT("ReplacingTaskStatus", "Replacing {0}: Matching '{1}' with '{2}' in {3} Assets..."), TypeName, FText::FromString(FindString), FText::FromString(ReplaceString), FText::AsNumber(InAssetDatas.Num()));
+		FScopedSlowTask ScopedSlowTask(static_cast<float>(InAssetDatas.Num()), MessageText);
+		ScopedSlowTask.MakeDialog(true);
+
+		FScopedTransaction ScopedTransaction(FText::Format(LOCTEXT("ReplaceTransaction", "Replace {0}."), TypeName));
+		
+		for(const FAssetData& AssetData : InAssetDatas)
+		{
+			ScopedSlowTask.EnterProgressFrame();
+
+			ReplaceInAsset(AssetData);
+
+			if(ScopedSlowTask.ShouldCancel())
+			{
+				break;
+			}
+		}
+	}
+
+	RefreshSearchResults();
+}
+
+void SAnimAssetFindReplace::ReplaceInAsset(const FAssetData& InAssetData) const
+{
+	switch(Type)
+	{
+	case EAnimAssetFindReplaceType::Curves:
+		ReplaceCurvesInAsset(InAssetData);
+		break;
+	case EAnimAssetFindReplaceType::Notifies:
+		ReplaceNotifiesInAsset(InAssetData);
+		break;
+	default:
+		break;
+	}
+}
+
+void SAnimAssetFindReplace::ReplaceCurvesInAsset(const FAssetData& InAssetData) const
+{
+	if(UObject* Asset = InAssetData.GetAsset())
+	{
+		if(UAnimSequenceBase* AnimSequenceBase = Cast<UAnimSequenceBase>(Asset))
+		{
+			Asset->MarkPackageDirty();
+
+			if(bFindWholeWord)
+			{
+				const FAnimationCurveIdentifier FindCurveId(*FindString, ERawCurveTrackTypes::RCT_Float);
+				const FAnimationCurveIdentifier ReplaceCurveId(*ReplaceString, ERawCurveTrackTypes::RCT_Float);
+				IAnimationDataController::FScopedBracket ScopedBracket(AnimSequenceBase->GetController(), LOCTEXT("ReplaceCurves", "Replace Curves"));
+				AnimSequenceBase->GetController().RenameCurve(FindCurveId, ReplaceCurveId);
+			}
+			else
+			{
+				TArray<TPair<FAnimationCurveIdentifier, FAnimationCurveIdentifier>> FindReplacePairs;
+				const TArray<FFloatCurve>& Curves = AnimSequenceBase->GetDataModel()->GetFloatCurves();
+				for(const FFloatCurve& Curve : Curves)
+				{
+					FString CurveName = Curve.GetName().ToString();
+					if(NameMatches(CurveName))
+					{
+						const FAnimationCurveIdentifier FindCurveId(*CurveName, ERawCurveTrackTypes::RCT_Float);
+						const FString NewName = CurveName.Replace(*FindString, *ReplaceString, SearchCase);
+						const FAnimationCurveIdentifier ReplaceCurveId(*NewName, ERawCurveTrackTypes::RCT_Float);
+
+						FindReplacePairs.Emplace(FindCurveId, ReplaceCurveId);
+					}
+				}
+
+				if(FindReplacePairs.Num() > 0)
+				{
+					IAnimationDataController::FScopedBracket ScopedBracket(AnimSequenceBase->GetController(), LOCTEXT("ReplaceCurves", "Replace Curves"));
+
+					for(const TPair<FAnimationCurveIdentifier, FAnimationCurveIdentifier>& FindReplacePair : FindReplacePairs)
+					{
+						AnimSequenceBase->GetController().RenameCurve(FindReplacePair.Key, FindReplacePair.Value);
+					}
+				}
+			}
+		}
+		else if(UPoseAsset* PoseAsset = Cast<UPoseAsset>(Asset))
+		{
+			if(bFindWholeWord)
+			{
+				Asset->Modify();
+				
+				const FName FindCurveName(*FindString);
+				const FName ReplaceCurveName(*ReplaceString);
+				PoseAsset->RenamePoseOrCurveName(FindCurveName, ReplaceCurveName);
+			}
+			else
+			{
+				TArray<TPair<FName, FName>> FindReplacePairs;
+
+				for(const FName& PoseName : PoseAsset->GetPoseFNames())
+				{
+					FString CurveName = PoseName.ToString();
+					if(NameMatches(CurveName))
+					{
+						const FName FindCurveName(*CurveName);
+						const FString NewName = CurveName.Replace(*FindString, *ReplaceString, SearchCase);
+						const FName ReplaceCurveName(*NewName);
+
+						FindReplacePairs.Emplace(FindCurveName, ReplaceCurveName);
+					}
+				}
+
+				if(FindReplacePairs.Num() > 0)
+				{
+					Asset->Modify();
+
+					for(const TPair<FName, FName>& FindReplacePair : FindReplacePairs)
+					{
+						PoseAsset->RenamePoseOrCurveName(FindReplacePair.Key, FindReplacePair.Value);
+					}
+				}
+			}
+		}
+		else if(USkeleton* Skeleton = Cast<USkeleton>(Asset))
+		{
+			if(bFindWholeWord)
+			{
+				Asset->Modify();
+
+				const FName FindCurveName(*FindString);
+				const FName ReplaceCurveName(*ReplaceString);
+				Skeleton->RenameCurveMetaData(FindCurveName, ReplaceCurveName);
+			}
+			else
+			{
+				TArray<TPair<FName, FName>> FindReplacePairs;
+
+				Skeleton->ForEachCurveMetaData([this, &FindReplacePairs](FName InCurveName, const FCurveMetaData& InMetaData)
+				{
+					FString CurveNameString = InCurveName.ToString();
+					if(NameMatches(CurveNameString))
+					{
+						const FName FindCurveName(InCurveName);
+						const FString NewName = CurveNameString.Replace(*FindString, *ReplaceString, SearchCase);
+						const FName ReplaceCurveName(*NewName);
+
+						FindReplacePairs.Emplace(FindCurveName, ReplaceCurveName);
+					}
+				});
+
+				if(FindReplacePairs.Num() > 0)
+				{
+					Asset->Modify();
+
+					for(const TPair<FName, FName>& FindReplacePair : FindReplacePairs)
+					{
+						Skeleton->RenameCurveMetaData(FindReplacePair.Key, FindReplacePair.Value);
+					}
+				}
+			}
+		}
+		else if(USkeletalMesh* SkeletalMesh = Cast<USkeletalMesh>(Asset))
+		{
+			if(UAnimCurveMetaData* AnimCurveMetaData = SkeletalMesh->GetAssetUserData<UAnimCurveMetaData>())
+			{
+				if(bFindWholeWord)
+				{
+					Asset->Modify();
+
+					const FName FindCurveName(*FindString);
+					const FName ReplaceCurveName(*ReplaceString);
+					AnimCurveMetaData->RenameCurveMetaData(FindCurveName, ReplaceCurveName);
+				}
+				else
+				{
+					TArray<TPair<FName, FName>> FindReplacePairs;
+
+					AnimCurveMetaData->ForEachCurveMetaData([this, &FindReplacePairs](FName InCurveName, const FCurveMetaData& InMetaData)
+					{
+						FString CurveNameString = InCurveName.ToString();
+						if(NameMatches(CurveNameString))
+						{
+							const FName FindCurveName(InCurveName);
+							const FString NewName = CurveNameString.Replace(*FindString, *ReplaceString, SearchCase);
+							const FName ReplaceCurveName(*NewName);
+
+							FindReplacePairs.Emplace(FindCurveName, ReplaceCurveName);
+						}
+					});
+
+					if(FindReplacePairs.Num() > 0)
+					{
+						Asset->Modify();
+
+						for(const TPair<FName, FName>& FindReplacePair : FindReplacePairs)
+						{
+							AnimCurveMetaData->RenameCurveMetaData(FindReplacePair.Key, FindReplacePair.Value);
+						}
+					}
+				}
+			}
+		}
+	}
+}
+
+void SAnimAssetFindReplace::ReplaceNotifiesInAsset(const FAssetData& InAssetData) const
+{
+	if(UObject* Asset = InAssetData.GetAsset())
+	{
+		if(UAnimSequenceBase* AnimSequenceBase = Cast<UAnimSequenceBase>(Asset))
+		{
+			if(bFindWholeWord)
+			{
+				Asset->Modify();
+				AnimSequenceBase->RenameNotifies(*FindString, *ReplaceString);
+			}
+			else
+			{
+				TArray<TPair<FName, FName>> FindReplacePairs;
+
+				for(const FAnimNotifyEvent& Notify : AnimSequenceBase->Notifies)
+				{
+					// Only handle named notifies
+					if(!Notify.IsBlueprintNotify())
+					{
+						FString NotifyName = Notify.NotifyName.ToString();
+						if(NameMatches(NotifyName))
+						{
+							const FName FindNotifyName(*NotifyName);
+							const FString NewName = NotifyName.Replace(*FindString, *ReplaceString, SearchCase);
+							const FName ReplaceNotifyName(*NewName);
+							FindReplacePairs.AddUnique(TPair<FName, FName>(FindNotifyName, ReplaceNotifyName));
+						}
+					}
+				}
+
+				if(FindReplacePairs.Num() > 0)
+				{
+					Asset->Modify();
+
+					for(const TPair<FName, FName>& FindReplacePair : FindReplacePairs)
+					{
+						AnimSequenceBase->RenameNotifies(FindReplacePair.Key, FindReplacePair.Value);
+					}
+				}
+			}
+		}
+		else if(USkeleton* Skeleton = Cast<USkeleton>(Asset))
+		{
+			if(bFindWholeWord)
+			{
+				Asset->Modify();
+				Skeleton->RenameAnimationNotify(*FindString, *ReplaceString);
+			}
+			else
+			{
+				TArray<TPair<FName, FName>> FindReplacePairs;
+
+				for(const FName& NotifyName : Skeleton->AnimationNotifies)
+				{
+					FString NotifyString = NotifyName.ToString();
+					if(NameMatches(NotifyString))
+					{
+						const FString NewName = NotifyString.Replace(*FindString, *ReplaceString, SearchCase);
+						const FName ReplaceNotifyName(*NewName);
+						FindReplacePairs.AddUnique(TPair<FName, FName>(NotifyName, ReplaceNotifyName));
+					}
+				}
+
+				if(FindReplacePairs.Num() > 0)
+				{
+					Asset->Modify();
+					for(const TPair<FName, FName>& FindReplacePair : FindReplacePairs)
+					{
+						Skeleton->RenameAnimationNotify(FindReplacePair.Key, FindReplacePair.Value);
+					}
+				}
+			}
+		}
+	}
+}
+
+FReply SAnimAssetFindReplace::HandleRemove()
+{
+	if(GetCurrentSelectionDelegate.IsBound())
+	{
+		TArray<FAssetData> SelectedAssets = GetCurrentSelectionDelegate.Execute();
+		if(SelectedAssets.Num() > 0)
+		{
+			RemoveInAssets(SelectedAssets);
+		}
+	}
+
+	return FReply::Handled();
+}
+
+FReply SAnimAssetFindReplace::HandleRemoveAll()
+{
+	// Apply current filter
+	FAssetRegistryModule& AssetRegistryModule = FModuleManager::LoadModuleChecked<FAssetRegistryModule>("AssetRegistry");
+
+	TArray<FAssetData> FilteredAssets;
+	AssetRegistryModule.Get().GetAssets(AssetPickerConfig.Filter,FilteredAssets);
+
+	FilteredAssets.RemoveAll([this](const FAssetData& InAssetData)
+	{
+		bool bIsOldAsset = false;
+		return ShouldFilterOutAsset(InAssetData, bIsOldAsset);
+	});
+
+	RemoveInAssets(FilteredAssets);
+
+	return FReply::Handled();
+}
+
+void SAnimAssetFindReplace::RemoveInAssets(const TArray<FAssetData>& InAssetDatas)
+{
+	FText TypeName;
+	switch(Type)
+	{
+	case EAnimAssetFindReplaceType::Curves:
+		TypeName = LOCTEXT("FindReplaceTypeCurves", "Curves");
+		break;
+	case EAnimAssetFindReplaceType::Notifies:
+		TypeName = LOCTEXT("FindReplaceTypeNotifies", "Notifies");
+		break;
+	default:
+		break;
+	}
+
+	{
+		const FText MessageText = FText::Format(LOCTEXT("RemovingTaskStatus", "Removing {0}: Matching '{1}' in {2} Assets..."), TypeName, FText::FromString(FindString), FText::AsNumber(InAssetDatas.Num()));
+		FScopedSlowTask ScopedSlowTask(static_cast<float>(InAssetDatas.Num()), MessageText);
+		ScopedSlowTask.MakeDialog(true);
+
+		FScopedTransaction ScopedTransaction(FText::Format(LOCTEXT("RemoveTransaction", "Remove {0}."), TypeName));
+
+		for(const FAssetData& AssetData : InAssetDatas)
+		{
+			ScopedSlowTask.EnterProgressFrame();
+
+			RemoveInAsset(AssetData);
+
+			if(ScopedSlowTask.ShouldCancel())
+			{
+				break;
+			}
+		}
+	}
+
+	RefreshSearchResults();
+}
+
+void SAnimAssetFindReplace::RemoveInAsset(const FAssetData& InAssetData) const
+{
+	switch(Type)
+	{
+	case EAnimAssetFindReplaceType::Curves:
+		RemoveCurvesInAsset(InAssetData);
+		break;
+	case EAnimAssetFindReplaceType::Notifies:
+		RemoveNotifiesInAsset(InAssetData);
+		break;
+	default:
+		break;
+	}
+}
+
+void SAnimAssetFindReplace::RemoveCurvesInAsset(const FAssetData& InAssetData) const
+{
+	if(UObject* Asset = InAssetData.GetAsset())
+	{
+		if(UAnimSequenceBase* AnimSequenceBase = Cast<UAnimSequenceBase>(Asset))
+		{
+			Asset->MarkPackageDirty();
+
+			if(bFindWholeWord)
+			{
+				const FAnimationCurveIdentifier CurveId(*FindString, ERawCurveTrackTypes::RCT_Float);
+				IAnimationDataController::FScopedBracket ScopedBracket(AnimSequenceBase->GetController(), LOCTEXT("RemoveCurves", "Remove Curves"));
+				AnimSequenceBase->GetController().RemoveCurve(CurveId);
+			}
+			else
+			{
+				TSet<FAnimationCurveIdentifier> CurveIdsToRemove;
+				const TArray<FFloatCurve>& Curves = AnimSequenceBase->GetDataModel()->GetFloatCurves();
+				for(const FFloatCurve& Curve : Curves)
+				{
+					FString CurveName = Curve.GetName().ToString();
+					if(NameMatches(CurveName))
+					{
+						const FAnimationCurveIdentifier CurveId(*CurveName, ERawCurveTrackTypes::RCT_Float);
+						CurveIdsToRemove.Add(CurveId);
+					}
+				}
+
+				if(CurveIdsToRemove.Num() > 0)
+				{
+					IAnimationDataController::FScopedBracket ScopedBracket(AnimSequenceBase->GetController(), LOCTEXT("RemoveCurves", "Remove Curves"));
+					for(const FAnimationCurveIdentifier& CurveIdToRemove : CurveIdsToRemove)
+					{
+						AnimSequenceBase->GetController().RemoveCurve(CurveIdToRemove);
+					}
+				}
+			}
+		}
+		else if(UPoseAsset* PoseAsset = Cast<UPoseAsset>(Asset))
+		{
+			if(bFindWholeWord)
+			{
+				Asset->Modify();
+				PoseAsset->RemovePoseOrCurveNames({ *FindString });
+			}
+			else
+			{
+				TArray<FName> CurveIdsToRemove;
+
+				for(const FName& PoseName : PoseAsset->GetPoseFNames())
+				{
+					if(NameMatches(PoseName.ToString()))
+					{
+						CurveIdsToRemove.AddUnique(PoseName);
+					}
+				}
+
+				for(const FName& CurveName : PoseAsset->GetCurveFNames())
+				{
+					if(NameMatches(CurveName.ToString()))
+					{
+						CurveIdsToRemove.AddUnique(CurveName);
+					}
+				}
+				
+				if(CurveIdsToRemove.Num() > 0)
+				{
+					Asset->Modify();
+					PoseAsset->RemovePoseOrCurveNames(CurveIdsToRemove);
+				}
+			}
+		}
+		else if(USkeleton* Skeleton = Cast<USkeleton>(Asset))
+		{
+			if(bFindWholeWord)
+			{
+				Asset->Modify();
+				Skeleton->RemoveCurveMetaData(*FindString);
+			}
+			else
+			{
+				TArray<FName> CurvesToRemove;
+				Skeleton->ForEachCurveMetaData([this, &CurvesToRemove](FName InCurveName, const FCurveMetaData& InCurveMetaData)
+				{
+				 	if(NameMatches(InCurveName.ToString()))
+				 	{
+						 CurvesToRemove.AddUnique(InCurveName);
+					 }
+				});
+
+				if(CurvesToRemove.Num() > 0)
+				{
+					Asset->Modify();
+					Skeleton->RemoveCurveMetaData(CurvesToRemove);
+				}
+			}
+		}
+	}
+}
+
+void SAnimAssetFindReplace::RemoveNotifiesInAsset(const FAssetData& InAssetData) const
+{
+	if(UObject* Asset = InAssetData.GetAsset())
+	{
+		if(UAnimSequenceBase* AnimSequenceBase = Cast<UAnimSequenceBase>(Asset))
+		{
+			if(bFindWholeWord)
+			{
+				Asset->Modify();
+				AnimSequenceBase->RemoveNotifies( { *FindString } );
+			}
+			else
+			{
+				TArray<FName> NotifiesToRemove;
+				for(const FAnimNotifyEvent& Notify : AnimSequenceBase->Notifies)
+				{
+					FString NotifyNameString = Notify.NotifyName.ToString();
+					if(NameMatches(NotifyNameString))
+					{
+						NotifiesToRemove.AddUnique(Notify.NotifyName);
+					}
+				}
+
+				if(NotifiesToRemove.Num() > 0)
+				{
+					Asset->Modify();
+					AnimSequenceBase->RemoveNotifies(NotifiesToRemove);
+				}
+			}
+		}
+		else if(USkeleton* Skeleton = Cast<USkeleton>(Asset))
+		{
+			if(bFindWholeWord)
+			{
+				Asset->Modify();
+				Skeleton->RemoveAnimationNotify(*FindString);
+			}
+			else
+			{
+				TArray<FName> NotifiesToRemove;
+
+				for(const FName& NotifyName : Skeleton->AnimationNotifies)
+				{
+					if(NameMatches(NotifyName.ToString()))
+					{
+						NotifiesToRemove.AddUnique(NotifyName);
+					}
+				}
+
+				if(NotifiesToRemove.Num() > 0)
+				{
+					Asset->Modify();
+					for(const FName& NotifyToRemove : NotifiesToRemove)
+					{
+						Skeleton->RemoveAnimationNotify(NotifyToRemove);
+					}
+				}
+			}
+		}
+	}
+}
+
+void SAnimAssetFindReplace::RefreshAutoCompleteItems()
+{
+	AutoCompleteItems->Empty();
+
+	// We use the asset registry to query all assets and accumulate their curve names
+	const FAssetRegistryModule& AssetRegistryModule = FModuleManager::LoadModuleChecked<FAssetRegistryModule>("AssetRegistry");
+
+	FARFilter Filter;
+	Filter.bRecursiveClasses = true;
+	Filter.ClassPaths = { UAnimationAsset::StaticClass()->GetClassPathName(), USkeleton::StaticClass()->GetClassPathName() };
+
+	TArray<FAssetData> FoundAssetData;
+	AssetRegistryModule.Get().GetAssets(Filter, FoundAssetData);
+
+	TSet<FString> UniqueNames;
+	
+	switch(Type)
+	{
+	case EAnimAssetFindReplaceType::Curves:
+		for (const FAssetData& AssetData : FoundAssetData)
+		{
+			const FString TagValue = AssetData.GetTagValueRef<FString>(USkeleton::CurveNameTag);
+			if (!TagValue.IsEmpty())
+			{
+				TArray<FString> AssetCurveNames;
+				if (TagValue.ParseIntoArray(AssetCurveNames, *USkeleton::CurveTagDelimiter, true) > 0)
+				{
+					for (FString& CurveNameString : AssetCurveNames)
+					{
+						UniqueNames.Add(MoveTemp(CurveNameString));
+					}
+				}
+			}
+		}
+		break;
+	case EAnimAssetFindReplaceType::Notifies:
+		for (const FAssetData& AssetData : FoundAssetData)
+		{
+			const FString TagValue = AssetData.GetTagValueRef<FString>(USkeleton::AnimNotifyTag);
+			if (!TagValue.IsEmpty())
+			{
+				TArray<FString> AssetNotifyNames;
+				if (TagValue.ParseIntoArray(AssetNotifyNames, *USkeleton::AnimNotifyTagDelimiter, true) > 0)
+				{
+					for (FString& NotifyNameString : AssetNotifyNames)
+					{
+						UniqueNames.Add(MoveTemp(NotifyNameString));
+					}
+				}
+			}
+		}
+		break;
+	default:
+		break;
+	}
+
+	for(FString& UniqueName : UniqueNames)
+	{
+		AutoCompleteItems->Add(MakeShared<FString>(MoveTemp(UniqueName)));
+	}
+
+	AutoCompleteItems->Sort([](TSharedPtr<FString> InLHS, TSharedPtr<FString> InRHS)
+	{
+		return *InLHS.Get() < *InRHS.Get();
+	});
+
+	FindSearchBox->RefreshAutoCompleteItems();
+	ReplaceSearchBox->RefreshAutoCompleteItems();
+}
+
+#undef LOCTEXT_NAMESPACE
\ No newline at end of file
diff --git a/Engine/Source/Editor/Persona/Private/SAnimAssetFindReplace.h b/Engine/Source/Editor/Persona/Private/SAnimAssetFindReplace.h
new file mode 100644
--- /dev/null
+++ b/Engine/Source/Editor/Persona/Private/SAnimAssetFindReplace.h
@@ -0,0 +1,126 @@
+﻿// Copyright Epic Games, Inc. All Rights Reserved.
+
+#pragma once
+
+#include "CoreTypes.h"
+#include "IContentBrowserSingleton.h"
+#include "IDocumentation.h"
+#include "PersonaTabs.h"
+#include "Widgets/SWindow.h"
+#include "Widgets/SCompoundWidget.h"
+#include "WorkflowOrientedApp/WorkflowTabFactory.h"
+#include "AnimAssetFindReplace.h"
+
+class SBox;
+class SAutoCompleteSearchBox;
+
+struct FAnimAssetFindReplaceSummoner : public FWorkflowTabFactory
+{
+public:
+	FAnimAssetFindReplaceSummoner(TSharedPtr<FAssetEditorToolkit> InHostingApp, const FAnimAssetFindReplaceConfig& InConfig);
+
+	virtual TSharedRef<SWidget> CreateTabBody(const FWorkflowTabSpawnInfo& Info) const override;
+	virtual TSharedPtr<SToolTip> CreateTabToolTipWidget(const FWorkflowTabSpawnInfo& Info) const override;
+
+private:
+	FAnimAssetFindReplaceConfig Config;
+};
+
+class SAnimAssetFindReplace : public SCompoundWidget
+{
+public:
+	SLATE_BEGIN_ARGS(SAnimAssetFindReplace) {}
+	
+	SLATE_ARGUMENT(FAnimAssetFindReplaceConfig, Config)
+
+	SLATE_END_ARGS()
+
+	void Construct(const FArguments& InArgs);
+
+	/** Sets the type of thing we are finding/replacing, then flags the UI for a refresh */
+	void SetFindReplaceType(EAnimAssetFindReplaceType InType);
+
+private:
+	void RefreshSearchResults();
+
+	bool HandleFilterAsset(const FAssetData& InAssetData);
+
+	bool FilterByCurve(const FAssetData& InAssetData, bool& bOutIsOldAsset) const;
+	
+	bool FilterByNotify(const FAssetData& InAssetData, bool& bOutIsOldAsset) const;
+
+	void GetMatchingCurveNamesForAsset(const FAssetData& InAssetData, TArray<FString>& OutCurveNames) const;
+
+	void GetMatchingNotifyNamesForAsset(const FAssetData& InAssetData, TArray<FString>& OutNotifyNames) const;
+
+	FReply HandleReplace();
+
+	FReply HandleReplaceAll();
+
+	void ReplaceInAssets(const TArray<FAssetData>& InAssetDatas);
+
+	void ReplaceInAsset(const FAssetData& InAssetData) const;
+
+	void ReplaceCurvesInAsset(const FAssetData& InAssetData) const;
+
+	void ReplaceNotifiesInAsset(const FAssetData& InAssetData) const;
+	
+	FReply HandleRemove();
+
+	FReply HandleRemoveAll();
+	
+	void RemoveInAssets(const TArray<FAssetData>& InAssetDatas);
+
+	void RemoveInAsset(const FAssetData& InAssetData) const;
+
+	void RemoveCurvesInAsset(const FAssetData& InAssetData) const;
+
+	void RemoveNotifiesInAsset(const FAssetData& InAssetData) const;
+
+	bool ShouldFilterOutAsset(const FAssetData& InAssetData, bool& bOutIsOldAsset) const;
+
+	bool NameMatches(const FString& InNameString) const;
+
+	void RefreshAutoCompleteItems();
+
+	FARFilter MakeARFilter() const;
+
+	bool IsAssetWithoutTagOldAsset(FName InTag, const FAssetData& InAssetData) const;
+
+private:
+	FAssetPickerConfig AssetPickerConfig;
+
+	EAnimAssetFindReplaceMode Mode = EAnimAssetFindReplaceMode::Find;
+
+	EAnimAssetFindReplaceType Type = EAnimAssetFindReplaceType::Curves;
+
+	FString FindString;
+
+	FString ReplaceString;
+
+	FRefreshAssetViewDelegate RefreshAssetViewDelegate;
+
+	FGetCurrentSelectionDelegate GetCurrentSelectionDelegate;
+
+	FSetARFilterDelegate SetARFilterDelegate;
+
+	bool bFindWholeWord = true;
+
+	bool bAssetsSelected = false;
+
+	bool bFoundAssets = false;
+
+	TArray<FAssetData> OldAssets;
+
+	ESearchCase::Type SearchCase = ESearchCase::IgnoreCase;
+
+	TStrongObjectPtr<UAnimAssetFindReplaceContext> ToolbarContext;
+
+	TSharedPtr<TArray<TSharedPtr<FString>>> AutoCompleteItems;
+
+	TSharedPtr<SAutoCompleteSearchBox> FindSearchBox;
+	
+	TSharedPtr<SAutoCompleteSearchBox> ReplaceSearchBox;
+
+	FAssetData SkeletonFilter;
+};
\ No newline at end of file
diff --git a/Engine/Source/Editor/Persona/Private/SAnimCurvePicker.cpp b/Engine/Source/Editor/Persona/Private/SAnimCurvePicker.cpp
--- a/Engine/Source/Editor/Persona/Private/SAnimCurvePicker.cpp
+++ b/Engine/Source/Editor/Persona/Private/SAnimCurvePicker.cpp
@@ -15,6 +15,8 @@
 #include "Widgets/SOverlay.h"
 #include "Widgets/Layout/SMenuOwner.h"
 #include "Widgets/Layout/SBorder.h"
+#include "Widgets/Input/SCheckBox.h"
+#include "Engine/SkeletalMesh.h"
 
 #define LOCTEXT_NAMESPACE "SAnimCurvePicker"
 
@@ -22,17 +24,18 @@ SAnimCurvePicker::~SAnimCurvePicker()
 {
 }
 
-void SAnimCurvePicker::Construct(const FArguments& InArgs, const TSharedRef<IEditableSkeleton>& InEditableSkeleton)
+void SAnimCurvePicker::Construct(const FArguments& InArgs, const USkeleton* InSkeleton)
 {
-	OnCurveNamePicked = InArgs._OnCurveNamePicked;
-	IsCurveMarkedForExclusion = InArgs._IsCurveMarkedForExclusion;
-	EditableSkeleton = InEditableSkeleton;
+	OnCurvePicked = InArgs._OnCurvePicked;
+	IsCurveNameMarkedForExclusion = InArgs._IsCurveNameMarkedForExclusion;
+	Skeleton = InSkeleton;
+	bShowOtherSkeletonCurves = false;
 
 	SAssignNew(SearchBox, SSearchBox)
 	.HintText(LOCTEXT("SearchBoxHint", "Search"))
 	.OnTextChanged(this, &SAnimCurvePicker::HandleFilterTextChanged);
 
-	SAssignNew(NameListView, SListView<TSharedPtr<FSmartName>>)
+	SAssignNew(NameListView, SListView<TSharedPtr<FName>>)
 	.SelectionMode(ESelectionMode::Single)
 	.ListItemsSource(&CurveNames)
 	.OnSelectionChanged(this, &SAnimCurvePicker::HandleSelectionChanged)
@@ -41,14 +44,16 @@ void SAnimCurvePicker::Construct(const FArguments& InArgs, const TSharedRef<IEdi
 	const float HorizontalPadding = 8.0f;
 	const float VerticalPadding = 2.0f;
 	const float WeightOverride = 300.0f;
-	
+
+	TSharedPtr<SVerticalBox> VerticalBox;
+
 	ChildSlot
 	[
 		SNew(SMenuOwner)
 		[
-			SNew(SListViewSelectorDropdownMenu<TSharedPtr<FSmartName>>, SearchBox, NameListView)
+			SNew(SListViewSelectorDropdownMenu<TSharedPtr<FName>>, SearchBox, NameListView)
 			[
-				SNew(SVerticalBox)
+				SAssignNew(VerticalBox, SVerticalBox)
 				+SVerticalBox::Slot()
 				.AutoHeight()
 				.Padding(HorizontalPadding, VerticalPadding)
@@ -57,7 +62,6 @@ void SAnimCurvePicker::Construct(const FArguments& InArgs, const TSharedRef<IEdi
 				]
 				+SVerticalBox::Slot()
 				.FillHeight(1.0f)
-				.VAlign(VAlign_Fill)
 				.Padding(HorizontalPadding, VerticalPadding)
 				[
 					SNew(SBox)
@@ -75,7 +79,7 @@ void SAnimCurvePicker::Construct(const FArguments& InArgs, const TSharedRef<IEdi
 						[
 							SNew(SScrollBox)
 							.Orientation(EOrientation::Orient_Vertical)
-							+ SScrollBox::Slot()
+							+SScrollBox::Slot()
 							.HAlign(HAlign_Fill)
 							.VAlign(VAlign_Fill)
 							[
@@ -88,10 +92,35 @@ void SAnimCurvePicker::Construct(const FArguments& InArgs, const TSharedRef<IEdi
 		]
 	];
 
+	if(Skeleton.IsValid())
+	{
+		VerticalBox->AddSlot()
+			.AutoHeight()
+			.Padding(HorizontalPadding, VerticalPadding)
+			[
+				SNew(SCheckBox)
+				.ToolTipText(LOCTEXT("ShowOtherSkeletonsTooltip", "Whether to show all curves or just the curves from the current skeleton"))
+				.IsChecked_Lambda([this]()
+				{
+					return bShowOtherSkeletonCurves ? ECheckBoxState::Checked : ECheckBoxState::Unchecked;
+				})
+				.OnCheckStateChanged_Lambda([this](ECheckBoxState InState)
+				{
+					bShowOtherSkeletonCurves = (InState == ECheckBoxState::Checked);
+					RefreshListItems();
+				})
+				.Content()
+				[
+					SNew(STextBlock)
+					.Text(LOCTEXT("ShowOtherSkeletons", "Show curves from other skeletons"))
+				]
+			];
+	}
+	
 	RefreshListItems();
 }
 
-void SAnimCurvePicker::HandleSelectionChanged(TSharedPtr<FSmartName> InItem, ESelectInfo::Type InSelectionType)
+void SAnimCurvePicker::HandleSelectionChanged(TSharedPtr<FName> InItem, ESelectInfo::Type InSelectionType)
 {
 	// When the user is navigating, do not act upon the selection change
 	if (InSelectionType == ESelectInfo::OnNavigation)
@@ -101,21 +130,21 @@ void SAnimCurvePicker::HandleSelectionChanged(TSharedPtr<FSmartName> InItem, ESe
 	
 	if (InItem.IsValid())
 	{
-		OnCurveNamePicked.ExecuteIfBound(*InItem);
+		OnCurvePicked.ExecuteIfBound(*InItem);
 	}
 }
 
-TSharedRef<ITableRow> SAnimCurvePicker::HandleGenerateRow(TSharedPtr<FSmartName> InItem, const TSharedRef<STableViewBase>& InOwnerTable)
+TSharedRef<ITableRow> SAnimCurvePicker::HandleGenerateRow(TSharedPtr<FName> InItem, const TSharedRef<STableViewBase>& InOwnerTable)
 {
 	return 
-		SNew(STableRow<TSharedPtr<FSmartName>>, InOwnerTable)
+		SNew(STableRow<TSharedPtr<FName>>, InOwnerTable)
 		.Padding(FMargin(8.0f, 0.0f))
 		[
 			SNew(SBox)
 			.VAlign(VAlign_Center)
 			[
 				SNew(STextBlock)
-				.Text(FText::FromName(InItem->DisplayName))
+				.Text(FText::FromName(*InItem))
 				.HighlightText_Lambda([this]() { return FText::FromString(FilterText); })
 			]
 		];
@@ -123,26 +152,32 @@ TSharedRef<ITableRow> SAnimCurvePicker::HandleGenerateRow(TSharedPtr<FSmartName>
 
 void SAnimCurvePicker::RefreshListItems()
 {
-	// Ensure valid skeleton exists
-	const TSharedPtr<IEditableSkeleton> CurrentEditableSkeleton = EditableSkeleton.Pin();
-	check(CurrentEditableSkeleton);
-	const USkeleton & CurrentSkeleton = CurrentEditableSkeleton->GetSkeleton();
-
-	// Get name container for all the curves stored by the skeleton
-	const FSmartNameMapping* CurveNameMapping = CurrentSkeleton.GetSmartNameContainer(USkeleton::AnimCurveMappingName);
-	CurveNameMapping->FillUidArray(CurveSmartNameUids);
-	
+	const USkeleton* CurrentSkeleton = Skeleton.Get();
+	FString CurrentSkeletonName;
+	if(CurrentSkeleton)
+	{
+		CurrentSkeletonName = FAssetData(CurrentSkeleton).GetExportTextName();
+	}
+
 	CurveNames.Reset();
 	UniqueCurveNames.Reset();
 
 	{
-		// We use the asset registry to query all assets with the supplied skeleton, and accumulate their curve names
+		// First check skeleton metadata
+		if(CurrentSkeleton)
+		{
+			CurrentSkeleton->ForEachCurveMetaData([this](FName InCurveName, const FCurveMetaData& InMetaData)
+			{
+				UniqueCurveNames.Add(InCurveName);
+			});
+		}
+
+		// We use the asset registry to query all assets (optionally with the supplied skeleton) and accumulate their curve names
 		const FAssetRegistryModule& AssetRegistryModule = FModuleManager::LoadModuleChecked<FAssetRegistryModule>("AssetRegistry");
 
 		FARFilter Filter;
 		Filter.bRecursiveClasses = true;
-		Filter.ClassPaths.Add(UAnimationAsset::StaticClass()->GetClassPathName());
-		Filter.TagsAndValues.Add(TEXT("Skeleton"), FAssetData(&CurrentSkeleton).GetExportTextName());
+		Filter.ClassPaths.Append({ UAnimationAsset::StaticClass()->GetClassPathName(), USkeletalMesh::StaticClass()->GetClassPathName(), USkeleton::StaticClass()->GetClassPathName() } );
 
 		TArray<FAssetData> FoundAssetData;
 		AssetRegistryModule.Get().GetAssets(Filter, FoundAssetData);
@@ -150,43 +185,45 @@ void SAnimCurvePicker::RefreshListItems()
 		// Build set of unique curve smart names
 		for (FAssetData& AssetData : FoundAssetData)
 		{
+			if(!bShowOtherSkeletonCurves && AssetData.GetClass() != USkeleton::StaticClass())
+			{
+				if(!CurrentSkeletonName.IsEmpty())
+				{
+					// Check skeleton tag
+					const FString SkeletonName = AssetData.GetTagValueRef<FString>(USkeleton::StaticClass()->GetFName());
+					if(SkeletonName != CurrentSkeletonName)
+					{
+						continue;
+					}
+				}
+			}
+			
 			const FString TagValue = AssetData.GetTagValueRef<FString>(USkeleton::CurveNameTag);
 			if (!TagValue.IsEmpty())
 			{
-				// parse notifies
-				if (TagValue.ParseIntoArray(CurveNamesQueriedFromAssetRegistry, *USkeleton::CurveTagDelimiter, true) > 0)
+				TArray<FString> AssetCurveNames;
+				if (TagValue.ParseIntoArray(AssetCurveNames, *USkeleton::CurveTagDelimiter, true) > 0)
 				{
-					for (const FString& CurveNameString : CurveNamesQueriedFromAssetRegistry)
+					for (const FString& CurveNameString : AssetCurveNames)
 					{
 						FName CurveName = FName(*CurveNameString);
-						FSmartName InCurveSmartName = {CurveName, CurveNameMapping->FindUID(CurveName)};
-						
-						if (IsCurveMarkedForExclusion.IsBound() && IsCurveMarkedForExclusion.Execute(InCurveSmartName))
+
+						if(CurveName == NAME_None)
 						{
 							continue;
 						}
-						
-						UniqueCurveNames.Add(InCurveSmartName);
+
+						if (IsCurveNameMarkedForExclusion.IsBound() && IsCurveNameMarkedForExclusion.Execute(CurveName))
+						{
+							continue;
+						}
+
+						UniqueCurveNames.Add(CurveName);
 					}
 				}
 			}
 		}
 	}
-	
-	// Add unique curves that were not found using asset registry
-	for (SmartName::UID_Type Id : CurveSmartNameUids)
-	{
-		FName OutCurveName;
-		if (CurveNameMapping->GetName(Id, OutCurveName))
-		{
-			FSmartName InCurveSmartName(OutCurveName, Id);
-
-			if (IsCurveMarkedForExclusion.IsBound() && IsCurveMarkedForExclusion.Execute(InCurveSmartName))
-				continue;
-			
-			UniqueCurveNames.Add(InCurveSmartName);
-		}
-	}
 
 	FilterAvailableCurves();
 }
@@ -196,24 +233,24 @@ void SAnimCurvePicker::FilterAvailableCurves()
 	CurveNames.Reset();
 	
 	// Exact filtering
-	for (const FSmartName& UniqueCurveName : UniqueCurveNames)
+	for (const FName& UniqueCurveName : UniqueCurveNames)
 	{
-		if (FilterText.IsEmpty() || UniqueCurveName.DisplayName.ToString().Contains(FilterText))
+		if (FilterText.IsEmpty() || UniqueCurveName.ToString().Contains(FilterText))
 		{
-			CurveNames.Add(MakeShared<FSmartName>(UniqueCurveName));
+			CurveNames.Add(MakeShared<FName>(UniqueCurveName));
 		}
 	}
 
 	// Alphabetical sorting
 	{
-		struct FSmartNameSortItemSortOp
+		struct FNameSortItemSortOp
 		{
-			FORCEINLINE bool operator()( const TSharedPtr<FSmartName>& A, const TSharedPtr<FSmartName>& B ) const
+			FORCEINLINE bool operator()( const TSharedPtr<FName>& A, const TSharedPtr<FName>& B ) const
 			{
-				return (A->DisplayName.Compare(B->DisplayName) < 0);
+				return (A->ToString().Compare(B->ToString()) < 0);
 			}
 		};
-		CurveNames.Sort(FSmartNameSortItemSortOp());
+		CurveNames.Sort(FNameSortItemSortOp());
 	}
 
 	// Rebuild list view
diff --git a/Engine/Source/Editor/Persona/Private/SAnimCurvePicker.h b/Engine/Source/Editor/Persona/Private/SAnimCurvePicker.h
--- a/Engine/Source/Editor/Persona/Private/SAnimCurvePicker.h
+++ b/Engine/Source/Editor/Persona/Private/SAnimCurvePicker.h
@@ -8,12 +8,10 @@
 #include "Widgets/Views/STableRow.h"
 #include "Animation/SmartName.h"
 #include "Widgets/Input/SSearchBox.h"
+#include "PersonaDelegates.h"
 
 class IEditableSkeleton;
 
-DECLARE_DELEGATE_OneParam(FOnCurveNamePicked, const FSmartName& /*PickedName*/)
-DECLARE_DELEGATE_RetVal_OneParam(bool, FIsCurveMarkedForExclusion, const FSmartName& /*CurveName*/)
-
 class SAnimCurvePicker : public SCompoundWidget
 {
 public:
@@ -23,17 +21,18 @@ public:
 	
 	SLATE_BEGIN_ARGS(SAnimCurvePicker) {}
 
-	SLATE_EVENT(FOnCurveNamePicked, OnCurveNamePicked)
-	SLATE_EVENT(FIsCurveMarkedForExclusion, IsCurveMarkedForExclusion)
+	SLATE_EVENT(FOnCurvePicked, OnCurvePicked)
+
+	SLATE_EVENT(FIsCurveNameMarkedForExclusion, IsCurveNameMarkedForExclusion)
 	
 	SLATE_END_ARGS()
 
 	/**
 	 * Construct this widget.
 	 * @param InArgs - The declaration data for this widget
-	 * @param InEditableSkeleton - The skeleton from which the widget extracts curve information 
+	 * @param InSkeleton - The skeleton from which the widget extracts curve information 
 	 */
-	void Construct(const FArguments& InArgs, const TSharedRef<IEditableSkeleton>& InEditableSkeleton);
+	void Construct(const FArguments& InArgs, const USkeleton* InSkeleton);
 
 private:
 	/** Refresh the list of available curves */
@@ -43,45 +42,35 @@ private:
 	void FilterAvailableCurves();
 
 	/** UI handlers */
-	void HandleSelectionChanged(TSharedPtr<FSmartName> InItem, ESelectInfo::Type InSelectionType);
-	TSharedRef<ITableRow> HandleGenerateRow(TSharedPtr<FSmartName> InItem, const TSharedRef<STableViewBase>& InOwnerTable);
+	void HandleSelectionChanged(TSharedPtr<FName> InItem, ESelectInfo::Type InSelectionType);
+	TSharedRef<ITableRow> HandleGenerateRow(TSharedPtr<FName> InItem, const TSharedRef<STableViewBase>& InOwnerTable);
 	void HandleFilterTextChanged(const FText& InFilterText);
 
 private:
 	/** Delegate fired when a curve name is picked */
-	FOnCurveNamePicked OnCurveNamePicked;
+	FOnCurvePicked OnCurvePicked;
 
 	/* Filter predicate to determine if curve should be excluded from the picker's curve list */
-	FIsCurveMarkedForExclusion IsCurveMarkedForExclusion;
-	
+	FIsCurveNameMarkedForExclusion IsCurveNameMarkedForExclusion;
+
 	/** The editable skeleton we use to grab curves from */
-	TWeakPtr<IEditableSkeleton> EditableSkeleton;
+	TWeakObjectPtr<const USkeleton> Skeleton;
 
 	/** The names of the curves we are displaying */
-	TArray<TSharedPtr<FSmartName>> CurveNames;
-
-	struct FSmartNameKeyFuncs : public DefaultKeyFuncs<FSmartName>
-	{
-		static FORCEINLINE FSmartName GetSetKey(FSmartName const& Element) { return Element; }
-		static FORCEINLINE uint32 GetKeyHash(FSmartName const& Key) { return GetTypeHash(Key.DisplayName); }
-		static FORCEINLINE bool Matches(FSmartName const& A, FSmartName const& B) { return A == B; }
-	};
-	
+	TArray<TSharedPtr<FName>> CurveNames;
+
 	/** All the unique curve names we can find */
-	TSet<FSmartName, FSmartNameKeyFuncs> UniqueCurveNames;
+	TSet<FName> UniqueCurveNames;
 
 	/** The string we use to filter curve names */
 	FString FilterText;
 
 	/** The list view used to display names */
-	TSharedPtr<SListView<TSharedPtr<FSmartName>>> NameListView;
+	TSharedPtr<SListView<TSharedPtr<FName>>> NameListView;
 
 	/** The search box used to filter curves */
 	TSharedPtr<SSearchBox> SearchBox;
 
-	/** The Uids of all the skeleton's curves in its mapping container */
-	TArray<SmartName::UID_Type> CurveSmartNameUids;
-
-	/** The container used to store curve names obtained from asset registry */
-	TArray<FString> CurveNamesQueriedFromAssetRegistry;
+	/** Whether we should show other skeleton's curves */
+	bool bShowOtherSkeletonCurves;
 };
\ No newline at end of file
