Animation Curve Runtime & Editor Improvements

diff --git a/Engine/Source/Editor/Persona/Private/SAnimCurveViewer.cpp b/Engine/Source/Editor/Persona/Private/SAnimCurveViewer.cpp
--- a/Engine/Source/Editor/Persona/Private/SAnimCurveViewer.cpp
+++ b/Engine/Source/Editor/Persona/Private/SAnimCurveViewer.cpp
@@ -2,51 +2,100 @@
 
 
 #include "SAnimCurveViewer.h"
+
+#include "AnimationEditorUtils.h"
 #include "Framework/Notifications/NotificationManager.h"
 #include "Widgets/Input/SCheckBox.h"
 #include "Widgets/Notifications/SNotificationList.h"
-#include "Framework/Commands/UICommandList.h"
 #include "Framework/MultiBox/MultiBoxBuilder.h"
 #include "Styling/AppStyle.h"
-#include "Layout/WidgetPath.h"
-#include "Framework/Application/MenuStack.h"
-#include "Framework/Application/SlateApplication.h"
 #include "Widgets/Input/SSpinBox.h"
 #include "Animation/DebugSkelMeshComponent.h"
 #include "Widgets/Input/SSearchBox.h"
-#include "Widgets/Text/SInlineEditableTextBlock.h"
-#include "Widgets/Input/STextEntryPopup.h"
 #include "Animation/AnimSingleNodeInstance.h"
 #include "IEditableSkeleton.h"
-
-#include "Framework/Commands/GenericCommands.h"
+#include "Widgets/Colors/SColorBlock.h"
 #include "CurveViewerCommands.h"
+#include "Animation/AnimBlueprintGeneratedClass.h"
 #include "Animation/EditorAnimCurveBoneLinks.h"
 #include "HAL/PlatformApplicationMisc.h"
-#include "ScopedTransaction.h"
+#include "AssetRegistry/AssetRegistryModule.h"
+#include "Misc/ScopedSlowTask.h"
+#include "Engine/PoseWatch.h"
+#include "Filters/GenericFilter.h"
+#include "Filters/SBasicFilterBar.h"
+#include "Widgets/Input/STextComboBox.h"
+#include "AnimPreviewInstance.h"
+#include "PersonaTabs.h"
+#include "SAnimAssetFindReplace.h"
+#include "Widgets/Docking/SDockTab.h"
 
 #define LOCTEXT_NAMESPACE "SAnimCurveViewer"
 
-static const FName ColumnId_AnimCurveNameLabel( "Curve Name" );
-static const FName ColumnID_AnimCurveTypeLabel("Type");
-static const FName ColumnID_AnimCurveWeightLabel( "Weight" );
-static const FName ColumnID_AnimCurveEditLabel( "Edit" );
-static const FName ColumnID_AnimCurveNumBoneLabel("Num Bones");
+namespace CurveViewerColumns
+{
+static const FName AnimCurveNameLabel( "Curve Name" );
+static const FName AnimCurveTypeLabel("Type");
+static const FName AnimCurveWeightLabel( "Weight" );
+static const FName AnimCurveEditLabel( "Edit" );
+}
 
-extern const float MaxMorphWeight = 5.f;
 //////////////////////////////////////////////////////////////////////////
 // SAnimCurveListRow
 
 typedef TSharedPtr< FDisplayedAnimCurveInfo > FDisplayedAnimCurveInfoPtr;
-//  This is a flag that is used to filter UI part
-enum EAnimCurveEditorFlags
+
+// This is a flag that is used to filter UI part
+enum class EAnimCurveViewerFilterFlags : uint8 
 {
-	// Used as morph target curve
-	ACEF_DriveMorphTarget		= 0x00000001, 
-	// Used as triggering event
-	ACEF_DriveAttribute			= 0x00000002, 
-	// Used as a material curve
-	ACEF_DriveMaterial			= 0x00000004, 
+	// Show all
+	ShowAll			= 0, 
+	// Show active curves
+	Active			= 0x01, 
+	// Show morph target curves
+	MorphTarget		= 0x02, 
+	// Show material curves
+	Material		= 0x04, 
+};
+
+ENUM_CLASS_FLAGS(EAnimCurveViewerFilterFlags);
+
+class FAnimCurveViewerFilter : public FGenericFilter<EAnimCurveViewerFilterFlags>
+{
+public:
+	FAnimCurveViewerFilter(EAnimCurveViewerFilterFlags InFlags, const FString& InName, const FText& InDisplayName, const FText& InToolTipText, FLinearColor InColor, TSharedPtr<FFilterCategory> InCategory)
+		: FGenericFilter<EAnimCurveViewerFilterFlags>(InCategory, InName, InDisplayName, FGenericFilter<EAnimCurveViewerFilterFlags>::FOnItemFiltered())
+		, Flags(InFlags)
+	{
+		ToolTip = InToolTipText;
+		Color = InColor;
+	}
+
+	bool IsActive() const
+	{
+		return bIsActive;
+	}
+
+	EAnimCurveViewerFilterFlags GetFlags() const
+	{
+		return Flags;
+	}
+	
+private:
+	// FFilterBase interface
+	virtual void ActiveStateChanged(bool bActive) override
+	{
+		bIsActive = bActive;
+	}
+
+	virtual bool PassesFilter(EAnimCurveViewerFilterFlags InItem) const override
+	{
+		return EnumHasAnyFlags(InItem, Flags);
+	}
+	
+private:
+	EAnimCurveViewerFilterFlags Flags;
+	bool bIsActive = false;
 };
 
 void SAnimCurveListRow::Construct( const FArguments& InArgs, const TSharedRef<STableViewBase>& InOwnerTableView, const TSharedRef<IPersonaPreviewScene>& InPreviewScene)
@@ -62,116 +111,84 @@ void SAnimCurveListRow::Construct( const FArguments& InArgs, const TSharedRef<ST
 
 TSharedRef< SWidget > SAnimCurveListRow::GenerateWidgetForColumn( const FName& ColumnName )
 {
-	if ( ColumnName == ColumnId_AnimCurveNameLabel )
+	if ( ColumnName == CurveViewerColumns::AnimCurveNameLabel )
 	{
 		TSharedPtr<SAnimCurveViewer> AnimCurveViewer = AnimCurveViewerPtr.Pin();
 		if (AnimCurveViewer.IsValid())
 		{
 			return
 				SNew(SVerticalBox)
-
+				.ToolTipText(this, &SAnimCurveListRow::GetItemName)
+			
 				+ SVerticalBox::Slot()
-				.AutoHeight()
+				.FillHeight(1.0f)
 				.Padding(4)
 				.VAlign(VAlign_Center)
 				[
-					SAssignNew(Item->EditableText, SInlineEditableTextBlock)
-					.OnTextCommitted(AnimCurveViewer.Get(), &SAnimCurveViewer::OnNameCommitted, Item)
+					SNew(STextBlock)
 					.Font(this, &SAnimCurveListRow::GetItemFont)
-					.IsSelected(this, &SAnimCurveListRow::IsSelected)
 					.Text(this, &SAnimCurveListRow::GetItemName)
 					.HighlightText(this, &SAnimCurveListRow::GetFilterText)
 				];
 		}
-		else
-		{
-			return SNullWidget::NullWidget;
-		}
 	}
-	else if (ColumnName == ColumnID_AnimCurveTypeLabel)
+	else if (ColumnName == CurveViewerColumns::AnimCurveTypeLabel)
 	{
 		TSharedPtr<SAnimCurveViewer> AnimCurveViewer = AnimCurveViewerPtr.Pin();
 		if (AnimCurveViewer.IsValid())
 		{
 			return
 				SNew(SVerticalBox)
+				.ToolTipText(LOCTEXT("AnimCurveTypeTooltip", "The type of the curve (e.g. morph target, material)."))
 
 				+ SVerticalBox::Slot()
-				.AutoHeight()
-				.Padding(4)
+				.FillHeight(1.0f)
 				.VAlign(VAlign_Center)
 				[
 					GetCurveTypeWidget()
 				];
 		}
-		else
-		{
-			return SNullWidget::NullWidget;
-		}
 	}
-	else if ( ColumnName == ColumnID_AnimCurveWeightLabel )
+	else if ( ColumnName == CurveViewerColumns::AnimCurveWeightLabel )
 	{
 		// Encase the SSpinbox in an SVertical box so we can apply padding. Setting ItemHeight on the containing SListView has no effect :-(
 		return
 			SNew( SVerticalBox )
-
+			.ToolTipText(LOCTEXT("AnimCurveWeightTooltip", "The current weight of the curve."))
+		
 			+ SVerticalBox::Slot()
-			.AutoHeight()
+			.FillHeight(1.0f)
 			.Padding( 0.0f, 1.0f )
 			.VAlign( VAlign_Center )
 			[
 				SNew( SSpinBox<float> )
-				.MinSliderValue(-1.f)
-				.MaxSliderValue(1.f)
-				.MinValue(-MaxMorphWeight)
-				.MaxValue(MaxMorphWeight)
-				.Value( this, &SAnimCurveListRow::GetWeight )
+				.Font(FAppStyle::Get().GetFontStyle("SmallFont"))
+				.Value(this, &SAnimCurveListRow::GetWeight)
+				.MinSliderValue(this, &SAnimCurveListRow::GetMinWeight)
+				.MaxSliderValue(this, &SAnimCurveListRow::GetMaxWeight)
 				.OnValueChanged( this, &SAnimCurveListRow::OnAnimCurveWeightChanged )
 				.OnValueCommitted( this, &SAnimCurveListRow::OnAnimCurveWeightValueCommitted )
 			];
 	}
-	else if ( ColumnName == ColumnID_AnimCurveEditLabel)
+	else if ( ColumnName == CurveViewerColumns::AnimCurveEditLabel)
 	{
 		return
 			SNew(SVerticalBox)
-
+			.ToolTipText(LOCTEXT("OverrideTooltip", "Whether the value of this curve is being overriden, or populated from the debugged data."))
+		
 			+ SVerticalBox::Slot()
-			.AutoHeight()
+			.FillHeight(1.0f)
 			.Padding(0.0f, 1.0f)
 			.VAlign(VAlign_Center)
 			.HAlign(HAlign_Center)
 			[
 				SNew(SCheckBox)
-				.OnCheckStateChanged(this, &SAnimCurveListRow::OnAnimCurveAutoFillChecked)
-				.IsChecked(this, &SAnimCurveListRow::IsAnimCurveAutoFillChangedChecked)
-			];
-	}
-	else
-	{
-		return
-			SNew(SVerticalBox)
-
-			+ SVerticalBox::Slot()
-			.AutoHeight()
-			.Padding(0.0f, 1.0f)
-			.VAlign(VAlign_Center)
-			.HAlign(HAlign_Center)
-			[
-				SNew(STextBlock)
-				.Text(this, &SAnimCurveListRow::GetNumConntectedBones)
+				.OnCheckStateChanged(this, &SAnimCurveListRow::OnAnimCurveOverrideChecked)
+				.IsChecked(this, &SAnimCurveListRow::IsAnimCurveOverrideChecked)
 			];
 	}
-}
-
-FText SAnimCurveListRow::GetNumConntectedBones() const
-{
-	const FCurveMetaData* CurveMetaData = Item->EditableSkeleton.Pin()->GetSkeleton().GetCurveMetaData(Item->SmartName);
-	if (CurveMetaData)
-	{
-		return FText::AsNumber(CurveMetaData->LinkedBones.Num());
-	}
-
-	return FText::AsNumber(0);
+	
+	return SNullWidget::NullWidget;
 }
 
 TSharedRef< SWidget > SAnimCurveListRow::GetCurveTypeWidget()
@@ -185,9 +202,7 @@ TSharedRef< SWidget > SAnimCurveListRow::GetCurveTypeWidget()
 		.HAlign(HAlign_Center)
 		[
 			SNew(SCheckBox)
-			.OnCheckStateChanged(this, &SAnimCurveListRow::OnAnimCurveTypeBoxChecked, true)
 			.IsChecked(this, &SAnimCurveListRow::IsAnimCurveTypeBoxChangedChecked, true)
-			.IsEnabled(false)
 			.CheckedImage(FAppStyle::GetBrush("AnimCurveViewer.MorphTargetOn"))
 			.CheckedPressedImage(FAppStyle::GetBrush("AnimCurveViewer.MorphTargetOn"))
 			.UncheckedImage(FAppStyle::GetBrush("AnimCurveViewer.MorphTargetOff"))
@@ -204,7 +219,6 @@ TSharedRef< SWidget > SAnimCurveListRow::GetCurveTypeWidget()
 		.HAlign(HAlign_Center)
 		[
 			SNew(SCheckBox)
-			.OnCheckStateChanged(this, &SAnimCurveListRow::OnAnimCurveTypeBoxChecked, false)
 			.IsChecked(this, &SAnimCurveListRow::IsAnimCurveTypeBoxChangedChecked, false)
 			.CheckedImage(FAppStyle::GetBrush("AnimCurveViewer.MaterialOn"))
 			.CheckedPressedImage(FAppStyle::GetBrush("AnimCurveViewer.MaterialOn"))
@@ -217,71 +231,81 @@ TSharedRef< SWidget > SAnimCurveListRow::GetCurveTypeWidget()
 
 }
 
-void SAnimCurveListRow::OnAnimCurveTypeBoxChecked(ECheckBoxState InState, bool bMorphTarget)
+bool SAnimCurveListRow::GetActiveFlag(bool bMorphTarget) const
 {
-	// currently only material curve is set
-	bool bNewData = (InState == ECheckBoxState::Checked);
-	if (!bMorphTarget)
+	// If anim viewer
+	TSharedPtr<SAnimCurveViewer> AnimCurveViewer = AnimCurveViewerPtr.Pin();
+	if (AnimCurveViewer.IsValid())
 	{
-		Item->EditableSkeleton.Pin()->SetCurveMetaDataMaterial(Item->SmartName, bNewData);
-
-		UAnimInstance* AnimInstance = PreviewScenePtr.Pin()->GetPreviewMeshComponent()->GetAnimInstance();
-		if (AnimInstance)
-		{		
-			AnimInstance->RecalcRequiredCurves(FCurveEvaluationOption());
+		if(const UAnimInstance* AnimInstance = AnimCurveViewer->GetAnimInstance())
+		{
+			// Find if we want to use a pose watch
+			if(UPoseWatchPoseElement* PoseWatchPoseElement = AnimCurveViewer->PoseWatch.Get())
+			{
+				if(UAnimBlueprintGeneratedClass* AnimClass = Cast<UAnimBlueprintGeneratedClass>(AnimInstance->GetClass()))
+				{
+					const FAnimBlueprintDebugData& DebugData = AnimClass->GetAnimBlueprintDebugData();
+					for(const FAnimNodePoseWatch& AnimNodePoseWatch : DebugData.AnimNodePoseWatch)
+					{
+						if(AnimNodePoseWatch.PoseWatchPoseElement == PoseWatchPoseElement)
+						{
+							UE::Anim::ECurveElementFlags Flags = AnimNodePoseWatch.GetCurves().GetFlags(Item->CurveName);
+							return EnumHasAnyFlags(Flags, bMorphTarget ? UE::Anim::ECurveElementFlags::MorphTarget : UE::Anim::ECurveElementFlags::Material);
+						}
+					}
+				}
+			}
+			else
+			{
+				// See if curve is in active set, attribute curve should have everything
+				const TMap<FName, float>& CurveList = AnimInstance->GetAnimationCurveList(bMorphTarget ? EAnimCurveType::MorphTargetCurve : EAnimCurveType::MaterialCurve);
+				const float* CurrentValue = CurveList.Find(Item->CurveName);
+				if (CurrentValue)
+				{
+					return true;
+				}
+			}
 		}
 	}
+
+	return false;
 }
 
 ECheckBoxState SAnimCurveListRow::IsAnimCurveTypeBoxChangedChecked(bool bMorphTarget) const
 {
-	const FCurveMetaData* CurveMetaData = Item->EditableSkeleton.Pin()->GetSkeleton().GetCurveMetaData(Item->SmartName);
-
-	bool bData = false;
-	if (CurveMetaData)
-	{
-		if (bMorphTarget)
-		{
-			bData = CurveMetaData->Type.bMorphtarget != 0;
-		}
-		else
-		{
-			bData = CurveMetaData->Type.bMaterial != 0;
-		}
-	}
-
-	return (bData)? ECheckBoxState::Checked : ECheckBoxState::Unchecked;
+	bool bHasCurveType = GetActiveFlag(bMorphTarget);
+	
+	return (bHasCurveType) ? ECheckBoxState::Checked : ECheckBoxState::Unchecked;
 }
 
-void SAnimCurveListRow::OnAnimCurveAutoFillChecked(ECheckBoxState InState)
+void SAnimCurveListRow::OnAnimCurveOverrideChecked(ECheckBoxState InState)
 {
-	Item->bAutoFillData = InState == ECheckBoxState::Checked;
+	Item->bOverrideData = InState == ECheckBoxState::Checked;
 
 	TSharedPtr<SAnimCurveViewer> AnimCurveViewer = AnimCurveViewerPtr.Pin();
 	if (AnimCurveViewer.IsValid())
 	{
-		if (Item->bAutoFillData)
+		if (Item->bOverrideData)
 		{
-			// clear value so that it can be filled up
-			AnimCurveViewer->RemoveAnimCurveOverride(Item->SmartName.DisplayName);
+			AnimCurveViewer->AddAnimCurveOverride(Item->CurveName, Item->Weight);
 		}
 		else
 		{
-			AnimCurveViewer->AddAnimCurveOverride(Item->SmartName.DisplayName, Item->Weight);
+			// clear value so that it can be filled up
+			AnimCurveViewer->RemoveAnimCurveOverride(Item->CurveName);
 		}
 	}
 }
 
-ECheckBoxState SAnimCurveListRow::IsAnimCurveAutoFillChangedChecked() const
+ECheckBoxState SAnimCurveListRow::IsAnimCurveOverrideChecked() const
 {
-	return (Item->bAutoFillData)? ECheckBoxState::Checked: ECheckBoxState::Unchecked;
+	return (Item->bOverrideData) ? ECheckBoxState::Checked : ECheckBoxState::Unchecked;
 }
 
 void SAnimCurveListRow::OnAnimCurveWeightChanged( float NewWeight )
 {
-	float NewValidWeight = FMath::Clamp(NewWeight, -MaxMorphWeight, MaxMorphWeight);
-	Item->Weight = NewValidWeight;
-	Item->bAutoFillData = false;
+	Item->Weight = NewWeight;
+	Item->bOverrideData = true;
 
 	TSharedPtr<SAnimCurveViewer> AnimCurveViewer = AnimCurveViewerPtr.Pin();
 	if (AnimCurveViewer.IsValid())
@@ -294,7 +318,7 @@ void SAnimCurveListRow::OnAnimCurveWeightChanged( float NewWeight )
 		}
 
 		// Add override
-		AnimCurveViewer->AddAnimCurveOverride(Item->SmartName.DisplayName, Item->Weight);
+		AnimCurveViewer->AddAnimCurveOverride(Item->CurveName, Item->Weight);
 
 		// ...then any selected rows need changing by the same delta
 		TArray< TSharedPtr< FDisplayedAnimCurveInfo > > SelectedRows = AnimCurveViewer->AnimCurveListView->GetSelectedItems();
@@ -304,9 +328,9 @@ void SAnimCurveListRow::OnAnimCurveWeightChanged( float NewWeight )
 
 			if (RowItem != Item) // Don't do "this" row again if it's selected
 			{
-				RowItem->Weight = NewValidWeight;
-				RowItem->bAutoFillData = false;
-				AnimCurveViewer->AddAnimCurveOverride(RowItem->SmartName.DisplayName, RowItem->Weight);
+				RowItem->Weight = NewWeight;
+				RowItem->bOverrideData = true;
+				AnimCurveViewer->AddAnimCurveOverride(RowItem->CurveName, RowItem->Weight);
 			}
 		}
 
@@ -327,17 +351,7 @@ void SAnimCurveListRow::OnAnimCurveWeightValueCommitted( float NewWeight, ETextC
 
 FText SAnimCurveListRow::GetItemName() const
 {
-	FName ItemName;
-
-	TSharedPtr<SAnimCurveViewer> AnimCurveViewer = AnimCurveViewerPtr.Pin();
-	if (AnimCurveViewer.IsValid())
-	{
-		const FSmartNameMapping* Mapping = Item->EditableSkeleton.Pin()->GetSkeleton().GetSmartNameContainer(AnimCurveViewer->ContainerName);
-		check(Mapping);
-		Mapping->GetName(Item->SmartName.UID, ItemName);
-	}
-
-	return FText::FromName(ItemName);
+	return FText::FromName(Item->CurveName);
 }
 
 FText SAnimCurveListRow::GetFilterText() const
@@ -362,18 +376,40 @@ bool SAnimCurveListRow::GetActiveWeight(float& OutWeight) const
 	TSharedPtr<SAnimCurveViewer> AnimCurveViewer = AnimCurveViewerPtr.Pin();
 	if (AnimCurveViewer.IsValid())
 	{
-		// If anim instance
-		UAnimInstance* AnimInstance = PreviewScenePtr.Pin()->GetPreviewMeshComponent()->GetAnimInstance();
-		if (AnimInstance)
+		if(const UAnimInstance* AnimInstance = AnimCurveViewer->GetAnimInstance())
 		{
-			// See if curve is in active set, attribute curve should have everything
-			const TMap<FName, float>& CurveList = AnimInstance->GetAnimationCurveList(EAnimCurveType::AttributeCurve);
-			const float* CurrentValue = CurveList.Find(Item->SmartName.DisplayName);
-			if (CurrentValue)
+			// Find if we want to use a pose watch
+			if(UPoseWatchPoseElement* PoseWatchPoseElement = AnimCurveViewer->PoseWatch.Get())
 			{
-				OutWeight = *CurrentValue;
-				// Remember we found it
-				bFoundActive = true;
+				if(UAnimBlueprintGeneratedClass* AnimClass = Cast<UAnimBlueprintGeneratedClass>(AnimInstance->GetClass()))
+				{
+					const FAnimBlueprintDebugData& DebugData = AnimClass->GetAnimBlueprintDebugData();
+					for(const FAnimNodePoseWatch& AnimNodePoseWatch : DebugData.AnimNodePoseWatch)
+					{
+						if(AnimNodePoseWatch.PoseWatchPoseElement == PoseWatchPoseElement)
+						{
+							bool bHasElement = false;
+							float CurrentValue = AnimNodePoseWatch.GetCurves().Get(Item->CurveName, bHasElement);
+							if(bHasElement)
+							{
+								OutWeight = CurrentValue;
+								bFoundActive = true;
+							}
+							break;
+						}
+					}
+				}
+			}
+			else
+			{
+				// See if curve is in active set, attribute curve should have everything
+				const TMap<FName, float>& CurveList = AnimInstance->GetAnimationCurveList(EAnimCurveType::AttributeCurve);
+				const float* CurrentValue = CurveList.Find(Item->CurveName);
+				if (CurrentValue)
+				{
+					OutWeight = *CurrentValue;
+					bFoundActive = true;
+				}
 			}
 		}
 	}
@@ -384,147 +420,42 @@ bool SAnimCurveListRow::GetActiveWeight(float& OutWeight) const
 
 FSlateFontInfo SAnimCurveListRow::GetItemFont() const
 {
-	// If not selected, show bright if active
-	bool bItemActive = true;
-	if (Item->bAutoFillData)
-	{
-		float Weight = 0.f;
-		GetActiveWeight(Weight);
-		// change so that print white if it has weight on it. 
-		bItemActive = (Weight != 0.f);
-	}
+	// Show bright if active
+	float Weight = 0.f;
+	const bool bItemActive = GetActiveWeight(Weight);
 
-	return bItemActive ? FAppStyle::Get().GetFontStyle("NormalFontBoldItalic") : FAppStyle::Get().GetFontStyle("NormalFont");
+	return bItemActive ? FAppStyle::Get().GetFontStyle("AnimCurveViewer.ActiveCurveFont") : FAppStyle::Get().GetFontStyle("SmallFont");
 }
 
-float SAnimCurveListRow::GetWeight() const 
-{ 
-	float Weight = Item->Weight;
-
-	if (Item->bAutoFillData)
+float SAnimCurveListRow::GetWeight() const
+{
+	float Weight = 0.0f;
+	bool bItemActive = false;
+	if (!Item->bOverrideData)
 	{
-		GetActiveWeight(Weight);
+		bItemActive = GetActiveWeight(Weight);
+		MinWeight = FMath::Min(MinWeight, Weight);
+		MaxWeight = FMath::Max(MaxWeight, Weight);
 	}
 
-	return Weight;
-}
-
-//////////////////////////////////////////////////////////////////////////
-// SAnimCurveTypeList
-
-class SAnimCurveTypeList : public SCompoundWidget
-{
-public:
-
-	SLATE_BEGIN_ARGS(SAnimCurveTypeList) {}
-
-		/** The item for this row **/
-		SLATE_ARGUMENT(int32, CurveFlags)
-
-		/* The SAnimCurveViewer that we push the morph target weights into */
-		SLATE_ARGUMENT(TWeakPtr<SAnimCurveViewer>, AnimCurveViewerPtr)
-
-	SLATE_END_ARGS()
-
-	void Construct(const FArguments& InArgs);
-
-	/** Auto fill check call back functions */
-	void OnAnimCurveTypeShowChecked(ECheckBoxState InState);
-	ECheckBoxState IsAnimCurveTypeShowChangedChecked() const;
-
-	FText GetAnimCurveType() const;
-private:
-
-	/* The SAnimCurveViewer that we push the morph target weights into */
-	TWeakPtr<SAnimCurveViewer> AnimCurveViewerPtr;
-
-	/** The name and weight of the morph target */
-	int32 CurveFlags;
-};
-
-void SAnimCurveTypeList::Construct(const FArguments& InArgs)
-{
-	CurveFlags = InArgs._CurveFlags;
-	AnimCurveViewerPtr = InArgs._AnimCurveViewerPtr;
-
-	TSharedPtr<SAnimCurveViewer> AnimCurveViewer = AnimCurveViewerPtr.Pin();
-	if (AnimCurveViewer.IsValid())
+	if(!bItemActive)
 	{
-		ChildSlot
-		[
-			SNew(SHorizontalBox)
-
-			+ SHorizontalBox::Slot()
-			.AutoWidth()
-			.Padding(0.0f, 1.0f)
-			.VAlign(VAlign_Center)
-			.HAlign(HAlign_Center)
-			[
-				SNew(SCheckBox)
-				.OnCheckStateChanged(this, &SAnimCurveTypeList::OnAnimCurveTypeShowChecked)
-				.IsChecked(this, &SAnimCurveTypeList::IsAnimCurveTypeShowChangedChecked)
-			]
-
-			+ SHorizontalBox::Slot()
-			.AutoWidth()
-			.Padding(3.0f, 1.0f)
-			.VAlign(VAlign_Center)
-			[
-				SNew(STextBlock)
-				.Text(this, &SAnimCurveTypeList::GetAnimCurveType)
-				.HighlightText(AnimCurveViewer->GetFilterText())
-			]
-		];
+		Weight = Item->Weight;
 	}
 
+	return Weight;
 }
 
-void SAnimCurveTypeList::OnAnimCurveTypeShowChecked(ECheckBoxState InState)
+TOptional<float> SAnimCurveListRow::GetMinWeight() const 
 {
-	TSharedPtr<SAnimCurveViewer> AnimCurveViewer = AnimCurveViewerPtr.Pin();
-	if (AnimCurveViewer.IsValid())
-	{
-		// clear value so that it can be filled up
-		if (InState == ECheckBoxState::Checked)
-		{
-			AnimCurveViewer->CurrentCurveFlag |= CurveFlags;
-		}
-		else
-		{
-			AnimCurveViewer->CurrentCurveFlag &= ~CurveFlags;
-		}
-
-		AnimCurveViewer->RefreshCurveList(true);
-	}
+	return MinWeight;
 }
 
-ECheckBoxState SAnimCurveTypeList::IsAnimCurveTypeShowChangedChecked() const
+TOptional<float> SAnimCurveListRow::GetMaxWeight() const 
 {
-	TSharedPtr<SAnimCurveViewer> AnimCurveViewer = AnimCurveViewerPtr.Pin();
-	if (AnimCurveViewer.IsValid())
-	{
-		return (AnimCurveViewer->CurrentCurveFlag & CurveFlags) ? ECheckBoxState::Checked : ECheckBoxState::Unchecked;
-	}
-	else
-	{
-		return ECheckBoxState::Unchecked;
-	}
+	return MaxWeight;
 }
 
-FText SAnimCurveTypeList::GetAnimCurveType() const
-{
-	switch (CurveFlags)
-	{
-	case ACEF_DriveMorphTarget:
-		return LOCTEXT("AnimCurveType_Morphtarget", "Morph Target");
-	case ACEF_DriveAttribute:
-		return LOCTEXT("AnimCurveType_Attribute", "Attribute");
-	case ACEF_DriveMaterial:
-		return LOCTEXT("AnimCurveType_Material", "Material");
-	}
-
-	return LOCTEXT("AnimCurveType_Unknown", "Unknown");
-}
 //////////////////////////////////////////////////////////////////////////
 // SAnimCurveViewer
 
@@ -532,97 +463,165 @@ void SAnimCurveViewer::Construct(const FArguments& InArgs, const TSharedRef<clas
 {
 	OnObjectsSelected = InOnObjectsSelected;
 
-	bShowAllCurves = true;
-
-	EditorObjectTracker.SetAllowOnePerClass(false);
-
-	ContainerName = USkeleton::AnimCurveMappingName;
-
 	PreviewScenePtr = InPreviewScene;
 	EditableSkeletonPtr = InEditableSkeleton;
 
 	InPreviewScene->RegisterOnPreviewMeshChanged(FOnPreviewMeshChanged::CreateSP(this, &SAnimCurveViewer::OnPreviewMeshChanged));
 	InPreviewScene->RegisterOnAnimChanged(FOnAnimChanged::CreateSP(this, &SAnimCurveViewer::OnPreviewAssetChanged));
 
-	SmartNameChangedHandle = InEditableSkeleton->RegisterOnSmartNameChanged(FOnSmartNameChanged::FDelegate::CreateSP(this, &SAnimCurveViewer::HandleSmartNamesChange));
-
 	// Register and bind all our menu commands
 	FCurveViewerCommands::Register();
-	BindCommands();
 
-	// @todo fix this to be filtered
-	CurrentCurveFlag = ACEF_DriveMorphTarget | ACEF_DriveMaterial | ACEF_DriveAttribute;
-	 
-	TSharedPtr<SHorizontalBox> AnimTypeBoxContainer = SNew(SHorizontalBox);
+	CurrentCurveFlag = EAnimCurveViewerFilterFlags::Active;
+
+	TSharedPtr<FFilterCategory> FilterCategory = MakeShared<FFilterCategory>(LOCTEXT("CurveFiltersLabel", "Curve Filters"), LOCTEXT("CurveFiltersToolTip", "Filter what kind fo curves can be displayed."));
+
+	Filters.Add(MakeShared<FAnimCurveViewerFilter>(
+		EAnimCurveViewerFilterFlags::Active,
+		"Active",
+		LOCTEXT("ShowActiveLabel", "Active"),
+		LOCTEXT("ShowActiveTooltip", "Show only active curves"),
+		FLinearColor::Yellow,
+		FilterCategory
+		));
+
+	Filters.Add(MakeShared<FAnimCurveViewerFilter>(
+		EAnimCurveViewerFilterFlags::MorphTarget,
+		"MorphTarget",
+		LOCTEXT("MorphTargetLabel", "Morph Target"),
+		LOCTEXT("MorphTargetTooltip", "Show morph target curves"),
+		FLinearColor::Red,
+		FilterCategory
+		));
+
+	Filters.Add(MakeShared<FAnimCurveViewerFilter>(
+		EAnimCurveViewerFilterFlags::Material,
+		"Material",
+		LOCTEXT("MaterialLabel", "Material"),
+		LOCTEXT("MaterialTooltip", "Show material curves"),
+		FLinearColor::Green,
+		FilterCategory
+		));
+	
+	TSharedRef<SBasicFilterBar<EAnimCurveViewerFilterFlags>> FilterBar = SNew(SBasicFilterBar<EAnimCurveViewerFilterFlags>)
+	.CustomFilters(Filters)
+	.UseSectionsForCategories(true)
+	.OnFilterChanged_Lambda([this]()
+	{
+		CurrentCurveFlag = EAnimCurveViewerFilterFlags::ShowAll;
+
+		for(const TSharedRef<FFilterBase<EAnimCurveViewerFilterFlags>>& Filter : Filters)
+		{
+			TSharedRef<FAnimCurveViewerFilter> AnimCurveFilter = StaticCastSharedRef<FAnimCurveViewerFilter>(Filter);
+			if(AnimCurveFilter->IsActive())
+			{
+				CurrentCurveFlag |= AnimCurveFilter->GetFlags();
+			}
+		}
+
+		RefreshCurveList(false);
+	});
 
+	FilterBar->SetFilterCheckState(Filters[0], ECheckBoxState::Checked);
+	
 	ChildSlot
 	[
 		SNew( SVerticalBox )
-		
-		+ SVerticalBox::Slot()
+		+SVerticalBox::Slot()
+		.Padding(2.0f,2.0f)
 		.AutoHeight()
-		.Padding(0,2)
+		.HAlign(HAlign_Left)
+		[
+			SNew(SSimpleButton)
+			.Text(LOCTEXT("FindReplaceCurvesButton", "Find/Replace Curves..."))
+			.ToolTipText(LOCTEXT("FindReplaceCurvesButtonTooltip", "Find and replace curves across multiple assets"))
+			.Icon(FAppStyle::GetBrush("Kismet.Tabs.FindResults"))
+			.OnClicked_Lambda([this]()
+			{
+				if(TSharedPtr<SDockTab> ActiveTab = FGlobalTabmanager::Get()->GetActiveTab())
+				{
+					if(TSharedPtr<FTabManager> TabManager = ActiveTab->GetTabManagerPtr())
+					{
+						if(TSharedPtr<SDockTab> Tab = TabManager->TryInvokeTab(FPersonaTabs::FindReplaceID))
+						{
+							TSharedRef<SAnimAssetFindReplace> FindReplaceWidget = StaticCastSharedRef<SAnimAssetFindReplace>(Tab->GetContent());
+							FindReplaceWidget->SetFindReplaceType(EAnimAssetFindReplaceType::Curves);
+						}
+					}
+				}
+				return FReply::Handled();
+			})
+		]
+		+SVerticalBox::Slot()
+		.AutoHeight()
+		.Padding(0.0f,2.0f)
 		[
 			SNew(SHorizontalBox)
-			// Filter entry
 			+SHorizontalBox::Slot()
-			.FillWidth( 1 )
+			.AutoWidth()
+			.Padding(2.0f,0.0f)
+			[
+				SBasicFilterBar<EAnimCurveViewerFilterFlags>::MakeAddFilterButton(FilterBar)
+			]
+			+SHorizontalBox::Slot()
+			.FillWidth(1.0f)
+			.Padding(2.0f,0.0f)
 			[
 				SAssignNew( NameFilterBox, SSearchBox )
 				.SelectAllTextWhenFocused( true )
 				.OnTextChanged( this, &SAnimCurveViewer::OnFilterTextChanged )
 				.OnTextCommitted( this, &SAnimCurveViewer::OnFilterTextCommitted )
 			]
-		]
-
-		+ SVerticalBox::Slot()
-		.AutoHeight()		// This is required to make the scrollbar work, as content overflows Slate containers by default
-		[
-			SNew(SBox)
-			.WidthOverride(150.f)
-			.Content()
+			+SHorizontalBox::Slot()
+			.AutoWidth()
+			.Padding(2.0f,0.0f)
 			[
-				AnimTypeBoxContainer.ToSharedRef()
+				SNew(SBox)
+				.MaxDesiredWidth(200.0f)
+				[
+					CreateCurveSourceSelector()
+				]
 			]
 		]
-
-		+ SVerticalBox::Slot()
-		.FillHeight( 1.0f )		// This is required to make the scrollbar work, as content overflows Slate containers by default
+		+SVerticalBox::Slot()
+		.AutoHeight()
 		[
-			SAssignNew( AnimCurveListView, SAnimCurveListType )
+			FilterBar
+		]
+		+SVerticalBox::Slot()
+		.FillHeight(1.0f)
+		[
+			SAssignNew(AnimCurveListView, SAnimCurveListType)
 			.ListItemsSource( &AnimCurveList )
 			.OnGenerateRow( this, &SAnimCurveViewer::GenerateAnimCurveRow )
-			.OnContextMenuOpening( this, &SAnimCurveViewer::OnGetContextMenuContent )
-			.ItemHeight( 22.0f )
+			.ItemHeight( 18.0f )
 			.SelectionMode(ESelectionMode::Multi)
-			.OnSelectionChanged( this, &SAnimCurveViewer::OnSelectionChanged )
 			.HeaderRow
 			(
 				SNew( SHeaderRow )
-				+ SHeaderRow::Column( ColumnId_AnimCurveNameLabel )
+				+ SHeaderRow::Column(CurveViewerColumns::AnimCurveNameLabel)
 				.FillWidth(1.f)
 				.DefaultLabel( LOCTEXT( "AnimCurveNameLabel", "Curve Name" ) )
-
-				+ SHeaderRow::Column(ColumnID_AnimCurveTypeLabel)
-				.FillWidth(0.5f)
+				.DefaultTooltip(LOCTEXT("AnimCurveNameTooltip", "The name of the curve."))
+				
+				+ SHeaderRow::Column(CurveViewerColumns::AnimCurveTypeLabel)
+				.FixedWidth(48.0f)
 				.DefaultLabel(LOCTEXT("AnimCurveTypeLabel", "Type"))
+				.DefaultTooltip(LOCTEXT("AnimCurveTypeTooltip", "The type of the curve (e.g. morph target, material)."))
 
-				+ SHeaderRow::Column( ColumnID_AnimCurveWeightLabel )
+				+ SHeaderRow::Column(CurveViewerColumns::AnimCurveWeightLabel )
 				.FillWidth(1.f)
 				.DefaultLabel( LOCTEXT( "AnimCurveWeightLabel", "Weight" ) )
+				.DefaultTooltip(LOCTEXT("AnimCurveWeightTooltip", "The current weight of the curve."))
 
-				+ SHeaderRow::Column(ColumnID_AnimCurveEditLabel)
-				.FillWidth(0.25f)
-				.DefaultLabel(LOCTEXT("AnimCurveEditLabel", "Auto"))
-
-				+ SHeaderRow::Column(ColumnID_AnimCurveNumBoneLabel)
-				.FillWidth(0.5f)
-				.DefaultLabel(LOCTEXT("AnimCurveNumBoneLabel", "Bones"))
+				+ SHeaderRow::Column(CurveViewerColumns::AnimCurveEditLabel)
+				.FixedWidth(24.0f)
+				.DefaultLabel(FText::GetEmpty())
+				.DefaultTooltip(LOCTEXT("OverrideTooltip", "Whether the value of this curve is being overriden, or populated from the debugged data."))
 			)
 		]
 	];
 
-	CreateAnimCurveTypeList(AnimTypeBoxContainer.ToSharedRef());
 	RefreshCurveList(true);
 }
 
@@ -634,75 +633,7 @@ SAnimCurveViewer::~SAnimCurveViewer()
 		PreviewScenePtr.Pin()->UnregisterOnAnimChanged(this);
 	}
 
-	if (EditableSkeletonPtr.IsValid())
-	{
-		EditableSkeletonPtr.Pin()->UnregisterOnSmartNameChanged(SmartNameChangedHandle);
-	}
-}
-
-bool SAnimCurveViewer::IsCurveFilterEnabled() const
-{
-	return !bShowAllCurves;
-}
-
-ECheckBoxState SAnimCurveViewer::IsShowingAllCurves() const
-{
-	return bShowAllCurves ? ECheckBoxState::Checked : ECheckBoxState::Unchecked;
-}
-
-void SAnimCurveViewer::OnToggleShowingAllCurves(ECheckBoxState NewState)
-{
-	bShowAllCurves = (NewState == ECheckBoxState::Checked);
-
-	RefreshCurveList(false);
-}
-FReply SAnimCurveViewer::OnKeyDown(const FGeometry& MyGeometry, const FKeyEvent& InKeyEvent)
-{
-	if (UICommandList.IsValid() && UICommandList->ProcessCommandBindings(InKeyEvent))
-	{
-		return FReply::Handled();
-	}
-	return FReply::Unhandled();
-}
-
-void SAnimCurveViewer::BindCommands()
-{
-	// This should not be called twice on the same instance
-	check(!UICommandList.IsValid());
-
-	UICommandList = MakeShareable(new FUICommandList);
-
-	FUICommandList& CommandList = *UICommandList;
-
-	// Grab the list of menu commands to bind...
-	const FCurveViewerCommands& MenuActions = FCurveViewerCommands::Get();
-
-	// ...and bind them all
-
-	CommandList.MapAction(
-		FGenericCommands::Get().Rename,
-		FExecuteAction::CreateSP(this, &SAnimCurveViewer::OnRenameClicked),
-		FCanExecuteAction::CreateSP(this, &SAnimCurveViewer::CanRename));
-
-	CommandList.MapAction(
-		FGenericCommands::Get().Delete,
-		FExecuteAction::CreateSP(this, &SAnimCurveViewer::OnDeleteNameClicked),
-		FCanExecuteAction::CreateSP(this, &SAnimCurveViewer::CanDelete));
-
-	CommandList.MapAction(
-		FGenericCommands::Get().Copy,
-		FExecuteAction::CreateSP(this, &SAnimCurveViewer::OnCopyClicked),
-		FCanExecuteAction::CreateSP(this, &SAnimCurveViewer::CanCopy));
-
-	CommandList.MapAction(
-		FGenericCommands::Get().Paste,
-		FExecuteAction::CreateSP(this, &SAnimCurveViewer::OnPasteClicked),
-		FCanExecuteAction::CreateSP(this, &SAnimCurveViewer::CanPaste));
-	
-	CommandList.MapAction(
-		MenuActions.AddCurve,
-		FExecuteAction::CreateSP(this, &SAnimCurveViewer::OnAddClicked),
-		FCanExecuteAction());
+	AnimationEditorUtils::OnPoseWatchesChanged().RemoveAll(this);
 }
 
 void SAnimCurveViewer::OnPreviewMeshChanged(class USkeletalMesh* OldPreviewMesh, class USkeletalMesh* NewPreviewMesh)
@@ -717,11 +648,6 @@ void SAnimCurveViewer::OnFilterTextChanged( const FText& SearchText )
 	RefreshCurveList(false);
 }
 
-void SAnimCurveViewer::OnCurvesChanged()
-{
-	RefreshCurveList(true);
-}
-
 void SAnimCurveViewer::OnFilterTextCommitted( const FText& SearchText, ETextCommit::Type CommitInfo )
 {
 	// Just do the same as if the user typed in the box
@@ -738,223 +664,196 @@ TSharedRef<ITableRow> SAnimCurveViewer::GenerateAnimCurveRow(TSharedPtr<FDisplay
 		.AnimCurveViewerPtr( SharedThis(this) );
 }
 
-TSharedPtr<SWidget> SAnimCurveViewer::OnGetContextMenuContent() const
-{
-	const bool bShouldCloseWindowAfterMenuSelection = true;
-	FMenuBuilder MenuBuilder( bShouldCloseWindowAfterMenuSelection, UICommandList);
-
-	const FCurveViewerCommands& Actions = FCurveViewerCommands::Get();
-
-	MenuBuilder.BeginSection("AnimCurveAction", LOCTEXT( "CurveAction", "Curve Actions" ) );
-
-	MenuBuilder.AddMenuEntry(FGenericCommands::Get().Rename, NAME_None, LOCTEXT("RenameSmartNameLabel", "Rename Curve"), LOCTEXT("RenameSmartNameToolTip", "Rename the selected curve"));
-	MenuBuilder.AddMenuEntry(FGenericCommands::Get().Delete, NAME_None, LOCTEXT("DeleteSmartNameLabel", "Delete Curve"), LOCTEXT("DeleteSmartNameToolTip", "Delete the selected curve"));
-	MenuBuilder.AddMenuEntry(FGenericCommands::Get().Copy);
-	MenuBuilder.AddMenuEntry(FGenericCommands::Get().Paste);
-	MenuBuilder.AddMenuEntry(Actions.AddCurve);
-
-	MenuBuilder.EndSection();
-
-	return MenuBuilder.MakeWidget();
-}
-
-void SAnimCurveViewer::OnRenameClicked()
-{
-	TArray< TSharedPtr<FDisplayedAnimCurveInfo> > SelectedItems = AnimCurveListView->GetSelectedItems();
-
-	SelectedItems[0]->EditableText->EnterEditingMode();
-}
-
-bool SAnimCurveViewer::CanRename()
-{
-	return AnimCurveListView->GetNumItemsSelected() == 1;
-}
-
-void SAnimCurveViewer::OnAddClicked()
-{
-	TSharedRef<STextEntryPopup> TextEntry =
-		SNew(STextEntryPopup)
-		.Label(LOCTEXT("NewSmartnameLabel", "New Name"))
-		.OnTextCommitted(this, &SAnimCurveViewer::CreateNewNameEntry);
-
-	FSlateApplication& SlateApp = FSlateApplication::Get();
-	SlateApp.PushMenu(
-		AsShared(),
-		FWidgetPath(),
-		TextEntry,
-		SlateApp.GetCursorPos(),
-		FPopupTransitionEffect::TypeInPopup
-		);
-}
-
-
-
 void SAnimCurveViewer::Tick(const FGeometry& AllottedGeometry, const double InCurrentTime, const float InDeltaTime)
 {
-	UDebugSkelMeshComponent* MeshComponent = PreviewScenePtr.Pin()->GetPreviewMeshComponent();
-	UAnimInstance* AnimInstance = MeshComponent->GetAnimInstance();
-
-	// Only need to refresh the list on tick() if we are filtering to active curves only
-	if (AnimInstance && !bShowAllCurves)
+	if (GetAnimInstance())
 	{
+		// We refresh when ticking each time as curve flags can potentially vary
 		RefreshCurveList(false);
 	}
 }
 
-void SAnimCurveViewer::CreateNewNameEntry(const FText& CommittedText, ETextCommit::Type CommitType)
+UAnimInstance* SAnimCurveViewer::GetAnimInstance() const
 {
-	FSlateApplication::Get().DismissAllMenus();
-	if (!CommittedText.IsEmpty() && CommitType == ETextCommit::OnEnter)
+	UDebugSkelMeshComponent* MeshComponent = PreviewScenePtr.Pin()->GetPreviewMeshComponent();
+	UAnimInstance* AnimInstance = MeshComponent->GetAnimInstance();
+	if (AnimInstance)
 	{
-		if (const FSmartNameMapping* NameMapping = GetAnimCurveMapping())
+		// Look at the debugged anim instance if we are targeting one
+		if(AnimInstance == MeshComponent->PreviewInstance.Get())
 		{
-			FName NewName = FName(*CommittedText.ToString());
-			FSmartName NewCurveName;
-			if (EditableSkeletonPtr.Pin()->AddSmartname(ContainerName, NewName, NewCurveName))
+			UAnimPreviewInstance* AnimPreviewInstance = MeshComponent->PreviewInstance;
+			if(USkeletalMeshComponent* DebuggedComponent = AnimPreviewInstance->GetDebugSkeletalMeshComponent())
 			{
-				// Successfully added
-				RefreshCurveList(true);
+				AnimInstance = DebuggedComponent->GetAnimInstance();
 			}
 		}
 	}
-}
 
-
-const FSmartNameMapping* SAnimCurveViewer::GetAnimCurveMapping()
-{
-	return EditableSkeletonPtr.Pin()->GetSkeleton().GetSmartNameContainer(ContainerName);
+	return AnimInstance;
 }
 
-UAnimInstance* SAnimCurveViewer::GetAnimInstance() const
+void SAnimCurveViewer::CreateAnimCurveList( const FString& SearchText, bool bInFullRefresh )
 {
-	return PreviewScenePtr.Pin()->GetPreviewMeshComponent()->GetAnimInstance();
-}
+	if(!AnimCurveListView.IsValid())
+	{
+		return;
+	}
+	
+	bool bDirty = bInFullRefresh;
 
-int32 FindIndexOfAnimCurveInfo(TArray<TSharedPtr<FDisplayedAnimCurveInfo>>& AnimCurveInfos, const FSmartName& CurveName)
-{
-	for (int32 CurveIdx = 0; CurveIdx < AnimCurveInfos.Num(); CurveIdx++)
+	AnimCurveList.Reset();
+
+	if(bInFullRefresh)
 	{
-		// check UID to make sure they match what it's looking for
-		if (AnimCurveInfos[CurveIdx]->SmartName.UID == CurveName.UID)
-		{
-			return CurveIdx;
-		}
+		AllSeenAnimCurvesMap.Reset();
 	}
 
-	return INDEX_NONE;
-}
+	auto AddCurve = [this](FName InCurveName, UE::Anim::ECurveElementFlags InFlags)
+	{
+		// Only add if the curve doesnt exist
+		TSharedPtr<FDisplayedAnimCurveInfo>* ExistingItem = AllSeenAnimCurvesMap.Find(InCurveName);
+		if(ExistingItem == nullptr)
+		{
+			TSharedRef<FDisplayedAnimCurveInfo> NewInfo = FDisplayedAnimCurveInfo::Make(EditableSkeletonPtr, InCurveName);
 
-void SAnimCurveViewer::CreateAnimCurveList( const FString& SearchText, bool bInFullRefresh )
-{
-	bool bDirty = bInFullRefresh;
+			float Weight = 0.f;
+			const bool bOverride = GetAnimCurveOverride(InCurveName, Weight);
+			NewInfo->bOverrideData = bOverride;
+			NewInfo->Weight = Weight;
+			NewInfo->bMaterial = EnumHasAnyFlags(InFlags, UE::Anim::ECurveElementFlags::Material);
+			NewInfo->bMorphTarget = EnumHasAnyFlags(InFlags, UE::Anim::ECurveElementFlags::MorphTarget);
+
+			AllSeenAnimCurvesMap.Add(InCurveName, NewInfo);
+		}
+		else
+		{
+			(*ExistingItem)->bMaterial = EnumHasAnyFlags(InFlags, UE::Anim::ECurveElementFlags::Material);
+			(*ExistingItem)->bMorphTarget = EnumHasAnyFlags(InFlags, UE::Anim::ECurveElementFlags::MorphTarget);
+		}
+	};
 	
-	const FSmartNameMapping* Mapping = GetAnimCurveMapping();
-	if (Mapping)
+	// Add curve items from skeleton metadata
+	EditableSkeletonPtr.Pin()->GetSkeleton().ForEachCurveMetaData([&AddCurve](FName InCurveName, const FCurveMetaData& InCurveMetaData)
 	{
-		TArray<SmartName::UID_Type> UidList;
-		Mapping->FillUidArray(UidList);
-
-		AnimCurveList.Reset();
-		if(bInFullRefresh)
+		UE::Anim::ECurveElementFlags Flags = UE::Anim::ECurveElementFlags::None;
+		if(InCurveMetaData.Type.bMaterial)
 		{
-			AnimCurvesByUID.Reset();
+			Flags |= UE::Anim::ECurveElementFlags::Material; 
 		}
-		
-		if(AnimCurvesByUID.Num() != UidList.Num())
+		if(InCurveMetaData.Type.bMorphtarget)
 		{
-			// Resize the displayed curves we are tracking by UID
-			AnimCurvesByUID.SetNum(UidList.Num());
-
-			// Fill in any missing items
-			int32 UIDIndex = 0;
-			for (SmartName::UID_Type Uid : UidList)
-			{
-				if(!AnimCurvesByUID[UIDIndex].IsValid())
-				{
-					FSmartName SmartName;
-					Mapping->FindSmartNameByUID(Uid, SmartName);
-					
-					UEditorAnimCurveBoneLinks* EditorMirrorObj = Cast<UEditorAnimCurveBoneLinks> (EditorObjectTracker.GetEditorObjectForClass(UEditorAnimCurveBoneLinks::StaticClass()));
-					EditorMirrorObj->Initialize(EditableSkeletonPtr, SmartName, FOnAnimCurveBonesChange::CreateSP(this, &SAnimCurveViewer::ApplyCurveBoneLinks));
-					TSharedRef<FDisplayedAnimCurveInfo> NewInfo = FDisplayedAnimCurveInfo::Make(EditableSkeletonPtr, ContainerName, SmartName, EditorMirrorObj);
-
-					// See if we have an override set, and if so, grab info
-					float Weight = 0.f;
-					bool bOverride = GetAnimCurveOverride(SmartName.DisplayName, Weight);
-					NewInfo->bAutoFillData = !bOverride;
-					NewInfo->Weight = Weight;
-
-					AnimCurvesByUID[UIDIndex] = NewInfo;
-				}
-				UIDIndex++;
-			}
+			Flags |= UE::Anim::ECurveElementFlags::MorphTarget; 
 		}
-		
-		// Get set of active curves
-		TMap<FName, float> ActiveCurves;
+		AddCurve(InCurveName, Flags);
+	});
 
-		if(!bShowAllCurves)
+	// Add active curves if required
+	TSet<FName> ActiveCurves;
+	if(const UAnimInstance* AnimInstance = GetAnimInstance())
+	{
+		// Find if we want to use a pose watch
+		if(UPoseWatchPoseElement* PoseWatchPoseElement = PoseWatch.Get())
 		{
-			UDebugSkelMeshComponent* MeshComponent = PreviewScenePtr.Pin()->GetPreviewMeshComponent();
-			UAnimInstance* AnimInstance = MeshComponent->GetAnimInstance();
-			if (!bShowAllCurves && AnimInstance)
+			if(UAnimBlueprintGeneratedClass* AnimClass = Cast<UAnimBlueprintGeneratedClass>(AnimInstance->GetClass()))
 			{
-				// attribute curve should contain everything
-				// so only search other container if attribute is off
-				if (CurrentCurveFlag & ACEF_DriveAttribute)
-				{
-					AnimInstance->AppendAnimationCurveList(EAnimCurveType::AttributeCurve, ActiveCurves);
-				}
-				else
+				const FAnimBlueprintDebugData& DebugData = AnimClass->GetAnimBlueprintDebugData();
+				for(const FAnimNodePoseWatch& AnimNodePoseWatch : DebugData.AnimNodePoseWatch)
 				{
-					if (CurrentCurveFlag & ACEF_DriveMorphTarget)
-					{
-						AnimInstance->AppendAnimationCurveList(EAnimCurveType::MorphTargetCurve, ActiveCurves);
-					}
-					if (CurrentCurveFlag & ACEF_DriveMaterial)
+					if(AnimNodePoseWatch.PoseWatchPoseElement == PoseWatchPoseElement)
 					{
-						AnimInstance->AppendAnimationCurveList(EAnimCurveType::MaterialCurve, ActiveCurves);
+						AnimNodePoseWatch.GetCurves().ForEachElement([this, &AddCurve, &ActiveCurves](const UE::Anim::FCurveElement& InElement)
+						{
+							if (EnumHasAnyFlags(CurrentCurveFlag, EAnimCurveViewerFilterFlags::MorphTarget))
+							{
+								if(EnumHasAnyFlags(InElement.Flags, UE::Anim::ECurveElementFlags::MorphTarget))
+								{
+									AddCurve(InElement.Name, InElement.Flags);
+									ActiveCurves.Add(InElement.Name);
+								}
+							}
+							if (EnumHasAnyFlags(CurrentCurveFlag, EAnimCurveViewerFilterFlags::Material))
+							{
+								if(EnumHasAnyFlags(InElement.Flags, UE::Anim::ECurveElementFlags::Material))
+								{
+									AddCurve(InElement.Name, InElement.Flags);
+									ActiveCurves.Add(InElement.Name);
+								}
+							}
+
+							// If we arent filtering by curve type, just show all curves
+							if(!EnumHasAnyFlags(CurrentCurveFlag, EAnimCurveViewerFilterFlags::MorphTarget | EAnimCurveViewerFilterFlags::Material))
+							{
+								AddCurve(InElement.Name, InElement.Flags);
+								ActiveCurves.Add(InElement.Name);
+							}
+						});
+						break;
 					}
 				}
 			}
 		}
-
-		// Iterate through all curves..
-		int32 UIDIndex = 0;
-		for (SmartName::UID_Type Uid : UidList)
+		else
 		{
-			bool bAddToList = true;
-
-			FSmartName SmartName;
-			Mapping->FindSmartNameByUID(Uid, SmartName);
-
-			// See if we pass the search filter
-			if (!FilterText.IsEmpty())
+			if (EnumHasAnyFlags(CurrentCurveFlag, EAnimCurveViewerFilterFlags::MorphTarget))
 			{
-				if (!SmartName.DisplayName.ToString().Contains(*FilterText.ToString()))
+				for(const TPair<FName, float>& CurveValuePair : AnimInstance->GetAnimationCurveList(EAnimCurveType::MorphTargetCurve))
 				{
-					bAddToList = false;
+					AddCurve(CurveValuePair.Key, UE::Anim::ECurveElementFlags::MorphTarget);
+					ActiveCurves.Add(CurveValuePair.Key);
 				}
 			}
-
-			// If we passed that, see if we are filtering to only active
-			if (bAddToList && !bShowAllCurves)
+			if (EnumHasAnyFlags(CurrentCurveFlag, EAnimCurveViewerFilterFlags::Material))
 			{
-				bAddToList = ActiveCurves.Contains(SmartName.DisplayName);
+				for(const TPair<FName, float>& CurveValuePair : AnimInstance->GetAnimationCurveList(EAnimCurveType::MaterialCurve))
+				{
+					AddCurve(CurveValuePair.Key, UE::Anim::ECurveElementFlags::Material);
+					ActiveCurves.Add(CurveValuePair.Key);
+				}
 			}
 
-			if(AnimCurvesByUID[UIDIndex]->bShown != bAddToList)
+			// If we arent filtering by curve type, show 'attribute' curves
+			if(!EnumHasAnyFlags(CurrentCurveFlag, EAnimCurveViewerFilterFlags::MorphTarget | EAnimCurveViewerFilterFlags::Material))
 			{
-				AnimCurvesByUID[UIDIndex]->bShown = bAddToList;
-				bDirty = true;
+				for(const TPair<FName, float>& CurveValuePair : AnimInstance->GetAnimationCurveList(EAnimCurveType::AttributeCurve))
+				{
+					AddCurve(CurveValuePair.Key, UE::Anim::ECurveElementFlags::None);
+					ActiveCurves.Add(CurveValuePair.Key);
+				}
 			}
-			
-			// If we still want to add
-			if (bAddToList)
+		}
+	}
+	
+	// Iterate through all curves that have been seen
+	for (const TPair<FName, TSharedPtr<FDisplayedAnimCurveInfo>>& CurveNameValuePair : AllSeenAnimCurvesMap)
+	{
+		bool bAddToList = true;
+
+		// See if we pass the search filter
+		if (!FilterText.IsEmpty())
+		{
+			if (!CurveNameValuePair.Key.ToString().Contains(*FilterText.ToString()))
 			{
-				AnimCurveList.Add(AnimCurvesByUID[UIDIndex]);
+				bAddToList = false;
 			}
-			UIDIndex++;
+		}
+
+		// If we passed that, see if we are filtering to only active
+		if (bAddToList && EnumHasAnyFlags(CurrentCurveFlag, EAnimCurveViewerFilterFlags::Active))
+		{
+			bAddToList = ActiveCurves.Contains(CurveNameValuePair.Key);
+		}
+
+		if(CurveNameValuePair.Value->bShown != bAddToList)
+		{
+			CurveNameValuePair.Value->bShown = bAddToList;
+			bDirty = true;
+		}
+		
+		// If we still want to add
+		if (bAddToList)
+		{
+			AnimCurveList.Add(CurveNameValuePair.Value);
 		}
 	}
 
@@ -965,57 +864,194 @@ void SAnimCurveViewer::CreateAnimCurveList( const FString& SearchText, bool bInF
 		{
 			bool operator()(const TSharedPtr<FDisplayedAnimCurveInfo>& A, const TSharedPtr<FDisplayedAnimCurveInfo>& B) const
 			{
-				return (A.Get()->SmartName.DisplayName.Compare(B.Get()->SmartName.DisplayName) < 0);
+				return (A.Get()->CurveName.Compare(B.Get()->CurveName) < 0);
 			}
 		};
 
 		AnimCurveList.Sort(FSortSmartNamesAlphabetically());
 
 		AnimCurveListView->RequestListRefresh();
-	}	
+	}
 }
 
-void SAnimCurveViewer::CreateAnimCurveTypeList(TSharedRef<SHorizontalBox> HorizontalBox)
+TSharedRef<SWidget> SAnimCurveViewer::CreateCurveSourceSelector()
 {
-	// Add toggle button for 'all curves'
-	HorizontalBox->AddSlot()
-	.AutoWidth()
-	.Padding(3.f, 1.f)
-	[
-		SNew(SCheckBox)
-		.Style(FAppStyle::Get(), "ToggleButtonCheckbox")
-		.ToolTipText(LOCTEXT("ShowAllCurvesTooltip", "Show all curves, or only active curves."))
-		.Type(ESlateCheckBoxType::ToggleButton)
-		.IsChecked(this, &SAnimCurveViewer::IsShowingAllCurves)
-		.OnCheckStateChanged(this, &SAnimCurveViewer::OnToggleShowingAllCurves)
-		.Padding(4.f)
+	bool bShowPoseWatches = false;
+	if(const UAnimInstance* AnimInstance = GetAnimInstance())
+	{
+		if(UAnimBlueprintGeneratedClass* AnimClass = Cast<UAnimBlueprintGeneratedClass>(AnimInstance->GetClass()))
+		{
+			bShowPoseWatches = true;
+		}
+	}
+
+	if(bShowPoseWatches)
+	{
+		AnimationEditorUtils::OnPoseWatchesChanged().AddSP(this, &SAnimCurveViewer::HandlePoseWatchesChanged);
+
+		// Add selector for pose watches
+		SAssignNew(PoseWatchCombo, SComboBox<TSharedPtr<TWeakObjectPtr<UPoseWatchPoseElement>>>)
+		.ToolTipText_Lambda([this]()
+		{
+			return PoseWatch.Get() != nullptr && PoseWatch->GetParent() != nullptr ? FText::Format(LOCTEXT("PoseWatchTooltipFormat", "Viewing Pose Watch Curves for '{0}'"), PoseWatch->GetParent()->GetLabel()) : LOCTEXT("NoPoseWatchSelected", "Viewing output Curves");
+		})
+		.OptionsSource(&PoseWatches)
+		.OnGenerateWidget_Lambda([this](TSharedPtr<TWeakObjectPtr<UPoseWatchPoseElement>> InElement)
+		{
+			TWeakObjectPtr<UPoseWatchPoseElement> Element = *InElement;
+
+			return
+				SNew(SHorizontalBox)
+				+SHorizontalBox::Slot()
+				.AutoWidth()
+				.VAlign(VAlign_Center)
+				.Padding(2.0f, 0.0f)
+				[
+					SNew(SColorBlock)
+					.Visibility_Lambda([Element]()
+					{
+						return Element.Get() != nullptr ? EVisibility::Visible : EVisibility::Collapsed;
+					})
+					.AlphaBackgroundBrush(FAppStyle::Get().GetBrush("ColorPicker.RoundedAlphaBackground"))
+					.ShowBackgroundForAlpha(true)
+					.AlphaDisplayMode(EColorBlockAlphaDisplayMode::Ignore)
+					.Size(FVector2D(16.0f, 16.0f))
+					.CornerRadius(FVector4(4.0f, 4.0f, 4.0f, 4.0f))	
+					.Color_Lambda([Element]()
+					{
+						if(UPoseWatchPoseElement* PoseWatchPoseElement = Element.Get())
+						{
+							return FLinearColor(PoseWatchPoseElement->GetColor());
+						}
+						return FLinearColor::Gray;
+					})
+				]
+				+SHorizontalBox::Slot()
+				.AutoWidth()
+				.VAlign(VAlign_Center)
+				.Padding(2.0f, 0.0f)
+				[
+					SNew(STextBlock)
+					.Text_Lambda([Element]()
+					{
+						return Element.Get() != nullptr && Element->GetParent() != nullptr ? Element->GetParent()->GetLabel() : LOCTEXT("NoPoseWatchSelected", "Output Curves");
+					})
+				];
+		})
+		.OnSelectionChanged_Lambda([this](TSharedPtr<TWeakObjectPtr<UPoseWatchPoseElement>> InElement, ESelectInfo::Type InSelectionType)
+		{
+			if(!InElement.IsValid())
+			{
+				return;
+			}
+			
+			if((*InElement).Get() == nullptr)
+			{
+				PoseWatch = nullptr;
+				return;
+			}
+
+			if(const UAnimInstance* AnimInstance = GetAnimInstance())
+			{
+				if(UAnimBlueprintGeneratedClass* AnimClass = Cast<UAnimBlueprintGeneratedClass>(AnimInstance->GetClass()))
+				{
+					TWeakObjectPtr<UPoseWatchPoseElement> Element = *InElement;
+					
+					for(const FAnimNodePoseWatch& AnimNodePoseWatch : AnimClass->AnimBlueprintDebugData.AnimNodePoseWatch)
+					{
+						if(AnimNodePoseWatch.PoseWatchPoseElement && Element.Get() == AnimNodePoseWatch.PoseWatchPoseElement)
+						{
+							PoseWatch = AnimNodePoseWatch.PoseWatchPoseElement;
+							break;
+						}
+					}
+				}
+			}
+		})
 		.Content()
 		[
-			SNew(STextBlock)
-			.Text(LOCTEXT("ShowAllCurves", "All Curves"))
-		]
-	];
+			SNew(SHorizontalBox)
+			.Clipping(EWidgetClipping::OnDemand)
+			+SHorizontalBox::Slot()
+			.AutoWidth()
+			.VAlign(VAlign_Center)
+			.Padding(2.0f, 0.0f)
+			[
+				SNew(SColorBlock)
+				.Visibility_Lambda([this]()
+				{
+					return PoseWatch.Get() != nullptr ? EVisibility::Visible : EVisibility::Collapsed;
+				})
+				.AlphaBackgroundBrush(FAppStyle::Get().GetBrush("ColorPicker.RoundedAlphaBackground"))
+				.ShowBackgroundForAlpha(true)
+				.AlphaDisplayMode(EColorBlockAlphaDisplayMode::Ignore)
+				.Size(FVector2D(16.0f, 16.0f))
+				.CornerRadius(FVector4(4.0f, 4.0f, 4.0f, 4.0f))	
+				.Color_Lambda([this]()
+				{
+					if(UPoseWatchPoseElement* Element = PoseWatch.Get())
+					{
+						return FLinearColor(Element->GetColor());
+					}
+					return FLinearColor::Gray;
+				})
+			]
+			+SHorizontalBox::Slot()
+			.AutoWidth()
+			.VAlign(VAlign_Center)
+			.Padding(2.0f, 0.0f)
+			[
+				SNew(STextBlock)
+				.Text_Lambda([this]()
+				{
+					return PoseWatch.Get() != nullptr && PoseWatch->GetParent() != nullptr ? PoseWatch->GetParent()->GetLabel() : LOCTEXT("NoPoseWatchSelected", "Output Curves");
+				})
+			]
+		];
 
-	// Add check box for each curve type flag
-	TArray<int32> CurveFlagsToList;
-	CurveFlagsToList.Add(ACEF_DriveMorphTarget);
-	CurveFlagsToList.Add(ACEF_DriveAttribute);
-	CurveFlagsToList.Add(ACEF_DriveMaterial);
+		RebuildPoseWatches();
 
-	for (int32 CurveFlagIndex = 0; CurveFlagIndex < CurveFlagsToList.Num(); ++CurveFlagIndex)
+		return PoseWatchCombo.ToSharedRef();
+	}
+
+	return SNullWidget::NullWidget;
+}
+
+void SAnimCurveViewer::RebuildPoseWatches()
+{
+	if(PoseWatchCombo.IsValid())
 	{
-		HorizontalBox->AddSlot()
-		.AutoWidth()
-		.Padding(3.f, 1.f)
-		[
-			SNew(SAnimCurveTypeList)
-			.IsEnabled(this, &SAnimCurveViewer::IsCurveFilterEnabled)
-			.CurveFlags(CurveFlagsToList[CurveFlagIndex])
-			.AnimCurveViewerPtr(SharedThis(this))
-		];
+		PoseWatches.Empty();
+
+		PoseWatches.Add(MakeShared<TWeakObjectPtr<UPoseWatchPoseElement>>());
+
+		if(const UAnimInstance* AnimInstance = GetAnimInstance())
+		{
+			if(UAnimBlueprintGeneratedClass* AnimClass = Cast<UAnimBlueprintGeneratedClass>(AnimInstance->GetClass()))
+			{
+				for(const FAnimNodePoseWatch& AnimNodePoseWatch : AnimClass->AnimBlueprintDebugData.AnimNodePoseWatch)
+				{
+					if(AnimNodePoseWatch.PoseWatchPoseElement)
+					{
+						PoseWatches.Add(MakeShared<TWeakObjectPtr<UPoseWatchPoseElement>>(AnimNodePoseWatch.PoseWatchPoseElement));
+					}
+				}
+			}
+		}
+
+		PoseWatchCombo->RefreshOptions();
 	}
 }
 
+void SAnimCurveViewer::HandlePoseWatchesChanged(UAnimBlueprint* /*InAnimBlueprint*/, UEdGraphNode* /*InNode*/)
+{
+	RegisterActiveTimer(0.0f, FWidgetActiveTimerDelegate::CreateLambda([this](double /*InCurrentTime*/, float /*InDeltaTime*/)
+	{
+		RebuildPoseWatches();
+		return EActiveTimerReturnType::Stop;
+	}));
+}
+
 void SAnimCurveViewer::AddAnimCurveOverride( FName& Name, float Weight)
 {
 	float& Value = OverrideCurves.FindOrAdd(Name);
@@ -1079,167 +1115,10 @@ void SAnimCurveViewer::RefreshCurveList(bool bInFullRefresh)
 	CreateAnimCurveList(FilterText.ToString(), bInFullRefresh);
 }
 
-void SAnimCurveViewer::OnNameCommitted(const FText& InNewName, ETextCommit::Type, TSharedPtr<FDisplayedAnimCurveInfo> Item)
-{
-	const FSmartNameMapping* Mapping = GetAnimCurveMapping();
-	if (Mapping)
-	{
-		FName NewName(*InNewName.ToString());
-		if (NewName == Item->SmartName.DisplayName)
-		{
-			// Do nothing if trying to rename to existing name...
-		}
-		else if (NewName != NAME_None && !Mapping->Exists(NewName))
-		{
-			EditableSkeletonPtr.Pin()->RenameSmartname(ContainerName, Item->SmartName.UID, NewName);
-			// remove it, so that it can readd it. 
-			AnimCurveList.Remove(Item);
-		}
-		else
-		{
-			FFormatNamedArguments Args;
-			Args.Add(TEXT("InvalidName"), FText::FromName(NewName) );
-			FNotificationInfo Info(FText::Format(LOCTEXT("AnimCurveRenamed", "The name \"{InvalidName}\" is invalid or already used."), Args));
-
-			Info.bUseLargeFont = false;
-			Info.ExpireDuration = 5.0f;
-
-			TSharedPtr<SNotificationItem> Notification = FSlateNotificationManager::Get().AddNotification(Info);
-			if (Notification.IsValid())
-			{
-				Notification->SetCompletionState(SNotificationItem::CS_Fail);
-			}
-		}
-	}
-}
-
-void SAnimCurveViewer::OnDeleteNameClicked()
-{
-	TArray< TSharedPtr<FDisplayedAnimCurveInfo> > SelectedItems = AnimCurveListView->GetSelectedItems();
-	TArray<FName> SelectedNames;
-
-	for (TSharedPtr<FDisplayedAnimCurveInfo> Item : SelectedItems)
-	{
-		SelectedNames.Add(Item->SmartName.DisplayName);
-	}
-
-	EditableSkeletonPtr.Pin()->RemoveSmartnamesAndFixupAnimations(ContainerName, SelectedNames);
-}
-
-bool SAnimCurveViewer::CanDelete()
-{
-	return AnimCurveListView->GetNumItemsSelected() > 0;
-}
-
-static const TCHAR* ClipboardHeader = TEXT("AnimCurveViewer");
-
-void SAnimCurveViewer::OnCopyClicked()
-{
-	TArray<TSharedPtr<FDisplayedAnimCurveInfo>> SelectedItems = AnimCurveListView->GetSelectedItems();
-
-	FAnimCurveViewerClipboard Clipboard;
-
-	for (TSharedPtr<FDisplayedAnimCurveInfo> Item : SelectedItems)
-	{
-		FAnimCurveViewerClipboardEntry Entry;
-		Entry.CurveName = Item->SmartName.DisplayName;
-
-		if(const FCurveMetaData* CurveMetaData = EditableSkeletonPtr.Pin()->GetSkeleton().GetCurveMetaData(Item->SmartName.DisplayName))
-		{
-			Entry.MetaData = *CurveMetaData;
-		}
-
-		Clipboard.Entries.Add(Entry);
-	}
-
-	FString ClipboardString = ClipboardHeader;
-	FAnimCurveViewerClipboard::StaticStruct()->ExportText(ClipboardString, &Clipboard, nullptr, nullptr, PPF_None, nullptr);
-	FPlatformApplicationMisc::ClipboardCopy(*ClipboardString);
-}
-
-bool SAnimCurveViewer::CanCopy() const
-{
-	return AnimCurveListView->GetNumItemsSelected() > 0;
-}
-
-void SAnimCurveViewer::OnPasteClicked()
-{
-	if (const FSmartNameMapping* NameMapping = GetAnimCurveMapping())
-	{
-		FScopedTransaction Transaction(LOCTEXT("PasteCurves", "Paste Curves"));
-		
-		FString TextToImport;
-		FPlatformApplicationMisc::ClipboardPaste(TextToImport);
-		TextToImport.RemoveFromStart(ClipboardHeader);
-
-		FAnimCurveViewerClipboard Clipboard;
-		FAnimCurveViewerClipboard::StaticStruct()->ImportText(*TextToImport, &Clipboard, nullptr, PPF_None, GLog, FAnimCurveViewerClipboard::StaticStruct()->GetName());
-
-		for(const FAnimCurveViewerClipboardEntry& Entry : Clipboard.Entries)
-		{
-			FSmartName NewCurveName;
-			EditableSkeletonPtr.Pin()->AddSmartname(ContainerName, Entry.CurveName, NewCurveName);
-			EditableSkeletonPtr.Pin()->SetCurveMetaBoneLinks(NewCurveName, Entry.MetaData.LinkedBones, Entry.MetaData.MaxLOD);
-			EditableSkeletonPtr.Pin()->SetCurveMetaDataMaterial(NewCurveName, Entry.MetaData.Type.bMaterial);
-			EditableSkeletonPtr.Pin()->SetCurveMetaDataMorphTarget(NewCurveName, Entry.MetaData.Type.bMorphtarget);
-		}
-
-		RefreshCurveList(true);
-	}
-}
-
-bool SAnimCurveViewer::CanPaste() const
-{
-	FString TextToImport;
-	FPlatformApplicationMisc::ClipboardPaste(TextToImport);
-	return TextToImport.StartsWith(ClipboardHeader);
-}
-
-void SAnimCurveViewer::OnSelectionChanged(TSharedPtr<FDisplayedAnimCurveInfo> InItem, ESelectInfo::Type SelectInfo)
-{
-	// make sure the currently selected ones are refreshed if it's first time
-	TArray<UObject*> SelectedObjects;
-
-	TArray< TSharedPtr< FDisplayedAnimCurveInfo > > SelectedRows = AnimCurveListView->GetSelectedItems();
-	for (auto ItemIt = SelectedRows.CreateIterator(); ItemIt; ++ItemIt)
-	{
-		TSharedPtr< FDisplayedAnimCurveInfo > RowItem = (*ItemIt);
-		UEditorAnimCurveBoneLinks* EditorMirrorObj = RowItem->EditorMirrorObject;
-		if (RowItem == InItem)
-		{
-			// first time selected, refresh
-			TArray<FBoneReference> BoneLinks;
-			FSmartName CurrentName = RowItem->SmartName;
-			const FCurveMetaData* CurveMetaData = EditableSkeletonPtr.Pin()->GetSkeleton().GetCurveMetaData(CurrentName);
-			uint8 MaxLOD = 0xFF;
-			if (CurveMetaData)
-			{
-				BoneLinks = CurveMetaData->LinkedBones;
-				MaxLOD = CurveMetaData->MaxLOD;
-			}
-
-			EditorMirrorObj->Refresh(CurrentName, BoneLinks, MaxLOD);
-		}
-
-		SelectedObjects.Add(EditorMirrorObj);
-	}
-
-	OnObjectsSelected.ExecuteIfBound(SelectedObjects);
-}
-
-void SAnimCurveViewer::ApplyCurveBoneLinks(class UEditorAnimCurveBoneLinks* EditorObj)
-{
-	if (EditorObj)
-	{
-		EditableSkeletonPtr.Pin()->SetCurveMetaBoneLinks(EditorObj->CurveName, EditorObj->ConnectedBones, EditorObj->MaxLOD);
-	}
-}
-
-void SAnimCurveViewer::HandleSmartNamesChange(const FName& InContainerName)
+void SAnimCurveViewer::HandleCurveMetaDataChange()
 {
 	AnimCurveList.Empty();
 	RefreshCurveList(true);
-	EditableSkeletonPtr.Pin()->GetSkeleton().OnSmartNamesChangedEvent.Broadcast();
 }
 
 #undef LOCTEXT_NAMESPACE
diff --git a/Engine/Source/Editor/Persona/Private/SAnimCurveViewer.h b/Engine/Source/Editor/Persona/Private/SAnimCurveViewer.h
--- a/Engine/Source/Editor/Persona/Private/SAnimCurveViewer.h
+++ b/Engine/Source/Editor/Persona/Private/SAnimCurveViewer.h
@@ -16,36 +16,19 @@
 #include "Widgets/Views/STableRow.h"
 #include "Widgets/Views/SListView.h"
 #include "Animation/AnimInstance.h"
-#include "EditorObjectsTracker.h"
 #include "PersonaDelegates.h"
 #include "EditorUndoClient.h"
-#include "SAnimCurveViewer.generated.h"
+#include "Filters/FilterBase.h"
+#include "Widgets/Input/SComboBox.h"
 
 class FUICommandList;
 class IEditableSkeleton;
 class SAnimCurveViewer;
-
-/** Clipboard contents for anim curves */
-USTRUCT()
-struct FAnimCurveViewerClipboardEntry
-{
-	GENERATED_BODY()
-	
-	UPROPERTY()
-	FName CurveName;
-
-	UPROPERTY()
-	FCurveMetaData MetaData;
-};
-
-USTRUCT()
-struct FAnimCurveViewerClipboard
-{
-	GENERATED_BODY()
-	
-	UPROPERTY()
-	TArray<FAnimCurveViewerClipboardEntry> Entries;
-};
+class UEditorAnimCurveBoneLinks;
+class UPoseWatchPoseElement;
+class UAnimBlueprint;
+class UEdGraphNode;
+enum class EAnimCurveViewerFilterFlags: uint8;
 
 //////////////////////////////////////////////////////////////////////////
 // FDisplayedAnimCurveInfo
@@ -53,30 +36,30 @@ struct FAnimCurveViewerClipboard
 class FDisplayedAnimCurveInfo
 {
 public:
-	FSmartName SmartName;
+	FName CurveName;
 	float Weight;
-	bool bAutoFillData;
+	bool bOverrideData;
 	TWeakPtr<class IEditableSkeleton> EditableSkeleton;		// The skeleton we're associated with
-	TSharedPtr<SInlineEditableTextBlock> EditableText;	// The editable text box in the list, used to focus from the context menu
-	FName ContainerName;	// The container in the skeleton this name resides in
-	class UEditorAnimCurveBoneLinks* EditorMirrorObject;
 	bool bShown;
+	bool bMorphTarget;
+	bool bMaterial;
+	
 	/** Static function for creating a new item, but ensures that you can only have a TSharedRef to one */
-	static TSharedRef<FDisplayedAnimCurveInfo> Make(TWeakPtr<class IEditableSkeleton> InEditableSkeleton, const FName& InContainerName, const FSmartName& InSmartName, class UEditorAnimCurveBoneLinks* InEditorMirrorObject)
+	static TSharedRef<FDisplayedAnimCurveInfo> Make(TWeakPtr<class IEditableSkeleton> InEditableSkeleton, const FName& InCurveName)
 	{
-		return MakeShareable(new FDisplayedAnimCurveInfo(InEditableSkeleton, InContainerName, InSmartName, InEditorMirrorObject));
+		return MakeShareable(new FDisplayedAnimCurveInfo(InEditableSkeleton, InCurveName));
 	}
 
 protected:
 	/** Hidden constructor, always use Make above */
-	FDisplayedAnimCurveInfo(TWeakPtr<class IEditableSkeleton> InEditableSkeleton, const FName& InContainerName, const FSmartName& InSmartName, class UEditorAnimCurveBoneLinks* InEditorMirrorObject)
-		: SmartName(InSmartName)
-		, Weight( 0 )
-		, bAutoFillData(true)
+	FDisplayedAnimCurveInfo(TWeakPtr<class IEditableSkeleton> InEditableSkeleton, const FName& InCurveName)
+		: CurveName(InCurveName)
+		, Weight(0.0f)
+		, bOverrideData(false)
 		, EditableSkeleton(InEditableSkeleton)
-		, ContainerName(InContainerName)
-		, EditorMirrorObject(InEditorMirrorObject)
 		, bShown(false)
+		, bMorphTarget(false)
+		, bMaterial(false)
 	{}
 };
 
@@ -121,16 +104,22 @@ private:
 	*/
 	void OnAnimCurveWeightValueCommitted(float NewWeight, ETextCommit::Type CommitType);
 
-	/** Auto fill check call back functions */
-	void OnAnimCurveAutoFillChecked(ECheckBoxState InState);
-	ECheckBoxState IsAnimCurveAutoFillChangedChecked() const;
+	/** Override check call back functions */
+	void OnAnimCurveOverrideChecked(ECheckBoxState InState);
+	ECheckBoxState IsAnimCurveOverrideChecked() const;
 
+	/** Check for an active morph target or material flag */
+	bool GetActiveFlag(bool bMorphTarget) const;
+	
 	/* Curve Flag checks for morphtarget or material */
-	void OnAnimCurveTypeBoxChecked(ECheckBoxState InState, bool bMorphTarget);
 	ECheckBoxState IsAnimCurveTypeBoxChangedChecked(bool bMorphTarget) const;
 
 	/** Returns the weight of this curve */
 	float GetWeight() const;
+	/** Returns the min seen weight of this curve */
+	TOptional<float> GetMinWeight() const;
+	/** Returns the max seen weight of this curve */
+	TOptional<float> GetMaxWeight() const;
 	/** Returns name of this curve */
 	FText GetItemName() const;
 	/** Get text we are filtering for */
@@ -146,14 +135,18 @@ private:
 
 	/** The name and weight of the morph target */
 	TSharedPtr<FDisplayedAnimCurveInfo>	Item;
-
+	
 	/** Preview scene used to update on scrub */
 	TWeakPtr<class IPersonaPreviewScene> PreviewScenePtr;
+
+	/** The min weight we have seen */
+	mutable float MinWeight = -1.0f;
+
+	/** The max weight we have seen */
+	mutable float MaxWeight = 1.0f;
+	
 	/** Returns curve type widget constructed */
 	TSharedRef<SWidget> GetCurveTypeWidget();
-
-	/** returns display text for number of connected joint setting */
-	FText GetNumConntectedBones() const;
 };
 
 //////////////////////////////////////////////////////////////////////////
@@ -180,10 +173,7 @@ public:
 	*
 	*/
 	virtual ~SAnimCurveViewer();
-
-	/** SWidget interface */
-	virtual FReply OnKeyDown(const FGeometry& MyGeometry, const FKeyEvent& InKeyEvent);
-
+	
 	/**
 	* Is registered with Persona to handle when its preview mesh is changed.
 	*
@@ -199,9 +189,6 @@ public:
 	*/
 	void OnPreviewAssetChanged(class UAnimationAsset* NewPreviewAsset);
 
-	/** Is registered with Persona to handle when curves change. */
-	void OnCurvesChanged();
-
 	/**
 	* Filters the SListView when the user changes the search text box (NameFilterBox)
 	*
@@ -210,9 +197,6 @@ public:
 	*/
 	void OnFilterTextChanged( const FText& SearchText );
 
-
-
-
 	/**
 	* Filters the SListView when the user hits enter or clears the search box
 	* Simply calls OnFilterTextChanged
@@ -248,12 +232,6 @@ public:
 	/** Is there currently an override for this curve, and if so, what is its weight */
 	bool GetAnimCurveOverride(FName& Name, float& Weight);
 
-	/**
-	* Tells the AnimInstance to reset all of its morph target curves
-	*
-	*/
-	void ResetAnimCurves();
-	
 	/**
 	* Accessor so our rows can grab the filtertext for highlighting
 	*
@@ -273,17 +251,8 @@ public:
 
 	void RefreshCurveList(bool bInFullRefresh);
 
-	// When a name is committed after being edited in the list
-	virtual void OnNameCommitted(const FText& NewName, ETextCommit::Type CommitType, TSharedPtr<FDisplayedAnimCurveInfo> Item);
-
 private:
 
-	void BindCommands();
-
-	/** Handler for context menus */
-	TSharedPtr<SWidget> OnGetContextMenuContent() const;
-	void OnSelectionChanged(TSharedPtr<FDisplayedAnimCurveInfo> InItem, ESelectInfo::Type SelectInfo);
-
 	/**
 	* Clears and rebuilds the table, according to an optional search string
 	*
@@ -291,42 +260,21 @@ private:
 	*
 	*/
 	void CreateAnimCurveList( const FString& SearchText = FString(), bool bInFullRefresh = false );
-	void CreateAnimCurveTypeList(TSharedRef<SHorizontalBox> HorizontalBox);
 
 	void ApplyCustomCurveOverride(UAnimInstance* AnimInstance) const;
 
-	void OnDeleteNameClicked();
-	bool CanDelete();
-
-	void OnRenameClicked();
-	bool CanRename();
-
-	void OnCopyClicked();
-	bool CanCopy() const;
-
-	void OnPasteClicked();
-	bool CanPaste() const;
-
-	void OnAddClicked();
-
-	ECheckBoxState IsShowingAllCurves() const;
-	void OnToggleShowingAllCurves(ECheckBoxState NewState);
-
-	bool IsCurveFilterEnabled() const;
-
-
-	// Adds a new smartname entry to the skeleton in the container we are managing
-	void CreateNewNameEntry(const FText& CommittedText, ETextCommit::Type CommitType);
-
-	/** Handle smart name (i.e. curve) removal */
-	void HandleSmartNamesChange(const FName& InContainerName);
-
-	/** Get the SmartNameMapping for anim curves */
-	const struct FSmartNameMapping* GetAnimCurveMapping();
+	/** Handle curve meta data removal */
+	void HandleCurveMetaDataChange();
 
 	/** Get the anim instance we are viewing */
 	UAnimInstance* GetAnimInstance() const;
 
+	TSharedRef<SWidget> CreateCurveSourceSelector();
+	
+	/** Handle building list of available pose watches */
+	void HandlePoseWatchesChanged(UAnimBlueprint* /*InAnimBlueprint*/, UEdGraphNode* /*InNode*/);
+	void RebuildPoseWatches();
+
 	/** Pointer to the preview scene we are bound to */
 	TWeakPtr<class IPersonaPreviewScene> PreviewScenePtr;
 
@@ -339,24 +287,16 @@ private:
 	/** A list of animation curve. Used by the AnimCurveListView. */
 	TArray< TSharedPtr<FDisplayedAnimCurveInfo> > AnimCurveList;
 
-	/** Tracking array of anim curves indexed by UID */
-	TArray< TSharedPtr<FDisplayedAnimCurveInfo> > AnimCurvesByUID;
-	
-	/** The skeletal mesh that we grab the animation curve from */
-	UAnimInstance* CachedPreviewInstance;						
+	/** Tracking map of anim curves */
+	TMap< FName, TSharedPtr<FDisplayedAnimCurveInfo> > AllSeenAnimCurvesMap;
 
 	/** Widget used to display the list of animation curve */
 	TSharedPtr<SAnimCurveListType> AnimCurveListView;
 
-	/** Name of the skeleton smart name container to display in the list */
-	FName ContainerName;
-
 	/** Current text typed into NameFilterBox */
 	FText FilterText;
 
-	int32 CurrentCurveFlag;
-
-	bool bShowAllCurves;
+	EAnimCurveViewerFilterFlags CurrentCurveFlag;
 
 	TMap<FName, float> OverrideCurves;
 
@@ -366,15 +306,18 @@ private:
 	friend class SAnimCurveListRow;
 	friend class SAnimCurveTypeList;
 
-	/** Tracks objects created for the details panel */
-	FEditorObjectTracker EditorObjectTracker;
-
 	/** Delegate called to select objects */
 	FOnObjectsSelected OnObjectsSelected;
 
-	/** apply curve bone links from editor mirror object to skeleton */
-	void ApplyCurveBoneLinks(class UEditorAnimCurveBoneLinks* EditorObj);
+	/** Valid when viewing a pose watch's data */
+	TWeakObjectPtr<UPoseWatchPoseElement> PoseWatch;
+
+	/** Names of each pose watch that can be selected */
+	TArray<TSharedPtr<TWeakObjectPtr<UPoseWatchPoseElement>>> PoseWatches;
+	
+	/** Combobox used to select pose watch */
+	TSharedPtr<SComboBox<TSharedPtr<TWeakObjectPtr<UPoseWatchPoseElement>>>> PoseWatchCombo;
 
-	/** Delegate handle for HandleSmartNameRemoved callback */
-	FDelegateHandle SmartNameChangedHandle;
+	/** All filters that can be applied to the widget's display */
+	TArray<TSharedRef<FFilterBase<EAnimCurveViewerFilterFlags>>> Filters;
 };
\ No newline at end of file
diff --git a/Engine/Source/Editor/Persona/Private/SAnimSequenceCurveEditor.cpp b/Engine/Source/Editor/Persona/Private/SAnimSequenceCurveEditor.cpp
--- a/Engine/Source/Editor/Persona/Private/SAnimSequenceCurveEditor.cpp
+++ b/Engine/Source/Editor/Persona/Private/SAnimSequenceCurveEditor.cpp
@@ -23,13 +23,19 @@
 
 #define LOCTEXT_NAMESPACE "SAnimSequenceCurveEditor"
 
-FRichCurveEditorModelNamed::FRichCurveEditorModelNamed(const FSmartName& InName, ERawCurveTrackTypes InType, int32 InCurveIndex, UAnimSequenceBase* InAnimSequence, FCurveEditorTreeItemID InTreeId /*= FCurveEditorTreeItemID()*/) : FRichCurveEditorModel(InAnimSequence)
-, Name(InName)
+FRichCurveEditorModelNamed::FRichCurveEditorModelNamed(const FSmartName& InName, ERawCurveTrackTypes InType, int32 InCurveIndex, UAnimSequenceBase* InAnimSequence, FCurveEditorTreeItemID InTreeId /*= FCurveEditorTreeItemID()*/)
+	: FRichCurveEditorModelNamed(InName.DisplayName, InType, InCurveIndex, InAnimSequence, InTreeId)
+{
+}
+
+FRichCurveEditorModelNamed::FRichCurveEditorModelNamed(const FName& InName, ERawCurveTrackTypes InType, int32 InCurveIndex, UAnimSequenceBase* InAnimSequence, FCurveEditorTreeItemID InTreeId /*= FCurveEditorTreeItemID()*/)
+: FRichCurveEditorModel(InAnimSequence)
+, CurveName(InName)
 , AnimSequence(InAnimSequence)
 , CurveIndex(InCurveIndex)
 , Type(InType)
 , TreeId(InTreeId)
-, CurveId(FAnimationCurveIdentifier(Name, Type))
+, CurveId(FAnimationCurveIdentifier(InName, Type))
 , bCurveRemoved(false)
 {
 	CurveModifiedDelegate.AddRaw(this, &FRichCurveEditorModelNamed::CurveHasChanged);
@@ -44,6 +50,7 @@ FRichCurveEditorModelNamed::FRichCurveEditorModelNamed(const FSmartName& InName,
 	UpdateCachedCurve();
 }
 
+
 FRichCurveEditorModelNamed::~FRichCurveEditorModelNamed()
 {
 	AnimSequence->GetDataModel()->GetModifiedEvent().RemoveAll(this);
@@ -51,7 +58,7 @@ FRichCurveEditorModelNamed::~FRichCurveEditorModelNamed()
 
 bool FRichCurveEditorModelNamed::IsValid() const
 {
-	return AnimSequence->GetDataModel()->FindCurve(FAnimationCurveIdentifier(Name, Type)) != nullptr;
+	return AnimSequence->GetDataModel()->FindCurve(FAnimationCurveIdentifier(CurveName, Type)) != nullptr;
 }
 
 FRichCurve& FRichCurveEditorModelNamed::GetRichCurve()
@@ -143,7 +150,7 @@ void FRichCurveEditorModelNamed::OnModelHasChanged(const EAnimDataModelNotifyTyp
 		case EAnimDataModelNotifyType::CurveScaled:
 		{
 			const FCurvePayload& TypedPayload = Payload.GetPayload<FCurvePayload>();
-			if (TypedPayload.Identifier.InternalName == Name)
+			if (TypedPayload.Identifier.CurveName == CurveName)
 			{
 				if (NotifyCollector.IsNotWithinBracket())
 				{
@@ -166,7 +173,7 @@ void FRichCurveEditorModelNamed::OnModelHasChanged(const EAnimDataModelNotifyTyp
 		{
 			// Curve was removed
 			const FCurveRemovedPayload& TypedPayload = Payload.GetPayload<FCurveRemovedPayload>();
-			if (TypedPayload.Identifier.InternalName == Name)
+			if (TypedPayload.Identifier.CurveName == CurveName)
 			{
 				bCurveRemoved = true;
 			}
@@ -178,7 +185,7 @@ void FRichCurveEditorModelNamed::OnModelHasChanged(const EAnimDataModelNotifyTyp
 			const FCurveRenamedPayload& TypedPayload = Payload.GetPayload<FCurveRenamedPayload>();
 			if (TypedPayload.Identifier == CurveId)
 			{
-				Name = TypedPayload.NewIdentifier.InternalName;
+				CurveName = TypedPayload.NewIdentifier.CurveName;
 				CurveId = TypedPayload.NewIdentifier;
 
 				if (NotifyCollector.IsNotWithinBracket())
@@ -267,7 +274,7 @@ void FRichCurveEditorModelNamed::UpdateCachedCurve()
 class FAnimSequenceCurveEditorItem : public ICurveEditorTreeItem
 {
 public:
-	FAnimSequenceCurveEditorItem(const FSmartName& InName, ERawCurveTrackTypes InType, int32 InCurveIndex, UAnimSequenceBase* InAnimSequence, const FText& InCurveDisplayName, const FLinearColor& InCurveColor, FSimpleDelegate InOnCurveModified, FCurveEditorTreeItemID InTreeId)
+	FAnimSequenceCurveEditorItem(const FName& InName, ERawCurveTrackTypes InType, int32 InCurveIndex, UAnimSequenceBase* InAnimSequence, const FText& InCurveDisplayName, const FLinearColor& InCurveColor, FSimpleDelegate InOnCurveModified, FCurveEditorTreeItemID InTreeId)
 		: Name(InName)
 		, Type(InType)
 		, CurveIndex(InCurveIndex)
@@ -340,7 +347,7 @@ public:
 		return false;
 	}
 
-	FSmartName Name;
+	FName Name;
 	ERawCurveTrackTypes Type;
 	int32 CurveIndex;
 	TWeakObjectPtr<UAnimSequenceBase> AnimSequence;
@@ -448,14 +455,14 @@ void SAnimSequenceCurveEditor::Construct(const FArguments& InArgs, const TShared
 
 void SAnimSequenceCurveEditor::OnModelHasChanged(const EAnimDataModelNotifyType& NotifyType, IAnimationDataModel* Model, const FAnimDataModelNotifPayload& Payload)
 {
-	auto StopEditingCurve = [this, NotifyType, &Payload, Model]()
+	auto StopEditingCurve = [this, &Payload, Model]()
 	{
 		const FCurvePayload& TypedPayload = Payload.GetPayload<FCurvePayload>();
 		const FAnimationCurveIdentifier& CurveId = TypedPayload.Identifier;
 		const int32 ChannelIndices = CurveId.CurveType == ERawCurveTrackTypes::RCT_Transform ? 9 : 1;
 		for (int32 ChannelIndex = 0; ChannelIndex < ChannelIndices; ++ChannelIndex)
 		{
-			RemoveCurve(CurveId.InternalName, CurveId.CurveType, ChannelIndex);
+			RemoveCurve(CurveId.CurveName, CurveId.CurveType, ChannelIndex);
 		}
 	};
 	
@@ -539,13 +546,13 @@ void SAnimSequenceCurveEditor::ResetCurves()
 	CurveEditor->RemoveAllCurves();
 }
 
-void SAnimSequenceCurveEditor::AddCurve(const FText& InCurveDisplayName, const FLinearColor& InCurveColor, const FSmartName& InName, ERawCurveTrackTypes InType, int32 InCurveIndex, FSimpleDelegate InOnCurveModified)
+void SAnimSequenceCurveEditor::AddCurve(const FText& InCurveDisplayName, const FLinearColor& InCurveColor, const FName& InName, ERawCurveTrackTypes InType, int32 InCurveIndex, FSimpleDelegate InOnCurveModified)
 {
 	// Ensure that curve is not already being edited
 	for (const TPair<FCurveModelID, TUniquePtr<FCurveModel>>& CurvePair : CurveEditor->GetCurves())
 	{
 		const FRichCurveEditorModelNamed* Model = static_cast<FRichCurveEditorModelNamed*>(CurvePair.Value.Get());
-		if(Model->Name == InName && Model->Type == InType && Model->CurveIndex == InCurveIndex)
+		if(Model->CurveName == InName && Model->Type == InType && Model->CurveIndex == InCurveIndex)
 		{
 			return;
 		}
@@ -568,7 +575,7 @@ void SAnimSequenceCurveEditor::AddCurve(const FText& InCurveDisplayName, const F
 	CurveEditor->SetTreeSelection(MoveTemp(NewSelection));
 }
 
-void SAnimSequenceCurveEditor::RemoveCurve(const FSmartName& InName, ERawCurveTrackTypes InType, int32 InCurveIndex)
+void SAnimSequenceCurveEditor::RemoveCurve(const FName& InName, ERawCurveTrackTypes InType, int32 InCurveIndex)
 {
 	for(const FCurveEditorTreeItemID& TreeItemID : CurveEditor->GetRootTreeItems())
 	{
@@ -584,7 +591,7 @@ void SAnimSequenceCurveEditor::RemoveCurve(const FSmartName& InName, ERawCurveTr
 	for(const auto& CurvePair : CurveEditor->GetCurves())
 	{
 		FRichCurveEditorModelNamed* Model = static_cast<FRichCurveEditorModelNamed*>(CurvePair.Value.Get());
-		if(Model->Name == InName && Model->Type == InType && Model->CurveIndex == InCurveIndex)
+		if(Model->CurveName == InName && Model->Type == InType && Model->CurveIndex == InCurveIndex)
 		{
 			// Cache ID to prevent use after release
 			CurveEditor->RemoveCurve(CurvePair.Key);
diff --git a/Engine/Source/Editor/Persona/Private/SAnimSequenceCurveEditor.h b/Engine/Source/Editor/Persona/Private/SAnimSequenceCurveEditor.h
--- a/Engine/Source/Editor/Persona/Private/SAnimSequenceCurveEditor.h
+++ b/Engine/Source/Editor/Persona/Private/SAnimSequenceCurveEditor.h
@@ -23,8 +23,11 @@ class FTabManager;
 class FRichCurveEditorModelNamed : public FRichCurveEditorModel
 {
 public:
+	UE_DEPRECATED(5.3, "Please use the constructor that takes a FName.")
 	FRichCurveEditorModelNamed(const FSmartName& InName, ERawCurveTrackTypes InType, int32 InCurveIndex, UAnimSequenceBase* InAnimSequence, FCurveEditorTreeItemID InTreeId = FCurveEditorTreeItemID());
 
+	FRichCurveEditorModelNamed(const FName& InName, ERawCurveTrackTypes InType, int32 InCurveIndex, UAnimSequenceBase* InAnimSequence, FCurveEditorTreeItemID InTreeId = FCurveEditorTreeItemID());
+	
 	virtual ~FRichCurveEditorModelNamed();
 
 	virtual bool IsValid() const override;
@@ -39,7 +42,9 @@ public:
 	void OnModelHasChanged(const EAnimDataModelNotifyType& NotifyType, IAnimationDataModel* Model, const FAnimDataModelNotifPayload& Payload);
 	void UpdateCachedCurve();
 
+	UE_DEPRECATED(5.3, "Please use CurveName.")
 	FSmartName Name;
+	FName CurveName;
 	TWeakObjectPtr<UAnimSequenceBase> AnimSequence;
 	int32 CurveIndex;
 	ERawCurveTrackTypes Type;
@@ -69,8 +74,8 @@ class SAnimSequenceCurveEditor : public IAnimSequenceCurveEditor
 
 	/** IAnimSequenceCurveEditor interface */
 	virtual void ResetCurves() override;
-	virtual void AddCurve(const FText& InCurveDisplayName, const FLinearColor& InCurveColor, const FSmartName& InName, ERawCurveTrackTypes InType, int32 InCurveIndex, FSimpleDelegate InOnCurveModified) override;
-	virtual void RemoveCurve(const FSmartName& InName, ERawCurveTrackTypes InType, int32 InCurveIndex) override;
+	virtual void AddCurve(const FText& InCurveDisplayName, const FLinearColor& InCurveColor, const FName& InName, ERawCurveTrackTypes InType, int32 InCurveIndex, FSimpleDelegate InOnCurveModified) override;
+	virtual void RemoveCurve(const FName& InName, ERawCurveTrackTypes InType, int32 InCurveIndex) override;
 	virtual void ZoomToFit() override;
 	
 	void OnModelHasChanged(const EAnimDataModelNotifyType& NotifyType, IAnimationDataModel* Model, const FAnimDataModelNotifPayload& Payload);
diff --git a/Engine/Source/Editor/Persona/Private/SAnimationSequenceBrowser.cpp b/Engine/Source/Editor/Persona/Private/SAnimationSequenceBrowser.cpp
--- a/Engine/Source/Editor/Persona/Private/SAnimationSequenceBrowser.cpp
+++ b/Engine/Source/Editor/Persona/Private/SAnimationSequenceBrowser.cpp
@@ -265,11 +265,11 @@ public:
 			.HeightOverride(300.0f)
 			.WidthOverride(200.0f)
 			[
-				SNew(SAnimCurvePicker, EditableSkeleton.ToSharedRef())
-				.OnCurveNamePicked_Lambda([this](const FSmartName& InCurveSmartName)
+				SNew(SAnimCurvePicker, &EditableSkeleton->GetSkeleton())
+				.OnCurvePicked_Lambda([this](const FName& InCurveName)
 				{
 					FSlateApplication::Get().DismissAllMenus();
-					CurveString = InCurveSmartName.DisplayName.ToString();
+					CurveString = InCurveName.ToString();
 					OnChanged().Broadcast();
 				})
 			];
diff --git a/Engine/Source/Editor/Persona/Private/SMorphTargetViewer.cpp b/Engine/Source/Editor/Persona/Private/SMorphTargetViewer.cpp
--- a/Engine/Source/Editor/Persona/Private/SMorphTargetViewer.cpp
+++ b/Engine/Source/Editor/Persona/Private/SMorphTargetViewer.cpp
@@ -25,7 +25,7 @@ static const FName ColumnID_MorphTargetWeightLabel( "Weight" );
 static const FName ColumnID_MorphTargetEditLabel( "Edit" );
 static const FName ColumnID_MorphTargetVertCountLabel( "NumberOfVerts" );
 
-extern const float MaxMorphWeight;
+const float MaxMorphWeight = 5.0f;
 
 //////////////////////////////////////////////////////////////////////////
 // SMorphTargetListRow
diff --git a/Engine/Source/Editor/Persona/Private/SPoseAssetNameWidget.cpp b/Engine/Source/Editor/Persona/Private/SPoseAssetNameWidget.cpp
--- a/Engine/Source/Editor/Persona/Private/SPoseAssetNameWidget.cpp
+++ b/Engine/Source/Editor/Persona/Private/SPoseAssetNameWidget.cpp
@@ -40,14 +40,10 @@ void SPoseAssetNameWidget::RefreshBasePoseChanged()
 	{
 		BasePoseComboList.Reset();
 		// add pose names
-		TArray<FSmartName> PoseNames = PoseAsset->GetPoseNames();
-		if (PoseNames.Num() > 0)
+		// go through profile and see if it has mine
+		for (const FName& PoseName : PoseAsset->GetPoseFNames())
 		{
-			// go through profile and see if it has mine
-			for (const FSmartName& PoseName : PoseNames)
-			{
-				BasePoseComboList.Add(MakeShareable(new FString(PoseName.DisplayName.ToString())));
-			}
+			BasePoseComboList.Add(MakeShareable(new FString(PoseName.ToString())));
 		}
 	}
 
diff --git a/Engine/Source/Editor/Persona/Private/SPoseEditor.cpp b/Engine/Source/Editor/Persona/Private/SPoseEditor.cpp
--- a/Engine/Source/Editor/Persona/Private/SPoseEditor.cpp
+++ b/Engine/Source/Editor/Persona/Private/SPoseEditor.cpp
@@ -303,7 +303,7 @@ void SPoseViewer::Construct(const FArguments& InArgs, const TSharedRef<IPersonaT
 	EditableSkeletonPtr = InEditableSkeleton;
 	PoseAssetPtr = InArgs._PoseAsset;
 	
-	NewPoseName = UPoseAsset::GetUniquePoseName(&(EditableSkeletonPtr.Pin()->GetSkeleton()));
+	NewPoseName = UPoseAsset::GetUniquePoseName(PoseAssetPtr.Get());
 
 	InPreviewScene->RegisterOnPreviewMeshChanged(FOnPreviewMeshChanged::CreateSP(this, &SPoseViewer::OnPreviewMeshChanged));
 
@@ -798,14 +798,13 @@ void SPoseViewer::CreatePoseList(const FString& SearchText)
 	{
 		UPoseAsset* PoseAsset = PoseAssetPtr.Get();
 
-		TArray<FSmartName> PoseNames = PoseAsset->GetPoseNames();
+		TArray<FName> PoseNames = PoseAsset->GetPoseFNames();
 		if (PoseNames.Num() > 0)
 		{
 			bool bDoFiltering = !SearchText.IsEmpty();
 
-			for (const FSmartName& PoseSmartName : PoseNames)
+			for (const FName& PoseName : PoseNames)
 			{
-				FName PoseName = PoseSmartName.DisplayName;
 				if (bDoFiltering && !PoseName.ToString().Contains(SearchText))
 				{
 					continue; // Skip items that don't match our filter
@@ -834,16 +833,10 @@ void SPoseViewer::CreateCurveList(const FString& SearchText)
 	{
 		UPoseAsset* PoseAsset = PoseAssetPtr.Get();
 
-		TArray<FSmartName> CurveNames = PoseAsset->GetCurveNames();
-		if (CurveNames.Num() > 0)
+		for (const FName& CurveName : PoseAsset->GetCurveFNames())
 		{
-			for (const FSmartName& CurveSmartName : CurveNames)
-			{
-				FName CurveName = CurveSmartName.DisplayName;
-
-				const TSharedRef<FDisplayedCurveInfo> Info = FDisplayedCurveInfo::Make(CurveName);
-				CurveList.Add(Info);
-			}
+			const TSharedRef<FDisplayedCurveInfo> Info = FDisplayedCurveInfo::Make(CurveName);
+			CurveList.Add(Info);
 		}
 	}
 
@@ -933,14 +926,7 @@ UAnimInstance* SPoseViewer::GetAnimInstance() const
 
 bool SPoseViewer::ModifyName(FName OldName, FName NewName, bool bSilence)
 {
-	FScopedTransaction Transaction(LOCTEXT("RenamePoses", "Rename Pose"));
-	PoseAssetPtr.Get()->Modify();
-
-	// get smart name
-	const USkeleton& Skeleton = EditableSkeletonPtr.Pin()->GetSkeleton();
-	const SmartName::UID_Type ExistingUID = Skeleton.GetUIDByName(USkeleton::AnimCurveMappingName, NewName);
-	// verify if this name exists in smart naming
-	if (ExistingUID != SmartName::MaxUID)
+	if(PoseAssetPtr.Get()->ContainsPose(NewName))
 	{
 		// warn users
 		// if so, verify if this name is still okay
@@ -953,20 +939,15 @@ bool SPoseViewer::ModifyName(FName OldName, FName NewName, bool bSilence)
 				return false;
 			}
 		}
-
-		// I think this might have to be delegate of the top window
-		if (PoseAssetPtr.Get()->ModifyPoseName(OldName, NewName, &ExistingUID) == false)
-		{
-			return false;
-		}
 	}
-	else
+
+	FScopedTransaction Transaction(LOCTEXT("RenamePoses", "Rename Pose"));
+	PoseAssetPtr.Get()->Modify();
+	
+	// I think this might have to be delegate of the top window
+	if (PoseAssetPtr.Get()->ModifyPoseName(OldName, NewName) == false)
 	{
-		// I think this might have to be delegate of the top window
-		if (PoseAssetPtr.Get()->ModifyPoseName(OldName, NewName, nullptr) == false)
-		{
-			return false;
-		}
+		return false;
 	}
 
 	// now refresh pose data
@@ -1001,28 +982,23 @@ void SPoseViewer::UpdateSelectedPoseWithCurrent()
 		UDebugSkelMeshComponent* PreviewComponent = PreviewScenePtr.Pin()->GetPreviewMeshComponent();
 		const USkeleton* Skeleton = PoseAsset->GetSkeleton();
 		if (PreviewComponent && Skeleton)
-        {
-		    const FName PoseName = SelectedRows[0]->Name;
-		    FSmartName PoseCurveSmartName;
-		    // Ensure the PoseName exists as a curve on the skeleton
-		    if(Skeleton->GetSmartNameByName(USkeleton::AnimCurveMappingName, PoseName, PoseCurveSmartName))
-		    {
-			    FScopedTransaction Transaction(LOCTEXT("UpdatePose", "Update Pose from Viewport"));
-			    PoseAsset->Modify();			
-			    PoseAsset->AddOrUpdatePose(PoseCurveSmartName, PreviewComponent, false);
-
-		    	// Reset bone modifiers in case the user has created a pose using them - resetting them to the 'base' pose
-			    if(UAnimPreviewInstance* PreviewInstance = Cast<UAnimPreviewInstance>(PreviewComponent->GetAnimInstance()))
-			    {
-				    PreviewInstance->ResetModifiedBone();
-			    }
-		    	
-			    // Reinitialize animation preview
-			    RestartAnimations(Skeleton);
-			    RestartPreviewComponent();
-		    }		
-        }
+		{
+			const FName PoseName = SelectedRows[0]->Name;
+
+			FScopedTransaction Transaction(LOCTEXT("UpdatePose", "Update Pose from Viewport"));
+			PoseAsset->Modify();			
+			PoseAsset->AddOrUpdatePose(PoseName, PreviewComponent, false);
 
+			// Reset bone modifiers in case the user has created a pose using them - resetting them to the 'base' pose
+			if(UAnimPreviewInstance* PreviewInstance = Cast<UAnimPreviewInstance>(PreviewComponent->GetAnimInstance()))
+			{
+				PreviewInstance->ResetModifiedBone();
+			}
+
+			// Reinitialize animation preview
+			RestartAnimations(Skeleton);
+			RestartPreviewComponent();
+		}
 	}
 }
 
@@ -1034,18 +1010,10 @@ void SPoseViewer::AddPoseWithCurrent()
 		USkeleton* Skeleton = PoseAsset->GetSkeleton();
 		if (PreviewComponent && Skeleton)
 		{
-			FSmartName PoseCurveSmartName;
-			PoseCurveSmartName.DisplayName = NewPoseName;
-			
-			Skeleton->VerifySmartName(USkeleton::AnimCurveMappingName, PoseCurveSmartName);
-			if (PoseCurveSmartName.IsValid())
-			{
-				PoseAsset->AddOrUpdatePose(PoseCurveSmartName, PreviewComponent);
-				NewPoseName = UPoseAsset::GetUniquePoseName(&(EditableSkeletonPtr.Pin()->GetSkeleton()));
-			}
+			PoseAsset->AddOrUpdatePose(NewPoseName, PreviewComponent);
+			NewPoseName = UPoseAsset::GetUniquePoseName(PoseAsset);
 		}
 	}
-
 }
 
 void SPoseViewer::AddPoseWithReference()
@@ -1056,15 +1024,8 @@ void SPoseViewer::AddPoseWithReference()
 		USkeleton* Skeleton = PoseAsset->GetSkeleton();
 		if (PreviewComponent && Skeleton)
 		{
-			FSmartName PoseCurveSmartName;
-			PoseCurveSmartName.DisplayName = NewPoseName;
-			
-			Skeleton->VerifySmartName(USkeleton::AnimCurveMappingName, PoseCurveSmartName);
-			if (PoseCurveSmartName.IsValid())
-			{
-				PoseAsset->AddReferencePose(PoseCurveSmartName, PreviewComponent->GetReferenceSkeleton());
-				NewPoseName = UPoseAsset::GetUniquePoseName(&(EditableSkeletonPtr.Pin()->GetSkeleton()));
-			}
+			PoseAsset->AddReferencePose(NewPoseName, PreviewComponent->GetReferenceSkeleton());
+			NewPoseName = UPoseAsset::GetUniquePoseName(PoseAsset);
 		}
 	}	
 }
diff --git a/Engine/Source/Editor/Persona/Private/TabSpawners.cpp b/Engine/Source/Editor/Persona/Private/TabSpawners.cpp
--- a/Engine/Source/Editor/Persona/Private/TabSpawners.cpp
+++ b/Engine/Source/Editor/Persona/Private/TabSpawners.cpp
@@ -28,6 +28,7 @@
 #include "Engine/PreviewMeshCollection.h"
 #include "PoseWatchManagerPublicTypes.h"
 #include "PoseWatchManagerDefaultMode.h"
+#include "SAnimCurveMetadataEditor.h"
 
 #define LOCTEXT_NAMESPACE "PersonaModes"
 
@@ -37,6 +38,7 @@
 // Tab constants
 const FName FPersonaTabs::MorphTargetsID("MorphTargetsTab");
 const FName FPersonaTabs::AnimCurveViewID("AnimCurveViewerTab");
+const FName FPersonaTabs::AnimCurveMetadataEditorID("AnimCurveMetadataEditorTab");
 const FName FPersonaTabs::SkeletonTreeViewID("SkeletonTreeView");		//@TODO: Name
 // Skeleton Pose manager
 const FName FPersonaTabs::RetargetManagerID("RetargetManager");
@@ -72,6 +74,8 @@ const FName FPersonaTabs::AdvancedPreviewSceneSettingsID("AdvancedPreviewTab");
 
 const FName FPersonaTabs::DetailsID("DetailsTab");
 
+const FName FPersonaTabs::FindReplaceID("FindReplaceTab");
+
 /////////////////////////////////////////////////////
 // FPersonaMode
 
@@ -166,14 +170,14 @@ FAnimCurveViewerTabSummoner::FAnimCurveViewerTabSummoner(TSharedPtr<class FAsset
 	, PreviewScene(InPreviewScene)
 	, OnObjectsSelected(InOnObjectsSelected)
 {
-	TabLabel = LOCTEXT("AnimCurveViewTabTitle", "Anim Curves");
-	TabIcon = FSlateIcon(FAppStyle::GetAppStyleSetName(), "Persona.Tabs.AnimCurvePreviewer");
+	TabLabel = LOCTEXT("AnimCurveViewerTabTitle", "Curve Debugger");
+	TabIcon = FSlateIcon(FAppStyle::GetAppStyleSetName(), "Persona.Tabs.AnimCurveDebugger");
 
 	EnableTabPadding();
 	bIsSingleton = true;
 
-	ViewMenuDescription = LOCTEXT("AnimCurveTabView", "Animation Curves");
-	ViewMenuTooltip = LOCTEXT("AnimCurveTabView_ToolTip", "Shows the animation curve viewer");
+	ViewMenuDescription = LOCTEXT("AnimCurveTabView", "Curve Debugger");
+	ViewMenuTooltip = LOCTEXT("AnimCurveTabView_ToolTip", "Shows the animation curve debugger. This shows the state of animation curves.");
 }
 
 TSharedRef<SWidget> FAnimCurveViewerTabSummoner::CreateTabBody(const FWorkflowTabSpawnInfo& Info) const
@@ -181,6 +185,31 @@ TSharedRef<SWidget> FAnimCurveViewerTabSummoner::CreateTabBody(const FWorkflowTa
 	return SNew(SAnimCurveViewer, EditableSkeleton.Pin().ToSharedRef(), PreviewScene.Pin().ToSharedRef(), OnObjectsSelected);
 }
 
+/////////////////////////////////////////////////////
+// FAnimCurveViewerTabSummoner
+
+FAnimCurveMetadataEditorTabSummoner::FAnimCurveMetadataEditorTabSummoner(TSharedPtr<class FAssetEditorToolkit> InHostingApp, UObject* InMetadataHost, const TSharedRef<IPersonaPreviewScene>& InPreviewScene, FOnObjectsSelected InOnObjectsSelected)
+	: FWorkflowTabFactory(FPersonaTabs::AnimCurveMetadataEditorID, InHostingApp)
+	, MetadataHost(InMetadataHost)
+	, PreviewScene(InPreviewScene)
+	, OnObjectsSelected(InOnObjectsSelected)
+{
+	TabLabel = LOCTEXT("AnimCurveMetadataEditorTabTitle", "Curves");
+	TabIcon = FSlateIcon(FAppStyle::GetAppStyleSetName(), "Persona.Tabs.AnimCurveMetadataEditor");
+
+	EnableTabPadding();
+	bIsSingleton = true;
+
+	ViewMenuDescription = LOCTEXT("AnimCurveMetadataEditorTab", "Animation Curves");
+	ViewMenuTooltip = LOCTEXT("AnimCurveMetadataEditorTab_ToolTip", "Shows the animation curve editor. This lets you add, remove, rename and edit curve metadata.");
+}
+
+TSharedRef<SWidget> FAnimCurveMetadataEditorTabSummoner::CreateTabBody(const FWorkflowTabSpawnInfo& Info) const
+{
+	return SNew(SAnimCurveMetadataEditor, MetadataHost.Get(), PreviewScene.Pin().ToSharedRef(), OnObjectsSelected);
+}
+
+
 /////////////////////////////////////////////////////
 // FAnimationAssetBrowserSummoner
 
diff --git a/Engine/Source/Editor/Persona/Private/TabSpawners.h b/Engine/Source/Editor/Persona/Private/TabSpawners.h
--- a/Engine/Source/Editor/Persona/Private/TabSpawners.h
+++ b/Engine/Source/Editor/Persona/Private/TabSpawners.h
@@ -160,6 +160,31 @@ private:
 	FOnObjectsSelected OnObjectsSelected;
 };
 
+/////////////////////////////////////////////////////
+// FAnimCurveViewerTabSummoner
+
+struct FAnimCurveMetadataEditorTabSummoner : public FWorkflowTabFactory
+{
+public:
+	FAnimCurveMetadataEditorTabSummoner(TSharedPtr<class FAssetEditorToolkit> InHostingApp, UObject* InMetadataHost, const TSharedRef<IPersonaPreviewScene>& InPreviewScene, FOnObjectsSelected InOnObjectsSelected);
+
+	virtual TSharedRef<SWidget> CreateTabBody(const FWorkflowTabSpawnInfo& Info) const override;
+
+	virtual TSharedPtr<SToolTip> CreateTabToolTipWidget(const FWorkflowTabSpawnInfo& Info) const override
+	{
+		return  IDocumentation::Get()->CreateToolTip(LOCTEXT("AnimCurveMetadataEditorTooltip", "The Anim Curve Metadata Editor tab lets you add, remove, rename and edit curve metadata."), NULL, TEXT("Shared/Editors/Persona"), TEXT("AnimCurveMetadataEditor_Window"));
+	}
+
+	virtual FText GetTabToolTipText(const FWorkflowTabSpawnInfo& Info) const
+	{
+		return LOCTEXT("AnimCurveTabView_ToolTip", "Shows the animation curve debugger. This shows the state of animation curves.");
+	}
+
+private:
+	TWeakObjectPtr<UObject> MetadataHost;
+	TWeakPtr<class IPersonaPreviewScene> PreviewScene;
+	FOnObjectsSelected OnObjectsSelected;
+};
 
 /////////////////////////////////////////////////////
 // FAnimationAssetBrowserSummoner
diff --git a/Engine/Source/Editor/Persona/Public/AnimAssetFindReplace.h b/Engine/Source/Editor/Persona/Public/AnimAssetFindReplace.h
new file mode 100644
--- /dev/null
+++ b/Engine/Source/Editor/Persona/Public/AnimAssetFindReplace.h
@@ -0,0 +1,52 @@
+﻿// Copyright Epic Games, Inc. All Rights Reserved.
+
+#pragma once
+
+#include "CoreTypes.h"
+#include "AnimAssetFindReplace.generated.h"
+
+class SAnimAssetFindReplace;
+
+/** Context for toolbar */
+UCLASS()
+class UAnimAssetFindReplaceContext : public UObject
+{
+	GENERATED_BODY()
+
+public:
+	/** Weak ptr to the find/replace widget */
+	TWeakPtr<SAnimAssetFindReplace> Widget;
+};
+
+enum class EAnimAssetFindReplaceMode : int32
+{
+	Find,
+	Replace,
+};
+
+enum class EAnimAssetFindReplaceType : int32
+{
+	Curves,
+	Notifies,
+};
+
+/** Configuration for the find/replace tab */
+struct FAnimAssetFindReplaceConfig
+{
+	FAnimAssetFindReplaceConfig() = default;
+
+	/** The initial find string **/
+	FString FindString;
+
+	/** The initial replace string **/
+	FString ReplaceString;
+
+	/** The initial type */
+	EAnimAssetFindReplaceType Type = EAnimAssetFindReplaceType::Curves;
+
+	/** The initial mode */
+	EAnimAssetFindReplaceMode Mode = EAnimAssetFindReplaceMode::Find;
+
+	/** The initial skeleton to filter by */
+	FAssetData SkeletonFilter;
+};
\ No newline at end of file
diff --git a/Engine/Source/Editor/Persona/Public/IAnimSequenceCurveEditor.h b/Engine/Source/Editor/Persona/Public/IAnimSequenceCurveEditor.h
--- a/Engine/Source/Editor/Persona/Public/IAnimSequenceCurveEditor.h
+++ b/Engine/Source/Editor/Persona/Public/IAnimSequenceCurveEditor.h
@@ -17,10 +17,16 @@ public:
 	virtual void ResetCurves() = 0;
 
 	/** Add a curve for editing */
-	virtual void AddCurve(const FText& InCurveDisplayName, const FLinearColor& InCurveColor, const FSmartName& InName, ERawCurveTrackTypes InType, int32 InCurveIndex, FSimpleDelegate InOnCurveModified) = 0;
+	virtual void AddCurve(const FText& InCurveDisplayName, const FLinearColor& InCurveColor, const FName& InName, ERawCurveTrackTypes InType, int32 InCurveIndex, FSimpleDelegate InOnCurveModified) = 0;
+
+	UE_DEPRECATED(5.3, "Please use AddCurve that takes an FName.")
+	virtual void AddCurve(const FText& InCurveDisplayName, const FLinearColor& InCurveColor, const FSmartName& InName, ERawCurveTrackTypes InType, int32 InCurveIndex, FSimpleDelegate InOnCurveModified) {}
 
 	/** Remove a curve for editing */
-	virtual void RemoveCurve(const FSmartName& InName, ERawCurveTrackTypes InType, int32 InCurveIndex) = 0;
+	virtual void RemoveCurve(const FName& InName, ERawCurveTrackTypes InType, int32 InCurveIndex) = 0;
+	
+	UE_DEPRECATED(5.3, "Please use AddCurve that takes an FName.")
+	virtual void RemoveCurve(const FSmartName& InName, ERawCurveTrackTypes InType, int32 InCurveIndex) {}
 
 	/** Zoom to fit all the curves */
 	virtual void ZoomToFit() = 0;
diff --git a/Engine/Source/Editor/Persona/Public/PersonaDelegates.h b/Engine/Source/Editor/Persona/Public/PersonaDelegates.h
--- a/Engine/Source/Editor/Persona/Public/PersonaDelegates.h
+++ b/Engine/Source/Editor/Persona/Public/PersonaDelegates.h
@@ -70,3 +70,9 @@ DECLARE_DELEGATE_RetVal_OneParam(TSharedRef<SWidget>, FOnExtendBlendSpaceSampleT
 
 // Called to let the blendspace editor UI set the preview position of an external blendspace node
 DECLARE_DELEGATE_OneParam(FOnSetBlendSpacePreviewPosition, FVector /*InBlendSample*/);
+
+// Called by the curve picker to indicate that a curve was chosen by the user
+DECLARE_DELEGATE_OneParam(FOnCurvePicked, const FName& /*PickedName*/);
+
+// Called by the curve picker to filter out curve names that are displayed to the user
+DECLARE_DELEGATE_RetVal_OneParam(bool, FIsCurveNameMarkedForExclusion, const FName& /*CurveName*/);
\ No newline at end of file
