Fix non-PCH build error

diff --git a/Engine/Source/Editor/Persona/Public/AnimAssetFindReplace.h b/Engine/Source/Editor/Persona/Public/AnimAssetFindReplace.h
--- a/Engine/Source/Editor/Persona/Public/AnimAssetFindReplace.h
+++ b/Engine/Source/Editor/Persona/Public/AnimAssetFindReplace.h
@@ -2,7 +2,8 @@
 
 #pragma once
 
-#include "CoreTypes.h"
+#include "CoreMinimal.h"
+#include "AssetRegistry/AssetData.h"
 #include "AnimAssetFindReplace.generated.h"
 
 class SAnimAssetFindReplace;
