Fixup conflicting duplicate LOCTEXT keys

diff --git a/Engine/Plugins/Animation/AnimationData/Source/AnimationData/Private/AnimSequencerController.cpp b/Engine/Plugins/Animation/AnimationData/Source/AnimationData/Private/AnimSequencerController.cpp
--- a/Engine/Plugins/Animation/AnimationData/Source/AnimationData/Private/AnimSequencerController.cpp
+++ b/Engine/Plugins/Animation/AnimationData/Source/AnimationData/Private/AnimSequencerController.cpp
@@ -603,7 +603,7 @@ bool UAnimSequencerController::DuplicateCurve(const FAnimationCurveIdentifier& C
 	}
 	else
 	{
-		ReportWarning(LOCTEXT("InvalidCurveIdentifierWarning", "Invalid curve identifier provided"));
+		ReportWarningf(LOCTEXT("DuplicateInvalidCurveIdentifierWarning", "Invalid curve identifier provided for duplicate: copy: {0}, new: {1}"), FText::FromName(CopyCurveId.CurveName), FText::FromName(NewCurveId.CurveName));
 	}
 
 	return false;
diff --git a/Engine/Source/Developer/AnimationDataController/Private/AnimDataController.cpp b/Engine/Source/Developer/AnimationDataController/Private/AnimDataController.cpp
--- a/Engine/Source/Developer/AnimationDataController/Private/AnimDataController.cpp
+++ b/Engine/Source/Developer/AnimationDataController/Private/AnimDataController.cpp
@@ -507,7 +507,7 @@ bool UAnimDataController::DuplicateCurve(const FAnimationCurveIdentifier& CopyCu
 	}
 	else
 	{
-		ReportWarning(LOCTEXT("InvalidCurveIdentifierWarning", "Invalid curve identifier provided"));
+		ReportWarningf(LOCTEXT("DuplicateInvalidCurveIdentifierWarning", "Invalid curve identifier provided for duplicate: copy: {0}, new: {1}"), FText::FromName(CopyCurveId.CurveName), FText::FromName(NewCurveId.CurveName));
 	}
 
 	return false;
diff --git a/Engine/Source/Editor/AnimGraph/Private/AnimGraphNode_ModifyCurve.cpp b/Engine/Source/Editor/AnimGraph/Private/AnimGraphNode_ModifyCurve.cpp
--- a/Engine/Source/Editor/AnimGraph/Private/AnimGraphNode_ModifyCurve.cpp
+++ b/Engine/Source/Editor/AnimGraph/Private/AnimGraphNode_ModifyCurve.cpp
@@ -48,7 +48,7 @@ void UAnimGraphNode_ModifyCurve::GetAddCurveMenuActions(FMenuBuilder& MenuBuilde
 
 			TSharedRef<STextEntryPopup> TextEntry =
 				SNew(STextEntryPopup)
-				.Label(LOCTEXT("NewCurveLabel", "New Curve Name"))
+				.Label(LOCTEXT("NewCurvePopupLabel", "New Curve Name"))
 				.OnTextCommitted_Lambda([this](FText InText, ETextCommit::Type InCommitType)
 				{
 					FSlateApplication::Get().DismissAllMenus();
diff --git a/Engine/Source/Editor/Persona/Private/SAnimCurveMetadataEditor.cpp b/Engine/Source/Editor/Persona/Private/SAnimCurveMetadataEditor.cpp
--- a/Engine/Source/Editor/Persona/Private/SAnimCurveMetadataEditor.cpp
+++ b/Engine/Source/Editor/Persona/Private/SAnimCurveMetadataEditor.cpp
@@ -504,7 +504,7 @@ void SAnimCurveMetadataEditor::Construct(const FArguments& InArgs, UObject* InAn
 				
 				+SHeaderRow::Column(CurveMetadataEditorColumns::AnimCurveMaxLODLabel)
 				.FillWidth(0.17f)
-				.DefaultLabel(LOCTEXT("AnimCurveNumBoneLabel", "Max LOD"))	
+				.DefaultLabel(LOCTEXT("AnimCurveMaxLODLabel", "Max LOD"))	
 			)
 		]
 	];
diff --git a/Engine/Source/Editor/Persona/Private/SAnimCurveViewer.cpp b/Engine/Source/Editor/Persona/Private/SAnimCurveViewer.cpp
--- a/Engine/Source/Editor/Persona/Private/SAnimCurveViewer.cpp
+++ b/Engine/Source/Editor/Persona/Private/SAnimCurveViewer.cpp
@@ -893,7 +893,7 @@ TSharedRef<SWidget> SAnimCurveViewer::CreateCurveSourceSelector()
 		SAssignNew(PoseWatchCombo, SComboBox<TSharedPtr<TWeakObjectPtr<UPoseWatchPoseElement>>>)
 		.ToolTipText_Lambda([this]()
 		{
-			return PoseWatch.Get() != nullptr && PoseWatch->GetParent() != nullptr ? FText::Format(LOCTEXT("PoseWatchTooltipFormat", "Viewing Pose Watch Curves for '{0}'"), PoseWatch->GetParent()->GetLabel()) : LOCTEXT("NoPoseWatchSelected", "Viewing output Curves");
+			return PoseWatch.Get() != nullptr && PoseWatch->GetParent() != nullptr ? FText::Format(LOCTEXT("PoseWatchTooltipFormat", "Viewing Pose Watch Curves for '{0}'"), PoseWatch->GetParent()->GetLabel()) : LOCTEXT("PoseWatchViewingOutputCurves", "Viewing output Curves");
 		})
 		.OptionsSource(&PoseWatches)
 		.OnGenerateWidget_Lambda([this](TSharedPtr<TWeakObjectPtr<UPoseWatchPoseElement>> InElement)
@@ -934,7 +934,7 @@ TSharedRef<SWidget> SAnimCurveViewer::CreateCurveSourceSelector()
 					SNew(STextBlock)
 					.Text_Lambda([Element]()
 					{
-						return Element.Get() != nullptr && Element->GetParent() != nullptr ? Element->GetParent()->GetLabel() : LOCTEXT("NoPoseWatchSelected", "Output Curves");
+						return Element.Get() != nullptr && Element->GetParent() != nullptr ? Element->GetParent()->GetLabel() : LOCTEXT("PoseWatchOutputCurves", "Output Curves");
 					})
 				];
 		})
@@ -1004,7 +1004,7 @@ TSharedRef<SWidget> SAnimCurveViewer::CreateCurveSourceSelector()
 				SNew(STextBlock)
 				.Text_Lambda([this]()
 				{
-					return PoseWatch.Get() != nullptr && PoseWatch->GetParent() != nullptr ? PoseWatch->GetParent()->GetLabel() : LOCTEXT("NoPoseWatchSelected", "Output Curves");
+					return PoseWatch.Get() != nullptr && PoseWatch->GetParent() != nullptr ? PoseWatch->GetParent()->GetLabel() : LOCTEXT("PoseWatchOutputCurves", "Output Curves");
 				})
 			]
 		];
