Add option to hide SnapToNotifies UI command

diff --git a/Engine/Source/Editor/Persona/Private/AnimTimeline/SAnimTimeline.cpp b/Engine/Source/Editor/Persona/Private/AnimTimeline/SAnimTimeline.cpp
--- a/Engine/Source/Editor/Persona/Private/AnimTimeline/SAnimTimeline.cpp
+++ b/Engine/Source/Editor/Persona/Private/AnimTimeline/SAnimTimeline.cpp
@@ -375,7 +375,10 @@ FReply SAnimTimeline::OnMouseButtonUp(const FGeometry& MyGeometry, const FPointe
 		MenuBuilder.BeginSection("SnapOptions", LOCTEXT("SnapOptions", "Snapping"));
 		{
 			MenuBuilder.AddMenuEntry(FAnimSequenceTimelineCommands::Get().SnapToFrames);
-			MenuBuilder.AddMenuEntry(FAnimSequenceTimelineCommands::Get().SnapToNotifies);
+			if (GetDefault<UPersonaOptions>()->bExposeNotifiesUICommands)
+			{
+				MenuBuilder.AddMenuEntry(FAnimSequenceTimelineCommands::Get().SnapToNotifies);
+			}
 			MenuBuilder.AddMenuEntry(FAnimSequenceTimelineCommands::Get().SnapToCompositeSegments);
 			MenuBuilder.AddMenuEntry(FAnimSequenceTimelineCommands::Get().SnapToMontageSections);
 		}
diff --git a/Engine/Source/Editor/UnrealEd/Classes/Preferences/PersonaOptions.h b/Engine/Source/Editor/UnrealEd/Classes/Preferences/PersonaOptions.h
--- a/Engine/Source/Editor/UnrealEd/Classes/Preferences/PersonaOptions.h
+++ b/Engine/Source/Editor/UnrealEd/Classes/Preferences/PersonaOptions.h
@@ -288,6 +288,8 @@ public:
 
 	bool bExposeClothingSceneElementMenu = true;
 
+	bool bExposeNotifiesUICommands = true;
+
 	DECLARE_DELEGATE_RetVal_OneParam(bool, FOnIsClassAllowed, const UClass* /*InClass*/);
 
 	/** Sets a delegate that allows external code to restrict which features can be used within the Persona editor by filtering which classes are allowed. */
