Improved Skeleton reference pose drawing behaviour
- Allow for drawing sub-set of bones (through selection)
- Bone/socket details now pull correct data from reference skeleton

diff --git a/Engine/Source/Editor/Persona/Private/AnimationEditorPreviewScene.cpp b/Engine/Source/Editor/Persona/Private/AnimationEditorPreviewScene.cpp
--- a/Engine/Source/Editor/Persona/Private/AnimationEditorPreviewScene.cpp
+++ b/Engine/Source/Editor/Persona/Private/AnimationEditorPreviewScene.cpp
@@ -858,23 +858,15 @@ void FAnimationEditorPreviewScene::SetSelectedBone(const FName& BoneName, ESelec
 		ClearSelectedSocket();
 		ClearSelectedActor();
 
-		// Add in bone of interest only if we have a preview instance set-up
-		if (SkeletalMeshComponent->PreviewInstance != nullptr)
-		{
-			// need to get mesh bone base since BonesOfInterest is saved in SkeletalMeshComponent
-			// and it is used by renderer. It is not Skeleton base
-			const int32 MeshBoneIndex = SkeletalMeshComponent->GetBoneIndex(BoneName);
-
-			if (MeshBoneIndex != INDEX_NONE)
-			{
-				SelectedBoneIndex = MeshBoneIndex;
-				SkeletalMeshComponent->BonesOfInterest.Add(SelectedBoneIndex);
-			}
+		// need to get mesh bone base since BonesOfInterest is saved in SkeletalMeshComponent
+		// and it is used by renderer. It is not Skeleton base
+		const int32 MeshBoneIndex = SkeletalMeshComponent->GetBoneIndex(BoneName);
+		SelectedBoneIndex = MeshBoneIndex != INDEX_NONE ? MeshBoneIndex : BoneIndex;
+		SkeletalMeshComponent->BonesOfInterest.Add(SelectedBoneIndex);
 
-			InvalidateViews();
+		InvalidateViews();
 
-			OnSelectedBoneChanged.Broadcast(BoneName, InSelectInfo);
-		}
+		OnSelectedBoneChanged.Broadcast(BoneName, InSelectInfo);
 	}
 }
 
diff --git a/Engine/Source/Editor/Persona/Private/AnimationEditorViewportClient.cpp b/Engine/Source/Editor/Persona/Private/AnimationEditorViewportClient.cpp
--- a/Engine/Source/Editor/Persona/Private/AnimationEditorViewportClient.cpp
+++ b/Engine/Source/Editor/Persona/Private/AnimationEditorViewportClient.cpp
@@ -592,11 +592,11 @@ void FAnimationViewportClient::Draw(const FSceneView* View, FPrimitiveDrawInterf
 				DrawAttributes(PreviewMeshComponent, PDI);
 			}
 		}
-		else if (bValidComponent && !bValidSkeletalMesh && GetBoneDrawMode() == EBoneDrawMode::All)
+		else if (bValidComponent && !bValidSkeletalMesh)
 		{
 			if (const USkeleton* Skeleton = GetPreviewScene()->GetPersonaToolkit()->GetSkeleton())
 			{
-				DrawBonesFromSkeleton(Skeleton, PDI);
+				DrawBonesFromSkeleton(Skeleton, PreviewMeshComponent->BonesOfInterest, PDI);
 			}
 		}
 	}
@@ -1652,7 +1652,7 @@ void FAnimationViewportClient::DrawMeshBonesBakedAnimation(UDebugSkelMeshCompone
 	}
 }
 
-void FAnimationViewportClient::DrawBonesFromSkeleton(const USkeleton* Skeleton, FPrimitiveDrawInterface* PDI) const
+void FAnimationViewportClient::DrawBonesFromSkeleton(const USkeleton* Skeleton, const TArray<int32>& InSelectedBones,FPrimitiveDrawInterface* PDI) const
 {
 	check(Skeleton);
 
@@ -1688,15 +1688,20 @@ void FAnimationViewportClient::DrawBonesFromSkeleton(const USkeleton* Skeleton,
 		BoneColours[BoneIndex] = PersonaOptions->DefaultBoneColor;
 	}
 
-	constexpr bool bForceDraw = true;
-	constexpr bool bAddHitProxy = false;
-	const TArray<int32> BonesOfInterest;
+	// color virtual bones
+	for (const int16 VirtualBoneIndex : RefSkeleton.GetRequiredVirtualBones())
+	{
+		BoneColours[VirtualBoneIndex] = PersonaOptions->VirtualBoneColor;
+	}
+
+	constexpr bool bForceDraw = false;
+	constexpr bool bAddHitProxy = true;
 	DrawBones(
 		FVector::ZeroVector,
 		RequiredBones,
 		RefSkeleton,
 		WorldTransforms,
-		BonesOfInterest,
+		InSelectedBones,
 		BoneColours,
 		PDI,
 		bForceDraw,
diff --git a/Engine/Source/Editor/Persona/Private/EditModes/SkeletonSelectionEditMode.cpp b/Engine/Source/Editor/Persona/Private/EditModes/SkeletonSelectionEditMode.cpp
--- a/Engine/Source/Editor/Persona/Private/EditModes/SkeletonSelectionEditMode.cpp
+++ b/Engine/Source/Editor/Persona/Private/EditModes/SkeletonSelectionEditMode.cpp
@@ -32,20 +32,26 @@ FSkeletonSelectionEditMode::FSkeletonSelectionEditMode()
 	bDrawGrid = false;
 }
 
+const FReferenceSkeleton& FSkeletonSelectionEditMode::GetReferenceSkeletonForComponent(const USkeletalMeshComponent* Component) const
+{
+	check(Component);
+	return Component->GetSkeletalMeshAsset() && !Component->GetSkeletalMeshAsset()->IsCompiling() ? Component->GetSkeletalMeshAsset()->GetRefSkeleton() : GetAnimPreviewScene().GetPersonaToolkit()->GetSkeleton()->GetReferenceSkeleton();
+}
+
 bool FSkeletonSelectionEditMode::GetCameraTarget(FSphere& OutTarget) const
 {
 	bool bHandled = false;
 
-	UDebugSkelMeshComponent* PreviewMeshComponent = GetAnimPreviewScene().GetPreviewMeshComponent();
-
-	if(PreviewMeshComponent && PreviewMeshComponent->GetSkeletalMeshAsset())
+	const UDebugSkelMeshComponent* PreviewMeshComponent = GetAnimPreviewScene().GetPreviewMeshComponent();
+	if(PreviewMeshComponent)
 	{
 		if (GetAnimPreviewScene().GetSelectedBoneIndex() != INDEX_NONE)
 		{
 			const int32 FocusBoneIndex = GetAnimPreviewScene().GetSelectedBoneIndex();
-			if (FocusBoneIndex != INDEX_NONE && PreviewMeshComponent->GetSkeletalMeshAsset()->GetRefSkeleton().IsValidIndex(FocusBoneIndex))
+			const FReferenceSkeleton& ReferenceSkeleton = GetReferenceSkeletonForComponent(PreviewMeshComponent);
+			if (FocusBoneIndex != INDEX_NONE && ReferenceSkeleton.IsValidIndex(FocusBoneIndex))
 			{
-				const FName BoneName = PreviewMeshComponent->GetSkeletalMeshAsset()->GetRefSkeleton().GetBoneName(FocusBoneIndex);
+				const FName BoneName = ReferenceSkeleton.GetBoneName(FocusBoneIndex);
 				OutTarget.Center = PreviewMeshComponent->GetBoneLocation(BoneName);
 				OutTarget.W = 30.0f;
 				bHandled = true;
@@ -97,11 +103,13 @@ bool FSkeletonSelectionEditMode::StartTracking(FEditorViewportClient* InViewport
 	UDebugSkelMeshComponent* PreviewMeshComponent = GetAnimPreviewScene().GetPreviewMeshComponent();
 	if(PreviewMeshComponent != nullptr && PreviewMeshComponent->GetSkeletalMeshAsset() != nullptr)
 	{
-		int32 BoneIndex = GetAnimPreviewScene().GetSelectedBoneIndex();
-		USkeletalMeshSocket* SelectedSocket = GetAnimPreviewScene().GetSelectedSocket().Socket;
-		AActor* SelectedActor = GetAnimPreviewScene().GetSelectedActor();
+		const int32 BoneIndex = GetAnimPreviewScene().GetSelectedBoneIndex();
+		const USkeletalMeshSocket* SelectedSocket = GetAnimPreviewScene().GetSelectedSocket().Socket;
+		const AActor* SelectedActor = GetAnimPreviewScene().GetSelectedActor();
 
-		if ((BoneIndex >= 0 && PreviewMeshComponent->GetSkeletalMeshAsset()->GetRefSkeleton().IsValidIndex(BoneIndex)) || SelectedSocket != nullptr || SelectedActor != nullptr)
+		// Retrieve reference skeleton from either current USkeletalMesh, or USkeleton if no mesh is set
+		const FReferenceSkeleton& ReferenceSkeleton = GetReferenceSkeletonForComponent(PreviewMeshComponent);
+		if ((BoneIndex >= 0 && ReferenceSkeleton.IsValidIndex(BoneIndex)) || SelectedSocket != nullptr || SelectedActor != nullptr)
 		{
 			bool bValidAxis = false;
 			FVector WorldAxisDir;
@@ -157,7 +165,7 @@ bool FSkeletonSelectionEditMode::StartTracking(FEditorViewportClient* InViewport
 						bInTransaction = true;
 
 						// now modify the bone array
-						const FName BoneName = PreviewMeshComponent->GetSkeletalMeshAsset()->GetRefSkeleton().GetBoneName(BoneIndex);
+						const FName BoneName = ReferenceSkeleton.GetBoneName(BoneIndex);
 						PreviewMeshComponent->PreviewInstance->ModifyBone(BoneName);
 					}
 				}
@@ -215,8 +223,10 @@ bool FSkeletonSelectionEditMode::InputDelta(FEditorViewportClient* InViewportCli
 
 		if ( BoneIndex >= 0 )
 		{
+			// Retrieve reference skeleton from either current USkeletalMesh, or USkeleton if no mesh is set
+			const FReferenceSkeleton& ReferenceSkeleton = GetReferenceSkeletonForComponent(PreviewMeshComponent);
+			const FName BoneName = ReferenceSkeleton.GetBoneName(BoneIndex);
 			//Get the skeleton control manipulating this bone
-			const FName BoneName = PreviewMeshComponent->GetSkeletalMeshAsset()->GetRefSkeleton().GetBoneName(BoneIndex);
 			SkelControl = &(PreviewMeshComponent->PreviewInstance->ModifyBone(BoneName));
 		}
 
@@ -339,12 +349,13 @@ FIntPoint FSkeletonSelectionEditMode::GetDPIUnscaledSize(FViewport* Viewport, FV
 void FSkeletonSelectionEditMode::DrawHUD(FEditorViewportClient* ViewportClient, FViewport* Viewport, const FSceneView* View, FCanvas* Canvas)
 {
 	const UDebugSkelMeshComponent* PreviewMeshComponent = GetAnimPreviewScene().GetPreviewMeshComponent();
-	if (PreviewMeshComponent == nullptr || PreviewMeshComponent->GetSkeletalMeshAsset() == nullptr || PreviewMeshComponent->GetSkeletalMeshAsset()->IsCompiling())
+	if (PreviewMeshComponent == nullptr)
 	{
 		return;
 	}
 
-	const FReferenceSkeleton& RefSkeleton = PreviewMeshComponent->GetSkeletalMeshAsset()->GetRefSkeleton();
+	// Retrieve reference skeleton from either current USkeletalMesh, or USkeleton if no mesh is set
+	const FReferenceSkeleton& RefSkeleton = PreviewMeshComponent->GetSkeletalMeshAsset() && !PreviewMeshComponent->GetSkeletalMeshAsset()->IsCompiling() ? PreviewMeshComponent->GetSkeletalMeshAsset()->GetRefSkeleton() : GetAnimPreviewScene().GetPersonaToolkit()->GetSkeleton()->GetReferenceSkeleton();
 	const int32 BoneIndex = GetAnimPreviewScene().GetSelectedBoneIndex();
 
 	// Draw name of selected bone
@@ -354,9 +365,9 @@ void FSkeletonSelectionEditMode::DrawHUD(FEditorViewportClient* ViewportClient,
 		const int32 HalfX = ViewPortSize.X / 2;
 		const int32 HalfY = ViewPortSize.Y / 2;
 
-		const FName BoneName = PreviewMeshComponent->GetSkeletalMeshAsset()->GetRefSkeleton().GetBoneName(BoneIndex);
+		const FName BoneName = RefSkeleton.GetBoneName(BoneIndex);
 
-		const FMatrix BoneMatrix = PreviewMeshComponent->GetBoneMatrix(BoneIndex);
+		const FMatrix BoneMatrix = GetBoneTransform(BoneIndex).ToMatrixNoScale();
 		const FPlane Proj = View->Project(BoneMatrix.GetOrigin());
 		if (Proj.W > 0.f)
 		{
@@ -374,8 +385,7 @@ void FSkeletonSelectionEditMode::DrawHUD(FEditorViewportClient* ViewportClient,
 	{
 		const USkeletalMeshSocket* Socket = GetAnimPreviewScene().GetSelectedSocket().Socket;
 
-		FMatrix SocketMatrix;
-		Socket->GetSocketMatrix(SocketMatrix, PreviewMeshComponent);
+		const FMatrix SocketMatrix = GetSocketTransform(Socket).ToMatrixNoScale();
 		const FVector SocketPos = SocketMatrix.GetOrigin();
 
 		const FPlane Proj = View->Project(SocketPos);
@@ -416,6 +426,10 @@ bool FSkeletonSelectionEditMode::IsSelectedBoneRequired() const
 			return LODData.RequiredBones.Find(static_cast<FBoneIndexType>(SelectedBoneIndex)) != INDEX_NONE;
 		}
 	}
+	else if (SelectedBoneIndex != INDEX_NONE && !PreviewMeshComponent->GetSkeletalMeshRenderData())
+	{
+		return true;
+	}
 
 	return false;
 }
@@ -441,29 +455,103 @@ bool FSkeletonSelectionEditMode::UsesTransformWidget(UE::Widget::EWidgetMode Che
 	return ShouldDrawWidget() && (CheckMode == UE::Widget::WM_Scale || CheckMode == UE::Widget::WM_Translate || CheckMode == UE::Widget::WM_Rotate);
 }
 
-bool FSkeletonSelectionEditMode::GetCustomDrawingCoordinateSystem(FMatrix& InMatrix, void* InData)
+FTransform FSkeletonSelectionEditMode::GetWorldSpaceBoneTransform(const FReferenceSkeleton& ReferenceSkeleton, const int32 BoneIndex) const
 {
-	UDebugSkelMeshComponent* PreviewMeshComponent = GetAnimPreviewScene().GetPreviewMeshComponent();
-	if(PreviewMeshComponent && PreviewMeshComponent->GetSkeletalMeshAsset() && !PreviewMeshComponent->GetSkeletalMeshAsset()->IsCompiling())
+	const TArray<FTransform>& BonePoses = ReferenceSkeleton.GetRefBonePose();
+
+	if (BonePoses.IsValidIndex(BoneIndex))
 	{
-		int32 BoneIndex = GetAnimPreviewScene().GetSelectedBoneIndex();
-		if (BoneIndex != INDEX_NONE && PreviewMeshComponent->GetSkeletalMeshAsset()->GetRefSkeleton().IsValidIndex(BoneIndex))
+		FTransform WorldSpacePose = BonePoses[BoneIndex];
+
+		int32 ParentIndex = ReferenceSkeleton.GetParentIndex(BoneIndex);
+
+		while(ParentIndex != INDEX_NONE)
 		{
-			FTransform BoneMatrix = PreviewMeshComponent->GetBoneTransform(BoneIndex);
+			WorldSpacePose = WorldSpacePose * BonePoses[ParentIndex];
+			ParentIndex = ReferenceSkeleton.GetParentIndex(ParentIndex);
+		}
+
+		return WorldSpacePose;
+	}
+
+	return FTransform::Identity;
+}
 
+
+
+FTransform FSkeletonSelectionEditMode::GetBoneTransform(const int32 BoneIndex) const
+{
+	const UDebugSkelMeshComponent* PreviewMeshComponent = GetAnimPreviewScene().GetPreviewMeshComponent();
+	if (PreviewMeshComponent && PreviewMeshComponent->GetSkeletalMeshAsset())
+	{
+		if (!PreviewMeshComponent->GetSkeletalMeshAsset()->IsCompiling())
+		{
+			return PreviewMeshComponent->GetBoneTransform(BoneIndex);
+		}
+	}
+	else if (const USkeleton* Skeleton = GetAnimPreviewScene().GetPersonaToolkit()->GetSkeleton())
+	{
+		return GetWorldSpaceBoneTransform(Skeleton->GetReferenceSkeleton(), BoneIndex);
+	}
+
+	return FTransform::Identity;
+}
+
+FTransform FSkeletonSelectionEditMode::GetSocketTransform(const USkeletalMeshSocket* Socket) const
+{
+	const UDebugSkelMeshComponent* PreviewMeshComponent = GetAnimPreviewScene().GetPreviewMeshComponent();
+	if (PreviewMeshComponent && PreviewMeshComponent->GetSkeletalMeshAsset())
+	{
+		if (!PreviewMeshComponent->GetSkeletalMeshAsset()->IsCompiling())
+		{
+			const int32 BoneIndex = PreviewMeshComponent->GetBoneIndex(Socket->BoneName);
+			if(BoneIndex != INDEX_NONE)
+			{
+				const FTransform BoneTM = PreviewMeshComponent->GetBoneTransform(BoneIndex);
+				const FTransform RelSocketTM( Socket->RelativeRotation, Socket->RelativeLocation, Socket->RelativeScale );
+				return RelSocketTM * BoneTM;
+			}
+			
+		}
+	}
+	else if (const USkeleton* Skeleton = GetAnimPreviewScene().GetPersonaToolkit()->GetSkeleton())
+	{
+		const FReferenceSkeleton& ReferenceSkeleton = Skeleton->GetReferenceSkeleton();
+		const int32 BoneIndex = ReferenceSkeleton.FindBoneIndex(Socket->BoneName);
+		if(BoneIndex != INDEX_NONE)
+		{
+			const FTransform BoneTM = GetWorldSpaceBoneTransform(Skeleton->GetReferenceSkeleton(), BoneIndex);
+			const FTransform RelSocketTM( Socket->RelativeRotation, Socket->RelativeLocation, Socket->RelativeScale );
+			return RelSocketTM * BoneTM;
+		}
+	}
+
+	return FTransform::Identity;
+}
+
+bool FSkeletonSelectionEditMode::GetCustomDrawingCoordinateSystem(FMatrix& InMatrix, void* InData)
+{
+	const UDebugSkelMeshComponent* PreviewMeshComponent = GetAnimPreviewScene().GetPreviewMeshComponent();
+	if(PreviewMeshComponent)
+	{
+		// Retrieve reference skeleton from either current USkeletalMesh, or USkeleton if no mesh is set
+		const FReferenceSkeleton& ReferenceSkeleton = GetReferenceSkeletonForComponent(PreviewMeshComponent);
+		
+		const int32 BoneIndex = GetAnimPreviewScene().GetSelectedBoneIndex();
+		if (BoneIndex != INDEX_NONE && ReferenceSkeleton.IsValidIndex(BoneIndex))
+		{
+			const FTransform BoneMatrix = GetBoneTransform(BoneIndex);
 			InMatrix = BoneMatrix.ToMatrixNoScale().RemoveTranslation();
 			return true;
 		}
 		else if (GetAnimPreviewScene().GetSelectedSocket().IsValid())
 		{
-			USkeletalMeshSocket* Socket = GetAnimPreviewScene().GetSelectedSocket().Socket;
-
-			FTransform SocketMatrix = Socket->GetSocketTransform(PreviewMeshComponent);
-
+			const USkeletalMeshSocket* Socket = GetAnimPreviewScene().GetSelectedSocket().Socket;
+			const FTransform SocketMatrix = GetSocketTransform(Socket);
 			InMatrix = SocketMatrix.ToMatrixNoScale().RemoveTranslation();
 			return true;
 		}
-		else if (AActor* SelectedActor = GetAnimPreviewScene().GetSelectedActor())
+		else if (const AActor* SelectedActor = GetAnimPreviewScene().GetSelectedActor())
 		{
 			InMatrix = SelectedActor->GetTransform().ToMatrixNoScale().RemoveTranslation();
 			return true;
@@ -480,37 +568,19 @@ bool FSkeletonSelectionEditMode::GetCustomInputCoordinateSystem(FMatrix& InMatri
 
 FVector FSkeletonSelectionEditMode::GetWidgetLocation() const
 {
-	UDebugSkelMeshComponent* PreviewMeshComponent = GetAnimPreviewScene().GetPreviewMeshComponent();
-	USkeletalMesh* SkeletalMesh = PreviewMeshComponent ? PreviewMeshComponent->GetSkeletalMeshAsset() : nullptr;
-	const bool bAssetIsCompiling = SkeletalMesh ? SkeletalMesh->IsCompiling() : false;
-
-	int32 BoneIndex = GetAnimPreviewScene().GetSelectedBoneIndex();
+	const int32 BoneIndex = GetAnimPreviewScene().GetSelectedBoneIndex();
 	if (BoneIndex != INDEX_NONE)
 	{
-		if (!bAssetIsCompiling
-			&& SkeletalMesh
-			&& SkeletalMesh->GetRefSkeleton().IsValidIndex(BoneIndex))
-		{
-			const FName BoneName = SkeletalMesh->GetRefSkeleton().GetBoneName(BoneIndex);
-
-			const FMatrix BoneMatrix = PreviewMeshComponent->GetBoneMatrix(BoneIndex);
-
-			return BoneMatrix.GetOrigin();
-		}		
+		const FMatrix BoneMatrix = GetBoneTransform(BoneIndex).ToMatrixNoScale();
+		return BoneMatrix.GetOrigin();
 	}
 	else if (GetAnimPreviewScene().GetSelectedSocket().IsValid())
 	{
-		if (!bAssetIsCompiling)
-		{
-			USkeletalMeshSocket* Socket = GetAnimPreviewScene().GetSelectedSocket().Socket;
-
-			FMatrix SocketMatrix;
-			Socket->GetSocketMatrix(SocketMatrix, PreviewMeshComponent);
-
-			return SocketMatrix.GetOrigin();
-		}
+		const USkeletalMeshSocket* Socket = GetAnimPreviewScene().GetSelectedSocket().Socket;
+		const FMatrix SocketMatrix = GetSocketTransform(Socket).ToMatrixNoScale();
+		return SocketMatrix.GetOrigin();
 	}
-	else if (AActor* SelectedActor = GetAnimPreviewScene().GetSelectedActor())
+	else if (const AActor* SelectedActor = GetAnimPreviewScene().GetSelectedActor())
 	{
 		return SelectedActor->GetActorLocation();
 	}
diff --git a/Engine/Source/Editor/Persona/Private/EditModes/SkeletonSelectionEditMode.h b/Engine/Source/Editor/Persona/Private/EditModes/SkeletonSelectionEditMode.h
--- a/Engine/Source/Editor/Persona/Private/EditModes/SkeletonSelectionEditMode.h
+++ b/Engine/Source/Editor/Persona/Private/EditModes/SkeletonSelectionEditMode.h
@@ -40,6 +40,12 @@ public:
 	virtual bool HandleClick(FEditorViewportClient* InViewportClient, HHitProxy *HitProxy, const FViewportClick &Click) override;
 	virtual bool CanCycleWidgetMode() const override;
 
+protected:
+	FTransform GetBoneTransform(const int32 BoneIndex) const;
+	FTransform GetSocketTransform(const USkeletalMeshSocket* Socket) const;
+	FTransform GetWorldSpaceBoneTransform(const FReferenceSkeleton& ReferenceSkeleton, const int32 BoneIndex) const;
+	const FReferenceSkeleton& GetReferenceSkeletonForComponent(const USkeletalMeshComponent* Component) const;
+
 private:
 	/** Duplicates and selects a socket when we alt-drag */
 	FSelectedSocketInfo DuplicateAndSelectSocket(const FSelectedSocketInfo& SocketInfoToDuplicate);
diff --git a/Engine/Source/Editor/Persona/Public/AnimationEditorViewportClient.h b/Engine/Source/Editor/Persona/Public/AnimationEditorViewportClient.h
--- a/Engine/Source/Editor/Persona/Public/AnimationEditorViewportClient.h
+++ b/Engine/Source/Editor/Persona/Public/AnimationEditorViewportClient.h
@@ -419,7 +419,7 @@ private:
 	/** Draw Bones for non retargeted animation. */
 	void DrawMeshBonesBakedAnimation(UDebugSkelMeshComponent * MeshComponent, FPrimitiveDrawInterface* PDI) const;
 	/** Draw Bones from skeleton reference pose. */
-	void DrawBonesFromSkeleton(const USkeleton* Skeleton, FPrimitiveDrawInterface* PDI) const;
+	void DrawBonesFromSkeleton(const USkeleton* Skeleton, const TArray<int32>& InSelectedBones, FPrimitiveDrawInterface* PDI) const;
 	/** Draws Bones for RequiredBones with WorldTransform **/
 	void DrawBones(
 		const FVector& ComponentOrigin,
diff --git a/Engine/Source/Editor/SkeletonEditor/Private/BoneProxy.cpp b/Engine/Source/Editor/SkeletonEditor/Private/BoneProxy.cpp
--- a/Engine/Source/Editor/SkeletonEditor/Private/BoneProxy.cpp
+++ b/Engine/Source/Editor/SkeletonEditor/Private/BoneProxy.cpp
@@ -4,6 +4,7 @@
 #include "IPersonaPreviewScene.h"
 #include "BoneControllers/AnimNode_ModifyBone.h"
 #include "AnimPreviewInstance.h"
+#include "IPersonaToolkit.h"
 #include "SAdvancedTransformInputBox.h"
 #include "ScopedTransaction.h"
 #include "Animation/DebugSkelMeshComponent.h"
@@ -22,6 +23,28 @@ UBoneProxy::UBoneProxy()
 {
 }
 
+FTransform GetWorldSpaceBoneTransform(const FReferenceSkeleton& ReferenceSkeleton, const int32 BoneIndex)
+{
+	const TArray<FTransform>& BonePoses = ReferenceSkeleton.GetRefBonePose();
+
+	if (BonePoses.IsValidIndex(BoneIndex))
+	{
+		FTransform WorldSpacePose = BonePoses[BoneIndex];
+
+		int32 ParentIndex = ReferenceSkeleton.GetParentIndex(BoneIndex);
+
+		while(ParentIndex != INDEX_NONE)
+		{
+			WorldSpacePose = WorldSpacePose * BonePoses[ParentIndex];
+			ParentIndex = ReferenceSkeleton.GetParentIndex(ParentIndex);
+		}
+
+		return WorldSpacePose;		
+	}
+
+	return FTransform::Identity;
+}
+
 void UBoneProxy::Tick(float DeltaTime)
 {
 	if (!bManipulating)
@@ -34,55 +57,101 @@ void UBoneProxy::Tick(float DeltaTime)
 				return;
 			}
 
-			TArray<FTransform> LocalBoneTransforms = Component->GetBoneSpaceTransforms();
-
-			int32 BoneIndex = Component->GetBoneIndex(BoneName);
-			if (LocalBoneTransforms.IsValidIndex(BoneIndex))
+			if (Component->GetSkeletalMeshAsset())
 			{
-				FTransform LocalTransform = LocalBoneTransforms[BoneIndex];
-				FTransform BoneTransform = Component->GetBoneTransform(BoneIndex);
+				TArray<FTransform> LocalBoneTransforms = Component->GetBoneSpaceTransforms();
 
-				if (bLocalLocation)
-				{
-					Location = LocalTransform.GetLocation();
-				}
-				else
+				const int32 BoneIndex = Component->GetBoneIndex(BoneName);
+				if (LocalBoneTransforms.IsValidIndex(BoneIndex))
 				{
-					Location = BoneTransform.GetLocation();
-				}
+					const FTransform& LocalTransform = LocalBoneTransforms[BoneIndex];
+					const FTransform BoneTransform = Component->GetBoneTransform(BoneIndex);
 
-				if (bLocalRotation)
-				{
-					Rotation = LocalTransform.GetRotation().Rotator();
-				}
-				else
-				{
-					Rotation = BoneTransform.GetRotation().Rotator();
-				}
+					if (bLocalLocation)
+					{
+						Location = LocalTransform.GetLocation();
+					}
+					else
+					{
+						Location = BoneTransform.GetLocation();
+					}
 
-				if(bLocalScale)
-				{
-					Scale = LocalTransform.GetScale3D();
-				}
-				else
-				{
-					Scale = BoneTransform.GetScale3D();
+					if (bLocalRotation)
+					{
+						Rotation = LocalTransform.GetRotation().Rotator();
+					}
+					else
+					{
+						Rotation = BoneTransform.GetRotation().Rotator();
+					}
+
+					if(bLocalScale)
+					{
+						Scale = LocalTransform.GetScale3D();
+					}
+					else
+					{
+						Scale = BoneTransform.GetScale3D();
+					}
+
+					const FTransform ReferenceTransform = Component->GetSkeletalMeshAsset()->GetRefSkeleton().GetRefBonePose()[BoneIndex];
+					ReferenceLocation = ReferenceTransform.GetLocation();
+					ReferenceRotation = ReferenceTransform.GetRotation().Rotator();
+					ReferenceScale = ReferenceTransform.GetScale3D();
 				}
 
-				FTransform ReferenceTransform = Component->GetSkeletalMeshAsset()->GetRefSkeleton().GetRefBonePose()[BoneIndex];
-				ReferenceLocation = ReferenceTransform.GetLocation();
-				ReferenceRotation = ReferenceTransform.GetRotation().Rotator();
-				ReferenceScale = ReferenceTransform.GetScale3D();
+				// Show mesh relative transform on the details panel so we have a way to visualize the root transform when processing root motion
+				// Note that this doesn't always represent the actual transform of the root in the animation at current time but where root motion has taken us so far
+				// It will not match the root transform at the current time in the animation after lopping multiple times if we are using ProcessRootMotion::Loop
+				// or if we are visualizing a complex section from a montage, for example
+				const FTransform MeshRelativeTransform = Component->GetRelativeTransform();
+				MeshLocation = MeshRelativeTransform.GetLocation();
+				MeshRotation = MeshRelativeTransform.GetRotation().Rotator();
+				MeshScale = MeshRelativeTransform.GetScale3D();
 			}
+			else if (const TSharedPtr<IPersonaPreviewScene> PreviewScene = WeakPreviewScene.Pin())
+			{
+				if (USkeleton* Skeleton = PreviewScene->GetPersonaToolkit()->GetSkeleton())
+				{
+					const int32 BoneIndex = Skeleton->GetReferenceSkeleton().FindBoneIndex(BoneName);
+					if (BoneIndex != INDEX_NONE)
+					{
+						const FTransform& ReferenceTransform = Skeleton->GetReferenceSkeleton().GetRefBonePose()[BoneIndex];
+						const FTransform WorldReferenceTransform = GetWorldSpaceBoneTransform(Skeleton->GetReferenceSkeleton(), BoneIndex);
+					
+						if (bLocalLocation)
+						{
+							Location = ReferenceTransform.GetLocation();
+						}
+						else
+						{
+							Location = WorldReferenceTransform.GetLocation();
+						}
+
+						if (bLocalRotation)
+						{
+							Rotation = ReferenceTransform.GetRotation().Rotator();
+						}
+						else
+						{
+							Rotation = WorldReferenceTransform.GetRotation().Rotator();
+						}
+
+						if(bLocalScale)
+						{
+							Scale = ReferenceTransform.GetScale3D();
+						}
+						else
+						{
+							Scale = WorldReferenceTransform.GetScale3D();
+						}
 
-			// Show mesh relative transform on the details panel so we have a way to visualize the root transform when processing root motion
-			// Note that this doesn't always represent the actual transform of the root in the animation at current time but where root motion has taken us so far
-			// It will not match the root transform at the current time in the animation after lopping multiple times if we are using ProcessRootMotion::Loop
-			// or if we are visualizing a complex section from a montage, for example
-			const FTransform MeshRelativeTransform = Component->GetRelativeTransform();
-			MeshLocation = MeshRelativeTransform.GetLocation();
-			MeshRotation = MeshRelativeTransform.GetRotation().Rotator();
-			MeshScale = MeshRelativeTransform.GetScale3D();
+						ReferenceLocation = ReferenceTransform.GetLocation();
+						ReferenceRotation = ReferenceTransform.GetRotation().Rotator();
+						ReferenceScale = ReferenceTransform.GetScale3D();
+					}
+				}
+			}
 		}
 	}
 }
diff --git a/Engine/Source/Editor/SkeletonEditor/Private/SkeletonTreeBoneItem.cpp b/Engine/Source/Editor/SkeletonEditor/Private/SkeletonTreeBoneItem.cpp
--- a/Engine/Source/Editor/SkeletonEditor/Private/SkeletonTreeBoneItem.cpp
+++ b/Engine/Source/Editor/SkeletonEditor/Private/SkeletonTreeBoneItem.cpp
@@ -38,10 +38,11 @@ FSkeletonTreeBoneItem::FSkeletonTreeBoneItem(const FName& InBoneName, const TSha
 	BoneProxy = NewObject<UBoneProxy>(GetTransientPackage(), *(BoneProxyPrefix + FString::Printf(TEXT("%p"), &InSkeletonTree.Get()) + InBoneName.ToString()));
 	BoneProxy->SetFlags(RF_Transactional);
 	BoneProxy->BoneName = InBoneName;
-	TSharedPtr<IPersonaPreviewScene> PreviewScene = InSkeletonTree->GetPreviewScene();
+	const TSharedPtr<IPersonaPreviewScene> PreviewScene = InSkeletonTree->GetPreviewScene();
 	if (PreviewScene.IsValid())
 	{
 		BoneProxy->SkelMeshComponent = PreviewScene->GetPreviewMeshComponent();
+		BoneProxy->WeakPreviewScene = PreviewScene.ToWeakPtr();
 	}
 }
 
@@ -337,7 +338,7 @@ FText FSkeletonTreeBoneItem::GetBoneToolTip()
 	{
 		UDebugSkelMeshComponent* PreviewComponent = GetSkeletonTree()->GetPreviewScene()->GetPreviewMeshComponent();
 
-		if (PreviewComponent)
+		if (PreviewComponent && PreviewComponent->GetSkeletalMeshAsset())
 		{
 			bMeshExists = true;
 
diff --git a/Engine/Source/Editor/SkeletonEditor/Private/SkeletonTreeVirtualBoneItem.cpp b/Engine/Source/Editor/SkeletonEditor/Private/SkeletonTreeVirtualBoneItem.cpp
--- a/Engine/Source/Editor/SkeletonEditor/Private/SkeletonTreeVirtualBoneItem.cpp
+++ b/Engine/Source/Editor/SkeletonEditor/Private/SkeletonTreeVirtualBoneItem.cpp
@@ -31,6 +31,7 @@ FSkeletonTreeVirtualBoneItem::FSkeletonTreeVirtualBoneItem(const FName& InBoneNa
 	if (PreviewScene.IsValid())
 	{
 		BoneProxy->SkelMeshComponent = PreviewScene->GetPreviewMeshComponent();
+		BoneProxy->WeakPreviewScene = PreviewScene.ToWeakPtr();
 	}
 }
 
diff --git a/Engine/Source/Editor/SkeletonEditor/Public/BoneProxy.h b/Engine/Source/Editor/SkeletonEditor/Public/BoneProxy.h
--- a/Engine/Source/Editor/SkeletonEditor/Public/BoneProxy.h
+++ b/Engine/Source/Editor/SkeletonEditor/Public/BoneProxy.h
@@ -7,6 +7,7 @@
 #include "TickableEditorObject.h"
 #include "UObject/Object.h"
 #include "UObject/WeakObjectPtr.h"
+#include "IPersonaPreviewScene.h"
 #if UE_ENABLE_INCLUDE_ORDER_DEPRECATED_IN_5_2
 #include "SAdvancedTransformInputBox.h"
 #endif
@@ -92,6 +93,8 @@ public:
 	UPROPERTY()
 	TWeakObjectPtr<UDebugSkelMeshComponent> SkelMeshComponent;
 
+	TWeakPtr<IPersonaPreviewScene> WeakPreviewScene;
+
 	/** Whether to use local or world location */
 	bool bLocalLocation;
 
diff --git a/Engine/Source/Runtime/Engine/Private/SkeletalDebugRendering.cpp b/Engine/Source/Runtime/Engine/Private/SkeletalDebugRendering.cpp
--- a/Engine/Source/Runtime/Engine/Private/SkeletalDebugRendering.cpp
+++ b/Engine/Source/Runtime/Engine/Private/SkeletalDebugRendering.cpp
@@ -314,7 +314,7 @@ void DrawBonesInternal(
 	{
 		for (int32 BoneIndex : InSelectedBones)
 		{
-			if (BoneIndex != INDEX_NONE)
+			if (BoneIndex != INDEX_NONE && BonesToDraw.IsValidIndex(BoneIndex))
 			{
 				BonesToDraw[BoneIndex] = true;
 			}
