AnimSandbox - Assertion failed: bUseRawDataOnly in AnimSequence.cpp when copy/pasting anim curves from one animation sequence to another
#fix add scoped bracket for curve changes

diff --git a/Engine/Source/Editor/Persona/Private/AnimTimeline/AnimModel_AnimSequenceBase.cpp b/Engine/Source/Editor/Persona/Private/AnimTimeline/AnimModel_AnimSequenceBase.cpp
--- a/Engine/Source/Editor/Persona/Private/AnimTimeline/AnimModel_AnimSequenceBase.cpp
+++ b/Engine/Source/Editor/Persona/Private/AnimTimeline/AnimModel_AnimSequenceBase.cpp
@@ -606,9 +606,7 @@ void FAnimModel_AnimSequenceBase::PasteDataFromClipboardToSelectedCurve()
 {
 	if (const UAnimTimelineClipboardContent* ClipboardContent = FAnimTimelineClipboardUtilities::GetContentFromClipboard())
 	{
-		// Setup for modifying animation sequence 
 		const FScopedTransaction Transaction(LOCTEXT("AnimSequenceBase_PasteCurveData", "Paste Data To Selected Curve"));
-		AnimSequenceBase->Modify();
 
 		// Paste data
 		FAnimTimelineClipboardUtilities::OverwriteSelectedCurveDataFromClipboard(ClipboardContent, SelectedTracks, AnimSequenceBase);
@@ -628,10 +626,7 @@ void FAnimModel_AnimSequenceBase::PasteFromClipboard()
 {
 	if (const UAnimTimelineClipboardContent* ClipboardContent = FAnimTimelineClipboardUtilities::GetContentFromClipboard())
 	{
-		// Setup for modifying animation sequence 
 		const FScopedTransaction Transaction(LOCTEXT("AnimSequenceBase_PasteCurves", "Paste"));
-		AnimSequenceBase->Modify();
-
 		// Paste curves from clipboard
 		FAnimTimelineClipboardUtilities::OverwriteOrAddCurvesFromClipboardContent(ClipboardContent, AnimSequenceBase);
 	}
diff --git a/Engine/Source/Editor/Persona/Private/AnimTimeline/AnimTimelineClipboard.cpp b/Engine/Source/Editor/Persona/Private/AnimTimeline/AnimTimelineClipboard.cpp
--- a/Engine/Source/Editor/Persona/Private/AnimTimeline/AnimTimelineClipboard.cpp
+++ b/Engine/Source/Editor/Persona/Private/AnimTimeline/AnimTimelineClipboard.cpp
@@ -181,6 +181,8 @@ void FAnimTimelineClipboardUtilities::OverwriteSelectedCurveDataFromClipboard(co
 	
 	if (const TSharedRef<FAnimTimelineTrack>& Track = *CurveTracks.CreateConstIterator(); Track->IsA<FAnimTimelineTrack_Curve>())
 	{
+		IAnimationDataController::FScopedBracket ScopedBracket(Controller, NSLOCTEXT("AnimTimelineClipboardUtilities", "PasteCurveKeys_Bracket", "Paste Animation Curve Keys"));
+		
 		const FAnimTimelineTrack_Curve & CurveTrack = Track->As<FAnimTimelineTrack_Curve>();
 
 		// Get selected curve track information
@@ -257,42 +259,47 @@ void FAnimTimelineClipboardUtilities::OverwriteOrAddCurvesFromClipboardContent(c
 	check(Skeleton != nullptr)
 	IAnimationDataController& Controller = InTargetSequence->GetController();
 	
-	for (TObjectPtr<UAnimCurveBaseCopyObject> InCurveCopyObj : ClipboardContent->Curves)
+	if (ClipboardContent->Curves.Num())
 	{
-		// Exit on types not supported by animation controller
-		if (InCurveCopyObj->CurveType != ERawCurveTrackTypes::RCT_Float && InCurveCopyObj->CurveType != ERawCurveTrackTypes::RCT_Transform)
-		{
-			UE_LOG(LogAnimation, Warning, TEXT("Attempting to paste curve { %s } of unsupported type { %s } by animation controller. "), *InCurveCopyObj->CurveName.ToString(), *UEnum::GetValueAsString(InCurveCopyObj->CurveType));
-			continue;
-		}
-		
-		// Ensure we are not pasting data into the same curve we copying data from
-		if (InCurveCopyObj->OriginName.Compare(Controller.GetModelInterface().GetObject()->GetOuter()->GetFName()) == 0 && Controller.GetModel()->FindCurve(InCurveCopyObj->GetAnimationCurveIdentifier()))
-		{
-			// Do nothing and exit.
-			UE_LOG(LogAnimation, Warning, TEXT("Attempting to paste already existing curve: %s"), *InCurveCopyObj->CurveName.ToString());
-		}
-		else
-		{
-			FAnimationCurveIdentifier CurveIdentifier = InCurveCopyObj->GetAnimationCurveIdentifier();
-			if (InCurveCopyObj->CurveType == ERawCurveTrackTypes::RCT_Float)
-			{
-				const UFloatCurveCopyObject* FloatCopyObject = Cast<UFloatCurveCopyObject>(InCurveCopyObj);
-				
-				Controller.SetCurveKeys(CurveIdentifier, FloatCopyObject->Curve.FloatCurve.GetConstRefOfKeys());
-				Controller.SetCurveFlags(CurveIdentifier, FloatCopyObject->Curve.GetCurveTypeFlags());
-			}
-			else if (InCurveCopyObj->CurveType == ERawCurveTrackTypes::RCT_Transform)
-			{
-				const UTransformCurveCopyObject* TransformCopyObject = Cast<UTransformCurveCopyObject>(InCurveCopyObj);
-				
-				TArray<float> TimeKeys;
-				TArray<FTransform> TransformValues;
-				TransformCopyObject->Curve.GetKeys(TimeKeys, TransformValues); // TODO: Optimize this if possible!
-	
-				Controller.SetTransformCurveKeys(CurveIdentifier, TransformValues, TimeKeys);
-				Controller.SetCurveFlags(CurveIdentifier, TransformCopyObject->Curve.GetCurveTypeFlags());
-			}
-		}
+		IAnimationDataController::FScopedBracket ScopedBracket(Controller, NSLOCTEXT("AnimTimelineClipboardUtilities", "PasteCurves_Bracket", "Paste Animation Curves"));
+	    for (TObjectPtr<UAnimCurveBaseCopyObject> InCurveCopyObj : ClipboardContent->Curves)
+	    {
+		    // Exit on types not supported by animation controller
+		    if (InCurveCopyObj->CurveType != ERawCurveTrackTypes::RCT_Float && InCurveCopyObj->CurveType != ERawCurveTrackTypes::RCT_Transform)
+		    {
+			    UE_LOG(LogAnimation, Warning, TEXT("Attempting to paste curve { %s } of unsupported type { %s } by animation controller. "), *InCurveCopyObj->CurveName.ToString(), *UEnum::GetValueAsString(InCurveCopyObj->CurveType));
+			    continue;
+		    }
+		    
+		    // Ensure we are not pasting data into the same curve we copying data from
+		    if (InCurveCopyObj->OriginName.Compare(Controller.GetModelInterface().GetObject()->GetOuter()->GetFName()) == 0 && Controller.GetModel()->FindCurve(InCurveCopyObj->GetAnimationCurveIdentifier()))
+		    {
+			    // Do nothing and exit.
+				UE_LOG(LogAnimation, Warning, TEXT("Attempting to paste already existing curve: %s"), *InCurveCopyObj->CurveName.ToString());
+		    }
+		    else
+		    {
+			    FAnimationCurveIdentifier CurveIdentifier = InCurveCopyObj->GetAnimationCurveIdentifier();
+				    if (InCurveCopyObj->CurveType == ERawCurveTrackTypes::RCT_Float)
+				    {
+					    const UFloatCurveCopyObject* FloatCopyObject = Cast<UFloatCurveCopyObject>(InCurveCopyObj);
+					    
+					    Controller.SetCurveKeys(CurveIdentifier, FloatCopyObject->Curve.FloatCurve.GetConstRefOfKeys());
+					    Controller.SetCurveFlags(CurveIdentifier, FloatCopyObject->Curve.GetCurveTypeFlags());
+				    }
+				    else if (InCurveCopyObj->CurveType == ERawCurveTrackTypes::RCT_Transform)
+				    {
+					    const UTransformCurveCopyObject* TransformCopyObject = Cast<UTransformCurveCopyObject>(InCurveCopyObj);
+					    
+					    TArray<float> TimeKeys;
+					    TArray<FTransform> TransformValues;
+					    TransformCopyObject->Curve.GetKeys(TimeKeys, TransformValues); // TODO: Optimize this if possible!
+		    
+					    Controller.SetTransformCurveKeys(CurveIdentifier, TransformValues, TimeKeys);
+					    Controller.SetCurveFlags(CurveIdentifier, TransformCopyObject->Curve.GetCurveTypeFlags());
+				    }
+			    }
+		    }
+	    }
 	}
 }
