Fix up incorrect whitespace and closing bracket when merging CL 24653993

diff --git a/Engine/Source/Editor/Persona/Private/AnimTimeline/AnimTimelineClipboard.cpp b/Engine/Source/Editor/Persona/Private/AnimTimeline/AnimTimelineClipboard.cpp
--- a/Engine/Source/Editor/Persona/Private/AnimTimeline/AnimTimelineClipboard.cpp
+++ b/Engine/Source/Editor/Persona/Private/AnimTimeline/AnimTimelineClipboard.cpp
@@ -280,26 +280,26 @@ void FAnimTimelineClipboardUtilities::OverwriteOrAddCurvesFromClipboardContent(c
 		    else
 		    {
 			    FAnimationCurveIdentifier CurveIdentifier = InCurveCopyObj->GetAnimationCurveIdentifier();
-				    if (InCurveCopyObj->CurveType == ERawCurveTrackTypes::RCT_Float)
-				    {
-					    const UFloatCurveCopyObject* FloatCopyObject = Cast<UFloatCurveCopyObject>(InCurveCopyObj);
-					    
-					    Controller.SetCurveKeys(CurveIdentifier, FloatCopyObject->Curve.FloatCurve.GetConstRefOfKeys());
-					    Controller.SetCurveFlags(CurveIdentifier, FloatCopyObject->Curve.GetCurveTypeFlags());
-				    }
-				    else if (InCurveCopyObj->CurveType == ERawCurveTrackTypes::RCT_Transform)
-				    {
-					    const UTransformCurveCopyObject* TransformCopyObject = Cast<UTransformCurveCopyObject>(InCurveCopyObj);
-					    
-					    TArray<float> TimeKeys;
-					    TArray<FTransform> TransformValues;
-					    TransformCopyObject->Curve.GetKeys(TimeKeys, TransformValues); // TODO: Optimize this if possible!
-		    
-					    Controller.SetTransformCurveKeys(CurveIdentifier, TransformValues, TimeKeys);
-					    Controller.SetCurveFlags(CurveIdentifier, TransformCopyObject->Curve.GetCurveTypeFlags());
-				    }
+			    if (InCurveCopyObj->CurveType == ERawCurveTrackTypes::RCT_Float)
+			    {
+				    const UFloatCurveCopyObject* FloatCopyObject = Cast<UFloatCurveCopyObject>(InCurveCopyObj);
+				    
+				    Controller.SetCurveKeys(CurveIdentifier, FloatCopyObject->Curve.FloatCurve.GetConstRefOfKeys());
+				    Controller.SetCurveFlags(CurveIdentifier, FloatCopyObject->Curve.GetCurveTypeFlags());
+			    }
+			    else if (InCurveCopyObj->CurveType == ERawCurveTrackTypes::RCT_Transform)
+			    {
+				    const UTransformCurveCopyObject* TransformCopyObject = Cast<UTransformCurveCopyObject>(InCurveCopyObj);
+				    
+				    TArray<float> TimeKeys;
+				    TArray<FTransform> TransformValues;
+				    TransformCopyObject->Curve.GetKeys(TimeKeys, TransformValues); // TODO: Optimize this if possible!
+	    
+				    Controller.SetTransformCurveKeys(CurveIdentifier, TransformValues, TimeKeys);
+				    Controller.SetCurveFlags(CurveIdentifier, TransformCopyObject->Curve.GetCurveTypeFlags());
 			    }
 		    }
 	    }
-	}
+    }
 }
+
