Focus bone selection (missing file, follow-up of 24483126)

diff --git a/Engine/Source/Editor/Persona/Private/AnimationEditorViewportClient.cpp b/Engine/Source/Editor/Persona/Private/AnimationEditorViewportClient.cpp
--- a/Engine/Source/Editor/Persona/Private/AnimationEditorViewportClient.cpp
+++ b/Engine/Source/Editor/Persona/Private/AnimationEditorViewportClient.cpp
@@ -1993,6 +1993,13 @@ FSphere FAnimationViewportClient::GetCameraTarget()
 		return DefaultSphere;
 	}
 
+	AActor* Actor = PreviewMeshComponent->GetOwner();
+	FBox LocalBox(ForceInit);
+	if (GetModeTools()->ComputeBoundingBoxForViewportFocus(Actor, PreviewMeshComponent, LocalBox))
+	{
+		return FBoxSphereBounds(LocalBox).GetSphere();
+	}
+	
 	FBoxSphereBounds Bounds = PreviewMeshComponent->CalcGameBounds(FTransform::Identity);
 	return Bounds.GetSphere();
 }
