PersonaEditor : Fix for crash creating animation assets multiple times, by enabling replacing existing animation in the options to create new assets (if not enabled and we try to save to the same asset we are previewing, the system tries to delete the target asset, which is currently being used as source for the new asset, and it crashes because the PersonaToolkit is no longer available)

diff --git a/Engine/Source/Editor/Persona/Private/PersonaModule.cpp b/Engine/Source/Editor/Persona/Private/PersonaModule.cpp
--- a/Engine/Source/Editor/Persona/Private/PersonaModule.cpp
+++ b/Engine/Source/Editor/Persona/Private/PersonaModule.cpp
@@ -1306,7 +1306,7 @@ void FPersonaModule::FillCreateAnimationMenu(FMenuBuilder& MenuBuilder, TWeakPtr
 	}
 
 	const bool bDoNotShowNameDialog = false;
-	const bool bAllowReplaceExisting = false;
+	const bool bAllowReplaceExisting = true;  // This avoids a crash when saving on top of the currently previewed animation (if false, it tries to generate using a PersonaToolkit that no longer exists, as the editor gets closed for deleting the asset)
 
 	// create rig
 	MenuBuilder.BeginSection("CreateAnimationSubMenu", LOCTEXT("CreateAnimationSubMenuHeading", "Create Animation"));
@@ -1360,7 +1360,7 @@ void FPersonaModule::FillCreateAnimationFromCurrentAnimationMenu(FMenuBuilder& M
 	}
 
 	const bool bDoNotShowNameDialog = false;
-	const bool bAllowReplaceExisting = false;
+	const bool bAllowReplaceExisting = true;  // This avoids a crash when saving on top of the currently previewed animation (if false, it tries to generate using a PersonaToolkit that no longer exists, as the editor gets closed for deleting the asset)
 
 	// create rig
 	MenuBuilder.BeginSection("CreateAnimationSubMenu", LOCTEXT("CreateAnimationFromCurrentAnimationSubmenuHeading", "Create Animation"));
@@ -1401,7 +1401,7 @@ void FPersonaModule::FillCreatePoseAssetMenu(FMenuBuilder& MenuBuilder, TWeakPtr
 	}
 
 	const bool bDoNotShowNameDialog = false;
-	const bool bAllowReplaceExisting = false;
+	const bool bAllowReplaceExisting = true;  // This avoids a crash when saving on top of the currently previewed animation (if false, it tries to generate using a PersonaToolkit that no longer exists, as the editor gets closed for deleting the asset)
 
 	// create rig
 	MenuBuilder.BeginSection("CreatePoseAssetSubMenu", LOCTEXT("CreatePoseAssetSubMenuHeading", "Create PoseAsset"));
