Fixing deprecation warnings from deprecating RetargetBasePose.

diff --git a/Engine/Source/Editor/Persona/Private/SRetargetSources.cpp b/Engine/Source/Editor/Persona/Private/SRetargetSources.cpp
--- a/Engine/Source/Editor/Persona/Private/SRetargetSources.cpp
+++ b/Engine/Source/Editor/Persona/Private/SRetargetSources.cpp
@@ -282,6 +282,7 @@ bool SRetargetSources::ShouldFilterAsset(const FAssetData& InAssetData)
 	return false;
 }
 
+PRAGMA_DISABLE_DEPRECATION_WARNINGS
 void SRetargetSources::ResetRetargetBasePose()
 {
 	UDebugSkelMeshComponent * PreviewMeshComp = PreviewScenePtr.Pin()->GetPreviewMeshComponent();
@@ -303,6 +304,7 @@ void SRetargetSources::ResetRetargetBasePose()
 
 	FSlateApplication::Get().DismissAllMenus();
 }
+PRAGMA_ENABLE_DEPRECATION_WARNINGS
 
 void SRetargetSources::UseCurrentPose()
 {
@@ -321,7 +323,9 @@ void SRetargetSources::UseCurrentPose()
 			const TArray<FTransform> & SpaceBases = PreviewMeshComp->GetComponentSpaceTransforms();
 			// @todo check to see if skeleton vs preview mesh makes it different for missing bones
 			const FReferenceSkeleton& RefSkeleton = PreviewMesh->GetRefSkeleton();
+			PRAGMA_DISABLE_DEPRECATION_WARNINGS
 			TArray<FTransform> & NewRetargetBasePose = PreviewMesh->GetRetargetBasePose();
+			PRAGMA_ENABLE_DEPRECATION_WARNINGS
 			// if you're using leader pose component in preview, this won't work
 			check(PreviewMesh->GetRefSkeleton().GetNum() == SpaceBases.Num());
 			int32 TotalNumBones = PreviewMesh->GetRefSkeleton().GetNum();
@@ -410,15 +414,16 @@ void SRetargetSources::ImportPose(const UPoseAsset* PoseAsset, const FName& Pose
 					const FScopedTransaction Transaction(LOCTEXT("ImportRetargetBasePose_Action", "Import Retarget Base Pose"));
 					PreviewMesh->Modify();
 
+					PRAGMA_DISABLE_DEPRECATION_WARNINGS
 					// reset to original ref pose first
 					PreviewMesh->SetRetargetBasePose(PreviewMesh->GetRefSkeleton().GetRefBonePose());
-
 					// now override imported pose
 					for (int32 TrackIndex = 0; TrackIndex < PoseTrackNames.Num(); ++TrackIndex)
 					{
 						const int32 BoneIndex = PreviewMesh->GetRefSkeleton().FindBoneIndex(PoseTrackNames[TrackIndex]);
 						PreviewMesh->GetRetargetBasePose()[BoneIndex] = PoseTransforms[TrackIndex];
 					}
+					PRAGMA_ENABLE_DEPRECATION_WARNINGS
 
 					TurnOnPreviewRetargetBasePose();
 				}
diff --git a/Engine/Source/Editor/UnrealEd/Private/ImportUtils/SkeletalMeshImportUtils.cpp b/Engine/Source/Editor/UnrealEd/Private/ImportUtils/SkeletalMeshImportUtils.cpp
--- a/Engine/Source/Editor/UnrealEd/Private/ImportUtils/SkeletalMeshImportUtils.cpp
+++ b/Engine/Source/Editor/UnrealEd/Private/ImportUtils/SkeletalMeshImportUtils.cpp
@@ -350,7 +350,9 @@ TSharedPtr<FExistingSkelMeshData> SkeletalMeshImportUtils::SaveExistingSkelMeshD
 	{
 		ExistingMeshDataPtr->ExistingMaterials = SourceMaterials;
 	}
+	PRAGMA_DISABLE_DEPRECATION_WARNINGS
 	ExistingMeshDataPtr->ExistingRetargetBasePose = SourceSkeletalMesh->GetRetargetBasePose();
+	PRAGMA_ENABLE_DEPRECATION_WARNINGS
 
 	ExistingMeshDataPtr->ExistingInlineReductionCacheDatas = SourceMeshModel->InlineReductionCacheDatas;
 
@@ -708,12 +710,14 @@ void SkeletalMeshImportUtils::RestoreExistingSkelMeshData(const TSharedPtr<const
 	//Do everything we need for base LOD re-import
 	if (SafeReimportLODIndex == 0)
 	{
+		PRAGMA_DISABLE_DEPRECATION_WARNINGS
 		// this is not ideal. Ideally we'll have to save only diff with indicating which joints, 
 		// but for now, we allow them to keep the previous pose IF the element count is same
 		if (MeshData->ExistingRetargetBasePose.Num() == SkeletalMesh->GetRefSkeleton().GetRawBoneNum())
 		{
 			SkeletalMesh->SetRetargetBasePose(MeshData->ExistingRetargetBasePose);
 		}
+		PRAGMA_ENABLE_DEPRECATION_WARNINGS
 
 		// Assign sockets from old version of this SkeletalMesh.
 		// Only copy ones for bones that exist in the new mesh.
diff --git a/Engine/Source/Runtime/Engine/Private/Animation/AnimationRuntime.cpp b/Engine/Source/Runtime/Engine/Private/Animation/AnimationRuntime.cpp
--- a/Engine/Source/Runtime/Engine/Private/Animation/AnimationRuntime.cpp
+++ b/Engine/Source/Runtime/Engine/Private/Animation/AnimationRuntime.cpp
@@ -1554,6 +1554,7 @@ void FAnimationRuntime::FillWithRefPose(TArray<FTransform>& OutAtoms, const FBon
 }
 
 #if WITH_EDITOR
+PRAGMA_DISABLE_DEPRECATION_WARNINGS
 void FAnimationRuntime::FillWithRetargetBaseRefPose(FCompactPose& OutPose, const USkeletalMesh* Mesh)
 {
 	// Copy Target Asset's ref pose.
@@ -1569,6 +1570,7 @@ void FAnimationRuntime::FillWithRetargetBaseRefPose(FCompactPose& OutPose, const
 		}
 	}
 }
+PRAGMA_ENABLE_DEPRECATION_WARNINGS
 #endif // WITH_EDITOR
 
 void FAnimationRuntime::ConvertPoseToMeshSpace(const TArray<FTransform>& LocalTransforms, TArray<FTransform>& MeshSpaceTransforms, const FBoneContainer& RequiredBones)
@@ -2407,6 +2409,7 @@ void FAnimationRuntime::FillUpComponentSpaceTransformsRefPose(const USkeleton* S
 	FillUpComponentSpaceTransforms(RefSkeleton, ReferencePose, ComponentSpaceTransforms);
 }
 
+PRAGMA_DISABLE_DEPRECATION_WARNINGS
 void FAnimationRuntime::FillUpComponentSpaceTransformsRetargetBasePose(const USkeletalMesh* Mesh, TArray<FTransform> &ComponentSpaceTransforms)
 {
 	if (Mesh)
@@ -2428,6 +2431,7 @@ void FAnimationRuntime::FillUpComponentSpaceTransformsRetargetBasePose(const USk
 		FillUpComponentSpaceTransformsRetargetBasePose(PreviewMesh, ComponentSpaceTransforms);
 	}
 }
+PRAGMA_ENABLE_DEPRECATION_WARNINGS
 #endif // WITH_EDITOR
 
 void FAnimationRuntime::AppendActiveMorphTargets(
diff --git a/Engine/Source/Runtime/Engine/Private/SkeletalMesh.cpp b/Engine/Source/Runtime/Engine/Private/SkeletalMesh.cpp
--- a/Engine/Source/Runtime/Engine/Private/SkeletalMesh.cpp
+++ b/Engine/Source/Runtime/Engine/Private/SkeletalMesh.cpp
@@ -2108,15 +2108,19 @@ void USkeletalMesh::CalculateInvRefMatrices()
 		}
 
 #if WITH_EDITORONLY_DATA
+		PRAGMA_DISABLE_DEPRECATION_WARNINGS
 		if(GetRetargetBasePose().Num() == 0)
 		{
 			SetRetargetBasePose(GetRefSkeleton().GetRefBonePose());
 		}
+		PRAGMA_ENABLE_DEPRECATION_WARNINGS
 #endif // WITH_EDITORONLY_DATA
 	}
 }
 
 #if WITH_EDITOR
+
+PRAGMA_DISABLE_DEPRECATION_WARNINGS
 void USkeletalMesh::ReallocateRetargetBasePose()
 {
 	// if you're adding other things here, please note that this function is called during postLoad
@@ -2151,6 +2155,7 @@ void USkeletalMesh::ReallocateRetargetBasePose()
 		GetRetargetBasePose() = GetRefSkeleton().GetRefBonePose();
 	}
 }
+PRAGMA_ENABLE_DEPRECATION_WARNINGS
 
 void USkeletalMesh::CalculateRequiredBones(FSkeletalMeshLODModel& LODModel, const struct FReferenceSkeleton& InRefSkeleton, const TMap<FBoneIndexType, FBoneIndexType> * BonesToRemove)
 {
@@ -3259,10 +3264,12 @@ void USkeletalMesh::FinishPostLoadInternal(FSkinnedAssetPostLoadContext& Context
 	CalculateInvRefMatrices();
 
 #if WITH_EDITORONLY_DATA
+	PRAGMA_DISABLE_DEPRECATION_WARNINGS
 	if (GetRetargetBasePose().Num() == 0 && !GetOutermost()->bIsCookedForEditor)
 	{
 		GetRetargetBasePose() = GetRefSkeleton().GetRefBonePose();
 	}
+	PRAGMA_ENABLE_DEPRECATION_WARNINGS
 
 	if (GetLinkerCustomVersion(FFortniteMainBranchObjectVersion::GUID) < FFortniteMainBranchObjectVersion::SupportVirtualBoneInRetargeting)
 	{
