Fixed crash opening a template animation blueprint with curves tab open

diff --git a/Engine/Source/Editor/AnimationBlueprintEditor/Private/AnimationBlueprintEditor.cpp b/Engine/Source/Editor/AnimationBlueprintEditor/Private/AnimationBlueprintEditor.cpp
--- a/Engine/Source/Editor/AnimationBlueprintEditor/Private/AnimationBlueprintEditor.cpp
+++ b/Engine/Source/Editor/AnimationBlueprintEditor/Private/AnimationBlueprintEditor.cpp
@@ -1582,7 +1582,10 @@ void FAnimationBlueprintEditor::OnActiveTabChanged( TSharedPtr<SDockTab> Previou
 
 void FAnimationBlueprintEditor::SetPreviewMesh(USkeletalMesh* NewPreviewMesh)
 {
-	GetSkeletonTree()->SetSkeletalMesh(NewPreviewMesh);
+	if(SkeletonTree.IsValid())
+	{
+		SkeletonTree->SetSkeletalMesh(NewPreviewMesh);
+	}
 }
 
 void FAnimationBlueprintEditor::RefreshPreviewInstanceTrackCurves()
@@ -1719,7 +1722,10 @@ void FAnimationBlueprintEditor::ClearSelectedAnimGraphNodes()
 
 void FAnimationBlueprintEditor::DeselectAll()
 {
-	GetSkeletonTree()->DeselectAll();
+	if(SkeletonTree)
+	{
+		SkeletonTree->DeselectAll();
+	}
 	ClearSelectedActor();
 	ClearSelectedAnimGraphNodes();
 }
diff --git a/Engine/Source/Editor/AnimationBlueprintEditor/Private/AnimationBlueprintEditor.h b/Engine/Source/Editor/AnimationBlueprintEditor/Private/AnimationBlueprintEditor.h
--- a/Engine/Source/Editor/AnimationBlueprintEditor/Private/AnimationBlueprintEditor.h
+++ b/Engine/Source/Editor/AnimationBlueprintEditor/Private/AnimationBlueprintEditor.h
@@ -199,7 +199,7 @@ public:
 	void RecompileAnimBlueprintIfDirty();
 
 	/** Get the skeleton tree this Persona editor is hosting */
-	TSharedRef<class ISkeletonTree> GetSkeletonTree() const { return SkeletonTree.ToSharedRef(); }
+	TSharedPtr<class ISkeletonTree> GetSkeletonTree() const { return SkeletonTree; }
 
 	/** Make this available to allow us to create title bar widgets for other container types - e.g. blendspaces */
 	using FBlueprintEditor::CreateGraphTitleBarWidget;
diff --git a/Engine/Source/Editor/AnimationBlueprintEditor/Private/AnimationBlueprintEditorMode.cpp b/Engine/Source/Editor/AnimationBlueprintEditor/Private/AnimationBlueprintEditorMode.cpp
--- a/Engine/Source/Editor/AnimationBlueprintEditor/Private/AnimationBlueprintEditorMode.cpp
+++ b/Engine/Source/Editor/AnimationBlueprintEditor/Private/AnimationBlueprintEditorMode.cpp
@@ -176,13 +176,14 @@ FAnimationBlueprintEditorMode::FAnimationBlueprintEditorMode(const TSharedRef<FA
 	TabFactories.RegisterFactory(PersonaModule.CreateAnimBlueprintPreviewTabFactory(InAnimationBlueprintEditor, InAnimationBlueprintEditor->GetPersonaToolkit()->GetPreviewScene()));
 	TabFactories.RegisterFactory(PersonaModule.CreateAnimBlueprintAssetOverridesTabFactory(InAnimationBlueprintEditor, InAnimationBlueprintEditor->GetPersonaToolkit()->GetAnimBlueprint(), InAnimationBlueprintEditor->OnPostUndo));
 	TabFactories.RegisterFactory(PersonaModule.CreatePoseWatchTabFactory(InAnimationBlueprintEditor));
-	TabFactories.RegisterFactory(PersonaModule.CreateCurveViewerTabFactory(InAnimationBlueprintEditor, InAnimationBlueprintEditor->GetSkeletonTree()->GetEditableSkeleton(), InAnimationBlueprintEditor->GetPersonaToolkit()->GetPreviewScene(), FOnObjectsSelected::CreateSP(&InAnimationBlueprintEditor.Get(), &FAnimationBlueprintEditor::HandleObjectsSelected)));
+	TSharedPtr<ISkeletonTree> SkeletonTree = InAnimationBlueprintEditor->GetSkeletonTree();
+	TabFactories.RegisterFactory(PersonaModule.CreateCurveViewerTabFactory(InAnimationBlueprintEditor, SkeletonTree.IsValid() ? SkeletonTree->GetEditableSkeleton().ToSharedPtr() : nullptr, InAnimationBlueprintEditor->GetPersonaToolkit()->GetPreviewScene(), FOnObjectsSelected::CreateSP(&InAnimationBlueprintEditor.Get(), &FAnimationBlueprintEditor::HandleObjectsSelected)));
 	TabFactories.RegisterFactory(PersonaModule.CreateAnimAssetFindReplaceTabFactory(InAnimationBlueprintEditor, FAnimAssetFindReplaceConfig()));
 
 	if (!bIsTemplate)
 	{
 		ISkeletonEditorModule& SkeletonEditorModule = FModuleManager::LoadModuleChecked<ISkeletonEditorModule>("SkeletonEditor");
-		TabFactories.RegisterFactory(SkeletonEditorModule.CreateSkeletonTreeTabFactory(InAnimationBlueprintEditor, InAnimationBlueprintEditor->GetSkeletonTree()));
+		TabFactories.RegisterFactory(SkeletonEditorModule.CreateSkeletonTreeTabFactory(InAnimationBlueprintEditor, InAnimationBlueprintEditor->GetSkeletonTree().ToSharedRef()));
 		TabFactories.RegisterFactory(PersonaModule.CreateSkeletonSlotNamesTabFactory(InAnimationBlueprintEditor, InAnimationBlueprintEditor->GetSkeletonTree()->GetEditableSkeleton(), FOnObjectSelected::CreateSP(&InAnimationBlueprintEditor.Get(), &FAnimationBlueprintEditor::HandleObjectSelected)));
 	}
 
diff --git a/Engine/Source/Editor/Persona/Private/PersonaModule.cpp b/Engine/Source/Editor/Persona/Private/PersonaModule.cpp
--- a/Engine/Source/Editor/Persona/Private/PersonaModule.cpp
+++ b/Engine/Source/Editor/Persona/Private/PersonaModule.cpp
@@ -279,12 +279,12 @@ TSharedRef<class FWorkflowTabFactory> FPersonaModule::CreateAnimNotifiesTabFacto
 	return MakeShareable(new FSkeletonAnimNotifiesSummoner(InHostingApp, InEditableSkeleton, InOnObjectsSelected));
 }
 
-TSharedRef<class FWorkflowTabFactory> FPersonaModule::CreateCurveViewerTabFactory(const TSharedRef<class FWorkflowCentricApplication>& InHostingApp, const TSharedRef<class IEditableSkeleton>& InEditableSkeleton, const TSharedRef<IPersonaPreviewScene>& InPreviewScene, FSimpleMulticastDelegate& InOnPostUndo, FOnObjectsSelected InOnObjectsSelected) const
+TSharedRef<class FWorkflowTabFactory> FPersonaModule::CreateCurveViewerTabFactory(const TSharedRef<class FWorkflowCentricApplication>& InHostingApp, const TSharedPtr<class IEditableSkeleton>& InEditableSkeleton, const TSharedRef<IPersonaPreviewScene>& InPreviewScene, FSimpleMulticastDelegate& InOnPostUndo, FOnObjectsSelected InOnObjectsSelected) const
 {
 	return MakeShareable(new FAnimCurveViewerTabSummoner(InHostingApp, InEditableSkeleton, InPreviewScene, InOnObjectsSelected));
 }
 
-TSharedRef<class FWorkflowTabFactory> FPersonaModule::CreateCurveViewerTabFactory(const TSharedRef<class FWorkflowCentricApplication>& InHostingApp, const TSharedRef<class IEditableSkeleton>& InEditableSkeleton, const TSharedRef<IPersonaPreviewScene>& InPreviewScene, FOnObjectsSelected InOnObjectsSelected) const
+TSharedRef<class FWorkflowTabFactory> FPersonaModule::CreateCurveViewerTabFactory(const TSharedRef<class FWorkflowCentricApplication>& InHostingApp, const TSharedPtr<class IEditableSkeleton>& InEditableSkeleton, const TSharedRef<IPersonaPreviewScene>& InPreviewScene, FOnObjectsSelected InOnObjectsSelected) const
 {
 	return MakeShareable(new FAnimCurveViewerTabSummoner(InHostingApp, InEditableSkeleton, InPreviewScene, InOnObjectsSelected));
 }
diff --git a/Engine/Source/Editor/Persona/Private/SAnimCurveViewer.cpp b/Engine/Source/Editor/Persona/Private/SAnimCurveViewer.cpp
--- a/Engine/Source/Editor/Persona/Private/SAnimCurveViewer.cpp
+++ b/Engine/Source/Editor/Persona/Private/SAnimCurveViewer.cpp
@@ -459,12 +459,12 @@ TOptional<float> SAnimCurveListRow::GetMaxWeight() const
 //////////////////////////////////////////////////////////////////////////
 // SAnimCurveViewer
 
-void SAnimCurveViewer::Construct(const FArguments& InArgs, const TSharedRef<class IEditableSkeleton>& InEditableSkeleton, const TSharedRef<IPersonaPreviewScene>& InPreviewScene, FOnObjectsSelected InOnObjectsSelected)
+void SAnimCurveViewer::Construct(const FArguments& InArgs,  const TSharedRef<IPersonaPreviewScene>& InPreviewScene, FOnObjectsSelected InOnObjectsSelected)
 {
 	OnObjectsSelected = InOnObjectsSelected;
 
 	PreviewScenePtr = InPreviewScene;
-	EditableSkeletonPtr = InEditableSkeleton;
+	EditableSkeletonPtr = InArgs._EditableSkeleton;
 
 	InPreviewScene->RegisterOnPreviewMeshChanged(FOnPreviewMeshChanged::CreateSP(this, &SAnimCurveViewer::OnPreviewMeshChanged));
 	InPreviewScene->RegisterOnAnimChanged(FOnAnimChanged::CreateSP(this, &SAnimCurveViewer::OnPreviewAssetChanged));
@@ -715,7 +715,7 @@ void SAnimCurveViewer::CreateAnimCurveList( const FString& SearchText, bool bInF
 		TSharedPtr<FDisplayedAnimCurveInfo>* ExistingItem = AllSeenAnimCurvesMap.Find(InCurveName);
 		if(ExistingItem == nullptr)
 		{
-			TSharedRef<FDisplayedAnimCurveInfo> NewInfo = FDisplayedAnimCurveInfo::Make(EditableSkeletonPtr, InCurveName);
+			TSharedRef<FDisplayedAnimCurveInfo> NewInfo = FDisplayedAnimCurveInfo::Make(InCurveName);
 
 			float Weight = 0.f;
 			const bool bOverride = GetAnimCurveOverride(InCurveName, Weight);
@@ -732,21 +732,24 @@ void SAnimCurveViewer::CreateAnimCurveList( const FString& SearchText, bool bInF
 			(*ExistingItem)->bMorphTarget = EnumHasAnyFlags(InFlags, UE::Anim::ECurveElementFlags::MorphTarget);
 		}
 	};
-	
+
 	// Add curve items from skeleton metadata
-	EditableSkeletonPtr.Pin()->GetSkeleton().ForEachCurveMetaData([&AddCurve](FName InCurveName, const FCurveMetaData& InCurveMetaData)
+	if(TSharedPtr<IEditableSkeleton> EditableSkeleton = EditableSkeletonPtr.Pin())
 	{
-		UE::Anim::ECurveElementFlags Flags = UE::Anim::ECurveElementFlags::None;
-		if(InCurveMetaData.Type.bMaterial)
-		{
-			Flags |= UE::Anim::ECurveElementFlags::Material; 
-		}
-		if(InCurveMetaData.Type.bMorphtarget)
+		EditableSkeleton->GetSkeleton().ForEachCurveMetaData([&AddCurve](FName InCurveName, const FCurveMetaData& InCurveMetaData)
 		{
-			Flags |= UE::Anim::ECurveElementFlags::MorphTarget; 
-		}
-		AddCurve(InCurveName, Flags);
-	});
+			UE::Anim::ECurveElementFlags Flags = UE::Anim::ECurveElementFlags::None;
+			if(InCurveMetaData.Type.bMaterial)
+			{
+				Flags |= UE::Anim::ECurveElementFlags::Material; 
+			}
+			if(InCurveMetaData.Type.bMorphtarget)
+			{
+				Flags |= UE::Anim::ECurveElementFlags::MorphTarget; 
+			}
+			AddCurve(InCurveName, Flags);
+		});
+	}
 
 	// Add active curves if required
 	TSet<FName> ActiveCurves;
diff --git a/Engine/Source/Editor/Persona/Private/SAnimCurveViewer.h b/Engine/Source/Editor/Persona/Private/SAnimCurveViewer.h
--- a/Engine/Source/Editor/Persona/Private/SAnimCurveViewer.h
+++ b/Engine/Source/Editor/Persona/Private/SAnimCurveViewer.h
@@ -39,24 +39,22 @@ public:
 	FName CurveName;
 	float Weight;
 	bool bOverrideData;
-	TWeakPtr<class IEditableSkeleton> EditableSkeleton;		// The skeleton we're associated with
 	bool bShown;
 	bool bMorphTarget;
 	bool bMaterial;
 	
 	/** Static function for creating a new item, but ensures that you can only have a TSharedRef to one */
-	static TSharedRef<FDisplayedAnimCurveInfo> Make(TWeakPtr<class IEditableSkeleton> InEditableSkeleton, const FName& InCurveName)
+	static TSharedRef<FDisplayedAnimCurveInfo> Make(const FName& InCurveName)
 	{
-		return MakeShareable(new FDisplayedAnimCurveInfo(InEditableSkeleton, InCurveName));
+		return MakeShareable(new FDisplayedAnimCurveInfo(InCurveName));
 	}
 
 protected:
 	/** Hidden constructor, always use Make above */
-	FDisplayedAnimCurveInfo(TWeakPtr<class IEditableSkeleton> InEditableSkeleton, const FName& InCurveName)
+	FDisplayedAnimCurveInfo(const FName& InCurveName)
 		: CurveName(InCurveName)
 		, Weight(0.0f)
 		, bOverrideData(false)
-		, EditableSkeleton(InEditableSkeleton)
 		, bShown(false)
 		, bMorphTarget(false)
 		, bMaterial(false)
@@ -155,9 +153,10 @@ private:
 class SAnimCurveViewer : public SCompoundWidget, public FSelfRegisteringEditorUndoClient
 {
 public:
-	SLATE_BEGIN_ARGS( SAnimCurveViewer )
-	{}
-	
+	SLATE_BEGIN_ARGS( SAnimCurveViewer ) {}
+
+	SLATE_ARGUMENT(TSharedPtr<IEditableSkeleton>, EditableSkeleton)
+
 	SLATE_END_ARGS()
 
 	/**
@@ -166,7 +165,7 @@ public:
 	* @param InArgs - Arguments passed from Slate
 	*
 	*/
-	void Construct( const FArguments& InArgs, const TSharedRef<class IEditableSkeleton>& InEditableSkeleton, const TSharedRef<class IPersonaPreviewScene>& InPreviewScene, FOnObjectsSelected InOnObjectsSelected);
+	void Construct( const FArguments& InArgs, const TSharedRef<class IPersonaPreviewScene>& InPreviewScene, FOnObjectsSelected InOnObjectsSelected);
 
 	/**
 	* Destructor - resets the animation curve
diff --git a/Engine/Source/Editor/Persona/Private/TabSpawners.cpp b/Engine/Source/Editor/Persona/Private/TabSpawners.cpp
--- a/Engine/Source/Editor/Persona/Private/TabSpawners.cpp
+++ b/Engine/Source/Editor/Persona/Private/TabSpawners.cpp
@@ -164,7 +164,7 @@ TSharedRef<SWidget> FMorphTargetTabSummoner::CreateTabBody(const FWorkflowTabSpa
 /////////////////////////////////////////////////////
 // FAnimCurveViewerTabSummoner
 
-FAnimCurveViewerTabSummoner::FAnimCurveViewerTabSummoner(TSharedPtr<class FAssetEditorToolkit> InHostingApp, const TSharedRef<IEditableSkeleton>& InEditableSkeleton, const TSharedRef<IPersonaPreviewScene>& InPreviewScene, FOnObjectsSelected InOnObjectsSelected)
+FAnimCurveViewerTabSummoner::FAnimCurveViewerTabSummoner(TSharedPtr<class FAssetEditorToolkit> InHostingApp, const TSharedPtr<IEditableSkeleton>& InEditableSkeleton, const TSharedRef<IPersonaPreviewScene>& InPreviewScene, FOnObjectsSelected InOnObjectsSelected)
 	: FWorkflowTabFactory(FPersonaTabs::AnimCurveViewID, InHostingApp)
 	, EditableSkeleton(InEditableSkeleton)
 	, PreviewScene(InPreviewScene)
@@ -182,7 +182,8 @@ FAnimCurveViewerTabSummoner::FAnimCurveViewerTabSummoner(TSharedPtr<class FAsset
 
 TSharedRef<SWidget> FAnimCurveViewerTabSummoner::CreateTabBody(const FWorkflowTabSpawnInfo& Info) const
 {
-	return SNew(SAnimCurveViewer, EditableSkeleton.Pin().ToSharedRef(), PreviewScene.Pin().ToSharedRef(), OnObjectsSelected);
+	return SNew(SAnimCurveViewer, PreviewScene.Pin().ToSharedRef(), OnObjectsSelected)
+		.EditableSkeleton(EditableSkeleton.Pin());
 }
 
 /////////////////////////////////////////////////////
diff --git a/Engine/Source/Editor/Persona/Private/TabSpawners.h b/Engine/Source/Editor/Persona/Private/TabSpawners.h
--- a/Engine/Source/Editor/Persona/Private/TabSpawners.h
+++ b/Engine/Source/Editor/Persona/Private/TabSpawners.h
@@ -144,7 +144,7 @@ private:
 struct FAnimCurveViewerTabSummoner : public FWorkflowTabFactory
 {
 public:
-	FAnimCurveViewerTabSummoner(TSharedPtr<class FAssetEditorToolkit> InHostingApp, const TSharedRef<IEditableSkeleton>& InEditableSkeleton, const TSharedRef<IPersonaPreviewScene>& InPreviewScene, FOnObjectsSelected InOnObjectsSelected);
+	FAnimCurveViewerTabSummoner(TSharedPtr<class FAssetEditorToolkit> InHostingApp, const TSharedPtr<IEditableSkeleton>& InEditableSkeleton, const TSharedRef<IPersonaPreviewScene>& InPreviewScene, FOnObjectsSelected InOnObjectsSelected);
 
 	virtual TSharedRef<SWidget> CreateTabBody(const FWorkflowTabSpawnInfo& Info) const override;
 
diff --git a/Engine/Source/Editor/Persona/Public/PersonaModule.h b/Engine/Source/Editor/Persona/Public/PersonaModule.h
--- a/Engine/Source/Editor/Persona/Public/PersonaModule.h
+++ b/Engine/Source/Editor/Persona/Public/PersonaModule.h
@@ -295,10 +295,10 @@ public:
 	virtual TSharedRef<FWorkflowTabFactory> CreateAnimNotifiesTabFactory(const TSharedRef<class FWorkflowCentricApplication>& InHostingApp, const TSharedRef<class IEditableSkeleton>& InEditableSkeleton, FOnObjectsSelected InOnObjectsSelected) const;
 
 	UE_DEPRECATED(5.0, "Please use the overload that does not take a post-undo delegate")
-	virtual TSharedRef<FWorkflowTabFactory> CreateCurveViewerTabFactory(const TSharedRef<class FWorkflowCentricApplication>& InHostingApp, const TSharedRef<class IEditableSkeleton>& InEditableSkeleton, const TSharedRef<class IPersonaPreviewScene>& InPreviewScene, FSimpleMulticastDelegate& InOnPostUndo, FOnObjectsSelected InOnObjectsSelected) const;
+	virtual TSharedRef<FWorkflowTabFactory> CreateCurveViewerTabFactory(const TSharedRef<class FWorkflowCentricApplication>& InHostingApp, const TSharedPtr<class IEditableSkeleton>& InEditableSkeleton, const TSharedRef<class IPersonaPreviewScene>& InPreviewScene, FSimpleMulticastDelegate& InOnPostUndo, FOnObjectsSelected InOnObjectsSelected) const;
 
 	/** Create a skeleton curve viewer tab factory */
-	virtual TSharedRef<FWorkflowTabFactory> CreateCurveViewerTabFactory(const TSharedRef<class FWorkflowCentricApplication>& InHostingApp, const TSharedRef<class IEditableSkeleton>& InEditableSkeleton, const TSharedRef<class IPersonaPreviewScene>& InPreviewScene, FOnObjectsSelected InOnObjectsSelected) const;
+	virtual TSharedRef<FWorkflowTabFactory> CreateCurveViewerTabFactory(const TSharedRef<class FWorkflowCentricApplication>& InHostingApp, const TSharedPtr<class IEditableSkeleton>& InEditableSkeleton, const TSharedRef<class IPersonaPreviewScene>& InPreviewScene, FOnObjectsSelected InOnObjectsSelected) const;
 
 	/** Create a skeleton curve metadata editor tab factory */
 	virtual TSharedRef<FWorkflowTabFactory> CreateCurveMetadataEditorTabFactory(const TSharedRef<class FWorkflowCentricApplication>& InHostingApp, UObject* InMetadataHost, const TSharedRef<class IPersonaPreviewScene>& InPreviewScene, FOnObjectsSelected InOnObjectsSelected) const;
