Revert 24046871 so current ScreenSize is calculated per frame in EditorViewport and not static to the per platform value

diff --git a/Engine/Source/Editor/Persona/Private/AnimationEditorViewportClient.cpp b/Engine/Source/Editor/Persona/Private/AnimationEditorViewportClient.cpp
--- a/Engine/Source/Editor/Persona/Private/AnimationEditorViewportClient.cpp
+++ b/Engine/Source/Editor/Persona/Private/AnimationEditorViewportClient.cpp
@@ -2412,13 +2412,7 @@ void FAnimationViewportClient::SetupViewForRendering( FSceneViewFamily& ViewFami
 	if (PreviewMeshComponent != nullptr && PreviewMeshComponent->MeshObject != nullptr)
 	{
 		const FBoxSphereBounds& SkelBounds = PreviewMeshComponent->Bounds;
-
-		const int32 LODIndex = PreviewMeshComponent->GetPredictedLODLevel();
-		USkinnedAsset* SkinnedAsset = PreviewMeshComponent->GetSkinnedAsset();
-		check(SkinnedAsset);
-
-		FPerPlatformFloat& ScreenSize = SkinnedAsset->GetLODInfo(LODIndex)->ScreenSize;
-		CachedScreenSize = SkinnedAsset->IsValidLODIndex(LODIndex) ? ScreenSize.GetValue(): CachedScreenSize;
+		CachedScreenSize = ComputeBoundsScreenSize(SkelBounds.Origin, static_cast<float>(SkelBounds.SphereRadius), View);
 	}
 }
 
