Animation sequences containing curves are always marked as modified after opening
#fix ensure that SetColor does not modify the outer package (unnecessarily dirtying Animation Assets)

diff --git a/Engine/Source/Editor/Persona/Private/AnimTimeline/AnimTimelineTrack_FloatCurve.cpp b/Engine/Source/Editor/Persona/Private/AnimTimeline/AnimTimelineTrack_FloatCurve.cpp
--- a/Engine/Source/Editor/Persona/Private/AnimTimeline/AnimTimelineTrack_FloatCurve.cpp
+++ b/Engine/Source/Editor/Persona/Private/AnimTimeline/AnimTimelineTrack_FloatCurve.cpp
@@ -291,7 +291,7 @@ void FAnimTimelineTrack_FloatCurve::AddCurveTrackButton(TSharedPtr<SHorizontalBo
 		// Set display curves too
 		for(const TPair<FCurveModelID, TUniquePtr<FCurveModel>>& CurvePair : CurveEditor->GetCurves())
 		{
-			CurvePair.Value->SetColor(InNewColor);
+			CurvePair.Value->SetColor(InNewColor, false);
 		}
 	};
 
