Added 0.75 and custom speed to Animation Viewport ToolBar, for Animation Speed and TurnTable speed

diff --git a/Engine/Plugins/Experimental/Mutable/Source/CustomizableObjectEditor/Private/MuCOE/SCustomizableObjectEditorViewportToolBar.cpp b/Engine/Plugins/Experimental/Mutable/Source/CustomizableObjectEditor/Private/MuCOE/SCustomizableObjectEditorViewportToolBar.cpp
--- a/Engine/Plugins/Experimental/Mutable/Source/CustomizableObjectEditor/Private/MuCOE/SCustomizableObjectEditorViewportToolBar.cpp
+++ b/Engine/Plugins/Experimental/Mutable/Source/CustomizableObjectEditor/Private/MuCOE/SCustomizableObjectEditorViewportToolBar.cpp
@@ -1174,7 +1174,7 @@ FText SCustomizableObjectEditorViewportToolBar::GetPlaybackMenuLabel() const
 	//		if (Viewport.Pin()->IsPlaybackSpeedSelected(i))
 	//		{
 	//			Label = FText::FromString(FString::Printf(
-	//				(i == EAnimationPlaybackSpeeds::Quarter) ? TEXT("x%.2f") : TEXT("x%.1f"),
+	//				(i == EAnimationPlaybackSpeeds::Quarter || i == EAnimationPlaybackSpeeds::ThreeQuarters) ? TEXT("x%.2f") : TEXT("x%.1f"),
 	//				EAnimationPlaybackSpeeds::Values[i]
 	//			));
 	//			break;
diff --git a/Engine/Source/Editor/Persona/Private/AnimViewportPlaybackCommands.cpp b/Engine/Source/Editor/Persona/Private/AnimViewportPlaybackCommands.cpp
--- a/Engine/Source/Editor/Persona/Private/AnimViewportPlaybackCommands.cpp
+++ b/Engine/Source/Editor/Persona/Private/AnimViewportPlaybackCommands.cpp
@@ -23,18 +23,22 @@ void FAnimViewportPlaybackCommands::RegisterCommands()
 	UI_COMMAND( PlaybackSpeedCommands[EAnimationPlaybackSpeeds::OneTenth],	"x0.1", "Set the animation playback speed to a tenth of normal", EUserInterfaceActionType::RadioButton, FInputChord() );
 	UI_COMMAND( PlaybackSpeedCommands[EAnimationPlaybackSpeeds::Quarter],		"x0.25", "Set the animation playback speed to a quarter of normal", EUserInterfaceActionType::RadioButton, FInputChord() );
 	UI_COMMAND( PlaybackSpeedCommands[EAnimationPlaybackSpeeds::Half],		"x0.5", "Set the animation playback speed to a half of normal", EUserInterfaceActionType::RadioButton, FInputChord() );
+	UI_COMMAND( PlaybackSpeedCommands[EAnimationPlaybackSpeeds::ThreeQuarters],		"x0.75", "Set the animation playback speed to three quarters of normal", EUserInterfaceActionType::RadioButton, FInputChord() );
 	UI_COMMAND( PlaybackSpeedCommands[EAnimationPlaybackSpeeds::Normal],		"x1.0", "Set the animation playback speed to normal", EUserInterfaceActionType::RadioButton, FInputChord() );
 	UI_COMMAND( PlaybackSpeedCommands[EAnimationPlaybackSpeeds::Double],		"x2.0", "Set the animation playback speed to double the speed of normal", EUserInterfaceActionType::RadioButton, FInputChord() );
 	UI_COMMAND( PlaybackSpeedCommands[EAnimationPlaybackSpeeds::FiveTimes],	"x5.0", "Set the animation playback speed to five times the normal speed", EUserInterfaceActionType::RadioButton, FInputChord() );
 	UI_COMMAND( PlaybackSpeedCommands[EAnimationPlaybackSpeeds::TenTimes],	"x10.0", "Set the animation playback speed to ten times the normal speed", EUserInterfaceActionType::RadioButton, FInputChord() );
+	UI_COMMAND( PlaybackSpeedCommands[EAnimationPlaybackSpeeds::Custom],	"xCustom", "Set the animation playback speed to assigned custom speed", EUserInterfaceActionType::RadioButton, FInputChord() );
 
 	UI_COMMAND(TurnTableSpeeds[EAnimationPlaybackSpeeds::OneTenth], "x0.1", "Set the turn table rotation speed to a tenth of normal", EUserInterfaceActionType::RadioButton, FInputChord());
 	UI_COMMAND(TurnTableSpeeds[EAnimationPlaybackSpeeds::Quarter], "x0.25", "Set the turn table rotation speed to a quarter of normal", EUserInterfaceActionType::RadioButton, FInputChord());
 	UI_COMMAND(TurnTableSpeeds[EAnimationPlaybackSpeeds::Half], "x0.5", "Set the turn table rotation speed to a half of normal", EUserInterfaceActionType::RadioButton, FInputChord());
+	UI_COMMAND(TurnTableSpeeds[EAnimationPlaybackSpeeds::ThreeQuarters], "x0.75", "Set the turn table rotation speed to three quarters of normal", EUserInterfaceActionType::RadioButton, FInputChord());
 	UI_COMMAND(TurnTableSpeeds[EAnimationPlaybackSpeeds::Normal], "x1.0", "Set the turn table rotation speed to normal", EUserInterfaceActionType::RadioButton, FInputChord());
 	UI_COMMAND(TurnTableSpeeds[EAnimationPlaybackSpeeds::Double], "x2.0", "Set the turn table rotation speed to double the speed of normal", EUserInterfaceActionType::RadioButton, FInputChord());
 	UI_COMMAND(TurnTableSpeeds[EAnimationPlaybackSpeeds::FiveTimes], "x5.0", "Set the turn table rotation speed to five times the normal speed", EUserInterfaceActionType::RadioButton, FInputChord());
 	UI_COMMAND(TurnTableSpeeds[EAnimationPlaybackSpeeds::TenTimes], "x10.0", "Set the turn table rotation speed to ten times the normal speed", EUserInterfaceActionType::RadioButton, FInputChord());
+	UI_COMMAND(TurnTableSpeeds[EAnimationPlaybackSpeeds::Custom], "xCustom", "Set the animation playback speed to assigned custom speed", EUserInterfaceActionType::RadioButton, FInputChord());
 
 	UI_COMMAND(PersonaTurnTablePlay, "Play", "Turn table rotates", EUserInterfaceActionType::RadioButton, FInputChord());
 	UI_COMMAND(PersonaTurnTablePause, "Pause", "Freeze with current rotation", EUserInterfaceActionType::RadioButton, FInputChord());
diff --git a/Engine/Source/Editor/Persona/Private/AnimationEditorViewportClient.cpp b/Engine/Source/Editor/Persona/Private/AnimationEditorViewportClient.cpp
--- a/Engine/Source/Editor/Persona/Private/AnimationEditorViewportClient.cpp
+++ b/Engine/Source/Editor/Persona/Private/AnimationEditorViewportClient.cpp
@@ -57,7 +57,7 @@ namespace {
 namespace EAnimationPlaybackSpeeds
 {
 	// Speed scales for animation playback, must match EAnimationPlaybackSpeeds::Type
-	float Values[EAnimationPlaybackSpeeds::NumPlaybackSpeeds] = { 0.1f, 0.25f, 0.5f, 1.0f, 2.0f, 5.0f, 10.0f };
+	float Values[EAnimationPlaybackSpeeds::NumPlaybackSpeeds] = { 0.1f, 0.25f, 0.5f, 0.75f, 1.0f, 2.0f, 5.0f, 10.0f, 0.f };
 }
 
 #define LOCTEXT_NAMESPACE "FAnimationViewportClient"
@@ -2339,7 +2339,10 @@ void FAnimationViewportClient::SetPlaybackSpeedMode(EAnimationPlaybackSpeeds::Ty
 
 	if (GetWorld())
 	{
-		GetWorld()->GetWorldSettings()->TimeDilation = EAnimationPlaybackSpeeds::Values[AnimationPlaybackSpeedMode];
+		const float AnimationSpeed = (InMode == EAnimationPlaybackSpeeds::Custom)
+			? GetCustomAnimationSpeed()
+			: EAnimationPlaybackSpeeds::Values[AnimationPlaybackSpeedMode];
+		GetWorld()->GetWorldSettings()->TimeDilation = AnimationSpeed;
 	}
 }
 
@@ -2348,6 +2351,17 @@ EAnimationPlaybackSpeeds::Type FAnimationViewportClient::GetPlaybackSpeedMode()
 	return AnimationPlaybackSpeedMode;
 }
 
+void FAnimationViewportClient::SetCustomAnimationSpeed(const float InCustomAnimationSpeed)
+{
+	CustomAnimationSpeed = InCustomAnimationSpeed;
+	SetPlaybackSpeedMode(EAnimationPlaybackSpeeds::Custom);
+}
+
+float FAnimationViewportClient::GetCustomAnimationSpeed() const
+{
+	return CustomAnimationSpeed;
+}
+
 TSharedPtr<class FAnimationEditorPreviewScene> FAnimationViewportClient::GetAnimPreviewScenePtr() const
 {
 	return StaticCastSharedPtr<FAnimationEditorPreviewScene>(PreviewScenePtr);
diff --git a/Engine/Source/Editor/Persona/Private/SAnimViewportToolBar.cpp b/Engine/Source/Editor/Persona/Private/SAnimViewportToolBar.cpp
--- a/Engine/Source/Editor/Persona/Private/SAnimViewportToolBar.cpp
+++ b/Engine/Source/Editor/Persona/Private/SAnimViewportToolBar.cpp
@@ -203,6 +203,52 @@ protected:
 	TWeakPtr<SAnimationEditorViewportTabBody> AnimViewportPtr;
 };
 
+//Class definition which represents widget to modify Bone Draw Size in viewport
+class SCusttomAnimationSpeedSetting : public SCompoundWidget
+{
+public:
+	/** Notification for numeric value change */
+	DECLARE_DELEGATE_OneParam(FOnCustomSpeedChanged, float);
+
+	SLATE_BEGIN_ARGS(SCusttomAnimationSpeedSetting) {}
+		SLATE_ARGUMENT(TWeakPtr<SAnimationEditorViewportTabBody>, AnimEditorViewport)
+		SLATE_ATTRIBUTE(float, CustomSpeed)
+		SLATE_EVENT(FOnCustomSpeedChanged, OnCustomSpeedChanged)
+	SLATE_END_ARGS()
+
+	void Construct(const FArguments& InArgs )
+	{
+		AnimViewportPtr = InArgs._AnimEditorViewport;
+		CustomSpeed = InArgs._CustomSpeed;
+		OnCustomSpeedChanged = InArgs._OnCustomSpeedChanged;
+
+		this->ChildSlot
+		[
+			SNew(SBox)
+			.HAlign(HAlign_Right)
+			[
+				SNew(SBox)
+				.Padding(FMargin(4.0f, 0.0f, 0.0f, 0.0f))
+				.WidthOverride(100.0f)
+				[
+					SNew(SSpinBox<float>)
+					.Font(FAppStyle::GetFontStyle(TEXT("MenuItem.Font")))
+					.ToolTipText(LOCTEXT("AnimationCustomSpeed", "Set Custom Speed."))
+					.MinValue(0.f)
+					.MaxSliderValue(10.f)
+					.SupportDynamicSliderMaxValue(true)
+					.Value(CustomSpeed)
+					.OnValueChanged(OnCustomSpeedChanged)
+				]
+			]
+		];
+	}
+
+protected:
+	TWeakPtr<SAnimationEditorViewportTabBody> AnimViewportPtr;
+	TAttribute<float> CustomSpeed = 1.0f;
+	FOnCustomSpeedChanged OnCustomSpeedChanged;
+};
 ///////////////////////////////////////////////////////////
 // SAnimViewportToolBar
 
@@ -1129,6 +1175,17 @@ TSharedRef<SWidget> SAnimViewportToolBar::GeneratePlaybackMenu() const
 				{
 					InMenuBuilder.AddMenuEntry( Actions.PlaybackSpeedCommands[PlaybackSpeedIndex] );
 				}
+				TSharedPtr<SWidget> AnimSpeedWidget = SNew(SCusttomAnimationSpeedSetting)
+														.AnimEditorViewport(Viewport)
+														.CustomSpeed_Lambda([Viewport = Viewport]()
+															{
+																return Viewport.Pin()->GetCustomAnimationSpeed();
+															})
+														.OnCustomSpeedChanged_Lambda([Viewport = Viewport](float CustomSpeed)
+															{
+																return Viewport.Pin()->SetCustomAnimationSpeed(CustomSpeed);
+															});
+				InMenuBuilder.AddWidget(AnimSpeedWidget.ToSharedRef(), LOCTEXT("PlaybackMenu_Speed_Custom", "Custom Speed:"));
 			}
 			InMenuBuilder.EndSection();
 		}
@@ -1160,6 +1217,17 @@ void SAnimViewportToolBar::GenerateTurnTableMenu(FMenuBuilder& MenuBuilder) cons
 		{
 			MenuBuilder.AddMenuEntry(Actions.TurnTableSpeeds[i]);
 		}
+		TSharedPtr<SWidget> AnimSpeedWidget = SNew(SCusttomAnimationSpeedSetting)
+												.AnimEditorViewport(Viewport)
+												.CustomSpeed_Lambda([Viewport = Viewport]()
+													{
+														return Viewport.Pin()->GetCustomTurnTableSpeed();
+													})
+												.OnCustomSpeedChanged_Lambda([Viewport = Viewport](float CustomSpeed)
+													{
+														Viewport.Pin()->SetCustomTurnTableSpeed(CustomSpeed);
+													});
+		MenuBuilder.AddWidget(AnimSpeedWidget.ToSharedRef(), LOCTEXT("PlaybackMenu_Speed_Custom", "Custom Speed:"));
 	}
 	MenuBuilder.EndSection();
 	MenuBuilder.PopCommandList();
@@ -1205,7 +1273,7 @@ FText SAnimViewportToolBar::GetPlaybackMenuLabel() const
 		{
 			if (Viewport.Pin()->IsPlaybackSpeedSelected(i))
 			{
-				int32 NumFractionalDigits = (i == EAnimationPlaybackSpeeds::Quarter) ? 2 : 1;
+				const int32 NumFractionalDigits = (i == EAnimationPlaybackSpeeds::Quarter || i == EAnimationPlaybackSpeeds::ThreeQuarters) ? 2 : 1;
 
 				const FNumberFormattingOptions FormatOptions = FNumberFormattingOptions()
 					.SetMinimumFractionalDigits(NumFractionalDigits)
diff --git a/Engine/Source/Editor/Persona/Private/SAnimationEditorViewport.cpp b/Engine/Source/Editor/Persona/Private/SAnimationEditorViewport.cpp
--- a/Engine/Source/Editor/Persona/Private/SAnimationEditorViewport.cpp
+++ b/Engine/Source/Editor/Persona/Private/SAnimationEditorViewport.cpp
@@ -1131,7 +1131,11 @@ void SAnimationEditorViewportTabBody::OnSetTurnTableSpeed(int32 SpeedIndex)
 	UDebugSkelMeshComponent* PreviewComponent = GetPreviewScene()->GetPreviewMeshComponent();
 	if (PreviewComponent)
 	{
-		PreviewComponent->TurnTableSpeedScaling = EAnimationPlaybackSpeeds::Values[SelectedTurnTableSpeed];
+		const float TurnTableSpeed = (SelectedTurnTableSpeed == EAnimationPlaybackSpeeds::Custom)
+			? GetCustomTurnTableSpeed()
+			: EAnimationPlaybackSpeeds::Values[SelectedTurnTableSpeed];
+
+		PreviewComponent->TurnTableSpeedScaling = TurnTableSpeed;
 	}
 }
 
@@ -1403,6 +1407,29 @@ float SAnimationEditorViewportTabBody::GetBoneDrawSize() const
 	return AnimViewportClient->GetBoneDrawSize();
 }
 
+void SAnimationEditorViewportTabBody::SetCustomAnimationSpeed(float InCusteomAnimationSpeed)
+{
+	TSharedRef<FAnimationViewportClient> AnimViewportClient = StaticCastSharedRef<FAnimationViewportClient>(LevelViewportClient.ToSharedRef());
+	AnimViewportClient->SetCustomAnimationSpeed(InCusteomAnimationSpeed);
+}
+
+float SAnimationEditorViewportTabBody::GetCustomAnimationSpeed() const
+{
+	TSharedRef<FAnimationViewportClient> AnimViewportClient = StaticCastSharedRef<FAnimationViewportClient>(LevelViewportClient.ToSharedRef());
+	return AnimViewportClient->GetCustomAnimationSpeed();
+}
+
+void SAnimationEditorViewportTabBody::SetCustomTurnTableSpeed(float InCustomTurnTableSpeed)
+{
+	CustomTurnTableSpeed = InCustomTurnTableSpeed;
+	OnSetTurnTableSpeed(EAnimationPlaybackSpeeds::Custom);
+}
+
+float SAnimationEditorViewportTabBody::GetCustomTurnTableSpeed() const
+{
+	return CustomTurnTableSpeed;
+}
+
 void SAnimationEditorViewportTabBody::OnSetBoneDrawMode(int32 BoneDrawMode)
 {
 	TSharedRef<FAnimationViewportClient> AnimViewportClient = StaticCastSharedRef<FAnimationViewportClient>(LevelViewportClient.ToSharedRef());
diff --git a/Engine/Source/Editor/Persona/Private/SAnimationEditorViewport.h b/Engine/Source/Editor/Persona/Private/SAnimationEditorViewport.h
--- a/Engine/Source/Editor/Persona/Private/SAnimationEditorViewport.h
+++ b/Engine/Source/Editor/Persona/Private/SAnimationEditorViewport.h
@@ -293,6 +293,12 @@ public:
 	void SetBoneDrawSize(float BoneDrawSize);
 	float GetBoneDrawSize() const;
 
+	void SetCustomAnimationSpeed(float AnimationSpeed);
+	float GetCustomAnimationSpeed() const;
+
+	void SetCustomTurnTableSpeed(float InCustomTurnTableSpeed);
+	float GetCustomTurnTableSpeed() const;
+
 	/** Function to set LOD model selection*/
 	void OnSetLODModel(int32 LODSelectionType);
 	void OnSetLODTrackDebuggedInstance();
@@ -515,6 +521,8 @@ private:
 	EAnimationPlaybackSpeeds::Type SelectedTurnTableSpeed;
 	/** Selected turn table mode */
 	EPersonaTurnTableMode::Type SelectedTurnTableMode;
+	/** Custom Animation speed in the viewport. Transient setting. */
+	float CustomTurnTableSpeed = 1.0f;
 
 	void OnSetTurnTableSpeed(int32 SpeedIndex);
 	void OnSetTurnTableMode(int32 ModeIndex);
diff --git a/Engine/Source/Editor/Persona/Public/AnimationEditorViewportClient.h b/Engine/Source/Editor/Persona/Public/AnimationEditorViewportClient.h
--- a/Engine/Source/Editor/Persona/Public/AnimationEditorViewportClient.h
+++ b/Engine/Source/Editor/Persona/Public/AnimationEditorViewportClient.h
@@ -60,10 +60,12 @@ namespace EAnimationPlaybackSpeeds
 		OneTenth = 0,
 		Quarter,
 		Half,
+		ThreeQuarters,
 		Normal,
 		Double,
 		FiveTimes,
 		TenTimes,
+		Custom,
 		NumPlaybackSpeeds
 	};
 
@@ -206,7 +208,11 @@ public:
 	/** Access Bone Draw size config option*/
 	void SetBoneDrawSize(const float InBoneDrawSize);
 	float GetBoneDrawSize() const;
-	
+
+	/** Access CustomAnimationSpeed config option*/
+	void SetCustomAnimationSpeed(const float InCustomAnimationSpeed);
+	float GetCustomAnimationSpeed() const;
+
 	/** Function to set Bone Draw  mode for the EBoneDrawType */
 	void SetBoneDrawMode(EBoneDrawMode::Type AxesMode);
 
@@ -477,6 +483,9 @@ private:
 	void HandlePreviewScenePostTick();
 
 private:
+	/** Custom Animation speed in the viewport. Transient setting. */
+	float CustomAnimationSpeed = 1.0f;
+
 	/** Size to draw bones in the viewport. Transient setting. */
 	float BoneDrawSize=1.0f;
 	
