Sync markers can now be managed in the same way as notifies

Sync markers are now visible in the notifies tab
Added sync marker functionality to find/replace tab
Refactored find/replace tab to allow it to be extended via implementing new UAnimAssetFindReplaceProcessor classes
Fixed a few workflow edge cases
Optimized string handling while searching/replacing

diff --git a/Engine/Source/Editor/Persona/Private/AnimAssetFindReplace.cpp b/Engine/Source/Editor/Persona/Private/AnimAssetFindReplace.cpp
new file mode 100644
--- /dev/null
+++ b/Engine/Source/Editor/Persona/Private/AnimAssetFindReplace.cpp
@@ -0,0 +1,424 @@
+﻿// Copyright Epic Games, Inc. All Rights Reserved.
+
+#include "AnimAssetFindReplace.h"
+
+#include "PropertyCustomizationHelpers.h"
+#include "SAnimAssetFindReplace.h"
+#include "Animation/PoseAsset.h"
+#include "Animation/Skeleton.h"
+#include "Animation/AnimationAsset.h"
+#include "Animation/AnimSequence.h"
+#include "Engine/SkeletalMesh.h"
+#include "AssetRegistry/AssetRegistryModule.h"
+#include "Widgets/Views/SListView.h"
+#include "Widgets/Input/SSearchBox.h"
+#include "Framework/Application/SlateApplication.h"
+
+#define LOCTEXT_NAMESPACE "AnimAssetFindReplace"
+
+class SAutoCompleteSearchBox : public SCompoundWidget
+{
+public:
+	SLATE_BEGIN_ARGS(SAutoCompleteSearchBox) {}
+
+	/** The text displayed in the SearchBox when no text has been entered */
+	SLATE_ATTRIBUTE(FText, HintText)
+
+	/** The text displayed in the SearchBox when it's created */
+	SLATE_ATTRIBUTE(FText, InitialText)
+
+	/** Invoked whenever the text changes */
+	SLATE_EVENT(FOnTextChanged, OnTextChanged)
+
+	/** Items to show in the autocomplete popup */
+	SLATE_ARGUMENT(TSharedPtr<TArray<TSharedPtr<FString>>>, AutoCompleteItems)
+	
+	SLATE_END_ARGS()
+
+	void Construct(const FArguments& InArgs)
+	{
+		AutoCompleteItems = InArgs._AutoCompleteItems;
+
+		ChildSlot
+		[
+			SAssignNew(MenuAnchor, SMenuAnchor)
+			.Method(EPopupMethod::CreateNewWindow)
+			.Placement(EMenuPlacement::MenuPlacement_BelowAnchor)
+			.MenuContent
+			(
+				SNew(SBox)
+				.MaxDesiredHeight(200.0f)
+				.MinDesiredWidth(200.0f)
+				[
+					SAssignNew(AutoCompleteList, SListView<TSharedPtr<FString>>)
+					.SelectionMode(ESelectionMode::Single)
+					.ListItemsSource(&FilteredAutoCompleteItems)
+					.OnSelectionChanged_Lambda([this](TSharedPtr<FString> InString, ESelectInfo::Type InSelectInfo)
+					{
+						if(InString.IsValid() && InSelectInfo == ESelectInfo::OnMouseClick)
+						{
+							TGuardValue<bool> GuardValue(bSettingTextFromSearchItem, true);
+							SearchBox->SetText(FText::FromString(*InString.Get()));
+							MenuAnchor->SetIsOpen(false);
+						}
+					})
+					.OnGenerateRow_Lambda([this](TSharedPtr<FString> InString, const TSharedRef<STableViewBase>& InTableView)
+					{
+						if(InString.IsValid())
+						{
+							return
+								SNew(STableRow<TSharedPtr<FString>>, InTableView)
+								.Content()
+								[
+									SNew(STextBlock)
+									.Text(FText::FromString(*InString.Get()))
+									.HighlightText_Lambda([this]()
+									{
+										return SearchBox->GetText();
+									})
+								];
+						}
+						return SNew(STableRow<TSharedPtr<FString>>, InTableView);
+					})
+					.OnKeyDownHandler_Lambda([this](const FGeometry& InGeometry, const FKeyEvent& InKeyEvent)
+					{
+						if(InKeyEvent.GetKey() == EKeys::Enter)
+						{
+							TArray<TSharedPtr<FString>> SelectedItems;
+							AutoCompleteList->GetSelectedItems(SelectedItems);
+							if(SelectedItems.Num() > 0 && SelectedItems[0].IsValid())
+							{
+								TGuardValue<bool> GuardValue(bSettingTextFromSearchItem, true);
+								SearchBox->SetText(FText::FromString(*SelectedItems[0].Get()));
+								MenuAnchor->SetIsOpen(false);
+								FReply::Handled();
+							}
+						}
+						return FReply::Unhandled();
+					})
+				]
+			)
+			[
+				SAssignNew(SearchBox, SSearchBox)
+				.HintText(InArgs._HintText)
+				.InitialText(InArgs._InitialText)
+				.OnTextChanged_Lambda([this, OnTextChanged = InArgs._OnTextChanged](const FText& InText)
+				{
+					RefreshAutoCompleteItems();
+
+					if(!bSettingTextFromSearchItem)
+					{
+						MenuAnchor->SetIsOpen(FilteredAutoCompleteItems.Num() > 0, false);
+					}
+
+					OnTextChanged.ExecuteIfBound(InText);
+				})
+				.OnVerifyTextChanged_Lambda([](const FText& InText, FText& OutErrorMessage)
+				{
+					return FName::IsValidXName(InText.ToString(), FString(INVALID_NAME_CHARACTERS), &OutErrorMessage);
+				})
+			]
+		];
+		
+		RefreshAutoCompleteItems();
+	}
+
+	void FilterItems(const FText& InText)
+	{
+		FilteredAutoCompleteItems.Empty();
+
+		for(TSharedPtr<FString> String : *AutoCompleteItems.Get())
+		{
+			if(String.Get()->Contains(InText.ToString()))
+			{
+				FilteredAutoCompleteItems.Add(String);
+			}
+		}
+	}
+
+	void RefreshAutoCompleteItems()
+	{
+		FilterItems(SearchBox->GetText());
+		AutoCompleteList->RequestListRefresh();
+	}
+
+	TSharedRef<SSearchBox> GetSearchBox() const
+	{
+		return SearchBox.ToSharedRef();
+	}
+	
+	virtual void OnFocusChanging(const FWeakWidgetPath& PreviousFocusPath, const FWidgetPath& NewWidgetPath, const FFocusEvent& InFocusEvent) override
+	{
+		if(PreviousFocusPath.ContainsWidget(SearchBox.Get()) && !NewWidgetPath.ContainsWidget(MenuAnchor->GetMenuWindow().Get()))
+		{
+			MenuAnchor->SetIsOpen(false);
+		}
+		SCompoundWidget::OnFocusChanging(PreviousFocusPath, NewWidgetPath, InFocusEvent);
+	}
+
+	virtual FReply OnPreviewKeyDown(const FGeometry& MyGeometry, const FKeyEvent& InKeyEvent) override
+	{
+		if(InKeyEvent.GetKey() == EKeys::Down && MenuAnchor->IsOpen())
+		{
+			// switch focus to the drop-down autocomplete list
+			return FReply::Handled().SetUserFocus(AutoCompleteList.ToSharedRef(), EFocusCause::Navigation);
+		}
+		return SCompoundWidget::OnPreviewKeyDown(MyGeometry, InKeyEvent); 
+	}
+	
+private:
+	TArray<TSharedPtr<FString>> FilteredAutoCompleteItems;
+	TSharedPtr<TArray<TSharedPtr<FString>>> AutoCompleteItems;
+	TSharedPtr<SListView<TSharedPtr<FString>>> AutoCompleteList;
+	TSharedPtr<SMenuAnchor> MenuAnchor;
+	TSharedPtr<SSearchBox> SearchBox;
+	bool bSettingTextFromSearchItem = false;
+};
+
+void UAnimAssetFindReplaceProcessor::Initialize(TSharedRef<SAnimAssetFindReplace> InWidget)
+{
+	Widget = InWidget;
+}
+
+void UAnimAssetFindReplaceProcessor::RequestUIRefresh()
+{
+	if(TSharedPtr<SAnimAssetFindReplace> PinnedWidget = Widget.Pin())
+	{
+		PinnedWidget->RequestRefresh();
+	}
+}
+
+void UAnimAssetFindReplaceProcessor::RequestRefreshSearchResults()
+{
+	if(TSharedPtr<SAnimAssetFindReplace> PinnedWidget = Widget.Pin())
+	{
+		PinnedWidget->RequestRefreshSearchResults();
+	}
+}
+
+void UAnimAssetFindReplaceProcessor_StringBase::SetFindString(const FString& InString)
+{
+	FindString = InString;
+	RequestRefreshSearchResults();
+}
+
+void UAnimAssetFindReplaceProcessor_StringBase::SetReplaceString(const FString& InString)
+{
+	ReplaceString = InString;
+	RequestRefreshSearchResults();
+}
+
+void UAnimAssetFindReplaceProcessor_StringBase::SetSkeletonFilter(const FAssetData& InSkeletonFilter)
+{
+	SkeletonFilter = InSkeletonFilter;
+	RequestRefreshSearchResults();
+}
+
+void UAnimAssetFindReplaceProcessor_StringBase::SetFindWholeWord(bool bInFindWholeWord)
+{
+	bFindWholeWord = bInFindWholeWord;
+	RequestRefreshSearchResults();
+}
+
+void UAnimAssetFindReplaceProcessor_StringBase::SetSearchCase(ESearchCase::Type InSearchCase)
+{
+	SearchCase = InSearchCase;
+	RequestRefreshSearchResults();
+}
+
+void UAnimAssetFindReplaceProcessor_StringBase::ExtendToolbar(FToolMenuSection& InSection)
+{
+	FToolUIAction MatchCaseCheckbox;
+	MatchCaseCheckbox.ExecuteAction = FToolMenuExecuteAction::CreateLambda([this](const FToolMenuContext& InContext)
+	{
+		SetSearchCase(GetSearchCase() == ESearchCase::CaseSensitive ? ESearchCase::IgnoreCase : ESearchCase::CaseSensitive);
+	});
+
+	MatchCaseCheckbox.GetActionCheckState = FToolMenuGetActionCheckState::CreateLambda([this](const FToolMenuContext& InContext)
+	{
+		return GetSearchCase() == ESearchCase::CaseSensitive ? ECheckBoxState::Checked : ECheckBoxState::Unchecked;
+	});
+	
+	InSection.AddEntry(
+		FToolMenuEntry::InitToolBarButton(
+			"MatchCase",
+			MatchCaseCheckbox,
+			LOCTEXT("MatchCaseCheckboxLabel", "Match Case"),
+			LOCTEXT("MatchCaseCheckboxTooltip", "Whether to match case when searching."),
+			FSlateIcon(FAppStyle::GetAppStyleSetName(), "Icons.Clipboard"),
+			EUserInterfaceActionType::ToggleButton));
+
+	FToolUIAction MatchWholeWordCheckbox;
+	MatchWholeWordCheckbox.ExecuteAction = FToolMenuExecuteAction::CreateLambda([this](const FToolMenuContext& InContext)
+	{
+		SetFindWholeWord(!GetFindWholeWord());
+	});
+
+	MatchWholeWordCheckbox.GetActionCheckState = FToolMenuGetActionCheckState::CreateLambda([this](const FToolMenuContext& InContext)
+	{
+		return GetFindWholeWord() ? ECheckBoxState::Checked : ECheckBoxState::Unchecked;
+	});
+	
+	InSection.AddEntry(
+		FToolMenuEntry::InitToolBarButton(
+			"MatchWholeWord",
+			MatchWholeWordCheckbox,
+			LOCTEXT("MatchWholeWordCheckboxLabel", "Match Whole Word"),
+			LOCTEXT("MatchWholeWordCheckboxTooltip", "Whether to match the whole word or just part of the word when searching."),
+			FSlateIcon(FAppStyle::GetAppStyleSetName(), "Icons.Use"),
+			EUserInterfaceActionType::ToggleButton));
+
+	InSection.AddDynamicEntry("SkeletonFilter", FNewToolMenuSectionDelegate::CreateLambda([this](FToolMenuSection& InSection)
+	{
+		if(TSharedPtr<SAnimAssetFindReplace> Widget = AnimAssetFindReplacePrivate::GetWidgetFromContext(InSection.Context))
+		{
+			InSection.AddEntry(
+			FToolMenuEntry::InitWidget(
+				"SkeletonFilterWidget",
+				SNew(SHorizontalBox)
+				.ToolTipText(LOCTEXT("SkeletonFilterTooltip", "Choose a Skeleton asset to filter results by."))
+				+SHorizontalBox::Slot()
+				.AutoWidth()
+				.VAlign(VAlign_Center)
+				.Padding(5.0f, 0.0f)
+				[
+					SNew(STextBlock)
+					.Text(LOCTEXT("SkeletonFilterLabel", "Skeleton"))
+				]
+				+SHorizontalBox::Slot()
+				.AutoWidth()
+				.VAlign(VAlign_Center)
+				[
+					SNew(SObjectPropertyEntryBox)
+					.ObjectPath_Lambda([this]()
+					{
+						return GetSkeletonFilter().GetObjectPathString();
+					})
+					.OnObjectChanged_Lambda([this](const FAssetData& InAssetData)
+					{
+						SetSkeletonFilter(InAssetData);
+					})
+					.AllowedClass(USkeleton::StaticClass())
+				],
+			FText::GetEmpty(),
+			true, true, true));
+		}
+	}));
+}
+
+TSharedRef<SWidget> UAnimAssetFindReplaceProcessor_StringBase::MakeFindReplaceWidget()
+{
+	return SAssignNew(FindReplaceWidget, SVerticalBox)
+		+SVerticalBox::Slot()
+		.Padding(6.0f, 2.0f)
+		[
+			SAssignNew(FindSearchBox, SAutoCompleteSearchBox)
+			.AutoCompleteItems(AutoCompleteItems)
+			.HintText(LOCTEXT("FindLabel", "Find"))
+			.ToolTipText(LOCTEXT("FindLabel", "Find"))
+			.InitialText_Lambda([this](){ return FText::FromString(FindString); })
+			.OnTextChanged_Lambda([this](const FText& InText)
+			{
+				FindString = InText.ToString();
+				RequestRefreshSearchResults();
+			})
+		]
+		+SVerticalBox::Slot()
+		.Padding(6.0f, 2.0f)
+		[
+			SAssignNew(ReplaceSearchBox, SAutoCompleteSearchBox)
+			.AutoCompleteItems(AutoCompleteItems)
+			.HintText(LOCTEXT("ReplaceLabel", "Replace With"))
+			.ToolTipText(LOCTEXT("ReplaceLabel", "Replace With"))
+			.InitialText_Lambda([this](){ return FText::FromString(ReplaceString); })
+			.OnTextChanged_Lambda([this](const FText& InText)
+			{
+				ReplaceString = InText.ToString();
+			})
+		];
+}
+
+void UAnimAssetFindReplaceProcessor_StringBase::FocusInitialWidget()
+{
+	if (FindSearchBox.IsValid())
+	{
+		FWidgetPath WidgetToFocusPath;
+		FSlateApplication::Get().GeneratePathToWidgetUnchecked(FindSearchBox.Pin()->GetSearchBox(), WidgetToFocusPath);
+		FSlateApplication::Get().SetKeyboardFocus(WidgetToFocusPath, EFocusCause::SetDirectly);
+		WidgetToFocusPath.GetWindow()->SetWidgetToFocusOnActivate(FindSearchBox.Pin()->GetSearchBox());
+	}
+}
+
+bool UAnimAssetFindReplaceProcessor_StringBase::CanCurrentlyReplace() const
+{
+	return !ReplaceString.IsEmpty();
+}
+
+bool UAnimAssetFindReplaceProcessor_StringBase::CanCurrentlyRemove() const
+{
+	return true;
+}
+
+void UAnimAssetFindReplaceProcessor_StringBase::RefreshCachedData()
+{
+	AutoCompleteItems->Empty();
+
+	// We use the asset registry to query all assets and accumulate their curve names
+	const FAssetRegistryModule& AssetRegistryModule = FModuleManager::LoadModuleChecked<FAssetRegistryModule>("AssetRegistry");
+
+	FARFilter Filter;
+	for(const UClass* AssetClass : GetSupportedAssetTypes())
+	{
+		Filter.ClassPaths.Add(AssetClass->GetClassPathName());
+	}
+	Filter.bRecursiveClasses = true;
+
+	TArray<FAssetData> FoundAssetData;
+	AssetRegistryModule.Get().GetAssets(Filter, FoundAssetData);
+
+	TSet<FString> UniqueNames;
+
+	GetAutoCompleteNames(FoundAssetData, UniqueNames);
+
+	for(FString& UniqueName : UniqueNames)
+	{
+		AutoCompleteItems->Add(MakeShared<FString>(MoveTemp(UniqueName)));
+	}
+
+	AutoCompleteItems->Sort([](TSharedPtr<FString> InLHS, TSharedPtr<FString> InRHS)
+	{
+		return *InLHS.Get() < *InRHS.Get();
+	});
+
+	if(TSharedPtr<SAutoCompleteSearchBox> FindWidget = FindSearchBox.Pin())
+	{
+		FindWidget->RefreshAutoCompleteItems();
+	}
+
+	if(TSharedPtr<SAutoCompleteSearchBox> ReplaceWidget = ReplaceSearchBox.Pin())
+	{
+		ReplaceWidget->RefreshAutoCompleteItems();
+	}
+}
+
+bool UAnimAssetFindReplaceProcessor_StringBase::NameMatches(FStringView InNameStringView) const
+{
+	if(bFindWholeWord)
+	{
+		if(InNameStringView.Compare(FindString, SearchCase) == 0)
+		{
+			return true;
+		}
+	}
+	else
+	{
+		if(UE::String::FindFirst(InNameStringView, FindString, SearchCase) != INDEX_NONE)
+		{
+			return true;
+		}
+	}
+
+	return false;
+}
+
+#undef LOCTEXT_NAMESPACE
\ No newline at end of file
diff --git a/Engine/Source/Editor/Persona/Private/AnimAssetFindReplaceCurves.cpp b/Engine/Source/Editor/Persona/Private/AnimAssetFindReplaceCurves.cpp
new file mode 100644
--- /dev/null
+++ b/Engine/Source/Editor/Persona/Private/AnimAssetFindReplaceCurves.cpp
@@ -0,0 +1,462 @@
+﻿// Copyright Epic Games, Inc. All Rights Reserved.
+
+#include "AnimAssetFindReplaceCurves.h"
+
+#include "ISkeletonEditorModule.h"
+#include "SAnimAssetFindReplace.h"
+#include "Animation/PoseAsset.h"
+#include "Animation/Skeleton.h"
+#include "Animation/AnimationAsset.h"
+#include "Animation/AnimSequence.h"
+#include "AssetRegistry/AssetRegistryModule.h"
+#include "Engine/SkeletalMesh.h"
+#include "String/ParseTokens.h"
+
+#define LOCTEXT_NAMESPACE "AnimAssetFindReplaceCurves"
+
+FString UAnimAssetFindReplaceCurves::GetFindResultStringFromAssetData(const FAssetData& InAssetData) const
+{
+	if(GetFindString().IsEmpty())
+	{
+		return FString();
+	}
+
+	auto GetMatchingCurveNamesForAsset = [this](const FAssetData& InAssetData, TArray<FString>& OutCurveNames)
+	{
+		const FString TagValue = InAssetData.GetTagValueRef<FString>(USkeleton::CurveNameTag);
+		if (!TagValue.IsEmpty())
+		{
+			UE::String::ParseTokens(TagValue, *USkeleton::CurveTagDelimiter, [this, &OutCurveNames](FStringView InToken)
+			{
+				if(GetFindWholeWord())
+				{
+					if(InToken.Compare(GetFindString(), GetSearchCase()) == 0)
+					{
+						OutCurveNames.Add(FString(InToken));
+					}
+				}
+				else
+				{
+					if(UE::String::FindFirst(InToken, GetFindString(), GetSearchCase()) != INDEX_NONE)
+					{
+						OutCurveNames.Add(FString(InToken));
+					}
+				}
+			}, UE::String::EParseTokensOptions::SkipEmpty);
+		}
+	};
+
+	TStringBuilder<128> Builder;
+	TArray<FString> CurveNames;
+	GetMatchingCurveNamesForAsset(InAssetData, CurveNames);
+	if(CurveNames.Num() > 0)
+	{
+		for(int32 NameIndex = 0; NameIndex < CurveNames.Num(); ++NameIndex)
+		{
+			Builder.Append(CurveNames[NameIndex]);
+			if(NameIndex != CurveNames.Num() - 1)
+			{
+				Builder.Append(TEXT(", "));
+			}
+		}
+	}
+	return FString(Builder.ToString());
+}
+
+TConstArrayView<UClass*> UAnimAssetFindReplaceCurves::GetSupportedAssetTypes() const
+{
+	static UClass* Types[] = { UPoseAsset::StaticClass(), UAnimSequenceBase::StaticClass(), USkeleton::StaticClass(), USkeletalMesh::StaticClass() };
+	return Types;
+}
+
+bool UAnimAssetFindReplaceCurves::ShouldFilterOutAsset(const FAssetData& InAssetData, bool& bOutIsOldAsset) const
+{
+	FString TagValue;
+	if(InAssetData.GetTagValue<FString>(USkeleton::CurveNameTag, TagValue))
+	{
+		bOutIsOldAsset = false;
+
+		if(GetFindString().IsEmpty())
+		{
+			return true;
+		}
+		
+		if(GetSkeletonFilter().IsValid() )
+		{
+			if(InAssetData.GetClass() != USkeleton::StaticClass())
+			{
+				FString SkeletonPath;
+				if(InAssetData.GetTagValue<FString>(TEXT("Skeleton"), SkeletonPath))
+				{
+					if(SkeletonPath != GetSkeletonFilter().GetExportTextName())
+					{
+						return true;
+					}
+				}
+			}
+			else
+			{
+				if(InAssetData.ToSoftObjectPath() != GetSkeletonFilter().ToSoftObjectPath())
+				{
+					return true;
+				}
+			}
+		}
+
+		bool bFoundMatch = false;
+		UE::String::ParseTokens(TagValue, *USkeleton::CurveTagDelimiter, [this, &bFoundMatch](FStringView InToken)
+		{
+			if(NameMatches(InToken))
+			{
+				bFoundMatch = true;
+			}
+		}, UE::String::EParseTokensOptions::SkipEmpty);
+
+		if(bFoundMatch)
+		{
+			return false;
+		}
+	}
+	else
+	{
+		FAssetRegistryModule& AssetRegistryModule = FModuleManager::LoadModuleChecked<FAssetRegistryModule>("AssetRegistry");
+		UClass* AssetClass = InAssetData.GetClass();
+		if(AssetClass->IsChildOf(UAnimSequenceBase::StaticClass()))
+		{
+			// Check the package object version - the asset was saving empty tags for curves, so the absence of curves is
+			// not the same as an empty value
+			FAssetPackageData PackageData;
+			const UE::AssetRegistry::EExists PackageExists = AssetRegistryModule.Get().TryGetAssetPackageData(InAssetData.PackageName, PackageData);
+			if (PackageExists == UE::AssetRegistry::EExists::Exists)
+			{
+				bOutIsOldAsset = PackageData.FileVersionUE < VER_UE4_SKELETON_ADD_SMARTNAMES;
+			}
+			else
+			{
+				// Does not exist or unknown - treat it as 'old'
+				bOutIsOldAsset = true;
+			}
+		}
+		else if(AssetClass->IsChildOf(UPoseAsset::StaticClass()))
+		{
+			// Check the package custom version - the asset was saving empty tags for curves, so the absence of curves is
+			// not the same as an empty value
+			FAssetPackageData PackageData;
+			const UE::AssetRegistry::EExists PackageExists = AssetRegistryModule.Get().TryGetAssetPackageData(InAssetData.PackageName, PackageData);
+			if (PackageExists == UE::AssetRegistry::EExists::Exists)
+			{
+				for(const UE::AssetRegistry::FPackageCustomVersion& CustomVersion : PackageData.GetCustomVersions())
+				{
+					if(CustomVersion.Key == FAnimPhysObjectVersion::GUID)
+					{
+						bOutIsOldAsset = CustomVersion.Version < FAnimPhysObjectVersion::SmartNameRefactorForDeterministicCooking;
+					}
+				}
+
+				// No FAnimPhysObjectVersion, treat as old
+				bOutIsOldAsset = true;
+			}
+			else
+			{
+				// Does not exist or unknown - treat it as 'old'
+				bOutIsOldAsset = true;
+			}
+		}
+		else if(AssetClass->IsChildOf(USkeleton::StaticClass()))
+		{
+			bOutIsOldAsset = true;
+		}
+		else if(AssetClass->IsChildOf(USkeletalMesh::StaticClass()))
+		{
+			// Skeletal meshes didn't have curves before, so cant be 'old' 
+			bOutIsOldAsset = false;
+		}
+		else
+		{
+			// Assume unknown assets are not 'old'
+			bOutIsOldAsset = false;
+		}
+	}
+
+	return true;
+}
+
+void UAnimAssetFindReplaceCurves::ReplaceInAsset(const FAssetData& InAssetData) const
+{
+	if(UObject* Asset = InAssetData.GetAsset())
+	{
+		if(UAnimSequenceBase* AnimSequenceBase = Cast<UAnimSequenceBase>(Asset))
+		{
+			Asset->MarkPackageDirty();
+
+			if(GetFindWholeWord())
+			{
+				const FAnimationCurveIdentifier FindCurveId(FName(GetFindString()), ERawCurveTrackTypes::RCT_Float);
+				const FAnimationCurveIdentifier ReplaceCurveId(FName(GetReplaceString()), ERawCurveTrackTypes::RCT_Float);
+				IAnimationDataController::FScopedBracket ScopedBracket(AnimSequenceBase->GetController(), LOCTEXT("ReplaceCurves", "Replace Curves"));
+				AnimSequenceBase->GetController().RenameCurve(FindCurveId, ReplaceCurveId);
+			}
+			else
+			{
+				TArray<TPair<FAnimationCurveIdentifier, FAnimationCurveIdentifier>> FindReplacePairs;
+				const TArray<FFloatCurve>& Curves = AnimSequenceBase->GetDataModel()->GetFloatCurves();
+				for(const FFloatCurve& Curve : Curves)
+				{
+					FString CurveName = Curve.GetName().ToString();
+					if(NameMatches(CurveName))
+					{
+						const FAnimationCurveIdentifier FindCurveId(*CurveName, ERawCurveTrackTypes::RCT_Float);
+						const FString NewName = CurveName.Replace(*GetFindStringRef(), *GetReplaceStringRef(), GetSearchCase());
+						const FAnimationCurveIdentifier ReplaceCurveId(*NewName, ERawCurveTrackTypes::RCT_Float);
+
+						FindReplacePairs.Emplace(FindCurveId, ReplaceCurveId);
+					}
+				}
+
+				if(FindReplacePairs.Num() > 0)
+				{
+					IAnimationDataController::FScopedBracket ScopedBracket(AnimSequenceBase->GetController(), LOCTEXT("ReplaceCurves", "Replace Curves"));
+
+					for(const TPair<FAnimationCurveIdentifier, FAnimationCurveIdentifier>& FindReplacePair : FindReplacePairs)
+					{
+						AnimSequenceBase->GetController().RenameCurve(FindReplacePair.Key, FindReplacePair.Value);
+					}
+				}
+			}
+		}
+		else if(UPoseAsset* PoseAsset = Cast<UPoseAsset>(Asset))
+		{
+			if(GetFindWholeWord())
+			{
+				Asset->Modify();
+				
+				const FName FindCurveName(GetFindString());
+				const FName ReplaceCurveName(GetReplaceString());
+				PoseAsset->RenamePoseOrCurveName(FindCurveName, ReplaceCurveName);
+			}
+			else
+			{
+				TArray<TPair<FName, FName>> FindReplacePairs;
+
+				for(const FName& PoseName : PoseAsset->GetPoseFNames())
+				{
+					FString CurveName = PoseName.ToString();
+					if(NameMatches(CurveName))
+					{
+						const FName FindCurveName(*CurveName);
+						const FString NewName = CurveName.Replace(*GetFindStringRef(), *GetReplaceStringRef(), GetSearchCase());
+						const FName ReplaceCurveName(*NewName);
+
+						FindReplacePairs.Emplace(FindCurveName, ReplaceCurveName);
+					}
+				}
+
+				if(FindReplacePairs.Num() > 0)
+				{
+					Asset->Modify();
+
+					for(const TPair<FName, FName>& FindReplacePair : FindReplacePairs)
+					{
+						PoseAsset->RenamePoseOrCurveName(FindReplacePair.Key, FindReplacePair.Value);
+					}
+				}
+			}
+		}
+		else if(USkeleton* Skeleton = Cast<USkeleton>(Asset))
+		{
+			if(GetFindWholeWord())
+			{
+				Asset->Modify();
+
+				const FName FindCurveName(GetFindString());
+				const FName ReplaceCurveName(GetReplaceString());
+				Skeleton->RenameCurveMetaData(FindCurveName, ReplaceCurveName);
+			}
+			else
+			{
+				TArray<TPair<FName, FName>> FindReplacePairs;
+
+				Skeleton->ForEachCurveMetaData([this, &FindReplacePairs](FName InCurveName, const FCurveMetaData& InMetaData)
+				{
+					const FString CurveNameString = InCurveName.ToString();
+					if(NameMatches(CurveNameString))
+					{
+						const FName FindCurveName(InCurveName);
+						const FString NewName = CurveNameString.Replace(*GetFindStringRef(), *GetReplaceStringRef(), GetSearchCase());
+						const FName ReplaceCurveName(*NewName);
+
+						FindReplacePairs.Emplace(FindCurveName, ReplaceCurveName);
+					}
+				});
+
+				if(FindReplacePairs.Num() > 0)
+				{
+					Asset->Modify();
+
+					for(const TPair<FName, FName>& FindReplacePair : FindReplacePairs)
+					{
+						Skeleton->RenameCurveMetaData(FindReplacePair.Key, FindReplacePair.Value);
+					}
+				}
+			}
+		}
+		else if(USkeletalMesh* SkeletalMesh = Cast<USkeletalMesh>(Asset))
+		{
+			if(UAnimCurveMetaData* AnimCurveMetaData = SkeletalMesh->GetAssetUserData<UAnimCurveMetaData>())
+			{
+				if(GetFindWholeWord())
+				{
+					Asset->Modify();
+
+					const FName FindCurveName(GetFindString());
+					const FName ReplaceCurveName(GetReplaceString());
+					AnimCurveMetaData->RenameCurveMetaData(FindCurveName, ReplaceCurveName);
+				}
+				else
+				{
+					TArray<TPair<FName, FName>> FindReplacePairs;
+
+					AnimCurveMetaData->ForEachCurveMetaData([this, &FindReplacePairs](FName InCurveName, const FCurveMetaData& InMetaData)
+					{
+						const FString CurveNameString = InCurveName.ToString();
+						if(NameMatches(CurveNameString))
+						{
+							const FName FindCurveName(InCurveName);
+							const FString NewName = CurveNameString.Replace(*GetFindStringRef(), *GetReplaceStringRef(), GetSearchCase());
+							const FName ReplaceCurveName(*NewName);
+
+							FindReplacePairs.Emplace(FindCurveName, ReplaceCurveName);
+						}
+					});
+
+					if(FindReplacePairs.Num() > 0)
+					{
+						Asset->Modify();
+
+						for(const TPair<FName, FName>& FindReplacePair : FindReplacePairs)
+						{
+							AnimCurveMetaData->RenameCurveMetaData(FindReplacePair.Key, FindReplacePair.Value);
+						}
+					}
+				}
+			}
+		}
+	}
+}
+
+void UAnimAssetFindReplaceCurves::RemoveInAsset(const FAssetData& InAssetData) const
+{
+	if(UObject* Asset = InAssetData.GetAsset())
+	{
+		if(UAnimSequenceBase* AnimSequenceBase = Cast<UAnimSequenceBase>(Asset))
+		{
+			Asset->MarkPackageDirty();
+
+			if(GetFindWholeWord())
+			{
+				const FAnimationCurveIdentifier CurveId(FName(GetFindString()), ERawCurveTrackTypes::RCT_Float);
+				IAnimationDataController::FScopedBracket ScopedBracket(AnimSequenceBase->GetController(), LOCTEXT("RemoveCurves", "Remove Curves"));
+				AnimSequenceBase->GetController().RemoveCurve(CurveId);
+			}
+			else
+			{
+				TSet<FAnimationCurveIdentifier> CurveIdsToRemove;
+				const TArray<FFloatCurve>& Curves = AnimSequenceBase->GetDataModel()->GetFloatCurves();
+				for(const FFloatCurve& Curve : Curves)
+				{
+					FString CurveName = Curve.GetName().ToString();
+					if(NameMatches(CurveName))
+					{
+						const FAnimationCurveIdentifier CurveId(*CurveName, ERawCurveTrackTypes::RCT_Float);
+						CurveIdsToRemove.Add(CurveId);
+					}
+				}
+
+				if(CurveIdsToRemove.Num() > 0)
+				{
+					IAnimationDataController::FScopedBracket ScopedBracket(AnimSequenceBase->GetController(), LOCTEXT("RemoveCurves", "Remove Curves"));
+					for(const FAnimationCurveIdentifier& CurveIdToRemove : CurveIdsToRemove)
+					{
+						AnimSequenceBase->GetController().RemoveCurve(CurveIdToRemove);
+					}
+				}
+			}
+		}
+		else if(UPoseAsset* PoseAsset = Cast<UPoseAsset>(Asset))
+		{
+			if(GetFindWholeWord())
+			{
+				Asset->Modify();
+				PoseAsset->RemovePoseOrCurveNames({ FName(GetFindString()) });
+			}
+			else
+			{
+				TArray<FName> CurveIdsToRemove;
+
+				for(const FName& PoseName : PoseAsset->GetPoseFNames())
+				{
+					if(NameMatches(PoseName.ToString()))
+					{
+						CurveIdsToRemove.AddUnique(PoseName);
+					}
+				}
+
+				for(const FName& CurveName : PoseAsset->GetCurveFNames())
+				{
+					if(NameMatches(CurveName.ToString()))
+					{
+						CurveIdsToRemove.AddUnique(CurveName);
+					}
+				}
+				
+				if(CurveIdsToRemove.Num() > 0)
+				{
+					Asset->Modify();
+					PoseAsset->RemovePoseOrCurveNames(CurveIdsToRemove);
+				}
+			}
+		}
+		else if(USkeleton* Skeleton = Cast<USkeleton>(Asset))
+		{
+			if(GetFindWholeWord())
+			{
+				Asset->Modify();
+				Skeleton->RemoveCurveMetaData(FName(GetFindString()));
+			}
+			else
+			{
+				TArray<FName> CurvesToRemove;
+				Skeleton->ForEachCurveMetaData([this, &CurvesToRemove](FName InCurveName, const FCurveMetaData& InCurveMetaData)
+				{
+				 	if(NameMatches(InCurveName.ToString()))
+				 	{
+						 CurvesToRemove.AddUnique(InCurveName);
+					 }
+				});
+
+				if(CurvesToRemove.Num() > 0)
+				{
+					Asset->Modify();
+					Skeleton->RemoveCurveMetaData(CurvesToRemove);
+				}
+			}
+		}
+	}
+}
+
+void UAnimAssetFindReplaceCurves::GetAutoCompleteNames(TArrayView<FAssetData> InAssetDatas, TSet<FString>& OutUniqueNames) const
+{
+	for (const FAssetData& AssetData : InAssetDatas)
+	{
+		const FString TagValue = AssetData.GetTagValueRef<FString>(USkeleton::CurveNameTag);
+		if (!TagValue.IsEmpty())
+		{
+			UE::String::ParseTokens(TagValue, *USkeleton::CurveTagDelimiter, [&OutUniqueNames](FStringView InToken)
+			{
+				OutUniqueNames.Add(FString(InToken));
+			}, UE::String::EParseTokensOptions::SkipEmpty);
+		}
+	}
+}
+
+#undef LOCTEXT_NAMESPACE
\ No newline at end of file
diff --git a/Engine/Source/Editor/Persona/Private/AnimAssetFindReplaceNotifies.cpp b/Engine/Source/Editor/Persona/Private/AnimAssetFindReplaceNotifies.cpp
new file mode 100644
--- /dev/null
+++ b/Engine/Source/Editor/Persona/Private/AnimAssetFindReplaceNotifies.cpp
@@ -0,0 +1,293 @@
+﻿// Copyright Epic Games, Inc. All Rights Reserved.
+
+#include "AnimAssetFindReplaceNotifies.h"
+
+#include "IEditableSkeleton.h"
+#include "ISkeletonEditorModule.h"
+#include "SAnimAssetFindReplace.h"
+#include "Animation/Skeleton.h"
+#include "Animation/AnimationAsset.h"
+#include "Animation/AnimSequence.h"
+#include "Engine/SkeletalMesh.h"
+#include "String/ParseTokens.h"
+
+#define LOCTEXT_NAMESPACE "AnimAssetFindReplaceNotifies"
+
+FString UAnimAssetFindReplaceNotifies::GetFindResultStringFromAssetData(const FAssetData& InAssetData) const
+{
+	if(GetFindString().IsEmpty())
+	{
+		return FString();
+	}
+
+	auto GetMatchingNotifyNamesForAsset = [this](const FAssetData& InAssetData, TArray<FString>& OutNotifyNames)
+	{
+		const FString TagValue = InAssetData.GetTagValueRef<FString>(USkeleton::AnimNotifyTag);
+		if (!TagValue.IsEmpty())
+		{
+			UE::String::ParseTokens(TagValue, *USkeleton::AnimNotifyTagDelimiter, [this, &OutNotifyNames](FStringView InToken)
+			{
+				if(GetFindWholeWord())
+				{
+					if(InToken.Compare(GetFindString(), GetSearchCase()) == 0)
+					{
+						OutNotifyNames.Add(FString(InToken));
+					}
+				}
+				else
+				{
+					if(UE::String::FindFirst(InToken, GetFindString(), GetSearchCase()) != INDEX_NONE)
+					{
+						OutNotifyNames.Add(FString(InToken));
+					}
+				}
+			}, UE::String::EParseTokensOptions::SkipEmpty);
+		}
+	};
+
+	TStringBuilder<128> Builder;
+	TArray<FString> NotifyNames;
+	GetMatchingNotifyNamesForAsset(InAssetData, NotifyNames);
+	if(NotifyNames.Num() > 0)
+	{
+		for(int32 NameIndex = 0; NameIndex < NotifyNames.Num(); ++NameIndex)
+		{
+			Builder.Append(NotifyNames[NameIndex]);
+			if(NameIndex != NotifyNames.Num() - 1)
+			{
+				Builder.Append(TEXT(", "));
+			}
+		}
+	}
+
+	return FString(Builder.ToString());
+}
+
+TConstArrayView<UClass*> UAnimAssetFindReplaceNotifies::GetSupportedAssetTypes() const
+{
+	static UClass* Types[] = { UAnimSequenceBase::StaticClass(), USkeleton::StaticClass() };
+	return Types;
+}
+
+bool UAnimAssetFindReplaceNotifies::ShouldFilterOutAsset(const FAssetData& InAssetData, bool& bOutIsOldAsset) const
+{
+	FString TagValue;
+	if(InAssetData.GetTagValue<FString>(USkeleton::AnimNotifyTag, TagValue))
+	{
+		bOutIsOldAsset = false;
+
+		if(GetFindString().IsEmpty())
+		{
+			return true;
+		}
+
+		if(GetSkeletonFilter().IsValid() )
+		{
+			if(InAssetData.GetClass() != USkeleton::StaticClass())
+			{
+				FString SkeletonPath;
+				if(InAssetData.GetTagValue<FString>(TEXT("Skeleton"), SkeletonPath))
+				{
+					if(SkeletonPath != GetSkeletonFilter().GetExportTextName())
+					{
+						return true;
+					}
+				}
+			}
+			else
+			{
+				if(InAssetData.ToSoftObjectPath() != GetSkeletonFilter().ToSoftObjectPath())
+				{
+					return true;
+				}
+			}
+		}
+
+		bool bFoundMatch = false;
+		UE::String::ParseTokens(TagValue, *USkeleton::AnimNotifyTagDelimiter, [this, &bFoundMatch](FStringView InToken)
+		{
+			if(NameMatches(InToken))
+			{
+				bFoundMatch = true;
+			}
+		}, UE::String::EParseTokensOptions::SkipEmpty);
+
+		if(bFoundMatch)
+		{
+			return false;
+		}
+	}
+	else
+	{
+		const UClass* AssetClass = InAssetData.GetClass();
+		if(AssetClass->IsChildOf(UAnimSequenceBase::StaticClass()))
+		{
+			bOutIsOldAsset = true;
+		}
+		else if(AssetClass->IsChildOf(USkeleton::StaticClass()))
+		{
+			bOutIsOldAsset = true;
+		}
+		else
+		{
+			// Assume unknown assets are not 'old'
+			bOutIsOldAsset = false;
+		}
+	}
+
+	return true;
+}
+
+void UAnimAssetFindReplaceNotifies::ReplaceInAsset(const FAssetData& InAssetData) const
+{
+	if(UObject* Asset = InAssetData.GetAsset())
+	{
+		if(UAnimSequenceBase* AnimSequenceBase = Cast<UAnimSequenceBase>(Asset))
+		{
+			if(GetFindWholeWord())
+			{
+				Asset->Modify();
+				AnimSequenceBase->RenameNotifies(FName(GetFindString()), FName(GetReplaceString()));
+			}
+			else
+			{
+				TArray<TPair<FName, FName>> FindReplacePairs;
+
+				for(const FAnimNotifyEvent& Notify : AnimSequenceBase->Notifies)
+				{
+					// Only handle named notifiesif(!Notify.IsBlueprintNotify())
+					{
+						FString NotifyName = Notify.NotifyName.ToString();
+						if(NameMatches(NotifyName))
+						{
+							const FName FindNotifyName(*NotifyName);
+							const FString NewName = NotifyName.Replace(*GetFindStringRef(), *GetReplaceStringRef(), GetSearchCase());
+							const FName ReplaceNotifyName(*NewName);
+							FindReplacePairs.AddUnique(TPair<FName, FName>(FindNotifyName, ReplaceNotifyName));
+						}
+					}
+				}
+
+				if(FindReplacePairs.Num() > 0)
+				{
+					Asset->Modify();
+
+					for(const TPair<FName, FName>& FindReplacePair : FindReplacePairs)
+					{
+						AnimSequenceBase->RenameNotifies(FindReplacePair.Key, FindReplacePair.Value);
+					}
+				}
+			}
+		}
+		else if(USkeleton* Skeleton = Cast<USkeleton>(Asset))
+		{
+			ISkeletonEditorModule& SkeletonEditorModule = FModuleManager::GetModuleChecked<ISkeletonEditorModule>("SkeletonEditor");
+			TSharedRef<IEditableSkeleton> EditableSkeleton = SkeletonEditorModule.CreateEditableSkeleton(Skeleton);
+			if(GetFindWholeWord())
+			{
+				Asset->Modify();
+				EditableSkeleton->RenameNotify(FName(GetReplaceString()), FName(GetFindString()), false);
+			}
+			else
+			{
+				TArray<TPair<FName, FName>> FindReplacePairs;
+
+				for(const FName& NotifyName : Skeleton->AnimationNotifies)
+				{
+					FString NotifyString = NotifyName.ToString();
+					if(NameMatches(NotifyString))
+					{
+						const FString NewName = NotifyString.Replace(*GetFindStringRef(), *GetReplaceStringRef(), GetSearchCase());
+						const FName ReplaceNotifyName(*NewName);
+						FindReplacePairs.AddUnique(TPair<FName, FName>(NotifyName, ReplaceNotifyName));
+					}
+				}
+
+				if(FindReplacePairs.Num() > 0)
+				{
+					Asset->Modify();
+					for(const TPair<FName, FName>& FindReplacePair : FindReplacePairs)
+					{
+						EditableSkeleton->RenameNotify(FindReplacePair.Value, FindReplacePair.Key, false);
+					}
+				}
+			}
+		}
+	}
+}
+
+void UAnimAssetFindReplaceNotifies::RemoveInAsset(const FAssetData& InAssetData) const
+{
+	if(UObject* Asset = InAssetData.GetAsset())
+	{
+		if(UAnimSequenceBase* AnimSequenceBase = Cast<UAnimSequenceBase>(Asset))
+		{
+			if(GetFindWholeWord())
+			{
+				Asset->Modify();
+				AnimSequenceBase->RemoveNotifies( { FName(GetFindString()) } );
+			}
+			else
+			{
+				TArray<FName> NotifiesToRemove;
+				for(const FAnimNotifyEvent& Notify : AnimSequenceBase->Notifies)
+				{
+					FString NotifyNameString = Notify.NotifyName.ToString();
+					if(NameMatches(NotifyNameString))
+					{
+						NotifiesToRemove.AddUnique(Notify.NotifyName);
+					}
+				}
+
+				if(NotifiesToRemove.Num() > 0)
+				{
+					Asset->Modify();
+					AnimSequenceBase->RemoveNotifies(NotifiesToRemove);
+				}
+			}
+		}
+		else if(USkeleton* Skeleton = Cast<USkeleton>(Asset))
+		{
+			ISkeletonEditorModule& SkeletonEditorModule = FModuleManager::GetModuleChecked<ISkeletonEditorModule>("SkeletonEditor");
+			TSharedRef<IEditableSkeleton> EditableSkeleton = SkeletonEditorModule.CreateEditableSkeleton(Skeleton);
+			if(GetFindWholeWord())
+			{
+				EditableSkeleton->DeleteAnimNotifies( { FName(GetFindString()) }, false);
+			}
+			else
+			{
+				TArray<FName> NotifiesToRemove;
+
+				for(const FName& NotifyName : Skeleton->AnimationNotifies)
+				{
+					if(NameMatches(NotifyName.ToString()))
+					{
+						NotifiesToRemove.AddUnique(NotifyName);
+					}
+				}
+
+				if(NotifiesToRemove.Num() > 0)
+				{
+					EditableSkeleton->DeleteAnimNotifies(NotifiesToRemove, false);
+				}
+			}
+		}
+	}
+}
+
+void UAnimAssetFindReplaceNotifies::GetAutoCompleteNames(TArrayView<FAssetData> InAssetDatas, TSet<FString>& OutUniqueNames) const
+{
+	for (const FAssetData& AssetData : InAssetDatas)
+	{
+		const FString TagValue = AssetData.GetTagValueRef<FString>(USkeleton::AnimNotifyTag);
+		if (!TagValue.IsEmpty())
+		{
+			UE::String::ParseTokens(TagValue, *USkeleton::AnimNotifyTagDelimiter, [&OutUniqueNames](FStringView InToken)
+			{
+				OutUniqueNames.Add(FString(InToken));
+			}, UE::String::EParseTokensOptions::SkipEmpty);
+		}
+	}
+}
+
+#undef LOCTEXT_NAMESPACE
\ No newline at end of file
diff --git a/Engine/Source/Editor/Persona/Private/AnimAssetFindReplaceSyncMarkers.cpp b/Engine/Source/Editor/Persona/Private/AnimAssetFindReplaceSyncMarkers.cpp
new file mode 100644
--- /dev/null
+++ b/Engine/Source/Editor/Persona/Private/AnimAssetFindReplaceSyncMarkers.cpp
@@ -0,0 +1,290 @@
+﻿// Copyright Epic Games, Inc. All Rights Reserved.
+
+#include "AnimAssetFindReplaceSyncMarkers.h"
+
+#include "IEditableSkeleton.h"
+#include "ISkeletonEditorModule.h"
+#include "SAnimAssetFindReplace.h"
+#include "Animation/Skeleton.h"
+#include "Animation/AnimationAsset.h"
+#include "Animation/AnimSequence.h"
+#include "Engine/SkeletalMesh.h"
+#include "String/ParseTokens.h"
+
+#define LOCTEXT_NAMESPACE "AnimAssetFindReplaceSyncMarkers"
+
+FString UAnimAssetFindReplaceSyncMarkers::GetFindResultStringFromAssetData(const FAssetData& InAssetData) const
+{
+	if(GetFindString().IsEmpty())
+	{
+		return FString();
+	}
+
+	auto GetMatchingSyncMarkerNamesForAsset = [this](const FAssetData& InAssetData, TArray<FString>& OutSyncMarkerNames)
+	{
+		const FString TagValue = InAssetData.GetTagValueRef<FString>(USkeleton::AnimSyncMarkerTag);
+		if (!TagValue.IsEmpty())
+		{
+			UE::String::ParseTokens(TagValue, *USkeleton::AnimSyncMarkerTagDelimiter, [this, &OutSyncMarkerNames](FStringView InToken)
+			{
+				if(GetFindWholeWord())
+				{
+					if(InToken.Compare(GetFindString(), GetSearchCase()) == 0)
+					{
+						OutSyncMarkerNames.Add(FString(InToken));
+					}
+				}
+				else
+				{
+					if(UE::String::FindFirst(InToken, GetFindString(), GetSearchCase()) != INDEX_NONE)
+					{
+						OutSyncMarkerNames.Add(FString(InToken));
+					}
+				}
+			}, UE::String::EParseTokensOptions::SkipEmpty);
+		}
+	};
+
+	TStringBuilder<128> Builder;
+	TArray<FString> SyncMarkerNames;
+	GetMatchingSyncMarkerNamesForAsset(InAssetData, SyncMarkerNames);
+	if(SyncMarkerNames.Num() > 0)
+	{
+		for(int32 NameIndex = 0; NameIndex < SyncMarkerNames.Num(); ++NameIndex)
+		{
+			Builder.Append(SyncMarkerNames[NameIndex]);
+			if(NameIndex != SyncMarkerNames.Num() - 1)
+			{
+				Builder.Append(TEXT(", "));
+			}
+		}
+	}
+
+	return FString(Builder.ToString());
+}
+
+TConstArrayView<UClass*> UAnimAssetFindReplaceSyncMarkers::GetSupportedAssetTypes() const
+{
+	static UClass* Types[] = { UAnimSequence::StaticClass(), USkeleton::StaticClass() };
+	return Types;
+}
+
+bool UAnimAssetFindReplaceSyncMarkers::ShouldFilterOutAsset(const FAssetData& InAssetData, bool& bOutIsOldAsset) const
+{
+	FString TagValue;
+	if(InAssetData.GetTagValue<FString>(USkeleton::AnimSyncMarkerTag, TagValue))
+	{
+		bOutIsOldAsset = false;
+
+		if(GetFindString().IsEmpty())
+		{
+			return true;
+		}
+
+		if(GetSkeletonFilter().IsValid() )
+		{
+			if(InAssetData.GetClass() != USkeleton::StaticClass())
+			{
+				FString SkeletonPath;
+				if(InAssetData.GetTagValue<FString>(TEXT("Skeleton"), SkeletonPath))
+				{
+					if(SkeletonPath != GetSkeletonFilter().GetExportTextName())
+					{
+						return true;
+					}
+				}
+			}
+			else
+			{
+				if(InAssetData.ToSoftObjectPath() != GetSkeletonFilter().ToSoftObjectPath())
+				{
+					return true;
+				}
+			}
+		}
+
+		bool bFoundMatch = false;
+		UE::String::ParseTokens(TagValue, *USkeleton::AnimSyncMarkerTagDelimiter, [this, &bFoundMatch](FStringView InToken)
+		{
+			if(NameMatches(InToken))
+			{
+				bFoundMatch = true;
+			}
+		}, UE::String::EParseTokensOptions::SkipEmpty);
+
+		if(bFoundMatch)
+		{
+			return false;
+		}
+	}
+	else
+	{
+		const UClass* AssetClass = InAssetData.GetClass();
+		if(AssetClass->IsChildOf(UAnimSequence::StaticClass()))
+		{
+			bOutIsOldAsset = true;
+		}
+		else if(AssetClass->IsChildOf(USkeleton::StaticClass()))
+		{
+			bOutIsOldAsset = true;
+		}
+		else
+		{
+			// Assume unknown assets are not 'old'
+			bOutIsOldAsset = false;
+		}
+	}
+
+	return true;
+}
+
+void UAnimAssetFindReplaceSyncMarkers::ReplaceInAsset(const FAssetData& InAssetData) const
+{
+	if(UObject* Asset = InAssetData.GetAsset())
+	{
+		if(UAnimSequence* AnimSequence = Cast<UAnimSequence>(Asset))
+		{
+			if(GetFindWholeWord())
+			{
+				Asset->Modify();
+				AnimSequence->RenameSyncMarkers(FName(GetFindString()), FName(GetReplaceString()));
+			}
+			else
+			{
+				TArray<TPair<FName, FName>> FindReplacePairs;
+
+				for(const FAnimSyncMarker& SyncMarker : AnimSequence->AuthoredSyncMarkers)
+				{
+					FString SyncMarkerName = SyncMarker.MarkerName.ToString();
+					if(NameMatches(SyncMarkerName))
+					{
+						const FName FindSyncMarkerName(*SyncMarkerName);
+						const FString NewName = SyncMarkerName.Replace(*GetFindStringRef(), *GetReplaceStringRef(), GetSearchCase());
+						const FName ReplaceSyncMarkerName(*NewName);
+						FindReplacePairs.AddUnique(TPair<FName, FName>(FindSyncMarkerName, ReplaceSyncMarkerName));
+					}
+				}
+
+				if(FindReplacePairs.Num() > 0)
+				{
+					Asset->Modify();
+
+					for(const TPair<FName, FName>& FindReplacePair : FindReplacePairs)
+					{
+						AnimSequence->RenameSyncMarkers(FindReplacePair.Key, FindReplacePair.Value);
+					}
+				}
+			}
+		}
+		else if(USkeleton* Skeleton = Cast<USkeleton>(Asset))
+		{
+			ISkeletonEditorModule& SkeletonEditorModule = FModuleManager::GetModuleChecked<ISkeletonEditorModule>("SkeletonEditor");
+			TSharedRef<IEditableSkeleton> EditableSkeleton = SkeletonEditorModule.CreateEditableSkeleton(Skeleton);
+			if(GetFindWholeWord())
+			{
+				Asset->Modify();
+				EditableSkeleton->RenameSyncMarker(FName(GetReplaceString()), FName(GetFindString()));
+			}
+			else
+			{
+				TArray<TPair<FName, FName>> FindReplacePairs;
+
+				for(const FName& SyncMarkerName : Skeleton->GetExistingMarkerNames())
+				{
+					FString SyncMarkerString = SyncMarkerName.ToString();
+					if(NameMatches(SyncMarkerString))
+					{
+						const FString NewName = SyncMarkerString.Replace(*GetFindStringRef(), *GetReplaceStringRef(), GetSearchCase());
+						const FName ReplaceSyncMarkerName(*NewName);
+						FindReplacePairs.AddUnique(TPair<FName, FName>(SyncMarkerName, ReplaceSyncMarkerName));
+					}
+				}
+
+				if(FindReplacePairs.Num() > 0)
+				{
+					Asset->Modify();
+					for(const TPair<FName, FName>& FindReplacePair : FindReplacePairs)
+					{
+						EditableSkeleton->RenameSyncMarker(FindReplacePair.Value, FindReplacePair.Key);
+					}
+				}
+			}
+		}
+	}
+}
+
+void UAnimAssetFindReplaceSyncMarkers::RemoveInAsset(const FAssetData& InAssetData) const
+{
+	if(UObject* Asset = InAssetData.GetAsset())
+	{
+		if(UAnimSequence* AnimSequence = Cast<UAnimSequence>(Asset))
+		{
+			if(GetFindWholeWord())
+			{
+				Asset->Modify();
+				AnimSequence->RemoveSyncMarkers( { FName(GetFindString()) } );
+			}
+			else
+			{
+				TArray<FName> SyncMarkerToRemove;
+				for(const FAnimSyncMarker& SyncMarker : AnimSequence->AuthoredSyncMarkers)
+				{
+					FString SyncMarkerNameString = SyncMarker.MarkerName.ToString();
+					if(NameMatches(SyncMarkerNameString))
+					{
+						SyncMarkerToRemove.AddUnique(SyncMarker.MarkerName);
+					}
+				}
+
+				if(SyncMarkerToRemove.Num() > 0)
+				{
+					Asset->Modify();
+					AnimSequence->RemoveSyncMarkers(SyncMarkerToRemove);
+				}
+			}
+		}
+		else if(USkeleton* Skeleton = Cast<USkeleton>(Asset))
+		{
+			ISkeletonEditorModule& SkeletonEditorModule = FModuleManager::GetModuleChecked<ISkeletonEditorModule>("SkeletonEditor");
+			TSharedRef<IEditableSkeleton> EditableSkeleton = SkeletonEditorModule.CreateEditableSkeleton(Skeleton);
+			if(GetFindWholeWord())
+			{
+				EditableSkeleton->DeleteSyncMarkers( { FName(GetFindString()) } );
+			}
+			else
+			{
+				TArray<FName> SyncMarkersToRemove;
+
+				for(const FName& SyncMarkerName : Skeleton->GetExistingMarkerNames())
+				{
+					if(NameMatches(SyncMarkerName.ToString()))
+					{
+						SyncMarkersToRemove.AddUnique(SyncMarkerName);
+					}
+				}
+
+				if(SyncMarkersToRemove.Num() > 0)
+				{
+					EditableSkeleton->DeleteSyncMarkers(SyncMarkersToRemove);
+				}
+			}
+		}
+	}
+}
+
+void UAnimAssetFindReplaceSyncMarkers::GetAutoCompleteNames(TArrayView<FAssetData> InAssetDatas, TSet<FString>& OutUniqueNames) const
+{
+	for (const FAssetData& AssetData : InAssetDatas)
+	{
+		const FString TagValue = AssetData.GetTagValueRef<FString>(USkeleton::AnimSyncMarkerTag);
+		if (!TagValue.IsEmpty())
+		{
+			UE::String::ParseTokens(TagValue, *USkeleton::AnimSyncMarkerTagDelimiter, [&OutUniqueNames](FStringView InToken)
+			{
+				OutUniqueNames.Add(FString(InToken));
+			}, UE::String::EParseTokensOptions::SkipEmpty);
+		}
+	}
+}
+
+#undef LOCTEXT_NAMESPACE
\ No newline at end of file
diff --git a/Engine/Source/Editor/Persona/Private/SAnimAssetFindReplace.cpp b/Engine/Source/Editor/Persona/Private/SAnimAssetFindReplace.cpp
--- a/Engine/Source/Editor/Persona/Private/SAnimAssetFindReplace.cpp
+++ b/Engine/Source/Editor/Persona/Private/SAnimAssetFindReplace.cpp
@@ -2,6 +2,7 @@
 
 #include "SAnimAssetFindReplace.h"
 
+#include "AnimAssetFindReplaceCurves.h"
 #include "ContentBrowserDataSource.h"
 #include "ContentBrowserModule.h"
 #include "IContentBrowserSingleton.h"
@@ -12,23 +13,18 @@
 #include "Widgets/Layout/SGridPanel.h"
 #include "Widgets/Input/SButton.h"
 #include "Widgets/Text/STextBlock.h"
-#include "Animation/AnimationAsset.h"
-#include "Animation/Skeleton.h"
 #include "Framework/Application/SlateApplication.h"
 #include "Widgets/Input/SEditableTextBox.h"
-#include "Widgets/Input/SSearchBox.h"
-#include "Animation/AnimSequence.h"
-#include "Engine/SkeletalMesh.h"
 #include "ToolMenus.h"
 #include "ToolMenuSection.h"
 #include "AssetRegistry/AssetRegistryModule.h"
 #include "Misc/ScopedSlowTask.h"
 #include "ScopedTransaction.h"
-#include "Animation/PoseAsset.h"
 #include "Widgets/Views/SListView.h"
 #include "UObject/AnimPhysObjectVersion.h"
 #include "UObject/ObjectVersion.h"
 #include "Widgets/Input/SHyperlink.h"
+#include "UObject/UObjectIterator.h"
 
 #define LOCTEXT_NAMESPACE "SAnimAssetFindReplace"
 
@@ -54,7 +50,7 @@ TSharedPtr<SToolTip> FAnimAssetFindReplaceSummoner::CreateTabToolTipWidget(const
 namespace AnimAssetFindReplacePrivate
 {
 
-static TSharedPtr<SAnimAssetFindReplace> GetWidgetFromContext(const FToolMenuContext& InContext)
+TSharedPtr<SAnimAssetFindReplace> GetWidgetFromContext(const FToolMenuContext& InContext)
 {
 	if(UAnimAssetFindReplaceContext* Context = InContext.FindContext<UAnimAssetFindReplaceContext>())
 	{
@@ -66,173 +62,31 @@ static TSharedPtr<SAnimAssetFindReplace> GetWidgetFromContext(const FToolMenuCon
 
 }
 
-class SAutoCompleteSearchBox : public SCompoundWidget
+void SAnimAssetFindReplace::Construct(const FArguments& InArgs)
 {
-public:
-	SLATE_BEGIN_ARGS(SAutoCompleteSearchBox) {}
-
-	/** The text displayed in the SearchBox when no text has been entered */
-	SLATE_ATTRIBUTE(FText, HintText)
-
-	/** The text displayed in the SearchBox when it's created */
-	SLATE_ATTRIBUTE(FText, InitialText)
-
-	/** Invoked whenever the text changes */
-	SLATE_EVENT(FOnTextChanged, OnTextChanged)
-
-	/** Items to show in the autocomplete popup */
-	SLATE_ARGUMENT(TSharedPtr<TArray<TSharedPtr<FString>>>, AutoCompleteItems)
-	
-	SLATE_END_ARGS()
-
-	void Construct(const FArguments& InArgs)
-	{
-		AutoCompleteItems = InArgs._AutoCompleteItems;
-
-		ChildSlot
-		[
-			SAssignNew(MenuAnchor, SMenuAnchor)
-			.Method(EPopupMethod::CreateNewWindow)
-			.Placement(EMenuPlacement::MenuPlacement_BelowAnchor)
-			.MenuContent
-			(
-				SNew(SBox)
-				.MaxDesiredHeight(200.0f)
-				.MinDesiredWidth(200.0f)
-				[
-					SAssignNew(AutoCompleteList, SListView<TSharedPtr<FString>>)
-					.SelectionMode(ESelectionMode::Single)
-					.ListItemsSource(&FilteredAutoCompleteItems)
-					.OnSelectionChanged_Lambda([this](TSharedPtr<FString> InString, ESelectInfo::Type InSelectInfo)
-					{
-						if(InString.IsValid() && InSelectInfo == ESelectInfo::OnMouseClick)
-						{
-							TGuardValue<bool> GuardValue(bSettingTextFromSearchItem, true);
-							SearchBox->SetText(FText::FromString(*InString.Get()));
-							MenuAnchor->SetIsOpen(false);
-						}
-					})
-					.OnGenerateRow_Lambda([this](TSharedPtr<FString> InString, const TSharedRef<STableViewBase>& InTableView)
-					{
-						if(InString.IsValid())
-						{
-							return
-								SNew(STableRow<TSharedPtr<FString>>, InTableView)
-								.Content()
-								[
-									SNew(STextBlock)
-									.Text(FText::FromString(*InString.Get()))
-									.HighlightText_Lambda([this]()
-									{
-										return SearchBox->GetText();
-									})
-								];
-						}
-						return SNew(STableRow<TSharedPtr<FString>>, InTableView);
-					})
-					.OnKeyDownHandler_Lambda([this](const FGeometry& InGeometry, const FKeyEvent& InKeyEvent)
-					{
-						if(InKeyEvent.GetKey() == EKeys::Enter)
-						{
-							TArray<TSharedPtr<FString>> SelectedItems;
-							AutoCompleteList->GetSelectedItems(SelectedItems);
-							if(SelectedItems.Num() > 0 && SelectedItems[0].IsValid())
-							{
-								TGuardValue<bool> GuardValue(bSettingTextFromSearchItem, true);
-								SearchBox->SetText(FText::FromString(*SelectedItems[0].Get()));
-								MenuAnchor->SetIsOpen(false);
-								FReply::Handled();
-							}
-						}
-						return FReply::Unhandled();
-					})
-				]
-			)
-			[
-				SAssignNew(SearchBox, SSearchBox)
-				.HintText(InArgs._HintText)
-				.InitialText(InArgs._InitialText)
-				.OnTextChanged_Lambda([this, OnTextChanged = InArgs._OnTextChanged](const FText& InText)
-				{
-					FilterItems(InText);
-
-					if(!bSettingTextFromSearchItem)
-					{
-						MenuAnchor->SetIsOpen(FilteredAutoCompleteItems.Num() > 0, false);
-					}
-
-					OnTextChanged.ExecuteIfBound(InText);
-				})
-				.OnVerifyTextChanged_Lambda([](const FText& InText, FText& OutErrorMessage)
-				{
-					return FName::IsValidXName(InText.ToString(), FString(INVALID_NAME_CHARACTERS), &OutErrorMessage);
-				})
-			]
-		];
-		
-		FilterItems(InArgs._InitialText.Get());
-	}
+	Config = InArgs._Config;
 
-	void FilterItems(const FText& InText)
+	// Create instances of all processor types we can use
+	for (TObjectIterator<UClass> It; It; ++It)
 	{
-		FilteredAutoCompleteItems.Empty();
-
-		for(TSharedPtr<FString> String : *AutoCompleteItems.Get())
+		if(!It->HasAnyClassFlags(CLASS_Abstract) && It->IsChildOf(UAnimAssetFindReplaceProcessor::StaticClass()))
 		{
-			if(String.Get()->Contains(InText.ToString()))
+			UAnimAssetFindReplaceProcessor* Processor = NewObject<UAnimAssetFindReplaceProcessor>(GetTransientPackage(), *It, NAME_None, RF_Transient);
+			Processor->Initialize(SharedThis(this));
+			Processors.Add(*It, Processor);
+
+			if(*It == Config.InitialProcessorClass)
 			{
-				FilteredAutoCompleteItems.Add(String);
+				CurrentProcessor = Processor;
 			}
 		}
-
-		RefreshAutoCompleteItems();
-	}
-
-	void RefreshAutoCompleteItems()
-	{
-		AutoCompleteList->RequestListRefresh();
-	}
-
-	TSharedRef<SSearchBox> GetSearchBox() const
-	{
-		return SearchBox.ToSharedRef();
-	}
-	
-	virtual void OnFocusChanging(const FWeakWidgetPath& PreviousFocusPath, const FWidgetPath& NewWidgetPath, const FFocusEvent& InFocusEvent) override
-	{
-		if(PreviousFocusPath.ContainsWidget(SearchBox.Get()) && !NewWidgetPath.ContainsWidget(MenuAnchor->GetMenuWindow().Get()))
-		{
-			MenuAnchor->SetIsOpen(false);
-		}
-		SCompoundWidget::OnFocusChanging(PreviousFocusPath, NewWidgetPath, InFocusEvent);
 	}
 
-	virtual FReply OnPreviewKeyDown(const FGeometry& MyGeometry, const FKeyEvent& InKeyEvent) override
+	// Assign a default processor if we didn't have one set
+	if(CurrentProcessor == nullptr && Processors.Num() > 0)
 	{
-		if(InKeyEvent.GetKey() == EKeys::Down && MenuAnchor->IsOpen())
-		{
-			// switch focus to the drop-down autocomplete list
-			return FReply::Handled().SetUserFocus(AutoCompleteList.ToSharedRef(), EFocusCause::Navigation);
-		}
-		return SCompoundWidget::OnPreviewKeyDown(MyGeometry, InKeyEvent); 
+		CurrentProcessor = Processors.FindRef(UAnimAssetFindReplaceCurves::StaticClass());
 	}
-	
-private:
-	TArray<TSharedPtr<FString>> FilteredAutoCompleteItems;
-	TSharedPtr<TArray<TSharedPtr<FString>>> AutoCompleteItems;
-	TSharedPtr<SListView<TSharedPtr<FString>>> AutoCompleteList;
-	TSharedPtr<SMenuAnchor> MenuAnchor;
-	TSharedPtr<SSearchBox> SearchBox;
-	bool bSettingTextFromSearchItem = false;
-};
-
-void SAnimAssetFindReplace::Construct(const FArguments& InArgs)
-{
-	Mode = InArgs._Config.Mode;
-	Type = InArgs._Config.Type;
-	FindString = InArgs._Config.FindString;
-	ReplaceString = InArgs._Config.ReplaceString;
-	SkeletonFilter = InArgs._Config.SkeletonFilter;
 
 	const FContentBrowserModule& ContentBrowserModule = FModuleManager::LoadModuleChecked<FContentBrowserModule>("ContentBrowser");
 
@@ -262,49 +116,26 @@ void SAnimAssetFindReplace::Construct(const FArguments& InArgs)
 				UObject::FAssetRegistryTag::TT_Alphabetical,
 				FOnGetCustomAssetColumnData::CreateLambda([this](FAssetData& InAssetData, FName InColumnName)
 				{
-					TStringBuilder<128> Builder;
-
-					if(Type == EAnimAssetFindReplaceType::Curves)
-					{
-						TArray<FString> CurveNames;
-						GetMatchingCurveNamesForAsset(InAssetData, CurveNames);
-						if(CurveNames.Num() > 0)
-						{
-							for(int32 NameIndex = 0; NameIndex < CurveNames.Num(); ++NameIndex)
-							{
-								Builder.Append(CurveNames[NameIndex]);
-								if(NameIndex != CurveNames.Num() - 1)
-								{
-									Builder.Append(TEXT(", "));
-								}
-							}
-						}
-					}
-					else if(Type == EAnimAssetFindReplaceType::Notifies)
-					{
-						TArray<FString> NotifyNames;
-						GetMatchingNotifyNamesForAsset(InAssetData, NotifyNames);
-						if(NotifyNames.Num() > 0)
-						{
-							for(int32 NameIndex = 0; NameIndex < NotifyNames.Num(); ++NameIndex)
-							{
-								Builder.Append(NotifyNames[NameIndex]);
-								if(NameIndex != NotifyNames.Num() - 1)
-								{
-									Builder.Append(TEXT(", "));
-								}
-							}
-						}
-					}
-
-					return FString(Builder.ToString());
+					return CurrentProcessor->GetFindResultStringFromAssetData(InAssetData);
 				})
 			));
 
-	TArray<UClass*> ClassesWithAssetRegistryTags = { UAnimSequenceBase::StaticClass(), USkeleton::StaticClass(), USkeletalMesh::StaticClass() };
-	GetDerivedClasses(UAnimSequenceBase::StaticClass(), ClassesWithAssetRegistryTags);
-	GetDerivedClasses(USkeleton::StaticClass(), ClassesWithAssetRegistryTags);
-	GetDerivedClasses(USkeletalMesh::StaticClass(), ClassesWithAssetRegistryTags);
+	TSet<UClass*> ClassesWithAssetRegistryTagsSet;
+	for(TPair<TSubclassOf<UAnimAssetFindReplaceProcessor>, UAnimAssetFindReplaceProcessor*> ProcessorPair : Processors)
+	{
+		for(const UClass* AssetType : ProcessorPair.Value->GetSupportedAssetTypes())
+		{
+			ClassesWithAssetRegistryTagsSet.Add(const_cast<UClass*>(AssetType));
+		}
+	}
+
+	TArray<UClass*> ClassesWithAssetRegistryTags = ClassesWithAssetRegistryTagsSet.Array();
+	TArray<UClass*> DerivedClassesWithAssetRegistryTags;
+	for(UClass* AssetType : ClassesWithAssetRegistryTags)
+	{
+		GetDerivedClasses(AssetType, DerivedClassesWithAssetRegistryTags);
+	}
+	ClassesWithAssetRegistryTags.Append(DerivedClassesWithAssetRegistryTags);
 
 	TArray<UObject::FAssetRegistryTag> AssetRegistryTags;
 	for(UClass* Class : ClassesWithAssetRegistryTags)
@@ -335,16 +166,65 @@ void SAnimAssetFindReplace::Construct(const FArguments& InArgs)
 		Toolbar = UToolMenus::Get()->RegisterMenu("AnimAssetFindReplaceToolbar", NAME_None, EMultiBoxType::SlimHorizontalToolBar);
 		Toolbar->StyleName = "CalloutToolbar";	// This style displays button text
 
+		{
+			FToolMenuSection& Section = Toolbar->AddSection("FindReplaceOptions");
+
+			Section.AddDynamicEntry("ProcessorSelector", FNewToolMenuSectionDelegate::CreateLambda([](FToolMenuSection& InSection)
+			{
+				if(TSharedPtr<SAnimAssetFindReplace> Widget = AnimAssetFindReplacePrivate::GetWidgetFromContext(InSection.Context))
+				{
+					InSection.AddEntry(
+						FToolMenuEntry::InitWidget(
+							"ProcessorSelector",
+							SNew(SHorizontalBox)
+							+SHorizontalBox::Slot()
+							.AutoWidth()
+							.Padding(5.0f, 0.0f)
+							.VAlign(VAlign_Center)
+							[
+								SNew(STextBlock)
+								.Text(LOCTEXT("FindLabel", "Find:"))
+							]
+							+SHorizontalBox::Slot()
+							.AutoWidth()
+							.VAlign(VAlign_Center)
+							[
+								SNew(SClassPropertyEntryBox)
+									.MetaClass(UAnimAssetFindReplaceProcessor::StaticClass())
+									.AllowNone(false)
+									.ShowDisplayNames(true)
+									.HideViewOptions(true)
+									.SelectedClass_Lambda([WeakWidget = TWeakPtr<SAnimAssetFindReplace>(Widget)]() -> UClass*
+									{
+										if(TSharedPtr<SAnimAssetFindReplace> PinnedWidget = WeakWidget.Pin())
+										{
+											return PinnedWidget->GetCurrentProcessor()->GetClass();
+										}
+
+										return nullptr;
+									})
+									.OnSetClass_Lambda([WeakWidget = TWeakPtr<SAnimAssetFindReplace>(Widget)](const UClass* InClass)
+									{
+										if(TSharedPtr<SAnimAssetFindReplace> PinnedWidget = WeakWidget.Pin())
+										{
+											PinnedWidget->SetCurrentProcessor(const_cast<UClass*>(InClass));
+										}
+									})
+							],
+							FText::GetEmpty(), true));
+				}
+			}));
+		}
+
 		{
 			FToolMenuSection& Section = Toolbar->AddSection("FindReplaceActions");
-			
+		
 			FToolUIAction RefreshButton;
 			RefreshButton.ExecuteAction = FToolMenuExecuteAction::CreateLambda([](const FToolMenuContext& InContext)
 			{
 				if(TSharedPtr<SAnimAssetFindReplace> Widget = AnimAssetFindReplacePrivate::GetWidgetFromContext(InContext))
 				{
-					Widget->RefreshAutoCompleteItems();
-					Widget->RefreshSearchResults();
+					Widget->RequestRefreshSearchResults();
 				}
 			});
 
@@ -354,213 +234,43 @@ void SAnimAssetFindReplace::Construct(const FArguments& InArgs)
 					RefreshButton,
 					LOCTEXT("RefreshRadioLabel", "Refresh"),
 					LOCTEXT("RefreshRadioTooltip", "Refresh search results."),
-					FSlateIcon(),
-					EUserInterfaceActionType::RadioButton));
+					FSlateIcon(FAppStyle::GetAppStyleSetName(), "Icons.Find"),
+					EUserInterfaceActionType::Button));
 		}
 
 		{
-			FToolMenuSection& Section = Toolbar->AddSection("FindReplaceOptions");
-			
-			FToolUIAction CurvesRadio;
-			CurvesRadio.ExecuteAction = FToolMenuExecuteAction::CreateLambda([](const FToolMenuContext& InContext)
-			{
-				if(TSharedPtr<SAnimAssetFindReplace> Widget = AnimAssetFindReplacePrivate::GetWidgetFromContext(InContext))
-				{
-					Widget->SetFindReplaceType(EAnimAssetFindReplaceType::Curves);
-				}
-			});
-
-			CurvesRadio.GetActionCheckState = FToolMenuGetActionCheckState::CreateLambda([](const FToolMenuContext& InContext)
-			{
-				if(TSharedPtr<SAnimAssetFindReplace> Widget = AnimAssetFindReplacePrivate::GetWidgetFromContext(InContext))
-				{
-					return Widget->Type == EAnimAssetFindReplaceType::Curves ? ECheckBoxState::Checked : ECheckBoxState::Unchecked;
-				}
-				return ECheckBoxState::Undetermined;
-			});
-			
-			Section.AddEntry(
-				FToolMenuEntry::InitToolBarButton(
-					"Curves",
-					CurvesRadio,
-					LOCTEXT("CurvesRadioLabel", "Curves"),
-					LOCTEXT("CurvesRadioTooltip", "Search for Curves in Animation Assets."),
-					FSlateIcon(),
-					EUserInterfaceActionType::RadioButton));
-
-			FToolUIAction NotifiesRadio;
-			NotifiesRadio.ExecuteAction = FToolMenuExecuteAction::CreateLambda([](const FToolMenuContext& InContext)
-			{
-				if(TSharedPtr<SAnimAssetFindReplace> Widget = AnimAssetFindReplacePrivate::GetWidgetFromContext(InContext))
-				{
-					Widget->SetFindReplaceType(EAnimAssetFindReplaceType::Notifies);
-				}
-			});
-
-			NotifiesRadio.GetActionCheckState = FToolMenuGetActionCheckState::CreateLambda([](const FToolMenuContext& InContext)
-			{
-				if(TSharedPtr<SAnimAssetFindReplace> Widget = AnimAssetFindReplacePrivate::GetWidgetFromContext(InContext))
-				{
-					return Widget->Type == EAnimAssetFindReplaceType::Notifies ? ECheckBoxState::Checked : ECheckBoxState::Unchecked;
-				}
-				return ECheckBoxState::Undetermined;
-			});
-			
-			Section.AddEntry(
-				FToolMenuEntry::InitToolBarButton(
-					"Notifies", NotifiesRadio,
-					LOCTEXT("NotifiesRadioLabel", "Notifies"),
-					LOCTEXT("NotifiesRadioTooltip", "Search for Named Notifies in Animation Assets."),
-					FSlateIcon(),
-					EUserInterfaceActionType::RadioButton));
-
-			FToolUIAction MatchCaseCheckbox;
-			MatchCaseCheckbox.ExecuteAction = FToolMenuExecuteAction::CreateLambda([](const FToolMenuContext& InContext)
-			{
-				if(TSharedPtr<SAnimAssetFindReplace> Widget = AnimAssetFindReplacePrivate::GetWidgetFromContext(InContext))
-				{
-					Widget->SearchCase = Widget->SearchCase == ESearchCase::CaseSensitive ? ESearchCase::IgnoreCase : ESearchCase::CaseSensitive;
-					Widget->RefreshSearchResults();
-				}
-			});
-
-			MatchCaseCheckbox.GetActionCheckState = FToolMenuGetActionCheckState::CreateLambda([](const FToolMenuContext& InContext)
-			{
-				if(TSharedPtr<SAnimAssetFindReplace> Widget = AnimAssetFindReplacePrivate::GetWidgetFromContext(InContext))
-				{
-					return Widget->SearchCase == ESearchCase::CaseSensitive ? ECheckBoxState::Checked : ECheckBoxState::Unchecked;
-				}
-				return ECheckBoxState::Undetermined;
-			});
-			
-			Section.AddEntry(
-				FToolMenuEntry::InitToolBarButton(
-					"MatchCase",
-					MatchCaseCheckbox,
-					LOCTEXT("MatchCaseCheckboxLabel", "Match Case"),
-					LOCTEXT("MatchCaseCheckboxTooltip", "Whether to match case when searching."),
-					FSlateIcon(),
-					EUserInterfaceActionType::ToggleButton));
-
-			FToolUIAction MatchWholeWordCheckbox;
-			MatchWholeWordCheckbox.ExecuteAction = FToolMenuExecuteAction::CreateLambda([](const FToolMenuContext& InContext)
-			{
-				if(TSharedPtr<SAnimAssetFindReplace> Widget = AnimAssetFindReplacePrivate::GetWidgetFromContext(InContext))
-				{
-					Widget->bFindWholeWord = !Widget->bFindWholeWord;
-					Widget->RefreshSearchResults();
-				}
-			});
-
-			MatchWholeWordCheckbox.GetActionCheckState = FToolMenuGetActionCheckState::CreateLambda([](const FToolMenuContext& InContext)
-			{
-				if(TSharedPtr<SAnimAssetFindReplace> Widget = AnimAssetFindReplacePrivate::GetWidgetFromContext(InContext))
-				{
-					return Widget->bFindWholeWord ? ECheckBoxState::Checked : ECheckBoxState::Unchecked;
-				}
-				return ECheckBoxState::Undetermined;
-			});
-			
-			Section.AddEntry(
-				FToolMenuEntry::InitToolBarButton(
-					"MatchWholeWord",
-					MatchWholeWordCheckbox,
-					LOCTEXT("MatchWholeWordCheckboxLabel", "Match Whole Word"),
-					LOCTEXT("MatchWholeWordCheckboxTooltip", "Whether to match the whole word or just part of the word when searching."),
-					FSlateIcon(),
-					EUserInterfaceActionType::ToggleButton));
-
-			Section.AddDynamicEntry("SkeletonFilter", FNewToolMenuSectionDelegate::CreateLambda([](FToolMenuSection& InSection)
+			FToolMenuSection& Section = Toolbar->AddSection("ProcessorOptions");
+			Section.AddDynamicEntry("ProcessorOptions", FNewToolMenuSectionDelegate::CreateLambda([this](FToolMenuSection& InSection)
 			{
 				if(TSharedPtr<SAnimAssetFindReplace> Widget = AnimAssetFindReplacePrivate::GetWidgetFromContext(InSection.Context))
 				{
-					InSection.AddEntry(
-					FToolMenuEntry::InitWidget(
-						"SkeletonFilterWidget",
-						SNew(SHorizontalBox)
-						.ToolTipText(LOCTEXT("SkeletonFilterTooltip", "Choose a Skeleton asset to filter results by."))
-						+SHorizontalBox::Slot()
-						.AutoWidth()
-						.VAlign(VAlign_Center)
-						.Padding(5.0f, 0.0f)
-						[
-							SNew(STextBlock)
-							.Text(LOCTEXT("SkeletonFilterLabel", "Skeleton"))
-						]
-						+SHorizontalBox::Slot()
-						.AutoWidth()
-						.VAlign(VAlign_Center)
-						[
-							SNew(SObjectPropertyEntryBox)
-							.ObjectPath_Lambda([WeakWidget = TWeakPtr<SAnimAssetFindReplace>(Widget)]()
-							{
-								if(TSharedPtr<SAnimAssetFindReplace> PinnedWidget = WeakWidget.Pin())
-								{
-									return PinnedWidget->SkeletonFilter.GetObjectPathString();
-								}
-								return FString();
-							})
-							.OnObjectChanged_Lambda([WeakWidget = TWeakPtr<SAnimAssetFindReplace>(Widget)](const FAssetData& InAssetData)
-							{
-								if(TSharedPtr<SAnimAssetFindReplace> PinnedWidget = WeakWidget.Pin())
-								{
-									PinnedWidget->SkeletonFilter = InAssetData;
-									PinnedWidget->RefreshSearchResults();
-								}
-							})
-							.AllowedClass(USkeleton::StaticClass())
-						],
-					FText::GetEmpty(),
-					true, true, true));
+					Widget->GetCurrentProcessor()->ExtendToolbar(InSection);
 				}
 			}));
 		}
 	}
 
-	ToolbarContext = TStrongObjectPtr(NewObject<UAnimAssetFindReplaceContext>());
+	ToolbarContext = NewObject<UAnimAssetFindReplaceContext>();
 	ToolbarContext->Widget = SharedThis(this);
 
-	AutoCompleteItems = MakeShared<TArray<TSharedPtr<FString>>>();
-
 	ChildSlot
 	[
 		SNew(SVerticalBox)
 		+SVerticalBox::Slot()
 		.AutoHeight()
 		[
-			UToolMenus::Get()->GenerateWidget("AnimAssetFindReplaceToolbar", FToolMenuContext(ToolbarContext.Get()))
+			SAssignNew(ToolbarContainer, SBox)
+			[
+				UToolMenus::Get()->GenerateWidget("AnimAssetFindReplaceToolbar", FToolMenuContext(ToolbarContext))
+			]
 		]
 		+SVerticalBox::Slot()
 		.AutoHeight()
 		.Padding(5.0f, 10.0f)
 		[
-			SNew(SVerticalBox)
-			+SVerticalBox::Slot()
-			.Padding(6.0f, 2.0f)
-			[
-				SAssignNew(FindSearchBox, SAutoCompleteSearchBox)
-				.AutoCompleteItems(AutoCompleteItems)
-				.HintText(LOCTEXT("FindLabel", "Find"))
-				.ToolTipText(LOCTEXT("FindLabel", "Find"))
-				.InitialText_Lambda([this](){ return FText::FromString(FindString); })
-				.OnTextChanged_Lambda([this](const FText& InText)
-				{
-					FindString = InText.ToString();
-					RefreshSearchResults();
-				})
-			]
-			+SVerticalBox::Slot()
-			.Padding(6.0f, 2.0f)
+			SAssignNew(FindReplaceWidgetContainer, SBox)
 			[
-				SAssignNew(ReplaceSearchBox, SAutoCompleteSearchBox)
-				.AutoCompleteItems(AutoCompleteItems)
-				.HintText(LOCTEXT("ReplaceLabel", "Replace With"))
-				.ToolTipText(LOCTEXT("ReplaceLabel", "Replace With"))
-				.InitialText_Lambda([this](){ return FText::FromString(ReplaceString); })
-				.OnTextChanged_Lambda([this](const FText& InText)
-				{
-					ReplaceString = InText.ToString();
-				})
+				CurrentProcessor->MakeFindReplaceWidget()
 			]
 		]
 		+SVerticalBox::Slot()
@@ -597,7 +307,7 @@ void SAnimAssetFindReplace::Construct(const FArguments& InArgs)
 					for(const FAssetData& AssetData : OldAssets)
 					{
 						SlowTask.EnterProgressFrame();
-						
+
 						AssetData.GetAsset();
 
 						if(SlowTask.ShouldCancel())
@@ -606,7 +316,7 @@ void SAnimAssetFindReplace::Construct(const FArguments& InArgs)
 						}
 					}
 
-					RefreshSearchResults();
+					RequestRefreshSearchResults();
 				})
 			]
 			+SHorizontalBox::Slot()
@@ -626,7 +336,7 @@ void SAnimAssetFindReplace::Construct(const FArguments& InArgs)
 					.ToolTipText(LOCTEXT("RemoveButtonTooltip", "Remove selected items"))
 					.IsEnabled_Lambda([this]()
 					{
-						return bAssetsSelected;
+						return bAssetsSelected && CurrentProcessor->CanCurrentlyRemove();
 					})
 					.OnClicked(this, &SAnimAssetFindReplace::HandleRemove)
 				]
@@ -639,7 +349,7 @@ void SAnimAssetFindReplace::Construct(const FArguments& InArgs)
 					.ToolTipText(LOCTEXT("RemoveAllButtonTooltip", "Remove all matching items"))
 					.IsEnabled_Lambda([this]()
 					{
-						return bFoundAssets;
+						return bFoundAssets && CurrentProcessor->CanCurrentlyRemove();
 					})
 					.OnClicked(this, &SAnimAssetFindReplace::HandleRemoveAll)
 				]
@@ -652,7 +362,7 @@ void SAnimAssetFindReplace::Construct(const FArguments& InArgs)
 					.ToolTipText(LOCTEXT("ReplaceButtonTooltip", "Replace selected items"))
 					.IsEnabled_Lambda([this]()
 					{
-						return bAssetsSelected && !ReplaceString.IsEmpty();
+						return bAssetsSelected && CurrentProcessor->CanCurrentlyReplace();
 					})
 					.OnClicked(this, &SAnimAssetFindReplace::HandleReplace)
 				]
@@ -665,46 +375,78 @@ void SAnimAssetFindReplace::Construct(const FArguments& InArgs)
 					.ToolTipText(LOCTEXT("ReplaceAllButtonTooltip", "Replace all matching items"))
 					.IsEnabled_Lambda([this]()
 					{
-						return bFoundAssets && !ReplaceString.IsEmpty();
+						return bFoundAssets && CurrentProcessor->CanCurrentlyReplace();
 					})
 					.OnClicked(this, &SAnimAssetFindReplace::HandleReplaceAll)
 				]
 			]
 		]
 	];
+
+	RequestRefresh();
+
+	RegisterActiveTimer(0.0f, FWidgetActiveTimerDelegate::CreateLambda([this](double InCurrentTime, float InDeltaTime)
+	{
+		CurrentProcessor->FocusInitialWidget();
+		return EActiveTimerReturnType::Stop;
+	}));
 	
-	RefreshAutoCompleteItems();
-	
-	RegisterActiveTimer(0.0f, FWidgetActiveTimerDelegate::CreateLambda([FindWidgetPtr = TWeakPtr<SAutoCompleteSearchBox>(FindSearchBox)](double InCurrentTime, float InDeltaTime)
+	RegisterActiveTimer(1.0f / 60.0f, FWidgetActiveTimerDelegate::CreateLambda([this](double InCurrentTime, float InDeltaTime)
 	{
-		if (FindWidgetPtr.IsValid())
+		if(bRefreshRequested)
 		{
-			FWidgetPath WidgetToFocusPath;
-			FSlateApplication::Get().GeneratePathToWidgetUnchecked(FindWidgetPtr.Pin()->GetSearchBox(), WidgetToFocusPath);
-			FSlateApplication::Get().SetKeyboardFocus(WidgetToFocusPath, EFocusCause::SetDirectly);
-			WidgetToFocusPath.GetWindow()->SetWidgetToFocusOnActivate(FindWidgetPtr.Pin()->GetSearchBox());
+			bRefreshRequested = false;
+			ToolbarContainer->SetContent(UToolMenus::Get()->GenerateWidget("AnimAssetFindReplaceToolbar", FToolMenuContext(ToolbarContext)));
+			FindReplaceWidgetContainer->SetContent(CurrentProcessor->MakeFindReplaceWidget());
+			CurrentProcessor->RefreshCachedData();
+			bRefreshSearchResultsRequested = true;
 		}
 
-		return EActiveTimerReturnType::Stop;
+		if(bRefreshSearchResultsRequested)
+		{
+			bRefreshSearchResultsRequested = false;
+			RefreshSearchResults();
+		}
+
+		return EActiveTimerReturnType::Continue;
 	}));
 }
 
-void SAnimAssetFindReplace::SetFindReplaceType(EAnimAssetFindReplaceType InType)
+void SAnimAssetFindReplace::SetCurrentProcessor(TSubclassOf<UAnimAssetFindReplaceProcessor> InProcessorClass)
+{
+	UAnimAssetFindReplaceProcessor** FoundProcessor = Processors.Find(InProcessorClass);
+	if(FoundProcessor)
+	{
+		CurrentProcessor = *FoundProcessor;
+	}
+
+	RequestRefresh();
+}
+
+UAnimAssetFindReplaceProcessor* SAnimAssetFindReplace::GetProcessor(TSubclassOf<UAnimAssetFindReplaceProcessor> InProcessorClass) const
+{
+	UAnimAssetFindReplaceProcessor* const* FoundProcessor = Processors.Find(InProcessorClass);
+	return FoundProcessor ? *FoundProcessor : nullptr;
+}
+
+void SAnimAssetFindReplace::AddReferencedObjects(FReferenceCollector& Collector)
+{
+	Collector.AddReferencedObjects(Processors);
+	Collector.AddReferencedObject(ToolbarContext);
+}
+
+FString SAnimAssetFindReplace::GetReferencerName() const
 {
-	Type = InType;
-	RefreshAutoCompleteItems();
-	RefreshSearchResults();
+	return TEXT("SAnimAssetFindReplace");
 }
 
 FARFilter SAnimAssetFindReplace::MakeARFilter() const
 {
 	FARFilter Filter;
-	Filter.ClassPaths =
+	for(const UClass* AssetClass : CurrentProcessor->GetSupportedAssetTypes())
 	{
-		UAnimationAsset::StaticClass()->GetClassPathName(),
-		USkeleton::StaticClass()->GetClassPathName(),
-		USkeletalMesh::StaticClass()->GetClassPathName()
-	};
+		Filter.ClassPaths.Add(AssetClass->GetClassPathName());
+	}
 	Filter.bRecursiveClasses = true;
 
 	return Filter;
@@ -720,17 +462,7 @@ void SAnimAssetFindReplace::RefreshSearchResults()
 
 bool SAnimAssetFindReplace::ShouldFilterOutAsset(const FAssetData& InAssetData, bool& bOutIsOldAsset) const
 {
-	switch(Type)
-	{
-	case EAnimAssetFindReplaceType::Curves:
-		return FilterByCurve(InAssetData, bOutIsOldAsset);
-	case EAnimAssetFindReplaceType::Notifies:
-		return FilterByNotify(InAssetData, bOutIsOldAsset);
-	default:
-		break;
-	}
-
-	return true;
+	return CurrentProcessor->ShouldFilterOutAsset(InAssetData, bOutIsOldAsset);
 }
 
 bool SAnimAssetFindReplace::HandleFilterAsset(const FAssetData& InAssetData)
@@ -745,613 +477,86 @@ bool SAnimAssetFindReplace::HandleFilterAsset(const FAssetData& InAssetData)
 	return bShouldFilterOut;
 }
 
-void SAnimAssetFindReplace::GetMatchingCurveNamesForAsset(const FAssetData& InAssetData, TArray<FString>& OutCurveNames) const
+FReply SAnimAssetFindReplace::HandleReplace()
 {
-	const FString TagValue = InAssetData.GetTagValueRef<FString>(USkeleton::CurveNameTag);
-	if (!TagValue.IsEmpty())
+	if(GetCurrentSelectionDelegate.IsBound())
 	{
-		check(IsInGameThread());
-		static TArray<FString> CurveNames;
-		CurveNames.Reset();
-		if (TagValue.ParseIntoArray(CurveNames, *USkeleton::CurveTagDelimiter, true) > 0)
+		TArray<FAssetData> SelectedAssets = GetCurrentSelectionDelegate.Execute();
+		if(SelectedAssets.Num() > 0)
 		{
-			for (const FString& CurveNameString : CurveNames)
-			{
-				if(bFindWholeWord)
-				{
-					if(CurveNameString.Compare(FindString, SearchCase) == 0)
-					{
-						OutCurveNames.Add(CurveNameString);
-					}
-				}
-				else
-				{
-					if(CurveNameString.Contains(FindString, SearchCase))
-					{
-						OutCurveNames.Add(CurveNameString);
-					}
-				}
-			}
+			ReplaceInAssets(SelectedAssets);
 		}
 	}
+
+	return FReply::Handled();
 }
 
-void SAnimAssetFindReplace::GetMatchingNotifyNamesForAsset(const FAssetData& InAssetData, TArray<FString>& OutNotifyNames) const
+FReply SAnimAssetFindReplace::HandleReplaceAll()
 {
-	const FString TagValue = InAssetData.GetTagValueRef<FString>(USkeleton::AnimNotifyTag);
-	if (!TagValue.IsEmpty())
+	// Apply current filter
+	FAssetRegistryModule& AssetRegistryModule = FModuleManager::LoadModuleChecked<FAssetRegistryModule>("AssetRegistry");
+
+	TArray<FAssetData> FilteredAssets;
+	AssetRegistryModule.Get().GetAssets(AssetPickerConfig.Filter,FilteredAssets);
+
+	FilteredAssets.RemoveAll([this](const FAssetData& InAssetData)
 	{
-		check(IsInGameThread());
-		static TArray<FString> NotifyNames;
-		NotifyNames.Reset();
-		if (TagValue.ParseIntoArray(NotifyNames, *USkeleton::AnimNotifyTagDelimiter, true) > 0)
-		{
-			for (const FString& NotifyNameString : NotifyNames)
-			{
-				if(bFindWholeWord)
-				{
-					if(NotifyNameString.Compare(FindString, SearchCase) == 0)
-					{
-						OutNotifyNames.Add(NotifyNameString);
-					}
-				}
-				else
-				{
-					if(NotifyNameString.Contains(FindString, SearchCase))
-					{
-						OutNotifyNames.Add(NotifyNameString);
-					}
-				}
-			}
-		}
-	}
+		bool bIsOldAsset = false;
+		return ShouldFilterOutAsset(InAssetData, bIsOldAsset);
+	});
+
+	ReplaceInAssets(FilteredAssets);
+
+	return FReply::Handled();
 }
 
-bool SAnimAssetFindReplace::NameMatches(const FString& InNameString) const
+void SAnimAssetFindReplace::ReplaceInAssets(const TArray<FAssetData>& InAssetDatas)
 {
-	if(bFindWholeWord)
+	const FText TypeName = CurrentProcessor->GetClass()->GetDisplayNameText();
+	const FText MessageText = FText::Format(LOCTEXT("ReplacingTaskStatus", "Replacing {0} in {1} Assets..."), TypeName, FText::AsNumber(InAssetDatas.Num()));
+	FScopedSlowTask ScopedSlowTask(static_cast<float>(InAssetDatas.Num()), MessageText);
+	ScopedSlowTask.MakeDialog(true);
+
+	FScopedTransaction ScopedTransaction(FText::Format(LOCTEXT("ReplaceTransaction", "Replace {0}."), TypeName));
+	
+	for(const FAssetData& AssetData : InAssetDatas)
 	{
-		if(InNameString.Compare(FindString, SearchCase) == 0)
+		ScopedSlowTask.EnterProgressFrame();
+
+		ReplaceInAsset(AssetData);
+
+		if(ScopedSlowTask.ShouldCancel())
 		{
-			return true;
+			break;
 		}
 	}
-	else
+
+	RequestRefreshSearchResults();
+}
+
+void SAnimAssetFindReplace::ReplaceInAsset(const FAssetData& InAssetData) const
+{
+	CurrentProcessor->ReplaceInAsset(InAssetData);
+}
+
+FReply SAnimAssetFindReplace::HandleRemove()
+{
+	if(GetCurrentSelectionDelegate.IsBound())
 	{
-		if(InNameString.Contains(FindString, SearchCase))
+		TArray<FAssetData> SelectedAssets = GetCurrentSelectionDelegate.Execute();
+		if(SelectedAssets.Num() > 0)
 		{
-			return true;
+			RemoveInAssets(SelectedAssets);
 		}
 	}
 
-	return false;
+	return FReply::Handled();
 }
 
-bool SAnimAssetFindReplace::FilterByCurve(const FAssetData& InAssetData, bool& bOutIsOldAsset) const
+FReply SAnimAssetFindReplace::HandleRemoveAll()
 {
-	FString TagValue;
-	if(InAssetData.GetTagValue<FString>(USkeleton::CurveNameTag, TagValue))
-	{
-		bOutIsOldAsset = false;
-
-		if(FindString.IsEmpty())
-		{
-			return true;
-		}
-		
-		if(SkeletonFilter.IsValid() )
-		{
-			if(InAssetData.GetClass() != USkeleton::StaticClass())
-			{
-				FString SkeletonPath;
-				if(InAssetData.GetTagValue<FString>(TEXT("Skeleton"), SkeletonPath))
-				{
-					if(SkeletonPath != SkeletonFilter.GetExportTextName())
-					{
-						return true;
-					}
-				}
-			}
-			else
-			{
-				if(InAssetData.ToSoftObjectPath() != SkeletonFilter.ToSoftObjectPath())
-				{
-					return true;
-				}
-			}
-		}
-
-		check(IsInGameThread());
-		static TArray<FString> CurveNames;
-		CurveNames.Reset();
-		if (TagValue.ParseIntoArray(CurveNames, *USkeleton::CurveTagDelimiter, true) > 0)
-		{
-			for (const FString& CurveNameString : CurveNames)
-			{
-				if(NameMatches(CurveNameString))
-				{
-					return false;
-				}
-			}
-		}
-	}
-	else
-	{
-		bOutIsOldAsset = IsAssetWithoutTagOldAsset(USkeleton::CurveNameTag, InAssetData);
-	}
-
-	return true;
-}
-
-bool SAnimAssetFindReplace::FilterByNotify(const FAssetData& InAssetData, bool& bOutIsOldAsset) const
-{
-	FString TagValue;
-	if(InAssetData.GetTagValue<FString>(USkeleton::AnimNotifyTag, TagValue))
-	{
-		bOutIsOldAsset = false;
-
-		if(FindString.IsEmpty())
-		{
-			return true;
-		}
-
-		if(SkeletonFilter.IsValid() )
-		{
-			if(InAssetData.GetClass() != USkeleton::StaticClass())
-			{
-				FString SkeletonPath;
-				if(InAssetData.GetTagValue<FString>(TEXT("Skeleton"), SkeletonPath))
-				{
-					if(SkeletonPath != SkeletonFilter.GetExportTextName())
-					{
-						return true;
-					}
-				}
-			}
-			else
-			{
-				if(InAssetData.ToSoftObjectPath() != SkeletonFilter.ToSoftObjectPath())
-				{
-					return true;
-				}
-			}
-		}
-		
-		check(IsInGameThread());
-		static TArray<FString> NotifyNames;
-		NotifyNames.Reset();
-		if (TagValue.ParseIntoArray(NotifyNames, *USkeleton::AnimNotifyTagDelimiter, true) > 0)
-		{
-			for (const FString& NotifyNameString : NotifyNames)
-			{
-				if(NameMatches(NotifyNameString))
-				{
-					return false;
-				}
-			}
-		}
-	}
-	else
-	{
-		bOutIsOldAsset = IsAssetWithoutTagOldAsset(USkeleton::AnimNotifyTag, InAssetData);
-	}
-
-	return true;
-}
-
-bool SAnimAssetFindReplace::IsAssetWithoutTagOldAsset(FName InTag, const FAssetData& InAssetData) const
-{
-	FAssetRegistryModule& AssetRegistryModule = FModuleManager::LoadModuleChecked<FAssetRegistryModule>("AssetRegistry");
-	UClass* AssetClass = InAssetData.GetClass();
-	if(AssetClass->IsChildOf(UAnimSequenceBase::StaticClass()))
-	{
-		if(InTag == USkeleton::CurveNameTag)
-		{
-			// Check the package object version - the asset was saving empty tags for curves, so the absence of curves is
-			// not the same as an empty value
-			FAssetPackageData PackageData;
-			const UE::AssetRegistry::EExists PackageExists = AssetRegistryModule.Get().TryGetAssetPackageData(InAssetData.PackageName, PackageData);
-			if (PackageExists == UE::AssetRegistry::EExists::Exists)
-			{
-				return PackageData.FileVersionUE < VER_UE4_SKELETON_ADD_SMARTNAMES;
-			}
-			else
-			{
-				// Does not exist or unknown - treat it as 'old'
-				return true;
-			}
-		}
-		else if(InTag == USkeleton::AnimNotifyTag)
-		{
-			return true;
-		}
-	}
-	else if(AssetClass->IsChildOf(UPoseAsset::StaticClass()))
-	{
-		if(InTag == USkeleton::CurveNameTag)
-		{
-			// Check the package custom version - the asset was saving empty tags for curves, so the absence of curves is
-			// not the same as an empty value
-			FAssetPackageData PackageData;
-			const UE::AssetRegistry::EExists PackageExists = AssetRegistryModule.Get().TryGetAssetPackageData(InAssetData.PackageName, PackageData);
-			if (PackageExists == UE::AssetRegistry::EExists::Exists)
-			{
-				for(const UE::AssetRegistry::FPackageCustomVersion& CustomVersion : PackageData.GetCustomVersions())
-				{
-					if(CustomVersion.Key == FAnimPhysObjectVersion::GUID)
-					{
-						return CustomVersion.Version < FAnimPhysObjectVersion::SmartNameRefactorForDeterministicCooking;
-					}
-				}
-
-				// No FAnimPhysObjectVersion, treat as old
-				return true;
-			}
-			else
-			{
-				// Does not exist or unknown - treat it as 'old'
-				return true;
-			}
-		}
-	}
-	else if(AssetClass->IsChildOf(USkeleton::StaticClass()))
-	{
-		return true;
-	}
-	else if(AssetClass->IsChildOf(USkeletalMesh::StaticClass()))
-	{
-		// Skeletal meshes didn't have curves before, so cant be 'old' 
-		return false;
-	}
-
-	// Assume unknown assets are not 'old'
-	return false;
-}
-
-FReply SAnimAssetFindReplace::HandleReplace()
-{
-	if(GetCurrentSelectionDelegate.IsBound())
-	{
-		TArray<FAssetData> SelectedAssets = GetCurrentSelectionDelegate.Execute();
-		if(SelectedAssets.Num() > 0)
-		{
-			ReplaceInAssets(SelectedAssets);
-		}
-	}
-
-	return FReply::Handled();
-}
-
-FReply SAnimAssetFindReplace::HandleReplaceAll()
-{
-	// Apply current filter
-	FAssetRegistryModule& AssetRegistryModule = FModuleManager::LoadModuleChecked<FAssetRegistryModule>("AssetRegistry");
-
-	TArray<FAssetData> FilteredAssets;
-	AssetRegistryModule.Get().GetAssets(AssetPickerConfig.Filter,FilteredAssets);
-
-	FilteredAssets.RemoveAll([this](const FAssetData& InAssetData)
-	{
-		bool bIsOldAsset = false;
-		return ShouldFilterOutAsset(InAssetData, bIsOldAsset);
-	});
-
-	ReplaceInAssets(FilteredAssets);
-
-	return FReply::Handled();
-}
-
-void SAnimAssetFindReplace::ReplaceInAssets(const TArray<FAssetData>& InAssetDatas)
-{
-	FText TypeName;
-	switch(Type)
-	{
-	case EAnimAssetFindReplaceType::Curves:
-		TypeName = LOCTEXT("FindReplaceTypeCurves", "Curves");
-		break;
-	case EAnimAssetFindReplaceType::Notifies:
-		TypeName = LOCTEXT("FindReplaceTypeNotifies", "Notifies");
-		break;
-	default:
-		break;
-	}
-
-	{
-		const FText MessageText = FText::Format(LOCTEXT("ReplacingTaskStatus", "Replacing {0}: Matching '{1}' with '{2}' in {3} Assets..."), TypeName, FText::FromString(FindString), FText::FromString(ReplaceString), FText::AsNumber(InAssetDatas.Num()));
-		FScopedSlowTask ScopedSlowTask(static_cast<float>(InAssetDatas.Num()), MessageText);
-		ScopedSlowTask.MakeDialog(true);
-
-		FScopedTransaction ScopedTransaction(FText::Format(LOCTEXT("ReplaceTransaction", "Replace {0}."), TypeName));
-		
-		for(const FAssetData& AssetData : InAssetDatas)
-		{
-			ScopedSlowTask.EnterProgressFrame();
-
-			ReplaceInAsset(AssetData);
-
-			if(ScopedSlowTask.ShouldCancel())
-			{
-				break;
-			}
-		}
-	}
-
-	RefreshSearchResults();
-}
-
-void SAnimAssetFindReplace::ReplaceInAsset(const FAssetData& InAssetData) const
-{
-	switch(Type)
-	{
-	case EAnimAssetFindReplaceType::Curves:
-		ReplaceCurvesInAsset(InAssetData);
-		break;
-	case EAnimAssetFindReplaceType::Notifies:
-		ReplaceNotifiesInAsset(InAssetData);
-		break;
-	default:
-		break;
-	}
-}
-
-void SAnimAssetFindReplace::ReplaceCurvesInAsset(const FAssetData& InAssetData) const
-{
-	if(UObject* Asset = InAssetData.GetAsset())
-	{
-		if(UAnimSequenceBase* AnimSequenceBase = Cast<UAnimSequenceBase>(Asset))
-		{
-			Asset->MarkPackageDirty();
-
-			if(bFindWholeWord)
-			{
-				const FAnimationCurveIdentifier FindCurveId(*FindString, ERawCurveTrackTypes::RCT_Float);
-				const FAnimationCurveIdentifier ReplaceCurveId(*ReplaceString, ERawCurveTrackTypes::RCT_Float);
-				IAnimationDataController::FScopedBracket ScopedBracket(AnimSequenceBase->GetController(), LOCTEXT("ReplaceCurves", "Replace Curves"));
-				AnimSequenceBase->GetController().RenameCurve(FindCurveId, ReplaceCurveId);
-			}
-			else
-			{
-				TArray<TPair<FAnimationCurveIdentifier, FAnimationCurveIdentifier>> FindReplacePairs;
-				const TArray<FFloatCurve>& Curves = AnimSequenceBase->GetDataModel()->GetFloatCurves();
-				for(const FFloatCurve& Curve : Curves)
-				{
-					FString CurveName = Curve.GetName().ToString();
-					if(NameMatches(CurveName))
-					{
-						const FAnimationCurveIdentifier FindCurveId(*CurveName, ERawCurveTrackTypes::RCT_Float);
-						const FString NewName = CurveName.Replace(*FindString, *ReplaceString, SearchCase);
-						const FAnimationCurveIdentifier ReplaceCurveId(*NewName, ERawCurveTrackTypes::RCT_Float);
-
-						FindReplacePairs.Emplace(FindCurveId, ReplaceCurveId);
-					}
-				}
-
-				if(FindReplacePairs.Num() > 0)
-				{
-					IAnimationDataController::FScopedBracket ScopedBracket(AnimSequenceBase->GetController(), LOCTEXT("ReplaceCurves", "Replace Curves"));
-
-					for(const TPair<FAnimationCurveIdentifier, FAnimationCurveIdentifier>& FindReplacePair : FindReplacePairs)
-					{
-						AnimSequenceBase->GetController().RenameCurve(FindReplacePair.Key, FindReplacePair.Value);
-					}
-				}
-			}
-		}
-		else if(UPoseAsset* PoseAsset = Cast<UPoseAsset>(Asset))
-		{
-			if(bFindWholeWord)
-			{
-				Asset->Modify();
-				
-				const FName FindCurveName(*FindString);
-				const FName ReplaceCurveName(*ReplaceString);
-				PoseAsset->RenamePoseOrCurveName(FindCurveName, ReplaceCurveName);
-			}
-			else
-			{
-				TArray<TPair<FName, FName>> FindReplacePairs;
-
-				for(const FName& PoseName : PoseAsset->GetPoseFNames())
-				{
-					FString CurveName = PoseName.ToString();
-					if(NameMatches(CurveName))
-					{
-						const FName FindCurveName(*CurveName);
-						const FString NewName = CurveName.Replace(*FindString, *ReplaceString, SearchCase);
-						const FName ReplaceCurveName(*NewName);
-
-						FindReplacePairs.Emplace(FindCurveName, ReplaceCurveName);
-					}
-				}
-
-				if(FindReplacePairs.Num() > 0)
-				{
-					Asset->Modify();
-
-					for(const TPair<FName, FName>& FindReplacePair : FindReplacePairs)
-					{
-						PoseAsset->RenamePoseOrCurveName(FindReplacePair.Key, FindReplacePair.Value);
-					}
-				}
-			}
-		}
-		else if(USkeleton* Skeleton = Cast<USkeleton>(Asset))
-		{
-			if(bFindWholeWord)
-			{
-				Asset->Modify();
-
-				const FName FindCurveName(*FindString);
-				const FName ReplaceCurveName(*ReplaceString);
-				Skeleton->RenameCurveMetaData(FindCurveName, ReplaceCurveName);
-			}
-			else
-			{
-				TArray<TPair<FName, FName>> FindReplacePairs;
-
-				Skeleton->ForEachCurveMetaData([this, &FindReplacePairs](FName InCurveName, const FCurveMetaData& InMetaData)
-				{
-					FString CurveNameString = InCurveName.ToString();
-					if(NameMatches(CurveNameString))
-					{
-						const FName FindCurveName(InCurveName);
-						const FString NewName = CurveNameString.Replace(*FindString, *ReplaceString, SearchCase);
-						const FName ReplaceCurveName(*NewName);
-
-						FindReplacePairs.Emplace(FindCurveName, ReplaceCurveName);
-					}
-				});
-
-				if(FindReplacePairs.Num() > 0)
-				{
-					Asset->Modify();
-
-					for(const TPair<FName, FName>& FindReplacePair : FindReplacePairs)
-					{
-						Skeleton->RenameCurveMetaData(FindReplacePair.Key, FindReplacePair.Value);
-					}
-				}
-			}
-		}
-		else if(USkeletalMesh* SkeletalMesh = Cast<USkeletalMesh>(Asset))
-		{
-			if(UAnimCurveMetaData* AnimCurveMetaData = SkeletalMesh->GetAssetUserData<UAnimCurveMetaData>())
-			{
-				if(bFindWholeWord)
-				{
-					Asset->Modify();
-
-					const FName FindCurveName(*FindString);
-					const FName ReplaceCurveName(*ReplaceString);
-					AnimCurveMetaData->RenameCurveMetaData(FindCurveName, ReplaceCurveName);
-				}
-				else
-				{
-					TArray<TPair<FName, FName>> FindReplacePairs;
-
-					AnimCurveMetaData->ForEachCurveMetaData([this, &FindReplacePairs](FName InCurveName, const FCurveMetaData& InMetaData)
-					{
-						FString CurveNameString = InCurveName.ToString();
-						if(NameMatches(CurveNameString))
-						{
-							const FName FindCurveName(InCurveName);
-							const FString NewName = CurveNameString.Replace(*FindString, *ReplaceString, SearchCase);
-							const FName ReplaceCurveName(*NewName);
-
-							FindReplacePairs.Emplace(FindCurveName, ReplaceCurveName);
-						}
-					});
-
-					if(FindReplacePairs.Num() > 0)
-					{
-						Asset->Modify();
-
-						for(const TPair<FName, FName>& FindReplacePair : FindReplacePairs)
-						{
-							AnimCurveMetaData->RenameCurveMetaData(FindReplacePair.Key, FindReplacePair.Value);
-						}
-					}
-				}
-			}
-		}
-	}
-}
-
-void SAnimAssetFindReplace::ReplaceNotifiesInAsset(const FAssetData& InAssetData) const
-{
-	if(UObject* Asset = InAssetData.GetAsset())
-	{
-		if(UAnimSequenceBase* AnimSequenceBase = Cast<UAnimSequenceBase>(Asset))
-		{
-			if(bFindWholeWord)
-			{
-				Asset->Modify();
-				AnimSequenceBase->RenameNotifies(*FindString, *ReplaceString);
-			}
-			else
-			{
-				TArray<TPair<FName, FName>> FindReplacePairs;
-
-				for(const FAnimNotifyEvent& Notify : AnimSequenceBase->Notifies)
-				{
-					// Only handle named notifies
-					if(!Notify.IsBlueprintNotify())
-					{
-						FString NotifyName = Notify.NotifyName.ToString();
-						if(NameMatches(NotifyName))
-						{
-							const FName FindNotifyName(*NotifyName);
-							const FString NewName = NotifyName.Replace(*FindString, *ReplaceString, SearchCase);
-							const FName ReplaceNotifyName(*NewName);
-							FindReplacePairs.AddUnique(TPair<FName, FName>(FindNotifyName, ReplaceNotifyName));
-						}
-					}
-				}
-
-				if(FindReplacePairs.Num() > 0)
-				{
-					Asset->Modify();
-
-					for(const TPair<FName, FName>& FindReplacePair : FindReplacePairs)
-					{
-						AnimSequenceBase->RenameNotifies(FindReplacePair.Key, FindReplacePair.Value);
-					}
-				}
-			}
-		}
-		else if(USkeleton* Skeleton = Cast<USkeleton>(Asset))
-		{
-			if(bFindWholeWord)
-			{
-				Asset->Modify();
-				Skeleton->RenameAnimationNotify(*FindString, *ReplaceString);
-			}
-			else
-			{
-				TArray<TPair<FName, FName>> FindReplacePairs;
-
-				for(const FName& NotifyName : Skeleton->AnimationNotifies)
-				{
-					FString NotifyString = NotifyName.ToString();
-					if(NameMatches(NotifyString))
-					{
-						const FString NewName = NotifyString.Replace(*FindString, *ReplaceString, SearchCase);
-						const FName ReplaceNotifyName(*NewName);
-						FindReplacePairs.AddUnique(TPair<FName, FName>(NotifyName, ReplaceNotifyName));
-					}
-				}
-
-				if(FindReplacePairs.Num() > 0)
-				{
-					Asset->Modify();
-					for(const TPair<FName, FName>& FindReplacePair : FindReplacePairs)
-					{
-						Skeleton->RenameAnimationNotify(FindReplacePair.Key, FindReplacePair.Value);
-					}
-				}
-			}
-		}
-	}
-}
-
-FReply SAnimAssetFindReplace::HandleRemove()
-{
-	if(GetCurrentSelectionDelegate.IsBound())
-	{
-		TArray<FAssetData> SelectedAssets = GetCurrentSelectionDelegate.Execute();
-		if(SelectedAssets.Num() > 0)
-		{
-			RemoveInAssets(SelectedAssets);
-		}
-	}
-
-	return FReply::Handled();
-}
-
-FReply SAnimAssetFindReplace::HandleRemoveAll()
-{
-	// Apply current filter
-	FAssetRegistryModule& AssetRegistryModule = FModuleManager::LoadModuleChecked<FAssetRegistryModule>("AssetRegistry");
+	// Apply current filter
+	FAssetRegistryModule& AssetRegistryModule = FModuleManager::LoadModuleChecked<FAssetRegistryModule>("AssetRegistry");
 
 	TArray<FAssetData> FilteredAssets;
 	AssetRegistryModule.Get().GetAssets(AssetPickerConfig.Filter,FilteredAssets);
@@ -1369,287 +574,31 @@ FReply SAnimAssetFindReplace::HandleRemoveAll()
 
 void SAnimAssetFindReplace::RemoveInAssets(const TArray<FAssetData>& InAssetDatas)
 {
-	FText TypeName;
-	switch(Type)
-	{
-	case EAnimAssetFindReplaceType::Curves:
-		TypeName = LOCTEXT("FindReplaceTypeCurves", "Curves");
-		break;
-	case EAnimAssetFindReplaceType::Notifies:
-		TypeName = LOCTEXT("FindReplaceTypeNotifies", "Notifies");
-		break;
-	default:
-		break;
-	}
-
-	{
-		const FText MessageText = FText::Format(LOCTEXT("RemovingTaskStatus", "Removing {0}: Matching '{1}' in {2} Assets..."), TypeName, FText::FromString(FindString), FText::AsNumber(InAssetDatas.Num()));
-		FScopedSlowTask ScopedSlowTask(static_cast<float>(InAssetDatas.Num()), MessageText);
-		ScopedSlowTask.MakeDialog(true);
-
-		FScopedTransaction ScopedTransaction(FText::Format(LOCTEXT("RemoveTransaction", "Remove {0}."), TypeName));
-
-		for(const FAssetData& AssetData : InAssetDatas)
-		{
-			ScopedSlowTask.EnterProgressFrame();
-
-			RemoveInAsset(AssetData);
+	const FText TypeName = CurrentProcessor->GetClass()->GetDisplayNameText();
+	const FText MessageText = FText::Format(LOCTEXT("RemovingTaskStatus", "Removing {0} in {1} Assets..."), TypeName, FText::AsNumber(InAssetDatas.Num()));
+	FScopedSlowTask ScopedSlowTask(static_cast<float>(InAssetDatas.Num()), MessageText);
+	ScopedSlowTask.MakeDialog(true);
 
-			if(ScopedSlowTask.ShouldCancel())
-			{
-				break;
-			}
-		}
-	}
-
-	RefreshSearchResults();
-}
+	FScopedTransaction ScopedTransaction(FText::Format(LOCTEXT("RemoveTransaction", "Remove {0}."), TypeName));
 
-void SAnimAssetFindReplace::RemoveInAsset(const FAssetData& InAssetData) const
-{
-	switch(Type)
+	for(const FAssetData& AssetData : InAssetDatas)
 	{
-	case EAnimAssetFindReplaceType::Curves:
-		RemoveCurvesInAsset(InAssetData);
-		break;
-	case EAnimAssetFindReplaceType::Notifies:
-		RemoveNotifiesInAsset(InAssetData);
-		break;
-	default:
-		break;
-	}
-}
+		ScopedSlowTask.EnterProgressFrame();
 
-void SAnimAssetFindReplace::RemoveCurvesInAsset(const FAssetData& InAssetData) const
-{
-	if(UObject* Asset = InAssetData.GetAsset())
-	{
-		if(UAnimSequenceBase* AnimSequenceBase = Cast<UAnimSequenceBase>(Asset))
-		{
-			Asset->MarkPackageDirty();
+		RemoveInAsset(AssetData);
 
-			if(bFindWholeWord)
-			{
-				const FAnimationCurveIdentifier CurveId(*FindString, ERawCurveTrackTypes::RCT_Float);
-				IAnimationDataController::FScopedBracket ScopedBracket(AnimSequenceBase->GetController(), LOCTEXT("RemoveCurves", "Remove Curves"));
-				AnimSequenceBase->GetController().RemoveCurve(CurveId);
-			}
-			else
-			{
-				TSet<FAnimationCurveIdentifier> CurveIdsToRemove;
-				const TArray<FFloatCurve>& Curves = AnimSequenceBase->GetDataModel()->GetFloatCurves();
-				for(const FFloatCurve& Curve : Curves)
-				{
-					FString CurveName = Curve.GetName().ToString();
-					if(NameMatches(CurveName))
-					{
-						const FAnimationCurveIdentifier CurveId(*CurveName, ERawCurveTrackTypes::RCT_Float);
-						CurveIdsToRemove.Add(CurveId);
-					}
-				}
-
-				if(CurveIdsToRemove.Num() > 0)
-				{
-					IAnimationDataController::FScopedBracket ScopedBracket(AnimSequenceBase->GetController(), LOCTEXT("RemoveCurves", "Remove Curves"));
-					for(const FAnimationCurveIdentifier& CurveIdToRemove : CurveIdsToRemove)
-					{
-						AnimSequenceBase->GetController().RemoveCurve(CurveIdToRemove);
-					}
-				}
-			}
-		}
-		else if(UPoseAsset* PoseAsset = Cast<UPoseAsset>(Asset))
+		if(ScopedSlowTask.ShouldCancel())
 		{
-			if(bFindWholeWord)
-			{
-				Asset->Modify();
-				PoseAsset->RemovePoseOrCurveNames({ *FindString });
-			}
-			else
-			{
-				TArray<FName> CurveIdsToRemove;
-
-				for(const FName& PoseName : PoseAsset->GetPoseFNames())
-				{
-					if(NameMatches(PoseName.ToString()))
-					{
-						CurveIdsToRemove.AddUnique(PoseName);
-					}
-				}
-
-				for(const FName& CurveName : PoseAsset->GetCurveFNames())
-				{
-					if(NameMatches(CurveName.ToString()))
-					{
-						CurveIdsToRemove.AddUnique(CurveName);
-					}
-				}
-				
-				if(CurveIdsToRemove.Num() > 0)
-				{
-					Asset->Modify();
-					PoseAsset->RemovePoseOrCurveNames(CurveIdsToRemove);
-				}
-			}
-		}
-		else if(USkeleton* Skeleton = Cast<USkeleton>(Asset))
-		{
-			if(bFindWholeWord)
-			{
-				Asset->Modify();
-				Skeleton->RemoveCurveMetaData(*FindString);
-			}
-			else
-			{
-				TArray<FName> CurvesToRemove;
-				Skeleton->ForEachCurveMetaData([this, &CurvesToRemove](FName InCurveName, const FCurveMetaData& InCurveMetaData)
-				{
-				 	if(NameMatches(InCurveName.ToString()))
-				 	{
-						 CurvesToRemove.AddUnique(InCurveName);
-					 }
-				});
-
-				if(CurvesToRemove.Num() > 0)
-				{
-					Asset->Modify();
-					Skeleton->RemoveCurveMetaData(CurvesToRemove);
-				}
-			}
+			break;
 		}
 	}
-}
 
-void SAnimAssetFindReplace::RemoveNotifiesInAsset(const FAssetData& InAssetData) const
-{
-	if(UObject* Asset = InAssetData.GetAsset())
-	{
-		if(UAnimSequenceBase* AnimSequenceBase = Cast<UAnimSequenceBase>(Asset))
-		{
-			if(bFindWholeWord)
-			{
-				Asset->Modify();
-				AnimSequenceBase->RemoveNotifies( { *FindString } );
-			}
-			else
-			{
-				TArray<FName> NotifiesToRemove;
-				for(const FAnimNotifyEvent& Notify : AnimSequenceBase->Notifies)
-				{
-					FString NotifyNameString = Notify.NotifyName.ToString();
-					if(NameMatches(NotifyNameString))
-					{
-						NotifiesToRemove.AddUnique(Notify.NotifyName);
-					}
-				}
-
-				if(NotifiesToRemove.Num() > 0)
-				{
-					Asset->Modify();
-					AnimSequenceBase->RemoveNotifies(NotifiesToRemove);
-				}
-			}
-		}
-		else if(USkeleton* Skeleton = Cast<USkeleton>(Asset))
-		{
-			if(bFindWholeWord)
-			{
-				Asset->Modify();
-				Skeleton->RemoveAnimationNotify(*FindString);
-			}
-			else
-			{
-				TArray<FName> NotifiesToRemove;
-
-				for(const FName& NotifyName : Skeleton->AnimationNotifies)
-				{
-					if(NameMatches(NotifyName.ToString()))
-					{
-						NotifiesToRemove.AddUnique(NotifyName);
-					}
-				}
-
-				if(NotifiesToRemove.Num() > 0)
-				{
-					Asset->Modify();
-					for(const FName& NotifyToRemove : NotifiesToRemove)
-					{
-						Skeleton->RemoveAnimationNotify(NotifyToRemove);
-					}
-				}
-			}
-		}
-	}
+	RequestRefreshSearchResults();
 }
 
-void SAnimAssetFindReplace::RefreshAutoCompleteItems()
+void SAnimAssetFindReplace::RemoveInAsset(const FAssetData& InAssetData) const
 {
-	AutoCompleteItems->Empty();
-
-	// We use the asset registry to query all assets and accumulate their curve names
-	const FAssetRegistryModule& AssetRegistryModule = FModuleManager::LoadModuleChecked<FAssetRegistryModule>("AssetRegistry");
-
-	FARFilter Filter;
-	Filter.bRecursiveClasses = true;
-	Filter.ClassPaths = { UAnimationAsset::StaticClass()->GetClassPathName(), USkeleton::StaticClass()->GetClassPathName() };
-
-	TArray<FAssetData> FoundAssetData;
-	AssetRegistryModule.Get().GetAssets(Filter, FoundAssetData);
-
-	TSet<FString> UniqueNames;
-	
-	switch(Type)
-	{
-	case EAnimAssetFindReplaceType::Curves:
-		for (const FAssetData& AssetData : FoundAssetData)
-		{
-			const FString TagValue = AssetData.GetTagValueRef<FString>(USkeleton::CurveNameTag);
-			if (!TagValue.IsEmpty())
-			{
-				TArray<FString> AssetCurveNames;
-				if (TagValue.ParseIntoArray(AssetCurveNames, *USkeleton::CurveTagDelimiter, true) > 0)
-				{
-					for (FString& CurveNameString : AssetCurveNames)
-					{
-						UniqueNames.Add(MoveTemp(CurveNameString));
-					}
-				}
-			}
-		}
-		break;
-	case EAnimAssetFindReplaceType::Notifies:
-		for (const FAssetData& AssetData : FoundAssetData)
-		{
-			const FString TagValue = AssetData.GetTagValueRef<FString>(USkeleton::AnimNotifyTag);
-			if (!TagValue.IsEmpty())
-			{
-				TArray<FString> AssetNotifyNames;
-				if (TagValue.ParseIntoArray(AssetNotifyNames, *USkeleton::AnimNotifyTagDelimiter, true) > 0)
-				{
-					for (FString& NotifyNameString : AssetNotifyNames)
-					{
-						UniqueNames.Add(MoveTemp(NotifyNameString));
-					}
-				}
-			}
-		}
-		break;
-	default:
-		break;
-	}
-
-	for(FString& UniqueName : UniqueNames)
-	{
-		AutoCompleteItems->Add(MakeShared<FString>(MoveTemp(UniqueName)));
-	}
-
-	AutoCompleteItems->Sort([](TSharedPtr<FString> InLHS, TSharedPtr<FString> InRHS)
-	{
-		return *InLHS.Get() < *InRHS.Get();
-	});
-
-	FindSearchBox->RefreshAutoCompleteItems();
-	ReplaceSearchBox->RefreshAutoCompleteItems();
+	CurrentProcessor->RemoveInAsset(InAssetData);
 }
 
 #undef LOCTEXT_NAMESPACE
\ No newline at end of file
diff --git a/Engine/Source/Editor/Persona/Private/SAnimAssetFindReplace.h b/Engine/Source/Editor/Persona/Private/SAnimAssetFindReplace.h
--- a/Engine/Source/Editor/Persona/Private/SAnimAssetFindReplace.h
+++ b/Engine/Source/Editor/Persona/Private/SAnimAssetFindReplace.h
@@ -14,6 +14,11 @@
 class SBox;
 class SAutoCompleteSearchBox;
 
+namespace AnimAssetFindReplacePrivate
+{
+extern TSharedPtr<SAnimAssetFindReplace> GetWidgetFromContext(const FToolMenuContext& InContext);
+}
+
 struct FAnimAssetFindReplaceSummoner : public FWorkflowTabFactory
 {
 public:
@@ -26,33 +31,33 @@ private:
 	FAnimAssetFindReplaceConfig Config;
 };
 
-class SAnimAssetFindReplace : public SCompoundWidget
+class SAnimAssetFindReplace : public IAnimAssetFindReplace, public FGCObject
 {
 public:
+	friend class UAnimAssetFindReplaceProcessor;
+
 	SLATE_BEGIN_ARGS(SAnimAssetFindReplace) {}
-	
+
 	SLATE_ARGUMENT(FAnimAssetFindReplaceConfig, Config)
 
 	SLATE_END_ARGS()
 
 	void Construct(const FArguments& InArgs);
 
-	/** Sets the type of thing we are finding/replacing, then flags the UI for a refresh */
-	void SetFindReplaceType(EAnimAssetFindReplaceType InType);
+	// IAnimAssetFindReplace interface
+	virtual void SetCurrentProcessor(TSubclassOf<UAnimAssetFindReplaceProcessor> InProcessor) override;
+	virtual UAnimAssetFindReplaceProcessor* GetCurrentProcessor() const override { return CurrentProcessor; }
+	virtual UAnimAssetFindReplaceProcessor* GetProcessor(TSubclassOf<UAnimAssetFindReplaceProcessor> InProcessor) const override;
 
 private:
+	// FGCObject interface
+	virtual void AddReferencedObjects(FReferenceCollector& Collector) override;
+	virtual FString GetReferencerName() const override;
+
 	void RefreshSearchResults();
 
 	bool HandleFilterAsset(const FAssetData& InAssetData);
 
-	bool FilterByCurve(const FAssetData& InAssetData, bool& bOutIsOldAsset) const;
-	
-	bool FilterByNotify(const FAssetData& InAssetData, bool& bOutIsOldAsset) const;
-
-	void GetMatchingCurveNamesForAsset(const FAssetData& InAssetData, TArray<FString>& OutCurveNames) const;
-
-	void GetMatchingNotifyNamesForAsset(const FAssetData& InAssetData, TArray<FString>& OutNotifyNames) const;
-
 	FReply HandleReplace();
 
 	FReply HandleReplaceAll();
@@ -61,10 +66,6 @@ private:
 
 	void ReplaceInAsset(const FAssetData& InAssetData) const;
 
-	void ReplaceCurvesInAsset(const FAssetData& InAssetData) const;
-
-	void ReplaceNotifiesInAsset(const FAssetData& InAssetData) const;
-	
 	FReply HandleRemove();
 
 	FReply HandleRemoveAll();
@@ -73,30 +74,24 @@ private:
 
 	void RemoveInAsset(const FAssetData& InAssetData) const;
 
-	void RemoveCurvesInAsset(const FAssetData& InAssetData) const;
-
-	void RemoveNotifiesInAsset(const FAssetData& InAssetData) const;
-
 	bool ShouldFilterOutAsset(const FAssetData& InAssetData, bool& bOutIsOldAsset) const;
 
-	bool NameMatches(const FString& InNameString) const;
+	void RequestRefresh() { bRefreshRequested = true; }
 
-	void RefreshAutoCompleteItems();
+	void RequestRefreshSearchResults() { bRefreshSearchResultsRequested = true; }
 
 	FARFilter MakeARFilter() const;
 
-	bool IsAssetWithoutTagOldAsset(FName InTag, const FAssetData& InAssetData) const;
-
 private:
+	FAnimAssetFindReplaceConfig Config;
+
 	FAssetPickerConfig AssetPickerConfig;
 
 	EAnimAssetFindReplaceMode Mode = EAnimAssetFindReplaceMode::Find;
 
-	EAnimAssetFindReplaceType Type = EAnimAssetFindReplaceType::Curves;
+	TMap<TSubclassOf<UAnimAssetFindReplaceProcessor>, UAnimAssetFindReplaceProcessor*> Processors;
 
-	FString FindString;
-
-	FString ReplaceString;
+	UAnimAssetFindReplaceProcessor* CurrentProcessor = nullptr;
 
 	FRefreshAssetViewDelegate RefreshAssetViewDelegate;
 
@@ -104,23 +99,19 @@ private:
 
 	FSetARFilterDelegate SetARFilterDelegate;
 
-	bool bFindWholeWord = true;
-
 	bool bAssetsSelected = false;
 
 	bool bFoundAssets = false;
 
-	TArray<FAssetData> OldAssets;
+	bool bRefreshRequested = false;
 
-	ESearchCase::Type SearchCase = ESearchCase::IgnoreCase;
+	bool bRefreshSearchResultsRequested = false;
 
-	TStrongObjectPtr<UAnimAssetFindReplaceContext> ToolbarContext;
+	TArray<FAssetData> OldAssets;
 
-	TSharedPtr<TArray<TSharedPtr<FString>>> AutoCompleteItems;
+	UAnimAssetFindReplaceContext* ToolbarContext;
 
-	TSharedPtr<SAutoCompleteSearchBox> FindSearchBox;
-	
-	TSharedPtr<SAutoCompleteSearchBox> ReplaceSearchBox;
+	TSharedPtr<SBox> ToolbarContainer;
 
-	FAssetData SkeletonFilter;
+	TSharedPtr<SBox> FindReplaceWidgetContainer;
 };
\ No newline at end of file
diff --git a/Engine/Source/Editor/Persona/Private/SAnimCurveMetadataEditor.cpp b/Engine/Source/Editor/Persona/Private/SAnimCurveMetadataEditor.cpp
--- a/Engine/Source/Editor/Persona/Private/SAnimCurveMetadataEditor.cpp
+++ b/Engine/Source/Editor/Persona/Private/SAnimCurveMetadataEditor.cpp
@@ -2,6 +2,8 @@
 
 
 #include "SAnimCurveMetadataEditor.h"
+
+#include "AnimAssetFindReplaceCurves.h"
 #include "Framework/Notifications/NotificationManager.h"
 #include "Widgets/Input/SCheckBox.h"
 #include "Widgets/Notifications/SNotificationList.h"
@@ -441,8 +443,8 @@ void SAnimCurveMetadataEditor::Construct(const FArguments& InArgs, UObject* InAn
 						{
 							if(TSharedPtr<SDockTab> Tab = TabManager->TryInvokeTab(FPersonaTabs::FindReplaceID))
 							{
-								TSharedRef<SAnimAssetFindReplace> FindReplaceWidget = StaticCastSharedRef<SAnimAssetFindReplace>(Tab->GetContent());
-								FindReplaceWidget->SetFindReplaceType(EAnimAssetFindReplaceType::Curves);
+								TSharedRef<IAnimAssetFindReplace> FindReplaceWidget = StaticCastSharedRef<IAnimAssetFindReplace>(Tab->GetContent());
+								FindReplaceWidget->SetCurrentProcessor(UAnimAssetFindReplaceCurves::StaticClass());
 							}
 						}
 					}
diff --git a/Engine/Source/Editor/Persona/Private/SAnimCurveViewer.cpp b/Engine/Source/Editor/Persona/Private/SAnimCurveViewer.cpp
--- a/Engine/Source/Editor/Persona/Private/SAnimCurveViewer.cpp
+++ b/Engine/Source/Editor/Persona/Private/SAnimCurveViewer.cpp
@@ -3,6 +3,7 @@
 
 #include "SAnimCurveViewer.h"
 
+#include "AnimAssetFindReplaceCurves.h"
 #include "AnimationEditorUtils.h"
 #include "Framework/Notifications/NotificationManager.h"
 #include "Widgets/Input/SCheckBox.h"
@@ -544,8 +545,8 @@ void SAnimCurveViewer::Construct(const FArguments& InArgs,  const TSharedRef<IPe
 					{
 						if(TSharedPtr<SDockTab> Tab = TabManager->TryInvokeTab(FPersonaTabs::FindReplaceID))
 						{
-							TSharedRef<SAnimAssetFindReplace> FindReplaceWidget = StaticCastSharedRef<SAnimAssetFindReplace>(Tab->GetContent());
-							FindReplaceWidget->SetFindReplaceType(EAnimAssetFindReplaceType::Curves);
+							TSharedRef<IAnimAssetFindReplace> FindReplaceWidget = StaticCastSharedRef<IAnimAssetFindReplace>(Tab->GetContent());
+							FindReplaceWidget->SetCurrentProcessor(UAnimAssetFindReplaceCurves::StaticClass());
 						}
 					}
 				}
diff --git a/Engine/Source/Editor/Persona/Private/SAnimNotifyPanel.cpp b/Engine/Source/Editor/Persona/Private/SAnimNotifyPanel.cpp
--- a/Engine/Source/Editor/Persona/Private/SAnimNotifyPanel.cpp
+++ b/Engine/Source/Editor/Persona/Private/SAnimNotifyPanel.cpp
@@ -2353,6 +2353,7 @@ void SAnimNotifyTrack::FillNewNotifyMenu(FMenuBuilder& MenuBuilder, bool bIsRepl
 						[
 							SNew(SSkeletonAnimNotifies, EditableSkeleton)
 							.IsPicker(true)
+							.ShowNotifies(true)
 							.OnItemSelected_Lambda([this, bIsReplaceWithMenu](const FName& InNotifyName)
 							{
 								FSlateApplication::Get().DismissAllMenus();
@@ -2411,7 +2412,8 @@ void SAnimNotifyTrack::FillNewSyncMarkerMenu(FMenuBuilder& MenuBuilder)
 					.MaxDesiredHeight(400.0f)
 					[
 						SNew(SSkeletonAnimNotifies, EditableSkeleton)
-						.IsSyncMarker(true)
+						.ShowSyncMarkers(true)
+						.ShowNotifies(false)
 						.OnItemSelected_Lambda([this](const FName& InNotifyName)
 						{
 							FSlateApplication::Get().DismissAllMenus();
diff --git a/Engine/Source/Editor/Persona/Private/SAnimationSequenceBrowser.cpp b/Engine/Source/Editor/Persona/Private/SAnimationSequenceBrowser.cpp
--- a/Engine/Source/Editor/Persona/Private/SAnimationSequenceBrowser.cpp
+++ b/Engine/Source/Editor/Persona/Private/SAnimationSequenceBrowser.cpp
@@ -140,6 +140,7 @@ public:
 			[
 				SNew(SSkeletonAnimNotifies, EditableSkeleton.ToSharedRef())
 				.IsPicker(true)
+				.ShowNotifies(true)
 				.OnItemSelected_Lambda([this](const FName& InNotifyName)
 				{
 					FSlateApplication::Get().DismissAllMenus();
diff --git a/Engine/Source/Editor/Persona/Private/SSkeletonAnimNotifies.cpp b/Engine/Source/Editor/Persona/Private/SSkeletonAnimNotifies.cpp
--- a/Engine/Source/Editor/Persona/Private/SSkeletonAnimNotifies.cpp
+++ b/Engine/Source/Editor/Persona/Private/SSkeletonAnimNotifies.cpp
@@ -2,6 +2,9 @@
 
 
 #include "SSkeletonAnimNotifies.h"
+
+#include "AnimAssetFindReplaceNotifies.h"
+#include "AnimAssetFindReplaceSyncMarkers.h"
 #include "Framework/Notifications/NotificationManager.h"
 #include "Widgets/Notifications/SNotificationList.h"
 #include "Framework/MultiBox/MultiBoxBuilder.h"
@@ -20,14 +23,53 @@
 #include "IAnimationBlueprintEditor.h"
 #include "IAnimationSequenceBrowser.h"
 #include "ISkeletonEditor.h"
+#include "SAnimAssetFindReplace.h"
+#include "Filters/GenericFilter.h"
+#include "Filters/SBasicFilterBar.h"
+#include "Widgets/Docking/SDockTab.h"
 
 #define LOCTEXT_NAMESPACE "SkeletonAnimNotifies"
 
-//////////////////////////////////////////////////////////////////////////
-// SMorphTargetListRow
-
 typedef TSharedPtr< FDisplayedAnimNotifyInfo > FDisplayedAnimNotifyInfoPtr;
 
+class FSkeletonAnimNotifiesFilter : public FGenericFilter<EAnimNotifyFilterFlags>
+{
+public:
+	FSkeletonAnimNotifiesFilter(EAnimNotifyFilterFlags InFlags, const FString& InName, const FText& InDisplayName, const FText& InToolTipText, FLinearColor InColor, TSharedPtr<FFilterCategory> InCategory)
+		: FGenericFilter<EAnimNotifyFilterFlags>(InCategory, InName, InDisplayName, FGenericFilter<EAnimNotifyFilterFlags>::FOnItemFiltered())
+		, Flags(InFlags)
+	{
+		ToolTip = InToolTipText;
+		Color = InColor;
+	}
+
+	bool IsActive() const
+	{
+		return bIsActive;
+	}
+
+	EAnimNotifyFilterFlags GetFlags() const
+	{
+		return Flags;
+	}
+
+private:
+	// FFilterBase interface
+	virtual void ActiveStateChanged(bool bActive) override
+	{
+		bIsActive = bActive;
+	}
+
+	virtual bool PassesFilter(EAnimNotifyFilterFlags InItem) const override
+	{
+		return EnumHasAnyFlags(InItem, Flags);
+	}
+	
+private:
+	EAnimNotifyFilterFlags Flags;
+	bool bIsActive = false;
+};
+
 /////////////////////////////////////////////////////
 // FSkeletonAnimNotifiesSummoner
 
@@ -49,6 +91,8 @@ FSkeletonAnimNotifiesSummoner::FSkeletonAnimNotifiesSummoner(TSharedPtr<class FA
 TSharedRef<SWidget> FSkeletonAnimNotifiesSummoner::CreateTabBody(const FWorkflowTabSpawnInfo& Info) const
 {
 	return SNew(SSkeletonAnimNotifies, EditableSkeleton.Pin().ToSharedRef(), HostingApp.Pin())
+		.ShowNotifies(true)
+		.ShowSyncMarkers(true)
 		.OnObjectsSelected(OnObjectsSelected);
 }
 
@@ -60,52 +104,123 @@ void SSkeletonAnimNotifies::Construct(const FArguments& InArgs, const TSharedRef
 	OnObjectsSelected = InArgs._OnObjectsSelected;
 	OnItemSelected = InArgs._OnItemSelected;
 	bIsPicker = InArgs._IsPicker;
-	bIsSyncMarker = InArgs._IsSyncMarker;
+	bShowSyncMarkers = InArgs._ShowSyncMarkers;
+	bShowNotifies = InArgs._ShowNotifies;
 
 	EditableSkeleton = InEditableSkeleton;
 	WeakHostingApp = InHostingApp;
 
-	if (bIsSyncMarker)
+	EditableSkeleton->RegisterOnNotifiesChanged(FSimpleDelegate::CreateSP(this, &SSkeletonAnimNotifies::OnNotifiesChanged));
+
+	if(GEditor)
 	{
-		bIsPicker = false; // Sync Markers are never pickers
+		GEditor->RegisterForUndo(this);
 	}
-	else
+
+	FOnContextMenuOpening OnContextMenuOpening = !bIsPicker ? FOnContextMenuOpening::CreateSP(this, &SSkeletonAnimNotifies::OnGetContextMenuContent) : FOnContextMenuOpening();
+
+	NameFilterBox = SNew( SSearchBox )
+		.SelectAllTextWhenFocused( true )
+		.OnTextChanged( this, &SSkeletonAnimNotifies::OnFilterTextChanged )
+		.OnTextCommitted( this, &SSkeletonAnimNotifies::OnFilterTextCommitted );
+
+	NotifiesListView = SNew( SAnimNotifyListType )
+		.ListItemsSource( &NotifyList )
+		.OnGenerateRow( this, &SSkeletonAnimNotifies::GenerateNotifyRow )
+		.OnContextMenuOpening( OnContextMenuOpening )
+		.OnSelectionChanged( this, &SSkeletonAnimNotifies::OnNotifySelectionChanged )
+		.ItemHeight( 18.0f )
+		.OnItemScrolledIntoView( this, &SSkeletonAnimNotifies::OnItemScrolledIntoView );
+	
+	TSharedPtr<FFilterCategory> FilterCategory = MakeShared<FFilterCategory>(LOCTEXT("AnimNotifyFiltersLabel", "Anim Notify Filters"), LOCTEXT("AnimNotifyFiltersLabelToolTip", "Filter what kind fo notifies and sync markers can be displayed."));
+
+	// Hide filter UI if we are only displaying one type of thing
+	const EVisibility FilterVisibility = (bShowNotifies ^ bShowSyncMarkers) ? EVisibility::Collapsed : EVisibility::Visible;
+
+	TSharedRef<SBasicFilterBar<EAnimNotifyFilterFlags>> FilterBar = SNew(SBasicFilterBar<EAnimNotifyFilterFlags>)
+	.Visibility(FilterVisibility)
+	.CustomFilters(Filters)
+	.UseSectionsForCategories(true)
+	.OnFilterChanged_Lambda([this]()
 	{
-		EditableSkeleton->RegisterOnNotifiesChanged(FSimpleDelegate::CreateSP(this, &SSkeletonAnimNotifies::OnNotifiesChanged));
-	}
+		CurrentFilterFlags = EAnimNotifyFilterFlags::None;
 
-	if(GEditor)
+		for(const TSharedRef<FFilterBase<EAnimNotifyFilterFlags>>& Filter : Filters)
+		{
+			TSharedRef<FSkeletonAnimNotifiesFilter> AnimNotifiesFilter = StaticCastSharedRef<FSkeletonAnimNotifiesFilter>(Filter);
+			if(AnimNotifiesFilter->IsActive())
+			{
+				CurrentFilterFlags |= AnimNotifiesFilter->GetFlags();
+			}
+		}
+
+		RefreshNotifiesListWithFilter();
+	});
+	
+	if(bShowNotifies)
 	{
-		GEditor->RegisterForUndo(this);
+		TSharedRef<FFilterBase<EAnimNotifyFilterFlags>> Filter = Filters.Add_GetRef(MakeShared<FSkeletonAnimNotifiesFilter>(
+			EAnimNotifyFilterFlags::Notifies,
+			"Notifies",
+			LOCTEXT("ShowNotifiesLabel", "Notifies"),
+			LOCTEXT("ShowNotifiesTooltip", "Show notifies"),
+			FLinearColor::Red.Desaturate(0.5f),
+			FilterCategory
+			));
+
+		FilterBar->SetFilterCheckState(Filter, ECheckBoxState::Checked);
 	}
 
-	FOnContextMenuOpening OnContextMenuOpening = (!bIsPicker && !bIsSyncMarker) ? FOnContextMenuOpening::CreateSP(this, &SSkeletonAnimNotifies::OnGetContextMenuContent) : FOnContextMenuOpening();
+	if(bShowSyncMarkers)
+	{
+		TSharedRef<FFilterBase<EAnimNotifyFilterFlags>> Filter = Filters.Add_GetRef(MakeShared<FSkeletonAnimNotifiesFilter>(
+			EAnimNotifyFilterFlags::SyncMarkers,
+			"SyncMarkers",
+			LOCTEXT("ShowSyncMarkersLabel", "Sync Markers"),
+			LOCTEXT("ShowSyncMarkersTooltip", "Show sync markers"),
+			FLinearColor::Green.Desaturate(0.5f),
+			FilterCategory
+			));
+
+		FilterBar->SetFilterCheckState(Filter, ECheckBoxState::Checked);
+	}
 
-	this->ChildSlot
+	TSharedRef<SWidget> AddFilterButton = SBasicFilterBar<EAnimNotifyFilterFlags>::MakeAddFilterButton(FilterBar);
+	AddFilterButton->SetVisibility(FilterVisibility);
+	
+	ChildSlot
 	[
 		SNew( SVerticalBox )
 
-		+ SVerticalBox::Slot()
+		+SVerticalBox::Slot()
 		.AutoHeight()
 		.Padding( FMargin( 0.0f, 0.0f, 0.0f, 4.0f ) )
 		[
-			SAssignNew( NameFilterBox, SSearchBox )
-			.SelectAllTextWhenFocused( true )
-			.OnTextChanged( this, &SSkeletonAnimNotifies::OnFilterTextChanged )
-			.OnTextCommitted( this, &SSkeletonAnimNotifies::OnFilterTextCommitted )
-			.HintText( LOCTEXT( "NotifiesSearchBoxHint", "Search Animation Notifies...") )
+			SNew(SHorizontalBox)
+			+SHorizontalBox::Slot()
+			.AutoWidth()
+			.Padding(2.0f,0.0f)
+			[
+				AddFilterButton
+			]
+
+			+SHorizontalBox::Slot()
+			.FillWidth(1.0f)
+			[
+				NameFilterBox.ToSharedRef()
+			]
 		]
 
-		+ SVerticalBox::Slot()
+		+SVerticalBox::Slot()
+		.AutoHeight()
+		[
+			FilterBar
+		]
+
+		+SVerticalBox::Slot()
 		.FillHeight( 1.0f )		// This is required to make the scrollbar work, as content overflows Slate containers by default
 		[
-			SAssignNew( NotifiesListView, SAnimNotifyListType )
-			.ListItemsSource( &NotifyList )
-			.OnGenerateRow( this, &SSkeletonAnimNotifies::GenerateNotifyRow )
-			.OnContextMenuOpening( OnContextMenuOpening )
-			.OnSelectionChanged( this, &SSkeletonAnimNotifies::OnNotifySelectionChanged )
-			.ItemHeight( 22.0f )
-			.OnItemScrolledIntoView( this, &SSkeletonAnimNotifies::OnItemScrolledIntoView )
+			NotifiesListView.ToSharedRef()
 		]
 	];
 
@@ -153,6 +268,7 @@ TSharedRef<ITableRow> SSkeletonAnimNotifies::GenerateNotifyRow(TSharedPtr<FDispl
 			[
 				SAssignNew(InInfo->InlineEditableText, SInlineEditableTextBlock)
 				.Text(FText::FromName(InInfo->Name))
+				.Font(FAppStyle::Get().GetFontStyle("SmallFont"))
 				.OnVerifyTextChanged(this, &SSkeletonAnimNotifies::OnVerifyNotifyNameCommit, InInfo)
 				.OnTextCommitted(this, &SSkeletonAnimNotifies::OnNotifyNameCommitted, InInfo)
 				.IsSelected(this, &SSkeletonAnimNotifies::IsSelected)
@@ -167,57 +283,50 @@ TSharedPtr<SWidget> SSkeletonAnimNotifies::OnGetContextMenuContent() const
 	const bool bShouldCloseWindowAfterMenuSelection = true;
 	FMenuBuilder MenuBuilder( bShouldCloseWindowAfterMenuSelection, NULL);
 
-	if (bIsSyncMarker)
+	MenuBuilder.BeginSection("AnimItemAction", LOCTEXT("ItemActions", "New Item"));
 	{
-		MenuBuilder.BeginSection("AnimNotifyAction", LOCTEXT("SelectedSyncMarkerActions", "Selected Sync Marker Actions"));
-		{
-			FUIAction Action = FUIAction(FExecuteAction::CreateSP(const_cast<SSkeletonAnimNotifies*>(this), &SSkeletonAnimNotifies::OnDeleteSyncMarker));
-			const FText Label = LOCTEXT("DeleteSyncMarkerButtonLabel", "Delete");
-			const FText ToolTipText = LOCTEXT("DeleteSyncMarkerButtonTooltip", "Deletes the sync marker from the suggestions");
-			MenuBuilder.AddMenuEntry(Label, ToolTipText, FSlateIcon(), Action);
-		}
-		MenuBuilder.EndSection();
+		FUIAction Action = FUIAction(FExecuteAction::CreateSP(const_cast<SSkeletonAnimNotifies*>(this), &SSkeletonAnimNotifies::OnAddItem, false));
+		const FText Label = LOCTEXT("NewAnimNotifyButtonLabel", "New Notify...");
+		const FText ToolTipText = LOCTEXT("NewAnimNotifyButtonTooltip", "Creates a new anim notify.");
+		MenuBuilder.AddMenuEntry(Label, ToolTipText, FSlateIcon(), Action);
 	}
-	else
+
+	{
+		FUIAction Action = FUIAction(FExecuteAction::CreateSP(const_cast<SSkeletonAnimNotifies*>(this), &SSkeletonAnimNotifies::OnAddItem, true));
+		const FText Label = LOCTEXT("NewSyncMarkerButtonLabel", "New Sync Marker...");
+		const FText ToolTipText = LOCTEXT("NewSyncMarkerButtonTooltip", "Creates a new sync marker.");
+		MenuBuilder.AddMenuEntry(Label, ToolTipText, FSlateIcon(), Action);
+	}
+	MenuBuilder.EndSection();
+
+	MenuBuilder.BeginSection("AnimNotifyAction", LOCTEXT("SelectedItemActions", "Selected Item Actions"));
 	{
-		MenuBuilder.BeginSection("AnimNotifyAction", LOCTEXT("AnimNotifyActions", "Notifies"));
 		{
-			FUIAction Action = FUIAction(FExecuteAction::CreateSP(const_cast<SSkeletonAnimNotifies*>(this), &SSkeletonAnimNotifies::OnAddAnimNotify));
-			const FText Label = LOCTEXT("NewAnimNotifyButtonLabel", "New...");
-			const FText ToolTipText = LOCTEXT("NewAnimNotifyButtonTooltip", "Creates a new anim notify.");
+			FUIAction Action = FUIAction(FExecuteAction::CreateSP(const_cast<SSkeletonAnimNotifies*>(this), &SSkeletonAnimNotifies::OnRenameItem),
+				FCanExecuteAction::CreateSP(this, &SSkeletonAnimNotifies::CanPerformRename));
+			const FText Label = LOCTEXT("RenameAnimNotifyButtonLabel", "Rename");
+			const FText ToolTipText = LOCTEXT("RenameAnimNotifyButtonTooltip", "Renames the selected item.");
 			MenuBuilder.AddMenuEntry(Label, ToolTipText, FSlateIcon(), Action);
 		}
-		MenuBuilder.EndSection();
 
-		MenuBuilder.BeginSection("AnimNotifyAction", LOCTEXT("SelectedAnimNotifyActions", "Selected Notify Actions"));
 		{
-			{
-				FUIAction Action = FUIAction(FExecuteAction::CreateSP(const_cast<SSkeletonAnimNotifies*>(this), &SSkeletonAnimNotifies::OnRenameAnimNotify),
-					FCanExecuteAction::CreateSP(this, &SSkeletonAnimNotifies::CanPerformRename));
-				const FText Label = LOCTEXT("RenameAnimNotifyButtonLabel", "Rename");
-				const FText ToolTipText = LOCTEXT("RenameAnimNotifyButtonTooltip", "Renames the selected anim notifies.");
-				MenuBuilder.AddMenuEntry(Label, ToolTipText, FSlateIcon(), Action);
-			}
-
-			{
-				FUIAction Action = FUIAction(FExecuteAction::CreateSP(const_cast<SSkeletonAnimNotifies*>(this), &SSkeletonAnimNotifies::OnDeleteAnimNotify),
-					FCanExecuteAction::CreateSP(this, &SSkeletonAnimNotifies::CanPerformDelete));
-				const FText Label = LOCTEXT("DeleteAnimNotifyButtonLabel", "Delete");
-				const FText ToolTipText = LOCTEXT("DeleteAnimNotifyButtonTooltip", "Deletes the selected anim notifies.");
-				MenuBuilder.AddMenuEntry(Label, ToolTipText, FSlateIcon(), Action);
-			}
+			FUIAction Action = FUIAction(FExecuteAction::CreateSP(const_cast<SSkeletonAnimNotifies*>(this), &SSkeletonAnimNotifies::OnDeleteItems),
+				FCanExecuteAction::CreateSP(this, &SSkeletonAnimNotifies::CanPerformDelete));
+			const FText Label = LOCTEXT("DeleteAnimNotifyButtonLabel", "Delete");
+			const FText ToolTipText = LOCTEXT("DeleteAnimNotifyButtonTooltip", "Deletes the selected items.");
+			MenuBuilder.AddMenuEntry(Label, ToolTipText, FSlateIcon(), Action);
+		}
 
-			if(WeakHostingApp.IsValid() && NotifiesListView->GetNumItemsSelected() == 1)
-			{
-				FUIAction Action = FUIAction(FExecuteAction::CreateSP(const_cast<SSkeletonAnimNotifies*>(this), &SSkeletonAnimNotifies::OnFindReferences),
-					FCanExecuteAction::CreateSP(this, &SSkeletonAnimNotifies::CanPerformDelete));
-				const FText Label = LOCTEXT("FindNotifyReferences", "Find References");
-				const FText ToolTipText = LOCTEXT("FindNotifyReferencesTooltip", "Find all references to this skeleton notify in the asset browser");
-				MenuBuilder.AddMenuEntry(Label, ToolTipText, FSlateIcon(), Action);
-			}
+		if(WeakHostingApp.IsValid() && NotifiesListView->GetNumItemsSelected() == 1)
+		{
+			FUIAction Action = FUIAction(FExecuteAction::CreateSP(const_cast<SSkeletonAnimNotifies*>(this), &SSkeletonAnimNotifies::OnFindReferences),
+				FCanExecuteAction::CreateSP(this, &SSkeletonAnimNotifies::CanPerformDelete));
+			const FText Label = LOCTEXT("FindNotifyReferences", "Find References");
+			const FText ToolTipText = LOCTEXT("FindNotifyReferencesTooltip", "Find all references to this item in the asset browser");
+			MenuBuilder.AddMenuEntry(Label, ToolTipText, FSlateIcon(), Action);
 		}
-		MenuBuilder.EndSection();
 	}
+	MenuBuilder.EndSection();
 
 	return MenuBuilder.MakeWidget();
 }
@@ -226,7 +335,7 @@ void SSkeletonAnimNotifies::OnNotifySelectionChanged(TSharedPtr<FDisplayedAnimNo
 {
 	if(Selection.IsValid())
 	{
-		if (!bIsSyncMarker)
+		if (!Selection->bIsSyncMarker)
 		{
 			ShowNotifyInDetailsView(Selection->Name);
 		}
@@ -245,10 +354,10 @@ bool SSkeletonAnimNotifies::CanPerformRename() const
 	return NotifiesListView->GetNumItemsSelected() == 1;
 }
 
-void SSkeletonAnimNotifies::OnAddAnimNotify()
+void SSkeletonAnimNotifies::OnAddItem(bool bIsSyncMarker)
 {
 	// Find a unique name for this notify
-	const TCHAR* BaseNotifyString = TEXT("NewNotify");
+	const TCHAR* BaseNotifyString = bIsSyncMarker ? TEXT("NewNotify") : TEXT("NewSyncMarker");
 	FString NewNotifyString = BaseNotifyString;
 	int32 NumericExtension = 0;
 
@@ -259,7 +368,7 @@ void SSkeletonAnimNotifies::OnAddAnimNotify()
 	}
 
 	// Add an item. The subsequent rename will commit the item.
-	TSharedPtr<FDisplayedAnimNotifyInfo> NewItem = FDisplayedAnimNotifyInfo::Make(*NewNotifyString);
+	TSharedPtr<FDisplayedAnimNotifyInfo> NewItem = FDisplayedAnimNotifyInfo::Make(*NewNotifyString, bIsSyncMarker);
 	NewItem->bIsNew = true;
 	NotifyList.Add(NewItem);
 
@@ -276,53 +385,53 @@ void SSkeletonAnimNotifies::OnItemScrolledIntoView(TSharedPtr<FDisplayedAnimNoti
 	}
 }
 
-void SSkeletonAnimNotifies::OnDeleteAnimNotify()
+void SSkeletonAnimNotifies::OnDeleteItems()
 {
 	TArray< TSharedPtr< FDisplayedAnimNotifyInfo > > SelectedRows = NotifiesListView->GetSelectedItems();
 
-	// this one deletes all notifies with same name. 
+	TArray<FName> SelectedSyncMarkerNames;
 	TArray<FName> SelectedNotifyNames;
 
-	for(int Selection = 0; Selection < SelectedRows.Num(); ++Selection)
+	for (TSharedPtr< FDisplayedAnimNotifyInfo > Selection : SelectedRows)
 	{
-		SelectedNotifyNames.Add(SelectedRows[Selection]->Name);
+		if(Selection->bIsSyncMarker)
+		{
+			SelectedSyncMarkerNames.Add(Selection->Name);
+		}
+		else
+		{
+			SelectedNotifyNames.Add(Selection->Name);
+		}
 	}
 
-	int32 NumAnimationsModified = EditableSkeleton->DeleteAnimNotifies(SelectedNotifyNames);
+	int32 NumAnimationsModified = 0;
+	
+	if(SelectedSyncMarkerNames.Num())
+	{
+		NumAnimationsModified += EditableSkeleton->DeleteSyncMarkers(SelectedSyncMarkerNames);
+	}
+
+	if(SelectedNotifyNames.Num())
+	{
+		NumAnimationsModified += EditableSkeleton->DeleteAnimNotifies(SelectedNotifyNames);
+	}
 
 	if(NumAnimationsModified > 0)
 	{
-		// Tell the user that the socket is a duplicate
 		FFormatNamedArguments Args;
 		Args.Add( TEXT("NumAnimationsModified"), NumAnimationsModified );
-		FNotificationInfo Info( FText::Format( LOCTEXT( "AnimNotifiesDeleted", "{NumAnimationsModified} animation(s) modified to delete notifications" ), Args ) );
+		FNotificationInfo Info( FText::Format( LOCTEXT( "ItemsDeleted", "{NumAnimationsModified} animation(s) modified to delete items" ), Args ) );
 
 		Info.bUseLargeFont = false;
 		Info.ExpireDuration = 5.0f;
 
 		NotifyUser( Info );
 	}
-
-	CreateNotifiesList( NameFilterBox->GetText().ToString() );
-}
-
-void SSkeletonAnimNotifies::OnDeleteSyncMarker()
-{
-	TArray< TSharedPtr< FDisplayedAnimNotifyInfo > > SelectedRows = NotifiesListView->GetSelectedItems();
-
-	TArray<FName> SelectedSyncMarkerNames;
-
-	for (int Selection = 0; Selection < SelectedRows.Num(); ++Selection)
-	{
-		SelectedSyncMarkerNames.Add(SelectedRows[Selection]->Name);
-	}
-
-	EditableSkeleton->DeleteSyncMarkers(SelectedSyncMarkerNames);
-
+	
 	CreateNotifiesList(NameFilterBox->GetText().ToString());
 }
 
-void SSkeletonAnimNotifies::OnRenameAnimNotify()
+void SSkeletonAnimNotifies::OnRenameItem()
 {
 	TArray< TSharedPtr< FDisplayedAnimNotifyInfo > > SelectedRows = NotifiesListView->GetSelectedItems();
 
@@ -344,10 +453,21 @@ bool SSkeletonAnimNotifies::OnVerifyNotifyNameCommit( const FText& NewName, FTex
 	FName NotifyName( *NewName.ToString() );
 	if(NotifyName != Item->Name || Item->bIsNew)
 	{
-		if(EditableSkeleton->GetSkeleton().AnimationNotifies.Contains(NotifyName))
+		if(Item->bIsSyncMarker)
 		{
-			OutErrorMessage = FText::Format( LOCTEXT("AlreadyInUseMessage", "'{0}' is already in use."), NewName );
-			bValid = false;
+			if(EditableSkeleton->GetSkeleton().GetExistingMarkerNames().Contains(NotifyName))
+			{
+				OutErrorMessage = FText::Format( LOCTEXT("AlreadyInUseMessage", "'{0}' is already in use."), NewName );
+				bValid = false;
+			}
+		}
+		else
+		{
+			if(EditableSkeleton->GetSkeleton().AnimationNotifies.Contains(NotifyName))
+			{
+				OutErrorMessage = FText::Format( LOCTEXT("AlreadyInUseMessage", "'{0}' is already in use."), NewName );
+				bValid = false;
+			}
 		}
 	}
 
@@ -359,28 +479,51 @@ void SSkeletonAnimNotifies::OnNotifyNameCommitted( const FText& NewName, ETextCo
 	FName NewFName = FName(*NewName.ToString());
 	if(Item->bIsNew)
 	{
-		EditableSkeleton->AddNotify(NewFName);
+		if(Item->bIsSyncMarker)
+		{
+			EditableSkeleton->AddNotify(NewFName);
+		}
+		else
+		{
+			EditableSkeleton->AddSyncMarker(NewFName);
+		}
 		Item->bIsNew = false;
 	}
 	else
 	{
 		if(NewFName != Item->Name)
 		{
-			int32 NumAnimationsModified = EditableSkeleton->RenameNotify(FName(*NewName.ToString()), Item->Name);
-
-			if(NumAnimationsModified > 0)
+			if(Item->bIsSyncMarker)
 			{
-				// Tell the user that the socket is a duplicate
-				FFormatNamedArguments Args;
-				Args.Add( TEXT("NumAnimationsModified"), NumAnimationsModified );
-				FNotificationInfo Info( FText::Format( LOCTEXT( "AnimNotifiesRenamed", "{NumAnimationsModified} animation(s) modified to rename notification" ), Args ) );
-
-				Info.bUseLargeFont = false;
-				Info.ExpireDuration = 5.0f;
-
-				NotifyUser( Info );
+				int32 NumAnimationsModified = EditableSkeleton->RenameSyncMarker(FName(*NewName.ToString()), Item->Name);
+				if(NumAnimationsModified > 0)
+				{
+					FFormatNamedArguments Args;
+					Args.Add( TEXT("NumAnimationsModified"), NumAnimationsModified );
+					FNotificationInfo Info( FText::Format( LOCTEXT( "SyncMarkersRenamed", "{NumAnimationsModified} animation(s) modified to rename sync marker" ), Args ) );
+
+					Info.bUseLargeFont = false;
+					Info.ExpireDuration = 5.0f;
+
+					NotifyUser( Info );
+				}
 			}
-
+			else
+			{
+				int32 NumAnimationsModified = EditableSkeleton->RenameNotify(FName(*NewName.ToString()), Item->Name);
+				if(NumAnimationsModified > 0)
+				{
+					FFormatNamedArguments Args;
+					Args.Add( TEXT("NumAnimationsModified"), NumAnimationsModified );
+					FNotificationInfo Info( FText::Format( LOCTEXT( "AnimNotifiesRenamed", "{NumAnimationsModified} animation(s) modified to rename notification" ), Args ) );
+
+					Info.bUseLargeFont = false;
+					Info.ExpireDuration = 5.0f;
+
+					NotifyUser( Info );
+				}
+			}
+			
 			RefreshNotifiesListWithFilter();
 		}
 	}
@@ -397,20 +540,34 @@ void SSkeletonAnimNotifies::CreateNotifiesList( const FString& SearchText )
 
 	const USkeleton& TargetSkeleton = EditableSkeleton->GetSkeleton();
 
-	const TArray<FName>& ItemNames = bIsSyncMarker ? TargetSkeleton.GetExistingMarkerNames() : TargetSkeleton.AnimationNotifies;
-
-	for(const FName& ItemName : ItemNames)
+	auto AddItem = [this, SearchText](FName InItemName, bool bInIsSyncMarker)
 	{
 		if ( !SearchText.IsEmpty() )
 		{
-			if (ItemName.ToString().Contains( SearchText ) )
+			if (InItemName.ToString().Contains( SearchText ) )
 			{
-				NotifyList.Add( FDisplayedAnimNotifyInfo::Make(ItemName) );
+				NotifyList.Add( FDisplayedAnimNotifyInfo::Make(InItemName, bInIsSyncMarker) );
 			}
 		}
 		else
 		{
-			NotifyList.Add( FDisplayedAnimNotifyInfo::Make(ItemName) );
+			NotifyList.Add( FDisplayedAnimNotifyInfo::Make(InItemName, bInIsSyncMarker) );
+		}
+	};
+
+	if(EnumHasAnyFlags(CurrentFilterFlags, EAnimNotifyFilterFlags::Notifies))
+	{
+		for(const FName& ItemName : TargetSkeleton.AnimationNotifies)
+		{
+			AddItem(ItemName, false);
+		}
+	}
+
+	if(EnumHasAnyFlags(CurrentFilterFlags, EAnimNotifyFilterFlags::SyncMarkers))
+	{
+		for(const FName& ItemName : TargetSkeleton.GetExistingMarkerNames())
+		{
+			AddItem(ItemName, true);
 		}
 	}
 
@@ -485,22 +642,12 @@ void SSkeletonAnimNotifies::OnFindReferences()
 		NotifiesListView->GetSelectedItems(SelectedItems);
 		FName Name = SelectedItems[0]->Name;
 
-		HostingApp->InvokeTab(FPersonaTabs::AssetBrowserID);
-
-		if (HostingApp->GetEditorName() == TEXT("AnimationEditor"))
-		{
-			TSharedPtr<IAnimationEditor> AnimationEditor = StaticCastSharedPtr<IAnimationEditor>(HostingApp);
-			AnimationEditor->GetAssetBrowser()->FilterBySkeletonNotify(Name);
-		}
-		else if (HostingApp->GetEditorName() == TEXT("AnimationBlueprintEditor"))
-		{
-			TSharedPtr<IAnimationBlueprintEditor> AnimationBlueprintEditor = StaticCastSharedPtr<IAnimationBlueprintEditor>(HostingApp);
-			AnimationBlueprintEditor->GetAssetBrowser()->FilterBySkeletonNotify(Name);
-		}
-		else if (HostingApp->GetEditorName() == TEXT("SkeletonEditor"))
+		if(TSharedPtr<SDockTab> Tab = HostingApp->GetTabManager()->TryInvokeTab(FPersonaTabs::FindReplaceID))
 		{
-			TSharedPtr<ISkeletonEditor> SkeletonEditor = StaticCastSharedPtr<ISkeletonEditor>(HostingApp);
-			SkeletonEditor->GetAssetBrowser()->FilterBySkeletonNotify(Name);
+			TSharedRef<IAnimAssetFindReplace> FindReplaceWidget = StaticCastSharedRef<IAnimAssetFindReplace>(Tab->GetContent());
+			FindReplaceWidget->SetCurrentProcessor(SelectedItems[0]->bIsSyncMarker ? UAnimAssetFindReplaceSyncMarkers::StaticClass() : UAnimAssetFindReplaceNotifies::StaticClass());
+			UAnimAssetFindReplaceProcessor_StringBase* Processor = Cast<UAnimAssetFindReplaceProcessor_StringBase>(FindReplaceWidget->GetCurrentProcessor());
+			Processor->SetFindString(Name.ToString());
 		}
 	}
 }
diff --git a/Engine/Source/Editor/Persona/Private/SSkeletonAnimNotifies.h b/Engine/Source/Editor/Persona/Private/SSkeletonAnimNotifies.h
--- a/Engine/Source/Editor/Persona/Private/SSkeletonAnimNotifies.h
+++ b/Engine/Source/Editor/Persona/Private/SSkeletonAnimNotifies.h
@@ -16,6 +16,7 @@
 #include "Widgets/Views/SListView.h"
 #include "EditorObjectsTracker.h"
 #include "EditorUndoClient.h"
+#include "Filters/FilterBase.h"
 
 #define LOCTEXT_NAMESPACE "SkeletonAnimnotifies"
 
@@ -60,23 +61,40 @@ public:
 	/** Flag to say whether this is a new item we are creating */
 	bool bIsNew;
 
+	/** Identifies sync markers vs notifies */
+	bool bIsSyncMarker;
+
 	/** Static function for creating a new item, but ensures that you can only have a TSharedRef to one */
-	static TSharedRef<FDisplayedAnimNotifyInfo> Make(const FName& NotifyName)
+	static TSharedRef<FDisplayedAnimNotifyInfo> Make(const FName& NotifyName, bool bInIsSyncMarker)
 	{
-		return MakeShareable(new FDisplayedAnimNotifyInfo(NotifyName));
+		return MakeShareable(new FDisplayedAnimNotifyInfo(NotifyName, bInIsSyncMarker));
 	}
 
 protected:
 	/** Hidden constructor, always use Make above */
-	FDisplayedAnimNotifyInfo(const FName& InNotifyName)
+	FDisplayedAnimNotifyInfo(const FName& InNotifyName, bool bInIsSyncMarker)
 		: Name( InNotifyName )
 		, bIsNew(false)
+		, bIsSyncMarker(bInIsSyncMarker)
 	{}
 
 	/** Hidden constructor, always use Make above */
 	FDisplayedAnimNotifyInfo() {}
 };
 
+// This is a flag that is used to filter UI part
+enum class EAnimNotifyFilterFlags : uint8 
+{
+	// Show none
+	None			= 0, 
+	// Show notifies
+	Notifies		= 0x01, 
+	// Show sync markers
+	SyncMarkers		= 0x02, 
+};
+
+ENUM_CLASS_FLAGS(EAnimNotifyFilterFlags);
+
 /** Widgets list type */
 typedef SListView< TSharedPtr<FDisplayedAnimNotifyInfo> > SAnimNotifyListType;
 
@@ -85,7 +103,8 @@ class SSkeletonAnimNotifies : public SCompoundWidget, public FGCObject, public F
 public:
 	SLATE_BEGIN_ARGS( SSkeletonAnimNotifies )
 		: _IsPicker(false)
-		, _IsSyncMarker(false)
+		, _ShowSyncMarkers(false)
+		, _ShowNotifies(true)
 	{}
 
 	/** Delegate called to select an object in the details panel */
@@ -97,8 +116,14 @@ public:
 	/** Whether we should use this dialog as a picker or an editor. In picker mode we cant add, remove or rename notifies. */
 	SLATE_ARGUMENT(bool, IsPicker)
 
-	/** Whether we should use this dialog for sync markers instead of notifies */
-	SLATE_ARGUMENT(bool, IsSyncMarker)
+	/** Whether we should show sync markers */
+	SLATE_ARGUMENT_DEPRECATED(bool, IsSyncMarker, 5.3, "Please use ShowSyncMarkers or ShowNotifies instead")
+
+	/** Whether we should show sync markers */
+	SLATE_ARGUMENT(bool, ShowSyncMarkers)
+
+	/** Whether we should show notifies */
+	SLATE_ARGUMENT(bool, ShowNotifies)
 
 	SLATE_END_ARGS()
 public:
@@ -157,21 +182,18 @@ private:
 
 	/** Delegate handler for determining whether we can show the delete menu options */
 	bool CanPerformDelete() const;
-
-	/** Delegate handler for deleting anim notifies */
-	void OnDeleteAnimNotify();
-
-	/** Delegate handler for deleting a sync marker */
-	void OnDeleteSyncMarker();
+	
+	/** Delegate handler for deleting items */
+	void OnDeleteItems();
 
 	/** Delegate handler for determining whether we can show the rename menu options */
 	bool CanPerformRename() const;
 
-	/** Delegate handler for adding anim notifies */
-	void OnAddAnimNotify();
+	/** Delegate handler for adding anim notifies & sync markers */
+	void OnAddItem(bool bIsSyncMarker);
 
-	/** Delegate handler for renaming anim notifies */
-	void OnRenameAnimNotify();
+	/** Delegate handler for renaming items */
+	void OnRenameItem();
 
 	/** Wrapper that populates NotifiesListView using current filter test */
 	void RefreshNotifiesListWithFilter();
@@ -222,10 +244,19 @@ private:
 	TWeakPtr<class FAssetEditorToolkit> WeakHostingApp;
 
 	/** Whether we should use this dialog as a picker or an editor. In picker mode we cant add, remove or rename notifies. */
-	bool bIsPicker;
+	bool bIsPicker = false;
 
-	/** Whether we are using this dialog for sync markers instead of notifies */
-	bool bIsSyncMarker;
+	/** Whether we should display markers */
+	bool bShowSyncMarkers = false;
+
+	/** Whether we should display notifies */
+	bool bShowNotifies = false;
+
+	/** All filters that can be applied to the widget's display */
+	TArray<TSharedRef<FFilterBase<EAnimNotifyFilterFlags>>> Filters;
+
+	/** Current filter flags */
+	EAnimNotifyFilterFlags CurrentFilterFlags = EAnimNotifyFilterFlags::None;
 };
 
 #undef LOCTEXT_NAMESPACE
diff --git a/Engine/Source/Editor/Persona/Public/AnimAssetFindReplace.h b/Engine/Source/Editor/Persona/Public/AnimAssetFindReplace.h
--- a/Engine/Source/Editor/Persona/Public/AnimAssetFindReplace.h
+++ b/Engine/Source/Editor/Persona/Public/AnimAssetFindReplace.h
@@ -4,9 +4,14 @@
 
 #include "CoreMinimal.h"
 #include "AssetRegistry/AssetData.h"
+#include "Widgets/SCompoundWidget.h"
+#include "Widgets/SNullWidget.h"
+#include "Templates/SubclassOf.h"
 #include "AnimAssetFindReplace.generated.h"
 
 class SAnimAssetFindReplace;
+struct FToolMenuSection;
+class SAutoCompleteSearchBox;
 
 /** Context for toolbar */
 UCLASS()
@@ -19,35 +24,225 @@ public:
 	TWeakPtr<SAnimAssetFindReplace> Widget;
 };
 
+/** The mode configuration that the dialog can take */
 enum class EAnimAssetFindReplaceMode : int32
 {
 	Find,
 	Replace,
 };
 
+/** Processor base class to allow systems to add their own find/replace functionality */
+UCLASS(MinimalAPI, Abstract)
+class UAnimAssetFindReplaceProcessor : public UObject
+{
+	GENERATED_BODY()
+
+	friend class SAnimAssetFindReplace;
+
+private:
+	/** Initially setup the processor, binding it to the specified widget */
+	void Initialize(TSharedRef<SAnimAssetFindReplace> InWidget);
+
+public:
+	/** Whether this processor supports the specified mode */
+	virtual bool SupportsMode(EAnimAssetFindReplaceMode InMode) const PURE_VIRTUAL(UAnimAssetFindReplaceProcessor::SupportsMode, return false;);
+
+	/** Builds results string to display from the supplied asset data */
+	virtual FString GetFindResultStringFromAssetData(const FAssetData& InAssetData) const PURE_VIRTUAL(UAnimAssetFindReplaceProcessor::GetFindResultStringFromAssetData, return FString(););
+
+	/** Get all the asset types that are supported by this processor */
+	virtual TConstArrayView<UClass*> GetSupportedAssetTypes() const PURE_VIRTUAL(UAnimAssetFindReplaceProcessor::GetSupportedAssetTypes, return TConstArrayView<UClass*>(););
+
+	/**
+	 * Whether the asset in question is a valid search result or not
+	 *
+	 * @param InAssetData		The asset in question
+	 * @param bOutIsOldAsset	Whether the asset is 'old', i.e. it needs to be fully loaded to determine whether it is
+	 *							filtered, as relevant data is not saved in the asset registry for it
+	 */
+	virtual bool ShouldFilterOutAsset(const FAssetData& InAssetData, bool& bOutIsOldAsset) const PURE_VIRTUAL(UAnimAssetFindReplaceProcessor::ShouldFilterOutAsset, return false;);
+
+	/** Perform a replace operation on the specified asset */
+	virtual void ReplaceInAsset(const FAssetData& InAssetData) const PURE_VIRTUAL(UAnimAssetFindReplaceProcessor::ReplaceInAsset, );
+
+	/** Perform a remove operation on the specified asset */
+	virtual void RemoveInAsset(const FAssetData& InAssetData) const PURE_VIRTUAL(UAnimAssetFindReplaceProcessor::RemoveInAsset, );
+	
+	/** Optional override to extend the toolbar when this processor is selected */
+	virtual void ExtendToolbar(FToolMenuSection& InSection) {}
+
+	/** Make or return the widget used to edit find/replace parameters */
+	virtual TSharedRef<SWidget> MakeFindReplaceWidget() PURE_VIRTUAL(UAnimAssetFindReplaceProcessor::MakeFindReplaceWidget, return SNullWidget::NullWidget;);
+
+	/** Optional override allowing a widget to be focused when the find/replace widget is initially displayed */
+	virtual void FocusInitialWidget() {}
+	
+	/** Allows control over when the replace option is available to users */
+	virtual bool CanCurrentlyReplace() const PURE_VIRTUAL(UAnimAssetFindReplaceProcessor::CanCurrentlyReplace, return false;);
+
+	/** Allows control over when the remove option is available to users */
+	virtual bool CanCurrentlyRemove() const PURE_VIRTUAL(UAnimAssetFindReplaceProcessor::CanCurrentlyRemove, return false;);
+
+	/** Called by the UI to deferred-refresh any cached data that this processor holds after asset operations have completed */
+	virtual void RefreshCachedData() {}
+
+protected:
+	/** Requests the entire find UI to be refreshed */
+	void RequestUIRefresh();
+
+	/** Requests search results to be refreshed */
+	void RequestRefreshSearchResults();
+
+private:
+	/** Weak ptr to the find/replace widget */
+	TWeakPtr<SAnimAssetFindReplace> Widget;
+};
+
+/** Processor for string-based animation asset operations */
+UCLASS(MinimalAPI, Abstract)
+class UAnimAssetFindReplaceProcessor_StringBase : public UAnimAssetFindReplaceProcessor
+{
+	GENERATED_BODY()
+
+public:
+	UAnimAssetFindReplaceProcessor_StringBase()
+		: AutoCompleteItems(MakeShared<TArray<TSharedPtr<FString>>>())
+	{}
+
+	/** Set the current find string */
+	PERSONA_API void SetFindString(const FString& InString);
+
+	/** Get the current find string */
+	FStringView GetFindString() const { return FindString; }
+
+	/** Get the current find string */
+	const FString& GetFindStringRef() const { return FindString; }
+
+	/** Set the current replace string */
+	PERSONA_API void SetReplaceString(const FString& InString);
+
+	/** Get the current replace string */
+	FStringView GetReplaceString() const { return ReplaceString; }
+
+	/** Get the current replace string */
+	const FString& GetReplaceStringRef() const { return ReplaceString; }
+	
+	/** Set the current skeleton filter */
+	PERSONA_API void SetSkeletonFilter(const FAssetData& InSkeletonFilter);
+
+	/** Get the current skeleton filter **/
+	const FAssetData& GetSkeletonFilter() const { return SkeletonFilter; }
+
+	/** Sets whether to match whole words when searching strings */
+	PERSONA_API void SetFindWholeWord(bool bInFindWholeWord);
+
+	/** Gets whether to match whole words when searching strings */
+	bool GetFindWholeWord() const { return bFindWholeWord; }
+
+	/** Sets case comparison to use when searching */
+	PERSONA_API void SetSearchCase(ESearchCase::Type InSearchCase);
+
+	/** Gets case comparison to use when searching */
+	ESearchCase::Type GetSearchCase() const { return SearchCase; }
+
+	/** Optional override point for derived classes to supply auto complete names to display in the UI */
+	virtual void GetAutoCompleteNames(TArrayView<FAssetData> InAssetDatas, TSet<FString>& OutUniqueNames) const {}
+
+private:
+	/** UAnimAssetFindReplaceProcessor interface */
+	virtual bool SupportsMode(EAnimAssetFindReplaceMode InMode) const override { return true; }
+	virtual void ExtendToolbar(FToolMenuSection& InSection) override;
+	virtual TSharedRef<SWidget> MakeFindReplaceWidget() override;
+	virtual void FocusInitialWidget() override;
+	virtual bool CanCurrentlyReplace() const override;
+	virtual bool CanCurrentlyRemove() const override;
+	virtual void RefreshCachedData() override;
+
+protected:
+	/** Check if the supplied string matches the find parameters */
+	bool NameMatches(FStringView InNameStringView) const;
+
+private:
+	/** The find string */
+	FString FindString;
+
+	/** The replace string */
+	FString ReplaceString;
+
+	/** The skeleton to filter by */
+	FAssetData SkeletonFilter;
+
+	/** Whether to match whole words when searching strings */
+	bool bFindWholeWord = false;
+
+	/** Case comparison to use when searching */
+	ESearchCase::Type SearchCase = ESearchCase::IgnoreCase;
+
+	/** Cached auto-complete items */
+	TSharedPtr<TArray<TSharedPtr<FString>>> AutoCompleteItems;
+
+	/** Cached find/replace widget */
+	TWeakPtr<SWidget> FindReplaceWidget;
+
+	/** Search box widget for 'find' */
+	TWeakPtr<SAutoCompleteSearchBox> FindSearchBox;
+
+	/** Search box widget for 'replace' */
+	TWeakPtr<SAutoCompleteSearchBox> ReplaceSearchBox;
+};
+
+// DEPRECATED - use UAnimAssetFindReplaceProcessor
 enum class EAnimAssetFindReplaceType : int32
 {
 	Curves,
-	Notifies,
+	Notifies
 };
 
 /** Configuration for the find/replace tab */
 struct FAnimAssetFindReplaceConfig
 {
+	PRAGMA_DISABLE_DEPRECATION_WARNINGS
 	FAnimAssetFindReplaceConfig() = default;
+	FAnimAssetFindReplaceConfig(const FAnimAssetFindReplaceConfig&) = default;
+	FAnimAssetFindReplaceConfig& operator=(const FAnimAssetFindReplaceConfig&) = default;
+	PRAGMA_ENABLE_DEPRECATION_WARNINGS
 
-	/** The initial find string **/
+	/** The initial processor class to use */
+	TSubclassOf<UAnimAssetFindReplaceProcessor> InitialProcessorClass = nullptr;
+
+	UE_DEPRECATED(5.3, "Please use the UAnimAssetFindReplaceProcessor classes to configure the find/replace UI")
 	FString FindString;
 
-	/** The initial replace string **/
+	UE_DEPRECATED(5.3, "Please use the UAnimAssetFindReplaceProcessor classes to configure the find/replace UI")
 	FString ReplaceString;
 
-	/** The initial type */
+	UE_DEPRECATED(5.3, "Please use the UAnimAssetFindReplaceProcessor classes to configure the find/replace UI")
 	EAnimAssetFindReplaceType Type = EAnimAssetFindReplaceType::Curves;
 
-	/** The initial mode */
+	UE_DEPRECATED(5.3, "Please use the UAnimAssetFindReplaceProcessor classes to configure the find/replace UI")
 	EAnimAssetFindReplaceMode Mode = EAnimAssetFindReplaceMode::Find;
 
-	/** The initial skeleton to filter by */
+	UE_DEPRECATED(5.3, "Please use the UAnimAssetFindReplaceProcessor classes to configure the find/replace UI")
 	FAssetData SkeletonFilter;
+};
+
+/** Public interface to the find/replace widget */
+class IAnimAssetFindReplace : public SCompoundWidget
+{
+public:
+	/** Sets the type of the current processor we are using to find/replace with */
+	virtual void SetCurrentProcessor(TSubclassOf<UAnimAssetFindReplaceProcessor> InProcessorClass) = 0;
+
+	/** Gets the current processor instance */
+	virtual UAnimAssetFindReplaceProcessor* GetCurrentProcessor() const = 0;
+	
+	/** Gets the processor instance of the specified type */
+	virtual UAnimAssetFindReplaceProcessor* GetProcessor(TSubclassOf<UAnimAssetFindReplaceProcessor> InProcessorClass) const = 0;
+
+	/** Gets the processor instance of the specified type */
+	template <typename ProcessorType>
+	ProcessorType* GetProcessor() const
+	{
+		return Cast<ProcessorType>(GetProcessor(ProcessorType::StaticClass()));
+	}
 };
\ No newline at end of file
diff --git a/Engine/Source/Editor/Persona/Public/AnimAssetFindReplaceCurves.h b/Engine/Source/Editor/Persona/Public/AnimAssetFindReplaceCurves.h
new file mode 100644
--- /dev/null
+++ b/Engine/Source/Editor/Persona/Public/AnimAssetFindReplaceCurves.h
@@ -0,0 +1,24 @@
+﻿// Copyright Epic Games, Inc. All Rights Reserved.
+
+#pragma once
+
+#include "CoreMinimal.h"
+#include "AnimAssetFindReplace.h"
+#include "AnimAssetFindReplaceCurves.generated.h"
+
+/** Find, replace and remove curves across assets */
+UCLASS(MinimalAPI, DisplayName="Curves")
+class UAnimAssetFindReplaceCurves : public UAnimAssetFindReplaceProcessor_StringBase
+{
+	GENERATED_BODY()
+
+	// UAnimAssetFindReplaceProcessor interface
+	virtual FString GetFindResultStringFromAssetData(const FAssetData& InAssetData) const override;
+	virtual TConstArrayView<UClass*> GetSupportedAssetTypes() const override;
+	virtual bool ShouldFilterOutAsset(const FAssetData& InAssetData, bool& bOutIsOldAsset) const override;
+	virtual void ReplaceInAsset(const FAssetData& InAssetData) const override;
+	virtual void RemoveInAsset(const FAssetData& InAssetData) const override;
+
+	// UAnimAssetFindReplaceProcessor_StringBase interface
+	virtual void GetAutoCompleteNames(TArrayView<FAssetData> InAssetDatas, TSet<FString>& OutUniqueNames) const override;
+};
diff --git a/Engine/Source/Editor/Persona/Public/AnimAssetFindReplaceNotifies.h b/Engine/Source/Editor/Persona/Public/AnimAssetFindReplaceNotifies.h
new file mode 100644
--- /dev/null
+++ b/Engine/Source/Editor/Persona/Public/AnimAssetFindReplaceNotifies.h
@@ -0,0 +1,24 @@
+﻿// Copyright Epic Games, Inc. All Rights Reserved.
+
+#pragma once
+
+#include "CoreMinimal.h"
+#include "AnimAssetFindReplace.h"
+#include "AnimAssetFindReplaceNotifies.generated.h"
+
+/** Find, replace and remove notifies across assets */
+UCLASS(MinimalAPI, DisplayName="Notifies")
+class UAnimAssetFindReplaceNotifies : public UAnimAssetFindReplaceProcessor_StringBase
+{
+	GENERATED_BODY()
+
+	// UAnimAssetFindReplaceProcessor interface
+	virtual FString GetFindResultStringFromAssetData(const FAssetData& InAssetData) const override;
+	virtual TConstArrayView<UClass*> GetSupportedAssetTypes() const override;
+	virtual bool ShouldFilterOutAsset(const FAssetData& InAssetData, bool& bOutIsOldAsset) const override;
+	virtual void ReplaceInAsset(const FAssetData& InAssetData) const override;
+	virtual void RemoveInAsset(const FAssetData& InAssetData) const override;
+
+	// UAnimAssetFindReplaceProcessor_StringBase interface
+	virtual void GetAutoCompleteNames(TArrayView<FAssetData> InAssetDatas, TSet<FString>& OutUniqueNames) const override;
+};
\ No newline at end of file
diff --git a/Engine/Source/Editor/Persona/Public/AnimAssetFindReplaceSyncMarkers.h b/Engine/Source/Editor/Persona/Public/AnimAssetFindReplaceSyncMarkers.h
new file mode 100644
--- /dev/null
+++ b/Engine/Source/Editor/Persona/Public/AnimAssetFindReplaceSyncMarkers.h
@@ -0,0 +1,24 @@
+﻿// Copyright Epic Games, Inc. All Rights Reserved.
+
+#pragma once
+
+#include "CoreMinimal.h"
+#include "AnimAssetFindReplace.h"
+#include "AnimAssetFindReplaceSyncMarkers.generated.h"
+
+/** Find, replace and remove sync markers across assets */
+UCLASS(MinimalAPI, DisplayName="Sync Markers")
+class UAnimAssetFindReplaceSyncMarkers : public UAnimAssetFindReplaceProcessor_StringBase
+{
+	GENERATED_BODY()
+
+	// UAnimAssetFindReplaceProcessor interface
+	virtual FString GetFindResultStringFromAssetData(const FAssetData& InAssetData) const override;
+	virtual TConstArrayView<UClass*> GetSupportedAssetTypes() const override;
+	virtual bool ShouldFilterOutAsset(const FAssetData& InAssetData, bool& bOutIsOldAsset) const override;
+	virtual void ReplaceInAsset(const FAssetData& InAssetData) const override;
+	virtual void RemoveInAsset(const FAssetData& InAssetData) const override;
+
+	// UAnimAssetFindReplaceProcessor_StringBase interface
+	virtual void GetAutoCompleteNames(TArrayView<FAssetData> InAssetDatas, TSet<FString>& OutUniqueNames) const override;
+};
\ No newline at end of file
diff --git a/Engine/Source/Editor/SkeletonEditor/Private/EditableSkeleton.cpp b/Engine/Source/Editor/SkeletonEditor/Private/EditableSkeleton.cpp
--- a/Engine/Source/Editor/SkeletonEditor/Private/EditableSkeleton.cpp
+++ b/Engine/Source/Editor/SkeletonEditor/Private/EditableSkeleton.cpp
@@ -888,8 +888,9 @@ TSharedRef<SWidget> FEditableSkeleton::CreateBlendProfilePicker(const FBlendProf
 	return BlendProfilePicker;
 }
 
-int32 FEditableSkeleton::DeleteAnimNotifies(const TArray<FName>& InNotifyNames)
+int32 FEditableSkeleton::DeleteAnimNotifies(const TArray<FName>& InNotifyNames, bool bDeleteFromAnimations)
 {
+	int32 NumAnimationsModified = 0;
 	const FScopedTransaction Transaction(LOCTEXT("DeleteAnimNotify", "Delete Anim Notify"));
 	Skeleton->Modify();
 
@@ -898,42 +899,67 @@ int32 FEditableSkeleton::DeleteAnimNotifies(const TArray<FName>& InNotifyNames)
 		Skeleton->AnimationNotifies.Remove(Notify);
 	}
 
-	TArray<FAssetData> CompatibleAnimSequences;
-	GetCompatibleAnimSequences(CompatibleAnimSequences);
-
-	int32 NumAnimationsModified = 0;
-
-	for (int32 AssetIndex = 0; AssetIndex < CompatibleAnimSequences.Num(); ++AssetIndex)
+	if(bDeleteFromAnimations)
 	{
-		const FAssetData& PossibleAnimSequence = CompatibleAnimSequences[AssetIndex];
-		if(UObject* LoadedAsset = PossibleAnimSequence.GetAsset())
-		{
-			UAnimSequenceBase* Sequence = CastChecked<UAnimSequenceBase>(LoadedAsset);
+		TArray<FAssetData> CompatibleAnimSequences;
+		GetCompatibleAnimSequences(CompatibleAnimSequences);
 
-			if (Sequence->RemoveNotifies(InNotifyNames))
+		for (int32 AssetIndex = 0; AssetIndex < CompatibleAnimSequences.Num(); ++AssetIndex)
+		{
+			const FAssetData& PossibleAnimSequence = CompatibleAnimSequences[AssetIndex];
+			if(UObject* LoadedAsset = PossibleAnimSequence.GetAsset())
 			{
-				++NumAnimationsModified;
+				UAnimSequenceBase* Sequence = CastChecked<UAnimSequenceBase>(LoadedAsset);
+
+				if (Sequence->RemoveNotifies(InNotifyNames))
+				{
+					++NumAnimationsModified;
+				}
 			}
 		}
 	}
 
 	FBlueprintActionDatabase::Get().RefreshAssetActions(Skeleton);
-	
 	OnNotifiesChanged.Broadcast();
 
 	return NumAnimationsModified;
 }
 
-void FEditableSkeleton::DeleteSyncMarkers(const TArray<FName>& InSyncMarkerNames)
+int32 FEditableSkeleton::DeleteSyncMarkers(const TArray<FName>& InSyncMarkerNames, bool bDeleteFromAnimations)
 {
+	int32 NumAnimationsModified = 0;
 	const FScopedTransaction Transaction(LOCTEXT("DeleteSyncMarkers", "Delete Sync Markers"));
-
 	Skeleton->Modify();
 
 	for (FName Marker : InSyncMarkerNames)
 	{
 		Skeleton->RemoveMarkerName(Marker);
 	}
+
+	if(bDeleteFromAnimations)
+	{
+		TArray<FAssetData> CompatibleAnimSequences;
+		GetCompatibleAnimSequences(CompatibleAnimSequences);
+
+		for (int32 AssetIndex = 0; AssetIndex < CompatibleAnimSequences.Num(); ++AssetIndex)
+		{
+			const FAssetData& PossibleAnimSequence = CompatibleAnimSequences[AssetIndex];
+			if(UObject* LoadedAsset = PossibleAnimSequence.GetAsset())
+			{
+				if(UAnimSequence* Sequence = CastChecked<UAnimSequence>(LoadedAsset))
+				{
+					if (Sequence->RemoveSyncMarkers(InSyncMarkerNames))
+					{
+						++NumAnimationsModified;
+					}
+				}
+			}
+		}
+	}
+	
+	OnNotifiesChanged.Broadcast();
+
+	return NumAnimationsModified;
 }
 
 void FEditableSkeleton::AddNotify(FName NewName)
@@ -951,53 +977,113 @@ void FEditableSkeleton::AddSyncMarker(FName NewName)
 	const FScopedTransaction Transaction(LOCTEXT("AddNewSyncMarkerToSkeleton", "Add New Sync Marker To Skeleton"));
 	Skeleton->Modify();
 	Skeleton->RegisterMarkerName(NewName);
+
+	OnNotifiesChanged.Broadcast();
 }
 
-int32 FEditableSkeleton::RenameNotify(const FName NewName, const FName OldName)
+int32 FEditableSkeleton::RenameNotify(const FName NewName, const FName OldName, bool bRenameInAnimations)
 {
-	const FScopedTransaction Transaction(LOCTEXT("RenameAnimNotify", "Rename Anim Notify"));
-	Skeleton->Modify();
-
+	int32 NumAnimationsModified = 0;
 	int32 Index = Skeleton->AnimationNotifies.IndexOfByKey(OldName);
-	Skeleton->AnimationNotifies[Index] = NewName;
+	if(Index != INDEX_NONE)
+	{
+		const FScopedTransaction Transaction(LOCTEXT("RenameAnimNotify", "Rename Anim Notify"));
+		Skeleton->Modify();
 
-	TArray<FAssetData> CompatibleAnimSequences;
-	GetCompatibleAnimSequences(CompatibleAnimSequences);
+		Skeleton->AnimationNotifies[Index] = NewName;
 
-	int32 NumAnimationsModified = 0;
-
-	for (int32 AssetIndex = 0; AssetIndex < CompatibleAnimSequences.Num(); ++AssetIndex)
-	{
-		const FAssetData& PossibleAnimSequence = CompatibleAnimSequences[AssetIndex];
-		if(UObject* Asset = PossibleAnimSequence.GetAsset())
+		if(bRenameInAnimations)
 		{
-			UAnimSequenceBase* Sequence = Cast<UAnimSequenceBase>(Asset);
-
-			bool SequenceModified = false;
-			for (int32 NotifyIndex = Sequence->Notifies.Num() - 1; NotifyIndex >= 0; --NotifyIndex)
+			TArray<FAssetData> CompatibleAnimSequences;
+			GetCompatibleAnimSequences(CompatibleAnimSequences);
+			
+			for (int32 AssetIndex = 0; AssetIndex < CompatibleAnimSequences.Num(); ++AssetIndex)
 			{
-				FAnimNotifyEvent& AnimNotify = Sequence->Notifies[NotifyIndex];
-				if (OldName == AnimNotify.NotifyName)
+				const FAssetData& PossibleAnimSequence = CompatibleAnimSequences[AssetIndex];
+				if(UObject* Asset = PossibleAnimSequence.GetAsset())
 				{
-					if (!SequenceModified)
+					UAnimSequenceBase* Sequence = Cast<UAnimSequenceBase>(Asset);
+
+					bool SequenceModified = false;
+					for (int32 NotifyIndex = Sequence->Notifies.Num() - 1; NotifyIndex >= 0; --NotifyIndex)
 					{
-						Sequence->Modify();
-						++NumAnimationsModified;
-						SequenceModified = true;
+						FAnimNotifyEvent& AnimNotify = Sequence->Notifies[NotifyIndex];
+						if (OldName == AnimNotify.NotifyName)
+						{
+							if (!SequenceModified)
+							{
+								Sequence->Modify();
+								++NumAnimationsModified;
+								SequenceModified = true;
+							}
+							AnimNotify.NotifyName = NewName;
+						}
+					}
+
+					if (SequenceModified)
+					{
+						Sequence->MarkPackageDirty();
 					}
-					AnimNotify.NotifyName = NewName;
 				}
 			}
+		}
 
-			if (SequenceModified)
+		FBlueprintActionDatabase::Get().RefreshAssetActions(Skeleton);
+		OnNotifiesChanged.Broadcast();
+	}
+
+	return NumAnimationsModified;
+}
+
+int32 FEditableSkeleton::RenameSyncMarker(const FName NewName, const FName OldName, bool bRenameInAnimations)
+{
+	int32 NumAnimationsModified = 0;
+	int32 Index = Skeleton->ExistingMarkerNames.IndexOfByKey(OldName);
+	if(Index != INDEX_NONE)
+	{
+		const FScopedTransaction Transaction(LOCTEXT("RenameSyncMarker", "Rename Sync Marker"));
+		Skeleton->Modify();
+		Skeleton->ExistingMarkerNames[Index] = NewName;
+
+		if(bRenameInAnimations)
+		{
+			TArray<FAssetData> CompatibleAnimSequences;
+			GetCompatibleAnimSequences(CompatibleAnimSequences);
+
+			for (int32 AssetIndex = 0; AssetIndex < CompatibleAnimSequences.Num(); ++AssetIndex)
 			{
-				Sequence->MarkPackageDirty();
+				const FAssetData& PossibleAnimSequence = CompatibleAnimSequences[AssetIndex];
+				if(UObject* Asset = PossibleAnimSequence.GetAsset())
+				{
+					if(UAnimSequence* Sequence = Cast<UAnimSequence>(Asset))
+					{
+						bool SequenceModified = false;
+						for (int32 MarkerIndex = Sequence->AuthoredSyncMarkers.Num() - 1; MarkerIndex >= 0; --MarkerIndex)
+						{
+							FAnimSyncMarker& AnimSyncMarker = Sequence->AuthoredSyncMarkers[MarkerIndex];
+							if (OldName == AnimSyncMarker.MarkerName)
+							{
+								if (!SequenceModified)
+								{
+									Sequence->Modify();
+									++NumAnimationsModified;
+									SequenceModified = true;
+								}
+								AnimSyncMarker.MarkerName = NewName;
+							}
+						}
+
+						if (SequenceModified)
+						{
+							Sequence->MarkPackageDirty();
+						}
+					}
+				}
 			}
 		}
-	}
 
-	FBlueprintActionDatabase::Get().RefreshAssetActions(Skeleton);
-	OnNotifiesChanged.Broadcast();
+		OnNotifiesChanged.Broadcast();
+	}
 
 	return NumAnimationsModified;
 }
diff --git a/Engine/Source/Editor/SkeletonEditor/Private/EditableSkeleton.h b/Engine/Source/Editor/SkeletonEditor/Private/EditableSkeleton.h
--- a/Engine/Source/Editor/SkeletonEditor/Private/EditableSkeleton.h
+++ b/Engine/Source/Editor/SkeletonEditor/Private/EditableSkeleton.h
@@ -38,11 +38,12 @@ public:
 	virtual void SetBlendProfileMode(FName InBlendProfileName, EBlendProfileMode ProfileMode);
 	virtual USkeletalMeshSocket* DuplicateSocket(const FSelectedSocketInfo& SocketInfoToDuplicate, const FName& NewParentBoneName, USkeletalMesh* InSkeletalMesh) override;
 	virtual int32 ValidatePreviewAttachedObjects() override;
-	virtual int32 DeleteAnimNotifies(const TArray<FName>& InSelectedNotifyNames) override;
-	virtual void DeleteSyncMarkers(const TArray<FName>& ISyncMarkerNames) override;
+	virtual int32 DeleteAnimNotifies(const TArray<FName>& InSelectedNotifyNames, bool bDeleteFromAnimations = true) override;
+	virtual int32 DeleteSyncMarkers(const TArray<FName>& ISyncMarkerNames, bool bDeleteFromAnimations = true) override;
 	virtual void AddNotify(FName NewName) override;
 	virtual void AddSyncMarker(FName NewName) override;
-	virtual int32 RenameNotify(const FName NewName, const FName OldName) override;
+	virtual int32 RenameNotify(const FName NewName, const FName OldName, bool bRenameInAnimations = true) override;
+	virtual int32 RenameSyncMarker(const FName NewName, const FName OldName, bool bRenameInAnimations = true) override;
 	virtual void BroadcastNotifyChanged() override;
 	virtual void GetCompatibleAnimSequences(TArray<struct FAssetData>& OutAssets) override;
 	virtual void RenameSocket(const FName OldSocketName, const FName NewSocketName, USkeletalMesh* InSkeletalMesh) override;
diff --git a/Engine/Source/Editor/SkeletonEditor/Public/IEditableSkeleton.h b/Engine/Source/Editor/SkeletonEditor/Public/IEditableSkeleton.h
--- a/Engine/Source/Editor/SkeletonEditor/Public/IEditableSkeleton.h
+++ b/Engine/Source/Editor/SkeletonEditor/Public/IEditableSkeleton.h
@@ -103,12 +103,13 @@ public:
 	 * Delete anim notifies by name
 	 * @return the number of animations modified
 	 */	
-	virtual int32 DeleteAnimNotifies(const TArray<FName>& InotifyNames) = 0;
+	virtual int32 DeleteAnimNotifies(const TArray<FName>& InotifyNames, bool bDeleteFromAnimations = true) = 0;
 
 	/**
 	 * Delete sync markers from the skeleton cache by name
+	 * @return the number of animations modified
 	 */
-	virtual void DeleteSyncMarkers(const TArray<FName>& ISyncMarkerNames) = 0;
+	virtual int32 DeleteSyncMarkers(const TArray<FName>& ISyncMarkerNames, bool bDeleteFromAnimations = true) = 0;
 
 	/**
 	* Add a notify
@@ -119,13 +120,19 @@ public:
 	 * Rename a notify
 	 * @return the number of animations modified
 	 */	
-	virtual int32 RenameNotify(const FName NewName, const FName OldName) = 0;
+	virtual int32 RenameNotify(const FName NewName, const FName OldName, bool bRenameInAnimations = true) = 0;
 
 	/**
 	* Add a sync marker
 	*/
 	virtual void AddSyncMarker(FName NewName) = 0;
 
+	/**
+	* Rename a sync marker
+	 * @return the number of animations modified
+	*/
+	virtual int32 RenameSyncMarker(FName NewName, const FName OldName, bool bRenameInAnimations = true) = 0;
+	
 	/** Inform the system that something about a notify changed */	
 	virtual void BroadcastNotifyChanged() = 0;
 
diff --git a/Engine/Source/Runtime/Engine/Classes/Animation/AnimSequence.h b/Engine/Source/Runtime/Engine/Classes/Animation/AnimSequence.h
--- a/Engine/Source/Runtime/Engine/Classes/Animation/AnimSequence.h
+++ b/Engine/Source/Runtime/Engine/Classes/Animation/AnimSequence.h
@@ -659,6 +659,14 @@ public:
 	/** Sort the sync markers array by time, earliest first. */
 	void SortSyncMarkers();
 
+#if WITH_EDITOR
+	/** Remove all markers with the specified names */
+	bool RemoveSyncMarkers(const TArray<FName>& MarkersToRemove);
+
+	/** Rename the markers with the specified name */
+	bool RenameSyncMarkers(FName InOldName, FName InNewName);
+#endif
+
 	// Advancing based on markers
 	float GetCurrentTimeFromMarkers(FMarkerPair& PrevMarker, FMarkerPair& NextMarker, float PositionBetweenMarkers) const;
 	virtual void AdvanceMarkerPhaseAsLeader(bool bLooping, float MoveDelta, const TArray<FName>& ValidMarkerNames, float& CurrentTime, FMarkerPair& PrevMarker, FMarkerPair& NextMarker, TArray<FPassedMarker>& MarkersPassed, const UMirrorDataTable* MirrorTable) const;
diff --git a/Engine/Source/Runtime/Engine/Classes/Animation/Skeleton.h b/Engine/Source/Runtime/Engine/Classes/Animation/Skeleton.h
--- a/Engine/Source/Runtime/Engine/Classes/Animation/Skeleton.h
+++ b/Engine/Source/Runtime/Engine/Classes/Animation/Skeleton.h
@@ -38,6 +38,7 @@ class USkeleton;
 typedef SmartName::UID_Type SkeletonAnimCurveUID;
 class USkeleton;
 struct FSkeletonRemapping;
+class FEditableSkeleton;
 
 // Delegate used to control global skeleton compatibility
 DECLARE_DELEGATE_RetVal(bool, FAreAllSkeletonsCompatible);
@@ -472,7 +473,10 @@ public:
 	void RegisterMarkerName(FName MarkerName) { ExistingMarkerNames.AddUnique(MarkerName); ExistingMarkerNames.Sort(FNameLexicalLess()); }
 
 	// Remove a sync marker name
-	void RemoveMarkerName(FName MarkerName) { ExistingMarkerNames.Remove(MarkerName); }
+	ENGINE_API bool RemoveMarkerName(FName MarkerName);
+
+	// Rename a sync marker name
+	ENGINE_API bool RenameMarkerName(FName InOldName, FName InNewName);
 #endif
 
 protected:
@@ -1001,6 +1005,10 @@ public:
 	ENGINE_API static const FName AnimNotifyTag;
 	ENGINE_API static const FString AnimNotifyTagDelimiter;
 
+	// Asset registry information for animation sync markers
+	ENGINE_API static const FName AnimSyncMarkerTag;
+	ENGINE_API static const FString AnimSyncMarkerTagDelimiter;
+	
 	// Asset registry information for animation curves
 	ENGINE_API static const FName CurveNameTag;
 	ENGINE_API static const FString CurveTagDelimiter;
@@ -1051,5 +1059,6 @@ protected:
 	TArray<TObjectPtr<UAssetUserData>> AssetUserData;
 
 	friend struct FReferenceSkeletonModifier;
+	friend class FEditableSkeleton;
 };
 
diff --git a/Engine/Source/Runtime/Engine/Private/Animation/AnimSequence.cpp b/Engine/Source/Runtime/Engine/Private/Animation/AnimSequence.cpp
--- a/Engine/Source/Runtime/Engine/Private/Animation/AnimSequence.cpp
+++ b/Engine/Source/Runtime/Engine/Private/Animation/AnimSequence.cpp
@@ -46,6 +46,7 @@
 #include "ITimeManagementModule.h"
 #include "Animation/SkeletonRemappingRegistry.h"
 #include "Animation/SkeletonRemapping.h"
+#include "Animation/Skeleton.h"
 
 LLM_DEFINE_TAG(SequenceData);
 
@@ -432,6 +433,18 @@ void UAnimSequence::GetAssetRegistryTags(TArray<FAssetRegistryTag>& OutTags) con
 		OutTags.Add(FAssetRegistryTag(TEXT("Number of Frames"), TEXT("0"), FAssetRegistryTag::TT_Numerical));
 		OutTags.Add(FAssetRegistryTag(TEXT("Number of Keys"), TEXT("0"), FAssetRegistryTag::TT_Numerical));
 	}
+
+	// Output unique sync marker names we use
+	TStringBuilder<256> SyncMarkersBuilder;
+	SyncMarkersBuilder.Append(USkeleton::AnimSyncMarkerTagDelimiter);
+
+	for(FName SyncMarker : UniqueMarkerNames)
+	{
+		SyncMarkersBuilder.Append(SyncMarker.ToString());
+		SyncMarkersBuilder.Append(USkeleton::AnimSyncMarkerTagDelimiter);
+	}
+
+	OutTags.Add(FAssetRegistryTag(USkeleton::AnimSyncMarkerTag, SyncMarkersBuilder.ToString(), FAssetRegistryTag::TT_Hidden));
 #endif
 
 	OutTags.Add(FAssetRegistryTag(TEXT("Compressed Size (KB)"), FString::Printf(TEXT("%.02f"), (float)GetApproxCompressedSize() / 1024.0f), FAssetRegistryTag::TT_Numerical));
@@ -833,6 +846,60 @@ void UAnimSequence::SortSyncMarkers()
 	RefreshSyncMarkerDataFromAuthored();
 }
 
+#if WITH_EDITOR
+
+bool UAnimSequence::RemoveSyncMarkers(const TArray<FName>& NotifiesToRemove)
+{
+	bool bSequenceModified = false;
+	for (int32 MarkerIndex = AuthoredSyncMarkers.Num() - 1; MarkerIndex >= 0; --MarkerIndex)
+	{
+		FAnimSyncMarker& Marker = AuthoredSyncMarkers[MarkerIndex];
+		if (NotifiesToRemove.Contains(Marker.MarkerName))
+		{
+			if (!bSequenceModified)
+			{
+				Modify();
+				bSequenceModified = true;
+			}
+			AuthoredSyncMarkers.RemoveAtSwap(MarkerIndex);
+		}
+	}
+
+	if (bSequenceModified)
+	{
+		MarkPackageDirty();
+		RefreshCacheData();
+	}
+	return bSequenceModified;
+}
+
+bool UAnimSequence::RenameSyncMarkers(FName InOldName, FName InNewName)
+{
+	bool bSequenceModified = false;
+	for(FAnimSyncMarker& Marker : AuthoredSyncMarkers)
+	{
+		if(Marker.MarkerName == InOldName)
+		{
+			if (!bSequenceModified)
+			{
+				Modify();
+				bSequenceModified = true;
+			}
+
+			Marker.MarkerName = InNewName;
+		}
+	}
+
+	if (bSequenceModified)
+	{
+		MarkPackageDirty();
+		RefreshCacheData();
+	}
+	return bSequenceModified;
+}
+
+#endif
+
 void UAnimSequence::GetPreloadDependencies(TArray<UObject*>& OutDeps)
 {
 	Super::GetPreloadDependencies(OutDeps);
diff --git a/Engine/Source/Runtime/Engine/Private/Animation/Skeleton.cpp b/Engine/Source/Runtime/Engine/Private/Animation/Skeleton.cpp
--- a/Engine/Source/Runtime/Engine/Private/Animation/Skeleton.cpp
+++ b/Engine/Source/Runtime/Engine/Private/Animation/Skeleton.cpp
@@ -45,6 +45,9 @@ PRAGMA_ENABLE_DEPRECATION_WARNINGS
 const FName USkeleton::AnimNotifyTag = FName(TEXT("AnimNotifyList"));
 const FString USkeleton::AnimNotifyTagDelimiter = TEXT(";");
 
+const FName USkeleton::AnimSyncMarkerTag = FName(TEXT("AnimSyncMarkerList"));
+const FString USkeleton::AnimSyncMarkerTagDelimiter = TEXT(";");
+
 const FName USkeleton::CurveNameTag = FName(TEXT("CurveNameList"));
 const FString USkeleton::CurveTagDelimiter = TEXT(";");
 
@@ -558,6 +561,44 @@ const FSkeletonRemapping* USkeleton::GetSkeletonRemapping(const USkeleton* Sourc
 	return &UE::Anim::FSkeletonRemappingRegistry::Get().GetRemapping(SourceSkeleton, this);
 }
 
+#if WITH_EDITOR
+
+bool USkeleton::RemoveMarkerName(FName MarkerName)
+{
+	if(ExistingMarkerNames.Contains(MarkerName))
+	{
+		Modify();
+	}
+
+	return ExistingMarkerNames.Remove(MarkerName) != 0;
+}
+
+bool USkeleton::RenameMarkerName(FName InOldName, FName InNewName)
+{
+	if(ExistingMarkerNames.Contains(InOldName))
+	{
+		Modify();
+	}
+
+	if(ExistingMarkerNames.Contains(InNewName))
+	{
+		return ExistingMarkerNames.Remove(InOldName) != 0;
+	}
+
+	for(FName& MarkerName : ExistingMarkerNames)
+	{
+		if(MarkerName == InOldName)
+		{
+			MarkerName = InNewName;
+			return true;
+		}
+	}
+
+	return false;
+}
+
+#endif
+
 bool USkeleton::DoesParentChainMatch(int32 StartBoneIndex, const USkinnedAsset* InSkinnedAsset) const
 {
 	const FReferenceSkeleton& SkeletonRefSkel = ReferenceSkeleton;
@@ -2078,6 +2119,30 @@ void USkeleton::GetAssetRegistryTags(TArray<FAssetRegistryTag>& OutTags) const
 
 	OutTags.Add(FAssetRegistryTag(USkeleton::CompatibleSkeletonsNameTag, CompatibleSkeletonsBuilder.ToString(), FAssetRegistryTag::TT_Hidden));
 
+	// Output sync notify names we use
+	TStringBuilder<256> NotifiesBuilder;
+	NotifiesBuilder.Append(USkeleton::AnimNotifyTagDelimiter);
+
+	for(FName NotifyName : AnimationNotifies)
+	{
+		NotifiesBuilder.Append(NotifyName.ToString());
+		NotifiesBuilder.Append(USkeleton::AnimNotifyTagDelimiter);
+	}
+
+	OutTags.Add(FAssetRegistryTag(USkeleton::AnimNotifyTag, NotifiesBuilder.ToString(), FAssetRegistryTag::TT_Hidden));
+	
+	// Output sync marker names we use
+	TStringBuilder<256> SyncMarkersBuilder;
+	SyncMarkersBuilder.Append(USkeleton::AnimSyncMarkerTagDelimiter);
+
+	for(FName SyncMarker : ExistingMarkerNames)
+	{
+		SyncMarkersBuilder.Append(SyncMarker.ToString());
+		SyncMarkersBuilder.Append(USkeleton::AnimSyncMarkerTagDelimiter);
+	}
+
+	OutTags.Add(FAssetRegistryTag(USkeleton::AnimSyncMarkerTag, SyncMarkersBuilder.ToString(), FAssetRegistryTag::TT_Hidden));
+	
 	// Allow asset user data to output tags
 	for(UAssetUserData* AssetUserDataItem : AssetUserData)
 	{
