Fixed a crash that occurred when pasting anim notifies from multiple tracks to a location that did not contain the same tracks

diff --git a/Engine/Source/Editor/Persona/Private/SAnimNotifyPanel.cpp b/Engine/Source/Editor/Persona/Private/SAnimNotifyPanel.cpp
--- a/Engine/Source/Editor/Persona/Private/SAnimNotifyPanel.cpp
+++ b/Engine/Source/Editor/Persona/Private/SAnimNotifyPanel.cpp
@@ -4573,6 +4573,7 @@ void SAnimNotifyPanel::OnPasteNodes(SAnimNotifyTrack* RequestTrack, float ClickT
 				AddTrack();
 				--TracksToAdd;
 			}
+			RefreshNotifyTracks(); 
 			NumTracks = NotifyAnimTracks.Num();
 		}
 
@@ -4606,18 +4607,21 @@ void SAnimNotifyPanel::OnPasteNodes(SAnimNotifyTrack* RequestTrack, float ClickT
 				float TimeOffset = OrigTime - OrigBeginTime;
 				float TimeToPaste = ClickTime + TimeOffset * ScaleMultiplier;
 
-				TSharedPtr<SAnimNotifyTrack> TrackToUse = NotifyAnimTracks[PasteIdx + TrackOffset];
-				if (NodeObjectType == ENodeObjectTypes::NOTIFY)
+				if (PasteIdx + TrackOffset < NotifyAnimTracks.Num())
 				{
-					TrackToUse->PasteSingleNotify(NotifyExportString, TimeToPaste);
-				}
-				else if (NodeObjectType == ENodeObjectTypes::SYNC_MARKER)
-				{
-					TrackToUse->PasteSingleSyncMarker(NotifyExportString, TimeToPaste);
-				}
-				else
-				{
-					check(false); //Unknown value in paste
+					TSharedPtr<SAnimNotifyTrack> TrackToUse = NotifyAnimTracks[PasteIdx + TrackOffset];
+					if (NodeObjectType == ENodeObjectTypes::NOTIFY)
+					{
+						TrackToUse->PasteSingleNotify(NotifyExportString, TimeToPaste);
+					}
+					else if (NodeObjectType == ENodeObjectTypes::SYNC_MARKER)
+					{
+						TrackToUse->PasteSingleSyncMarker(NotifyExportString, TimeToPaste);
+					}
+					else
+					{
+						check(false); //Unknown value in paste
+					}
 				}
 			}
 		}
