Fix bugs related to updating SAnimViewportToolBar to use UToolsMenu.

I was passing in "this" to lambdas when registering the new menus when I should have been using a context instead.

diff --git a/Engine/Source/Editor/Persona/Private/AnimViewportToolBarToolMenuContext.h b/Engine/Source/Editor/Persona/Private/AnimViewportToolBarToolMenuContext.h
new file mode 100644
--- /dev/null
+++ b/Engine/Source/Editor/Persona/Private/AnimViewportToolBarToolMenuContext.h
@@ -0,0 +1,17 @@
+// Copyright Epic Games, Inc. All Rights Reserved.
+
+#pragma once
+
+#include "CoreMinimal.h"
+#include "AnimViewportToolBarToolMenuContext.generated.h"
+
+class SAnimViewportToolBar;
+
+UCLASS()
+class UAnimViewportToolBarToolMenuContext : public UObject
+{
+	GENERATED_BODY()
+
+public:
+	TWeakPtr<const SAnimViewportToolBar> AnimViewportToolBar;
+};
diff --git a/Engine/Source/Editor/Persona/Private/SAnimViewportToolBar.cpp b/Engine/Source/Editor/Persona/Private/SAnimViewportToolBar.cpp
--- a/Engine/Source/Editor/Persona/Private/SAnimViewportToolBar.cpp
+++ b/Engine/Source/Editor/Persona/Private/SAnimViewportToolBar.cpp
@@ -3,6 +3,7 @@
 
 #include "SAnimViewportToolBar.h"
 
+#include "AnimViewportToolBarToolMenuContext.h"
 #include "AnimPreviewInstance.h"
 #include "Widgets/Layout/SBorder.h"
 #include "Widgets/Text/STextBlock.h"
@@ -739,8 +740,11 @@ TSharedRef<SWidget> SAnimViewportToolBar::GenerateCharacterMenu() const
 			Section.AddSubMenu(TEXT("AnimationSubMenu"),
 				LOCTEXT("CharacterMenu_AnimationSubMenu", "Animation"),
 				LOCTEXT("CharacterMenu_AnimationSubMenuToolTip", "Animation-related options"),
-				FNewToolMenuDelegate::CreateLambda([this, WeakSharedViewport = Viewport](UToolMenu* InSubMenu)
-				{
+				FNewToolMenuDelegate::CreateLambda([](UToolMenu* InSubMenu)
+			{
+					UAnimViewportToolBarToolMenuContext* Context = InSubMenu->FindContext<UAnimViewportToolBarToolMenuContext>();
+					const SAnimViewportToolBar* ContextThis = Context ? Context->AnimViewportToolBar.Pin().Get() : nullptr;
+
 					{
 						FToolMenuSection& Section = InSubMenu->AddSection("AnimViewportRootMotion", LOCTEXT("CharacterMenu_RootMotionLabel", "Root Motion"));
 						Section.AddMenuEntry(FAnimViewportShowCommands::Get().DoNotProcessRootMotion);
@@ -755,9 +759,9 @@ TSharedRef<SWidget> SAnimViewportToolBar::GenerateCharacterMenu() const
 						Section.AddMenuEntry(FAnimViewportShowCommands::Get().ShowAdditiveBaseBones);
 						Section.AddMenuEntry(FAnimViewportShowCommands::Get().ShowSourceRawAnimation);
 
-						if (WeakSharedViewport.IsValid())
+						if (ContextThis && ContextThis->Viewport.IsValid())
 						{
-							if ( UDebugSkelMeshComponent* PreviewComponent = WeakSharedViewport.Pin()->GetPreviewScene()->GetPreviewMeshComponent())
+							if ( UDebugSkelMeshComponent* PreviewComponent = ContextThis->Viewport.Pin()->GetPreviewScene()->GetPreviewMeshComponent())
 							{
 								FUIAction DisableUnlessPreviewInstance(
 									FExecuteAction::CreateLambda([](){}),
@@ -770,7 +774,7 @@ TSharedRef<SWidget> SAnimViewportToolBar::GenerateCharacterMenu() const
 								Section.AddSubMenu(TEXT("MirrorSubMenu"),
 									LOCTEXT("CharacterMenu_AnimationSubMenu_MirrorSubMenu", "Mirror"),
 									LOCTEXT("CharacterMenu_AnimationSubMenu_MirrorSubMenuToolTip", "Mirror the animation using the selected mirror data table"),
-									FNewToolMenuChoice(FNewMenuDelegate::CreateRaw(const_cast<SAnimViewportToolBar*>(this), &SAnimViewportToolBar::FillCharacterMirrorMenu)),
+									FNewToolMenuChoice(FNewMenuDelegate::CreateRaw(ContextThis, &SAnimViewportToolBar::FillCharacterMirrorMenu)),
 									FToolUIActionChoice(DisableUnlessPreviewInstance),
 									EUserInterfaceActionType::Button,
 									false,
@@ -781,10 +785,10 @@ TSharedRef<SWidget> SAnimViewportToolBar::GenerateCharacterMenu() const
 						Section.AddMenuEntry(FAnimViewportShowCommands::Get().ShowBakedAnimation);
 						Section.AddMenuEntry(FAnimViewportShowCommands::Get().DisablePostProcessBlueprint);
 					}
-					if ( WeakSharedViewport.IsValid())
+					if (ContextThis && ContextThis->Viewport.IsValid())
 					{
 						FToolMenuSection& Section = InSubMenu->AddSection("SkinWeights", LOCTEXT("SkinWeights_Label", "Skin Weight Profiles"));
-						Section.AddEntry(FToolMenuEntry::InitWidget(TEXT("SkinWeightCombo"), WeakSharedViewport.Pin()->SkinWeightCombo.ToSharedRef(), FText()));
+						Section.AddEntry(FToolMenuEntry::InitWidget(TEXT("SkinWeightCombo"), ContextThis->Viewport.Pin()->SkinWeightCombo.ToSharedRef(), FText()));
 					}
 				})
 			);
@@ -792,8 +796,10 @@ TSharedRef<SWidget> SAnimViewportToolBar::GenerateCharacterMenu() const
 			Section.AddSubMenu(TEXT("BonesSubMenu"),
 				LOCTEXT("CharacterMenu_BoneDrawSubMenu", "Bones"),
 				LOCTEXT("CharacterMenu_BoneDrawSubMenuToolTip", "Bone Drawing Options"),
-				FNewToolMenuDelegate::CreateLambda([this](UToolMenu* InSubMenu)
+				FNewToolMenuDelegate::CreateLambda([](UToolMenu* InSubMenu)
 				{
+					UAnimViewportToolBarToolMenuContext* Context = InSubMenu->FindContext<UAnimViewportToolBarToolMenuContext>();
+					const SAnimViewportToolBar* ContextThis = Context ? Context->AnimViewportToolBar.Pin().Get() : nullptr;
 					{
 						FToolMenuSection& Section = InSubMenu->AddSection("BonesAndSockets", LOCTEXT("CharacterMenu_BonesAndSocketsLabel", "Show"));
 						Section.AddMenuEntry(FAnimViewportShowCommands::Get().ShowSockets);
@@ -803,8 +809,11 @@ TSharedRef<SWidget> SAnimViewportToolBar::GenerateCharacterMenu() const
 
 					{
 						FToolMenuSection& Section = InSubMenu->AddSection("AnimViewportPreviewHierarchyBoneDraw", LOCTEXT("CharacterMenu_Actions_BoneDrawing", "Bone Drawing"));
-						TSharedPtr<SWidget> BoneSizeWidget = SNew(SBoneDrawSizeSetting).AnimEditorViewport(Viewport);
-						Section.AddEntry(FToolMenuEntry::InitWidget(TEXT("BoneDrawSize"), BoneSizeWidget.ToSharedRef(), LOCTEXT("CharacterMenu_Actions_BoneDrawSize", "Bone Draw Size:")));
+						if (ContextThis)
+						{
+							TSharedPtr<SWidget> BoneSizeWidget = SNew(SBoneDrawSizeSetting).AnimEditorViewport(ContextThis->Viewport);
+							Section.AddEntry(FToolMenuEntry::InitWidget(TEXT("BoneDrawSize"), BoneSizeWidget.ToSharedRef(), LOCTEXT("CharacterMenu_Actions_BoneDrawSize", "Bone Draw Size:")));
+						}
 						Section.AddMenuEntry(FAnimViewportShowCommands::Get().ShowBoneDrawAll);
 						Section.AddMenuEntry(FAnimViewportShowCommands::Get().ShowBoneDrawSelected);
 						Section.AddMenuEntry(FAnimViewportShowCommands::Get().ShowBoneDrawSelectedAndParents);
@@ -815,18 +824,22 @@ TSharedRef<SWidget> SAnimViewportToolBar::GenerateCharacterMenu() const
 				})
 				);
 
-			UDebugSkelMeshComponent* PreviewComp = Viewport.Pin()->GetPreviewScene()->GetPreviewMeshComponent();
-
-			if (PreviewComp && GetDefault<UPersonaOptions>()->bExposeClothingSceneElementMenu)
+			Section.AddDynamicEntry("ClothingSubMenu", FNewToolMenuSectionDelegate::CreateLambda([](FToolMenuSection& InSection)
 			{
-				constexpr bool bInOpenSubMenuOnClick = false;
-				constexpr bool bShouldCloseWindowAfterMenuSelection = false;
-				Section.AddSubMenu(TEXT("ClothingSubMenu"),
-					LOCTEXT("CharacterMenu_ClothingSubMenu", "Clothing"),
-					LOCTEXT("CharacterMenu_ClothingSubMenuToolTip", "Options relating to clothing"),
-					FNewToolMenuChoice(FNewMenuDelegate::CreateRaw(const_cast<SAnimViewportToolBar*>(this), &SAnimViewportToolBar::FillCharacterClothingMenu)),
-					bInOpenSubMenuOnClick, TAttribute<FSlateIcon>(), bShouldCloseWindowAfterMenuSelection);
-			}
+				UAnimViewportToolBarToolMenuContext* Context = InSection.FindContext<UAnimViewportToolBarToolMenuContext>();
+				const SAnimViewportToolBar* ContextThis = Context ? Context->AnimViewportToolBar.Pin().Get() : nullptr; 
+				UDebugSkelMeshComponent* PreviewComp = ContextThis && ContextThis->Viewport.IsValid() ? ContextThis->Viewport.Pin()->GetPreviewScene()->GetPreviewMeshComponent() : nullptr;
+				if (PreviewComp && GetDefault<UPersonaOptions>()->bExposeClothingSceneElementMenu)
+				{
+					constexpr bool bInOpenSubMenuOnClick = false;
+					constexpr bool bShouldCloseWindowAfterMenuSelection = false;
+					InSection.AddSubMenu(TEXT("ClothingSubMenu"),
+						LOCTEXT("CharacterMenu_ClothingSubMenu", "Clothing"),
+						LOCTEXT("CharacterMenu_ClothingSubMenuToolTip", "Options relating to clothing"),
+						FNewToolMenuChoice(FNewMenuDelegate::CreateRaw(const_cast<SAnimViewportToolBar *>(ContextThis), &SAnimViewportToolBar::FillCharacterClothingMenu)),
+						bInOpenSubMenuOnClick, TAttribute<FSlateIcon>(), bShouldCloseWindowAfterMenuSelection);
+				}
+			}));
 
 			Section.AddSubMenu(TEXT("AudioSubMenu"),
 				LOCTEXT("CharacterMenu_AudioSubMenu", "Audio"),
@@ -838,10 +851,18 @@ TSharedRef<SWidget> SAnimViewportToolBar::GenerateCharacterMenu() const
 				Section.AddMenuEntry(FAnimViewportShowCommands::Get().UseAudioAttenuation);
 			}));
 
-			Section.AddSubMenu(TEXT("AdvancedSubMenu"),
-				LOCTEXT("CharacterMenu_AdvancedSubMenu", "Advanced"),
-				LOCTEXT("CharacterMenu_AdvancedSubMenuToolTip", "Advanced options"),
-				FNewToolMenuChoice(FNewMenuDelegate::CreateRaw(const_cast<SAnimViewportToolBar*>(this), &SAnimViewportToolBar::FillCharacterAdvancedMenu)));
+			Section.AddDynamicEntry("AdvancedSubMenu", FNewToolMenuSectionDelegate::CreateLambda([](FToolMenuSection& InSection)
+			{
+				UAnimViewportToolBarToolMenuContext* Context = InSection.FindContext<UAnimViewportToolBarToolMenuContext>();
+				const SAnimViewportToolBar* ContextThis = Context ? Context->AnimViewportToolBar.Pin().Get() : nullptr;
+				if (ContextThis)
+				{
+					InSection.AddSubMenu(TEXT("AdvancedSubMenu"),
+						LOCTEXT("CharacterMenu_AdvancedSubMenu", "Advanced"),
+						LOCTEXT("CharacterMenu_AdvancedSubMenuToolTip", "Advanced options"),
+						FNewToolMenuChoice(FNewMenuDelegate::CreateRaw(ContextThis, &SAnimViewportToolBar::FillCharacterAdvancedMenu)));
+				}
+			}));
 		}
 	}
 
@@ -849,6 +870,9 @@ TSharedRef<SWidget> SAnimViewportToolBar::GenerateCharacterMenu() const
 	TSharedPtr<SAnimationEditorViewportTabBody> PinnedViewport = Viewport.Pin();
 	FToolMenuContext MenuContext(PinnedViewport->GetCommandList(), MenuExtender);
 	PinnedViewport->GetAssetEditorToolkit()->InitToolMenuContext(MenuContext);
+	UAnimViewportToolBarToolMenuContext* AnimViewportContext = NewObject<UAnimViewportToolBarToolMenuContext>();
+	AnimViewportContext->AnimViewportToolBar = SharedThis(this);
+	MenuContext.AddObject(AnimViewportContext);
 	return UToolMenus::Get()->GenerateWidget(MenuName, MenuContext);
 }
 
