Fix build error (due to CL 25773495)

diff --git a/Engine/Source/Editor/Persona/Private/SAnimAssetFindReplace.h b/Engine/Source/Editor/Persona/Private/SAnimAssetFindReplace.h
--- a/Engine/Source/Editor/Persona/Private/SAnimAssetFindReplace.h
+++ b/Engine/Source/Editor/Persona/Private/SAnimAssetFindReplace.h
@@ -89,7 +89,7 @@ private:
 
 	EAnimAssetFindReplaceMode Mode = EAnimAssetFindReplaceMode::Find;
 
-	TMap<TSubclassOf<UAnimAssetFindReplaceProcessor>, UAnimAssetFindReplaceProcessor*> Processors;
+	TMap<TSubclassOf<UAnimAssetFindReplaceProcessor>, TObjectPtr<UAnimAssetFindReplaceProcessor>> Processors;
 
 	UAnimAssetFindReplaceProcessor* CurrentProcessor = nullptr;
 
@@ -114,4 +114,4 @@ private:
 	TSharedPtr<SBox> ToolbarContainer;
 
 	TSharedPtr<SBox> FindReplaceWidgetContainer;
-};
\ No newline at end of file
+};
