Fix build errors

diff --git a/Engine/Source/Editor/Persona/Private/SAnimAssetFindReplace.h b/Engine/Source/Editor/Persona/Private/SAnimAssetFindReplace.h
--- a/Engine/Source/Editor/Persona/Private/SAnimAssetFindReplace.h
+++ b/Engine/Source/Editor/Persona/Private/SAnimAssetFindReplace.h
@@ -1,4 +1,4 @@
-﻿// Copyright Epic Games, Inc. All Rights Reserved.
+// Copyright Epic Games, Inc. All Rights Reserved.
 
 #pragma once
 
@@ -89,7 +89,7 @@ private:
 
 	EAnimAssetFindReplaceMode Mode = EAnimAssetFindReplaceMode::Find;
 
-	TMap<TSubclassOf<UAnimAssetFindReplaceProcessor>, TObjectPtr<UAnimAssetFindReplaceProcessor>> Processors;
+	TMap<TSubclassOf<UAnimAssetFindReplaceProcessor>, UAnimAssetFindReplaceProcessor*> Processors;
 
 	UAnimAssetFindReplaceProcessor* CurrentProcessor = nullptr;
 
