Fixed Animation Notifies Panel treating Sync Marker Notifies as regular Anim Notifies in some cases.

#rnx

diff --git a/Engine/Source/Editor/Persona/Private/SSkeletonAnimNotifies.cpp b/Engine/Source/Editor/Persona/Private/SSkeletonAnimNotifies.cpp
--- a/Engine/Source/Editor/Persona/Private/SSkeletonAnimNotifies.cpp
+++ b/Engine/Source/Editor/Persona/Private/SSkeletonAnimNotifies.cpp
@@ -357,7 +357,7 @@ bool SSkeletonAnimNotifies::CanPerformRename() const
 void SSkeletonAnimNotifies::OnAddItem(bool bIsSyncMarker)
 {
 	// Find a unique name for this notify
-	const TCHAR* BaseNotifyString = bIsSyncMarker ? TEXT("NewNotify") : TEXT("NewSyncMarker");
+	const TCHAR* BaseNotifyString = bIsSyncMarker ? TEXT("NewSyncMarker") : TEXT("NewNotify");
 	FString NewNotifyString = BaseNotifyString;
 	int32 NumericExtension = 0;
 
@@ -481,11 +481,11 @@ void SSkeletonAnimNotifies::OnNotifyNameCommitted( const FText& NewName, ETextCo
 	{
 		if(Item->bIsSyncMarker)
 		{
-			EditableSkeleton->AddNotify(NewFName);
+			EditableSkeleton->AddSyncMarker(NewFName);
 		}
 		else
 		{
-			EditableSkeleton->AddSyncMarker(NewFName);
+			EditableSkeleton->AddNotify(NewFName);
 		}
 		Item->bIsNew = false;
 	}
