Added the option to filter by "Sync Marker..." to the animation sequence browser.

diff --git a/Engine/Source/Editor/Persona/Private/SAnimationSequenceBrowser.cpp b/Engine/Source/Editor/Persona/Private/SAnimationSequenceBrowser.cpp
--- a/Engine/Source/Editor/Persona/Private/SAnimationSequenceBrowser.cpp
+++ b/Engine/Source/Editor/Persona/Private/SAnimationSequenceBrowser.cpp
@@ -206,6 +206,135 @@ public:
 	TSharedPtr<IEditableSkeleton> EditableSkeleton;
 };
 
+/** A filter that displays animations that use a skeleton sync marker */
+class FFrontendFilter_SkeletonSyncMarker : public FFrontendFilter
+{
+public:
+	FFrontendFilter_SkeletonSyncMarker(TSharedPtr<FFrontendFilterCategory> InCategory, TSharedPtr<IEditableSkeleton> InEditableSkeleton)
+		: FFrontendFilter(InCategory) 
+		, EditableSkeleton(InEditableSkeleton)
+	{}
+
+	// FFrontendFilter implementation
+	virtual FString GetName() const override 
+	{
+		return TEXT("SkeletonSyncMarkerAssets"); 
+	}
+
+	virtual FText GetDisplayName() const override 
+	{ 
+		return SyncMarkerString.IsEmpty() ? LOCTEXT("FFrontendFilter_SkeletonSyncMarkerAssetsEmpty", "Uses Sync Marker...") : FText::Format(LOCTEXT("FFrontendFilter_SkeletonSyncMarkerAssets", "Uses Sync Marker: {0}"), FText::FromString(SyncMarkerString));
+	}
+
+	virtual FText GetToolTipText() const override
+	{
+		return SyncMarkerString.IsEmpty() ? LOCTEXT("FFrontendFilter_SkeletonSyncMarkerAssetsEmptyTooltip", "Show assets that contains a (specified) sync marker") : FText::Format(LOCTEXT("FFrontendFilter_SkeletonSyncMarkerAssetsTooltip", "Show assets that use sync marker '{0}'"), FText::FromString(SyncMarkerString));
+	}
+
+	virtual void ModifyContextMenu(FMenuBuilder& MenuBuilder) override
+	{
+		MenuBuilder.BeginSection(TEXT("SyncMarkerSection"), LOCTEXT("SyncMarkerSectionHeading", "Choose Sync Marker"));
+
+		MenuBuilder.AddMenuEntry(LOCTEXT("AllNotifyFilterName", "Any Animation Sync Marker"),
+			LOCTEXT("AllSyncMarkerFilterName_ToolTip", "Consider all Sync Markers, rather than a specific one, for filtering."),
+			FSlateIcon(),
+			FUIAction(
+				FExecuteAction::CreateLambda([this]()
+				{
+					if (!SyncMarkerString.IsEmpty())
+					{
+						SyncMarkerString.Empty();
+					}
+
+					OnChanged().Broadcast();
+				}), 
+				FCanExecuteAction::CreateLambda([this]()
+				{
+					return !SyncMarkerString.IsEmpty();
+				}),
+				FGetActionCheckState::CreateLambda([this]()
+				{
+					return SyncMarkerString.IsEmpty() ? ECheckBoxState::Checked : ECheckBoxState::Unchecked;
+				})
+			), 
+			NAME_None,
+			EUserInterfaceActionType::ToggleButton
+		);
+		
+		TSharedRef<SWidget> Widget =
+			SNew(SBox)
+			.HeightOverride(300.0f)
+			.WidthOverride(200.0f)
+			[
+				SNew(SSkeletonAnimNotifies, EditableSkeleton.ToSharedRef())
+				.IsPicker(true)
+				.ShowNotifies(false)
+				.ShowSyncMarkers(true)
+				.OnItemSelected_Lambda([this](const FName& InNotifyName)
+				{
+					FSlateApplication::Get().DismissAllMenus();
+					SyncMarkerString = InNotifyName.ToString();
+					OnChanged().Broadcast();
+				})
+			];
+
+		MenuBuilder.AddWidget(Widget, FText(), true);
+
+		MenuBuilder.EndSection();
+	}
+
+	virtual FLinearColor GetColor() const override
+	{
+		return FLinearColor(0.62f, 0.43f, 0.0f, 1);
+	}
+
+	// IFilter implementation
+	virtual bool PassesFilter(FAssetFilterType InItem) const override
+	{
+		FAssetData ItemAssetData;
+		if (InItem.Legacy_TryGetAssetData(ItemAssetData))
+		{
+			const FString TagValue = ItemAssetData.GetTagValueRef<FString>(USkeleton::AnimSyncMarkerTag);
+			if (!TagValue.IsEmpty())
+			{
+				if (!SyncMarkerString.IsEmpty())
+				{
+					// Parse sync markers
+					TArray<FString> SyncMarkerValues;
+					if (TagValue.ParseIntoArray(SyncMarkerValues, *USkeleton::AnimSyncMarkerTagDelimiter, true) > 0)
+					{
+						for (const FString& SyncMarkerValue : SyncMarkerValues)
+						{
+							if (SyncMarkerValue == SyncMarkerString)
+							{
+								return true;
+							}
+						}
+					}
+				}
+
+				return SyncMarkerString.IsEmpty();
+			}
+		}
+
+		return false;
+	}
+
+	void SetSyncMarkerFilter(const FName& InNotifyFilter)
+	{
+		SyncMarkerString = InNotifyFilter.ToString();
+		OnChanged().Broadcast();
+		SetActive(true);
+	}
+
+public:
+	/** Sync marker string to use when filtering */
+	FString SyncMarkerString;
+
+	/** Editable skeleton used to access available sync markers */
+	TSharedPtr<IEditableSkeleton> EditableSkeleton;
+};
+
 /** A filter that displays animations that use a curve */
 class FFrontendFilter_Curves : public FFrontendFilter
 {
@@ -904,6 +1033,7 @@ void SAnimationSequenceBrowser::Construct(const FArguments& InArgs, const TShare
 	SkeletonNotifyFilter = MakeShareable(new FFrontendFilter_SkeletonNotify(AnimCategory, InPersonaToolkit->GetEditableSkeleton()));
 	Config.ExtraFrontendFilters.Add(SkeletonNotifyFilter.ToSharedRef());
 	Config.ExtraFrontendFilters.Add(MakeShareable(new FFrontendFilter_Curves(AnimCategory, InPersonaToolkit->GetEditableSkeleton())));
+	Config.ExtraFrontendFilters.Add(MakeShareable(new FFrontendFilter_SkeletonSyncMarker(AnimCategory, InPersonaToolkit->GetEditableSkeleton())));
 
 	Config.OnIsAssetValidForCustomToolTip = FOnIsAssetValidForCustomToolTip::CreateLambda([](const FAssetData& AssetData) {return AssetData.IsAssetLoaded(); });
 	Config.OnGetCustomAssetToolTip = FOnGetCustomAssetToolTip::CreateSP(this, &SAnimationSequenceBrowser::CreateCustomAssetToolTip);
