Fix for crash when filtering preview mesh on template anim bps

diff --git a/Engine/Source/Editor/Persona/Private/PersonaModule.cpp b/Engine/Source/Editor/Persona/Private/PersonaModule.cpp
--- a/Engine/Source/Editor/Persona/Private/PersonaModule.cpp
+++ b/Engine/Source/Editor/Persona/Private/PersonaModule.cpp
@@ -940,7 +940,7 @@ namespace UE::Persona::Private
 			AssetPickerConfig.Filter.ClassPaths.Add(USkeletalMesh::StaticClass()->GetClassPathName());
 			AssetPickerConfig.OnShouldFilterAsset = FOnShouldFilterAsset::CreateLambda([InWeakPersonaToolkit](const FAssetData& AssetData)
 			{
-				if (InWeakPersonaToolkit.IsValid())
+				if (InWeakPersonaToolkit.IsValid() && InWeakPersonaToolkit.Pin()->GetSkeleton())
 				{
 					if (InWeakPersonaToolkit.Pin()->GetContext() == UPhysicsAsset::StaticClass()->GetFName())
 					{
@@ -1006,7 +1006,7 @@ namespace UE::Persona::Private
 			AssetPickerConfig.Filter.ClassPaths.Add(UAnimationAsset::StaticClass()->GetClassPathName());
 			AssetPickerConfig.OnShouldFilterAsset = FOnShouldFilterAsset::CreateLambda([InWeakPersonaToolkit](const FAssetData& AssetData)
 			{
-				if (InWeakPersonaToolkit.IsValid())
+				if (InWeakPersonaToolkit.IsValid() && InWeakPersonaToolkit.Pin()->GetSkeleton())
 				{
 					FString TagValue;
 					if (AssetData.GetTagValue("Skeleton", TagValue))
