Fix duplicate loc tag

diff --git a/Engine/Source/Editor/Persona/Private/SAnimationSequenceBrowser.cpp b/Engine/Source/Editor/Persona/Private/SAnimationSequenceBrowser.cpp
--- a/Engine/Source/Editor/Persona/Private/SAnimationSequenceBrowser.cpp
+++ b/Engine/Source/Editor/Persona/Private/SAnimationSequenceBrowser.cpp
@@ -235,7 +235,7 @@ public:
 	{
 		MenuBuilder.BeginSection(TEXT("SyncMarkerSection"), LOCTEXT("SyncMarkerSectionHeading", "Choose Sync Marker"));
 
-		MenuBuilder.AddMenuEntry(LOCTEXT("AllNotifyFilterName", "Any Animation Sync Marker"),
+		MenuBuilder.AddMenuEntry(LOCTEXT("AllSyncMarkerFilterName", "Any Animation Sync Marker"),
 			LOCTEXT("AllSyncMarkerFilterName_ToolTip", "Consider all Sync Markers, rather than a specific one, for filtering."),
 			FSlateIcon(),
 			FUIAction(
