Skeletal Mesh: When adding/removing vertex attributes, ensure the asset detail view is synchronized with the LOD info list of vertex attributes. Also fixes other undo issues with Bones to Prioritize, Bones to Remove.

#rnx

diff --git a/Engine/Source/Editor/Persona/Private/PersonaMeshDetails.cpp b/Engine/Source/Editor/Persona/Private/PersonaMeshDetails.cpp
--- a/Engine/Source/Editor/Persona/Private/PersonaMeshDetails.cpp
+++ b/Engine/Source/Editor/Persona/Private/PersonaMeshDetails.cpp
@@ -2125,6 +2125,11 @@ FPersonaMeshDetails::~FPersonaMeshDetails()
 		PreviewScene->UnregisterOnPreviewMeshChanged(this);
 	}
 
+	if (SkeletalMeshPtr.Get())
+	{
+		SkeletalMeshPtr.Get()->OnPostMeshCached().RemoveAll(this);
+	}
+
 	GEditor->GetEditorSubsystem<UImportSubsystem>()->OnAssetPostLODImport.RemoveAll(this);
 	GEditor->GetEditorSubsystem<UImportSubsystem>()->OnAssetReimport.RemoveAll(this);
 }
@@ -4367,6 +4372,12 @@ void FPersonaMeshDetails::CustomizeDetails( IDetailLayoutBuilder& DetailLayout )
 	IDetailCategoryBuilder& ClothingCategory = DetailLayout.EditCategory("Clothing", FText::GetEmpty(), ECategoryPriority::TypeSpecific);
 	CustomizeClothingProperties(DetailLayout,ClothingCategory);
 
+	// Get notified when the mesh has finished building
+	if (SkeletalMeshPtr.Get())
+	{
+		SkeletalMeshPtr.Get()->OnPostMeshCached().AddSP(this, &FPersonaMeshDetails::OnMeshRebuildCompleted);
+	}
+
 	// Post process selector
 	IDetailCategoryBuilder& SkelMeshCategory = DetailLayout.EditCategory("SkeletalMesh");
 	TSharedRef<IPropertyHandle> PostProcessHandle = DetailLayout.GetProperty(USkeletalMesh::GetPostProcessAnimBlueprintMemberName(), USkeletalMesh::StaticClass());
@@ -4479,6 +4490,15 @@ bool FPersonaMeshDetails::GetVertexOverrideColorEnabledState() const
 	return (VertexColorImportOption == EVertexColorImportOption::Override);
 }
 
+
+void FPersonaMeshDetails::OnMeshRebuildCompleted(USkeletalMesh* InMesh)
+{
+	if (InMesh == SkeletalMeshPtr.Get())
+	{
+		RefreshMeshDetailLayout();
+	}
+}
+
 void FPersonaMeshDetails::HideUnnecessaryProperties(IDetailLayoutBuilder& DetailLayout)
 {
 	// LODInfo doesn't need to be showed anymore because it was moved to each LOD category
diff --git a/Engine/Source/Editor/Persona/Private/PersonaMeshDetails.h b/Engine/Source/Editor/Persona/Private/PersonaMeshDetails.h
--- a/Engine/Source/Editor/Persona/Private/PersonaMeshDetails.h
+++ b/Engine/Source/Editor/Persona/Private/PersonaMeshDetails.h
@@ -837,4 +837,7 @@ private:
 
 	// Delegate used at runtime to determine the state of the VertexOverrideColor property
 	bool GetVertexOverrideColorEnabledState() const;
+
+	// Called when the skeletal mesh has finished rebuilding. This may affect some settings, such as vertex attributes.
+	void OnMeshRebuildCompleted(USkeletalMesh* InMesh);
 };
