Fix 'Add Curve' adding a metadata curve rather than a variable one

diff --git a/Engine/Source/Editor/Persona/Private/AnimTimeline/AnimTimelineTrack_Curves.cpp b/Engine/Source/Editor/Persona/Private/AnimTimeline/AnimTimelineTrack_Curves.cpp
--- a/Engine/Source/Editor/Persona/Private/AnimTimeline/AnimTimelineTrack_Curves.cpp
+++ b/Engine/Source/Editor/Persona/Private/AnimTimeline/AnimTimelineTrack_Curves.cpp
@@ -304,7 +304,7 @@ void FAnimTimelineTrack_Curves::OnVariableCurveNamePicked(const FName& InCurveNa
 	
 	if(InCurveName != NAME_None)
 	{
-		AddMetadataEntry(InCurveName);
+		AddVariableCurve(InCurveName);
 	}
 }
 
