PR #10519: Fixed typo in AnimationEditorPreviewScene.cpp

diff --git a/Engine/Source/Editor/Persona/Private/AnimationEditorPreviewScene.cpp b/Engine/Source/Editor/Persona/Private/AnimationEditorPreviewScene.cpp
--- a/Engine/Source/Editor/Persona/Private/AnimationEditorPreviewScene.cpp
+++ b/Engine/Source/Editor/Persona/Private/AnimationEditorPreviewScene.cpp
@@ -1171,7 +1171,7 @@ void FAnimationEditorPreviewScene::SetAllowMeshHitProxies(bool bState)
 
 void FAnimationEditorPreviewScene::FlagTickable()
 {
-	// Set the last tick time so we tick kwhen we are visible in a viewport
+	// Set the last tick time so we tick when we are visible in a viewport
 	LastTickTime = FPlatformTime::Seconds();
 }
 
