Allow editing/visualizing properties in debug instance AnimBP.

diff --git a/Engine/Source/Editor/Persona/Private/TabSpawners.cpp b/Engine/Source/Editor/Persona/Private/TabSpawners.cpp
--- a/Engine/Source/Editor/Persona/Private/TabSpawners.cpp
+++ b/Engine/Source/Editor/Persona/Private/TabSpawners.cpp
@@ -375,9 +375,11 @@ TSharedRef<SWidget> FRetargetSourcesTabSummoner::CreateTabBody(const FWorkflowTa
 /////////////////////////////////////////////////////
 // SPersonaPreviewPropertyEditor
 
-void SPersonaPreviewPropertyEditor::Construct(const FArguments& InArgs, TSharedRef<IPersonaPreviewScene> InPreviewScene)
+void SPersonaPreviewPropertyEditor::Construct(const FArguments& InArgs, TSharedRef<IPersonaPreviewScene> InPreviewScene, TSharedRef<FBlueprintEditor> InBPEditor)
 {
 	PreviewScene = InPreviewScene;
+	BPEditor = InBPEditor;
+
 	bPropertyEdited = false;
 
 	SSingleObjectDetailsPanel::Construct(SSingleObjectDetailsPanel::FArguments(), /*bAutomaticallyObserveViaGetObjectToObserve*/ true, /*bAllowSearch*/ true);
@@ -388,6 +390,16 @@ void SPersonaPreviewPropertyEditor::Construct(const FArguments& InArgs, TSharedR
 
 UObject* SPersonaPreviewPropertyEditor::GetObjectToObserve() const
 {
+	// Use the selected debug instance, if we have one.
+	if (UBlueprint* Blueprint = BPEditor.Pin()->GetBlueprintObj())
+	{
+		if (UObject* DebugObject = Blueprint->GetObjectBeingDebugged())
+		{
+			return DebugObject;
+		}
+	}
+
+	// If there's no selection, fall-back to preview instance.
 	if (UDebugSkelMeshComponent* PreviewMeshComponent = PreviewScene.Pin()->GetPreviewMeshComponent())
 	{
 		if (PreviewMeshComponent->GetAnimInstance() != nullptr)
@@ -421,7 +433,7 @@ TSharedRef<SWidget> SPersonaPreviewPropertyEditor::PopulateSlot(TSharedRef<SWidg
 				.Padding(2.0f)
 				[
 					SNew(STextBlock)
-					.Text(LOCTEXT("AnimBlueprintEditPreviewText", "Changes made to preview only. Changes will not be saved!"))
+					.Text(LOCTEXT("AnimBlueprintEditPreviewText", "Changes made to the selected instance only. Changes will not be saved!"))
 					.ColorAndOpacity(FLinearColor::Yellow)
 					.ShadowOffset(FVector2D::UnitVector)
 					.AutoWrapText(true)
@@ -434,7 +446,7 @@ TSharedRef<SWidget> SPersonaPreviewPropertyEditor::PopulateSlot(TSharedRef<SWidg
 				[
 					SNew(SButton)
 					.OnClicked(this, &SPersonaPreviewPropertyEditor::HandleApplyChanges)
-					.ToolTipText(LOCTEXT("AnimBlueprintEditApplyChanges_Tooltip", "Apply any changes that have been made to the preview to the defaults."))
+					.ToolTipText(LOCTEXT("AnimBlueprintEditApplyChanges_Tooltip", "Apply the selected instance's current property values to the defaults."))
 					[
 						SNew(STextBlock)
 						.Text(LOCTEXT("AnimBlueprintEditApplyChanges", "Apply"))
@@ -446,15 +458,28 @@ TSharedRef<SWidget> SPersonaPreviewPropertyEditor::PopulateSlot(TSharedRef<SWidg
 
 void SPersonaPreviewPropertyEditor::HandlePropertyChanged(const FPropertyChangedEvent& InPropertyChangedEvent)
 {
-	if (UDebugSkelMeshComponent* PreviewMeshComponent = PreviewScene.Pin()->GetPreviewMeshComponent())
+	UAnimInstance* AnimInstance = nullptr;
+
+	if (UBlueprint* Blueprint = BPEditor.Pin()->GetBlueprintObj())
 	{
-		if (UAnimInstance* AnimInstance = PreviewMeshComponent->GetAnimInstance())
+		if (UAnimInstance* DebugInstance = Cast<UAnimInstance>(Blueprint->GetObjectBeingDebugged()))
+		{
+			// Compare the selected debug instance (if we have one) against the CDO.
+			AnimInstance = DebugInstance;
+		}
+		else if (UDebugSkelMeshComponent* PreviewMeshComponent = PreviewScene.Pin()->GetPreviewMeshComponent())
 		{
-			// check to see how many properties have changed
-			const int32 NumChangedProperties = PersonaUtils::CopyPropertiesToCDO(AnimInstance, PersonaUtils::FCopyOptions(PersonaUtils::ECopyOptions::PreviewOnly));
-			bPropertyEdited = (NumChangedProperties > 0);
+			// If there's no selection, fall-back to preview instance.
+			AnimInstance = PreviewMeshComponent->GetAnimInstance();
 		}
 	}
+
+	if (AnimInstance != nullptr)
+	{
+		// check to see how many properties have changed
+		const int32 NumChangedProperties = PersonaUtils::CopyPropertiesToCDO(AnimInstance, PersonaUtils::FCopyOptions(PersonaUtils::ECopyOptions::PreviewOnly));
+		bPropertyEdited = (NumChangedProperties > 0);
+	}
 }
 
 FReply SPersonaPreviewPropertyEditor::HandleApplyChanges()
@@ -512,7 +537,7 @@ TSharedRef<SWidget> FAnimBlueprintPreviewEditorSummoner::CreateTabBody(const FWo
 																	TEXT("AnimBlueprintPropertyEditorPreviewMode")))
 					[
 						SNew( STextBlock )
-						.Text( LOCTEXT("AnimBlueprintDefaultsPreviewMode", "Edit Preview") )
+						.Text( LOCTEXT("AnimBlueprintDefaultsPreviewMode", "Edit Selected Instance") )
 					]
 				]
 				+SHorizontalBox::Slot()
@@ -543,7 +568,7 @@ TSharedRef<SWidget> FAnimBlueprintPreviewEditorSummoner::CreateTabBody(const FWo
 					.BorderImage( FAppStyle::GetBrush("NoBorder") )
 					.Visibility(this, &FAnimBlueprintPreviewEditorSummoner::IsEditorVisible, EAnimBlueprintEditorMode::PreviewMode)
 					[
-						SNew(SPersonaPreviewPropertyEditor, PreviewScene.Pin().ToSharedRef())
+						SNew(SPersonaPreviewPropertyEditor, PreviewScene.Pin().ToSharedRef(), BlueprintEditor.Pin().ToSharedRef())
 					]
 				]
 				+SOverlay::Slot()
@@ -561,7 +586,7 @@ TSharedRef<SWidget> FAnimBlueprintPreviewEditorSummoner::CreateTabBody(const FWo
 
 FText FAnimBlueprintPreviewEditorSummoner::GetTabToolTipText(const FWorkflowTabSpawnInfo& Info) const
 {
-	return LOCTEXT("AnimBlueprintPreviewEditorTooltip", "The editor lets you change the values of the preview instance");
+	return LOCTEXT("AnimBlueprintPreviewEditorTooltip", "The editor lets you change the values of the selected instance");
 }
 
 EVisibility FAnimBlueprintPreviewEditorSummoner::IsEditorVisible(EAnimBlueprintEditorMode::Type Mode) const
diff --git a/Engine/Source/Editor/Persona/Private/TabSpawners.h b/Engine/Source/Editor/Persona/Private/TabSpawners.h
--- a/Engine/Source/Editor/Persona/Private/TabSpawners.h
+++ b/Engine/Source/Editor/Persona/Private/TabSpawners.h
@@ -271,9 +271,11 @@ public:
 private:
 	// Pointer to preview scene
 	TWeakPtr<IPersonaPreviewScene> PreviewScene;
+	TWeakPtr<FBlueprintEditor> BPEditor;
 
 public:
-	void Construct(const FArguments& InArgs, TSharedRef<IPersonaPreviewScene> InPreviewScene);
+
+	void Construct(const FArguments& InArgs, TSharedRef<IPersonaPreviewScene> InPreviewScene, TSharedRef<FBlueprintEditor> InBPEditor);
 
 	// SSingleObjectDetailsPanel interface
 	virtual UObject* GetObjectToObserve() const override;
