Fix details panel disappearing on curve metadata edits

Selection was getting reset after edits by the list view rebuild, so adding a check when the selection is set directly to skip that case

diff --git a/Engine/Source/Editor/Persona/Private/SAnimCurveMetadataEditor.cpp b/Engine/Source/Editor/Persona/Private/SAnimCurveMetadataEditor.cpp
--- a/Engine/Source/Editor/Persona/Private/SAnimCurveMetadataEditor.cpp
+++ b/Engine/Source/Editor/Persona/Private/SAnimCurveMetadataEditor.cpp
@@ -920,36 +920,40 @@ bool SAnimCurveMetadataEditor::CanPaste() const
 
 void SAnimCurveMetadataEditor::OnSelectionChanged(TSharedPtr<FAnimCurveMetadataEditorItem> InItem, ESelectInfo::Type SelectInfo)
 {
-	if(UAnimCurveMetaData* AnimCurveMetaData = GetAnimCurveMetaData())
+	if(SelectInfo != ESelectInfo::Direct)
 	{
-		// make sure the currently selected ones are refreshed if it's first time
-		TArray<UObject*> SelectedObjects;
 
-		TArray< TSharedPtr< FAnimCurveMetadataEditorItem > > SelectedRows = AnimCurveListView->GetSelectedItems();
-		for (auto ItemIt = SelectedRows.CreateIterator(); ItemIt; ++ItemIt)
+		if(UAnimCurveMetaData* AnimCurveMetaData = GetAnimCurveMetaData())
 		{
-			TSharedPtr< FAnimCurveMetadataEditorItem > RowItem = (*ItemIt);
-			UEditorAnimCurveBoneLinks* EditorMirrorObj = RowItem->EditorMirrorObject;
-			if (RowItem == InItem)
+			// make sure the currently selected ones are refreshed if it's first time
+			TArray<UObject*> SelectedObjects;
+
+			TArray< TSharedPtr< FAnimCurveMetadataEditorItem > > SelectedRows = AnimCurveListView->GetSelectedItems();
+			for (auto ItemIt = SelectedRows.CreateIterator(); ItemIt; ++ItemIt)
 			{
-				// first time selected, refresh
-				TArray<FBoneReference> BoneLinks;
-				FName CurrentName = RowItem->CurveName;
-				const FCurveMetaData* CurveMetaData = AnimCurveMetaData->GetCurveMetaData(CurrentName);
-				uint8 MaxLOD = 0xFF;
-				if (CurveMetaData)
+				TSharedPtr< FAnimCurveMetadataEditorItem > RowItem = (*ItemIt);
+				UEditorAnimCurveBoneLinks* EditorMirrorObj = RowItem->EditorMirrorObject;
+				if (RowItem == InItem)
 				{
-					BoneLinks = CurveMetaData->LinkedBones;
-					MaxLOD = CurveMetaData->MaxLOD;
+					// first time selected, refresh
+					TArray<FBoneReference> BoneLinks;
+					FName CurrentName = RowItem->CurveName;
+					const FCurveMetaData* CurveMetaData = AnimCurveMetaData->GetCurveMetaData(CurrentName);
+					uint8 MaxLOD = 0xFF;
+					if (CurveMetaData)
+					{
+						BoneLinks = CurveMetaData->LinkedBones;
+						MaxLOD = CurveMetaData->MaxLOD;
+					}
+
+					EditorMirrorObj->Refresh(CurrentName, BoneLinks, MaxLOD);
 				}
 
-				EditorMirrorObj->Refresh(CurrentName, BoneLinks, MaxLOD);
+				SelectedObjects.Add(EditorMirrorObj);
 			}
 
-			SelectedObjects.Add(EditorMirrorObj);
+			OnObjectsSelected.ExecuteIfBound(SelectedObjects);
 		}
-
-		OnObjectsSelected.ExecuteIfBound(SelectedObjects);
 	}
 }
 
@@ -957,7 +961,7 @@ void SAnimCurveMetadataEditor::ApplyCurveBoneLinks(UEditorAnimCurveBoneLinks* Ed
 {
 	if(UAnimCurveMetaData* AnimCurveMetaData = GetAnimCurveMetaData())
 	{
-		if (EditorObj)
+		if (EditorObj && !GIsTransacting)
 		{
 			AnimCurveMetaData->SetCurveMetaDataBoneLinks(EditorObj->CurveName, EditorObj->ConnectedBones, EditorObj->MaxLOD, GetSkeleton());
 		}
