Cloth - Fixed crash in the Skeletal Mesh editor Character Clothing menu when the Chaos Cloth plugin is disabled.

diff --git a/Engine/Source/Editor/Persona/Private/SAnimViewportToolBar.cpp b/Engine/Source/Editor/Persona/Private/SAnimViewportToolBar.cpp
--- a/Engine/Source/Editor/Persona/Private/SAnimViewportToolBar.cpp
+++ b/Engine/Source/Editor/Persona/Private/SAnimViewportToolBar.cpp
@@ -986,20 +986,22 @@ void SAnimViewportToolBar::FillCharacterClothingMenu(FMenuBuilder& MenuBuilder)
 
 	// Call into the clothing editor module to customize the menu (this is mainly for debug visualizations and sim-specific options)
 	TSharedPtr<SAnimationEditorViewportTabBody> SharedViewport = Viewport.Pin();
-	if(SharedViewport.IsValid())
+	if (SharedViewport.IsValid())
 	{
 		TSharedRef<IPersonaPreviewScene> PreviewScene = SharedViewport->GetAnimationViewportClient()->GetPreviewScene();
-		if(UDebugSkelMeshComponent* PreviewComponent = PreviewScene->GetPreviewMeshComponent())
+		if (UDebugSkelMeshComponent* PreviewComponent = PreviewScene->GetPreviewMeshComponent())
 		{
-			FClothingSystemEditorInterfaceModule& ClothingEditorModule = FModuleManager::LoadModuleChecked<FClothingSystemEditorInterfaceModule>(TEXT("ClothingSystemEditorInterface"));
-
-			if(ISimulationEditorExtender* Extender = ClothingEditorModule.GetSimulationEditorExtender(PreviewComponent->ClothingSimulationFactory->GetFName()))
+			if (PreviewComponent->ClothingSimulationFactory)  // The cloth plugin could be disabled, and the factory would be null in this case
 			{
-				Extender->ExtendViewportShowMenu(MenuBuilder, PreviewScene);
+				FClothingSystemEditorInterfaceModule& ClothingEditorModule = FModuleManager::LoadModuleChecked<FClothingSystemEditorInterfaceModule>(TEXT("ClothingSystemEditorInterface"));
+
+				if (ISimulationEditorExtender* Extender = ClothingEditorModule.GetSimulationEditorExtender(PreviewComponent->ClothingSimulationFactory->GetFName()))
+				{
+					Extender->ExtendViewportShowMenu(MenuBuilder, PreviewScene);
+				}
 			}
 		}
 	}
-
 }
 
 TSharedRef<SWidget> SAnimViewportToolBar::GenerateShowMenu() const
