Added DoubleClick support to AnimSegmentsPanel, in order to jump to the target asset (shift + DoubleClick opens in new window)

diff --git a/Engine/Source/Editor/Persona/Private/SAnimSegmentsPanel.cpp b/Engine/Source/Editor/Persona/Private/SAnimSegmentsPanel.cpp
--- a/Engine/Source/Editor/Persona/Private/SAnimSegmentsPanel.cpp
+++ b/Engine/Source/Editor/Persona/Private/SAnimSegmentsPanel.cpp
@@ -36,13 +36,14 @@ void SAnimSegmentsPanel::Construct(const FArguments& InArgs)
 	ViewInputMin = InArgs._ViewInputMin;
 	ViewInputMax = InArgs._ViewInputMax;
 
-	OnAnimSegmentNodeClickedDelegate	= InArgs._OnAnimSegmentNodeClicked;
-	OnPreAnimUpdateDelegate				= InArgs._OnPreAnimUpdate;
-	OnPostAnimUpdateDelegate			= InArgs._OnPostAnimUpdate;
-	OnAnimSegmentRemovedDelegate		= InArgs._OnAnimSegmentRemoved;
-	OnAnimReplaceMapping				= InArgs._OnAnimReplaceMapping;
-	OnDiffFromParentAsset				= InArgs._OnDiffFromParentAsset;
-	OnGetNodeColor						= InArgs._OnGetNodeColor;
+	OnAnimSegmentNodeClickedDelegate		= InArgs._OnAnimSegmentNodeClicked;
+	OnAnimSegmentNodeDoubleClickedDelegate	= InArgs._OnAnimSegmentNodeDoubleClicked;
+	OnPreAnimUpdateDelegate					= InArgs._OnPreAnimUpdate;
+	OnPostAnimUpdateDelegate				= InArgs._OnPostAnimUpdate;
+	OnAnimSegmentRemovedDelegate			= InArgs._OnAnimSegmentRemoved;
+	OnAnimReplaceMapping					= InArgs._OnAnimReplaceMapping;
+	OnDiffFromParentAsset					= InArgs._OnDiffFromParentAsset;
+	OnGetNodeColor							= InArgs._OnGetNodeColor;
 
 	bChildAnimMontage = InArgs._bChildAnimMontage;
 
@@ -147,6 +148,7 @@ void SAnimSegmentsPanel::Construct(const FArguments& InArgs)
 				.OnTrackNodeDropped(this, &SAnimSegmentsPanel::OnSegmentDropped, SegmentIdx)
 				.OnNodeRightClickContextMenu(this, &SAnimSegmentsPanel::SummonSegmentNodeContextMenu, SegmentIdx)
 				.OnTrackNodeClicked(this, &SAnimSegmentsPanel::OnAnimSegmentNodeClicked, SegmentIdx)
+				.OnTrackNodeDoubleClicked(this, &SAnimSegmentsPanel::OnAnimSegmentNodeDoubleClicked, SegmentIdx)
 				.NodeSelectionSet(InArgs._NodeSelectionSet)
 			);
 		}
@@ -513,6 +515,18 @@ void SAnimSegmentsPanel::OnAnimSegmentNodeClicked(int32 SegmentIdx)
 	OnAnimSegmentNodeClickedDelegate.ExecuteIfBound(SegmentIdx);
 }
 
+void SAnimSegmentsPanel::OnAnimSegmentNodeDoubleClicked(int32 SegmentIdx)
+{
+	if (OnAnimSegmentNodeDoubleClickedDelegate.IsBound())
+	{
+		OnAnimSegmentNodeDoubleClickedDelegate.ExecuteIfBound(SegmentIdx);
+	}
+	else
+	{
+		OpenAsset(SegmentIdx);
+	}
+}
+
 void SAnimSegmentsPanel::RemoveSelectedAnimSegments()
 {
 	TArray<int32> SelectedNodeIndices;
diff --git a/Engine/Source/Editor/Persona/Private/SAnimSegmentsPanel.h b/Engine/Source/Editor/Persona/Private/SAnimSegmentsPanel.h
--- a/Engine/Source/Editor/Persona/Private/SAnimSegmentsPanel.h
+++ b/Engine/Source/Editor/Persona/Private/SAnimSegmentsPanel.h
@@ -23,6 +23,7 @@ class USkeleton;
 DECLARE_DELEGATE( FOnPreAnimUpdate )
 DECLARE_DELEGATE( FOnPostAnimUpdate )
 DECLARE_DELEGATE_OneParam( FOnAnimSegmentNodeClicked, int32 )
+DECLARE_DELEGATE_OneParam( FOnAnimSegmentNodeDoubleClicked, int32)
 DECLARE_DELEGATE_OneParam( FOnAnimSegmentRemoved, int32 )
 DECLARE_DELEGATE_FourParams( FOnAnimReplaceMapping, FName, int32, UAnimSequenceBase*, UAnimSequenceBase*)
 DECLARE_DELEGATE_RetVal_ThreeParams(bool, FOnDiffFromParentAsset, FName, int32, const FAnimSegment& )
@@ -73,6 +74,7 @@ public:
 	SLATE_ATTRIBUTE(int32, TrackNumDiscreteValues)
 
 	SLATE_EVENT(FOnAnimSegmentNodeClicked, OnAnimSegmentNodeClicked)
+	SLATE_EVENT(FOnAnimSegmentNodeDoubleClicked, OnAnimSegmentNodeDoubleClicked)
 	SLATE_EVENT(FOnPreAnimUpdate, OnPreAnimUpdate)
 	SLATE_EVENT(FOnPostAnimUpdate, OnPostAnimUpdate)
 	SLATE_EVENT(FOnAnimSegmentRemoved, OnAnimSegmentRemoved)
@@ -98,13 +100,14 @@ public:
 
 private:
 
-	FOnPreAnimUpdate			OnPreAnimUpdateDelegate;
-	FOnPostAnimUpdate			OnPostAnimUpdateDelegate;
-	FOnAnimSegmentNodeClicked	OnAnimSegmentNodeClickedDelegate;
-	FOnAnimSegmentRemoved		OnAnimSegmentRemovedDelegate;
-	FOnAnimReplaceMapping		OnAnimReplaceMapping;
-	FOnDiffFromParentAsset		OnDiffFromParentAsset;
-	FOnGetNodeColor				OnGetNodeColor;
+	FOnPreAnimUpdate				OnPreAnimUpdateDelegate;
+	FOnPostAnimUpdate				OnPostAnimUpdateDelegate;
+	FOnAnimSegmentNodeClicked		OnAnimSegmentNodeClickedDelegate;
+	FOnAnimSegmentNodeDoubleClicked OnAnimSegmentNodeDoubleClickedDelegate;
+	FOnAnimSegmentRemoved			OnAnimSegmentRemovedDelegate;
+	FOnAnimReplaceMapping			OnAnimReplaceMapping;
+	FOnDiffFromParentAsset			OnDiffFromParentAsset;
+	FOnGetNodeColor					OnGetNodeColor;
 
 	enum ETrackViewStyle
 	{
@@ -131,6 +134,7 @@ private:
 	bool				IsValidToAdd(UAnimSequenceBase* NewSequenceBase) const;
 	void				OnTrackDragDrop( TSharedPtr<FDragDropOperation> DragDropOp, float DataPos );
 	void				OnAnimSegmentNodeClicked(int32 SegmentIdx);
+	void				OnAnimSegmentNodeDoubleClicked(int32 SegmentIdx);
 
 	// child anim montage
 	void				ReplaceAnimSegment(UAnimSequenceBase* NewSequenceBase, float NewStartPos);
diff --git a/Engine/Source/Editor/Persona/Private/STrack.cpp b/Engine/Source/Editor/Persona/Private/STrack.cpp
--- a/Engine/Source/Editor/Persona/Private/STrack.cpp
+++ b/Engine/Source/Editor/Persona/Private/STrack.cpp
@@ -112,6 +112,7 @@ void STrackNode::Construct(const FArguments& InArgs)
 	OnTrackNodeDropped = InArgs._OnTrackNodeDropped;
 	OnNodeRightClickContextMenu = InArgs._OnNodeRightClickContextMenu;
 	OnTrackNodeClicked = InArgs._OnTrackNodeClicked;
+	OnTrackNodeDoubleClicked = InArgs._OnTrackNodeDoubleClicked;
 	bCenterOnPosition = InArgs._CenterOnPosition;
 	
 	NodeSelectionSet = InArgs._NodeSelectionSet;
@@ -227,6 +228,18 @@ FReply STrackNode::OnMouseButtonUp( const FGeometry& MyGeometry, const FPointerE
 	return FReply::Unhandled();
 }
 
+FReply STrackNode::OnMouseButtonDoubleClick(const FGeometry& MyGeometry, const FPointerEvent& MouseEvent)
+{
+	if (MouseEvent.GetEffectingButton() == EKeys::LeftMouseButton)
+	{
+		OnTrackNodeDoubleClicked.ExecuteIfBound();
+
+		return FReply::Handled();
+	}
+
+	return FReply::Unhandled();
+}
+
 FReply STrackNode::OnMouseMove( const FGeometry& MyGeometry, const FPointerEvent& MouseEvent )
 {
 	return FReply::Unhandled();
diff --git a/Engine/Source/Editor/Persona/Private/STrack.h b/Engine/Source/Editor/Persona/Private/STrack.h
--- a/Engine/Source/Editor/Persona/Private/STrack.h
+++ b/Engine/Source/Editor/Persona/Private/STrack.h
@@ -52,6 +52,7 @@ DECLARE_DELEGATE_RetVal( FString, FOnGetNodeName )
 DECLARE_DELEGATE_OneParam( FOnTrackNodeDragged, float )
 DECLARE_DELEGATE( FOnTrackNodeDropped )
 DECLARE_DELEGATE( FOnTrackNodeClicked )
+DECLARE_DELEGATE( FOnTrackNodeDoubleClicked )
 
 DECLARE_DELEGATE_RetVal_TwoParams( TSharedPtr<SWidget>, FOnSummonContextMenu, const FGeometry&, const FPointerEvent& );
 DECLARE_DELEGATE_ThreeParams( FOnTrackRightClickContextMenu, FMenuBuilder&, float, int32 )
@@ -112,6 +113,7 @@ public:
 		, _OnSelectionChanged()
 		, _OnNodeRightClickContextMenu()
 		, _OnTrackNodeClicked()
+		, _OnTrackNodeDoubleClicked()
 		, _CenterOnPosition(false)
 	{
 	}
@@ -131,6 +133,7 @@ public:
 	SLATE_EVENT( FOnNodeSelectionChanged, OnSelectionChanged )
 	SLATE_EVENT( FOnNodeRightClickContextMenu, OnNodeRightClickContextMenu )
 	SLATE_EVENT( FOnTrackNodeClicked, OnTrackNodeClicked )
+	SLATE_EVENT( FOnTrackNodeDoubleClicked, OnTrackNodeDoubleClicked)
 	SLATE_ARGUMENT( bool, CenterOnPosition )
 	SLATE_NAMED_SLOT(FArguments, OverrideContent)
 
@@ -144,6 +147,7 @@ public:
 	virtual void OnMouseEnter( const FGeometry& MyGeometry, const FPointerEvent& MouseEvent ) override;
 	virtual void OnMouseLeave( const FPointerEvent& MouseEvent ) override;
 	virtual void OnDragged(const class FDragDropEvent& DragDropEvent );
+	virtual FReply OnMouseButtonDoubleClick(const FGeometry& InMyGeometry, const FPointerEvent& InMouseEvent) override;
 
 	// virtual draw related functions
 	virtual FVector2D GetOffsetRelativeToParent(const FGeometry& ParentAllottedGeometry) const;
@@ -211,6 +215,7 @@ protected:
 	FOnTrackNodeDropped			OnTrackNodeDropped;
 	
 	FOnTrackNodeClicked				OnTrackNodeClicked;
+	FOnTrackNodeDoubleClicked		OnTrackNodeDoubleClicked;
 	FOnNodeRightClickContextMenu	OnNodeRightClickContextMenu;
 	
 	FSlateFontInfo Font;
