Added the ability to replace sync markers on the animation timeline

Also tweaked the various context menus to reduce nesting levels

diff --git a/Engine/Source/Editor/Persona/Private/SAnimNotifyPanel.cpp b/Engine/Source/Editor/Persona/Private/SAnimNotifyPanel.cpp
--- a/Engine/Source/Editor/Persona/Private/SAnimNotifyPanel.cpp
+++ b/Engine/Source/Editor/Persona/Private/SAnimNotifyPanel.cpp
@@ -722,6 +722,7 @@ public:
 		, _OnGetIsAnimNotifySelectionValidForReplacement()
 		, _OnReplaceSelectedWithNotify()
 		, _OnReplaceSelectedWithBlueprintNotify()
+		, _OnReplaceSelectedWithSyncMarker()
 		, _OnDeselectAllNotifies()
 		, _OnCopyNodes()
 		, _OnPasteNodes()
@@ -754,6 +755,7 @@ public:
 		SLATE_EVENT( FOnGetIsAnimNotifySelectionValidForReplacement, OnGetIsAnimNotifySelectionValidForReplacement)
 		SLATE_EVENT( FReplaceWithNotify, OnReplaceSelectedWithNotify )
 		SLATE_EVENT( FReplaceWithBlueprintNotify, OnReplaceSelectedWithBlueprintNotify)
+		SLATE_EVENT( FReplaceWithSyncMarker, OnReplaceSelectedWithSyncMarker)
 		SLATE_EVENT( FDeselectAllNotifies, OnDeselectAllNotifies)
 		SLATE_EVENT( FCopyNodes, OnCopyNodes )
 		SLATE_EVENT(FPasteNodes, OnPasteNodes)
@@ -860,13 +862,13 @@ protected:
 	// Build up a "New Notify..." menu
 	void FillNewNotifyMenu(FMenuBuilder& MenuBuilderbool, bool bIsReplaceWithMenu = false);
 	void FillNewNotifyStateMenu(FMenuBuilder& MenuBuilder, bool bIsReplaceWithMenu  = false);
-	void FillNewSyncMarkerMenu(FMenuBuilder& MenuBuilder);
+	void FillNewSyncMarkerMenu(FMenuBuilder& MenuBuilder, bool bIsReplaceWithMenu = false);
 	void OnAnimNotifyClassPicked(UClass* NotifyClass, bool bIsReplaceWithMenu = false);
 
 	// New notify functions
 	void CreateNewBlueprintNotifyAtCursor(FString NewNotifyName, FString BlueprintPath);
 	void CreateNewNotifyAtCursor(FString NewNotifyName, UClass* NotifyClass);
-	void CreateNewSyncMarkerAtCursor(FString NewSyncMarkerName, UClass* NotifyClass);
+	void CreateNewSyncMarkerAtCursor(FString NewSyncMarkerName);
 	void OnNewNotifyClicked();
 	void OnNewSyncMarkerClicked();
 	void AddNewNotify(const FText& NewNotifyName, ETextCommit::Type CommitInfo);
@@ -879,6 +881,7 @@ protected:
 	// "Replace with... " commands
 	void ReplaceSelectedWithBlueprintNotify(FString NewNotifyName, FString BlueprintPath);
 	void ReplaceSelectedWithNotify(FString NewNotifyName, UClass* NotifyClass);
+	void ReplaceSelectedWithSyncMarker(FString NewSyncMarkerName);
 	bool IsValidToPlace(UClass* NotifyClass) const;
 
 	// Whether we have one node selected
@@ -1026,6 +1029,7 @@ protected:
 	FOnGetIsAnimNotifySelectionValidForReplacement OnGetIsAnimNotifySelectionValidforReplacement;
 	FReplaceWithNotify						OnReplaceSelectedWithNotify;
 	FReplaceWithBlueprintNotify				OnReplaceSelectedWithBlueprintNotify;
+	FReplaceWithSyncMarker					OnReplaceSelectedWithSyncMarker;
 
 	FOnInvokeTab							OnInvokeTab;
 
@@ -2195,6 +2199,7 @@ void SAnimNotifyTrack::Construct(const FArguments& InArgs)
 	OnGetIsAnimNotifySelectionValidforReplacement = InArgs._OnGetIsAnimNotifySelectionValidForReplacement;
 	OnReplaceSelectedWithNotify = InArgs._OnReplaceSelectedWithNotify;
 	OnReplaceSelectedWithBlueprintNotify = InArgs._OnReplaceSelectedWithBlueprintNotify;
+	OnReplaceSelectedWithSyncMarker = InArgs._OnReplaceSelectedWithSyncMarker;
 	OnDeselectAllNotifies = InArgs._OnDeselectAllNotifies;
 	OnCopyNodes = InArgs._OnCopyNodes;
 	OnPasteNodes = InArgs._OnPasteNodes;
@@ -2318,7 +2323,13 @@ void SAnimNotifyTrack::FillNewNotifyStateMenu(FMenuBuilder& MenuBuilder, bool bI
 		MenuBuilder.AddMenuSeparator();
 	}
 
-	TSharedRef<SWidget> Widget = PersonaUtils::MakeAnimNotifyStatePicker(Sequence, FOnClassPicked::CreateRaw(this, &SAnimNotifyTrack::OnAnimNotifyClassPicked, bIsReplaceWithMenu));
+	TSharedRef<SWidget> Widget = 
+		SNew(SBox)
+		.WidthOverride(300.0f)
+		.HeightOverride(300.0f)
+		[	
+			PersonaUtils::MakeAnimNotifyStatePicker(Sequence, FOnClassPicked::CreateRaw(this, &SAnimNotifyTrack::OnAnimNotifyClassPicked, bIsReplaceWithMenu))
+		];
 	MenuBuilder.AddWidget(Widget, FText(), true, false);
 }
 
@@ -2338,55 +2349,52 @@ void SAnimNotifyTrack::FillNewNotifyMenu(FMenuBuilder& MenuBuilder, bool bIsRepl
 				MenuBuilder.AddMenuEntry(LOCTEXT("NewNotify", "New Notify..."), LOCTEXT("NewNotifyToolTip", "Create a new animation notify on the skeleton"), FSlateIcon(), UIAction);
 			}
 
-			MenuBuilder.AddSubMenu(
-				LOCTEXT("NewNotifySubMenu_Skeleton", "Skeleton Notifies"),
-				LOCTEXT("NewNotifySubMenu_Skeleton_Tooltip", "Choose from custom notifies on the skeleton"),
-				FNewMenuDelegate::CreateLambda([this, SeqSkeleton, bIsReplaceWithMenu](FMenuBuilder& InSubMenuBuilder)
-				{
-					ISkeletonEditorModule& SkeletonEditorModule = FModuleManager::LoadModuleChecked<ISkeletonEditorModule>("SkeletonEditor");
-					TSharedRef<IEditableSkeleton> EditableSkeleton = SkeletonEditorModule.CreateEditableSkeleton(SeqSkeleton);
+			ISkeletonEditorModule& SkeletonEditorModule = FModuleManager::LoadModuleChecked<ISkeletonEditorModule>("SkeletonEditor");
+			TSharedRef<IEditableSkeleton> EditableSkeleton = SkeletonEditorModule.CreateEditableSkeleton(SeqSkeleton);
 
-					InSubMenuBuilder.AddWidget(
-						SNew(SBox)
-						.MinDesiredWidth(300.0f)
-						.MaxDesiredHeight(400.0f)
-						[
-							SNew(SSkeletonAnimNotifies, EditableSkeleton)
-							.IsPicker(true)
-							.ShowNotifies(true)
-							.OnItemSelected_Lambda([this, bIsReplaceWithMenu](const FName& InNotifyName)
-							{
-								FSlateApplication::Get().DismissAllMenus();
+			MenuBuilder.AddWidget(
+				SNew(SBox)
+				.WidthOverride(300.0f)
+				.HeightOverride(250.0f)
+				[
+					SNew(SSkeletonAnimNotifies, EditableSkeleton)
+					.IsPicker(true)
+					.ShowNotifies(true)
+					.OnItemSelected_Lambda([this, bIsReplaceWithMenu](const FName& InNotifyName)
+					{
+						FSlateApplication::Get().DismissAllMenus();
 
-								if (!bIsReplaceWithMenu)
-								{
-									CreateNewNotifyAtCursor(InNotifyName.ToString(), nullptr);
-								}
-								else
-								{
-									ReplaceSelectedWithNotify(InNotifyName.ToString(), nullptr);
-								}
-							})
-						],
-						FText(), true, false
-					);
-				}));
+						if (!bIsReplaceWithMenu)
+						{
+							CreateNewNotifyAtCursor(InNotifyName.ToString(), nullptr);
+						}
+						else
+						{
+							ReplaceSelectedWithNotify(InNotifyName.ToString(), nullptr);
+						}
+					})
+				],
+				FText(), true, false);
 		}
 		MenuBuilder.EndSection();
 	}
 
-	// MenuBuilder always has a search widget added to it by default, hence if larger then 1 then something else has been added to it
-	if (MenuBuilder.GetMultiBox()->GetBlocks().Num() > 1)
+	MenuBuilder.BeginSection("AnimNotifyNotifySubMenu", LOCTEXT("NewNotifySubMenu_Notifies", "Notifies"));
 	{
-		MenuBuilder.AddMenuSeparator();
+		// Add a notify picker
+		TSharedRef<SWidget> Widget = 
+			SNew(SBox)
+			.WidthOverride(300.0f)
+			.HeightOverride(250.0f)
+			[
+				PersonaUtils::MakeAnimNotifyPicker(Sequence, FOnClassPicked::CreateRaw(this, &SAnimNotifyTrack::OnAnimNotifyClassPicked, bIsReplaceWithMenu))
+			];
+		MenuBuilder.AddWidget(Widget, FText(), true, false);
 	}
-
-	// Add a notify picker
-	TSharedRef<SWidget> Widget = PersonaUtils::MakeAnimNotifyPicker(Sequence, FOnClassPicked::CreateRaw(this, &SAnimNotifyTrack::OnAnimNotifyClassPicked, bIsReplaceWithMenu));
-	MenuBuilder.AddWidget(Widget, FText(), true, false);
+	MenuBuilder.EndSection();
 }
 
-void SAnimNotifyTrack::FillNewSyncMarkerMenu(FMenuBuilder& MenuBuilder)
+void SAnimNotifyTrack::FillNewSyncMarkerMenu(FMenuBuilder& MenuBuilder, bool bIsReplaceWithMenu /* = false */)
 {
 	USkeleton* SeqSkeleton = Sequence->GetSkeleton();
 	if (SeqSkeleton)
@@ -2394,36 +2402,39 @@ void SAnimNotifyTrack::FillNewSyncMarkerMenu(FMenuBuilder& MenuBuilder)
 		MenuBuilder.BeginSection("AnimSyncMarkerSubMenu", LOCTEXT("NewSyncMarkerSubMenu_Skeleton", "Sync Markers"));
 		{
 			FUIAction UIAction;
-			UIAction.ExecuteAction.BindSP(
-				this, &SAnimNotifyTrack::OnNewSyncMarkerClicked);
-			MenuBuilder.AddMenuEntry(LOCTEXT("NewSyncMarker", "New Sync Marker..."), LOCTEXT("NewSyncMarkerToolTip", "Create a new animation sync marker"), FSlateIcon(), UIAction);
-
-			MenuBuilder.AddSubMenu(
-				LOCTEXT("NewSyncMarkerSubMenu_Existing", "Existing Sync Markers"),
-				LOCTEXT("NewSyncMarkerSubMenu_Existing_Tooltip", "Choose from existing sync marker names on the skeleton"),
-				FNewMenuDelegate::CreateLambda([this, SeqSkeleton](FMenuBuilder& InSubMenuBuilder)
+			if (!bIsReplaceWithMenu)
 			{
-				ISkeletonEditorModule& SkeletonEditorModule = FModuleManager::LoadModuleChecked<ISkeletonEditorModule>("SkeletonEditor");
-				TSharedRef<IEditableSkeleton> EditableSkeleton = SkeletonEditorModule.CreateEditableSkeleton(SeqSkeleton);
+				UIAction.ExecuteAction.BindSP(
+					this, &SAnimNotifyTrack::OnNewSyncMarkerClicked);
+				MenuBuilder.AddMenuEntry(LOCTEXT("NewSyncMarker", "New Sync Marker..."), LOCTEXT("NewSyncMarkerToolTip", "Create a new animation sync marker"), FSlateIcon(), UIAction);
+			}
 
-				InSubMenuBuilder.AddWidget(
-					SNew(SBox)
-					.MinDesiredWidth(300.0f)
-					.MaxDesiredHeight(400.0f)
-					[
-						SNew(SSkeletonAnimNotifies, EditableSkeleton)
-						.ShowSyncMarkers(true)
-						.ShowNotifies(false)
-						.OnItemSelected_Lambda([this](const FName& InNotifyName)
+			ISkeletonEditorModule& SkeletonEditorModule = FModuleManager::LoadModuleChecked<ISkeletonEditorModule>("SkeletonEditor");
+			TSharedRef<IEditableSkeleton> EditableSkeleton = SkeletonEditorModule.CreateEditableSkeleton(SeqSkeleton);
+
+			MenuBuilder.AddWidget(
+				SNew(SBox)
+				.WidthOverride(300.0f)
+				.HeightOverride(250.0f)
+				[
+					SNew(SSkeletonAnimNotifies, EditableSkeleton)
+					.ShowSyncMarkers(true)
+					.ShowNotifies(false)
+					.OnItemSelected_Lambda([this, bIsReplaceWithMenu](const FName& InNotifyName)
+					{
+						FSlateApplication::Get().DismissAllMenus();
+
+						if (!bIsReplaceWithMenu)
+						{
+							CreateNewSyncMarkerAtCursor(InNotifyName.ToString());
+						}
+						else
 						{
-							FSlateApplication::Get().DismissAllMenus();
-
-							CreateNewSyncMarkerAtCursor(InNotifyName.ToString(), nullptr);
-						})
-					],
-					FText(), true, false
-					);
-			}));
+							ReplaceSelectedWithSyncMarker(InNotifyName.ToString());
+						}
+					})
+				],
+				FText(), true, false);
 		}
 		MenuBuilder.EndSection();
 	}
@@ -2524,7 +2535,7 @@ void SAnimNotifyTrack::CreateNewNotifyAtCursor(FString NewNotifyName, UClass* No
 	OnUpdatePanel.ExecuteIfBound();
 }
 
-void SAnimNotifyTrack::CreateNewSyncMarkerAtCursor(FString NewSyncMarkerName, UClass* NotifyClass)
+void SAnimNotifyTrack::CreateNewSyncMarkerAtCursor(FString NewSyncMarkerName)
 {
 	UAnimSequence* Seq = CastChecked<UAnimSequence>(Sequence);
 
@@ -2554,6 +2565,11 @@ void SAnimNotifyTrack::ReplaceSelectedWithNotify(FString NewNotifyName, UClass*
 	OnReplaceSelectedWithNotify.ExecuteIfBound(NewNotifyName, NotifyClass);
 }
 
+void SAnimNotifyTrack::ReplaceSelectedWithSyncMarker(FString NewNotifyName)
+{
+	OnReplaceSelectedWithSyncMarker.ExecuteIfBound(NewNotifyName);
+}
+
 bool SAnimNotifyTrack::IsValidToPlace(UClass* NotifyClass) const
 {
 	if (NotifyClass && NotifyClass->IsChildOf(UAnimNotify::StaticClass()))
@@ -2985,7 +3001,7 @@ TSharedPtr<SWidget> SAnimNotifyTrack::SummonContextMenu(const FGeometry& MyGeome
 				MenuBuilder.AddSubMenu(
 					NSLOCTEXT("NewSyncMarkerSubMenu", "NewSyncMarkerSubMenuAddNotifyState", "Add Sync Marker..."),
 					NSLOCTEXT("NewSyncMarkerSubMenu", "NewSyncMarkerSubMenuAddNotifyStateToolTip", "Create a new animation sync marker"),
-					FNewMenuDelegate::CreateRaw(this, &SAnimNotifyTrack::FillNewSyncMarkerMenu));
+					FNewMenuDelegate::CreateRaw(this, &SAnimNotifyTrack::FillNewSyncMarkerMenu, false));
 			}
 
 			MenuBuilder.AddMenuEntry(
@@ -3032,6 +3048,16 @@ TSharedPtr<SWidget> SAnimNotifyTrack::SummonContextMenu(const FGeometry& MyGeome
 					}
 				}
 			}
+			else
+			{
+				if (NodeObject->GetType() == ENodeObjectTypes::SYNC_MARKER)
+				{
+					MenuBuilder.AddSubMenu(
+						NSLOCTEXT("NewNotifySubMenu", "NewNotifySubMenuReplaceWithSyncMarker", "Replace Sync Marker(s)..."),
+						NSLOCTEXT("NewNotifySubMenu", "NewNotifySubMenuReplaceWithSyncMarkerToolTip", "Replace the selected sync markers"),
+						FNewMenuDelegate::CreateRaw(this, &SAnimNotifyTrack::FillNewSyncMarkerMenu, true));
+				}
+			}
 		}
 		else
 		{
@@ -3244,7 +3270,7 @@ void SAnimNotifyTrack::AddNewSyncMarker(const FText& NewNotifyName, ETextCommit:
 
 		FBlueprintActionDatabase::Get().RefreshAssetActions(SeqSkeleton);
 
-		CreateNewSyncMarkerAtCursor(NewNotifyName.ToString(), (UClass*)nullptr);
+		CreateNewSyncMarkerAtCursor(NewNotifyName.ToString());
 	}
 
 	FSlateApplication::Get().DismissAllMenus();
@@ -3804,6 +3830,7 @@ void SNotifyEdTrack::Construct(const FArguments& InArgs)
 				.OnGetIsAnimNotifySelectionValidForReplacement(PanelRef, &SAnimNotifyPanel::IsNotifySelectionValidForReplacement)
 				.OnReplaceSelectedWithNotify(PanelRef, &SAnimNotifyPanel::OnReplaceSelectedWithNotify)
 				.OnReplaceSelectedWithBlueprintNotify(PanelRef, &SAnimNotifyPanel::OnReplaceSelectedWithNotifyBlueprint)
+				.OnReplaceSelectedWithSyncMarker(PanelRef, &SAnimNotifyPanel::OnReplaceSelectedWithSyncMarker)
 				.OnDeselectAllNotifies(InArgs._OnDeselectAllNotifies)
 				.OnCopyNodes(InArgs._OnCopyNodes)
 				.OnPasteNodes(InArgs._OnPasteNodes)
@@ -4515,6 +4542,50 @@ void SAnimNotifyPanel::OnReplaceSelectedWithNotifyBlueprint(FString NewBlueprint
 	OnReplaceSelectedWithNotify(NewBlueprintNotifyName, BlueprintClass);
 }
 
+void SAnimNotifyPanel::OnReplaceSelectedWithSyncMarker(FString NewSyncMarkerName)
+{
+	if (UAnimSequence* Seq = Cast<UAnimSequence>(Sequence))
+	{
+		TArray<INodeObjectInterface*> SelectedNodes;
+		for (TSharedPtr<SAnimNotifyTrack> Track : NotifyAnimTracks)
+		{
+			Track->AppendSelectionToArray(SelectedNodes);
+		}
+
+		// Sort these since order is important for deletion
+		SelectedNodes.Sort();
+
+		const FScopedTransaction Transaction(LOCTEXT("ReplaceSyncMarker", "Replace Sync Marker"));
+		Seq->Modify(true);
+
+		for (INodeObjectInterface* NodeObject : SelectedNodes)
+		{
+			if (NodeObject->GetType() == ENodeObjectTypes::SYNC_MARKER)
+			{
+				float Time = NodeObject->GetTime();
+				int32 TrackIndex = NodeObject->GetTrackIndex();
+
+				NodeObject->Delete(Seq);
+
+				FAnimSyncMarker& SyncMarker = Seq->AuthoredSyncMarkers.AddDefaulted_GetRef();
+				SyncMarker.MarkerName = FName(*NewSyncMarkerName);
+				SyncMarker.TrackIndex = TrackIndex;
+				SyncMarker.Time = Time;
+				SyncMarker.Guid = FGuid::NewGuid();
+			}
+		}
+
+		// clear selection  
+		TArray<UObject*> Objects;
+		OnSelectionChanged.ExecuteIfBound(Objects);
+
+		Seq->PostEditChange();
+		Seq->MarkPackageDirty();
+
+		RequestUpdate();
+	}
+}
+
 void SAnimNotifyPanel::OnPasteNodes(SAnimNotifyTrack* RequestTrack, float ClickTime, ENotifyPasteMode::Type PasteMode, ENotifyPasteMultipleMode::Type MultiplePasteType)
 {
 	if(RequestTrack == nullptr)
diff --git a/Engine/Source/Editor/Persona/Private/SAnimNotifyPanel.h b/Engine/Source/Editor/Persona/Private/SAnimNotifyPanel.h
--- a/Engine/Source/Editor/Persona/Private/SAnimNotifyPanel.h
+++ b/Engine/Source/Editor/Persona/Private/SAnimNotifyPanel.h
@@ -36,6 +36,7 @@ DECLARE_DELEGATE( FDeleteNotify )
 DECLARE_DELEGATE_RetVal( bool, FOnGetIsAnimNotifySelectionValidForReplacement )
 DECLARE_DELEGATE_TwoParams( FReplaceWithNotify, FString, UClass* )
 DECLARE_DELEGATE_TwoParams( FReplaceWithBlueprintNotify, FString, FString )
+DECLARE_DELEGATE_OneParam( FReplaceWithSyncMarker, FString )
 DECLARE_DELEGATE( FDeselectAllNotifies )
 DECLARE_DELEGATE_OneParam( FOnGetBlueprintNotifyData, TArray<FAssetData>& )
 DECLARE_DELEGATE_OneParam( FOnGetNativeNotifyClasses, TArray<UClass*>&)
@@ -260,6 +261,9 @@ public:
 	/** Handler for replacing with notify blueprint */
 	void OnReplaceSelectedWithNotifyBlueprint(FString NewBlueprintNotifyName, FString NewBlueprintNotifyClass);
 
+	/** Handler for replacing with sync marker */
+	void OnReplaceSelectedWithSyncMarker(FString NewBlueprintNotifyName);
+
 	void HandleObjectsSelected(const TArray<UObject*>& InObjects);
 
 	TSharedRef<FUICommandList> GetCommandList() const { return CommandList.ToSharedRef(); }
