Fixed duplicate curves being added to curve editor when curves are not selected

Iterate the active tree items, not the items returned from GetCurves (which are actually the selected, edited curves)

diff --git a/Engine/Source/Editor/Persona/Private/SAnimSequenceCurveEditor.cpp b/Engine/Source/Editor/Persona/Private/SAnimSequenceCurveEditor.cpp
--- a/Engine/Source/Editor/Persona/Private/SAnimSequenceCurveEditor.cpp
+++ b/Engine/Source/Editor/Persona/Private/SAnimSequenceCurveEditor.cpp
@@ -549,10 +549,10 @@ void SAnimSequenceCurveEditor::ResetCurves()
 void SAnimSequenceCurveEditor::AddCurve(const FText& InCurveDisplayName, const FLinearColor& InCurveColor, const FName& InName, ERawCurveTrackTypes InType, int32 InCurveIndex, FSimpleDelegate InOnCurveModified)
 {
 	// Ensure that curve is not already being edited
-	for (const TPair<FCurveModelID, TUniquePtr<FCurveModel>>& CurvePair : CurveEditor->GetCurves())
+	for (const TPair<FCurveEditorTreeItemID, FCurveEditorTreeItem>& ItemPair : CurveEditor->GetTree()->GetAllItems())
 	{
-		const FRichCurveEditorModelNamed* Model = static_cast<FRichCurveEditorModelNamed*>(CurvePair.Value.Get());
-		if(Model->CurveName == InName && Model->Type == InType && Model->CurveIndex == InCurveIndex)
+		TSharedPtr<FAnimSequenceCurveEditorItem> Item = StaticCastSharedPtr<FAnimSequenceCurveEditorItem>(ItemPair.Value.GetItem());
+		if (Item->Name == InName && Item->Type == InType && Item->CurveIndex == InCurveIndex)
 		{
 			return;
 		}
