Removing an added Notify track with a Sync Marker attached causes an Ensure
#fix also handle moving up notifies/syncmarkers for 1-indexed notify tracks

diff --git a/Engine/Source/Editor/Persona/Private/AnimTimeline/AnimTimelineTrack_NotifiesPanel.cpp b/Engine/Source/Editor/Persona/Private/AnimTimeline/AnimTimelineTrack_NotifiesPanel.cpp
--- a/Engine/Source/Editor/Persona/Private/AnimTimeline/AnimTimelineTrack_NotifiesPanel.cpp
+++ b/Engine/Source/Editor/Persona/Private/AnimTimeline/AnimTimelineTrack_NotifiesPanel.cpp
@@ -214,7 +214,7 @@ void FAnimTimelineTrack_NotifiesPanel::RemoveTrack(int32 InTrackIndexToRemove)
 			AnimSequence->Modify();
 
 			// before insert, make sure everything behind is fixed
-			for (int32 TrackIndex = InTrackIndexToRemove + 1; TrackIndex < AnimSequence->AnimNotifyTracks.Num(); ++TrackIndex)
+			for (int32 TrackIndex = InTrackIndexToRemove; TrackIndex < AnimSequence->AnimNotifyTracks.Num() && TrackIndex > 0; ++TrackIndex)
 			{
 				FAnimNotifyTrack& Track = AnimSequence->AnimNotifyTracks[TrackIndex];
 				const int32 NewTrackIndex = TrackIndex - 1;
