Fix asset switching using the Persona shortcut bar opening the 'wrong' previous animation/mesh etc.

diff --git a/Engine/Source/Editor/Persona/Private/PersonaAssetFamily.cpp b/Engine/Source/Editor/Persona/Private/PersonaAssetFamily.cpp
--- a/Engine/Source/Editor/Persona/Private/PersonaAssetFamily.cpp
+++ b/Engine/Source/Editor/Persona/Private/PersonaAssetFamily.cpp
@@ -17,6 +17,7 @@ FPersonaAssetFamily::FPersonaAssetFamily(const UObject* InFromObject)
 	, Mesh(nullptr)
 	, AnimBlueprint(nullptr)
 	, AnimationAsset(nullptr)
+	, PhysicsAsset(nullptr)
 {
 	if (InFromObject)
 	{
@@ -45,6 +46,38 @@ FPersonaAssetFamily::FPersonaAssetFamily(const UObject* InFromObject)
 	}
 }
 
+FPersonaAssetFamily::FPersonaAssetFamily(const UObject* InFromObject, const TSharedRef<FPersonaAssetFamily> InFromFamily)
+	: Skeleton(InFromFamily->Skeleton)
+	, Mesh(InFromFamily->Mesh)
+	, AnimBlueprint(InFromFamily->AnimBlueprint)
+	, AnimationAsset(InFromFamily->AnimationAsset)
+	, PhysicsAsset(InFromFamily->PhysicsAsset)
+{
+	if (InFromObject)
+	{
+		if (InFromObject->IsA<USkeleton>())
+		{
+			Skeleton = CastChecked<USkeleton>(InFromObject);
+		}
+		else if (InFromObject->IsA<UAnimationAsset>())
+		{
+			AnimationAsset = CastChecked<UAnimationAsset>(InFromObject);
+		}
+		else if (InFromObject->IsA<USkeletalMesh>())
+		{
+			Mesh = CastChecked<USkeletalMesh>(InFromObject);
+		}
+		else if (InFromObject->IsA<UAnimBlueprint>())
+		{
+			AnimBlueprint = CastChecked<UAnimBlueprint>(InFromObject);
+		}
+		else if (InFromObject->IsA<UPhysicsAsset>())
+		{
+			PhysicsAsset = CastChecked<UPhysicsAsset>(InFromObject);
+		}
+	}
+}
+
 void FPersonaAssetFamily::Initialize()
 {
 	GetMutableDefault<UPersonaOptions>()->RegisterOnUpdateSettings(UPersonaOptions::FOnUpdateSettingsMulticaster::FDelegate::CreateSP(this, &FPersonaAssetFamily::OnSettingsChange));
diff --git a/Engine/Source/Editor/Persona/Private/PersonaAssetFamily.h b/Engine/Source/Editor/Persona/Private/PersonaAssetFamily.h
--- a/Engine/Source/Editor/Persona/Private/PersonaAssetFamily.h
+++ b/Engine/Source/Editor/Persona/Private/PersonaAssetFamily.h
@@ -17,6 +17,8 @@ class FPersonaAssetFamily : public IAssetFamily, public TSharedFromThis<FPersona
 {
 public:
 	FPersonaAssetFamily(const UObject* InFromObject);
+	FPersonaAssetFamily(const UObject* InFromObject, const TSharedRef<FPersonaAssetFamily> InFromFamily);
+
 	virtual ~FPersonaAssetFamily() {}
 
 	/** IAssetFamily interface */
diff --git a/Engine/Source/Editor/Persona/Private/PersonaAssetFamilyManager.cpp b/Engine/Source/Editor/Persona/Private/PersonaAssetFamilyManager.cpp
--- a/Engine/Source/Editor/Persona/Private/PersonaAssetFamilyManager.cpp
+++ b/Engine/Source/Editor/Persona/Private/PersonaAssetFamilyManager.cpp
@@ -14,10 +14,22 @@ FPersonaAssetFamilyManager& FPersonaAssetFamilyManager::Get()
 TSharedRef<IAssetFamily> FPersonaAssetFamilyManager::CreatePersonaAssetFamily(const UObject* InAsset)
 {
 	// compact any invalid entries
-	AssetFamilies.RemoveAll([](const TWeakPtr<IAssetFamily>& InAssetFamily) { return !InAssetFamily.IsValid(); });
+	AssetFamilies.RemoveAll([](const TWeakPtr<FPersonaAssetFamily>& InAssetFamily) { return !InAssetFamily.IsValid(); });
+
+	// look for an existing matching asset family
+	TWeakPtr<FPersonaAssetFamily>* ExistingAssetFamily = nullptr;
+	FAssetData AssetData(InAsset);
+	for (TWeakPtr<FPersonaAssetFamily>& AssetFamily : AssetFamilies)
+	{
+		if (AssetFamily.Pin()->IsAssetCompatible(AssetData))
+		{
+			ExistingAssetFamily = &AssetFamily;
+			break;
+		}
+	}
 
 	// Create new asset family
-	TSharedRef<FPersonaAssetFamily> NewAssetFamily = MakeShared<FPersonaAssetFamily>(InAsset);
+	TSharedRef<FPersonaAssetFamily> NewAssetFamily = ExistingAssetFamily != nullptr ? MakeShared<FPersonaAssetFamily>(InAsset, ExistingAssetFamily->Pin().ToSharedRef()) : MakeShared<FPersonaAssetFamily>(InAsset);
 	NewAssetFamily->Initialize();
 	AssetFamilies.Add(NewAssetFamily);
 	return NewAssetFamily;
@@ -26,10 +38,10 @@ TSharedRef<IAssetFamily> FPersonaAssetFamilyManager::CreatePersonaAssetFamily(co
 void FPersonaAssetFamilyManager::BroadcastAssetFamilyChange()
 {
 	// Create copy as delegate can modify the AssetFamilies array
-	TArray<TWeakPtr<IAssetFamily>> AssetFamiliesCopy = AssetFamilies;
-	for (TWeakPtr<IAssetFamily>& AssetFamily : AssetFamiliesCopy)
+	TArray<TWeakPtr<FPersonaAssetFamily>> AssetFamiliesCopy = AssetFamilies;
+	for (TWeakPtr<FPersonaAssetFamily>& AssetFamily : AssetFamiliesCopy)
 	{
-		if (TSharedPtr<IAssetFamily> PinnedAssetFamily = AssetFamily.Pin())
+		if (TSharedPtr<FPersonaAssetFamily> PinnedAssetFamily = AssetFamily.Pin())
 		{
 			PinnedAssetFamily->GetOnAssetFamilyChanged().Broadcast();
 		}
@@ -38,9 +50,9 @@ void FPersonaAssetFamilyManager::BroadcastAssetFamilyChange()
 
 void FPersonaAssetFamilyManager::RecordAssetOpened(const FAssetData& InAssetData) const
 {
-	for (const TWeakPtr<IAssetFamily>& AssetFamily : AssetFamilies)
+	for (const TWeakPtr<FPersonaAssetFamily>& AssetFamily : AssetFamilies)
 	{
-		if (TSharedPtr<IAssetFamily> PinnedAssetFamily = AssetFamily.Pin())
+		if (TSharedPtr<FPersonaAssetFamily> PinnedAssetFamily = AssetFamily.Pin())
 		{
 			if(PinnedAssetFamily->IsAssetCompatible(InAssetData))
 			{
diff --git a/Engine/Source/Editor/Persona/Private/PersonaAssetFamilyManager.h b/Engine/Source/Editor/Persona/Private/PersonaAssetFamilyManager.h
--- a/Engine/Source/Editor/Persona/Private/PersonaAssetFamilyManager.h
+++ b/Engine/Source/Editor/Persona/Private/PersonaAssetFamilyManager.h
@@ -5,6 +5,7 @@
 #include "CoreMinimal.h"
 
 class IAssetFamily;
+class FPersonaAssetFamily;
 struct FAssetData;
 
 /** Central registry of persona asset families */
@@ -29,5 +30,5 @@ private:
 
 private:
 	/** All current asset families */
-	TArray<TWeakPtr<IAssetFamily>> AssetFamilies;
+	TArray<TWeakPtr<FPersonaAssetFamily>> AssetFamilies;
 };
