Lowers the angle threshold for flagging bad triangulation from 4 to 1 degree, to reduce the chance of bogus warnings.

diff --git a/Engine/Source/Editor/Persona/Private/SAnimationBlendSpaceGridWidget.cpp b/Engine/Source/Editor/Persona/Private/SAnimationBlendSpaceGridWidget.cpp
--- a/Engine/Source/Editor/Persona/Private/SAnimationBlendSpaceGridWidget.cpp
+++ b/Engine/Source/Editor/Persona/Private/SAnimationBlendSpaceGridWidget.cpp
@@ -47,18 +47,20 @@
 #include "Framework/Commands/GenericCommands.h"
 #include "Framework/Commands/UICommandList.h"
 
+//UE_DISABLE_OPTIMIZATION
 
 #define LOCTEXT_NAMESPACE "SAnimationBlendSpaceGridWidget"
 
 // Draws additional data on the triangulation to help debugging
 //#define DEBUG_BLENDSPACE_TRIANGULATION
 
-// Threshold for it being considered a problem that when a lookup is made at the same location as a sample, the returned
-// weight is less than this.
+// Threshold for it being considered a problem that when a lookup is made at the same location as a
+// sample, the returned weight is less than this.
 static const float SampleLookupWeightThreshold = 0.2f;
 
-// Flag any triangle that has an interior angle smaller than this (degrees)
-static const float CriticalTriangulationAngle = 4.0f;
+// Flag any triangle that has an interior angle smaller than this (degrees), calculated from the
+// normalized positions
+static const float CriticalTriangulationAngle = 1.0f;
 
 // Flag any triangle that has a smaller area than this (normalized units)
 static const float CriticalTriangulationArea = 5e-4f;
