The user cannot paste curves created in an Animation Sequence into another Animation Sequence
#fix reinstate adding of any missing curve(s) that the user is trying to paste

diff --git a/Engine/Source/Editor/Persona/Private/AnimTimeline/AnimTimelineClipboard.cpp b/Engine/Source/Editor/Persona/Private/AnimTimeline/AnimTimelineClipboard.cpp
--- a/Engine/Source/Editor/Persona/Private/AnimTimeline/AnimTimelineClipboard.cpp
+++ b/Engine/Source/Editor/Persona/Private/AnimTimeline/AnimTimelineClipboard.cpp
@@ -280,6 +280,13 @@ void FAnimTimelineClipboardUtilities::OverwriteOrAddCurvesFromClipboardContent(c
 		    else
 		    {
 			    FAnimationCurveIdentifier CurveIdentifier = InCurveCopyObj->GetAnimationCurveIdentifier();
+
+		    	// Create curve with clipboard information if needed
+		    	if (!Controller.GetModel()->FindCurve(CurveIdentifier))
+		    	{
+		    		Controller.AddCurve(CurveIdentifier);
+		    	}
+		    	
 			    if (InCurveCopyObj->CurveType == ERawCurveTrackTypes::RCT_Float)
 			    {
 				    const UFloatCurveCopyObject* FloatCopyObject = Cast<UFloatCurveCopyObject>(InCurveCopyObj);
