Fix issue where screen size property isn't correctly hidden in skel mesh details customization when driving LOD settings from data asset

diff --git a/Engine/Source/Editor/Persona/Private/PersonaMeshDetails.cpp b/Engine/Source/Editor/Persona/Private/PersonaMeshDetails.cpp
--- a/Engine/Source/Editor/Persona/Private/PersonaMeshDetails.cpp
+++ b/Engine/Source/Editor/Persona/Private/PersonaMeshDetails.cpp
@@ -2623,7 +2623,7 @@ void FPersonaMeshDetails::CustomizeLODInfoSetingsDetails(IDetailLayoutBuilder& D
 	const bool bForceComponentRefreshFalse = false;
 	// enable/disable handler - because we want to make sure not editable if LOD sharing is on
 	TSharedPtr<IPropertyHandle> ScreenSizeHandle = LODInfoChild->GetChildHandle(GET_MEMBER_NAME_CHECKED(FSkeletalMeshLODInfo, ScreenSize));
-	IDetailPropertyRow& ScreenSizeRow = LODInfoGroup.AddPropertyRow(ScreenSizeHandle->AsShared());
+	IDetailPropertyRow& ScreenSizeRow = LODInfoGroup.AddPropertyRow(ScreenSizeHandle->AsShared()).EditCondition(EnabledAttrib, NULL);
 	ScreenSizeRow.IsEnabled(EnabledAttrib);
 	PreventAttributePostEditChange(ScreenSizeHandle, LODIndex, ScreenSizeHandle->GetProperty()->GetFName(), bForceComponentRefreshTrue);
 	DetailLayout.HideProperty(ScreenSizeHandle);
