Enable bone rotation/translation while painting weights.

diff --git a/Engine/Plugins/Experimental/Animation/SkeletalMeshModelingTools/Source/SkeletalMeshModelingTools/Private/SkeletalMeshModelingToolsEditorMode.cpp b/Engine/Plugins/Experimental/Animation/SkeletalMeshModelingTools/Source/SkeletalMeshModelingTools/Private/SkeletalMeshModelingToolsEditorMode.cpp
--- a/Engine/Plugins/Experimental/Animation/SkeletalMeshModelingTools/Source/SkeletalMeshModelingTools/Private/SkeletalMeshModelingToolsEditorMode.cpp
+++ b/Engine/Plugins/Experimental/Animation/SkeletalMeshModelingTools/Source/SkeletalMeshModelingTools/Private/SkeletalMeshModelingToolsEditorMode.cpp
@@ -38,6 +38,7 @@
 #include "MeshVertexSculptTool.h"
 #include "ModelingToolsManagerActions.h"
 #include "OffsetMeshTool.h"
+#include "PersonaModule.h"
 #include "PolygonOnMeshTool.h"
 #include "ProjectToTargetTool.h"
 #include "RemeshMeshTool.h"
@@ -263,12 +264,21 @@ void USkeletalMeshModelingToolsEditorMode::Exit()
 	UEdMode::Exit();
 }
 
-
 void USkeletalMeshModelingToolsEditorMode::CreateToolkit()
 {
 	Toolkit = MakeShareable(new FSkeletalMeshModelingToolsEditorModeToolkit);
 }
 
+bool USkeletalMeshModelingToolsEditorMode::IsCompatibleWith(FEditorModeID OtherModeID) const
+{
+	if (OtherModeID == FPersonaEditModes::SkeletonSelection)
+	{
+		return true;
+	}
+	
+	return Super::IsCompatibleWith(OtherModeID);
+}
+
 void USkeletalMeshModelingToolsEditorMode::Tick(FEditorViewportClient* ViewportClient, float DeltaTime)
 {
 	Super::Tick(ViewportClient, DeltaTime);
diff --git a/Engine/Plugins/Experimental/Animation/SkeletalMeshModelingTools/Source/SkeletalMeshModelingTools/Private/SkeletalMeshModelingToolsEditorMode.h b/Engine/Plugins/Experimental/Animation/SkeletalMeshModelingTools/Source/SkeletalMeshModelingTools/Private/SkeletalMeshModelingToolsEditorMode.h
--- a/Engine/Plugins/Experimental/Animation/SkeletalMeshModelingTools/Source/SkeletalMeshModelingTools/Private/SkeletalMeshModelingToolsEditorMode.h
+++ b/Engine/Plugins/Experimental/Animation/SkeletalMeshModelingTools/Source/SkeletalMeshModelingTools/Private/SkeletalMeshModelingToolsEditorMode.h
@@ -34,6 +34,9 @@ public:
 	virtual void Enter() override;
 	virtual void Exit() override;
 	virtual void CreateToolkit() override;
+	virtual bool IsCompatibleWith(FEditorModeID OtherModeID) const override;
+	virtual bool UsesTransformWidget() const override { return false; };
+	virtual bool UsesPropertyWidgets() const override { return false; };
 	virtual void Tick(FEditorViewportClient* InViewportClient, float InDeltaTime) override;
 	virtual bool HandleClick(FEditorViewportClient* InViewportClient, HHitProxy *HitProxy, const FViewportClick &Click) override;
 	virtual bool ComputeBoundingBoxForViewportFocus(AActor* Actor, UPrimitiveComponent* PrimitiveComponent, FBox& InOutBox) const override;
diff --git a/Engine/Plugins/Experimental/Animation/SkeletalMeshModelingTools/Source/SkeletalMeshModelingTools/Private/SkeletalMeshModelingToolsModule.cpp b/Engine/Plugins/Experimental/Animation/SkeletalMeshModelingTools/Source/SkeletalMeshModelingTools/Private/SkeletalMeshModelingToolsModule.cpp
--- a/Engine/Plugins/Experimental/Animation/SkeletalMeshModelingTools/Source/SkeletalMeshModelingTools/Private/SkeletalMeshModelingToolsModule.cpp
+++ b/Engine/Plugins/Experimental/Animation/SkeletalMeshModelingTools/Source/SkeletalMeshModelingTools/Private/SkeletalMeshModelingToolsModule.cpp
@@ -166,7 +166,7 @@ void FSkeletalMeshModelingToolsModule::OnToggleEditingToolsMode(TWeakPtr<ISkelet
 		}
 		else
 		{
-			EditorModeManager.ActivateDefaultMode();
+			EditorModeManager.DeactivateMode(USkeletalMeshModelingToolsEditorMode::Id);
 		}
 	}
 }
diff --git a/Engine/Plugins/Experimental/Animation/SkeletalMeshModelingTools/Source/SkeletalMeshModelingTools/SkeletalMeshModelingTools.Build.cs b/Engine/Plugins/Experimental/Animation/SkeletalMeshModelingTools/Source/SkeletalMeshModelingTools/SkeletalMeshModelingTools.Build.cs
--- a/Engine/Plugins/Experimental/Animation/SkeletalMeshModelingTools/Source/SkeletalMeshModelingTools/SkeletalMeshModelingTools.Build.cs
+++ b/Engine/Plugins/Experimental/Animation/SkeletalMeshModelingTools/Source/SkeletalMeshModelingTools/SkeletalMeshModelingTools.Build.cs
@@ -36,7 +36,8 @@ public class SkeletalMeshModelingTools : ModuleRules
 				"AnimationWidgets",
 				"ApplicationCore",
 				"WidgetRegistration",
-				"SkeletalMeshModifiers"
+				"SkeletalMeshModifiers",
+				"Persona"
 			}
 		);
 
diff --git a/Engine/Plugins/Experimental/MeshModelingToolsetExp/Source/MeshModelingToolsEditorOnlyExp/MeshModelingToolsEditorOnlyExp.Build.cs b/Engine/Plugins/Experimental/MeshModelingToolsetExp/Source/MeshModelingToolsEditorOnlyExp/MeshModelingToolsEditorOnlyExp.Build.cs
--- a/Engine/Plugins/Experimental/MeshModelingToolsetExp/Source/MeshModelingToolsEditorOnlyExp/MeshModelingToolsEditorOnlyExp.Build.cs
+++ b/Engine/Plugins/Experimental/MeshModelingToolsetExp/Source/MeshModelingToolsEditorOnlyExp/MeshModelingToolsEditorOnlyExp.Build.cs
@@ -55,6 +55,7 @@ public class MeshModelingToolsEditorOnlyExp : ModuleRules
 				"RenderCore",
 				"InputCore",
 				"UnrealEd",
+				"Persona"
 				// ... add private dependencies that you statically link with here ...	
 			}
 		);
diff --git a/Engine/Plugins/Experimental/MeshModelingToolsetExp/Source/MeshModelingToolsEditorOnlyExp/Private/SkeletalMesh/SkinWeightsPaintTool.cpp b/Engine/Plugins/Experimental/MeshModelingToolsetExp/Source/MeshModelingToolsEditorOnlyExp/Private/SkeletalMesh/SkinWeightsPaintTool.cpp
--- a/Engine/Plugins/Experimental/MeshModelingToolsetExp/Source/MeshModelingToolsEditorOnlyExp/Private/SkeletalMesh/SkinWeightsPaintTool.cpp
+++ b/Engine/Plugins/Experimental/MeshModelingToolsetExp/Source/MeshModelingToolsEditorOnlyExp/Private/SkeletalMesh/SkinWeightsPaintTool.cpp
@@ -25,6 +25,8 @@
 #include "Util/ColorConstants.h"
 #include "Operations/SmoothBoneWeights.h"
 #include "ContextObjectStore.h"
+#include "Editor/Persona/Public/IPersonaEditorModeManager.h"
+#include "Editor/Persona/Public/PersonaModule.h"
 
 #include UE_INLINE_GENERATED_CPP_BY_NAME(SkinWeightsPaintTool)
 
@@ -702,6 +704,8 @@ void USkinWeightsPaintTool::Init(const FToolBuilderState& InSceneState)
 {
 	const UContextObjectStore* ContextObjectStore = InSceneState.ToolManager->GetContextObjectStore();
 	EditorContext = ContextObjectStore->FindContext<USkeletalMeshEditorContextObjectBase>();
+
+	PersonaModeManagerContext = ContextObjectStore->FindContext<UPersonaEditorModeManagerContext>();
 }
 
 void USkinWeightsPaintTool::Setup()
@@ -1532,6 +1536,11 @@ void USkinWeightsPaintTool::OnShutdown(EToolShutdownType ShutdownType)
 	{
 		EditorContext->UnbindFrom(this);
 	}
+
+	if (PersonaModeManagerContext.IsValid())
+	{
+		PersonaModeManagerContext->GetPersonaEditorModeManager()->DeactivateMode(FPersonaEditModes::SkeletonSelection);
+	}
 }
 
 void USkinWeightsPaintTool::BeginChange()
@@ -2002,8 +2011,26 @@ void USkinWeightsPaintTool::HandleSkeletalMeshModified(const TArray<FName>& InBo
 void USkinWeightsPaintTool::ToggleEditingMode()
 {
 	Weights.Deformer.SetAllVerticesToBeUpdated();
+
+	// toggle brush mode
 	SetBrushEnabled(WeightToolProperties->EditingMode == EWeightEditMode::Brush);
+
+	// toggle vertex mode
 	PolygonSelectionMechanic->SetIsEnabled(WeightToolProperties->EditingMode == EWeightEditMode::Vertices);
+
+	// toggle bone select mode
+	// this mode is set to be compatible with the 
+	if (PersonaModeManagerContext.IsValid())
+	{
+		if (WeightToolProperties->EditingMode == EWeightEditMode::Bones)
+		{
+			PersonaModeManagerContext->GetPersonaEditorModeManager()->ActivateMode(FPersonaEditModes::SkeletonSelection);	
+		}
+		else
+		{
+			PersonaModeManagerContext->GetPersonaEditorModeManager()->DeactivateMode(FPersonaEditModes::SkeletonSelection);	
+		}
+	}
 }
 
 void USkinWeightsPaintTool::GetSelectedVertices(TArray<int32>& OutVertexIndices) const
diff --git a/Engine/Plugins/Experimental/MeshModelingToolsetExp/Source/MeshModelingToolsEditorOnlyExp/Public/SkeletalMesh/SkinWeightsPaintTool.h b/Engine/Plugins/Experimental/MeshModelingToolsetExp/Source/MeshModelingToolsEditorOnlyExp/Public/SkeletalMesh/SkinWeightsPaintTool.h
--- a/Engine/Plugins/Experimental/MeshModelingToolsetExp/Source/MeshModelingToolsEditorOnlyExp/Public/SkeletalMesh/SkinWeightsPaintTool.h
+++ b/Engine/Plugins/Experimental/MeshModelingToolsetExp/Source/MeshModelingToolsEditorOnlyExp/Public/SkeletalMesh/SkinWeightsPaintTool.h
@@ -20,6 +20,7 @@
 struct FMeshDescription;
 class USkinWeightsPaintTool;
 class UPolygonSelectionMechanic;
+class UPersonaEditorModeManagerContext;
 
 namespace UE::Geometry 
 {
@@ -35,7 +36,8 @@ UENUM()
 enum class EWeightEditMode : uint8
 {
 	Brush,
-	Vertices
+	Vertices,
+	Bones,
 };
 
 // weight color mode
@@ -499,6 +501,8 @@ protected:
 
 	UPROPERTY()
 	TWeakObjectPtr<USkeletalMeshEditorContextObjectBase> EditorContext = nullptr;
+	UPROPERTY()
+	TWeakObjectPtr<UPersonaEditorModeManagerContext> PersonaModeManagerContext = nullptr;
 
 	friend SkinPaintTool::FSkinToolDeformer;
 };
diff --git a/Engine/Plugins/Experimental/MeshModelingToolsetExp/Source/ModelingEditorUI/Private/SkinWeightDetailCustomization.cpp b/Engine/Plugins/Experimental/MeshModelingToolsetExp/Source/ModelingEditorUI/Private/SkinWeightDetailCustomization.cpp
--- a/Engine/Plugins/Experimental/MeshModelingToolsetExp/Source/ModelingEditorUI/Private/SkinWeightDetailCustomization.cpp
+++ b/Engine/Plugins/Experimental/MeshModelingToolsetExp/Source/ModelingEditorUI/Private/SkinWeightDetailCustomization.cpp
@@ -45,8 +45,9 @@ void FSkinWeightDetailCustomization::CustomizeDetails(IDetailLayoutBuilder& Deta
 		[
 			SNew(SSegmentedControl<EWeightEditMode>)
 			.ToolTipText(LOCTEXT("EditingModeTooltip",
-					"Brush: the interactive viewport brush interface for editing weights.\n"
-					"Selection: the 1-off editing operations for weights on selected elements (bones and vertices).\n"))
+					"Brush: edit weights by painting directly on mesh.\n"
+					"Vertices: select vertices and edit weights directly.\n"
+					"Bones: select and manipulate bones to preview deformations.\n"))
 			.Value_Lambda([this]()
 			{
 				return SkinToolSettings->EditingMode;
@@ -64,6 +65,8 @@ void FSkinWeightDetailCustomization::CustomizeDetails(IDetailLayoutBuilder& Deta
 			.Text(LOCTEXT("BrushEditMode", "Brush"))
 			+ SSegmentedControl<EWeightEditMode>::Slot(EWeightEditMode::Vertices)
 			.Text(LOCTEXT("VertexEditMode", "Vertices"))
+			+ SSegmentedControl<EWeightEditMode>::Slot(EWeightEditMode::Bones)
+			.Text(LOCTEXT("BoneEditMode", "Bones"))
 		]
 	];
 
diff --git a/Engine/Source/Editor/Persona/Private/EditModes/SkeletonSelectionEditMode.cpp b/Engine/Source/Editor/Persona/Private/EditModes/SkeletonSelectionEditMode.cpp
deleted file mode 100644
diff --git a/Engine/Source/Editor/Persona/Private/EditModes/SkeletonSelectionEditMode.h b/Engine/Source/Editor/Persona/Private/EditModes/SkeletonSelectionEditMode.h
deleted file mode 100644
diff --git a/Engine/Source/Editor/Persona/Private/PersonaModule.cpp b/Engine/Source/Editor/Persona/Private/PersonaModule.cpp
--- a/Engine/Source/Editor/Persona/Private/PersonaModule.cpp
+++ b/Engine/Source/Editor/Persona/Private/PersonaModule.cpp
@@ -57,7 +57,7 @@
 #include "Customization/BlendSpaceDetails.h"
 #include "Customization/BlendParameterDetails.h"
 #include "Customization/InterpolationParameterDetails.h"
-#include "EditModes/SkeletonSelectionEditMode.h"
+#include "SkeletonSelectionEditMode.h"
 #include "PersonaEditorModeManager.h"
 #include "PreviewSceneCustomizations.h"
 #include "SSkeletonSlotNames.h"
diff --git a/Engine/Source/Editor/Persona/Private/SkeletonSelectionEditMode.cpp b/Engine/Source/Editor/Persona/Private/SkeletonSelectionEditMode.cpp
new file mode 100644
--- /dev/null
+++ b/Engine/Source/Editor/Persona/Private/SkeletonSelectionEditMode.cpp
@@ -0,0 +1,651 @@
+// Copyright Epic Games, Inc. All Rights Reserved.
+
+#include "SkeletonSelectionEditMode.h"
+#include "Animation/DebugSkelMeshComponent.h"
+#include "AnimationEditorViewportClient.h"
+#include "CanvasItem.h"
+#include "CanvasTypes.h"
+#include "AnimPreviewInstance.h"
+#include "ISkeletonTree.h"
+#include "AssetEditorModeManager.h"
+#include "Engine/SkeletalMeshSocket.h"
+#include "EngineUtils.h"
+#include "Rendering/SkeletalMeshRenderData.h"
+#include "Engine/WindDirectionalSource.h"
+
+#include "IPersonaToolkit.h"
+#include "IEditableSkeleton.h"
+
+#define LOCTEXT_NAMESPACE "SkeletonSelectionEditMode"
+
+namespace SkeletonSelectionModeConstants
+{
+	/** Distance to trace for physics bodies */
+	static const float BodyTraceDistance = 10000.0f;
+}
+
+FSkeletonSelectionEditMode::FSkeletonSelectionEditMode()
+	: bManipulating(false)
+	, bInTransaction(false)
+{
+	// Disable grid drawing for this mode as the viewport handles this
+	bDrawGrid = false;
+}
+
+const FReferenceSkeleton& FSkeletonSelectionEditMode::GetReferenceSkeletonForComponent(const USkeletalMeshComponent* Component) const
+{
+	check(Component);
+	return Component->GetSkeletalMeshAsset() && !Component->GetSkeletalMeshAsset()->IsCompiling() ? Component->GetSkeletalMeshAsset()->GetRefSkeleton() : GetAnimPreviewScene().GetPersonaToolkit()->GetSkeleton()->GetReferenceSkeleton();
+}
+
+bool FSkeletonSelectionEditMode::GetCameraTarget(FSphere& OutTarget) const
+{
+	bool bHandled = false;
+
+	const UDebugSkelMeshComponent* PreviewMeshComponent = GetAnimPreviewScene().GetPreviewMeshComponent();
+	if(PreviewMeshComponent)
+	{
+		if (GetAnimPreviewScene().GetSelectedBoneIndex() != INDEX_NONE)
+		{
+			const int32 FocusBoneIndex = GetAnimPreviewScene().GetSelectedBoneIndex();
+			const FReferenceSkeleton& ReferenceSkeleton = GetReferenceSkeletonForComponent(PreviewMeshComponent);
+			if (FocusBoneIndex != INDEX_NONE && ReferenceSkeleton.IsValidIndex(FocusBoneIndex))
+			{
+				const FName BoneName = ReferenceSkeleton.GetBoneName(FocusBoneIndex);
+				OutTarget.Center = PreviewMeshComponent->GetBoneLocation(BoneName);
+				OutTarget.W = 30.0f;
+				bHandled = true;
+			}
+		}
+
+		if (!bHandled && GetAnimPreviewScene().GetSelectedSocket().IsValid())
+		{
+			USkeletalMeshSocket * Socket = GetAnimPreviewScene().GetSelectedSocket().Socket;
+			if (Socket)
+			{
+				OutTarget.Center = Socket->GetSocketLocation(PreviewMeshComponent);
+				OutTarget.W = 30.0f;
+				bHandled = true;
+			}
+		}
+	}
+
+	return bHandled;
+}
+
+IPersonaPreviewScene& FSkeletonSelectionEditMode::GetAnimPreviewScene() const
+{
+	return *static_cast<IPersonaPreviewScene*>(static_cast<FAssetEditorModeManager*>(Owner)->GetPreviewScene());
+}
+
+void FSkeletonSelectionEditMode::GetOnScreenDebugInfo(TArray<FText>& OutDebugInfo) const
+{
+
+}
+
+FSelectedSocketInfo FSkeletonSelectionEditMode::DuplicateAndSelectSocket(const FSelectedSocketInfo& SocketInfoToDuplicate)
+{
+	USkeletalMesh* SkeletalMesh = GetAnimPreviewScene().GetPreviewMeshComponent()->GetSkeletalMeshAsset();
+	USkeletalMeshSocket* NewSocket = GetAnimPreviewScene().GetPersonaToolkit()->GetEditableSkeleton()->DuplicateSocket(SocketInfoToDuplicate, SocketInfoToDuplicate.Socket->BoneName, SkeletalMesh);
+
+	FSelectedSocketInfo NewSocketInfo(NewSocket, SocketInfoToDuplicate.bSocketIsOnSkeleton);
+	GetAnimPreviewScene().DeselectAll();
+	GetAnimPreviewScene().SetSelectedSocket(NewSocketInfo);
+
+	return NewSocketInfo;
+}
+
+bool FSkeletonSelectionEditMode::StartTracking(FEditorViewportClient* InViewportClient, FViewport* InViewport)
+{
+	EAxisList::Type CurrentAxis = InViewportClient->GetCurrentWidgetAxis();
+	UE::Widget::EWidgetMode WidgetMode = InViewportClient->GetWidgetMode();
+
+	UDebugSkelMeshComponent* PreviewMeshComponent = GetAnimPreviewScene().GetPreviewMeshComponent();
+	if(PreviewMeshComponent != nullptr && PreviewMeshComponent->GetSkeletalMeshAsset() != nullptr)
+	{
+		const int32 BoneIndex = GetAnimPreviewScene().GetSelectedBoneIndex();
+		const USkeletalMeshSocket* SelectedSocket = GetAnimPreviewScene().GetSelectedSocket().Socket;
+		const AActor* SelectedActor = GetAnimPreviewScene().GetSelectedActor();
+
+		// Retrieve reference skeleton from either current USkeletalMesh, or USkeleton if no mesh is set
+		const FReferenceSkeleton& ReferenceSkeleton = GetReferenceSkeletonForComponent(PreviewMeshComponent);
+		if ((BoneIndex >= 0 && ReferenceSkeleton.IsValidIndex(BoneIndex)) || SelectedSocket != nullptr || SelectedActor != nullptr)
+		{
+			bool bValidAxis = false;
+			FVector WorldAxisDir;
+
+			if ( ((CurrentAxis & EAxisList::XYZ) | (CurrentAxis & EAxisList::Screen)) != 0)
+			{
+				FSelectedSocketInfo SelectedSocketInfo = GetAnimPreviewScene().GetSelectedSocket();
+				if (SelectedSocketInfo.IsValid())
+				{
+					const bool bAltDown = InViewportClient->IsAltPressed();
+
+					if (bAltDown)
+					{
+						// Rather than moving/rotating the selected socket, copy it and move the copy instead
+						SelectedSocketInfo = DuplicateAndSelectSocket(SelectedSocketInfo);
+					}
+
+					// Socket movement is transactional - we want undo/redo and saving of it
+					USkeletalMeshSocket* Socket = SelectedSocketInfo.Socket;
+
+					if (Socket && bInTransaction == false)
+					{
+						if (WidgetMode == UE::Widget::WM_Rotate)
+						{
+							GEditor->BeginTransaction(LOCTEXT("AnimationEditorViewport_RotateSocket", "Rotate Socket"));
+						}
+						else
+						{
+							GEditor->BeginTransaction(LOCTEXT("AnimationEditorViewport_TranslateSocket", "Translate Socket"));
+						}
+
+						Socket->SetFlags(RF_Transactional);	// Undo doesn't work without this!
+						Socket->Modify();
+						bInTransaction = true;
+					}
+				}
+				else if (BoneIndex >= 0)
+				{
+					if (bInTransaction == false)
+					{
+						// we also allow undo/redo of bone manipulations
+						if (WidgetMode == UE::Widget::WM_Rotate)
+						{
+							GEditor->BeginTransaction(LOCTEXT("AnimationEditorViewport_RotateBone", "Rotate Bone"));
+						}
+						else
+						{
+							GEditor->BeginTransaction(LOCTEXT("AnimationEditorViewport_TranslateBone", "Translate Bone"));
+						}
+
+						PreviewMeshComponent->PreviewInstance->SetFlags(RF_Transactional);	// Undo doesn't work without this!
+						PreviewMeshComponent->PreviewInstance->Modify();
+						bInTransaction = true;
+
+						// now modify the bone array
+						const FName BoneName = ReferenceSkeleton.GetBoneName(BoneIndex);
+						PreviewMeshComponent->PreviewInstance->ModifyBone(BoneName);
+					}
+				}
+			}
+
+			bManipulating = true;
+			return true;
+		}
+	}
+	
+	return false;
+}
+
+bool FSkeletonSelectionEditMode::EndTracking(FEditorViewportClient* InViewportClient, FViewport* InViewport)
+{
+	if (bManipulating)
+	{
+		// Socket movement is transactional - we want undo/redo and saving of it
+		if (bInTransaction)
+		{
+			GEditor->EndTransaction();
+			bInTransaction = false;
+		}
+
+		bManipulating = false;
+		return true;
+	}
+
+	return false;
+}
+
+bool FSkeletonSelectionEditMode::InputDelta(FEditorViewportClient* InViewportClient, FViewport* InViewport, FVector& InDrag, FRotator& InRot, FVector& InScale)
+{
+	const EAxisList::Type CurrentAxis = InViewportClient->GetCurrentWidgetAxis();
+	const UE::Widget::EWidgetMode WidgetMode = InViewportClient->GetWidgetMode();
+	const ECoordSystem CoordSystem = InViewportClient->GetWidgetCoordSystemSpace();
+
+	// Get some useful info about buttons being held down
+	const bool bCtrlDown = InViewport->KeyState(EKeys::LeftControl) || InViewport->KeyState(EKeys::RightControl);
+	const bool bShiftDown = InViewport->KeyState(EKeys::LeftShift) || InViewport->KeyState(EKeys::RightShift);
+	const bool bMouseButtonDown = InViewport->KeyState( EKeys::LeftMouseButton ) || InViewport->KeyState( EKeys::MiddleMouseButton ) || InViewport->KeyState( EKeys::RightMouseButton );
+
+	bool bHandled = false;
+
+	UDebugSkelMeshComponent* PreviewMeshComponent = GetAnimPreviewScene().GetPreviewMeshComponent();
+
+	if ( bManipulating && CurrentAxis != EAxisList::None )
+	{
+		bHandled = true;
+
+		int32 BoneIndex = GetAnimPreviewScene().GetSelectedBoneIndex();
+		USkeletalMeshSocket* SelectedSocket = GetAnimPreviewScene().GetSelectedSocket().Socket;
+		AActor* SelectedActor = GetAnimPreviewScene().GetSelectedActor();
+		FAnimNode_ModifyBone* SkelControl = nullptr;
+
+		if ( BoneIndex >= 0 )
+		{
+			// Retrieve reference skeleton from either current USkeletalMesh, or USkeleton if no mesh is set
+			const FReferenceSkeleton& ReferenceSkeleton = GetReferenceSkeletonForComponent(PreviewMeshComponent);
+			const FName BoneName = ReferenceSkeleton.GetBoneName(BoneIndex);
+			//Get the skeleton control manipulating this bone
+			SkelControl = &(PreviewMeshComponent->PreviewInstance->ModifyBone(BoneName));
+		}
+
+		if ( SkelControl || SelectedSocket )
+		{
+			FTransform CurrentSkelControlTM(
+				SelectedSocket ? SelectedSocket->RelativeRotation : SkelControl->Rotation,
+				SelectedSocket ? SelectedSocket->RelativeLocation : SkelControl->Translation,
+				SelectedSocket ? SelectedSocket->RelativeScale : SkelControl->Scale);
+
+			FTransform BaseTM;
+
+			if ( SelectedSocket )
+			{
+				BaseTM = SelectedSocket->GetSocketTransform( PreviewMeshComponent );
+			}
+			else
+			{
+				BaseTM = PreviewMeshComponent->GetBoneTransform( BoneIndex );
+			}
+
+			// Remove SkelControl's orientation from BoneMatrix, as we need to translate/rotate in the non-SkelControlled space
+			BaseTM = BaseTM.GetRelativeTransformReverse(CurrentSkelControlTM);
+
+			const bool bDoRotation    = WidgetMode == UE::Widget::WM_Rotate    || WidgetMode == UE::Widget::WM_TranslateRotateZ;
+			const bool bDoTranslation = WidgetMode == UE::Widget::WM_Translate || WidgetMode == UE::Widget::WM_TranslateRotateZ;
+			const bool bDoScale = WidgetMode == UE::Widget::WM_Scale;
+
+			if (bDoRotation)
+			{
+				FVector RotAxis;
+				float RotAngle;
+				InRot.Quaternion().ToAxisAndAngle( RotAxis, RotAngle );
+
+				FVector4 BoneSpaceAxis = BaseTM.TransformVectorNoScale( RotAxis );
+
+				//Calculate the new delta rotation
+				FQuat DeltaQuat( BoneSpaceAxis, RotAngle );
+				DeltaQuat.Normalize();
+
+				FRotator NewRotation = ( CurrentSkelControlTM * FTransform( DeltaQuat )).Rotator();
+
+				if ( SelectedSocket )
+				{
+					SelectedSocket->RelativeRotation = NewRotation;
+				}
+				else
+				{
+					SkelControl->Rotation = NewRotation;
+				}
+			}
+
+			if (bDoTranslation)
+			{
+				FVector4 BoneSpaceOffset = BaseTM.TransformVector(InDrag);
+				if (SelectedSocket)
+				{
+					SelectedSocket->RelativeLocation += BoneSpaceOffset;
+				}
+				else
+				{
+					SkelControl->Translation += BoneSpaceOffset;
+				}
+			}
+			if(bDoScale)
+			{
+				FVector4 BoneSpaceScaleOffset;
+
+				if (CoordSystem == COORD_World)
+				{
+					BoneSpaceScaleOffset = BaseTM.TransformVector(InScale);
+				}
+				else
+				{
+					BoneSpaceScaleOffset = InScale;
+				}
+
+				if(SelectedSocket)
+				{
+					SelectedSocket->RelativeScale += BoneSpaceScaleOffset;
+				}
+				else
+				{
+					SkelControl->Scale += BoneSpaceScaleOffset;
+				}
+			}
+
+		}
+		else if( SelectedActor != nullptr )
+		{
+			if (WidgetMode == UE::Widget::WM_Rotate)
+			{
+				FTransform Transform = SelectedActor->GetTransform();
+				FRotator NewRotation = (Transform * FTransform( InRot ) ).Rotator();
+
+				SelectedActor->SetActorRotation( NewRotation );
+			}
+			else
+			{
+				FVector Location = SelectedActor->GetActorLocation();
+				Location += InDrag;
+				SelectedActor->SetActorLocation(Location);
+			}
+		}
+
+		InViewport->Invalidate();
+	}
+
+	return bHandled;
+}
+
+FIntPoint FSkeletonSelectionEditMode::GetDPIUnscaledSize(FViewport* Viewport, FViewportClient* Client)
+{
+	const FIntPoint Size = Viewport->GetSizeXY();
+	const float DPIScale = Client->GetDPIScale();
+	// (FIntPoint / float) implicitly casts the float to an int if you try to divide it directly
+	return FIntPoint(static_cast<int32>(Size.X / DPIScale), static_cast<int32>(Size.Y / DPIScale));
+}
+
+void FSkeletonSelectionEditMode::DrawHUD(FEditorViewportClient* ViewportClient, FViewport* Viewport, const FSceneView* View, FCanvas* Canvas)
+{
+	const UDebugSkelMeshComponent* PreviewMeshComponent = GetAnimPreviewScene().GetPreviewMeshComponent();
+	if (PreviewMeshComponent == nullptr)
+	{
+		return;
+	}
+
+	// Retrieve reference skeleton from either current USkeletalMesh, or USkeleton if no mesh is set
+	const FReferenceSkeleton& RefSkeleton = PreviewMeshComponent->GetSkeletalMeshAsset() && !PreviewMeshComponent->GetSkeletalMeshAsset()->IsCompiling() ? PreviewMeshComponent->GetSkeletalMeshAsset()->GetRefSkeleton() : GetAnimPreviewScene().GetPersonaToolkit()->GetSkeleton()->GetReferenceSkeleton();
+	const int32 BoneIndex = GetAnimPreviewScene().GetSelectedBoneIndex();
+
+	// Draw name of selected bone
+	if (RefSkeleton.IsValidIndex(BoneIndex) && (IsSelectedBoneRequired() || RefSkeleton.GetRequiredVirtualBones().Contains(BoneIndex)))
+	{
+		const FIntPoint ViewPortSize = GetDPIUnscaledSize(Viewport, ViewportClient);
+		const int32 HalfX = ViewPortSize.X / 2;
+		const int32 HalfY = ViewPortSize.Y / 2;
+
+		const FName BoneName = RefSkeleton.GetBoneName(BoneIndex);
+
+		const FMatrix BoneMatrix = GetBoneTransform(BoneIndex).ToMatrixNoScale();
+		const FPlane Proj = View->Project(BoneMatrix.GetOrigin());
+		if (Proj.W > 0.f)
+		{
+			const int32 XPos = HalfX + static_cast<int32>(HalfX * Proj.X);
+			const int32 YPos = HalfY + static_cast<int32>(HalfY * Proj.Y * -1);
+
+			FCanvasTextItem TextItem(FVector2D(XPos, YPos), FText::FromString(BoneName.ToString()), GEngine->GetSmallFont(), FLinearColor::White);
+			TextItem.EnableShadow(FLinearColor::Black);
+			Canvas->DrawItem(TextItem);
+		}
+	}
+
+	// Draw name of selected socket
+	if (GetAnimPreviewScene().GetSelectedSocket().IsValid())
+	{
+		const USkeletalMeshSocket* Socket = GetAnimPreviewScene().GetSelectedSocket().Socket;
+
+		const FMatrix SocketMatrix = GetSocketTransform(Socket).ToMatrixNoScale();
+		const FVector SocketPos = SocketMatrix.GetOrigin();
+
+		const FPlane Proj = View->Project(SocketPos);
+		if (Proj.W > 0.f)
+		{
+			const FIntPoint ViewPortSize = GetDPIUnscaledSize(Viewport, ViewportClient);
+			const int32 HalfX = ViewPortSize.X / 2;
+			const int32 HalfY = ViewPortSize.Y / 2;
+
+			const int32 XPos = HalfX + static_cast<int32>(HalfX * Proj.X);
+			const int32 YPos = HalfY + static_cast<int32>(HalfY * (Proj.Y * -1));
+			FCanvasTextItem TextItem(FVector2D(XPos, YPos), FText::FromString(Socket->SocketName.ToString()), GEngine->GetSmallFont(), FLinearColor::White);
+			TextItem.EnableShadow(FLinearColor::Black);
+			Canvas->DrawItem(TextItem);
+		}
+	}
+}
+
+bool FSkeletonSelectionEditMode::AllowWidgetMove()
+{
+	return ShouldDrawWidget();
+}
+
+bool FSkeletonSelectionEditMode::IsSelectedBoneRequired() const
+{
+	const UDebugSkelMeshComponent* PreviewMeshComponent = GetAnimPreviewScene().GetPreviewMeshComponent();
+	const int32 SelectedBoneIndex = GetAnimPreviewScene().GetSelectedBoneIndex();
+	if (SelectedBoneIndex != INDEX_NONE && PreviewMeshComponent->GetSkeletalMeshRenderData())
+	{
+		//Get current LOD
+		const FSkeletalMeshRenderData* SkelMeshRenderData = PreviewMeshComponent->GetSkeletalMeshRenderData();
+		if(SkelMeshRenderData->LODRenderData.Num() > 0)
+		{
+			const int32 LODIndex = FMath::Clamp(PreviewMeshComponent->GetPredictedLODLevel(), 0, SkelMeshRenderData->LODRenderData.Num() - 1);
+			const FSkeletalMeshLODRenderData& LODData = SkelMeshRenderData->LODRenderData[LODIndex];
+
+			//Check whether the bone is vertex weighted
+			return LODData.RequiredBones.Find(static_cast<FBoneIndexType>(SelectedBoneIndex)) != INDEX_NONE;
+		}
+	}
+	else if (SelectedBoneIndex != INDEX_NONE && !PreviewMeshComponent->GetSkeletalMeshRenderData())
+	{
+		return true;
+	}
+
+	return false;
+}
+
+bool FSkeletonSelectionEditMode::ShouldDrawWidget() const
+{
+	UDebugSkelMeshComponent* PreviewMeshComponent = GetAnimPreviewScene().GetPreviewMeshComponent();
+	if (!PreviewMeshComponent->IsAnimBlueprintInstanced())
+	{
+		return IsSelectedBoneRequired() || GetAnimPreviewScene().GetSelectedSocket().IsValid() || GetAnimPreviewScene().GetSelectedActor() != nullptr;
+	}
+
+	return false;
+}
+
+bool FSkeletonSelectionEditMode::UsesTransformWidget() const
+{
+	return true;
+}
+
+bool FSkeletonSelectionEditMode::UsesTransformWidget(UE::Widget::EWidgetMode CheckMode) const
+{
+	return ShouldDrawWidget() && (CheckMode == UE::Widget::WM_Scale || CheckMode == UE::Widget::WM_Translate || CheckMode == UE::Widget::WM_Rotate);
+}
+
+FTransform FSkeletonSelectionEditMode::GetWorldSpaceBoneTransform(const FReferenceSkeleton& ReferenceSkeleton, const int32 BoneIndex) const
+{
+	const TArray<FTransform>& BonePoses = ReferenceSkeleton.GetRefBonePose();
+
+	if (BonePoses.IsValidIndex(BoneIndex))
+	{
+		FTransform WorldSpacePose = BonePoses[BoneIndex];
+
+		int32 ParentIndex = ReferenceSkeleton.GetParentIndex(BoneIndex);
+
+		while(ParentIndex != INDEX_NONE)
+		{
+			WorldSpacePose = WorldSpacePose * BonePoses[ParentIndex];
+			ParentIndex = ReferenceSkeleton.GetParentIndex(ParentIndex);
+		}
+
+		return WorldSpacePose;
+	}
+
+	return FTransform::Identity;
+}
+
+
+
+FTransform FSkeletonSelectionEditMode::GetBoneTransform(const int32 BoneIndex) const
+{
+	const UDebugSkelMeshComponent* PreviewMeshComponent = GetAnimPreviewScene().GetPreviewMeshComponent();
+	if (PreviewMeshComponent && PreviewMeshComponent->GetSkeletalMeshAsset())
+	{
+		if (!PreviewMeshComponent->GetSkeletalMeshAsset()->IsCompiling())
+		{
+			return PreviewMeshComponent->GetBoneTransform(BoneIndex);
+		}
+	}
+	else if (const USkeleton* Skeleton = GetAnimPreviewScene().GetPersonaToolkit()->GetSkeleton())
+	{
+		return GetWorldSpaceBoneTransform(Skeleton->GetReferenceSkeleton(), BoneIndex);
+	}
+
+	return FTransform::Identity;
+}
+
+FTransform FSkeletonSelectionEditMode::GetSocketTransform(const USkeletalMeshSocket* Socket) const
+{
+	const UDebugSkelMeshComponent* PreviewMeshComponent = GetAnimPreviewScene().GetPreviewMeshComponent();
+	if (PreviewMeshComponent && PreviewMeshComponent->GetSkeletalMeshAsset())
+	{
+		if (!PreviewMeshComponent->GetSkeletalMeshAsset()->IsCompiling())
+		{
+			const int32 BoneIndex = PreviewMeshComponent->GetBoneIndex(Socket->BoneName);
+			if(BoneIndex != INDEX_NONE)
+			{
+				const FTransform BoneTM = PreviewMeshComponent->GetBoneTransform(BoneIndex);
+				const FTransform RelSocketTM( Socket->RelativeRotation, Socket->RelativeLocation, Socket->RelativeScale );
+				return RelSocketTM * BoneTM;
+			}
+			
+		}
+	}
+	else if (const USkeleton* Skeleton = GetAnimPreviewScene().GetPersonaToolkit()->GetSkeleton())
+	{
+		const FReferenceSkeleton& ReferenceSkeleton = Skeleton->GetReferenceSkeleton();
+		const int32 BoneIndex = ReferenceSkeleton.FindBoneIndex(Socket->BoneName);
+		if(BoneIndex != INDEX_NONE)
+		{
+			const FTransform BoneTM = GetWorldSpaceBoneTransform(Skeleton->GetReferenceSkeleton(), BoneIndex);
+			const FTransform RelSocketTM( Socket->RelativeRotation, Socket->RelativeLocation, Socket->RelativeScale );
+			return RelSocketTM * BoneTM;
+		}
+	}
+
+	return FTransform::Identity;
+}
+
+bool FSkeletonSelectionEditMode::GetCustomDrawingCoordinateSystem(FMatrix& InMatrix, void* InData)
+{
+	const UDebugSkelMeshComponent* PreviewMeshComponent = GetAnimPreviewScene().GetPreviewMeshComponent();
+	if(PreviewMeshComponent)
+	{
+		// Retrieve reference skeleton from either current USkeletalMesh, or USkeleton if no mesh is set
+		const FReferenceSkeleton& ReferenceSkeleton = GetReferenceSkeletonForComponent(PreviewMeshComponent);
+		
+		const int32 BoneIndex = GetAnimPreviewScene().GetSelectedBoneIndex();
+		if (BoneIndex != INDEX_NONE && ReferenceSkeleton.IsValidIndex(BoneIndex))
+		{
+			const FTransform BoneMatrix = GetBoneTransform(BoneIndex);
+			InMatrix = BoneMatrix.ToMatrixNoScale().RemoveTranslation();
+			return true;
+		}
+		else if (GetAnimPreviewScene().GetSelectedSocket().IsValid())
+		{
+			const USkeletalMeshSocket* Socket = GetAnimPreviewScene().GetSelectedSocket().Socket;
+			const FTransform SocketMatrix = GetSocketTransform(Socket);
+			InMatrix = SocketMatrix.ToMatrixNoScale().RemoveTranslation();
+			return true;
+		}
+		else if (const AActor* SelectedActor = GetAnimPreviewScene().GetSelectedActor())
+		{
+			InMatrix = SelectedActor->GetTransform().ToMatrixNoScale().RemoveTranslation();
+			return true;
+		}
+	}
+	
+	return false;
+}
+
+bool FSkeletonSelectionEditMode::GetCustomInputCoordinateSystem(FMatrix& InMatrix, void* InData)
+{
+	return GetCustomDrawingCoordinateSystem(InMatrix, InData);
+}
+
+FVector FSkeletonSelectionEditMode::GetWidgetLocation() const
+{
+	const int32 BoneIndex = GetAnimPreviewScene().GetSelectedBoneIndex();
+	if (BoneIndex != INDEX_NONE)
+	{
+		const FMatrix BoneMatrix = GetBoneTransform(BoneIndex).ToMatrixNoScale();
+		return BoneMatrix.GetOrigin();
+	}
+	else if (GetAnimPreviewScene().GetSelectedSocket().IsValid())
+	{
+		const USkeletalMeshSocket* Socket = GetAnimPreviewScene().GetSelectedSocket().Socket;
+		const FMatrix SocketMatrix = GetSocketTransform(Socket).ToMatrixNoScale();
+		return SocketMatrix.GetOrigin();
+	}
+	else if (const AActor* SelectedActor = GetAnimPreviewScene().GetSelectedActor())
+	{
+		return SelectedActor->GetActorLocation();
+	}
+
+	return FVector::ZeroVector;
+}
+
+bool FSkeletonSelectionEditMode::HandleClick(FEditorViewportClient* InViewportClient, HHitProxy* HitProxy, const FViewportClick& Click)
+{
+	GetAnimPreviewScene().DeselectAll();
+	USkeletalMeshComponent* MeshComponent = GetAnimPreviewScene().GetPreviewMeshComponent();
+	if (MeshComponent)
+	{
+		MeshComponent->SetSelectedEditorSection(INDEX_NONE);
+	}
+
+	if (!HitProxy)
+	{
+		return false;
+	}
+
+	if (HPersonaBoneHitProxy* BoneHitProxy = HitProxyCast<HPersonaBoneHitProxy>(HitProxy))
+	{
+		GetAnimPreviewScene().SetSelectedBone(BoneHitProxy->BoneName, ESelectInfo::OnMouseClick);
+		return true;
+	}
+	if (HPersonaSocketHitProxy* SocketHitProxy = HitProxyCast<HPersonaSocketHitProxy>(HitProxy))
+	{
+		if (USkeleton* Skeleton = GetAnimPreviewScene().GetPersonaToolkit()->GetSkeleton())
+		{
+			FSelectedSocketInfo SocketInfo;
+			SocketInfo.Socket = SocketHitProxy->Socket;
+			SocketInfo.bSocketIsOnSkeleton = !SocketInfo.Socket->GetOuter()->IsA<USkeletalMesh>();
+			GetAnimPreviewScene().SetSelectedSocket(SocketInfo);
+			return true;
+		}
+	}
+	if (HActor* ActorHitProxy = HitProxyCast<HActor>(HitProxy))
+	{
+		if (ActorHitProxy->Actor)
+		{
+			if (ActorHitProxy->Actor->IsA<AWindDirectionalSource>())
+			{
+				AWindDirectionalSource* WindSourceActor = CastChecked<AWindDirectionalSource>(ActorHitProxy->Actor);
+				if (WindSourceActor->IsSelectable())
+				{
+					GetAnimPreviewScene().SetSelectedActor(WindSourceActor);
+					return true;
+				}
+			}
+		}
+		GetAnimPreviewScene().BroadcastMeshClick(ActorHitProxy, Click); // This can pop up menu which redraws viewport and invalidates HitProxy!
+		return true;
+	}
+
+	return false;
+}
+
+bool FSkeletonSelectionEditMode::CanCycleWidgetMode() const
+{
+	int32 SelectedBoneIndex = GetAnimPreviewScene().GetSelectedBoneIndex();
+	USkeletalMeshSocket* SelectedSocket = GetAnimPreviewScene().GetSelectedSocket().Socket;
+	AActor* SelectedActor = GetAnimPreviewScene().GetSelectedActor();
+
+	return (SelectedBoneIndex >= 0 || SelectedSocket || SelectedActor != nullptr);
+}
+
+#undef LOCTEXT_NAMESPACE
diff --git a/Engine/Source/Editor/Persona/Public/IPersonaEditorModeManager.h b/Engine/Source/Editor/Persona/Public/IPersonaEditorModeManager.h
--- a/Engine/Source/Editor/Persona/Public/IPersonaEditorModeManager.h
+++ b/Engine/Source/Editor/Persona/Public/IPersonaEditorModeManager.h
@@ -35,7 +35,7 @@ public:
 class IPersonaEditorModeManager;
 
 UCLASS()
-class UPersonaEditorModeManagerContext : public UObject, public IPersonaManagerContext
+class PERSONA_API UPersonaEditorModeManagerContext : public UObject, public IPersonaManagerContext
 {
 	GENERATED_BODY()
 public:
diff --git a/Engine/Source/Editor/Persona/Public/SkeletonSelectionEditMode.h b/Engine/Source/Editor/Persona/Public/SkeletonSelectionEditMode.h
new file mode 100644
--- /dev/null
+++ b/Engine/Source/Editor/Persona/Public/SkeletonSelectionEditMode.h
@@ -0,0 +1,68 @@
+// Copyright Epic Games, Inc. All Rights Reserved.
+
+#pragma once
+
+#include "CoreMinimal.h"
+#include "UnrealWidgetFwd.h"
+#include "IPersonaEditMode.h"
+
+class FCanvas;
+class FEditorViewportClient;
+class FPrimitiveDrawInterface;
+class FSceneView;
+class FViewport;
+class USkeletalMeshSocket;
+class USkeletalMeshComponent;
+struct FViewportClick;
+struct FSelectedSocketInfo;
+class FViewportClient;
+struct FReferenceSkeleton;
+
+class FSkeletonSelectionEditMode : public IPersonaEditMode
+{
+public:
+	FSkeletonSelectionEditMode();
+
+	/** IPersonaEditMode interface */
+	virtual bool GetCameraTarget(FSphere& OutTarget) const override;
+	virtual class IPersonaPreviewScene& GetAnimPreviewScene() const override;
+	virtual void GetOnScreenDebugInfo(TArray<FText>& OutDebugInfo) const override;
+
+	/** FEdMode interface */
+	virtual bool StartTracking(FEditorViewportClient* InViewportClient, FViewport* InViewport) override;
+	virtual bool EndTracking(FEditorViewportClient* InViewportClient, FViewport* InViewport) override;
+	virtual bool InputDelta(FEditorViewportClient* InViewportClient, FViewport* InViewport, FVector& InDrag, FRotator& InRot, FVector& InScale) override;
+	virtual void DrawHUD(FEditorViewportClient* ViewportClient, FViewport* Viewport, const FSceneView* View, FCanvas* Canvas) override;
+	virtual bool AllowWidgetMove() override;
+	virtual bool ShouldDrawWidget() const override;
+	virtual bool UsesTransformWidget() const override;
+	virtual bool UsesTransformWidget(UE::Widget::EWidgetMode CheckMode) const override;
+	virtual FVector GetWidgetLocation() const override;
+	virtual bool GetCustomDrawingCoordinateSystem(FMatrix& InMatrix, void* InData) override;
+	virtual bool GetCustomInputCoordinateSystem(FMatrix& InMatrix, void* InData) override;
+	virtual bool HandleClick(FEditorViewportClient* InViewportClient, HHitProxy *HitProxy, const FViewportClick &Click) override;
+	virtual bool CanCycleWidgetMode() const override;
+
+protected:
+	FTransform GetBoneTransform(const int32 BoneIndex) const;
+	FTransform GetSocketTransform(const USkeletalMeshSocket* Socket) const;
+	FTransform GetWorldSpaceBoneTransform(const FReferenceSkeleton& ReferenceSkeleton, const int32 BoneIndex) const;
+	const FReferenceSkeleton& GetReferenceSkeletonForComponent(const USkeletalMeshComponent* Component) const;
+
+private:
+	/** Duplicates and selects a socket when we alt-drag */
+	FSelectedSocketInfo DuplicateAndSelectSocket(const FSelectedSocketInfo& SocketInfoToDuplicate);
+
+	/** Check whether the currently selected bone is in the required bones list */
+	bool IsSelectedBoneRequired() const;
+
+	/** Unscale a viewport's size by its DPI factor */
+	static FIntPoint GetDPIUnscaledSize(FViewport* Viewport, FViewportClient* Client);
+
+private:
+	/** Whether we are currently in a manipulation  */
+	bool bManipulating;	
+
+	/** Whether we are currently in a transaction  */
+	bool bInTransaction;
+};
