Remove impossible Cast from AnimComposite to AnimMontage
#rnx

diff --git a/Engine/Source/Editor/Persona/Private/SAnimCompositePanel.cpp b/Engine/Source/Editor/Persona/Private/SAnimCompositePanel.cpp
--- a/Engine/Source/Editor/Persona/Private/SAnimCompositePanel.cpp
+++ b/Engine/Source/Editor/Persona/Private/SAnimCompositePanel.cpp
@@ -73,8 +73,6 @@ void SAnimCompositePanel::Update()
 			SAssignNew( CompositeSlots, SVerticalBox )
 			);
 
-		UAnimMontage* AnimMontage = Cast<UAnimMontage>(Composite);
-
 		CompositeSlots->AddSlot()
 			.AutoHeight()
 			.VAlign(VAlign_Center)
@@ -87,7 +85,7 @@ void SAnimCompositePanel::Update()
 				.OnGetNodeColor(this,  &SAnimCompositePanel::HandleGetNodeColor)
 				.TrackMaxValue(Composite->GetPlayLength())
 				.TrackNumDiscreteValues(Composite->GetNumberOfSampledKeys())
-				.bChildAnimMontage(AnimMontage && AnimMontage->HasParentAsset())
+				.bChildAnimMontage(false)
 				.OnAnimSegmentNodeClicked( this, &SAnimCompositePanel::ShowSegmentInDetailsView )
 				.OnPreAnimUpdate( this, &SAnimCompositePanel::PreAnimUpdate )
 				.OnPostAnimUpdate( this, &SAnimCompositePanel::PostAnimUpdate )
