Fix Persona ReferencePose Toolbar button by making it a "ToolBarButton" instead of MenuEntry.

diff --git a/Engine/Source/Editor/Persona/Private/PersonaModule.cpp b/Engine/Source/Editor/Persona/Private/PersonaModule.cpp
--- a/Engine/Source/Editor/Persona/Private/PersonaModule.cpp
+++ b/Engine/Source/Editor/Persona/Private/PersonaModule.cpp
@@ -1219,12 +1219,12 @@ void FPersonaModule::AddCommonToolbarExtensions(UToolMenu* InToolMenu, const FCo
 
 	if (InArgs.bReferencePose)
 	{
-		PersonaSection.AddMenuEntry(
+		PersonaSection.AddEntry(FToolMenuEntry::InitToolBarButton(
 			"ReferencePose",
+			FUIAction(FExecuteAction::CreateStatic(&UE::Persona::Private::ShowReferencePose, WeakPersonaToolkit)),
 			LOCTEXT("ShowReferencePose", "Reference Pose"),
 			LOCTEXT("ShowReferencePoseTooltip", "Show the reference pose. Clears all bone modifications. More advanced settings are available in Preview Scene Settings."),
-			FSlateIcon(FAppStyle::GetAppStyleSetName(), "Persona.ToggleReferencePose", "Persona.ToggleReferencePose.Small"),
-			FUIAction(FExecuteAction::CreateStatic(&UE::Persona::Private::ShowReferencePose, WeakPersonaToolkit))
+			FSlateIcon(FAppStyle::GetAppStyleSetName(), "Persona.ToggleReferencePose", "Persona.ToggleReferencePose.Small"))
 		);
 	}
 
