Generalized SkeletalMeshEditorToolbox to be used throughout Persona-based editors

diff --git a/Engine/Source/Editor/Persona/Private/PersonaAssetEditorToolkit.cpp b/Engine/Source/Editor/Persona/Private/PersonaAssetEditorToolkit.cpp
--- a/Engine/Source/Editor/Persona/Private/PersonaAssetEditorToolkit.cpp
+++ b/Engine/Source/Editor/Persona/Private/PersonaAssetEditorToolkit.cpp
@@ -6,8 +6,55 @@
 #include "AssetEditorModeManager.h"
 #include "Modules/ModuleManager.h"
 #include "IPersonaEditorModeManager.h"
+#include "PersonaTabs.h"
 
 void FPersonaAssetEditorToolkit::CreateEditorModeManager()
 {
 	EditorModeManager = MakeShareable(FModuleManager::LoadModuleChecked<FPersonaModule>("Persona").CreatePersonaEditorModeManager());
+
+	// Make sure we get told when the editor mode changes so we can switch to the appropriate tab
+	// if there's a toolbox available.
+	GetEditorModeManager().OnEditorModeIDChanged().AddSP(this, &FPersonaAssetEditorToolkit::OnEditorModeIdChanged);
+}
+
+void FPersonaAssetEditorToolkit::OnToolkitHostingStarted(const TSharedRef<IToolkit>& Toolkit)
+{
+	if (ensure(!HostedToolkit.IsValid()))
+	{
+		HostedToolkit = Toolkit;
+		OnAttachToolkit.Broadcast(Toolkit);
+	}
+}
+
+void FPersonaAssetEditorToolkit::OnToolkitHostingFinished(const TSharedRef<IToolkit>& Toolkit)
+{
+	if (ensure(Toolkit == HostedToolkit))
+	{
+		HostedToolkit.Reset();
+		OnDetachToolkit.Broadcast(Toolkit);
+	}
+}
+
+void FPersonaAssetEditorToolkit::OnEditorModeIdChanged(const FEditorModeID& ModeChangedID, bool bIsEnteringMode)
+{
+	if (GetEditorModeManager().IsDefaultMode(ModeChangedID))
+	{
+		return;
+	}
+
+	if (bIsEnteringMode)
+	{
+		if (HostedToolkit.IsValid())
+		{
+			TabManager->TryInvokeTab(FPersonaTabs::ToolboxID);
+		}
+	}
+	else
+	{
+		const TSharedPtr<SDockTab> ToolboxTab = TabManager->FindExistingLiveTab(FPersonaTabs::ToolboxID);
+		if (ToolboxTab.IsValid())
+		{
+			ToolboxTab->RequestCloseTab();
+		}
+	}
 }
diff --git a/Engine/Source/Editor/Persona/Private/PersonaModule.cpp b/Engine/Source/Editor/Persona/Private/PersonaModule.cpp
--- a/Engine/Source/Editor/Persona/Private/PersonaModule.cpp
+++ b/Engine/Source/Editor/Persona/Private/PersonaModule.cpp
@@ -94,6 +94,7 @@
 #include "AnimSequenceTimelineCommands.h"
 #include "SAnimCurvePicker.h"
 #include "SAnimMontageSectionsPanel.h"
+#include "SPersonaToolBox.h"
 #include "Animation/AnimSequenceHelpers.h"
 #include "SkeletalMeshReferenceSectionDetails.h"
 #include "PersonaToolMenuContext.h"
@@ -358,6 +359,11 @@ TSharedRef<SWidget> FPersonaModule::CreateBlendSpacePreviewWidget(const FBlendSp
 		.OnGetBlendSpaceSampleName(InArgs.OnGetBlendSpaceSampleName);
 }
 
+TSharedRef<FWorkflowTabFactory> FPersonaModule::CreatePersonaToolboxTabFactory(const TSharedRef<FPersonaAssetEditorToolkit>& InHostingApp) const
+{
+	return MakeShareable(new FToolBoxSummoner(InHostingApp));
+}
+
 TSharedRef<SWidget> FPersonaModule::CreateBlendSpacePreviewWidget(TAttribute<const UBlendSpace*> InBlendSpace, TAttribute<FVector> InPosition, TAttribute<FVector> InFilteredPosition) const
 {
 PRAGMA_DISABLE_DEPRECATION_WARNINGS
diff --git a/Engine/Source/Editor/Persona/Private/SPersonaToolBox.cpp b/Engine/Source/Editor/Persona/Private/SPersonaToolBox.cpp
new file mode 100644
--- /dev/null
+++ b/Engine/Source/Editor/Persona/Private/SPersonaToolBox.cpp
@@ -0,0 +1,333 @@
+// Copyright Epic Games, Inc. All Rights Reserved.
+
+#include "SPersonaToolBox.h"
+
+#include "EditorModeManager.h"
+#include "EdMode.h"
+#include "PersonaTabs.h"
+#include "Framework/MultiBox/MultiBoxBuilder.h"
+#include "Tools/UEdMode.h"
+#include "Widgets/Input/SCheckBox.h"
+#include "Widgets/Layout/SWidgetSwitcher.h"
+#include "Widgets/Layout/SUniformWrapPanel.h"
+#include "Widgets/Layout/SScrollBox.h"
+#include "Widgets/Input/SSegmentedControl.h"
+
+#define LOCTEXT_NAMESPACE "PersonaToolbox"
+
+FToolBoxSummoner::FToolBoxSummoner(TSharedPtr<FPersonaAssetEditorToolkit> InPersonaEditorToolkit) : FWorkflowTabFactory(FPersonaTabs::ToolboxID, InPersonaEditorToolkit), PersonaAssetEditorToolkit(InPersonaEditorToolkit)
+{
+	TabLabel = LOCTEXT("ToolboxTabTitle", "Toolbox");
+	TabIcon = FSlateIcon(FAppStyle::Get().GetStyleSetName(), "LevelEditor.Tabs.Modes");
+
+	bIsSingleton = true;
+
+	ViewMenuDescription = LOCTEXT("ToolBoxMenu", "Toolbox");
+	ViewMenuTooltip = LOCTEXT("ToolBoxMenu_ToolTip", "Used to display Edit Mode widget.");	
+}
+
+TSharedRef<SDockTab> FToolBoxSummoner::SpawnTab(const FWorkflowTabSpawnInfo& Info) const
+{
+	TSharedRef<SDockTab> DockTab = FWorkflowTabFactory::SpawnTab(Info);
+
+	const TWeakPtr<FPersonaAssetEditorToolkit> WeakToolkit = PersonaAssetEditorToolkit;	
+	DockTab->SetOnTabClosed(SDockTab::FOnTabClosedCallback::CreateLambda([WeakToolkit](TSharedRef<SDockTab> DockTab)
+	{
+		if (WeakToolkit.IsValid())
+		{
+			WeakToolkit.Pin()->GetEditorModeManager().ActivateDefaultMode();
+		}
+	}));
+
+	return DockTab;
+}
+
+TSharedRef<SWidget> FToolBoxSummoner::CreateTabBody(const FWorkflowTabSpawnInfo& Info) const
+{
+	return SNew(SBox)
+		.AddMetaData<FTagMetaData>(FTagMetaData(TEXT("ToolboxTab")))
+		[
+			SNew(SPersonaToolbox, PersonaAssetEditorToolkit.Pin().ToSharedRef())
+		];
+}
+
+bool FToolBoxSummoner::CanSpawnTab(const FSpawnTabArgs& SpawnArgs, TWeakPtr<FTabManager> WeakTabManager) const
+{
+	return FWorkflowTabFactory::CanSpawnTab(SpawnArgs, WeakTabManager) && PersonaAssetEditorToolkit.Pin()->GetHostedToolkit().IsValid();
+}
+
+SPersonaToolbox::~SPersonaToolbox()
+{
+	if (StatusBarMessageHandle.IsValid() && PersonaEditor.IsValid())
+	{
+		GEditor->GetEditorSubsystem<UStatusBarSubsystem>()->PopStatusBarMessage(PersonaEditor.Pin()->GetToolMenuName(), StatusBarMessageHandle);
+		StatusBarMessageHandle.Reset();
+	}
+
+	if (ModeToolBarContainer)
+	{
+		ModeToolBarContainer->SetContent(SNullWidget::NullWidget);
+	}
+	if (ModeToolHeader)
+	{
+		ModeToolHeader->SetContent(SNullWidget::NullWidget);
+	}
+	if (InlineContentHolder)
+	{
+		InlineContentHolder->SetContent(SNullWidget::NullWidget);
+	}
+
+	if (PersonaEditor.IsValid())
+	{
+		PersonaEditor.Pin()->GetOnAttachToolkit().RemoveAll(this);
+		PersonaEditor.Pin()->GetOnDetachToolkit().RemoveAll(this);
+	}
+	
+	PersonaEditor.Reset();
+}
+
+void SPersonaToolbox::Construct(const FArguments& InArgs, const TSharedRef<FPersonaAssetEditorToolkit>& InOwningEditor)
+{
+	PersonaEditor = InOwningEditor;
+	// Hook into delegates to handle attaching and detaching of toolkits
+	InOwningEditor->GetOnAttachToolkit().AddSP(this, &SPersonaToolbox::AttachToolkit);
+	InOwningEditor->GetOnDetachToolkit().AddSP(this, &SPersonaToolbox::DetachToolkit);
+
+	// Modeled after FModeToolkit::CreatePrimaryModePanel
+	ChildSlot
+	[
+		SNew(SBorder)
+		.BorderImage(FAppStyle::GetBrush("ToolPanel.GroupBorder"))
+		.Padding(0.0f)
+		[
+			SNew(SVerticalBox)
+			+ SVerticalBox::Slot()
+			.AutoHeight()
+			.HAlign(HAlign_Left)
+			[
+				SAssignNew(ModeToolBarContainer, SBorder)
+				.BorderImage(FAppStyle::GetBrush("ToolPanel.GroupBorder"))
+				.Padding(FMargin(4, 0, 0, 0))
+			]
+
+			+ SVerticalBox::Slot()
+			.FillHeight(1.0f)
+			[
+				SNew(SVerticalBox)
+				+SVerticalBox::Slot()
+				.Padding(0.0, 8.0, 0.0, 0.0)
+				.AutoHeight()
+				[
+					SAssignNew(ModeToolHeader, SBorder)
+					.BorderImage(FAppStyle::GetBrush("ToolPanel.GroupBorder"))
+				]
+				+ SVerticalBox::Slot()
+				.FillHeight(1)
+				[
+					SAssignNew(InlineContentHolder, SBorder)
+					.BorderImage(FAppStyle::GetBrush("ToolPanel.GroupBorder"))
+					.Visibility(this, &SPersonaToolbox::GetInlineContentHolderVisibility)
+				]
+			]
+		]
+	];
+
+	// Directly attach toolkit if one is being hosted
+	if (InOwningEditor->GetHostedToolkit().IsValid())
+	{
+		AttachToolkit(InOwningEditor->GetHostedToolkit().ToSharedRef());
+	}
+}
+
+void SPersonaToolbox::AttachToolkit(const TSharedRef<IToolkit>& InToolkit)
+{
+	UpdateInlineContent(InToolkit, InToolkit->GetInlineContent());
+}
+
+void SPersonaToolbox::DetachToolkit(const TSharedRef<IToolkit>& InToolkit)
+{
+	UpdateInlineContent(nullptr, SNullWidget::NullWidget);
+}
+
+void SPersonaToolbox::SetOwningTab(const TSharedRef<SDockTab>& InOwningTab)
+{
+	OwningTab = InOwningTab;
+}
+
+void SPersonaToolbox::UpdateInlineContent(const TSharedPtr<IToolkit>& Toolkit, TSharedPtr<SWidget> InlineContent)
+{
+	// The display name that the owning tab should have as its label
+	FText TabName;
+
+	// The icon that should be displayed in the parent tab
+	const FSlateBrush* TabIcon = nullptr;
+
+	if (StatusBarMessageHandle.IsValid() && PersonaEditor.IsValid())
+	{
+		GEditor->GetEditorSubsystem<UStatusBarSubsystem>()->PopStatusBarMessage(PersonaEditor.Pin()->GetToolMenuName(), StatusBarMessageHandle);
+		StatusBarMessageHandle.Reset();	
+
+		const TWeakPtr<FModeToolkit> ModeToolkit = StaticCastSharedPtr<FModeToolkit>(Toolkit);
+		if (ModeToolkit.IsValid())
+		{
+			TabName = ModeToolkit.Pin()->GetEditorModeDisplayName();
+			TabIcon = ModeToolkit.Pin()->GetEditorModeIcon().GetSmallIcon();
+
+			UpdatePalette(ModeToolkit);
+			
+			// Show the name of the active tool in the statusbar.
+			// FIXME: We should also be showing Ctrl/Shift/Alt LMB/RMB shortcuts.
+			StatusBarMessageHandle = GEditor->GetEditorSubsystem<UStatusBarSubsystem>()->PushStatusBarMessage(
+					PersonaEditor.Pin()->GetToolMenuName(),
+					TAttribute<FText>::Create(TAttribute<FText>::FGetter::CreateSP(ModeToolkit.Pin()->AsShared(), &FModeToolkit::GetActiveToolDisplayName)));
+			
+		}
+	}
+	else
+	{
+		TabName = NSLOCTEXT("Persona", "ToolboxTab", "Toolbox");
+		TabIcon = FAppStyle::Get().GetBrush("LevelEditor.Tabs.Modes");
+	}
+
+	if (InlineContent.IsValid() && InlineContentHolder.IsValid())
+	{
+		InlineContentHolder->SetContent(InlineContent.ToSharedRef());
+	}
+
+	const TSharedPtr<SDockTab> OwningTabPinned = OwningTab.Pin();
+	if (OwningTabPinned.IsValid())
+	{
+		OwningTabPinned->SetLabel(TabName);
+		OwningTabPinned->SetTabIcon(TabIcon);
+	}
+}
+
+void SPersonaToolbox::UpdatePalette(const TWeakPtr<FModeToolkit>& InWeakToolkit) const
+{
+	if (!InWeakToolkit.IsValid())
+	{
+		return;
+	}
+
+	if (const TSharedPtr<FModeToolkit> Toolkit = InWeakToolkit.Pin())
+	{
+		if (Toolkit->HasIntegratedToolPalettes())
+		{
+			TSharedRef<SSegmentedControl<FName>> PaletteTabBox = SNew(SSegmentedControl<FName>)
+			.UniformPadding(FMargin(8.f, 3.f))
+			.Value_Lambda([InWeakToolkit]()
+			{
+				return InWeakToolkit.IsValid() ? InWeakToolkit.Pin()->GetCurrentPalette() : NAME_None;
+			})
+			.OnValueChanged_Lambda([InWeakToolkit](const FName& Palette)
+			{
+				if (InWeakToolkit.IsValid())
+				{
+					InWeakToolkit.Pin()->SetCurrentPalette(Palette);
+				}
+			});
+
+			// Only show if there's more than one entry.
+			PaletteTabBox->SetVisibility(TAttribute<EVisibility>::Create(
+				TAttribute<EVisibility>::FGetter::CreateLambda([PaletteTabBox]() -> EVisibility { 
+					return PaletteTabBox->NumSlots() > 1 ? EVisibility::Visible : EVisibility::Collapsed;
+				})));
+			
+
+			const TSharedPtr<FModeToolkit> ModeToolkit = InWeakToolkit.Pin();
+			
+			// Also build the toolkit here
+			TArray<FName> PaletteNames;
+			ModeToolkit->GetToolPaletteNames(PaletteNames);
+
+			const TSharedRef<SWidgetSwitcher> PaletteSwitcher = SNew(SWidgetSwitcher)
+				.WidgetIndex_Lambda( [PaletteNames, InWeakToolkit] () -> int32
+				{
+					if (InWeakToolkit.IsValid())
+					{
+						int32 FoundIndex;
+						if (PaletteNames.Find(InWeakToolkit.Pin()->GetCurrentPalette(), FoundIndex))
+						{
+							return FoundIndex;	
+						}
+					}
+					return 0;
+				});
+					
+			const TSharedPtr<FUICommandList> CommandList = ModeToolkit->GetToolkitCommands();
+			for(const FName& Palette : PaletteNames)
+			{
+				const FName ToolbarCustomizationName = ModeToolkit->GetEditorMode() ?
+					ModeToolkit->GetEditorMode()->GetModeInfo().ToolbarCustomizationName :
+					ModeToolkit->GetScriptableEditorMode()->GetModeInfo().ToolbarCustomizationName;
+				FUniformToolBarBuilder ModeToolbarBuilder(CommandList, FMultiBoxCustomization(ToolbarCustomizationName));
+				ModeToolbarBuilder.SetStyle(&FAppStyle::Get(), "PaletteToolBar");
+
+				ModeToolkit->BuildToolPalette(Palette, ModeToolbarBuilder);
+
+				TSharedRef<SWidget> PaletteWidget = ModeToolbarBuilder.MakeWidget();
+				constexpr bool bRebuildChildren = false;
+				PaletteTabBox->AddSlot(Palette, false)
+				.Text(ModeToolkit->GetToolPaletteDisplayName(Palette));
+
+				PaletteSwitcher->AddSlot()
+				[
+					PaletteWidget
+				];
+			}
+
+			PaletteTabBox->RebuildChildren();
+
+			if (ModeToolHeader)
+			{
+				ModeToolHeader->SetContent
+				(
+					SNew(SVerticalBox)
+					+SVerticalBox::Slot()
+					.Padding(8.f, 0.f, 0.f, 8.f)
+					.AutoHeight()
+					.HAlign(HAlign_Center)
+					[
+						PaletteTabBox
+					]
+					+SVerticalBox::Slot()
+					.AutoHeight()
+					[
+						PaletteSwitcher
+					]
+				);
+			}
+		}
+		else
+		{
+			if (ModeToolHeader)
+			{
+				ModeToolHeader->SetContent(SNullWidget::NullWidget);
+			}
+			if (ModeToolBarContainer)
+			{
+				ModeToolBarContainer->SetContent(SNullWidget::NullWidget);
+			}
+		}
+
+		if (InlineContentHolder.IsValid())
+		{
+			if (const TSharedPtr<SWidget> InlineContent = Toolkit->GetInlineContent())
+			{
+				InlineContentHolder->SetContent(
+				SNew(SScrollBox)
+				+ SScrollBox::Slot()
+				[
+					InlineContent.ToSharedRef()
+				]);
+			}
+		}
+	}
+}
+struct 
+EVisibility SPersonaToolbox::GetInlineContentHolderVisibility() const
+{
+	return InlineContentHolder->GetContent() == SNullWidget::NullWidget ? EVisibility::Collapsed : EVisibility::Visible;
+}
+
+#undef LOCTEXT_NAMESPACE // "PersonaToolbox"
\ No newline at end of file
diff --git a/Engine/Source/Editor/Persona/Private/TabSpawners.cpp b/Engine/Source/Editor/Persona/Private/TabSpawners.cpp
--- a/Engine/Source/Editor/Persona/Private/TabSpawners.cpp
+++ b/Engine/Source/Editor/Persona/Private/TabSpawners.cpp
@@ -76,6 +76,8 @@ const FName FPersonaTabs::DetailsID("DetailsTab");
 
 const FName FPersonaTabs::FindReplaceID("FindReplaceTab");
 
+const FName FPersonaTabs::ToolboxID("ToolBox");
+
 /////////////////////////////////////////////////////
 // FPersonaMode
 
diff --git a/Engine/Source/Editor/Persona/Public/PersonaAssetEditorToolkit.h b/Engine/Source/Editor/Persona/Public/PersonaAssetEditorToolkit.h
--- a/Engine/Source/Editor/Persona/Public/PersonaAssetEditorToolkit.h
+++ b/Engine/Source/Editor/Persona/Public/PersonaAssetEditorToolkit.h
@@ -4,10 +4,33 @@
 
 #include "WorkflowOrientedApp/WorkflowCentricApplication.h"
 
+DECLARE_MULTICAST_DELEGATE_OneParam(FOnAttachToolkit, const TSharedRef<IToolkit>&);
+DECLARE_MULTICAST_DELEGATE_OneParam(FOnDetachToolkit, const TSharedRef<IToolkit>&);
+
 /** Persona asset editor toolkit wrapper, used to auto inject the persona editor mode manager */
 class PERSONA_API FPersonaAssetEditorToolkit : public FWorkflowCentricApplication
 {
 public:
 	/** FAssetEditorToolkit interface  */
 	virtual void CreateEditorModeManager() override;
+
+	// IToolkitHost Interface
+	virtual void OnToolkitHostingStarted(const TSharedRef<class IToolkit>& Toolkit) override;
+	virtual void OnToolkitHostingFinished(const TSharedRef<class IToolkit>& Toolkit) override;
+
+	// Returns the currently hosted toolkit. Can be invalid if no toolkit is being hosted.
+	TSharedPtr<IToolkit> GetHostedToolkit() const { return HostedToolkit; }
+	FOnAttachToolkit& GetOnAttachToolkit() { return OnAttachToolkit; }
+	FOnDetachToolkit& GetOnDetachToolkit() { return OnDetachToolkit; }
+	
+protected:
+	virtual void OnEditorModeIdChanged(const FEditorModeID& ModeChangedID, bool bIsEnteringMode);
+	
+protected:
+	FOnAttachToolkit OnAttachToolkit;
+	FOnDetachToolkit OnDetachToolkit;
+	
+	// The toolkit we're currently hosting.
+	TSharedPtr<IToolkit> HostedToolkit;
+		
 };
diff --git a/Engine/Source/Editor/Persona/Public/PersonaModule.h b/Engine/Source/Editor/Persona/Public/PersonaModule.h
--- a/Engine/Source/Editor/Persona/Public/PersonaModule.h
+++ b/Engine/Source/Editor/Persona/Public/PersonaModule.h
@@ -10,6 +10,7 @@
 #include "Editor.h"
 #include "PersonaDelegates.h"
 #include "Factories/FbxImportUI.h"
+#include "SPersonaToolBox.h"
 
 class FBlueprintEditor;
 class IDetailsView;
@@ -333,6 +334,9 @@ public:
 	/** Create a tab factory for editing slot names and groups */
 	virtual TSharedRef<FWorkflowTabFactory> CreateSkeletonSlotNamesTabFactory(const TSharedRef<class FWorkflowCentricApplication>& InHostingApp, const TSharedRef<class IEditableSkeleton>& InEditableSkeleton, FOnObjectSelected InOnObjectSelected) const;
 
+	/** Create a toolbox tab factory */
+	virtual TSharedRef<FWorkflowTabFactory> CreatePersonaToolboxTabFactory(const TSharedRef<class FPersonaAssetEditorToolkit>& InHostingApp) const;
+
 	/** Deprecated */
 	UE_DEPRECATED(5.0, "Please use the overload that takes a FBlendSpacePreviewArgs struct")
 	virtual TSharedRef<SWidget> CreateBlendSpacePreviewWidget(TAttribute<const UBlendSpace*> InBlendSpace, TAttribute<FVector> InBlendPosition, TAttribute<FVector> InFilteredBlendPosition) const;
diff --git a/Engine/Source/Editor/Persona/Public/PersonaTabs.h b/Engine/Source/Editor/Persona/Public/PersonaTabs.h
--- a/Engine/Source/Editor/Persona/Public/PersonaTabs.h
+++ b/Engine/Source/Editor/Persona/Public/PersonaTabs.h
@@ -47,6 +47,7 @@ struct PERSONA_API FPersonaTabs
 	static const FName DetailsID;
 	static const FName FindReplaceID;
 
+	static const FName ToolboxID;
 private:
 	FPersonaTabs() {}
 };
\ No newline at end of file
diff --git a/Engine/Source/Editor/Persona/Public/SPersonaToolBox.h b/Engine/Source/Editor/Persona/Public/SPersonaToolBox.h
new file mode 100644
--- /dev/null
+++ b/Engine/Source/Editor/Persona/Public/SPersonaToolBox.h
@@ -0,0 +1,82 @@
+// Copyright Epic Games, Inc. All Rights Reserved.
+
+#pragma once
+
+#include "PersonaAssetEditorToolkit.h"
+
+#include "CoreMinimal.h"
+#include "StatusBarSubsystem.h"
+#include "Widgets/DeclarativeSyntaxSupport.h"
+#include "Widgets/Docking/SDockTab.h"
+#include "Widgets/SCompoundWidget.h"
+#include "WorkflowOrientedApp/WorkflowTabFactory.h"
+
+#include "Delegates/DelegateCombinations.h"
+#include "Templates/SharedPointer.h"
+#include "Toolkits/IToolkit.h"
+
+class FModeToolkit;
+class FPersonaAssetEditorToolkit;
+class IToolkit;
+class SBorder;
+
+
+struct FToolBoxSummoner : public FWorkflowTabFactory
+{
+public:
+	FToolBoxSummoner(TSharedPtr<FPersonaAssetEditorToolkit> InPersonaEditorToolkit);
+
+	virtual TSharedRef<SDockTab> SpawnTab(const FWorkflowTabSpawnInfo& Info) const override;
+	virtual TSharedRef<SWidget> CreateTabBody(const FWorkflowTabSpawnInfo& Info) const override;
+	virtual bool CanSpawnTab(const FSpawnTabArgs& SpawnArgs, TWeakPtr<FTabManager> WeakTabManager) const override;
+
+protected:
+	TWeakPtr<FPersonaAssetEditorToolkit> PersonaAssetEditorToolkit;
+};
+
+/** Widget to surface an Edit Mode its widget, similar to LevelEditor approach */
+class PERSONA_API SPersonaToolbox : public SCompoundWidget
+{
+public:
+	SLATE_BEGIN_ARGS( SPersonaToolbox ) {}
+	SLATE_END_ARGS()
+
+	~SPersonaToolbox();
+
+	void Construct(const FArguments& InArgs, const TSharedRef<FPersonaAssetEditorToolkit>& InOwningEditor);
+
+	// Called to attach a toolkit to this toolbox widget (e.g. when the mode changes).
+	void AttachToolkit(const TSharedRef<IToolkit>& InToolkit);
+
+	// Called to detach a toolkit from this toolbox widget (e.g. when the mode changes).
+	void DetachToolkit(const TSharedRef<IToolkit>& InToolkit);
+
+	// Sets the owning tab so that we can update the tab icon based on what the toolkit provides.
+	void SetOwningTab(const TSharedRef<SDockTab>& InOwningTab);
+
+private:	
+	void UpdateInlineContent(const TSharedPtr<IToolkit>& Toolkit, TSharedPtr<SWidget> InlineContent);
+
+	void UpdatePalette(const TWeakPtr<FModeToolkit>& InWeakToolkit) const;
+
+	// Whether to show the inline content or not.
+	EVisibility GetInlineContentHolderVisibility() const;
+
+	// The owning tab where the toolkit's interface is hosted. We need this to update the icons
+	// and title as the toolkit changes.
+	TWeakPtr<SDockTab> OwningTab;
+
+	TWeakPtr<FPersonaAssetEditorToolkit> PersonaEditor;
+
+	// Inline content area for editor modes
+	TSharedPtr<SBorder> InlineContentHolder;
+
+	// The container holding the mode toolbar
+	TSharedPtr<SBorder> ModeToolBarContainer;
+
+	// The active tool header area
+	TSharedPtr<SBorder> ModeToolHeader;
+
+	// The global status bar to show the name of the active tool.
+	FStatusBarMessageHandle StatusBarMessageHandle;
+};
\ No newline at end of file
diff --git a/Engine/Source/Editor/SkeletalMeshEditor/Private/SSkeletalMeshEditorToolbox.cpp b/Engine/Source/Editor/SkeletalMeshEditor/Private/SSkeletalMeshEditorToolbox.cpp
deleted file mode 100644
diff --git a/Engine/Source/Editor/SkeletalMeshEditor/Private/SSkeletalMeshEditorToolbox.h b/Engine/Source/Editor/SkeletalMeshEditor/Private/SSkeletalMeshEditorToolbox.h
deleted file mode 100644
diff --git a/Engine/Source/Editor/SkeletalMeshEditor/Private/SkeletalMeshEditor.cpp b/Engine/Source/Editor/SkeletalMeshEditor/Private/SkeletalMeshEditor.cpp
--- a/Engine/Source/Editor/SkeletalMeshEditor/Private/SkeletalMeshEditor.cpp
+++ b/Engine/Source/Editor/SkeletalMeshEditor/Private/SkeletalMeshEditor.cpp
@@ -4,8 +4,6 @@
 
 #include "SkeletalMeshEditorCommands.h"
 #include "SkeletalMeshEditorMode.h"
-#include "SSkeletalMeshEditorToolbox.h"
-
 
 #include "Algo/Transform.h"
 #include "Animation/DebugSkelMeshComponent.h"
@@ -213,17 +211,6 @@ bool FSkeletalMeshEditor::OnRequestClose(EAssetEditorCloseReason InCloseReason)
 
 void FSkeletalMeshEditor::RegisterTabSpawners(const TSharedRef<class FTabManager>& InTabManager)
 {
-	WorkspaceMenuCategory = InTabManager->AddLocalWorkspaceMenuCategory(LOCTEXT("WorkspaceMenu_SkeletalMeshEditor", "Skeletal Mesh Editor"));
-	
-	InTabManager->RegisterTabSpawner(
-		SkeletalMeshEditorTabs::ToolboxDetailsTab, 
-		FOnSpawnTab::CreateSP(this, &FSkeletalMeshEditor::SpawnToolboxTab),
-		FCanSpawnTab::CreateSP(this, &FSkeletalMeshEditor::CanSpawnToolboxTab)
-		)
-		.SetDisplayName(LOCTEXT("ToolboxTab", "Toolbox"))
-		.SetGroup(WorkspaceMenuCategory.ToSharedRef())
-		.SetIcon(FSlateIcon(FAppStyle::GetAppStyleSetName(), "LevelEditor.Tabs.Modes" ));
-
 	FAssetEditorToolkit::RegisterTabSpawners(InTabManager);
 }
 
@@ -346,7 +333,6 @@ TSharedPtr<FSkeletalMeshEditor> FSkeletalMeshEditor::GetSkeletalMeshEditor(const
 	return TSharedPtr<FSkeletalMeshEditor>();
 }
 
-
 void FSkeletalMeshEditor::OnEditorModeIdChanged(const FEditorModeID& ModeChangedID, bool bIsEnteringMode)
 {
 	if (GetEditorModeManager().IsDefaultMode(ModeChangedID))
@@ -382,40 +368,6 @@ void FSkeletalMeshEditor::OnEditorModeIdChanged(const FEditorModeID& ModeChanged
 	}
 }
 
-
-bool FSkeletalMeshEditor::CanSpawnToolboxTab(const FSpawnTabArgs& InArgs) const
-{
-	return HostedToolkit.IsValid();
-}
-
-
-TSharedRef<SDockTab> FSkeletalMeshEditor::SpawnToolboxTab(const FSpawnTabArgs& Args)
-{
-	ToolboxWidget = SNew(SSkeletalMeshEditorToolbox, SharedThis<ISkeletalMeshEditor>(this));
-
-	TSharedRef<SDockTab> DockTab = SNew(SDockTab)
-		.Label(LOCTEXT("ToolboxTab", "Toolbox"))
-		[
-			SNew(SBox)
-			.AddMetaData<FTagMetaData>(FTagMetaData(TEXT("ToolboxTab")))
-			[
-				ToolboxWidget.ToSharedRef()
-			]
-		];
-	
-	ToolboxWidget->AttachToolkit(HostedToolkit.ToSharedRef());
-
-	return DockTab;
-}
-
-
-void FSkeletalMeshEditor::OnToolboxTabClosed(TSharedRef<SDockTab> InClosedTab)
-{
-	// If the user closed the tab, then we want to go back to the base skel mesh editor mode.
-	GetEditorModeManager().ActivateDefaultMode();
-}
-
-
 void FSkeletalMeshEditor::RegisterReimportContextMenu(const FName InBaseMenuName)
 {
 	static auto ReimportMeshWithNewFileAction = [](const FToolMenuContext& InMenuContext, int32 SourceFileIndex)
@@ -709,35 +661,6 @@ void FSkeletalMeshEditor::InitToolMenuContext(FToolMenuContext& MenuContext)
 	MenuContext.AddObject(PersonaContext);
 }
 
-
-void FSkeletalMeshEditor::OnToolkitHostingStarted(const TSharedRef<IToolkit>& Toolkit)
-{
-	if (!HostedToolkit.IsValid())
-	{
-		HostedToolkit = Toolkit;
-
-		if (ToolboxWidget.IsValid())
-		{
-			ToolboxWidget->AttachToolkit(Toolkit);
-		}
-	}
-}
-
-
-void FSkeletalMeshEditor::OnToolkitHostingFinished(const TSharedRef<IToolkit>& Toolkit)
-{
-	if (Toolkit == HostedToolkit)
-	{
-		if (ToolboxWidget.IsValid())
-		{
-			ToolboxWidget->DetachToolkit(Toolkit);
-		}
-		
-		HostedToolkit.Reset();
-	}
-}
-
-
 void FSkeletalMeshEditor::AddViewportOverlayWidget(TSharedRef<SWidget> InOverlaidWidget)
 {
 	if (Viewport.IsValid())
diff --git a/Engine/Source/Editor/SkeletalMeshEditor/Private/SkeletalMeshEditor.h b/Engine/Source/Editor/SkeletalMeshEditor/Private/SkeletalMeshEditor.h
--- a/Engine/Source/Editor/SkeletalMeshEditor/Private/SkeletalMeshEditor.h
+++ b/Engine/Source/Editor/SkeletalMeshEditor/Private/SkeletalMeshEditor.h
@@ -71,9 +71,6 @@ public:
 	virtual FString GetWorldCentricTabPrefix() const override;
 	virtual FLinearColor GetWorldCentricTabColorScale() const override;
 	virtual void InitToolMenuContext(FToolMenuContext& MenuContext) override;
-	// IToolkitHost Interface
-	virtual void OnToolkitHostingStarted(const TSharedRef<class IToolkit>& Toolkit) override;
-	virtual void OnToolkitHostingFinished(const TSharedRef<class IToolkit>& Toolkit) override;
 
 	virtual void AddViewportOverlayWidget(TSharedRef<SWidget> InOverlaidWidget) override;
 	virtual void RemoveViewportOverlayWidget(TSharedRef<SWidget> InOverlaidWidget) override;
@@ -252,4 +249,4 @@ public:
 private:
 	TWeakPtr<FSkeletalMeshEditor> Editor;
 	FSkeletalMeshEditorNotifier Notifier;
-};
\ No newline at end of file
+};
diff --git a/Engine/Source/Editor/SkeletalMeshEditor/Private/SkeletalMeshEditorMode.cpp b/Engine/Source/Editor/SkeletalMeshEditor/Private/SkeletalMeshEditorMode.cpp
--- a/Engine/Source/Editor/SkeletalMeshEditor/Private/SkeletalMeshEditorMode.cpp
+++ b/Engine/Source/Editor/SkeletalMeshEditor/Private/SkeletalMeshEditorMode.cpp
@@ -41,6 +41,7 @@ FSkeletalMeshEditorMode::FSkeletalMeshEditorMode(TSharedRef<FWorkflowCentricAppl
 	TabFactories.RegisterFactory(PersonaModule.CreateCurveMetadataEditorTabFactory(InHostingApp, SkeletalMeshEditor->HandleGetAsset(), SkeletalMeshEditor->GetPersonaToolkit()->GetPreviewScene(), OnObjectsSelected));
 	TabFactories.RegisterFactory(CreateMeshControllerMappingTabFactory(InHostingApp, Cast<USkeletalMesh> (SkeletalMeshEditor->HandleGetAsset()), SkeletalMeshEditor->OnPostUndo));
 	TabFactories.RegisterFactory(PersonaModule.CreateAnimAssetFindReplaceTabFactory(InHostingApp, FAnimAssetFindReplaceConfig()));
+	TabFactories.RegisterFactory(PersonaModule.CreatePersonaToolboxTabFactory(SkeletalMeshEditor));
 
 	TabLayout = FTabManager::NewLayout("Standalone_SkeletalMeshEditor_Layout_v4.5")
 		->AddArea
diff --git a/Engine/Source/Editor/UnrealEd/Public/WorkflowOrientedApp/WorkflowTabFactory.h b/Engine/Source/Editor/UnrealEd/Public/WorkflowOrientedApp/WorkflowTabFactory.h
--- a/Engine/Source/Editor/UnrealEd/Public/WorkflowOrientedApp/WorkflowTabFactory.h
+++ b/Engine/Source/Editor/UnrealEd/Public/WorkflowOrientedApp/WorkflowTabFactory.h
@@ -134,7 +134,7 @@ protected:
 	UNREALED_API virtual TSharedRef<SDockTab> OnSpawnTab(const FSpawnTabArgs& SpawnArgs, TWeakPtr<FTabManager> WeakTabManager) const;
 
 	/** Callback function for checking to see if a tab can/should be spawned */
-	UNREALED_API bool CanSpawnTab(const FSpawnTabArgs& SpawnArgs, TWeakPtr<FTabManager> WeakTabManager) const;
+	UNREALED_API virtual bool CanSpawnTab(const FSpawnTabArgs& SpawnArgs, TWeakPtr<FTabManager> WeakTabManager) const;
 
 public:
 	// Creates the widget that will be the body of the new tab
diff --git a/Engine/Source/Runtime/Slate/Private/Framework/Docking/TabManager.cpp b/Engine/Source/Runtime/Slate/Private/Framework/Docking/TabManager.cpp
--- a/Engine/Source/Runtime/Slate/Private/Framework/Docking/TabManager.cpp
+++ b/Engine/Source/Runtime/Slate/Private/Framework/Docking/TabManager.cpp
@@ -2079,6 +2079,11 @@ TSharedPtr<SDockTab> FTabManager::SpawnTab(const FTabId& TabId, const TSharedPtr
 			bSpawningAllowedBySpawner = false;
 		}
 	}
+	else
+	{
+	
+		bSpawningAllowedBySpawner = false;
+	}
 
 	// The tab was allowed to be spawned but failed for some reason
 	if (bSpawningAllowedBySpawner && !NewTabWidget.IsValid())
@@ -2735,9 +2740,10 @@ void FGlobalTabmanager::SaveAllVisualState()
 
 	for( int32 ManagerIndex=0; ManagerIndex < SubTabManagers.Num(); ++ManagerIndex )
 	{
-		const TSharedPtr<FTabManager> SubManagerTab = SubTabManagers[ManagerIndex].TabManager.Pin();
+		const TSharedPtr<FTabManager> SubManagerTab = SubTabManagers[ManagerIndex].TabManager.Pin();		
 		if (SubManagerTab.IsValid())
 		{
+			ensure(SubTabManagers[ManagerIndex].MajorTab.IsValid());
 			SubManagerTab->SavePersistentLayout();
 		}
 	}
