Fix curve debugger not working correctly with pose watches when connected to an external instance

As external instances can be child classes, some class-level checks were failing to retrieve the correct set of pose watches, instead looking at the child anim BP which always has empty pose watches.

diff --git a/Engine/Source/Editor/Persona/Private/SAnimCurveViewer.cpp b/Engine/Source/Editor/Persona/Private/SAnimCurveViewer.cpp
--- a/Engine/Source/Editor/Persona/Private/SAnimCurveViewer.cpp
+++ b/Engine/Source/Editor/Persona/Private/SAnimCurveViewer.cpp
@@ -245,13 +245,17 @@ bool SAnimCurveListRow::GetActiveFlag(bool bMorphTarget) const
 			{
 				if(UAnimBlueprintGeneratedClass* AnimClass = Cast<UAnimBlueprintGeneratedClass>(AnimInstance->GetClass()))
 				{
-					const FAnimBlueprintDebugData& DebugData = AnimClass->GetAnimBlueprintDebugData();
-					for(const FAnimNodePoseWatch& AnimNodePoseWatch : DebugData.AnimNodePoseWatch)
+					// We have to grab our pose watches from the root class as no pose watches can be set on child anim BPs
+					if(const UAnimBlueprintGeneratedClass* RootClass = Cast<UAnimBlueprintGeneratedClass>(AnimClass->GetRootClass()))
 					{
-						if(AnimNodePoseWatch.PoseWatchPoseElement == PoseWatchPoseElement)
+						const FAnimBlueprintDebugData& DebugData = RootClass->AnimBlueprintDebugData;
+						for(const FAnimNodePoseWatch& AnimNodePoseWatch : DebugData.AnimNodePoseWatch)
 						{
-							UE::Anim::ECurveElementFlags Flags = AnimNodePoseWatch.GetCurves().GetFlags(Item->CurveName);
-							return EnumHasAnyFlags(Flags, bMorphTarget ? UE::Anim::ECurveElementFlags::MorphTarget : UE::Anim::ECurveElementFlags::Material);
+							if(AnimNodePoseWatch.PoseWatchPoseElement == PoseWatchPoseElement)
+							{
+								UE::Anim::ECurveElementFlags Flags = AnimNodePoseWatch.GetCurves().GetFlags(Item->CurveName);
+								return EnumHasAnyFlags(Flags, bMorphTarget ? UE::Anim::ECurveElementFlags::MorphTarget : UE::Anim::ECurveElementFlags::Material);
+							}
 						}
 					}
 				}
@@ -384,19 +388,23 @@ bool SAnimCurveListRow::GetActiveWeight(float& OutWeight) const
 			{
 				if(UAnimBlueprintGeneratedClass* AnimClass = Cast<UAnimBlueprintGeneratedClass>(AnimInstance->GetClass()))
 				{
-					const FAnimBlueprintDebugData& DebugData = AnimClass->GetAnimBlueprintDebugData();
-					for(const FAnimNodePoseWatch& AnimNodePoseWatch : DebugData.AnimNodePoseWatch)
+					// We have to grab our pose watches from the root class as no pose watches can be set on child anim BPs
+					if(const UAnimBlueprintGeneratedClass* RootClass = Cast<UAnimBlueprintGeneratedClass>(AnimClass->GetRootClass()))
 					{
-						if(AnimNodePoseWatch.PoseWatchPoseElement == PoseWatchPoseElement)
+						const FAnimBlueprintDebugData& DebugData = RootClass->AnimBlueprintDebugData;
+						for(const FAnimNodePoseWatch& AnimNodePoseWatch : DebugData.AnimNodePoseWatch)
 						{
-							bool bHasElement = false;
-							float CurrentValue = AnimNodePoseWatch.GetCurves().Get(Item->CurveName, bHasElement);
-							if(bHasElement)
+							if(AnimNodePoseWatch.PoseWatchPoseElement == PoseWatchPoseElement)
 							{
-								OutWeight = CurrentValue;
-								bFoundActive = true;
+								bool bHasElement = false;
+								float CurrentValue = AnimNodePoseWatch.GetCurves().Get(Item->CurveName, bHasElement);
+								if(bHasElement)
+								{
+									OutWeight = CurrentValue;
+									bFoundActive = true;
+								}
+								break;
 							}
-							break;
 						}
 					}
 				}
@@ -761,38 +769,42 @@ void SAnimCurveViewer::CreateAnimCurveList( const FString& SearchText, bool bInF
 		{
 			if(UAnimBlueprintGeneratedClass* AnimClass = Cast<UAnimBlueprintGeneratedClass>(AnimInstance->GetClass()))
 			{
-				const FAnimBlueprintDebugData& DebugData = AnimClass->GetAnimBlueprintDebugData();
-				for(const FAnimNodePoseWatch& AnimNodePoseWatch : DebugData.AnimNodePoseWatch)
+				// We have to grab our pose watches from the root class as no pose watches can be set on child anim BPs
+				if(const UAnimBlueprintGeneratedClass* RootClass = Cast<UAnimBlueprintGeneratedClass>(AnimClass->GetRootClass()))
 				{
-					if(AnimNodePoseWatch.PoseWatchPoseElement == PoseWatchPoseElement)
+					const FAnimBlueprintDebugData& DebugData = RootClass->AnimBlueprintDebugData;
+					for(const FAnimNodePoseWatch& AnimNodePoseWatch : DebugData.AnimNodePoseWatch)
 					{
-						AnimNodePoseWatch.GetCurves().ForEachElement([this, &AddCurve, &ActiveCurves](const UE::Anim::FCurveElement& InElement)
+						if(AnimNodePoseWatch.PoseWatchPoseElement == PoseWatchPoseElement)
 						{
-							if (EnumHasAnyFlags(CurrentCurveFlag, EAnimCurveViewerFilterFlags::MorphTarget))
+							AnimNodePoseWatch.GetCurves().ForEachElement([this, &AddCurve, &ActiveCurves](const UE::Anim::FCurveElement& InElement)
 							{
-								if(EnumHasAnyFlags(InElement.Flags, UE::Anim::ECurveElementFlags::MorphTarget))
+								if (EnumHasAnyFlags(CurrentCurveFlag, EAnimCurveViewerFilterFlags::MorphTarget))
 								{
-									AddCurve(InElement.Name, InElement.Flags);
-									ActiveCurves.Add(InElement.Name);
+									if(EnumHasAnyFlags(InElement.Flags, UE::Anim::ECurveElementFlags::MorphTarget))
+									{
+										AddCurve(InElement.Name, InElement.Flags);
+										ActiveCurves.Add(InElement.Name);
+									}
 								}
-							}
-							if (EnumHasAnyFlags(CurrentCurveFlag, EAnimCurveViewerFilterFlags::Material))
-							{
-								if(EnumHasAnyFlags(InElement.Flags, UE::Anim::ECurveElementFlags::Material))
+								if (EnumHasAnyFlags(CurrentCurveFlag, EAnimCurveViewerFilterFlags::Material))
+								{
+									if(EnumHasAnyFlags(InElement.Flags, UE::Anim::ECurveElementFlags::Material))
+									{
+										AddCurve(InElement.Name, InElement.Flags);
+										ActiveCurves.Add(InElement.Name);
+									}
+								}
+
+								// If we arent filtering by curve type, just show all curves
+								if(!EnumHasAnyFlags(CurrentCurveFlag, EAnimCurveViewerFilterFlags::MorphTarget | EAnimCurveViewerFilterFlags::Material))
 								{
 									AddCurve(InElement.Name, InElement.Flags);
 									ActiveCurves.Add(InElement.Name);
 								}
-							}
-
-							// If we arent filtering by curve type, just show all curves
-							if(!EnumHasAnyFlags(CurrentCurveFlag, EAnimCurveViewerFilterFlags::MorphTarget | EAnimCurveViewerFilterFlags::Material))
-							{
-								AddCurve(InElement.Name, InElement.Flags);
-								ActiveCurves.Add(InElement.Name);
-							}
-						});
-						break;
+							});
+							break;
+						}
 					}
 				}
 			}
@@ -960,13 +972,17 @@ TSharedRef<SWidget> SAnimCurveViewer::CreateCurveSourceSelector()
 				if(UAnimBlueprintGeneratedClass* AnimClass = Cast<UAnimBlueprintGeneratedClass>(AnimInstance->GetClass()))
 				{
 					TWeakObjectPtr<UPoseWatchPoseElement> Element = *InElement;
-					
-					for(const FAnimNodePoseWatch& AnimNodePoseWatch : AnimClass->AnimBlueprintDebugData.AnimNodePoseWatch)
+
+					// We have to grab our pose watches from the root class as no pose watches can be set on child anim BPs
+					if(const UAnimBlueprintGeneratedClass* RootClass = Cast<UAnimBlueprintGeneratedClass>(AnimClass->GetRootClass()))
 					{
-						if(AnimNodePoseWatch.PoseWatchPoseElement && Element.Get() == AnimNodePoseWatch.PoseWatchPoseElement)
+						for(const FAnimNodePoseWatch& AnimNodePoseWatch : RootClass->AnimBlueprintDebugData.AnimNodePoseWatch)
 						{
-							PoseWatch = AnimNodePoseWatch.PoseWatchPoseElement;
-							break;
+							if(AnimNodePoseWatch.PoseWatchPoseElement && Element.Get() == AnimNodePoseWatch.PoseWatchPoseElement)
+							{
+								PoseWatch = AnimNodePoseWatch.PoseWatchPoseElement;
+								break;
+							}
 						}
 					}
 				}
@@ -1033,11 +1049,15 @@ void SAnimCurveViewer::RebuildPoseWatches()
 		{
 			if(UAnimBlueprintGeneratedClass* AnimClass = Cast<UAnimBlueprintGeneratedClass>(AnimInstance->GetClass()))
 			{
-				for(const FAnimNodePoseWatch& AnimNodePoseWatch : AnimClass->AnimBlueprintDebugData.AnimNodePoseWatch)
+				// We have to grab our pose watches from the root class as no pose watches can be set on child anim BPs
+				if(const UAnimBlueprintGeneratedClass* RootClass = Cast<UAnimBlueprintGeneratedClass>(AnimClass->GetRootClass()))
 				{
-					if(AnimNodePoseWatch.PoseWatchPoseElement)
+					for(const FAnimNodePoseWatch& AnimNodePoseWatch : RootClass->AnimBlueprintDebugData.AnimNodePoseWatch)
 					{
-						PoseWatches.Add(MakeShared<TWeakObjectPtr<UPoseWatchPoseElement>>(AnimNodePoseWatch.PoseWatchPoseElement));
+						if(AnimNodePoseWatch.PoseWatchPoseElement)
+						{
+							PoseWatches.Add(MakeShared<TWeakObjectPtr<UPoseWatchPoseElement>>(AnimNodePoseWatch.PoseWatchPoseElement));
+						}
 					}
 				}
 			}
@@ -1047,13 +1067,19 @@ void SAnimCurveViewer::RebuildPoseWatches()
 	}
 }
 
-void SAnimCurveViewer::HandlePoseWatchesChanged(UAnimBlueprint* /*InAnimBlueprint*/, UEdGraphNode* /*InNode*/)
+void SAnimCurveViewer::HandlePoseWatchesChanged(UAnimBlueprint* InAnimBlueprint, UEdGraphNode* /*InNode*/)
 {
-	RegisterActiveTimer(0.0f, FWidgetActiveTimerDelegate::CreateLambda([this](double /*InCurrentTime*/, float /*InDeltaTime*/)
+	if(UAnimInstance* AnimInstance = GetAnimInstance())
 	{
-		RebuildPoseWatches();
-		return EActiveTimerReturnType::Stop;
-	}));
+		if(AnimInstance->GetClass()->IsChildOf(InAnimBlueprint->GeneratedClass))
+		{
+			RegisterActiveTimer(0.0f, FWidgetActiveTimerDelegate::CreateLambda([this](double /*InCurrentTime*/, float /*InDeltaTime*/)
+			{
+				RebuildPoseWatches();
+				return EActiveTimerReturnType::Stop;
+			}));
+		}
+	}
 }
 
 void SAnimCurveViewer::AddAnimCurveOverride( FName& Name, float Weight)
