Niagara Audio - Attenuation is ignored when previewing outside of PIE

diff --git a/Engine/Source/Editor/AudioEditor/Classes/AudioEditorSettings.h b/Engine/Source/Editor/AudioEditor/Classes/AudioEditorSettings.h
--- a/Engine/Source/Editor/AudioEditor/Classes/AudioEditorSettings.h
+++ b/Engine/Source/Editor/AudioEditor/Classes/AudioEditorSettings.h
@@ -7,7 +7,7 @@
 #include "UObject/NameTypes.h"
 #include "UObject/ObjectMacros.h"
 #include "UObject/UObjectGlobals.h"
-
+#include "AudioDefines.h"
 #include "AudioEditorSettings.generated.h"
 
 class UObject;
@@ -18,6 +18,12 @@ class AUDIOEDITOR_API UAudioEditorSettings : public UDeveloperSettings
 {
 	GENERATED_BODY()
 
+protected:
+
+	/** Whether or not should Audio Attenuation be used by default, for Non-Game Worlds*/
+	UPROPERTY(EditAnywhere, config, Category = NonGameWorld)
+	bool bUseAudioAttenuation = true;
+
 public:
 	/** Whether to pin the Sound Cue asset type when creating new assets. Requires editor restart to take effect. */
 	UPROPERTY(EditAnywhere, config, Category = AssetMenu)
@@ -35,12 +41,26 @@ public:
 	UPROPERTY(EditAnywhere, config, Category = AssetMenu)
 	bool bPinSoundConcurrencyInAssetMenu = true;
 
+	/** Set and apply, whether audio attenuation is used for non-game worlds*/
+	void SetUseAudioAttenuation(bool bInUseAudioAttenuation);
+
+	/** Is audio attenuation used for non-game worlds*/
+	bool IsUsingAudioAttenuation() const { return bUseAudioAttenuation; };
+
 	//~ Begin UDeveloperSettings
 	virtual FName GetCategoryName() const override { return TEXT("General"); }
 	//~ End UDeveloperSettings
 
 	//~ Begin UObject
 	virtual void PostEditChangeProperty(FPropertyChangedEvent& PropertyChangedEvent) override;
+	virtual void PostInitProperties();
 	//~ End UObject
-};
 
+private:
+
+	/** Apply non-game world attenuation setting for all AudioDevices.*/
+	void ApplyAttenuationForAllAudioDevices();
+
+	/** Apply non-game world attenuation setting for AudioDevice with given ID.*/
+	void ApplyAttenuationForAudioDevice(Audio::FDeviceId InDeviceID);
+};
diff --git a/Engine/Source/Editor/AudioEditor/Private/AudioEditorSettings.cpp b/Engine/Source/Editor/AudioEditor/Private/AudioEditorSettings.cpp
--- a/Engine/Source/Editor/AudioEditor/Private/AudioEditorSettings.cpp
+++ b/Engine/Source/Editor/AudioEditor/Private/AudioEditorSettings.cpp
@@ -1,10 +1,59 @@
 // Copyright Epic Games, Inc. All Rights Reserved.
 
 #include "AudioEditorSettings.h"
+#include "AudioDeviceManager.h"
+#include "AudioDevice.h"
 
 struct FPropertyChangedEvent;
 
+void UAudioEditorSettings::PostInitProperties()
+{
+	Super::PostInitProperties();
+
+	ApplyAttenuationForAllAudioDevices();
+	FAudioDeviceManagerDelegates::OnAudioDeviceCreated.AddUObject(this, &UAudioEditorSettings::ApplyAttenuationForAudioDevice);
+}
+
 void UAudioEditorSettings::PostEditChangeProperty(FPropertyChangedEvent& PropertyChangedEvent)
 {
 	Super::PostEditChangeProperty(PropertyChangedEvent);
+
+	if (PropertyChangedEvent.Property && PropertyChangedEvent.Property->GetFName() == GET_MEMBER_NAME_CHECKED(UAudioEditorSettings, bUseAudioAttenuation))
+	{
+		ApplyAttenuationForAllAudioDevices();
+	}
+}
+
+void UAudioEditorSettings::SetUseAudioAttenuation(bool bInUseAudioAttenuation) 
+{
+	bUseAudioAttenuation = bInUseAudioAttenuation;
+	SaveConfig();
+	ApplyAttenuationForAllAudioDevices();
+}
+
+void UAudioEditorSettings::ApplyAttenuationForAllAudioDevices()
+{
+	if (FAudioDeviceManager* AudioDeviceManager = FAudioDeviceManager::Get())
+	{
+		TArray<FAudioDevice*> AudioDevices = AudioDeviceManager->GetAudioDevices();
+		for (FAudioDevice* Device : AudioDevices)
+		{
+			if (Device)
+			{
+				Device->SetUseAttenuationForNonGameWorlds(bUseAudioAttenuation);
+			}
+		}
+	}
+}
+
+void UAudioEditorSettings::ApplyAttenuationForAudioDevice(Audio::FDeviceId InDeviceID)
+{
+	if (FAudioDeviceManager* AudioDeviceManager = FAudioDeviceManager::Get())
+	{
+		FAudioDeviceHandle Device = AudioDeviceManager->GetAudioDevice(InDeviceID);
+		if (Device.IsValid())
+		{
+			Device->SetUseAttenuationForNonGameWorlds(bUseAudioAttenuation);
+		}
+	}	
 }
diff --git a/Engine/Source/Editor/Persona/Persona.Build.cs b/Engine/Source/Editor/Persona/Persona.Build.cs
--- a/Engine/Source/Editor/Persona/Persona.Build.cs
+++ b/Engine/Source/Editor/Persona/Persona.Build.cs
@@ -40,7 +40,8 @@ public class Persona : ModuleRules
                 "AppFramework",
                 "AnimationModifiers",
                 "AnimationBlueprintLibrary",
-                "Core", 
+				"AudioEditor",
+				"Core", 
                 "CoreUObject", 
 				"ApplicationCore",
                 "Slate", 
diff --git a/Engine/Source/Editor/Persona/Private/AnimationEditorViewportClient.cpp b/Engine/Source/Editor/Persona/Private/AnimationEditorViewportClient.cpp
--- a/Engine/Source/Editor/Persona/Private/AnimationEditorViewportClient.cpp
+++ b/Engine/Source/Editor/Persona/Private/AnimationEditorViewportClient.cpp
@@ -46,6 +46,7 @@
 #include "Engine/PoseWatchRenderData.h"
 #include "UObject/UE5MainStreamObjectVersion.h"
 #include "Animation/AnimCompositeBase.h"
+#include "AudioEditorSettings.h"
 
 namespace {
 	static const float AnimationEditorViewport_RotateSpeed = 0.02f;
@@ -142,10 +143,9 @@ FAnimationViewportClient::FAnimationViewportClient(const TSharedRef<IPersonaPrev
 	{
 		World->bAllowAudioPlayback = !ConfigOption->bMuteAudio;
 
-		if(FAudioDevice* AudioDevice = World->GetAudioDeviceRaw())
-		{
-			AudioDevice->SetUseAttenuationForNonGameWorlds(ConfigOption->bUseAudioAttenuation);
-		}
+		UAudioEditorSettings* AudioConfigOption = GetMutableDefault<UAudioEditorSettings>();
+		check(AudioConfigOption);
+		AudioConfigOption->SetUseAudioAttenuation(true);
 	}
 }
 
@@ -242,21 +242,16 @@ bool FAnimationViewportClient::IsAudioMuted() const
 
 void FAnimationViewportClient::OnToggleUseAudioAttenuation()
 {
-	ConfigOption->SetUseAudioAttenuation(!ConfigOption->bUseAudioAttenuation);
-
-	UWorld* World = PreviewScene->GetWorld();
-	if(World)
-	{
-		if(FAudioDevice* AudioDevice = GetWorld()->GetAudioDeviceRaw())
-		{
-			AudioDevice->SetUseAttenuationForNonGameWorlds(ConfigOption->bUseAudioAttenuation);
-		}
-	}
+	UAudioEditorSettings* AudioConfigOption = GetMutableDefault<UAudioEditorSettings>();
+	check(AudioConfigOption);
+	AudioConfigOption->SetUseAudioAttenuation(!AudioConfigOption->IsUsingAudioAttenuation());
 }
 
 bool FAnimationViewportClient::IsUsingAudioAttenuation() const
 {
-	return ConfigOption->bUseAudioAttenuation;
+	const UAudioEditorSettings* AudioConfigOption = GetDefault<UAudioEditorSettings>();
+	check(AudioConfigOption);
+	return AudioConfigOption->IsUsingAudioAttenuation();
 }
 
 void FAnimationViewportClient::ToggleRotateCameraToFollowBone()
diff --git a/Engine/Source/Editor/UnrealEd/Classes/Preferences/PersonaOptions.h b/Engine/Source/Editor/UnrealEd/Classes/Preferences/PersonaOptions.h
--- a/Engine/Source/Editor/UnrealEd/Classes/Preferences/PersonaOptions.h
+++ b/Engine/Source/Editor/UnrealEd/Classes/Preferences/PersonaOptions.h
@@ -131,9 +131,6 @@ class UPersonaOptions : public UObject
 	UPROPERTY(EditAnywhere, config, Category = "Audio")
 	uint32 bMuteAudio:1;
 
-	UPROPERTY(EditAnywhere, config, Category = "Audio")
-	uint32 bUseAudioAttenuation:1;
-
 	/** Currently Stats can have None, Basic and Detailed. Please refer to EDisplayInfoMode. */
 	UPROPERTY(EditAnywhere, config, Category = "Viewport", meta=(ClampMin ="0", ClampMax = "3", UIMin = "0", UIMax = "3"))
 	int32 ShowMeshStats;
@@ -255,7 +252,6 @@ public:
 	UNREALED_API void SetHighlightOrigin( bool bInHighlightOrigin );
 	UNREALED_API void SetAutoAlignFloorToMesh(bool bInAutoAlignFloorToMesh);
 	UNREALED_API void SetMuteAudio( bool bInMuteAudio );
-	UNREALED_API void SetUseAudioAttenuation( bool bInUseAudioAttenuation );
 	UNREALED_API void SetViewModeIndex( FName InContext, EViewModeIndex InViewModeIndex, int32 InViewportIndex );
 	UNREALED_API void SetViewFOV( FName InContext, float InViewFOV, int32 InViewportIndex );
 	UNREALED_API void SetCameraSpeed(FName InContext, int32 InCameraSpeed, int32 InViewportIndex);
diff --git a/Engine/Source/Editor/UnrealEd/Private/PreferenceStubs.cpp b/Engine/Source/Editor/UnrealEd/Private/PreferenceStubs.cpp
--- a/Engine/Source/Editor/UnrealEd/Private/PreferenceStubs.cpp
+++ b/Engine/Source/Editor/UnrealEd/Private/PreferenceStubs.cpp
@@ -150,8 +150,6 @@ UPersonaOptions::UPersonaOptions(const FObjectInitializer& ObjectInitializer)
 
 	NumFolderFiltersInAssetBrowser = 2;
 
-	bUseAudioAttenuation = true;
-
 	CurveEditorSnapInterval = 0.01f;
 
 	// Default to millisecond resolution
@@ -206,12 +204,6 @@ void UPersonaOptions::SetMuteAudio( bool bInMuteAudio )
 	SaveConfig();
 }
 
-void UPersonaOptions::SetUseAudioAttenuation( bool bInUseAudioAttenuation )
-{
-	bUseAudioAttenuation = bInUseAudioAttenuation;
-	SaveConfig();
-}
-
 void UPersonaOptions::SetViewFOV( FName InContext, float InViewFOV, int32 InViewportIndex )
 {
 	check(InViewportIndex >= 0 && InViewportIndex < 4);
