Fixed renaming of float curves in the animation editor

diff --git a/Engine/Source/Editor/Persona/Private/AnimTimeline/AnimTimelineTrack_FloatCurve.cpp b/Engine/Source/Editor/Persona/Private/AnimTimeline/AnimTimelineTrack_FloatCurve.cpp
--- a/Engine/Source/Editor/Persona/Private/AnimTimeline/AnimTimelineTrack_FloatCurve.cpp
+++ b/Engine/Source/Editor/Persona/Private/AnimTimeline/AnimTimelineTrack_FloatCurve.cpp
@@ -201,7 +201,7 @@ void FAnimTimelineTrack_FloatCurve::OnCommitCurveName(const FText& InText, EText
 		const FName RequestedName = FName(*InText.ToString());
 		
 		const TArray<FFloatCurve>& FloatCurves = AnimSequenceBase->GetDataModel()->GetFloatCurves();
-		if (!FloatCurves.ContainsByPredicate([RequestedName](const FFloatCurve& Curve)
+		if (FloatCurves.ContainsByPredicate([RequestedName](const FFloatCurve& Curve)
 		{
 			return Curve.GetName() == RequestedName;
 		}))
