Add support for blendspace graph nodes to FBlendProfileCustomization::GetSkeletonFromOuter

diff --git a/Engine/Source/Editor/Persona/Private/BlendProfileCustomization.cpp b/Engine/Source/Editor/Persona/Private/BlendProfileCustomization.cpp
--- a/Engine/Source/Editor/Persona/Private/BlendProfileCustomization.cpp
+++ b/Engine/Source/Editor/Persona/Private/BlendProfileCustomization.cpp
@@ -12,6 +12,7 @@
 #include "Animation/AnimBlueprint.h"
 #include "Modules/ModuleManager.h"
 #include "Engine/PoseWatch.h"
+#include "Animation/BlendSpace.h"
 
 #define LOCTEXT_NAMESPACE "BlendProfileCustomization"
 
@@ -99,6 +100,14 @@ void FBlendProfileCustomization::OnBlendProfileChanged(UBlendProfile* NewProfile
 USkeleton* FBlendProfileCustomization::GetSkeletonFromOuter(const UObject* Outer)
 {
 	const UAnimBlueprint* AnimBlueprint = nullptr;
+	if (const UBlendSpace* BlendSpace = Cast<UBlendSpace>(Outer))
+	{
+		// Check for blend space graph nodes
+		if (!BlendSpace->IsAsset())
+		{
+			AnimBlueprint = BlendSpace->GetTypedOuter<UAnimBlueprint>();
+		}			
+	}
 	if (const UEdGraphNode* OuterEdGraphNode = Cast<UEdGraphNode>(Outer))
 	{
 		AnimBlueprint = Cast<UAnimBlueprint>(FBlueprintEditorUtils::FindBlueprintForNode(OuterEdGraphNode));
