Allow for a multi-selection of anim assets in the sequence browser to set their preview mesh not only for anim sequences.

[RN] minor

diff --git a/Engine/Source/Editor/Persona/Private/SAnimationSequenceBrowser.cpp b/Engine/Source/Editor/Persona/Private/SAnimationSequenceBrowser.cpp
--- a/Engine/Source/Editor/Persona/Private/SAnimationSequenceBrowser.cpp
+++ b/Engine/Source/Editor/Persona/Private/SAnimationSequenceBrowser.cpp
@@ -714,16 +714,6 @@ TSharedPtr<SWidget> SAnimationSequenceBrowser::OnGetAssetContextMenu(const TArra
 					FCanExecuteAction()
 				)
 			);
-
-			Section.AddMenuEntry("SetCurrentPreviewMesh",
-				LOCTEXT("SetCurrentPreviewMesh", "Set Current Preview Mesh"),
-				LOCTEXT("SetCurrentPreviewMesh_ToolTip", "Set current preview mesh to be used when previewed by this asset. This only applies when you open Persona using this asset."),
-				FSlateIcon(),
-				FUIAction(
-					FExecuteAction::CreateSP(this, &SAnimationSequenceBrowser::OnSetCurrentPreviewMesh, SelectedAssets),
-					FCanExecuteAction()
-				)
-			);
 		}
 
 		if (SelectedAssets.Num() == 1 && SelectedAssets[0].IsInstanceOf(USoundWave::StaticClass()))
@@ -758,6 +748,20 @@ TSharedPtr<SWidget> SAnimationSequenceBrowser::OnGetAssetContextMenu(const TArra
 			}
 		}
 
+		{
+			FToolMenuSection& Section = InMenu->FindOrAddSection("SequenceOptions", LOCTEXT("AssetHeading", "Asset"));
+
+			Section.AddMenuEntry("SetCurrentPreviewMesh",
+				LOCTEXT("SetCurrentPreviewMesh", "Set Current Preview Mesh"),
+				LOCTEXT("SetCurrentPreviewMesh_ToolTip", "Set current preview mesh to be used when previewed by this asset. This only applies when you open Persona using this asset."),
+				FSlateIcon(),
+				FUIAction(
+					FExecuteAction::CreateSP(this, &SAnimationSequenceBrowser::OnSetCurrentPreviewMesh, SelectedAssets),
+					FCanExecuteAction()
+				)
+			);
+		}
+		
 		{
 			FToolMenuSection& Section = InMenu->AddSection("SequenceGeneralOptions", LOCTEXT("AnimationSequenceGeneralOptions", "Options"));
 
@@ -904,7 +908,6 @@ void SAnimationSequenceBrowser::OnSetCurrentPreviewMesh(TArray<FAssetData> Selec
 		USkeletalMesh* PreviewMesh = PersonaToolkitPtr.Pin()->GetPreviewScene()->GetPreviewMeshComponent()->GetSkeletalMeshAsset();
 		if (PreviewMesh)
 		{
-			TArray<TWeakObjectPtr<UAnimSequence>> AnimSequences;
 			for(auto Iter = SelectedAssets.CreateIterator(); Iter; ++Iter)
 			{
 				UAnimationAsset * AnimAsset = Cast<UAnimationAsset>(Iter->GetAsset());
