Fix ensures and crashes when interacting with the Motion Matching node

Fixed crash when selecting "On Motion Matching State Updated" function
Fixed ensure when undoing actions on a node with a blend profile selected

diff --git a/Engine/Source/Editor/AnimationBlueprintEditor/Private/AnimationNodes/SAnimationGraphNode.cpp b/Engine/Source/Editor/AnimationBlueprintEditor/Private/AnimationNodes/SAnimationGraphNode.cpp
--- a/Engine/Source/Editor/AnimationBlueprintEditor/Private/AnimationNodes/SAnimationGraphNode.cpp
+++ b/Engine/Source/Editor/AnimationBlueprintEditor/Private/AnimationNodes/SAnimationGraphNode.cpp
@@ -304,7 +304,7 @@ public:
 			const TSharedPtr<IPropertyHandle> ChildPropertyHandle = ChildNode->CreatePropertyHandle();
 				
 			// Try to match node
-			if (ChildPropertyHandle.IsValid() && ChildPropertyHandle->GetProperty()->GetFName() == InMemberName && ChildPropertyHandle->GetMetaData(TEXT("Category")) == InCategory)
+			if (ChildPropertyHandle.IsValid() && ChildPropertyHandle->GetProperty() && ChildPropertyHandle->GetProperty()->GetFName() == InMemberName && ChildPropertyHandle->GetMetaData(TEXT("Category")) == InCategory)
 			{
 				OutDetailTreeNode = ChildNode;
 				OutPropertyHandle = ChildPropertyHandle;
diff --git a/Engine/Source/Editor/Persona/Private/BlendProfileCustomization.cpp b/Engine/Source/Editor/Persona/Private/BlendProfileCustomization.cpp
--- a/Engine/Source/Editor/Persona/Private/BlendProfileCustomization.cpp
+++ b/Engine/Source/Editor/Persona/Private/BlendProfileCustomization.cpp
@@ -40,7 +40,7 @@ void FBlendProfileCustomization::CustomizeHeader(TSharedRef<class IPropertyHandl
 		// try to get skeleton from first outer
 		if (USkeleton* TargetSkeleton = GetSkeletonFromOuter(OuterObjects[0]))
 		{
-			TSharedPtr<IPropertyHandle> PropertyPtr(InStructPropertyHandle);
+			TWeakPtr<IPropertyHandle> PropertyPtr(InStructPropertyHandle);
 
 			UObject* PropertyValue = nullptr;
 			InStructPropertyHandle->GetValue(PropertyValue);
@@ -69,7 +69,7 @@ void FBlendProfileCustomization::CustomizeHeader(TSharedRef<class IPropertyHandl
 			Args.OnBlendProfileSelected = FOnBlendProfileSelected::CreateSP(this, &FBlendProfileCustomization::OnBlendProfileChanged, PropertyPtr);
 			Args.InitialProfile = CurrentProfile;
 			Args.SupportedBlendProfileModes = SupportedBlendProfileModes;
-			Args.PropertyHandle = PropertyPtr;
+			Args.PropertyHandle = InStructPropertyHandle;
 
 			ISkeletonEditorModule& SkeletonEditorModule = FModuleManager::Get().LoadModuleChecked<ISkeletonEditorModule>("SkeletonEditor");
 			ValueCustomWidget = SkeletonEditorModule.CreateBlendProfilePicker(TargetSkeleton, Args);
@@ -89,11 +89,14 @@ void FBlendProfileCustomization::CustomizeHeader(TSharedRef<class IPropertyHandl
 		];
 }
 
-void FBlendProfileCustomization::OnBlendProfileChanged(UBlendProfile* NewProfile, TSharedPtr<IPropertyHandle> PropertyHandle)
+void FBlendProfileCustomization::OnBlendProfileChanged(UBlendProfile* NewProfile, TWeakPtr<IPropertyHandle> WeakPropertyHandle)
 {
-	if (PropertyHandle.IsValid())
+	if(!GIsTransacting)
 	{
-		PropertyHandle->SetValue(NewProfile);
+		if (TSharedPtr<IPropertyHandle> PropertyHandle = WeakPropertyHandle.Pin())
+		{
+			PropertyHandle->SetValue(NewProfile);
+		}
 	}
 }
 
diff --git a/Engine/Source/Editor/Persona/Private/BlendProfileCustomization.h b/Engine/Source/Editor/Persona/Private/BlendProfileCustomization.h
--- a/Engine/Source/Editor/Persona/Private/BlendProfileCustomization.h
+++ b/Engine/Source/Editor/Persona/Private/BlendProfileCustomization.h
@@ -19,6 +19,6 @@ public:
 
 private:
 
-	void OnBlendProfileChanged(class UBlendProfile* NewProfile, TSharedPtr<IPropertyHandle> PropertyHandle);
+	void OnBlendProfileChanged(class UBlendProfile* NewProfile, TWeakPtr<IPropertyHandle> WeakPropertyHandle);
 	class USkeleton* GetSkeletonFromOuter(const UObject* Outer);
 };
