Fixed a crash when reloading a skeletal mesh while an animation sequence is open, due to a stored raw pointer in the PersonaToolkit

diff --git a/Engine/Source/Editor/Persona/Private/PersonaToolkit.cpp b/Engine/Source/Editor/Persona/Private/PersonaToolkit.cpp
--- a/Engine/Source/Editor/Persona/Private/PersonaToolkit.cpp
+++ b/Engine/Source/Editor/Persona/Private/PersonaToolkit.cpp
@@ -32,10 +32,10 @@ FPersonaToolkit::~FPersonaToolkit()
 	PreviewScene.Reset();
 }
 
-static void FindCounterpartAssets(const UObject* InAsset, TWeakObjectPtr<USkeleton>& OutSkeleton, USkeletalMesh*& OutMesh)
+static void FindCounterpartAssets(const UObject* InAsset, TWeakObjectPtr<USkeleton>& OutSkeleton, TWeakObjectPtr <USkeletalMesh>& OutMesh)
 {
 	const USkeleton* CounterpartSkeleton = OutSkeleton.Get();
-	const USkeletalMesh* CounterpartMesh = OutMesh;
+	const USkeletalMesh* CounterpartMesh = OutMesh.Get();
 	FPersonaAssetFamily::FindCounterpartAssets(InAsset, CounterpartSkeleton, CounterpartMesh);
 	OutSkeleton = MakeWeakObjectPtr(const_cast<USkeleton*>(CounterpartSkeleton));
 	OutMesh = const_cast<USkeletalMesh*>(CounterpartMesh);
@@ -179,9 +179,9 @@ void FPersonaToolkit::CreatePreviewScene(const FPersonaToolkitArgs& PersonaToolk
 
 		bool bSetMesh = false;
 		// Set the mesh
-		if (Mesh)
+		if (Mesh.IsValid())
 		{
-			PreviewScene->SetPreviewMesh(Mesh, bAllowOverrideMesh);
+			PreviewScene->SetPreviewMesh(Mesh.Get(), bAllowOverrideMesh);
 			bSetMesh = true;
 			
 		}
@@ -219,7 +219,7 @@ UDebugSkelMeshComponent* FPersonaToolkit::GetPreviewMeshComponent() const
 
 USkeletalMesh* FPersonaToolkit::GetMesh() const
 {
-	return Mesh;
+	return Mesh.Get();
 }
 
 void FPersonaToolkit::SetMesh(class USkeletalMesh* InSkeletalMesh)
@@ -296,8 +296,8 @@ USkeletalMesh* FPersonaToolkit::GetPreviewMesh() const
 	}
 	else if(InitialAssetClass == USkeletalMesh::StaticClass())
 	{
-		check(Mesh);
-		return Mesh;
+		check(Mesh.IsValid());
+		return Mesh.Get();
 	}
 	else if(InitialAssetClass == USkeleton::StaticClass())
 	{
diff --git a/Engine/Source/Editor/Persona/Private/PersonaToolkit.h b/Engine/Source/Editor/Persona/Private/PersonaToolkit.h
--- a/Engine/Source/Editor/Persona/Private/PersonaToolkit.h
+++ b/Engine/Source/Editor/Persona/Private/PersonaToolkit.h
@@ -66,7 +66,7 @@ private:
 	TSharedPtr<IEditableSkeleton> EditableSkeleton;
 
 	/** The mesh we are editing */
-	USkeletalMesh* Mesh;
+	TWeakObjectPtr <USkeletalMesh> Mesh;
 
 	/** The anim blueprint we are editing */
 	UAnimBlueprint* AnimBlueprint;
