Fix for animation editor preview scene OnAnimChange delegate not being called with the animation asset being set in the scene (or null ptr if preview is getting closed)

diff --git a/Engine/Source/Editor/Persona/Private/AnimationEditorPreviewScene.cpp b/Engine/Source/Editor/Persona/Private/AnimationEditorPreviewScene.cpp
--- a/Engine/Source/Editor/Persona/Private/AnimationEditorPreviewScene.cpp
+++ b/Engine/Source/Editor/Persona/Private/AnimationEditorPreviewScene.cpp
@@ -278,8 +278,8 @@ void FAnimationEditorPreviewScene::SetPreviewMeshInternal(USkeletalMesh* NewPrev
 		UAnimationAsset* AnimAssetToPlay = nullptr;
 		float PlayPosition = 0.f;
 		bool bPlaying = false;
-		bool bNeedsToCopyAnimationData = SkeletalMeshComponent->GetAnimInstance() && SkeletalMeshComponent->GetAnimInstance() == SkeletalMeshComponent->PreviewInstance;
-		if (bNeedsToCopyAnimationData)
+		bool bNeedsToCopyAnimationData = bIsBeingConstructed == false && SkeletalMeshComponent->GetAnimInstance() && SkeletalMeshComponent->GetAnimInstance() == SkeletalMeshComponent->PreviewInstance;
+		if (bNeedsToCopyAnimationData && NewPreviewMesh)
 		{
 			AnimAssetToPlay = SkeletalMeshComponent->PreviewInstance->GetCurrentAsset();
 			PlayPosition = SkeletalMeshComponent->PreviewInstance->GetCurrentTime();
@@ -289,7 +289,7 @@ void FAnimationEditorPreviewScene::SetPreviewMeshInternal(USkeletalMesh* NewPrev
 		SkeletalMeshComponent->EmptyOverrideMaterials();
 		SkeletalMeshComponent->SetSkeletalMesh(NewPreviewMesh);
 
-		if (bNeedsToCopyAnimationData)
+		if (bNeedsToCopyAnimationData && (NewPreviewMesh == nullptr || AnimAssetToPlay != nullptr))
 		{
 			SetPreviewAnimationAsset(AnimAssetToPlay);
 			SkeletalMeshComponent->PreviewInstance->SetPosition(PlayPosition);
diff --git a/Engine/Source/Editor/Persona/Private/AnimationEditorPreviewScene.h b/Engine/Source/Editor/Persona/Private/AnimationEditorPreviewScene.h
--- a/Engine/Source/Editor/Persona/Private/AnimationEditorPreviewScene.h
+++ b/Engine/Source/Editor/Persona/Private/AnimationEditorPreviewScene.h
@@ -324,6 +324,16 @@ public:
 		EditableSkeletonPtr = InEditableSkeleton;
 	}
 
+	void SetIsBeingConstructed(bool bInIsBeingConstructed) 
+	{
+		bIsBeingConstructed = bInIsBeingConstructed;
+	}
+
+	bool IsBeingConstructed() const
+	{
+		return bIsBeingConstructed;
+	}
+
 private:
 	/** Set preview mesh internal use only. The mesh should be verified by now. */
 	void SetPreviewMeshInternal(USkeletalMesh* NewPreviewMesh);
@@ -470,4 +480,7 @@ private:
 
 	/** Delegate Remove attach component */
 	FOnRemoveAttachedComponentFilter OnRemoveAttachedComponentFilter;
+
+	/** True during initial creation, so some code can be skipped */
+	bool bIsBeingConstructed = false;
 };
diff --git a/Engine/Source/Editor/Persona/Private/PersonaToolkit.cpp b/Engine/Source/Editor/Persona/Private/PersonaToolkit.cpp
--- a/Engine/Source/Editor/Persona/Private/PersonaToolkit.cpp
+++ b/Engine/Source/Editor/Persona/Private/PersonaToolkit.cpp
@@ -76,6 +76,11 @@ void FPersonaToolkit::Initialize(UAnimationAsset* InAnimationAsset, const FPerso
 	FindCounterpartAssets(InAnimationAsset, Skeleton, Mesh);
 
 	CommonInitialSetup(PersonaToolkitArgs);
+
+	if (AnimationAsset != nullptr)
+	{
+		PreviewScene->SetPreviewAnimationAsset(AnimationAsset);
+	}
 }
 
 void FPersonaToolkit::Initialize(USkeletalMesh* InSkeletalMesh, const FPersonaToolkitArgs& PersonaToolkitArgs)
@@ -139,6 +144,12 @@ void FPersonaToolkit::CreatePreviewScene(const FPersonaToolkitArgs& PersonaToolk
 
 		PreviewScene = MakeShareable(new FAnimationEditorPreviewScene(FPreviewScene::ConstructionValues().AllowAudioPlayback(true).ShouldSimulatePhysics(true), EditableSkeleton, AsShared()));
 
+		PreviewScene->SetIsBeingConstructed(true);
+		ON_SCOPE_EXIT
+		{
+			PreviewScene->SetIsBeingConstructed(false);
+		};
+
 		//Temporary fix for missing attached assets - MDW
 		PreviewScene->GetWorld()->GetWorldSettings()->SetIsTemporarilyHiddenInEditor(false);
 
