Curve workflow fixes and improvements

Fixed assets in find/replace being identified as 'old' incorrectly - metadata was not getting added to 'new' assets, so they werent outputting any data to the asset registry.
Added shortcut in curve panels to find uses of specific curves
Added ability find/replace morph targets and materials
Fixed curve debugger to correctly indicate filter flags (initial active filter was not visible)
Added manual refresh of curve debugger to explicitly remove 'old' curves if required instread of guessing when it would be appropriate

diff --git a/Engine/Source/Editor/Persona/Private/AnimAssetFindReplace.cpp b/Engine/Source/Editor/Persona/Private/AnimAssetFindReplace.cpp
--- a/Engine/Source/Editor/Persona/Private/AnimAssetFindReplace.cpp
+++ b/Engine/Source/Editor/Persona/Private/AnimAssetFindReplace.cpp
@@ -180,11 +180,19 @@ void UAnimAssetFindReplaceProcessor::Initialize(TSharedRef<SAnimAssetFindReplace
 	Widget = InWidget;
 }
 
-void UAnimAssetFindReplaceProcessor::RequestUIRefresh()
+void UAnimAssetFindReplaceProcessor::RequestRefreshUI()
 {
 	if(TSharedPtr<SAnimAssetFindReplace> PinnedWidget = Widget.Pin())
 	{
-		PinnedWidget->RequestRefresh();
+		PinnedWidget->RequestRefreshUI();
+	}
+}
+
+void UAnimAssetFindReplaceProcessor::RequestRefreshCachedData()
+{
+	if(TSharedPtr<SAnimAssetFindReplace> PinnedWidget = Widget.Pin())
+	{
+		PinnedWidget->RequestRefreshCachedData();
 	}
 }
 
@@ -268,6 +276,8 @@ void UAnimAssetFindReplaceProcessor_StringBase::ExtendToolbar(FToolMenuSection&
 			FSlateIcon(FAppStyle::GetAppStyleSetName(), "Icons.Use"),
 			EUserInterfaceActionType::ToggleButton));
 
+	InSection.AddSeparator(NAME_None);
+
 	InSection.AddDynamicEntry("SkeletonFilter", FNewToolMenuSectionDelegate::CreateLambda([this](FToolMenuSection& InSection)
 	{
 		if(TSharedPtr<SAnimAssetFindReplace> Widget = AnimAssetFindReplacePrivate::GetWidgetFromContext(InSection.Context))
@@ -403,18 +413,21 @@ void UAnimAssetFindReplaceProcessor_StringBase::RefreshCachedData()
 
 bool UAnimAssetFindReplaceProcessor_StringBase::NameMatches(FStringView InNameStringView) const
 {
-	if(bFindWholeWord)
+	if(!FindString.IsEmpty())
 	{
-		if(InNameStringView.Compare(FindString, SearchCase) == 0)
+		if(bFindWholeWord)
 		{
-			return true;
+			if(InNameStringView.Compare(FindString, SearchCase) == 0)
+			{
+				return true;
+			}
 		}
-	}
-	else
-	{
-		if(UE::String::FindFirst(InNameStringView, FindString, SearchCase) != INDEX_NONE)
+		else
 		{
-			return true;
+			if(UE::String::FindFirst(InNameStringView, FindString, SearchCase) != INDEX_NONE)
+			{
+				return true;
+			}
 		}
 	}
 
diff --git a/Engine/Source/Editor/Persona/Private/AnimAssetFindReplaceCurves.cpp b/Engine/Source/Editor/Persona/Private/AnimAssetFindReplaceCurves.cpp
--- a/Engine/Source/Editor/Persona/Private/AnimAssetFindReplaceCurves.cpp
+++ b/Engine/Source/Editor/Persona/Private/AnimAssetFindReplaceCurves.cpp
@@ -10,10 +10,114 @@
 #include "Animation/AnimSequence.h"
 #include "AssetRegistry/AssetRegistryModule.h"
 #include "Engine/SkeletalMesh.h"
+#include "Engine/SkinnedAssetCommon.h"
+#include "Materials/Material.h"
+#include "Materials/MaterialExpressionParameter.h"
+#include "Materials/MaterialInterface.h"
 #include "String/ParseTokens.h"
 
 #define LOCTEXT_NAMESPACE "AnimAssetFindReplaceCurves"
 
+namespace UE::AnimAssetFindReplaceCurves::Private
+{
+
+// Helper function used to rename material parameters in skeletal meshes
+static void RenameMaterialParameter(USkeletalMesh* InSkeletalMesh, FName InOldName, FName InNewName)
+{
+	for(const FSkeletalMaterial& SkeletalMaterial : InSkeletalMesh->GetMaterials())
+	{
+		UMaterial* Material = (SkeletalMaterial.MaterialInterface != nullptr) ? SkeletalMaterial.MaterialInterface->GetMaterial() : nullptr;
+		if (Material == nullptr)
+		{
+			continue;
+		}
+
+		// Search all expressions for matching scalar parameters
+		for(UMaterialExpression* Expression : Material->GetExpressions())
+		{
+			UMaterialExpressionParameter* Parameter = Cast<UMaterialExpressionParameter>(Expression);
+			if(Parameter == nullptr)
+			{
+				continue;
+			}
+
+			FMaterialParameterMetadata ParameterMeta;
+			if (!Parameter->GetParameterValue(ParameterMeta))
+			{
+				continue;
+			}
+
+			if(ParameterMeta.Value.Type != EMaterialParameterType::Scalar)
+			{
+				continue;
+			}
+
+			if(Parameter->GetParameterName() == InOldName)
+			{
+				FProperty* ParameterNameProperty = UMaterialExpressionParameter::StaticClass()->FindPropertyByName(GET_MEMBER_NAME_CHECKED(UMaterialExpressionParameter, ParameterName));
+				check(ParameterNameProperty);
+				Parameter->PreEditChange(ParameterNameProperty);
+
+				Parameter->ParameterName = InNewName;
+
+				FPropertyChangedEvent ChangedEvent(ParameterNameProperty);
+				Parameter->PostEditChangeProperty(ChangedEvent);
+			}
+		}
+	}
+}
+
+// Helper function used to remove material parameters in skeletal meshes
+static void RemoveMaterialParameters(USkeletalMesh* InSkeletalMesh, TConstArrayView<FName> InMaterialParameterNames)
+{
+	for(const FSkeletalMaterial& SkeletalMaterial : InSkeletalMesh->GetMaterials())
+	{
+		bool bMaterialModified = false;
+		UMaterial* Material = (SkeletalMaterial.MaterialInterface != nullptr) ? SkeletalMaterial.MaterialInterface->GetMaterial() : nullptr;
+		if (Material == nullptr)
+		{
+			continue;
+		}
+
+		// Search all expressions for matching scalar parameters
+		for(UMaterialExpression* Expression : Material->GetExpressions())
+		{
+			UMaterialExpressionParameter* Parameter = Cast<UMaterialExpressionParameter>(Expression);
+			if(Parameter == nullptr)
+			{
+				continue;
+			}
+
+			FMaterialParameterMetadata ParameterMeta;
+			if (!Parameter->GetParameterValue(ParameterMeta))
+			{
+				continue;
+			}
+
+			if(ParameterMeta.Value.Type != EMaterialParameterType::Scalar)
+			{
+				continue;
+			}
+
+			for(FName NameToRemove : InMaterialParameterNames)
+			{
+				if(Parameter->GetParameterName() == NameToRemove)
+				{
+					Material->PreEditChange(nullptr);
+
+					Material->GetExpressionCollection().RemoveExpression(Parameter);
+					Material->RemoveExpressionParameter(Parameter);
+					Parameter->MarkAsGarbage();
+
+					Material->PostEditChange();
+				}
+			}
+		}
+	}
+}
+
+}
+
 FString UAnimAssetFindReplaceCurves::GetFindResultStringFromAssetData(const FAssetData& InAssetData) const
 {
 	if(GetFindString().IsEmpty())
@@ -21,12 +125,12 @@ FString UAnimAssetFindReplaceCurves::GetFindResultStringFromAssetData(const FAss
 		return FString();
 	}
 
-	auto GetMatchingCurveNamesForAsset = [this](const FAssetData& InAssetData, TArray<FString>& OutCurveNames)
+	auto GetMatchingCurveNamesForAsset = [this](const FAssetData& InAssetData, TArray<FString>& OutCurveNames, FName InTag, const FString& InDelimiter)
 	{
-		const FString TagValue = InAssetData.GetTagValueRef<FString>(USkeleton::CurveNameTag);
+		const FString TagValue = InAssetData.GetTagValueRef<FString>(InTag);
 		if (!TagValue.IsEmpty())
 		{
-			UE::String::ParseTokens(TagValue, *USkeleton::CurveTagDelimiter, [this, &OutCurveNames](FStringView InToken)
+			UE::String::ParseTokens(TagValue, *InDelimiter, [this, &OutCurveNames](FStringView InToken)
 			{
 				if(GetFindWholeWord())
 				{
@@ -48,7 +152,7 @@ FString UAnimAssetFindReplaceCurves::GetFindResultStringFromAssetData(const FAss
 
 	TStringBuilder<128> Builder;
 	TArray<FString> CurveNames;
-	GetMatchingCurveNamesForAsset(InAssetData, CurveNames);
+	GetMatchingCurveNamesForAsset(InAssetData, CurveNames, USkeleton::CurveNameTag, USkeleton::CurveTagDelimiter);
 	if(CurveNames.Num() > 0)
 	{
 		for(int32 NameIndex = 0; NameIndex < CurveNames.Num(); ++NameIndex)
@@ -60,67 +164,80 @@ FString UAnimAssetFindReplaceCurves::GetFindResultStringFromAssetData(const FAss
 			}
 		}
 	}
-	return FString(Builder.ToString());
-}
 
-TConstArrayView<UClass*> UAnimAssetFindReplaceCurves::GetSupportedAssetTypes() const
-{
-	static UClass* Types[] = { UPoseAsset::StaticClass(), UAnimSequenceBase::StaticClass(), USkeleton::StaticClass(), USkeletalMesh::StaticClass() };
-	return Types;
-}
-
-bool UAnimAssetFindReplaceCurves::ShouldFilterOutAsset(const FAssetData& InAssetData, bool& bOutIsOldAsset) const
-{
-	FString TagValue;
-	if(InAssetData.GetTagValue<FString>(USkeleton::CurveNameTag, TagValue))
+	if(GetSearchMorphTargets())
 	{
-		bOutIsOldAsset = false;
-
-		if(GetFindString().IsEmpty())
-		{
-			return true;
-		}
-		
-		if(GetSkeletonFilter().IsValid() )
+		UClass* AssetClass = InAssetData.GetClass();
+		if(AssetClass->IsChildOf(USkeletalMesh::StaticClass()))
 		{
-			if(InAssetData.GetClass() != USkeleton::StaticClass())
+			TArray<FString> MorphNames;
+			GetMatchingCurveNamesForAsset(InAssetData, MorphNames, USkeletalMesh::MorphNamesTag, USkeletalMesh::MorphNamesTagDelimiter);
+			if(MorphNames.Num() > 0)
 			{
-				FString SkeletonPath;
-				if(InAssetData.GetTagValue<FString>(TEXT("Skeleton"), SkeletonPath))
+				if(Builder.Len() > 0)
 				{
-					if(SkeletonPath != GetSkeletonFilter().GetExportTextName())
+					Builder.Append(TEXT(" "));
+				}
+				Builder.Append(LOCTEXT("MorphTargets", "Morph targets: ").ToString());
+				for(int32 NameIndex = 0; NameIndex < MorphNames.Num(); ++NameIndex)
+				{
+					Builder.Append(MorphNames[NameIndex]);
+					if(NameIndex != MorphNames.Num() - 1)
 					{
-						return true;
+						Builder.Append(TEXT(", "));
 					}
 				}
 			}
-			else
+		}
+	}
+
+	if(GetSearchMaterials())
+	{
+		UClass* AssetClass = InAssetData.GetClass();
+		if(AssetClass->IsChildOf(USkeletalMesh::StaticClass()))
+		{
+			TArray<FString> MaterialNames;
+			GetMatchingCurveNamesForAsset(InAssetData, MaterialNames, USkeletalMesh::MaterialParamNamesTag, USkeletalMesh::MaterialParamNamesTagDelimiter);
+			if(MaterialNames.Num() > 0)
 			{
-				if(InAssetData.ToSoftObjectPath() != GetSkeletonFilter().ToSoftObjectPath())
+				if(Builder.Len() > 0)
 				{
-					return true;
+					Builder.Append(TEXT(" "));
+				}
+				Builder.Append(LOCTEXT("MaterialParameters", "Material parameters: ").ToString());
+				for(int32 NameIndex = 0; NameIndex < MaterialNames.Num(); ++NameIndex)
+				{
+					Builder.Append(MaterialNames[NameIndex]);
+					if(NameIndex != MaterialNames.Num() - 1)
+					{
+						Builder.Append(TEXT(", "));
+					}
 				}
 			}
 		}
+	}
 
-		bool bFoundMatch = false;
-		UE::String::ParseTokens(TagValue, *USkeleton::CurveTagDelimiter, [this, &bFoundMatch](FStringView InToken)
-		{
-			if(NameMatches(InToken))
-			{
-				bFoundMatch = true;
-			}
-		}, UE::String::EParseTokensOptions::SkipEmpty);
+	return FString(Builder.ToString());
+}
 
-		if(bFoundMatch)
-		{
-			return false;
-		}
+TConstArrayView<UClass*> UAnimAssetFindReplaceCurves::GetSupportedAssetTypes() const
+{
+	static UClass* Types[] = { UPoseAsset::StaticClass(), UAnimSequenceBase::StaticClass(), USkeleton::StaticClass(), USkeletalMesh::StaticClass() };
+	return Types;
+}
+
+bool UAnimAssetFindReplaceCurves::ShouldFilterOutAsset(const FAssetData& InAssetData, bool& bOutIsOldAsset) const
+{
+	UClass* AssetClass = InAssetData.GetClass();
+
+	FString TagValue;
+	if(InAssetData.GetTagValue<FString>(USkeleton::CurveNameTag, TagValue))
+	{
+		bOutIsOldAsset = false;
 	}
 	else
 	{
 		FAssetRegistryModule& AssetRegistryModule = FModuleManager::LoadModuleChecked<FAssetRegistryModule>("AssetRegistry");
-		UClass* AssetClass = InAssetData.GetClass();
 		if(AssetClass->IsChildOf(UAnimSequenceBase::StaticClass()))
 		{
 			// Check the package object version - the asset was saving empty tags for curves, so the absence of curves is
@@ -177,8 +294,104 @@ bool UAnimAssetFindReplaceCurves::ShouldFilterOutAsset(const FAssetData& InAsset
 			bOutIsOldAsset = false;
 		}
 	}
+	
+	if(GetSearchMorphTargets())
+	{
+		if(AssetClass->IsChildOf(USkeletalMesh::StaticClass()))
+		{
+			FString MorphTagValue;
+			if(!InAssetData.GetTagValue<FString>(USkeletalMesh::MorphNamesTag, MorphTagValue))
+			{
+				bOutIsOldAsset = true;
+			}
+		}
+	}
 
-	return true;
+	if(GetSearchMaterials())
+	{
+		if(AssetClass->IsChildOf(USkeletalMesh::StaticClass()))
+		{
+			FString MaterialParamTagValue;
+			if(!InAssetData.GetTagValue<FString>(USkeletalMesh::MaterialParamNamesTag, MaterialParamTagValue))
+			{
+				bOutIsOldAsset = true;
+			}
+		}
+	}
+
+	if(GetFindString().IsEmpty())
+	{
+		return true;
+	}
+
+	if(GetSkeletonFilter().IsValid() )
+	{
+		if(InAssetData.GetClass() != USkeleton::StaticClass())
+		{
+			FString SkeletonPath;
+			if(InAssetData.GetTagValue<FString>(TEXT("Skeleton"), SkeletonPath))
+			{
+				if(SkeletonPath != GetSkeletonFilter().GetExportTextName())
+				{
+					return true;
+				}
+			}
+		}
+		else
+		{
+			if(InAssetData.ToSoftObjectPath() != GetSkeletonFilter().ToSoftObjectPath())
+			{
+				return true;
+			}
+		}
+	}
+
+	bool bFoundMatch = false;
+	UE::String::ParseTokens(TagValue, *USkeleton::CurveTagDelimiter, [this, &bFoundMatch](FStringView InToken)
+	{
+		if(NameMatches(InToken))
+		{
+			bFoundMatch = true;
+		}
+	}, UE::String::EParseTokensOptions::SkipEmpty);
+	
+	if(GetSearchMorphTargets())
+	{
+		if(AssetClass->IsChildOf(USkeletalMesh::StaticClass()))
+		{
+			FString MorphTagValue;
+			if(InAssetData.GetTagValue<FString>(USkeletalMesh::MorphNamesTag, MorphTagValue))
+			{
+				UE::String::ParseTokens(MorphTagValue, *USkeletalMesh::MorphNamesTagDelimiter, [this, &bFoundMatch](FStringView InToken)
+				{
+					if(NameMatches(InToken))
+					{
+						bFoundMatch = true;
+					}
+				}, UE::String::EParseTokensOptions::SkipEmpty);
+			}
+		}
+	}
+
+	if(GetSearchMaterials())
+	{
+		if(AssetClass->IsChildOf(USkeletalMesh::StaticClass()))
+		{
+			FString MaterialParamTagValue;
+			if(InAssetData.GetTagValue<FString>(USkeletalMesh::MaterialParamNamesTag, MaterialParamTagValue))
+			{
+				UE::String::ParseTokens(MaterialParamTagValue, *USkeletalMesh::MaterialParamNamesTagDelimiter, [this, &bFoundMatch](FStringView InToken)
+				{
+					if(NameMatches(InToken))
+					{
+						bFoundMatch = true;
+					}
+				}, UE::String::EParseTokensOptions::SkipEmpty);
+			}
+		}
+	}
+
+	return !bFoundMatch;
 }
 
 void UAnimAssetFindReplaceCurves::ReplaceInAsset(const FAssetData& InAssetData) const
@@ -340,6 +553,95 @@ void UAnimAssetFindReplaceCurves::ReplaceInAsset(const FAssetData& InAssetData)
 					}
 				}
 			}
+
+			if(GetSearchMorphTargets())
+			{
+				if(GetFindWholeWord())
+				{
+					Asset->Modify();
+
+					const FName FindMorphName(GetFindString());
+					const FName ReplaceMorphName(GetReplaceString());
+					SkeletalMesh->RenameMorphTarget(FindMorphName, ReplaceMorphName);
+				}
+				else
+				{
+					TArray<TPair<FName, FName>> FindReplacePairs;
+
+					for(UMorphTarget* MorphTarget : SkeletalMesh->GetMorphTargets())
+					{
+						const FString MorphNameString = MorphTarget->GetName();
+						if(NameMatches(MorphNameString))
+						{
+							const FName FindMorphName(MorphTarget->GetFName());
+							const FString NewName = MorphNameString.Replace(*GetFindStringRef(), *GetReplaceStringRef(), GetSearchCase());
+							const FName ReplaceMorphName(*NewName);
+
+							FindReplacePairs.Emplace(FindMorphName, ReplaceMorphName);
+						}
+					}
+
+					if(FindReplacePairs.Num() > 0)
+					{
+						Asset->Modify();
+
+						for(const TPair<FName, FName>& FindReplacePair : FindReplacePairs)
+						{
+							SkeletalMesh->RenameMorphTarget(FindReplacePair.Key, FindReplacePair.Value);
+						}
+					}
+				}
+			}
+
+			if(GetSearchMaterials())
+			{
+				if(GetFindWholeWord())
+				{
+					Asset->Modify();
+
+					const FName FindMorphName(GetFindString());
+					const FName ReplaceMorphName(GetReplaceString());
+					UE::AnimAssetFindReplaceCurves::Private::RenameMaterialParameter(SkeletalMesh, FindMorphName, ReplaceMorphName);
+				}
+				else
+				{
+					TArray<TPair<FName, FName>> FindReplacePairs;
+
+					for (const FSkeletalMaterial& SkeletalMaterial : SkeletalMesh->GetMaterials())
+					{
+						UMaterial* Material = (SkeletalMaterial.MaterialInterface != nullptr) ? SkeletalMaterial.MaterialInterface->GetMaterial() : nullptr;
+						if (Material)
+						{
+							TArray<FMaterialParameterInfo> OutParameterInfo;
+							TArray<FGuid> OutParameterIds;
+							SkeletalMaterial.MaterialInterface->GetAllScalarParameterInfo(OutParameterInfo, OutParameterIds);
+
+							for (const FMaterialParameterInfo& MaterialParameterInfo : OutParameterInfo)
+							{
+								const FString ParameterNameString = MaterialParameterInfo.Name.ToString();
+								if(NameMatches(ParameterNameString))
+								{
+									const FName FindMorphName(MaterialParameterInfo.Name);
+									const FString NewName = ParameterNameString.Replace(*GetFindStringRef(), *GetReplaceStringRef(), GetSearchCase());
+									const FName ReplaceMorphName(*NewName);
+
+									FindReplacePairs.Emplace(FindMorphName, ReplaceMorphName);
+								}
+							}
+						}
+					}
+
+					if(FindReplacePairs.Num() > 0)
+					{
+						Asset->Modify();
+
+						for(const TPair<FName, FName>& FindReplacePair : FindReplacePairs)
+						{
+							UE::AnimAssetFindReplaceCurves::Private::RenameMaterialParameter(SkeletalMesh, FindReplacePair.Key, FindReplacePair.Value);
+						}
+					}
+				}
+			}
 		}
 	}
 }
@@ -430,8 +732,8 @@ void UAnimAssetFindReplaceCurves::RemoveInAsset(const FAssetData& InAssetData) c
 				{
 				 	if(NameMatches(InCurveName.ToString()))
 				 	{
-						 CurvesToRemove.AddUnique(InCurveName);
-					 }
+						CurvesToRemove.AddUnique(InCurveName);
+				 	}
 				});
 
 				if(CurvesToRemove.Num() > 0)
@@ -441,22 +743,174 @@ void UAnimAssetFindReplaceCurves::RemoveInAsset(const FAssetData& InAssetData) c
 				}
 			}
 		}
+		else if(USkeletalMesh* SkeletalMesh = Cast<USkeletalMesh>(Asset))
+		{
+			if(UAnimCurveMetaData* AnimCurveMetaData = SkeletalMesh->GetAssetUserData<UAnimCurveMetaData>())
+			{
+				if(GetFindWholeWord())
+				{
+					Asset->Modify();
+					AnimCurveMetaData->RemoveCurveMetaData(FName(GetFindString()));
+				}
+				else
+				{
+					TArray<FName> CurvesToRemove;
+					AnimCurveMetaData->ForEachCurveMetaData([this, &CurvesToRemove](FName InCurveName, const FCurveMetaData& InCurveMetaData)
+					{
+						if(NameMatches(InCurveName.ToString()))
+						{
+							CurvesToRemove.AddUnique(InCurveName);
+						}
+					});
+
+					if(CurvesToRemove.Num() > 0)
+					{
+						Asset->Modify();
+						AnimCurveMetaData->RemoveCurveMetaData(CurvesToRemove);
+					}
+				}
+			}
+
+			if(GetSearchMorphTargets())
+			{
+				if(GetFindWholeWord())
+				{
+					Asset->Modify();
+					SkeletalMesh->RemoveMorphTargets({FName(GetFindString())});
+				}
+				else
+				{
+					TArray<FName> MorphsToRemove;
+					for(UMorphTarget* MorphTarget : SkeletalMesh->GetMorphTargets())
+					{
+						if(NameMatches(MorphTarget->GetFName().ToString()))
+						{
+							MorphsToRemove.AddUnique(MorphTarget->GetFName());
+						}
+					}
+
+					if(MorphsToRemove.Num() > 0)
+					{
+						Asset->Modify();
+						SkeletalMesh->RemoveMorphTargets(MorphsToRemove);
+					}
+				}
+			}
+
+			if(GetSearchMaterials())
+			{
+				if(GetFindWholeWord())
+				{
+					Asset->Modify();
+					UE::AnimAssetFindReplaceCurves::Private::RemoveMaterialParameters(SkeletalMesh, {FName(GetFindString())});
+				}
+				else
+				{
+					TArray<FName> MaterialParametersToRemove;
+					for(UMorphTarget* MorphTarget : SkeletalMesh->GetMorphTargets())
+					{
+						if(NameMatches(MorphTarget->GetFName().ToString()))
+						{
+							MaterialParametersToRemove.AddUnique(MorphTarget->GetFName());
+						}
+					}
+
+					if(MaterialParametersToRemove.Num() > 0)
+					{
+						Asset->Modify();
+						UE::AnimAssetFindReplaceCurves::Private::RemoveMaterialParameters(SkeletalMesh, MaterialParametersToRemove);
+					}
+				}
+			}
+		}
 	}
 }
 
+void UAnimAssetFindReplaceCurves::ExtendToolbar(FToolMenuSection& InSection)
+{
+	Super::ExtendToolbar(InSection);
+
+	InSection.AddDynamicEntry("CurveSearchOptions", FNewToolMenuSectionDelegate::CreateLambda([this](FToolMenuSection& InSection)
+	{
+		FToolUIAction SearchMorphTargetsCheckbox;
+		SearchMorphTargetsCheckbox.ExecuteAction = FToolMenuExecuteAction::CreateLambda([this](const FToolMenuContext& InContext)
+		{
+			SetSearchMorphTargets(!GetSearchMorphTargets());
+		});
+
+		SearchMorphTargetsCheckbox.GetActionCheckState = FToolMenuGetActionCheckState::CreateLambda([this](const FToolMenuContext& InContext)
+		{
+			return GetSearchMorphTargets() ? ECheckBoxState::Checked : ECheckBoxState::Unchecked;
+		});
+		
+		InSection.AddEntry(
+		FToolMenuEntry::InitToolBarButton(
+			"SearchMorphTargets",
+			SearchMorphTargetsCheckbox,
+			LOCTEXT("SearchMorphTargetsCheckboxLabel", "Morph Targets"),
+			LOCTEXT("SearchMorphTargetsCheckboxTooltip", "Whether to search morph targets in skeletal meshes while searching for curves."),
+			FSlateIcon(FAppStyle::GetAppStyleSetName(), "Persona.Tabs.MorphTargetPreviewer"),
+			EUserInterfaceActionType::ToggleButton));
+
+		FToolUIAction SearchMaterialsCheckbox;
+		SearchMaterialsCheckbox.ExecuteAction = FToolMenuExecuteAction::CreateLambda([this](const FToolMenuContext& InContext)
+		{
+			SetSearchMaterials(!GetSearchMaterials());
+		});
+
+		SearchMaterialsCheckbox.GetActionCheckState = FToolMenuGetActionCheckState::CreateLambda([this](const FToolMenuContext& InContext)
+		{
+			return GetSearchMaterials() ? ECheckBoxState::Checked : ECheckBoxState::Unchecked;
+		});
+		
+		InSection.AddEntry(
+		FToolMenuEntry::InitToolBarButton(
+			"SearchMaterials",
+			SearchMaterialsCheckbox,
+			LOCTEXT("SearchMaterialsCheckboxLabel", "Materials"),
+			LOCTEXT("SearchMaterialsCheckboxTooltip", "Whether to search material parameters in skeletal meshes while searching for curves."),
+			FSlateIcon(FAppStyle::GetAppStyleSetName(), "ClassIcon.Material"),
+			EUserInterfaceActionType::ToggleButton));
+	}));
+}
+
 void UAnimAssetFindReplaceCurves::GetAutoCompleteNames(TArrayView<FAssetData> InAssetDatas, TSet<FString>& OutUniqueNames) const
 {
 	for (const FAssetData& AssetData : InAssetDatas)
 	{
-		const FString TagValue = AssetData.GetTagValueRef<FString>(USkeleton::CurveNameTag);
-		if (!TagValue.IsEmpty())
+		auto GetNamesForAsset = [this, &OutUniqueNames](const FAssetData& InAssetData, FName InTag, const FString& InDelimiter)
 		{
-			UE::String::ParseTokens(TagValue, *USkeleton::CurveTagDelimiter, [&OutUniqueNames](FStringView InToken)
+			const FString TagValue = InAssetData.GetTagValueRef<FString>(InTag);
+			if (!TagValue.IsEmpty())
 			{
-				OutUniqueNames.Add(FString(InToken));
-			}, UE::String::EParseTokensOptions::SkipEmpty);
+				UE::String::ParseTokens(TagValue, *USkeleton::CurveTagDelimiter, [&OutUniqueNames](FStringView InToken)
+				{
+					OutUniqueNames.Add(FString(InToken));
+				}, UE::String::EParseTokensOptions::SkipEmpty);
+			}
+		};
+
+		GetNamesForAsset(AssetData, USkeleton::CurveNameTag, USkeleton::CurveTagDelimiter);
+
+		UClass* AssetClass = AssetData.GetClass();
+		if(AssetClass->IsChildOf(USkeletalMesh::StaticClass()))
+		{
+			GetNamesForAsset(AssetData, USkeletalMesh::MorphNamesTag, USkeletalMesh::MorphNamesTagDelimiter);
+			GetNamesForAsset(AssetData, USkeletalMesh::MaterialParamNamesTag, USkeletalMesh::MaterialParamNamesTagDelimiter);
 		}
 	}
 }
 
+void UAnimAssetFindReplaceCurves::SetSearchMorphTargets(bool bInSearchMorphTargets)
+{
+	bSearchMorphTargets = bInSearchMorphTargets;
+	RequestRefreshSearchResults();
+}
+
+void UAnimAssetFindReplaceCurves::SetSearchMaterials(bool bInSearchMaterials)
+{
+	bSearchMaterials = bInSearchMaterials;
+	RequestRefreshSearchResults();
+}
+
 #undef LOCTEXT_NAMESPACE
\ No newline at end of file
diff --git a/Engine/Source/Editor/Persona/Private/CurveViewerCommands.cpp b/Engine/Source/Editor/Persona/Private/CurveViewerCommands.cpp
--- a/Engine/Source/Editor/Persona/Private/CurveViewerCommands.cpp
+++ b/Engine/Source/Editor/Persona/Private/CurveViewerCommands.cpp
@@ -10,6 +10,7 @@
 void FCurveViewerCommands::RegisterCommands()
 {
 	UI_COMMAND( AddCurve, "Add Curve", "Add a curve to the Skeleton", EUserInterfaceActionType::Button, FInputChord());
+	UI_COMMAND( FindCurveUses, "Find Curve Uses", "Open the find/replace tab to find all uses of this curve", EUserInterfaceActionType::Button, FInputChord(EKeys::F3));
 }
 
 #undef LOCTEXT_NAMESPACE
diff --git a/Engine/Source/Editor/Persona/Private/CurveViewerCommands.h b/Engine/Source/Editor/Persona/Private/CurveViewerCommands.h
--- a/Engine/Source/Editor/Persona/Private/CurveViewerCommands.h
+++ b/Engine/Source/Editor/Persona/Private/CurveViewerCommands.h
@@ -33,6 +33,9 @@ public:
 
 	/** Add curve */
 	TSharedPtr< FUICommandInfo > AddCurve;
+
+	/** Fund curve uses */
+	TSharedPtr< FUICommandInfo > FindCurveUses;
 };
 
 
diff --git a/Engine/Source/Editor/Persona/Private/SAnimAssetFindReplace.cpp b/Engine/Source/Editor/Persona/Private/SAnimAssetFindReplace.cpp
--- a/Engine/Source/Editor/Persona/Private/SAnimAssetFindReplace.cpp
+++ b/Engine/Source/Editor/Persona/Private/SAnimAssetFindReplace.cpp
@@ -20,6 +20,7 @@
 #include "AssetRegistry/AssetRegistryModule.h"
 #include "Misc/ScopedSlowTask.h"
 #include "ScopedTransaction.h"
+#include "SkeletalMeshCompiler.h"
 #include "Widgets/Views/SListView.h"
 #include "UObject/AssetRegistryTagsContext.h"
 #include "UObject/AnimPhysObjectVersion.h"
@@ -227,7 +228,7 @@ void SAnimAssetFindReplace::Construct(const FArguments& InArgs)
 			{
 				if(TSharedPtr<SAnimAssetFindReplace> Widget = AnimAssetFindReplacePrivate::GetWidgetFromContext(InContext))
 				{
-					Widget->RequestRefreshSearchResults();
+					Widget->RequestRefreshUI();
 				}
 			});
 
@@ -319,7 +320,10 @@ void SAnimAssetFindReplace::Construct(const FArguments& InArgs)
 						}
 					}
 
-					RequestRefreshSearchResults();
+					// Ensure all meshes are compiled after the load, as asset registry data isnt available correctly until they are
+					FSkinnedAssetCompilingManager::Get().FinishAllCompilation();
+
+					RequestRefreshCachedData();
 				})
 			]
 			+SHorizontalBox::Slot()
@@ -386,7 +390,7 @@ void SAnimAssetFindReplace::Construct(const FArguments& InArgs)
 		]
 	];
 
-	RequestRefresh();
+	RequestRefreshUI();
 
 	RegisterActiveTimer(0.0f, FWidgetActiveTimerDelegate::CreateLambda([this](double InCurrentTime, float InDeltaTime)
 	{
@@ -396,11 +400,17 @@ void SAnimAssetFindReplace::Construct(const FArguments& InArgs)
 	
 	RegisterActiveTimer(1.0f / 60.0f, FWidgetActiveTimerDelegate::CreateLambda([this](double InCurrentTime, float InDeltaTime)
 	{
-		if(bRefreshRequested)
+		if(bRefreshUIRequested)
 		{
-			bRefreshRequested = false;
+			bRefreshUIRequested = false;
 			ToolbarContainer->SetContent(UToolMenus::Get()->GenerateWidget("AnimAssetFindReplaceToolbar", FToolMenuContext(ToolbarContext)));
 			FindReplaceWidgetContainer->SetContent(CurrentProcessor->MakeFindReplaceWidget());
+			bRefreshCachedDataRequested = true;
+		}
+
+		if(bRefreshCachedDataRequested)
+		{
+			bRefreshCachedDataRequested = false;
 			CurrentProcessor->RefreshCachedData();
 			bRefreshSearchResultsRequested = true;
 		}
@@ -423,7 +433,7 @@ void SAnimAssetFindReplace::SetCurrentProcessor(TSubclassOf<UAnimAssetFindReplac
 		CurrentProcessor = *FoundProcessor;
 	}
 
-	RequestRefresh();
+	RequestRefreshUI();
 }
 
 UAnimAssetFindReplaceProcessor* SAnimAssetFindReplace::GetProcessor(TSubclassOf<UAnimAssetFindReplaceProcessor> InProcessorClass) const
@@ -535,7 +545,7 @@ void SAnimAssetFindReplace::ReplaceInAssets(const TArray<FAssetData>& InAssetDat
 		}
 	}
 
-	RequestRefreshSearchResults();
+	RequestRefreshCachedData();
 }
 
 void SAnimAssetFindReplace::ReplaceInAsset(const FAssetData& InAssetData) const
@@ -597,7 +607,7 @@ void SAnimAssetFindReplace::RemoveInAssets(const TArray<FAssetData>& InAssetData
 		}
 	}
 
-	RequestRefreshSearchResults();
+	RequestRefreshCachedData();
 }
 
 void SAnimAssetFindReplace::RemoveInAsset(const FAssetData& InAssetData) const
diff --git a/Engine/Source/Editor/Persona/Private/SAnimAssetFindReplace.h b/Engine/Source/Editor/Persona/Private/SAnimAssetFindReplace.h
--- a/Engine/Source/Editor/Persona/Private/SAnimAssetFindReplace.h
+++ b/Engine/Source/Editor/Persona/Private/SAnimAssetFindReplace.h
@@ -76,8 +76,13 @@ private:
 
 	bool ShouldFilterOutAsset(const FAssetData& InAssetData, bool& bOutIsOldAsset) const;
 
-	void RequestRefresh() { bRefreshRequested = true; }
+	// Request a 'full' refresh. Rebuilds find/replace widgets and also falls through to bRefreshRequested
+	void RequestRefreshUI() { bRefreshUIRequested = true; }
 
+	// Request a refresh of cached data and search results. Re-caches data also falls through to bRefreshSearchResultsRequested
+	void RequestRefreshCachedData() { bRefreshCachedDataRequested = true; }
+
+	// Refresh search results
 	void RequestRefreshSearchResults() { bRefreshSearchResultsRequested = true; }
 
 	FARFilter MakeARFilter() const;
@@ -103,7 +108,9 @@ private:
 
 	bool bFoundAssets = false;
 
-	bool bRefreshRequested = false;
+	bool bRefreshCachedDataRequested = false;
+
+	bool bRefreshUIRequested = false;
 
 	bool bRefreshSearchResultsRequested = false;
 
diff --git a/Engine/Source/Editor/Persona/Private/SAnimCurveMetadataEditor.cpp b/Engine/Source/Editor/Persona/Private/SAnimCurveMetadataEditor.cpp
--- a/Engine/Source/Editor/Persona/Private/SAnimCurveMetadataEditor.cpp
+++ b/Engine/Source/Editor/Persona/Private/SAnimCurveMetadataEditor.cpp
@@ -437,17 +437,7 @@ void SAnimCurveMetadataEditor::Construct(const FArguments& InArgs, UObject* InAn
 				.Icon(FAppStyle::GetBrush("Kismet.Tabs.FindResults"))
 				.OnClicked_Lambda([this]()
 				{
-					if(TSharedPtr<SDockTab> ActiveTab = FGlobalTabmanager::Get()->GetActiveTab())
-					{
-						if(TSharedPtr<FTabManager> TabManager = ActiveTab->GetTabManagerPtr())
-						{
-							if(TSharedPtr<SDockTab> Tab = TabManager->TryInvokeTab(FPersonaTabs::FindReplaceID))
-							{
-								TSharedRef<IAnimAssetFindReplace> FindReplaceWidget = StaticCastSharedRef<IAnimAssetFindReplace>(Tab->GetContent());
-								FindReplaceWidget->SetCurrentProcessor(UAnimAssetFindReplaceCurves::StaticClass());
-							}
-						}
-					}
+					FindReplaceCurves();
 					return FReply::Handled();
 				})
 			]
@@ -575,6 +565,11 @@ void SAnimCurveMetadataEditor::BindCommands()
 		MenuActions.AddCurve,
 		FExecuteAction::CreateSP(this, &SAnimCurveMetadataEditor::OnAddClicked),
 		FCanExecuteAction());
+
+	CommandList.MapAction(
+		MenuActions.FindCurveUses,
+		FExecuteAction::CreateSP(this, &SAnimCurveMetadataEditor::OnFindCurveUsesClicked),
+		FCanExecuteAction::CreateSP(this, &SAnimCurveMetadataEditor::CanFindCurveUses));
 }
 
 void SAnimCurveMetadataEditor::OnPreviewMeshChanged(class USkeletalMesh* OldPreviewMesh, class USkeletalMesh* NewPreviewMesh)
@@ -624,6 +619,7 @@ TSharedPtr<SWidget> SAnimCurveMetadataEditor::OnGetContextMenuContent() const
 	MenuBuilder.AddMenuEntry(FGenericCommands::Get().Copy);
 	MenuBuilder.AddMenuEntry(FGenericCommands::Get().Paste);
 	MenuBuilder.AddMenuEntry(Actions.AddCurve);
+	MenuBuilder.AddMenuEntry(Actions.FindCurveUses);
 
 	MenuBuilder.EndSection();
 
@@ -659,6 +655,52 @@ void SAnimCurveMetadataEditor::OnAddClicked()
 		);
 }
 
+void SAnimCurveMetadataEditor::OnFindCurveUsesClicked()
+{
+	FindReplaceCurves();
+}
+
+bool SAnimCurveMetadataEditor::CanFindCurveUses()
+{
+	return AnimCurveListView->GetNumItemsSelected() == 1;
+}
+
+void SAnimCurveMetadataEditor::FindReplaceCurves()
+{
+	FName CurveName = NAME_None;
+	bool bMorphTarget = false;
+	bool bMaterial = false;
+	TArray<TSharedPtr<FAnimCurveMetadataEditorItem>> SelectedItems = AnimCurveListView->GetSelectedItems();
+	if(SelectedItems.Num() > 0)
+	{
+		CurveName = SelectedItems[0]->CurveName;
+		bMorphTarget = EnumHasAnyFlags(SelectedItems[0]->Flags, EAnimCurveMetadataEditorFilterFlags::MorphTarget);
+		bMaterial = EnumHasAnyFlags(SelectedItems[0]->Flags, EAnimCurveMetadataEditorFilterFlags::Material);
+	}
+
+	if(TSharedPtr<SDockTab> ActiveTab = FGlobalTabmanager::Get()->GetActiveTab())
+	{
+		if(TSharedPtr<FTabManager> TabManager = ActiveTab->GetTabManagerPtr())
+		{
+			if(TSharedPtr<SDockTab> Tab = TabManager->TryInvokeTab(FPersonaTabs::FindReplaceID))
+			{
+				TSharedRef<IAnimAssetFindReplace> FindReplaceWidget = StaticCastSharedRef<IAnimAssetFindReplace>(Tab->GetContent());
+				FindReplaceWidget->SetCurrentProcessor(UAnimAssetFindReplaceCurves::StaticClass());
+				if(CurveName != NAME_None)
+				{
+					if(UAnimAssetFindReplaceCurves* Processor = FindReplaceWidget->GetProcessor<UAnimAssetFindReplaceCurves>())
+					{
+						Processor->SetFindString(CurveName.ToString());
+						Processor->SetFindWholeWord(true);
+						Processor->SetSearchMaterials(bMaterial);
+						Processor->SetSearchMorphTargets(bMorphTarget);
+					}
+				}
+			}
+		}
+	}
+}
+
 void SAnimCurveMetadataEditor::CreateNewNameEntry(const FText& CommittedText, ETextCommit::Type CommitType)
 {
 	FSlateApplication::Get().DismissAllMenus();
diff --git a/Engine/Source/Editor/Persona/Private/SAnimCurveMetadataEditor.h b/Engine/Source/Editor/Persona/Private/SAnimCurveMetadataEditor.h
--- a/Engine/Source/Editor/Persona/Private/SAnimCurveMetadataEditor.h
+++ b/Engine/Source/Editor/Persona/Private/SAnimCurveMetadataEditor.h
@@ -256,6 +256,9 @@ private:
 
 	void OnAddClicked();
 
+	void OnFindCurveUsesClicked();
+	bool CanFindCurveUses();
+	
 	// Adds a new curve metadata name entry to the skeleton
 	void CreateNewNameEntry(const FText& CommittedText, ETextCommit::Type CommitType);
 
@@ -271,6 +274,9 @@ private:
 	/** Get the anim instance we are viewing */
 	UAnimInstance* GetAnimInstance() const;
 
+	/** Open the find/replace tab to process curves. Automatically supplies the name of the selected curve, if any */
+	void FindReplaceCurves();
+
 	/** Pointer to the preview scene we are bound to */
 	TWeakPtr<class IPersonaPreviewScene> PreviewScenePtr;
 
diff --git a/Engine/Source/Editor/Persona/Private/SAnimCurveViewer.cpp b/Engine/Source/Editor/Persona/Private/SAnimCurveViewer.cpp
--- a/Engine/Source/Editor/Persona/Private/SAnimCurveViewer.cpp
+++ b/Engine/Source/Editor/Persona/Private/SAnimCurveViewer.cpp
@@ -99,6 +99,45 @@ private:
 	bool bIsActive = false;
 };
 
+bool FDisplayedAnimCurveInfo::GetActiveFlag(const TSharedPtr<SAnimCurveViewer>& InAnimCurveViewer, bool bInMorphTarget) const
+{
+	if(const UAnimInstance* AnimInstance = InAnimCurveViewer->GetAnimInstance())
+	{
+		// Find if we want to use a pose watch
+		if(UPoseWatchPoseElement* PoseWatchPoseElement = InAnimCurveViewer->PoseWatch.Get())
+		{
+			if(UAnimBlueprintGeneratedClass* AnimClass = Cast<UAnimBlueprintGeneratedClass>(AnimInstance->GetClass()))
+			{
+				// We have to grab our pose watches from the root class as no pose watches can be set on child anim BPs
+				if(const UAnimBlueprintGeneratedClass* RootClass = Cast<UAnimBlueprintGeneratedClass>(AnimClass->GetRootClass()))
+				{
+					const FAnimBlueprintDebugData& DebugData = RootClass->AnimBlueprintDebugData;
+					for(const FAnimNodePoseWatch& AnimNodePoseWatch : DebugData.AnimNodePoseWatch)
+					{
+						if(AnimNodePoseWatch.PoseWatchPoseElement == PoseWatchPoseElement)
+						{
+							UE::Anim::ECurveElementFlags Flags = AnimNodePoseWatch.GetCurves().GetFlags(CurveName);
+							return EnumHasAnyFlags(Flags, bInMorphTarget ? UE::Anim::ECurveElementFlags::MorphTarget : UE::Anim::ECurveElementFlags::Material);
+						}
+					}
+				}
+			}
+		}
+		else
+		{
+			// See if curve is in active set, attribute curve should have everything
+			const TMap<FName, float>& CurveList = AnimInstance->GetAnimationCurveList(bInMorphTarget ? EAnimCurveType::MorphTargetCurve : EAnimCurveType::MaterialCurve);
+			const float* CurrentValue = CurveList.Find(CurveName);
+			if (CurrentValue)
+			{
+				return true;
+			}
+		}
+	}
+
+	return false;
+}
+
 void SAnimCurveListRow::Construct( const FArguments& InArgs, const TSharedRef<STableViewBase>& InOwnerTableView, const TSharedRef<IPersonaPreviewScene>& InPreviewScene)
 {
 	Item = InArgs._Item;
@@ -202,15 +241,20 @@ TSharedRef< SWidget > SAnimCurveListRow::GetCurveTypeWidget()
 		.VAlign(VAlign_Center)
 		.HAlign(HAlign_Center)
 		[
-			SNew(SCheckBox)
-			.IsChecked(this, &SAnimCurveListRow::IsAnimCurveTypeBoxChangedChecked, true)
-			.CheckedImage(FAppStyle::GetBrush("AnimCurveViewer.MorphTargetOn"))
-			.CheckedPressedImage(FAppStyle::GetBrush("AnimCurveViewer.MorphTargetOn"))
-			.UncheckedImage(FAppStyle::GetBrush("AnimCurveViewer.MorphTargetOff"))
-			.CheckedHoveredImage(FAppStyle::GetBrush("AnimCurveViewer.MorphTargetOn"))
-			.UncheckedHoveredImage(FAppStyle::GetBrush("AnimCurveViewer.MorphTargetOff"))
+			SNew(SImage)
 			.ToolTipText(LOCTEXT("CurveTypeMorphTarget_Tooltip", "MorphTarget"))
-			.ForegroundColor(FAppStyle::GetSlateColor("DefaultForeground"))
+			.Image_Lambda([this]()
+			{
+				bool bHasCurveType = GetActiveFlag(true);
+				if(bHasCurveType)
+				{
+					return FAppStyle::GetBrush("AnimCurveViewer.MorphTargetOn");
+				}
+				else
+				{
+					return FAppStyle::GetBrush("AnimCurveViewer.MorphTargetOff");
+				}
+			})
 		]
 
 		+SHorizontalBox::Slot()
@@ -219,15 +263,20 @@ TSharedRef< SWidget > SAnimCurveListRow::GetCurveTypeWidget()
 		.VAlign(VAlign_Center)
 		.HAlign(HAlign_Center)
 		[
-			SNew(SCheckBox)
-			.IsChecked(this, &SAnimCurveListRow::IsAnimCurveTypeBoxChangedChecked, false)
-			.CheckedImage(FAppStyle::GetBrush("AnimCurveViewer.MaterialOn"))
-			.CheckedPressedImage(FAppStyle::GetBrush("AnimCurveViewer.MaterialOn"))
-			.UncheckedImage(FAppStyle::GetBrush("AnimCurveViewer.MaterialOff"))
-			.CheckedHoveredImage(FAppStyle::GetBrush("AnimCurveViewer.MaterialOn"))
-			.UncheckedHoveredImage(FAppStyle::GetBrush("AnimCurveViewer.MaterialOff"))
+			SNew(SImage)
 			.ToolTipText(LOCTEXT("CurveTypeMaterial_Tooltip", "Material"))
-			.ForegroundColor(FAppStyle::GetSlateColor("DefaultForeground"))
+			.Image_Lambda([this]()
+			{
+				bool bHasCurveType = GetActiveFlag(false);
+				if(bHasCurveType)
+				{
+					return FAppStyle::GetBrush("AnimCurveViewer.MaterialOn");
+				}
+				else
+				{
+					return FAppStyle::GetBrush("AnimCurveViewer.MaterialOff");
+				}
+			})
 		];
 
 }
@@ -238,39 +287,7 @@ bool SAnimCurveListRow::GetActiveFlag(bool bMorphTarget) const
 	TSharedPtr<SAnimCurveViewer> AnimCurveViewer = AnimCurveViewerPtr.Pin();
 	if (AnimCurveViewer.IsValid())
 	{
-		if(const UAnimInstance* AnimInstance = AnimCurveViewer->GetAnimInstance())
-		{
-			// Find if we want to use a pose watch
-			if(UPoseWatchPoseElement* PoseWatchPoseElement = AnimCurveViewer->PoseWatch.Get())
-			{
-				if(UAnimBlueprintGeneratedClass* AnimClass = Cast<UAnimBlueprintGeneratedClass>(AnimInstance->GetClass()))
-				{
-					// We have to grab our pose watches from the root class as no pose watches can be set on child anim BPs
-					if(const UAnimBlueprintGeneratedClass* RootClass = Cast<UAnimBlueprintGeneratedClass>(AnimClass->GetRootClass()))
-					{
-						const FAnimBlueprintDebugData& DebugData = RootClass->AnimBlueprintDebugData;
-						for(const FAnimNodePoseWatch& AnimNodePoseWatch : DebugData.AnimNodePoseWatch)
-						{
-							if(AnimNodePoseWatch.PoseWatchPoseElement == PoseWatchPoseElement)
-							{
-								UE::Anim::ECurveElementFlags Flags = AnimNodePoseWatch.GetCurves().GetFlags(Item->CurveName);
-								return EnumHasAnyFlags(Flags, bMorphTarget ? UE::Anim::ECurveElementFlags::MorphTarget : UE::Anim::ECurveElementFlags::Material);
-							}
-						}
-					}
-				}
-			}
-			else
-			{
-				// See if curve is in active set, attribute curve should have everything
-				const TMap<FName, float>& CurveList = AnimInstance->GetAnimationCurveList(bMorphTarget ? EAnimCurveType::MorphTargetCurve : EAnimCurveType::MaterialCurve);
-				const float* CurrentValue = CurveList.Find(Item->CurveName);
-				if (CurrentValue)
-				{
-					return true;
-				}
-			}
-		}
+		return Item->GetActiveFlag(AnimCurveViewer, bMorphTarget);
 	}
 
 	return false;
@@ -480,6 +497,7 @@ void SAnimCurveViewer::Construct(const FArguments& InArgs,  const TSharedRef<IPe
 
 	// Register and bind all our menu commands
 	FCurveViewerCommands::Register();
+	BindCommands();
 
 	CurrentCurveFlag = EAnimCurveViewerFilterFlags::Active;
 
@@ -514,6 +532,7 @@ void SAnimCurveViewer::Construct(const FArguments& InArgs,  const TSharedRef<IPe
 	
 	TSharedRef<SBasicFilterBar<EAnimCurveViewerFilterFlags>> FilterBar = SNew(SBasicFilterBar<EAnimCurveViewerFilterFlags>)
 	.CustomFilters(Filters)
+	.bPinAllFrontendFilters(true)
 	.UseSectionsForCategories(true)
 	.OnFilterChanged_Lambda([this]()
 	{
@@ -528,11 +547,11 @@ void SAnimCurveViewer::Construct(const FArguments& InArgs,  const TSharedRef<IPe
 			}
 		}
 
-		RefreshCurveList(false);
+		RefreshCurveList(true);
 	});
 
-	FilterBar->SetFilterCheckState(Filters[0], ECheckBoxState::Checked);
-	
+	Filters[0]->SetActive(true);
+
 	ChildSlot
 	[
 		SNew( SVerticalBox )
@@ -541,25 +560,34 @@ void SAnimCurveViewer::Construct(const FArguments& InArgs,  const TSharedRef<IPe
 		.AutoHeight()
 		.HAlign(HAlign_Left)
 		[
-			SNew(SSimpleButton)
-			.Text(LOCTEXT("FindReplaceCurvesButton", "Find/Replace Curves..."))
-			.ToolTipText(LOCTEXT("FindReplaceCurvesButtonTooltip", "Find and replace curves across multiple assets"))
-			.Icon(FAppStyle::GetBrush("Kismet.Tabs.FindResults"))
-			.OnClicked_Lambda([this]()
-			{
-				if(TSharedPtr<SDockTab> ActiveTab = FGlobalTabmanager::Get()->GetActiveTab())
+			SNew(SHorizontalBox)
+			+SHorizontalBox::Slot()
+			.AutoWidth()
+			[
+				SNew(SSimpleButton)
+				.Text(LOCTEXT("RefreshButton", "Refresh"))
+				.ToolTipText(LOCTEXT("RefreshButtonTooltip", "Refresh the displayed curves, clearing out any curves that are not currently active"))
+				.Icon(FAppStyle::GetBrush("Icons.Refresh"))
+				.OnClicked_Lambda([this]()
 				{
-					if(TSharedPtr<FTabManager> TabManager = ActiveTab->GetTabManagerPtr())
-					{
-						if(TSharedPtr<SDockTab> Tab = TabManager->TryInvokeTab(FPersonaTabs::FindReplaceID))
-						{
-							TSharedRef<IAnimAssetFindReplace> FindReplaceWidget = StaticCastSharedRef<IAnimAssetFindReplace>(Tab->GetContent());
-							FindReplaceWidget->SetCurrentProcessor(UAnimAssetFindReplaceCurves::StaticClass());
-						}
-					}
-				}
-				return FReply::Handled();
-			})
+					AllSeenAnimCurvesMap.Reset();
+					RefreshCurveList(true);
+					return FReply::Handled();
+				})
+			]
+			+SHorizontalBox::Slot()
+			.AutoWidth()
+			[
+				SNew(SSimpleButton)
+				.Text(LOCTEXT("FindReplaceCurvesButton", "Find/Replace Curves..."))
+				.ToolTipText(LOCTEXT("FindReplaceCurvesButtonTooltip", "Find and replace curves across multiple assets"))
+				.Icon(FAppStyle::GetBrush("Kismet.Tabs.FindResults"))
+				.OnClicked_Lambda([this]()
+				{
+					FindReplaceCurves();
+					return FReply::Handled();
+				})
+			]
 		]
 		+SVerticalBox::Slot()
 		.AutoHeight()
@@ -601,6 +629,7 @@ void SAnimCurveViewer::Construct(const FArguments& InArgs,  const TSharedRef<IPe
 		.FillHeight(1.0f)
 		[
 			SAssignNew(AnimCurveListView, SAnimCurveListType)
+			.OnContextMenuOpening( this, &SAnimCurveViewer::OnGetContextMenuContent )
 			.ListItemsSource( &AnimCurveList )
 			.OnGenerateRow( this, &SAnimCurveViewer::GenerateAnimCurveRow )
 			.ItemHeight( 18.0f )
@@ -682,6 +711,15 @@ void SAnimCurveViewer::Tick(const FGeometry& AllottedGeometry, const double InCu
 	}
 }
 
+FReply SAnimCurveViewer::OnKeyDown(const FGeometry& MyGeometry, const FKeyEvent& InKeyEvent)
+{
+	if (UICommandList.IsValid() && UICommandList->ProcessCommandBindings(InKeyEvent))
+	{
+		return FReply::Handled();
+	}
+	return FReply::Unhandled();
+}
+
 UAnimInstance* SAnimCurveViewer::GetAnimInstance() const
 {
 	UDebugSkelMeshComponent* MeshComponent = PreviewScenePtr.Pin()->GetPreviewMeshComponent();
@@ -713,11 +751,6 @@ void SAnimCurveViewer::CreateAnimCurveList( const FString& SearchText, bool bInF
 
 	AnimCurveList.Reset();
 
-	if(bInFullRefresh)
-	{
-		AllSeenAnimCurvesMap.Reset();
-	}
-
 	auto AddCurve = [this](FName InCurveName, UE::Anim::ECurveElementFlags InFlags)
 	{
 		// Only add if the curve doesnt exist
@@ -1151,5 +1184,85 @@ void SAnimCurveViewer::HandleCurveMetaDataChange()
 	RefreshCurveList(true);
 }
 
+void SAnimCurveViewer::BindCommands()
+{
+	// This should not be called twice on the same instance
+	check(!UICommandList.IsValid());
+
+	UICommandList = MakeShared<FUICommandList>();
+
+	FUICommandList& CommandList = *UICommandList;
+
+	// Grab the list of menu commands to bind...
+	const FCurveViewerCommands& MenuActions = FCurveViewerCommands::Get();
+
+	// ...and bind them all
+	CommandList.MapAction(
+		MenuActions.FindCurveUses,
+		FExecuteAction::CreateSP(this, &SAnimCurveViewer::OnFindCurveUsesClicked),
+		FCanExecuteAction::CreateSP(this, &SAnimCurveViewer::CanFindCurveUses));
+}
+
+TSharedPtr<SWidget> SAnimCurveViewer::OnGetContextMenuContent() const
+{
+	const bool bShouldCloseWindowAfterMenuSelection = true;
+	FMenuBuilder MenuBuilder( bShouldCloseWindowAfterMenuSelection, UICommandList);
+
+	const FCurveViewerCommands& Actions = FCurveViewerCommands::Get();
+
+	MenuBuilder.BeginSection("AnimCurveAction", LOCTEXT( "CurveAction", "Curve Actions" ) );
+	MenuBuilder.AddMenuEntry(Actions.FindCurveUses);
+	MenuBuilder.EndSection();
+
+	return MenuBuilder.MakeWidget();
+}
+
+void SAnimCurveViewer::OnFindCurveUsesClicked()
+{
+	FindReplaceCurves();
+}
+
+bool SAnimCurveViewer::CanFindCurveUses()
+{
+	return AnimCurveListView->GetNumItemsSelected() == 1;
+}
+
+void SAnimCurveViewer::FindReplaceCurves()
+{
+	FName CurveName = NAME_None;
+	bool bMorphTarget = false;
+	bool bMaterial = false;
+	TArray<TSharedPtr<FDisplayedAnimCurveInfo>> SelectedItems = AnimCurveListView->GetSelectedItems();
+	if(SelectedItems.Num() > 0)
+	{
+		CurveName = SelectedItems[0]->CurveName;
+		bMorphTarget = SelectedItems[0]->GetActiveFlag(SharedThis(this), true);
+		bMaterial = SelectedItems[0]->GetActiveFlag(SharedThis(this), false);
+	}
+
+	if(TSharedPtr<SDockTab> ActiveTab = FGlobalTabmanager::Get()->GetActiveTab())
+	{
+		if(TSharedPtr<FTabManager> TabManager = ActiveTab->GetTabManagerPtr())
+		{
+			if(TSharedPtr<SDockTab> Tab = TabManager->TryInvokeTab(FPersonaTabs::FindReplaceID))
+			{
+				TSharedRef<IAnimAssetFindReplace> FindReplaceWidget = StaticCastSharedRef<IAnimAssetFindReplace>(Tab->GetContent());
+				FindReplaceWidget->SetCurrentProcessor(UAnimAssetFindReplaceCurves::StaticClass());
+				if(CurveName != NAME_None)
+				{
+					if(UAnimAssetFindReplaceCurves* Processor = FindReplaceWidget->GetProcessor<UAnimAssetFindReplaceCurves>())
+					{
+						Processor->SetFindString(CurveName.ToString());
+						Processor->SetFindWholeWord(true);
+						Processor->SetSearchMaterials(bMaterial);
+						Processor->SetSearchMorphTargets(bMorphTarget);
+					}
+				}
+			}
+		}
+	}
+}
+
+
 #undef LOCTEXT_NAMESPACE
 
diff --git a/Engine/Source/Editor/Persona/Private/SAnimCurveViewer.h b/Engine/Source/Editor/Persona/Private/SAnimCurveViewer.h
--- a/Engine/Source/Editor/Persona/Private/SAnimCurveViewer.h
+++ b/Engine/Source/Editor/Persona/Private/SAnimCurveViewer.h
@@ -21,6 +21,11 @@
 #include "Filters/FilterBase.h"
 #include "Widgets/Input/SComboBox.h"
 
+namespace UE::Anim
+{
+enum class ECurveElementFlags : uint8;
+}
+
 class FUICommandList;
 class IEditableSkeleton;
 class SAnimCurveViewer;
@@ -49,6 +54,9 @@ public:
 		return MakeShareable(new FDisplayedAnimCurveInfo(InCurveName));
 	}
 
+	// Get the active morph/material flag for this curve 
+	bool GetActiveFlag(const TSharedPtr<SAnimCurveViewer>& InAnimCurveViewer, bool bMorphTarget) const;
+
 protected:
 	/** Hidden constructor, always use Make above */
 	FDisplayedAnimCurveInfo(const FName& InCurveName)
@@ -246,7 +254,9 @@ public:
 	 */
 	void PostUndoRedo();
 
-	virtual void Tick(const FGeometry& AllottedGeometry, const double InCurrentTime, const float InDeltaTime);
+	/** SWidget interface */
+	virtual void Tick(const FGeometry& AllottedGeometry, const double InCurrentTime, const float InDeltaTime) override;
+	virtual FReply OnKeyDown(const FGeometry& MyGeometry, const FKeyEvent& InKeyEvent) override;
 
 	void RefreshCurveList(bool bInFullRefresh);
 
@@ -274,6 +284,15 @@ private:
 	void HandlePoseWatchesChanged(UAnimBlueprint* /*InAnimBlueprint*/, UEdGraphNode* /*InNode*/);
 	void RebuildPoseWatches();
 
+	void BindCommands();
+	
+	/** Build context menu */
+	TSharedPtr<SWidget> OnGetContextMenuContent() const;
+	
+	void OnFindCurveUsesClicked();
+	bool CanFindCurveUses();
+	void FindReplaceCurves();
+	
 	/** Pointer to the preview scene we are bound to */
 	TWeakPtr<class IPersonaPreviewScene> PreviewScenePtr;
 
@@ -304,6 +323,7 @@ private:
 
 	friend class SAnimCurveListRow;
 	friend class SAnimCurveTypeList;
+	friend class FDisplayedAnimCurveInfo;
 
 	/** Delegate called to select objects */
 	FOnObjectsSelected OnObjectsSelected;
diff --git a/Engine/Source/Editor/Persona/Private/SMorphTargetViewer.cpp b/Engine/Source/Editor/Persona/Private/SMorphTargetViewer.cpp
--- a/Engine/Source/Editor/Persona/Private/SMorphTargetViewer.cpp
+++ b/Engine/Source/Editor/Persona/Private/SMorphTargetViewer.cpp
@@ -517,55 +517,22 @@ bool SMorphTargetViewer::CanPerformDelete() const
 void SMorphTargetViewer::OnDeleteMorphTargets()
 {
 	TArray< TSharedPtr< FDisplayedMorphTargetInfo > > SelectedRows = MorphTargetListView->GetSelectedItems();
-	
-	for (int RowIndex = 0; RowIndex < SelectedRows.Num(); ++RowIndex)
+
+	// Clean up override usage
+	TArray<FName> MorphTargetNames;
+	for (int32 RowIndex = 0; RowIndex < SelectedRows.Num(); ++RowIndex)
 	{
 		UMorphTarget* MorphTarget = SkeletalMesh->FindMorphTarget(SelectedRows[RowIndex]->Name);
 		if(MorphTarget)
 		{
-			MorphTarget->RemoveFromRoot();
-			MorphTarget->ClearFlags(RF_Standalone);
-
-			FScopedTransaction Transaction(LOCTEXT("DeleteMorphTarget", "Delete Morph Target"));
-			SkeletalMesh->Modify();
-			MorphTarget->Modify();
-
-			//Clean up override usage
 			AddMorphTargetOverride(SelectedRows[RowIndex]->Name, 0.0f, true);
-
-			if (!SkeletalMesh->IsLODImportedDataEmpty(0) && SkeletalMesh->IsLODImportedDataBuildAvailable(0))
-			{
-				//Remove the morph target from the raw import data
-				FSkeletalMeshImportData SkelMeshImportData;
-				SkeletalMesh->LoadLODImportedData(0, SkelMeshImportData);
-				int32 ToDeleteIndex = INDEX_NONE;
-				for (int32 MoprhTargetIndex = 0; MoprhTargetIndex < SkelMeshImportData.MorphTargetNames.Num(); ++MoprhTargetIndex)
-				{
-					if (SkelMeshImportData.MorphTargetNames[MoprhTargetIndex].Equals(SelectedRows[RowIndex]->Name.ToString()))
-					{
-						ToDeleteIndex = MoprhTargetIndex;
-						break;
-					}
-				}
-
-				if (ToDeleteIndex != INDEX_NONE)
-				{
-					SkelMeshImportData.MorphTargetNames.RemoveAt(ToDeleteIndex);
-					SkelMeshImportData.MorphTargetModifiedPoints.RemoveAt(ToDeleteIndex);
-					SkelMeshImportData.MorphTargets.RemoveAt(ToDeleteIndex);
-					SkeletalMesh->SaveLODImportedData(0, SkelMeshImportData);
-				}
-			}
-			else
-			{
-				//If we deal with an old asset (pre 4.24) and we do not have some valid import data, we need to dirty the ddc key so it wont take the ddc with the old morph target we just delete
-				SkeletalMesh->InvalidateDeriveDataCacheGUID();
-			}
-
-			SkeletalMesh->UnregisterMorphTarget(MorphTarget);
+			MorphTargetNames.Add(SelectedRows[RowIndex]->Name);
 		}
 	}
 
+	// Remove from mesh
+	SkeletalMesh->RemoveMorphTargets(MorphTargetNames);
+
 	CreateMorphTargetList( NameFilterBox->GetText().ToString() );
 }
 
diff --git a/Engine/Source/Editor/Persona/Public/AnimAssetFindReplace.h b/Engine/Source/Editor/Persona/Public/AnimAssetFindReplace.h
--- a/Engine/Source/Editor/Persona/Public/AnimAssetFindReplace.h
+++ b/Engine/Source/Editor/Persona/Public/AnimAssetFindReplace.h
@@ -87,8 +87,14 @@ public:
 	virtual void RefreshCachedData() {}
 
 protected:
+	UE_DEPRECATED(5.4, "Please use RequestRefreshUI.")
+	void RequestUIRefresh() { RequestRefreshUI(); }
+
 	/** Requests the entire find UI to be refreshed */
-	void RequestUIRefresh();
+	void RequestRefreshUI();
+
+	/** Requests cached data and search results */
+	void RequestRefreshCachedData();
 
 	/** Requests search results to be refreshed */
 	void RequestRefreshSearchResults();
@@ -148,7 +154,7 @@ public:
 	/** Optional override point for derived classes to supply auto complete names to display in the UI */
 	virtual void GetAutoCompleteNames(TArrayView<FAssetData> InAssetDatas, TSet<FString>& OutUniqueNames) const {}
 
-private:
+protected:
 	/** UAnimAssetFindReplaceProcessor interface */
 	virtual bool SupportsMode(EAnimAssetFindReplaceMode InMode) const override { return true; }
 	virtual void ExtendToolbar(FToolMenuSection& InSection) override;
diff --git a/Engine/Source/Editor/Persona/Public/AnimAssetFindReplaceCurves.h b/Engine/Source/Editor/Persona/Public/AnimAssetFindReplaceCurves.h
--- a/Engine/Source/Editor/Persona/Public/AnimAssetFindReplaceCurves.h
+++ b/Engine/Source/Editor/Persona/Public/AnimAssetFindReplaceCurves.h
@@ -18,7 +18,28 @@ class UAnimAssetFindReplaceCurves : public UAnimAssetFindReplaceProcessor_String
 	virtual bool ShouldFilterOutAsset(const FAssetData& InAssetData, bool& bOutIsOldAsset) const override;
 	virtual void ReplaceInAsset(const FAssetData& InAssetData) const override;
 	virtual void RemoveInAsset(const FAssetData& InAssetData) const override;
+	virtual void ExtendToolbar(FToolMenuSection& InSection) override;
 
 	// UAnimAssetFindReplaceProcessor_StringBase interface
 	virtual void GetAutoCompleteNames(TArrayView<FAssetData> InAssetDatas, TSet<FString>& OutUniqueNames) const override;
+
+public:
+	// Set whether to search morph targets when searching for curves
+	PERSONA_API void SetSearchMorphTargets(bool bInSearchMorphTargets);
+
+	// Get whether to search morph targets when searching for curves
+	bool GetSearchMorphTargets() const { return bSearchMorphTargets; }
+
+	// Set whether to search material parameters when searching for curves
+	PERSONA_API void SetSearchMaterials(bool bInSearchMaterials);
+
+	// Get whether to search material parameters when searching for curves
+	bool GetSearchMaterials() const { return bSearchMaterials; }
+
+private:
+	// Whether to search morph targets when searching for curves
+	bool bSearchMorphTargets = false;
+
+	// Whether to search materials when searching for curves
+	bool bSearchMaterials = false;
 };
diff --git a/Engine/Source/Runtime/Engine/Classes/Engine/SkeletalMesh.h b/Engine/Source/Runtime/Engine/Classes/Engine/SkeletalMesh.h
--- a/Engine/Source/Runtime/Engine/Classes/Engine/SkeletalMesh.h
+++ b/Engine/Source/Runtime/Engine/Classes/Engine/SkeletalMesh.h
@@ -2436,6 +2436,16 @@ public:
 
 	ENGINE_API void CalculateInvRefMatrices();
 
+#if WITH_EDITORONLY_DATA
+	// Asset registry information for morph target names
+	ENGINE_API static const FName MorphNamesTag;
+	ENGINE_API static const FString MorphNamesTagDelimiter;
+
+	// Asset registry information for scalar material param names
+	ENGINE_API static const FName MaterialParamNamesTag;
+	ENGINE_API static const FString MaterialParamNamesTagDelimiter;
+#endif
+
 #if WITH_EDITOR
 	/** Calculate the required bones for a Skeletal Mesh LOD, including possible extra influences */
 	static ENGINE_API void CalculateRequiredBones(FSkeletalMeshLODModel& LODModel, const struct FReferenceSkeleton& RefSkeleton, const TMap<FBoneIndexType, FBoneIndexType> * BonesToRemove);
@@ -2561,13 +2571,21 @@ public:
 	/** if name conflicts, it will overwrite the reference */
 	ENGINE_API bool RegisterMorphTarget(UMorphTarget* MorphTarget, bool bInvalidateRenderData = true);
 
-	ENGINE_API void UnregisterMorphTarget(UMorphTarget* MorphTarget);
+	ENGINE_API void UnregisterMorphTarget(UMorphTarget* MorphTarget, bool bInvalidateRenderData = true);
 
 	ENGINE_API void UnregisterAllMorphTarget();
 
 	/** Initialize MorphSets look up table : MorphTargetIndexMap */
 	ENGINE_API void InitMorphTargets();
 
+#if WITH_EDITOR
+	/** Remove the morph targets with the specified names */
+	ENGINE_API bool RemoveMorphTargets(TConstArrayView<FName> InMorphTargetNames);
+
+	/** Rename the morph target named InOldName to InNewName */
+	ENGINE_API bool RenameMorphTarget(FName InOldName, FName InNewName);
+#endif
+
 	/** 
 	 * Checks whether the provided section is using APEX cloth. if bCheckCorrespondingSections is true
 	 * disabled sections will defer to correspond sections to see if they use cloth (non-cloth sections
diff --git a/Engine/Source/Runtime/Engine/Private/Animation/Skeleton.cpp b/Engine/Source/Runtime/Engine/Private/Animation/Skeleton.cpp
--- a/Engine/Source/Runtime/Engine/Private/Animation/Skeleton.cpp
+++ b/Engine/Source/Runtime/Engine/Private/Animation/Skeleton.cpp
@@ -393,6 +393,11 @@ void USkeleton::PostLoad()
 			}
 		}
 	}
+	else
+	{
+		// Ensure we have curve metadata, even if empty, so we can correctly pick up older objects in the asset registry
+		GetOrCreateCurveMetaDataObject();
+	}
 #endif
 
 	// refresh linked bone indices
diff --git a/Engine/Source/Runtime/Engine/Private/SkeletalMesh.cpp b/Engine/Source/Runtime/Engine/Private/SkeletalMesh.cpp
--- a/Engine/Source/Runtime/Engine/Private/SkeletalMesh.cpp
+++ b/Engine/Source/Runtime/Engine/Private/SkeletalMesh.cpp
@@ -93,6 +93,7 @@
 #endif
 #include "Animation/SkinWeightProfileManager.h"
 #include "BoneWeights.h"
+#include "Logging/StructuredLog.h"
 
 #define LOCTEXT_NAMESPACE "SkeltalMesh"
 
@@ -141,6 +142,14 @@ static FAutoConsoleVariableRef CVarSkeletalMeshMinLodQualityLevel(
 	FConsoleVariableDelegate::CreateStatic(&USkeletalMesh::OnLodStrippingQualityLevelChanged),
 	ECVF_Scalability);
 
+#if WITH_EDITOR
+const FName USkeletalMesh::MorphNamesTag("MorphTargetNames");
+const FString USkeletalMesh::MorphNamesTagDelimiter(TEXT(";"));
+
+const FName USkeletalMesh::MaterialParamNamesTag("MaterialParamNames");
+const FString USkeletalMesh::MaterialParamNamesTagDelimiter(TEXT(";"));
+#endif
+
 /*-----------------------------------------------------------------------------
 FGPUSkinVertexBase
 -----------------------------------------------------------------------------*/
@@ -3573,6 +3582,45 @@ void USkeletalMesh::GetAssetRegistryTags(FAssetRegistryTagsContext Context) cons
 	// tags are available.
 	Context.AddTag(FAssetRegistryTag("MaxBoneInfluences", MaxBoneInfluencesString, FAssetRegistryTag::TT_Numerical));
 
+	// Expose morph target names to the asset registry
+	{
+		TStringBuilder<256> MorphNamesBuilder;
+		MorphNamesBuilder.Append(MorphNamesTagDelimiter);
+
+		for(UMorphTarget* MorphTarget : GetMorphTargets())
+		{
+			MorphTarget->GetFName().AppendString(MorphNamesBuilder);
+			MorphNamesBuilder.Append(MorphNamesTagDelimiter);
+		}
+
+		Context.AddTag(FAssetRegistryTag(MorphNamesTag, MorphNamesBuilder.ToString(), FAssetRegistryTag::TT_Hidden));
+	}
+
+	// Expose material scalar params (these can be driven by curves)
+	{
+		TStringBuilder<256> MaterialParamNamesBuilder;
+		MaterialParamNamesBuilder.Append(MaterialParamNamesTagDelimiter);
+
+		for (const FSkeletalMaterial& SkeletalMaterial : GetMaterials())
+		{
+			UMaterial* Material = (SkeletalMaterial.MaterialInterface != nullptr) ? SkeletalMaterial.MaterialInterface->GetMaterial() : nullptr;
+			if (Material)
+			{
+				TArray<FMaterialParameterInfo> OutParameterInfo;
+				TArray<FGuid> OutParameterIds;
+				SkeletalMaterial.MaterialInterface->GetAllScalarParameterInfo(OutParameterInfo, OutParameterIds);
+
+				for (const FMaterialParameterInfo& MaterialParameterInfo : OutParameterInfo)
+				{
+					MaterialParameterInfo.Name.AppendString(MaterialParamNamesBuilder);
+					MaterialParamNamesBuilder.Append(MorphNamesTagDelimiter);
+				}
+			}
+		}
+
+		Context.AddTag(FAssetRegistryTag(MaterialParamNamesTag, MaterialParamNamesBuilder.ToString(), FAssetRegistryTag::TT_Hidden));
+	}
+
 	// Allow asset user data to output tags
 	for(UAssetUserData* AssetUserDataItem : *GetAssetUserDataArray())
 	{
@@ -3690,7 +3738,7 @@ void USkeletalMesh::UnregisterAllMorphTarget()
 	InitMorphTargetsAndRebuildRenderData();
 }
 
-void USkeletalMesh::UnregisterMorphTarget(UMorphTarget* MorphTarget)
+void USkeletalMesh::UnregisterMorphTarget(UMorphTarget* MorphTarget, bool bInvalidateRenderData)
 {
 	if ( MorphTarget )
 	{
@@ -3702,7 +3750,10 @@ void USkeletalMesh::UnregisterMorphTarget(UMorphTarget* MorphTarget)
 			{
 				GetMorphTargets().RemoveAt(I);
 				--I;
-				InitMorphTargetsAndRebuildRenderData();
+				if(bInvalidateRenderData)
+				{
+					InitMorphTargetsAndRebuildRenderData();
+				}
 				return;
 			}
 		}
@@ -3760,6 +3811,126 @@ UMorphTarget* USkeletalMesh::FindMorphTargetAndIndex(FName MorphTargetName, int3
 	return nullptr;
 }
 
+#if WITH_EDITOR
+bool USkeletalMesh::RemoveMorphTargets(TConstArrayView<FName> InMorphTargetNames)
+{
+	if(InMorphTargetNames.Num() == 0)
+	{
+		return false;
+	}
+
+	FScopedTransaction Transaction(LOCTEXT("DeleteMorphTargets", "Delete Morph Targets"));
+
+	bool bRemoved = false;
+	for(FName MorphTargetName : InMorphTargetNames)
+	{
+		UMorphTarget* MorphTarget = FindMorphTarget(MorphTargetName);
+		if(MorphTarget)
+		{
+			MorphTarget->RemoveFromRoot();
+			MorphTarget->ClearFlags(RF_Standalone);
+
+			Modify();
+			MorphTarget->Modify();
+
+			if (!IsLODImportedDataEmpty(0) && IsLODImportedDataBuildAvailable(0))
+			{
+				//Remove the morph target from the raw import data
+				FSkeletalMeshImportData SkelMeshImportData;
+				LoadLODImportedData(0, SkelMeshImportData);
+				int32 ToDeleteIndex = INDEX_NONE;
+				for (int32 MorphTargetIndex = 0; MorphTargetIndex < SkelMeshImportData.MorphTargetNames.Num(); ++MorphTargetIndex)
+				{
+					if (SkelMeshImportData.MorphTargetNames[MorphTargetIndex].Equals(MorphTargetName.ToString()))
+					{
+						ToDeleteIndex = MorphTargetIndex;
+						break;
+					}
+				}
+
+				if (ToDeleteIndex != INDEX_NONE)
+				{
+					SkelMeshImportData.MorphTargetNames.RemoveAt(ToDeleteIndex);
+					SkelMeshImportData.MorphTargetModifiedPoints.RemoveAt(ToDeleteIndex);
+					SkelMeshImportData.MorphTargets.RemoveAt(ToDeleteIndex);
+					SaveLODImportedData(0, SkelMeshImportData);
+				}
+			}
+			else
+			{
+				//If we deal with an old asset (pre 4.24) and we do not have some valid import data, we need to dirty the ddc key so it wont take the ddc with the old morph target we just delete
+				InvalidateDeriveDataCacheGUID();
+			}
+
+			UnregisterMorphTarget(MorphTarget);
+
+			bRemoved = true;
+		}
+	}
+
+	return bRemoved;
+}
+
+bool USkeletalMesh::RenameMorphTarget(FName InOldName, FName InNewName)
+{
+	FText Reason;
+	if(!InOldName.IsValidObjectName(Reason) || !InNewName.IsValidObjectName(Reason))
+	{
+		UE_LOGFMT(LogAnimation, Warning, "Could not rename morph target from {0} to {1}. {2}", InOldName, InNewName, Reason.ToString());
+		return false;
+	}
+
+	if(FindObject<UObject>(this, *InNewName.ToString()))
+	{
+		UE_LOGFMT(LogAnimation, Warning, "Could not rename morph target from {0} to {1}. Destination object already exists.", InOldName, InNewName);
+		return false;
+	}
+	
+	UMorphTarget* MorphTarget = FindMorphTarget(InOldName);
+	if(MorphTarget == nullptr)
+	{
+		UE_LOGFMT(LogAnimation, Warning, "Could not rename morph target from {0} to {1}. Could not find morph target.", InOldName, InNewName);
+		return false;
+	}
+
+	FScopedTransaction Transaction(LOCTEXT("RenameMorphTarget", "Rename Morph Target"));
+
+	// Unregister the morph target (but dont invalidate renderdata yet, we will recreate it below in RegisterMorphTarget)
+	UnregisterMorphTarget(MorphTarget, false);
+
+	Modify();
+	MorphTarget->Modify();
+
+	if (!IsLODImportedDataEmpty(0) && IsLODImportedDataBuildAvailable(0))
+	{
+		// Rename the morph target in the raw import data
+		FSkeletalMeshImportData SkelMeshImportData;
+		LoadLODImportedData(0, SkelMeshImportData);
+		for (int32 MorphTargetIndex = 0; MorphTargetIndex < SkelMeshImportData.MorphTargetNames.Num(); ++MorphTargetIndex)
+		{
+			if (SkelMeshImportData.MorphTargetNames[MorphTargetIndex].Equals(InOldName.ToString()))
+			{
+				SkelMeshImportData.MorphTargetNames[MorphTargetIndex] = InNewName.ToString();
+				SaveLODImportedData(0, SkelMeshImportData);
+				break;
+			}
+		}
+	}
+	else
+	{
+		InvalidateDeriveDataCacheGUID();
+	}
+
+	// Rename the morph target itself
+	MorphTarget->Rename(*InNewName.ToString(), nullptr, REN_ForceNoResetLoaders | REN_DontCreateRedirectors);
+
+	// Re-register the morph target
+	RegisterMorphTarget(MorphTarget);
+
+	return true;
+}
+#endif
+
 USkeletalMeshSocket* USkeletalMesh::FindSocket(FName InSocketName) const
 {
 	int32 DummyIdx;
