Fix ensure when removing a notify track

The zeroth track was always skipped for re-indexing. This fix just removes that and ensures that the index is not out of range.

diff --git a/Engine/Source/Editor/Persona/Private/AnimTimeline/AnimTimelineTrack_NotifiesPanel.cpp b/Engine/Source/Editor/Persona/Private/AnimTimeline/AnimTimelineTrack_NotifiesPanel.cpp
--- a/Engine/Source/Editor/Persona/Private/AnimTimeline/AnimTimelineTrack_NotifiesPanel.cpp
+++ b/Engine/Source/Editor/Persona/Private/AnimTimeline/AnimTimelineTrack_NotifiesPanel.cpp
@@ -208,34 +208,31 @@ void FAnimTimelineTrack_NotifiesPanel::RemoveTrack(int32 InTrackIndexToRemove)
 
 	if (AnimSequence->AnimNotifyTracks.IsValidIndex(InTrackIndexToRemove))
 	{
-		if (AnimSequence->AnimNotifyTracks[InTrackIndexToRemove].Notifies.Num() == 0)
+		FScopedTransaction Transaction(LOCTEXT("RemoveNotifyTrack", "Remove Notify Track"));
+		AnimSequence->Modify();
+
+		// before insert, make sure everything behind is fixed
+		for (int32 TrackIndex = InTrackIndexToRemove; TrackIndex < AnimSequence->AnimNotifyTracks.Num(); ++TrackIndex)
 		{
-			FScopedTransaction Transaction(LOCTEXT("RemoveNotifyTrack", "Remove Notify Track"));
-			AnimSequence->Modify();
+			FAnimNotifyTrack& Track = AnimSequence->AnimNotifyTracks[TrackIndex];
+			const int32 NewTrackIndex = FMath::Max(0, TrackIndex - 1);
 
-			// before insert, make sure everything behind is fixed
-			for (int32 TrackIndex = InTrackIndexToRemove; TrackIndex < AnimSequence->AnimNotifyTracks.Num() && TrackIndex > 0; ++TrackIndex)
+			for (FAnimNotifyEvent* Notify : Track.Notifies)
 			{
-				FAnimNotifyTrack& Track = AnimSequence->AnimNotifyTracks[TrackIndex];
-				const int32 NewTrackIndex = TrackIndex - 1;
-
-				for (FAnimNotifyEvent* Notify : Track.Notifies)
-				{
-					// fix notifies indices
-					Notify->TrackIndex = NewTrackIndex;
-				}
+				// fix notifies indices
+				Notify->TrackIndex = NewTrackIndex;
+			}
 
-				for (FAnimSyncMarker* SyncMarker : Track.SyncMarkers)
-				{
-					// fix notifies indices
-					SyncMarker->TrackIndex = NewTrackIndex;
-				}
+			for (FAnimSyncMarker* SyncMarker : Track.SyncMarkers)
+			{
+				// fix notifies indices
+				SyncMarker->TrackIndex = NewTrackIndex;
 			}
+		}
 
-			AnimSequence->AnimNotifyTracks.RemoveAt(InTrackIndexToRemove);
+		AnimSequence->AnimNotifyTracks.RemoveAt(InTrackIndexToRemove);
 
-			Update();
-		}
+		Update();
 	}
 }
 
