Ensure that both notify and notify state properties are copied when copy/pasting

PR #10681: fixed replace anim notify state that copy properties failed

diff --git a/Engine/Source/Editor/Persona/Private/SAnimNotifyPanel.cpp b/Engine/Source/Editor/Persona/Private/SAnimNotifyPanel.cpp
--- a/Engine/Source/Editor/Persona/Private/SAnimNotifyPanel.cpp
+++ b/Engine/Source/Editor/Persona/Private/SAnimNotifyPanel.cpp
@@ -4491,6 +4491,7 @@ void SAnimNotifyPanel::OnReplaceSelectedWithNotify(FString NewNotifyName, UClass
 
 			FColor OldColor = OldEvent->NotifyColor;
 			UAnimNotify* OldEventPayload = OldEvent->Notify;
+			UAnimNotifyState* OldEventStatePayload = OldEvent->NotifyStateClass;
 
 			// Delete old one before creating new one to avoid potential array re-allocation when array temporarily increases by 1 in size
 			NodeObject->Delete(Sequence);
@@ -4513,6 +4514,13 @@ void SAnimNotifyPanel::OnReplaceSelectedWithNotify(FString NewNotifyName, UClass
 			// For Anim Notify States, handle the end time and link
 			if (NewEvent.NotifyStateClass != nullptr)
 			{
+				if (OldEventStatePayload != nullptr)
+				{
+					UEngine::FCopyPropertiesForUnrelatedObjectsParams CopyParams;
+					CopyParams.bNotifyObjectReplacement = true;
+					UEngine::CopyPropertiesForUnrelatedObjects(OldEventStatePayload, NewEvent.NotifyStateClass, CopyParams);
+				}
+
 				NewEvent.SetDuration(Length);
 				NewEvent.EndTriggerTimeOffset = EndTriggerTimeOffset;
 				NewEvent.EndLink.ChangeSlotIndex(EndSlotIndex);
