Anim timeline crash fixes

Fixed accessing deleted memory - widgets outlive the memory pointed to by curves when refreshing the view (by a frame) so added cached values for attribute-based display.
Also fixed metadata track display (was using an out of date brush).

diff --git a/Engine/Source/Editor/Persona/Private/AnimTimeline/AnimModel_AnimSequenceBase.cpp b/Engine/Source/Editor/Persona/Private/AnimTimeline/AnimModel_AnimSequenceBase.cpp
--- a/Engine/Source/Editor/Persona/Private/AnimTimeline/AnimModel_AnimSequenceBase.cpp
+++ b/Engine/Source/Editor/Persona/Private/AnimTimeline/AnimModel_AnimSequenceBase.cpp
@@ -385,7 +385,7 @@ void FAnimModel_AnimSequenceBase::EditSelectedCurves()
 		if(SelectedTrack->IsA<FAnimTimelineTrack_Curve>())
 		{
 			TSharedRef<FAnimTimelineTrack_Curve> CurveTrack = StaticCastSharedRef<FAnimTimelineTrack_Curve>(SelectedTrack);
-			const TArray<const FRichCurve*> Curves = CurveTrack->GetCurves();
+			const TArray<const FRichCurve*>& Curves = CurveTrack->GetCurves();
 			int32 NumCurves = Curves.Num();
 			for(int32 CurveIndex = 0; CurveIndex < NumCurves; ++CurveIndex)
 			{
diff --git a/Engine/Source/Editor/Persona/Private/AnimTimeline/AnimTimelineTrack_FloatCurve.cpp b/Engine/Source/Editor/Persona/Private/AnimTimeline/AnimTimelineTrack_FloatCurve.cpp
--- a/Engine/Source/Editor/Persona/Private/AnimTimeline/AnimTimelineTrack_FloatCurve.cpp
+++ b/Engine/Source/Editor/Persona/Private/AnimTimeline/AnimTimelineTrack_FloatCurve.cpp
@@ -33,6 +33,9 @@ FAnimTimelineTrack_FloatCurve::FAnimTimelineTrack_FloatCurve(const FFloatCurve*
 	, FloatCurve(InCurve)
 	, CurveName(InCurve->GetName())
 	, CurveId(FAnimationCurveIdentifier(InCurve->GetName(), ERawCurveTrackTypes::RCT_Float))
+	, Color(InCurve->Color)
+	, Comment(InCurve->Comment) 
+	, bIsMetadata(InCurve->GetCurveTypeFlag(AACF_Metadata))
 {
 	SetHeight(32.0f);
 }
@@ -54,7 +57,7 @@ TSharedRef<SWidget> FAnimTimelineTrack_FloatCurve::MakeTimelineWidgetContainer()
 			[
 				SAssignNew(TimelineWidgetContainer, SBorder)
 				.Padding(0.0f)
-				.BorderImage(FloatCurve->GetCurveTypeFlag(AACF_Metadata) ? FAppStyle::GetBrush("Sequencer.Section.SelectedSectionOverlay") : FAppStyle::GetBrush("AnimTimeline.Outliner.DefaultBorder"))
+				.BorderImage(bIsMetadata ? FAppStyle::GetBrush("Sequencer.Section.CollapsedSelectedSectionOverlay") : FAppStyle::GetBrush("AnimTimeline.Outliner.DefaultBorder"))
 				.BorderBackgroundColor(this, &FAnimTimelineTrack_FloatCurve::GetTrackColor, false)
 				[
 					CurveWidget
@@ -85,13 +88,13 @@ TSharedRef<SWidget> FAnimTimelineTrack_FloatCurve::MakeTimelineWidgetContainer()
 
 float FAnimTimelineTrack_FloatCurve::GetCommentHeight() const
 {
-	if(FloatCurve->Comment.IsEmpty())
+	if(Comment.IsEmpty())
 	{
 		return 0.0f;
 	}
 
 	const TSharedRef<FSlateFontMeasure> FontMeasureService = FSlateApplication::Get().GetRenderer()->GetFontMeasureService();
-	UE::Slate::FDeprecateVector2DResult Result = FontMeasureService->Measure(FloatCurve->Comment, FCoreStyle::Get().GetWidgetStyle<FTextBlockStyle>( "NormalText" ).Font);
+	UE::Slate::FDeprecateVector2DResult Result = FontMeasureService->Measure(Comment, FCoreStyle::Get().GetWidgetStyle<FTextBlockStyle>( "NormalText" ).Font);
 	return Result.Y + 6.0f;
 }
 
@@ -141,7 +144,7 @@ TSharedRef<SWidget> FAnimTimelineTrack_FloatCurve::BuildCurveTrackMenu()
 {
 	FMenuBuilder MenuBuilder(true, GetModel()->GetCommandList());
 
-	bool bIsMetadata = FloatCurve->GetCurveTypeFlag(AACF_Metadata);
+	bIsMetadata = FloatCurve->GetCurveTypeFlag(AACF_Metadata);
 
 	MenuBuilder.BeginSection("Curve", bIsMetadata ? LOCTEXT("CurveMetadataMenuSection", "Curve Metadata") : LOCTEXT("CurveMenuSection", "Curve"));
 	{
@@ -209,6 +212,8 @@ void FAnimTimelineTrack_FloatCurve::ConvertCurveToMetaData()
 	Controller.SetCurveFlag(CurveId, AACF_Metadata, true);
 	Controller.SetCurveKeys(CurveId, { FRichCurveKey(0.f, 1.f) });	
 
+	bIsMetadata = true;
+
 	ZoomToFit();
 }
 
@@ -218,7 +223,8 @@ void FAnimTimelineTrack_FloatCurve::ConvertMetaDataToCurve()
 
 	IAnimationDataController& Controller = AnimSequenceBase->GetController();
 	IAnimationDataController::FScopedBracket ScopedBracket(Controller, LOCTEXT("CurvePanel_ConvertMetaDataToCurve", "Convert metadata to curve"));
-	Controller.SetCurveFlag(CurveId, AACF_Metadata, false);	
+	Controller.SetCurveFlag(CurveId, AACF_Metadata, false);
+	bIsMetadata = false;
 }
 
 void FAnimTimelineTrack_FloatCurve::RemoveCurve()
@@ -310,7 +316,7 @@ void FAnimTimelineTrack_FloatCurve::Copy(UAnimTimelineClipboardContent* InOutCli
 
 bool FAnimTimelineTrack_FloatCurve::CanEditCurve(int32 InCurveIndex) const
 {
-	return !FloatCurve->GetCurveTypeFlag(AACF_Metadata);
+	return !bIsMetadata;
 }
 
 void FAnimTimelineTrack_FloatCurve::RequestRename()
@@ -337,6 +343,7 @@ void FAnimTimelineTrack_FloatCurve::AddCurveTrackButton(TSharedPtr<SHorizontalBo
 	{
 		UAnimSequenceBase* AnimSequenceBase = GetModel()->GetAnimSequenceBase();
 		AnimSequenceBase->GetController().SetCurveColor(CurveId, InNewColor);
+		Color = InNewColor;
 
 		// Set display curves too
 		for(const TPair<FCurveModelID, TUniquePtr<FCurveModel>>& CurvePair : CurveEditor->GetCurves())
@@ -392,7 +399,7 @@ void FAnimTimelineTrack_FloatCurve::AddCurveTrackButton(TSharedPtr<SHorizontalBo
 
 FLinearColor FAnimTimelineTrack_FloatCurve::GetCurveColor(int32 InCurveIndex) const
 { 
-	return FloatCurve->Color; 
+	return Color; 
 }
 
 void FAnimTimelineTrack_FloatCurve::GetCurveEditInfo(int32 InCurveIndex, FName& OutName, ERawCurveTrackTypes& OutType, int32& OutCurveIndex) const
@@ -404,10 +411,12 @@ void FAnimTimelineTrack_FloatCurve::GetCurveEditInfo(int32 InCurveIndex, FName&
 
 void FAnimTimelineTrack_FloatCurve::HandleAddComment()
 {
-	if(FloatCurve->Comment.IsEmpty())
+	if(Comment.IsEmpty())
 	{
 		UAnimSequenceBase* AnimSequenceBase = GetModel()->GetAnimSequenceBase();
-		AnimSequenceBase->GetController().SetCurveComment(CurveId, LOCTEXT("DefaultComment", "Comment").ToString());
+		FString NewComment = LOCTEXT("DefaultComment", "Comment").ToString();
+		AnimSequenceBase->GetController().SetCurveComment(CurveId, NewComment);
+		Comment = NewComment;
 	}
 
 	ExecuteOnGameThread(UE_SOURCE_LOCATION, [WeakThis = TWeakPtr<FAnimTimelineTrack_FloatCurve>(SharedThis(this))]()
@@ -421,21 +430,22 @@ void FAnimTimelineTrack_FloatCurve::HandleAddComment()
 
 void FAnimTimelineTrack_FloatCurve::OnCommitCurveComment(const FText& InText, ETextCommit::Type CommitInfo)
 {
-	if(FloatCurve->Comment != InText.ToString())
+	if(Comment != InText.ToString())
 	{
 		UAnimSequenceBase* AnimSequenceBase = GetModel()->GetAnimSequenceBase();
 		AnimSequenceBase->GetController().SetCurveComment(CurveId, InText.ToString());
+		Comment = FloatCurve->Comment;
 	}
 }
 
 FText FAnimTimelineTrack_FloatCurve::GetCommentText() const
 {
-	return FText::FromString(FloatCurve->Comment);
+	return FText::FromString(Comment);
 }
 
 EVisibility FAnimTimelineTrack_FloatCurve::GetCommentVisibility() const
 {
-	return FloatCurve->Comment.IsEmpty() ? EVisibility::Collapsed : EVisibility::Visible;
+	return Comment.IsEmpty() ? EVisibility::Collapsed : EVisibility::Visible;
 }
 
 bool FAnimTimelineTrack_FloatCurve::IsSelected() const
@@ -451,7 +461,7 @@ FSlateColor FAnimTimelineTrack_FloatCurve::GetTrackColor(bool bForComment) const
 	}
 	else
 	{
-		FLinearColor CurveColor = FloatCurve->GetCurveTypeFlag(AACF_Metadata) ? FloatCurve->GetColor().Desaturate(0.25f) : FloatCurve->GetColor().Desaturate(0.75f);
+		FLinearColor CurveColor = bIsMetadata ? Color.Desaturate(0.25f) : Color.Desaturate(0.75f);
 		return bForComment ? CurveColor.CopyWithNewOpacity(CurveColor.A * 0.5f) : CurveColor;
 	}
 }
diff --git a/Engine/Source/Editor/Persona/Private/AnimTimeline/AnimTimelineTrack_FloatCurve.h b/Engine/Source/Editor/Persona/Private/AnimTimeline/AnimTimelineTrack_FloatCurve.h
--- a/Engine/Source/Editor/Persona/Private/AnimTimeline/AnimTimelineTrack_FloatCurve.h
+++ b/Engine/Source/Editor/Persona/Private/AnimTimeline/AnimTimelineTrack_FloatCurve.h
@@ -69,6 +69,15 @@ private:
 	FName CurveName;
 	FAnimationCurveIdentifier CurveId;
 
+	/** Cached color */
+	FLinearColor Color;
+
+	/** Cached comment */
+	FString Comment;
+
+	/** Cached metadata flag */
+	bool bIsMetadata = false;
+
 	/** Label we can edit */
 	TSharedPtr<SInlineEditableTextBlock> EditableTextLabel;
 
