Persona: Clear out the preview scene when destroying the animation viewport client to avoid any reference to it while the rest of the asset editor is being torn down.

#rnx

diff --git a/Engine/Source/Editor/Persona/Private/AnimationEditorViewportClient.cpp b/Engine/Source/Editor/Persona/Private/AnimationEditorViewportClient.cpp
--- a/Engine/Source/Editor/Persona/Private/AnimationEditorViewportClient.cpp
+++ b/Engine/Source/Editor/Persona/Private/AnimationEditorViewportClient.cpp
@@ -181,6 +181,10 @@ FAnimationViewportClient::~FAnimationViewportClient()
 	OnPhysicsCreatedDelegateHandle.Reset();
 	OnMeshChangedDelegateHandle.Reset();
 
+	// Clear out the preview scene so any subsequent editor object destruction that tries to poke into the
+	// world will not crash.
+	static_cast<FAssetEditorModeManager*>(ModeTools.Get())->SetPreviewScene(nullptr);
+	
 	UAssetViewerSettings::Get()->OnAssetViewerSettingsChanged().RemoveAll(this);
 }
 
