Persona: Dynamically switch between the bone picker and text box in the FBoneReference widget, if the skeleton is valid or not. Similar for the FBoneSocketTarget.

#rnx

diff --git a/Engine/Plugins/Experimental/Animation/SkeletalMeshModelingTools/Source/SkeletalMeshModelingTools/Private/SkeletalMeshModelingToolsMeshConverter.cpp b/Engine/Plugins/Experimental/Animation/SkeletalMeshModelingTools/Source/SkeletalMeshModelingTools/Private/SkeletalMeshModelingToolsMeshConverter.cpp
--- a/Engine/Plugins/Experimental/Animation/SkeletalMeshModelingTools/Source/SkeletalMeshModelingTools/Private/SkeletalMeshModelingToolsMeshConverter.cpp
+++ b/Engine/Plugins/Experimental/Animation/SkeletalMeshModelingTools/Source/SkeletalMeshModelingTools/Private/SkeletalMeshModelingToolsMeshConverter.cpp
@@ -164,8 +164,6 @@ void UStaticMeshToSkeletalMeshConvertOptions::PostEditChangeProperty(FPropertyCh
 					BindingBoneName.Reset();
 				}
 			}
-			
-			(void)SkeletonProviderChanged.ExecuteIfBound();
 		}
 	}
 }
@@ -397,11 +395,6 @@ void ConvertStaticMeshAssetsToSkeletalMeshesInteractive(
 		SCustomDialog::FButton(LOCTEXT("DialogButtonCancel", "Cancel"))
 	});
 
-	Options->SkeletonProviderChanged.BindLambda([DetailsView]()
-	{
-		DetailsView->ForceRefresh();
-	});
-
 	OptionsDialog->Show();
 }
 
diff --git a/Engine/Plugins/Experimental/Animation/SkeletalMeshModelingTools/Source/SkeletalMeshModelingTools/Private/SkeletalMeshModelingToolsMeshConverter.h b/Engine/Plugins/Experimental/Animation/SkeletalMeshModelingTools/Source/SkeletalMeshModelingTools/Private/SkeletalMeshModelingToolsMeshConverter.h
--- a/Engine/Plugins/Experimental/Animation/SkeletalMeshModelingTools/Source/SkeletalMeshModelingTools/Private/SkeletalMeshModelingToolsMeshConverter.h
+++ b/Engine/Plugins/Experimental/Animation/SkeletalMeshModelingTools/Source/SkeletalMeshModelingTools/Private/SkeletalMeshModelingToolsMeshConverter.h
@@ -5,6 +5,7 @@
 #include "BoneContainer.h"
 #include "Factories/Factory.h"
 #include "ReferenceSkeleton.h"
+#include "Animation/BoneSocketReference.h"
 #include "Interfaces/Interface_BoneReferenceSkeletonProvider.h"
 
 #include "SkeletalMeshModelingToolsMeshConverter.generated.h"
@@ -150,8 +151,6 @@ public:
 	
 	// UObject overrides.
 	virtual void PostEditChangeProperty(FPropertyChangedEvent& PropertyChangedEvent) override;
-
-	FSimpleDelegate SkeletonProviderChanged;
 };
 
 
diff --git a/Engine/Source/Editor/Persona/Private/AnimGraphNodeDetails.cpp b/Engine/Source/Editor/Persona/Private/AnimGraphNodeDetails.cpp
--- a/Engine/Source/Editor/Persona/Private/AnimGraphNodeDetails.cpp
+++ b/Engine/Source/Editor/Persona/Private/AnimGraphNodeDetails.cpp
@@ -46,6 +46,7 @@
 #include "Algo/Accumulate.h"
 #include "Kismet2/BlueprintEditorUtils.h"
 #include "Widgets/Input/SEditableTextBox.h"
+#include "Widgets/Layout/SWidgetSwitcher.h"
 
 #define LOCTEXT_NAMESPACE "KismetNodeWithOptionalPinsDetails"
 
@@ -507,35 +508,34 @@ TSharedRef<IPropertyTypeCustomization> FBoneReferenceCustomization::MakeInstance
 
 void FBoneReferenceCustomization::CustomizeHeader( TSharedRef<IPropertyHandle> StructPropertyHandle, class FDetailWidgetRow& HeaderRow, IPropertyTypeCustomizationUtils& StructCustomizationUtils )
 {
-	// set property handle 
-	SetPropertyHandle(StructPropertyHandle);
-	// set editable skeleton info from struct
-	SetEditableSkeleton(StructPropertyHandle);
-	if (TargetEditableSkeleton.IsValid() && BoneNameProperty->IsValidHandle())
-	{
-		HeaderRow
-		.NameContent()
-		[
-			StructPropertyHandle->CreatePropertyNameWidget()
-		]
-		.ValueContent()
-		.MaxDesiredWidth(0.0f)
+	StructProperty = StructPropertyHandle;
+	
+	BoneNameProperty = FindStructMemberProperty(StructProperty, GET_MEMBER_NAME_CHECKED(FBoneReference, BoneName));
+	check(BoneNameProperty->IsValidHandle());
+	
+	HeaderRow
+	.NameContent()
+	[
+		StructPropertyHandle->CreatePropertyNameWidget()
+	]
+	.ValueContent()
+	.MaxDesiredWidth(0.0f)
+	[
+		SNew(SWidgetSwitcher)
+		.WidgetIndex_Lambda([this]()
+		{
+			// If we have a skeleton, show the picker widget, otherwise default to the text box.
+			return GetSkeleton() != nullptr ? 0 : 1;
+		})
+		+ SWidgetSwitcher::Slot()
 		[
 			SNew(SBoneSelectionWidget)
 			.ToolTipText(StructPropertyHandle->GetToolTipText())
 			.OnBoneSelectionChanged(this, &FBoneReferenceCustomization::OnBoneSelectionChanged)
 			.OnGetSelectedBone(this, &FBoneReferenceCustomization::GetSelectedBone)
 			.OnGetReferenceSkeleton(this, &FBoneReferenceCustomization::GetReferenceSkeleton)
-		];
-	}
-	else
-	{
-		HeaderRow
-		.NameContent()
-		[
-			StructPropertyHandle->CreatePropertyNameWidget()
 		]
-		.ValueContent()
+		+ SWidgetSwitcher::Slot()
 		[
 			SNew(SEditableTextBox)
 			.Font(StructCustomizationUtils.GetRegularFont())
@@ -550,54 +550,46 @@ void FBoneReferenceCustomization::CustomizeHeader( TSharedRef<IPropertyHandle> S
 				BoneNameProperty->GetValue(Name);
 				return FText::FromName(Name);
 			})
-		];
-	}
+		]
+	];
 }
 
-void FBoneReferenceCustomization::CustomizeChildren(TSharedRef<IPropertyHandle> StructPropertyHandle, class IDetailChildrenBuilder& ChildBuilder, IPropertyTypeCustomizationUtils& StructCustomizationUtils)
-{
 
-}
-void FBoneReferenceCustomization::SetEditableSkeleton(TSharedRef<IPropertyHandle> StructPropertyHandle) 
+USkeleton* FBoneReferenceCustomization::GetSkeleton() const 
 {
 	TArray<UObject*> Objects;
-	StructPropertyHandle->GetOuterObjects(Objects);
-
-	USkeleton* TargetSkeleton = nullptr;
-	TSharedPtr<IEditableSkeleton> EditableSkeleton;
-
-	bEnsureOnInvalidSkeleton = true;
+	StructProperty->GetOuterObjects(Objects);
 
-	auto FindSkeletonForObject = [this, &TargetSkeleton, &EditableSkeleton, &StructPropertyHandle](UObject* InObject)
+	auto FindSkeletonForObject = [this](UObject* InObject) -> USkeleton* 
 	{
 		for( ; InObject; InObject = InObject->GetOuter())
 		{
-			if (UAnimGraphNode_Base* AnimGraphNode = Cast<UAnimGraphNode_Base>(InObject))
+			if (const UAnimGraphNode_Base* AnimGraphNode = Cast<UAnimGraphNode_Base>(InObject))
 			{
-				TargetSkeleton = AnimGraphNode->GetAnimBlueprint()->TargetSkeleton;
-				break;
+				return AnimGraphNode->GetAnimBlueprint()->TargetSkeleton;
 			}
 
 			if (USkeletalMesh* SkeletalMesh = Cast<USkeletalMesh>(InObject))
 			{
-				TargetSkeleton = SkeletalMesh->GetSkeleton();
-				break;
+				return SkeletalMesh->GetSkeleton();
 			}
 
-			if (ULODInfoUILayout* LODInfoUILayout = Cast<ULODInfoUILayout>(InObject))
+			if (const ULODInfoUILayout* LODInfoUILayout = Cast<ULODInfoUILayout>(InObject))
 			{
 				USkeletalMesh* SkeletalMesh = LODInfoUILayout->GetPersonaToolkit()->GetPreviewMesh();
-				check(SkeletalMesh);
-				TargetSkeleton = SkeletalMesh->GetSkeleton();
-				break;
+				if (ensure(SkeletalMesh))
+				{
+					return SkeletalMesh->GetSkeleton();
+				}
+
+				return nullptr;
 			}
 
-			if (UAnimationAsset* AnimationAsset = Cast<UAnimationAsset>(InObject))
+			if (const UAnimationAsset* AnimationAsset = Cast<UAnimationAsset>(InObject))
 			{
 				if(AnimationAsset->IsAsset())
 				{
-					TargetSkeleton = AnimationAsset->GetSkeleton();
-					break;
+					return AnimationAsset->GetSkeleton();
 				}
 			}
 
@@ -605,59 +597,50 @@ void FBoneReferenceCustomization::SetEditableSkeleton(TSharedRef<IPropertyHandle
 			{
 				if (AnimInstance->CurrentSkeleton)
 				{
-					TargetSkeleton = AnimInstance->CurrentSkeleton;
-					break;
+					return AnimInstance->CurrentSkeleton;
 				}
-				else if (UAnimBlueprintGeneratedClass* AnimBPClass = Cast<UAnimBlueprintGeneratedClass>(AnimInstance->GetClass()))
+				if (UAnimBlueprintGeneratedClass* AnimBPClass = Cast<UAnimBlueprintGeneratedClass>(AnimInstance->GetClass()))
 				{
-					TargetSkeleton = AnimBPClass->TargetSkeleton;
-					break;
+					return AnimBPClass->TargetSkeleton;
 				}
 			}
 
 			// editor animation curve bone links are responsible for linking joints to curve
 			// this is editor object that only exists for editor
-			if (UEditorAnimCurveBoneLinks* AnimCurveObj = Cast<UEditorAnimCurveBoneLinks>(InObject))
+			if (const UEditorAnimCurveBoneLinks* AnimCurveObj = Cast<UEditorAnimCurveBoneLinks>(InObject))
 			{
 				if(USkeleton* Skeleton = Cast<USkeleton>(AnimCurveObj->AnimCurveMetaData->GetOuter()))
 				{
-					TargetSkeleton = Skeleton;
+					return Skeleton;
 				}
-				else if(USkeletalMesh* SkeletalMesh = Cast<USkeletalMesh>(AnimCurveObj->AnimCurveMetaData->GetOuter()))
+				if(USkeletalMesh* SkeletalMesh = Cast<USkeletalMesh>(AnimCurveObj->AnimCurveMetaData->GetOuter()))
 				{
-					TargetSkeleton = SkeletalMesh->GetSkeleton();
+					return SkeletalMesh->GetSkeleton();
 				}
-				break;
 			}
 
 			if (IBoneReferenceSkeletonProvider* SkeletonProvider = Cast<IBoneReferenceSkeletonProvider>(InObject))
 			{
-				TargetSkeleton = SkeletonProvider->GetSkeleton(bEnsureOnInvalidSkeleton, &StructPropertyHandle.Get());
-				break;
+				bool bInvalidSkeletonIsError = false;
+				return SkeletonProvider->GetSkeleton(bInvalidSkeletonIsError, StructProperty.Get());
 			}
 		}
 
-		return TargetSkeleton != nullptr || EditableSkeleton != nullptr;
+		return nullptr;
 	};
 	
 	for (UObject* Object : Objects)
 	{
-		if(FindSkeletonForObject(Object))
+		if(USkeleton* Skeleton = FindSkeletonForObject(Object))
 		{
-			break;
+			return Skeleton;
 		}
 	}
 
-	if (TargetSkeleton != nullptr)
-	{
-		ISkeletonEditorModule& SkeletonEditorModule = FModuleManager::LoadModuleChecked<ISkeletonEditorModule>("SkeletonEditor");
-		EditableSkeleton = SkeletonEditorModule.CreateEditableSkeleton(TargetSkeleton);
-	}
-
-	TargetEditableSkeleton = EditableSkeleton;
+	return nullptr;
 }
 
-TSharedPtr<IPropertyHandle> FBoneReferenceCustomization::FindStructMemberProperty(TSharedRef<IPropertyHandle> PropertyHandle, const FName& PropertyName)
+TSharedPtr<IPropertyHandle> FBoneReferenceCustomization::FindStructMemberProperty(TSharedPtr<IPropertyHandle> PropertyHandle, const FName& PropertyName)
 {
 	uint32 NumChildren = 0;
 	PropertyHandle->GetNumChildren(NumChildren);
@@ -673,12 +656,6 @@ TSharedPtr<IPropertyHandle> FBoneReferenceCustomization::FindStructMemberPropert
 	return TSharedPtr<IPropertyHandle>();
 }
 
-void FBoneReferenceCustomization::SetPropertyHandle(TSharedRef<IPropertyHandle> StructPropertyHandle)
-{
-	BoneNameProperty = FindStructMemberProperty(StructPropertyHandle, GET_MEMBER_NAME_CHECKED(FBoneReference, BoneName));
-	check(BoneNameProperty->IsValidHandle());
-}
-
 void FBoneReferenceCustomization::OnBoneSelectionChanged(FName Name)
 {
 	BoneNameProperty->SetValue(Name);
@@ -688,7 +665,7 @@ FName FBoneReferenceCustomization::GetSelectedBone(bool& bMultipleValues) const
 {
 	FString OutText;
 	
-	FPropertyAccess::Result Result = BoneNameProperty->GetValueAsFormattedString(OutText);
+	const FPropertyAccess::Result Result = BoneNameProperty->GetValueAsFormattedString(OutText);
 	bMultipleValues = (Result == FPropertyAccess::MultipleValues);
 
 	return FName(*OutText);
@@ -696,10 +673,10 @@ FName FBoneReferenceCustomization::GetSelectedBone(bool& bMultipleValues) const
 
 const struct FReferenceSkeleton&  FBoneReferenceCustomization::GetReferenceSkeleton() const
 {
-	// retruning dummy skeleton if any reason, it is invalid
-	static FReferenceSkeleton DummySkeleton;
-
-	return (TargetEditableSkeleton.IsValid()) ? TargetEditableSkeleton.Get()->GetSkeleton().GetReferenceSkeleton() : DummySkeleton;
+	static FReferenceSkeleton EmptySkeleton;
+	
+	const USkeleton* Skeleton = GetSkeleton();
+	return Skeleton ? Skeleton->GetReferenceSkeleton() : EmptySkeleton; 
 }
 
 /////////////////////////////////////////////////////////////////////////////////////////////
@@ -710,59 +687,61 @@ TSharedRef<IPropertyTypeCustomization> FBoneSocketTargetCustomization::MakeInsta
 	return MakeShareable(new FBoneSocketTargetCustomization());
 }
 
-void FBoneSocketTargetCustomization::CustomizeChildren(TSharedRef<IPropertyHandle> StructPropertyHandle, class IDetailChildrenBuilder& ChildBuilder, IPropertyTypeCustomizationUtils& StructCustomizationUtils)
+void FBoneSocketTargetCustomization::CustomizeHeader(TSharedRef<IPropertyHandle> StructPropertyHandle, class FDetailWidgetRow& HeaderRow, IPropertyTypeCustomizationUtils& StructCustomizationUtils)
 {
-	// set property handle 
-	SetPropertyHandle(StructPropertyHandle);
-	// set editable skeleton info from struct
-	SetEditableSkeleton(StructPropertyHandle);
-	Build(StructPropertyHandle, ChildBuilder);
+	// set property handle
+	StructProperty = StructPropertyHandle;
+	
+	ResolveChildProperties();
+	
+	HeaderRow
+	.NameContent()
+	[
+		StructPropertyHandle->CreatePropertyNameWidget()
+	]
+	.ValueContent()
+	[
+		SNew(SWidgetSwitcher)
+		.WidgetIndex_Lambda([this]()
+		{
+			// If we have a skeleton, show the picker widget, otherwise default to a text label indicating that
+			// none is available.
+			return GetSkeleton() != nullptr ? 0 : 1;
+		})
+		+ SWidgetSwitcher::Slot()
+		[
+			SNew(SBoneSelectionWidget)
+			.ToolTipText(StructPropertyHandle->GetToolTipText())
+			.bShowSocket(true)
+			.OnBoneSelectionChanged(this, &FBoneSocketTargetCustomization::OnBoneSelectionChanged)
+			.OnGetSelectedBone(this, &FBoneSocketTargetCustomization::GetSelectedBone)
+			.OnGetReferenceSkeleton(this, &FBoneReferenceCustomization::GetReferenceSkeleton)
+			.OnGetSocketList(this, &FBoneSocketTargetCustomization::GetSocketList)
+		]
+		+ SWidgetSwitcher::Slot()
+		[
+			SNew(STextBlock)
+			.Font(StructCustomizationUtils.GetRegularFont())
+			.Text(LOCTEXT("BoneSelectionNoSkeleton", "No Skeleton Available"))
+		]
+	];
 }
 
-void FBoneSocketTargetCustomization::SetPropertyHandle(TSharedRef<IPropertyHandle> StructPropertyHandle)
+void FBoneSocketTargetCustomization::ResolveChildProperties()
 {
-	TSharedPtr<IPropertyHandle> BoneReferenceProperty = FindStructMemberProperty(StructPropertyHandle, GET_MEMBER_NAME_CHECKED(FBoneSocketTarget, BoneReference));
+	const TSharedPtr<IPropertyHandle> BoneReferenceProperty = FindStructMemberProperty(StructProperty, GET_MEMBER_NAME_CHECKED(FBoneSocketTarget, BoneReference));
 	check(BoneReferenceProperty->IsValidHandle());
 	BoneNameProperty = FindStructMemberProperty(BoneReferenceProperty.ToSharedRef(), GET_MEMBER_NAME_CHECKED(FBoneReference, BoneName));
-	TSharedPtr<IPropertyHandle> SocketReferenceProperty = FindStructMemberProperty(StructPropertyHandle, GET_MEMBER_NAME_CHECKED(FBoneSocketTarget, SocketReference));
+	
+	const TSharedPtr<IPropertyHandle> SocketReferenceProperty = FindStructMemberProperty(StructProperty, GET_MEMBER_NAME_CHECKED(FBoneSocketTarget, SocketReference));
 	check(SocketReferenceProperty->IsValidHandle());
 	SocketNameProperty = FindStructMemberProperty(SocketReferenceProperty.ToSharedRef(), GET_MEMBER_NAME_CHECKED(FSocketReference, SocketName));
-	UseSocketProperty = FindStructMemberProperty(StructPropertyHandle, GET_MEMBER_NAME_CHECKED(FBoneSocketTarget, bUseSocket));
+	
+	UseSocketProperty = FindStructMemberProperty(StructProperty, GET_MEMBER_NAME_CHECKED(FBoneSocketTarget, bUseSocket));
 
 	check(BoneNameProperty->IsValidHandle() && SocketNameProperty->IsValidHandle() && UseSocketProperty->IsValidHandle());
 }
 
-BEGIN_SLATE_FUNCTION_BUILD_OPTIMIZATION
-void FBoneSocketTargetCustomization::Build(TSharedRef<IPropertyHandle> StructPropertyHandle, class IDetailChildrenBuilder& ChildBuilder)
-{
-	if (TargetEditableSkeleton.IsValid() && BoneNameProperty->IsValidHandle())
-	{
-		ChildBuilder
-		.AddProperty(StructPropertyHandle)
-		.CustomWidget()
-		.NameContent()
-		[
-			StructPropertyHandle->CreatePropertyNameWidget()
-		]
-		.ValueContent()
-		[
-			SNew(SBoneSelectionWidget)
-			.ToolTipText(StructPropertyHandle->GetToolTipText())
-			.bShowSocket(true)
-			.OnBoneSelectionChanged(this, &FBoneSocketTargetCustomization::OnBoneSelectionChanged)
-			.OnGetSelectedBone(this, &FBoneSocketTargetCustomization::GetSelectedBone)
-			.OnGetReferenceSkeleton(this, &FBoneReferenceCustomization::GetReferenceSkeleton)
-			.OnGetSocketList(this, &FBoneSocketTargetCustomization::GetSocketList)
-		];
-	}
-	else
-	{
-		// if this FBoneSocketTarget is used by some other Outers, this will fail	
-		// should warn programmers instead of silent fail
-		ensureAlways(false);
-	}
-}
-END_SLATE_FUNCTION_BUILD_OPTIMIZATION
 
 TSharedPtr<IPropertyHandle> FBoneSocketTargetCustomization::GetNameProperty() const
 {
@@ -782,34 +761,31 @@ TSharedPtr<IPropertyHandle> FBoneSocketTargetCustomization::GetNameProperty() co
 void FBoneSocketTargetCustomization::OnBoneSelectionChanged(FName Name)
 {
 	// figure out if the name is BoneName or socket name
-	if (TargetEditableSkeleton.IsValid())
+	bool bUseSocket = false;
+	if (GetReferenceSkeleton().FindBoneIndex(Name) == INDEX_NONE)
 	{
-		bool bUseSocket = false;
-		if (GetReferenceSkeleton().FindBoneIndex(Name) == INDEX_NONE)
+		// make sure socket exists
+		const TArray<class USkeletalMeshSocket*>& Sockets = GetSocketList();
+		for (int32 Idx = 0; Idx < Sockets.Num(); ++Idx)
 		{
-			// make sure socket exists
-			const TArray<class USkeletalMeshSocket*>& Sockets = GetSocketList();
-			for (int32 Idx = 0; Idx < Sockets.Num(); ++Idx)
+			if (Sockets[Idx]->SocketName == Name)
 			{
-				if (Sockets[Idx]->SocketName == Name)
-				{
-					bUseSocket = true;
-					break;
-				}
+				bUseSocket = true;
+				break;
 			}
-
-			// we should find one
-			ensure(bUseSocket);
 		}
 
-		// set correct value
-		UseSocketProperty->SetValue(bUseSocket);
+		// we should find one
+		ensure(bUseSocket);
+	}
+
+	// set correct value
+	UseSocketProperty->SetValue(bUseSocket);
 
-		TSharedPtr<IPropertyHandle> NameProperty = GetNameProperty();
-		if (ensureAlways(NameProperty.IsValid()))
-		{
-			NameProperty->SetValue(Name);
-		}
+	const TSharedPtr<IPropertyHandle> NameProperty = GetNameProperty();
+	if (ensureAlways(NameProperty.IsValid()))
+	{
+		NameProperty->SetValue(Name);
 	}
 }
 
@@ -817,10 +793,10 @@ FName FBoneSocketTargetCustomization::GetSelectedBone(bool& bMultipleValues) con
 {
 	FString OutText;
 
-	TSharedPtr<IPropertyHandle> NameProperty = GetNameProperty();
+	const TSharedPtr<IPropertyHandle> NameProperty = GetNameProperty();
 	if (NameProperty.IsValid())
 	{
-		FPropertyAccess::Result Result = NameProperty->GetValueAsFormattedString(OutText);
+		const FPropertyAccess::Result Result = NameProperty->GetValueAsFormattedString(OutText);
 		bMultipleValues = (Result == FPropertyAccess::MultipleValues);
 	}
 	else
@@ -833,15 +809,15 @@ FName FBoneSocketTargetCustomization::GetSelectedBone(bool& bMultipleValues) con
 	return FName(*OutText);
 }
 
-const TArray<class USkeletalMeshSocket*>& FBoneSocketTargetCustomization::GetSocketList() const
+const TArray<USkeletalMeshSocket*>& FBoneSocketTargetCustomization::GetSocketList() const
 {
-	if (TargetEditableSkeleton.IsValid())
+	static TArray<USkeletalMeshSocket*> EmptySocketList;
+	
+	if (const USkeleton* Skeleton = GetSkeleton())
 	{
-		return  TargetEditableSkeleton.Get()->GetSkeleton().Sockets;
+		return Skeleton->Sockets;
 	}
-
-	static TArray<class USkeletalMeshSocket*> DummyList;
-	return DummyList;
+	return EmptySocketList;
 }
 
 /////////////////////////////////////////////////////////////////////////////////////////////
diff --git a/Engine/Source/Editor/Persona/Private/AnimGraphNodeDetails.h b/Engine/Source/Editor/Persona/Private/AnimGraphNodeDetails.h
--- a/Engine/Source/Editor/Persona/Private/AnimGraphNodeDetails.h
+++ b/Engine/Source/Editor/Persona/Private/AnimGraphNodeDetails.h
@@ -108,20 +108,19 @@ public:
 
 	// IPropertyTypeCustomization interface
 	virtual void CustomizeHeader(TSharedRef<IPropertyHandle> StructPropertyHandle, class FDetailWidgetRow& HeaderRow, IPropertyTypeCustomizationUtils& StructCustomizationUtils) override;
-	virtual void CustomizeChildren(TSharedRef<IPropertyHandle> StructPropertyHandle, class IDetailChildrenBuilder& ChildBuilder, IPropertyTypeCustomizationUtils& StructCustomizationUtils) override;
+	virtual void CustomizeChildren(TSharedRef<IPropertyHandle> StructPropertyHandle, class IDetailChildrenBuilder& ChildBuilder, IPropertyTypeCustomizationUtils& StructCustomizationUtils) override {}
 	const struct FReferenceSkeleton&  GetReferenceSkeleton() const;
 
 protected:
-	void SetEditableSkeleton(TSharedRef<IPropertyHandle> StructPropertyHandle);
-	virtual void SetPropertyHandle(TSharedRef<IPropertyHandle> StructPropertyHandle);
-	TSharedPtr<IPropertyHandle> FindStructMemberProperty(TSharedRef<IPropertyHandle> PropertyHandle, const FName& PropertyName);
+	USkeleton* GetSkeleton() const;
+	
+	static TSharedPtr<IPropertyHandle> FindStructMemberProperty(TSharedPtr<IPropertyHandle> PropertyHandle, const FName& PropertyName);
+	
+	TSharedPtr<IPropertyHandle> StructProperty;
+	
 	// Property to change after bone has been picked
 	TSharedPtr<IPropertyHandle> BoneNameProperty;
 
-	// Target Skeleton this widget is referencing
-	TSharedPtr<IEditableSkeleton> TargetEditableSkeleton;
-
-	bool bEnsureOnInvalidSkeleton = true;
 private:
 
 	// Bone tree widget delegates
@@ -137,16 +136,15 @@ class FBoneSocketTargetCustomization : public FBoneReferenceCustomization
 public:
 	static TSharedRef<IPropertyTypeCustomization> MakeInstance();
 
-	virtual void CustomizeHeader(TSharedRef<IPropertyHandle> StructPropertyHandle, class FDetailWidgetRow& HeaderRow, IPropertyTypeCustomizationUtils& StructCustomizationUtils) override {};
-	virtual void CustomizeChildren(TSharedRef<IPropertyHandle> StructPropertyHandle, class IDetailChildrenBuilder& ChildBuilder, IPropertyTypeCustomizationUtils& StructCustomizationUtils) override;
+	virtual void CustomizeHeader(TSharedRef<IPropertyHandle> StructPropertyHandle, class FDetailWidgetRow& HeaderRow, IPropertyTypeCustomizationUtils& StructCustomizationUtils) override;
 
 private:
+	void ResolveChildProperties();	
+	
 	// Property to change after bone has been picked
 	TSharedPtr<IPropertyHandle> SocketNameProperty;
 	TSharedPtr<IPropertyHandle> UseSocketProperty;
 
-	virtual void SetPropertyHandle(TSharedRef<IPropertyHandle> StructPropertyHandle) override;
-	void Build(TSharedRef<IPropertyHandle> StructPropertyHandle, class IDetailChildrenBuilder& ChildBuilder);
 	// Bone tree widget delegates
 	virtual void OnBoneSelectionChanged(FName Name) override;
 	virtual FName GetSelectedBone(bool& bMultipleValues) const override;
