Provide alternative orientation calculation method for blend space analysis - avoids problems when the pointing bone is not restricted to an "upright" orientation

diff --git a/Engine/Source/Editor/Persona/Public/BlendSpaceAnalysis.h b/Engine/Source/Editor/Persona/Public/BlendSpaceAnalysis.h
--- a/Engine/Source/Editor/Persona/Public/BlendSpaceAnalysis.h
+++ b/Engine/Source/Editor/Persona/Public/BlendSpaceAnalysis.h
@@ -62,6 +62,13 @@ enum class EAnalysisLinearAxis : uint8
 	MinusZ UMETA(DisplayName = "-Z", ToolTip = "The axis points in the negative Z direction"),
 };
 
+UENUM()
+enum class EEulerCalculationMethod : uint8
+{
+	AimDirection    UMETA(ToolTip = "Calculates the yaw by looking at the BoneRightAxis. This can provide better yaw values, especially when aiming (e.g. with a weapon that has minimal rotation around its pointing axis) and covering extreme angles up and down, but only if this rightwards facing axis is reliable. It won't work well if the bone is also rolling around its axis."),
+	PointDirection  UMETA(ToolTip = "Calculates the yaw based only on the BoneFacingAxis. This will work when you're most interested in the yaw and pitch from a pointing direction, but can produce undesirable results when pointing almost directly up or down."),
+};
+
 UENUM()
 enum class EAnalysisEulerAxis : uint8
 {
@@ -158,6 +165,10 @@ public:
 	UPROPERTY(EditAnywhere, Category = AnalysisProperties)
 	EAnalysisLinearAxis BoneRightAxis = EAnalysisLinearAxis::PlusY;
 
+	/** Used for some analysis functions - specifies how yaw should be calculated from the bone axes */
+	UPROPERTY(EditAnywhere, Category = AnalysisProperties)
+	EEulerCalculationMethod EulerCalculationMethod = EEulerCalculationMethod::AimDirection;
+
 	/**
 	* The space in which to perform the analysis. Fixed will use the analysis bone/socket at the first frame
 	* of the analysis time range. Changing will use the analysis bone/socket at the relevant frame during the
@@ -469,10 +480,22 @@ void CalculateBoneOrientation(
 	const FVector AimForwardDir = BlendSpaceAnalysis::GetAxisFromTM(TM, AnalysisProperties->BoneFacingAxis);
 	const FVector AimRightDir = BlendSpaceAnalysis::GetAxisFromTM(TM, AnalysisProperties->BoneRightAxis);
 
-	// Note that Yaw is best taken from the AimRightDir - this is to avoid problems when the gun is pointing
-	// up or down - especially if it goes beyond 90 degrees in pitch.
-	const double Yaw = FMath::RadiansToDegrees(FMath::Atan2(
-		FVector::DotProduct(AimRightDir, -FrameFacingDir), FVector::DotProduct(AimRightDir, FrameRightDir)));
+	double Yaw;
+	if (AnalysisProperties->EulerCalculationMethod == EEulerCalculationMethod::AimDirection)
+	{
+		// Yaw is taken from the AimRightDir to avoid problems when the gun is pointing up or down - especially if it 
+		// goes beyond 90 degrees in pitch. However, if there is roll around the gun axis, then this can produce
+		// incorrect/undesirable results.
+		Yaw = FMath::RadiansToDegrees(FMath::Atan2(
+			FVector::DotProduct(AimRightDir, -FrameFacingDir), FVector::DotProduct(AimRightDir, FrameRightDir)));
+	}
+	else
+	{
+		// This takes yaw directly from the forwards direction. Note that if the pose is really one with small yaw 
+		// and pitch more than 90 degrees, then this will calculate a yaw that is nearer to 180 degrees.
+		Yaw = FMath::RadiansToDegrees(FMath::Atan2(
+			FVector::DotProduct(AimForwardDir, FrameRightDir), FVector::DotProduct(AimForwardDir, FrameFacingDir)));
+	}
 
 	// Undo the yaw to get pitch
 	const FQuat YawQuat(FrameUpDir, FMath::DegreesToRadians(Yaw));
