Anim notify and sync marker workflow changes

Skeletons are no longer dirtied when adding notifies and sync markers from animation timelines.
Skeletons can still hold sync markers and notifies, but they are not mandatory and they can all be held in animation assets if required.
'Find References' for notifies and sync markers now launches the find/replace dialog when selected from the animation timeline.
Removing notifies and sync markers from the skeleton now does only that. To remove from animations, use the find/replace tab.
Updated BP node spawning for anim notifies to use a custom spawner, correctly use the asset registry, apply filtering and context sensitivity and move some of the code out of BlueprintGraph.
Added icons to distinguish sync markers/notifies in the editing tab.
Updated notify editing tab filtering to allow users to filter items from 'this' skeleton, compatible skeletons or other assets.

Additionally fixed up icons for some Persona tabs that were incorrect.

diff --git a/Engine/Source/Editor/AnimationBlueprintEditor/Private/AnimationBlueprintEditorModule.cpp b/Engine/Source/Editor/AnimationBlueprintEditor/Private/AnimationBlueprintEditorModule.cpp
--- a/Engine/Source/Editor/AnimationBlueprintEditor/Private/AnimationBlueprintEditorModule.cpp
+++ b/Engine/Source/Editor/AnimationBlueprintEditor/Private/AnimationBlueprintEditorModule.cpp
@@ -28,6 +28,15 @@
 #include "UObject/ObjectPtr.h"
 #include "UObject/UObjectGlobals.h"
 #include "UObject/WeakObjectPtr.h"
+#include "BlueprintActionDatabaseRegistrar.h"
+#include "AssetRegistry/AssetRegistryModule.h"
+#include "String/ParseTokens.h"
+#include "K2Node_Event.h"
+#include "AnimNotifyEventNodeSpawner.h"
+#include "BlueprintActionDatabase.h"
+#include "Animation/Skeleton.h"
+#include "Animation/AnimBlueprint.h"
+#include "Animation/AnimBlueprintGeneratedClass.h"
 
 class IAnimationBlueprintEditor;
 class IToolkitHost;
@@ -88,6 +97,88 @@ TSharedRef<IAnimationBlueprintEditor> FAnimationBlueprintEditorModule::CreateAni
 	return NewAnimationBlueprintEditor;
 }
 
+void FAnimationBlueprintEditorModule::GetTypeActions(FBlueprintActionDatabaseRegistrar& ActionRegistrar) const
+{
+	if (!ActionRegistrar.IsOpenForRegistration(UAnimBlueprint::StaticClass()))
+	{
+		return;
+	}
+
+	// Grab notifies from skeletons and anim sequences
+	FAssetRegistryModule& AssetRegistryModule = FModuleManager::LoadModuleChecked<FAssetRegistryModule>(TEXT("AssetRegistry"));
+	FARFilter Filter;
+	Filter.bRecursiveClasses = true;
+
+	if (FBlueprintActionDatabase::IsClassAllowed(UAnimSequenceBase::StaticClass(), FBlueprintActionDatabase::EPermissionsContext::Asset))
+	{
+		Filter.ClassPaths.Add(UAnimSequenceBase::StaticClass()->GetClassPathName());
+	}
+
+	if (FBlueprintActionDatabase::IsClassAllowed(USkeleton::StaticClass(), FBlueprintActionDatabase::EPermissionsContext::Asset))
+	{
+		Filter.ClassPaths.Add(USkeleton::StaticClass()->GetClassPathName());
+	}
+
+	if (Filter.ClassPaths.Num() == 0)
+	{
+		return;
+	}
+
+	TArray<FAssetData> FoundAssetData;
+	AssetRegistryModule.Get().GetAssets(Filter, FoundAssetData);
+
+	TMap<FSoftObjectPath, TSet<FName>> NotifiesPerSkeleton;
+	for (const FAssetData& AssetData : FoundAssetData)
+	{
+		const FString TagValue = AssetData.GetTagValueRef<FString>(USkeleton::AnimNotifyTag);
+		if (!TagValue.IsEmpty())
+		{
+			FSoftObjectPath SkeletonPath(AssetData.GetTagValueRef<FString>("Skeleton"));
+			TSet<FName>& NotifyNames = NotifiesPerSkeleton.FindOrAdd(SkeletonPath);
+
+			UE::String::ParseTokens(TagValue, USkeleton::AnimNotifyTagDelimiter, [&NotifyNames](FStringView InToken)
+			{
+				FName NotifyName(InToken);
+				if(NotifyName != NAME_None)
+				{
+					NotifyNames.Add(NotifyName);
+				}
+			}, UE::String::EParseTokensOptions::SkipEmpty);
+		}
+	}
+
+	for (const TPair<FSoftObjectPath, TSet<FName>>& SkeletonPathNamesPair : NotifiesPerSkeleton)
+	{
+		for (FName NotifyName : SkeletonPathNamesPair.Value)
+		{
+			UAnimNotifyEventNodeSpawner* NodeSpawner = UAnimNotifyEventNodeSpawner::Create(SkeletonPathNamesPair.Key, NotifyName);
+			ActionRegistrar.AddBlueprintAction(UAnimBlueprint::StaticClass(), NodeSpawner);
+		}
+	}
+}
+
+void FAnimationBlueprintEditorModule::GetInstanceActions(const UAnimBlueprint* InAnimBlueprint, FBlueprintActionDatabaseRegistrar& ActionRegistrar) const
+{
+	if (UAnimBlueprintGeneratedClass* GeneratedClass = InAnimBlueprint->GetAnimBlueprintGeneratedClass())
+	{
+		FSoftObjectPath SkeletonPath;
+		if (InAnimBlueprint->TargetSkeleton)
+		{
+			SkeletonPath = FSoftObjectPath(InAnimBlueprint->TargetSkeleton);
+		}
+
+		for (int32 NotifyIdx = 0; NotifyIdx < GeneratedClass->GetAnimNotifies().Num(); NotifyIdx++)
+		{
+			FName NotifyName = GeneratedClass->GetAnimNotifies()[NotifyIdx].NotifyName;
+			if (NotifyName != NAME_None)
+			{
+				UAnimNotifyEventNodeSpawner* NodeSpawner = UAnimNotifyEventNodeSpawner::Create(SkeletonPath, NotifyName);
+				ActionRegistrar.AddBlueprintAction(InAnimBlueprint, NodeSpawner);
+			}
+		}
+	}
+}
+
 void FAnimationBlueprintEditorModule::OnNewBlueprintCreated(UBlueprint* InBlueprint)
 {
 	if (InBlueprint->UbergraphPages.Num() > 0)
diff --git a/Engine/Source/Editor/AnimationBlueprintEditor/Private/AnimationBlueprintEditorModule.h b/Engine/Source/Editor/AnimationBlueprintEditor/Private/AnimationBlueprintEditorModule.h
--- a/Engine/Source/Editor/AnimationBlueprintEditor/Private/AnimationBlueprintEditorModule.h
+++ b/Engine/Source/Editor/AnimationBlueprintEditor/Private/AnimationBlueprintEditorModule.h
@@ -31,6 +31,8 @@ public:
 	/** IAnimationBlueprintEditorModule interface */
 	virtual TSharedRef<class IAnimationBlueprintEditor> CreateAnimationBlueprintEditor(const EToolkitMode::Type Mode, const TSharedPtr< class IToolkitHost >& InitToolkitHost, class UAnimBlueprint* Blueprint) override;
 	virtual TArray<FAnimationBlueprintEditorToolbarExtender>& GetAllAnimationBlueprintEditorToolbarExtenders() { return AnimationBlueprintEditorToolbarExtenders; }
+	virtual void GetTypeActions(FBlueprintActionDatabaseRegistrar& ActionRegistrar) const override;
+	virtual void GetInstanceActions(const UAnimBlueprint* InAnimBlueprint, FBlueprintActionDatabaseRegistrar& ActionRegistrar) const override;
 
 	/** Gets the extensibility managers for outside entities to extend this editor's menus and toolbars */
 	virtual TSharedPtr<FExtensibilityManager> GetMenuExtensibilityManager() override { return MenuExtensibilityManager; }
diff --git a/Engine/Source/Editor/AnimationBlueprintEditor/Public/IAnimationBlueprintEditorModule.h b/Engine/Source/Editor/AnimationBlueprintEditor/Public/IAnimationBlueprintEditorModule.h
--- a/Engine/Source/Editor/AnimationBlueprintEditor/Public/IAnimationBlueprintEditorModule.h
+++ b/Engine/Source/Editor/AnimationBlueprintEditor/Public/IAnimationBlueprintEditorModule.h
@@ -12,6 +12,8 @@
 
 DECLARE_LOG_CATEGORY_EXTERN(LogAnimationBlueprintEditor, Log, All);
 
+class FBlueprintActionDatabaseRegistrar;
+
 class IAnimationBlueprintEditorModule : public IModuleInterface, public IHasMenuExtensibility, public IHasToolBarExtensibility
 {
 public:
@@ -37,4 +39,10 @@ public:
 
 	DECLARE_MULTICAST_DELEGATE_TwoParams(FOnGetCustomDebugObjects, const IAnimationBlueprintEditor& /*AnimationBlueprintEditor*/, TArray<FCustomDebugObject>& /*DebugList*/);
 	virtual FOnGetCustomDebugObjects& OnGetCustomDebugObjects() = 0;
+
+	/** Get anim blueprint type actions */
+	virtual void GetTypeActions(FBlueprintActionDatabaseRegistrar& ActionRegistrar) const = 0;
+
+	/** Get anim blueprint instance actions */
+	virtual void GetInstanceActions(const UAnimBlueprint* InAnimBlueprint, FBlueprintActionDatabaseRegistrar& ActionRegistrar) const = 0;
 };
diff --git a/Engine/Source/Editor/BlueprintGraph/Classes/BlueprintEventNodeSpawner.h b/Engine/Source/Editor/BlueprintGraph/Classes/BlueprintEventNodeSpawner.h
--- a/Engine/Source/Editor/BlueprintGraph/Classes/BlueprintEventNodeSpawner.h
+++ b/Engine/Source/Editor/BlueprintGraph/Classes/BlueprintEventNodeSpawner.h
@@ -87,7 +87,7 @@ public:
 	 */
 	virtual UK2Node_Event const* FindPreExistingEvent(UBlueprint* Blueprint, FBindingSet const& Bindings) const;
 
-private:
+protected:
 	/** The function to configure new nodes with. */
 	UPROPERTY()
     TObjectPtr<UFunction const> EventFunc;
diff --git a/Engine/Source/Editor/BlueprintGraph/Private/AnimNotifyEventNodeSpawner.cpp b/Engine/Source/Editor/BlueprintGraph/Private/AnimNotifyEventNodeSpawner.cpp
new file mode 100644
--- /dev/null
+++ b/Engine/Source/Editor/BlueprintGraph/Private/AnimNotifyEventNodeSpawner.cpp
@@ -0,0 +1,38 @@
+// Copyright Epic Games, Inc. All Rights Reserved.
+
+#include "AnimNotifyEventNodeSpawner.h"
+#include "EditorCategoryUtils.h"
+#include "Styling/AppStyle.h"
+#include "Animation/AnimInstance.h"
+
+#define LOCTEXT_NAMESPACE "AnimNotifyEventNodeSpawner"
+
+UAnimNotifyEventNodeSpawner* UAnimNotifyEventNodeSpawner::Create(const FSoftObjectPath& InSkeletonObjectPath, FName InNotifyName)
+{
+	check(InNotifyName != NAME_None);
+
+	FString Label = InNotifyName.ToString();
+	FString CustomEventName = FString::Printf(TEXT("AnimNotify_%s"), *Label);
+
+	UAnimNotifyEventNodeSpawner* NodeSpawner = NewObject<UAnimNotifyEventNodeSpawner>(GetTransientPackage());
+	NodeSpawner->SkeletonObjectPath = InSkeletonObjectPath;
+	NodeSpawner->NodeClass = UK2Node_Event::StaticClass();
+	NodeSpawner->CustomEventName = *CustomEventName;
+
+	FBlueprintActionUiSpec& MenuSignature = NodeSpawner->DefaultMenuSignature;
+	MenuSignature.MenuName = FText::Format(LOCTEXT("EventWithSignatureName", "Event {0}"), FText::FromName(NodeSpawner->CustomEventName));
+	MenuSignature.Icon = FSlateIcon(FAppStyle::GetAppStyleSetName(), "GraphEditor.Event_16x");
+
+	auto PostSpawnSetupLambda = [](UEdGraphNode* NewNode, bool /*bIsTemplateNode*/)
+	{
+		UK2Node_Event* ActorRefNode = CastChecked<UK2Node_Event>(NewNode);
+		ActorRefNode->EventReference.SetExternalMember(ActorRefNode->CustomFunctionName, UAnimInstance::StaticClass());
+	};
+
+	NodeSpawner->CustomizeNodeDelegate = UBlueprintNodeSpawner::FCustomizeNodeDelegate::CreateLambda(PostSpawnSetupLambda);
+	NodeSpawner->DefaultMenuSignature.Category = FEditorCategoryUtils::GetCommonCategory(FCommonEditorCategory::AnimNotify);
+
+	return NodeSpawner;
+}
+
+#undef LOCTEXT_NAMESPACE
\ No newline at end of file
diff --git a/Engine/Source/Editor/BlueprintGraph/Private/BlueprintActionDatabase.cpp b/Engine/Source/Editor/BlueprintGraph/Private/BlueprintActionDatabase.cpp
--- a/Engine/Source/Editor/BlueprintGraph/Private/BlueprintActionDatabase.cpp
+++ b/Engine/Source/Editor/BlueprintGraph/Private/BlueprintActionDatabase.cpp
@@ -2,11 +2,6 @@
 
 #include "BlueprintActionDatabase.h"
 
-#include "Animation/AnimBlueprint.h"
-#include "Animation/AnimBlueprintGeneratedClass.h"
-#include "Animation/AnimInstance.h"
-#include "Animation/AnimTypes.h"
-#include "Animation/Skeleton.h"
 #include "AssetRegistry/AssetData.h"
 #include "AssetRegistry/AssetRegistryModule.h"
 #include "AssetRegistry/IAssetRegistry.h"
@@ -176,15 +171,6 @@ namespace FBlueprintNodeSpawnerFactory
 	 * @return 
 	 */
 	static UBlueprintNodeSpawner* MakeActorBoundEventSpawner(FMulticastDelegateProperty* DelegateProperty);
-
-	/**
-	 * Constructs UK2Node_Event spawner that is owned by UAnimInstance. Used for Anim Notificatios and montage 
-	 * branching points.
-	 *
-	 * @param EventName
-	 * @return A new node-spawner, set up to spawn UK2Node_Event
-	 */
-	static UBlueprintNodeSpawner* MakeAnimOwnedEventSpawner(FName SignatureName, FText CustomCategory);
 };
 
 //------------------------------------------------------------------------------
@@ -305,22 +291,6 @@ static UBlueprintNodeSpawner* FBlueprintNodeSpawnerFactory::MakeActorBoundEventS
 	return UBlueprintBoundEventNodeSpawner::Create(UK2Node_ActorBoundEvent::StaticClass(), DelegateProperty);
 }
 
-//------------------------------------------------------------------------------
-static UBlueprintNodeSpawner* FBlueprintNodeSpawnerFactory::MakeAnimOwnedEventSpawner(FName SignatureName, FText CustomCategory)
-{
-	auto PostSpawnSetupLambda = [](UEdGraphNode* NewNode, bool /*bIsTemplateNode*/)
-	{
-		UK2Node_Event* ActorRefNode = CastChecked<UK2Node_Event>(NewNode);
-		ActorRefNode->EventReference.SetExternalMember(ActorRefNode->CustomFunctionName, UAnimInstance::StaticClass());
-	};
-
-	UBlueprintNodeSpawner* NodeSpawner = UBlueprintEventNodeSpawner::Create(UK2Node_Event::StaticClass(), SignatureName);
-	NodeSpawner->CustomizeNodeDelegate = UBlueprintNodeSpawner::FCustomizeNodeDelegate::CreateStatic(PostSpawnSetupLambda);
-	NodeSpawner->DefaultMenuSignature.Category = CustomCategory;
-
-	return NodeSpawner;
-}
-
 /*******************************************************************************
  * Static FBlueprintActionDatabase Helpers
  ******************************************************************************/
@@ -412,14 +382,6 @@ namespace BlueprintActionDatabaseImpl
 	 */
 	static void AddClassCastActions(UClass* const Class, FActionList& ActionListOut);
 
-	/**
-	 * Adds custom actions to operate on the provided skeleton. Used primarily
-	 * to find AnimNotify event vocabulary
-	 *
-	 * @Param Skeleton	The skeleton that may have anim notifies defined on it.
-	 */
-	static void AddSkeletonActions( const USkeleton& Skeleton, FActionList& ActionListOut);
-
 	/**
 	 * If the associated class is a blueprint generated class, then this will
 	 * loop over the blueprint's graphs and create any node-spawners associated
@@ -429,16 +391,6 @@ namespace BlueprintActionDatabaseImpl
 	 * @param  ActionListOut	The list you want populated with new spawners.
 	 */
 	static void AddBlueprintGraphActions(UBlueprint const* const Blueprint, FActionList& ActionListOut);
-	
-	/**
-	 * If the associated class is an anim blueprint generated class, then this
-	 * will loop over AnimNotification events in the anim blueprint generated
-	 * class and create node spawners for those events.
-	 *
-	 * @param  Blueprint		The blueprint which you want graph associated node-spawners for.
-	 * @param  ActionListOut	The list you want populated with new spawners.
-	 */
-	static void AddAnimBlueprintGraphActions( UAnimBlueprint const* AnimBlueprint, FActionList& ActionListOut );
 
 	/**
 	 * Emulates UEdGraphSchema::GetGraphContextActions(). If the supplied class  
@@ -881,21 +833,6 @@ static void BlueprintActionDatabaseImpl::AddClassCastActions(UClass* Class, FAct
 	}
 }
 
-//------------------------------------------------------------------------------
-static void BlueprintActionDatabaseImpl::AddSkeletonActions(const USkeleton& Skeleton, FActionList& ActionListOut)
-{
-	TArray<FName> NotifyNames;
-	Skeleton.CollectAnimationNotifies(NotifyNames);
-
-	for (const FName& NotifyName : NotifyNames)
-	{
-		FString Label = NotifyName.ToString();
-
-		FString SignatureName = FString::Printf(TEXT("AnimNotify_%s"), *Label);
-		ActionListOut.Add(FBlueprintNodeSpawnerFactory::MakeAnimOwnedEventSpawner(FName(*SignatureName), FEditorCategoryUtils::GetCommonCategory(FCommonEditorCategory::AnimNotify)));
-	}
-}
-
 //------------------------------------------------------------------------------
 static void BlueprintActionDatabaseImpl::AddBlueprintGraphActions(UBlueprint const* const Blueprint, FActionList& ActionListOut)
 {
@@ -968,24 +905,6 @@ static void BlueprintActionDatabaseImpl::AddBlueprintGraphActions(UBlueprint con
 	}
 }
 
-//------------------------------------------------------------------------------
-static void BlueprintActionDatabaseImpl::AddAnimBlueprintGraphActions(UAnimBlueprint const* AnimBlueprint, FActionList& ActionListOut)
-{
-	if (UAnimBlueprintGeneratedClass* GeneratedClass = AnimBlueprint->GetAnimBlueprintGeneratedClass())
-	{
-		for (int32 NotifyIdx = 0; NotifyIdx < GeneratedClass->GetAnimNotifies().Num(); NotifyIdx++)
-		{
-			FName NotifyName = GeneratedClass->GetAnimNotifies()[NotifyIdx].NotifyName;
-			if (NotifyName != NAME_None)
-			{
-				FString Label = NotifyName.ToString();
-				FString SignatureName = FString::Printf(TEXT("AnimNotify_%s"), *Label);
-				ActionListOut.Add(FBlueprintNodeSpawnerFactory::MakeAnimOwnedEventSpawner(FName(*SignatureName), FEditorCategoryUtils::GetCommonCategory(FCommonEditorCategory::AnimNotify)));
-			}
-		}
-	}
-}
-
 //------------------------------------------------------------------------------
 static void BlueprintActionDatabaseImpl::GetNodeSpecificActions(TSubclassOf<UEdGraphNode const> const NodeClass, FBlueprintActionDatabaseRegistrar& Registrar)
 {
@@ -1534,16 +1453,6 @@ void FBlueprintActionDatabase::RefreshAll()
 		RefreshClassActions(Class);
 	}
 
-	if(IsClassAllowed(USkeleton::StaticClass(), EPermissionsContext::Asset))
-	{
-		// this handles creating entries for skeletons that were loaded before the database was alive:
-		for( TObjectIterator<USkeleton> SkeletonIt; SkeletonIt; ++SkeletonIt )
-		{
-			FActionList& ClassActionList = ActionRegistry.FindOrAdd(*SkeletonIt);
-			BlueprintActionDatabaseImpl::AddSkeletonActions(**SkeletonIt, ClassActionList);
-		}
-	}
-
 	FComponentTypeRegistry::Get().SubscribeToComponentList(ComponentTypes).RemoveAll(this);
 
 	// this handles creating entries for components that were loaded before the database was alive:
@@ -1746,11 +1655,6 @@ void FBlueprintActionDatabase::RefreshAssetActions(UObject* const AssetObject)
 		return;
 	}
 
-	if(const USkeleton* Skeleton = Cast<USkeleton>(AssetObject))
-	{
-		AddSkeletonActions(*Skeleton, AssetActionList);
-	}
-
 	UBlueprint* BlueprintAsset = Cast<UBlueprint>(AssetObject);
 	if (BlueprintAsset != nullptr)
 	{
@@ -1760,11 +1664,6 @@ void FBlueprintActionDatabase::RefreshAssetActions(UObject* const AssetObject)
 			GetClassMemberActions(SkeletonClass, AssetActionList);
 		}
 
-		if( const UAnimBlueprint* AnimBlueprint = Cast<UAnimBlueprint>(BlueprintAsset) )
-		{
-			AddAnimBlueprintGraphActions( AnimBlueprint, AssetActionList );
-		}
-
 		FBlueprintActionDatabaseRegistrar Registrar(ActionRegistry, UnloadedActionRegistry, ActionPrimingQueue);
 		if (!bIsInitializing)
 		{
diff --git a/Engine/Source/Editor/BlueprintGraph/Private/BlueprintActionFilter.cpp b/Engine/Source/Editor/BlueprintGraph/Private/BlueprintActionFilter.cpp
--- a/Engine/Source/Editor/BlueprintGraph/Private/BlueprintActionFilter.cpp
+++ b/Engine/Source/Editor/BlueprintGraph/Private/BlueprintActionFilter.cpp
@@ -41,6 +41,7 @@
 #include "Kismet2/KismetEditorUtilities.h"
 #include "Editor/EditorEngine.h"
 #include "BlueprintEditorSettings.h"
+#include "AnimNotifyEventNodeSpawner.h"
 
 #if ENABLE_BLUEPRINT_ACTION_FILTER_PROFILING
 #include "HAL/PlatformStackWalk.h"
@@ -1789,19 +1790,21 @@ static bool BlueprintActionFilterImpl::IsIncompatibleAnimNotification(FBlueprint
 
 	if ( BlueprintAction.GetNodeClass()->IsChildOf<UK2Node_Event>() )
 	{
-		if( const USkeleton* SkeletonOwningEvent = Cast<USkeleton>(BlueprintAction.GetActionOwner()) )
+		if (UAnimNotifyEventNodeSpawner const* const AnimNotifyEventNodeSpawner = Cast<UAnimNotifyEventNodeSpawner>(BlueprintAction.NodeSpawner))
 		{
 			// The event is owned by a skeleton. Only show if it the current anim blueprint is targetting
 			// that skeleton:
+			const FSoftObjectPath& SkeletonObjectPath = AnimNotifyEventNodeSpawner->GetSkeletonObjectPath();
 			FBlueprintActionContext const& FilterContext = Filter.Context;
 			bool bFoundInAllBlueprints = true;
+			const bool bUseSkeletonCheck = Filter.TargetClasses.Num() != 0;	// If we are 'context sensitive' then we will have classes here
 
 			for (const UBlueprint* Blueprint : FilterContext.Blueprints)
 			{
 				bool bFoundInCurrentBlueprint = false;
 				if (const UAnimBlueprint* AnimBlueprint = Cast<UAnimBlueprint>(Blueprint))
 				{
-					if( AnimBlueprint->TargetSkeleton == SkeletonOwningEvent )
+					if( !bUseSkeletonCheck || FSoftObjectPath(AnimBlueprint->TargetSkeleton) == SkeletonObjectPath )
 					{
 						bFoundInCurrentBlueprint = true;
 						break;
diff --git a/Engine/Source/Editor/BlueprintGraph/Public/AnimNotifyEventNodeSpawner.h b/Engine/Source/Editor/BlueprintGraph/Public/AnimNotifyEventNodeSpawner.h
new file mode 100644
--- /dev/null
+++ b/Engine/Source/Editor/BlueprintGraph/Public/AnimNotifyEventNodeSpawner.h
@@ -0,0 +1,30 @@
+// Copyright Epic Games, Inc. All Rights Reserved.
+
+#pragma once
+
+#include "BlueprintEventNodeSpawner.h"
+
+#include "AnimNotifyEventNodeSpawner.generated.h"
+
+UCLASS(Transient)
+class BLUEPRINTGRAPH_API UAnimNotifyEventNodeSpawner : public UBlueprintEventNodeSpawner
+{
+	GENERATED_BODY()
+
+public:
+	/**
+	 * Creates a new UAnimNotifyEventNodeSpawner 
+	 *
+	 * @param  NotifyName	The name you want assigned to the event.
+	 * @return A newly allocated instance of this class.
+	 */
+	static UAnimNotifyEventNodeSpawner* Create(const FSoftObjectPath& InSkeletonObjectPath, FName InNotifyName);
+
+	/** @return the skeleton object path */
+	const FSoftObjectPath& GetSkeletonObjectPath() const { return SkeletonObjectPath; }
+
+private:
+	/** The skeleton that supplied this notify, used for filtering */
+	UPROPERTY()
+	FSoftObjectPath SkeletonObjectPath;
+};
diff --git a/Engine/Source/Editor/EditorStyle/Private/StarshipStyle.cpp b/Engine/Source/Editor/EditorStyle/Private/StarshipStyle.cpp
--- a/Engine/Source/Editor/EditorStyle/Private/StarshipStyle.cpp
+++ b/Engine/Source/Editor/EditorStyle/Private/StarshipStyle.cpp
@@ -4921,6 +4921,8 @@ void FStarshipEditorStyle::FStyle::SetupPersonaStyle()
 
 		// Anim Notify Editor
 		Set("AnimNotifyEditor.BranchingPoint", new IMAGE_BRUSH("Persona/NotifyEditor/BranchingPoints_24x", Icon24x24));
+		Set("AnimNotifyEditor.AnimNotify", new IMAGE_BRUSH_SVG("Starship/Persona/Notify", Icon16x16));
+		Set("AnimNotifyEditor.AnimSyncMarker", new IMAGE_BRUSH_SVG("Starship/Persona/SyncMarker", Icon16x16));
 
 		// AnimBlueprint Preview Warning Background
 		FSlateColor PreviewPropertiesWarningColour(FLinearColor::Gray);
@@ -4929,8 +4931,8 @@ void FStarshipEditorStyle::FStyle::SetupPersonaStyle()
 		// Persona-specific tabs
 		Set("Persona.Tabs.SkeletonTree", new IMAGE_BRUSH(TEXT("Icons/icon_Persona_Skeleton_Tree_16x"), Icon16x16));
 		Set("Persona.Tabs.MorphTargetPreviewer", new IMAGE_BRUSH(TEXT("Icons/icon_Persona_Morph_Target_Previewer_16x"), Icon16x16));
-		Set("Persona.Tabs.AnimCurveDebugger", new IMAGE_BRUSH(TEXT("Icons/icon_Persona_AnimCurve_Previewer_16x"), Icon16x16));
-		Set("Persona.Tabs.AnimCurveMetadataEditor", new IMAGE_BRUSH(TEXT("Icons/icon_Persona_AnimCurve_Previewer_16x"), Icon16x16));
+		Set("Persona.Tabs.AnimCurveDebugger", new IMAGE_BRUSH_SVG(TEXT("Starship/Persona/CurveDebugger"), Icon16x16));
+		Set("Persona.Tabs.AnimCurveMetadataEditor", new IMAGE_BRUSH_SVG(TEXT("Starship/Persona/CurveMetadata"), Icon16x16));
 		Set("Persona.Tabs.AnimationNotifies", new IMAGE_BRUSH(TEXT("Icons/icon_Persona_Animation_Notifies_16x"), Icon16x16));
 		Set("Persona.Tabs.RetargetManager", new IMAGE_BRUSH_SVG("Starship/Persona/RetargetManager", Icon16x16));
 		Set("Persona.Tabs.AnimSlotManager", new IMAGE_BRUSH(TEXT("Icons/icon_Persona_Anim_Slot_Manager_16x"), Icon16x16));
diff --git a/Engine/Source/Editor/Persona/Private/SAnimNotifyPanel.cpp b/Engine/Source/Editor/Persona/Private/SAnimNotifyPanel.cpp
--- a/Engine/Source/Editor/Persona/Private/SAnimNotifyPanel.cpp
+++ b/Engine/Source/Editor/Persona/Private/SAnimNotifyPanel.cpp
@@ -48,6 +48,9 @@
 #include "IAnimationSequenceBrowser.h"
 #include "AnimTimeline/AnimTimelineTrack_NotifiesPanel.h"
 #include "PersonaUtils.h"
+#include "AnimAssetFindReplace.h"
+#include "AnimAssetFindReplaceSyncMarkers.h"
+#include "AnimAssetFindReplaceNotifies.h"
 
 // AnimNotify Drawing
 const float NotifyHeightOffset = 0.f;
@@ -907,8 +910,8 @@ protected:
 	/** Opens the supplied blueprint in an editor */
 	void OnOpenNotifySource(UBlueprint* InSourceBlueprint) const;
 
-	/** Filters the asset browser by the selected notify */
-	void OnFilterSkeletonNotify(FName InName);
+	/** Filters the asset browser by the selected notify/sync marker */
+	void OnFindReferences(FName InName, bool bInIsSyncMarker);
 
 	/**
 	 * Selects a node on the track. Supports multi selection
@@ -2339,14 +2342,14 @@ void SAnimNotifyTrack::FillNewNotifyMenu(FMenuBuilder& MenuBuilder, bool bIsRepl
 	USkeleton* SeqSkeleton = Sequence->GetSkeleton();
 	if (SeqSkeleton)
 	{
-		MenuBuilder.BeginSection("AnimNotifySkeletonSubMenu", LOCTEXT("NewNotifySubMenu_Skeleton", "Skeleton Notifies"));
+		MenuBuilder.BeginSection("AnimNotifySubMenu", LOCTEXT("NewNotifySubMenu", "Notifies"));
 		{
 			if (!bIsReplaceWithMenu)
 			{
 				FUIAction UIAction;
 				UIAction.ExecuteAction.BindSP(
 					this, &SAnimNotifyTrack::OnNewNotifyClicked);
-				MenuBuilder.AddMenuEntry(LOCTEXT("NewNotify", "New Notify..."), LOCTEXT("NewNotifyToolTip", "Create a new animation notify on the skeleton"), FSlateIcon(), UIAction);
+				MenuBuilder.AddMenuEntry(LOCTEXT("NewNotify", "New Notify..."), LOCTEXT("NewNotifyToolTip", "Create a new animation notify"), FSlateIcon(), UIAction);
 			}
 
 			ISkeletonEditorModule& SkeletonEditorModule = FModuleManager::LoadModuleChecked<ISkeletonEditorModule>("SkeletonEditor");
@@ -2357,9 +2360,13 @@ void SAnimNotifyTrack::FillNewNotifyMenu(FMenuBuilder& MenuBuilder, bool bIsRepl
 				.WidthOverride(300.0f)
 				.HeightOverride(250.0f)
 				[
-					SNew(SSkeletonAnimNotifies, EditableSkeleton)
+					SNew(SSkeletonAnimNotifies)
 					.IsPicker(true)
+					.ShowSyncMarkers(false)
 					.ShowNotifies(true)
+					.ShowCompatibleSkeletonAssets(true)
+					.ShowOtherAssets(true)
+					.EditableSkeleton(EditableSkeleton)
 					.OnItemSelected_Lambda([this, bIsReplaceWithMenu](const FName& InNotifyName)
 					{
 						FSlateApplication::Get().DismissAllMenus();
@@ -2399,7 +2406,7 @@ void SAnimNotifyTrack::FillNewSyncMarkerMenu(FMenuBuilder& MenuBuilder, bool bIs
 	USkeleton* SeqSkeleton = Sequence->GetSkeleton();
 	if (SeqSkeleton)
 	{
-		MenuBuilder.BeginSection("AnimSyncMarkerSubMenu", LOCTEXT("NewSyncMarkerSubMenu_Skeleton", "Sync Markers"));
+		MenuBuilder.BeginSection("AnimSyncMarkerSubMenu", LOCTEXT("NewSyncMarkerSubMenu", "Sync Markers"));
 		{
 			FUIAction UIAction;
 			if (!bIsReplaceWithMenu)
@@ -2417,9 +2424,13 @@ void SAnimNotifyTrack::FillNewSyncMarkerMenu(FMenuBuilder& MenuBuilder, bool bIs
 				.WidthOverride(300.0f)
 				.HeightOverride(250.0f)
 				[
-					SNew(SSkeletonAnimNotifies, EditableSkeleton)
+					SNew(SSkeletonAnimNotifies)
+					.IsPicker(true)
 					.ShowSyncMarkers(true)
 					.ShowNotifies(false)
+					.ShowCompatibleSkeletonAssets(true)
+					.ShowOtherAssets(true)
+					.EditableSkeleton(EditableSkeleton)
 					.OnItemSelected_Lambda([this, bIsReplaceWithMenu](const FName& InNotifyName)
 					{
 						FSlateApplication::Get().DismissAllMenus();
@@ -3122,12 +3133,21 @@ TSharedPtr<SWidget> SAnimNotifyTrack::SummonContextMenu(const FGeometry& MyGeome
 		{
 			// skeleton notify
 			NewAction.ExecuteAction.BindRaw(
-				this, &SAnimNotifyTrack::OnFilterSkeletonNotify, NotifyEvent->NotifyName);
-			MenuBuilder.AddMenuEntry(LOCTEXT("FindNotifyReferences", "Find References"), LOCTEXT("FindNotifyReferencesTooltip", "Find all references to this skeleton notify in the asset browser"), FSlateIcon(), NewAction);
+				this, &SAnimNotifyTrack::OnFindReferences, NotifyEvent->NotifyName, false);
+			MenuBuilder.AddMenuEntry(LOCTEXT("FindNotifyReferences", "Find/Replace References..."), LOCTEXT("FindNotifyReferencesTooltip", "Find, replace and remove references to this notify in the find/replace taby"), FSlateIcon(), NewAction);
 		}
 
 		MenuBuilder.EndSection(); //ViewSource
 	}
+	else
+	{
+		if (NodeObject && NodeObject->GetType() == ENodeObjectTypes::SYNC_MARKER)
+		{
+			NewAction.ExecuteAction.BindRaw(
+				this, &SAnimNotifyTrack::OnFindReferences, NodeObject->GetName(), true);
+			MenuBuilder.AddMenuEntry(LOCTEXT("FindSyncMarkerReferences", "Find/Replace References..."), LOCTEXT("FindSyncMarkerReferencesTooltip", "Find, replace and remove references to this sync marker in the find/replace tab"), FSlateIcon(), NewAction);
+		}
+	}
 
 	FWidgetPath WidgetPath = MouseEvent.GetEventPath() != nullptr ? *MouseEvent.GetEventPath() : FWidgetPath();
 
@@ -3163,15 +3183,17 @@ void SAnimNotifyTrack::OnOpenNotifySource(UBlueprint* InSourceBlueprint) const
 	GEditor->GetEditorSubsystem<UAssetEditorSubsystem>()->OpenEditorForAsset(InSourceBlueprint);
 }
 
-void SAnimNotifyTrack::OnFilterSkeletonNotify(FName InName)
+void SAnimNotifyTrack::OnFindReferences(FName InName, bool bInIsSyncMarker)
 {
-	// Open asset browser first
-	OnInvokeTab.ExecuteIfBound(FPersonaTabs::AssetBrowserID);
-
 	IAssetEditorInstance* AssetEditor = GEditor->GetEditorSubsystem<UAssetEditorSubsystem>()->FindEditorForAsset(Sequence, true);
 	check(AssetEditor->GetEditorName() == TEXT("AnimationEditor"));
-	IAnimationEditor* AnimationEditor = static_cast<IAnimationEditor*>(AssetEditor);
-	AnimationEditor->GetAssetBrowser()->FilterBySkeletonNotify(InName);
+	if (TSharedPtr<SDockTab> Tab = AssetEditor->GetAssociatedTabManager()->TryInvokeTab(FPersonaTabs::FindReplaceID))
+	{
+		TSharedRef<IAnimAssetFindReplace> FindReplaceWidget = StaticCastSharedRef<IAnimAssetFindReplace>(Tab->GetContent());
+		FindReplaceWidget->SetCurrentProcessor(bInIsSyncMarker ? UAnimAssetFindReplaceSyncMarkers::StaticClass() : UAnimAssetFindReplaceNotifies::StaticClass());
+		UAnimAssetFindReplaceProcessor_StringBase* Processor = Cast<UAnimAssetFindReplaceProcessor_StringBase>(FindReplaceWidget->GetCurrentProcessor());
+		Processor->SetFindString(InName.ToString());
+	}
 }
 
 bool SAnimNotifyTrack::IsSingleNodeSelected()
@@ -3235,20 +3257,16 @@ void SAnimNotifyTrack::OnNewSyncMarkerClicked()
 
 void SAnimNotifyTrack::AddNewNotify(const FText& NewNotifyName, ETextCommit::Type CommitInfo)
 {
-	USkeleton* SeqSkeleton = Sequence->GetSkeleton();
-	if ((CommitInfo == ETextCommit::OnEnter) && SeqSkeleton)
+	if (CommitInfo == ETextCommit::OnEnter)
 	{
 		const FScopedTransaction Transaction( LOCTEXT("AddNewNotifyEvent", "Add New Anim Notify") );
 		FName NewName = FName( *NewNotifyName.ToString() );
 
-		ISkeletonEditorModule& SkeletonEditorModule = FModuleManager::LoadModuleChecked<ISkeletonEditorModule>("SkeletonEditor");
-		TSharedRef<IEditableSkeleton> EditableSkeleton = SkeletonEditorModule.CreateEditableSkeleton(SeqSkeleton);
-
-		EditableSkeleton->AddNotify(NewName);
-
-		FBlueprintActionDatabase::Get().RefreshAssetActions(SeqSkeleton);
-
 		CreateNewNotifyAtCursor(NewNotifyName.ToString(), (UClass*)nullptr);
+
+		FBlueprintActionDatabase& ActionDatabase = FBlueprintActionDatabase::Get();
+		ActionDatabase.ClearAssetActions(UAnimBlueprint::StaticClass());
+		ActionDatabase.RefreshClassActions(UAnimBlueprint::StaticClass());
 	}
 
 	FSlateApplication::Get().DismissAllMenus();
@@ -3256,20 +3274,10 @@ void SAnimNotifyTrack::AddNewNotify(const FText& NewNotifyName, ETextCommit::Typ
 
 void SAnimNotifyTrack::AddNewSyncMarker(const FText& NewNotifyName, ETextCommit::Type CommitInfo) 
 {
-	USkeleton* SeqSkeleton = Sequence->GetSkeleton();
-	if ((CommitInfo == ETextCommit::OnEnter) && SeqSkeleton)
+	if (CommitInfo == ETextCommit::OnEnter)
 	{
 		const FScopedTransaction Transaction(LOCTEXT("AddNewSyncMarker", "Add New Sync Marker"));
 
-		FName NewName = FName(*NewNotifyName.ToString());
-
-		ISkeletonEditorModule& SkeletonEditorModule = FModuleManager::LoadModuleChecked<ISkeletonEditorModule>("SkeletonEditor");
-		TSharedRef<IEditableSkeleton> EditableSkeleton = SkeletonEditorModule.CreateEditableSkeleton(SeqSkeleton);
-
-		EditableSkeleton->AddSyncMarker(NewName);
-
-		FBlueprintActionDatabase::Get().RefreshAssetActions(SeqSkeleton);
-
 		CreateNewSyncMarkerAtCursor(NewNotifyName.ToString());
 	}
 
diff --git a/Engine/Source/Editor/Persona/Private/SAnimationSequenceBrowser.cpp b/Engine/Source/Editor/Persona/Private/SAnimationSequenceBrowser.cpp
--- a/Engine/Source/Editor/Persona/Private/SAnimationSequenceBrowser.cpp
+++ b/Engine/Source/Editor/Persona/Private/SAnimationSequenceBrowser.cpp
@@ -139,9 +139,13 @@ public:
 			.HeightOverride(300.0f)
 			.WidthOverride(200.0f)
 			[
-				SNew(SSkeletonAnimNotifies, EditableSkeleton.ToSharedRef())
+				SNew(SSkeletonAnimNotifies)
 				.IsPicker(true)
 				.ShowNotifies(true)
+				.ShowSyncMarkers(false)
+				.ShowCompatibleSkeletonAssets(true)
+				.ShowOtherAssets(true)
+				.EditableSkeleton(EditableSkeleton.Pin())
 				.OnItemSelected_Lambda([this](const FName& InNotifyName)
 				{
 					FSlateApplication::Get().DismissAllMenus();
@@ -203,8 +207,8 @@ public:
 	/** Notify string to use when filtering */
 	FString NotifyString;
 
-	/** Editable skeleton used to access available notifies */
-	TSharedPtr<IEditableSkeleton> EditableSkeleton;
+	/** Editable skeleton used to filter available notifies */
+	TWeakPtr<IEditableSkeleton> EditableSkeleton;
 };
 
 /** A filter that displays animations that use a skeleton sync marker */
@@ -267,10 +271,13 @@ public:
 			.HeightOverride(300.0f)
 			.WidthOverride(200.0f)
 			[
-				SNew(SSkeletonAnimNotifies, EditableSkeleton.ToSharedRef())
+				SNew(SSkeletonAnimNotifies)
 				.IsPicker(true)
 				.ShowNotifies(false)
 				.ShowSyncMarkers(true)
+				.ShowCompatibleSkeletonAssets(true)
+				.ShowOtherAssets(true)
+				.EditableSkeleton(EditableSkeleton.Pin())
 				.OnItemSelected_Lambda([this](const FName& InNotifyName)
 				{
 					FSlateApplication::Get().DismissAllMenus();
@@ -332,8 +339,8 @@ public:
 	/** Sync marker string to use when filtering */
 	FString SyncMarkerString;
 
-	/** Editable skeleton used to access available sync markers */
-	TSharedPtr<IEditableSkeleton> EditableSkeleton;
+	/** Editable skeleton used to filter available sync markers */
+	TWeakPtr<IEditableSkeleton> EditableSkeleton;
 };
 
 /** A filter that displays animations that use a curve */
diff --git a/Engine/Source/Editor/Persona/Private/SSkeletonAnimNotifies.cpp b/Engine/Source/Editor/Persona/Private/SSkeletonAnimNotifies.cpp
--- a/Engine/Source/Editor/Persona/Private/SSkeletonAnimNotifies.cpp
+++ b/Engine/Source/Editor/Persona/Private/SSkeletonAnimNotifies.cpp
@@ -12,7 +12,6 @@
 #include "Animation/AnimSequenceBase.h"
 #include "Styling/AppStyle.h"
 #include "Animation/EditorSkeletonNotifyObj.h"
-
 #include "Widgets/Input/SSearchBox.h"
 #include "Widgets/Text/SInlineEditableTextBlock.h"
 #include "ScopedTransaction.h"
@@ -27,6 +26,8 @@
 #include "Filters/GenericFilter.h"
 #include "Filters/SBasicFilterBar.h"
 #include "Widgets/Docking/SDockTab.h"
+#include "AssetRegistry/AssetRegistryModule.h"
+#include "String/ParseTokens.h"
 
 #define LOCTEXT_NAMESPACE "SkeletonAnimNotifies"
 
@@ -57,6 +58,8 @@ private:
 	// FFilterBase interface
 	virtual void ActiveStateChanged(bool bActive) override
 	{
+		FGenericFilter<EAnimNotifyFilterFlags>::ActiveStateChanged(bActive);
+
 		bIsActive = bActive;
 	}
 
@@ -79,45 +82,49 @@ FSkeletonAnimNotifiesSummoner::FSkeletonAnimNotifiesSummoner(TSharedPtr<class FA
 	, OnObjectsSelected(InOnObjectsSelected)
 {
 	TabLabel = LOCTEXT("SkeletonAnimNotifiesTabTitle", "Animation Notifies");
-	TabIcon = FSlateIcon(FAppStyle::GetAppStyleSetName(), "Persona.Tabs.AnimationNotifies");
+	TabIcon = FSlateIcon(FAppStyle::GetAppStyleSetName(), "Persona.AnimNotifyWindow");
 
 	EnableTabPadding();
 	bIsSingleton = true;
 
 	ViewMenuDescription = LOCTEXT("SkeletonAnimNotifiesMenu", "Animation Notifies");
-	ViewMenuTooltip = LOCTEXT("SkeletonAnimNotifies_ToolTip", "Shows the skeletons notifies list");
+	ViewMenuTooltip = LOCTEXT("SkeletonAnimNotifies_ToolTip", "Shows the notify and sync marker list");
 }
 
 TSharedRef<SWidget> FSkeletonAnimNotifiesSummoner::CreateTabBody(const FWorkflowTabSpawnInfo& Info) const
 {
-	return SNew(SSkeletonAnimNotifies, EditableSkeleton.Pin().ToSharedRef(), HostingApp.Pin())
-		.ShowNotifies(true)
-		.ShowSyncMarkers(true)
-		.OnObjectsSelected(OnObjectsSelected);
+	return SNew(SSkeletonAnimNotifies, HostingApp.Pin())
+		.OnObjectsSelected(OnObjectsSelected)
+		.EditableSkeleton(EditableSkeleton.Pin());
 }
 
 /////////////////////////////////////////////////////
 // SSkeletonAnimNotifies
 
-void SSkeletonAnimNotifies::Construct(const FArguments& InArgs, const TSharedRef<IEditableSkeleton>& InEditableSkeleton, const TSharedPtr<class FAssetEditorToolkit>& InHostingApp)
+void SSkeletonAnimNotifies::Construct(const FArguments& InArgs, const TSharedPtr<class FAssetEditorToolkit>& InHostingApp)
 {
 	OnObjectsSelected = InArgs._OnObjectsSelected;
 	OnItemSelected = InArgs._OnItemSelected;
 	bIsPicker = InArgs._IsPicker;
 	bShowSyncMarkers = InArgs._ShowSyncMarkers;
+	bShowOtherAssets = InArgs._ShowOtherAssets;
+	bShowCompatibleSkeletonAssets = InArgs._ShowCompatibleSkeletonAssets;
 	bShowNotifies = InArgs._ShowNotifies;
+	EditableSkeleton = InArgs._EditableSkeleton;
 
-	EditableSkeleton = InEditableSkeleton;
 	WeakHostingApp = InHostingApp;
 
-	EditableSkeleton->RegisterOnNotifiesChanged(FSimpleDelegate::CreateSP(this, &SSkeletonAnimNotifies::OnNotifiesChanged));
+	if (EditableSkeleton.IsValid())
+	{
+		EditableSkeleton->RegisterOnNotifiesChanged(FSimpleDelegate::CreateSP(this, &SSkeletonAnimNotifies::OnNotifiesChanged));
+	}
 
 	if(GEditor)
 	{
 		GEditor->RegisterForUndo(this);
 	}
 
-	FOnContextMenuOpening OnContextMenuOpening = !bIsPicker ? FOnContextMenuOpening::CreateSP(this, &SSkeletonAnimNotifies::OnGetContextMenuContent) : FOnContextMenuOpening();
+	FOnContextMenuOpening OnContextMenuOpening = (!bIsPicker && EditableSkeleton.IsValid()) ? FOnContextMenuOpening::CreateSP(this, &SSkeletonAnimNotifies::OnGetContextMenuContent) : FOnContextMenuOpening();
 
 	NameFilterBox = SNew( SSearchBox )
 		.SelectAllTextWhenFocused( true )
@@ -132,19 +139,18 @@ void SSkeletonAnimNotifies::Construct(const FArguments& InArgs, const TSharedRef
 		.ItemHeight( 18.0f )
 		.OnItemScrolledIntoView( this, &SSkeletonAnimNotifies::OnItemScrolledIntoView );
 	
+	CurrentFilterFlags = EAnimNotifyFilterFlags::None;
+
 	TSharedPtr<FFilterCategory> FilterCategory = MakeShared<FFilterCategory>(LOCTEXT("AnimNotifyFiltersLabel", "Anim Notify Filters"), LOCTEXT("AnimNotifyFiltersLabelToolTip", "Filter what kind fo notifies and sync markers can be displayed."));
 
-	// Hide filter UI if we are only displaying one type of thing
-	const EVisibility FilterVisibility = (bShowNotifies ^ bShowSyncMarkers) ? EVisibility::Collapsed : EVisibility::Visible;
+	const bool bSingleType = (bShowNotifies ^ bShowSyncMarkers);
 
 	TSharedRef<SBasicFilterBar<EAnimNotifyFilterFlags>> FilterBar = SNew(SBasicFilterBar<EAnimNotifyFilterFlags>)
-	.Visibility(FilterVisibility)
 	.CustomFilters(Filters)
+	.bPinAllFrontendFilters(true)
 	.UseSectionsForCategories(true)
 	.OnFilterChanged_Lambda([this]()
 	{
-		CurrentFilterFlags = EAnimNotifyFilterFlags::None;
-
 		for(const TSharedRef<FFilterBase<EAnimNotifyFilterFlags>>& Filter : Filters)
 		{
 			TSharedRef<FSkeletonAnimNotifiesFilter> AnimNotifiesFilter = StaticCastSharedRef<FSkeletonAnimNotifiesFilter>(Filter);
@@ -152,12 +158,38 @@ void SSkeletonAnimNotifies::Construct(const FArguments& InArgs, const TSharedRef
 			{
 				CurrentFilterFlags |= AnimNotifiesFilter->GetFlags();
 			}
+			else
+			{
+				CurrentFilterFlags &= ~AnimNotifiesFilter->GetFlags();
+			}
 		}
 
 		RefreshNotifiesListWithFilter();
 	});
 	
-	if(bShowNotifies)
+	if (EditableSkeleton.IsValid())
+	{
+		CurrentFilterFlags |= EAnimNotifyFilterFlags::CurrentSkeleton;
+
+		TSharedRef<FFilterBase<EAnimNotifyFilterFlags>> Filter = Filters.Add_GetRef(MakeShared<FSkeletonAnimNotifiesFilter>(
+			EAnimNotifyFilterFlags::CurrentSkeleton,
+			"CurrentSkeleton",
+			LOCTEXT("ShowSkeletonItemsLabel", "Skeleton"),
+			LOCTEXT("ShowSkeletonItemsTooltip", "Show items for the current skeleton"),
+			FLinearColor::Blue.Desaturate(0.25f),
+			FilterCategory
+		));
+
+		FilterBar->AddFilter(Filter);
+		Filter->SetActive(true);
+	}
+
+	if (bShowNotifies)
+	{
+		CurrentFilterFlags |= EAnimNotifyFilterFlags::Notifies;
+	}
+
+	if (!bSingleType && bShowNotifies)
 	{
 		TSharedRef<FFilterBase<EAnimNotifyFilterFlags>> Filter = Filters.Add_GetRef(MakeShared<FSkeletonAnimNotifiesFilter>(
 			EAnimNotifyFilterFlags::Notifies,
@@ -166,12 +198,18 @@ void SSkeletonAnimNotifies::Construct(const FArguments& InArgs, const TSharedRef
 			LOCTEXT("ShowNotifiesTooltip", "Show notifies"),
 			FLinearColor::Red.Desaturate(0.5f),
 			FilterCategory
-			));
+		));
+
+		FilterBar->AddFilter(Filter);
+		Filter->SetActive(true);
+	}
 
-		FilterBar->SetFilterCheckState(Filter, ECheckBoxState::Checked);
+	if (bShowSyncMarkers)
+	{
+		CurrentFilterFlags |= EAnimNotifyFilterFlags::SyncMarkers;
 	}
 
-	if(bShowSyncMarkers)
+	if (!bSingleType && bShowSyncMarkers)
 	{
 		TSharedRef<FFilterBase<EAnimNotifyFilterFlags>> Filter = Filters.Add_GetRef(MakeShared<FSkeletonAnimNotifiesFilter>(
 			EAnimNotifyFilterFlags::SyncMarkers,
@@ -180,13 +218,45 @@ void SSkeletonAnimNotifies::Construct(const FArguments& InArgs, const TSharedRef
 			LOCTEXT("ShowSyncMarkersTooltip", "Show sync markers"),
 			FLinearColor::Green.Desaturate(0.5f),
 			FilterCategory
-			));
+		));
 
-		FilterBar->SetFilterCheckState(Filter, ECheckBoxState::Checked);
+		FilterBar->AddFilter(Filter);
+		Filter->SetActive(true);
+	}
+
+	{
+		CurrentFilterFlags |= bShowOtherAssets ? EAnimNotifyFilterFlags::OtherAssets : EAnimNotifyFilterFlags::None;
+
+		TSharedRef<FFilterBase<EAnimNotifyFilterFlags>> Filter = Filters.Add_GetRef(MakeShared<FSkeletonAnimNotifiesFilter>(
+			EAnimNotifyFilterFlags::OtherAssets,
+			"OtherAssets",
+			LOCTEXT("ShowOtherAssetsLabel", "Other Assets"),
+			LOCTEXT("ShowOtherAssetsTooltip", "Show items that are present on other assets"),
+			FLinearColor::Yellow.Desaturate(0.5f),
+			FilterCategory
+		));
+
+		FilterBar->AddFilter(Filter);
+		Filter->SetActive(bShowOtherAssets);
+	}
+
+	{
+		CurrentFilterFlags |= bShowCompatibleSkeletonAssets ? EAnimNotifyFilterFlags::CompatibleAssets : EAnimNotifyFilterFlags::None;
+
+		TSharedRef<FFilterBase<EAnimNotifyFilterFlags>> Filter = Filters.Add_GetRef(MakeShared<FSkeletonAnimNotifiesFilter>(
+			EAnimNotifyFilterFlags::CompatibleAssets,
+			"CompatibleAssets",
+			LOCTEXT("ShowCompatibleAssetsLabel", "Compatible"),
+			LOCTEXT("ShowCompatibleAssetsTooltip", "Show items that are present on other assets that are compatible with the current skeleton"),
+			FLinearColor::Blue.Desaturate(0.5f),
+			FilterCategory
+		));
+
+		FilterBar->AddFilter(Filter);
+		Filter->SetActive(bShowCompatibleSkeletonAssets);
 	}
 
 	TSharedRef<SWidget> AddFilterButton = SBasicFilterBar<EAnimNotifyFilterFlags>::MakeAddFilterButton(FilterBar);
-	AddFilterButton->SetVisibility(FilterVisibility);
 	
 	ChildSlot
 	[
@@ -224,7 +294,7 @@ void SSkeletonAnimNotifies::Construct(const FArguments& InArgs, const TSharedRef
 		]
 	];
 
-	CreateNotifiesList();
+	RefreshNotifiesListWithFilter();
 }
 
 SSkeletonAnimNotifies::~SSkeletonAnimNotifies()
@@ -257,23 +327,58 @@ TSharedRef<ITableRow> SSkeletonAnimNotifies::GenerateNotifyRow(TSharedPtr<FDispl
 {
 	check( InInfo.IsValid() );
 
+	FTextBuilder TooltipBuilder;
+	if (InInfo->bIsSyncMarker)
+	{
+		TooltipBuilder.AppendLineFormat(LOCTEXT("SyncMarkerTooltip", "Sync Marker '{0}'"), FText::FromName(InInfo->Name));
+	}
+	else
+	{
+		TooltipBuilder.AppendLineFormat(LOCTEXT("NotifyTooltip", "Notify '{0}'"), FText::FromName(InInfo->Name));
+	}
+
+	if (EnumHasAnyFlags(InInfo->ItemFlags, EAnimNotifyFilterFlags::CurrentSkeleton))
+	{
+		TooltipBuilder.AppendLine(LOCTEXT("CurrentSkeleton", "Item is on the current skeleton"));
+	}
+
+	if (EnumHasAnyFlags(InInfo->ItemFlags, EAnimNotifyFilterFlags::CompatibleAssets))
+	{
+		TooltipBuilder.AppendLine(LOCTEXT("CompatibleAsset", "Item is on a compatible asset"));
+	}
+
+	if (EnumHasAnyFlags(InInfo->ItemFlags, EAnimNotifyFilterFlags::OtherAssets))
+	{
+		TooltipBuilder.AppendLine(LOCTEXT("AnotherAsset", "Item is on another asset"));
+	}
+	
 	return
 		SNew( STableRow<TSharedPtr<FDisplayedAnimNotifyInfo>>, OwnerTable )
+		.ToolTipText(TooltipBuilder.ToText())
 		[
-			SNew( SVerticalBox )
-			+ SVerticalBox::Slot()
-			.AutoHeight()
-			.Padding( 0.0f, 4.0f )
+			SNew( SHorizontalBox )
+			+SHorizontalBox::Slot()
+			.AutoWidth()
+			.Padding( 4.0f, 4.0f )
+			.VAlign( VAlign_Center )
+			[
+				SNew(SImage)
+				.Image(InInfo->bIsSyncMarker ? FAppStyle::Get().GetBrush("AnimNotifyEditor.AnimSyncMarker") : FAppStyle::Get().GetBrush("AnimNotifyEditor.AnimNotify"))
+			]
+			+SHorizontalBox::Slot()
+			.AutoWidth()
+			.Padding( 4.0f, 4.0f )
 			.VAlign( VAlign_Center )
 			[
 				SAssignNew(InInfo->InlineEditableText, SInlineEditableTextBlock)
 				.Text(FText::FromName(InInfo->Name))
 				.Font(FAppStyle::Get().GetFontStyle("SmallFont"))
+				.ColorAndOpacity(EnumHasAnyFlags(InInfo->ItemFlags, EAnimNotifyFilterFlags::CurrentSkeleton) ? FSlateColor::UseForeground() : FSlateColor::UseSubduedForeground())
 				.OnVerifyTextChanged(this, &SSkeletonAnimNotifies::OnVerifyNotifyNameCommit, InInfo)
 				.OnTextCommitted(this, &SSkeletonAnimNotifies::OnNotifyNameCommitted, InInfo)
 				.IsSelected(this, &SSkeletonAnimNotifies::IsSelected)
 				.HighlightText_Lambda([this](){ return FilterText; })
-				.IsReadOnly(bIsPicker)
+				.IsReadOnly(bIsPicker && EditableSkeleton.IsValid())
 			]
 		];
 }
@@ -320,9 +425,9 @@ TSharedPtr<SWidget> SSkeletonAnimNotifies::OnGetContextMenuContent() const
 		if(WeakHostingApp.IsValid() && NotifiesListView->GetNumItemsSelected() == 1)
 		{
 			FUIAction Action = FUIAction(FExecuteAction::CreateSP(const_cast<SSkeletonAnimNotifies*>(this), &SSkeletonAnimNotifies::OnFindReferences),
-				FCanExecuteAction::CreateSP(this, &SSkeletonAnimNotifies::CanPerformDelete));
-			const FText Label = LOCTEXT("FindNotifyReferences", "Find References");
-			const FText ToolTipText = LOCTEXT("FindNotifyReferencesTooltip", "Find all references to this item in the asset browser");
+			FCanExecuteAction::CreateSP(this, &SSkeletonAnimNotifies::CanPerformFindReferences));
+			const FText Label = LOCTEXT("FindNotifyReferences", "Find/Replace References...");
+			const FText ToolTipText = LOCTEXT("FindNotifyReferencesTooltip", "Find, replace and remove references to this item in the find/replace tab");
 			MenuBuilder.AddMenuEntry(Label, ToolTipText, FSlateIcon(), Action);
 		}
 	}
@@ -346,10 +451,33 @@ void SSkeletonAnimNotifies::OnNotifySelectionChanged(TSharedPtr<FDisplayedAnimNo
 
 bool SSkeletonAnimNotifies::CanPerformDelete() const
 {
-	return NotifiesListView->GetNumItemsSelected() > 0;
+	if (NotifiesListView->GetNumItemsSelected() > 0)
+	{
+		for (const TSharedPtr<FDisplayedAnimNotifyInfo>& Item : NotifiesListView->GetSelectedItems())
+		{
+			// We can delete only if an item was from 'this' skeleton
+			if (EnumHasAnyFlags(Item->ItemFlags, EAnimNotifyFilterFlags::CurrentSkeleton))
+			{
+				return true;
+			}
+		}
+	}
+
+	return false;
 }
 
 bool SSkeletonAnimNotifies::CanPerformRename() const
+{
+	if (NotifiesListView->GetNumItemsSelected() == 1)
+	{
+		// We can rename only if the item was from 'this' skeleton
+		return EnumHasAnyFlags(NotifiesListView->GetSelectedItems()[0]->ItemFlags, EAnimNotifyFilterFlags::CurrentSkeleton);
+	}
+
+	return false;
+}
+
+bool SSkeletonAnimNotifies::CanPerformFindReferences() const
 {
 	return NotifiesListView->GetNumItemsSelected() == 1;
 }
@@ -361,14 +489,19 @@ void SSkeletonAnimNotifies::OnAddItem(bool bIsSyncMarker)
 	FString NewNotifyString = BaseNotifyString;
 	int32 NumericExtension = 0;
 
-	while(EditableSkeleton->GetSkeleton().AnimationNotifies.ContainsByPredicate([&NewNotifyString](const FName& InNotifyName){ return InNotifyName.ToString() == NewNotifyString; }))
+	auto NameAndTypeMatches = [&NewNotifyString, bIsSyncMarker](const TSharedPtr<FDisplayedAnimNotifyInfo>& InNotify)
+	{
+		return InNotify->Name.ToString() == NewNotifyString && InNotify->bIsSyncMarker == bIsSyncMarker;
+	};
+
+	while(NotifyList.ContainsByPredicate(NameAndTypeMatches))
 	{
 		NewNotifyString = FString::Printf(TEXT("%s_%d"), BaseNotifyString, NumericExtension);
 		NumericExtension++;
 	}
 
 	// Add an item. The subsequent rename will commit the item.
-	TSharedPtr<FDisplayedAnimNotifyInfo> NewItem = FDisplayedAnimNotifyInfo::Make(*NewNotifyString, bIsSyncMarker);
+	TSharedPtr<FDisplayedAnimNotifyInfo> NewItem = FDisplayedAnimNotifyInfo::Make(*NewNotifyString, bIsSyncMarker, EAnimNotifyFilterFlags::CurrentSkeleton);
 	NewItem->bIsNew = true;
 	NotifyList.Add(NewItem);
 
@@ -387,48 +520,40 @@ void SSkeletonAnimNotifies::OnItemScrolledIntoView(TSharedPtr<FDisplayedAnimNoti
 
 void SSkeletonAnimNotifies::OnDeleteItems()
 {
-	TArray< TSharedPtr< FDisplayedAnimNotifyInfo > > SelectedRows = NotifiesListView->GetSelectedItems();
+	if (EditableSkeleton.IsValid())
+	{
+		TArray<TSharedPtr<FDisplayedAnimNotifyInfo>> SelectedRows = NotifiesListView->GetSelectedItems();
 
-	TArray<FName> SelectedSyncMarkerNames;
-	TArray<FName> SelectedNotifyNames;
+		TArray<FName> SelectedSyncMarkerNames;
+		TArray<FName> SelectedNotifyNames;
 
-	for (TSharedPtr< FDisplayedAnimNotifyInfo > Selection : SelectedRows)
-	{
-		if(Selection->bIsSyncMarker)
+		for (TSharedPtr<FDisplayedAnimNotifyInfo> Selection : SelectedRows)
 		{
-			SelectedSyncMarkerNames.Add(Selection->Name);
+			if (EnumHasAnyFlags(Selection->ItemFlags, EAnimNotifyFilterFlags::CurrentSkeleton))
+			{
+				if (Selection->bIsSyncMarker)
+				{
+					SelectedSyncMarkerNames.Add(Selection->Name);
+				}
+				else
+				{
+					SelectedNotifyNames.Add(Selection->Name);
+				}
+			}
 		}
-		else
+
+		if (SelectedSyncMarkerNames.Num())
 		{
-			SelectedNotifyNames.Add(Selection->Name);
+			EditableSkeleton->DeleteSyncMarkers(SelectedSyncMarkerNames, false);
 		}
-	}
 
-	int32 NumAnimationsModified = 0;
+		if (SelectedNotifyNames.Num())
+		{
+			EditableSkeleton->DeleteAnimNotifies(SelectedNotifyNames, false);
+		}
 	
-	if(SelectedSyncMarkerNames.Num())
-	{
-		NumAnimationsModified += EditableSkeleton->DeleteSyncMarkers(SelectedSyncMarkerNames);
-	}
-
-	if(SelectedNotifyNames.Num())
-	{
-		NumAnimationsModified += EditableSkeleton->DeleteAnimNotifies(SelectedNotifyNames);
-	}
-
-	if(NumAnimationsModified > 0)
-	{
-		FFormatNamedArguments Args;
-		Args.Add( TEXT("NumAnimationsModified"), NumAnimationsModified );
-		FNotificationInfo Info( FText::Format( LOCTEXT( "ItemsDeleted", "{NumAnimationsModified} animation(s) modified to delete items" ), Args ) );
-
-		Info.bUseLargeFont = false;
-		Info.ExpireDuration = 5.0f;
-
-		NotifyUser( Info );
+		RefreshNotifiesListWithFilter();
 	}
-	
-	CreateNotifiesList(NameFilterBox->GetText().ToString());
 }
 
 void SSkeletonAnimNotifies::OnRenameItem()
@@ -436,147 +561,244 @@ void SSkeletonAnimNotifies::OnRenameItem()
 	TArray< TSharedPtr< FDisplayedAnimNotifyInfo > > SelectedRows = NotifiesListView->GetSelectedItems();
 
 	check(SelectedRows.Num() == 1); // Should be guaranteed by CanPerformRename
-
-	SelectedRows[0]->InlineEditableText->EnterEditingMode();
+	if (EnumHasAnyFlags(SelectedRows[0]->ItemFlags, EAnimNotifyFilterFlags::CurrentSkeleton))
+	{
+		SelectedRows[0]->InlineEditableText->EnterEditingMode();
+	}
 }
 
 bool SSkeletonAnimNotifies::OnVerifyNotifyNameCommit( const FText& NewName, FText& OutErrorMessage, TSharedPtr<FDisplayedAnimNotifyInfo> Item )
 {
-	bool bValid(true);
-
-	if(NewName.IsEmpty())
+	if (EditableSkeleton.IsValid())
 	{
-		OutErrorMessage = LOCTEXT( "NameMissing_Error", "You must provide a name." );
-		bValid = false;
-	}
+		bool bValid(true);
 
-	FName NotifyName( *NewName.ToString() );
-	if(NotifyName != Item->Name || Item->bIsNew)
-	{
-		if(Item->bIsSyncMarker)
+		if (NewName.IsEmpty())
 		{
-			if(EditableSkeleton->GetSkeleton().GetExistingMarkerNames().Contains(NotifyName))
-			{
-				OutErrorMessage = FText::Format( LOCTEXT("AlreadyInUseMessage", "'{0}' is already in use."), NewName );
-				bValid = false;
-			}
+			OutErrorMessage = LOCTEXT("NameMissing_Error", "You must provide a name.");
+			bValid = false;
 		}
-		else
+
+		FName NotifyName(*NewName.ToString());
+		if (NotifyName != Item->Name || Item->bIsNew)
 		{
-			if(EditableSkeleton->GetSkeleton().AnimationNotifies.Contains(NotifyName))
+			if (Item->bIsSyncMarker)
 			{
-				OutErrorMessage = FText::Format( LOCTEXT("AlreadyInUseMessage", "'{0}' is already in use."), NewName );
-				bValid = false;
+				if (EditableSkeleton->GetSkeleton().GetExistingMarkerNames().Contains(NotifyName))
+				{
+					OutErrorMessage = FText::Format(LOCTEXT("AlreadyInUseMessage", "'{0}' is already in use."), NewName);
+					bValid = false;
+				}
+			}
+			else
+			{
+				if (EditableSkeleton->GetSkeleton().AnimationNotifies.Contains(NotifyName))
+				{
+					OutErrorMessage = FText::Format(LOCTEXT("AlreadyInUseMessage", "'{0}' is already in use."), NewName);
+					bValid = false;
+				}
 			}
 		}
+
+		return bValid;
 	}
 
-	return bValid;
+	return false;
 }
 
 void SSkeletonAnimNotifies::OnNotifyNameCommitted( const FText& NewName, ETextCommit::Type, TSharedPtr<FDisplayedAnimNotifyInfo> Item )
 {
-	FName NewFName = FName(*NewName.ToString());
-	if(Item->bIsNew)
+	if (EditableSkeleton.IsValid())
 	{
-		if(Item->bIsSyncMarker)
+		FName NewFName = FName(*NewName.ToString());
+		if (Item->bIsNew)
 		{
-			EditableSkeleton->AddSyncMarker(NewFName);
+			if (Item->bIsSyncMarker)
+			{
+				EditableSkeleton->AddSyncMarker(NewFName);
+			}
+			else
+			{
+				EditableSkeleton->AddNotify(NewFName);
+			}
+			Item->bIsNew = false;
 		}
 		else
 		{
-			EditableSkeleton->AddNotify(NewFName);
-		}
-		Item->bIsNew = false;
-	}
-	else
-	{
-		if(NewFName != Item->Name)
-		{
-			if(Item->bIsSyncMarker)
+			if (NewFName != Item->Name)
 			{
-				int32 NumAnimationsModified = EditableSkeleton->RenameSyncMarker(FName(*NewName.ToString()), Item->Name);
-				if(NumAnimationsModified > 0)
+				if (Item->bIsSyncMarker)
 				{
-					FFormatNamedArguments Args;
-					Args.Add( TEXT("NumAnimationsModified"), NumAnimationsModified );
-					FNotificationInfo Info( FText::Format( LOCTEXT( "SyncMarkersRenamed", "{NumAnimationsModified} animation(s) modified to rename sync marker" ), Args ) );
-
-					Info.bUseLargeFont = false;
-					Info.ExpireDuration = 5.0f;
-
-					NotifyUser( Info );
+					int32 NumAnimationsModified = EditableSkeleton->RenameSyncMarker(FName(*NewName.ToString()), Item->Name, false);
+					if (NumAnimationsModified > 0)
+					{
+						FFormatNamedArguments Args;
+						Args.Add(TEXT("NumAnimationsModified"), NumAnimationsModified);
+						FNotificationInfo Info(FText::Format(LOCTEXT("SyncMarkersRenamed", "{NumAnimationsModified} animation(s) modified to rename sync marker"), Args));
+
+						Info.bUseLargeFont = false;
+						Info.ExpireDuration = 5.0f;
+
+						NotifyUser(Info);
+					}
 				}
-			}
-			else
-			{
-				int32 NumAnimationsModified = EditableSkeleton->RenameNotify(FName(*NewName.ToString()), Item->Name);
-				if(NumAnimationsModified > 0)
+				else
 				{
-					FFormatNamedArguments Args;
-					Args.Add( TEXT("NumAnimationsModified"), NumAnimationsModified );
-					FNotificationInfo Info( FText::Format( LOCTEXT( "AnimNotifiesRenamed", "{NumAnimationsModified} animation(s) modified to rename notification" ), Args ) );
-
-					Info.bUseLargeFont = false;
-					Info.ExpireDuration = 5.0f;
-
-					NotifyUser( Info );
+					int32 NumAnimationsModified = EditableSkeleton->RenameNotify(FName(*NewName.ToString()), Item->Name, false);
+					if (NumAnimationsModified > 0)
+					{
+						FFormatNamedArguments Args;
+						Args.Add(TEXT("NumAnimationsModified"), NumAnimationsModified);
+						FNotificationInfo Info(FText::Format(LOCTEXT("AnimNotifiesRenamed", "{NumAnimationsModified} animation(s) modified to rename notification"), Args));
+
+						Info.bUseLargeFont = false;
+						Info.ExpireDuration = 5.0f;
+
+						NotifyUser(Info);
+					}
 				}
-			}
 			
-			RefreshNotifiesListWithFilter();
+				RefreshNotifiesListWithFilter();
+			}
 		}
 	}
 }
 
 void SSkeletonAnimNotifies::RefreshNotifiesListWithFilter()
 {
-	CreateNotifiesList( NameFilterBox->GetText().ToString() );
+	CreateNotifiesList();
+	FilterNotifiesList(NameFilterBox->GetText().ToString());
+}
+
+void SSkeletonAnimNotifies::FilterNotifiesList(const FString& InSearchText)
+{
+	FilteredNotifyList.Empty();
+
+	for (const TSharedPtr<FDisplayedAnimNotifyInfo>& Item : NotifyList)
+	{
+		if (!InSearchText.IsEmpty())
+		{
+			if (Item->Name.ToString().Contains(InSearchText))
+			{
+				FilteredNotifyList.Add(Item);
+			}
+		}
+		else
+		{
+			FilteredNotifyList.Add(Item);
+		}
+	}
+
+	NotifiesListView->RequestListRefresh();
 }
 
-void SSkeletonAnimNotifies::CreateNotifiesList( const FString& SearchText )
+void SSkeletonAnimNotifies::CreateNotifiesList()
 {
+	const USkeleton* CurrentSkeleton = EditableSkeleton.IsValid() ? &EditableSkeleton->GetSkeleton() : nullptr;
+	FString CurrentSkeletonName;
+	if(CurrentSkeleton)
+	{
+		CurrentSkeletonName = FAssetData(CurrentSkeleton).GetExportTextName();
+	}
+
 	NotifyList.Empty();
+	FilteredNotifyList.Empty();
+
+	// Query the asset registry for our notifies and sync markers
+	const FAssetRegistryModule& AssetRegistryModule = FModuleManager::LoadModuleChecked<FAssetRegistryModule>("AssetRegistry");
 
-	const USkeleton& TargetSkeleton = EditableSkeleton->GetSkeleton();
+	FARFilter Filter;
+	Filter.bRecursiveClasses = true;
+	Filter.ClassPaths.Append({ UAnimationAsset::StaticClass()->GetClassPathName(), USkeleton::StaticClass()->GetClassPathName() } );
 
-	auto AddItem = [this, SearchText](FName InItemName, bool bInIsSyncMarker)
+	TArray<FAssetData> FoundAssetData;
+	AssetRegistryModule.Get().GetAssets(Filter, FoundAssetData);
+
+	// Build set of unique names
+	TMap<FName, EAnimNotifyFilterFlags> NotifyNames;
+	TMap<FName, EAnimNotifyFilterFlags> SyncMarkerNames;
+
+	for (const FAssetData& AssetData : FoundAssetData)
 	{
-		if ( !SearchText.IsEmpty() )
+		EAnimNotifyFilterFlags AssetFlags = EAnimNotifyFilterFlags::None;
+
+		if (AssetData.GetClass() != USkeleton::StaticClass())
 		{
-			if (InItemName.ToString().Contains( SearchText ) )
+			if (CurrentSkeleton && CurrentSkeleton->IsCompatibleForEditor(AssetData))
+			{
+				AssetFlags |= EAnimNotifyFilterFlags::CompatibleAssets;
+			}
+			else
 			{
-				NotifyList.Add( FDisplayedAnimNotifyInfo::Make(InItemName, bInIsSyncMarker) );
+				AssetFlags |= EAnimNotifyFilterFlags::OtherAssets;
 			}
 		}
 		else
 		{
-			NotifyList.Add( FDisplayedAnimNotifyInfo::Make(InItemName, bInIsSyncMarker) );
+			if (AssetData.GetExportTextName() == CurrentSkeletonName)
+			{
+				AssetFlags |= EAnimNotifyFilterFlags::CurrentSkeleton;
+			}
+			else if (CurrentSkeleton && CurrentSkeleton->IsCompatibleForEditor(AssetData))
+			{
+				AssetFlags |= EAnimNotifyFilterFlags::CompatibleAssets;
+			}
+		}
+
+		if(!EnumHasAnyFlags(CurrentFilterFlags, AssetFlags))
+		{
+			continue;
 		}
-	};
+
+		if (EnumHasAnyFlags(CurrentFilterFlags, EAnimNotifyFilterFlags::Notifies))
+		{
+			const FString TagValue = AssetData.GetTagValueRef<FString>(USkeleton::AnimNotifyTag);
+			if (!TagValue.IsEmpty())
+			{
+				UE::String::ParseTokens(TagValue, USkeleton::AnimNotifyTagDelimiter, 
+				 [&NotifyNames, AssetFlags](FStringView InToken)
+				{
+					EAnimNotifyFilterFlags& ItemFlags = NotifyNames.FindOrAdd(FName(InToken));
+					ItemFlags |= AssetFlags;
+				}, UE::String::EParseTokensOptions::SkipEmpty);
+			}
+		}
+
+		if (EnumHasAnyFlags(CurrentFilterFlags, EAnimNotifyFilterFlags::SyncMarkers))
+		{
+			const FString TagValue = AssetData.GetTagValueRef<FString>(USkeleton::AnimSyncMarkerTag);
+			if (!TagValue.IsEmpty())
+			{
+				UE::String::ParseTokens(TagValue, USkeleton::AnimSyncMarkerTagDelimiter, 
+				[&SyncMarkerNames, AssetFlags](FStringView InToken)
+				{
+					EAnimNotifyFilterFlags& ItemFlags = SyncMarkerNames.FindOrAdd(FName(InToken));
+					ItemFlags |= AssetFlags;
+				}, UE::String::EParseTokensOptions::SkipEmpty);
+			}
+		}
+	}
 
 	if(EnumHasAnyFlags(CurrentFilterFlags, EAnimNotifyFilterFlags::Notifies))
 	{
-		for(const FName& ItemName : TargetSkeleton.AnimationNotifies)
+		for(const TPair<FName, EAnimNotifyFilterFlags>& Item : NotifyNames)
 		{
-			AddItem(ItemName, false);
+			NotifyList.Add(FDisplayedAnimNotifyInfo::Make(Item.Key, false, Item.Value));
 		}
 	}
 
 	if(EnumHasAnyFlags(CurrentFilterFlags, EAnimNotifyFilterFlags::SyncMarkers))
 	{
-		for(const FName& ItemName : TargetSkeleton.GetExistingMarkerNames())
+		for(const TPair<FName, EAnimNotifyFilterFlags>& Item : SyncMarkerNames)
 		{
-			AddItem(ItemName, true);
+			NotifyList.Add(FDisplayedAnimNotifyInfo::Make(Item.Key, true, Item.Value));
 		}
 	}
-
-	NotifiesListView->RequestListRefresh();
 }
 
 void SSkeletonAnimNotifies::ShowNotifyInDetailsView(FName NotifyName)
 {
-	if(OnObjectsSelected.IsBound())
+	if(EditableSkeleton.IsValid() && OnObjectsSelected.IsBound())
 	{
 		ClearDetailsView();
 
diff --git a/Engine/Source/Editor/Persona/Private/SSkeletonAnimNotifies.h b/Engine/Source/Editor/Persona/Private/SSkeletonAnimNotifies.h
--- a/Engine/Source/Editor/Persona/Private/SSkeletonAnimNotifies.h
+++ b/Engine/Source/Editor/Persona/Private/SSkeletonAnimNotifies.h
@@ -47,6 +47,25 @@ private:
 	FOnObjectsSelected OnObjectsSelected;
 };
 
+// This is a flag that is used to filter UI part
+enum class EAnimNotifyFilterFlags : uint8 
+{
+	// Show none
+	None			= 0, 
+	// Show notifies
+	Notifies		= 0x01, 
+	// Show sync markers
+	SyncMarkers		= 0x02, 
+	// Show the current skeleton's items
+	CurrentSkeleton	= 0x04, 
+	// Show other compatible asset's items
+	CompatibleAssets = 0x08, 
+	// Show other asset's items
+	OtherAssets		= 0x10, 
+};
+
+ENUM_CLASS_FLAGS(EAnimNotifyFilterFlags);
+
 //////////////////////////////////////////////////////////////////////////
 // FDisplayedAnimNotifyInfo
 
@@ -64,37 +83,28 @@ public:
 	/** Identifies sync markers vs notifies */
 	bool bIsSyncMarker;
 
+	/** Flags for this item */
+	EAnimNotifyFilterFlags ItemFlags;
+
 	/** Static function for creating a new item, but ensures that you can only have a TSharedRef to one */
-	static TSharedRef<FDisplayedAnimNotifyInfo> Make(const FName& NotifyName, bool bInIsSyncMarker)
+	static TSharedRef<FDisplayedAnimNotifyInfo> Make(const FName& NotifyName, bool bInIsSyncMarker, EAnimNotifyFilterFlags InItemFlags)
 	{
-		return MakeShareable(new FDisplayedAnimNotifyInfo(NotifyName, bInIsSyncMarker));
+		return MakeShareable(new FDisplayedAnimNotifyInfo(NotifyName, bInIsSyncMarker, InItemFlags));
 	}
 
 protected:
 	/** Hidden constructor, always use Make above */
-	FDisplayedAnimNotifyInfo(const FName& InNotifyName, bool bInIsSyncMarker)
+	FDisplayedAnimNotifyInfo(const FName& InNotifyName, bool bInIsSyncMarker, EAnimNotifyFilterFlags InItemFlags)
 		: Name( InNotifyName )
 		, bIsNew(false)
 		, bIsSyncMarker(bInIsSyncMarker)
+		, ItemFlags(InItemFlags)
 	{}
 
 	/** Hidden constructor, always use Make above */
 	FDisplayedAnimNotifyInfo() {}
 };
 
-// This is a flag that is used to filter UI part
-enum class EAnimNotifyFilterFlags : uint8 
-{
-	// Show none
-	None			= 0, 
-	// Show notifies
-	Notifies		= 0x01, 
-	// Show sync markers
-	SyncMarkers		= 0x02, 
-};
-
-ENUM_CLASS_FLAGS(EAnimNotifyFilterFlags);
-
 /** Widgets list type */
 typedef SListView< TSharedPtr<FDisplayedAnimNotifyInfo> > SAnimNotifyListType;
 
@@ -103,8 +113,10 @@ class SSkeletonAnimNotifies : public SCompoundWidget, public FGCObject, public F
 public:
 	SLATE_BEGIN_ARGS( SSkeletonAnimNotifies )
 		: _IsPicker(false)
-		, _ShowSyncMarkers(false)
+		, _ShowSyncMarkers(true)
 		, _ShowNotifies(true)
+		, _ShowCompatibleSkeletonAssets(false)
+		, _ShowOtherAssets(false)
 	{}
 
 	/** Delegate called to select an object in the details panel */
@@ -125,9 +137,18 @@ public:
 	/** Whether we should show notifies */
 	SLATE_ARGUMENT(bool, ShowNotifies)
 
+	/** Whether we notifies and sync markers from assets compatible with the current skeleton */
+	SLATE_ARGUMENT(bool, ShowCompatibleSkeletonAssets)
+
+	/** Whether we notifies and sync markers from assets other than those compatible with the current skeleton */
+	SLATE_ARGUMENT(bool, ShowOtherAssets)
+
+	/** Editable skeleton - if this is set, notifies and sync markers will be added to the skeleton on creation */
+	SLATE_ARGUMENT(TSharedPtr<IEditableSkeleton>, EditableSkeleton)
+
 	SLATE_END_ARGS()
 public:
-	void Construct(const FArguments& InArgs, const TSharedRef<IEditableSkeleton>& InEditableSkeleton, const TSharedPtr<class FAssetEditorToolkit>& InHostingApp = nullptr);
+	void Construct(const FArguments& InArgs, const TSharedPtr<class FAssetEditorToolkit>& InHostingApp = nullptr);
 
 	~SSkeletonAnimNotifies();
 
@@ -183,6 +204,9 @@ private:
 	/** Delegate handler for determining whether we can show the delete menu options */
 	bool CanPerformDelete() const;
 	
+	/** Delegate handler for determining whether we can show the find references menu option */
+	bool CanPerformFindReferences() const;
+	
 	/** Delegate handler for deleting items */
 	void OnDeleteItems();
 
@@ -195,11 +219,14 @@ private:
 	/** Delegate handler for renaming items */
 	void OnRenameItem();
 
-	/** Wrapper that populates NotifiesListView using current filter test */
+	/** Wrapper that populates NotifiesListView using current filter text */
 	void RefreshNotifiesListWithFilter();
 
-	/** Populates NotifiesListView based on the skeletons notifies and the supplied filter text */
-	void CreateNotifiesList( const FString& SearchText = FString("") );
+	/** Populates NotifyList based on the project's notifies */
+	void CreateNotifiesList();
+
+	/** Populates FilteredNotifyList based on the project's notifies and the supplied filter text */
+	void FilterNotifiesList(const FString& InSearchText);
 
 	/** handler for user selecting a Notify in NotifiesListView - populates the details panel */
 	void ShowNotifyInDetailsView( FName NotifyName );
@@ -225,9 +252,12 @@ private:
 	/** Widget used to display the list of notifies */
 	TSharedPtr<SAnimNotifyListType> NotifiesListView;
 
-	/** A list of notifies. Used by the NotifiesListView. */
+	/** A list of all notifies and sync markers. */
 	TArray< TSharedPtr<FDisplayedAnimNotifyInfo> > NotifyList;
 
+	/** Filtered list of notifies. Used by the NotifiesListView. */
+	TArray< TSharedPtr<FDisplayedAnimNotifyInfo> > FilteredNotifyList;
+
 	/** Current text typed into NameFilterBox */
 	FText FilterText;
 
@@ -252,6 +282,15 @@ private:
 	/** Whether we should display notifies */
 	bool bShowNotifies = false;
 
+	/** Whether we should display notifies or sync markers from other skeletons */
+	bool bShowOtherSkeletonItems = false;
+
+	/** Whether we notifies and sync markers from assets other than the current skeleton */
+	bool bShowOtherAssets = false;
+	
+	/** Whether we notifies and sync markers from assets compatible with the current skeleton */
+	bool bShowCompatibleSkeletonAssets = false;
+
 	/** All filters that can be applied to the widget's display */
 	TArray<TSharedRef<FFilterBase<EAnimNotifyFilterFlags>>> Filters;
 
diff --git a/Engine/Source/Editor/SkeletonEditor/Private/EditableSkeleton.cpp b/Engine/Source/Editor/SkeletonEditor/Private/EditableSkeleton.cpp
--- a/Engine/Source/Editor/SkeletonEditor/Private/EditableSkeleton.cpp
+++ b/Engine/Source/Editor/SkeletonEditor/Private/EditableSkeleton.cpp
@@ -136,7 +136,6 @@ FEditableSkeleton::FEditableSkeleton(USkeleton* InSkeleton)
 	: Skeleton(InSkeleton)
 {
 	Skeleton->SetFlags(RF_Transactional);
-	Skeleton->CollectAnimationNotifies();
 }
 
 const USkeleton& FEditableSkeleton::GetSkeleton() const
@@ -930,7 +929,9 @@ int32 FEditableSkeleton::DeleteAnimNotifies(const TArray<FName>& InNotifyNames,
 		}
 	}
 
-	FBlueprintActionDatabase::Get().RefreshAssetActions(Skeleton);
+	FBlueprintActionDatabase& ActionDatabase = FBlueprintActionDatabase::Get();
+	ActionDatabase.ClearAssetActions(UAnimBlueprint::StaticClass());
+	ActionDatabase.RefreshClassActions(UAnimBlueprint::StaticClass());
 	OnNotifiesChanged.Broadcast();
 
 	return NumAnimationsModified;
@@ -979,7 +980,9 @@ void FEditableSkeleton::AddNotify(FName NewName)
 	Skeleton->Modify();
 	Skeleton->AddNewAnimationNotify(NewName);
 
-	FBlueprintActionDatabase::Get().RefreshAssetActions(Skeleton);
+	FBlueprintActionDatabase& ActionDatabase = FBlueprintActionDatabase::Get();
+	ActionDatabase.ClearAssetActions(UAnimBlueprint::StaticClass());
+	ActionDatabase.RefreshClassActions(UAnimBlueprint::StaticClass());
 	OnNotifiesChanged.Broadcast();
 }
 
@@ -1039,7 +1042,9 @@ int32 FEditableSkeleton::RenameNotify(const FName NewName, const FName OldName,
 			}
 		}
 
-		FBlueprintActionDatabase::Get().RefreshAssetActions(Skeleton);
+		FBlueprintActionDatabase& ActionDatabase = FBlueprintActionDatabase::Get();
+		ActionDatabase.ClearAssetActions(UAnimBlueprint::StaticClass());
+		ActionDatabase.RefreshClassActions(UAnimBlueprint::StaticClass());
 		OnNotifiesChanged.Broadcast();
 	}
 
diff --git a/Engine/Source/Runtime/Engine/Classes/Animation/AnimBlueprint.h b/Engine/Source/Runtime/Engine/Classes/Animation/AnimBlueprint.h
--- a/Engine/Source/Runtime/Engine/Classes/Animation/AnimBlueprint.h
+++ b/Engine/Source/Runtime/Engine/Classes/Animation/AnimBlueprint.h
@@ -190,7 +190,9 @@ class UAnimBlueprint : public UBlueprint, public IInterface_PreviewMeshProvider
 	ENGINE_API virtual bool SupportsInputEvents() const override;
 	ENGINE_API virtual bool AllowFunctionOverride(const UFunction* const InFunction) const override;
 	ENGINE_API virtual void PostEditChangeProperty(FPropertyChangedEvent& PropertyChangedEvent) override;
-	
+	ENGINE_API virtual void GetTypeActions(FBlueprintActionDatabaseRegistrar& ActionRegistrar) const override;
+	ENGINE_API virtual void GetInstanceActions(FBlueprintActionDatabaseRegistrar& ActionRegistrar) const override;
+
 protected:
 	// Broadcast when an override is changed, allowing derived blueprints to be updated
 	FOnOverrideChangedMulticaster OnOverrideChanged;
diff --git a/Engine/Source/Runtime/Engine/Classes/Animation/Skeleton.h b/Engine/Source/Runtime/Engine/Classes/Animation/Skeleton.h
--- a/Engine/Source/Runtime/Engine/Classes/Animation/Skeleton.h
+++ b/Engine/Source/Runtime/Engine/Classes/Animation/Skeleton.h
@@ -669,10 +669,7 @@ public:
 
 #if WITH_EDITORONLY_DATA
 
-	/*
-	 * Collect animation notifies that are referenced in all animations that use this skeleton (uses the asset registry).
-	 * Updates the cached AnimationNotifies array.
-	 */
+	UE_DEPRECATED(5.4, "Please do not use this function - notifies are stored collectively in the asset registry now rather than centrally on the skeleton")
 	ENGINE_API void CollectAnimationNotifies();
 
 	/*
diff --git a/Engine/Source/Runtime/Engine/Engine.Build.cs b/Engine/Source/Runtime/Engine/Engine.Build.cs
--- a/Engine/Source/Runtime/Engine/Engine.Build.cs
+++ b/Engine/Source/Runtime/Engine/Engine.Build.cs
@@ -293,6 +293,9 @@ public class Engine : ModuleRules
 			PublicIncludePathModuleNames.Add("AnimationDataController");
 			DynamicallyLoadedModuleNames.Add("AnimationDataController");
 
+			PublicIncludePathModuleNames.Add("AnimationBlueprintEditor");
+			DynamicallyLoadedModuleNames.Add("AnimationBlueprintEditor");
+
 			PrivateDependencyModuleNames.AddRange(
 				new string[] {
 					"RawMesh"
diff --git a/Engine/Source/Runtime/Engine/Private/AnimBlueprint.cpp b/Engine/Source/Runtime/Engine/Private/AnimBlueprint.cpp
--- a/Engine/Source/Runtime/Engine/Private/AnimBlueprint.cpp
+++ b/Engine/Source/Runtime/Engine/Private/AnimBlueprint.cpp
@@ -8,6 +8,7 @@
 #include UE_INLINE_GENERATED_CPP_BY_NAME(AnimBlueprint)
 
 #if WITH_EDITOR
+#include "IAnimationBlueprintEditorModule.h"
 #include "Settings/AnimBlueprintSettings.h"
 #endif
 #if WITH_EDITORONLY_DATA
@@ -290,6 +291,22 @@ void UAnimBlueprint::PostEditChangeProperty(struct FPropertyChangedEvent& Proper
 	}
 }
 
+void UAnimBlueprint::GetTypeActions(FBlueprintActionDatabaseRegistrar& ActionRegistrar) const
+{
+	Super::GetTypeActions(ActionRegistrar);
+
+	const IAnimationBlueprintEditorModule& AnimationBlueprintEditorModule = FModuleManager::LoadModuleChecked<IAnimationBlueprintEditorModule>(TEXT("AnimationBlueprintEditor"));
+	AnimationBlueprintEditorModule.GetTypeActions(ActionRegistrar);
+}
+
+void UAnimBlueprint::GetInstanceActions(FBlueprintActionDatabaseRegistrar& ActionRegistrar) const
+{
+	Super::GetInstanceActions(ActionRegistrar);
+
+	const IAnimationBlueprintEditorModule& AnimationBlueprintEditorModule = FModuleManager::LoadModuleChecked<IAnimationBlueprintEditorModule>(TEXT("AnimationBlueprintEditor"));
+	AnimationBlueprintEditorModule.GetInstanceActions(this, ActionRegistrar);
+}
+
 #endif
 
 USkeletalMesh* UAnimBlueprint::GetPreviewMesh(bool bFindIfNotSet/*=false*/)
