Added a Project Render Setting to configure MorphTargets maximum blend weight, by default with the value that was previously hardcoded (5.0)
Updated Rendering and Animation to use the configured value.
Also added a cvar that enables testing different values from the console    r.MorphTarget.MaxBlendWeight

diff --git a/Engine/Source/Editor/Persona/Private/SMorphTargetViewer.cpp b/Engine/Source/Editor/Persona/Private/SMorphTargetViewer.cpp
--- a/Engine/Source/Editor/Persona/Private/SMorphTargetViewer.cpp
+++ b/Engine/Source/Editor/Persona/Private/SMorphTargetViewer.cpp
@@ -18,6 +18,7 @@
 #include "IPersonaPreviewScene.h"
 #include "SkeletalMeshAttributes.h"
 #include "Rendering/SkeletalMeshLODImporterData.h"
+#include "SkeletalRenderPublic.h"
 
 #define LOCTEXT_NAMESPACE "SMorphTargetViewer"
 
@@ -26,7 +27,6 @@ static const FName ColumnID_MorphTargetWeightLabel( "Weight" );
 static const FName ColumnID_MorphTargetEditLabel( "Edit" );
 static const FName ColumnID_MorphTargetVertCountLabel( "NumberOfVerts" );
 
-const float MaxMorphWeight = 5.0f;
 
 //////////////////////////////////////////////////////////////////////////
 // SMorphTargetListRow
@@ -152,8 +152,6 @@ TSharedRef< SWidget > SMorphTargetListRow::GenerateWidgetForColumn( const FName&
 				SNew( SSpinBox<float> )
 				.MinSliderValue(-1.f)
 				.MaxSliderValue(1.f)
-				.MinValue(-MaxMorphWeight)
-				.MaxValue(MaxMorphWeight)
 				.Value( this, &SMorphTargetListRow::GetWeight )
 				.OnValueChanged( this, &SMorphTargetListRow::OnMorphTargetWeightChanged )
 				.OnValueCommitted( this, &SMorphTargetListRow::OnMorphTargetWeightValueCommitted )
@@ -227,6 +225,10 @@ void SMorphTargetListRow::OnMorphTargetWeightChanged( float NewWeight )
 	// First change this item...
 	// the delta feature is a bit confusing when debugging morphtargets, and you're not sure why it's changing, so I'm disabling it for now. 
 	// I think in practice, you want each morph target to move independentaly. It is very unlikely you'd like to move multiple things together. 
+
+	const float MorphTargetMaxBlendWeight = UE::SkeletalRender::Settings::GetMorphTargetMaxBlendWeight();
+	NewWeight = FMath::Clamp(NewWeight, -MorphTargetMaxBlendWeight, MorphTargetMaxBlendWeight);
+
 #if 0 
 	float Delta = NewWeight - GetWeight();
 #endif
@@ -247,7 +249,7 @@ void SMorphTargetListRow::OnMorphTargetWeightChanged( float NewWeight )
 
 		if ( RowItem != Item ) // Don't do "this" row again if it's selected
 		{
-			RowItem->Weight = FMath::Clamp(RowItem->Weight + Delta, -MaxMorphWeight, MaxMorphWeight);
+			RowItem->Weight = FMath::Clamp(RowItem->Weight + Delta, -MorphTargetMaxBlendWeight, MorphTargetMaxBlendWeight);
 			RowItem->bAutoFillData = false;
 			MorphTargetViewer->AddMorphTargetOverride( RowItem->Name, RowItem->Weight, false );
 		}
@@ -259,8 +261,10 @@ void SMorphTargetListRow::OnMorphTargetWeightValueCommitted( float NewWeight, ET
 {
 	if (CommitType == ETextCommit::OnEnter || CommitType == ETextCommit::OnUserMovedFocus)
 	{
-		float NewValidWeight = FMath::Clamp(NewWeight, -MaxMorphWeight, MaxMorphWeight);
-		Item->Weight = NewValidWeight;
+		const float MorphTargetMaxBlendWeight = UE::SkeletalRender::Settings::GetMorphTargetMaxBlendWeight();
+		NewWeight = FMath::Clamp(NewWeight, -MorphTargetMaxBlendWeight, MorphTargetMaxBlendWeight);
+
+		Item->Weight = NewWeight;
 		Item->bAutoFillData = false;
 
 		MorphTargetViewer->AddMorphTargetOverride(Item->Name, Item->Weight, false);
@@ -274,7 +278,7 @@ void SMorphTargetListRow::OnMorphTargetWeightValueCommitted( float NewWeight, ET
 
 			if(RowItem != Item) // Don't do "this" row again if it's selected
 			{
-				RowItem->Weight = NewValidWeight;
+				RowItem->Weight = NewWeight;
 				RowItem->bAutoFillData = false;
 				MorphTargetViewer->AddMorphTargetOverride(RowItem->Name, RowItem->Weight, false);
 			}
diff --git a/Engine/Source/Runtime/Engine/Classes/Engine/RendererSettings.h b/Engine/Source/Runtime/Engine/Classes/Engine/RendererSettings.h
--- a/Engine/Source/Runtime/Engine/Classes/Engine/RendererSettings.h
+++ b/Engine/Source/Runtime/Engine/Classes/Engine/RendererSettings.h
@@ -846,6 +846,11 @@ class URendererSettings : public UDeveloperSettings
 		ToolTip = "Whether to use original CPU method (loop per morph then by vertex) or use a GPU-based method on Shader Model 5 hardware."))
 	uint32 bUseGPUMorphTargets : 1;
 
+	UPROPERTY(config, EditAnywhere, Category = DefaultSettings, meta = (
+		ConsoleVariable = "r.MorphTarget.MaxBlendWeight", DisplayName = "Maximum absolute value accepted as a morph target blend weight, positive or negative.",
+		ToolTip = "Blend target weights will be checked against this value for validation. Absolue values greather than this number will be clamped to [-MorphTargetMaxBlendWeight, MorphTargetMaxBlendWeight]."))
+	float MorphTargetMaxBlendWeight;
+
 	/**
 	"The sky atmosphere component requires extra samplers/textures to be bound to apply aerial perspective on transparent surfaces (and all surfaces on mobile via per vertex evaluation)."
 	*/
diff --git a/Engine/Source/Runtime/Engine/Private/Animation/AnimationRuntime.cpp b/Engine/Source/Runtime/Engine/Private/Animation/AnimationRuntime.cpp
--- a/Engine/Source/Runtime/Engine/Private/Animation/AnimationRuntime.cpp
+++ b/Engine/Source/Runtime/Engine/Private/Animation/AnimationRuntime.cpp
@@ -12,6 +12,7 @@
 #include "BonePose.h"
 #include "Engine/SkeletalMesh.h"
 #include "SkeletalRender.h"
+#include "SkeletalRenderPublic.h"
 #include "Animation/AttributesRuntime.h"
 #include "Animation/AnimationPoseData.h"
 #include "Animation/BlendProfile.h"
@@ -2594,6 +2595,8 @@ void FAnimationRuntime::AppendActiveMorphTargets(
 		return;
 	}
 	
+	const float MorphTargetMaxBlendWeight = UE::SkeletalRender::Settings::GetMorphTargetMaxBlendWeight();
+
 	// Then go over the CurveKeys finding morph targets by name
 	for(const TPair<FName, float>& MorphCurveAnim : MorphCurveAnims)
 	{
@@ -2611,17 +2614,18 @@ void FAnimationRuntime::AppendActiveMorphTargets(
 			// If it has a valid weight
 			if (FMath::Abs(Weight) > MinMorphTargetBlendWeight)
 			{
+				const float ClampedWeight = FMath::Clamp(Weight, -MorphTargetMaxBlendWeight, MorphTargetMaxBlendWeight);
 				// If not, add it
 				if (FoundMorphIndex == nullptr)
 				{
 					InOutActiveMorphTargets.Add(Target, SkeletalMorphIndex);
-					InOutMorphTargetWeights[SkeletalMorphIndex] = Weight;
+					InOutMorphTargetWeights[SkeletalMorphIndex] = ClampedWeight;
 				}
 				else
 				{
 					// If it does, use the max weight
 					check(SkeletalMorphIndex == *FoundMorphIndex);
-					InOutMorphTargetWeights[SkeletalMorphIndex] = Weight;
+					InOutMorphTargetWeights[SkeletalMorphIndex] = ClampedWeight;
 				}
 			}
 			else if (FoundMorphIndex != nullptr)
diff --git a/Engine/Source/Runtime/Engine/Private/RendererSettings.cpp b/Engine/Source/Runtime/Engine/Private/RendererSettings.cpp
--- a/Engine/Source/Runtime/Engine/Private/RendererSettings.cpp
+++ b/Engine/Source/Runtime/Engine/Private/RendererSettings.cpp
@@ -80,6 +80,7 @@ URendererSettings::URendererSettings(const FObjectInitializer& ObjectInitializer
 	TranslucentSortAxis = FVector(0.0f, -1.0f, 0.0f);
 	bSupportStationarySkylight = true;
 	bSupportPointLightWholeSceneShadows = true;
+	MorphTargetMaxBlendWeight = 5.f;
 	bSupportSkyAtmosphere = true;
 	bSupportSkinCacheShaders = false;
 	bSkipCompilingGPUSkinVF = false;
diff --git a/Engine/Source/Runtime/Engine/Private/SkeletalRender.cpp b/Engine/Source/Runtime/Engine/Private/SkeletalRender.cpp
--- a/Engine/Source/Runtime/Engine/Private/SkeletalRender.cpp
+++ b/Engine/Source/Runtime/Engine/Private/SkeletalRender.cpp
@@ -12,6 +12,8 @@
 #include "SceneView.h"
 #include "GPUSkinCache.h"
 #include "Rendering/SkeletalMeshRenderData.h"
+#include "Engine/RendererSettings.h"
+
 
 /*-----------------------------------------------------------------------------
 Globals
@@ -22,6 +24,25 @@ const float MinMorphTargetBlendWeight = UE_SMALL_NUMBER;
 // largest blend weight for vertex anims
 const float MaxMorphTargetBlendWeight = 5.0f;
 
+static float GMorphTargetMaxBlendWeight = 5.f;
+static FAutoConsoleVariableRef CVarMorphTargetMinBlendWeight(
+	TEXT("r.MorphTarget.MaxBlendWeight"),
+	GMorphTargetMaxBlendWeight,
+	TEXT("Maximum value accepted as a morph target blend weight..\n")
+	TEXT("Blend target weights will be checked against this value for validation.Values smaller than this number will be clamped.\n"),
+	ECVF_Default
+);
+
+namespace UE::SkeletalRender::Settings
+{
+	float GetMorphTargetMaxBlendWeight()
+	{
+		static const IConsoleVariable* MorphTargetMaxBlendWeightCVar = IConsoleManager::Get().FindConsoleVariable(TEXT("r.MorphTarget.MaxBlendWeight"));
+
+		return (MorphTargetMaxBlendWeightCVar != nullptr) ? MorphTargetMaxBlendWeightCVar->GetFloat() : GetDefault<URendererSettings>()->MorphTargetMaxBlendWeight;
+	}
+}
+
 #if RHI_RAYTRACING
 static bool IsSkeletalMeshRayTracingSupported()
 {
diff --git a/Engine/Source/Runtime/Engine/Private/SkeletalRender.h b/Engine/Source/Runtime/Engine/Private/SkeletalRender.h
--- a/Engine/Source/Runtime/Engine/Private/SkeletalRender.h
+++ b/Engine/Source/Runtime/Engine/Private/SkeletalRender.h
@@ -15,6 +15,7 @@ class FSkeletalMeshRenderData;
 // smallest blend weight for vertex anims
 extern const float MinMorphTargetBlendWeight;
 // largest blend weight for vertex anims
+UE_DEPRECATED(5.4, "Please, use UE::SkeletalRender::Settings::GetMorphTargetMaxBlendWeight");
 extern const float MaxMorphTargetBlendWeight;
 
 
diff --git a/Engine/Source/Runtime/Engine/Private/SkeletalRenderCPUSkin.cpp b/Engine/Source/Runtime/Engine/Private/SkeletalRenderCPUSkin.cpp
--- a/Engine/Source/Runtime/Engine/Private/SkeletalRenderCPUSkin.cpp
+++ b/Engine/Source/Runtime/Engine/Private/SkeletalRenderCPUSkin.cpp
@@ -570,6 +570,8 @@ static uint32 InitEvalInfos(const FMorphTargetWeightMap& InActiveMorphTargets, c
 {
 	uint32 NumValidMorphTargets=0;
 
+	const float MorphTargetMaxBlendWeight = UE::SkeletalRender::Settings::GetMorphTargetMaxBlendWeight();
+
 	for(const TTuple<const UMorphTarget*, int32>& MorphItem: InActiveMorphTargets)
 	{
 		FMorphTargetInfo NewInfo;
@@ -580,7 +582,7 @@ static uint32 InitEvalInfos(const FMorphTargetWeightMap& InActiveMorphTargets, c
 
 		if( MorphTarget != nullptr &&
 			ActiveMorphAbsVertexWeight >= MinMorphTargetBlendWeight &&
-			ActiveMorphAbsVertexWeight <= MaxMorphTargetBlendWeight &&
+			ActiveMorphAbsVertexWeight <= MorphTargetMaxBlendWeight &&
 			MorphItem.Key->HasDataForLOD(LODIndex) )
 		{
 			// start at the first vertex since they affect base mesh verts in ascending order
diff --git a/Engine/Source/Runtime/Engine/Private/SkeletalRenderGPUSkin.cpp b/Engine/Source/Runtime/Engine/Private/SkeletalRenderGPUSkin.cpp
--- a/Engine/Source/Runtime/Engine/Private/SkeletalRenderGPUSkin.cpp
+++ b/Engine/Source/Runtime/Engine/Private/SkeletalRenderGPUSkin.cpp
@@ -1404,6 +1404,8 @@ void FSkeletalMeshObjectGPUSkin::FSkeletalMeshObjectLOD::UpdateMorphVertexBuffer
 		{
 			SCOPE_CYCLE_COUNTER(STAT_MorphVertexBuffer_ApplyDelta);
 
+			const float MorphTargetMaxBlendWeight = UE::SkeletalRender::Settings::GetMorphTargetMaxBlendWeight();
+
 			// iterate over all active morph targets and accumulate their vertex deltas
 			for(const TTuple<const UMorphTarget*, int32>& MorphItem: InActiveMorphTargets)
 			{
@@ -1413,8 +1415,7 @@ void FSkeletalMeshObjectGPUSkin::FSkeletalMeshObjectLOD::UpdateMorphVertexBuffer
 				checkSlow(MorphTarget->HasDataForLOD(LODIndex));
 				const float MorphTargetWeight = MorphTargetWeights.IsValidIndex(WeightIndex) ? MorphTargetWeights[WeightIndex] : 0.0f;
 				const float MorphAbsWeight = FMath::Abs(MorphTargetWeight);
-				checkSlow(MorphAbsWeight >= MinMorphTargetBlendWeight && MorphAbsWeight <= MaxMorphTargetBlendWeight);
-
+				checkSlow(MorphAbsWeight >= MinMorphTargetBlendWeight && MorphAbsWeight <= MorphTargetMaxBlendWeight);
 
 				// Get deltas
 				int32 NumDeltas;
@@ -2550,6 +2551,8 @@ void FDynamicSkelMeshObjectDataGPUSkin::InitDynamicSkelMeshObjectDataGPUSkin(
 		}
 	}
 
+	const float MorphTargetMaxBlendWeight = UE::SkeletalRender::Settings::GetMorphTargetMaxBlendWeight();
+
 	// find number of morphs that are currently weighted and will affect the mesh
 	ActiveMorphTargets.Reserve(InActiveMorphTargets.Num());
 	for(const TTuple<const UMorphTarget*, int32>& MorphItem: InActiveMorphTargets)
@@ -2561,7 +2564,7 @@ void FDynamicSkelMeshObjectDataGPUSkin::InitDynamicSkelMeshObjectDataGPUSkin(
 
 		if( MorphTarget != nullptr && 
 			MorphAbsWeight >= MinMorphTargetBlendWeight &&
-			MorphAbsWeight <= MaxMorphTargetBlendWeight &&
+			MorphAbsWeight <= MorphTargetMaxBlendWeight &&
 			MorphTarget->HasDataForLOD(LODIndex) ) 
 		{
 			NumWeightedActiveMorphTargets++;
diff --git a/Engine/Source/Runtime/Engine/Public/SkeletalRenderPublic.h b/Engine/Source/Runtime/Engine/Public/SkeletalRenderPublic.h
--- a/Engine/Source/Runtime/Engine/Public/SkeletalRenderPublic.h
+++ b/Engine/Source/Runtime/Engine/Public/SkeletalRenderPublic.h
@@ -31,6 +31,12 @@ class FMeshDeformerGeometry;
 class FRayTracingGeometry;
 class FRHICommandList;
 
+namespace UE::SkeletalRender::Settings
+{
+	// Returns the maximum value allowed for morph targets blend weights, configured at RenderSettings
+	ENGINE_API float GetMorphTargetMaxBlendWeight();
+}
+
 /** data for a single skinned skeletal mesh vertex */
 struct FFinalSkinVertex
 {
