Add editor-only option to enable/disable each sample in the "analyse all"

diff --git a/Engine/Source/Editor/Persona/Private/Customization/BlendSampleDetails.cpp b/Engine/Source/Editor/Persona/Private/Customization/BlendSampleDetails.cpp
--- a/Engine/Source/Editor/Persona/Private/Customization/BlendSampleDetails.cpp
+++ b/Engine/Source/Editor/Persona/Private/Customization/BlendSampleDetails.cpp
@@ -193,6 +193,10 @@ void FBlendSampleDetails::CustomizeDetails(class IDetailLayoutBuilder& DetailBui
 			GET_MEMBER_NAME_CHECKED(FBlendSample, RateScale), FBlendSample::StaticStruct());
 		CategoryBuilder.AddProperty(RateScaleProperty);
 
+		TSharedPtr<IPropertyHandle> IncludeInAnalyseAllProperty = DetailBuilder.GetProperty(
+			GET_MEMBER_NAME_CHECKED(FBlendSample, bIncludeInAnalyseAll), FBlendSample::StaticStruct());
+		CategoryBuilder.AddProperty(IncludeInAnalyseAllProperty);
+
 		bool bShowAnalysis = false;
 		bool bAnalyzed[3] = { false, false, false };
 		// Note that the analyzed position won't change whilst this menu is open, but the original value might.
diff --git a/Engine/Source/Editor/Persona/Private/Customization/BlendSpaceDetails.cpp b/Engine/Source/Editor/Persona/Private/Customization/BlendSpaceDetails.cpp
--- a/Engine/Source/Editor/Persona/Private/Customization/BlendSpaceDetails.cpp
+++ b/Engine/Source/Editor/Persona/Private/Customization/BlendSpaceDetails.cpp
@@ -68,17 +68,20 @@ FReply FBlendSpaceDetails::HandleAnalyzeSamples()
 		for (int32 SampleIndex = 0 ; SampleIndex != NumSamples ; ++SampleIndex)
 		{
 			bool bAnalyzed[3] = { false, false, false };
-			FVector NewValue = BlendSpaceAnalysis::CalculateSampleValue(
-				*BlendSpace, *BlendSpace->SampleData[SampleIndex].Animation, 
-				BlendSpace->SampleData[SampleIndex].RateScale, 
-				BlendSpace->SampleData[SampleIndex].SampleValue, bAnalyzed);
-			if (bAnalyzed[0] || bAnalyzed[1] || bAnalyzed[2])
+			if (BlendSpace->SampleData[SampleIndex].bIncludeInAnalyseAll)
 			{
-				BlendSpace->EditSampleValue(SampleIndex, NewValue);
-				// Note that the sample might not move if the destination position is in use
-				if (NewValue == BlendSpace->SampleData[SampleIndex].SampleValue)
+				FVector NewValue = BlendSpaceAnalysis::CalculateSampleValue(
+					*BlendSpace, *BlendSpace->SampleData[SampleIndex].Animation, 
+					BlendSpace->SampleData[SampleIndex].RateScale, 
+					BlendSpace->SampleData[SampleIndex].SampleValue, bAnalyzed);
+				if (bAnalyzed[0] || bAnalyzed[1] || bAnalyzed[2])
 				{
-					bChangedOne = true;
+					BlendSpace->EditSampleValue(SampleIndex, NewValue);
+					// Note that the sample might not move if the destination position is in use
+					if (NewValue == BlendSpace->SampleData[SampleIndex].SampleValue)
+					{
+						bChangedOne = true;
+					}
 				}
 			}
 		}
@@ -447,6 +450,7 @@ void FBlendSpaceDetails::CustomizeDetails(class IDetailLayoutBuilder& DetailBuil
 				TSharedPtr<IPropertyHandle> AnimationProperty = BlendSampleProperty->GetChildHandle(GET_MEMBER_NAME_CHECKED(FBlendSample, Animation));
 				TSharedPtr<IPropertyHandle> SampleValueProperty = BlendSampleProperty->GetChildHandle(GET_MEMBER_NAME_CHECKED(FBlendSample, SampleValue));
 				TSharedPtr<IPropertyHandle> RateScaleProperty = BlendSampleProperty->GetChildHandle(GET_MEMBER_NAME_CHECKED(FBlendSample, RateScale));
+				TSharedPtr<IPropertyHandle> IncludeInAnalyseAllProperty = BlendSampleProperty->GetChildHandle(GET_MEMBER_NAME_CHECKED(FBlendSample, bIncludeInAnalyseAll));
 
 				IDetailGroup& Group = DetailCategoryBuilder.AddGroup(FName("BlendSamples_Samples"), FText::GetEmpty());
 				Group.HeaderRow()
@@ -506,6 +510,7 @@ void FBlendSpaceDetails::CustomizeDetails(class IDetailLayoutBuilder& DetailBuil
 					IDetailPropertyRow& AnimationRow = Group.AddPropertyRow(AnimationProperty.ToSharedRef());
 					FBlendSampleDetails::GenerateAnimationWidget(AnimationRow, BlendSpace, AnimationProperty);
 					Group.AddPropertyRow(RateScaleProperty.ToSharedRef());
+					Group.AddPropertyRow(IncludeInAnalyseAllProperty.ToSharedRef());
 				}
 				else if(BlendSpaceNode.Get())
 				{
diff --git a/Engine/Source/Runtime/Engine/Classes/Animation/BlendSpace.h b/Engine/Source/Runtime/Engine/Classes/Animation/BlendSpace.h
--- a/Engine/Source/Runtime/Engine/Classes/Animation/BlendSpace.h
+++ b/Engine/Source/Runtime/Engine/Classes/Animation/BlendSpace.h
@@ -178,6 +178,11 @@ struct FBlendSample
 	float RateScale = 1.0f;
 
 #if WITH_EDITORONLY_DATA
+	// Whether or not this sample will be moved when the "analyse all" button is used. Note that, even if disabled,
+	// it will still be available for individual sample analysis/moving
+	UPROPERTY(EditAnywhere, Category = BlendSample, meta=(UIMin="0.01", UIMax="2.0", ClampMin="0.01", ClampMax="64.0"))
+	uint8 bIncludeInAnalyseAll : 1;
+
 	UPROPERTY(transient)
 	uint8 bIsValid : 1;
 
@@ -191,6 +196,7 @@ struct FBlendSample
 		, SampleValue(0.f)
 		, RateScale(1.0f)
 #if WITH_EDITORONLY_DATA
+		, bIncludeInAnalyseAll(true)
 		, bIsValid(false)
 		, CachedMarkerDataUpdateCounter(INDEX_NONE)
 #endif // WITH_EDITORONLY_DATA
