Makes the action buttons in the BlendSpace details appear when searching

diff --git a/Engine/Source/Editor/Persona/Private/Customization/BlendSpaceDetails.cpp b/Engine/Source/Editor/Persona/Private/Customization/BlendSpaceDetails.cpp
--- a/Engine/Source/Editor/Persona/Private/Customization/BlendSpaceDetails.cpp
+++ b/Engine/Source/Editor/Persona/Private/Customization/BlendSpaceDetails.cpp
@@ -269,18 +269,19 @@ void FBlendSpaceDetails::CustomizeDetails(class IDetailLayoutBuilder& DetailBuil
 		if (BlendSpace->IsAsset())
 		{
 			IDetailCategoryBuilder& DetailCategoryBuilder = DetailBuilder.EditCategory(FName("Analysis"));
-			TSharedPtr<IPropertyHandle> AnalysisPropertiesHandle = DetailBuilder.GetProperty(GET_MEMBER_NAME_CHECKED(UBlendSpace, AnalysisProperties), UBlendSpace::StaticClass());
+			TSharedPtr<IPropertyHandle> AnalysisPropertiesHandle = DetailBuilder.GetProperty(
+				GET_MEMBER_NAME_CHECKED(UBlendSpace, AnalysisProperties), UBlendSpace::StaticClass());
 
 			int32 HideIndex = b1DBlendSpace ? 1 : 2;
 
 			// Re-analyse button
 			if (NumBlendSampleEntries)
 			{
-				DetailCategoryBuilder
-					.AddGroup(FName("BlendSamples_RemoveSamples"), FText::GetEmpty())
-					.HeaderRow()
+				// Use AddCustomRow rather than AddGroup because the latter fails to include the button when searching
+				DetailCategoryBuilder.AddCustomRow(
+					LOCTEXT("AnalyzeAllSamplesRow", "Analyse all samples"))
 					.Visibility(TAttribute<EVisibility>::Create([AnalysisPropertiesHandle, HideIndex]() {
-						return GetAnalyzeButtonVisibility(AnalysisPropertiesHandle, HideIndex);}))
+						return GetAnalyzeButtonVisibility(AnalysisPropertiesHandle, HideIndex); }))
 					.NameContent()
 					[
 						SNew(STextBlock)
@@ -419,9 +420,8 @@ void FBlendSpaceDetails::CustomizeDetails(class IDetailLayoutBuilder& DetailBuil
 			// Add a "Remove all" button if there are some samples. Only in the asset blendspace for now.
 			if (NumBlendSampleEntries && BlendSpace->IsAsset())
 			{
-				DetailCategoryBuilder
-					.AddGroup(FName("BlendSamples_RemoveSamples"), FText::GetEmpty())
-					.HeaderRow()
+				DetailCategoryBuilder.AddCustomRow(
+					LOCTEXT("RemoveSamplesRow", "Remove all samples"))
 					.NameContent()
 					[
 						SNew(STextBlock)
@@ -488,7 +488,10 @@ void FBlendSpaceDetails::CustomizeDetails(class IDetailLayoutBuilder& DetailBuil
 				];
 
 				FBlendSampleDetails::GenerateBlendSampleWidget(
-					[&Group, &SampleValueProperty]() -> IDetailPropertyRow& { return Group.AddPropertyRow(SampleValueProperty.ToSharedRef()); }, FOnSampleMoved::CreateLambda(
+					[&Group, &SampleValueProperty]() -> IDetailPropertyRow& 
+					{ 
+						return Group.AddPropertyRow(SampleValueProperty.ToSharedRef()); 
+					}, FOnSampleMoved::CreateLambda(
 						[this](const uint32 Index, const FVector& SampleValue, bool bIsInteractive) 
 				{
 					if (BlendSpace->IsValidBlendSampleIndex(Index) && BlendSpace->GetBlendSample(Index).SampleValue !=
@@ -499,7 +502,8 @@ void FBlendSpaceDetails::CustomizeDetails(class IDetailLayoutBuilder& DetailBuil
 						if (bMoveSuccesful)
 						{
 							BlendSpace->ValidateSampleData();
-							FPropertyChangedEvent ChangedEvent(nullptr, bIsInteractive ? EPropertyChangeType::Interactive : EPropertyChangeType::ValueSet);
+							FPropertyChangedEvent ChangedEvent(
+								nullptr, bIsInteractive ? EPropertyChangeType::Interactive : EPropertyChangeType::ValueSet);
 							BlendSpace->PostEditChangeProperty(ChangedEvent);
 						}
 					}
