Persona: Only force a full layout rebuild when LOD counts change. In all other cases, just refresh the layout.

diff --git a/Engine/Source/Editor/Persona/Private/LODInfoUILayout.cpp b/Engine/Source/Editor/Persona/Private/LODInfoUILayout.cpp
--- a/Engine/Source/Editor/Persona/Private/LODInfoUILayout.cpp
+++ b/Engine/Source/Editor/Persona/Private/LODInfoUILayout.cpp
@@ -28,12 +28,15 @@ void ULODInfoUILayout::SetReferenceLODInfo(TWeakPtr<IPersonaToolkit> InPersonaTo
 void ULODInfoUILayout::PostEditChangeProperty(struct FPropertyChangedEvent& PropertyChangedEvent)
 {
 	Super::PostEditChangeProperty(PropertyChangedEvent);
-	
-	USkeletalMesh* SkeletalMesh = PersonaToolkit.Pin()->GetPreviewMesh();
-	FSkeletalMeshLODInfo* SkeletalMeshLODInfo = SkeletalMesh->GetLODInfo(LODIndex);
-	check(SkeletalMeshLODInfo != nullptr);
-	SkeletalMesh->Modify();
-	//Copy the LODInfo into the real skeletal mesh LODInfo data
-	*SkeletalMeshLODInfo = LODInfo;
-	FScopedSkeletalMeshPostEditChange ScopeSkeletalmeshPostEdit(SkeletalMesh);
+
+	if ((PropertyChangedEvent.ChangeType & EPropertyChangeType::Interactive) == 0)
+	{
+		USkeletalMesh* SkeletalMesh = PersonaToolkit.Pin()->GetPreviewMesh();
+		FSkeletalMeshLODInfo* SkeletalMeshLODInfo = SkeletalMesh->GetLODInfo(LODIndex);
+		check(SkeletalMeshLODInfo != nullptr);
+		SkeletalMesh->Modify();
+		//Copy the LODInfo into the real skeletal mesh LODInfo data
+		*SkeletalMeshLODInfo = LODInfo;
+		FScopedSkeletalMeshPostEditChange ScopeSkeletalmeshPostEdit(SkeletalMesh);
+	}
 }
diff --git a/Engine/Source/Editor/Persona/Private/PersonaMeshDetails.cpp b/Engine/Source/Editor/Persona/Private/PersonaMeshDetails.cpp
--- a/Engine/Source/Editor/Persona/Private/PersonaMeshDetails.cpp
+++ b/Engine/Source/Editor/Persona/Private/PersonaMeshDetails.cpp
@@ -80,6 +80,7 @@
 #include "Factories/FbxSkeletalMeshImportData.h"
 #include "PropertyCustomizationHelpers.h"
 #include "ComponentReregisterContext.h"
+#include "IPropertyUtilities.h"
 #include "LODInfoUILayout.h"
 #include "SWarningOrErrorBox.h"
 #include "Interfaces/ITargetPlatform.h"
@@ -595,7 +596,7 @@ void FSkeletalMeshReductionSettingsLayout::GenerateChildContent(IDetailChildrenB
 			LOCTEXT("MaxBonesPerVertex_ToolTip", "Maximum number of bones that can be assigned to each vertex."),
 			TEXT("MaxBonesPerVertex"),
 			1,
-			INT_MAX,
+			MAX_TOTAL_INFLUENCES,
 			FGetIntegerDelegate::CreateRaw(this, &FSkeletalMeshReductionSettingsLayout::GetMaxBonesPerVertex),
 			FSetIntegerDelegate::CreateRaw(this, &FSkeletalMeshReductionSettingsLayout::SetMaxBonesPerVertex));
 	}
@@ -3236,7 +3237,7 @@ TSharedRef<SWidget> FPersonaMeshDetails::CreateSkinWeightProfileMenuContent()
 			SkeletalMesh->Modify();
 
 			FSkinWeightProfileHelpers::ImportSkinWeightProfile(SkeletalMesh);
-			RefreshMeshDetailLayout();
+			RequestLayoutUpdate();
 		}
 	})));
 	
@@ -3290,7 +3291,7 @@ TSharedRef<SWidget> FPersonaMeshDetails::CreateSkinWeightProfileMenuContent()
 												SkeletalMesh->Modify();
 
 												FSkinWeightProfileHelpers::ImportSkinWeightProfileLOD(SkeletalMesh, ProfileName, LODIndex);
-												RefreshMeshDetailLayout();
+												RequestLayoutUpdate();
 											}
 										})));
 									}
@@ -3631,7 +3632,7 @@ void FPersonaMeshDetails::OnAssetPostLODImported(UObject* InObject, int32 InLODI
 {
 	if (InObject == GetPersonaToolkit()->GetMesh())
 	{
-		RefreshMeshDetailLayout();
+		RequestLayoutUpdate();
 	}
 }
 
@@ -3639,7 +3640,7 @@ void FPersonaMeshDetails::OnAssetReimport(UObject* InObject)
 {
 	if (InObject == GetPersonaToolkit()->GetMesh())
 	{
-		RefreshMeshDetailLayout();
+		RequestLayoutUpdate();
 	}
 }
 
@@ -3672,8 +3673,10 @@ void FPersonaMeshDetails::OnLODCountCommitted(int32 InValue, ETextCommit::Type C
 	OnLODCountChanged(InValue);
 }
 
-void FPersonaMeshDetails::RefreshMeshDetailLayout()
+void FPersonaMeshDetails::ForceLayoutRebuild()
 {
+	// Force a full rebuild of the layout.
+	
 	//Unbound all UI delegate in case the user move the cursor during the refresh details	
 	for (TPair<int32, TSharedPtr<FSkeletalMeshBuildSettingsLayout>>& BuildSettingWidgetPair : BuildSettingsWidgetsPerLOD)
 	{
@@ -3685,9 +3688,17 @@ void FPersonaMeshDetails::RefreshMeshDetailLayout()
 		TSharedPtr<FSkeletalMeshReductionSettingsLayout>& LodReductionSettings = ReductionSettingWidgetPair.Value;
 		LodReductionSettings->UnbindReductionSettings();
 	}
+	
 	MeshDetailLayout->ForceRefreshDetails();
 }
 
+void FPersonaMeshDetails::RequestLayoutUpdate()
+{
+	// Update when possible but don't fully rebuild the layout. Useful to ensure that complex widgets, such as FPerPlatformX specializations, get
+	// time to update properly.
+	MeshDetailLayout->GetPropertyUtilities()->RequestRefresh();
+}
+
 void FPersonaMeshDetails::OnNoRefStreamingLODBiasChanged(int32 NewValue, FName QualityLevel)
 {
 	USkeletalMesh* SkelMesh = GetPersonaToolkit()->GetMesh();
@@ -3713,7 +3724,7 @@ void FPersonaMeshDetails::OnNoRefStreamingLODBiasChanged(int32 NewValue, FName Q
 		SkelMesh->SetNoRefStreamingLODBias(MoveTemp(NoRefStreamingLODBias));
 		SkelMesh->Modify();
 	}
-	RefreshMeshDetailLayout();
+	RequestLayoutUpdate();
 }
 
 void FPersonaMeshDetails::OnNoRefStreamingLODBiasCommitted(int32 InValue, ETextCommit::Type CommitInfo, FName QualityLevel)
@@ -3857,7 +3868,7 @@ void FPersonaMeshDetails::OnMinQualityLevelLodChanged(int32 NewValue, FName Qual
 		SkelMesh->SetQualityLevelMinLod(MoveTemp(MinLOD));
 		SkelMesh->Modify();
 	}
-	RefreshMeshDetailLayout();
+	RequestLayoutUpdate();
 }
 
 void FPersonaMeshDetails::OnMinQualityLevelLodCommitted(int32 InValue, ETextCommit::Type CommitInfo, FName QualityLevel)
@@ -3947,7 +3958,7 @@ FReply FPersonaMeshDetails::ResetToDefault()
 		SkelMesh->SetMinLod(MoveTemp(PlatformMinLOD));
 		SkelMesh->Modify();
 
-		RefreshMeshDetailLayout();
+		RequestLayoutUpdate();
 	}
 	return FReply::Handled();
 }
@@ -4042,7 +4053,8 @@ FReply FPersonaMeshDetails::ApplyLODChanges(int32 LODIndex)
 		SkelMesh->MarkPackageDirty();
 	}
 
-	RefreshMeshDetailLayout();
+	RequestLayoutUpdate();
+	
 	if (SkelMesh->GetMorphTargets().Num() > 0)
 	{
 		GetPersonaToolkit()->GetPreviewScene()->BroadcastOnMorphTargetsChanged();
@@ -4227,7 +4239,14 @@ FReply FPersonaMeshDetails::RemoveOneLOD(int32 LODIndex)
 				}
 			}
 
-			RefreshMeshDetailLayout();
+			int32 CurrentDisplayLOD = GetPersonaToolkit()->GetPreviewMeshComponent()->GetForcedLOD();
+			if (CurrentDisplayLOD == LODIndex)
+			{
+				SetCurrentLOD(FMath::Min(CurrentDisplayLOD, SkelMesh->GetLODNum() - 1));
+			}
+			
+			// We need to rebuild the layout to ensure all the LOD details are up-to-date.
+			ForceLayoutRebuild();
 		}
 	}
 	return FReply::Handled();
@@ -4235,16 +4254,14 @@ FReply FPersonaMeshDetails::RemoveOneLOD(int32 LODIndex)
 
 FText FPersonaMeshDetails::GetApplyButtonText() const
 {
-	if (IsApplyNeeded())	
+	if (IsApplyNeeded() || !IsGenerateAvailable())	
 	{
 		return LOCTEXT("ApplyChanges", "Apply Changes");
 	}
-	else if (IsGenerateAvailable())
+	else
 	{
 		return LOCTEXT("Regenerate", "Regenerate");
 	}
-
-	return LOCTEXT("ApplyChanges", "Apply Changes");
 }
 
 void FPersonaMeshDetails::ApplyChanges()
@@ -4306,7 +4323,8 @@ void FPersonaMeshDetails::ApplyChanges()
 		GetPersonaToolkit()->GetPreviewScene()->BroadcastOnMorphTargetsChanged();
 	}
 
-	RefreshMeshDetailLayout();
+	// We need to rebuild the layout to ensure all the LOD details are up-to-date.
+	ForceLayoutRebuild();
 }
 
 void FPersonaMeshDetails::UpdateLODNames()
@@ -4517,7 +4535,7 @@ void FPersonaMeshDetails::OnMeshRebuildCompleted(USkeletalMesh* InMesh)
 {
 	if (InMesh == SkeletalMeshPtr.Get())
 	{
-		RefreshMeshDetailLayout();
+		RequestLayoutUpdate();
 	}
 }
 
@@ -4554,7 +4572,7 @@ void FPersonaMeshDetails::HideUnnecessaryProperties(IDetailLayoutBuilder& Detail
 
 void FPersonaMeshDetails::OnPostProcessBlueprintChanged(IDetailLayoutBuilder* DetailBuilder)
 {
-	RefreshMeshDetailLayout();
+	RequestLayoutUpdate();
 }
 
 FString FPersonaMeshDetails::GetCurrentPostProcessBlueprintPath() const
@@ -4669,7 +4687,7 @@ FReply FPersonaMeshDetails::OnReimportLodClicked(EReimportButtonType InReimportT
 				else if (InReimportType == EReimportButtonType::ReimportWithNewFile)
 				{
 					//Refresh the layout so the BaseLOD min max get recompute
-					RefreshMeshDetailLayout();
+					RequestLayoutUpdate();
 				}
 			});
 		return FReply::Handled();
@@ -6512,9 +6530,6 @@ FReply FPersonaMeshDetails::OnRemoveClothingAssetClicked(int32 AssetIndex, IDeta
 	UpdateClothingEntries();
 	RefreshClothingComboBoxes();
 
-	// Force layout to refresh
-	//RefreshMeshDetailLayout();
-	
 	return FReply::Handled();
 }
 
@@ -6688,7 +6703,7 @@ void FPersonaMeshDetails::OnPreviewMeshChanged(USkeletalMesh* OldSkeletalMesh, U
 {
 	if (IsApplyNeeded())
 	{
-		RefreshMeshDetailLayout();
+		RequestLayoutUpdate();
 	}
 }
 
diff --git a/Engine/Source/Editor/Persona/Private/PersonaMeshDetails.h b/Engine/Source/Editor/Persona/Private/PersonaMeshDetails.h
--- a/Engine/Source/Editor/Persona/Private/PersonaMeshDetails.h
+++ b/Engine/Source/Editor/Persona/Private/PersonaMeshDetails.h
@@ -607,7 +607,8 @@ private:
 
 	FText GetMaterialSlotNameText(int32 MaterialIndex) const;
 
-	void RefreshMeshDetailLayout();
+	void ForceLayoutRebuild();
+	void RequestLayoutUpdate();
 
 	void OnNoRefStreamingLODBiasChanged(int32 NewValue, FName QualityLevel);
 	void OnNoRefStreamingLODBiasCommitted(int32 InValue, ETextCommit::Type CommitInfo, FName QualityLevel);
