Use UToolMenu for AnimNotifyTrack ContextMenu.  This allows users to customize the menu and apply permissions.

diff --git a/Engine/Source/Editor/Persona/Private/AnimNotifyPanelContextMenuContext.h b/Engine/Source/Editor/Persona/Private/AnimNotifyPanelContextMenuContext.h
new file mode 100644
--- /dev/null
+++ b/Engine/Source/Editor/Persona/Private/AnimNotifyPanelContextMenuContext.h
@@ -0,0 +1,24 @@
+// Copyright Epic Games, Inc. All Rights Reserved.
+
+#pragma once
+
+#include "CoreMinimal.h"
+#include "AnimNotifyPanelContextMenuContext.generated.h"
+
+class SAnimNotifyTrack;
+struct INodeObjectInterface;
+struct FAnimNotifyEvent; 
+
+UCLASS()
+class UAnimNotifyPanelContextMenuContext : public UObject
+{
+	GENERATED_BODY()
+
+public:
+	TWeakPtr<SAnimNotifyTrack> NotifyTrack = nullptr;
+	INodeObjectInterface* NodeObject = nullptr;
+	FAnimNotifyEvent* NotifyEvent = nullptr;
+	int32 NotifyIndex = INDEX_NONE;
+	int32 NodeIndex = INDEX_NONE;
+	FPointerEvent MouseEvent; 
+};
diff --git a/Engine/Source/Editor/Persona/Private/SAnimNotifyPanel.cpp b/Engine/Source/Editor/Persona/Private/SAnimNotifyPanel.cpp
--- a/Engine/Source/Editor/Persona/Private/SAnimNotifyPanel.cpp
+++ b/Engine/Source/Editor/Persona/Private/SAnimNotifyPanel.cpp
@@ -51,6 +51,9 @@
 #include "AnimAssetFindReplace.h"
 #include "AnimAssetFindReplaceSyncMarkers.h"
 #include "AnimAssetFindReplaceNotifies.h"
+#include "ToolMenus.h"
+#include "ToolMenuMisc.h"
+#include "AnimNotifyPanelContextMenuContext.h"
 
 // AnimNotify Drawing
 const float NotifyHeightOffset = 0.f;
@@ -939,6 +942,9 @@ protected:
 	int32 GetHitNotifyNode(const FGeometry& MyGeometry, const FVector2D& Position);
 
 	TSharedPtr<SWidget> SummonContextMenu(const FGeometry& MyGeometry, const FPointerEvent& MouseEvent);
+
+	UToolMenu* CreateContextMenuContent(FName BaseMenuName);
+
 	virtual FVector2D ComputeDesiredSize(float) const override;
 	virtual FReply OnMouseButtonUp(const FGeometry& MyGeometry, const FPointerEvent& MouseEvent) override;
 
@@ -2176,7 +2182,7 @@ FCursorReply SAnimNotifyNode::OnCursorQuery(const FGeometry& MyGeometry, const F
 void SAnimNotifyTrack::Construct(const FArguments& InArgs)
 {
 	SetClipping(EWidgetClipping::ClipToBounds);
-
+	
 	WeakCommandList = InArgs._CommandList;
 	Sequence = InArgs._Sequence;
 	ViewInputMin = InArgs._ViewInputMin;
@@ -2217,8 +2223,8 @@ void SAnimNotifyTrack::Construct(const FArguments& InArgs)
 			.BorderImage( FAppStyle::GetBrush("NoBorder") )
 			.Padding( FMargin(0.f, 0.f) )
 	];
-
 	Update();
+
 }
 
 FVector2D SAnimNotifyTrack::ComputeDesiredSize( float ) const
@@ -2619,6 +2625,7 @@ FReply SAnimNotifyTrack::OnMouseButtonUp( const FGeometry& MyGeometry, const FPo
 	if ( bRightMouseButton )
 	{
 		TSharedPtr<SWidget> WidgetToFocus;
+
 		WidgetToFocus = SummonContextMenu(MyGeometry, MouseEvent);
 
 		return (WidgetToFocus.IsValid())
@@ -2758,33 +2765,70 @@ void SAnimNotifyTrack::SelectNodesByGuid(const TSet<FGuid>& InGuids, bool bUpdat
 
 TSharedPtr<SWidget> SAnimNotifyTrack::SummonContextMenu(const FGeometry& MyGeometry, const FPointerEvent& MouseEvent)
 {
+	UToolMenus* ToolMenus = UToolMenus::Get();
+	if (!ToolMenus)
+	{
+		return SNullWidget::NullWidget;
+	}
+
+	UAnimNotifyPanelContextMenuContext* MenuContext = NewObject<UAnimNotifyPanelContextMenuContext>();
+
 	FVector2D CursorPos = MouseEvent.GetScreenSpacePosition();
-	int32 NodeIndex = GetHitNotifyNode(MyGeometry, MyGeometry.AbsoluteToLocal(CursorPos));
 	LastClickedTime = CalculateTime(MyGeometry, MouseEvent.GetScreenSpacePosition());
 
-	const bool bCloseWindowAfterMenuSelection = true;
-	FMenuBuilder MenuBuilder( bCloseWindowAfterMenuSelection, WeakCommandList.Pin() );
-	FUIAction NewAction;
+	MenuContext->NodeIndex = GetHitNotifyNode(MyGeometry, MyGeometry.AbsoluteToLocal(CursorPos));
+	MenuContext->NotifyTrack = SharedThis(this);
+	MenuContext->NodeObject = MenuContext->NodeIndex != INDEX_NONE ? NotifyNodes[MenuContext->NodeIndex]->NodeObjectInterface : nullptr;
+	MenuContext->NotifyEvent = MenuContext->NodeObject ? MenuContext->NodeObject->GetNotifyEvent() : nullptr;
+	MenuContext->NotifyIndex = MenuContext->NotifyEvent ? AnimNotifies.IndexOfByKey(MenuContext->NotifyEvent) : INDEX_NONE;
+	MenuContext->MouseEvent = MouseEvent; 
 
-	INodeObjectInterface* NodeObject = NodeIndex != INDEX_NONE ? NotifyNodes[NodeIndex]->NodeObjectInterface : nullptr;
-	FAnimNotifyEvent* NotifyEvent = NodeObject ? NodeObject->GetNotifyEvent(): nullptr;
-	int32 NotifyIndex = NotifyEvent ? AnimNotifies.IndexOfByKey(NotifyEvent) : INDEX_NONE;
+	static const FName BaseMenuName("Persona.AnimNotifyTrackContextMenu");
+	if (!ToolMenus->IsMenuRegistered(BaseMenuName))
+	{
+		CreateContextMenuContent(BaseMenuName);
+	}
 
-	MenuBuilder.BeginSection("AnimNotify", LOCTEXT("NotifyHeading", "Notify") );
+	TSharedPtr<SWidget> MenuWidget = ToolMenus->GenerateWidget(BaseMenuName, FToolMenuContext(MenuContext));
+
+	if (MenuWidget.IsValid())
 	{
-		if (NodeObject)
+		const FVector2D MouseCursorLocation = FSlateApplication::Get().GetCursorPos();
+
+		FSlateApplication::Get().PushMenu(
+			SharedThis(this),
+			FWidgetPath(),
+			MenuWidget.ToSharedRef(),
+			MouseCursorLocation,
+			FPopupTransitionEffect(FPopupTransitionEffect::ContextMenu)
+		);
+	}
+	return TSharedPtr<SWidget>();
+}
+
+UToolMenu* SAnimNotifyTrack::CreateContextMenuContent(FName BaseMenuName)
+{
+	UToolMenus* ToolMenus = UToolMenus::Get();
+	UToolMenu* Menu = ToolMenus->RegisterMenu(BaseMenuName);
+
+	Menu->AddDynamicSection(NAME_None, FNewToolMenuDelegate::CreateLambda([](UToolMenu* InMenu)
 		{
-			if (!NotifyNodes[NodeIndex]->bSelected)
+			FToolMenuSection& Section = InMenu->AddSection("AnimNotify", LOCTEXT("NotifyHeading", "Notify"));
+			UAnimNotifyPanelContextMenuContext* MenuContext = InMenu->FindContext<UAnimNotifyPanelContextMenuContext>();
+			TSharedPtr< SAnimNotifyTrack> SourceTrack = MenuContext->NotifyTrack.Pin();
+			if (MenuContext->NodeObject)
 			{
-				SelectTrackObjectNode(NodeIndex, MouseEvent.IsControlDown());
-			}
+				if (!SourceTrack->NotifyNodes[MenuContext->NodeIndex]->bSelected)
+				{
+					SourceTrack->SelectTrackObjectNode(MenuContext->NodeIndex, MenuContext->MouseEvent.IsControlDown());
+				}
 
-			if(IsSingleNodeSelected())
-			{
-				// Add item to directly set notify time
-				TSharedRef<SWidget> TimeWidget = 
-					SNew( SBox )
-					.HAlign( HAlign_Right )
+				if (SourceTrack->IsSingleNodeSelected())
+				{
+					// Add item to directly set notify time
+					TSharedRef<SWidget> TimeWidget =
+					SNew(SBox)
+					.HAlign(HAlign_Right)
 					.ToolTipText(LOCTEXT("SetTimeToolTip", "Set the time of this notify directly"))
 					[
 						SNew(SBox)
@@ -2794,43 +2838,46 @@ TSharedPtr<SWidget> SAnimNotifyTrack::SummonContextMenu(const FGeometry& MyGeome
 							SNew(SNumericEntryBox<float>)
 							.Font(FAppStyle::GetFontStyle(TEXT("MenuItem.Font")))
 							.MinValue(0.0f)
-							.MaxValue(Sequence->GetPlayLength())
-							.Value(NodeObject->GetTime())
+							.MaxValue(SourceTrack->Sequence->GetPlayLength())
+							.Value(MenuContext->NodeObject->GetTime())
 							.AllowSpin(false)
-							.OnValueCommitted_Lambda([this, NodeIndex](float InValue, ETextCommit::Type InCommitType)
-							{
-								if (InCommitType == ETextCommit::OnEnter && NotifyNodes.IsValidIndex(NodeIndex))
+							.OnValueCommitted_Lambda([SourceTrack, MenuContext](float InValue, ETextCommit::Type InCommitType)
 								{
-									const FScopedTransaction Transaction(LOCTEXT("SetNotifyTimeTransaction", "Set Anim Notify trigger time"));
-									Sequence->Modify();
-									
-									INodeObjectInterface* LocalNodeObject = NotifyNodes[NodeIndex]->NodeObjectInterface;
+									if (InCommitType == ETextCommit::OnEnter && SourceTrack->NotifyNodes.IsValidIndex(MenuContext->NodeIndex))
+									{
+										const FScopedTransaction Transaction(LOCTEXT("SetNotifyTimeTransaction", "Set Anim Notify trigger time"));
+										SourceTrack->Sequence->Modify();
 
-									float NewTime = FMath::Clamp(InValue, 0.0f, Sequence->GetPlayLength() - LocalNodeObject->GetDuration());
-									LocalNodeObject->SetTime(NewTime);
+										INodeObjectInterface* LocalNodeObject = SourceTrack->NotifyNodes[MenuContext->NodeIndex]->NodeObjectInterface;
 
-									if (FAnimNotifyEvent* Event = LocalNodeObject->GetNotifyEvent())
-									{
-										Event->RefreshTriggerOffset(Sequence->CalculateOffsetForNotify(Event->GetTime()));
-										if (Event->GetDuration() > 0.0f)
+										float NewTime = FMath::Clamp(InValue, 0.0f, SourceTrack->Sequence->GetPlayLength() - LocalNodeObject->GetDuration());
+										LocalNodeObject->SetTime(NewTime);
+
+										if (FAnimNotifyEvent* Event = LocalNodeObject->GetNotifyEvent())
 										{
-											Event->RefreshEndTriggerOffset(Sequence->CalculateOffsetForNotify(Event->GetTime() + Event->GetDuration()));
+											Event->RefreshTriggerOffset(SourceTrack->Sequence->CalculateOffsetForNotify(Event->GetTime()));
+											if (Event->GetDuration() > 0.0f)
+											{
+												Event->RefreshEndTriggerOffset(SourceTrack->Sequence->CalculateOffsetForNotify(Event->GetTime() + Event->GetDuration()));
+											}
 										}
-									}
-									OnUpdatePanel.ExecuteIfBound();
+										SourceTrack->OnUpdatePanel.ExecuteIfBound();
 
-									FSlateApplication::Get().DismissAllMenus();
-								}
-							})
+										FSlateApplication::Get().DismissAllMenus();
+									}
+								})
 						]
 					];
 
-				MenuBuilder.AddWidget(TimeWidget, LOCTEXT("TimeMenuText", "Notify Begin Time"));
+					Section.AddEntry(
+						FToolMenuEntry::InitWidget(TEXT("AnimNotifyContextMenuTimeWidget"), TimeWidget, FText::FromString("Notify Begin Time"), true, false)
+					);
 
-				// Add item to directly set notify frame
-				TSharedRef<SWidget> FrameWidget = 
-					SNew( SBox )
-					.HAlign( HAlign_Right )
+
+					// Add item to directly set notify frame
+					TSharedRef<SWidget> FrameWidget =
+					SNew(SBox)
+					.HAlign(HAlign_Right)
 					.ToolTipText(LOCTEXT("SetFrameToolTip", "Set the frame of this notify directly"))
 					[
 						SNew(SBox)
@@ -2840,45 +2887,47 @@ TSharedPtr<SWidget> SAnimNotifyTrack::SummonContextMenu(const FGeometry& MyGeome
 							SNew(SNumericEntryBox<int32>)
 							.Font(FAppStyle::GetFontStyle(TEXT("MenuItem.Font")))
 							.MinValue(0)
-							.MaxValue(Sequence->GetNumberOfSampledKeys())
-							.Value(Sequence->GetFrameAtTime(NodeObject->GetTime()))
-							.AllowSpin(false)						
-							.OnValueCommitted_Lambda([this, NodeIndex](int32 InValue, ETextCommit::Type InCommitType)
-							{
-								if (InCommitType == ETextCommit::OnEnter && NotifyNodes.IsValidIndex(NodeIndex))
+							.MaxValue(SourceTrack->Sequence->GetNumberOfSampledKeys())
+							.Value(SourceTrack->Sequence->GetFrameAtTime(MenuContext->NodeObject->GetTime()))
+							.AllowSpin(false)
+							.OnValueCommitted_Lambda([SourceTrack, MenuContext](int32 InValue, ETextCommit::Type InCommitType)
 								{
-									const FScopedTransaction Transaction(LOCTEXT("SetNotifyFrameTransaction", "Set Anim Notify trigger frame index"));
-									Sequence->Modify();
+									if (InCommitType == ETextCommit::OnEnter && SourceTrack->NotifyNodes.IsValidIndex(MenuContext->NodeIndex))
+									{
+										const FScopedTransaction Transaction(LOCTEXT("SetNotifyFrameTransaction", "Set Anim Notify trigger frame index"));
+										SourceTrack->Sequence->Modify();
 
-									INodeObjectInterface* LocalNodeObject = NotifyNodes[NodeIndex]->NodeObjectInterface;
+										INodeObjectInterface* LocalNodeObject = SourceTrack->NotifyNodes[MenuContext->NodeIndex]->NodeObjectInterface;
 
-									float NewTime = FMath::Clamp(Sequence->GetTimeAtFrame(InValue), 0.0f, Sequence->GetPlayLength() - LocalNodeObject->GetDuration());
-									LocalNodeObject->SetTime(NewTime);
+										float NewTime = FMath::Clamp(SourceTrack->Sequence->GetTimeAtFrame(InValue), 0.0f, SourceTrack->Sequence->GetPlayLength() - LocalNodeObject->GetDuration());
+										LocalNodeObject->SetTime(NewTime);
 
-									if (FAnimNotifyEvent* Event = LocalNodeObject->GetNotifyEvent())
-									{
-										Event->RefreshTriggerOffset(Sequence->CalculateOffsetForNotify(Event->GetTime()));
-										if (Event->GetDuration() > 0.0f)
+										if (FAnimNotifyEvent* Event = LocalNodeObject->GetNotifyEvent())
 										{
-											Event->RefreshEndTriggerOffset(Sequence->CalculateOffsetForNotify(Event->GetTime() + Event->GetDuration()));
+											Event->RefreshTriggerOffset(SourceTrack->Sequence->CalculateOffsetForNotify(Event->GetTime()));
+											if (Event->GetDuration() > 0.0f)
+											{
+												Event->RefreshEndTriggerOffset(SourceTrack->Sequence->CalculateOffsetForNotify(Event->GetTime() + Event->GetDuration()));
+											}
 										}
-									}
-									OnUpdatePanel.ExecuteIfBound();
+										SourceTrack->OnUpdatePanel.ExecuteIfBound();
 
-									FSlateApplication::Get().DismissAllMenus();
-								}
-							})
+										FSlateApplication::Get().DismissAllMenus();
+									}
+								})
 						]
 					];
+					Section.AddEntry(
+						FToolMenuEntry::InitWidget(TEXT("AnimNotifyContextMenuFrameWidget"), FrameWidget, FText::FromString("Notify Frame"), true, false)
+					);
 
-				MenuBuilder.AddWidget(FrameWidget, LOCTEXT("FrameMenuText", "Notify Frame"));
 
-				if (NotifyEvent)
-				{
-					// add menu to get threshold weight for triggering this notify
-					TSharedRef<SWidget> ThresholdWeightWidget = 
-						SNew( SBox )
-						.HAlign( HAlign_Right )
+					if (MenuContext->NotifyEvent)
+					{
+						// add menu to get threshold weight for triggering this notify
+						TSharedRef<SWidget> ThresholdWeightWidget =
+						SNew(SBox)
+						.HAlign(HAlign_Right)
 						.ToolTipText(LOCTEXT("MinTriggerWeightToolTip", "The minimum weight to trigger this notify"))
 						[
 							SNew(SBox)
@@ -2889,32 +2938,35 @@ TSharedPtr<SWidget> SAnimNotifyTrack::SummonContextMenu(const FGeometry& MyGeome
 								.Font(FAppStyle::GetFontStyle(TEXT("MenuItem.Font")))
 								.MinValue(0.0f)
 								.MaxValue(1.0f)
-								.Value(NotifyEvent->TriggerWeightThreshold)
-								.AllowSpin(false)						
-								.OnValueCommitted_Lambda([this, NotifyIndex](float InValue, ETextCommit::Type InCommitType)
-								{
-									if ( InCommitType == ETextCommit::OnEnter && AnimNotifies.IsValidIndex(NotifyIndex) )
+								.Value(MenuContext->NotifyEvent->TriggerWeightThreshold)
+								.AllowSpin(false)
+								.OnValueCommitted_Lambda([SourceTrack, MenuContext](float InValue, ETextCommit::Type InCommitType)
 									{
-										const FScopedTransaction Transaction(LOCTEXT("SetNotifyWeightTransaction", "Set Anim Notify trigger weight"));
-										Sequence->Modify();
-									
-										float NewWeight = FMath::Max(InValue, ZERO_ANIMWEIGHT_THRESH);
-										AnimNotifies[NotifyIndex]->TriggerWeightThreshold = NewWeight;
+										if (InCommitType == ETextCommit::OnEnter && SourceTrack->AnimNotifies.IsValidIndex(MenuContext->NotifyIndex))
+										{
+											const FScopedTransaction Transaction(LOCTEXT("SetNotifyWeightTransaction", "Set Anim Notify trigger weight"));
+											SourceTrack->Sequence->Modify();
 
-										FSlateApplication::Get().DismissAllMenus();
-									}
-								})
+											float NewWeight = FMath::Max(InValue, ZERO_ANIMWEIGHT_THRESH);
+											SourceTrack->AnimNotifies[MenuContext->NotifyIndex]->TriggerWeightThreshold = NewWeight;
+
+											FSlateApplication::Get().DismissAllMenus();
+										}
+									})
 							]
 						];
 
-					MenuBuilder.AddWidget(ThresholdWeightWidget, LOCTEXT("MinTriggerWeight", "Min Trigger Weight"));
+						Section.AddEntry(
+							FToolMenuEntry::InitWidget(TEXT("AnimNotifyContextMenuTriggerWeightWidget"), ThresholdWeightWidget, FText::FromString("Min Trigger Weight"), true, false)
+						);
 
-					// Add menu for changing duration if this is an AnimNotifyState
-					if (NotifyEvent->NotifyStateClass)
-					{
-						TSharedRef<SWidget> NotifyStateDurationWidget = 
-							SNew( SBox )
-							.HAlign( HAlign_Right )
+
+						// Add menu for changing duration if this is an AnimNotifyState
+						if (MenuContext->NotifyEvent->NotifyStateClass)
+						{
+							TSharedRef<SWidget> NotifyStateDurationWidget =
+							SNew(SBox)
+							.HAlign(HAlign_Right)
 							.ToolTipText(LOCTEXT("SetAnimStateDuration_ToolTip", "The duration of this Anim Notify State"))
 							[
 								SNew(SBox)
@@ -2926,35 +2978,37 @@ TSharedPtr<SWidget> SAnimNotifyTrack::SummonContextMenu(const FGeometry& MyGeome
 									.MinValue(SAnimNotifyNode::MinimumStateDuration)
 									.MinSliderValue(SAnimNotifyNode::MinimumStateDuration)
 									.MaxSliderValue(100.0f)
-									.Value(NotifyEvent->GetDuration())
-									.AllowSpin(false)						
-									.OnValueCommitted_Lambda([this, NotifyIndex](float InValue, ETextCommit::Type InCommitType)
-									{
-										if ( InCommitType == ETextCommit::OnEnter && AnimNotifies.IsValidIndex(NotifyIndex) )
+									.Value(MenuContext->NotifyEvent->GetDuration())
+									.AllowSpin(false)
+									.OnValueCommitted_Lambda([SourceTrack, MenuContext](float InValue, ETextCommit::Type InCommitType)
 										{
-											const FScopedTransaction Transaction(LOCTEXT("SetNotifyDurationSecondsTransaction", "Set Anim Notify State duration in seconds"));
-											Sequence->Modify();
-											
-											float NewDuration = FMath::Max(InValue, SAnimNotifyNode::MinimumStateDuration);
-											float MaxDuration = Sequence->GetPlayLength() - AnimNotifies[NotifyIndex]->GetTime();
-											NewDuration = FMath::Min(NewDuration, MaxDuration);
-											AnimNotifies[NotifyIndex]->SetDuration(NewDuration);
-
-											// If we have a delegate bound to refresh the offsets, call it.
-											// This is used by the montage editor to keep the offsets up to date.
-											OnRequestRefreshOffsets.ExecuteIfBound();
-
-											FSlateApplication::Get().DismissAllMenus();
-										}
-									})
+											if (InCommitType == ETextCommit::OnEnter && SourceTrack->AnimNotifies.IsValidIndex(MenuContext->NotifyIndex))
+											{
+												const FScopedTransaction Transaction(LOCTEXT("SetNotifyDurationSecondsTransaction", "Set Anim Notify State duration in seconds"));
+												SourceTrack->Sequence->Modify();
+
+												float NewDuration = FMath::Max(InValue, SAnimNotifyNode::MinimumStateDuration);
+												float MaxDuration = SourceTrack->Sequence->GetPlayLength() - SourceTrack->AnimNotifies[MenuContext->NotifyIndex]->GetTime();
+												NewDuration = FMath::Min(NewDuration, MaxDuration);
+												SourceTrack->AnimNotifies[MenuContext->NotifyIndex]->SetDuration(NewDuration);
+
+												// If we have a delegate bound to refresh the offsets, call it.
+												// This is used by the montage editor to keep the offsets up to date.
+												SourceTrack->OnRequestRefreshOffsets.ExecuteIfBound();
+
+												FSlateApplication::Get().DismissAllMenus();
+											}
+										})
 								]
 							];
 
-						MenuBuilder.AddWidget(NotifyStateDurationWidget, LOCTEXT("SetAnimStateDuration", "Anim Notify State Duration"));
+							Section.AddEntry(
+								FToolMenuEntry::InitWidget(TEXT("AnimNotifyContextStateDurationWidget"), NotifyStateDurationWidget, FText::FromString("Anim Notify State Duration"), true, false)
+							);
 
-						TSharedRef<SWidget> NotifyStateDurationFramesWidget = 
-							SNew( SBox )
-							.HAlign( HAlign_Right )
+							TSharedRef<SWidget> NotifyStateDurationFramesWidget =
+							SNew(SBox)
+							.HAlign(HAlign_Right)
 							.ToolTipText(LOCTEXT("SetAnimStateDurationFrames_ToolTip", "The duration of this Anim Notify State in frames"))
 							[
 								SNew(SBox)
@@ -2965,198 +3019,229 @@ TSharedPtr<SWidget> SAnimNotifyTrack::SummonContextMenu(const FGeometry& MyGeome
 									.Font(FAppStyle::GetFontStyle(TEXT("MenuItem.Font")))
 									.MinValue(1)
 									.MinSliderValue(1)
-									.MaxSliderValue(Sequence->GetNumberOfSampledKeys())
-									.Value(Sequence->GetFrameAtTime(NotifyEvent->GetDuration()))
-									.AllowSpin(false)						
-									.OnValueCommitted_Lambda([this, NotifyIndex](int32 InValue, ETextCommit::Type InCommitType)
-									{
-										if ( InCommitType == ETextCommit::OnEnter && AnimNotifies.IsValidIndex(NotifyIndex) )
+									.MaxSliderValue(SourceTrack->Sequence->GetNumberOfSampledKeys())
+									.Value(SourceTrack->Sequence->GetFrameAtTime(MenuContext->NotifyEvent->GetDuration()))
+									.AllowSpin(false)
+									.OnValueCommitted_Lambda([SourceTrack, MenuContext](int32 InValue, ETextCommit::Type InCommitType)
 										{
-											const FScopedTransaction Transaction(LOCTEXT("SetNotifyDurationFramesTransaction", "Set Anim Notify State duration in frames"));
-											Sequence->Modify();
-											
-											float NewDuration = FMath::Max(Sequence->GetTimeAtFrame(InValue), SAnimNotifyNode::MinimumStateDuration);
-											float MaxDuration = Sequence->GetPlayLength() - AnimNotifies[NotifyIndex]->GetTime();
-											NewDuration = FMath::Min(NewDuration, MaxDuration);
-											AnimNotifies[NotifyIndex]->SetDuration(NewDuration);
-
-											// If we have a delegate bound to refresh the offsets, call it.
-											// This is used by the montage editor to keep the offsets up to date.
-											OnRequestRefreshOffsets.ExecuteIfBound();
-
-											FSlateApplication::Get().DismissAllMenus();
-										}
-									})
+											if (InCommitType == ETextCommit::OnEnter && SourceTrack->AnimNotifies.IsValidIndex(MenuContext->NotifyIndex))
+											{
+												const FScopedTransaction Transaction(LOCTEXT("SetNotifyDurationFramesTransaction", "Set Anim Notify State duration in frames"));
+												SourceTrack->Sequence->Modify();
+
+												float NewDuration = FMath::Max(SourceTrack->Sequence->GetTimeAtFrame(InValue), SAnimNotifyNode::MinimumStateDuration);
+												float MaxDuration = SourceTrack->Sequence->GetPlayLength() - SourceTrack->AnimNotifies[MenuContext->NotifyIndex]->GetTime();
+												NewDuration = FMath::Min(NewDuration, MaxDuration);
+												SourceTrack->AnimNotifies[MenuContext->NotifyIndex]->SetDuration(NewDuration);
+
+												// If we have a delegate bound to refresh the offsets, call it.
+												// This is used by the montage editor to keep the offsets up to date.
+												SourceTrack->OnRequestRefreshOffsets.ExecuteIfBound();
+
+												FSlateApplication::Get().DismissAllMenus();
+											}
+										})
 								]
 							];
 
-						MenuBuilder.AddWidget(NotifyStateDurationFramesWidget, LOCTEXT("SetAnimStateDurationFrames", "Anim Notify State Frames"));
+							Section.AddEntry(
+								FToolMenuEntry::InitWidget(TEXT("AnimNotifyContextMenuStateFramesWidget"), NotifyStateDurationFramesWidget, FText::FromString("Anim Notify State Frames"), true, false)
+							);
+
+						}
+
 					}
 				}
-			}
-		}
-		else
-		{
-			MenuBuilder.AddSubMenu(
-				NSLOCTEXT("NewNotifySubMenu", "NewNotifySubMenuAddNotify", "Add Notify..."),
-				NSLOCTEXT("NewNotifySubMenu", "NewNotifySubMenuAddNotifyToolTip", "Add AnimNotifyEvent"),
-				FNewMenuDelegate::CreateRaw( this, &SAnimNotifyTrack::FillNewNotifyMenu, false ) );
 
-			MenuBuilder.AddSubMenu(
-				NSLOCTEXT("NewNotifySubMenu", "NewNotifySubMenuAddNotifyState", "Add Notify State..."),
-				NSLOCTEXT("NewNotifySubMenu", "NewNotifySubMenuAddNotifyStateToolTip","Add AnimNotifyState"),
-				FNewMenuDelegate::CreateRaw( this, &SAnimNotifyTrack::FillNewNotifyStateMenu, false ) );
-
-			if (Sequence->IsA(UAnimSequence::StaticClass()))
-			{
-				MenuBuilder.AddSubMenu(
-					NSLOCTEXT("NewSyncMarkerSubMenu", "NewSyncMarkerSubMenuAddNotifyState", "Add Sync Marker..."),
-					NSLOCTEXT("NewSyncMarkerSubMenu", "NewSyncMarkerSubMenuAddNotifyStateToolTip", "Create a new animation sync marker"),
-					FNewMenuDelegate::CreateRaw(this, &SAnimNotifyTrack::FillNewSyncMarkerMenu, false));
 			}
+			else
+			{
+				Section.AddSubMenu(TEXT("AddNotify"),
+					NSLOCTEXT("NewNotifySubMenu", "NewNotifySubMenuAddNotify", "Add Notify..."),
+					NSLOCTEXT("NewNotifySubMenu", "NewNotifySubMenuAddNotifyToolTip", "Add AnimNotifyEvent"),
+					FNewMenuDelegate::CreateRaw(SourceTrack.Get(), &SAnimNotifyTrack::FillNewNotifyMenu, false),
+					false,
+					FSlateIcon());
+
+				Section.AddSubMenu(TEXT("AddNotifyState"),
+					NSLOCTEXT("NewNotifySubMenu", "NewNotifySubMenuAddNotifyState", "Add Notify State..."),
+					NSLOCTEXT("NewNotifySubMenu", "NewNotifySubMenuAddNotifyStateToolTip", "Add AnimNotifyState"),
+					FNewMenuDelegate::CreateRaw(SourceTrack.Get(), &SAnimNotifyTrack::FillNewNotifyStateMenu, false),
+					false,
+					FSlateIcon());
+
+
+				if (SourceTrack->Sequence->IsA(UAnimSequence::StaticClass()))
+				{
+					Section.AddSubMenu(TEXT("AddSyncMarker"),
+						NSLOCTEXT("NewSyncMarkerSubMenu", "NewSyncMarkerSubMenuAddNotifyState", "Add Sync Marker..."),
+						NSLOCTEXT("NewSyncMarkerSubMenu", "NewSyncMarkerSubMenuAddNotifyStateToolTip", "Create a new animation sync marker"),
+						FNewMenuDelegate::CreateRaw(SourceTrack.Get(), &SAnimNotifyTrack::FillNewSyncMarkerMenu, false),
+						false,
+						FSlateIcon());
+				}
 
-			MenuBuilder.AddMenuEntry(
-				NSLOCTEXT("NewNotifySubMenu", "ManageNotifies", "Manage Notifies..."),
-				NSLOCTEXT("NewNotifySubMenu", "ManageNotifiesToolTip", "Opens the Manage Notifies window"),
-				FSlateIcon(),
-				FUIAction( FExecuteAction::CreateSP( this, &SAnimNotifyTrack::OnManageNotifies ) ) );
-		}
-	}
-	MenuBuilder.EndSection(); //AnimNotify
-
-	NewAction.CanExecuteAction = 0;
+				Section.AddMenuEntry(TEXT("ManageNotifies"),
+					NSLOCTEXT("NewNotifySubMenu", "ManageNotifies", "Manage Notifies..."),
+					NSLOCTEXT("NewNotifySubMenu", "ManageNotifiesToolTip", "Opens the Manage Notifies window"),
+					FSlateIcon(),
+					FUIAction(FExecuteAction::CreateSP(SourceTrack.Get(), &SAnimNotifyTrack::OnManageNotifies)));
+			}
 
-	MenuBuilder.BeginSection("AnimEdit", LOCTEXT("NotifyEditHeading", "Edit") );
-	{
-		if ( NodeObject )
-		{
-			// copy notify menu item
-			MenuBuilder.AddMenuEntry(FAnimNotifyPanelCommands::Get().CopyNotifies);
+			FToolMenuSection& EditSection = InMenu->AddSection("AnimEdit", LOCTEXT("NotifyEditHeading", "Edit"));
+			if (MenuContext->NodeObject)
+			{
+				// copy notify menu item
+				EditSection.AddMenuEntry(FAnimNotifyPanelCommands::Get().CopyNotifies);
 
-			// allow it to delete
-			MenuBuilder.AddMenuEntry(FAnimNotifyPanelCommands::Get().DeleteNotify);
+				// allow it to delete
+				EditSection.AddMenuEntry(FAnimNotifyPanelCommands::Get().DeleteNotify);
 
-			if (NotifyEvent)
-			{
-				// For the "Replace With..." menu, make sure the current AnimNotify selection is valid for replacement
-				if (OnGetIsAnimNotifySelectionValidforReplacement.IsBound() && OnGetIsAnimNotifySelectionValidforReplacement.Execute())
+				if (MenuContext->NotifyEvent)
 				{
-					// If this is an AnimNotifyState (has duration) allow it to be replaced with other AnimNotifyStates
-					if (NotifyEvent->NotifyStateClass)
+					// For the "Replace With..." menu, make sure the current AnimNotify selection is valid for replacement
+					if (SourceTrack->OnGetIsAnimNotifySelectionValidforReplacement.IsBound() && SourceTrack->OnGetIsAnimNotifySelectionValidforReplacement.Execute())
 					{
-						MenuBuilder.AddSubMenu(
-							NSLOCTEXT("NewNotifySubMenu", "NewNotifySubMenuReplaceWithNotifyState", "Replace with Notify State..."),
-							NSLOCTEXT("NewNotifySubMenu", "NewNotifySubMenuReplaceWithNotifyStateToolTip", "Replace with AnimNotifyState"),
-							FNewMenuDelegate::CreateRaw(this, &SAnimNotifyTrack::FillNewNotifyStateMenu, true));
+						// If this is an AnimNotifyState (has duration) allow it to be replaced with other AnimNotifyStates
+						if (MenuContext->NotifyEvent->NotifyStateClass)
+						{
+							EditSection.AddSubMenu(TEXT("ReplaceWithNotifyState"),
+								NSLOCTEXT("NewNotifySubMenu", "NewNotifySubMenuReplaceWithNotifyState", "Replace with Notify State..."),
+								NSLOCTEXT("NewNotifySubMenu", "NewNotifySubMenuReplaceWithNotifyStateToolTip", "Replace with AnimNotifyState"),
+								FNewMenuDelegate::CreateRaw(SourceTrack.Get(), &SAnimNotifyTrack::FillNewNotifyStateMenu, true),
+								false,
+								FSlateIcon());
+						}
+						// If this is a regular AnimNotify (no duration) allow it to be replaced with other AnimNotifies
+						else
+						{
+							EditSection.AddSubMenu(TEXT("ReplaceWithNotify"),
+								NSLOCTEXT("NewNotifySubMenu", "NewNotifySubMenuReplaceWithNotify", "Replace with Notify..."),
+								NSLOCTEXT("NewNotifySubMenu", "NewNotifySubMenuReplaceWithNotifyToolTip", "Replace with AnimNotifyEvent"),
+								FNewMenuDelegate::CreateRaw(SourceTrack.Get(), &SAnimNotifyTrack::FillNewNotifyMenu, true),
+								false,
+								FSlateIcon()
+							);
+						}
 					}
-					// If this is a regular AnimNotify (no duration) allow it to be replaced with other AnimNotifies
-					else
+				}
+				else
+				{
+					if (MenuContext->NodeObject->GetType() == ENodeObjectTypes::SYNC_MARKER)
 					{
-						MenuBuilder.AddSubMenu(
-							NSLOCTEXT("NewNotifySubMenu", "NewNotifySubMenuReplaceWithNotify", "Replace with Notify..."),
-							NSLOCTEXT("NewNotifySubMenu", "NewNotifySubMenuReplaceWithNotifyToolTip", "Replace with AnimNotifyEvent"),
-							FNewMenuDelegate::CreateRaw(this, &SAnimNotifyTrack::FillNewNotifyMenu, true));
+						EditSection.AddSubMenu(TEXT("ReplaceSyncMarkers"),
+							NSLOCTEXT("NewNotifySubMenu", "NewNotifySubMenuReplaceWithSyncMarker", "Replace Sync Marker(s)..."),
+							NSLOCTEXT("NewNotifySubMenu", "NewNotifySubMenuReplaceWithSyncMarkerToolTip", "Replace the selected sync markers"),
+							FNewMenuDelegate::CreateRaw(SourceTrack.Get(), &SAnimNotifyTrack::FillNewSyncMarkerMenu, true),
+							false,
+							FSlateIcon()
+						);
 					}
 				}
 			}
 			else
 			{
-				if (NodeObject->GetType() == ENodeObjectTypes::SYNC_MARKER)
-				{
-					MenuBuilder.AddSubMenu(
-						NSLOCTEXT("NewNotifySubMenu", "NewNotifySubMenuReplaceWithSyncMarker", "Replace Sync Marker(s)..."),
-						NSLOCTEXT("NewNotifySubMenu", "NewNotifySubMenuReplaceWithSyncMarkerToolTip", "Replace the selected sync markers"),
-						FNewMenuDelegate::CreateRaw(this, &SAnimNotifyTrack::FillNewSyncMarkerMenu, true));
-				}
-			}
-		}
-		else
-		{
-			FString PropertyString;
-			const TCHAR* Buffer;
-			float OriginalTime;
-			float OriginalLength;
-			int32 TrackSpan;
-
-			//Check whether can we show menu item to paste anim notify event
-			if( ReadNotifyPasteHeader(PropertyString, Buffer, OriginalTime, OriginalLength,TrackSpan) )
-			{
-				// paste notify menu item
-				if (IsSingleNodeInClipboard())
+				FString PropertyString;
+				const TCHAR* Buffer;
+				float OriginalTime;
+				float OriginalLength;
+				int32 TrackSpan;
+				FUIAction NewAction;
+				//Check whether can we show menu item to paste anim notify event
+				if (ReadNotifyPasteHeader(PropertyString, Buffer, OriginalTime, OriginalLength, TrackSpan))
 				{
-					MenuBuilder.AddMenuEntry(FAnimNotifyPanelCommands::Get().PasteNotifies);
-				}
-				else
-				{
-					NewAction.ExecuteAction.BindRaw(
-						this, &SAnimNotifyTrack::OnPasteNotifyClicked, ENotifyPasteMode::MousePosition, ENotifyPasteMultipleMode::Relative);
+					// paste notify menu item
+					if (SourceTrack->IsSingleNodeInClipboard())
+					{
+						EditSection.AddMenuEntry(FAnimNotifyPanelCommands::Get().PasteNotifies);
+					}
+					else
+					{
+						NewAction.ExecuteAction.BindRaw(
+							SourceTrack.Get(), &SAnimNotifyTrack::OnPasteNotifyClicked, ENotifyPasteMode::MousePosition, ENotifyPasteMultipleMode::Relative);
+
+						EditSection.AddMenuEntry(TEXT("PasteMultipleRelative"),
+							LOCTEXT("PasteMultRel", "Paste Multiple Relative"), 
+							LOCTEXT("PasteMultRelToolTip", "Paste multiple notifies beginning at the mouse cursor, maintaining the same relative spacing as the source."),
+							FSlateIcon(),
+							NewAction);
+
+						EditSection.AddMenuEntry(FAnimNotifyPanelCommands::Get().PasteNotifies,
+							LOCTEXT("PasteMultAbs", "Paste Multiple Absolute"),
+							LOCTEXT("PasteMultAbsToolTip", "Paste multiple notifies beginning at the mouse cursor, maintaining absolute spacing."));
+					}
 
-					MenuBuilder.AddMenuEntry(LOCTEXT("PasteMultRel", "Paste Multiple Relative"), LOCTEXT("PasteMultRelToolTip", "Paste multiple notifies beginning at the mouse cursor, maintaining the same relative spacing as the source."), FSlateIcon(), NewAction);
+					if (OriginalTime < SourceTrack->Sequence->GetPlayLength())
+					{
+						NewAction.ExecuteAction.BindRaw(
+							SourceTrack.Get(), &SAnimNotifyTrack::OnPasteNotifyClicked, ENotifyPasteMode::OriginalTime, ENotifyPasteMultipleMode::Absolute);
 
-					MenuBuilder.AddMenuEntry(FAnimNotifyPanelCommands::Get().PasteNotifies, NAME_None, LOCTEXT("PasteMultAbs", "Paste Multiple Absolute"), LOCTEXT("PasteMultAbsToolTip", "Paste multiple notifies beginning at the mouse cursor, maintaining absolute spacing."));
-				}
+						FText DisplayText = FText::Format(LOCTEXT("PasteAtOriginalTime", "Paste at original time ({0})"), FText::AsNumber(OriginalTime));
 
-				if(OriginalTime < Sequence->GetPlayLength())
-				{
-					NewAction.ExecuteAction.BindRaw(
-						this, &SAnimNotifyTrack::OnPasteNotifyClicked, ENotifyPasteMode::OriginalTime, ENotifyPasteMultipleMode::Absolute);
+						EditSection.AddMenuEntry(TEXT("Paste at Original Time"),
+							DisplayText,
+							LOCTEXT("PasteAtOriginalTimeToolTip", "Paste animation notify event at the time it was set to when it was copied"),
+							FSlateIcon(), 
+							NewAction);
+					}
 
-					FText DisplayText = FText::Format( LOCTEXT("PasteAtOriginalTime", "Paste at original time ({0})"), FText::AsNumber( OriginalTime) );
-					MenuBuilder.AddMenuEntry(DisplayText, LOCTEXT("PasteAtOriginalTimeToolTip", "Paste animation notify event at the time it was set to when it was copied"), FSlateIcon(), NewAction);
 				}
-				
 			}
-		}
-	}
-	MenuBuilder.EndSection(); //AnimEdit
-
-	if (NotifyEvent)
-	{
-		UObject* NotifyObject = NotifyEvent->Notify;
-		NotifyObject = NotifyObject ? NotifyObject : ToRawPtr(NotifyEvent->NotifyStateClass);
 
-		MenuBuilder.BeginSection("ViewSource", LOCTEXT("NotifyViewHeading", "View"));
-
-		if (NotifyObject)
-		{
-			if (Cast<UBlueprintGeneratedClass>(NotifyObject->GetClass()))
+			FToolMenuSection& ViewSection = InMenu->AddSection("AnimView", LOCTEXT("NotifyViewHeading", "View"));
+			if (MenuContext->NotifyEvent)
 			{
-				if (UBlueprint * Blueprint = Cast<UBlueprint>(NotifyObject->GetClass()->ClassGeneratedBy))
+				UObject* NotifyObject = MenuContext->NotifyEvent->Notify;
+				NotifyObject = NotifyObject ? NotifyObject : ToRawPtr(MenuContext->NotifyEvent->NotifyStateClass);
+				FUIAction NewAction;
+				
+				if (NotifyObject)
+				{
+					if (Cast<UBlueprintGeneratedClass>(NotifyObject->GetClass()))
+					{
+						if (UBlueprint* Blueprint = Cast<UBlueprint>(NotifyObject->GetClass()->ClassGeneratedBy))
+						{
+							NewAction.ExecuteAction.BindRaw(
+								SourceTrack.Get(), &SAnimNotifyTrack::OnOpenNotifySource, Blueprint);
+							ViewSection.AddMenuEntry(TEXT("OpenNotifyBlueprint"),
+								LOCTEXT("OpenNotifyBlueprint", "Open Notify Blueprint"),
+								LOCTEXT("OpenNotifyBlueprintTooltip", "Opens the source blueprint for this notify"),
+								FSlateIcon(),
+								NewAction);
+						}
+					}
+				}
+				else
 				{
+					// skeleton notify
 					NewAction.ExecuteAction.BindRaw(
-						this, &SAnimNotifyTrack::OnOpenNotifySource, Blueprint);
-					MenuBuilder.AddMenuEntry(LOCTEXT("OpenNotifyBlueprint", "Open Notify Blueprint"), LOCTEXT("OpenNotifyBlueprintTooltip", "Opens the source blueprint for this notify"), FSlateIcon(), NewAction);
+						SourceTrack.Get(), &SAnimNotifyTrack::OnFindReferences, MenuContext->NodeObject->GetName(), true);
+					ViewSection.AddMenuEntry(TEXT("FindReferences"),
+						LOCTEXT("FindNotifyReferences", "Find/Replace References..."),
+						LOCTEXT("FindNotifyReferencesTooltip",
+							"Find, replace and remove references to this  notify in the find/replace tab"),
+						FSlateIcon(), 
+						NewAction);
 				}
 			}
-		}
-		else
-		{
-			// skeleton notify
-			NewAction.ExecuteAction.BindRaw(
-				this, &SAnimNotifyTrack::OnFindReferences, NotifyEvent->NotifyName, false);
-			MenuBuilder.AddMenuEntry(LOCTEXT("FindNotifyReferences", "Find/Replace References..."), LOCTEXT("FindNotifyReferencesTooltip", "Find, replace and remove references to this notify in the find/replace taby"), FSlateIcon(), NewAction);
-		}
-
-		MenuBuilder.EndSection(); //ViewSource
-	}
-	else
-	{
-		if (NodeObject && NodeObject->GetType() == ENodeObjectTypes::SYNC_MARKER)
-		{
-			NewAction.ExecuteAction.BindRaw(
-				this, &SAnimNotifyTrack::OnFindReferences, NodeObject->GetName(), true);
-			MenuBuilder.AddMenuEntry(LOCTEXT("FindSyncMarkerReferences", "Find/Replace References..."), LOCTEXT("FindSyncMarkerReferencesTooltip", "Find, replace and remove references to this sync marker in the find/replace tab"), FSlateIcon(), NewAction);
-		}
-	}
-
-	FWidgetPath WidgetPath = MouseEvent.GetEventPath() != nullptr ? *MouseEvent.GetEventPath() : FWidgetPath();
-
-	// Display the newly built menu
-	FSlateApplication::Get().PushMenu(SharedThis(this), WidgetPath, MenuBuilder.MakeWidget(), CursorPos, FPopupTransitionEffect(FPopupTransitionEffect::ContextMenu));
+			else if (MenuContext->NodeObject && MenuContext->NodeObject->GetType() == ENodeObjectTypes::SYNC_MARKER)
+			{
+				FUIAction NewAction;
+				NewAction.ExecuteAction.BindRaw(
+					SourceTrack.Get(), &SAnimNotifyTrack::OnFindReferences, MenuContext->NodeObject->GetName(), true);
+				
+				ViewSection.AddMenuEntry(TEXT("FindSyncReferences"),
+					LOCTEXT("FindSyncMarkerReferences", "Find/Replace References..."),
+					LOCTEXT("FindSyncMarkerReferencesTooltip", "Find, replace and remove references to this sync marker in the find/replace tab"),
+					FSlateIcon(),
+					NewAction);
 
-	return TSharedPtr<SWidget>();
+			}
+		}));
+	return Menu; 
 }
 
+
 bool SAnimNotifyTrack::CanPasteAnimNotify() const
 {
 	FString PropertyString;
