Fix build  issue with AnimNotifyPanelContextMenuContext
#trivial

diff --git a/Engine/Source/Editor/Persona/Private/AnimNotifyPanelContextMenuContext.h b/Engine/Source/Editor/Persona/Private/AnimNotifyPanelContextMenuContext.h
--- a/Engine/Source/Editor/Persona/Private/AnimNotifyPanelContextMenuContext.h
+++ b/Engine/Source/Editor/Persona/Private/AnimNotifyPanelContextMenuContext.h
@@ -3,6 +3,7 @@
 #pragma once
 
 #include "CoreMinimal.h"
+#include "Input/Events.h"
 #include "AnimNotifyPanelContextMenuContext.generated.h"
 
 class SAnimNotifyTrack;
