Respect Persona bExposeNotifiesUICommands option in NotifyPanel

diff --git a/Engine/Source/Editor/Persona/Private/SSkeletonAnimNotifies.cpp b/Engine/Source/Editor/Persona/Private/SSkeletonAnimNotifies.cpp
--- a/Engine/Source/Editor/Persona/Private/SSkeletonAnimNotifies.cpp
+++ b/Engine/Source/Editor/Persona/Private/SSkeletonAnimNotifies.cpp
@@ -11,6 +11,7 @@
 #include "AssetRegistry/AssetData.h"
 #include "Animation/AnimSequenceBase.h"
 #include "Styling/AppStyle.h"
+#include "Preferences/PersonaOptions.h"
 #include "Animation/EditorSkeletonNotifyObj.h"
 #include "Widgets/Input/SSearchBox.h"
 #include "Widgets/Text/SInlineEditableTextBlock.h"
@@ -103,13 +104,15 @@ TSharedRef<SWidget> FSkeletonAnimNotifiesSummoner::CreateTabBody(const FWorkflow
 
 void SSkeletonAnimNotifies::Construct(const FArguments& InArgs, const TSharedPtr<class FAssetEditorToolkit>& InHostingApp)
 {
+	bool bNotifiesAllowed = GetDefault<UPersonaOptions>()->bExposeNotifiesUICommands;
+
 	OnObjectsSelected = InArgs._OnObjectsSelected;
 	OnItemSelected = InArgs._OnItemSelected;
 	bIsPicker = InArgs._IsPicker;
 	bShowSyncMarkers = InArgs._ShowSyncMarkers;
 	bShowOtherAssets = InArgs._ShowOtherAssets;
 	bShowCompatibleSkeletonAssets = InArgs._ShowCompatibleSkeletonAssets;
-	bShowNotifies = InArgs._ShowNotifies;
+	bShowNotifies = InArgs._ShowNotifies && bNotifiesAllowed;
 	EditableSkeleton = InArgs._EditableSkeleton;
 
 	WeakHostingApp = InHostingApp;
@@ -141,7 +144,7 @@ void SSkeletonAnimNotifies::Construct(const FArguments& InArgs, const TSharedPtr
 	
 	CurrentFilterFlags = EAnimNotifyFilterFlags::None;
 
-	TSharedPtr<FFilterCategory> FilterCategory = MakeShared<FFilterCategory>(LOCTEXT("AnimNotifyFiltersLabel", "Anim Notify Filters"), LOCTEXT("AnimNotifyFiltersLabelToolTip", "Filter what kind fo notifies and sync markers can be displayed."));
+	TSharedPtr<FFilterCategory> FilterCategory = MakeShared<FFilterCategory>(LOCTEXT("AnimNotifyFiltersLabel", "Anim Notify Filters"), LOCTEXT("AnimNotifyFiltersLabelToolTip", "Filter what kind of notifies and sync markers can be displayed."));
 
 	const bool bSingleType = (bShowNotifies ^ bShowSyncMarkers);
 
@@ -388,12 +391,17 @@ TSharedPtr<SWidget> SSkeletonAnimNotifies::OnGetContextMenuContent() const
 	const bool bShouldCloseWindowAfterMenuSelection = true;
 	FMenuBuilder MenuBuilder( bShouldCloseWindowAfterMenuSelection, NULL);
 
-	MenuBuilder.BeginSection("AnimItemAction", LOCTEXT("ItemActions", "New Item"));
+	bool bNotifiesAllowed = GetDefault<UPersonaOptions>()->bExposeNotifiesUICommands;
+
+	if (bNotifiesAllowed)
 	{
-		FUIAction Action = FUIAction(FExecuteAction::CreateSP(const_cast<SSkeletonAnimNotifies*>(this), &SSkeletonAnimNotifies::OnAddItem, false));
-		const FText Label = LOCTEXT("NewAnimNotifyButtonLabel", "New Notify...");
-		const FText ToolTipText = LOCTEXT("NewAnimNotifyButtonTooltip", "Creates a new anim notify.");
-		MenuBuilder.AddMenuEntry(Label, ToolTipText, FSlateIcon(), Action);
+		MenuBuilder.BeginSection("AnimItemAction", LOCTEXT("ItemActions", "New Item"));
+		{
+			FUIAction Action = FUIAction(FExecuteAction::CreateSP(const_cast<SSkeletonAnimNotifies*>(this), &SSkeletonAnimNotifies::OnAddItem, false));
+			const FText Label = LOCTEXT("NewAnimNotifyButtonLabel", "New Notify...");
+			const FText ToolTipText = LOCTEXT("NewAnimNotifyButtonTooltip", "Creates a new anim notify.");
+			MenuBuilder.AddMenuEntry(Label, ToolTipText, FSlateIcon(), Action);
+		}
 	}
 
 	{
@@ -424,11 +432,16 @@ TSharedPtr<SWidget> SSkeletonAnimNotifies::OnGetContextMenuContent() const
 
 		if(WeakHostingApp.IsValid() && NotifiesListView->GetNumItemsSelected() == 1)
 		{
-			FUIAction Action = FUIAction(FExecuteAction::CreateSP(const_cast<SSkeletonAnimNotifies*>(this), &SSkeletonAnimNotifies::OnFindReferences),
-			FCanExecuteAction::CreateSP(this, &SSkeletonAnimNotifies::CanPerformFindReferences));
-			const FText Label = LOCTEXT("FindNotifyReferences", "Find/Replace References...");
-			const FText ToolTipText = LOCTEXT("FindNotifyReferencesTooltip", "Find, replace and remove references to this item in the find/replace tab");
-			MenuBuilder.AddMenuEntry(Label, ToolTipText, FSlateIcon(), Action);
+			TSharedPtr<FAssetEditorToolkit> HostingApp = WeakHostingApp.Pin();
+			bool bFindReplaceAvailable = HostingApp->GetTabManager()->GetTabPermissionList()->PassesFilter(FPersonaTabs::FindReplaceID);
+			if (bFindReplaceAvailable)
+			{
+				FUIAction Action = FUIAction(FExecuteAction::CreateSP(const_cast<SSkeletonAnimNotifies*>(this), &SSkeletonAnimNotifies::OnFindReferences),
+					FCanExecuteAction::CreateSP(this, &SSkeletonAnimNotifies::CanPerformFindReferences));
+				const FText Label = LOCTEXT("FindNotifyReferences", "Find/Replace References...");
+				const FText ToolTipText = LOCTEXT("FindNotifyReferencesTooltip", "Find, replace and remove references to this item in the find/replace tab");
+				MenuBuilder.AddMenuEntry(Label, ToolTipText, FSlateIcon(), Action);
+			}
 		}
 	}
 	MenuBuilder.EndSection();
