Fix viewport cursor style in animation editors when hovering mesh info text

diff --git a/Engine/Source/Editor/Persona/Private/SAnimViewportToolBar.cpp b/Engine/Source/Editor/Persona/Private/SAnimViewportToolBar.cpp
--- a/Engine/Source/Editor/Persona/Private/SAnimViewportToolBar.cpp
+++ b/Engine/Source/Editor/Persona/Private/SAnimViewportToolBar.cpp
@@ -410,6 +410,7 @@ void SAnimViewportToolBar::Construct(const FArguments& InArgs, TSharedPtr<class
 		[
 			// Display text (e.g., item being previewed)
 			SNew(SRichTextBlock)
+			.Visibility(EVisibility::SelfHitTestInvisible)
 			.DecoratorStyleSet(&FAppStyle::Get())
 			.Text(InViewport.Get(), &SAnimationEditorViewportTabBody::GetDisplayString)
 			.TextStyle(&FAppStyle::Get().GetWidgetStyle<FTextBlockStyle>("AnimViewport.MessageText"))
