Fix delete material to repatch the lod material map correctly.

#rnx

diff --git a/Engine/Source/Developer/SkeletalMeshUtilitiesCommon/Private/LODUtilities.cpp b/Engine/Source/Developer/SkeletalMeshUtilitiesCommon/Private/LODUtilities.cpp
--- a/Engine/Source/Developer/SkeletalMeshUtilitiesCommon/Private/LODUtilities.cpp
+++ b/Engine/Source/Developer/SkeletalMeshUtilitiesCommon/Private/LODUtilities.cpp
@@ -3899,6 +3899,7 @@ void FLODUtilities::AdjustImportDataFaceMaterialIndex(const TArray<FSkeletalMate
 	MaterialRemap.Reserve(RawMeshMaterials.Num());
 	//Optimization to avoid doing the remap if no material have to change
 	bool bNeedRemapping = false;
+	bool bFoundMaterial = false;
 	for (int32 MaterialIndex = 0; MaterialIndex < RawMeshMaterials.Num(); ++MaterialIndex)
 	{
 		MaterialRemap.Add(MaterialIndex);
@@ -3910,9 +3911,15 @@ void FLODUtilities::AdjustImportDataFaceMaterialIndex(const TArray<FSkeletalMate
 			{
 				bNeedRemapping |= (MaterialRemap[MaterialIndex] != MeshMaterialIndex);
 				MaterialRemap[MaterialIndex] = MeshMaterialIndex;
+				bFoundMaterial = true;
 				break;
 			}
 		}
+		//A material slot was deleted by the user, we cannot remap anything in that case, the LODMaterialMap should be active for the LOD 0 and will match the sections
+		if (!bFoundMaterial)
+		{
+			return;
+		}
 	}
 	if (bNeedRemapping)
 	{
diff --git a/Engine/Source/Editor/Persona/Private/PersonaMeshDetails.cpp b/Engine/Source/Editor/Persona/Private/PersonaMeshDetails.cpp
--- a/Engine/Source/Editor/Persona/Private/PersonaMeshDetails.cpp
+++ b/Engine/Source/Editor/Persona/Private/PersonaMeshDetails.cpp
@@ -5031,6 +5031,8 @@ void FPersonaMeshDetails::OnDeleteMaterialSlot(int32 MaterialIndex)
 	{
 		FScopedSkeletalMeshPostEditChange ScopedPostEditChange(SkeletalMeshPtr.Get());
 		//When we delete a material slot we must invalidate the DDC because material index is not part of the DDC key by design
+		SkeletalMeshPtr->InvalidateDeriveDataCacheGUID();
+
 		SkeletalMeshPtr->GetMaterials().RemoveAt(MaterialIndex);
 		FSkeletalMeshModel* Model = SkeletalMeshPtr->GetImportedModel();
 
@@ -5039,10 +5041,11 @@ void FPersonaMeshDetails::OnDeleteMaterialSlot(int32 MaterialIndex)
 		//When we delete a material slot we need to fix all MaterialIndex after the deleted index
 		for (int32 LODInfoIdx = 0; LODInfoIdx < NumLODInfos; LODInfoIdx++)
 		{
+			TArray<FSkelMeshSection>& Sections = Model->LODModels[LODInfoIdx].Sections;
 			TArray<int32>& LODMaterialMap = SkeletalMeshPtr->GetLODInfo(LODInfoIdx)->LODMaterialMap;
-			for (int32 SectionIndex = 0; SectionIndex < Model->LODModels[LODInfoIdx].Sections.Num(); ++SectionIndex)
+			for (int32 SectionIndex = 0; SectionIndex < Sections.Num(); ++SectionIndex)
 			{
-				int32 SectionMaterialIndex = Model->LODModels[LODInfoIdx].Sections[SectionIndex].MaterialIndex;
+				int32 SectionMaterialIndex = Sections[SectionIndex].MaterialIndex;
 				if (LODMaterialMap.IsValidIndex(SectionIndex) && LODMaterialMap[SectionIndex] != INDEX_NONE)
 				{
 					SectionMaterialIndex = LODMaterialMap[SectionIndex];
@@ -5050,14 +5053,13 @@ void FPersonaMeshDetails::OnDeleteMaterialSlot(int32 MaterialIndex)
 				if (SectionMaterialIndex > MaterialIndex)
 				{
 					SectionMaterialIndex--;
-				}
-				if (SectionMaterialIndex != Model->LODModels[LODInfoIdx].Sections[SectionIndex].MaterialIndex)
-				{
-					while(!LODMaterialMap.IsValidIndex(SectionIndex))
+					//Patch the lod material map
+					while (!LODMaterialMap.IsValidIndex(SectionIndex))
 					{
 						LODMaterialMap.Add(INDEX_NONE);
 					}
 					LODMaterialMap[SectionIndex] = SectionMaterialIndex;
+					Sections[SectionIndex].MaterialIndex--;
 				}
 			}
 		}
