Optimize the opening of notifies/sync marker widget

Scope refreshes so that list is not refresh multiple times as filters are constructed
Avoid expensive compatibility checks when none are needed (when no asset registry data is present)
Avoid memory allocations and FindObject() calls in skeleton compatibility checks

diff --git a/Engine/Source/Editor/Persona/Private/SSkeletonAnimNotifies.cpp b/Engine/Source/Editor/Persona/Private/SSkeletonAnimNotifies.cpp
--- a/Engine/Source/Editor/Persona/Private/SSkeletonAnimNotifies.cpp
+++ b/Engine/Source/Editor/Persona/Private/SSkeletonAnimNotifies.cpp
@@ -148,154 +148,158 @@ void SSkeletonAnimNotifies::Construct(const FArguments& InArgs, const TSharedPtr
 
 	const bool bSingleType = (bShowNotifies ^ bShowSyncMarkers);
 
-	TSharedRef<SBasicFilterBar<EAnimNotifyFilterFlags>> FilterBar = SNew(SBasicFilterBar<EAnimNotifyFilterFlags>)
-	.CustomFilters(Filters)
-	.bPinAllFrontendFilters(true)
-	.UseSectionsForCategories(true)
-	.OnFilterChanged_Lambda([this]()
 	{
-		for(const TSharedRef<FFilterBase<EAnimNotifyFilterFlags>>& Filter : Filters)
+		TGuardValue<bool> SuspendRefreshFilter(bAllowRefreshFilter, false);
+		
+		TSharedRef<SBasicFilterBar<EAnimNotifyFilterFlags>> FilterBar = SNew(SBasicFilterBar<EAnimNotifyFilterFlags>)
+		.CustomFilters(Filters)
+		.bPinAllFrontendFilters(true)
+		.UseSectionsForCategories(true)
+		.OnFilterChanged_Lambda([this]()
 		{
-			TSharedRef<FSkeletonAnimNotifiesFilter> AnimNotifiesFilter = StaticCastSharedRef<FSkeletonAnimNotifiesFilter>(Filter);
-			if(AnimNotifiesFilter->IsActive())
+			for(const TSharedRef<FFilterBase<EAnimNotifyFilterFlags>>& Filter : Filters)
 			{
-				CurrentFilterFlags |= AnimNotifiesFilter->GetFlags();
-			}
-			else
-			{
-				CurrentFilterFlags &= ~AnimNotifiesFilter->GetFlags();
+				TSharedRef<FSkeletonAnimNotifiesFilter> AnimNotifiesFilter = StaticCastSharedRef<FSkeletonAnimNotifiesFilter>(Filter);
+				if(AnimNotifiesFilter->IsActive())
+				{
+					CurrentFilterFlags |= AnimNotifiesFilter->GetFlags();
+				}
+				else
+				{
+					CurrentFilterFlags &= ~AnimNotifiesFilter->GetFlags();
+				}
 			}
-		}
 
-		RefreshNotifiesListWithFilter();
-	});
-	
-	if (EditableSkeleton.IsValid())
-	{
-		CurrentFilterFlags |= EAnimNotifyFilterFlags::CurrentSkeleton;
-
-		TSharedRef<FFilterBase<EAnimNotifyFilterFlags>> Filter = Filters.Add_GetRef(MakeShared<FSkeletonAnimNotifiesFilter>(
-			EAnimNotifyFilterFlags::CurrentSkeleton,
-			"CurrentSkeleton",
-			LOCTEXT("ShowSkeletonItemsLabel", "Skeleton"),
-			LOCTEXT("ShowSkeletonItemsTooltip", "Show items for the current skeleton"),
-			FLinearColor::Blue.Desaturate(0.25f),
-			FilterCategory
-		));
-
-		FilterBar->AddFilter(Filter);
-		Filter->SetActive(true);
-	}
-
-	if (bShowNotifies)
-	{
-		CurrentFilterFlags |= EAnimNotifyFilterFlags::Notifies;
-	}
-
-	if (!bSingleType && bShowNotifies)
-	{
-		TSharedRef<FFilterBase<EAnimNotifyFilterFlags>> Filter = Filters.Add_GetRef(MakeShared<FSkeletonAnimNotifiesFilter>(
-			EAnimNotifyFilterFlags::Notifies,
-			"Notifies",
-			LOCTEXT("ShowNotifiesLabel", "Notifies"),
-			LOCTEXT("ShowNotifiesTooltip", "Show notifies"),
-			FLinearColor::Red.Desaturate(0.5f),
-			FilterCategory
-		));
-
-		FilterBar->AddFilter(Filter);
-		Filter->SetActive(true);
-	}
-
-	if (bShowSyncMarkers)
-	{
-		CurrentFilterFlags |= EAnimNotifyFilterFlags::SyncMarkers;
-	}
-
-	if (!bSingleType && bShowSyncMarkers)
-	{
-		TSharedRef<FFilterBase<EAnimNotifyFilterFlags>> Filter = Filters.Add_GetRef(MakeShared<FSkeletonAnimNotifiesFilter>(
-			EAnimNotifyFilterFlags::SyncMarkers,
-			"SyncMarkers",
-			LOCTEXT("ShowSyncMarkersLabel", "Sync Markers"),
-			LOCTEXT("ShowSyncMarkersTooltip", "Show sync markers"),
-			FLinearColor::Green.Desaturate(0.5f),
-			FilterCategory
-		));
-
-		FilterBar->AddFilter(Filter);
-		Filter->SetActive(true);
-	}
-
-	{
-		CurrentFilterFlags |= bShowOtherAssets ? EAnimNotifyFilterFlags::OtherAssets : EAnimNotifyFilterFlags::None;
+			RefreshNotifiesListWithFilter();
+		});
+		
+		if (EditableSkeleton.IsValid())
+		{
+			CurrentFilterFlags |= EAnimNotifyFilterFlags::CurrentSkeleton;
+
+			TSharedRef<FFilterBase<EAnimNotifyFilterFlags>> Filter = Filters.Add_GetRef(MakeShared<FSkeletonAnimNotifiesFilter>(
+				EAnimNotifyFilterFlags::CurrentSkeleton,
+				"CurrentSkeleton",
+				LOCTEXT("ShowSkeletonItemsLabel", "Skeleton"),
+				LOCTEXT("ShowSkeletonItemsTooltip", "Show items for the current skeleton"),
+				FLinearColor::Blue.Desaturate(0.25f),
+				FilterCategory
+			));
+
+			FilterBar->AddFilter(Filter);
+			Filter->SetActive(true);
+		}
 
-		TSharedRef<FFilterBase<EAnimNotifyFilterFlags>> Filter = Filters.Add_GetRef(MakeShared<FSkeletonAnimNotifiesFilter>(
-			EAnimNotifyFilterFlags::OtherAssets,
-			"OtherAssets",
-			LOCTEXT("ShowOtherAssetsLabel", "Other Assets"),
-			LOCTEXT("ShowOtherAssetsTooltip", "Show items that are present on other assets"),
-			FLinearColor::Yellow.Desaturate(0.5f),
-			FilterCategory
-		));
+		if (bShowNotifies)
+		{
+			CurrentFilterFlags |= EAnimNotifyFilterFlags::Notifies;
+		}
 
-		FilterBar->AddFilter(Filter);
-		Filter->SetActive(bShowOtherAssets);
-	}
+		if (!bSingleType && bShowNotifies)
+		{
+			TSharedRef<FFilterBase<EAnimNotifyFilterFlags>> Filter = Filters.Add_GetRef(MakeShared<FSkeletonAnimNotifiesFilter>(
+				EAnimNotifyFilterFlags::Notifies,
+				"Notifies",
+				LOCTEXT("ShowNotifiesLabel", "Notifies"),
+				LOCTEXT("ShowNotifiesTooltip", "Show notifies"),
+				FLinearColor::Red.Desaturate(0.5f),
+				FilterCategory
+			));
+
+			FilterBar->AddFilter(Filter);
+			Filter->SetActive(true);
+		}
 
-	{
-		CurrentFilterFlags |= bShowCompatibleSkeletonAssets ? EAnimNotifyFilterFlags::CompatibleAssets : EAnimNotifyFilterFlags::None;
+		if (bShowSyncMarkers)
+		{
+			CurrentFilterFlags |= EAnimNotifyFilterFlags::SyncMarkers;
+		}
 
-		TSharedRef<FFilterBase<EAnimNotifyFilterFlags>> Filter = Filters.Add_GetRef(MakeShared<FSkeletonAnimNotifiesFilter>(
-			EAnimNotifyFilterFlags::CompatibleAssets,
-			"CompatibleAssets",
-			LOCTEXT("ShowCompatibleAssetsLabel", "Compatible"),
-			LOCTEXT("ShowCompatibleAssetsTooltip", "Show items that are present on other assets that are compatible with the current skeleton"),
-			FLinearColor::Blue.Desaturate(0.5f),
-			FilterCategory
-		));
+		if (!bSingleType && bShowSyncMarkers)
+		{
+			TSharedRef<FFilterBase<EAnimNotifyFilterFlags>> Filter = Filters.Add_GetRef(MakeShared<FSkeletonAnimNotifiesFilter>(
+				EAnimNotifyFilterFlags::SyncMarkers,
+				"SyncMarkers",
+				LOCTEXT("ShowSyncMarkersLabel", "Sync Markers"),
+				LOCTEXT("ShowSyncMarkersTooltip", "Show sync markers"),
+				FLinearColor::Green.Desaturate(0.5f),
+				FilterCategory
+			));
+
+			FilterBar->AddFilter(Filter);
+			Filter->SetActive(true);
+		}
 
-		FilterBar->AddFilter(Filter);
-		Filter->SetActive(bShowCompatibleSkeletonAssets);
-	}
+		{
+			CurrentFilterFlags |= bShowOtherAssets ? EAnimNotifyFilterFlags::OtherAssets : EAnimNotifyFilterFlags::None;
+
+			TSharedRef<FFilterBase<EAnimNotifyFilterFlags>> Filter = Filters.Add_GetRef(MakeShared<FSkeletonAnimNotifiesFilter>(
+				EAnimNotifyFilterFlags::OtherAssets,
+				"OtherAssets",
+				LOCTEXT("ShowOtherAssetsLabel", "Other Assets"),
+				LOCTEXT("ShowOtherAssetsTooltip", "Show items that are present on other assets"),
+				FLinearColor::Yellow.Desaturate(0.5f),
+				FilterCategory
+			));
+
+			FilterBar->AddFilter(Filter);
+			Filter->SetActive(bShowOtherAssets);
+		}
 
-	TSharedRef<SWidget> AddFilterButton = SBasicFilterBar<EAnimNotifyFilterFlags>::MakeAddFilterButton(FilterBar);
-	
-	ChildSlot
-	[
-		SNew( SVerticalBox )
+		{
+			CurrentFilterFlags |= bShowCompatibleSkeletonAssets ? EAnimNotifyFilterFlags::CompatibleAssets : EAnimNotifyFilterFlags::None;
+
+			TSharedRef<FFilterBase<EAnimNotifyFilterFlags>> Filter = Filters.Add_GetRef(MakeShared<FSkeletonAnimNotifiesFilter>(
+				EAnimNotifyFilterFlags::CompatibleAssets,
+				"CompatibleAssets",
+				LOCTEXT("ShowCompatibleAssetsLabel", "Compatible"),
+				LOCTEXT("ShowCompatibleAssetsTooltip", "Show items that are present on other assets that are compatible with the current skeleton"),
+				FLinearColor::Blue.Desaturate(0.5f),
+				FilterCategory
+			));
+
+			FilterBar->AddFilter(Filter);
+			Filter->SetActive(bShowCompatibleSkeletonAssets);
+		}
 
-		+SVerticalBox::Slot()
-		.AutoHeight()
-		.Padding( FMargin( 0.0f, 0.0f, 0.0f, 4.0f ) )
+		TSharedRef<SWidget> AddFilterButton = SBasicFilterBar<EAnimNotifyFilterFlags>::MakeAddFilterButton(FilterBar);
+		
+		ChildSlot
 		[
-			SNew(SHorizontalBox)
-			+SHorizontalBox::Slot()
-			.AutoWidth()
-			.Padding(2.0f,0.0f)
+			SNew( SVerticalBox )
+
+			+SVerticalBox::Slot()
+			.AutoHeight()
+			.Padding( FMargin( 0.0f, 0.0f, 0.0f, 4.0f ) )
 			[
-				AddFilterButton
+				SNew(SHorizontalBox)
+				+SHorizontalBox::Slot()
+				.AutoWidth()
+				.Padding(2.0f,0.0f)
+				[
+					AddFilterButton
+				]
+
+				+SHorizontalBox::Slot()
+				.FillWidth(1.0f)
+				[
+					NameFilterBox.ToSharedRef()
+				]
 			]
 
-			+SHorizontalBox::Slot()
-			.FillWidth(1.0f)
+			+SVerticalBox::Slot()
+			.AutoHeight()
 			[
-				NameFilterBox.ToSharedRef()
+				FilterBar
 			]
-		]
-
-		+SVerticalBox::Slot()
-		.AutoHeight()
-		[
-			FilterBar
-		]
 
-		+SVerticalBox::Slot()
-		.FillHeight( 1.0f )		// This is required to make the scrollbar work, as content overflows Slate containers by default
-		[
-			NotifiesListView.ToSharedRef()
-		]
-	];
+			+SVerticalBox::Slot()
+			.FillHeight( 1.0f )		// This is required to make the scrollbar work, as content overflows Slate containers by default
+			[
+				NotifiesListView.ToSharedRef()
+			]
+		];
+	}
 
 	RefreshNotifiesListWithFilter();
 }
@@ -679,8 +683,11 @@ void SSkeletonAnimNotifies::OnNotifyNameCommitted( const FText& NewName, ETextCo
 
 void SSkeletonAnimNotifies::RefreshNotifiesListWithFilter()
 {
-	CreateNotifiesList();
-	FilterNotifiesList(NameFilterBox->GetText().ToString());
+	if(bAllowRefreshFilter)
+	{
+		CreateNotifiesList();
+		FilterNotifiesList(NameFilterBox->GetText().ToString());
+	}
 }
 
 void SSkeletonAnimNotifies::FilterNotifiesList(const FString& InSearchText)
@@ -708,10 +715,10 @@ void SSkeletonAnimNotifies::FilterNotifiesList(const FString& InSearchText)
 void SSkeletonAnimNotifies::CreateNotifiesList()
 {
 	const USkeleton* CurrentSkeleton = EditableSkeleton.IsValid() ? &EditableSkeleton->GetSkeleton() : nullptr;
-	FString CurrentSkeletonName;
+	FAssetData CurrentSkeletonAsset;
 	if(CurrentSkeleton)
 	{
-		CurrentSkeletonName = FAssetData(CurrentSkeleton).GetExportTextName();
+		CurrentSkeletonAsset = FAssetData(CurrentSkeleton);
 	}
 
 	NotifyList.Empty();
@@ -722,7 +729,9 @@ void SSkeletonAnimNotifies::CreateNotifiesList()
 
 	FARFilter Filter;
 	Filter.bRecursiveClasses = true;
-	Filter.ClassPaths.Append({ UAnimationAsset::StaticClass()->GetClassPathName(), USkeleton::StaticClass()->GetClassPathName() } );
+	Filter.ClassPaths.Append({ UAnimSequenceBase::StaticClass()->GetClassPathName(), USkeleton::StaticClass()->GetClassPathName() } );
+	Filter.TagsAndValues.Add(USkeleton::AnimNotifyTag);
+	Filter.TagsAndValues.Add(USkeleton::AnimSyncMarkerTag);
 
 	TArray<FAssetData> FoundAssetData;
 	AssetRegistryModule.Get().GetAssets(Filter, FoundAssetData);
@@ -735,6 +744,27 @@ void SSkeletonAnimNotifies::CreateNotifiesList()
 	{
 		EAnimNotifyFilterFlags AssetFlags = EAnimNotifyFilterFlags::None;
 
+		auto HasNotifies = [&AssetData]()
+		{
+			FAssetTagValueRef TagRef = AssetData.TagsAndValues.FindTag(USkeleton::AnimNotifyTag);
+			return TagRef.IsSet() && !TagRef.Equals(USkeleton::AnimNotifyTagDelimiter);
+		};
+
+		auto HasSyncMarkers = [&AssetData]()
+		{
+			FAssetTagValueRef TagRef = AssetData.TagsAndValues.FindTag(USkeleton::AnimSyncMarkerTag);
+			return TagRef.IsSet() && !TagRef.Equals(USkeleton::AnimSyncMarkerTagDelimiter);
+		};
+		
+		bool bHasAssetRegistryData =
+			(EnumHasAnyFlags(CurrentFilterFlags, EAnimNotifyFilterFlags::Notifies) && HasNotifies()) ||
+			(EnumHasAnyFlags(CurrentFilterFlags, EAnimNotifyFilterFlags::SyncMarkers) && HasSyncMarkers());
+
+		if(!bHasAssetRegistryData)
+		{
+			continue;
+		}
+
 		if (AssetData.GetClass() != USkeleton::StaticClass())
 		{
 			if (CurrentSkeleton && CurrentSkeleton->IsCompatibleForEditor(AssetData))
@@ -748,7 +778,7 @@ void SSkeletonAnimNotifies::CreateNotifiesList()
 		}
 		else
 		{
-			if (AssetData.GetExportTextName() == CurrentSkeletonName)
+			if (AssetData == CurrentSkeletonAsset)
 			{
 				AssetFlags |= EAnimNotifyFilterFlags::CurrentSkeleton;
 			}
diff --git a/Engine/Source/Editor/Persona/Private/SSkeletonAnimNotifies.h b/Engine/Source/Editor/Persona/Private/SSkeletonAnimNotifies.h
--- a/Engine/Source/Editor/Persona/Private/SSkeletonAnimNotifies.h
+++ b/Engine/Source/Editor/Persona/Private/SSkeletonAnimNotifies.h
@@ -291,6 +291,9 @@ private:
 	/** Whether we notifies and sync markers from assets compatible with the current skeleton */
 	bool bShowCompatibleSkeletonAssets = false;
 
+	/** Whether to suspend refreshing the UI when filtering */
+	bool bAllowRefreshFilter = true;
+
 	/** All filters that can be applied to the widget's display */
 	TArray<TSharedRef<FFilterBase<EAnimNotifyFilterFlags>>> Filters;
 
diff --git a/Engine/Source/Runtime/Engine/Private/Animation/Skeleton.cpp b/Engine/Source/Runtime/Engine/Private/Animation/Skeleton.cpp
--- a/Engine/Source/Runtime/Engine/Private/Animation/Skeleton.cpp
+++ b/Engine/Source/Runtime/Engine/Private/Animation/Skeleton.cpp
@@ -143,8 +143,10 @@ bool USkeleton::IsCompatibleForEditor(const FAssetData& AssetData, const TCHAR*
 bool USkeleton::IsCompatibleForEditor(const FString& SkeletonAssetString) const
 {
 	// First check against itself.
-	const FString SkeletonString = FAssetData(this, FAssetData::ECreationFlags::SkipAssetRegistryTagsGathering).GetExportTextName();
-	if (SkeletonString == SkeletonAssetString)
+	TStringBuilder<128> SkeletonStringBuilder;
+	FAssetData(this, FAssetData::ECreationFlags::SkipAssetRegistryTagsGathering).GetExportTextName(SkeletonStringBuilder);
+	const TCHAR* SkeletonString = SkeletonStringBuilder.ToString();
+	if (SkeletonAssetString == SkeletonString)
 	{
 		return true;
 	}
@@ -156,18 +158,22 @@ bool USkeleton::IsCompatibleForEditor(const FString& SkeletonAssetString) const
 	}
 
 	// Now check against the list of compatible skeletons and see if we're dealing with the same asset.
-	const FSoftObjectPath InPath(SkeletonAssetString);
-	for (const TSoftObjectPtr<USkeleton>& CompatibleSkeleton : CompatibleSkeletons)
+	if(CompatibleSkeletons.Num() > 0)
 	{
-		if (CompatibleSkeleton.ToSoftObjectPath() == InPath)
+		const FSoftObjectPath InPath(SkeletonAssetString);
+		for (const TSoftObjectPtr<USkeleton>& CompatibleSkeleton : CompatibleSkeletons)
 		{
-			return true;
+			if (CompatibleSkeleton.ToSoftObjectPath() == InPath)
+			{
+				return true;
+			}
 		}
 	}
 
 	// Check if the other skeleton is compatible with this via the asset registry
 	const IAssetRegistry& AssetRegistry = FModuleManager::LoadModuleChecked<FAssetRegistryModule>("AssetRegistry").Get();
-	const FAssetData SkeletonAssetData = AssetRegistry.GetAssetByObjectPath(FSoftObjectPath(*SkeletonAssetString));	
+	constexpr bool bIncludeOnlyOnDiskAssets = true;
+	const FAssetData SkeletonAssetData = AssetRegistry.GetAssetByObjectPath(FSoftObjectPath(*SkeletonAssetString), bIncludeOnlyOnDiskAssets);
 	const FString TagValue = SkeletonAssetData.GetTagValueRef<FString>(USkeleton::CompatibleSkeletonsNameTag);
 	if (!TagValue.IsEmpty())
 	{
@@ -176,7 +182,7 @@ bool USkeleton::IsCompatibleForEditor(const FString& SkeletonAssetString) const
 		{
 			for (const FString& OtherCompatibleSkeleton : OtherCompatibleSkeletons)
 			{
-				if (SkeletonString == OtherCompatibleSkeleton)
+				if (OtherCompatibleSkeleton == SkeletonString)
 				{
 					return true;
 				}
