Fixed an issue where the animation notify copy / paste menu items stopped working

diff --git a/Engine/Source/Editor/Persona/Private/SAnimNotifyPanel.cpp b/Engine/Source/Editor/Persona/Private/SAnimNotifyPanel.cpp
--- a/Engine/Source/Editor/Persona/Private/SAnimNotifyPanel.cpp
+++ b/Engine/Source/Editor/Persona/Private/SAnimNotifyPanel.cpp
@@ -2789,7 +2789,13 @@ TSharedPtr<SWidget> SAnimNotifyTrack::SummonContextMenu(const FGeometry& MyGeome
 		CreateContextMenuContent(BaseMenuName);
 	}
 
-	TSharedPtr<SWidget> MenuWidget = ToolMenus->GenerateWidget(BaseMenuName, FToolMenuContext(MenuContext));
+	FToolMenuContext ToolMenuContext(MenuContext);
+	if (WeakCommandList.IsValid())
+	{
+		ToolMenuContext.AppendCommandList(WeakCommandList.Pin());
+	}
+
+	TSharedPtr<SWidget> MenuWidget = ToolMenus->GenerateWidget(BaseMenuName, ToolMenuContext);
 
 	if (MenuWidget.IsValid())
 	{
