FIx persona toolkit emitting an assert when switching preview mesh on a control rig blueprint
#rnx

diff --git a/Engine/Source/Editor/Persona/Private/PersonaToolkit.cpp b/Engine/Source/Editor/Persona/Private/PersonaToolkit.cpp
--- a/Engine/Source/Editor/Persona/Private/PersonaToolkit.cpp
+++ b/Engine/Source/Editor/Persona/Private/PersonaToolkit.cpp
@@ -359,14 +359,14 @@ void FPersonaToolkit::SetPreviewMesh(class USkeletalMesh* InSkeletalMesh, bool b
 				check(PhysicsAsset);
 				PhysicsAsset->SetPreviewMesh(InSkeletalMesh);
 			}
-			else if(EditableSkeleton.IsValid())
-			{
-				EditableSkeleton->SetPreviewMesh(InSkeletalMesh);
-			}
 			else if(IInterface_PreviewMeshProvider* PreviewMeshInterface = Cast<IInterface_PreviewMeshProvider>(Asset))
 			{
 				PreviewMeshInterface->SetPreviewMesh(InSkeletalMesh);
 			}
+			else if(EditableSkeleton.IsValid())
+			{
+				EditableSkeleton->SetPreviewMesh(InSkeletalMesh);
+			}
 		}
 
 		if(bReOpenEditor)
@@ -388,6 +388,10 @@ void FPersonaToolkit::SetPreviewMesh(class USkeletalMesh* InSkeletalMesh, bool b
 			{
 				AssetToReopen = Skeleton.Get();
 			}
+			else if(IInterface_PreviewMeshProvider* PreviewMeshInterface = Cast<IInterface_PreviewMeshProvider>(Asset))
+			{
+				AssetToReopen = Asset;
+			}
 
 			check(AssetToReopen);
 
