Changed PersonaToolkit raw object pointers to weak object pointers, so validity can be checked

diff --git a/Engine/Source/Editor/Persona/Private/PersonaToolkit.cpp b/Engine/Source/Editor/Persona/Private/PersonaToolkit.cpp
--- a/Engine/Source/Editor/Persona/Private/PersonaToolkit.cpp
+++ b/Engine/Source/Editor/Persona/Private/PersonaToolkit.cpp
@@ -79,7 +79,7 @@ void FPersonaToolkit::Initialize(UAnimationAsset* InAnimationAsset, const FPerso
 
 	if (AnimationAsset != nullptr)
 	{
-		PreviewScene->SetPreviewAnimationAsset(AnimationAsset);
+		PreviewScene->SetPreviewAnimationAsset(AnimationAsset.Get());
 	}
 }
 
@@ -245,12 +245,12 @@ void FPersonaToolkit::SetMesh(class USkeletalMesh* InSkeletalMesh)
 
 UAnimBlueprint* FPersonaToolkit::GetAnimBlueprint() const
 {
-	return AnimBlueprint;
+	return AnimBlueprint.Get();
 }
 
 UAnimationAsset* FPersonaToolkit::GetAnimationAsset() const
 {
-	return AnimationAsset;
+	return AnimationAsset.Get();
 }
 
 void FPersonaToolkit::SetAnimationAsset(class UAnimationAsset* InAnimationAsset)
@@ -264,7 +264,7 @@ void FPersonaToolkit::SetAnimationAsset(class UAnimationAsset* InAnimationAsset)
 
 	AnimationAsset = InAnimationAsset;
 	
-	if(AnimationAsset)
+	if(AnimationAsset.IsValid())
 	{
 		check(InitialAssetClass == UAnimationAsset::StaticClass());
 
@@ -292,17 +292,17 @@ USkeletalMesh* FPersonaToolkit::GetPreviewMesh() const
 {
 	if (InitialAssetClass == UAnimationAsset::StaticClass())
 	{
-		check(AnimationAsset);
+		check(AnimationAsset.IsValid());
 		return AnimationAsset->GetPreviewMesh();
 	}
 	else if (InitialAssetClass == UAnimBlueprint::StaticClass())
 	{
-		check(AnimBlueprint);
+		check(AnimBlueprint.IsValid());
 		return AnimBlueprint->GetPreviewMesh();
 	}
 	else if (InitialAssetClass == UPhysicsAsset::StaticClass())
 	{
-		check(PhysicsAsset);
+		check(PhysicsAsset.IsValid());
 		return PhysicsAsset->GetPreviewMesh();
 	}
 	else if(InitialAssetClass == USkeletalMesh::StaticClass())
@@ -342,21 +342,21 @@ void FPersonaToolkit::SetPreviewMesh(class USkeletalMesh* InSkeletalMesh, bool b
 			{
 				FScopedTransaction Transaction(NSLOCTEXT("PersonaToolkit", "SetAnimationPreviewMesh", "Set Animation Preview Mesh"));
 
-				check(AnimationAsset);
+				check(AnimationAsset.IsValid());
 				AnimationAsset->SetPreviewMesh(InSkeletalMesh);
 			}
 			else if (InitialAssetClass == UAnimBlueprint::StaticClass())
 			{
 				FScopedTransaction Transaction(NSLOCTEXT("PersonaToolkit", "SetAnimBlueprintPreviewMesh", "Set Animation Blueprint Preview Mesh"));
 
-				check(AnimBlueprint);
+				check(AnimBlueprint.IsValid());
 				AnimBlueprint->SetPreviewMesh(InSkeletalMesh);
 			}
 			else if (InitialAssetClass == UPhysicsAsset::StaticClass())
 			{
 				FScopedTransaction Transaction(NSLOCTEXT("PersonaToolkit", "SetPhysicsAssetPreviewMesh", "Set Physics Asset Preview Mesh"));
 
-				check(PhysicsAsset);
+				check(PhysicsAsset.IsValid());
 				PhysicsAsset->SetPreviewMesh(InSkeletalMesh);
 			}
 			else if(IInterface_PreviewMeshProvider* PreviewMeshInterface = Cast<IInterface_PreviewMeshProvider>(Asset))
@@ -374,15 +374,15 @@ void FPersonaToolkit::SetPreviewMesh(class USkeletalMesh* InSkeletalMesh, bool b
 			UObject* AssetToReopen = nullptr;
 			if (InitialAssetClass == UAnimationAsset::StaticClass())
 			{
-				AssetToReopen = AnimationAsset;
+				AssetToReopen = AnimationAsset.Get();
 			}
 			else if (InitialAssetClass == UAnimBlueprint::StaticClass())
 			{
-				AssetToReopen = AnimBlueprint;
+				AssetToReopen = AnimBlueprint.Get();
 			}
 			else if (InitialAssetClass == UPhysicsAsset::StaticClass())
 			{
-				AssetToReopen = PhysicsAsset;
+				AssetToReopen = PhysicsAsset.Get();
 			}
 			else if (InitialAssetClass == USkeleton::StaticClass())
 			{
@@ -390,7 +390,7 @@ void FPersonaToolkit::SetPreviewMesh(class USkeletalMesh* InSkeletalMesh, bool b
 			}
 			else if(IInterface_PreviewMeshProvider* PreviewMeshInterface = Cast<IInterface_PreviewMeshProvider>(Asset))
 			{
-				AssetToReopen = Asset;
+				AssetToReopen = Asset.Get();
 			}
 
 			check(AssetToReopen);
@@ -412,11 +412,11 @@ void FPersonaToolkit::SetPreviewAnimationBlueprint(UAnimBlueprint* InAnimBluepri
 	{
 		FScopedTransaction Transaction(NSLOCTEXT("PersonaToolkit", "SetAnimBlueprintPreviewBlueprint", "Set Animation Blueprint Preview Blueprint"));
 
-		check(AnimBlueprint);
+		check(AnimBlueprint.IsValid());
 		AnimBlueprint->SetPreviewAnimationBlueprint(InAnimBlueprint);
 
 		// Note setting the 'edited' blueprint as an overlay here
-		GetPreviewScene()->SetPreviewAnimationBlueprint(InAnimBlueprint, AnimBlueprint);
+		GetPreviewScene()->SetPreviewAnimationBlueprint(InAnimBlueprint, AnimBlueprint.Get());
 	}
 }
 
@@ -424,7 +424,7 @@ UAnimBlueprint* FPersonaToolkit::GetPreviewAnimationBlueprint() const
 {
 	if (InitialAssetClass == UAnimBlueprint::StaticClass())
 	{
-		check(AnimBlueprint);
+		check(AnimBlueprint.IsValid());
 		return AnimBlueprint->GetPreviewAnimationBlueprint();
 	}
 
diff --git a/Engine/Source/Editor/Persona/Private/PersonaToolkit.h b/Engine/Source/Editor/Persona/Private/PersonaToolkit.h
--- a/Engine/Source/Editor/Persona/Private/PersonaToolkit.h
+++ b/Engine/Source/Editor/Persona/Private/PersonaToolkit.h
@@ -69,16 +69,16 @@ private:
 	TWeakObjectPtr <USkeletalMesh> Mesh;
 
 	/** The anim blueprint we are editing */
-	UAnimBlueprint* AnimBlueprint;
+	TWeakObjectPtr<UAnimBlueprint> AnimBlueprint;
 
 	/** the animation asset we are editing */
-	UAnimationAsset* AnimationAsset;
+	TWeakObjectPtr<UAnimationAsset> AnimationAsset;
 
 	/** the physics asset we are editing */
-	UPhysicsAsset* PhysicsAsset;
+	TWeakObjectPtr<UPhysicsAsset> PhysicsAsset;
 
 	/** The generic asset we are editing */
-	UObject* Asset;
+	TWeakObjectPtr<UObject> Asset;
 
 	/** Allow custom data for this editor */
 	TMap<int32, int32> CustomEditorData;
