Expose Bone Influence Limit property in the Skeletal Mesh Editor

diff --git a/Engine/Source/Editor/Persona/Private/PersonaMeshDetails.cpp b/Engine/Source/Editor/Persona/Private/PersonaMeshDetails.cpp
--- a/Engine/Source/Editor/Persona/Private/PersonaMeshDetails.cpp
+++ b/Engine/Source/Editor/Persona/Private/PersonaMeshDetails.cpp
@@ -1737,6 +1737,7 @@ void FSkeletalMeshBuildSettingsLayout::GenerateChildContent(IDetailChildrenBuild
 			LOCTEXT("BoneInfluenceLimit_Row", "BoneInfluenceLimitRow"),
 			LOCTEXT("BoneInfluenceLimit", "Bone Influence Limit"),
 			LOCTEXT("BoneInfluenceLimit_ToolTip", "Limit the number of bone influences a vertex can have. If 0, the Default Bone Influence Limit from the project settings will be used."),
+			TEXT("BoneInfluenceLimit"),
 			0,
 			MAX_TOTAL_INFLUENCES,
 			FGetIntegerDelegate::CreateRaw(this, &FSkeletalMeshBuildSettingsLayout::GetBoneInfluenceLimit),
@@ -1821,6 +1822,7 @@ FDetailWidgetRow& FSkeletalMeshBuildSettingsLayout::AddIntegerRow(
 	const FText& RowTitleText,
 	const FText& RowNameContentText,
 	const FText& RowNameContentTooltipText,
+	FName RowTag,
 	const int32 MinSliderValue,
 	const int32 MaxSliderValue,
 	const FGetIntegerDelegate& GetterDelegate,
@@ -1869,6 +1871,7 @@ FDetailWidgetRow& FSkeletalMeshBuildSettingsLayout::AddIntegerRow(
 	};
 
 	FDetailWidgetRow& Row = ChildrenBuilder.AddCustomRow(RowTitleText)
+	.RowTag(RowTag)
 	.NameContent()
 	[
 		SNew(STextBlock)
diff --git a/Engine/Source/Editor/Persona/Private/PersonaMeshDetails.h b/Engine/Source/Editor/Persona/Private/PersonaMeshDetails.h
--- a/Engine/Source/Editor/Persona/Private/PersonaMeshDetails.h
+++ b/Engine/Source/Editor/Persona/Private/PersonaMeshDetails.h
@@ -285,6 +285,7 @@ private:
 		const FText& RowTitleText,
 		const FText& RowNameContentText,
 		const FText& RowNameContentTooltipText,
+		FName RowTag,
 		const int32 MinSliderValue,
 		const int32 MaxSliderValue,
 		const FGetIntegerDelegate& GetterDelegate,
