Persona Editor: Fix crash when importing custom LODs

#rnx

diff --git a/Engine/Source/Editor/Persona/Private/PersonaMeshDetails.cpp b/Engine/Source/Editor/Persona/Private/PersonaMeshDetails.cpp
--- a/Engine/Source/Editor/Persona/Private/PersonaMeshDetails.cpp
+++ b/Engine/Source/Editor/Persona/Private/PersonaMeshDetails.cpp
@@ -3643,7 +3643,7 @@ void FPersonaMeshDetails::OnAssetPostLODImported(UObject* InObject, int32 InLODI
 {
 	if (InObject == GetPersonaToolkit()->GetMesh())
 	{
-		RequestLayoutUpdate();
+		ForceLayoutRebuild();
 	}
 }
 
@@ -3651,7 +3651,7 @@ void FPersonaMeshDetails::OnAssetReimport(UObject* InObject)
 {
 	if (InObject == GetPersonaToolkit()->GetMesh())
 	{
-		RequestLayoutUpdate();
+		ForceLayoutRebuild();
 	}
 }
 
@@ -6716,7 +6716,7 @@ void FPersonaMeshDetails::OnPreviewMeshChanged(USkeletalMesh* OldSkeletalMesh, U
 {
 	if (IsApplyNeeded())
 	{
-		RequestLayoutUpdate();
+		ForceLayoutRebuild();
 	}
 }
 
