Anim Editor: Draw AssetUserData
- Added UAssetUserData::DrawCanvas() which can be used to draw info to canvas (e.g. text labels)
- Added DrawCanvas() support to static mesh editor
- Added UAssetUserData::Draw() and UAssetUserData::DrawCanvas() support to anim editors

diff --git a/Engine/Source/Editor/Persona/Private/AnimViewportShowCommands.cpp b/Engine/Source/Editor/Persona/Private/AnimViewportShowCommands.cpp
--- a/Engine/Source/Editor/Persona/Private/AnimViewportShowCommands.cpp
+++ b/Engine/Source/Editor/Persona/Private/AnimViewportShowCommands.cpp
@@ -74,6 +74,7 @@ void FAnimViewportShowCommands::RegisterCommands()
 
 	UI_COMMAND(ShowNotificationVisualizations, "Show Notification", "Display animation notifications visualizations in viewport", EUserInterfaceActionType::ToggleButton, FInputChord());
 	UI_COMMAND(ShowRootMotionVisualization, "Show Root Motion", "Display root motion visualizations in viewport", EUserInterfaceActionType::ToggleButton, FInputChord());
+	UI_COMMAND(ShowAssetUserDataVisualizations, "Show Meta Data", "Display animation meta data visualizations in viewport", EUserInterfaceActionType::ToggleButton, FInputChord());
 }
 
 #undef LOCTEXT_NAMESPACE
\ No newline at end of file
diff --git a/Engine/Source/Editor/Persona/Private/AnimViewportShowCommands.h b/Engine/Source/Editor/Persona/Private/AnimViewportShowCommands.h
--- a/Engine/Source/Editor/Persona/Private/AnimViewportShowCommands.h
+++ b/Engine/Source/Editor/Persona/Private/AnimViewportShowCommands.h
@@ -159,6 +159,9 @@ public:
 	/** Show Root motion visualization */
 	TSharedPtr< FUICommandInfo > ShowRootMotionVisualization;
 
+	/** Show AssetUserData Visualizations */
+	TSharedPtr< FUICommandInfo > ShowAssetUserDataVisualizations;
+
 public:
 	/** Registers our commands with the binding system */
 	virtual void RegisterCommands() override;
diff --git a/Engine/Source/Editor/Persona/Private/AnimationEditorViewportClient.cpp b/Engine/Source/Editor/Persona/Private/AnimationEditorViewportClient.cpp
--- a/Engine/Source/Editor/Persona/Private/AnimationEditorViewportClient.cpp
+++ b/Engine/Source/Editor/Persona/Private/AnimationEditorViewportClient.cpp
@@ -614,7 +614,8 @@ void FAnimationViewportClient::Draw(const FSceneView* View, FPrimitiveDrawInterf
 			DrawNotifies(PreviewMeshComponent, PDI);
 
 			DrawRootMotionTrajectory(PreviewMeshComponent, PDI);
-			
+
+			DrawAssetUserData(PDI);
 		}
 		else if (bValidComponent && !bValidSkeletalMesh)
 		{
@@ -672,6 +673,8 @@ void FAnimationViewportClient::DrawCanvas( FViewport& InViewport, FSceneView& Vi
 
 		DrawCanvasNotifies(PreviewMeshComponent, Canvas, View);
 
+		DrawCanvasAssetUserData(Canvas, View);
+		
 		if (bDrawUVs)
 		{
 			DrawUVsForMesh(Viewport, &Canvas, 1, PreviewMeshComponent);
@@ -2010,6 +2013,62 @@ void FAnimationViewportClient::DrawCanvasNotifies(UDebugSkelMeshComponent* MeshC
 	}
 }
 
+TArray<IInterface_AssetUserData*> FAnimationViewportClient::GetEditedObjectsWithAssetUserData() const
+{
+	TArray<IInterface_AssetUserData*> Result;
+	if (const TSharedPtr<FAssetEditorToolkit> AssetEditorToolkit = AssetEditorToolkitPtr.Pin())
+	{
+		if (const TArray<UObject*>* ObjectsCurrentlyBeingEdited = AssetEditorToolkit->GetObjectsCurrentlyBeingEdited())
+		{
+			for (UObject* Object : *ObjectsCurrentlyBeingEdited)
+			{
+				if (IInterface_AssetUserData* AssetUserDataInterface = Cast<IInterface_AssetUserData>(Object))
+				{
+					Result.Add(AssetUserDataInterface);
+				}
+			}
+		}
+	}
+
+	return Result;
+}
+
+void FAnimationViewportClient::DrawAssetUserData(FPrimitiveDrawInterface* PDI) const
+{
+	TArray<IInterface_AssetUserData*> AssetsWithUserData = GetEditedObjectsWithAssetUserData();
+	for (const IInterface_AssetUserData* AssetUserDataInterface : AssetsWithUserData)
+	{
+		if (const TArray<UAssetUserData*>* AssetUserDataArray = AssetUserDataInterface->GetAssetUserDataArray())
+		{
+			for (const UAssetUserData* AssetUserData : *AssetUserDataArray)
+			{
+				if (AssetUserData)
+				{
+					AssetUserData->Draw(PDI, PDI->View);
+				}
+			}
+		}
+	}	
+}
+
+void FAnimationViewportClient::DrawCanvasAssetUserData(FCanvas& Canvas, FSceneView& View) const
+{
+	TArray<IInterface_AssetUserData*> AssetsWithUserData = GetEditedObjectsWithAssetUserData();
+	for (const IInterface_AssetUserData* AssetUserDataInterface : AssetsWithUserData)
+	{
+		if (const TArray<UAssetUserData*>* AssetUserDataArray = AssetUserDataInterface->GetAssetUserDataArray())
+		{
+			for (const UAssetUserData* AssetUserData : *AssetUserDataArray)
+			{
+				if (AssetUserData)
+				{
+					AssetUserData->DrawCanvas(Canvas, View);
+				}
+			}
+		}
+	}
+}
+
 void FAnimationViewportClient::DrawRootMotionTrajectory(UDebugSkelMeshComponent* MeshComponent, FPrimitiveDrawInterface* PDI) const
 {
 	constexpr float DepthBias = 2.0f;
diff --git a/Engine/Source/Editor/Persona/Private/SAnimViewportToolBar.cpp b/Engine/Source/Editor/Persona/Private/SAnimViewportToolBar.cpp
--- a/Engine/Source/Editor/Persona/Private/SAnimViewportToolBar.cpp
+++ b/Engine/Source/Editor/Persona/Private/SAnimViewportToolBar.cpp
@@ -757,6 +757,7 @@ TSharedRef<SWidget> SAnimViewportToolBar::GenerateCharacterMenu() const
 						FToolMenuSection& Section = InSubMenu->AddSection("AnimViewportVisualization", LOCTEXT("CharacterMenu_VisualizationsLabel", "Visualizations"));
 						Section.AddMenuEntry(FAnimViewportShowCommands::Get().ShowNotificationVisualizations);
 						Section.AddMenuEntry(FAnimViewportShowCommands::Get().ShowRootMotionVisualization);
+						Section.AddMenuEntry(FAnimViewportShowCommands::Get().ShowAssetUserDataVisualizations);
 					}
 
 					{
diff --git a/Engine/Source/Editor/Persona/Private/SAnimationEditorViewport.cpp b/Engine/Source/Editor/Persona/Private/SAnimationEditorViewport.cpp
--- a/Engine/Source/Editor/Persona/Private/SAnimationEditorViewport.cpp
+++ b/Engine/Source/Editor/Persona/Private/SAnimationEditorViewport.cpp
@@ -1115,6 +1115,12 @@ void SAnimationEditorViewportTabBody::BindCommands()
 		FIsActionChecked(),
 		FIsActionChecked::CreateSP(this, &SAnimationEditorViewportTabBody::IsRootMotionVisualizationsEnabled));
 
+	CommandList.MapAction(
+		ViewportShowMenuCommands.ShowAssetUserDataVisualizations,
+		FExecuteAction::CreateSP(this, &SAnimationEditorViewportTabBody::ToggleAssetUserDataVisualizations),
+		FIsActionChecked(),
+		FIsActionChecked::CreateSP(this, &SAnimationEditorViewportTabBody::IsAssetUserDataVisualizationsEnabled));
+
 	CommandList.EndGroup();
 
 	CommandList.MapAction(
@@ -2185,6 +2191,22 @@ bool SAnimationEditorViewportTabBody::IsRootMotionVisualizationsEnabled() const
 	return false;
 }
 
+void SAnimationEditorViewportTabBody::ToggleAssetUserDataVisualizations()
+{
+	if (UDebugSkelMeshComponent* PreviewComponent = GetPreviewScene()->GetPreviewMeshComponent())
+	{
+		PreviewComponent->SetShowAssetUserDataVisualizations(!PreviewComponent->IsAssetUserDataVisualizationsEnabled());
+	}
+}
+
+bool SAnimationEditorViewportTabBody::IsAssetUserDataVisualizationsEnabled() const
+{
+	if (const UDebugSkelMeshComponent* PreviewComponent = GetPreviewScene()->GetPreviewMeshComponent())
+	{
+		return PreviewComponent->IsAssetUserDataVisualizationsEnabled();
+	}
+	return false;
+}
 
 bool SAnimationEditorViewportTabBody::IsClothSimulationEnabled() const
 {
diff --git a/Engine/Source/Editor/Persona/Private/SAnimationEditorViewport.h b/Engine/Source/Editor/Persona/Private/SAnimationEditorViewport.h
--- a/Engine/Source/Editor/Persona/Private/SAnimationEditorViewport.h
+++ b/Engine/Source/Editor/Persona/Private/SAnimationEditorViewport.h
@@ -525,6 +525,9 @@ private:
 	void ToggleRootMotionVisualizations();
 	bool IsRootMotionVisualizationsEnabled() const;
 
+	void ToggleAssetUserDataVisualizations();
+	bool IsAssetUserDataVisualizationsEnabled() const;
+
 private:
 	/** Selected Turn Table speed  */
 	EAnimationPlaybackSpeeds::Type SelectedTurnTableSpeed;
diff --git a/Engine/Source/Editor/Persona/Public/AnimationEditorViewportClient.h b/Engine/Source/Editor/Persona/Public/AnimationEditorViewportClient.h
--- a/Engine/Source/Editor/Persona/Public/AnimationEditorViewportClient.h
+++ b/Engine/Source/Editor/Persona/Public/AnimationEditorViewportClient.h
@@ -449,6 +449,12 @@ private:
 	/* Draws visualization from animation notifies into canvas. */
 	void DrawCanvasNotifies(UDebugSkelMeshComponent* MeshComponent, FCanvas& Canvas, FSceneView& View) const;
 
+	/* Draws visualization from Asset User Data into viewport. */
+	void DrawAssetUserData(FPrimitiveDrawInterface* PDI) const;
+
+	/* Draws visualization from Asset User Data into canvas. */
+	void DrawCanvasAssetUserData(FCanvas& Canvas, FSceneView& View) const;
+
 	/** Draws root motion trajectory */
 	void DrawRootMotionTrajectory(UDebugSkelMeshComponent* MeshComponent, FPrimitiveDrawInterface* PDI) const;
 
@@ -492,6 +498,10 @@ private:
 	void HandlePreviewScenePostTick();
 
 private:
+
+	/** @return array of AssetUserData interfaces from editable objects on associated asset toolkit. */
+	TArray<IInterface_AssetUserData*> GetEditedObjectsWithAssetUserData() const;
+	
 	/** Custom Animation speed in the viewport. Transient setting. */
 	float CustomAnimationSpeed = 1.0f;
 
diff --git a/Engine/Source/Editor/StaticMeshEditor/Private/StaticMeshEditorViewportClient.cpp b/Engine/Source/Editor/StaticMeshEditor/Private/StaticMeshEditorViewportClient.cpp
--- a/Engine/Source/Editor/StaticMeshEditor/Private/StaticMeshEditorViewportClient.cpp
+++ b/Engine/Source/Editor/StaticMeshEditor/Private/StaticMeshEditorViewportClient.cpp
@@ -796,6 +796,21 @@ void FStaticMeshEditorViewportClient::DrawCanvas(FViewport& InViewport, FSceneVi
 		}
 	}
 
+	if( bDrawAdditionalData )
+	{
+		const TArray<UAssetUserData*>* UserDataArray = StaticMesh->GetAssetUserDataArray();
+		if (UserDataArray != NULL)
+		{
+			for (int32 AdditionalDataIndex = 0; AdditionalDataIndex < UserDataArray->Num(); ++AdditionalDataIndex)
+			{
+				if ((*UserDataArray)[AdditionalDataIndex] != NULL)
+				{
+					(*UserDataArray)[AdditionalDataIndex]->DrawCanvas(Canvas, View);
+				}
+			}
+		}
+	}
+	
 	TArray<SStaticMeshEditorViewport::FOverlayTextItem, TInlineAllocator<10>> TextItems;
 
 	const int32 CurrentLODLevel = [this, &StaticMeshEditor, &View]()
diff --git a/Engine/Source/Editor/UnrealEd/Classes/Animation/DebugSkelMeshComponent.h b/Engine/Source/Editor/UnrealEd/Classes/Animation/DebugSkelMeshComponent.h
--- a/Engine/Source/Editor/UnrealEd/Classes/Animation/DebugSkelMeshComponent.h
+++ b/Engine/Source/Editor/UnrealEd/Classes/Animation/DebugSkelMeshComponent.h
@@ -242,6 +242,10 @@ class UDebugSkelMeshComponent : public USkeletalMeshComponent
 	UPROPERTY(Transient)
 	uint32 bShowRootMotionVisualizations:1;
 
+	/** Display Metadata visualizations in viewport */
+	UPROPERTY(Transient)
+	uint32 bShowAssetUserDataVisualizations:1;
+
 	/** Display Bound **/
 	UPROPERTY(transient)
 	bool bDisplayBound;
@@ -492,6 +496,10 @@ class UDebugSkelMeshComponent : public USkeletalMeshComponent
 	UNREALED_API void SetShowRootMotionVisualizations(const bool bShow) { bShowRootMotionVisualizations = bShow; }
 	UNREALED_API bool IsRootMotionVisualizationsEnabled() const { return bShowRootMotionVisualizations; }
 
+	/** Sets flags whether we AssetUserData visualizations should be drawn in the viewport. */
+	UNREALED_API void SetShowAssetUserDataVisualizations(const bool bShow) { bShowAssetUserDataVisualizations = bShow; }
+	UNREALED_API bool IsAssetUserDataVisualizationsEnabled() const { return bShowAssetUserDataVisualizations; }
+
 	/** Whether the current LOD of the debug mesh is being synced with the attached (preview) mesh instance. */
 	UNREALED_API bool IsTrackingAttachedLOD() const;
 
diff --git a/Engine/Source/Editor/UnrealEd/Private/Animation/DebugSkelMeshComponent.cpp b/Engine/Source/Editor/UnrealEd/Private/Animation/DebugSkelMeshComponent.cpp
--- a/Engine/Source/Editor/UnrealEd/Private/Animation/DebugSkelMeshComponent.cpp
+++ b/Engine/Source/Editor/UnrealEd/Private/Animation/DebugSkelMeshComponent.cpp
@@ -58,6 +58,7 @@ UDebugSkelMeshComponent::UDebugSkelMeshComponent(const FObjectInitializer& Objec
 
 	bShowNotificationVisualizations = true;
 	bShowRootMotionVisualizations = false;
+	bShowAssetUserDataVisualizations = true;
 
 	TurnTableSpeedScaling = 1.f;
 	TurnTableMode = EPersonaTurnTableMode::Stopped;
diff --git a/Engine/Source/Runtime/Engine/Classes/Engine/AssetUserData.h b/Engine/Source/Runtime/Engine/Classes/Engine/AssetUserData.h
--- a/Engine/Source/Runtime/Engine/Classes/Engine/AssetUserData.h
+++ b/Engine/Source/Runtime/Engine/Classes/Engine/AssetUserData.h
@@ -21,6 +21,9 @@ class UAssetUserData
 	/** used for debugging UAssetUserData data in editor */
 	virtual void Draw(class FPrimitiveDrawInterface* PDI, const class FSceneView* View) const {}
 
+	/** used for debugging UAssetUserData data in editor */
+	virtual void DrawCanvas(class FCanvas& Canvas, class FSceneView& View) const {}
+	
 	/** Called when the owner object is modified */
 	virtual void PostEditChangeOwner() {}
 };
