Fix issues copy and pasting notify states in montages

diff --git a/Engine/Source/Editor/Persona/Private/SAnimNotifyPanel.cpp b/Engine/Source/Editor/Persona/Private/SAnimNotifyPanel.cpp
--- a/Engine/Source/Editor/Persona/Private/SAnimNotifyPanel.cpp
+++ b/Engine/Source/Editor/Persona/Private/SAnimNotifyPanel.cpp
@@ -3717,7 +3717,15 @@ void SAnimNotifyTrack::PasteSingleNotify(FString& NotifyString, float PasteTime)
 			NewNotify.NotifyStateClass = NewNotifyStateObject;
 			bValidNotify = NewNotifyStateObject->CanBePlaced(Sequence);
 			// Clamp duration into the sequence
-			NewNotify.SetDuration(FMath::Clamp(NewNotify.GetDuration(), 1 / 30.0f, Sequence->GetPlayLength() - NewNotify.GetTime()));
+			if (UAnimMontage* Montage = Cast<UAnimMontage>(Sequence))
+			{
+				NewNotify.SetDuration(FMath::Clamp(NewNotify.Duration, 1 / 30.0f, Montage->CalculateSequenceLength() - NewNotify.GetTime()));
+
+			}
+			else
+			{
+				NewNotify.SetDuration(FMath::Clamp(NewNotify.Duration, 1 / 30.0f, Sequence->GetPlayLength() - NewNotify.GetTime()));
+			}
 			NewNotify.EndTriggerTimeOffset = GetTriggerTimeOffsetForType(Sequence->CalculateOffsetForNotify(NewNotify.GetTime() + NewNotify.GetDuration()));
 			NewNotify.EndLink.Link(Sequence, NewNotify.EndLink.GetTime());
 		}
