[Animation Editor Viewport] Added null checks for when there is no preview mesh, which is possible for control rig and deformer editors

diff --git a/Engine/Source/Editor/Persona/Private/SAnimationEditorViewport.cpp b/Engine/Source/Editor/Persona/Private/SAnimationEditorViewport.cpp
--- a/Engine/Source/Editor/Persona/Private/SAnimationEditorViewport.cpp
+++ b/Engine/Source/Editor/Persona/Private/SAnimationEditorViewport.cpp
@@ -1023,9 +1023,9 @@ void SAnimationEditorViewportTabBody::BindCommands()
 		CommandList.MapAction(
 			ViewportLODMenuCommands.LODDebug,
 			FExecuteAction::CreateSP(this, &SAnimationEditorViewportTabBody::OnSetLODTrackDebuggedInstance),
-			FCanExecuteAction::CreateLambda([PreviewComponent]() { return (bool)PreviewComponent->PreviewInstance->GetDebugSkeletalMeshComponent(); }),
+			FCanExecuteAction::CreateLambda([PreviewComponent]() { return PreviewComponent->PreviewInstance ? (bool)PreviewComponent->PreviewInstance->GetDebugSkeletalMeshComponent() : false; }),
 			FIsActionChecked::CreateLambda([PreviewComponent]() { return PreviewComponent->IsTrackingAttachedLOD(); }),
-			FIsActionButtonVisible::CreateLambda([PreviewComponent]() { return (bool)PreviewComponent->PreviewInstance->GetDebugSkeletalMeshComponent() ; }));
+			FIsActionButtonVisible::CreateLambda([PreviewComponent]() { return PreviewComponent->PreviewInstance ? (bool)PreviewComponent->PreviewInstance->GetDebugSkeletalMeshComponent() : false; }));
 
 		PreviewComponent->RegisterOnDebugForceLODChangedDelegate(FOnDebugForceLODChanged::CreateSP(this, &SAnimationEditorViewportTabBody::OnDebugForcedLODChanged));
 	}
