[Skeletal mesh editor] missed change from the last vertex attribute UI fix. This one line change should have been part of the last fix to actually use the delegate on the UI side

diff --git a/Engine/Source/Editor/Persona/Private/PersonaMeshDetails.cpp b/Engine/Source/Editor/Persona/Private/PersonaMeshDetails.cpp
--- a/Engine/Source/Editor/Persona/Private/PersonaMeshDetails.cpp
+++ b/Engine/Source/Editor/Persona/Private/PersonaMeshDetails.cpp
@@ -2909,6 +2909,7 @@ void FPersonaMeshDetails::AddLODLevelCategories(IDetailLayoutBuilder& DetailLayo
 			}
 			const FSkeletalMeshLODModel& LODModel = SkelMesh->GetImportedModel()->LODModels[LODIndex];
 			LODInfoUILayout->SetReferenceLODInfo(GetPersonaToolkit(), LODIndex);
+			SkelMesh->GetOnVertexAttributesArrayChanged().AddUObject(LODInfoUILayout, &ULODInfoUILayout::RefreshReferenceLODInfo);
 			LODInfoUILayouts.Add(LODInfoUILayout);
 
 			//Show the viewport LOD at start
