Removes deprecated retarget base pose

diff --git a/Engine/Source/Editor/AnimGraph/Private/AnimPreviewInstance.cpp b/Engine/Source/Editor/AnimGraph/Private/AnimPreviewInstance.cpp
--- a/Engine/Source/Editor/AnimGraph/Private/AnimPreviewInstance.cpp
+++ b/Engine/Source/Editor/AnimGraph/Private/AnimPreviewInstance.cpp
@@ -118,14 +118,6 @@ void FAnimPreviewInstanceProxy::Update(float DeltaSeconds)
 	// we cant update on a worker thread here because of the key delegate needing to be fired
 	check(IsInGameThread());
 
-#if WITH_EDITORONLY_DATA
-	if(bForceRetargetBasePose)
-	{
-		// nothing to be done here
-		return;
-	}
-#endif // #if WITH_EDITORONLY_DATA
-
 	FAnimSingleNodeInstanceProxy::Update(DeltaSeconds);
 }
 
@@ -156,10 +148,7 @@ void FAnimPreviewInstanceProxy::PreUpdate(UAnimInstance* InAnimInstance, float D
 		CopyPoseNode.PreUpdate(InAnimInstance);
 	}
 
-	if (!bForceRetargetBasePose)
-	{
-		CurveSource.PreUpdate(InAnimInstance);
-	}
+	CurveSource.PreUpdate(InAnimInstance);
 }
 
 bool FAnimPreviewInstanceProxy::Evaluate(FPoseContext& Output)
@@ -173,31 +162,13 @@ bool FAnimPreviewInstanceProxy::Evaluate(FPoseContext& Output)
 	}
 	else
 	{
-#if WITH_EDITORONLY_DATA
-		if(bForceRetargetBasePose)
+		if (UPoseAsset* PoseAsset = Cast<UPoseAsset>(CurrentAsset))
 		{
-			USkeletalMeshComponent* MeshComponent = Output.AnimInstanceProxy->GetSkelMeshComponent();
-			if(MeshComponent && MeshComponent->GetSkeletalMeshAsset())
-			{
-				FAnimationRuntime::FillWithRetargetBaseRefPose(Output.Pose, GetSkelMeshComponent()->GetSkeletalMeshAsset());
-			}
-			else
-			{
-				// ideally we'll return just ref pose, but not sure if this will work with LODs
-				Output.Pose.ResetToRefPose();
-			}
+			PoseBlendNode.Evaluate_AnyThread(Output);
 		}
 		else
-#endif // #if WITH_EDITORONLY_DATA
 		{
-			if (UPoseAsset* PoseAsset = Cast<UPoseAsset>(CurrentAsset))
-			{
-				PoseBlendNode.Evaluate_AnyThread(Output);
-			}
-			else
-			{
-				FAnimSingleNodeInstanceProxy::Evaluate(Output);
-			}
+			FAnimSingleNodeInstanceProxy::Evaluate(Output);
 		}
 
 		if (bEnableControllers)
@@ -1222,16 +1193,6 @@ void UAnimPreviewInstance::EnableControllers(bool bEnable)
 	GetProxyOnGameThread<FAnimPreviewInstanceProxy>().EnableControllers(bEnable);
 }
 
-void UAnimPreviewInstance::SetForceRetargetBasePose(bool bInForceRetargetBasePose)
-{
-	GetProxyOnGameThread<FAnimPreviewInstanceProxy>().SetForceRetargetBasePose(bInForceRetargetBasePose);
-}
-
-bool UAnimPreviewInstance::GetForceRetargetBasePose() const
-{
-	return GetProxyOnGameThread<FAnimPreviewInstanceProxy>().GetForceRetargetBasePose();
-}
-
 FAnimInstanceProxy* UAnimPreviewInstance::CreateAnimInstanceProxy()
 {
 	return new FAnimPreviewInstanceProxy(this);
diff --git a/Engine/Source/Editor/AnimGraph/Public/AnimPreviewInstance.h b/Engine/Source/Editor/AnimGraph/Public/AnimPreviewInstance.h
--- a/Engine/Source/Editor/AnimGraph/Public/AnimPreviewInstance.h
+++ b/Engine/Source/Editor/AnimGraph/Public/AnimPreviewInstance.h
@@ -40,9 +40,6 @@ public:
 	FAnimPreviewInstanceProxy(UAnimInstance* InAnimInstance)
 		: FAnimSingleNodeInstanceProxy(InAnimInstance)
 		, SkeletalControlAlpha(1.0f)
-#if WITH_EDITORONLY_DATA
-		, bForceRetargetBasePose(false)
-#endif
 		, bEnableControllers(true)
 		, bSetKey(false)
 	{
@@ -64,16 +61,6 @@ public:
 
 	void RemoveBoneModification(const FName& InBoneName, bool bCurveController = false);
 
-	void SetForceRetargetBasePose(bool bInForceRetargetBasePose)
-	{ 
-		bForceRetargetBasePose = bInForceRetargetBasePose; 
-	}
-
-	bool GetForceRetargetBasePose() const 
-	{ 
-		return bForceRetargetBasePose; 
-	}
-
 	void EnableControllers(bool bEnable)
 	{
 		bEnableControllers = bEnable;
@@ -153,11 +140,7 @@ private:
 
 	/** Shared parameters for previewing blendspace or animsequence **/
 	float SkeletalControlAlpha;
-
-#if WITH_EDITORONLY_DATA
-	bool bForceRetargetBasePose;
-#endif
-
+	
 	/*
 	 * Used to determine if controller has to be applied or not
 	 * Used to disable controller during editing
@@ -263,10 +246,6 @@ public:
 	 */
 	const TArray<FAnimNode_ModifyBone>& GetBoneControllers();
 
-	void SetForceRetargetBasePose(bool ForceRetargetBasePose);
-
-	bool GetForceRetargetBasePose() const;
-
 #if WITH_EDITOR	
 	/**
 	 * Convert current modified bone transforms (BoneControllers) to transform curves (CurveControllers)
diff --git a/Engine/Source/Editor/Persona/Private/AnimViewportShowCommands.cpp b/Engine/Source/Editor/Persona/Private/AnimViewportShowCommands.cpp
--- a/Engine/Source/Editor/Persona/Private/AnimViewportShowCommands.cpp
+++ b/Engine/Source/Editor/Persona/Private/AnimViewportShowCommands.cpp
@@ -21,7 +21,6 @@ void FAnimViewportShowCommands::RegisterCommands()
 
 	UI_COMMAND(DisablePostProcessBlueprint, "Disable Post Process", "Disable the evaluation of post process animation blueprints on the preview instance.", EUserInterfaceActionType::ToggleButton, FInputChord());
 
-	UI_COMMAND( ShowRetargetBasePose, "Retarget Base Pose", "Show retarget Base pose on preview mesh", EUserInterfaceActionType::ToggleButton, FInputChord() );
 	UI_COMMAND( ShowBound, "Bound", "Show bound on preview mesh", EUserInterfaceActionType::ToggleButton, FInputChord() );
 	UI_COMMAND( UseInGameBound, "In-game Bound", "Use in-game bound on preview mesh when showing bounds. Otherwise bounds will always be calculated from bones alone.", EUserInterfaceActionType::ToggleButton, FInputChord());
 	UI_COMMAND( UseFixedBounds, "Fixed Bounds", "Preview using the 'Fixed Bounds' option, which will use bounds from skel mesh and not bones at all, for speed.", EUserInterfaceActionType::ToggleButton, FInputChord());
diff --git a/Engine/Source/Editor/Persona/Private/AnimViewportShowCommands.h b/Engine/Source/Editor/Persona/Private/AnimViewportShowCommands.h
--- a/Engine/Source/Editor/Persona/Private/AnimViewportShowCommands.h
+++ b/Engine/Source/Editor/Persona/Private/AnimViewportShowCommands.h
@@ -49,9 +49,6 @@ public:
 
 	/** Option to enable/disable post process anim blueprint evaluation */
 	TSharedPtr< FUICommandInfo > DisablePostProcessBlueprint;
-
-	/** Show reference pose on preview mesh */
-	TSharedPtr< FUICommandInfo > ShowRetargetBasePose;
 	
 	/** Show Bound of preview mesh */
 	TSharedPtr< FUICommandInfo > ShowBound;
diff --git a/Engine/Source/Editor/Persona/Private/PersonaModule.cpp b/Engine/Source/Editor/Persona/Private/PersonaModule.cpp
--- a/Engine/Source/Editor/Persona/Private/PersonaModule.cpp
+++ b/Engine/Source/Editor/Persona/Private/PersonaModule.cpp
@@ -295,9 +295,9 @@ TSharedRef<class FWorkflowTabFactory> FPersonaModule::CreateCurveMetadataEditorT
 	return MakeShared<FAnimCurveMetadataEditorTabSummoner>(InHostingApp, InMetadataHost, InPreviewScene, InOnObjectsSelected);
 }
 
-TSharedRef<class FWorkflowTabFactory> FPersonaModule::CreateRetargetSourcesTabFactory(const TSharedRef<class FWorkflowCentricApplication>& InHostingApp, const TSharedRef<class IEditableSkeleton>& InEditableSkeleton, const TSharedRef<IPersonaPreviewScene>& InPreviewScene, FSimpleMulticastDelegate& InOnPostUndo) const
+TSharedRef<class FWorkflowTabFactory> FPersonaModule::CreateRetargetSourcesTabFactory(const TSharedRef<class FWorkflowCentricApplication>& InHostingApp, const TSharedRef<class IEditableSkeleton>& InEditableSkeleton, FSimpleMulticastDelegate& InOnPostUndo) const
 {
-	return MakeShareable(new FRetargetSourcesTabSummoner(InHostingApp, InEditableSkeleton, InPreviewScene, InOnPostUndo));
+	return MakeShareable(new FRetargetSourcesTabSummoner(InHostingApp, InEditableSkeleton, InOnPostUndo));
 }
 
 TSharedRef<class FWorkflowTabFactory> FPersonaModule::CreateAdvancedPreviewSceneTabFactory(const TSharedRef<class FWorkflowCentricApplication>& InHostingApp, const TSharedRef<IPersonaPreviewScene>& InPreviewScene) const
diff --git a/Engine/Source/Editor/Persona/Private/SAnimViewportToolBar.cpp b/Engine/Source/Editor/Persona/Private/SAnimViewportToolBar.cpp
--- a/Engine/Source/Editor/Persona/Private/SAnimViewportToolBar.cpp
+++ b/Engine/Source/Editor/Persona/Private/SAnimViewportToolBar.cpp
@@ -715,7 +715,6 @@ TSharedRef<SWidget> SAnimViewportToolBar::GenerateCharacterMenu() const
 				{
 					{
 						FToolMenuSection& Section = InSubMenu->AddSection("AnimViewportMesh", LOCTEXT("CharacterMenu_Actions_Mesh", "Mesh"));
-						Section.AddMenuEntry(FAnimViewportShowCommands::Get().ShowRetargetBasePose);
 						Section.AddMenuEntry(FAnimViewportShowCommands::Get().ShowBound);
 						Section.AddMenuEntry(FAnimViewportShowCommands::Get().UseInGameBound);
 						Section.AddMenuEntry(FAnimViewportShowCommands::Get().UseFixedBounds);
diff --git a/Engine/Source/Editor/Persona/Private/SAnimationEditorViewport.cpp b/Engine/Source/Editor/Persona/Private/SAnimationEditorViewport.cpp
--- a/Engine/Source/Editor/Persona/Private/SAnimationEditorViewport.cpp
+++ b/Engine/Source/Editor/Persona/Private/SAnimationEditorViewport.cpp
@@ -754,12 +754,6 @@ void SAnimationEditorViewportTabBody::BindCommands()
 	//Bind Show commands
 	const FAnimViewportShowCommands& ViewportShowMenuCommands = FAnimViewportShowCommands::Get();
 
-	CommandList.MapAction(
-		ViewportShowMenuCommands.ShowRetargetBasePose,
-		FExecuteAction::CreateSP(this, &SAnimationEditorViewportTabBody::ShowRetargetBasePose),
-		FCanExecuteAction::CreateSP(this, &SAnimationEditorViewportTabBody::CanShowRetargetBasePose),
-		FIsActionChecked::CreateSP(this, &SAnimationEditorViewportTabBody::IsShowRetargetBasePoseEnabled));
-
 	CommandList.MapAction(
 		ViewportShowMenuCommands.ShowBound,
 		FExecuteAction::CreateSP(this, &SAnimationEditorViewportTabBody::ShowBound),
@@ -1562,41 +1556,6 @@ bool SAnimationEditorViewportTabBody::IsPlaybackSpeedSelected(int32 PlaybackSpee
 	return PlaybackSpeedMode == AnimViewportClient->GetPlaybackSpeedMode();
 }
 
-void SAnimationEditorViewportTabBody::ShowRetargetBasePose()
-{
-	ForEachDebugMesh([](UDebugSkelMeshComponent* PreviewMeshComponent)
-	{
-		PreviewMeshComponent->PreviewInstance->SetForceRetargetBasePose(!PreviewMeshComponent->PreviewInstance->GetForceRetargetBasePose());
-	});
-}
-
-bool SAnimationEditorViewportTabBody::CanShowRetargetBasePose() const
-{
-	TArray<UDebugSkelMeshComponent*> PreviewMeshComponents = GetPreviewScene()->GetAllPreviewMeshComponents();
-	for (UDebugSkelMeshComponent* PreviewMeshComponent : PreviewMeshComponents)
-	{
-		if (PreviewMeshComponent->PreviewInstance)
-		{
-			return true;
-		}
-	}
-	
-	return false;
-}
-
-bool SAnimationEditorViewportTabBody::IsShowRetargetBasePoseEnabled() const
-{
-	TArray<UDebugSkelMeshComponent*> PreviewMeshComponents = GetPreviewScene()->GetAllPreviewMeshComponents();
-	for (UDebugSkelMeshComponent* PreviewMeshComponent : PreviewMeshComponents)
-	{
-		if (PreviewMeshComponent && PreviewMeshComponent->PreviewInstance)
-		{
-			return PreviewMeshComponent->PreviewInstance->GetForceRetargetBasePose();
-		}
-	}
-	return false;
-}
-
 void SAnimationEditorViewportTabBody::ShowBound()
 {
 	TSharedRef<FAnimationViewportClient> AnimViewportClient = StaticCastSharedRef<FAnimationViewportClient>(LevelViewportClient.ToSharedRef());	
diff --git a/Engine/Source/Editor/Persona/Private/SAnimationEditorViewport.h b/Engine/Source/Editor/Persona/Private/SAnimationEditorViewport.h
--- a/Engine/Source/Editor/Persona/Private/SAnimationEditorViewport.h
+++ b/Engine/Source/Editor/Persona/Private/SAnimationEditorViewport.h
@@ -420,11 +420,6 @@ private:
 	/** Toggles floor alignment in the preview scene */
 	void OnToggleAutoAlignFloor();
 
-	/** Called to toggle showing of reference pose on current preview mesh */
-	void ShowRetargetBasePose();
-	bool CanShowRetargetBasePose() const;
-	bool IsShowRetargetBasePoseEnabled() const;
-
 	/** Called to toggle showing of the bounds of the current preview mesh */
 	void ShowBound();
 	bool CanShowBound() const;
diff --git a/Engine/Source/Editor/Persona/Private/SPoseAssetNameWidget.cpp b/Engine/Source/Editor/Persona/Private/SPoseAssetNameWidget.cpp
deleted file mode 100644
diff --git a/Engine/Source/Editor/Persona/Private/SRetargetSources.cpp b/Engine/Source/Editor/Persona/Private/SRetargetSources.cpp
--- a/Engine/Source/Editor/Persona/Private/SRetargetSources.cpp
+++ b/Engine/Source/Editor/Persona/Private/SRetargetSources.cpp
@@ -24,68 +24,12 @@
 void SRetargetSources::Construct(
 	const FArguments& InArgs,
 	const TSharedRef<IEditableSkeleton>& InEditableSkeleton,
-	const TSharedRef<IPersonaPreviewScene>& InPreviewScene,
 	FSimpleMulticastDelegate& InOnPostUndo)
 {
-	EditableSkeletonPtr = InEditableSkeleton;
-	PreviewScenePtr = InPreviewScene;
-	InOnPostUndo.Add(FSimpleDelegate::CreateSP(this, &SRetargetSources::PostUndo));
-
-
 	const FString DocLink = TEXT("Shared/Editors/Persona");
 	ChildSlot
 	[
 		SNew (SVerticalBox)
-
-		+ SVerticalBox::Slot()
-		.Padding(5, 5)
-		.AutoHeight()
-		[
-			SNew(STextBlock)
-			.TextStyle(FAppStyle::Get(), "Persona.RetargetManager.ImportantText")
-			.Text(LOCTEXT("BasePose_Title", "Edit Retarget Base Pose"))
-		]
-
-		+ SVerticalBox::Slot()
-		.AutoHeight()		// This is required to make the scrollbar work, as content overflows Slate containers by default
-		.VAlign(VAlign_Center)
-		.HAlign(HAlign_Left)
-		.Padding(5, 5)
-		[
-			SNew(SHorizontalBox)
-
-			+SHorizontalBox::Slot()
-			.AutoWidth()
-			.HAlign(HAlign_Center)
-			[
-				SNew(SButton)
-				.OnClicked(FOnClicked::CreateSP(this, &SRetargetSources::OnModifyPose))
-				.HAlign(HAlign_Center)
-				.VAlign(VAlign_Center)
-				.Text(LOCTEXT("ModifyRetargetBasePose_Label", "Modify Pose"))
-				.ToolTipText(LOCTEXT("ModifyRetargetBasePose_Tooltip", "The Retarget Base Pose is used to retarget animation from Skeletal Meshes with varying ref poses to this Skeletal Mesh. \nPrefer fixing ref poses before importing if possible."))
-			]
-
-			+SHorizontalBox::Slot()
-			.AutoWidth()
-			.HAlign(HAlign_Center)
-			[
-				SNew(SButton)
-				.OnClicked(FOnClicked::CreateSP(this, &SRetargetSources::OnViewRetargetBasePose))
-				.HAlign(HAlign_Center)
-				.VAlign(VAlign_Center)
-				.Text(this, &SRetargetSources::GetToggleRetargetBasePose)
-				.ToolTipText(LOCTEXT("ViewRetargetBasePose_Tooltip", "Toggle to View/Edit Retarget Base Pose"))
-			]
-		]
-
-		+SVerticalBox::Slot()
-		.Padding(5, 5)
-		.AutoHeight()
-		[
-			SNew(SSeparator)
-			.Orientation(Orient_Horizontal)
-		]
 		
 		+ SVerticalBox::Slot()
 		.Padding(5, 5)
@@ -100,7 +44,7 @@ void SRetargetSources::Construct(
 		.Padding(5, 5)
 		.FillHeight(0.5)
 		[
-			// construct retarget source window
+			// construct retarget source UI
 			SNew(SRetargetSourceWindow, InEditableSkeleton, InOnPostUndo)
 		]
 
@@ -125,342 +69,11 @@ void SRetargetSources::Construct(
 		.Padding(5, 5)
 		.FillHeight(0.5)
 		[
+			// construct compatible skeletons UI
 			SNew(SCompatibleSkeletons, InEditableSkeleton, InOnPostUndo)
 		]
 	];
 }
 
-FReply SRetargetSources::OnViewRetargetBasePose() const
-{
-	UDebugSkelMeshComponent * PreviewMeshComp = PreviewScenePtr.Pin()->GetPreviewMeshComponent();
-	if (PreviewMeshComp && PreviewMeshComp->PreviewInstance)
-	{
-		const FScopedTransaction Transaction(LOCTEXT("ViewRetargetBasePose_Action", "Edit Retarget Base Pose"));
-		PreviewMeshComp->PreviewInstance->SetForceRetargetBasePose(!PreviewMeshComp->PreviewInstance->GetForceRetargetBasePose());
-		PreviewMeshComp->Modify();
-		// reset all bone transform since you don't want to keep any bone transform change
-		PreviewMeshComp->PreviewInstance->ResetModifiedBone();
-		// add root 
-		if (PreviewMeshComp->PreviewInstance->GetForceRetargetBasePose())
-		{
-			PreviewMeshComp->BonesOfInterest.Add(0);
-		}
-	}
-
-	return FReply::Handled();
-}
-
-FReply SRetargetSources::OnModifyPose()
-{
-	// create context menu
-	TSharedPtr< SWindow > Parent = FSlateApplication::Get().GetActiveTopLevelWindow();
-	if (Parent.IsValid())
-	{
-		FSlateApplication::Get().PushMenu(
-			Parent.ToSharedRef(),
-			FWidgetPath(),
-			OnModifyPoseContextMenu(),
-			FSlateApplication::Get().GetCursorPos(),
-			FPopupTransitionEffect(FPopupTransitionEffect::TypeInPopup));
-	}
-
-	return FReply::Handled();
-}
-
-TSharedRef<SWidget> SRetargetSources::OnModifyPoseContextMenu() 
-{
-	FMenuBuilder MenuBuilder(false, nullptr);
-
-	MenuBuilder.BeginSection("ModifyPose_Label", LOCTEXT("ModifyPose", "Set Pose"));
-	{
-		FUIAction Action_ReferencePose
-		(
-			FExecuteAction::CreateSP(this, &SRetargetSources::ResetRetargetBasePose)
-		);
-
-		MenuBuilder.AddMenuEntry
-		(
-			LOCTEXT("ModifyPoseContextMenu_Reset", "Reset"),
-			LOCTEXT("ModifyPoseContextMenu_Reset_Desc", "Reset to reference pose"),
-			FSlateIcon(FAppStyle::GetAppStyleSetName(), "Profiler.EventGraph.SelectStack"), Action_ReferencePose, NAME_None, EUserInterfaceActionType::Button
-		);
-
-		FUIAction Action_UseCurrentPose
-		(
-			FExecuteAction::CreateSP(this, &SRetargetSources::UseCurrentPose)
-		);
-
-		MenuBuilder.AddMenuEntry
-		(
-			LOCTEXT("ModifyPoseContextMenu_UseCurrentPose", "Use CurrentPose"),
-			LOCTEXT("ModifyPoseContextMenu_UseCurrentPose_Desc", "Use Current Pose"),
-			FSlateIcon(FAppStyle::GetAppStyleSetName(), "Profiler.EventGraph.SelectStack"), Action_UseCurrentPose, NAME_None, EUserInterfaceActionType::Button
-		);
-
-		MenuBuilder.AddMenuSeparator();
-
-		MenuBuilder.AddWidget(
-			SNew(SVerticalBox)
-
-			+ SVerticalBox::Slot()
-			.AutoHeight()
-			.Padding(3)
-			[
-				SNew(SObjectPropertyEntryBox)
-				.AllowedClass(UPoseAsset::StaticClass())
-				.OnObjectChanged(this, &SRetargetSources::SetSelectedPose)
-				.OnShouldFilterAsset(this, &SRetargetSources::ShouldFilterAsset)
-				.ObjectPath(this, &SRetargetSources::GetSelectedPose)
-			]
-			+ SVerticalBox::Slot()
-			.AutoHeight()
-			.Padding(3)
-			[
-				SAssignNew(PoseAssetNameWidget, SPoseAssetNameWidget)
-				.OnSelectionChanged(this, &SRetargetSources::SetPoseName)
-			]
-
-			+ SVerticalBox::Slot()
-			.AutoHeight()
-			.Padding(3)
-			[
-				SNew(SButton)
-				.OnClicked(FOnClicked::CreateSP(this, &SRetargetSources::OnImportPose))
-				.IsEnabled(this, &SRetargetSources::CanImportPose)
-				.HAlign(HAlign_Center)
-				.VAlign(VAlign_Center)
-				.Text(LOCTEXT("ImportRetargetBasePose_Label", "Import"))
-				.ToolTipText(LOCTEXT("ImportRetargetBasePose_Tooltip", "Import the selected pose to Retarget Base Pose"))
-			]
-			,
-			FText()
-		);
-
-		if (SelectedPoseAsset.IsValid())
-		{
-			PoseAssetNameWidget->SetPoseAsset(SelectedPoseAsset.Get());
-		}
-		// import pose 
-		MenuBuilder.EndSection();
-	}
-
-	return MenuBuilder.MakeWidget();
-}
-
-bool SRetargetSources::CanImportPose() const
-{
-	return (SelectedPoseAsset.IsValid() && SelectedPoseAsset.Get()->ContainsPose(FName(*SelectedPoseName)));
-}
-
-void SRetargetSources::SetSelectedPose(const FAssetData& InAssetData)
-{
-	if (PoseAssetNameWidget.IsValid())
-	{
-		SelectedPoseAsset = Cast<UPoseAsset>(InAssetData.GetAsset());
-		if (SelectedPoseAsset.IsValid())
-		{
-			PoseAssetNameWidget->SetPoseAsset(SelectedPoseAsset.Get());
-		}
-	}
-}
-
-FString SRetargetSources::GetSelectedPose() const
-{
-	return SelectedPoseAsset->GetPathName();
-}
-
-bool SRetargetSources::ShouldFilterAsset(const FAssetData& InAssetData)
-{
-	if (InAssetData.GetClass() == UPoseAsset::StaticClass() &&
-		EditableSkeletonPtr.IsValid())
-	{
-		FString SkeletonString = FAssetData(&EditableSkeletonPtr.Pin()->GetSkeleton()).GetExportTextName();
-		FAssetDataTagMapSharedView::FFindTagResult Result = InAssetData.TagsAndValues.FindTag("Skeleton");
-		return (!Result.IsSet() || SkeletonString != Result.GetValue());
-	}
-
-	return false;
-}
-
-PRAGMA_DISABLE_DEPRECATION_WARNINGS
-void SRetargetSources::ResetRetargetBasePose()
-{
-	UDebugSkelMeshComponent * PreviewMeshComp = PreviewScenePtr.Pin()->GetPreviewMeshComponent();
-	if(PreviewMeshComp && PreviewMeshComp->GetSkeletalMeshAsset())
-	{
-		USkeletalMesh * PreviewMesh = PreviewMeshComp->GetSkeletalMeshAsset();
-
-		check(PreviewMesh && &EditableSkeletonPtr.Pin()->GetSkeleton() == PreviewMesh->GetSkeleton());
-
-		if(PreviewMesh)
-		{
-			const FScopedTransaction Transaction(LOCTEXT("ResetRetargetBasePose_Action", "Reset Retarget Base Pose"));
-			PreviewMesh->Modify();
-			// reset to original ref pose
-			PreviewMesh->SetRetargetBasePose(PreviewMesh->GetRefSkeleton().GetRefBonePose());
-			TurnOnPreviewRetargetBasePose();
-		}
-	}
-
-	FSlateApplication::Get().DismissAllMenus();
-}
-PRAGMA_ENABLE_DEPRECATION_WARNINGS
-
-void SRetargetSources::UseCurrentPose()
-{
-	UDebugSkelMeshComponent * PreviewMeshComp = PreviewScenePtr.Pin()->GetPreviewMeshComponent();
-	if (PreviewMeshComp && PreviewMeshComp->GetSkeletalMeshAsset())
-	{
-		USkeletalMesh * PreviewMesh = PreviewMeshComp->GetSkeletalMeshAsset();
-
-		check(PreviewMesh && &EditableSkeletonPtr.Pin()->GetSkeleton() == PreviewMesh->GetSkeleton());
-
-		if (PreviewMesh)
-		{
-			const FScopedTransaction Transaction(LOCTEXT("RetargetBasePose_UseCurrentPose_Action", "Retarget Base Pose : Use Current Pose"));
-			PreviewMesh->Modify();
-			// get space bases and calculate local
-			const TArray<FTransform> & SpaceBases = PreviewMeshComp->GetComponentSpaceTransforms();
-			// @todo check to see if skeleton vs preview mesh makes it different for missing bones
-			const FReferenceSkeleton& RefSkeleton = PreviewMesh->GetRefSkeleton();
-			PRAGMA_DISABLE_DEPRECATION_WARNINGS
-			TArray<FTransform> & NewRetargetBasePose = PreviewMesh->GetRetargetBasePose();
-			PRAGMA_ENABLE_DEPRECATION_WARNINGS
-			// if you're using leader pose component in preview, this won't work
-			check(PreviewMesh->GetRefSkeleton().GetNum() == SpaceBases.Num());
-			int32 TotalNumBones = PreviewMesh->GetRefSkeleton().GetNum();
-			NewRetargetBasePose.Empty(TotalNumBones);
-			NewRetargetBasePose.AddUninitialized(TotalNumBones);
-
-			for (int32 BoneIndex = 0; BoneIndex < TotalNumBones; ++BoneIndex)
-			{
-				// this is slower, but skeleton can have more bones, so I can't just access
-				// Parent index from Skeleton. Going safer route
-				FName BoneName = PreviewMeshComp->GetBoneName(BoneIndex);
-				FName ParentBoneName = PreviewMeshComp->GetParentBone(BoneName);
-				int32 ParentIndex = RefSkeleton.FindBoneIndex(ParentBoneName);
-
-				if (ParentIndex != INDEX_NONE)
-				{
-					NewRetargetBasePose[BoneIndex] = SpaceBases[BoneIndex].GetRelativeTransform(SpaceBases[ParentIndex]);
-				}
-				else
-				{
-					NewRetargetBasePose[BoneIndex] = SpaceBases[BoneIndex];
-				}
-			}
-
-			// Clear PreviewMeshComp bone modified, they're baked now
-			PreviewMeshComp->PreviewInstance->ResetModifiedBone();
-			TurnOnPreviewRetargetBasePose();
-		}
-	}
-	FSlateApplication::Get().DismissAllMenus();
-}
-
-void SRetargetSources::SetPoseName(TSharedPtr<FString> PoseName, ESelectInfo::Type SelectionType)
-{
-	SelectedPoseName = *PoseName.Get();
-}
-
-FReply SRetargetSources::OnImportPose()
-{
-	if (CanImportPose())
-	{
-		UPoseAsset* RawPoseAsset = SelectedPoseAsset.Get();
-		ImportPose(RawPoseAsset, FName(*SelectedPoseName));
-	}
-
-	FSlateApplication::Get().DismissAllMenus();
-
-	return FReply::Handled();
-}
-
-void SRetargetSources::ImportPose(const UPoseAsset* PoseAsset, const FName& PoseName)
-{
-	// Get transforms from pose (this also converts from additive if necessary)
-	const int32 PoseIndex = PoseAsset->GetPoseIndexByName(PoseName);
-	if (PoseIndex != INDEX_NONE)
-	{
-		TArray<FTransform> PoseTransforms;
-		if (PoseAsset->GetFullPose(PoseIndex, PoseTransforms))
-		{
-			const TArray<FName>	PoseTrackNames = PoseAsset->GetTrackNames();
-
-			ensureAlways(PoseTrackNames.Num() == PoseTransforms.Num());
-
-			// now I have pose, I have to copy to the retarget base pose
-			UDebugSkelMeshComponent * PreviewMeshComp = PreviewScenePtr.Pin()->GetPreviewMeshComponent();
-			if (PreviewMeshComp && PreviewMeshComp->GetSkeletalMeshAsset())
-			{
-				USkeletalMesh * PreviewMesh = PreviewMeshComp->GetSkeletalMeshAsset();
-
-				check(PreviewMesh && &EditableSkeletonPtr.Pin()->GetSkeleton() == PreviewMesh->GetSkeleton());
-
-				if (PreviewMesh)
-				{
-					// Check if we have bones for all the tracks. If not, then fail so that the user doesn't end up
-					// with partial or broken retarget setup.
-					for (int32 TrackIndex = 0; TrackIndex < PoseTrackNames.Num(); ++TrackIndex)
-					{
-						const int32 BoneIndex = PreviewMesh->GetRefSkeleton().FindBoneIndex(PoseTrackNames[TrackIndex]);
-						if (BoneIndex == INDEX_NONE)
-						{
-							FMessageDialog::Open(EAppMsgType::Ok, FText::FromString(TEXT("Pose asset does not match the preview mesh skeleton. Aborting.")));
-							return;
-						}
-					}
-					
-					const FScopedTransaction Transaction(LOCTEXT("ImportRetargetBasePose_Action", "Import Retarget Base Pose"));
-					PreviewMesh->Modify();
-
-					PRAGMA_DISABLE_DEPRECATION_WARNINGS
-					// reset to original ref pose first
-					PreviewMesh->SetRetargetBasePose(PreviewMesh->GetRefSkeleton().GetRefBonePose());
-					// now override imported pose
-					for (int32 TrackIndex = 0; TrackIndex < PoseTrackNames.Num(); ++TrackIndex)
-					{
-						const int32 BoneIndex = PreviewMesh->GetRefSkeleton().FindBoneIndex(PoseTrackNames[TrackIndex]);
-						PreviewMesh->GetRetargetBasePose()[BoneIndex] = PoseTransforms[TrackIndex];
-					}
-					PRAGMA_ENABLE_DEPRECATION_WARNINGS
-
-					TurnOnPreviewRetargetBasePose();
-				}
-			}
-		}
-	}
-}
-
-void SRetargetSources::PostUndo()
-{
-}
-
-FText SRetargetSources::GetToggleRetargetBasePose() const
-{
-	UDebugSkelMeshComponent * PreviewMeshComp = PreviewScenePtr.Pin()->GetPreviewMeshComponent();
-	if(PreviewMeshComp && PreviewMeshComp->PreviewInstance)
-	{
-		if (PreviewMeshComp->PreviewInstance->GetForceRetargetBasePose())
-		{
-			return LOCTEXT("HideRetargetBasePose_Label", "Hide Pose");
-		}
-		else
-		{
-			return LOCTEXT("ViewRetargetBasePose_Label", "View Pose");
-		}
-	}
-
-	return LOCTEXT("InvalidRetargetBasePose_Label", "No Mesh for Base Pose");
-}
-
-void SRetargetSources::TurnOnPreviewRetargetBasePose()
-{
-	UDebugSkelMeshComponent * PreviewMeshComp = PreviewScenePtr.Pin()->GetPreviewMeshComponent();
-	if (PreviewMeshComp && PreviewMeshComp->PreviewInstance)
-	{
-		PreviewMeshComp->PreviewInstance->SetForceRetargetBasePose(true);
-	}
-}
 
 #undef LOCTEXT_NAMESPACE
diff --git a/Engine/Source/Editor/Persona/Private/SRetargetSources.h b/Engine/Source/Editor/Persona/Private/SRetargetSources.h
--- a/Engine/Source/Editor/Persona/Private/SRetargetSources.h
+++ b/Engine/Source/Editor/Persona/Private/SRetargetSources.h
@@ -7,8 +7,6 @@
 #include "Input/Reply.h"
 #include "Widgets/SCompoundWidget.h"
 #include "IPersonaPreviewScene.h"
-	
-#include "SPoseAssetNameWidget.h"
 
 class IEditableSkeleton;
 struct FAssetData;
@@ -21,41 +19,8 @@ public:
 
 	SLATE_END_ARGS()
 	
-	void Construct(const FArguments& InArgs, const TSharedRef<class IEditableSkeleton>& InEditableSkeleton, const TSharedRef<class IPersonaPreviewScene>& InPreviewScene, FSimpleMulticastDelegate& InOnPostUndo);
-
-private:
-
-	FReply OnViewRetargetBasePose() const;
-	FReply OnModifyPose();
-	FReply OnSaveRetargetBasePose();
-
-	/** The editable skeleton */
-	TWeakPtr<class IEditableSkeleton> EditableSkeletonPtr;
-
-	/** The preview scene  */
-	TWeakPtr<class IPersonaPreviewScene> PreviewScenePtr;
-
-	
-	/** Delegate for undo/redo transaction **/
-	void PostUndo();
-	FText GetToggleRetargetBasePose() const;	
-	void ResetRetargetBasePose();
-	void UseCurrentPose();
-	void ImportPose(const UPoseAsset* PoseAsset, const FName& PoseName);
-
-	TSharedRef<SWidget> OnModifyPoseContextMenu();
-
-	void SetSelectedPose(const FAssetData& InAssetData);
-	FString GetSelectedPose() const;
-	bool ShouldFilterAsset(const FAssetData& InAssetData);
-
-	TSharedPtr<SPoseAssetNameWidget> PoseAssetNameWidget;
-	TWeakObjectPtr<UPoseAsset> SelectedPoseAsset;
-	FString SelectedPoseName;
-
-	void SetPoseName(TSharedPtr<FString> PoseName, ESelectInfo::Type SelectionType);
-	FReply OnImportPose();
-
-	void TurnOnPreviewRetargetBasePose();
-	bool CanImportPose() const;
+	void Construct(
+		const FArguments& InArgs,
+		const TSharedRef<class IEditableSkeleton>& InEditableSkeleton,
+		FSimpleMulticastDelegate& InOnPostUndo);
 };
diff --git a/Engine/Source/Editor/Persona/Private/TabSpawners.cpp b/Engine/Source/Editor/Persona/Private/TabSpawners.cpp
--- a/Engine/Source/Editor/Persona/Private/TabSpawners.cpp
+++ b/Engine/Source/Editor/Persona/Private/TabSpawners.cpp
@@ -352,25 +352,27 @@ FTabSpawnerEntry& FPreviewViewportSummoner::RegisterTabSpawner(TSharedRef<FTabMa
 /////////////////////////////////////////////////////
 // FRetargetManagerTabSummoner
 
-FRetargetSourcesTabSummoner::FRetargetSourcesTabSummoner(TSharedPtr<class FAssetEditorToolkit> InHostingApp, const TSharedRef<IEditableSkeleton>& InEditableSkeleton, const TSharedRef<IPersonaPreviewScene>& InPreviewScene, FSimpleMulticastDelegate& InOnPostUndo)
+FRetargetSourcesTabSummoner::FRetargetSourcesTabSummoner(
+	TSharedPtr<class FAssetEditorToolkit> InHostingApp,
+	const TSharedRef<IEditableSkeleton>& InEditableSkeleton,
+	FSimpleMulticastDelegate& InOnPostUndo)
 	: FWorkflowTabFactory(FPersonaTabs::RetargetManagerID, InHostingApp)
 	, EditableSkeleton(InEditableSkeleton)
-	, PreviewScene(InPreviewScene)
 	, OnPostUndo(InOnPostUndo)
 {
-	TabLabel = LOCTEXT("RetargetSourcesTabTitle", "Retarget Sources");
+	TabLabel = LOCTEXT("RetargetingTabTitle", "Retargeting");
 	TabIcon = FSlateIcon(FAppStyle::GetAppStyleSetName(), "Persona.Tabs.RetargetManager");
 
 	EnableTabPadding();
 	bIsSingleton = true;
 
-	ViewMenuDescription = LOCTEXT("RetargetSourcesTabView", "Retarget Sources");
-	ViewMenuTooltip = LOCTEXT("RetargetSourcesTabView_ToolTip", "Retarget Sources indicate what proportions a sequence was authored with so that animation is correctly retargeted to other proportions.\n\nThese become 'Retarget Source' options on sequences.\n\nRetarget Sources are only needed when an animation sequence is authored on a skeletal mesh with proportions that are different than the default skeleton asset.");
+	ViewMenuDescription = LOCTEXT("RetargetingTabView", "Retargeting");
+	ViewMenuTooltip = LOCTEXT("RetargetSourcesTabView_ToolTip", "Specify Retarget Sources and Compatible Skeletons.");
 }
 
 TSharedRef<SWidget> FRetargetSourcesTabSummoner::CreateTabBody(const FWorkflowTabSpawnInfo& Info) const
 {
-	return SNew(SRetargetSources, EditableSkeleton.Pin().ToSharedRef(), PreviewScene.Pin().ToSharedRef(), OnPostUndo);
+	return SNew(SRetargetSources, EditableSkeleton.Pin().ToSharedRef(), OnPostUndo);
 }
 
 
diff --git a/Engine/Source/Editor/Persona/Private/TabSpawners.h b/Engine/Source/Editor/Persona/Private/TabSpawners.h
--- a/Engine/Source/Editor/Persona/Private/TabSpawners.h
+++ b/Engine/Source/Editor/Persona/Private/TabSpawners.h
@@ -243,19 +243,21 @@ struct FPreviewViewportSummoner : public FWorkflowTabFactory
 struct FRetargetSourcesTabSummoner : public FWorkflowTabFactory
 {
 public:
-	FRetargetSourcesTabSummoner(TSharedPtr<class FAssetEditorToolkit> InHostingApp, const TSharedRef<IEditableSkeleton>& InEditableSkeleton, const TSharedRef<IPersonaPreviewScene>& InPreviewScene, FSimpleMulticastDelegate& InOnPostUndo);
+	FRetargetSourcesTabSummoner(
+		TSharedPtr<class FAssetEditorToolkit> InHostingApp,
+		const TSharedRef<IEditableSkeleton>& InEditableSkeleton,
+		FSimpleMulticastDelegate& InOnPostUndo);
 
 	virtual TSharedRef<SWidget> CreateTabBody(const FWorkflowTabSpawnInfo& Info) const override;
 
 	// Create a tooltip widget for the tab
 	virtual TSharedPtr<SToolTip> CreateTabToolTipWidget(const FWorkflowTabSpawnInfo& Info) const override
 	{
-		return  IDocumentation::Get()->CreateToolTip(LOCTEXT("RetargetSourceTooltip", "In this panel, you can manage retarget sources for animations authored with varying proportions."), NULL, TEXT("Shared/Editors/Persona"), TEXT("RetargetSources"));
+		return IDocumentation::Get()->CreateToolTip(LOCTEXT("RetargetSourceTooltip", "In this panel, you can manage Retarget Sources (for playing animations authored with varying proportions) and Compatible Skeletons (for playing animations from other skeletal meshes)."), NULL, TEXT("Shared/Editors/Persona"), TEXT("RetargetSources"));
 	}
 
 private:
 	TWeakPtr<class IEditableSkeleton> EditableSkeleton;
-	TWeakPtr<class IPersonaPreviewScene> PreviewScene;
 	FSimpleMulticastDelegate& OnPostUndo;
 };
 
diff --git a/Engine/Source/Editor/Persona/Public/PersonaModule.h b/Engine/Source/Editor/Persona/Public/PersonaModule.h
--- a/Engine/Source/Editor/Persona/Public/PersonaModule.h
+++ b/Engine/Source/Editor/Persona/Public/PersonaModule.h
@@ -305,7 +305,7 @@ public:
 	virtual TSharedRef<FWorkflowTabFactory> CreateCurveMetadataEditorTabFactory(const TSharedRef<class FWorkflowCentricApplication>& InHostingApp, UObject* InMetadataHost, const TSharedRef<class IPersonaPreviewScene>& InPreviewScene, FOnObjectsSelected InOnObjectsSelected) const;
 
 	/** Create a retarget sources tab factory */
-	virtual TSharedRef<FWorkflowTabFactory> CreateRetargetSourcesTabFactory(const TSharedRef<class FWorkflowCentricApplication>& InHostingApp, const TSharedRef<class IEditableSkeleton>& InEditableSkeleton, const TSharedRef<IPersonaPreviewScene>& InPreviewScene, FSimpleMulticastDelegate& InOnPostUndo) const;
+	virtual TSharedRef<FWorkflowTabFactory> CreateRetargetSourcesTabFactory(const TSharedRef<class FWorkflowCentricApplication>& InHostingApp, const TSharedRef<class IEditableSkeleton>& InEditableSkeleton, FSimpleMulticastDelegate& InOnPostUndo) const;
 
 	/** Create a tab factory used to configure preview scene settings */
 	virtual TSharedRef<FWorkflowTabFactory> CreateAdvancedPreviewSceneTabFactory(const TSharedRef<class FWorkflowCentricApplication>& InHostingApp, const TSharedRef<IPersonaPreviewScene>& InPreviewScene) const;
diff --git a/Engine/Source/Editor/Persona/Public/SPoseAssetNameWidget.h b/Engine/Source/Editor/Persona/Public/SPoseAssetNameWidget.h
deleted file mode 100644
diff --git a/Engine/Source/Editor/SkeletonEditor/Private/SkeletonEditor.cpp b/Engine/Source/Editor/SkeletonEditor/Private/SkeletonEditor.cpp
--- a/Engine/Source/Editor/SkeletonEditor/Private/SkeletonEditor.cpp
+++ b/Engine/Source/Editor/SkeletonEditor/Private/SkeletonEditor.cpp
@@ -180,7 +180,7 @@ void FSkeletonEditor::BindCommands()
 	ToolkitCommands->MapAction(FSkeletonEditorCommands::Get().AnimNotifyWindow,
 		FExecuteAction::CreateSP(this, &FSkeletonEditor::OnAnimNotifyWindow));
 
-	ToolkitCommands->MapAction(FSkeletonEditorCommands::Get().RetargetManager,
+	ToolkitCommands->MapAction(FSkeletonEditorCommands::Get().Retargeting,
 		FExecuteAction::CreateSP(this, &FSkeletonEditor::OnRetargetManager));
 
 	ToolkitCommands->MapAction(FSkeletonEditorCommands::Get().ImportMesh,
@@ -229,7 +229,7 @@ void FSkeletonEditor::ExtendToolbar()
 	{
 		FToolMenuSection& SkeletonSection = ToolMenu->AddSection("Skeleton", LOCTEXT("ToolbarSkeletonSectionLabel", "Skeleton"), SectionInsertLocation);
 		SkeletonSection.AddEntry(FToolMenuEntry::InitToolBarButton(FSkeletonEditorCommands::Get().AnimNotifyWindow));
-		SkeletonSection.AddEntry(FToolMenuEntry::InitToolBarButton(FSkeletonEditorCommands::Get().RetargetManager, LOCTEXT("Toolbar_RetargetManager", "Retarget Manager")));
+		SkeletonSection.AddEntry(FToolMenuEntry::InitToolBarButton(FSkeletonEditorCommands::Get().Retargeting, LOCTEXT("Toolbar_RetargetManager", "Retarget Manager")));
 		SkeletonSection.AddEntry(FToolMenuEntry::InitToolBarButton(FSkeletonEditorCommands::Get().ImportMesh));
 	}
 
diff --git a/Engine/Source/Editor/SkeletonEditor/Private/SkeletonEditorCommands.cpp b/Engine/Source/Editor/SkeletonEditor/Private/SkeletonEditorCommands.cpp
--- a/Engine/Source/Editor/SkeletonEditor/Private/SkeletonEditorCommands.cpp
+++ b/Engine/Source/Editor/SkeletonEditor/Private/SkeletonEditorCommands.cpp
@@ -15,7 +15,7 @@ void FSkeletonEditorCommands::RegisterCommands()
 	UI_COMMAND(RemoveUnusedBones, "Remove Unused Bones from Skeleton", "Removes any bones from the skeleton that are not used by any of its meshes. The skeleton and associated animations will require saving after this action.", EUserInterfaceActionType::Button, FInputChord());
 	UI_COMMAND(TestSkeletonCurveMetaDataForUse, "Test Skeleton Curve Metadata for Use", "Iterates over all Animations and Meshes for this skeleton to look for unused morph/material curves", EUserInterfaceActionType::Button, FInputChord());
 	UI_COMMAND(AnimNotifyWindow, "Anim Notifies", "You can manage animation notifies that belong to the skeleton.", EUserInterfaceActionType::Button, FInputChord());
-	UI_COMMAND(RetargetManager, "Retarget Manager", "Manager retarget setups. ", EUserInterfaceActionType::Button, FInputChord());
+	UI_COMMAND(Retargeting, "Retargeting", "Edit Compatible Skeletons and Retarget Sources for skeletal meshes using this skeleton. ", EUserInterfaceActionType::Button, FInputChord());
 	UI_COMMAND(ImportMesh, "Import Mesh", "Import new mesh for this skeleton. ", EUserInterfaceActionType::Button, FInputChord());
 
 	UI_COMMAND(UpdateSkeletonRefPose, "Update Skeleton RefPose", "Update Skeleton ref pose based on current preview mesh", EUserInterfaceActionType::Button, FInputChord());
diff --git a/Engine/Source/Editor/SkeletonEditor/Private/SkeletonEditorCommands.h b/Engine/Source/Editor/SkeletonEditor/Private/SkeletonEditorCommands.h
--- a/Engine/Source/Editor/SkeletonEditor/Private/SkeletonEditorCommands.h
+++ b/Engine/Source/Editor/SkeletonEditor/Private/SkeletonEditorCommands.h
@@ -33,8 +33,8 @@ public:
 	TSharedPtr<FUICommandInfo> TestSkeletonCurveMetaDataForUse;
 	// Command to show Anim Notify window
 	TSharedPtr<FUICommandInfo> AnimNotifyWindow;
-	// Command to show Retarget Source Manager
-	TSharedPtr<FUICommandInfo> RetargetManager;
+	// Command to show retargeting tools
+	TSharedPtr<FUICommandInfo> Retargeting;
 	// Import Mesh for this Skeleton
 	TSharedPtr<FUICommandInfo> ImportMesh;
 
diff --git a/Engine/Source/Editor/SkeletonEditor/Private/SkeletonEditorMode.cpp b/Engine/Source/Editor/SkeletonEditor/Private/SkeletonEditorMode.cpp
--- a/Engine/Source/Editor/SkeletonEditor/Private/SkeletonEditorMode.cpp
+++ b/Engine/Source/Editor/SkeletonEditor/Private/SkeletonEditorMode.cpp
@@ -33,7 +33,7 @@ FSkeletonEditorMode::FSkeletonEditorMode(TSharedRef<FWorkflowCentricApplication>
 
 	TabFactories.RegisterFactory(PersonaModule.CreateAnimNotifiesTabFactory(InHostingApp, InSkeletonTree->GetEditableSkeleton(),  OnObjectsSelected));
 	TabFactories.RegisterFactory(PersonaModule.CreateAdvancedPreviewSceneTabFactory(InHostingApp, SkeletonEditor->GetPersonaToolkit()->GetPreviewScene()));
-	TabFactories.RegisterFactory(PersonaModule.CreateRetargetSourcesTabFactory(InHostingApp, InSkeletonTree->GetEditableSkeleton(), SkeletonEditor->GetPersonaToolkit()->GetPreviewScene(), SkeletonEditor->OnPostUndo));
+	TabFactories.RegisterFactory(PersonaModule.CreateRetargetSourcesTabFactory(InHostingApp, InSkeletonTree->GetEditableSkeleton(), SkeletonEditor->OnPostUndo));
 	TabFactories.RegisterFactory(PersonaModule.CreateCurveMetadataEditorTabFactory(InHostingApp, SkeletonEditor->HandleGetAsset(), SkeletonEditor->GetPersonaToolkit()->GetPreviewScene(), OnObjectsSelected));
 	TabFactories.RegisterFactory(PersonaModule.CreateSkeletonSlotNamesTabFactory(InHostingApp, InSkeletonTree->GetEditableSkeleton(), OnObjectSelected));
 	TabFactories.RegisterFactory(PersonaModule.CreateAssetDetailsTabFactory(InHostingApp, FOnGetAsset::CreateSP(&SkeletonEditor.Get(), &FSkeletonEditor::HandleGetAsset), FOnDetailsCreated()));
diff --git a/Engine/Source/Editor/UnrealEd/Private/ImportUtils/SkeletalMeshImportUtils.cpp b/Engine/Source/Editor/UnrealEd/Private/ImportUtils/SkeletalMeshImportUtils.cpp
--- a/Engine/Source/Editor/UnrealEd/Private/ImportUtils/SkeletalMeshImportUtils.cpp
+++ b/Engine/Source/Editor/UnrealEd/Private/ImportUtils/SkeletalMeshImportUtils.cpp
@@ -350,9 +350,6 @@ TSharedPtr<FExistingSkelMeshData> SkeletalMeshImportUtils::SaveExistingSkelMeshD
 	{
 		ExistingMeshDataPtr->ExistingMaterials = SourceMaterials;
 	}
-	PRAGMA_DISABLE_DEPRECATION_WARNINGS
-	ExistingMeshDataPtr->ExistingRetargetBasePose = SourceSkeletalMesh->GetRetargetBasePose();
-	PRAGMA_ENABLE_DEPRECATION_WARNINGS
 
 	ExistingMeshDataPtr->ExistingInlineReductionCacheDatas = SourceMeshModel->InlineReductionCacheDatas;
 
@@ -714,15 +711,6 @@ void SkeletalMeshImportUtils::RestoreExistingSkelMeshData(const TSharedPtr<const
 	//Do everything we need for base LOD re-import
 	if (SafeReimportLODIndex == 0)
 	{
-		PRAGMA_DISABLE_DEPRECATION_WARNINGS
-		// this is not ideal. Ideally we'll have to save only diff with indicating which joints, 
-		// but for now, we allow them to keep the previous pose IF the element count is same
-		if (MeshData->ExistingRetargetBasePose.Num() == SkeletalMesh->GetRefSkeleton().GetRawBoneNum())
-		{
-			SkeletalMesh->SetRetargetBasePose(MeshData->ExistingRetargetBasePose);
-		}
-		PRAGMA_ENABLE_DEPRECATION_WARNINGS
-
 		// Assign sockets from old version of this SkeletalMesh.
 		// Only copy ones for bones that exist in the new mesh.
 		for (int32 i = 0; i < MeshData->ExistingSockets.Num(); i++)
diff --git a/Engine/Source/Editor/UnrealEd/Private/SSkeletonWidget.cpp b/Engine/Source/Editor/UnrealEd/Private/SSkeletonWidget.cpp
--- a/Engine/Source/Editor/UnrealEd/Private/SSkeletonWidget.cpp
+++ b/Engine/Source/Editor/UnrealEd/Private/SSkeletonWidget.cpp
@@ -626,8 +626,6 @@ void SBasePoseViewport::SetSkeleton(USkeleton* Skeleton)
 			{
 				PreviewComponent->SetSkeletalMesh(PreviewSkeletalMesh);
 				PreviewComponent->EnablePreview(true, nullptr);
-//				PreviewComponent->AnimScriptInstance = PreviewComponent->PreviewInstance;
-				PreviewComponent->PreviewInstance->SetForceRetargetBasePose(true);
 				PreviewComponent->RefreshBoneTransforms(nullptr);
 
 				//Place the camera at a good viewer position
diff --git a/Engine/Source/Editor/UnrealEd/Public/ImportUtils/SkelImport.h b/Engine/Source/Editor/UnrealEd/Public/ImportUtils/SkelImport.h
--- a/Engine/Source/Editor/UnrealEd/Public/ImportUtils/SkelImport.h
+++ b/Engine/Source/Editor/UnrealEd/Public/ImportUtils/SkelImport.h
@@ -81,7 +81,6 @@ struct FExistingSkelMeshData
 	TArray<UPhysicsAsset*>					ExistingPhysicsAssets;
 	UPhysicsAsset*							ExistingShadowPhysicsAsset;
 	USkeleton*								ExistingSkeleton;
-	TArray<FTransform>						ExistingRetargetBasePose;
 	USkeletalMeshLODSettings*				ExistingLODSettings;
 	TSubclassOf<UAnimInstance>				ExistingPostProcessAnimBlueprint;
 	TSoftObjectPtr<UObject>					ExistingDefaultAnimatingRig;
diff --git a/Engine/Source/Runtime/Engine/Classes/Engine/SkeletalMesh.h b/Engine/Source/Runtime/Engine/Classes/Engine/SkeletalMesh.h
--- a/Engine/Source/Runtime/Engine/Classes/Engine/SkeletalMesh.h
+++ b/Engine/Source/Runtime/Engine/Classes/Engine/SkeletalMesh.h
@@ -2118,46 +2118,6 @@ public:
 		PRAGMA_ENABLE_DEPRECATION_WARNINGS
 	}
 
-	/** This is buffer that saves pose that is used by retargeting*/
-	UE_DEPRECATED(4.27, "Please do not access this member directly; use USkeletalMesh::GetRetargetBasePose() or USkeletalMesh::SetRetargetBasePose().")
-	UPROPERTY()
-	TArray<FTransform> RetargetBasePose;
-
-	UE_DEPRECATED(5.3, "Please do not use the retarget base pose. If you need a separate retarget pose, use the IK Retargeter or adjust the reference pose of the skeletal mesh.")
-	static FName GetRetargetBasePoseMemberName()
-	{
-		PRAGMA_DISABLE_DEPRECATION_WARNINGS
-		return GET_MEMBER_NAME_CHECKED(USkeletalMesh, RetargetBasePose);
-		PRAGMA_ENABLE_DEPRECATION_WARNINGS
-	}
-
-	UE_DEPRECATED(5.3, "Please do not use the retarget base pose. If you need a separate retarget pose, use the IK Retargeter or adjust the reference pose of the skeletal mesh.")
-	TArray<FTransform>& GetRetargetBasePose()
-	{
-		WaitUntilAsyncPropertyReleased(ESkeletalMeshAsyncProperties::RetargetBasePose);
-		PRAGMA_DISABLE_DEPRECATION_WARNINGS
-		return RetargetBasePose;
-		PRAGMA_ENABLE_DEPRECATION_WARNINGS
-	}
-
-	UE_DEPRECATED(5.3, "Please do not use the retarget base pose. If you need a separate retarget pose, use the IK Retargeter or adjust the reference pose of the skeletal mesh.")
-	const TArray<FTransform>& GetRetargetBasePose() const
-	{
-		WaitUntilAsyncPropertyReleased(ESkeletalMeshAsyncProperties::RetargetBasePose, ESkinnedAssetAsyncPropertyLockType::ReadOnly);
-		PRAGMA_DISABLE_DEPRECATION_WARNINGS
-		return RetargetBasePose;
-		PRAGMA_ENABLE_DEPRECATION_WARNINGS
-	}
-
-	UE_DEPRECATED(5.3, "Please do not use the retarget base pose. If you need a separate retarget pose, use the IK Retargeter or adjust the reference pose of the skeletal mesh.")
-	void SetRetargetBasePose(const TArray<FTransform>& InRetargetBasePose)
-	{
-		WaitUntilAsyncPropertyReleased(ESkeletalMeshAsyncProperties::RetargetBasePose);
-		PRAGMA_DISABLE_DEPRECATION_WARNINGS
-		RetargetBasePose = InRetargetBasePose;
-		PRAGMA_ENABLE_DEPRECATION_WARNINGS
-	}
-
 	/** Legacy clothing asset data, will be converted to new assets after loading */
 	UPROPERTY()
 	TArray<FClothingAssetData_Legacy>		ClothingAssets_DEPRECATED;
@@ -2564,9 +2524,6 @@ public:
 	/** Calculate the required bones for a Skeletal Mesh LOD, including possible extra influences */
 	static ENGINE_API void CalculateRequiredBones(FSkeletalMeshLODModel& LODModel, const struct FReferenceSkeleton& RefSkeleton, const TMap<FBoneIndexType, FBoneIndexType> * BonesToRemove);
 
-	/** Recalculate Retarget Base Pose BoneTransform */
-	ENGINE_API void ReallocateRetargetBasePose();
-
 	/**
 	 *	Add a skeletal socket object to this SkeletalMesh, and optionally promotes it to USkeleton socket.
 	 */
diff --git a/Engine/Source/Runtime/Engine/Private/Animation/AnimationRuntime.cpp b/Engine/Source/Runtime/Engine/Private/Animation/AnimationRuntime.cpp
--- a/Engine/Source/Runtime/Engine/Private/Animation/AnimationRuntime.cpp
+++ b/Engine/Source/Runtime/Engine/Private/Animation/AnimationRuntime.cpp
@@ -1624,26 +1624,6 @@ void FAnimationRuntime::FillWithRefPose(TArray<FTransform>& OutAtoms, const FBon
 	}
 }
 
-#if WITH_EDITOR
-PRAGMA_DISABLE_DEPRECATION_WARNINGS
-void FAnimationRuntime::FillWithRetargetBaseRefPose(FCompactPose& OutPose, const USkeletalMesh* Mesh)
-{
-	// Copy Target Asset's ref pose.
-	if (Mesh)
-	{
-		for (FCompactPoseBoneIndex BoneIndex : OutPose.ForEachBoneIndex())
-		{
-			int32 PoseIndex = OutPose.GetBoneContainer().MakeMeshPoseIndex(BoneIndex).GetInt();
-			if (Mesh->GetRetargetBasePose().IsValidIndex(PoseIndex))
-			{
-				OutPose[BoneIndex] = Mesh->GetRetargetBasePose()[PoseIndex];
-			}
-		}
-	}
-}
-PRAGMA_ENABLE_DEPRECATION_WARNINGS
-#endif // WITH_EDITOR
-
 void FAnimationRuntime::ConvertPoseToMeshSpace(const TArray<FTransform>& LocalTransforms, TArray<FTransform>& MeshSpaceTransforms, const FBoneContainer& RequiredBones)
 {
 	const int32 NumBones = RequiredBones.GetNumBones();
@@ -2523,6 +2503,7 @@ void FAnimationRuntime::MakeSkeletonRefPoseFromMesh(const USkeletalMesh* InMesh,
 		}
 	}
 }
+
 #if WITH_EDITOR
 void FAnimationRuntime::FillUpComponentSpaceTransformsRefPose(const USkeleton* Skeleton, TArray<FTransform> &ComponentSpaceTransforms)
 {
@@ -2532,31 +2513,7 @@ void FAnimationRuntime::FillUpComponentSpaceTransformsRefPose(const USkeleton* S
 	const TArray<FTransform>& ReferencePose = RefSkeleton.GetRefBonePose();
 	FillUpComponentSpaceTransforms(RefSkeleton, ReferencePose, ComponentSpaceTransforms);
 }
-
-PRAGMA_DISABLE_DEPRECATION_WARNINGS
-void FAnimationRuntime::FillUpComponentSpaceTransformsRetargetBasePose(const USkeletalMesh* Mesh, TArray<FTransform> &ComponentSpaceTransforms)
-{
-	if (Mesh)
-	{
-		const TArray<FTransform>& ReferencePose = Mesh->GetRetargetBasePose();
-		const FReferenceSkeleton& RefSkeleton = Mesh->GetRefSkeleton();
-		FillUpComponentSpaceTransforms(RefSkeleton, ReferencePose, ComponentSpaceTransforms);
-	}
-}
-
-void FAnimationRuntime::FillUpComponentSpaceTransformsRetargetBasePose(const USkeleton* Skeleton, TArray<FTransform> &ComponentSpaceTransforms)
-{
-	check(Skeleton);
-
-	// @Todo fixme: this has to get preview mesh instead of skeleton
-	const USkeletalMesh* PreviewMesh = Skeleton->GetPreviewMesh();
-	if (PreviewMesh)
-	{
-		FillUpComponentSpaceTransformsRetargetBasePose(PreviewMesh, ComponentSpaceTransforms);
-	}
-}
-PRAGMA_ENABLE_DEPRECATION_WARNINGS
-#endif // WITH_EDITOR
+#endif
 
 void FAnimationRuntime::AppendActiveMorphTargets(
 	const USkeletalMesh* InSkeletalMesh,
diff --git a/Engine/Source/Runtime/Engine/Private/Animation/Skeleton.cpp b/Engine/Source/Runtime/Engine/Private/Animation/Skeleton.cpp
--- a/Engine/Source/Runtime/Engine/Private/Animation/Skeleton.cpp
+++ b/Engine/Source/Runtime/Engine/Private/Animation/Skeleton.cpp
@@ -2045,10 +2045,6 @@ void USkeleton::HandleVirtualBoneChanges()
 			// also have to update retarget base pose
 			SkelMesh->GetRefSkeleton().RebuildRefSkeleton(this, bRebuildNameMap);
 			RebuildLinkup(SkelMesh);
-#if WITH_EDITOR
-			// whole bone count has changed, so it has to recalculate retarget base pose
-			SkelMesh->ReallocateRetargetBasePose();
-#endif // #if WITH_EDITOR
 		}
 	}
 
diff --git a/Engine/Source/Runtime/Engine/Private/SkeletalMesh.cpp b/Engine/Source/Runtime/Engine/Private/SkeletalMesh.cpp
--- a/Engine/Source/Runtime/Engine/Private/SkeletalMesh.cpp
+++ b/Engine/Source/Runtime/Engine/Private/SkeletalMesh.cpp
@@ -2191,57 +2191,11 @@ void USkeletalMesh::CalculateInvRefMatrices()
 			// Precompute inverse so we can use from-refpose-skin vertices.
 			GetRefBasesInvMatrix()[b] = FMatrix44f(ComposedRefPoseMatrices[b].Inverse());
 		}
-
-#if WITH_EDITORONLY_DATA
-		PRAGMA_DISABLE_DEPRECATION_WARNINGS
-		if(GetRetargetBasePose().Num() == 0)
-		{
-			SetRetargetBasePose(GetRefSkeleton().GetRefBonePose());
-		}
-		PRAGMA_ENABLE_DEPRECATION_WARNINGS
-#endif // WITH_EDITORONLY_DATA
 	}
 }
 
 #if WITH_EDITOR
 
-PRAGMA_DISABLE_DEPRECATION_WARNINGS
-void USkeletalMesh::ReallocateRetargetBasePose()
-{
-	// if you're adding other things here, please note that this function is called during postLoad
-	// fix up retarget base pose if VB has changed
-	// if we have virtual joints, we make sure Retarget Base Pose matches
-	const int32 RawNum = GetRefSkeleton().GetRawBoneNum();
-	const int32 VBNum = GetRefSkeleton().GetVirtualBoneRefData().Num();
-	const int32 BoneNum = GetRefSkeleton().GetNum();
-	check(RawNum + VBNum == BoneNum);
-
-	const int32 OldRetargetBasePoseNum = GetRetargetBasePose().Num();
-	// we want to make sure retarget base pose contains raw numbers PREVIOUSLY
-	// otherwise, we may override wrong transform
-	if (OldRetargetBasePoseNum >= RawNum)
-	{
-		// we have to do this in case buffer size changes (shrink for example)
-		GetRetargetBasePose().SetNum(BoneNum);
-
-		// if we have VB, we should override them
-		// they're not editable, so it's fine to override them from raw bones
-		if (VBNum > 0)
-		{
-			const TArray<FTransform>& BonePose = GetRefSkeleton().GetRefBonePose();
-			check(GetRetargetBasePose().GetTypeSize() == BonePose.GetTypeSize());
-			const int32 ElementSize = GetRetargetBasePose().GetTypeSize();
-			FMemory::Memcpy(GetRetargetBasePose().GetData() + RawNum, BonePose.GetData() + RawNum, ElementSize*VBNum);
-		}
-	}
-	else
-	{
-		// else we think, something has changed, we just override retarget base pose to current pose
-		GetRetargetBasePose() = GetRefSkeleton().GetRefBonePose();
-	}
-}
-PRAGMA_ENABLE_DEPRECATION_WARNINGS
-
 void USkeletalMesh::CalculateRequiredBones(FSkeletalMeshLODModel& LODModel, const struct FReferenceSkeleton& InRefSkeleton, const TMap<FBoneIndexType, FBoneIndexType> * BonesToRemove)
 {
 	// RequiredBones for base model includes all raw bones.
@@ -3554,20 +3508,6 @@ void USkeletalMesh::FinishPostLoadInternal(FSkinnedAssetPostLoadContext& Context
 
 	CalculateInvRefMatrices();
 
-#if WITH_EDITORONLY_DATA
-	PRAGMA_DISABLE_DEPRECATION_WARNINGS
-	if (GetRetargetBasePose().Num() == 0 && !GetOutermost()->bIsCookedForEditor)
-	{
-		GetRetargetBasePose() = GetRefSkeleton().GetRefBonePose();
-	}
-	PRAGMA_ENABLE_DEPRECATION_WARNINGS
-
-	if (GetLinkerCustomVersion(FFortniteMainBranchObjectVersion::GUID) < FFortniteMainBranchObjectVersion::SupportVirtualBoneInRetargeting)
-	{
-		ReallocateRetargetBasePose();
-	}
-#endif
-
 	// Bounds have been loaded - apply extensions.
 	CalculateExtendedBounds();
 
@@ -5747,25 +5687,15 @@ TSoftObjectPtr<UObject> USkeletalMesh::GetDefaultAnimatingRig() const
 
 void USkeletalMesh::GetMappableNodeData(TArray<FName>& OutNames, TArray<FNodeItem>& OutNodeItems) const
 {
-	TArray<FTransform> ComponentSpaceRefPose;
-#if WITH_EDITORONLY_DATA
-	FAnimationRuntime::FillUpComponentSpaceTransformsRetargetBasePose(this, ComponentSpaceRefPose);
-#else
-	// hasn't tested this route, but we don't have retarget base pose if not editor, wonder we should to non-editor soon
-	ensure(false);
-	FAnimationRuntime::FillUpComponentSpaceTransforms(GetRefSkeleton(), GetRefSkeleton().GetRefBonePose(), ComponentSpaceRefPose);
-#endif //
+	
 
 	const int32 NumJoint = GetRefSkeleton().GetNum();
 	// allocate buffer
 	OutNames.Reset(NumJoint);
 	OutNodeItems.Reset(NumJoint);
 
-	if(ComponentSpaceRefPose.Num() < NumJoint)
-	{
-		// if the mesh's RetargetBasePose is out of whack we should rely on the ref skeleton
-		FAnimationRuntime::FillUpComponentSpaceTransforms(GetRefSkeleton(), GetRefSkeleton().GetRefBonePose(), ComponentSpaceRefPose);
-	}
+	TArray<FTransform> ComponentSpaceRefPose;
+	FAnimationRuntime::FillUpComponentSpaceTransforms(GetRefSkeleton(), GetRefSkeleton().GetRefBonePose(), ComponentSpaceRefPose);
 
 	if (NumJoint > 0)
 	{
diff --git a/Engine/Source/Runtime/Engine/Public/AnimationRuntime.h b/Engine/Source/Runtime/Engine/Public/AnimationRuntime.h
--- a/Engine/Source/Runtime/Engine/Public/AnimationRuntime.h
+++ b/Engine/Source/Runtime/Engine/Public/AnimationRuntime.h
@@ -418,11 +418,6 @@ public:
 	/** Fill ref pose **/
 	static ENGINE_API void FillWithRefPose(TArray<FTransform>& OutAtoms, const FBoneContainer& RequiredBones);
 
-#if WITH_EDITOR
-	/** fill with retarget base ref pose but this isn't used during run-time, so it always copies all of them */
-	static ENGINE_API void FillWithRetargetBaseRefPose(FCompactPose& OutPose, const USkeletalMesh* Mesh);
-#endif
-
 	/** Convert LocalTransforms into MeshSpaceTransforms over RequiredBones. */
 	static ENGINE_API void ConvertPoseToMeshSpace(const TArray<FTransform>& LocalTransforms, TArray<FTransform>& MeshSpaceTransforms, const FBoneContainer& RequiredBones);
 
@@ -621,8 +616,6 @@ public:
 
 #if WITH_EDITOR
 	static ENGINE_API void FillUpComponentSpaceTransformsRefPose(const USkeleton* Skeleton, TArray<FTransform> &ComponentSpaceTransforms);
-	static ENGINE_API void FillUpComponentSpaceTransformsRetargetBasePose(const USkeleton* Skeleton, TArray<FTransform> &ComponentSpaceTransforms);
-	static ENGINE_API void FillUpComponentSpaceTransformsRetargetBasePose(const USkeletalMesh* Mesh, TArray<FTransform> &ComponentSpaceTransforms);
 #endif
 
 	/* Weight utility functions */
