[Skeletal mesh editor] Guard against invalid personal tookkit ptr since the refresh callback for lodinfo is fired async and by the time it fires this weak ptr can be invalid

diff --git a/Engine/Source/Editor/Persona/Private/LODInfoUILayout.cpp b/Engine/Source/Editor/Persona/Private/LODInfoUILayout.cpp
--- a/Engine/Source/Editor/Persona/Private/LODInfoUILayout.cpp
+++ b/Engine/Source/Editor/Persona/Private/LODInfoUILayout.cpp
@@ -27,11 +27,16 @@ void ULODInfoUILayout::SetReferenceLODInfo(TWeakPtr<IPersonaToolkit> InPersonaTo
 
 void ULODInfoUILayout::RefreshReferenceLODInfo()
 {
-	USkeletalMesh* SkeletalMesh = PersonaToolkit.Pin()->GetPreviewMesh();
-	const FSkeletalMeshLODInfo* SkeletalMeshLODInfo = SkeletalMesh->GetLODInfo(LODIndex);
-	check(SkeletalMeshLODInfo != nullptr);
-	//Copy the LODInfo Array to the temporary
-	LODInfo = *SkeletalMeshLODInfo;	
+	// Request to refresh can be queued from other threads and thus
+	// by the time it tries to refresh, it is possible that we no longer have a valid ptr
+	if (PersonaToolkit.IsValid())
+	{
+		USkeletalMesh* SkeletalMesh = PersonaToolkit.Pin()->GetPreviewMesh();
+		const FSkeletalMeshLODInfo* SkeletalMeshLODInfo = SkeletalMesh->GetLODInfo(LODIndex);
+		check(SkeletalMeshLODInfo != nullptr);
+		//Copy the LODInfo Array to the temporary
+		LODInfo = *SkeletalMeshLODInfo;	
+	}
 }
 
 void ULODInfoUILayout::PostEditChangeProperty(struct FPropertyChangedEvent& PropertyChangedEvent)
