Added multi-color bone support
Precursor to multi-color weights.

diff --git a/Engine/Source/Editor/Persona/Private/AnimViewportShowCommands.cpp b/Engine/Source/Editor/Persona/Private/AnimViewportShowCommands.cpp
--- a/Engine/Source/Editor/Persona/Private/AnimViewportShowCommands.cpp
+++ b/Engine/Source/Editor/Persona/Private/AnimViewportShowCommands.cpp
@@ -29,6 +29,7 @@ void FAnimViewportShowCommands::RegisterCommands()
 	UI_COMMAND( ShowMorphTargets, "Morph Targets", "Display applied morph targets of the mesh", EUserInterfaceActionType::ToggleButton, FInputChord() );
 
 	UI_COMMAND( ShowBoneNames, "Bone Names", "Display bone names in the viewport", EUserInterfaceActionType::ToggleButton, FInputChord() );
+	UI_COMMAND( ShowBoneColors, "Multi-Color Bones", "Display skeleton with multiple colors", EUserInterfaceActionType::ToggleButton, FInputChord() );
 
 	// below 3 menus are radio button styles
 	UI_COMMAND(ShowDisplayInfoBasic, "Basic", "Display basic mesh info in the viewport", EUserInterfaceActionType::RadioButton, FInputChord());
diff --git a/Engine/Source/Editor/Persona/Private/AnimViewportShowCommands.h b/Engine/Source/Editor/Persona/Private/AnimViewportShowCommands.h
--- a/Engine/Source/Editor/Persona/Private/AnimViewportShowCommands.h
+++ b/Engine/Source/Editor/Persona/Private/AnimViewportShowCommands.h
@@ -104,6 +104,9 @@ public:
 	/** Show skeletal mesh bone names */
 	TSharedPtr< FUICommandInfo > ShowBoneNames;
 
+	/** Show skeletal mesh bone names */
+	TSharedPtr< FUICommandInfo > ShowBoneColors;
+
 	/** Show skeletal mesh info */
 	TSharedPtr< FUICommandInfo > ShowDisplayInfoBasic;
 	TSharedPtr< FUICommandInfo > ShowDisplayInfoDetailed;
diff --git a/Engine/Source/Editor/Persona/Private/AnimationEditorViewportClient.cpp b/Engine/Source/Editor/Persona/Private/AnimationEditorViewportClient.cpp
--- a/Engine/Source/Editor/Persona/Private/AnimationEditorViewportClient.cpp
+++ b/Engine/Source/Editor/Persona/Private/AnimationEditorViewportClient.cpp
@@ -1777,20 +1777,17 @@ void FAnimationViewportClient::DrawMeshBones(UDebugSkelMeshComponent* MeshCompon
 	TArray<FLinearColor> BoneColours;
 	BoneColours.AddUninitialized(MeshComponent->GetNumDrawTransform());
 
-	// factor skeleton draw mode into color selection
-	const FLinearColor BoneColor = MeshComponent->SkeletonDrawMode == ESkeletonDrawMode::GreyedOut ? GetDefault<UPersonaOptions>()->DisabledBoneColor : GetDefault<UPersonaOptions>()->DefaultBoneColor;
-	const FLinearColor VirtualBoneColor = MeshComponent->SkeletonDrawMode == ESkeletonDrawMode::GreyedOut ? GetDefault<UPersonaOptions>()->DisabledBoneColor : GetDefault<UPersonaOptions>()->VirtualBoneColor;
-	
 	// we could cache parent bones as we calculate, but right now I'm not worried about perf issue of this
 	const TArray<FBoneIndexType>& DrawBoneIndices = MeshComponent->GetDrawBoneIndices();
 	for ( int32 Index=0; Index<DrawBoneIndices.Num(); ++Index )
 	{
 		const int32 BoneIndex = DrawBoneIndices[Index];
 		WorldTransforms[BoneIndex] = MeshComponent->GetDrawTransform(BoneIndex) * MeshComponent->GetComponentTransform();
-		BoneColours[BoneIndex] = BoneColor;
+		BoneColours[BoneIndex] = MeshComponent->GetBoneColor(Index);
 	}
 
-	// color virtual bones
+	// color virtual bones, factoring in skeleton draw mode into color selection
+	const FLinearColor VirtualBoneColor = MeshComponent->SkeletonDrawMode == ESkeletonDrawMode::GreyedOut ? GetDefault<UPersonaOptions>()->DisabledBoneColor : GetDefault<UPersonaOptions>()->VirtualBoneColor;
 	for (int16 VirtualBoneIndex : MeshComponent->GetReferenceSkeleton().GetRequiredVirtualBones())
 	{
 		BoneColours[VirtualBoneIndex] = VirtualBoneColor;
diff --git a/Engine/Source/Editor/Persona/Private/SAnimViewportToolBar.cpp b/Engine/Source/Editor/Persona/Private/SAnimViewportToolBar.cpp
--- a/Engine/Source/Editor/Persona/Private/SAnimViewportToolBar.cpp
+++ b/Engine/Source/Editor/Persona/Private/SAnimViewportToolBar.cpp
@@ -812,6 +812,7 @@ TSharedRef<SWidget> SAnimViewportToolBar::GenerateCharacterMenu() const
 						Section.AddMenuEntry(FAnimViewportShowCommands::Get().ShowSockets);
 						Section.AddMenuEntry(FAnimViewportShowCommands::Get().ShowAttributes);
 						Section.AddMenuEntry(FAnimViewportShowCommands::Get().ShowBoneNames);
+						Section.AddMenuEntry(FAnimViewportShowCommands::Get().ShowBoneColors);
 					}
 
 					{
diff --git a/Engine/Source/Editor/Persona/Private/SAnimationEditorViewport.cpp b/Engine/Source/Editor/Persona/Private/SAnimationEditorViewport.cpp
--- a/Engine/Source/Editor/Persona/Private/SAnimationEditorViewport.cpp
+++ b/Engine/Source/Editor/Persona/Private/SAnimationEditorViewport.cpp
@@ -796,6 +796,12 @@ void SAnimationEditorViewportTabBody::BindCommands()
 		FCanExecuteAction(),
 		FIsActionChecked::CreateSP(this, &SAnimationEditorViewportTabBody::IsShowingBoneNames));
 
+	CommandList.MapAction( 
+		ViewportShowMenuCommands.ShowBoneColors,
+		FExecuteAction::CreateSP(this, &SAnimationEditorViewportTabBody::OnShowBoneColors),
+		FCanExecuteAction(),
+		FIsActionChecked::CreateSP(this, &SAnimationEditorViewportTabBody::IsShowingBoneColors));
+
 	CommandList.MapAction( 
 		ViewportShowMenuCommands.ShowRawAnimation,
 		FExecuteAction::CreateSP(this, &SAnimationEditorViewportTabBody::OnShowRawAnimation),
@@ -1301,6 +1307,22 @@ bool SAnimationEditorViewportTabBody::IsShowingBoneNames() const
 	return PreviewComponent != NULL && PreviewComponent->bShowBoneNames;
 }
 
+void SAnimationEditorViewportTabBody::OnShowBoneColors()
+{
+	ForEachDebugMesh([](UDebugSkelMeshComponent* InMesh)
+	{
+		InMesh->bShowBoneColors = !InMesh->bShowBoneColors;
+		InMesh->MarkRenderStateDirty();
+	});
+	RefreshViewport();
+}
+
+bool SAnimationEditorViewportTabBody::IsShowingBoneColors() const
+{
+	const UDebugSkelMeshComponent* PreviewComponent = GetPreviewScene()->GetPreviewMeshComponent();
+	return PreviewComponent && PreviewComponent->bShowBoneColors;
+}
+
 bool SAnimationEditorViewportTabBody::IsShowingRawAnimation() const
 {
 	UDebugSkelMeshComponent* PreviewComponent = GetPreviewScene()->GetPreviewMeshComponent();
diff --git a/Engine/Source/Editor/Persona/Private/SAnimationEditorViewport.h b/Engine/Source/Editor/Persona/Private/SAnimationEditorViewport.h
--- a/Engine/Source/Editor/Persona/Private/SAnimationEditorViewport.h
+++ b/Engine/Source/Editor/Persona/Private/SAnimationEditorViewport.h
@@ -368,6 +368,12 @@ private:
 
 	/** Function to check whether bone names are displayed or not */
 	bool IsShowingBoneNames() const;
+
+	/** Function to show/hide bone colors */
+	void OnShowBoneColors();
+
+	/** Function to check whether bone's are using multiple colors */
+	bool IsShowingBoneColors() const;
 	
 	/** Function to show/hide selected bone weight */
 	void OnShowOverlayNone();
diff --git a/Engine/Source/Editor/UnrealEd/Classes/Animation/DebugSkelMeshComponent.h b/Engine/Source/Editor/UnrealEd/Classes/Animation/DebugSkelMeshComponent.h
--- a/Engine/Source/Editor/UnrealEd/Classes/Animation/DebugSkelMeshComponent.h
+++ b/Engine/Source/Editor/UnrealEd/Classes/Animation/DebugSkelMeshComponent.h
@@ -293,10 +293,6 @@ class UDebugSkelMeshComponent : public USkeletalMeshComponent
 	UPROPERTY(transient)
 	uint32 bRequiredBonesUpToDateDuringTick : 1;
 
-	/** Multiplier for the bone radius rendering */
-	UPROPERTY(transient)
-	float BoneRadiusMultiplier;
-
 	/* Bounds computed from cloth. */
 	FBoxSphereBounds CachedClothBounds;
 
@@ -357,6 +353,13 @@ class UDebugSkelMeshComponent : public USkeletalMeshComponent
 	UPROPERTY(transient)
 	bool bTrackAttachedInstanceLOD;
 
+	/** If true then the bones are displayed with multiple semi-random colors */
+	UPROPERTY()
+	bool bShowBoneColors = true;
+	
+	/** Palettized semi randomized colors for multi-color bone display */
+	UNREALED_API FLinearColor GetBoneColor(int32 BoneIndex);
+
 	// Helper method that sets the forced lod
 	UNREALED_API void SetDebugForcedLOD(int32 InNewForcedLOD);
 
@@ -686,7 +689,6 @@ public:
 
 		return FTransform::Identity;
 	}
-
 };
 
 
diff --git a/Engine/Source/Editor/UnrealEd/Private/Animation/DebugSkelMeshComponent.cpp b/Engine/Source/Editor/UnrealEd/Private/Animation/DebugSkelMeshComponent.cpp
--- a/Engine/Source/Editor/UnrealEd/Private/Animation/DebugSkelMeshComponent.cpp
+++ b/Engine/Source/Editor/UnrealEd/Private/Animation/DebugSkelMeshComponent.cpp
@@ -25,6 +25,7 @@
 #include "Utils/ClothingMeshUtils.h"
 #include "DynamicMeshBuilder.h"
 #include "Materials/MaterialInstanceDynamic.h"
+#include "Preferences/PersonaOptions.h"
 
 //////////////////////////////////////////////////////////////////////////
 // UDebugSkelMeshComponent
@@ -79,6 +80,30 @@ UDebugSkelMeshComponent::UDebugSkelMeshComponent(const FObjectInitializer& Objec
 	ConsumeRootMotionPreviousPlaybackTime = 0.f;
 }
 
+FLinearColor UDebugSkelMeshComponent::GetBoneColor(int32 InBoneIndex)
+{
+	// this returns the normal unmodified color of the bone, calling code must account
+	// for any editor specific states that might affect the final bone color (like selection)
+	
+	// skeleton greyed out
+	if (SkeletonDrawMode == ESkeletonDrawMode::GreyedOut)
+	{
+		return GetDefault<UPersonaOptions>()->DisabledBoneColor;
+	}
+
+	// using default color for all bones
+	if (!bShowBoneColors)
+	{
+		return GetDefault<UPersonaOptions>()->DefaultBoneColor;
+	}
+	
+	// uses deterministic, semi-random desaturated color unique to the bone index
+	constexpr float Saturation = 0.5f;
+	constexpr float Value = 1.0f;
+	constexpr float Rotation = 90.f;
+	return FLinearColor::IntToDistinctColor(InBoneIndex, Saturation, Value, Rotation);
+}
+
 void UDebugSkelMeshComponent::SetDebugForcedLOD(int32 InNewForcedLOD)
 {
 	SetForcedLOD(InNewForcedLOD);
diff --git a/Engine/Source/Runtime/Core/Public/Math/Color.h b/Engine/Source/Runtime/Core/Public/Math/Color.h
--- a/Engine/Source/Runtime/Core/Public/Math/Color.h
+++ b/Engine/Source/Runtime/Core/Public/Math/Color.h
@@ -443,6 +443,24 @@ struct FLinearColor
 		return (ClampedLo < 1.0f) ? ClampedLo : 1.0f;
 	}
 
+	/**
+	 * @brief Helper function to generate distinct colors from a sequence of integers where each integer increment
+	 * spins around the Hue wheel by increments of the golden ratio.
+	 * @param Integer		the "seed" for a semi-random, but deterministic color
+	 * @param Saturation	0-1 saturation of resulting color
+	 * @param Value			0-1 brightness of resulting color
+	 * @param HueRotation	0-360 amount to rotate around the hue wheel
+	 * @return color with semi-random hue */
+	static FLinearColor IntToDistinctColor(
+		const int32 Integer,
+		const float Saturation=1.f,
+		const float Value=1.f,
+		const float HueRotation=180.f)
+	{
+		constexpr float GoldenRatio = 1.618033988749895f; // (1 + FMath::Sqrt(5.f)) / 2.f;
+		return FLinearColor(Integer * GoldenRatio * HueRotation, Saturation, Value).HSVToLinearRGB();
+	}
+
 	// Common colors.	
 	static CORE_API const FLinearColor White;
 	static CORE_API const FLinearColor Gray;
