Added multi-color bone support
Precursor to multi-color weights.

diff --git a/Engine/Plugins/Animation/IKRig/Source/IKRigEditor/Private/RetargetEditor/IKRetargetEditorController.cpp b/Engine/Plugins/Animation/IKRig/Source/IKRigEditor/Private/RetargetEditor/IKRetargetEditorController.cpp
--- a/Engine/Plugins/Animation/IKRig/Source/IKRigEditor/Private/RetargetEditor/IKRetargetEditorController.cpp
+++ b/Engine/Plugins/Animation/IKRig/Source/IKRigEditor/Private/RetargetEditor/IKRetargetEditorController.cpp
@@ -1454,8 +1454,15 @@ void FIKRetargetEditorController::RenderSkeleton(FPrimitiveDrawInterface* PDI, E
 	// generate bone colors, blue on selected chains
 	TArray<FLinearColor> BoneColors;
 	{
-		// set all to default color
-		BoneColors.Init(DefaultColor, RefSkeleton.GetNum());
+		// set default colors
+		if (MeshComponent->bShowBoneColors)
+		{
+			SkeletalDebugRendering::FillWithMultiColors(BoneColors, RefSkeleton.GetNum());
+		}
+		else
+		{
+			BoneColors.Init(DefaultColor, RefSkeleton.GetNum());
+		}
 
 		// highlight selected chains in blue
 		const TArray<FName>& SelectedChainNames = GetSelectedChains();
diff --git a/Engine/Plugins/Animation/IKRig/Source/IKRigEditor/Private/RigEditor/IKRigEditMode.cpp b/Engine/Plugins/Animation/IKRig/Source/IKRigEditor/Private/RigEditor/IKRigEditMode.cpp
--- a/Engine/Plugins/Animation/IKRig/Source/IKRigEditor/Private/RigEditor/IKRigEditMode.cpp
+++ b/Engine/Plugins/Animation/IKRig/Source/IKRigEditor/Private/RigEditor/IKRigEditMode.cpp
@@ -190,7 +190,7 @@ void FIKRigEditMode::RenderBones(FPrimitiveDrawInterface* PDI)
 
 	// get bone colors
 	TArray<FLinearColor> BoneColors;
-	GetBoneColors(Controller.Get(), CurrentProcessor, RefSkeleton, BoneColors);
+	GetBoneColors(Controller.Get(), CurrentProcessor, RefSkeleton, MeshComponent->bShowBoneColors, BoneColors);
 
 	SkeletalDebugRendering::DrawBones(
 		PDI,
@@ -209,6 +209,7 @@ void FIKRigEditMode::GetBoneColors(
 	FIKRigEditorController* Controller,
 	const UIKRigProcessor* Processor,
 	const FReferenceSkeleton& RefSkeleton,
+	bool bUseMultiColorAsDefaultColor,
 	TArray<FLinearColor>& OutBoneColors) const
 {
 	const FLinearColor DefaultColor = GetMutableDefault<UPersonaOptions>()->DefaultBoneColor;
@@ -216,10 +217,13 @@ void FIKRigEditMode::GetBoneColors(
 	const FLinearColor ErrorColor = FLinearColor::Red;
 
 	// set all to default color
-	OutBoneColors.SetNum(RefSkeleton.GetNum());
-	for (int32 Index=0; Index<RefSkeleton.GetNum(); ++Index)
+	if (bUseMultiColorAsDefaultColor)
 	{
-		OutBoneColors[Index] = DefaultColor;
+		SkeletalDebugRendering::FillWithMultiColors(OutBoneColors, RefSkeleton.GetNum());
+	}
+	else
+	{
+		OutBoneColors.Init(DefaultColor, RefSkeleton.GetNum());		
 	}
 	
 	// highlight bones of the last selected UI element (could be solver or retarget chain) 
diff --git a/Engine/Plugins/Animation/IKRig/Source/IKRigEditor/Public/RigEditor/IKRigEditMode.h b/Engine/Plugins/Animation/IKRig/Source/IKRigEditor/Public/RigEditor/IKRigEditMode.h
--- a/Engine/Plugins/Animation/IKRig/Source/IKRigEditor/Public/RigEditor/IKRigEditMode.h
+++ b/Engine/Plugins/Animation/IKRig/Source/IKRigEditor/Public/RigEditor/IKRigEditMode.h
@@ -62,6 +62,7 @@ private:
 		FIKRigEditorController* Controller,
 		const UIKRigProcessor* Processor,
 		const FReferenceSkeleton& RefSkeleton,
+		bool bUseMultiColorAsDefaultColor,
 		TArray<FLinearColor>& OutBoneColors) const;
 
 	bool HandleBeginTransform() const;
diff --git a/Engine/Plugins/Animation/PoseSearch/Source/Editor/Private/PoseSearchDatabaseEdMode.cpp b/Engine/Plugins/Animation/PoseSearch/Source/Editor/Private/PoseSearchDatabaseEdMode.cpp
--- a/Engine/Plugins/Animation/PoseSearch/Source/Editor/Private/PoseSearchDatabaseEdMode.cpp
+++ b/Engine/Plugins/Animation/PoseSearch/Source/Editor/Private/PoseSearchDatabaseEdMode.cpp
@@ -61,20 +61,17 @@ namespace UE::PoseSearch
 					{
 						WorldTransforms.SetNumUninitialized(MeshComponent->GetNumDrawTransform());
 						BoneColors.SetNumUninitialized(MeshComponent->GetNumDrawTransform());
-
-						// factor skeleton draw mode into color selection
-						const FLinearColor BoneColor = MeshComponent->SkeletonDrawMode == ESkeletonDrawMode::GreyedOut ? PersonaOptions->DisabledBoneColor : PersonaOptions->DefaultBoneColor;
-						const FLinearColor VirtualBoneColor = MeshComponent->SkeletonDrawMode == ESkeletonDrawMode::GreyedOut ? PersonaOptions->DisabledBoneColor : PersonaOptions->VirtualBoneColor;
-
+						
 						const TArray<FBoneIndexType>& DrawBoneIndices = MeshComponent->GetDrawBoneIndices();
 						for (int32 Index = 0; Index < DrawBoneIndices.Num(); ++Index)
 						{
 							const int32 BoneIndex = DrawBoneIndices[Index];
 							WorldTransforms[BoneIndex] = MeshComponent->GetDrawTransform(BoneIndex) * MeshComponent->GetComponentTransform();
-							BoneColors[BoneIndex] = BoneColor;
+							BoneColors[BoneIndex] = MeshComponent->GetBoneColor(BoneIndex);
 						}
 
 						// color virtual bones
+						const FLinearColor VirtualBoneColor = MeshComponent->SkeletonDrawMode == ESkeletonDrawMode::GreyedOut ? PersonaOptions->DisabledBoneColor : PersonaOptions->VirtualBoneColor;
 						for (int16 VirtualBoneIndex : MeshComponent->GetReferenceSkeleton().GetRequiredVirtualBones())
 						{
 							BoneColors[VirtualBoneIndex] = VirtualBoneColor;
@@ -85,6 +82,7 @@ namespace UE::PoseSearch
 						DrawConfig.BoneDrawSize = 0.2f;
 						DrawConfig.bAddHitProxy = false;
 						DrawConfig.bForceDraw = false;
+						DrawConfig.bUseMultiColorAsDefaultColor = MeshComponent->bShowBoneColors;
 						DrawConfig.DefaultBoneColor = PersonaOptions->DefaultBoneColor;
 						DrawConfig.AffectedBoneColor = PersonaOptions->AffectedBoneColor;
 						DrawConfig.SelectedBoneColor = PersonaOptions->SelectedBoneColor;
diff --git a/Engine/Plugins/Experimental/MeshModelingToolsetExp/Source/MeshModelingToolsEditorOnlyExp/Private/SkeletalMesh/SkeletonEditingTool.cpp b/Engine/Plugins/Experimental/MeshModelingToolsetExp/Source/MeshModelingToolsEditorOnlyExp/Private/SkeletalMesh/SkeletonEditingTool.cpp
--- a/Engine/Plugins/Experimental/MeshModelingToolsetExp/Source/MeshModelingToolsEditorOnlyExp/Private/SkeletalMesh/SkeletonEditingTool.cpp
+++ b/Engine/Plugins/Experimental/MeshModelingToolsetExp/Source/MeshModelingToolsEditorOnlyExp/Private/SkeletalMesh/SkeletonEditingTool.cpp
@@ -18,6 +18,7 @@
 #include "BaseBehaviors/ClickDragBehavior.h"
 #include "BaseGizmos/GizmoViewContext.h"
 #include "DynamicMesh/MeshNormals.h"
+#include "Preferences/PersonaOptions.h"
 #include "Selection/PolygonSelectionMechanic.h"
 
 #include "SkeletalMesh/SkeletonTransformProxy.h"
@@ -1714,19 +1715,16 @@ void USkeletonEditingTool::Render(IToolsContextRenderAPI* RenderAPI)
 		return;
 	}
 
-	static const FLinearColor DefaultBoneColor(0.0f,0.0f,0.025f,1.0f);
-	static const FLinearColor SelectedBoneColor(0.2f,1.0f,0.2f,1.0f);
-	static const FLinearColor AffectedBoneColor(1.0f,1.0f,1.0f,1.0f);
-	static const FLinearColor ParentOfSelectedBoneColor(0.85f,0.45f,0.12f,1.0f);
+	const UPersonaOptions* PersonaOptions = GetDefault<UPersonaOptions>();
 	static FSkelDebugDrawConfig DrawConfig;
 		DrawConfig.BoneDrawMode = EBoneDrawMode::Type::All;
 		DrawConfig.BoneDrawSize = 1.f;
 		DrawConfig.bAddHitProxy = true;
 		DrawConfig.bForceDraw = false;
-		DrawConfig.DefaultBoneColor = DefaultBoneColor;
-		DrawConfig.AffectedBoneColor = AffectedBoneColor;
-		DrawConfig.SelectedBoneColor = SelectedBoneColor;
-		DrawConfig.ParentOfSelectedBoneColor = ParentOfSelectedBoneColor;
+		DrawConfig.DefaultBoneColor = PersonaOptions->DefaultBoneColor;
+		DrawConfig.AffectedBoneColor = PersonaOptions->AffectedBoneColor;
+		DrawConfig.SelectedBoneColor = PersonaOptions->SelectedBoneColor;
+		DrawConfig.ParentOfSelectedBoneColor = PersonaOptions->ParentOfSelectedBoneColor;
 		DrawConfig.AxisConfig.Thickness = Properties->AxisThickness;
 		DrawConfig.AxisConfig.Length = Properties->AxisLength;
 	
@@ -1748,7 +1746,7 @@ void USkeletonEditingTool::Render(IToolsContextRenderAPI* RenderAPI)
 		const FTransform& BoneTransform = Modifier->GetTransform(Index, true);
 		WorldTransforms[Index] = BoneTransform;
 		RequiredBones[Index] = Index;
-		BoneColors[Index] = DefaultBoneColor;
+		BoneColors[Index] = SkeletalDebugRendering::GetSemiRandomColorForBone(Index);
 		HitProxies.Add(new HBoneHitProxy(Index, RefSkeleton.GetBoneName(Index)));
 	}
 
diff --git a/Engine/Plugins/Experimental/Mutable/Source/CustomizableObjectEditor/Private/MuCOE/CustomizableObjectEditorViewportClient.cpp b/Engine/Plugins/Experimental/Mutable/Source/CustomizableObjectEditor/Private/MuCOE/CustomizableObjectEditorViewportClient.cpp
--- a/Engine/Plugins/Experimental/Mutable/Source/CustomizableObjectEditor/Private/MuCOE/CustomizableObjectEditorViewportClient.cpp
+++ b/Engine/Plugins/Experimental/Mutable/Source/CustomizableObjectEditor/Private/MuCOE/CustomizableObjectEditorViewportClient.cpp
@@ -1989,20 +1989,18 @@ void FCustomizableObjectEditorViewportClient::DrawMeshBones(UDebugSkelMeshCompon
 
 	TArray<FLinearColor> BoneColors;
 	BoneColors.AddUninitialized(MeshComponent->GetNumDrawTransform());
-
-	const FLinearColor BoneColor = GetDefault<UPersonaOptions>()->DefaultBoneColor;
-	const FLinearColor VirtualBoneColor = GetDefault<UPersonaOptions>()->VirtualBoneColor;
-	const TArray<FBoneIndexType>& DrawBoneIndices = MeshComponent->GetDrawBoneIndices();
 	
+	const TArray<FBoneIndexType>& DrawBoneIndices = MeshComponent->GetDrawBoneIndices();
 	for (int32 Index = 0; Index < DrawBoneIndices.Num(); ++Index)
 	{
 		const int32 BoneIndex = DrawBoneIndices[Index];
 		WorldTransforms[BoneIndex] = MeshComponent->GetDrawTransform(BoneIndex) * MeshComponent->GetComponentTransform();
-		BoneColors[BoneIndex] = BoneColor;
+		BoneColors[BoneIndex] = MeshComponent->GetBoneColor(BoneIndex);
 	}
 
 	// color virtual bones
-	for (int16 VirtualBoneIndex : MeshComponent->GetReferenceSkeleton().GetRequiredVirtualBones())
+	const FLinearColor VirtualBoneColor = GetDefault<UPersonaOptions>()->VirtualBoneColor;
+	for (const int16 VirtualBoneIndex : MeshComponent->GetReferenceSkeleton().GetRequiredVirtualBones())
 	{
 		BoneColors[VirtualBoneIndex] = VirtualBoneColor;
 	}
diff --git a/Engine/Source/Editor/Persona/Private/AnimViewportShowCommands.cpp b/Engine/Source/Editor/Persona/Private/AnimViewportShowCommands.cpp
--- a/Engine/Source/Editor/Persona/Private/AnimViewportShowCommands.cpp
+++ b/Engine/Source/Editor/Persona/Private/AnimViewportShowCommands.cpp
@@ -29,6 +29,7 @@ void FAnimViewportShowCommands::RegisterCommands()
 	UI_COMMAND( ShowMorphTargets, "Morph Targets", "Display applied morph targets of the mesh", EUserInterfaceActionType::ToggleButton, FInputChord() );
 
 	UI_COMMAND( ShowBoneNames, "Bone Names", "Display bone names in the viewport", EUserInterfaceActionType::ToggleButton, FInputChord() );
+	UI_COMMAND( ShowBoneColors, "Multi-Color Bones", "Display skeleton with multiple colors", EUserInterfaceActionType::ToggleButton, FInputChord() );
 
 	// below 3 menus are radio button styles
 	UI_COMMAND(ShowDisplayInfoBasic, "Basic", "Display basic mesh info in the viewport", EUserInterfaceActionType::RadioButton, FInputChord());
diff --git a/Engine/Source/Editor/Persona/Private/AnimViewportShowCommands.h b/Engine/Source/Editor/Persona/Private/AnimViewportShowCommands.h
--- a/Engine/Source/Editor/Persona/Private/AnimViewportShowCommands.h
+++ b/Engine/Source/Editor/Persona/Private/AnimViewportShowCommands.h
@@ -104,6 +104,9 @@ public:
 	/** Show skeletal mesh bone names */
 	TSharedPtr< FUICommandInfo > ShowBoneNames;
 
+	/** Show skeletal mesh bone names */
+	TSharedPtr< FUICommandInfo > ShowBoneColors;
+
 	/** Show skeletal mesh info */
 	TSharedPtr< FUICommandInfo > ShowDisplayInfoBasic;
 	TSharedPtr< FUICommandInfo > ShowDisplayInfoDetailed;
diff --git a/Engine/Source/Editor/Persona/Private/AnimationEditorViewportClient.cpp b/Engine/Source/Editor/Persona/Private/AnimationEditorViewportClient.cpp
--- a/Engine/Source/Editor/Persona/Private/AnimationEditorViewportClient.cpp
+++ b/Engine/Source/Editor/Persona/Private/AnimationEditorViewportClient.cpp
@@ -556,17 +556,14 @@ void FAnimationViewportClient::Draw(const FSceneView* View, FPrimitiveDrawInterf
 
 			const FReferenceSkeleton& RefSkeleton = PreviewMeshComponent->GetReferenceSkeleton();
 			const TArray<FBoneIndexType>& DrawBoneIndices = PreviewMeshComponent->GetDrawBoneIndices();
-
-			// if we have BonesOfInterest, draw sub set of the bones only
-			if (GetAnimPreviewScene()->GetSelectedBoneIndex() != INDEX_NONE)
-			{
-				DrawMeshSubsetBones(PreviewMeshComponent, PreviewMeshComponent->BonesOfInterest, PDI);
-			}
-			// otherwise, if we display bones, display
+			
+			// draw the skeleton normally
 			if ( GetBoneDrawMode() != EBoneDrawMode::None )
 			{
 				DrawMeshBones(PreviewMeshComponent, PDI);
 			}
+
+			// special draw modes for debugging various transforms...
 			if (PreviewMeshComponent->bDisplayRawAnimation )
 			{
 				DrawMeshBonesUncompressedAnimation(PreviewMeshComponent, PDI);
@@ -621,7 +618,7 @@ void FAnimationViewportClient::Draw(const FSceneView* View, FPrimitiveDrawInterf
 		{
 			if (const USkeleton* Skeleton = GetPreviewScene()->GetPersonaToolkit()->GetSkeleton())
 			{
-				DrawBonesFromSkeleton(Skeleton, PreviewMeshComponent->BonesOfInterest, PDI);
+				DrawBonesFromSkeleton(PreviewMeshComponent, Skeleton, PreviewMeshComponent->BonesOfInterest, PDI);
 			}
 		}
 	}
@@ -1577,6 +1574,7 @@ void FAnimationViewportClient::DrawBonesFromTransforms(
 
 	constexpr bool bForceDraw = false;
 	const bool bAddHitProxy = MeshComponent->SkeletonDrawMode != ESkeletonDrawMode::GreyedOut;
+	const bool bUseMuliColors = MeshComponent->bShowBoneColors;
 
 	DrawBones(
 		MeshComponent->GetComponentLocation(),
@@ -1587,7 +1585,8 @@ void FAnimationViewportClient::DrawBonesFromTransforms(
 		BoneColours,
 		PDI,
 		bForceDraw,
-		bAddHitProxy);
+		bAddHitProxy,
+		bUseMuliColors);
 }
 
 void FAnimationViewportClient::DrawBonesFromCompactPose(
@@ -1603,16 +1602,10 @@ void FAnimationViewportClient::DrawBonesFromCompactPose(
 	{
 		return;
 	}
-
-	// optionally override draw color
-	const FLinearColor BoneColor = MeshComponent->SkeletonDrawMode == ESkeletonDrawMode::GreyedOut ? GetDefault<UPersonaOptions>()->DisabledBoneColor : DrawColor;
 	
 	TArray<FTransform> WorldTransforms;
 	WorldTransforms.AddUninitialized(Pose.GetBoneContainer().GetNumBones());
 
-	TArray<FLinearColor> BoneColors;
-	BoneColors.AddUninitialized(Pose.GetBoneContainer().GetNumBones());
-
 	// we could cache parent bones as we calculate, but right now I'm not worried about perf issue of this
 	for (FCompactPoseBoneIndex BoneIndex : Pose.ForEachBoneIndex())
 	{
@@ -1628,11 +1621,11 @@ void FAnimationViewportClient::DrawBonesFromCompactPose(
 		{
 			WorldTransforms[MeshBoneIndex.GetInt()] = Pose[BoneIndex] * WorldTransforms[ParentIndex];
 		}
-		BoneColors[MeshBoneIndex.GetInt()] = BoneColor;
 	}
 
 	constexpr bool bForceDraw = true;
 	const bool bAddHitProxy = MeshComponent->SkeletonDrawMode != ESkeletonDrawMode::GreyedOut;
+	const bool bUseMultiColor = MeshComponent->bShowBoneColors;
 
 	DrawBones(
 		MeshComponent->GetComponentLocation(),
@@ -1640,10 +1633,11 @@ void FAnimationViewportClient::DrawBonesFromCompactPose(
 		MeshComponent->GetReferenceSkeleton(),
 		WorldTransforms,
 		MeshComponent->BonesOfInterest,
-		BoneColors,
+		TArray<FLinearColor>(),
 		PDI,
 		bForceDraw,
-		bAddHitProxy);
+		bAddHitProxy,
+		bUseMultiColor);
 }
 
 void FAnimationViewportClient::DrawMeshBonesUncompressedAnimation(UDebugSkelMeshComponent * MeshComponent, FPrimitiveDrawInterface* PDI) const
@@ -1705,7 +1699,7 @@ void FAnimationViewportClient::DrawMeshBonesBakedAnimation(UDebugSkelMeshCompone
 	}
 }
 
-void FAnimationViewportClient::DrawBonesFromSkeleton(const USkeleton* Skeleton, const TArray<int32>& InSelectedBones,FPrimitiveDrawInterface* PDI) const
+void FAnimationViewportClient::DrawBonesFromSkeleton(UDebugSkelMeshComponent * MeshComponent, const USkeleton* Skeleton, const TArray<int32>& InSelectedBones,FPrimitiveDrawInterface* PDI) const
 {
 	check(Skeleton);
 
@@ -1721,7 +1715,7 @@ void FAnimationViewportClient::DrawBonesFromSkeleton(const USkeleton* Skeleton,
 
 	const FReferenceSkeleton& RefSkeleton = Skeleton->GetReferenceSkeleton();
 
-	const UPersonaOptions* PersonaOptions = GetDefault<UPersonaOptions>();
+	
 	for (FBoneIndexType BoneIndex = 0; BoneIndex < SkeletonRefPose.Num(); ++BoneIndex)
 	{
 		const int32 ParentIndex = RefSkeleton.GetParentIndex(BoneIndex);
@@ -1738,17 +1732,20 @@ void FAnimationViewportClient::DrawBonesFromSkeleton(const USkeleton* Skeleton,
 			WorldTransforms[BoneIndex] = SkeletonRefPose[BoneIndex];
 		}
 
-		BoneColours[BoneIndex] = PersonaOptions->DefaultBoneColor;
+		BoneColours[BoneIndex] = MeshComponent->GetBoneColor(BoneIndex);
 	}
 
 	// color virtual bones
+	const FLinearColor VirtualBoneColor = GetDefault<UPersonaOptions>()->VirtualBoneColor;
 	for (const int16 VirtualBoneIndex : RefSkeleton.GetRequiredVirtualBones())
 	{
-		BoneColours[VirtualBoneIndex] = PersonaOptions->VirtualBoneColor;
+		BoneColours[VirtualBoneIndex] = VirtualBoneColor;
 	}
 
 	constexpr bool bForceDraw = false;
 	constexpr bool bAddHitProxy = true;
+	const bool bUseMultiColor = MeshComponent->bShowBoneColors;
+	
 	DrawBones(
 		FVector::ZeroVector,
 		RequiredBones,
@@ -1758,7 +1755,8 @@ void FAnimationViewportClient::DrawBonesFromSkeleton(const USkeleton* Skeleton,
 		BoneColours,
 		PDI,
 		bForceDraw,
-		bAddHitProxy);
+		bAddHitProxy,
+		bUseMultiColor);
 }
 
 void FAnimationViewportClient::DrawMeshBones(UDebugSkelMeshComponent* MeshComponent, FPrimitiveDrawInterface* PDI) const
@@ -1787,7 +1785,7 @@ void FAnimationViewportClient::DrawMeshBones(UDebugSkelMeshComponent* MeshCompon
 	{
 		const int32 BoneIndex = DrawBoneIndices[Index];
 		WorldTransforms[BoneIndex] = MeshComponent->GetDrawTransform(BoneIndex) * MeshComponent->GetComponentTransform();
-		BoneColours[BoneIndex] = BoneColor;
+		BoneColours[BoneIndex] = MeshComponent->GetBoneColor(BoneIndex);
 	}
 
 	// color virtual bones
@@ -1795,11 +1793,11 @@ void FAnimationViewportClient::DrawMeshBones(UDebugSkelMeshComponent* MeshCompon
 	{
 		BoneColours[VirtualBoneIndex] = VirtualBoneColor;
 	}
-
+	
 	constexpr bool bForceDraw = false;
-
 	// don't allow selection if the skeleton draw mode is greyed out
 	const bool bAddHitProxy = MeshComponent->SkeletonDrawMode != ESkeletonDrawMode::GreyedOut;
+	const bool bUseMultiColors = MeshComponent->bShowBoneColors;
 
 	DrawBones(
 		MeshComponent->GetComponentLocation(),
@@ -1810,7 +1808,8 @@ void FAnimationViewportClient::DrawMeshBones(UDebugSkelMeshComponent* MeshCompon
 		BoneColours,
 		PDI,
 		bForceDraw,
-		bAddHitProxy);
+		bAddHitProxy,
+		bUseMultiColors);
 }
 
 void FAnimationViewportClient::DrawBones(
@@ -1822,13 +1821,15 @@ void FAnimationViewportClient::DrawBones(
 	const TArray<FLinearColor>& BoneColors,
 	FPrimitiveDrawInterface* PDI,
 	bool bForceDraw,
-	bool bAddHitProxy) const
+	bool bAddHitProxy,
+	bool bUseMultiColors) const
 {
 	FSkelDebugDrawConfig DrawConfig;
 	DrawConfig.BoneDrawMode = GetBoneDrawMode();
 	DrawConfig.BoneDrawSize = GetBoneDrawSize();
 	DrawConfig.bAddHitProxy = bAddHitProxy;
 	DrawConfig.bForceDraw = bForceDraw;
+	DrawConfig.bUseMultiColorAsDefaultColor = bUseMultiColors;
 	DrawConfig.DefaultBoneColor = GetMutableDefault<UPersonaOptions>()->DefaultBoneColor;
 	DrawConfig.AffectedBoneColor = GetMutableDefault<UPersonaOptions>()->AffectedBoneColor;
 	DrawConfig.SelectedBoneColor = GetMutableDefault<UPersonaOptions>()->SelectedBoneColor;
@@ -1858,84 +1859,6 @@ void FAnimationViewportClient::DrawBones(
 	);
 }
 
-void FAnimationViewportClient::DrawMeshSubsetBones(const UDebugSkelMeshComponent* MeshComponent, const TArray<int32>& BonesOfInterest, FPrimitiveDrawInterface* PDI) const
-{
-	// this BonesOfInterest has to be in MeshComponent base, not Skeleton 
-	if (!MeshComponent ||
-		!MeshComponent->GetSkeletalMeshAsset() ||
-		BonesOfInterest.IsEmpty() ||
-		MeshComponent->SkeletonDrawMode == ESkeletonDrawMode::Hidden)
-	{
-		return;
-	}
-	
-	TArray<FTransform> WorldTransforms;
-	WorldTransforms.AddUninitialized(MeshComponent->GetNumDrawTransform());
-
-	TArray<FLinearColor> BoneColours;
-	BoneColours.AddUninitialized(MeshComponent->GetNumDrawTransform());
-
-	TArray<FBoneIndexType> RequiredBones;
-
-	const FReferenceSkeleton& RefSkeleton = MeshComponent->GetReferenceSkeleton();
-
-	// we could cache parent bones as we calculate, but right now I'm not worried about perf issue of this
-	const TArray<FBoneIndexType>& DrawBoneIndices = MeshComponent->GetDrawBoneIndices();
-	const UPersonaOptions* PersonaOptions = GetDefault<UPersonaOptions>();
-	for ( auto Iter = DrawBoneIndices.CreateConstIterator(); Iter; ++Iter)
-	{
-		const int32 BoneIndex = *Iter;
-		bool bDrawBone = false;
-
-		const int32 ParentIndex = RefSkeleton.GetParentIndex(BoneIndex);
-
-		// need to see if it's child of any of Bones of interest
-		for (auto SubIter=BonesOfInterest.CreateConstIterator(); SubIter; ++SubIter )
-		{
-			const int32 SubBoneIndex = *SubIter;
-			// if I'm child of the BonesOfInterest
-			if(BoneIndex == SubBoneIndex)
-			{
-				//found a bone we are interested in
-				if(ParentIndex >= 0)
-				{
-					WorldTransforms[ParentIndex] = MeshComponent->GetDrawTransform(ParentIndex)*MeshComponent->GetComponentTransform();
-				}
-				BoneColours[BoneIndex] = PersonaOptions->SelectedBoneColor;
-				bDrawBone = true;
-				break;
-			}
-			else if ( RefSkeleton.BoneIsChildOf(BoneIndex, SubBoneIndex) )
-			{
-				BoneColours[BoneIndex] = PersonaOptions->DefaultBoneColor;
-				bDrawBone = true;
-				break;
-			}
-		}
-
-		if (bDrawBone)
-		{
-			//add to the list
-			RequiredBones.AddUnique(static_cast<FBoneIndexType>(BoneIndex));
-			WorldTransforms[BoneIndex] = MeshComponent->GetDrawTransform(BoneIndex) * MeshComponent->GetComponentTransform();
-		}
-	}
-
-	constexpr bool bForceDraw = false;
-	const bool bAddHitProxy = MeshComponent->SkeletonDrawMode != ESkeletonDrawMode::GreyedOut;
-
-	DrawBones(
-		MeshComponent->GetComponentLocation(),
-		RequiredBones,
-		MeshComponent->GetReferenceSkeleton(),
-		WorldTransforms,
-		MeshComponent->BonesOfInterest,
-		BoneColours,
-		PDI,
-		bForceDraw,
-		bAddHitProxy);
-}
-
 void FAnimationViewportClient::DrawAttributes(UDebugSkelMeshComponent* MeshComponent, FPrimitiveDrawInterface* PDI) const
 {
 	if (MeshComponent && MeshComponent->GetSkeletalMeshAsset())
diff --git a/Engine/Source/Editor/Persona/Private/SAnimViewportToolBar.cpp b/Engine/Source/Editor/Persona/Private/SAnimViewportToolBar.cpp
--- a/Engine/Source/Editor/Persona/Private/SAnimViewportToolBar.cpp
+++ b/Engine/Source/Editor/Persona/Private/SAnimViewportToolBar.cpp
@@ -812,6 +812,7 @@ TSharedRef<SWidget> SAnimViewportToolBar::GenerateCharacterMenu() const
 						Section.AddMenuEntry(FAnimViewportShowCommands::Get().ShowSockets);
 						Section.AddMenuEntry(FAnimViewportShowCommands::Get().ShowAttributes);
 						Section.AddMenuEntry(FAnimViewportShowCommands::Get().ShowBoneNames);
+						Section.AddMenuEntry(FAnimViewportShowCommands::Get().ShowBoneColors);
 					}
 
 					{
diff --git a/Engine/Source/Editor/Persona/Private/SAnimationEditorViewport.cpp b/Engine/Source/Editor/Persona/Private/SAnimationEditorViewport.cpp
--- a/Engine/Source/Editor/Persona/Private/SAnimationEditorViewport.cpp
+++ b/Engine/Source/Editor/Persona/Private/SAnimationEditorViewport.cpp
@@ -796,6 +796,12 @@ void SAnimationEditorViewportTabBody::BindCommands()
 		FCanExecuteAction(),
 		FIsActionChecked::CreateSP(this, &SAnimationEditorViewportTabBody::IsShowingBoneNames));
 
+	CommandList.MapAction( 
+		ViewportShowMenuCommands.ShowBoneColors,
+		FExecuteAction::CreateSP(this, &SAnimationEditorViewportTabBody::OnShowBoneColors),
+		FCanExecuteAction(),
+		FIsActionChecked::CreateSP(this, &SAnimationEditorViewportTabBody::IsShowingBoneColors));
+
 	CommandList.MapAction( 
 		ViewportShowMenuCommands.ShowRawAnimation,
 		FExecuteAction::CreateSP(this, &SAnimationEditorViewportTabBody::OnShowRawAnimation),
@@ -1301,6 +1307,22 @@ bool SAnimationEditorViewportTabBody::IsShowingBoneNames() const
 	return PreviewComponent != NULL && PreviewComponent->bShowBoneNames;
 }
 
+void SAnimationEditorViewportTabBody::OnShowBoneColors()
+{
+	ForEachDebugMesh([](UDebugSkelMeshComponent* InMesh)
+	{
+		InMesh->bShowBoneColors = !InMesh->bShowBoneColors;
+		InMesh->MarkRenderStateDirty();
+	});
+	RefreshViewport();
+}
+
+bool SAnimationEditorViewportTabBody::IsShowingBoneColors() const
+{
+	const UDebugSkelMeshComponent* PreviewComponent = GetPreviewScene()->GetPreviewMeshComponent();
+	return PreviewComponent && PreviewComponent->bShowBoneColors;
+}
+
 bool SAnimationEditorViewportTabBody::IsShowingRawAnimation() const
 {
 	UDebugSkelMeshComponent* PreviewComponent = GetPreviewScene()->GetPreviewMeshComponent();
diff --git a/Engine/Source/Editor/Persona/Private/SAnimationEditorViewport.h b/Engine/Source/Editor/Persona/Private/SAnimationEditorViewport.h
--- a/Engine/Source/Editor/Persona/Private/SAnimationEditorViewport.h
+++ b/Engine/Source/Editor/Persona/Private/SAnimationEditorViewport.h
@@ -368,6 +368,12 @@ private:
 
 	/** Function to check whether bone names are displayed or not */
 	bool IsShowingBoneNames() const;
+
+	/** Function to show/hide bone colors */
+	void OnShowBoneColors();
+
+	/** Function to check whether bone's are using multiple colors */
+	bool IsShowingBoneColors() const;
 	
 	/** Function to show/hide selected bone weight */
 	void OnShowOverlayNone();
diff --git a/Engine/Source/Editor/Persona/Public/AnimationEditorViewportClient.h b/Engine/Source/Editor/Persona/Public/AnimationEditorViewportClient.h
--- a/Engine/Source/Editor/Persona/Public/AnimationEditorViewportClient.h
+++ b/Engine/Source/Editor/Persona/Public/AnimationEditorViewportClient.h
@@ -425,7 +425,7 @@ private:
 	/** Draw Bones for non retargeted animation. */
 	void DrawMeshBonesBakedAnimation(UDebugSkelMeshComponent * MeshComponent, FPrimitiveDrawInterface* PDI) const;
 	/** Draw Bones from skeleton reference pose. */
-	void DrawBonesFromSkeleton(const USkeleton* Skeleton, const TArray<int32>& InSelectedBones, FPrimitiveDrawInterface* PDI) const;
+	void DrawBonesFromSkeleton(UDebugSkelMeshComponent * MeshComponent, const USkeleton* Skeleton, const TArray<int32>& InSelectedBones, FPrimitiveDrawInterface* PDI) const;
 	/** Draws Bones for RequiredBones with WorldTransform **/
 	void DrawBones(
 		const FVector& ComponentOrigin,
@@ -436,9 +436,8 @@ private:
 		const TArray<FLinearColor>& BoneColors,
 		FPrimitiveDrawInterface* PDI,
 		bool bForceDraw,
-		bool bAddHitProxy) const;
-	/** Draw Sub set of Bones **/
-	void DrawMeshSubsetBones(const UDebugSkelMeshComponent* MeshComponent, const TArray<int32>& BonesOfInterest, FPrimitiveDrawInterface* PDI) const;
+		bool bAddHitProxy,
+		bool bUseMultiColors) const;
 
 	/** Draws active transform attributes */
 	void DrawAttributes(UDebugSkelMeshComponent* MeshComponent, FPrimitiveDrawInterface* PDI) const;
diff --git a/Engine/Source/Editor/UnrealEd/Classes/Animation/DebugSkelMeshComponent.h b/Engine/Source/Editor/UnrealEd/Classes/Animation/DebugSkelMeshComponent.h
--- a/Engine/Source/Editor/UnrealEd/Classes/Animation/DebugSkelMeshComponent.h
+++ b/Engine/Source/Editor/UnrealEd/Classes/Animation/DebugSkelMeshComponent.h
@@ -293,10 +293,6 @@ class UDebugSkelMeshComponent : public USkeletalMeshComponent
 	UPROPERTY(transient)
 	uint32 bRequiredBonesUpToDateDuringTick : 1;
 
-	/** Multiplier for the bone radius rendering */
-	UPROPERTY(transient)
-	float BoneRadiusMultiplier;
-
 	/* Bounds computed from cloth. */
 	FBoxSphereBounds CachedClothBounds;
 
@@ -357,6 +353,13 @@ class UDebugSkelMeshComponent : public USkeletalMeshComponent
 	UPROPERTY(transient)
 	bool bTrackAttachedInstanceLOD;
 
+	/** If true then the bones are displayed with multiple semi-random colors */
+	UPROPERTY()
+	bool bShowBoneColors = true;
+	
+	/** Palettized semi randomized colors for multi-color bone display */
+	UNREALED_API FLinearColor GetBoneColor(int32 BoneIndex) const;
+
 	// Helper method that sets the forced lod
 	UNREALED_API void SetDebugForcedLOD(int32 InNewForcedLOD);
 
@@ -686,7 +689,6 @@ public:
 
 		return FTransform::Identity;
 	}
-
 };
 
 
diff --git a/Engine/Source/Editor/UnrealEd/Private/Animation/DebugSkelMeshComponent.cpp b/Engine/Source/Editor/UnrealEd/Private/Animation/DebugSkelMeshComponent.cpp
--- a/Engine/Source/Editor/UnrealEd/Private/Animation/DebugSkelMeshComponent.cpp
+++ b/Engine/Source/Editor/UnrealEd/Private/Animation/DebugSkelMeshComponent.cpp
@@ -25,6 +25,7 @@
 #include "Utils/ClothingMeshUtils.h"
 #include "DynamicMeshBuilder.h"
 #include "Materials/MaterialInstanceDynamic.h"
+#include "Preferences/PersonaOptions.h"
 
 //////////////////////////////////////////////////////////////////////////
 // UDebugSkelMeshComponent
@@ -79,6 +80,30 @@ UDebugSkelMeshComponent::UDebugSkelMeshComponent(const FObjectInitializer& Objec
 	ConsumeRootMotionPreviousPlaybackTime = 0.f;
 }
 
+FLinearColor UDebugSkelMeshComponent::GetBoneColor(int32 InBoneIndex) const
+{
+	// this returns the normal unmodified color of the bone, calling code must account
+	// for any editor specific states that might affect the final bone color (like selection)
+	
+	// skeleton greyed out
+	if (SkeletonDrawMode == ESkeletonDrawMode::GreyedOut)
+	{
+		return GetDefault<UPersonaOptions>()->DisabledBoneColor;
+	}
+
+	// using default color for all bones
+	if (!bShowBoneColors)
+	{
+		return GetDefault<UPersonaOptions>()->DefaultBoneColor;
+	}
+	
+	// uses deterministic, semi-random desaturated color unique to the bone index
+	constexpr float Saturation = 0.5f;
+	constexpr float Value = 1.0f;
+	constexpr float Rotation = 90.f;
+	return FLinearColor::IntToDistinctColor(InBoneIndex, Saturation, Value, Rotation);
+}
+
 void UDebugSkelMeshComponent::SetDebugForcedLOD(int32 InNewForcedLOD)
 {
 	SetForcedLOD(InNewForcedLOD);
diff --git a/Engine/Source/Runtime/Core/Public/Math/Color.h b/Engine/Source/Runtime/Core/Public/Math/Color.h
--- a/Engine/Source/Runtime/Core/Public/Math/Color.h
+++ b/Engine/Source/Runtime/Core/Public/Math/Color.h
@@ -443,6 +443,24 @@ struct FLinearColor
 		return (ClampedLo < 1.0f) ? ClampedLo : 1.0f;
 	}
 
+	/**
+	 * @brief Helper function to generate distinct colors from a sequence of integers where each integer increment
+	 * spins around the Hue wheel by increments of the golden ratio.
+	 * @param Seed			the "seed" for a semi-random, but deterministic color
+	 * @param Saturation	0-1 saturation of resulting color
+	 * @param Value			0-1 brightness of resulting color
+	 * @param HueRotation	0-360 amount to rotate around the hue wheel
+	 * @return color with semi-random hue */
+	static FLinearColor IntToDistinctColor(
+		const int32 Seed,
+		const float Saturation=1.f,
+		const float Value=1.f,
+		const float HueRotation=180.f)
+	{
+		constexpr float GoldenRatio = 1.618033988749895f; // (1 + FMath::Sqrt(5.f)) / 2.f;
+		return FLinearColor(static_cast<float>(Seed) * GoldenRatio * HueRotation, Saturation, Value).HSVToLinearRGB();
+	}
+
 	// Common colors.	
 	static CORE_API const FLinearColor White;
 	static CORE_API const FLinearColor Gray;
diff --git a/Engine/Source/Runtime/Engine/Private/SkeletalDebugRendering.cpp b/Engine/Source/Runtime/Engine/Private/SkeletalDebugRendering.cpp
--- a/Engine/Source/Runtime/Engine/Private/SkeletalDebugRendering.cpp
+++ b/Engine/Source/Runtime/Engine/Private/SkeletalDebugRendering.cpp
@@ -392,7 +392,15 @@ void DrawBonesInternal(
 		// determine color of bone based on selection / affected state
 		const bool bIsSelected = InSelectedBones.Contains(BoneIndex);
 		const bool bIsAffected = AffectedBones[BoneIndex];
-		FLinearColor DefaultBoneColor = BoneColors.IsEmpty() ? DrawConfig.DefaultBoneColor : BoneColors[BoneIndex];
+		FLinearColor DefaultBoneColor;
+		if (BoneColors.IsEmpty())
+		{
+			DefaultBoneColor = DrawConfig.bUseMultiColorAsDefaultColor ? GetSemiRandomColorForBone(BoneIndex) : DrawConfig.DefaultBoneColor;
+		}
+		else
+		{
+			DefaultBoneColor = BoneColors[BoneIndex];
+		}
 		FLinearColor BoneColor = bIsAffected ? DrawConfig.AffectedBoneColor : DefaultBoneColor;
 		BoneColor = bIsSelected ? DrawConfig.SelectedBoneColor : BoneColor;
 
@@ -466,4 +474,22 @@ void DrawBonesInternal(
 	}
 }
 
+FLinearColor GetSemiRandomColorForBone(const int32 BoneIndex)
+{
+	// uses deterministic, semi-random desaturated color unique to the bone index
+	constexpr float Saturation = 0.5f;
+	constexpr float Value = 1.0f;
+	constexpr float Rotation = 90.f;
+	return FLinearColor::IntToDistinctColor(BoneIndex, Saturation, Value, Rotation);
+}
+
+void FillWithMultiColors(TArray<FLinearColor>& BoneColors, const int32 NumBones)
+{
+	BoneColors.SetNumUninitialized(NumBones);
+	for (int32 BoneIndex=0; BoneIndex<NumBones; ++BoneIndex)
+	{
+		BoneColors[BoneIndex] = GetSemiRandomColorForBone(BoneIndex);
+	}
+}
+
 }
diff --git a/Engine/Source/Runtime/Engine/Public/SkeletalDebugRendering.h b/Engine/Source/Runtime/Engine/Public/SkeletalDebugRendering.h
--- a/Engine/Source/Runtime/Engine/Public/SkeletalDebugRendering.h
+++ b/Engine/Source/Runtime/Engine/Public/SkeletalDebugRendering.h
@@ -40,6 +40,7 @@ public:
 	float BoneDrawSize;
 	bool bForceDraw;
 	bool bAddHitProxy;
+	bool bUseMultiColorAsDefaultColor;
 	FLinearColor DefaultBoneColor;
 	FLinearColor AffectedBoneColor;
 	FLinearColor SelectedBoneColor;
@@ -151,4 +152,7 @@ void DrawBonesInternal(
 	const TArray<FLinearColor>& BoneColors,
 	const TArray<TRefCountPtr<HHitProxy>>& HitProxies,
 	const FSkelDebugDrawConfig& DrawConfig);
-}
\ No newline at end of file
+
+ENGINE_API FLinearColor GetSemiRandomColorForBone(const int32 BoneIndex);
+ENGINE_API void FillWithMultiColors(TArray<FLinearColor>& BoneColors, const int32 NumBones);
+}
