Control Rig: Fix cannot swap preview mesh on modular rigs
#rnx

diff --git a/Engine/Source/Editor/Persona/Private/AnimationEditorPreviewScene.cpp b/Engine/Source/Editor/Persona/Private/AnimationEditorPreviewScene.cpp
--- a/Engine/Source/Editor/Persona/Private/AnimationEditorPreviewScene.cpp
+++ b/Engine/Source/Editor/Persona/Private/AnimationEditorPreviewScene.cpp
@@ -175,7 +175,8 @@ void FAnimationEditorPreviewScene::SetPreviewMesh(USkeletalMesh* NewPreviewMesh,
 	if (NewPreviewMesh != nullptr && GetEditableSkeleton().IsValid() && !GetEditableSkeleton()->GetSkeleton().IsCompatibleMesh(NewPreviewMesh))
 	{
 		const USkeleton& Skeleton = GetEditableSkeleton()->GetSkeleton();
-		if (NewPreviewMesh->GetSkeleton() && Skeleton.IsCompatibleForEditor(NewPreviewMesh->GetSkeleton()))
+		const bool bSkipSkeletonCompatibility = PersonaToolkit.IsValid() && PersonaToolkit.Pin()->CanPreviewMeshUseDifferentSkeleton(); 
+		if (NewPreviewMesh->GetSkeleton() && (bSkipSkeletonCompatibility || Skeleton.IsCompatibleForEditor(NewPreviewMesh->GetSkeleton())))
 		{
 			SetPreviewMeshInternal(NewPreviewMesh);
 		}	
