Turning bone colors.
Bone colors no longer default.
Saturation bumped up

diff --git a/Engine/Plugins/Animation/IKRig/Source/IKRigEditor/Private/RetargetEditor/IKRetargetEditorController.cpp b/Engine/Plugins/Animation/IKRig/Source/IKRigEditor/Private/RetargetEditor/IKRetargetEditorController.cpp
--- a/Engine/Plugins/Animation/IKRig/Source/IKRigEditor/Private/RetargetEditor/IKRetargetEditorController.cpp
+++ b/Engine/Plugins/Animation/IKRig/Source/IKRigEditor/Private/RetargetEditor/IKRetargetEditorController.cpp
@@ -1455,7 +1455,7 @@ void FIKRetargetEditorController::RenderSkeleton(FPrimitiveDrawInterface* PDI, E
 	TArray<FLinearColor> BoneColors;
 	{
 		// set default colors
-		if (MeshComponent->bShowBoneColors)
+		if (GetDefault<UPersonaOptions>()->bShowBoneColors)
 		{
 			SkeletalDebugRendering::FillWithMultiColors(BoneColors, RefSkeleton.GetNum());
 		}
diff --git a/Engine/Plugins/Animation/IKRig/Source/IKRigEditor/Private/RigEditor/IKRigEditMode.cpp b/Engine/Plugins/Animation/IKRig/Source/IKRigEditor/Private/RigEditor/IKRigEditMode.cpp
--- a/Engine/Plugins/Animation/IKRig/Source/IKRigEditor/Private/RigEditor/IKRigEditMode.cpp
+++ b/Engine/Plugins/Animation/IKRig/Source/IKRigEditor/Private/RigEditor/IKRigEditMode.cpp
@@ -190,7 +190,8 @@ void FIKRigEditMode::RenderBones(FPrimitiveDrawInterface* PDI)
 
 	// get bone colors
 	TArray<FLinearColor> BoneColors;
-	GetBoneColors(Controller.Get(), CurrentProcessor, RefSkeleton, MeshComponent->bShowBoneColors, BoneColors);
+	const bool bUseBoneColors = GetDefault<UPersonaOptions>()->bShowBoneColors;
+	GetBoneColors(Controller.Get(), CurrentProcessor, RefSkeleton, bUseBoneColors, BoneColors);
 
 	SkeletalDebugRendering::DrawBones(
 		PDI,
diff --git a/Engine/Plugins/Animation/PoseSearch/Source/Editor/Private/PoseSearchDatabaseEdMode.cpp b/Engine/Plugins/Animation/PoseSearch/Source/Editor/Private/PoseSearchDatabaseEdMode.cpp
--- a/Engine/Plugins/Animation/PoseSearch/Source/Editor/Private/PoseSearchDatabaseEdMode.cpp
+++ b/Engine/Plugins/Animation/PoseSearch/Source/Editor/Private/PoseSearchDatabaseEdMode.cpp
@@ -82,7 +82,7 @@ namespace UE::PoseSearch
 						DrawConfig.BoneDrawSize = 0.2f;
 						DrawConfig.bAddHitProxy = false;
 						DrawConfig.bForceDraw = false;
-						DrawConfig.bUseMultiColorAsDefaultColor = MeshComponent->bShowBoneColors;
+						DrawConfig.bUseMultiColorAsDefaultColor = GetDefault<UPersonaOptions>()->bShowBoneColors;
 						DrawConfig.DefaultBoneColor = PersonaOptions->DefaultBoneColor;
 						DrawConfig.AffectedBoneColor = PersonaOptions->AffectedBoneColor;
 						DrawConfig.SelectedBoneColor = PersonaOptions->SelectedBoneColor;
diff --git a/Engine/Source/Editor/Persona/Private/AnimationEditorViewportClient.cpp b/Engine/Source/Editor/Persona/Private/AnimationEditorViewportClient.cpp
--- a/Engine/Source/Editor/Persona/Private/AnimationEditorViewportClient.cpp
+++ b/Engine/Source/Editor/Persona/Private/AnimationEditorViewportClient.cpp
@@ -1573,7 +1573,7 @@ void FAnimationViewportClient::DrawBonesFromTransforms(
 
 	constexpr bool bForceDraw = false;
 	const bool bAddHitProxy = MeshComponent->SkeletonDrawMode != ESkeletonDrawMode::GreyedOut;
-	const bool bUseMuliColors = MeshComponent->bShowBoneColors;
+	const bool bUseMuliColors = GetDefault<UPersonaOptions>()->bShowBoneColors;
 
 	DrawBones(
 		MeshComponent->GetComponentLocation(),
@@ -1624,7 +1624,7 @@ void FAnimationViewportClient::DrawBonesFromCompactPose(
 
 	constexpr bool bForceDraw = true;
 	const bool bAddHitProxy = MeshComponent->SkeletonDrawMode != ESkeletonDrawMode::GreyedOut;
-	const bool bUseMultiColor = MeshComponent->bShowBoneColors;
+	const bool bUseMultiColor = GetDefault<UPersonaOptions>()->bShowBoneColors;
 
 	DrawBones(
 		MeshComponent->GetComponentLocation(),
@@ -1743,7 +1743,7 @@ void FAnimationViewportClient::DrawBonesFromSkeleton(UDebugSkelMeshComponent * M
 
 	constexpr bool bForceDraw = false;
 	constexpr bool bAddHitProxy = true;
-	const bool bUseMultiColor = MeshComponent->bShowBoneColors;
+	const bool bUseMultiColor = GetDefault<UPersonaOptions>()->bShowBoneColors;
 	
 	DrawBones(
 		FVector::ZeroVector,
@@ -1796,7 +1796,7 @@ void FAnimationViewportClient::DrawMeshBones(UDebugSkelMeshComponent* MeshCompon
 	constexpr bool bForceDraw = false;
 	// don't allow selection if the skeleton draw mode is greyed out
 	const bool bAddHitProxy = MeshComponent->SkeletonDrawMode != ESkeletonDrawMode::GreyedOut;
-	const bool bUseMultiColors = MeshComponent->bShowBoneColors;
+	const bool bUseMultiColors = GetDefault<UPersonaOptions>()->bShowBoneColors;
 
 	DrawBones(
 		MeshComponent->GetComponentLocation(),
diff --git a/Engine/Source/Editor/Persona/Private/SAnimationEditorViewport.cpp b/Engine/Source/Editor/Persona/Private/SAnimationEditorViewport.cpp
--- a/Engine/Source/Editor/Persona/Private/SAnimationEditorViewport.cpp
+++ b/Engine/Source/Editor/Persona/Private/SAnimationEditorViewport.cpp
@@ -1310,7 +1310,8 @@ void SAnimationEditorViewportTabBody::OnShowBoneColors()
 {
 	ForEachDebugMesh([](UDebugSkelMeshComponent* InMesh)
 	{
-		InMesh->bShowBoneColors = !InMesh->bShowBoneColors;
+		UPersonaOptions* Settings = GetMutableDefault<UPersonaOptions>();
+		Settings->bShowBoneColors = !Settings->bShowBoneColors;
 		InMesh->MarkRenderStateDirty();
 	});
 	RefreshViewport();
@@ -1319,7 +1320,7 @@ void SAnimationEditorViewportTabBody::OnShowBoneColors()
 bool SAnimationEditorViewportTabBody::IsShowingBoneColors() const
 {
 	const UDebugSkelMeshComponent* PreviewComponent = GetPreviewScene()->GetPreviewMeshComponent();
-	return PreviewComponent && PreviewComponent->bShowBoneColors;
+	return PreviewComponent && GetDefault<UPersonaOptions>()->bShowBoneColors;
 }
 
 bool SAnimationEditorViewportTabBody::IsShowingRawAnimation() const
diff --git a/Engine/Source/Editor/SkeletalMeshEditor/Private/SkeletalMeshEditor.cpp b/Engine/Source/Editor/SkeletalMeshEditor/Private/SkeletalMeshEditor.cpp
--- a/Engine/Source/Editor/SkeletalMeshEditor/Private/SkeletalMeshEditor.cpp
+++ b/Engine/Source/Editor/SkeletalMeshEditor/Private/SkeletalMeshEditor.cpp
@@ -2,6 +2,7 @@
 
 #include "SkeletalMeshEditor.h"
 
+#include "AnimationEditorViewportClient.h"
 #include "SkeletalMeshEditorCommands.h"
 #include "SkeletalMeshEditorMode.h"
 
@@ -1263,9 +1264,32 @@ void FSkeletalMeshEditor::HandleMeshDetailsCreated(const TSharedRef<IDetailsView
 void FSkeletalMeshEditor::HandleViewportCreated(const TSharedRef<IPersonaViewport>& InViewport)
 {
 	Viewport = InViewport;
+	FEditorViewportClient& ViewportClient = InViewport->GetViewportClient();
+
+	// register callbacks to allow the asset to store the Bone Size viewport setting
+	if (FAnimationViewportClient* AnimViewportClient = static_cast<FAnimationViewportClient*>(&ViewportClient))
+	{
+		AnimViewportClient->OnSetBoneSize.BindLambda([InSkeletalMesh=SkeletalMesh](float InBoneSize)
+		{
+			if (InSkeletalMesh)
+			{
+				InSkeletalMesh->Modify();
+				InSkeletalMesh->BoneDrawSize = InBoneSize;
+			}
+		});
+		
+		AnimViewportClient->OnGetBoneSize.BindLambda([InSkeletalMesh=SkeletalMesh]() -> float
+		{
+			if (InSkeletalMesh)
+			{
+				return InSkeletalMesh->BoneDrawSize;
+			}
+
+			return 1.0f;
+		});
+	}
 	
 	// we need the viewport client to start out focused, or else it won't get ticked until we click inside it.
-	FEditorViewportClient& ViewportClient = InViewport->GetViewportClient(); 
 	ViewportClient.ReceivedFocus(ViewportClient.Viewport);
 }
 
diff --git a/Engine/Source/Editor/UnrealEd/Classes/Animation/DebugSkelMeshComponent.h b/Engine/Source/Editor/UnrealEd/Classes/Animation/DebugSkelMeshComponent.h
--- a/Engine/Source/Editor/UnrealEd/Classes/Animation/DebugSkelMeshComponent.h
+++ b/Engine/Source/Editor/UnrealEd/Classes/Animation/DebugSkelMeshComponent.h
@@ -352,10 +352,6 @@ class UDebugSkelMeshComponent : public USkeletalMeshComponent
 	/** Should the LOD of the debug mesh component track the LOD of the instance being debugged */
 	UPROPERTY(transient)
 	bool bTrackAttachedInstanceLOD;
-
-	/** If true then the bones are displayed with multiple semi-random colors */
-	UPROPERTY()
-	bool bShowBoneColors = true;
 	
 	/** Palettized semi randomized colors for multi-color bone display */
 	UNREALED_API FLinearColor GetBoneColor(int32 BoneIndex) const;
diff --git a/Engine/Source/Editor/UnrealEd/Classes/Preferences/PersonaOptions.h b/Engine/Source/Editor/UnrealEd/Classes/Preferences/PersonaOptions.h
--- a/Engine/Source/Editor/UnrealEd/Classes/Preferences/PersonaOptions.h
+++ b/Engine/Source/Editor/UnrealEd/Classes/Preferences/PersonaOptions.h
@@ -143,6 +143,10 @@ class UPersonaOptions : public UObject
 	UPROPERTY(EditAnywhere, config, Category = "Viewport")
 	uint32 DefaultBoneDrawSelection;
 
+	/** If true then the bones are displayed with multiple semi-random colors */
+	UPROPERTY(EditAnywhere, config, Category = "Viewport")
+	bool bShowBoneColors = false;
+
 	UPROPERTY(EditAnywhere, config, Category = "Viewport")
 	FLinearColor DefaultBoneColor;
 
diff --git a/Engine/Source/Editor/UnrealEd/Private/Animation/DebugSkelMeshComponent.cpp b/Engine/Source/Editor/UnrealEd/Private/Animation/DebugSkelMeshComponent.cpp
--- a/Engine/Source/Editor/UnrealEd/Private/Animation/DebugSkelMeshComponent.cpp
+++ b/Engine/Source/Editor/UnrealEd/Private/Animation/DebugSkelMeshComponent.cpp
@@ -24,6 +24,7 @@
 #include "ClothingSimulation.h"
 #include "Utils/ClothingMeshUtils.h"
 #include "DynamicMeshBuilder.h"
+#include "SkeletalDebugRendering.h"
 #include "Materials/MaterialInstanceDynamic.h"
 #include "Preferences/PersonaOptions.h"
 
@@ -76,16 +77,13 @@ FLinearColor UDebugSkelMeshComponent::GetBoneColor(int32 InBoneIndex) const
 	}
 
 	// using default color for all bones
-	if (!bShowBoneColors)
+	if (!GetDefault<UPersonaOptions>()->bShowBoneColors)
 	{
 		return GetDefault<UPersonaOptions>()->DefaultBoneColor;
 	}
 	
 	// uses deterministic, semi-random desaturated color unique to the bone index
-	constexpr float Saturation = 0.5f;
-	constexpr float Value = 1.0f;
-	constexpr float Rotation = 90.f;
-	return FLinearColor::IntToDistinctColor(InBoneIndex, Saturation, Value, Rotation);
+	return SkeletalDebugRendering::GetSemiRandomColorForBone(InBoneIndex);
 }
 
 void UDebugSkelMeshComponent::SetDebugForcedLOD(int32 InNewForcedLOD)
diff --git a/Engine/Source/Runtime/Engine/Classes/Engine/SkeletalMesh.h b/Engine/Source/Runtime/Engine/Classes/Engine/SkeletalMesh.h
--- a/Engine/Source/Runtime/Engine/Classes/Engine/SkeletalMesh.h
+++ b/Engine/Source/Runtime/Engine/Classes/Engine/SkeletalMesh.h
@@ -2122,6 +2122,10 @@ public:
 	/** Legacy clothing asset data, will be converted to new assets after loading */
 	UPROPERTY()
 	TArray<FClothingAssetData_Legacy>		ClothingAssets_DEPRECATED;
+
+	// The visual size of the bones in the viewport (saved between sessions). This is set from the viewport Character>Bones menu
+	UPROPERTY()
+	float BoneDrawSize = 1.0f;
 #endif
 
 	/** Animation Blueprint class to run as a post process for this mesh.
diff --git a/Engine/Source/Runtime/Engine/Public/SkeletalDebugRendering.h b/Engine/Source/Runtime/Engine/Public/SkeletalDebugRendering.h
--- a/Engine/Source/Runtime/Engine/Public/SkeletalDebugRendering.h
+++ b/Engine/Source/Runtime/Engine/Public/SkeletalDebugRendering.h
@@ -153,6 +153,6 @@ void DrawBonesInternal(
 	const TArray<TRefCountPtr<HHitProxy>>& HitProxies,
 	const FSkelDebugDrawConfig& DrawConfig);
 
-ENGINE_API FLinearColor GetSemiRandomColorForBone(const int32 BoneIndex, float Value=1.0f, float Saturation=0.5f);
+ENGINE_API FLinearColor GetSemiRandomColorForBone(const int32 BoneIndex, float Value=1.0f, float Saturation=0.75f);
 ENGINE_API void FillWithMultiColors(TArray<FLinearColor>& BoneColors, const int32 NumBones);
 }
