Skeletal Mesh Editor: disable post process notification when the skeletal mesh is hidden
    - most of the skeletal mesh modeling tools will hide the target on setup so no need to display the notification when that's the case

#rnx

diff --git a/Engine/Source/Editor/Persona/Private/SAnimationEditorViewport.cpp b/Engine/Source/Editor/Persona/Private/SAnimationEditorViewport.cpp
--- a/Engine/Source/Editor/Persona/Private/SAnimationEditorViewport.cpp
+++ b/Engine/Source/Editor/Persona/Private/SAnimationEditorViewport.cpp
@@ -1339,12 +1339,12 @@ void SAnimationEditorViewportTabBody::OnToggleDisablePostProcess()
 	AddPostProcessNotification();
 }
 
-bool SAnimationEditorViewportTabBody::CanDisablePostProcess()
+bool SAnimationEditorViewportTabBody::CanDisablePostProcess() const
 {
-	TArray<UDebugSkelMeshComponent*> PreviewMeshComponents = GetPreviewScene()->GetAllPreviewMeshComponents();
-	for (UDebugSkelMeshComponent* PreviewMeshComponent : PreviewMeshComponents)
+	const TArray<UDebugSkelMeshComponent*> PreviewMeshComponents = GetPreviewScene()->GetAllPreviewMeshComponents();
+	for (const UDebugSkelMeshComponent* PreviewMeshComponent : PreviewMeshComponents)
 	{
-		if (PreviewMeshComponent->PostProcessAnimInstance)
+		if (PreviewMeshComponent->PostProcessAnimInstance && PreviewMeshComponent->IsVisible())
 		{
 			return true;
 		}
diff --git a/Engine/Source/Editor/Persona/Private/SAnimationEditorViewport.h b/Engine/Source/Editor/Persona/Private/SAnimationEditorViewport.h
--- a/Engine/Source/Editor/Persona/Private/SAnimationEditorViewport.h
+++ b/Engine/Source/Editor/Persona/Private/SAnimationEditorViewport.h
@@ -337,7 +337,7 @@ private:
 
 	/** Handlers for the disable post process flag */
 	void OnToggleDisablePostProcess();
-	bool CanDisablePostProcess();
+	bool CanDisablePostProcess() const;
 	bool IsDisablePostProcessChecked();
 
 	/** Show non retargeted animation. */
