[UE] Fix runtime warning on uninitialized data member
- BoolProperty FPhysicsAssetRenderSettings::bShowCOM_DEPRECATED is not initialized properly

#rnx

diff --git a/Engine/Source/Editor/Persona/Private/PhysicsAssetRenderUtils.cpp b/Engine/Source/Editor/Persona/Private/PhysicsAssetRenderUtils.cpp
--- a/Engine/Source/Editor/Persona/Private/PhysicsAssetRenderUtils.cpp
+++ b/Engine/Source/Editor/Persona/Private/PhysicsAssetRenderUtils.cpp
@@ -63,6 +63,7 @@ FPhysicsAssetRenderSettings::FPhysicsAssetRenderSettings()
 	, bHideSimulatedBodies(false)
 	, bHideBodyMass(false)
 	, bRenderOnlySelectedConstraints(false)
+	, bShowCOM_DEPRECATED(false)
 	, bShowConstraintsAsPoints(false)
 	, bDrawViolatedLimits(false)
 	, BoneUnselectedColor(170, 155, 225)
