Skeletal Mesh: Ignore refreshes on ULODInfoUILayout if the LOD no longer exists.
The assumption here is that the the LOD got removed but the details panel hasn't had time to refresh and delete no longer needed ULODInfoUILayout objects.

diff --git a/Engine/Source/Editor/Persona/Private/LODInfoUILayout.cpp b/Engine/Source/Editor/Persona/Private/LODInfoUILayout.cpp
--- a/Engine/Source/Editor/Persona/Private/LODInfoUILayout.cpp
+++ b/Engine/Source/Editor/Persona/Private/LODInfoUILayout.cpp
@@ -31,11 +31,16 @@ void ULODInfoUILayout::RefreshReferenceLODInfo()
 	// by the time it tries to refresh, it is possible that we no longer have a valid ptr
 	if (PersonaToolkit.IsValid())
 	{
-		USkeletalMesh* SkeletalMesh = PersonaToolkit.Pin()->GetPreviewMesh();
-		const FSkeletalMeshLODInfo* SkeletalMeshLODInfo = SkeletalMesh->GetLODInfo(LODIndex);
-		check(SkeletalMeshLODInfo != nullptr);
-		//Copy the LODInfo Array to the temporary
-		LODInfo = *SkeletalMeshLODInfo;	
+		if (USkeletalMesh* SkeletalMesh = PersonaToolkit.Pin()->GetPreviewMesh())
+		{
+			// If the LOD info is out of date, don't update. Most likely this object is going to be
+			// nuked when Persona details panel for the skelmesh updates.
+			if (const FSkeletalMeshLODInfo* SkeletalMeshLODInfo = SkeletalMesh->GetLODInfo(LODIndex))
+			{
+				//Copy the LODInfo Array to the temporary
+				LODInfo = *SkeletalMeshLODInfo;	
+			}
+		}
 	}
 }
 
