Displaying categories for custom notify class properties in the details panel.
Moving category text code from DetailCategoryBuilderImpl to ObjectEditorUtils.

All properties on the notify class will now show in their appropriate categories in the detail panel when selecting an animation notify in the animation editor. The categories will be placed as subgroups of the "Anim Notify" group. Any properties without a category will be placed directly under the "Anim Notify" group.

Properties that explicitly specify a category of "AnimNotify" will be added to the "Anim Notify" category directly as if they did not have the "AnimNotify" prefix to prevent showing an extra nested category.

#tests viewing and editing anim notifies with different categories and advanced properties in the animation editor

diff --git a/Engine/Source/Editor/Persona/Private/AnimNotifyDetails.cpp b/Engine/Source/Editor/Persona/Private/AnimNotifyDetails.cpp
--- a/Engine/Source/Editor/Persona/Private/AnimNotifyDetails.cpp
+++ b/Engine/Source/Editor/Persona/Private/AnimNotifyDetails.cpp
@@ -12,9 +12,160 @@
 #include "Animation/AnimMontage.h"
 #include "Animation/EditorNotifyObject.h"
 #include "AssetSearchBoxUtilPersona.h"
+#include "IDetailGroup.h"
+#include "ObjectEditorUtils.h"
 #include "Widgets/Input/STextComboBox.h"
 #include "Animation/AnimNotifies/AnimNotifyState.h"
 
+namespace UE::AnimNotifyDetails::Private
+{
+	constexpr TCHAR CategoryDelimiter = TEXT('|');
+	constexpr FStringView AnimNotifyCategory = TEXTVIEW("AnimNotify");
+	static const FName AnimNotifyCategoryName = FName(AnimNotifyCategory.Len(), AnimNotifyCategory.GetData());
+	static const FName AdvancedCategoryName = TEXT("Advanced");
+	
+	/**
+	 * Splits a category name into its parent and leaf category names. Returns a pair of (Parent, Leaf) category names.
+	 * If there is no parent category, then the parent name will be empty
+	 */
+	static TPair<FName, FName> SplitCategory(const FName CategoryName)
+	{
+		const FString CategoryString = CategoryName.ToString();
+		const FStringView CategoryView = CategoryString;
+
+		int32 DelimiterLocation = 0;
+		if (CategoryString.FindLastChar(CategoryDelimiter, DelimiterLocation))
+		{
+			const FStringView ParentCategoryView = CategoryView.Left(DelimiterLocation);
+			const FStringView LeafCategoryView = CategoryView.RightChop(DelimiterLocation + 1);
+			const FName ParentCategoryName = FName(ParentCategoryView.Len(), ParentCategoryView.GetData());
+			const FName LeafCategoryName = FName(LeafCategoryView.Len(), LeafCategoryView.GetData());
+			
+			return MakeTuple(ParentCategoryName, LeafCategoryName);
+		}
+		
+		return MakeTuple(NAME_None, CategoryName);
+	}
+
+	/**
+	 * Strips the leading "Anim Notify" category from the category name, if there is any.
+	 * 
+	 * Helps as a number of anim notifies were authored with the "Anim Notify" category but did not previously display
+	 * it correctly. This prevents such notifies from showing an extra category level
+	 */
+	static FName StripAnimNotifyPrefix(const FName CategoryName)
+	{
+		if (CategoryName.IsNone())
+		{
+			return CategoryName;
+		}
+
+		const FString CategoryString = CategoryName.ToString();
+		const FStringView CategoryView = CategoryString;
+
+		int32 DelimiterLocation = 0;
+		if (CategoryString.StartsWith(AnimNotifyCategory))
+		{
+			FStringView StrippedCategoryView = CategoryView.RightChop(AnimNotifyCategory.Len() + 1);
+			if (!StrippedCategoryView.IsEmpty() && StrippedCategoryView[0] == CategoryDelimiter)
+			{
+				StrippedCategoryView = StrippedCategoryView.RightChop(1);
+			}
+
+			if (StrippedCategoryView.IsEmpty())
+			{
+				return NAME_None;
+			}
+
+			return FName(StrippedCategoryView.Len(), StrippedCategoryView.GetData());
+		}
+
+		return CategoryName;
+	}
+
+	/**
+	 * Adds a series of subgroups for the specified category name, with a new subgroup for each category separated by
+	 * a |. Appends the category to the subgroup map to avoid creating categories multiple times
+	 */
+	static IDetailGroup& FindOrAddSubgroup(IDetailCategoryBuilder& Category, const FName CategoryName,
+		TMap<FName, IDetailGroup*>& SubgroupMap)
+	{
+		if (IDetailGroup** ExistingGroup = SubgroupMap.Find(CategoryName))
+		{
+			check(*ExistingGroup);
+			return **ExistingGroup;
+		}
+		
+		const TPair<FName, FName> SplitCategoryName = SplitCategory(CategoryName);
+		const FName ParentCategoryName = SplitCategoryName.Key;
+		const FName LeafCategoryName = SplitCategoryName.Value;
+		const FText DisplayName = FObjectEditorUtils::GetCategoryText(LeafCategoryName);
+
+		IDetailGroup* Subgroup = nullptr;
+
+		if (ParentCategoryName.IsNone())
+		{
+			Subgroup = &Category.AddGroup(LeafCategoryName, DisplayName);
+		}
+		else
+		{
+			IDetailGroup& ParentGroup = FindOrAddSubgroup(Category, ParentCategoryName, SubgroupMap);
+			Subgroup = &ParentGroup.AddGroup(LeafCategoryName, DisplayName);
+		}
+		check(Subgroup);
+
+		SubgroupMap.Add(CategoryName, Subgroup);
+		return *Subgroup;
+	}
+	
+	/**
+	 * Adds subgroups for the specified property, but not the property itself
+	 */
+	static void AddSubgroupForProperty(IDetailCategoryBuilder& Category, const FProperty* Property,
+		TMap<FName, IDetailGroup*>& SubgroupMap)
+	{
+		if (Property)
+		{
+			const FName CategoryName = StripAnimNotifyPrefix(FObjectEditorUtils::GetCategoryFName(Property));
+			if (!CategoryName.IsNone())
+			{
+				FindOrAddSubgroup(Category, CategoryName, SubgroupMap);
+			}
+		}
+	}
+
+	/**
+	 * Adds a subcategory to the specified category with the "Advanced" name. Uses a subgroup map and advanced subgroup
+	 * map to avoid creating duplicate categories
+	 */
+	static IDetailGroup& FindOrAddAdvancedCategory(FName CategoryName, TMap<FName, IDetailGroup*>& SubgroupMap,
+		TMap<FName, IDetailGroup*>& AdvancedSubgroupMap)
+	{
+		if (IDetailGroup** ExistingAdvancedGroup = AdvancedSubgroupMap.Find(CategoryName))
+		{
+			check(*ExistingAdvancedGroup);
+			return **ExistingAdvancedGroup;
+		}
+		else
+		{
+			IDetailGroup** PropertyGroup = SubgroupMap.Find(CategoryName);
+			check(PropertyGroup);
+			check(*PropertyGroup);
+
+			static FText AdvancedCategoryText;
+			if (AdvancedCategoryText.IsEmpty())
+			{
+				AdvancedCategoryText = FObjectEditorUtils::GetCategoryText(AdvancedCategoryName);
+			}
+
+			IDetailGroup& NewAdvancedGroup = (*PropertyGroup)->AddGroup(AdvancedCategoryName, AdvancedCategoryText);
+
+			AdvancedSubgroupMap.Add(CategoryName, &NewAdvancedGroup);
+			return NewAdvancedGroup;
+		}
+	}
+}
+
 TSharedRef<IDetailCustomization> FAnimNotifyDetails::MakeInstance()
 {
 	return MakeShareable( new FAnimNotifyDetails );
@@ -46,7 +197,8 @@ void FAnimNotifyDetails::CustomizeDetails( IDetailLayoutBuilder& DetailBuilder )
 	// Don't want to edit the notify name here.
 	DetailBuilder.HideProperty(TEXT("Event.NotifyName"));
 
-	IDetailCategoryBuilder& AnimNotifyCategory = DetailBuilder.EditCategory(TEXT("AnimNotify"), FText::GetEmpty(), ECategoryPriority::TypeSpecific);
+	IDetailCategoryBuilder& AnimNotifyCategory = DetailBuilder.EditCategory(
+		UE::AnimNotifyDetails::Private::AnimNotifyCategoryName, FText::GetEmpty(), ECategoryPriority::TypeSpecific);
 
 	// Check existence of notify, get rid of the property if not set
 	if(!NotifyPtr)
@@ -93,6 +245,11 @@ void FAnimNotifyDetails::CustomizeDetails( IDetailLayoutBuilder& DetailBuilder )
 		HideLinkProperties(DetailBuilder, EventHandle);
 	}
 
+	TMap<FName, IDetailGroup*> SubgroupMap;
+	TMap<FName, IDetailGroup*> AdvancedSubgroupMap;
+	TArray<TSharedPtr<IPropertyHandle>> PropertyHandles;
+	TArray<TSharedPtr<IPropertyHandle>> AdvancedPropertyHandles;
+
 	// Customizations do not run for instanced properties, so we have to resolve the properties and then
 	// customize them here instead.
 	if(NotifyPropHandle->IsValidHandle())
@@ -116,13 +273,71 @@ void FAnimNotifyDetails::CustomizeDetails( IDetailLayoutBuilder& DetailBuilder )
 				{
 					if(!CustomizeProperty(AnimNotifyCategory, NotifyPtr, NotifyProperty))
 					{
-						AnimNotifyCategory.AddProperty(NotifyProperty);
+						// Add our subgroups first, so we can make sure they are sorted before the normal properties
+						UE::AnimNotifyDetails::Private::AddSubgroupForProperty(AnimNotifyCategory, Prop, SubgroupMap);
+
+						if (Prop->HasAnyPropertyFlags(CPF_AdvancedDisplay))
+						{
+							AdvancedPropertyHandles.Add(NotifyProperty);
+						}
+						else
+						{
+							PropertyHandles.Add(NotifyProperty);
+						}
 					}
 				}
 			}
 		}
 	}
 
+	for (TSharedPtr<IPropertyHandle> PropertyHandle : PropertyHandles)
+	{
+		check(PropertyHandle);
+		FProperty* Property = PropertyHandle->GetProperty();
+		check(Property);
+		
+		const FName PropertyGroupName =
+			UE::AnimNotifyDetails::Private::StripAnimNotifyPrefix(FObjectEditorUtils::GetCategoryFName(Property));
+
+		if (!PropertyGroupName.IsNone())
+		{
+			IDetailGroup** PropertyGroup = SubgroupMap.Find(PropertyGroupName);
+			check(PropertyGroup);
+			check(*PropertyGroup);
+
+			(*PropertyGroup)->AddPropertyRow(PropertyHandle.ToSharedRef());
+		}
+		else
+		{
+			AnimNotifyCategory.AddProperty(PropertyHandle);
+		}
+	}
+
+	// Iterate over all of the advanced properties last so we can their advanced categories as needed, to sort after the
+	// normal properties
+	for (TSharedPtr<IPropertyHandle> PropertyHandle : AdvancedPropertyHandles)
+	{
+		check(PropertyHandle);
+		FProperty* Property = PropertyHandle->GetProperty();
+		check(Property);
+
+		const FName PropertyGroupName = 
+			UE::AnimNotifyDetails::Private::StripAnimNotifyPrefix(FObjectEditorUtils::GetCategoryFName(Property));
+
+		if (!PropertyGroupName.IsNone())
+		{
+			IDetailGroup& AdvancedCategory = UE::AnimNotifyDetails::Private::FindOrAddAdvancedCategory(
+				PropertyGroupName, SubgroupMap, AdvancedSubgroupMap);
+			AdvancedCategory.AddPropertyRow(PropertyHandle.ToSharedRef());
+		}
+		else
+		{
+			// If we're just adding the property to the top level category, then AddProperty will automatically handle
+			// setting whether it's advanced or not
+			AnimNotifyCategory.AddProperty(PropertyHandle);
+		}
+	}
+
 	struct FPropVisPair
 	{
 		const TCHAR* NotifyName;
diff --git a/Engine/Source/Editor/PropertyEditor/Private/DetailCategoryBuilderImpl.cpp b/Engine/Source/Editor/PropertyEditor/Private/DetailCategoryBuilderImpl.cpp
--- a/Engine/Source/Editor/PropertyEditor/Private/DetailCategoryBuilderImpl.cpp
+++ b/Engine/Source/Editor/PropertyEditor/Private/DetailCategoryBuilderImpl.cpp
@@ -10,6 +10,7 @@
 #include "DetailItemNode.h"
 #include "DetailPropertyRow.h"
 #include "IPropertyGenerationUtilities.h"
+#include "ObjectEditorUtils.h"
 #include "Misc/ConfigCacheIni.h"
 #include "Modules/ModuleManager.h"
 #include "ObjectPropertyNode.h"
@@ -843,26 +844,7 @@ void FDetailCategoryImpl::SetDisplayName(FName InCategoryName, const FText& Loca
 	}
 	else if (InCategoryName != NAME_None)
 	{
-		static const FTextKey CategoryLocalizationNamespace = TEXT("UObjectCategory");
-		static const FName CategoryMetaDataKey = TEXT("Category");
-
-		DisplayName = FText();
-
-		const FString NativeCategory = InCategoryName.ToString();
-		if (FText::FindText(CategoryLocalizationNamespace, NativeCategory, /*OUT*/DisplayName, &NativeCategory))
-		{
-			// Category names in English are typically gathered in their non-pretty form (eg "UserInterface" rather than "User Interface"), so skip 
-			// applying the localized variant if the text matches the raw category name, as in this case the pretty printer will do a better job
-			if (NativeCategory.Equals(DisplayName.ToString(), ESearchCase::CaseSensitive))
-			{
-				DisplayName = FText();
-			}
-		}
-		
-		if (DisplayName.IsEmpty())
-		{
-			DisplayName = FText::AsCultureInvariant(FName::NameToDisplayString(NativeCategory, false));
-		}
+		DisplayName = FObjectEditorUtils::GetCategoryText(InCategoryName);
 	}
 	else
 	{
diff --git a/Engine/Source/Runtime/Engine/Private/ObjectEditorUtils.cpp b/Engine/Source/Runtime/Engine/Private/ObjectEditorUtils.cpp
--- a/Engine/Source/Runtime/Engine/Private/ObjectEditorUtils.cpp
+++ b/Engine/Source/Runtime/Engine/Private/ObjectEditorUtils.cpp
@@ -10,6 +10,33 @@
 namespace FObjectEditorUtils
 {
 
+	FText GetCategoryText( const FName InName )
+	{
+		static const FTextKey CategoryLocalizationNamespace = TEXT("UObjectCategory");
+		static const FName CategoryMetaDataKey = TEXT("Category");
+
+		FText DisplayName;
+
+		const FString NativeCategory = InName.ToString();
+		if (FText::FindText(CategoryLocalizationNamespace, NativeCategory, /*OUT*/DisplayName, &NativeCategory))
+		{
+			// Category names in English are typically gathered in their non-pretty form (eg "UserInterface" rather than "User Interface"), so skip 
+			// applying the localized variant if the text matches the raw category name, as in this case the pretty printer will do a better job
+			if (NativeCategory.Equals(DisplayName.ToString(), ESearchCase::CaseSensitive))
+			{
+				DisplayName = FText();
+			}
+		}
+		
+		if (DisplayName.IsEmpty())
+		{
+			DisplayName = FText::AsCultureInvariant(FName::NameToDisplayString(NativeCategory, false));
+		}
+
+		return DisplayName;
+	}
+
+
 	FText GetCategoryText( const FField* InField )
 	{
 		static const FTextKey CategoryLocalizationNamespace = TEXT("UObjectCategory");
diff --git a/Engine/Source/Runtime/Engine/Public/ObjectEditorUtils.h b/Engine/Source/Runtime/Engine/Public/ObjectEditorUtils.h
--- a/Engine/Source/Runtime/Engine/Public/ObjectEditorUtils.h
+++ b/Engine/Source/Runtime/Engine/Public/ObjectEditorUtils.h
@@ -11,6 +11,13 @@ class FProperty;
 
 namespace FObjectEditorUtils
 {
+	/**
+	 * Gets the category text from this category FName
+	 *
+	 * @param	InName Category FName we want the text for
+	 * @return	Category text for the provided category name
+	 */
+	ENGINE_API FText GetCategoryText( const FName InName );
 
 	/**
 	 * Gets the category this field belongs to.
