Animation Viewport: notify the user of auto-generated LODs (to mimic the details panel)

#rnx

diff --git a/Engine/Source/Editor/Persona/Private/SAnimViewportToolBar.cpp b/Engine/Source/Editor/Persona/Private/SAnimViewportToolBar.cpp
--- a/Engine/Source/Editor/Persona/Private/SAnimViewportToolBar.cpp
+++ b/Engine/Source/Editor/Persona/Private/SAnimViewportToolBar.cpp
@@ -1114,6 +1114,18 @@ TSharedRef<SWidget> SAnimViewportToolBar::GenerateLODMenu() const
 	InMenuBuilder.PushCommandList(Viewport.Pin()->GetCommandList().ToSharedRef());
 	InMenuBuilder.PushExtender(MenuExtender.ToSharedRef());
 
+	const USkeletalMesh* PreviewMesh =  Viewport.Pin()->GetPreviewScene()->GetPreviewMesh();
+	auto IsAutoGenerated = [PreviewMesh](const int32 InLODId) -> bool
+	{
+		if (!PreviewMesh || PreviewMesh->IsCompiling())
+		{
+			return false;
+		}
+		return PreviewMesh->IsValidLODIndex(InLODId) && PreviewMesh->GetLODInfo(InLODId)->bHasBeenSimplified;
+	};
+	static const FText LODFormat = LOCTEXT("LODFmt", "LOD {0}");
+	static const FText GeneratedLODFormat = LOCTEXT("GeneratedLODFmt", "LOD {0} (generated)");
+	
 	{
 		// LOD Models
 		InMenuBuilder.BeginSection("AnimViewportPreviewLODs", LOCTEXT("ShowLOD_PreviewLabel", "Preview LODs") );
@@ -1122,16 +1134,17 @@ TSharedRef<SWidget> SAnimViewportToolBar::GenerateLODMenu() const
 			InMenuBuilder.AddMenuEntry( Actions.LODAuto );
 			InMenuBuilder.AddMenuEntry( Actions.LOD0 );
 
-			int32 LODCount = Viewport.Pin()->GetLODModelCount();
+			const int32 LODCount = Viewport.Pin()->GetLODModelCount();
 			for (int32 LODId = 1; LODId < LODCount; ++LODId)
 			{
-				FString TitleLabel = FString::Printf(TEXT("LOD %d"), LODId);
+				const FText LODNumber = FText::AsNumber(LODId);
+				const FText TitleLabel = IsAutoGenerated(LODId) ? FText::Format(GeneratedLODFormat, LODNumber) : FText::Format(LODFormat, LODNumber); 
 
 				FUIAction Action(FExecuteAction::CreateSP(Viewport.Pin().ToSharedRef(), &SAnimationEditorViewportTabBody::OnSetLODModel, LODId + 1),
 					FCanExecuteAction(),
 					FIsActionChecked::CreateSP(Viewport.Pin().ToSharedRef(), &SAnimationEditorViewportTabBody::IsLODModelSelected, LODId + 1));
 
-				InMenuBuilder.AddMenuEntry(FText::FromString(TitleLabel), FText::GetEmpty(), FSlateIcon(), Action, NAME_None, EUserInterfaceActionType::RadioButton);
+				InMenuBuilder.AddMenuEntry(TitleLabel, FText::GetEmpty(), FSlateIcon(), Action, NAME_None, EUserInterfaceActionType::RadioButton);
 			}
 		}
 		InMenuBuilder.EndSection();
