Expose Find/Replace tab in animation editors and filter out non-editable assets

diff --git a/Engine/Source/Editor/Persona/Private/SAnimAssetFindReplace.cpp b/Engine/Source/Editor/Persona/Private/SAnimAssetFindReplace.cpp
--- a/Engine/Source/Editor/Persona/Private/SAnimAssetFindReplace.cpp
+++ b/Engine/Source/Editor/Persona/Private/SAnimAssetFindReplace.cpp
@@ -27,6 +27,7 @@
 #include "UObject/ObjectVersion.h"
 #include "Widgets/Input/SHyperlink.h"
 #include "UObject/UObjectIterator.h"
+#include "AssetToolsModule.h"
 
 #define LOCTEXT_NAMESPACE "SAnimAssetFindReplace"
 
@@ -476,6 +477,14 @@ void SAnimAssetFindReplace::RefreshSearchResults()
 
 bool SAnimAssetFindReplace::ShouldFilterOutAsset(const FAssetData& InAssetData, bool& bOutIsOldAsset) const
 {
+	FAssetToolsModule& AssetToolsModule = FModuleManager::LoadModuleChecked<FAssetToolsModule>(TEXT("AssetTools"));
+	
+	const bool bIsEditableAsset = AssetToolsModule.Get().GetWritableFolderPermissionList()->PassesStartsWithFilter(InAssetData.PackageName);
+	if (!bIsEditableAsset)
+	{
+		return true;
+	}
+
 	return CurrentProcessor->ShouldFilterOutAsset(InAssetData, bOutIsOldAsset);
 }
 
