Fix animation editors selecting the incorrect bone for non-mesh bones leading to array OOBs
* Regression from 24612652

diff --git a/Engine/Source/Editor/Persona/Private/AnimationEditorPreviewScene.cpp b/Engine/Source/Editor/Persona/Private/AnimationEditorPreviewScene.cpp
--- a/Engine/Source/Editor/Persona/Private/AnimationEditorPreviewScene.cpp
+++ b/Engine/Source/Editor/Persona/Private/AnimationEditorPreviewScene.cpp
@@ -862,9 +862,12 @@ void FAnimationEditorPreviewScene::SetSelectedBone(const FName& BoneName, ESelec
 		// need to get mesh bone base since BonesOfInterest is saved in SkeletalMeshComponent
 		// and it is used by renderer. It is not Skeleton base
 		const int32 MeshBoneIndex = SkeletalMeshComponent->GetBoneIndex(BoneName);
-		SelectedBoneIndex = MeshBoneIndex != INDEX_NONE ? MeshBoneIndex : BoneIndex;
-		SkeletalMeshComponent->BonesOfInterest.Add(SelectedBoneIndex);
-
+		if (MeshBoneIndex != INDEX_NONE)
+		{
+			SelectedBoneIndex = MeshBoneIndex;
+			SkeletalMeshComponent->BonesOfInterest.Add(SelectedBoneIndex);
+		}
+		
 		InvalidateViews();
 
 		OnSelectedBoneChanged.Broadcast(BoneName, InSelectInfo);
