Fix CIS Location text conflict

#rnx

diff --git a/Engine/Source/Editor/Persona/Private/SMorphTargetViewer.cpp b/Engine/Source/Editor/Persona/Private/SMorphTargetViewer.cpp
--- a/Engine/Source/Editor/Persona/Private/SMorphTargetViewer.cpp
+++ b/Engine/Source/Editor/Persona/Private/SMorphTargetViewer.cpp
@@ -581,7 +581,7 @@ TSharedPtr<SWidget> SMorphTargetViewer::OnGetContextMenuContent() const
 										Action.ExecuteAction = FExecuteAction::CreateSP(const_cast<SMorphTargetViewer*>(this), &SMorphTargetViewer::OnReimportMorphTargetsWithNewFile, LodIndex);
 										Action.CanExecuteAction = nullptr;
 										const FText Label = LOCTEXT("ReimportWithNewFileMorphTargetLabel", "Reimport With New File");
-										const FText ToolTipText = LOCTEXT("ReimportMorphTargetTooltip", "Ask a file and re-import every selected morph target.");
+										const FText ToolTipText = LOCTEXT("ReimportWithNewFileMorphTargetTooltip", "Ask a file and re-import every selected morph target.");
 										InSubMenuBuilder.AddMenuEntry(Label, ToolTipText, FSlateIcon(), Action);
 									}
 								})
