[UE] Ensure PhAT constraint debug draw is always shown in Anim BP editor
- Fix bug Where constraint debug draw is not visible in Anim BP editor when the 'show only selected constraints' option is selected in PhAT.

#rnx

diff --git a/Engine/Source/Editor/Persona/Private/PhysicsAssetRenderUtils.cpp b/Engine/Source/Editor/Persona/Private/PhysicsAssetRenderUtils.cpp
--- a/Engine/Source/Editor/Persona/Private/PhysicsAssetRenderUtils.cpp
+++ b/Engine/Source/Editor/Persona/Private/PhysicsAssetRenderUtils.cpp
@@ -300,9 +300,8 @@ namespace PhysicsAssetRender
 
 		// Draw Constraints.
 		{
-			auto HitProxyFn = [](const int32) { return nullptr; };
-			auto IsSelectedFn = [](const uint32) { return false; };
-			DebugDrawConstraints(SkeletalMeshComponent, PhysicsAsset, PDI, IsSelectedFn, false, HitProxyFn);
+			auto HitProxyFunctor = [](const int32) { return nullptr; };
+			DebugDrawConstraints(SkeletalMeshComponent, PhysicsAsset, PDI, IsSelectedFn(), false, HitProxyFunctor);
 		}
 	}
 
@@ -497,7 +496,7 @@ namespace PhysicsAssetRender
 		}
 	}
 
-	void DebugDrawConstraints(USkeletalMeshComponent* const SkeletalMeshComponent, UPhysicsAsset* const PhysicsAsset, FPrimitiveDrawInterface* PDI, TFunctionRef< bool(const uint32) > IsConstraintSelected, const bool bRunningSimulation, CreateConstraintHitProxyFn CreateHitProxy)
+	void DebugDrawConstraints(USkeletalMeshComponent* const SkeletalMeshComponent, UPhysicsAsset* const PhysicsAsset, FPrimitiveDrawInterface* PDI, TFunction< bool(const uint32) > IsSelected, const bool bRunningSimulation, CreateConstraintHitProxyFn CreateHitProxy)
 	{
 		check(SkeletalMeshComponent);
 		check(PhysicsAsset);
@@ -511,11 +510,14 @@ namespace PhysicsAssetRender
 
 		if (RenderSettings->ConstraintViewMode != EPhysicsAssetEditorConstraintViewMode::None)
 		{
+			const bool bIsSelectedStateAvailable = (IsSelected != nullptr);
+			const bool bRenderOnlySelected = bIsSelectedStateAvailable && RenderSettings->bRenderOnlySelectedConstraints;
+
 			for (int32 ConstraintIndex = 0; ConstraintIndex < PhysicsAsset->ConstraintSetup.Num(); ++ConstraintIndex)
 			{
-				const bool bConstraintSelected = IsConstraintSelected(ConstraintIndex);
+				const bool bConstraintSelected = bIsSelectedStateAvailable && IsSelected(ConstraintIndex);
 
-				if ((!RenderSettings->bRenderOnlySelectedConstraints || (RenderSettings->bRenderOnlySelectedConstraints && bConstraintSelected)) &&
+				if ((!bRenderOnlySelected || bConstraintSelected) &&
 					!RenderSettings->IsConstraintHidden(ConstraintIndex))
 				{
 					const bool bDrawLimits = (RenderSettings->ConstraintViewMode == EPhysicsAssetEditorConstraintViewMode::AllLimits) || bConstraintSelected;
@@ -756,8 +758,7 @@ void FPhysicsAssetRenderInterface::DebugDrawBodies(USkeletalMeshComponent* const
 void FPhysicsAssetRenderInterface::DebugDrawConstraints(USkeletalMeshComponent* const SkeletalMeshComponent, UPhysicsAsset* const PhysicsAsset, FPrimitiveDrawInterface* PDI)
 {
 	auto HitProxyFn   = [](const int32) { return nullptr; };
-	auto IsSelectedFn = [](const uint32) { return false; };
-	PhysicsAssetRender::DebugDrawConstraints(SkeletalMeshComponent, PhysicsAsset, PDI, IsSelectedFn, false, HitProxyFn);
+	PhysicsAssetRender::DebugDrawConstraints(SkeletalMeshComponent, PhysicsAsset, PDI, PhysicsAssetRender::IsSelectedFn(), false, HitProxyFn);
 }
 
 void FPhysicsAssetRenderInterface::SaveConfig()
diff --git a/Engine/Source/Editor/Persona/Public/PhysicsAssetRenderUtils.h b/Engine/Source/Editor/Persona/Public/PhysicsAssetRenderUtils.h
--- a/Engine/Source/Editor/Persona/Public/PhysicsAssetRenderUtils.h
+++ b/Engine/Source/Editor/Persona/Public/PhysicsAssetRenderUtils.h
@@ -214,6 +214,7 @@ namespace PhysicsAssetRender
 	using CreateConstraintHitProxyFn = TFunctionRef< HHitProxy* (const int32 InConstraintIndex) >;
 	using CreateCoMHitProxyFn = TFunctionRef< HHitProxy* (const int32 InBodyIndex) >;
 	using COMAccessorFunctionFn = TFunctionRef< TPair< bool, FVector >(const int32) >;
+	using IsSelectedFn = TFunction< bool(const uint32) >;
 
 	/** Debug draw Physics Asset bodies and constraints using the default callbacks */
 	PERSONA_API void DebugDraw(class USkeletalMeshComponent* const SkeletalMeshComponent, class UPhysicsAsset* const PhysicsAsset, FPrimitiveDrawInterface* PDI);
@@ -225,7 +226,7 @@ namespace PhysicsAssetRender
 	PERSONA_API void DebugDrawCenterOfMass(USkeletalMeshComponent* const SkeletalMeshComponent, class UPhysicsAsset* const PhysicsAsset, FPrimitiveDrawInterface* PDI, TFunctionRef< FVector(const uint32) > GetCoMPosition, TFunctionRef< bool(const uint32) > IsSelected, CreateCoMHitProxyFn CreateHitProxy);
 
 	/** Debug draw Physics Asset constraints using the supplied custom callbacks */
-	PERSONA_API void DebugDrawConstraints(class USkeletalMeshComponent* const SkeletalMeshComponent, class UPhysicsAsset* const PhysicsAsset, FPrimitiveDrawInterface* PDI, TFunctionRef< bool(const uint32) > IsConstraintSelected, const bool bRunningSimulation, CreateConstraintHitProxyFn CreateHitProxy);
+	PERSONA_API void DebugDrawConstraints(class USkeletalMeshComponent* const SkeletalMeshComponent, class UPhysicsAsset* const PhysicsAsset, FPrimitiveDrawInterface* PDI, IsSelectedFn IsSelected, const bool bRunningSimulation, CreateConstraintHitProxyFn CreateHitProxy);
 	
 	/** Default callbacks used by DebugDraw */
 	PERSONA_API FTransform GetPrimitiveTransform(const UPhysicsAsset* PhysicsAsset, const FTransform& BoneTM, const int32 BodyIndex, const EAggCollisionShape::Type PrimType, const int32 PrimIndex, const float Scale);
