Fixes crash when previewing skeletal mesh normals.

diff --git a/Engine/Source/Editor/Persona/Private/AnimationEditorViewportClient.cpp b/Engine/Source/Editor/Persona/Private/AnimationEditorViewportClient.cpp
--- a/Engine/Source/Editor/Persona/Private/AnimationEditorViewportClient.cpp
+++ b/Engine/Source/Editor/Persona/Private/AnimationEditorViewportClient.cpp
@@ -2451,10 +2451,22 @@ void FAnimationViewportClient::OnAssetViewerSettingsChanged(const FName& InPrope
 		if (Settings->Profiles.IsValidIndex(ProfileIndex))
 		{			
 			Settings->Profiles[ProfileIndex].SetShowFlags(EngineShowFlags);
+			SetAdvancedShowFlagsForScene(Settings->Profiles[ProfileIndex].bPostProcessingEnabled);
 		}
 	}
 }
 
+void FAnimationViewportClient::SetAdvancedShowFlagsForScene(const bool bAdvancedShowFlags)
+{	
+	if (bAdvancedShowFlags)
+	{
+		EngineShowFlags.EnableAdvancedFeatures();
+	}
+	else
+	{
+		EngineShowFlags.DisableAdvancedFeatures();
+	}
+}
 void FAnimationViewportClient::SetPlaybackSpeedMode(EAnimationPlaybackSpeeds::Type InMode)
 {
 	AnimationPlaybackSpeedMode = InMode;
diff --git a/Engine/Source/Editor/Persona/Public/AnimationEditorViewportClient.h b/Engine/Source/Editor/Persona/Public/AnimationEditorViewportClient.h
--- a/Engine/Source/Editor/Persona/Public/AnimationEditorViewportClient.h
+++ b/Engine/Source/Editor/Persona/Public/AnimationEditorViewportClient.h
@@ -474,6 +474,9 @@ private:
 
 	/** Delegate for preview profile is changed (used for updating show flags) */
 	void OnAssetViewerSettingsChanged(const FName& InPropertyName);
+
+	/** Sets up the ShowFlag according to the current preview scene profile */
+	void SetAdvancedShowFlagsForScene(const bool bAdvancedShowFlags);
 	
 	/** Computes a bounding box for the selected section of the preview mesh component.
 	    If there is no selected section, returns an empty box. */
