Crash after creation a Local Variable
#fix guard against no outer-objects during anim bp node customization

diff --git a/Engine/Source/Editor/Persona/Private/AnimGraphNodeDetails.cpp b/Engine/Source/Editor/Persona/Private/AnimGraphNodeDetails.cpp
--- a/Engine/Source/Editor/Persona/Private/AnimGraphNodeDetails.cpp
+++ b/Engine/Source/Editor/Persona/Private/AnimGraphNodeDetails.cpp
@@ -1204,7 +1204,7 @@ void FAnimGraphNodeBindingExtension::GetOptionalPinData(const IPropertyHandle& P
 	PropertyHandle.GetOuterObjects(Objects);
 
 	FProperty* Property = PropertyHandle.GetProperty();
-	if(Property)
+	if(Property && Objects.Num())
 	{
 		OutAnimGraphNode = Cast<UAnimGraphNode_Base>(Objects[0]);
 		if (OutAnimGraphNode != nullptr)
