Fixed possible mismatch of types when assigning UBase* types to TSoftObjectPtr<UDerived>.

diff --git a/Engine/Source/Editor/Persona/Private/PersonaModule.cpp b/Engine/Source/Editor/Persona/Private/PersonaModule.cpp
--- a/Engine/Source/Editor/Persona/Private/PersonaModule.cpp
+++ b/Engine/Source/Editor/Persona/Private/PersonaModule.cpp
@@ -1003,7 +1003,7 @@ namespace UE::Persona::Private
 					PreviewScene->GetPreviewSceneDescription()->SetPreviewController(UPersonaPreviewSceneAnimationController::StaticClass(), &PreviewScene.Get());
 
 					UPersonaPreviewSceneAnimationController* AnimController = CastChecked<UPersonaPreviewSceneAnimationController>(PreviewScene->GetPreviewSceneDescription()->PreviewControllerInstance);
-					AnimController->Animation = AssetData.GetAsset();
+					AnimController->Animation = CastChecked<UAnimationAsset>(AssetData.GetAsset());
 					AnimController->InitializeView(PreviewScene->GetPreviewSceneDescription(), &PreviewScene.Get());
 
 					// Make sure any settings views are updated with the new settings
