Adding ability to preview live pose in persona.

diff --git a/Engine/Plugins/Animation/IKRig/Source/IKRig/Private/AnimNodes/AnimNode_IKRig.cpp b/Engine/Plugins/Animation/IKRig/Source/IKRig/Private/AnimNodes/AnimNode_IKRig.cpp
--- a/Engine/Plugins/Animation/IKRig/Source/IKRig/Private/AnimNodes/AnimNode_IKRig.cpp
+++ b/Engine/Plugins/Animation/IKRig/Source/IKRig/Private/AnimNodes/AnimNode_IKRig.cpp
@@ -30,13 +30,13 @@ void FAnimNode_IKRig::Evaluate_AnyThread(FPoseContext& Output)
 	DECLARE_SCOPE_HIERARCHICAL_COUNTER_FUNC()
 	ANIM_MT_SCOPE_CYCLE_COUNTER_VERBOSE(IKRig, !IsInGameThread());
 
-	if (Source.GetLinkNode() && !bStartFromRefPose)
+	if (bStartFromRefPose)
 	{
-		Source.Evaluate(Output);
+		Output.ResetToRefPose();
 	}
-	else
+	else if (Source.GetLinkNode())
 	{
-		Output.ResetToRefPose();
+		Source.Evaluate(Output);
 	}
 
 	if (!(RigDefinitionAsset && IKRigProcessor))
diff --git a/Engine/Plugins/Animation/IKRig/Source/IKRigEditor/Private/RigEditor/IKRigAnimInstanceProxy.cpp b/Engine/Plugins/Animation/IKRig/Source/IKRigEditor/Private/RigEditor/IKRigAnimInstanceProxy.cpp
--- a/Engine/Plugins/Animation/IKRig/Source/IKRigEditor/Private/RigEditor/IKRigAnimInstanceProxy.cpp
+++ b/Engine/Plugins/Animation/IKRig/Source/IKRigEditor/Private/RigEditor/IKRigAnimInstanceProxy.cpp
@@ -14,15 +14,13 @@ FIKRigAnimInstanceProxy::FIKRigAnimInstanceProxy(UAnimInstance* InAnimInstance,
 
 void FIKRigAnimInstanceProxy::Initialize(UAnimInstance* InAnimInstance)
 {
-	FAnimPreviewInstanceProxy::Initialize(InAnimInstance);
-	IKRigNode->Source.SetLinkNode(&SingleNode);
-
-	// force this instance of the IK Rig evaluation to copy setting from the source IK Rig asset
-	IKRigNode->bDriveWithSourceAsset = true; 
+	Super::Initialize(InAnimInstance);
+	IKRigNode->bDriveWithSourceAsset = true; // force this instance of the IK Rig evaluation to copy settings from the source IK Rig asset 
 }
 
 bool FIKRigAnimInstanceProxy::Evaluate(FPoseContext& Output)
 {
+	Super::Evaluate(Output);
 	IKRigNode->Evaluate_AnyThread(Output);
 	return true;
 }
@@ -39,14 +37,8 @@ void FIKRigAnimInstanceProxy::GetCustomNodes(TArray<FAnimNode_Base*>& OutNodes)
 
 void FIKRigAnimInstanceProxy::UpdateAnimationNode(const FAnimationUpdateContext& InContext)
 {
-	if (CurrentAsset != nullptr)
-	{
-		FAnimPreviewInstanceProxy::UpdateAnimationNode(InContext);
-	}
-	else
-	{
-		IKRigNode->Update_AnyThread(InContext);
-	}
+	Super::UpdateAnimationNode(InContext);
+	IKRigNode->Update_AnyThread(InContext);
 }
 
 void FIKRigAnimInstanceProxy::SetIKRigAsset(UIKRigDefinition* InIKRigAsset)
diff --git a/Engine/Source/Editor/Persona/Private/PersonaPreviewSceneSkelMeshInstanceController.cpp b/Engine/Source/Editor/Persona/Private/PersonaPreviewSceneSkelMeshInstanceController.cpp
new file mode 100644
--- /dev/null
+++ b/Engine/Source/Editor/Persona/Private/PersonaPreviewSceneSkelMeshInstanceController.cpp
@@ -0,0 +1,409 @@
+﻿// Copyright Epic Games, Inc. All Rights Reserved.
+
+#include "PersonaPreviewSceneSkelMeshInstanceController.h"
+#include "AnimationEditorPreviewScene.h"
+#include "AnimPreviewInstance.h"
+#include "DetailWidgetRow.h"
+#include "Editor.h"
+#include "EngineUtils.h"
+#include "PersonaPreviewSceneDescription.h"
+#include "IPersonaToolkit.h"
+#include "Selection.h"
+#include "Widgets/Layout/SGridPanel.h"
+
+#define LOCTEXT_NAMESPACE "UPersonaPreviewSceneSkelMeshInstanceController"
+
+AActor* FSkeletalMeshDebugInstance::GetActor() const
+{
+	return SkeletalMeshComponent.IsValid() ? SkeletalMeshComponent->GetOwner() : nullptr;
+}
+
+void UPersonaPreviewSceneSkelMeshInstanceController::InitializeView(
+	UPersonaPreviewSceneDescription* SceneDescription,
+	IPersonaPreviewScene* PreviewScene) const
+{
+	constexpr bool bShowReferencePose = true;
+	constexpr bool bResetTransforms = true;
+	PreviewScene->ShowReferencePose(bShowReferencePose, bResetTransforms);
+}
+
+void UPersonaPreviewSceneSkelMeshInstanceController::UninitializeView(
+	UPersonaPreviewSceneDescription* SceneDescription,
+	IPersonaPreviewScene* PreviewScene) const
+{
+	PreviewScene->ShowDefaultMode();
+}
+
+IDetailPropertyRow* UPersonaPreviewSceneSkelMeshInstanceController::AddPreviewControllerPropertyToDetails(
+	const TSharedRef<IPersonaToolkit>& PersonaToolkit,
+	IDetailLayoutBuilder& DetailBuilder,
+	IDetailCategoryBuilder& Category,
+	const FProperty* Property,
+	const EPropertyLocation::Type PropertyLocation)
+{
+	TSharedPtr<IPersonaPreviewScene> PreviewScenePtr = PersonaToolkit->GetPreviewScene();
+	const USkeletalMesh* SkeletalMesh = PersonaToolkit->GetPreviewMeshComponent()->GetSkeletalMeshAsset();
+	if (!SkeletalMesh)
+	{
+		return nullptr;
+	}
+
+	if (Property->GetName() != GET_MEMBER_NAME_CHECKED(UPersonaPreviewSceneSkelMeshInstanceController, ActivePreviewInstance))
+	{
+		return nullptr;
+	}
+
+	// create custom widget to select preview instance
+	const TArray<UObject*> ListOfPreviewController{ this };
+	const FName PropertyName = GET_MEMBER_NAME_CHECKED(UPersonaPreviewSceneSkelMeshInstanceController, ActivePreviewInstance);
+	IDetailPropertyRow* NewRow = Category.AddExternalObjectProperty(ListOfPreviewController, PropertyName, EPropertyLocation::Common);
+	NewRow->CustomWidget()
+	.NameContent()
+	[
+		NewRow->GetPropertyHandle()->CreatePropertyNameWidget()
+	]
+	.ValueContent()
+	.MaxDesiredWidth(250.0f)
+	.MinDesiredWidth(250.0f)
+	[
+		SNew(SSkeletalMeshDebugSelectionWidget).PreviewScene(PreviewScenePtr)
+	];
+	
+	return NewRow;
+}
+
+void SSkeletalMeshDebugSelectionWidget::Construct(const FArguments& InArgs)
+{
+	PreviewScenePtr = InArgs._PreviewScene;
+	
+	ChildSlot
+	[
+		SNew(SHorizontalBox)
+		+SHorizontalBox::Slot()
+		.AutoWidth()
+		.Padding(FMargin(2, 0, 0, 0))
+		.HAlign(HAlign_Left)
+		.VAlign(VAlign_Center)
+		[
+			SAssignNew(InstanceComboBox, SComboBox<TSharedPtr<FSkeletalMeshDebugInstance>>)
+			.OptionsSource(&AllMeshInstances)
+			.OnGenerateWidget(this, &SSkeletalMeshDebugSelectionWidget::OnGenerateComboBoxItemWidget)
+			.OnSelectionChanged(this, &SSkeletalMeshDebugSelectionWidget::OnSelectionChanged)
+			[
+				SNew(STextBlock)
+				.Text_Lambda([this]()
+				{
+					return ActiveInstance.IsValid() ? ActiveInstance->DisplayName : LOCTEXT("PreviewDisabledText", "Preview Disabled");
+				})
+				.ColorAndOpacity_Lambda([this]()
+				{
+					const bool bIsPreviewingInstance = ActiveInstance.IsValid() && ActiveInstance->SkeletalMeshComponent.IsValid();
+					return bIsPreviewingInstance ? FLinearColor::Green : FSlateColor::UseForeground();
+				})
+			]
+		]
+	];
+
+	Refresh();
+
+	FWorldDelegates::OnPostWorldInitialization.AddSPLambda(this, [this](UWorld* World, const UWorld::InitializationValues IVS) { Refresh(); });
+	FWorldDelegates::OnWorldCleanup.AddSPLambda(this, [this](UWorld* World, bool bSessionEnded, bool bCleanupResources) { Refresh(); } );
+	FWorldDelegates::OnPostDuplicate.AddSPLambda(this, [this](UWorld* World, bool bDuplicateForPIE, FWorldDelegates::FReplacementMap& ReplacementMap, TArray<UObject*>& ObjectsToFixReferences) { Refresh(); });
+	FWorldDelegates::OnPostWorldRename.AddSPLambda(this, [this](UWorld* World) { Refresh(); } );
+	
+	FWorldDelegates::OnPIEReady.AddSPLambda(this, [this](UGameInstance*) { Refresh(); });
+	FEditorDelegates::PostPIEStarted.AddSPLambda(this, [this](bool bSimulating) { Refresh(); });
+	FEditorDelegates::PausePIE.AddSPLambda(this, [this](bool bSimulating) { Refresh(); });
+	FEditorDelegates::ResumePIE.AddSPLambda(this, [this](bool bSimulating) { Refresh(); });
+	FEditorDelegates::SingleStepPIE.AddSPLambda(this, [this](bool bSimulating) { Refresh(); });
+	FEditorDelegates::EndPIE.AddSPLambda(this, [this](bool bSimulating) { Refresh(); });
+	FEditorDelegates::CancelPIE.AddSPLambda(this, [this]() { Refresh(); });
+	FEditorDelegates::OnNewActorsPlaced.AddSPLambda(this, [this](UObject*, const TArray<AActor*>&) { Refresh(); });
+	FEditorDelegates::OnDeleteActorsBegin.AddSPLambda(this, [this]() { Refresh(); });
+	FEditorDelegates::OnSwitchBeginPIEAndSIE.AddSPLambda(this, [this](bool bSimulating) { Refresh(); });
+	
+	USelection::SelectObjectEvent.AddSPLambda(this, [this](UObject* NewSelection) { Refresh(); });
+	USelection::SelectionChangedEvent.AddSPLambda(this, [this](UObject* NewSelection) { Refresh(); });
+	USelection::SelectNoneEvent.AddSPLambda(this, [this]() { Refresh(); });
+}
+
+void SSkeletalMeshDebugSelectionWidget::Refresh()
+{
+	// empty list of instances to refresh
+	AllMeshInstances.Reset();
+	
+	// create a default/empty item used to disable preview
+	TSharedPtr<FSkeletalMeshDebugInstance> EmptyItem = MakeShared<FSkeletalMeshDebugInstance>();
+	EmptyItem->DisplayName = LOCTEXT("DefaultItemText", "None");
+	AllMeshInstances.Emplace(EmptyItem);
+	
+	const IPersonaPreviewScene* PreviewScene = PreviewScenePtr.Get();
+	if (!PreviewScene)
+	{
+		return;
+	}
+	
+	const USkeletalMesh* PreviewMesh = PreviewScene->GetPreviewMesh();
+	if (!PreviewMesh)
+	{
+		return;
+	}
+
+	const UDebugSkelMeshComponent* PreviewComponent = PreviewScene->GetPreviewMeshComponent();
+	if (!PreviewComponent)
+	{
+		return;
+	}
+
+	// get all debug worlds
+	TArray<TWeakObjectPtr<UWorld>> AllDebugWorlds = []()
+	{
+		TArray<TWeakObjectPtr<UWorld>> AllDebugWorlds;
+		for (TObjectIterator<UWorld> It; It; ++It)
+		{
+			UWorld* World = *It;
+			// include only PIE and worlds that own the persistent level (i.e. non-streaming levels).
+			const bool bIsValidDebugWorld = (World != nullptr)
+				&& World->WorldType == EWorldType::PIE || World->WorldType == EWorldType::EditorPreview
+				&& World->PersistentLevel != nullptr
+				&& World->PersistentLevel->OwningWorld == World;
+			if (!bIsValidDebugWorld)
+			{
+				continue;
+			}
+			AllDebugWorlds.Add(World);
+		}
+		return MoveTemp(AllDebugWorlds);
+	}();
+	
+	// spin through all available worlds and find all skeletal mesh components using the target skeletal mesh
+	for (int32 WorldIndex=0; WorldIndex<AllDebugWorlds.Num(); ++WorldIndex)
+	{
+		TWeakObjectPtr<UWorld> World = AllDebugWorlds[WorldIndex];
+		for (TActorIterator<AActor> ActorItr(World.Get()); ActorItr; ++ActorItr)
+		{
+			const AActor* Actor = *ActorItr;
+			if (!Actor)
+			{
+				continue;
+			}
+
+			constexpr bool bIncludeChildActors = true;
+			Actor->ForEachComponent<USkeletalMeshComponent>(bIncludeChildActors,
+				[this, PreviewMesh, PreviewComponent](USkeletalMeshComponent* Component)
+			{
+				if (!Component)
+				{
+					return;
+				}
+
+				const bool bUsingPreviewMesh = Component->GetSkeletalMeshAsset() == PreviewMesh;
+				const bool bIsNotPreviewComponent = Component != PreviewComponent;
+				if (bUsingPreviewMesh && bIsNotPreviewComponent)
+				{
+					TSharedPtr<FSkeletalMeshDebugInstance> NewInstance = MakeShared<FSkeletalMeshDebugInstance>();
+					NewInstance->SkeletalMeshComponent = Component;
+					AActor* Actor = NewInstance->GetActor();
+					NewInstance->DisplayName = IsValid(Actor) ? FText::FromString(Actor->GetName()) : FText();
+					AllMeshInstances.Emplace(NewInstance);
+				}
+			});
+		}
+	}
+	
+	// update bIsSelected for each debug instance based on if the actor is selected in the level editor
+	{
+		// default to NOT selected
+		for (TSharedPtr<FSkeletalMeshDebugInstance>& MeshInstance : AllMeshInstances)
+		{
+			if (MeshInstance.IsValid())
+			{
+				MeshInstance->bIsSelected = false;
+			}
+		}
+
+		// get the selected actors in the editor.
+		const USelection* ActiveDebugActors = GEditor->GetSelectedActors();
+		if (ActiveDebugActors == nullptr)
+		{
+			return;
+		}
+	
+		// processed in reverse order, as we want the last selected item to be the one we pick.
+		// there can only be one actor selected to preview, while many can be selected in the editor itself.
+		for (int32 Index = ActiveDebugActors->Num() - 1; Index >= 0; --Index)
+		{
+			const AActor* Actor = Cast<AActor>(ActiveDebugActors->GetSelectedObject(Index));
+			if (Actor == nullptr)
+			{
+				continue;
+			}
+
+			// is this an actor with a preview mesh?
+			const TSharedPtr<FSkeletalMeshDebugInstance>* SelectedMeshInstance = AllMeshInstances.FindByPredicate
+			([Actor](const TSharedPtr<FSkeletalMeshDebugInstance>& MeshInstance)
+			{
+				if (!MeshInstance.IsValid())
+				{
+					return false;
+				}
+			
+				return MeshInstance->GetActor() == Actor;
+			});
+
+			// found a selected preview mesh?
+			if (SelectedMeshInstance && SelectedMeshInstance->IsValid())
+			{
+				// mark it as selected and break out (only one/last selection considered for debug preview)
+				SelectedMeshInstance->Get()->bIsSelected = true;
+				break;
+			}
+		}
+	} // END update selection state
+	
+	// restore active running instance if there is one AND it's still in the list of available instances
+	TSharedPtr<FSkeletalMeshDebugInstance> ItemToActivate = EmptyItem; // default to empty item
+	for (const TSharedPtr<FSkeletalMeshDebugInstance> Instance : AllMeshInstances)
+	{
+		// search for actor with same name as was previously selected
+		if (Instance->DisplayName.EqualTo(NameOfLastSelectedInstance, ETextComparisonLevel::Default))
+		{
+			ItemToActivate = Instance;
+			break;
+		}
+	}
+
+	// assign selection to combobox
+	InstanceComboBox.Get()->SetSelectedItem(ItemToActivate);
+}
+
+TSharedRef<SWidget> SSkeletalMeshDebugSelectionWidget::OnGenerateComboBoxItemWidget(TSharedPtr<FSkeletalMeshDebugInstance> Item)
+{
+	// If we have the first item in the actor list, generate a special widget.
+	const TWeakObjectPtr<USkeletalMeshComponent> Component = Item.IsValid() ? Item->SkeletalMeshComponent : nullptr;
+	const AActor* Actor = Component.IsValid() ? Component->GetOwner() : nullptr;
+	FText ActorName = Actor ? FText::FromString(Actor->GetName()) : LOCTEXT("DestroyedActorText", "<Destroyed>");
+
+	const UWorld* World = Actor ? Actor->GetWorld() : nullptr;
+	FText WorldName = World ? FText::FromString(GetDebugStringForWorld(World)) : LOCTEXT("DestroyedWorldText", "<Destroyed>");
+
+	constexpr FLinearColor ActiveColor = FLinearColor(0.0f, 1.0f, 0.0f);
+	constexpr FLinearColor EditorSelectedMarkColor = FLinearColor(1.0f, 0.5f, 0.0f);
+
+	// special item to disable debugging.
+	if (Actor == nullptr)
+	{
+		ActorName = LOCTEXT("DebugDisabledActorText", "None");
+		WorldName = LOCTEXT("DebugDisabledWorldText", "Editor Preview World");
+	}
+	
+	const TSharedPtr<SWidget> ItemWidget = 
+	SNew(SGridPanel)
+	+SGridPanel::Slot(0, 0)
+	.Padding(2.0f)
+	.HAlign(EHorizontalAlignment::HAlign_Right)
+	[
+		SNew(STextBlock)
+		.Text(LOCTEXT("ActorName", "Actor:"))
+		.ColorAndOpacity(FSlateColor::UseSubduedForeground())
+	]
+	+SGridPanel::Slot(1, 0)
+	.Padding(2.0f)
+	.HAlign(EHorizontalAlignment::HAlign_Left)
+	[
+		SNew(STextBlock)
+		.Text(ActorName)
+		.Font_Lambda([Item]()
+		{
+			if (Item.IsValid() && Item->bIsSelected)
+			{
+				return FAppStyle::GetFontStyle(TEXT("NormalFontBold"));
+			}
+			return FAppStyle::GetFontStyle(TEXT("NormalFont"));
+		})
+		.ColorAndOpacity_Lambda([this, Item, ActiveColor]() 
+		{ 
+			const AActor* ItemActor = Item.IsValid() ? Item->GetActor() : nullptr;
+			return (ActiveInstance.IsValid() && ActiveInstance->GetActor() == ItemActor) ? ActiveColor : FSlateColor::UseForeground();
+		})
+	]
+	+SGridPanel::Slot(2, 0)
+	.Padding(2.0f)
+	.HAlign(EHorizontalAlignment::HAlign_Left)
+	[
+		SNew(STextBlock)
+		.Text_Lambda([Item]()
+		{
+			if (Item.IsValid() && Item->bIsSelected)
+			{
+				return LOCTEXT("SelectedText", "(Selected)");
+			}
+			return FText();
+		})
+		.TextStyle(FAppStyle::Get(), "RichTextBlock.Bold")
+		.ColorAndOpacity(EditorSelectedMarkColor)
+	]
+	+SGridPanel::Slot(0, 1)
+	.Padding(2.0f)
+	.HAlign(EHorizontalAlignment::HAlign_Right)
+	[
+		SNew(STextBlock)
+		.Text(LOCTEXT("WorldName", "World:"))
+		.ColorAndOpacity(FSlateColor::UseSubduedForeground())
+	]
+	+SGridPanel::Slot(1, 1)
+	.Padding(2.0f)
+	.HAlign(EHorizontalAlignment::HAlign_Left)
+	[
+		SNew(STextBlock)
+		.Text(WorldName)
+		.ColorAndOpacity_Lambda([this, Item, ActiveColor]() 
+		{ 
+			const TObjectPtr<AActor> ItemActor = Item.IsValid() ? Item->GetActor() : nullptr;
+			return (ActiveInstance.IsValid() && ActiveInstance->GetActor() == ItemActor) ? ActiveColor : FSlateColor::UseForeground();
+		})
+	];
+
+	return ItemWidget.ToSharedRef();
+}
+
+void SSkeletalMeshDebugSelectionWidget::OnSelectionChanged(TSharedPtr<FSkeletalMeshDebugInstance> Item, ESelectInfo::Type SelectInfo)
+{
+	ActiveInstance = Item;
+	if (SelectInfo != ESelectInfo::Type::Direct)
+	{
+		NameOfLastSelectedInstance = ActiveInstance->DisplayName;	
+	}
+			
+	const IPersonaPreviewScene* PreviewScene = PreviewScenePtr.Get();
+	if (!PreviewScene)
+	{
+		return;
+	}
+
+	UDebugSkelMeshComponent* PreviewMeshComponent = PreviewScene->GetPreviewMeshComponent();
+	if (!PreviewMeshComponent)
+	{
+		return;
+	}
+
+	UAnimPreviewInstance* PreviewAnimInstance = PreviewMeshComponent->PreviewInstance;
+	if (!PreviewAnimInstance)
+	{
+		return;
+	}
+
+	// reset to show no preview
+	if (!Item->SkeletalMeshComponent.IsValid())
+	{
+		PreviewAnimInstance->SetDebugSkeletalMeshComponent(nullptr);
+		PreviewMeshComponent->bTrackAttachedInstanceLOD = false;
+		return;
+	}
+
+	// assign the preview mesh to the debug skel mesh component
+	PreviewAnimInstance->SetDebugSkeletalMeshComponent(Item->SkeletalMeshComponent.Get());
+	PreviewMeshComponent->bTrackAttachedInstanceLOD = true;
+}
+
+#undef LOCTEXT_NAMESPACE
diff --git a/Engine/Source/Editor/Persona/Private/PersonaPreviewSceneSkelMeshInstanceController.h b/Engine/Source/Editor/Persona/Private/PersonaPreviewSceneSkelMeshInstanceController.h
new file mode 100644
--- /dev/null
+++ b/Engine/Source/Editor/Persona/Private/PersonaPreviewSceneSkelMeshInstanceController.h
@@ -0,0 +1,75 @@
+﻿// Copyright Epic Games, Inc. All Rights Reserved.
+
+#pragma once
+
+#include "PersonaPreviewSceneController.h"
+#include "Widgets/Input/SComboBox.h"
+#include "PersonaPreviewSceneSkelMeshInstanceController.generated.h"
+
+struct FSkeletalMeshDebugInstance
+{
+	FSkeletalMeshDebugInstance() = default;
+	AActor* GetActor() const;
+	
+	// the component inside the PIE viewport, or nullptr when debugging is disabled. */
+	TWeakObjectPtr<USkeletalMeshComponent> SkeletalMeshComponent = nullptr;
+	// true if the actor is selected in the level editor (used to color UI to help user find intended instance)
+	bool bIsSelected = false;
+	// the display name
+	FText DisplayName;
+};
+
+class SSkeletalMeshDebugSelectionWidget : public SCompoundWidget
+{
+public:
+	SLATE_BEGIN_ARGS(SSkeletalMeshDebugSelectionWidget) {}
+	SLATE_ARGUMENT(TSharedPtr<IPersonaPreviewScene>, PreviewScene)
+	SLATE_END_ARGS()
+
+	void Construct(const FArguments& InArgs);
+
+private:
+
+	// gathers all skel mesh instances in all UWorlds and refresh the UI options
+	void Refresh();
+
+	// UI
+	TSharedRef<SWidget> OnGenerateComboBoxItemWidget(TSharedPtr<FSkeletalMeshDebugInstance> Item);
+	IDetailPropertyRow* GenerateInstancesCombobox(IDetailCategoryBuilder& Category);
+	void OnSelectionChanged(TSharedPtr<FSkeletalMeshDebugInstance> Item, ESelectInfo::Type SelectInfo);
+
+	// store the preview scene for callbacks to use
+	TSharedPtr<IPersonaPreviewScene> PreviewScenePtr;
+	// the combobox that contains the instance names
+	TSharedPtr<SComboBox<TSharedPtr<FSkeletalMeshDebugInstance>>> InstanceComboBox;
+	// all the instances of skeletal mesh components using the target skeletal mesh
+	TArray<TSharedPtr<FSkeletalMeshDebugInstance>> AllMeshInstances;
+	// the currently active component to copy the pose from 
+	TSharedPtr<FSkeletalMeshDebugInstance> ActiveInstance;
+	// the display name of the last instance selected by the user, used to restore their selection between PIE sessions
+	FText NameOfLastSelectedInstance;
+};
+
+UCLASS(DisplayName = "Running Instance")
+class UPersonaPreviewSceneSkelMeshInstanceController : public UPersonaPreviewSceneController
+{
+public:
+	GENERATED_BODY()
+
+	UPersonaPreviewSceneSkelMeshInstanceController() : ActivePreviewInstance(nullptr) {}
+
+	/** The instance to preview. */
+	UPROPERTY(EditAnywhere, Category = "Animation")
+	TWeakObjectPtr<USkeletalMeshComponent> ActivePreviewInstance;
+
+	// BEGIN UPersonaPreviewSceneController interface
+	virtual void InitializeView(UPersonaPreviewSceneDescription* SceneDescription, IPersonaPreviewScene* PreviewScene) const override;
+	virtual void UninitializeView(UPersonaPreviewSceneDescription* SceneDescription, IPersonaPreviewScene* PreviewScene) const override;
+	virtual IDetailPropertyRow* AddPreviewControllerPropertyToDetails(
+		const TSharedRef<IPersonaToolkit>& PersonaToolkit,
+		IDetailLayoutBuilder& DetailBuilder,
+		IDetailCategoryBuilder& Category,
+		const FProperty* Property,
+		const EPropertyLocation::Type PropertyLocation) override;
+	// END UPersonaPreviewSceneController interface
+};
diff --git a/Engine/Source/Editor/Persona/Private/PreviewSceneCustomizations.cpp b/Engine/Source/Editor/Persona/Private/PreviewSceneCustomizations.cpp
--- a/Engine/Source/Editor/Persona/Private/PreviewSceneCustomizations.cpp
+++ b/Engine/Source/Editor/Persona/Private/PreviewSceneCustomizations.cpp
@@ -32,6 +32,7 @@
 #include "Features/IModularFeatures.h"
 #include "AnimPreviewInstance.h"
 #include "PersonaModule.h"
+#include "PersonaPreviewSceneSkelMeshInstanceController.h"
 
 #define LOCTEXT_NAMESPACE "PreviewSceneCustomizations"
 
@@ -98,7 +99,8 @@ void FPreviewSceneDescriptionCustomization::CustomizeDetails(IDetailLayoutBuilde
 		TArray<UClass*> BuiltInPreviewControllers = {
 			UPersonaPreviewSceneDefaultController::StaticClass(),
 			UPersonaPreviewSceneRefPoseController::StaticClass(),
-			UPersonaPreviewSceneAnimationController::StaticClass() };
+			UPersonaPreviewSceneAnimationController::StaticClass(),
+			UPersonaPreviewSceneSkelMeshInstanceController::StaticClass() };
 		TArray<UClass*> DynamicPreviewControllers;
 		for (TObjectIterator<UClass> It; It; ++It)
 		{
