Persona: Fix Linux build errors.

#rnx

diff --git a/Engine/Source/Editor/Persona/Private/PersonaPreviewSceneSkelMeshInstanceController.cpp b/Engine/Source/Editor/Persona/Private/PersonaPreviewSceneSkelMeshInstanceController.cpp
--- a/Engine/Source/Editor/Persona/Private/PersonaPreviewSceneSkelMeshInstanceController.cpp
+++ b/Engine/Source/Editor/Persona/Private/PersonaPreviewSceneSkelMeshInstanceController.cpp
@@ -164,7 +164,7 @@ void SSkeletalMeshDebugSelectionWidget::Refresh()
 			UWorld* World = *It;
 			// include only PIE and worlds that own the persistent level (i.e. non-streaming levels).
 			const bool bIsValidDebugWorld = (World != nullptr)
-				&& World->WorldType == EWorldType::PIE || World->WorldType == EWorldType::EditorPreview
+				&& (World->WorldType == EWorldType::PIE || World->WorldType == EWorldType::EditorPreview)
 				&& World->PersistentLevel != nullptr
 				&& World->PersistentLevel->OwningWorld == World;
 			if (!bIsValidDebugWorld)
@@ -263,7 +263,7 @@ void SSkeletalMeshDebugSelectionWidget::Refresh()
 	
 	// restore active running instance if there is one AND it's still in the list of available instances
 	TSharedPtr<FSkeletalMeshDebugInstance> ItemToActivate = EmptyItem; // default to empty item
-	for (const TSharedPtr<FSkeletalMeshDebugInstance> Instance : AllMeshInstances)
+	for (const TSharedPtr<FSkeletalMeshDebugInstance>& Instance : AllMeshInstances)
 	{
 		// search for actor with same name as was previously selected
 		if (Instance->DisplayName.EqualTo(NameOfLastSelectedInstance, ETextComparisonLevel::Default))
