Interchange: Custom morph target will now be import with the file name by default. Renaming context menu added.

#rnx

diff --git a/Engine/Plugins/Interchange/Runtime/Source/Pipelines/Private/InterchangeGenericAssetsPipeline.cpp b/Engine/Plugins/Interchange/Runtime/Source/Pipelines/Private/InterchangeGenericAssetsPipeline.cpp
--- a/Engine/Plugins/Interchange/Runtime/Source/Pipelines/Private/InterchangeGenericAssetsPipeline.cpp
+++ b/Engine/Plugins/Interchange/Runtime/Source/Pipelines/Private/InterchangeGenericAssetsPipeline.cpp
@@ -165,7 +165,7 @@ void UInterchangeGenericAssetsPipeline::AdjustSettingsForContext(const FIntercha
 	else if (ContextParams.ContextType == EInterchangePipelineContext::AssetCustomMorphTargetImport
 		|| ContextParams.ContextType == EInterchangePipelineContext::AssetCustomMorphTargetReImport)
 	{
-		bUseSourceNameForAsset = false;
+		bUseSourceNameForAsset = true;
 	}
 
 	if (MaterialPipeline)
diff --git a/Engine/Source/Editor/Persona/Private/SMorphTargetViewer.cpp b/Engine/Source/Editor/Persona/Private/SMorphTargetViewer.cpp
--- a/Engine/Source/Editor/Persona/Private/SMorphTargetViewer.cpp
+++ b/Engine/Source/Editor/Persona/Private/SMorphTargetViewer.cpp
@@ -13,11 +13,13 @@
 #include "IPersonaPreviewScene.h"
 #include "InterchangeMeshUtilities.h"
 #include "InterchangeManager.h"
+#include "Misc/App.h"
 #include "Rendering/SkeletalMeshLODImporterData.h"
 #include "ScopedTransaction.h"
 #include "SkeletalMeshAttributes.h"
 #include "SkeletalRenderPublic.h"
 #include "SkinnedAssetCompiler.h"
+#include "SRenameMorphTargetDialog.h"
 #include "Textures/SlateIcon.h"
 #include "Widgets/Input/SButton.h"
 #include "Widgets/Input/SCheckBox.h"
@@ -598,6 +600,15 @@ TSharedPtr<SWidget> SMorphTargetViewer::OnGetContextMenuContent() const
 			{
 				FUIAction Action;
 
+				//Rename morph target
+				{
+					Action.ExecuteAction = FExecuteAction::CreateSP(const_cast<SMorphTargetViewer*>(this), &SMorphTargetViewer::OnRenameMorphTargets);
+					Action.CanExecuteAction = nullptr;
+					const FText Label = LOCTEXT("RenameMorphTargetLabel", "Rename");
+					const FText ToolTipText = LOCTEXT("RenameMorphTargetTooltip", "Rename the selected morph targets");
+					MenuBuilder.AddMenuEntry(Label, ToolTipText, FSlateIcon(), Action);
+				}
+
 				//Delete morph target
 				{
 					Action.ExecuteAction = FExecuteAction::CreateSP(const_cast<SMorphTargetViewer*>(this), &SMorphTargetViewer::OnDeleteMorphTargets);
@@ -683,6 +694,50 @@ bool SMorphTargetViewer::CanPerformDelete() const
 	return SelectedRows.Num() > 0;
 }
 
+void SMorphTargetViewer::OnRenameMorphTargets()
+{
+	auto RenameMorphTarget = [this](UMorphTarget* SelectMorphTarget)
+		{
+			TSharedRef <SRenameMorphTargetDialog> RenameWidgetDialog = SNew(SRenameMorphTargetDialog)
+				.SkeletalMesh(SkeletalMesh)
+				.MorphTarget(SelectMorphTarget);
+
+			TSharedRef<SWindow> RenameWindowDialog =
+				SNew(SWindow)
+				.Title(LOCTEXT("RenameMorphTargetWindowTitle", "Rename Morph target"))
+				.SizingRule(ESizingRule::Autosized)
+				.SupportsMaximize(false)
+				.SupportsMinimize(false);
+
+			RenameWindowDialog->SetContent(SNew(SBox)
+				.MinDesiredWidth(320.0f)
+				[
+					RenameWidgetDialog
+				]);
+			TSharedPtr<SWindow> CurrentWindow = FSlateApplication::Get().FindWidgetWindow(AsShared());
+			FSlateApplication::Get().AddModalWindow(RenameWindowDialog, CurrentWindow);
+		};
+
+	{
+		FScopedSkeletalMeshPostEditChange PostEditChangeScope(SkeletalMesh);
+		TArray< TSharedPtr< FDisplayedMorphTargetInfo > > SelectedRows = MorphTargetListView->GetSelectedItems();
+
+		for (int32 RowIndex = 0; RowIndex < SelectedRows.Num(); ++RowIndex)
+		{
+			UMorphTarget* MorphTarget = SkeletalMesh->FindMorphTarget(SelectedRows[RowIndex]->Name);
+			if (MorphTarget)
+			{
+				RenameMorphTarget(MorphTarget);
+			}
+		}
+	}
+
+	//Wait until the skeletal mesh compilation is done
+	FSkinnedAssetCompilingManager::Get().FinishCompilation({ SkeletalMesh });
+
+	CreateMorphTargetList(NameFilterBox->GetText().ToString());
+}
+
 void SMorphTargetViewer::OnDeleteMorphTargets()
 {
 	TArray< TSharedPtr< FDisplayedMorphTargetInfo > > SelectedRows = MorphTargetListView->GetSelectedItems();
diff --git a/Engine/Source/Editor/Persona/Private/SMorphTargetViewer.h b/Engine/Source/Editor/Persona/Private/SMorphTargetViewer.h
--- a/Engine/Source/Editor/Persona/Private/SMorphTargetViewer.h
+++ b/Engine/Source/Editor/Persona/Private/SMorphTargetViewer.h
@@ -135,6 +135,11 @@ public:
 	*/
 	bool CanPerformDelete() const;
 
+	/**
+	* Handler for the rename morph targets button
+	*/
+	void OnRenameMorphTargets();
+
 	/**
 	* Handler for the delete morph targets button
 	*/
diff --git a/Engine/Source/Editor/Persona/Private/SRenameMorphTargetDialog.cpp b/Engine/Source/Editor/Persona/Private/SRenameMorphTargetDialog.cpp
new file mode 100644
--- /dev/null
+++ b/Engine/Source/Editor/Persona/Private/SRenameMorphTargetDialog.cpp
@@ -0,0 +1,227 @@
+// Copyright Epic Games, Inc. All Rights Reserved.
+
+#include "SRenameMorphTargetDialog.h"
+
+#include "Animation/MorphTarget.h"
+#include "Engine/SkinnedAssetCommon.h"
+#include "Framework/Application/SlateApplication.h"
+#include "SkeletalMeshAttributes.h"
+#include "SkeletalMeshTypes.h"
+#include "Widgets/Input/SButton.h"
+#include "Widgets/Input/SEditableTextBox.h"
+#include "Widgets/Layout/SBox.h"
+#include "Widgets/Layout/SGridPanel.h"
+#include "Widgets/Text/STextBlock.h"
+#include "Widgets/SWindow.h"
+
+#define LOCTEXT_NAMESPACE "RenameMorphTarget"
+
+void SRenameMorphTargetDialog::Construct(const FArguments& InArgs)
+{
+	check(InArgs._SkeletalMesh)
+	check(InArgs._MorphTarget)
+
+	SkeletalMesh = InArgs._SkeletalMesh;
+	MorphTarget = InArgs._MorphTarget;
+
+	ChildSlot
+	[
+		SNew(SBox)
+		.Padding(InArgs._Padding)
+		[
+			SNew(SGridPanel)
+			.FillColumn(1, 1.0)
+
+			// Current name display
+			+ SGridPanel::Slot(0, 0)
+			.Padding(2)
+			.VAlign(VAlign_Center)
+			.HAlign(HAlign_Left)
+			[
+				SNew(STextBlock)
+				.Font(FAppStyle::GetFontStyle( TEXT("PropertyWindow.NormalFont")))
+				.Text(LOCTEXT("CurrentName", "Current Name:"))
+			]
+
+			+ SGridPanel::Slot(1, 0)
+			.Padding(2)
+			.VAlign(VAlign_Center)
+			.HAlign(HAlign_Fill)
+			[
+				SNew(STextBlock)
+				.Font(FAppStyle::GetFontStyle( TEXT("PropertyWindow.NormalFont")))
+				.Text(FText::FromName(MorphTarget->GetFName()))
+			]
+			
+			// New name controls
+			+ SGridPanel::Slot(0, 1)
+			.Padding(2)
+			.VAlign(VAlign_Center)
+			.HAlign(HAlign_Left)
+			[
+				SNew(STextBlock )
+				.Font(FAppStyle::GetFontStyle( TEXT("PropertyWindow.NormalFont")))
+				.Text(LOCTEXT("NewName", "New Name:"))
+			]
+
+			+ SGridPanel::Slot(1, 1)
+			.Padding(2)
+			.VAlign(VAlign_Center)
+			.HAlign(HAlign_Fill)
+			[
+				SAssignNew(NewMorphTargetNameTextBox, SEditableTextBox)
+				.Font(FAppStyle::GetFontStyle( TEXT("PropertyWindow.NormalFont")))
+				.Text(FText::FromName(MorphTarget->GetFName()))
+				.OnVerifyTextChanged(this, &SRenameMorphTargetDialog::OnVerifyNameChanged)
+				.OnTextCommitted(this, &SRenameMorphTargetDialog::OnRenameTextCommitted)
+			]
+
+			// Dialog controls
+			+ SGridPanel::Slot(0, 2)
+			.ColumnSpan(2)
+			.HAlign(HAlign_Right)
+			.Padding(FMargin(0, 16))
+			[
+				SNew(SHorizontalBox)
+
+				// Rename
+				+SHorizontalBox::Slot()
+				.AutoWidth()
+				.Padding(0, 0, 8.0f, 0)
+				[
+					SNew(SButton)
+					.IsFocusable(false)
+					.OnClicked(this, &SRenameMorphTargetDialog::OnRenameClicked)
+					[
+						SNew(STextBlock)
+						.Text(LOCTEXT("RenameMorphTargetButtonText", "Rename"))
+					]
+				]
+
+				// Cancel
+				+SHorizontalBox::Slot()
+				.AutoWidth()
+				[
+					SNew(SButton)
+					.IsFocusable(false)
+					.OnClicked(this, & SRenameMorphTargetDialog::OnCancelClicked)
+					[
+						SNew(STextBlock)
+						.Text(LOCTEXT("CancelRenameButtonText", "Cancel"))
+					]
+				]
+			]
+		]
+	];
+}
+
+bool SRenameMorphTargetDialog::OnVerifyNameChanged(const FText& InText, FText& OutErrorMessage)
+{
+	const FString OriginalMorphTargetName = MorphTarget->GetName();
+	if (InText.ToString().Equals(OriginalMorphTargetName))
+	{
+		return true;
+	}
+
+	if (!MorphTarget->Rename(*(InText.ToString()), nullptr, REN_Test))
+	{
+		OutErrorMessage = LOCTEXT("VerifyBadName", "Bad name");
+		return false;
+	}
+	return true;
+}
+
+void SRenameMorphTargetDialog::RenameAndClose()
+{
+	const FString MorphTargetToRename = MorphTarget->GetName();
+	const FString NewMorphTargetName = NewMorphTargetNameTextBox->GetText().ToString();
+	if (NewMorphTargetName.Equals(MorphTargetToRename))
+	{
+		CloseContainingWindow();
+		return;
+	}
+	if (MorphTarget->Rename(*NewMorphTargetName, nullptr, REN_Test))
+	{
+		if (MorphTarget->Rename(*NewMorphTargetName, nullptr))
+		{
+			for (int32 LodIndex = 0; LodIndex < SkeletalMesh->GetLODNum(); ++LodIndex)
+			{
+				//The rename is successful, we now update the source model so it build the skeletal mesh with the correct name
+				if (FMeshDescription* MeshDescription = SkeletalMesh->GetMeshDescription(LodIndex))
+				{
+					FSkeletalMeshAttributes MeshDescriptionAttribute(*MeshDescription);
+					FName OriginalMorphAttributeName = *MorphTargetToRename;
+					FName NewMorphAttributeName = *NewMorphTargetName;
+					bool bHasVertexMorphPositionDelta = MeshDescriptionAttribute.HasMorphTargetPositionsAttribute(OriginalMorphAttributeName);
+					bool bHasMorphTargetNormalAttribute = MeshDescriptionAttribute.HasMorphTargetNormalsAttribute(OriginalMorphAttributeName);
+					if (bHasVertexMorphPositionDelta)
+					{
+						bool bMorphTargetAttributeExist = MeshDescriptionAttribute.HasMorphTargetPositionsAttribute(NewMorphAttributeName);
+						if (!bMorphTargetAttributeExist && bHasVertexMorphPositionDelta)
+						{
+							FScopedSkeletalMeshPostEditChange PostEditChangeScope(SkeletalMesh);
+							MeshDescriptionAttribute.RegisterMorphTargetAttribute(NewMorphAttributeName, bHasMorphTargetNormalAttribute);
+							TVertexAttributesRef<FVector3f> OriginalVertexMorphPositionDelta = MeshDescriptionAttribute.GetVertexMorphPositionDelta(OriginalMorphAttributeName);
+							TVertexAttributesRef<FVector3f> RenamedVertexMorphPositionDelta = MeshDescriptionAttribute.GetVertexMorphPositionDelta(NewMorphAttributeName);
+							RenamedVertexMorphPositionDelta.Copy(OriginalVertexMorphPositionDelta);
+							if (bHasMorphTargetNormalAttribute)
+							{
+								TVertexInstanceAttributesRef<FVector3f> OriginalVertexInstanceMorphNormalDelta = MeshDescriptionAttribute.GetVertexInstanceMorphNormalDelta(OriginalMorphAttributeName);
+								TVertexInstanceAttributesRef<FVector3f> RenamedVertexInstanceMorphNormalDelta = MeshDescriptionAttribute.GetVertexInstanceMorphNormalDelta(NewMorphAttributeName);
+								RenamedVertexInstanceMorphNormalDelta.Copy(OriginalVertexInstanceMorphNormalDelta);
+							}
+							MeshDescriptionAttribute.UnregisterMorphTargetAttribute(OriginalMorphAttributeName);
+							SkeletalMesh->CommitMeshDescription(0);
+							FSkeletalMeshLODInfo* LodInfo = SkeletalMesh->GetLODInfo(LodIndex);
+							if (ensure(LodInfo))
+							{
+								if (FMorphTargetImportedSourceFileInfo* MorphTargetImportedSourceFileInfo = LodInfo->ImportedMorphTargetSourceFilename.Find(MorphTargetToRename))
+								{
+									const FString OriginalSourceFile = MorphTargetImportedSourceFileInfo->GetSourceFilename();
+									LodInfo->ImportedMorphTargetSourceFilename.Remove(MorphTargetToRename);
+									//Add the mapping to the lod info so we can re-import the morph target
+									LodInfo->ImportedMorphTargetSourceFilename.FindOrAdd(NewMorphTargetName).SetSourceFilename(OriginalSourceFile);
+								}
+							}
+						}
+					}
+				}
+			}
+		}
+	}
+	CloseContainingWindow();
+}
+
+void SRenameMorphTargetDialog::OnRenameTextCommitted(const FText& InText, ETextCommit::Type InCommitType)
+{
+	if (InCommitType == ETextCommit::OnEnter)
+	{
+		RenameAndClose();
+	}
+}
+
+FReply SRenameMorphTargetDialog::OnRenameClicked()
+{
+	RenameAndClose();
+
+	return FReply::Handled();
+}
+
+FReply SRenameMorphTargetDialog::OnCancelClicked()
+{
+	CloseContainingWindow();
+
+	return FReply::Handled();
+}
+
+void SRenameMorphTargetDialog::CloseContainingWindow()
+{
+	TSharedPtr<SWindow> CurrentWindow = FSlateApplication::Get().FindWidgetWindow( AsShared() );
+
+	if (CurrentWindow.IsValid())
+	{
+		CurrentWindow->RequestDestroyWindow();
+	}
+}
+
+#undef LOCTEXT_NAMESPACE
diff --git a/Engine/Source/Editor/Persona/Private/SRenameMorphTargetDialog.h b/Engine/Source/Editor/Persona/Private/SRenameMorphTargetDialog.h
new file mode 100644
--- /dev/null
+++ b/Engine/Source/Editor/Persona/Private/SRenameMorphTargetDialog.h
@@ -0,0 +1,57 @@
+// Copyright Epic Games, Inc. All Rights Reserved.
+
+#pragma once
+
+#include "Animation/MorphTarget.h"
+#include "Engine/SkeletalMesh.h"
+#include "UObject/ObjectPtr.h"
+#include "Widgets/SCompoundWidget.h"
+
+class SEditableTextBox;
+namespace ETextCommit { enum Type : int; }
+class UMorphTarget;
+
+class SRenameMorphTargetDialog : public SCompoundWidget
+{
+public:
+
+
+	SLATE_BEGIN_ARGS(SRenameMorphTargetDialog)
+		: _SkeletalMesh(nullptr)
+		, _MorphTarget(nullptr)
+		, _Padding(FMargin(15))
+		{}
+		SLATE_ARGUMENT(TObjectPtr<USkeletalMesh>, SkeletalMesh)
+		SLATE_ARGUMENT(TObjectPtr<UMorphTarget>, MorphTarget)
+		SLATE_ARGUMENT(FMargin, Padding)
+	SLATE_END_ARGS()
+
+	void Construct(const FArguments& InArgs);
+
+private:
+
+	/** Renames the morph target based on dialog parameters */
+	FReply OnRenameClicked();
+
+	/** Callback for when Cancel is clicked */
+	FReply OnCancelClicked();
+
+	/** Callback to verify the rename is good */
+	bool OnVerifyNameChanged(const FText& InText, FText& OutErrorMessage);
+
+	/** Renames the morph target and attempts to close the active window */
+	void RenameAndClose();
+
+	/** Attempts to rename the morph target if enter is pressed while editing the morph target name */
+	void OnRenameTextCommitted(const FText& InText, ETextCommit::Type InCommitType);
+
+	/** Closes the window that contains this widget */
+	void CloseContainingWindow();
+
+private:
+	TObjectPtr<USkeletalMesh> SkeletalMesh = nullptr;
+
+	TObjectPtr<UMorphTarget> MorphTarget = nullptr;
+
+	TSharedPtr<SEditableTextBox> NewMorphTargetNameTextBox;
+};
