Fix an issue where closing a stand alone viewport resulted in a crash.  This removes code added to fix issues in linux asan tests that were failing in Animation.CQTests

diff --git a/Engine/Source/Editor/Persona/Private/AnimationEditorViewportClient.cpp b/Engine/Source/Editor/Persona/Private/AnimationEditorViewportClient.cpp
--- a/Engine/Source/Editor/Persona/Private/AnimationEditorViewportClient.cpp
+++ b/Engine/Source/Editor/Persona/Private/AnimationEditorViewportClient.cpp
@@ -200,10 +200,6 @@ FAnimationViewportClient::~FAnimationViewportClient()
 	}
 	OnPhysicsCreatedDelegateHandle.Reset();
 	OnMeshChangedDelegateHandle.Reset();
-
-	// Clear out the preview scene so any subsequent editor object destruction that tries to poke into the
-	// world will not crash.
-	static_cast<FAssetEditorModeManager*>(ModeTools.Get())->SetPreviewScene(nullptr);
 	
 	UAssetViewerSettings::Get()->OnAssetViewerSettingsChanged().RemoveAll(this);
 }
