Added root motion orientation visualization + enable root motion vis by default

diff --git a/Engine/Source/Editor/Persona/Private/AnimViewportShowCommands.cpp b/Engine/Source/Editor/Persona/Private/AnimViewportShowCommands.cpp
--- a/Engine/Source/Editor/Persona/Private/AnimViewportShowCommands.cpp
+++ b/Engine/Source/Editor/Persona/Private/AnimViewportShowCommands.cpp
@@ -19,6 +19,10 @@ void FAnimViewportShowCommands::RegisterCommands()
 	UI_COMMAND(ProcessRootMotionLoopAndReset, "Loop and Reset", "Preview mesh will consume root motion resetting the position back to the origin every time the animation loops", EUserInterfaceActionType::RadioButton, FInputChord());
 	UI_COMMAND(ProcessRootMotionLoop, "Loop", "Preview mesh will consume root motion continually", EUserInterfaceActionType::RadioButton, FInputChord());
 
+	UI_COMMAND(DoNotVisualizeRootMotion, "None", "Do not show root motion", EUserInterfaceActionType::RadioButton, FInputChord());
+	UI_COMMAND(VisualizeRootMotionTrajectory, "Visualize Trajectory", "Show root motion trajectory", EUserInterfaceActionType::RadioButton, FInputChord());
+	UI_COMMAND(VisualizeRootMotionTrajectoryAndOrientation, "Visualize Trajectory and Orientation", "Show root motion trajectory and orientation", EUserInterfaceActionType::RadioButton, FInputChord());
+
 	UI_COMMAND(DisablePostProcessBlueprint, "Disable Post Process", "Disable the evaluation of post process animation blueprints on the preview instance.", EUserInterfaceActionType::ToggleButton, FInputChord());
 
 	UI_COMMAND( ShowBound, "Bound", "Show bound on preview mesh", EUserInterfaceActionType::ToggleButton, FInputChord() );
@@ -73,7 +77,6 @@ void FAnimViewportShowCommands::RegisterCommands()
 	UI_COMMAND(HideOnlyClothSections, "Hide Only Cloth Sections", "Display all except cloth mapped sections", EUserInterfaceActionType::RadioButton, FInputChord());
 
 	UI_COMMAND(ShowNotificationVisualizations, "Show Notification", "Display animation notifications visualizations in viewport", EUserInterfaceActionType::ToggleButton, FInputChord());
-	UI_COMMAND(ShowRootMotionVisualization, "Show Root Motion", "Display root motion visualizations in viewport", EUserInterfaceActionType::ToggleButton, FInputChord());
 	UI_COMMAND(ShowAssetUserDataVisualizations, "Show Meta Data", "Display animation meta data visualizations in viewport", EUserInterfaceActionType::ToggleButton, FInputChord());
 
 	UI_COMMAND(ShowTimecode, "Show Timecode", "Display timecode information associated with animation data (if exists).", EUserInterfaceActionType::ToggleButton, FInputChord());
diff --git a/Engine/Source/Editor/Persona/Private/AnimViewportShowCommands.h b/Engine/Source/Editor/Persona/Private/AnimViewportShowCommands.h
--- a/Engine/Source/Editor/Persona/Private/AnimViewportShowCommands.h
+++ b/Engine/Source/Editor/Persona/Private/AnimViewportShowCommands.h
@@ -47,6 +47,11 @@ public:
 	/** Option to set ProcessRootMotionMode to Loop (Preview mesh will consume root motion continually) */
 	TSharedPtr< FUICommandInfo > ProcessRootMotionLoop;
 
+	/** Visualize root motion mode */
+	TSharedPtr< FUICommandInfo > DoNotVisualizeRootMotion;
+	TSharedPtr< FUICommandInfo > VisualizeRootMotionTrajectory;
+	TSharedPtr< FUICommandInfo > VisualizeRootMotionTrajectoryAndOrientation;
+
 	/** Option to enable/disable post process anim blueprint evaluation */
 	TSharedPtr< FUICommandInfo > DisablePostProcessBlueprint;
 	
@@ -156,9 +161,6 @@ public:
 	/** Show Notification Visualizations */
 	TSharedPtr< FUICommandInfo > ShowNotificationVisualizations;
 
-	/** Show Root motion visualization */
-	TSharedPtr< FUICommandInfo > ShowRootMotionVisualization;
-
 	/** Show AssetUserData Visualizations */
 	TSharedPtr< FUICommandInfo > ShowAssetUserDataVisualizations;
 
diff --git a/Engine/Source/Editor/Persona/Private/AnimationEditorViewportClient.cpp b/Engine/Source/Editor/Persona/Private/AnimationEditorViewportClient.cpp
--- a/Engine/Source/Editor/Persona/Private/AnimationEditorViewportClient.cpp
+++ b/Engine/Source/Editor/Persona/Private/AnimationEditorViewportClient.cpp
@@ -19,6 +19,9 @@
 #include "PersonaModule.h"
 #include "Rendering/SkeletalMeshRenderData.h"
 #include "GenericPlatform/GenericPlatformInputDeviceMapper.h"
+#include "DynamicMeshBuilder.h"
+#include "Materials/MaterialInstanceDynamic.h"
+#include "Materials/MaterialRenderProxy.h"
 
 #include "SEditorViewport.h"
 #include "CanvasTypes.h"
@@ -80,6 +83,88 @@ namespace UE::Private
 		}
 		return true;
 	}
+
+	void DrawCoordinateSystem(FPrimitiveDrawInterface* PDI, const FTransform& Transform, const float Thickness, const float Length, const float DepthBias, const bool bScreenSpace, uint8 Alpha)
+	{
+		const FVector Location = Transform.GetLocation();
+		const FVector AxisX = Transform.GetUnitAxis(EAxis::X) * Length;
+		const FVector AxisY = Transform.GetUnitAxis(EAxis::Y) * Length;
+		const FVector AxisZ = Transform.GetUnitAxis(EAxis::Z) * Length;
+		PDI->DrawTranslucentLine(Location, Location + AxisX, FColor::Red.WithAlpha(Alpha), SDPG_World, 1.0f, DepthBias, bScreenSpace);
+		PDI->DrawTranslucentLine(Location, Location + AxisY, FColor::Green.WithAlpha(Alpha), SDPG_World, 1.0f, DepthBias, bScreenSpace);
+		PDI->DrawTranslucentLine(Location, Location + AxisZ, FColor::Blue.WithAlpha(Alpha), SDPG_World, 1.0f, DepthBias, bScreenSpace);
+	}
+
+	FColor GetColorForAxis(EAxis::Type InAxis)
+	{
+		// Just draw all forward-axis versions as black for now.
+		return FColor::Black;
+	}
+
+	void DrawFlatArrow(class FPrimitiveDrawInterface* PDI,const FVector& Base,const FVector& XAxis,const FVector& YAxis,FColor Color,float Length,int32 Width, const FMaterialRenderProxy* MaterialRenderProxy, uint8 DepthPriority, float Thickness = 0.0f)
+	{
+		float DistanceFromBaseToHead = Length/3.0f;
+		float DistanceFromBaseToTip = DistanceFromBaseToHead*2.0f;
+		float WidthOfBase = Width;
+		float WidthOfHead = 2*Width;
+
+		FVector ArrowPoints[7];
+		//base points
+		ArrowPoints[0] = Base - YAxis*(WidthOfBase*.5f);
+		ArrowPoints[1] = Base + YAxis*(WidthOfBase*.5f);
+		//inner head
+		ArrowPoints[2] = ArrowPoints[0] + XAxis*DistanceFromBaseToHead;
+		ArrowPoints[3] = ArrowPoints[1] + XAxis*DistanceFromBaseToHead;
+		//outer head
+		ArrowPoints[4] = ArrowPoints[2] - YAxis*(WidthOfBase*.5f);
+		ArrowPoints[5] = ArrowPoints[3] + YAxis*(WidthOfBase*.5f);
+		//tip
+		ArrowPoints[6] = Base + XAxis*Length;
+
+		//Draw lines
+		{
+			//base
+			PDI->DrawTranslucentLine(ArrowPoints[0], ArrowPoints[1], Color, DepthPriority, Thickness);
+			//base sides																 
+			PDI->DrawTranslucentLine(ArrowPoints[0], ArrowPoints[2], Color, DepthPriority, Thickness);
+			PDI->DrawTranslucentLine(ArrowPoints[1], ArrowPoints[3], Color, DepthPriority, Thickness);
+			//head base																	 
+			PDI->DrawTranslucentLine(ArrowPoints[2], ArrowPoints[4], Color, DepthPriority, Thickness);
+			PDI->DrawTranslucentLine(ArrowPoints[3], ArrowPoints[5], Color, DepthPriority, Thickness);
+			//head sides																 
+			PDI->DrawTranslucentLine(ArrowPoints[4], ArrowPoints[6], Color, DepthPriority, Thickness);
+			PDI->DrawTranslucentLine(ArrowPoints[5], ArrowPoints[6], Color, DepthPriority, Thickness);
+
+		}
+
+		if (MaterialRenderProxy != nullptr)
+		{
+			FDynamicMeshBuilder MeshBuilder(PDI->View->GetFeatureLevel());
+
+			//Compute vertices for base circle.
+			for(int32 i = 0; i< 7; ++i)
+			{
+				FDynamicMeshVertex MeshVertex;
+				MeshVertex.Position = (FVector3f)ArrowPoints[i];
+				MeshVertex.Color = Color;
+				MeshVertex.TextureCoordinate[0] = FVector2f(0.0f, 0.0f);;
+				MeshVertex.SetTangents(FVector3f(XAxis^YAxis), (FVector3f)YAxis, (FVector3f)XAxis);
+				MeshBuilder.AddVertex(MeshVertex); //Add bottom vertex
+			}
+
+			//Add triangles / double sided
+			{
+				MeshBuilder.AddTriangle(0, 2, 1); //base
+				MeshBuilder.AddTriangle(0, 1, 2); //base
+				MeshBuilder.AddTriangle(1, 2, 3); //base
+				MeshBuilder.AddTriangle(1, 3, 2); //base
+				MeshBuilder.AddTriangle(4, 5, 6); //head
+				MeshBuilder.AddTriangle(4, 6, 5); //head
+			}
+
+			MeshBuilder.Draw(PDI, FMatrix::Identity, MaterialRenderProxy, DepthPriority, 0.0f);
+		}
+	}
 }
 
 #define LOCTEXT_NAMESPACE "FAnimationViewportClient"
@@ -1331,9 +1416,9 @@ FText FAnimationViewportClient::GetDisplayInfo(bool bDisplayAllInfo) const
 		TextValue = ConcatenateLine(TextValue, LOCTEXT("MeshMaterialHiddenWarning", "Mesh Materials Hidden"));
 	}
 
-	if (const UAnimSequence* AnimSequence = Cast<UAnimSequence>(GetAnimPreviewScene()->GetPreviewAnimationAsset()))
+	if (const UAnimSequenceBase* AnimSequenceBase = Cast<UAnimSequenceBase>(GetAnimPreviewScene()->GetPreviewAnimationAsset()))
 	{
-		TextValue = ConcatenateLine(TextValue, FText::Format(LOCTEXT("FramerateFormat", "Framerate: {0}"), AnimSequence->GetSamplingFrameRate().ToPrettyText()));
+		TextValue = ConcatenateLine(TextValue, FText::Format(LOCTEXT("FramerateFormat", "Framerate: {0}"), AnimSequenceBase->GetSamplingFrameRate().ToPrettyText()));
 	}
 
 	if (const UPoseAsset* PoseAsset = Cast<UPoseAsset>(GetAnimPreviewScene()->GetPreviewAnimationAsset()))
@@ -2018,10 +2103,12 @@ void FAnimationViewportClient::DrawRootMotionTrajectory(UDebugSkelMeshComponent*
 	constexpr bool bScreenSpace = true;
 
 	if (MeshComponent
-		&& MeshComponent->IsRootMotionVisualizationsEnabled()
+		&& !MeshComponent->IsVisualizeRootMotionMode(EVisualizeRootMotionMode::None)
 		&& MeshComponent->GetSkeletalMeshAsset()
 		&& MeshComponent->DoesCurrentAssetHaveRootMotion())
 	{
+		const EVisualizeRootMotionMode VisMode = MeshComponent->GetVisualizeRootMotionMode();
+
 		const FTransform& ReferenceTransform = MeshComponent->RootMotionReferenceTransform;
 		const UMirrorDataTable* MirrorTable = MeshComponent->PreviewInstance ? MeshComponent->PreviewInstance->GetMirrorDataTable() : nullptr;
 		
@@ -2030,19 +2117,39 @@ void FAnimationViewportClient::DrawRootMotionTrajectory(UDebugSkelMeshComponent*
 			// Draw root motion trajectory
 			const int32 NumFrames = AnimSequenceBase->GetNumberOfSampledKeys();
 			const FFrameRate FrameRate = AnimSequenceBase->GetSamplingFrameRate();
+			const float CurrentTime = MeshComponent->GetPosition();
+			const USkeleton* Skeleton = AnimSequenceBase->GetSkeleton();
+			check(Skeleton);
+			EAxis::Type SkeletonForwardAxis = Skeleton->GetForwardAxis();
 
 			const FColor TrajectoryColor = FColor::Black.WithAlpha(64);
 
 			FVector PrevLocation;
+			float PlayLength = AnimSequenceBase->GetPlayLength();
+			PlayLength = AnimSequenceBase->GetPlayLength();
 			for (int32 Frame = 0; Frame <= NumFrames; Frame++)
 			{
-				const double Time = FMath::Clamp(FrameRate.AsSeconds(Frame), 0., (double)AnimSequenceBase->GetPlayLength());
+				const double Time = FMath::Clamp(FrameRate.AsSeconds(Frame), 0., (double)PlayLength);
 				const FTransform Transform = UE::Anim::ExtractRootMotionFromAnimationAsset(AnimSequenceBase, MirrorTable, 0.0, Time) * ReferenceTransform;
 				const FVector Location = Transform.GetLocation();
 
 				const bool bFirstOrLastPoint = Frame == 0 || Frame == NumFrames;
-				PDI->DrawPoint(Location, TrajectoryColor, bFirstOrLastPoint ? 2.0f : 1.0f, SDPG_World);
-				
+
+				PDI->DrawPoint(Location, TrajectoryColor, bFirstOrLastPoint ? 2.5f : 1.25f, SDPG_World);
+
+				if (VisMode == EVisualizeRootMotionMode::TrajectoryAndOrientation)
+				{
+					if (bFirstOrLastPoint || (Frame % 3 == 0))
+					{
+						const FVector XAxis = Transform.GetUnitAxis(SkeletonForwardAxis);
+						const FColor AxisColor = UE::Private::GetColorForAxis(SkeletonForwardAxis);
+
+						FVector YAxis, ZAxis;
+						XAxis.FindBestAxisVectors(YAxis,ZAxis);
+						UE::Private::DrawFlatArrow(PDI, Transform.GetLocation(), XAxis, ZAxis, AxisColor.WithAlpha(64), 15.0f, 8, nullptr, SDPG_World, 1.0f);
+					}
+				}
+
 				if (Frame > 0)
 				{
 					PDI->DrawTranslucentLine(PrevLocation, Location, TrajectoryColor, SDPG_World, 1.0f, DepthBias, bScreenSpace);
@@ -2052,18 +2159,19 @@ void FAnimationViewportClient::DrawRootMotionTrajectory(UDebugSkelMeshComponent*
 
 			// Draw current location on the root motion.
 			{
-				const float CurrentTime = MeshComponent->GetPosition();
 				const FTransform Transform = UE::Anim::ExtractRootMotionFromAnimationAsset(AnimSequenceBase, MirrorTable, 0.0, CurrentTime) * ReferenceTransform;
 
-				constexpr double CurrentPosTickSize = 5.0;
+				const FVector XAxis = Transform.GetUnitAxis(SkeletonForwardAxis);
+				const FColor AxisColor = UE::Private::GetColorForAxis(SkeletonForwardAxis);
 
-				const FVector Location = Transform.GetLocation();
-				const FVector AxisX = Transform.GetUnitAxis(EAxis::X) * CurrentPosTickSize;
-				const FVector AxisY = Transform.GetUnitAxis(EAxis::Y) * CurrentPosTickSize;
-				const FVector AxisZ = Transform.GetUnitAxis(EAxis::Z) * CurrentPosTickSize;
-				PDI->DrawTranslucentLine(Location, Location + AxisX, FColor::Red.WithAlpha(128), SDPG_World, 1.0f, DepthBias, bScreenSpace);
-				PDI->DrawTranslucentLine(Location, Location + AxisY, FColor::Green.WithAlpha(128), SDPG_World, 1.0f, DepthBias, bScreenSpace);
-				PDI->DrawTranslucentLine(Location, Location + AxisZ, FColor::Blue.WithAlpha(128), SDPG_World, 1.0f, DepthBias, bScreenSpace);
+				FVector YAxis, ZAxis;
+				XAxis.FindBestAxisVectors(YAxis,ZAxis);
+
+				if (VisMode == EVisualizeRootMotionMode::TrajectoryAndOrientation)
+				{
+					UE::Private::DrawFlatArrow(PDI, Transform.GetLocation(), XAxis, ZAxis, AxisColor, 30.0f, 15, GEngine->ArrowMaterialYellow->GetRenderProxy(), SDPG_Foreground, 1.0f);
+				}
+				UE::Private::DrawCoordinateSystem(PDI, Transform, 10.0f, 20.0f, DepthBias, bScreenSpace, 200);
 			}			
 		}
 	}
diff --git a/Engine/Source/Editor/Persona/Private/SAnimViewportToolBar.cpp b/Engine/Source/Editor/Persona/Private/SAnimViewportToolBar.cpp
--- a/Engine/Source/Editor/Persona/Private/SAnimViewportToolBar.cpp
+++ b/Engine/Source/Editor/Persona/Private/SAnimViewportToolBar.cpp
@@ -764,7 +764,9 @@ TSharedRef<SWidget> SAnimViewportToolBar::GenerateCharacterMenu() const
 					{
 						FToolMenuSection& Section = InSubMenu->AddSection("AnimViewportVisualization", LOCTEXT("CharacterMenu_VisualizationsLabel", "Visualizations"));
 						Section.AddMenuEntry(FAnimViewportShowCommands::Get().ShowNotificationVisualizations);
-						Section.AddMenuEntry(FAnimViewportShowCommands::Get().ShowRootMotionVisualization);
+						Section.AddMenuEntry(FAnimViewportShowCommands::Get().DoNotVisualizeRootMotion);
+						Section.AddMenuEntry(FAnimViewportShowCommands::Get().VisualizeRootMotionTrajectory);
+						Section.AddMenuEntry(FAnimViewportShowCommands::Get().VisualizeRootMotionTrajectoryAndOrientation);
 						Section.AddMenuEntry(FAnimViewportShowCommands::Get().ShowAssetUserDataVisualizations);
 					}
 
diff --git a/Engine/Source/Editor/Persona/Private/SAnimationEditorViewport.cpp b/Engine/Source/Editor/Persona/Private/SAnimationEditorViewport.cpp
--- a/Engine/Source/Editor/Persona/Private/SAnimationEditorViewport.cpp
+++ b/Engine/Source/Editor/Persona/Private/SAnimationEditorViewport.cpp
@@ -1112,18 +1112,30 @@ void SAnimationEditorViewportTabBody::BindCommands()
 		FIsActionChecked::CreateSP(this, &SAnimationEditorViewportTabBody::CanUseProcessRootMotionMode, EProcessRootMotionMode::Loop),
 		FIsActionChecked::CreateSP(this, &SAnimationEditorViewportTabBody::IsProcessRootMotionModeSet, EProcessRootMotionMode::Loop));
 
+	CommandList.MapAction(
+		ViewportShowMenuCommands.DoNotVisualizeRootMotion,
+		FExecuteAction::CreateSP(this, &SAnimationEditorViewportTabBody::SetVisualizeRootMotionMode, EVisualizeRootMotionMode::None),
+		FIsActionChecked::CreateSP(this, &SAnimationEditorViewportTabBody::CanVisualizeRootMotion),
+		FIsActionChecked::CreateSP(this, &SAnimationEditorViewportTabBody::IsVisualizeRootMotionModeSet, EVisualizeRootMotionMode::None));
+
+	CommandList.MapAction(
+		ViewportShowMenuCommands.VisualizeRootMotionTrajectory,
+		FExecuteAction::CreateSP(this, &SAnimationEditorViewportTabBody::SetVisualizeRootMotionMode, EVisualizeRootMotionMode::Trajectory),
+		FIsActionChecked::CreateSP(this, &SAnimationEditorViewportTabBody::CanVisualizeRootMotion),
+		FIsActionChecked::CreateSP(this, &SAnimationEditorViewportTabBody::IsVisualizeRootMotionModeSet, EVisualizeRootMotionMode::Trajectory));
+
+	CommandList.MapAction(
+		ViewportShowMenuCommands.VisualizeRootMotionTrajectoryAndOrientation,
+		FExecuteAction::CreateSP(this, &SAnimationEditorViewportTabBody::SetVisualizeRootMotionMode, EVisualizeRootMotionMode::TrajectoryAndOrientation),
+		FIsActionChecked::CreateSP(this, &SAnimationEditorViewportTabBody::CanVisualizeRootMotion),
+		FIsActionChecked::CreateSP(this, &SAnimationEditorViewportTabBody::IsVisualizeRootMotionModeSet, EVisualizeRootMotionMode::TrajectoryAndOrientation));
+
 	CommandList.MapAction(
 		ViewportShowMenuCommands.ShowNotificationVisualizations,
 		FExecuteAction::CreateSP(this, &SAnimationEditorViewportTabBody::ToggleNotificationVisualizations),
 		FIsActionChecked(),
 		FIsActionChecked::CreateSP(this, &SAnimationEditorViewportTabBody::IsNotificationVisualizationsEnabled));
 
-	CommandList.MapAction(
-		ViewportShowMenuCommands.ShowRootMotionVisualization,
-		FExecuteAction::CreateSP(this, &SAnimationEditorViewportTabBody::ToggleRootMotionVisualizations),
-		FIsActionChecked(),
-		FIsActionChecked::CreateSP(this, &SAnimationEditorViewportTabBody::IsRootMotionVisualizationsEnabled));
-
 	CommandList.MapAction(
 		ViewportShowMenuCommands.ShowAssetUserDataVisualizations,
 		FExecuteAction::CreateSP(this, &SAnimationEditorViewportTabBody::ToggleAssetUserDataVisualizations),
@@ -2153,36 +2165,47 @@ bool SAnimationEditorViewportTabBody::CanUseProcessRootMotionMode(EProcessRootMo
 	return false;
 }
 
-void SAnimationEditorViewportTabBody::ToggleNotificationVisualizations()
+void SAnimationEditorViewportTabBody::SetVisualizeRootMotionMode(EVisualizeRootMotionMode Mode)
 {
-	if (UDebugSkelMeshComponent* PreviewComponent = GetPreviewScene()->GetPreviewMeshComponent())
+	if(UDebugSkelMeshComponent* PreviewComponent = GetPreviewScene()->GetPreviewMeshComponent())
 	{
-		PreviewComponent->SetShowNotificationVisualizations(!PreviewComponent->IsNotificationVisualizationsEnabled());
+		return PreviewComponent->SetVisualizeRootMotionMode(Mode);
 	}
 }
 
-bool SAnimationEditorViewportTabBody::IsNotificationVisualizationsEnabled() const
+bool SAnimationEditorViewportTabBody::IsVisualizeRootMotionModeSet(EVisualizeRootMotionMode Mode) const
 {
-	if (const UDebugSkelMeshComponent* PreviewComponent = GetPreviewScene()->GetPreviewMeshComponent())
+	if(const UDebugSkelMeshComponent* PreviewComponent = GetPreviewScene()->GetPreviewMeshComponent())
 	{
-		return PreviewComponent->IsNotificationVisualizationsEnabled();
+		return PreviewComponent->IsVisualizeRootMotionMode(Mode);
+	}
+
+	return false;
+}
+
+bool SAnimationEditorViewportTabBody::CanVisualizeRootMotion() const
+{
+	if(const UDebugSkelMeshComponent* PreviewComponent = GetPreviewScene()->GetPreviewMeshComponent())
+	{
+		return PreviewComponent->DoesCurrentAssetHaveRootMotion();
 	}
+
 	return false;
 }
 
-void SAnimationEditorViewportTabBody::ToggleRootMotionVisualizations()
+void SAnimationEditorViewportTabBody::ToggleNotificationVisualizations()
 {
 	if (UDebugSkelMeshComponent* PreviewComponent = GetPreviewScene()->GetPreviewMeshComponent())
 	{
-		PreviewComponent->SetShowRootMotionVisualizations(!PreviewComponent->IsRootMotionVisualizationsEnabled());
+		PreviewComponent->SetShowNotificationVisualizations(!PreviewComponent->IsNotificationVisualizationsEnabled());
 	}
 }
 
-bool SAnimationEditorViewportTabBody::IsRootMotionVisualizationsEnabled() const
+bool SAnimationEditorViewportTabBody::IsNotificationVisualizationsEnabled() const
 {
 	if (const UDebugSkelMeshComponent* PreviewComponent = GetPreviewScene()->GetPreviewMeshComponent())
 	{
-		return PreviewComponent->IsRootMotionVisualizationsEnabled();
+		return PreviewComponent->IsNotificationVisualizationsEnabled();
 	}
 	return false;
 }
diff --git a/Engine/Source/Editor/Persona/Private/SAnimationEditorViewport.h b/Engine/Source/Editor/Persona/Private/SAnimationEditorViewport.h
--- a/Engine/Source/Editor/Persona/Private/SAnimationEditorViewport.h
+++ b/Engine/Source/Editor/Persona/Private/SAnimationEditorViewport.h
@@ -526,12 +526,13 @@ private:
 	/** Whether the supplied mode can be used */
 	bool CanUseProcessRootMotionMode(EProcessRootMotionMode Mode) const;
 
+	void SetVisualizeRootMotionMode(EVisualizeRootMotionMode Mode);
+	bool IsVisualizeRootMotionModeSet(EVisualizeRootMotionMode Mode) const;
+	bool CanVisualizeRootMotion() const;
+
 	void ToggleNotificationVisualizations();
 	bool IsNotificationVisualizationsEnabled() const;
 
-	void ToggleRootMotionVisualizations();
-	bool IsRootMotionVisualizationsEnabled() const;
-
 	void ToggleAssetUserDataVisualizations();
 	bool IsAssetUserDataVisualizationsEnabled() const;
 
diff --git a/Engine/Source/Editor/UnrealEd/Classes/Animation/DebugSkelMeshComponent.h b/Engine/Source/Editor/UnrealEd/Classes/Animation/DebugSkelMeshComponent.h
--- a/Engine/Source/Editor/UnrealEd/Classes/Animation/DebugSkelMeshComponent.h
+++ b/Engine/Source/Editor/UnrealEd/Classes/Animation/DebugSkelMeshComponent.h
@@ -80,6 +80,20 @@ enum class EProcessRootMotionMode : uint8
 	LoopAndReset
 };
 
+/** Different modes for visualizing root motion */
+UENUM()
+enum class EVisualizeRootMotionMode : uint8
+{
+	/** Preview will not show root motion */
+	None,
+
+	/** Preview will show root motion path */
+	Trajectory,
+
+	/** Preview will show root motion path and orientation. */
+	TrajectoryAndOrientation
+};
+
 //////////////////////////////////////////////////////////////////////////
 // FDebugSkelMeshSceneProxy
 
@@ -238,10 +252,6 @@ class UDebugSkelMeshComponent : public USkeletalMeshComponent
 	UPROPERTY(Transient)
 	uint32 bShowNotificationVisualizations:1;
 
-	/** Display Root motion visualizations in viewport */
-	UPROPERTY(Transient)
-	uint32 bShowRootMotionVisualizations:1;
-
 	/** Display Metadata visualizations in viewport */
 	UPROPERTY(Transient)
 	uint32 bShowAssetUserDataVisualizations:1;
@@ -268,6 +278,9 @@ class UDebugSkelMeshComponent : public USkeletalMeshComponent
 	UPROPERTY(transient)
 	EProcessRootMotionMode ProcessRootMotionMode;
 
+	UPROPERTY(transient)
+	EVisualizeRootMotionMode VisualizeRootMotionMode;
+
 	/** Playback time last time ConsumeRootmotion was called */
 	UPROPERTY(transient)
 	float ConsumeRootMotionPreviousPlaybackTime;
@@ -482,6 +495,11 @@ class UDebugSkelMeshComponent : public USkeletalMeshComponent
 	/** Sets process root motion mode, the request may be ignored if current asset does not support the mode. Note: disabling root motion preview resets transform. */
 	UNREALED_API void SetProcessRootMotionMode(EProcessRootMotionMode Mode);
 
+	/** Sets how we visualize root motion in the viewport. See EVisualizeRootMotionMode for details. */
+	UNREALED_API void SetVisualizeRootMotionMode(EVisualizeRootMotionMode Mode) { VisualizeRootMotionMode = Mode; };
+	UNREALED_API bool IsVisualizeRootMotionMode(EVisualizeRootMotionMode Mode) const { return VisualizeRootMotionMode == Mode; };
+	UNREALED_API EVisualizeRootMotionMode GetVisualizeRootMotionMode() const { return VisualizeRootMotionMode; };
+
 	/** Whether the supplied root motion mode can be used for the current asset */
 	UNREALED_API bool CanUseProcessRootMotionMode(EProcessRootMotionMode Mode) const;
 
@@ -492,9 +510,9 @@ class UDebugSkelMeshComponent : public USkeletalMeshComponent
 	UNREALED_API void SetShowNotificationVisualizations(const bool bShow) { bShowNotificationVisualizations = bShow; }
 	UNREALED_API bool IsNotificationVisualizationsEnabled() const { return bShowNotificationVisualizations; }
 
-	/** Sets flags whether we root motion visualization should be drawn in the viewport. */
-	UNREALED_API void SetShowRootMotionVisualizations(const bool bShow) { bShowRootMotionVisualizations = bShow; }
-	UNREALED_API bool IsRootMotionVisualizationsEnabled() const { return bShowRootMotionVisualizations; }
+	UE_DEPRECATED(5.5, "Use VisualizeRootMotionMode functions instead.")
+	UNREALED_API void SetShowRootMotionVisualizations(const bool bShow) { VisualizeRootMotionMode = bShow ? EVisualizeRootMotionMode::Trajectory : EVisualizeRootMotionMode::None; }
+	UNREALED_API bool IsRootMotionVisualizationsEnabled() const { return VisualizeRootMotionMode != EVisualizeRootMotionMode::None; }
 
 	/** Sets flags whether we AssetUserData visualizations should be drawn in the viewport. */
 	UNREALED_API void SetShowAssetUserDataVisualizations(const bool bShow) { bShowAssetUserDataVisualizations = bShow; }
diff --git a/Engine/Source/Editor/UnrealEd/Private/Animation/DebugSkelMeshComponent.cpp b/Engine/Source/Editor/UnrealEd/Private/Animation/DebugSkelMeshComponent.cpp
--- a/Engine/Source/Editor/UnrealEd/Private/Animation/DebugSkelMeshComponent.cpp
+++ b/Engine/Source/Editor/UnrealEd/Private/Animation/DebugSkelMeshComponent.cpp
@@ -43,7 +43,7 @@ UDebugSkelMeshComponent::UDebugSkelMeshComponent(const FObjectInitializer& Objec
 	bSkeletonSocketsVisible = true;
 
 	bShowNotificationVisualizations = true;
-	bShowRootMotionVisualizations = false;
+	VisualizeRootMotionMode = EVisualizeRootMotionMode::Trajectory;
 	bShowAssetUserDataVisualizations = true;
 
 	TurnTableSpeedScaling = 1.f;
diff --git a/Engine/Source/Runtime/Engine/Classes/Animation/Skeleton.h b/Engine/Source/Runtime/Engine/Classes/Animation/Skeleton.h
--- a/Engine/Source/Runtime/Engine/Classes/Animation/Skeleton.h
+++ b/Engine/Source/Runtime/Engine/Classes/Animation/Skeleton.h
@@ -307,6 +307,10 @@ protected:
 	/** Reference skeleton poses in local space */
 	UPROPERTY()
 	TArray<FTransform> RefLocalPoses_DEPRECATED;
+
+	// Preview axis to consider as "forward" for the skeleton. Only used for preview purposes.
+	UPROPERTY(EditAnywhere, Category = Preview)
+	TEnumAsByte<EAxis::Type> ForwardAxis;
 #endif
 
 	/** Reference Skeleton */
@@ -345,6 +349,7 @@ public:
 #if WITH_EDITOR
 	ENGINE_API virtual void PreEditUndo() override;
 	ENGINE_API virtual void PostEditUndo() override;
+	ENGINE_API EAxis::Type GetForwardAxis() const { return ForwardAxis; }
 #endif
 	ENGINE_API virtual void BeginDestroy() override;
 
diff --git a/Engine/Source/Runtime/Engine/Private/Animation/Skeleton.cpp b/Engine/Source/Runtime/Engine/Private/Animation/Skeleton.cpp
--- a/Engine/Source/Runtime/Engine/Private/Animation/Skeleton.cpp
+++ b/Engine/Source/Runtime/Engine/Private/Animation/Skeleton.cpp
@@ -335,6 +335,10 @@ USkeleton::USkeleton(const FObjectInitializer& ObjectInitializer)
 {
 	CachedSoftObjectPtr = TSoftObjectPtr<USkeleton>(this);
 
+#if WITH_EDITORONLY_DATA
+	ForwardAxis = EAxis::Y;
+#endif
+
 	if (HasAnyFlags(RF_ClassDefaultObject))
 	{
 		FCoreUObjectDelegates::OnPackageReloaded.AddStatic(&USkeleton::HandlePackageReloaded);
