Improve animation blueprint navigation speed when several animation blueprints are opened in the editor. Now only compatible (related) assets will be refreshed in persona when we double click a node and a new asset is opened in the editor.

diff --git a/Engine/Source/Editor/Persona/Private/SAssetFamilyShortcutBar.cpp b/Engine/Source/Editor/Persona/Private/SAssetFamilyShortcutBar.cpp
--- a/Engine/Source/Editor/Persona/Private/SAssetFamilyShortcutBar.cpp
+++ b/Engine/Source/Editor/Persona/Private/SAssetFamilyShortcutBar.cpp
@@ -405,7 +405,10 @@ public:
 
 	void HandleAssetOpened(UObject* InAsset)
 	{
-		RefreshAsset();
+		if (AssetFamily.IsValid() && AssetFamily->IsAssetCompatible(FAssetData(InAsset)))
+		{
+			RefreshAsset();
+		}
 	}
 
 	EVisibility GetThumbnailVisibility() const
@@ -443,12 +446,15 @@ public:
 			}
 
 			// switch to new asset if needed
-			FAssetData NewAssetData = AssetFamily->FindAssetOfType(AssetData.GetClass());
-			if (!bAssetBeingEdited && NewAssetData != AssetData)
+			if (!bAssetBeingEdited)
 			{
-				AssetData = NewAssetData;
+				FAssetData NewAssetData = AssetFamily->FindAssetOfType(AssetData.GetClass());
+				if (NewAssetData != AssetData)
+				{
+					AssetData = NewAssetData;
 
-				RegenerateThumbnail();
+					RegenerateThumbnail();
+				}
 			}
 		}
 	}
