Refresh notify list after Undo / Redo on animation sequences, to avoid crashes caused by active tooltips

diff --git a/Engine/Source/Editor/Persona/Private/SAnimNotifyPanel.cpp b/Engine/Source/Editor/Persona/Private/SAnimNotifyPanel.cpp
--- a/Engine/Source/Editor/Persona/Private/SAnimNotifyPanel.cpp
+++ b/Engine/Source/Editor/Persona/Private/SAnimNotifyPanel.cpp
@@ -4359,6 +4359,7 @@ void SAnimNotifyPanel::PostUndo( bool bSuccess )
 	{
 		Sequence->RefreshCacheData();
 	}
+	RefreshNotifyTracks();
 }
 
 void SAnimNotifyPanel::PostRedo( bool bSuccess )
@@ -4367,6 +4368,7 @@ void SAnimNotifyPanel::PostRedo( bool bSuccess )
 	{
 		Sequence->RefreshCacheData();
 	}
+	RefreshNotifyTracks();
 }
 
 void SAnimNotifyPanel::OnDeletePressed()
