Fixed Regenerate LOD button not getting enabled in the Skeletal Mesh Editor under some situations, due to Build GUID getting overriden at LODInfoUILayout  PostEditChangeProperty, which disalows detecting changes to the data

diff --git a/Engine/Source/Editor/Persona/Private/LODInfoUILayout.cpp b/Engine/Source/Editor/Persona/Private/LODInfoUILayout.cpp
--- a/Engine/Source/Editor/Persona/Private/LODInfoUILayout.cpp
+++ b/Engine/Source/Editor/Persona/Private/LODInfoUILayout.cpp
@@ -54,8 +54,9 @@ void ULODInfoUILayout::PostEditChangeProperty(struct FPropertyChangedEvent& Prop
 		FSkeletalMeshLODInfo* SkeletalMeshLODInfo = SkeletalMesh->GetLODInfo(LODIndex);
 		check(SkeletalMeshLODInfo != nullptr);
 		SkeletalMesh->Modify();
-		//Copy the LODInfo into the real skeletal mesh LODInfo data
-		*SkeletalMeshLODInfo = LODInfo;
+		const FGuid OldBuildGuid = SkeletalMeshLODInfo->BuildGUID; // Save the current Guid
+		*SkeletalMeshLODInfo = LODInfo; //Copy the LODInfo into the real skeletal mesh LODInfo data
+		SkeletalMeshLODInfo->BuildGUID = OldBuildGuid;	// Restore the saved Guid, as it is used to detect changes in the LODInfo data
 		FScopedSkeletalMeshPostEditChange ScopeSkeletalmeshPostEdit(SkeletalMesh);
 	}
 }
