Skeletal mesh editor now uses the required bones to display the skeleton tree, so the bones stays in sync with the preview scene. This will make clear which bones are excluded or not, as the required bones might have some bones that can not be removed from the LOD due to physics or other various needs.

diff --git a/Engine/Source/Editor/Persona/Private/AnimationEditorPreviewScene.h b/Engine/Source/Editor/Persona/Private/AnimationEditorPreviewScene.h
--- a/Engine/Source/Editor/Persona/Private/AnimationEditorPreviewScene.h
+++ b/Engine/Source/Editor/Persona/Private/AnimationEditorPreviewScene.h
@@ -79,6 +79,14 @@ public:
 		OnLODChanged.RemoveAll(Thing);
 	}
 
+	virtual void BroadcastOnLODChanged() override
+	{
+		if (OnLODChanged.IsBound())
+		{
+			OnLODChanged.Broadcast();
+		}
+	}
+
 	virtual void RegisterOnMorphTargetsChanged(const FSimpleDelegate& Delegate) override
 	{
 		OnMorphTargetsChanged.Add(Delegate);
diff --git a/Engine/Source/Editor/Persona/Private/PersonaMeshDetails.cpp b/Engine/Source/Editor/Persona/Private/PersonaMeshDetails.cpp
--- a/Engine/Source/Editor/Persona/Private/PersonaMeshDetails.cpp
+++ b/Engine/Source/Editor/Persona/Private/PersonaMeshDetails.cpp
@@ -344,42 +344,49 @@ private:
 		{
 			USkeletalMesh* SkeletalMesh = SharedToolkit->GetMesh();
 
-			if (!SkeletalMesh || SkeletalMesh->IsCompiling())
-			{
-				return false;
-			}
-			FSkeletalMeshLODInfo* LODInfo = SkeletalMesh->GetLODInfo(LODIndex);
-			if (LODInfo == nullptr)
-			{
-				return false;
-			}
-			bool bValidLODSettings = false;
-			if (SkeletalMesh->GetLODSettings() != nullptr)
-			{
-				const int32 NumSettings = FMath::Min(SkeletalMesh->GetLODSettings()->GetNumberOfSettings(), SkeletalMesh->GetLODNum());
-				if (LODIndex < NumSettings)
-				{
-					bValidLODSettings = true;
-				}
-			}
-			
-			const FSkeletalMeshLODGroupSettings* SkeletalMeshLODGroupSettings = bValidLODSettings ? &SkeletalMesh->GetLODSettings()->GetSettingsForLODLevel(LODIndex) : nullptr;
+			return LODSettingsChanged(SkeletalMesh, LODIndex);
+		}
+		return false;
+	}
 
-			FGuid BuildGUID = LODInfo->ComputeDeriveDataCacheKey(SkeletalMeshLODGroupSettings);
-			if (LODInfo->BuildGUID != BuildGUID)
-			{
-				return true;
-			}
-			else if(!SkeletalMesh->GetImportedModel() || !(SkeletalMesh->GetImportedModel()->LODModels.IsValidIndex(LODIndex)))
+public:
+	static bool LODSettingsChanged(USkeletalMesh* SkeletalMesh, int32 LODIndex)
+	{
+		if (!SkeletalMesh || SkeletalMesh->IsCompiling())
+		{
+			return false;
+		}
+		FSkeletalMeshLODInfo* LODInfo = SkeletalMesh->GetLODInfo(LODIndex);
+		if (LODInfo == nullptr)
+		{
+			return false;
+		}
+		bool bValidLODSettings = false;
+		if (SkeletalMesh->GetLODSettings() != nullptr)
+		{
+			const int32 NumSettings = FMath::Min(SkeletalMesh->GetLODSettings()->GetNumberOfSettings(), SkeletalMesh->GetLODNum());
+			if (LODIndex < NumSettings)
 			{
-				//If there is no valid LODIndex imported model we want to return false to force a build to happen
-				return false;
+				bValidLODSettings = true;
 			}
-			return SkeletalMesh->GetImportedModel()->LODModels[LODIndex].BuildStringID != SkeletalMesh->GetImportedModel()->LODModels[LODIndex].GetLODModelDeriveDataKey();
 		}
-		return false;
+
+		const FSkeletalMeshLODGroupSettings* SkeletalMeshLODGroupSettings = bValidLODSettings ? &SkeletalMesh->GetLODSettings()->GetSettingsForLODLevel(LODIndex) : nullptr;
+
+		FGuid BuildGUID = LODInfo->ComputeDeriveDataCacheKey(SkeletalMeshLODGroupSettings);
+		if (LODInfo->BuildGUID != BuildGUID)
+		{
+			return true;
+		}
+		else if (!SkeletalMesh->GetImportedModel() || !(SkeletalMesh->GetImportedModel()->LODModels.IsValidIndex(LODIndex)))
+		{
+			//If there is no valid LODIndex imported model we want to return false to force a build to happen
+			return false;
+		}
+		return SkeletalMesh->GetImportedModel()->LODModels[LODIndex].BuildStringID != SkeletalMesh->GetImportedModel()->LODModels[LODIndex].GetLODModelDeriveDataKey();
 	}
 
+private:
 	// Incoming arg data
 	int32 LODIndex;
 	TWeakPtr<IPersonaToolkit> PersonaToolkit;
@@ -4151,6 +4158,11 @@ FReply FPersonaMeshDetails::ApplyLODChanges(int32 LODIndex)
 	{
 		GetPersonaToolkit()->GetPreviewScene()->BroadcastOnMorphTargetsChanged();
 	}
+
+	// Update mesh required bones to match the new LOD settings
+	GetPersonaToolkit()->GetPreviewScene()->GetPreviewMeshComponent()->RecalcRequiredBones(LODIndex);
+	// Force an update (this will update the Skeleton Tree to match the bones in the preview skeletal mesh required bones)
+	GetPersonaToolkit()->GetPreviewScene()->BroadcastOnLODChanged();
 	
 	return FReply::Handled();
 }
@@ -4361,6 +4373,10 @@ void FPersonaMeshDetails::ApplyChanges()
 	USkeletalMesh* SkelMesh = GetPersonaToolkit()->GetMesh();
 	check(SkelMesh);
 
+	const int32 CurrentDisplayLOD = GetCurrentDisplayLODIndex(GetPersonaToolkit()->GetPreviewMeshComponent());
+	// Check if there are any LOD changes before we regenerate, as we will have to send a LOD update event to fully refresh the bone tree state
+	const bool bRequiresLODChangeUpdate = SSkeletalLODActions::LODSettingsChanged(SkelMesh, CurrentDisplayLOD);
+
 	if (NaniteSettings.IsValid())
 	{
 		NaniteSettings->ApplyChanges();
@@ -4420,6 +4436,14 @@ void FPersonaMeshDetails::ApplyChanges()
 		GetPersonaToolkit()->GetPreviewScene()->BroadcastOnMorphTargetsChanged();
 	}
 
+	// Update preview mesh required bones, in case LOD settings have changed
+	if (bRequiresLODChangeUpdate)
+	{
+		GetPersonaToolkit()->GetPreviewScene()->GetPreviewMeshComponent()->RecalcRequiredBones(CurrentDisplayLOD);
+		// Force an update (this will update the Skeleton Tree to match the bones in the preview skeletal mesh required bones)
+		GetPersonaToolkit()->GetPreviewScene()->BroadcastOnLODChanged();
+	}
+
 	// We need to rebuild the layout to ensure all the LOD details are up-to-date.
 	ForceLayoutRebuild();
 }
@@ -5756,13 +5780,23 @@ void FPersonaMeshDetails::UpdateLODCategoryVisibility() const
 
 FText FPersonaMeshDetails::GetCurrentLodName() const
 {
-	bool bAutoLod = false;
-	if (GetPersonaToolkit()->GetPreviewMeshComponent() != nullptr)
+	const int32 CurrentDisplayLOD = GetCurrentDisplayLODIndex(GetPersonaToolkit()->GetPreviewMeshComponent());
+	return FText::FromString(GetAutoLod(GetPersonaToolkit()->GetPreviewMeshComponent()) ? FString(TEXT("Auto (LOD0)")) : (FString(TEXT("LOD")) + FString::FromInt(CurrentDisplayLOD)));
+}
+
+bool FPersonaMeshDetails::GetAutoLod(USkeletalMeshComponent* InMeshComponent)
+{
+	return (InMeshComponent != nullptr) ? InMeshComponent->GetForcedLOD() == 0 : false;
+}
+
+int32 FPersonaMeshDetails::GetCurrentDisplayLODIndex(USkeletalMeshComponent* InMeshComponent)
+{
+	int32 CurrentDisplayLOD = 0;
+	if (InMeshComponent != nullptr)
 	{
-		bAutoLod = GetPersonaToolkit()->GetPreviewMeshComponent()->GetForcedLOD() == 0;
+		CurrentDisplayLOD = GetAutoLod(InMeshComponent) ? 0 : InMeshComponent->GetForcedLOD() - 1;
 	}
-	int32 CurrentDisplayLOD = bAutoLod ? 0 : GetPersonaToolkit()->GetPreviewMeshComponent()->GetForcedLOD() - 1;
-	return FText::FromString(bAutoLod ? FString(TEXT("Auto (LOD0)")) : (FString(TEXT("LOD")) + FString::FromInt(CurrentDisplayLOD)));
+	return CurrentDisplayLOD;
 }
 
 FText FPersonaMeshDetails::GetCurrentLodTooltip() const
diff --git a/Engine/Source/Editor/Persona/Private/PersonaMeshDetails.h b/Engine/Source/Editor/Persona/Private/PersonaMeshDetails.h
--- a/Engine/Source/Editor/Persona/Private/PersonaMeshDetails.h
+++ b/Engine/Source/Editor/Persona/Private/PersonaMeshDetails.h
@@ -447,6 +447,8 @@ private:
 	TSharedRef<SWidget> OnGenerateLodMenuForLodPicker();
 	FText GetCurrentLodName() const;
 	FText GetCurrentLodTooltip() const;
+	static bool GetAutoLod(USkeletalMeshComponent* InMeshComponent);
+	static int32 GetCurrentDisplayLODIndex(USkeletalMeshComponent* InMeshComponent);
 
 	void SetCurrentLOD(int32 NewLodIndex);
 
diff --git a/Engine/Source/Editor/Persona/Public/IPersonaPreviewScene.h b/Engine/Source/Editor/Persona/Public/IPersonaPreviewScene.h
--- a/Engine/Source/Editor/Persona/Public/IPersonaPreviewScene.h
+++ b/Engine/Source/Editor/Persona/Public/IPersonaPreviewScene.h
@@ -174,6 +174,9 @@ public:
 	/** Unregisters a delegate to be called when the preview mesh's LOD has changed */
 	virtual void UnregisterOnLODChanged(void* Thing) = 0;
 
+	/** Broadcasts that we want to force an update, after preview mesh's LOD has changed */
+	virtual void BroadcastOnLODChanged() = 0;
+
 	/** Registers a delegate to be called when the preview mesh's morph targets has changed */
 	virtual void RegisterOnMorphTargetsChanged(const FSimpleDelegate& Delegate) = 0;
 
diff --git a/Engine/Source/Editor/SkeletonEditor/Private/SkeletonTreeBoneItem.cpp b/Engine/Source/Editor/SkeletonEditor/Private/SkeletonTreeBoneItem.cpp
--- a/Engine/Source/Editor/SkeletonEditor/Private/SkeletonTreeBoneItem.cpp
+++ b/Engine/Source/Editor/SkeletonEditor/Private/SkeletonTreeBoneItem.cpp
@@ -556,12 +556,8 @@ bool FSkeletonTreeBoneItem::IsBoneRequired(int32 MeshBoneIndex, UDebugSkelMeshCo
 		return false;
 	}
 
-	//Get current LOD
-	const int32 LODIndex = FMath::Clamp(PreviewComponent->GetPredictedLODLevel(), 0, PreviewComponent->GetSkeletalMeshAsset()->GetResourceForRendering()->LODRenderData.Num() - 1);
-	FSkeletalMeshLODRenderData& LODData = PreviewComponent->GetSkeletalMeshAsset()->GetResourceForRendering()->LODRenderData[LODIndex];
-
-	//Check whether the bone is vertex weighted
-	int32 Index = LODData.RequiredBones.Find(IntCastChecked<FBoneIndexType>(MeshBoneIndex));
+	//Check whether the bone is vertex weighted or required by physics or other forced requirements (see USkeletalMeshComponent::ComputeRequiredBones)
+	const int32 Index = PreviewComponent->RequiredBones.Find(IntCastChecked<FBoneIndexType>(MeshBoneIndex));
 
 	return Index != INDEX_NONE;
 }
