Fixed problems editing various property types in anim notifies

Properties such as FGameplayTagQuery and FRuntimeFloatCurve can now be edited properly without the proxy object being destroyed underneath their editor sub-dialogs

[RN] Fixed problems editing various property types in anim notifies

diff --git a/Engine/Source/Editor/Persona/Private/SAnimNotifyPanel.cpp b/Engine/Source/Editor/Persona/Private/SAnimNotifyPanel.cpp
--- a/Engine/Source/Editor/Persona/Private/SAnimNotifyPanel.cpp
+++ b/Engine/Source/Editor/Persona/Private/SAnimNotifyPanel.cpp
@@ -4021,9 +4021,6 @@ void SAnimNotifyPanel::Construct(const FArguments& InArgs, const TSharedRef<FAni
 		.ColorAndOpacity(FLinearColor::White)
 	];
 
-	OnPropertyChangedHandle = FCoreUObjectDelegates::FOnObjectPropertyChanged::FDelegate::CreateSP(this, &SAnimNotifyPanel::OnPropertyChanged);
-	OnPropertyChangedHandleDelegateHandle = FCoreUObjectDelegates::OnObjectPropertyChanged.Add(OnPropertyChangedHandle);
-
 	// Base notify classes used to search asset data for children.
 	NotifyClassNames.Add(TEXT("Class'/Script/Engine.AnimNotify'"));
 	NotifyStateClassNames.Add(TEXT("Class'/Script/Engine.AnimNotifyState'"));
@@ -4038,8 +4035,6 @@ SAnimNotifyPanel::~SAnimNotifyPanel()
 {
 	Sequence->UnregisterOnNotifyChanged(this);
 
-	FCoreUObjectDelegates::OnObjectPropertyChanged.Remove(OnPropertyChangedHandleDelegateHandle);
-
 	if(GEditor)
 	{
 		GEditor->UnregisterForUndo(this);
@@ -4816,32 +4811,6 @@ void SAnimNotifyPanel::OnPasteNodes(SAnimNotifyTrack* RequestTrack, float ClickT
 	}
 }
 
-void SAnimNotifyPanel::OnPropertyChanged(UObject* ChangedObject, FPropertyChangedEvent& PropertyEvent)
-{
-	// Bail if it isn't a notify
-	if(!ChangedObject->GetClass()->IsChildOf(UAnimNotify::StaticClass()) &&
-	   !ChangedObject->GetClass()->IsChildOf(UAnimNotifyState::StaticClass()))
-	{
-		return;
-	}
-
-	const FName PropertyName = PropertyEvent.GetPropertyName();
-	
-	// Don't process if it's an interactive change; wait till we receive the final event.
-	// Skip notify color as otherwise we will end up refreshing the details panel before any edits are applied (e.g. with the tab key)
-	if(PropertyEvent.ChangeType != EPropertyChangeType::Interactive && PropertyName != GET_MEMBER_NAME_CHECKED(UAnimNotify, NotifyColor) && PropertyName != GET_MEMBER_NAME_CHECKED(UAnimNotifyState, NotifyColor))
-	{
-		for(FAnimNotifyEvent& Event : Sequence->Notifies)
-		{
-			if(Event.Notify == ChangedObject || Event.NotifyStateClass == ChangedObject)
-			{
-				// If we've changed a notify present in the sequence, refresh our tracks.
-				RequestUpdate();
-			}
-		}
-	}
-}
-
 void SAnimNotifyPanel::BindCommands()
 {
 	// This should not be called twice on the same instance
@@ -5107,17 +5076,20 @@ void SAnimNotifyPanel::OnGetNativeNotifyData(TArray<UClass*>& OutClasses, UClass
 
 void SAnimNotifyPanel::OnNotifyObjectChanged(UObject* EditorBaseObj, bool bRebuild)
 {
-	if(UEditorNotifyObject* NotifyObject = Cast<UEditorNotifyObject>(EditorBaseObj))
+	if(bRebuild)
 	{
-		FScopedSavedNotifySelection ScopedSelection(*this);
-
-		for(FAnimNotifyEvent& Notify : Sequence->Notifies)
+		if(UEditorNotifyObject* NotifyObject = Cast<UEditorNotifyObject>(EditorBaseObj))
 		{
-			if(Notify.Guid == NotifyObject->Event.Guid)
+			FScopedSavedNotifySelection ScopedSelection(*this);
+
+			for(FAnimNotifyEvent& Notify : Sequence->Notifies)
 			{
-				if(NotifyAnimTracks.IsValidIndex(Notify.TrackIndex))
+				if(Notify.Guid == NotifyObject->Event.Guid)
 				{
-					NotifyAnimTracks[Notify.TrackIndex]->Update();
+					if(NotifyAnimTracks.IsValidIndex(Notify.TrackIndex))
+					{
+						NotifyAnimTracks[Notify.TrackIndex]->Update();
+					}
 				}
 			}
 		}
diff --git a/Engine/Source/Editor/Persona/Private/SAnimNotifyPanel.h b/Engine/Source/Editor/Persona/Private/SAnimNotifyPanel.h
--- a/Engine/Source/Editor/Persona/Private/SAnimNotifyPanel.h
+++ b/Engine/Source/Editor/Persona/Private/SAnimNotifyPanel.h
@@ -233,10 +233,6 @@ public:
 	void CopySelectedNodesToClipboard() const;
 	void OnPasteNodes(SAnimNotifyTrack* RequestTrack, float ClickTime, ENotifyPasteMode::Type PasteMode, ENotifyPasteMultipleMode::Type MultiplePasteType);
 
-	/** Handler for properties changing on objects */
-	FCoreUObjectDelegates::FOnObjectPropertyChanged::FDelegate OnPropertyChangedHandle;
-	void OnPropertyChanged(UObject* ChangedObject, FPropertyChangedEvent& PropertyEvent);
-
 	/** SWidget Interface */
 	virtual FReply OnKeyDown(const FGeometry& MyGeometry, const FKeyEvent& InKeyEvent) override;
 	virtual FReply OnMouseButtonDown(const FGeometry& MyGeometry, const FPointerEvent& MouseEvent) override;
diff --git a/Engine/Source/Editor/UnrealEd/Classes/Animation/EditorNotifyObject.h b/Engine/Source/Editor/UnrealEd/Classes/Animation/EditorNotifyObject.h
--- a/Engine/Source/Editor/UnrealEd/Classes/Animation/EditorNotifyObject.h
+++ b/Engine/Source/Editor/UnrealEd/Classes/Animation/EditorNotifyObject.h
@@ -25,6 +25,7 @@ class UEditorNotifyObject : public UEditorAnimBaseObj
 	
 	/** Copy changes made to the event object back to the montage asset */
 	virtual bool ApplyChangesToMontage() override;
+	virtual bool PropertyChangeRequiresRebuild(FPropertyChangedEvent& PropertyChangedEvent) override;
 
 	/** The notify event to modify */
 	UPROPERTY(EditAnywhere, Category=Event)
diff --git a/Engine/Source/Editor/UnrealEd/Private/Animation/EditorNotifyObject.cpp b/Engine/Source/Editor/UnrealEd/Private/Animation/EditorNotifyObject.cpp
--- a/Engine/Source/Editor/UnrealEd/Private/Animation/EditorNotifyObject.cpp
+++ b/Engine/Source/Editor/UnrealEd/Private/Animation/EditorNotifyObject.cpp
@@ -2,6 +2,8 @@
 
 #include "Animation/EditorNotifyObject.h"
 #include "Animation/AnimSequenceBase.h"
+#include "Animation/AnimNotifies/AnimNotify.h"
+#include "Animation/AnimNotifies/AnimNotifyState.h"
 
 UEditorNotifyObject::UEditorNotifyObject(const FObjectInitializer& ObjectInitializer)
 	: Super(ObjectInitializer)
@@ -46,3 +48,14 @@ void UEditorNotifyObject::InitialiseNotify(const FAnimNotifyEvent& InNotify)
 		Event = InNotify;
 	}
 }
+
+bool UEditorNotifyObject::PropertyChangeRequiresRebuild(FPropertyChangedEvent& PropertyChangedEvent)
+{
+	// We dont need to rebuild the track UI when we change the properties of a notify
+	if(PropertyChangedEvent.GetMemberPropertyName() == GET_MEMBER_NAME_CHECKED(UEditorNotifyObject, Event) && PropertyChangedEvent.GetPropertyName() == GET_MEMBER_NAME_CHECKED(FAnimNotifyEvent, Notify))
+	{
+		return false;
+	}
+
+	return Super::PropertyChangeRequiresRebuild(PropertyChangedEvent);
+}
\ No newline at end of file
