Fix montage editor selection issues
* Dragging montage selections now correctly selects the montage.
* Changing marquee selection in notify track. This hijacks mouse selection and the marquee selection is jarring as other track types do not also have this functionality. Switching to default selection behaviour and gating marquee selection behind the alt modifier keys.

diff --git a/Engine/Source/Editor/Persona/Private/SAnimNotifyPanel.cpp b/Engine/Source/Editor/Persona/Private/SAnimNotifyPanel.cpp
--- a/Engine/Source/Editor/Persona/Private/SAnimNotifyPanel.cpp
+++ b/Engine/Source/Editor/Persona/Private/SAnimNotifyPanel.cpp
@@ -4843,11 +4843,14 @@ FReply SAnimNotifyPanel::OnKeyDown(const FGeometry& MyGeometry, const FKeyEvent&
 
 FReply SAnimNotifyPanel::OnMouseButtonDown(const FGeometry& MyGeometry, const FPointerEvent& MouseEvent)
 {
-	SAnimTrackPanel::OnMouseButtonDown(MyGeometry, MouseEvent);
+	if (SAnimTrackPanel::OnMouseButtonDown(MyGeometry, MouseEvent).IsEventHandled())
+	{
+		return FReply::Handled();
+	}
 
-	bool bLeftButton = MouseEvent.IsMouseButtonDown(EKeys::LeftMouseButton);
+	bool bCreateMarquee = MouseEvent.IsMouseButtonDown(EKeys::LeftMouseButton) && (MouseEvent.IsAltDown());
 
-	if(bLeftButton)
+	if(bCreateMarquee)
 	{
 		TArray<TSharedPtr<SAnimNotifyNode>> SelectedNodes;
 		for(TSharedPtr<SAnimNotifyTrack> Track : NotifyAnimTracks)
diff --git a/Engine/Source/Editor/Persona/Private/SAnimTimingPanel.cpp b/Engine/Source/Editor/Persona/Private/SAnimTimingPanel.cpp
--- a/Engine/Source/Editor/Persona/Private/SAnimTimingPanel.cpp
+++ b/Engine/Source/Editor/Persona/Private/SAnimTimingPanel.cpp
@@ -295,6 +295,7 @@ void SAnimTimingPanel::RefreshTrackNodes()
 				.DataStartPos(Element.ToSharedRef(), &FTimingRelevantElementBase::GetElementTime)
 				.NodeName(FString::FromInt(ElementIdx + 1))
 				.NodeColor(FLinearColor::Yellow)
+				.Visibility(EVisibility::HitTestInvisible)
 				.Element(Element)
 				);
 		}
