Fixes crash in "running instance" preview mesh controller when refreshing list with a null world

diff --git a/Engine/Source/Editor/Persona/Private/PersonaPreviewSceneSkelMeshInstanceController.cpp b/Engine/Source/Editor/Persona/Private/PersonaPreviewSceneSkelMeshInstanceController.cpp
--- a/Engine/Source/Editor/Persona/Private/PersonaPreviewSceneSkelMeshInstanceController.cpp
+++ b/Engine/Source/Editor/Persona/Private/PersonaPreviewSceneSkelMeshInstanceController.cpp
@@ -179,8 +179,14 @@ void SSkeletalMeshDebugSelectionWidget::Refresh()
 	// spin through all available worlds and find all skeletal mesh components using the target skeletal mesh
 	for (int32 WorldIndex=0; WorldIndex<AllDebugWorlds.Num(); ++WorldIndex)
 	{
-		TWeakObjectPtr<UWorld> World = AllDebugWorlds[WorldIndex];
-		for (TActorIterator<AActor> ActorItr(World.Get()); ActorItr; ++ActorItr)
+		UWorld* World = AllDebugWorlds[WorldIndex].Get();
+		if (!World)
+		{
+			// double-check because we have had crashes in TActorIterator below on null worlds
+			continue;
+		}
+		
+		for (TActorIterator<AActor> ActorItr(World); ActorItr; ++ActorItr)
 		{
 			const AActor* Actor = *ActorItr;
 			if (!Actor)
