Add nullptr check for pose watch column widget

diff --git a/Engine/Source/Editor/Persona/Private/PoseWatchManagerColumnVisibility.cpp b/Engine/Source/Editor/Persona/Private/PoseWatchManagerColumnVisibility.cpp
--- a/Engine/Source/Editor/Persona/Private/PoseWatchManagerColumnVisibility.cpp
+++ b/Engine/Source/Editor/Persona/Private/PoseWatchManagerColumnVisibility.cpp
@@ -103,7 +103,9 @@ namespace PoseWatchManager
 
 		bool IsVisible() const
 		{
-			return IsVisible(WeakTreeItem.Pin(), WeakColumn.Pin());
+			TSharedPtr<FPoseWatchManagerColumnVisibility> Column = WeakColumn.Pin();
+			TSharedPtr<IPoseWatchManagerTreeItem> TreeItem = WeakTreeItem.Pin();
+			return (TreeItem.IsValid() && Column.IsValid()) ? IsVisible(TreeItem, Column) : false;
 		}
 
 		void SetIsVisible(const bool bVisible)
