Fix crash with setting bone name in skeletal mesh region bone filter
* Also switch to SComboBox<FName> instead of SComboBox<TSharedPtr<FName>>

diff --git a/Engine/Source/Editor/Persona/Private/Customization/SkeletalMeshRegionCustomization.cpp b/Engine/Source/Editor/Persona/Private/Customization/SkeletalMeshRegionCustomization.cpp
--- a/Engine/Source/Editor/Persona/Private/Customization/SkeletalMeshRegionCustomization.cpp
+++ b/Engine/Source/Editor/Persona/Private/Customization/SkeletalMeshRegionCustomization.cpp
@@ -67,7 +67,7 @@ void FNiagaraSkeletalMeshRegionBoneFilterDetails::CustomizeChildren(TSharedRef<I
 				]
 				.ValueWidget
 				[
-					SNew(SComboBox<TSharedPtr<FName>>)
+					SNew(SComboBox<FName>)
 					.OptionsSource(&PossibleBoneNames)
 					.ContentPadding(2.0f)
 					.OnGenerateWidget(this, &FNiagaraSkeletalMeshRegionBoneFilterDetails::HandleBoneNameComboBoxGenerateWidget)
@@ -93,15 +93,15 @@ FText FNiagaraSkeletalMeshRegionBoneFilterDetails::HandleBoneNameComboBoxContent
 	return FText::FromName(OutName);
 }
 
-TSharedRef<SWidget> FNiagaraSkeletalMeshRegionBoneFilterDetails::HandleBoneNameComboBoxGenerateWidget(TSharedPtr<FName> StringItem)
+TSharedRef<SWidget> FNiagaraSkeletalMeshRegionBoneFilterDetails::HandleBoneNameComboBoxGenerateWidget(FName StringItem)
 {
-	FName DefaultName = StringItem.IsValid() ? *StringItem : NAME_None;
-	return SNew(STextBlock).Text(FText::FromName(DefaultName));
+	return SNew(STextBlock)
+		.Text(FText::FromName(StringItem));
 }
 
-void FNiagaraSkeletalMeshRegionBoneFilterDetails::HandleBoneNameComboBoxSelectionChanged(TSharedPtr<FName> StringItem, ESelectInfo::Type SelectInfo)
+void FNiagaraSkeletalMeshRegionBoneFilterDetails::HandleBoneNameComboBoxSelectionChanged(FName StringItem, ESelectInfo::Type SelectInfo)
 {
-	BoneNameHandle->SetValue(*StringItem);
+	BoneNameHandle->SetValue(StringItem);
 }
 
 
@@ -113,8 +113,7 @@ void FNiagaraSkeletalMeshRegionBoneFilterDetails::OnComboOpening()
 		//Populate PossibleBonesNames
 		for (FMeshBoneInfo Bone : MeshObject->GetSkeleton()->GetReferenceSkeleton().GetRefBoneInfo())
 		{
-			TSharedPtr<FName> BoneName(new FName(Bone.Name));
-			PossibleBoneNames.Add(BoneName);
+			PossibleBoneNames.Add(Bone.Name);
 		}
 	}
 }
diff --git a/Engine/Source/Editor/Persona/Private/Customization/SkeletalMeshRegionCustomization.h b/Engine/Source/Editor/Persona/Private/Customization/SkeletalMeshRegionCustomization.h
--- a/Engine/Source/Editor/Persona/Private/Customization/SkeletalMeshRegionCustomization.h
+++ b/Engine/Source/Editor/Persona/Private/Customization/SkeletalMeshRegionCustomization.h
@@ -27,15 +27,15 @@ public:
 
 private:
 	FText HandleBoneNameComboBoxContentText() const;
-	TSharedRef<SWidget> HandleBoneNameComboBoxGenerateWidget(TSharedPtr<FName> StringItem);
-	void HandleBoneNameComboBoxSelectionChanged(TSharedPtr<FName> StringItem, ESelectInfo::Type SelectInfo);
+	TSharedRef<SWidget> HandleBoneNameComboBoxGenerateWidget(FName StringItem);
+	void HandleBoneNameComboBoxSelectionChanged(FName StringItem, ESelectInfo::Type SelectInfo);
 
 	void OnComboOpening();
 
 	/** The mesh object whose details we're customizing. */
-	USkeletalMesh * MeshObject;
+	USkeletalMesh* MeshObject;
 
-	TArray<TSharedPtr<FName>> PossibleBoneNames;
+	TArray<FName> PossibleBoneNames;
 
 	TSharedPtr<IPropertyHandle> BoneNameHandle;
 };
