PR #10113: Fix pose watch not drawing correctly for debug skeletons larger than preview

diff --git a/Engine/Source/Editor/Persona/Private/AnimationEditorViewportClient.cpp b/Engine/Source/Editor/Persona/Private/AnimationEditorViewportClient.cpp
--- a/Engine/Source/Editor/Persona/Private/AnimationEditorViewportClient.cpp
+++ b/Engine/Source/Editor/Persona/Private/AnimationEditorViewportClient.cpp
@@ -1802,14 +1802,20 @@ void FAnimationViewportClient::DrawWatchedPoses(UDebugSkelMeshComponent * MeshCo
 	{
 		if (UAnimBlueprint* Blueprint = Cast<UAnimBlueprint>(AnimBPGenClass->ClassGeneratedBy))
 		{
-			if (Blueprint->GetObjectBeingDebugged() && MeshComponent)
+			if (const UAnimInstance* DebuggedAnimInstance = Cast<UAnimInstance>(Blueprint->GetObjectBeingDebugged()))
 			{
-				FAnimBlueprintDebugData& DebugData = AnimBPGenClass->GetAnimBlueprintDebugData();
-				DebugData.ForEachActiveVisiblePoseWatchPoseElement([PDI, MeshComponent](FAnimNodePoseWatch& PoseWatch)
+				if (const USkeletalMeshComponent* DebuggedSkeletalMeshComponent = DebuggedAnimInstance->GetSkelMeshComponent())
 				{
-					PoseWatch.CopyPoseWatchData(MeshComponent->GetReferenceSkeleton());
-					SkeletalDebugRendering::DrawBonesFromPoseWatch(PDI, PoseWatch, /*bUseWorldTransform*/false);
-				});
+					if (const USkeletalMesh* SkeletalMesh = DebuggedSkeletalMeshComponent->GetSkeletalMeshAsset())
+					{
+						FAnimBlueprintDebugData& DebugData = AnimBPGenClass->GetAnimBlueprintDebugData();
+						DebugData.ForEachActiveVisiblePoseWatchPoseElement([PDI, SkeletalMesh](FAnimNodePoseWatch& PoseWatch)
+						{
+							PoseWatch.CopyPoseWatchData(SkeletalMesh->GetRefSkeleton());
+							SkeletalDebugRendering::DrawBonesFromPoseWatch(PDI, PoseWatch, /*bUseWorldTransform*/false);
+						});
+					}
+				}
 			}
 		}
 	}
