Fix BonesToDraw not updating correctly when preview mesh changes and when Show Bone Names is set

diff --git a/Engine/Source/Editor/Persona/Private/AnimationEditorViewportClient.cpp b/Engine/Source/Editor/Persona/Private/AnimationEditorViewportClient.cpp
--- a/Engine/Source/Editor/Persona/Private/AnimationEditorViewportClient.cpp
+++ b/Engine/Source/Editor/Persona/Private/AnimationEditorViewportClient.cpp
@@ -614,6 +614,7 @@ void FAnimationViewportClient::HandleSkeletalMeshChanged(USkeletalMesh* OldSkele
 void FAnimationViewportClient::HandleOnMeshChanged()
 {
 	UpdateCameraSetup();
+	UpdateBonesToDraw();
 	Invalidate();
 }
 
@@ -953,7 +954,7 @@ void FAnimationViewportClient::ShowBoneNames( FCanvas* Canvas, FSceneView* View,
 	{
 		const int32 BoneIndex = LODData.RequiredBones[i];
 
-		if (!BonesToDraw[BoneIndex])
+		if (!BonesToDraw.IsValidIndex(BoneIndex) || !BonesToDraw[BoneIndex])
 		{
 			continue;
 		}
diff --git a/Engine/Source/Editor/Persona/Private/SAnimationEditorViewport.cpp b/Engine/Source/Editor/Persona/Private/SAnimationEditorViewport.cpp
--- a/Engine/Source/Editor/Persona/Private/SAnimationEditorViewport.cpp
+++ b/Engine/Source/Editor/Persona/Private/SAnimationEditorViewport.cpp
@@ -1258,6 +1258,8 @@ void SAnimationEditorViewportTabBody::OnShowBoneNames()
 		InMesh->MarkRenderStateDirty();
 	});
 	RefreshViewport();
+
+	GetAnimationViewportClient()->UpdateBonesToDraw();
 }
 
 void SAnimationEditorViewportTabBody::OnShowRawAnimation()
