Fixes gap in viewport toolbar when hiding preview profile menu.

diff --git a/Engine/Source/Editor/Persona/Private/SAnimViewportToolBar.cpp b/Engine/Source/Editor/Persona/Private/SAnimViewportToolBar.cpp
--- a/Engine/Source/Editor/Persona/Private/SAnimViewportToolBar.cpp
+++ b/Engine/Source/Editor/Persona/Private/SAnimViewportToolBar.cpp
@@ -337,7 +337,7 @@ void SAnimViewportToolBar::Construct(const FArguments& InArgs, TSharedPtr<class
 				// this behavior was requested by UX to match the behavior of the static mesh editor
 				if (PreviewProfileController.IsValid())
 				{
-					return PreviewProfileController->HasAnyUserProfiles() ? EVisibility::Visible : EVisibility::Hidden;
+					return PreviewProfileController->HasAnyUserProfiles() ? EVisibility::Visible : EVisibility::Collapsed;
 				}
 				return EVisibility::Hidden;
 			})
