[Persona] Add validity check on PersonaToolkit SharedPtr which crashes from time to time

While working on a local change, I noticed a crash happening with the following steps:

- Open editor for Skel Mesh
- Change floor offset
- Close editor
- Re-open editor
- Hit undo
- Crash

While I'm currently looking into why that specific change could lead to the crash, I think it might be good to have a check around the pointer as per this CL.

diff --git a/Engine/Source/Editor/Persona/Private/AnimationEditorPreviewScene.cpp b/Engine/Source/Editor/Persona/Private/AnimationEditorPreviewScene.cpp
--- a/Engine/Source/Editor/Persona/Private/AnimationEditorPreviewScene.cpp
+++ b/Engine/Source/Editor/Persona/Private/AnimationEditorPreviewScene.cpp
@@ -316,18 +316,20 @@ void FAnimationEditorPreviewScene::SetPreviewMeshInternal(USkeletalMesh* NewPrev
 
 	// Setting the skeletal mesh to in the PreviewScene can change AnimScriptInstance so we must re register it
 	// with the AnimBlueprint
-	UAnimBlueprint* SourceBlueprint = PersonaToolkit.Pin()->GetAnimBlueprint();
-	if (SourceBlueprint)
+	if (const TSharedPtr<IPersonaToolkit>& PersonaToolkitPinned = PersonaToolkit.Pin())
 	{
-		if (DebuggedSkeletalMeshComponent && DebuggedSkeletalMeshComponent->GetAnimInstance() && DebuggedSkeletalMeshComponent->GetAnimInstance()->IsA(SourceBlueprint->GeneratedClass))
+		if (UAnimBlueprint* SourceBlueprint = PersonaToolkitPinned->GetAnimBlueprint())
 		{
-			PersonaUtils::SetObjectBeingDebugged(SourceBlueprint, DebuggedSkeletalMeshComponent->GetAnimInstance());
-		}
+			if (DebuggedSkeletalMeshComponent && DebuggedSkeletalMeshComponent->GetAnimInstance() && DebuggedSkeletalMeshComponent->GetAnimInstance()->IsA(SourceBlueprint->GeneratedClass))
+			{
+				PersonaUtils::SetObjectBeingDebugged(SourceBlueprint, DebuggedSkeletalMeshComponent->GetAnimInstance());
+			}
 
-		// If we didn't have a preview mesh before and we select one now, set it up as the object being debugged
-		if (DebuggedSkeletalMeshComponent == nullptr && NewPreviewMesh != nullptr && SkeletalMeshComponent->GetAnimInstance() && SkeletalMeshComponent->GetAnimInstance()->IsA(SourceBlueprint->GeneratedClass))
-		{
-			PersonaUtils::SetObjectBeingDebugged(SourceBlueprint, SkeletalMeshComponent->GetAnimInstance());
+			// If we didn't have a preview mesh before and we select one now, set it up as the object being debugged
+			if (DebuggedSkeletalMeshComponent == nullptr && NewPreviewMesh != nullptr && SkeletalMeshComponent->GetAnimInstance() && SkeletalMeshComponent->GetAnimInstance()->IsA(SourceBlueprint->GeneratedClass))
+			{
+				PersonaUtils::SetObjectBeingDebugged(SourceBlueprint, SkeletalMeshComponent->GetAnimInstance());
+			}
 		}
 	}
 
