[Viewport Toolbar] Add new toolbar to Skeletal Editors

- Copied Animation Editor Toolbar Widgets from SAnimViewportToolbar.cpp to dedicated source files. Will remove old ones after updating Old Toolbar in a future CL.
- Adds most menus
- Still missing custom select menu
- Missing entries which are to be moved to timeline area

diff --git a/Engine/Source/Editor/Persona/Private/SAnimViewportToolBar.cpp b/Engine/Source/Editor/Persona/Private/SAnimViewportToolBar.cpp
--- a/Engine/Source/Editor/Persona/Private/SAnimViewportToolBar.cpp
+++ b/Engine/Source/Editor/Persona/Private/SAnimViewportToolBar.cpp
@@ -1079,16 +1079,19 @@ void SAnimViewportToolBar::FillCharacterClothingMenu(FMenuBuilder& MenuBuilder)
 	TSharedPtr<SAnimationEditorViewportTabBody> SharedViewport = Viewport.Pin();
 	if (SharedViewport.IsValid())
 	{
-		TSharedRef<IPersonaPreviewScene> PreviewScene = SharedViewport->GetAnimationViewportClient()->GetPreviewScene();
-		if (UDebugSkelMeshComponent* PreviewComponent = PreviewScene->GetPreviewMeshComponent())
+		if (const TSharedPtr<FAnimationViewportClient>& AnimationViewportClient = SharedViewport->GetAnimationViewportClient())
 		{
-			if (PreviewComponent->ClothingSimulationFactory)  // The cloth plugin could be disabled, and the factory would be null in this case
+			TSharedRef<IPersonaPreviewScene> PreviewScene = AnimationViewportClient->GetPreviewScene();
+			if (UDebugSkelMeshComponent* PreviewComponent = PreviewScene->GetPreviewMeshComponent())
 			{
-				FClothingSystemEditorInterfaceModule& ClothingEditorModule = FModuleManager::LoadModuleChecked<FClothingSystemEditorInterfaceModule>(TEXT("ClothingSystemEditorInterface"));
-
-				if (ISimulationEditorExtender* Extender = ClothingEditorModule.GetSimulationEditorExtender(PreviewComponent->ClothingSimulationFactory->GetFName()))
+				if (PreviewComponent->ClothingSimulationFactory)  // The cloth plugin could be disabled, and the factory would be null in this case
 				{
-					Extender->ExtendViewportShowMenu(MenuBuilder, PreviewScene);
+					FClothingSystemEditorInterfaceModule& ClothingEditorModule = FModuleManager::LoadModuleChecked<FClothingSystemEditorInterfaceModule>(TEXT("ClothingSystemEditorInterface"));
+
+					if (ISimulationEditorExtender* Extender = ClothingEditorModule.GetSimulationEditorExtender(PreviewComponent->ClothingSimulationFactory->GetFName()))
+					{
+						Extender->ExtendViewportShowMenu(MenuBuilder, PreviewScene);
+					}
 				}
 			}
 		}
diff --git a/Engine/Source/Editor/Persona/Private/SAnimationEditorViewport.cpp b/Engine/Source/Editor/Persona/Private/SAnimationEditorViewport.cpp
--- a/Engine/Source/Editor/Persona/Private/SAnimationEditorViewport.cpp
+++ b/Engine/Source/Editor/Persona/Private/SAnimationEditorViewport.cpp
@@ -37,9 +37,15 @@
 #include "IPersonaEditMode.h"
 
 #include "IPersonaToolkit.h"
+#include "IPinnedCommandList.h"
+#include "IPinnedCommandListModule.h"
 #include "Subsystems/AssetEditorSubsystem.h"
 #include "SNameComboBox.h"
+#include "ToolMenus.h"
 #include "Viewports.h"
+#include "ViewportToolbar/AnimViewportContext.h"
+#include "ViewportToolbar/AnimationEditorMenus.h"
+#include "ViewportToolbar/UnrealEdViewportToolbar.h"
 
 #define LOCTEXT_NAMESPACE "PersonaViewportToolbar"
 
@@ -82,10 +88,10 @@ void SAnimationEditorViewport::Construct(const FArguments& InArgs, const FAnimat
 	}
 
 	UEditorEngine* Editor = (UEditorEngine*)GEngine;
-	PreviewFeatureLevelChangedHandle = Editor->OnPreviewFeatureLevelChanged().AddLambda([this](ERHIFeatureLevel::Type NewFeatureLevel)
+	PreviewFeatureLevelChangedHandle = Editor->OnPreviewFeatureLevelChanged().AddLambda(
+		[PreviewScenePtrWeak = PreviewScenePtr](ERHIFeatureLevel::Type NewFeatureLevel)
 		{
-			auto ScenePtr = PreviewScenePtr.Pin();
-			if (ScenePtr.IsValid())
+			if (TSharedPtr<IPersonaPreviewScene> ScenePtr = PreviewScenePtrWeak.Pin())
 			{
 				UWorld* World = ScenePtr->GetWorld();
 				if (World != nullptr)
@@ -139,7 +145,7 @@ TSharedPtr<SWidget> SAnimationEditorViewport::MakeViewportToolbar()
 		.Cursor(EMouseCursor::Default)
 		.Extenders(Extenders)
 		.ContextName(ContextName)
-		.PreviewProfileController(MakeShared<FPreviewProfileController>())
+		.PreviewProfileController(GetPreviewProfileController()) // TODO: this line can be removed when updating SAnimViewportToolBar (CL 37180383)
 		.ShowShowMenu(bShowShowMenu)
 		.ShowLODMenu(bShowLODMenu)
 		.ShowPlaySpeedMenu(bShowPlaySpeedMenu)
@@ -148,6 +154,176 @@ TSharedPtr<SWidget> SAnimationEditorViewport::MakeViewportToolbar()
 		.ShowPhysicsMenu(bShowPhysicsMenu);
 }
 
+TSharedPtr<SWidget> SAnimationEditorViewport::BuildViewportToolbar()
+{
+	// Register the viewport toolbar if another viewport hasn't already (it's shared).
+	const FName ViewportToolbarName = "AnimationEditor.ViewportToolbar";
+
+	if (!UToolMenus::Get()->IsMenuRegistered(ViewportToolbarName))
+	{
+		UToolMenu* const ViewportToolbarMenu = UToolMenus::Get()->RegisterMenu(
+			ViewportToolbarName, NAME_None /* parent */, EMultiBoxType::SlimHorizontalToolBar
+		);
+
+		ViewportToolbarMenu->StyleName = "ViewportToolbar";
+
+		// Add the left-aligned part of the viewport toolbar.
+		{
+			FToolMenuSection& LeftSection = ViewportToolbarMenu->FindOrAddSection("Left");
+
+			// Add the "Transforms" sub menu.
+			{
+				FToolMenuEntry TransformsSubmenu = UE::UnrealEd::CreateViewportToolbarTransformsSection();
+				TransformsSubmenu.InsertPosition.Position = EToolMenuInsertType::First;
+				LeftSection.AddEntry(TransformsSubmenu);
+			}
+
+			// Add the "Selection" sub menu.
+			{
+				// todo: needs specific menu for Skel Mesh
+			}
+
+			// Add the "Snapping" sub menu.
+			{
+				FToolMenuEntry SnappingSubmenu = UE::UnrealEd::CreateViewportToolbarSnappingSubmenu();
+				SnappingSubmenu.InsertPosition.Position = EToolMenuInsertType::First;
+				LeftSection.AddEntry(SnappingSubmenu);
+			}
+		}
+
+		// Add the right-aligned part of the viewport toolbar.
+		{
+			// Add the submenus of this section as EToolMenuInsertType::Last to sort them after any
+			// default-positioned submenus external code might add.
+			FToolMenuSection& RightSection = ViewportToolbarMenu->FindOrAddSection("Right");
+			RightSection.Alignment = EToolMenuSectionAlign::Last;
+
+			// Add the "Camera" submenu.
+			{
+				const FName GrandParentSubmenuName = "UnrealEd.ViewportToolbar.Camera";
+				const FName ParentSubmenuName = "AnimationEditor.ViewportToolbar.Camera";
+				const FName SubmenuName = "AnimationEditor.ViewportToolbar.CameraOptions";
+
+				// Create our grandparent menu.
+				if (!UToolMenus::Get()->IsMenuRegistered(GrandParentSubmenuName))
+				{
+					UToolMenus::Get()->RegisterMenu(GrandParentSubmenuName);
+				}
+
+				// Create our parent menu.
+				if (!UToolMenus::Get()->IsMenuRegistered(ParentSubmenuName))
+				{
+					UToolMenus::Get()->RegisterMenu(ParentSubmenuName, GrandParentSubmenuName);
+				}
+
+				// Create our menu.
+				UToolMenus::Get()->RegisterMenu(SubmenuName, ParentSubmenuName);
+
+				constexpr bool bShowFarViewPlaneEntry = false;
+				UE::UnrealEd::ExtendCameraSubmenu(SubmenuName, bShowFarViewPlaneEntry);
+				UE::AnimationEditor::ExtendCameraMenu(SubmenuName);
+
+				FToolMenuEntry CameraSubmenu = UE::UnrealEd::CreateViewportToolbarCameraSubmenu();
+				CameraSubmenu.InsertPosition.Position = EToolMenuInsertType::Last;
+				RightSection.AddEntry(CameraSubmenu);
+			}
+
+			// Add the "View Modes" sub menu.
+			{
+				const FName ViewModesMenuName = "AnimationEditor.ViewportToolbar.ViewModes";
+
+				// Stay backward-compatible with the old viewport toolbar.
+				{
+					const FName ParentSubmenuName = "UnrealEd.ViewportToolbar.View";
+					// Create our parent menu.
+					if (!UToolMenus::Get()->IsMenuRegistered(ParentSubmenuName))
+					{
+						UToolMenus::Get()->RegisterMenu(ParentSubmenuName);
+					}
+
+					// Register our ToolMenu here first, before we create the submenu, so we can set our parent.
+					UToolMenus::Get()->RegisterMenu(ViewModesMenuName, ParentSubmenuName);
+				}
+
+				UE::AnimationEditor::ExtendViewModesSubmenu(ViewModesMenuName);
+
+				FToolMenuEntry ViewModesSubmenu = UE::UnrealEd::CreateViewportToolbarViewModesSubmenu();
+				ViewModesSubmenu.InsertPosition.Position = EToolMenuInsertType::Last;
+
+				RightSection.AddEntry(ViewModesSubmenu);
+			}
+
+			// Add the "Show" sub menu
+			{
+				FToolMenuEntry ShowSubmenu = UE::AnimationEditor::CreateShowSubmenu();
+				ShowSubmenu.InsertPosition.Position = EToolMenuInsertType::Last;
+				RightSection.AddEntry(ShowSubmenu);
+			}
+
+			// Add the "LOD" sub menu
+			{
+				FToolMenuEntry LODSubmenu = UE::AnimationEditor::CreateLODSubmenu();
+				LODSubmenu.InsertPosition.Position = EToolMenuInsertType::Last;
+				RightSection.AddEntry(LODSubmenu);
+			}
+
+			// Add the "Settings" sub menu
+			{
+				FToolMenuEntry SettingsSubmenu = UE::AnimationEditor::CreateSettingsSubmenu();
+				SettingsSubmenu.InsertPosition.Position = EToolMenuInsertType::Last;
+				RightSection.AddEntry(SettingsSubmenu);
+			}
+		}
+	}
+
+	FToolMenuContext ViewportToolbarContext;
+	{
+		ViewportToolbarContext.AppendCommandList(GetCommandList());
+
+		// Add extenders
+		{
+			Extenders.Add(UE::AnimationEditor::GetViewModesLegacyExtenders(TabBodyPtr));
+			ViewportToolbarContext.AddExtender(FExtender::Combine(Extenders));
+		}
+
+		// Add the UnrealEd viewport toolbar context.
+		{
+			UUnrealEdViewportToolbarContext* const ContextObject =
+				UE::UnrealEd::CreateViewportToolbarDefaultContext(SharedThis(this));
+
+			ViewportToolbarContext.AddObject(ContextObject);
+		}
+
+		// Add the Anim viewport toolbar context.
+		{
+			UAnimViewportContext* const ContextObject = NewObject<UAnimViewportContext>();
+			ContextObject->ViewportTabBody = TabBodyPtr.Pin();
+
+			ViewportToolbarContext.AddObject(ContextObject);
+		}
+	}
+
+	// clang-format off
+	const TSharedRef<SWidget> NewViewportToolbar = SNew(SBox)
+	.Visibility_Lambda(
+		[]() -> EVisibility
+		{
+			return  UE::UnrealEd::ShowNewViewportToolbars() ? EVisibility::Visible: EVisibility::Collapsed;
+		}
+	)
+	[
+		UToolMenus::Get()->GenerateWidget(ViewportToolbarName, ViewportToolbarContext)
+	];
+	// clang-format on
+
+	return NewViewportToolbar;
+}
+
+TSharedPtr<IPreviewProfileController> SAnimationEditorViewport::CreatePreviewProfileController()
+{
+	return MakeShared<FPreviewProfileController>();
+}
+
 void SAnimationEditorViewport::PostUndo( bool bSuccess )
 {
 	LevelViewportClient->Invalidate();
@@ -160,8 +336,11 @@ void SAnimationEditorViewport::PostRedo( bool bSuccess )
 
 void SAnimationEditorViewport::OnFocusViewportToSelection()
 {
-	TSharedRef<FAnimationViewportClient> AnimViewportClient = StaticCastSharedRef<FAnimationViewportClient>(LevelViewportClient.ToSharedRef());
-	AnimViewportClient->OnFocusViewportToSelection();
+	if (LevelViewportClient)
+	{
+		TSharedRef<FAnimationViewportClient> AnimViewportClient = StaticCastSharedRef<FAnimationViewportClient>(LevelViewportClient.ToSharedRef());
+		AnimViewportClient->OnFocusViewportToSelection();
+	}
 }
 
 void SAnimationEditorViewport::BindCommands()
@@ -226,11 +405,13 @@ SAnimationEditorViewportTabBody::~SAnimationEditorViewportTabBody()
 	// Close viewport
 	if (LevelViewportClient.IsValid())
 	{
-		LevelViewportClient->Viewport = NULL;
+		LevelViewportClient->Viewport = nullptr;
 	}
 
 	// Release our reference to the viewport client
 	LevelViewportClient.Reset();
+
+	PendingTransaction.Reset();
 }
 
 bool SAnimationEditorViewportTabBody::CanUseGizmos() const
@@ -354,14 +535,21 @@ FText SAnimationEditorViewportTabBody::GetDisplayString() const
 TSharedRef<IPersonaViewportState> SAnimationEditorViewportTabBody::SaveState() const
 {
 	TSharedRef<FPersonaModeSharedData> State = MakeShareable(new(FPersonaModeSharedData));
-	State->Save(StaticCastSharedRef<FAnimationViewportClient>(LevelViewportClient.ToSharedRef()));
+
+	if (const TSharedPtr<FAnimationViewportClient>& AnimViewportClient = GetAnimationViewportClient())
+	{
+		State->Save(AnimViewportClient.ToSharedRef());
+	}
 	return State;
 }
 
 void SAnimationEditorViewportTabBody::RestoreState(TSharedRef<IPersonaViewportState> InState)
 {
-	TSharedRef<FPersonaModeSharedData> State = StaticCastSharedRef<FPersonaModeSharedData>(InState);
-	State->Restore(StaticCastSharedRef<FAnimationViewportClient>(LevelViewportClient.ToSharedRef()));
+	if (const TSharedPtr<FAnimationViewportClient>& AnimViewportClient = GetAnimationViewportClient())
+	{
+		TSharedRef<FPersonaModeSharedData> State = StaticCastSharedRef<FPersonaModeSharedData>(InState);
+		State->Restore(AnimViewportClient.ToSharedRef());
+	}
 }
 
 FEditorViewportClient& SAnimationEditorViewportTabBody::GetViewportClient() const
@@ -374,6 +562,15 @@ TSharedRef<IPinnedCommandList> SAnimationEditorViewportTabBody::GetPinnedCommand
 	return ViewportWidget->GetViewportToolbar()->GetPinnedCommandList().ToSharedRef();
 }
 
+const TSharedPtr<IPinnedCommandList>& SAnimationEditorViewportTabBody::GetPinnedCommands()
+{
+	if (!PinnedCommands)
+	{
+		CreatePinnedCommands();
+	}
+	return PinnedCommands;
+}
+
 TWeakPtr<SWidget> SAnimationEditorViewportTabBody::AddNotification(TAttribute<EMessageSeverity::Type> InSeverity, TAttribute<bool> InCanBeDismissed, const TSharedRef<SWidget>& InNotificationWidget, FPersonaViewportNotificationOptions InOptions)
 {
 	TSharedPtr<SBorder> ContainingWidget = nullptr;
@@ -500,6 +697,11 @@ void SAnimationEditorViewportTabBody::RefreshViewport()
 	LevelViewportClient->Invalidate();
 }
 
+TSharedPtr<FAssetEditorToolkit> SAnimationEditorViewportTabBody::GetAssetEditorToolkit() const
+{
+	return AssetEditorToolkitPtr.Pin();
+}
+
 bool SAnimationEditorViewportTabBody::IsVisible() const
 {
 	return ViewportWidget.IsValid();
@@ -531,6 +733,7 @@ void SAnimationEditorViewportTabBody::Construct(const FArguments& InArgs, const
 	bAlwaysShowTransformToolbar = InArgs._AlwaysShowTransformToolbar;
 	OnInvokeTab = InArgs._OnInvokeTab;
 	OnGetViewportText = InArgs._OnGetViewportText;
+	ContextName = InArgs._ContextName;
 
 	// register delegates for change notifications
 	InPreviewScene->RegisterOnAnimChanged(FOnAnimChanged::CreateSP(this, &SAnimationEditorViewportTabBody::AnimChanged));
@@ -598,6 +801,9 @@ void SAnimationEditorViewportTabBody::Construct(const FArguments& InArgs, const
 		.ShowPhysicsMenu(InArgs._ShowPhysicsMenu);
 
 	TSharedPtr<SVerticalBox> ViewportContainer = nullptr;
+
+	CreatePinnedCommands();
+
 	this->ChildSlot
 	[
 		SAssignNew(ViewportContainer, SVerticalBox)
@@ -646,10 +852,15 @@ void SAnimationEditorViewportTabBody::Construct(const FArguments& InArgs, const
 
 	LevelViewportClient = ViewportWidget->GetViewportClient();
 
-	TSharedRef<FAnimationViewportClient> AnimViewportClient = StaticCastSharedRef<FAnimationViewportClient>(LevelViewportClient.ToSharedRef());
+	if (const TSharedPtr<FAnimationViewportClient>& AnimViewportClient = GetAnimationViewportClient())
+	{
+		// Load the view mode from config
+		AnimViewportClient->SetViewMode(AnimViewportClient->ConfigOption
+												 ->GetAssetEditorOptions(AssetEditorToolkitPtr.Pin()->GetEditorName())
+												 .ViewportConfigs[InViewportIndex]
+												 .ViewModeIndex);
+	}
 
-	// Load the view mode from config
-	AnimViewportClient->SetViewMode(AnimViewportClient->ConfigOption->GetAssetEditorOptions(AssetEditorToolkitPtr.Pin()->GetEditorName()).ViewportConfigs[InViewportIndex].ViewModeIndex);
 	UpdateShowFlagForMeshEdges();
 
 
@@ -724,31 +935,39 @@ void SAnimationEditorViewportTabBody::BindCommands()
 		MenuActions.PreviewSceneSettings,
 		FExecuteAction::CreateSP(this, &SAnimationEditorViewportTabBody::OpenPreviewSceneSettings));
 
-	TSharedRef<FAnimationViewportClient> EditorViewportClientRef = GetAnimationViewportClient();
-
-	CommandList.MapAction(
-		MenuActions.SetCPUSkinning,
-		FExecuteAction::CreateSP(EditorViewportClientRef, &FAnimationViewportClient::ToggleCPUSkinning),
-		FCanExecuteAction(),
-		FIsActionChecked::CreateSP(EditorViewportClientRef, &FAnimationViewportClient::IsSetCPUSkinningChecked));
-
-	CommandList.MapAction(
-		MenuActions.SetShowNormals,
-		FExecuteAction::CreateSP( EditorViewportClientRef, &FAnimationViewportClient::ToggleShowNormals ),
-		FCanExecuteAction(),
-		FIsActionChecked::CreateSP( EditorViewportClientRef, &FAnimationViewportClient::IsSetShowNormalsChecked ) );
-
-	CommandList.MapAction(
-		MenuActions.SetShowTangents,
-		FExecuteAction::CreateSP( EditorViewportClientRef, &FAnimationViewportClient::ToggleShowTangents ),
-		FCanExecuteAction(),
-		FIsActionChecked::CreateSP( EditorViewportClientRef, &FAnimationViewportClient::IsSetShowTangentsChecked ) );
-
-	CommandList.MapAction(
-		MenuActions.SetShowBinormals,
-		FExecuteAction::CreateSP( EditorViewportClientRef, &FAnimationViewportClient::ToggleShowBinormals ),
-		FCanExecuteAction(),
-		FIsActionChecked::CreateSP( EditorViewportClientRef, &FAnimationViewportClient::IsSetShowBinormalsChecked ) );
+	if (const TSharedPtr<FAnimationViewportClient>& AnimationViewportClientPtr = GetAnimationViewportClient())
+	{
+		if (FAnimationViewportClient* AnimViewportClient = AnimationViewportClientPtr.Get())
+		{
+			CommandList.MapAction(
+				MenuActions.SetCPUSkinning,
+				FExecuteAction::CreateSP(AnimViewportClient, &FAnimationViewportClient::ToggleCPUSkinning),
+				FCanExecuteAction(),
+				FIsActionChecked::CreateSP(AnimViewportClient, &FAnimationViewportClient::IsSetCPUSkinningChecked)
+			);
+
+			CommandList.MapAction(
+				MenuActions.SetShowNormals,
+				FExecuteAction::CreateSP(AnimViewportClient, &FAnimationViewportClient::ToggleShowNormals),
+				FCanExecuteAction(),
+				FIsActionChecked::CreateSP(AnimViewportClient, &FAnimationViewportClient::IsSetShowNormalsChecked)
+			);
+
+			CommandList.MapAction(
+				MenuActions.SetShowTangents,
+				FExecuteAction::CreateSP(AnimViewportClient, &FAnimationViewportClient::ToggleShowTangents),
+				FCanExecuteAction(),
+				FIsActionChecked::CreateSP(AnimViewportClient, &FAnimationViewportClient::IsSetShowTangentsChecked)
+			);
+
+			CommandList.MapAction(
+				MenuActions.SetShowBinormals,
+				FExecuteAction::CreateSP(AnimViewportClient, &FAnimationViewportClient::ToggleShowBinormals),
+				FCanExecuteAction(),
+				FIsActionChecked::CreateSP(AnimViewportClient, &FAnimationViewportClient::IsSetShowBinormalsChecked)
+			);
+		}
+	}
 
 	//Bind Show commands
 	const FAnimViewportShowCommands& ViewportShowMenuCommands = FAnimViewportShowCommands::Get();
@@ -1190,8 +1409,12 @@ void SAnimationEditorViewportTabBody::BindCommands()
 		FEditorViewportCommands::Get().FocusViewportToSelection,
 		FExecuteAction::CreateSP(this, &SAnimationEditorViewportTabBody::HandleFocusCamera));
 
-	TSharedPtr<FUICommandList> ToolkitCommandList = ConstCastSharedRef<FUICommandList>(GetAssetEditorToolkit()->GetToolkitCommands());
-	ToolkitCommandList->Append(UICommandList->AsShared());
+	if (const TSharedPtr<FAssetEditorToolkit>& AssetEditorToolkit = GetAssetEditorToolkit())
+	{
+		TSharedPtr<FUICommandList> ToolkitCommandList =
+			ConstCastSharedRef<FUICommandList>(AssetEditorToolkit->GetToolkitCommands());
+		ToolkitCommandList->Append(UICommandList->AsShared());
+	}
 }
 
 void SAnimationEditorViewportTabBody::OnSetTurnTableSpeed(int32 SpeedIndex)
@@ -1259,7 +1482,10 @@ void SAnimationEditorViewportTabBody::OnShowBoneNames()
 	});
 	RefreshViewport();
 
-	GetAnimationViewportClient()->UpdateBonesToDraw();
+	if (const TSharedPtr<FAnimationViewportClient>& AnimViewportClient = GetAnimationViewportClient())
+	{
+		AnimViewportClient->UpdateBonesToDraw();
+	}
 }
 
 void SAnimationEditorViewportTabBody::OnShowRawAnimation()
@@ -1416,12 +1642,20 @@ bool SAnimationEditorViewportTabBody::IsShowingBakedAnimation() const
 
 void SAnimationEditorViewportTabBody::OnShowDisplayInfo(int32 DisplayInfoMode)
 {
-	GetAnimationViewportClient()->OnSetShowMeshStats(DisplayInfoMode);
+	if (const TSharedPtr<FAnimationViewportClient>& AnimViewportClient = GetAnimationViewportClient())
+	{
+		AnimViewportClient->OnSetShowMeshStats(DisplayInfoMode);
+	}
 }
 
 bool SAnimationEditorViewportTabBody::IsShowingMeshInfo(int32 DisplayInfoMode) const
 {
-	return GetAnimationViewportClient()->GetShowMeshStats() == DisplayInfoMode;
+	if (const TSharedPtr<FAnimationViewportClient>& AnimViewportClient = GetAnimationViewportClient())
+	{
+		return AnimViewportClient->GetShowMeshStats() == DisplayInfoMode;
+	}
+
+	return false;
 }
 
 void SAnimationEditorViewportTabBody::OnShowOverlayNone()
@@ -1481,26 +1715,38 @@ bool SAnimationEditorViewportTabBody::IsShowingOverlayMorphTargetVerts() const
 
 void SAnimationEditorViewportTabBody::SetBoneDrawSize(float BoneDrawSize)
 {
-	TSharedRef<FAnimationViewportClient> AnimViewportClient = StaticCastSharedRef<FAnimationViewportClient>(LevelViewportClient.ToSharedRef());
-	AnimViewportClient->SetBoneDrawSize(BoneDrawSize);
+	if (const TSharedPtr<FAnimationViewportClient>& AnimViewportClient = GetAnimationViewportClient())
+	{
+		AnimViewportClient->SetBoneDrawSize(BoneDrawSize);
+	}
 }
 
 float SAnimationEditorViewportTabBody::GetBoneDrawSize() const
 {
-	TSharedRef<FAnimationViewportClient> AnimViewportClient = StaticCastSharedRef<FAnimationViewportClient>(LevelViewportClient.ToSharedRef());
-	return AnimViewportClient->GetBoneDrawSize();
+	if (const TSharedPtr<FAnimationViewportClient>& AnimViewportClient = GetAnimationViewportClient())
+	{
+		return AnimViewportClient->GetBoneDrawSize();
+	}
+
+	return 0.0f;
 }
 
 void SAnimationEditorViewportTabBody::SetCustomAnimationSpeed(float InCusteomAnimationSpeed)
 {
-	TSharedRef<FAnimationViewportClient> AnimViewportClient = StaticCastSharedRef<FAnimationViewportClient>(LevelViewportClient.ToSharedRef());
-	AnimViewportClient->SetCustomAnimationSpeed(InCusteomAnimationSpeed);
+	if (const TSharedPtr<FAnimationViewportClient>& AnimViewportClient = GetAnimationViewportClient())
+	{
+		AnimViewportClient->SetCustomAnimationSpeed(InCusteomAnimationSpeed);
+	}
 }
 
 float SAnimationEditorViewportTabBody::GetCustomAnimationSpeed() const
 {
-	TSharedRef<FAnimationViewportClient> AnimViewportClient = StaticCastSharedRef<FAnimationViewportClient>(LevelViewportClient.ToSharedRef());
-	return AnimViewportClient->GetCustomAnimationSpeed();
+	if (const TSharedPtr<FAnimationViewportClient>& AnimViewportClient = GetAnimationViewportClient())
+	{
+		return AnimViewportClient->GetCustomAnimationSpeed();
+	}
+
+	return 0.0f;
 }
 
 void SAnimationEditorViewportTabBody::SetCustomTurnTableSpeed(float InCustomTurnTableSpeed)
@@ -1516,26 +1762,38 @@ float SAnimationEditorViewportTabBody::GetCustomTurnTableSpeed() const
 
 void SAnimationEditorViewportTabBody::OnSetBoneDrawMode(int32 BoneDrawMode)
 {
-	TSharedRef<FAnimationViewportClient> AnimViewportClient = StaticCastSharedRef<FAnimationViewportClient>(LevelViewportClient.ToSharedRef());
-	AnimViewportClient->SetBoneDrawMode((EBoneDrawMode::Type)BoneDrawMode);
+	if (const TSharedPtr<FAnimationViewportClient>& AnimViewportClient = GetAnimationViewportClient())
+	{
+		AnimViewportClient->SetBoneDrawMode((EBoneDrawMode::Type)BoneDrawMode);
+	}
 }
 
 bool SAnimationEditorViewportTabBody::IsBoneDrawModeSet(int32 BoneDrawMode) const
 {
-	TSharedRef<FAnimationViewportClient> AnimViewportClient = StaticCastSharedRef<FAnimationViewportClient>(LevelViewportClient.ToSharedRef());
-	return AnimViewportClient->IsBoneDrawModeSet((EBoneDrawMode::Type)BoneDrawMode);
+	if (const TSharedPtr<FAnimationViewportClient>& AnimViewportClient = GetAnimationViewportClient())
+	{
+		return AnimViewportClient->IsBoneDrawModeSet((EBoneDrawMode::Type)BoneDrawMode);
+	}
+
+	return false;
 }
 
 void SAnimationEditorViewportTabBody::OnSetLocalAxesMode(int32 LocalAxesMode)
 {
-	TSharedRef<FAnimationViewportClient> AnimViewportClient = StaticCastSharedRef<FAnimationViewportClient>(LevelViewportClient.ToSharedRef());
-	AnimViewportClient->SetLocalAxesMode((ELocalAxesMode::Type)LocalAxesMode);
+	if (const TSharedPtr<FAnimationViewportClient>& AnimViewportClient = GetAnimationViewportClient())
+	{
+		AnimViewportClient->SetLocalAxesMode((ELocalAxesMode::Type)LocalAxesMode);
+	}
 }
 
 bool SAnimationEditorViewportTabBody::IsLocalAxesModeSet(int32 LocalAxesMode) const
 {
-	TSharedRef<FAnimationViewportClient> AnimViewportClient = StaticCastSharedRef<FAnimationViewportClient>(LevelViewportClient.ToSharedRef());
-	return AnimViewportClient->IsLocalAxesModeSet((ELocalAxesMode::Type)LocalAxesMode);
+	if (const TSharedPtr<FAnimationViewportClient>& AnimViewportClient = GetAnimationViewportClient())
+	{
+		return AnimViewportClient->IsLocalAxesModeSet((ELocalAxesMode::Type)LocalAxesMode);
+	}
+
+	return false;
 }
 
 void SAnimationEditorViewportTabBody::OnShowSockets()
@@ -1574,39 +1832,58 @@ bool SAnimationEditorViewportTabBody::IsShowingAttributes() const
 
 void SAnimationEditorViewportTabBody::OnToggleAutoAlignFloor()
 {
-	TSharedRef<FAnimationViewportClient> AnimViewportClient = StaticCastSharedRef<FAnimationViewportClient>(LevelViewportClient.ToSharedRef());
-	AnimViewportClient->OnToggleAutoAlignFloor();
+	if (const TSharedPtr<FAnimationViewportClient>& AnimViewportClient = GetAnimationViewportClient())
+	{
+		AnimViewportClient->OnToggleAutoAlignFloor();
+	}
 }
 
 bool SAnimationEditorViewportTabBody::IsAutoAlignFloor() const
 {
-	TSharedRef<FAnimationViewportClient> AnimViewportClient = StaticCastSharedRef<FAnimationViewportClient>(LevelViewportClient.ToSharedRef());
-	return AnimViewportClient->IsAutoAlignFloor();
+	if (const TSharedPtr<FAnimationViewportClient>& AnimViewportClient = GetAnimationViewportClient())
+	{
+		return AnimViewportClient->IsAutoAlignFloor();
+	}
+
+	return false;
 }
 
 /** Function to set the current playback speed*/
 void SAnimationEditorViewportTabBody::OnSetPlaybackSpeed(int32 PlaybackSpeedMode)
 {
-	TSharedRef<FAnimationViewportClient> AnimViewportClient = StaticCastSharedRef<FAnimationViewportClient>(LevelViewportClient.ToSharedRef());
-	AnimViewportClient->SetPlaybackSpeedMode((EAnimationPlaybackSpeeds::Type)PlaybackSpeedMode);
+	if (const TSharedPtr<FAnimationViewportClient>& AnimViewportClient = GetAnimationViewportClient())
+	{
+		AnimViewportClient->SetPlaybackSpeedMode((EAnimationPlaybackSpeeds::Type)PlaybackSpeedMode);
+	}
 }
 
 bool SAnimationEditorViewportTabBody::IsPlaybackSpeedSelected(int32 PlaybackSpeedMode)
 {
-	TSharedRef<FAnimationViewportClient> AnimViewportClient = StaticCastSharedRef<FAnimationViewportClient>(LevelViewportClient.ToSharedRef());
-	return PlaybackSpeedMode == AnimViewportClient->GetPlaybackSpeedMode();
+	if (const TSharedPtr<FAnimationViewportClient>& AnimViewportClient = GetAnimationViewportClient())
+	{
+		return PlaybackSpeedMode == AnimViewportClient->GetPlaybackSpeedMode();
+	}
+
+	return false;
 }
 
 void SAnimationEditorViewportTabBody::ShowBound()
 {
-	TSharedRef<FAnimationViewportClient> AnimViewportClient = StaticCastSharedRef<FAnimationViewportClient>(LevelViewportClient.ToSharedRef());	
-	AnimViewportClient->ToggleShowBounds();
-
-	ForEachDebugMesh([AnimViewportClient](UDebugSkelMeshComponent* PreviewMeshComponent)
+	if (const TSharedPtr<FAnimationViewportClient>& AnimViewportClient = GetAnimationViewportClient())
 	{
-		PreviewMeshComponent->bDisplayBound = AnimViewportClient->EngineShowFlags.Bounds;
-		PreviewMeshComponent->RecreateRenderState_Concurrent();
-	});
+		AnimViewportClient->ToggleShowBounds();
+
+		ForEachDebugMesh(
+			[AnimViewportClientWeak = AnimViewportClient.ToWeakPtr()](UDebugSkelMeshComponent* PreviewMeshComponent)
+			{
+				if (const TSharedPtr<FAnimationViewportClient>& AnimViewportClient = AnimViewportClientWeak.Pin())
+				{
+					PreviewMeshComponent->bDisplayBound = AnimViewportClient->EngineShowFlags.Bounds;
+					PreviewMeshComponent->RecreateRenderState_Concurrent();
+				}
+			}
+		);
+	}
 }
 
 bool SAnimationEditorViewportTabBody::CanShowBound() const
@@ -1616,8 +1893,12 @@ bool SAnimationEditorViewportTabBody::CanShowBound() const
 
 bool SAnimationEditorViewportTabBody::IsShowBoundEnabled() const
 {
-	TSharedRef<FAnimationViewportClient> AnimViewportClient = StaticCastSharedRef<FAnimationViewportClient>(LevelViewportClient.ToSharedRef());	
-	return AnimViewportClient->IsSetShowBoundsChecked();
+	if (const TSharedPtr<FAnimationViewportClient>& AnimViewportClient = GetAnimationViewportClient())
+	{
+		return AnimViewportClient->IsSetShowBoundsChecked();
+	}
+
+	return false;
 }
 
 void SAnimationEditorViewportTabBody::ToggleShowPreviewMesh()
@@ -1756,8 +2037,13 @@ void SAnimationEditorViewportTabBody::AnimChanged(UAnimationAsset* AnimAsset)
 
 void SAnimationEditorViewportTabBody::ComboBoxSelectionChanged(TSharedPtr<FString> NewSelection, ESelectInfo::Type SelectInfo)
 {
+	const TSharedPtr<FAnimationViewportClient>& AnimViewportClient = GetAnimationViewportClient();
+	if (!AnimViewportClient)
+	{
+		return;
+	}
+
 	int32 NewUVSelection = UVChannels.Find(NewSelection) - 1;
-	TSharedRef<FAnimationViewportClient> AnimViewportClient = StaticCastSharedRef<FAnimationViewportClient>(LevelViewportClient.ToSharedRef());
 
 	// "None" is index -1 here.
 	if ( NewUVSelection < 0 )
@@ -1794,6 +2080,12 @@ void SAnimationEditorViewportTabBody::PopulateNumUVChannels()
 
 void SAnimationEditorViewportTabBody::PopulateUVChoices()
 {
+	const TSharedPtr<FAnimationViewportClient>& AnimViewportClient = GetAnimationViewportClient();
+	if (!AnimViewportClient)
+	{
+		return;
+	}
+
 	// Fill out the UV channels combo.
 	UVChannels.Empty();
 
@@ -1810,7 +2102,6 @@ void SAnimationEditorViewportTabBody::PopulateUVChoices()
 				UVChannels.Add( MakeShareable( new FString( FText::Format( NSLOCTEXT("AnimationEditorViewport", "UVChannel_ID", "UV Channel {0}"), FText::AsNumber( UVChannelID ) ).ToString() ) ) );
 			}
 
-			TSharedRef<FAnimationViewportClient> AnimViewportClient = StaticCastSharedRef<FAnimationViewportClient>(LevelViewportClient.ToSharedRef());
 			int32 CurrentUVChannel = AnimViewportClient->GetUVChannelToDraw();
 			if (!UVChannels.IsValidIndex(CurrentUVChannel))
 			{
@@ -2037,9 +2328,52 @@ void SAnimationEditorViewportTabBody::OnDebugForcedLODChanged()
 	}
 }
 
-TSharedRef<FAnimationViewportClient> SAnimationEditorViewportTabBody::GetAnimationViewportClient() const
+void SAnimationEditorViewportTabBody::OnBeginSliderMovementFloorOffset()
 {
-	return StaticCastSharedRef<FAnimationViewportClient>(LevelViewportClient.ToSharedRef());
+	// This value is saved in a UPROPERTY for the floor mesh, so changes are transactional
+	PendingTransaction = MakeUnique<FScopedTransaction>(LOCTEXT("SetFloorOffset", "Set Floor Offset"));
+	PinnedCommands->AddCustomWidget(TEXT("FloorOffsetWidget"));
+}
+
+void SAnimationEditorViewportTabBody::OnFloorOffsetChanged(float InNewValue)
+{
+	FAnimationViewportClient& AnimViewportClient = (FAnimationViewportClient&)GetLevelViewportClient();
+
+	AnimViewportClient.SetFloorOffset(InNewValue);
+
+	PinnedCommands->AddCustomWidget(TEXT("FloorOffsetWidget"));
+}
+
+void SAnimationEditorViewportTabBody::OnFloorOffsetCommitted(float InNewValue, ETextCommit::Type InCommitType)
+{
+	if (!PendingTransaction)
+	{
+		// Create the transaction here if it doesn't already exist. This can happen when changes come via text entry to the slider.
+		PendingTransaction = MakeUnique<FScopedTransaction>(LOCTEXT("SetFloorOffset", "Set Floor Offset"));
+	}
+
+	FAnimationViewportClient& AnimViewportClient = (FAnimationViewportClient&)GetLevelViewportClient();
+
+	AnimViewportClient.SetFloorOffset(InNewValue);
+
+	PinnedCommands->AddCustomWidget(TEXT("FloorOffsetWidget"));
+
+	PendingTransaction.Reset();
+}
+
+void SAnimationEditorViewportTabBody::CreatePinnedCommands()
+{
+	// Create our pinned commands before we bind commands
+	IPinnedCommandListModule& PinnedCommandListModule =
+		FModuleManager::LoadModuleChecked<IPinnedCommandListModule>(TEXT("PinnedCommandList"));
+	PinnedCommands =
+		PinnedCommandListModule.CreatePinnedCommandList((ContextName != NAME_None) ? ContextName : TEXT("PersonaViewport"));
+	PinnedCommands->SetStyle(&FAppStyle::Get(), TEXT("ViewportPinnedCommandList"));
+}
+
+TSharedPtr<FAnimationViewportClient> SAnimationEditorViewportTabBody::GetAnimationViewportClient() const
+{
+	return StaticCastSharedPtr<FAnimationViewportClient>(LevelViewportClient);
 }
 
 void SAnimationEditorViewportTabBody::OpenPreviewSceneSettings()
@@ -2049,26 +2383,38 @@ void SAnimationEditorViewportTabBody::OpenPreviewSceneSettings()
 
 void SAnimationEditorViewportTabBody::SetCameraFollowMode(EAnimationViewportCameraFollowMode InCameraFollowMode, FName InBoneName)
 {
-	TSharedRef<FAnimationViewportClient> AnimViewportClient = StaticCastSharedRef<FAnimationViewportClient>(LevelViewportClient.ToSharedRef());
-	AnimViewportClient->SetCameraFollowMode(InCameraFollowMode, InBoneName);
+	if (const TSharedPtr<FAnimationViewportClient>& AnimViewportClient = GetAnimationViewportClient())
+	{
+		AnimViewportClient->SetCameraFollowMode(InCameraFollowMode, InBoneName);
+	}
 }
 
 bool SAnimationEditorViewportTabBody::IsCameraFollowEnabled(EAnimationViewportCameraFollowMode InCameraFollowMode) const
 {
-	TSharedRef<FAnimationViewportClient> AnimViewportClient = StaticCastSharedRef<FAnimationViewportClient>(LevelViewportClient.ToSharedRef());
-	return (AnimViewportClient->GetCameraFollowMode() == InCameraFollowMode);
+	if (const TSharedPtr<FAnimationViewportClient>& AnimViewportClient = GetAnimationViewportClient())
+	{
+		return (AnimViewportClient->GetCameraFollowMode() == InCameraFollowMode);
+	}
+
+	return false;
 }
 
 void SAnimationEditorViewportTabBody::ToggleRotateCameraToFollowBone()
 {
-	TSharedRef<FAnimationViewportClient> AnimViewportClient = StaticCastSharedRef<FAnimationViewportClient>(LevelViewportClient.ToSharedRef());
-	AnimViewportClient->ToggleRotateCameraToFollowBone();
+	if (const TSharedPtr<FAnimationViewportClient>& AnimViewportClient = GetAnimationViewportClient())
+	{
+		AnimViewportClient->ToggleRotateCameraToFollowBone();
+	}
 }
 
 bool SAnimationEditorViewportTabBody::GetShouldRotateCameraToFollowBone() const
 {
-	TSharedRef<FAnimationViewportClient> AnimViewportClient = StaticCastSharedRef<FAnimationViewportClient>(LevelViewportClient.ToSharedRef());
-	return AnimViewportClient->GetShouldRotateCameraToFollowBone();
+	if (const TSharedPtr<FAnimationViewportClient>& AnimViewportClient = GetAnimationViewportClient())
+	{
+		return AnimViewportClient->GetShouldRotateCameraToFollowBone();
+	}
+
+	return false;
 }
 
 void SAnimationEditorViewportTabBody::TogglePauseAnimationOnCameraMove()
@@ -2083,38 +2429,56 @@ bool SAnimationEditorViewportTabBody::GetShouldPauseAnimationOnCameraMove() cons
 
 FName SAnimationEditorViewportTabBody::GetCameraFollowBoneName() const
 {
-	TSharedRef<FAnimationViewportClient> AnimViewportClient = StaticCastSharedRef<FAnimationViewportClient>(LevelViewportClient.ToSharedRef());
-	return AnimViewportClient->GetCameraFollowBoneName();
+	if (const TSharedPtr<FAnimationViewportClient>& AnimViewportClient = GetAnimationViewportClient())
+	{
+		return AnimViewportClient->GetCameraFollowBoneName();
+	}
+
+	return NAME_None;
 }
 
 void SAnimationEditorViewportTabBody::SaveCameraAsDefault()
 {
-	TSharedRef<FAnimationViewportClient> AnimViewportClient = StaticCastSharedRef<FAnimationViewportClient>(LevelViewportClient.ToSharedRef());
-	AnimViewportClient->SaveCameraAsDefault();
+	if (const TSharedPtr<FAnimationViewportClient>& AnimViewportClient = GetAnimationViewportClient())
+	{
+		AnimViewportClient->SaveCameraAsDefault();
+	}
 }
 
 void SAnimationEditorViewportTabBody::ClearDefaultCamera()
 {
-	TSharedRef<FAnimationViewportClient> AnimViewportClient = StaticCastSharedRef<FAnimationViewportClient>(LevelViewportClient.ToSharedRef());
-	AnimViewportClient->ClearDefaultCamera();
+	if (const TSharedPtr<FAnimationViewportClient>& AnimViewportClient = GetAnimationViewportClient())
+	{
+		AnimViewportClient->ClearDefaultCamera();
+	}
 }
 
 void SAnimationEditorViewportTabBody::JumpToDefaultCamera()
 {
-	TSharedRef<FAnimationViewportClient> AnimViewportClient = StaticCastSharedRef<FAnimationViewportClient>(LevelViewportClient.ToSharedRef());
-	AnimViewportClient->JumpToDefaultCamera();
+	if (const TSharedPtr<FAnimationViewportClient>& AnimViewportClient = GetAnimationViewportClient())
+	{
+		AnimViewportClient->JumpToDefaultCamera();
+	}
 }
 
 bool SAnimationEditorViewportTabBody::CanSaveCameraAsDefault() const
 {
-	TSharedRef<FAnimationViewportClient> AnimViewportClient = StaticCastSharedRef<FAnimationViewportClient>(LevelViewportClient.ToSharedRef());
-	return AnimViewportClient->CanSaveCameraAsDefault();
+	if (const TSharedPtr<FAnimationViewportClient>& AnimViewportClient = GetAnimationViewportClient())
+	{
+		return AnimViewportClient->CanSaveCameraAsDefault();
+	}
+
+	return false;
 }
 
 bool SAnimationEditorViewportTabBody::HasDefaultCameraSet() const
 {
-	TSharedRef<FAnimationViewportClient> AnimViewportClient = StaticCastSharedRef<FAnimationViewportClient>(LevelViewportClient.ToSharedRef());
-	return (AnimViewportClient->HasDefaultCameraSet());
+	if (const TSharedPtr<FAnimationViewportClient>& AnimViewportClient = GetAnimationViewportClient())
+	{
+		return (AnimViewportClient->HasDefaultCameraSet());
+	}
+
+	return false;
 }
 
 bool SAnimationEditorViewportTabBody::CanChangeCameraMode() const
@@ -2125,22 +2489,38 @@ bool SAnimationEditorViewportTabBody::CanChangeCameraMode() const
 
 void SAnimationEditorViewportTabBody::OnToggleMuteAudio()
 {
-	GetAnimationViewportClient()->OnToggleMuteAudio();
+	if (const TSharedPtr<FAnimationViewportClient>& AnimViewportClient = GetAnimationViewportClient())
+	{
+		AnimViewportClient->OnToggleMuteAudio();
+	}
 }
 
 bool SAnimationEditorViewportTabBody::IsAudioMuted() const
 {
-	return GetAnimationViewportClient()->IsAudioMuted();
+	if (const TSharedPtr<FAnimationViewportClient>& AnimViewportClient = GetAnimationViewportClient())
+	{
+		return AnimViewportClient->IsAudioMuted();
+	}
+
+	return false;
 }
 
 void SAnimationEditorViewportTabBody::OnToggleUseAudioAttenuation()
 {
-	GetAnimationViewportClient()->OnToggleUseAudioAttenuation();
+	if (const TSharedPtr<FAnimationViewportClient>& AnimViewportClient = GetAnimationViewportClient())
+	{
+		AnimViewportClient->OnToggleUseAudioAttenuation();
+	}
 }
 
 bool SAnimationEditorViewportTabBody::IsAudioAttenuationEnabled() const
 {
-	return GetAnimationViewportClient()->IsUsingAudioAttenuation();
+	if (const TSharedPtr<FAnimationViewportClient>& AnimViewportClient = GetAnimationViewportClient())
+	{
+		return AnimViewportClient->IsUsingAudioAttenuation();
+	}
+
+	return false;
 }
 
 void SAnimationEditorViewportTabBody::SetProcessRootMotionMode(EProcessRootMotionMode Mode)
@@ -2801,8 +3181,10 @@ void SAnimationEditorViewportTabBody::AddSkinWeightProfileNotification()
 
 void SAnimationEditorViewportTabBody::HandleFocusCamera()
 {
-	TSharedRef<FAnimationViewportClient> AnimViewportClient = StaticCastSharedRef<FAnimationViewportClient>(LevelViewportClient.ToSharedRef());
-	AnimViewportClient->FocusViewportOnPreviewMesh(false);
+	if (const TSharedPtr<FAnimationViewportClient>& AnimViewportClient = GetAnimationViewportClient())
+	{
+		AnimViewportClient->FocusViewportOnPreviewMesh(false);
+	}
 }
 
 #undef LOCTEXT_NAMESPACE
diff --git a/Engine/Source/Editor/Persona/Private/SAnimationEditorViewport.h b/Engine/Source/Editor/Persona/Private/SAnimationEditorViewport.h
--- a/Engine/Source/Editor/Persona/Private/SAnimationEditorViewport.h
+++ b/Engine/Source/Editor/Persona/Private/SAnimationEditorViewport.h
@@ -1,6 +1,5 @@
 // Copyright Epic Games, Inc. All Rights Reserved.
 
-
 #pragma once
 
 #include "CoreMinimal.h"
@@ -20,9 +19,10 @@
 #include "PersonaModule.h"
 #include "SNameComboBox.h"
 
-class SAnimationEditorViewportTabBody;
+class FScopedTransaction;
 class FUICommandList_Pinnable;
 class SAnimViewportToolBar;
+class SAnimationEditorViewportTabBody;
 
 struct FAnimationEditorViewportRequiredArgs
 {
@@ -98,6 +98,8 @@ protected:
 	// SEditorViewport interface
 	virtual TSharedRef<FEditorViewportClient> MakeEditorViewportClient() override;
 	virtual TSharedPtr<SWidget> MakeViewportToolbar() override;
+	virtual TSharedPtr<SWidget> BuildViewportToolbar() override;
+	virtual TSharedPtr<IPreviewProfileController> CreatePreviewProfileController() override;
 	virtual void OnFocusViewportToSelection() override;
 	virtual void PopulateViewportOverlays(TSharedRef<SOverlay> Overlay) override;
 	// End of SEditorViewport interface
@@ -267,7 +269,7 @@ public:
 	}
 
 	/** Gets the animation viewport client */
-	TSharedRef<class FAnimationViewportClient> GetAnimationViewportClient() const;
+	TSharedPtr<class FAnimationViewportClient> GetAnimationViewportClient() const;
 
 	/** Returns Detail description of what's going with viewport **/
 	FText GetDisplayString() const;
@@ -308,14 +310,17 @@ public:
 	void OnLODModelChanged();
 	void OnDebugForcedLODChanged();
 
+	void OnBeginSliderMovementFloorOffset();
+	void OnFloorOffsetChanged(float InNewValue);
+	void OnFloorOffsetCommitted(float InNewValue, ETextCommit::Type InCommitType);
+
+	const TSharedPtr<IPinnedCommandList>& GetPinnedCommands();
+
 	/** Get the preview scene we are viewing */
 	TSharedRef<class FAnimationEditorPreviewScene> GetPreviewScene() const { return PreviewScenePtr.Pin().ToSharedRef(); }
 
 	/** Get the asset editor toolkit we are embedded in */
-	TSharedPtr<class FAssetEditorToolkit> GetAssetEditorToolkit() const { return AssetEditorToolkitPtr.Pin(); }
-
-protected:
-
+	TSharedPtr<class FAssetEditorToolkit> GetAssetEditorToolkit() const;
 
 private:
 	bool IsVisible() const;
@@ -497,6 +502,8 @@ private:
 	/** Focus the viewport on the preview mesh */
 	void HandleFocusCamera();
 
+	void CreatePinnedCommands();
+
 public:
 	/** Called to determine whether the camera mode menu options should be enabled */
 	bool CanChangeCameraMode() const;
@@ -617,6 +624,13 @@ private:
 	/** Delegate used to invoke tabs in the containing asset editor */
 	FOnInvokeTab OnInvokeTab;
 
+	/** Pinned commands widget */
+	TSharedPtr<IPinnedCommandList> PinnedCommands;
+
+	/** Transaction to handle scoping updates for sliders we own that transact objects (i.e. floor mesh) */
+	TUniquePtr<FScopedTransaction> PendingTransaction;
+
+	FName ContextName;
 public:
 	/** UV Channel Selector */
 	TSharedPtr< class STextComboBox > UVChannelCombo;
diff --git a/Engine/Source/Editor/Persona/Private/ViewportToolbar/AnimViewportContext.h b/Engine/Source/Editor/Persona/Private/ViewportToolbar/AnimViewportContext.h
new file mode 100644
--- /dev/null
+++ b/Engine/Source/Editor/Persona/Private/ViewportToolbar/AnimViewportContext.h
@@ -0,0 +1,15 @@
+// Copyright Epic Games, Inc. All Rights Reserved.
+
+#pragma once
+
+#include "SAnimationEditorViewport.h"
+#include "AnimViewportContext.generated.h"
+
+UCLASS()
+class UAnimViewportContext : public UObject
+{
+	GENERATED_BODY()
+
+public:
+	TWeakPtr<SAnimationEditorViewportTabBody> ViewportTabBody;
+};
diff --git a/Engine/Source/Editor/Persona/Private/ViewportToolbar/AnimationEditorMenus.cpp b/Engine/Source/Editor/Persona/Private/ViewportToolbar/AnimationEditorMenus.cpp
new file mode 100644
--- /dev/null
+++ b/Engine/Source/Editor/Persona/Private/ViewportToolbar/AnimationEditorMenus.cpp
@@ -0,0 +1,1318 @@
+﻿// Copyright Epic Games, Inc. All Rights Reserved.
+
+#include "AnimationEditorMenus.h"
+
+#include "AnimPreviewInstance.h"
+#include "AnimViewportContext.h"
+#include "AnimViewportLODCommands.h"
+#include "AnimViewportMenuCommands.h"
+#include "AnimViewportPlaybackCommands.h"
+#include "AnimViewportShowCommands.h"
+#include "Animation/MirrorDataTable.h"
+#include "AnimationEditorWidgets.h"
+#include "BoneSelectionWidget.h"
+#include "BufferVisualizationMenuCommands.h"
+#include "ClothingSystemEditorInterfaceModule.h"
+#include "ContentBrowserModule.h"
+#include "EditorViewportCommands.h"
+#include "IContentBrowserSingleton.h"
+#include "IPinnedCommandList.h"
+#include "SAnimationEditorViewport.h"
+#include "ShowFlagMenuCommands.h"
+#include "SimulationEditorExtender.h"
+#include "ToolMenus.h"
+#include "UICommandList_Pinnable.h"
+#include "ViewportToolbar/UnrealEdViewportToolbar.h"
+#include "Widgets/Input/SSpinBox.h"
+#include "Widgets/Input/STextComboBox.h"
+
+#define LOCTEXT_NAMESPACE "AnimEditorViewportToolbar"
+
+namespace UE::AnimationEditor::Private
+{
+// Convenience function to retrieve Anim Editor Viewport Tab from UToolMenu
+TSharedPtr<SAnimationEditorViewportTabBody> GetAnimationEditorViewportTab(const UToolMenu* InMenu)
+{
+	if (InMenu)
+	{
+		if (UAnimViewportContext* const AnimViewportContext = InMenu->FindContext<UAnimViewportContext>())
+		{
+			return AnimViewportContext->ViewportTabBody.Pin();
+		}
+	}
+
+	return nullptr;
+}
+
+// Convenience function to retrieve Anim Editor Viewport Tab from FToolMenuSection
+TSharedPtr<SAnimationEditorViewportTabBody> GetAnimationEditorViewportTab(const FToolMenuSection& InDynamicSection)
+{
+	if (UAnimViewportContext* const AnimViewportContext = InDynamicSection.FindContext<UAnimViewportContext>())
+	{
+		return AnimViewportContext->ViewportTabBody.Pin();
+	}
+
+	return nullptr;
+}
+
+void PopulateLODSubmenu(UToolMenu* InMenu)
+{
+	if (!InMenu)
+	{
+		return;
+	}
+
+	const TSharedPtr<SAnimationEditorViewportTabBody> AnimEditorViewportTab = GetAnimationEditorViewportTab(InMenu);
+	if (!AnimEditorViewportTab)
+	{
+		return;
+	}
+
+	const FAnimViewportLODCommands& Actions = FAnimViewportLODCommands::Get();
+
+	const USkeletalMesh* PreviewMesh = AnimEditorViewportTab->GetPreviewScene()->GetPreviewMesh();
+	auto IsAutoGenerated = [PreviewMeshWeak = TWeakObjectPtr<const USkeletalMesh>(PreviewMesh)](const int32 InLODId) -> bool
+	{
+		TStrongObjectPtr<const USkeletalMesh> PreviewMeshPinned = PreviewMeshWeak.Pin();
+		if (!PreviewMeshPinned || PreviewMeshPinned->IsCompiling())
+		{
+			return false;
+		}
+		return PreviewMeshPinned->IsValidLODIndex(InLODId) && PreviewMeshPinned->GetLODInfo(InLODId)->bHasBeenSimplified;
+	};
+	static const FText LODFormat = LOCTEXT("LODFmt", "LOD {0}");
+	static const FText GeneratedLODFormat = LOCTEXT("GeneratedLODFmt", "LOD {0} (generated)");
+
+	{
+		// LOD Models
+		FToolMenuSection& LODSection =
+			InMenu->AddSection("AnimViewportPreviewLODs", LOCTEXT("ShowLOD_PreviewLabel", "Preview LODs"));
+		{
+			LODSection.AddMenuEntry(Actions.LODDebug);
+			LODSection.AddMenuEntry(Actions.LODAuto);
+			LODSection.AddMenuEntry(Actions.LOD0);
+
+			const int32 LODCount = AnimEditorViewportTab->GetLODModelCount();
+			for (int32 LODId = 1; LODId < LODCount; ++LODId)
+			{
+				const FText LODNumber = FText::AsNumber(LODId);
+				const FText TitleLabel = IsAutoGenerated(LODId) ? FText::Format(GeneratedLODFormat, LODNumber)
+																: FText::Format(LODFormat, LODNumber);
+
+				FUIAction Action(
+					FExecuteAction::CreateSP(
+						AnimEditorViewportTab.ToSharedRef(), &SAnimationEditorViewportTabBody::OnSetLODModel, LODId + 1
+					),
+					FCanExecuteAction(),
+					FIsActionChecked::CreateSP(
+						AnimEditorViewportTab.ToSharedRef(), &SAnimationEditorViewportTabBody::IsLODModelSelected, LODId + 1
+					)
+				);
+
+				LODSection.AddMenuEntry(
+					FName(TitleLabel.ToString()), FText(), FText(), FSlateIcon(), Action, EUserInterfaceActionType::RadioButton
+				);
+			}
+		}
+	}
+}
+
+void FillCharacterMirrorMenu(UToolMenu* InMenu)
+{
+	if (!InMenu)
+	{
+		return;
+	}
+
+	TSharedPtr<SAnimationEditorViewportTabBody> ViewportTab = GetAnimationEditorViewportTab(InMenu);
+	if (!ViewportTab)
+	{
+		return;
+	}
+
+	UDebugSkelMeshComponent* PreviewComp = ViewportTab->GetPreviewScene()->GetPreviewMeshComponent();
+	USkeletalMesh* Mesh = PreviewComp->GetSkeletalMeshAsset();
+	UAnimPreviewInstance* PreviewInstance = PreviewComp->PreviewInstance;
+	if (Mesh && PreviewInstance)
+	{
+		USkeleton* Skeleton = Mesh->GetSkeleton();
+
+		FAssetPickerConfig AssetPickerConfig;
+		AssetPickerConfig.Filter.ClassPaths.Add(UMirrorDataTable::StaticClass()->GetClassPathName());
+		AssetPickerConfig.Filter.bRecursiveClasses = false;
+		AssetPickerConfig.bAllowNullSelection = true;
+		AssetPickerConfig.OnShouldFilterAsset =
+			FOnShouldFilterAsset::CreateUObject(Skeleton, &USkeleton::ShouldFilterAsset, TEXT("Skeleton"));
+		AssetPickerConfig.InitialAssetSelection = FAssetData(PreviewInstance->GetMirrorDataTable());
+		AssetPickerConfig.OnAssetSelected = FOnAssetSelected::CreateLambda(
+			[ViewportTabWeak = ViewportTab.ToWeakPtr()](const FAssetData& InSelectedMirrorTableData)
+			{
+				if (const TSharedPtr<SAnimationEditorViewportTabBody>& ViewportTabPinned = ViewportTabWeak.Pin())
+				{
+					UDebugSkelMeshComponent* PreviewComp = ViewportTabPinned->GetPreviewScene()->GetPreviewMeshComponent();
+					USkeletalMesh* Mesh = PreviewComp->GetSkeletalMeshAsset();
+					UAnimPreviewInstance* PreviewInstance = PreviewComp->PreviewInstance;
+					if (Mesh && PreviewInstance)
+					{
+						UMirrorDataTable* MirrorDataTable = Cast<UMirrorDataTable>(InSelectedMirrorTableData.GetAsset());
+						PreviewInstance->SetMirrorDataTable(MirrorDataTable);
+						PreviewComp->OnMirrorDataTableChanged();
+					}
+				}
+			}
+		);
+		AssetPickerConfig.InitialAssetViewType = EAssetViewType::List;
+		AssetPickerConfig.ThumbnailScale = 0.1f;
+		AssetPickerConfig.bAddFilterUI = false;
+
+		FContentBrowserModule& ContentBrowserModule =
+			FModuleManager::Get().LoadModuleChecked<FContentBrowserModule>(TEXT("ContentBrowser"));
+
+		FToolMenuEntry CharacterMirrorMenu =
+			FToolMenuEntry::InitWidget("", ContentBrowserModule.Get().CreateAssetPicker(AssetPickerConfig), FText());
+
+		InMenu->AddMenuEntry("CharacterMirrorMenu", CharacterMirrorMenu);
+	}
+}
+
+void FillCharacterClothingMenu(UToolMenu* InMenu)
+{
+	if (!InMenu)
+	{
+		return;
+	}
+
+	// Call into the clothing editor module to customize the menu (this is mainly for debug visualizations and sim-specific options)
+	TSharedPtr<SAnimationEditorViewportTabBody> AnimEditorViewportTab = GetAnimationEditorViewportTab(InMenu);
+	if (!AnimEditorViewportTab)
+	{
+		return;
+	}
+
+	const FAnimViewportShowCommands& Actions = FAnimViewportShowCommands::Get();
+
+	FToolMenuSection& ClothPreviewSection = InMenu->AddSection("ClothPreview", LOCTEXT("ClothPreview_Label", "Simulation"));
+	{
+		ClothPreviewSection.AddMenuEntry(Actions.EnableClothSimulation);
+		ClothPreviewSection.AddMenuEntry(Actions.ResetClothSimulation);
+
+		TSharedPtr<SWidget> WindWidget =
+			SNew(UE::AnimationEditor::SClothWindSettings).AnimEditorViewport(AnimEditorViewportTab);
+		ClothPreviewSection.AddEntry(FToolMenuEntry::InitWidget(
+			"", WindWidget.ToSharedRef(), LOCTEXT("ClothPreview_WindStrength", "Wind Strength:")
+		));
+
+		TSharedPtr<SWidget> GravityWidget =
+			SNew(UE::AnimationEditor::SGravitySettings).AnimEditorViewport(AnimEditorViewportTab);
+		ClothPreviewSection.AddEntry(FToolMenuEntry::InitWidget(
+			"", GravityWidget.ToSharedRef(), LOCTEXT("ClothPreview_GravityScale", "Gravity Scale:")
+		));
+
+		ClothPreviewSection.AddMenuEntry(Actions.EnableCollisionWithAttachedClothChildren);
+		ClothPreviewSection.AddMenuEntry(Actions.PauseClothWithAnim);
+	}
+
+	FToolMenuSection& ClothAdditionalVisualizationSection = InMenu->AddSection(
+		"ClothAdditionalVisualization", LOCTEXT("ClothAdditionalVisualization_Label", "Sections Display Mode")
+	);
+	{
+		ClothAdditionalVisualizationSection.AddMenuEntry(Actions.ShowAllSections);
+		ClothAdditionalVisualizationSection.AddMenuEntry(Actions.ShowOnlyClothSections);
+		ClothAdditionalVisualizationSection.AddMenuEntry(Actions.HideOnlyClothSections);
+	}
+
+	// Call into the clothing editor module to customize the menu (this is mainly for debug visualizations and sim-specific options)
+	if (const TSharedPtr<FAnimationViewportClient>& AnimationViewportClient =
+			AnimEditorViewportTab->GetAnimationViewportClient())
+	{
+		TSharedRef<IPersonaPreviewScene> PreviewScene = AnimationViewportClient->GetPreviewScene();
+		if (UDebugSkelMeshComponent* PreviewComponent = PreviewScene->GetPreviewMeshComponent())
+		{
+			if (PreviewComponent->ClothingSimulationFactory) // The cloth plugin could be disabled, and the factory would be null in this case
+			{
+				// Currently using FNewToolMenuDelegateLegacy since this extension is being done via
+				// ISimulationEditorExtender::ExtendViewportShowMenu(FMenuBuilder&InMenuBuilder, TSharedRef<IPersonaPreviewScene> InPreviewScene)
+				// todo: that function might need a new version using UToolMenu
+				FToolMenuSection& UnnamedSection = InMenu->AddSection(NAME_None);
+				UnnamedSection.AddDynamicEntry(
+					"SimulationEditorExtender",
+					FNewToolMenuDelegateLegacy::CreateLambda(
+						[](FMenuBuilder& InMenuBuilder, UToolMenu* InMenu)
+						{
+							TSharedPtr<SAnimationEditorViewportTabBody> AnimEditorViewportTab =
+								GetAnimationEditorViewportTab(InMenu);
+							if (!AnimEditorViewportTab)
+							{
+								return;
+							}
+
+							TSharedRef<FAnimationEditorPreviewScene> PreviewScene =
+								AnimEditorViewportTab->GetPreviewScene();
+
+							UDebugSkelMeshComponent* PreviewComponent = PreviewScene->GetPreviewMeshComponent();
+							if (!PreviewComponent)
+							{
+								return;
+							}
+
+							FClothingSystemEditorInterfaceModule& ClothingEditorModule =
+								FModuleManager::LoadModuleChecked<FClothingSystemEditorInterfaceModule>(
+									TEXT("ClothingSystemEditorInterface")
+								);
+
+							if (ISimulationEditorExtender* Extender = ClothingEditorModule.GetSimulationEditorExtender(
+									PreviewComponent->ClothingSimulationFactory->GetFName()
+								))
+							{
+								// Calling end section will set bSectionNeedsToBeApplied to false.
+								// Without doing so, calling ExtendViewportShowMenu will end up triggering an ensure
+								InMenuBuilder.EndSection();
+								Extender->ExtendViewportShowMenu(InMenuBuilder, PreviewScene);
+							}
+						}
+					)
+				);
+			}
+		}
+	}
+}
+
+void FillCharacterAdvancedMenu(UToolMenu* InMenu)
+{
+	if (!InMenu)
+	{
+		return;
+	}
+
+	const FAnimViewportShowCommands& Actions = FAnimViewportShowCommands::Get();
+
+	TSharedPtr<SAnimationEditorViewportTabBody> AnimEditorViewportTab = GetAnimationEditorViewportTab(InMenu);
+	if (!AnimEditorViewportTab)
+	{
+		return;
+	}
+
+	// Draw UVs
+	{
+		FToolMenuSection& UVSection =
+			InMenu->AddSection("UVVisualization", LOCTEXT("UVVisualization_Label", "UV Visualization"));
+
+		if (TSharedPtr<STextComboBox> UVChannelComboBox = AnimEditorViewportTab->UVChannelCombo)
+		{
+			UVSection.AddEntry(FToolMenuEntry::InitWidget("", UVChannelComboBox.ToSharedRef(), FText()));
+		}
+	}
+
+	// Skinning
+	{
+		FToolMenuSection& SkinningSection = InMenu->AddSection("Skinning", LOCTEXT("Skinning_Label", "Skinning"));
+		SkinningSection.AddMenuEntry(FAnimViewportMenuCommands::Get().SetCPUSkinning);
+	}
+
+	// Vertex visualization
+	{
+		FToolMenuSection& ShowVertexSection =
+			InMenu->AddSection("ShowVertex", LOCTEXT("ShowVertex_Label", "Vertex Normal Visualization"));
+
+		// Vertex debug flags
+		ShowVertexSection.AddMenuEntry(FAnimViewportMenuCommands::Get().SetShowNormals);
+		ShowVertexSection.AddMenuEntry(FAnimViewportMenuCommands::Get().SetShowTangents);
+		ShowVertexSection.AddMenuEntry(FAnimViewportMenuCommands::Get().SetShowBinormals);
+	}
+
+	// Local Axes
+	{
+		FToolMenuSection& LocalAxesSection = InMenu->AddSection(
+			"AnimViewportPreviewHierarchyLocalAxes", LOCTEXT("ShowMenu_Actions_HierarchyAxes", "Hierarchy Local Axes")
+		);
+
+		LocalAxesSection.AddMenuEntry(Actions.ShowLocalAxesAll);
+		LocalAxesSection.AddMenuEntry(Actions.ShowLocalAxesSelected);
+		LocalAxesSection.AddMenuEntry(Actions.ShowLocalAxesNone);
+	}
+}
+
+void FillCharacterTimecodeMenu(UToolMenu* InMenu)
+{
+	if (!InMenu)
+	{
+		return;
+	}
+
+	const FAnimViewportShowCommands& Actions = FAnimViewportShowCommands::Get();
+	FToolMenuSection& TimecodeSection = InMenu->AddSection("Timecode", LOCTEXT("Timecode_Label", "Timecode"));
+	TimecodeSection.AddMenuEntry(Actions.ShowTimecode);
+}
+} // namespace UE::AnimationEditor::Private
+
+TSharedRef<SWidget> UE::AnimationEditor::MakeFollowBoneWidget(
+	const TWeakPtr<SAnimationEditorViewportTabBody>& InViewport, const TWeakPtr<SComboButton>& InWeakComboButton
+)
+{
+	TSharedPtr<SAnimationEditorViewportTabBody> Viewport = InViewport.Pin();
+
+	if (!Viewport)
+	{
+		return SNullWidget::NullWidget;
+	}
+
+	TSharedPtr<SEditorViewport> ViewportWidget = Viewport->GetViewportWidget();
+
+	if (!ViewportWidget || !ViewportWidget->GetViewportClient())
+	{
+		return SNullWidget::NullWidget;
+	}
+
+	TSharedPtr<SBoneTreeMenu> BoneTreeMenu;
+
+	// clang-format off
+	TSharedRef<SWidget> MenuWidget =
+		SNew(SBox)
+		.MaxDesiredHeight(400.0f)
+		[
+			SAssignNew(BoneTreeMenu, SBoneTreeMenu)
+			.bShowVirtualBones(true)
+			.OnBoneSelectionChanged_Lambda(
+				[ViewportWeak = Viewport.ToWeakPtr()](FName InBoneName)
+				{
+					TSharedPtr<SAnimationEditorViewportTabBody> Viewport = ViewportWeak.Pin();
+					if (!Viewport)
+					{
+						return;
+					}
+
+					Viewport->SetCameraFollowMode(EAnimationViewportCameraFollowMode::Bone, InBoneName);
+					FSlateApplication::Get().DismissAllMenus();
+
+					if (const TSharedPtr<IPinnedCommandList>& PinnedCommands = Viewport->GetPinnedCommands())
+					{
+						PinnedCommands->SetStyle(&FAppStyle::Get(), TEXT("ViewportPinnedCommandList"));
+						PinnedCommands->AddCustomWidget(TEXT("FollowBoneWidget"));
+					}
+				})
+			.SelectedBone(Viewport->GetCameraFollowBoneName())
+			.OnGetReferenceSkeleton_Lambda(
+				[ViewportWeak = Viewport.ToWeakPtr()]() -> const FReferenceSkeleton&
+				{
+					static FReferenceSkeleton EmptySkeleton;
+
+					if (TSharedPtr<SAnimationEditorViewportTabBody> Viewport = ViewportWeak.Pin())
+					{
+						if (USkeletalMesh* PreviewMesh = Viewport->GetPreviewScene()->GetPreviewMesh())
+						{
+							return PreviewMesh->GetRefSkeleton();
+						}
+					}
+
+					return EmptySkeleton;
+				})
+		];
+	// clang-format on
+
+	if (TSharedPtr<SComboButton> ComboButton = InWeakComboButton.Pin())
+	{
+		ComboButton->SetMenuContentWidgetToFocus(BoneTreeMenu->GetFilterTextWidget());
+	}
+
+	return MenuWidget;
+}
+
+FToolMenuEntry UE::AnimationEditor::CreateShowSubmenu()
+{
+	return FToolMenuEntry::InitSubMenu(
+		"Show",
+		LOCTEXT("ShowSubmenuLabel", "Show"),
+		LOCTEXT("ShowSubmenuTooltip", "Show options"),
+		FNewToolMenuDelegate::CreateLambda(
+			[](UToolMenu* Submenu) -> void
+			{
+				UE::AnimationEditor::FillShowSubmenu(Submenu);
+				AddSceneElementsSection(Submenu);
+			}
+		)
+	);
+}
+
+void UE::AnimationEditor::FillShowSubmenu(UToolMenu* InMenu, bool bInShowViewportStatsToggle)
+{
+	if (!InMenu)
+	{
+		return;
+	}
+
+	InMenu->AddDynamicSection(
+		"AnimSection",
+		FNewToolMenuDelegate::CreateLambda(
+			[bInShowViewportStatsToggle](UToolMenu* InMenu)
+			{
+				if (bInShowViewportStatsToggle)
+				{
+					FToolMenuSection& UnnamedSection = InMenu->AddSection(NAME_None);
+					UnnamedSection.AddMenuEntry(
+						FEditorViewportCommands::Get().ToggleStats, LOCTEXT("ViewportStatsLabel", "Viewport Stats")
+					);
+
+					UnnamedSection.AddSeparator(NAME_None);
+				}
+
+				// Only include helpful show flags.
+				static const FShowFlagFilter ShowFlagFilter =
+					FShowFlagFilter(FShowFlagFilter::ExcludeAllFlagsByDefault)
+						// General
+						.IncludeFlag(FEngineShowFlags::SF_AntiAliasing)
+						.IncludeFlag(FEngineShowFlags::SF_Collision)
+						.IncludeFlag(FEngineShowFlags::SF_Grid)
+						.IncludeFlag(FEngineShowFlags::SF_Particles)
+						.IncludeFlag(FEngineShowFlags::SF_Translucency)
+						// Post Processing
+						.IncludeFlag(FEngineShowFlags::SF_Bloom)
+						.IncludeFlag(FEngineShowFlags::SF_DepthOfField)
+						.IncludeFlag(FEngineShowFlags::SF_EyeAdaptation)
+						.IncludeFlag(FEngineShowFlags::SF_HMDDistortion)
+						.IncludeFlag(FEngineShowFlags::SF_MotionBlur)
+						.IncludeFlag(FEngineShowFlags::SF_Tonemapper)
+						// Lighting Components
+						.IncludeGroup(SFG_LightingComponents)
+						// Lighting Features
+						.IncludeFlag(FEngineShowFlags::SF_AmbientCubemap)
+						.IncludeFlag(FEngineShowFlags::SF_DistanceFieldAO)
+						.IncludeFlag(FEngineShowFlags::SF_IndirectLightingCache)
+						.IncludeFlag(FEngineShowFlags::SF_LightFunctions)
+						.IncludeFlag(FEngineShowFlags::SF_LightShafts)
+						.IncludeFlag(FEngineShowFlags::SF_ReflectionEnvironment)
+						.IncludeFlag(FEngineShowFlags::SF_ScreenSpaceAO)
+						.IncludeFlag(FEngineShowFlags::SF_ContactShadows)
+						.IncludeFlag(FEngineShowFlags::SF_ScreenSpaceReflections)
+						.IncludeFlag(FEngineShowFlags::SF_SubsurfaceScattering)
+						.IncludeFlag(FEngineShowFlags::SF_TexturedLightProfiles)
+						// Developer
+						.IncludeFlag(FEngineShowFlags::SF_Refraction)
+						// Advanced
+						.IncludeFlag(FEngineShowFlags::SF_DeferredLighting)
+						.IncludeFlag(FEngineShowFlags::SF_Selection)
+						.IncludeFlag(FEngineShowFlags::SF_SeparateTranslucency)
+						.IncludeFlag(FEngineShowFlags::SF_TemporalAA)
+						.IncludeFlag(FEngineShowFlags::SF_VertexColors)
+						.IncludeFlag(FEngineShowFlags::SF_MeshEdges);
+
+				FShowFlagMenuCommands::Get().BuildShowFlagsMenu(InMenu, ShowFlagFilter);
+			}
+		)
+	);
+}
+
+FToolMenuEntry UE::AnimationEditor::CreateTurnTableMenu()
+{
+	return FToolMenuEntry::InitSubMenu(
+		"TurnTable",
+		LOCTEXT("TurnTableLabel", "Turn Table"),
+		LOCTEXT("TurnTableTooltip", "Set up auto-rotation of preview."),
+		FNewToolMenuDelegate::CreateLambda(
+			[](UToolMenu* Submenu) -> void
+			{
+				UE::AnimationEditor::FillTurnTableSubmenu(Submenu);
+			}
+		),
+		false,
+		FSlateIcon(FAppStyle::GetAppStyleSetName(), "AnimViewportMenu.TurnTableSpeed")
+	);
+}
+
+void UE::AnimationEditor::FillTurnTableSubmenu(UToolMenu* InMenu)
+{
+	if (!InMenu)
+	{
+		return;
+	}
+
+	const FAnimViewportPlaybackCommands& Actions = FAnimViewportPlaybackCommands::Get();
+	{
+		FToolMenuSection& TurnTableModeSection =
+			InMenu->FindOrAddSection("AnimViewportTurnTableMode", LOCTEXT("TurnTableMenu_ModeLabel", "Turn Table Mode"));
+
+		TurnTableModeSection.AddMenuEntry(Actions.PersonaTurnTablePlay);
+		TurnTableModeSection.AddMenuEntry(Actions.PersonaTurnTablePause);
+		TurnTableModeSection.AddMenuEntry(Actions.PersonaTurnTableStop);
+	}
+
+	{
+		FToolMenuSection& TurnTableSpeedSection =
+			InMenu->FindOrAddSection("AnimViewportTurnTableSpeed", LOCTEXT("TurnTableMenu_SpeedLabel", "Turn Table Speed"));
+		for (int i = 0; i < EAnimationPlaybackSpeeds::NumPlaybackSpeeds; ++i)
+		{
+			TurnTableSpeedSection.AddMenuEntry(Actions.TurnTableSpeeds[i]);
+		}
+
+		const TSharedPtr<SAnimationEditorViewportTabBody> AnimEditorViewportTab =
+			Private::GetAnimationEditorViewportTab(InMenu);
+
+		// clang-format off
+		TSharedPtr<SWidget> AnimSpeedWidget =
+			SNew(UE::AnimationEditor::SCustomAnimationSpeedSetting)
+			.AnimEditorViewport(AnimEditorViewportTab)
+			.CustomSpeed_Lambda(
+				[ViewportWeak = AnimEditorViewportTab.ToWeakPtr()]()
+				{
+					if (const TSharedPtr<SAnimationEditorViewportTabBody> Viewport = ViewportWeak.Pin())
+					{
+						return Viewport->GetCustomTurnTableSpeed();
+					}
+
+					return 1.0f;
+				})
+			.OnCustomSpeedChanged_Lambda(
+				[ViewportWeak = AnimEditorViewportTab.ToWeakPtr()](float CustomSpeed)
+				{
+					if (TSharedPtr<SAnimationEditorViewportTabBody> Viewport = ViewportWeak.Pin())
+					{
+						Viewport->SetCustomTurnTableSpeed(CustomSpeed);
+					}
+				});
+		// clang-format on
+
+		TurnTableSpeedSection.AddEntry(FToolMenuEntry::InitWidget(
+			"AnimSpeed", AnimSpeedWidget.ToSharedRef(), LOCTEXT("PlaybackMenu_Speed_Custom", "Custom Speed:")
+		));
+	}
+}
+
+void UE::AnimationEditor::AddSceneElementsSection(UToolMenu* InMenu, bool bInShowAudioEntries, bool bInShowTimecodeEntries)
+{
+	if (!InMenu)
+	{
+		return;
+	}
+
+	FToolMenuSection& Section =
+		InMenu->AddSection("AnimViewportSceneElements", LOCTEXT("CharacterMenu_SceneElements", "Scene Elements"));
+
+	Section.AddSubMenu(
+		TEXT("MeshSubMenu"),
+		LOCTEXT("CharacterMenu_MeshSubMenu", "Mesh"),
+		LOCTEXT("CharacterMenu_MeshSubMenuToolTip", "Mesh-related options"),
+		FNewToolMenuDelegate::CreateLambda(
+			[](UToolMenu* InSubMenu)
+			{
+				{
+					FToolMenuSection& Section =
+						InSubMenu->AddSection("AnimViewportMesh", LOCTEXT("CharacterMenu_Actions_Mesh", "Mesh"));
+					Section.AddMenuEntry(FAnimViewportShowCommands::Get().ShowBound);
+					Section.AddMenuEntry(FAnimViewportShowCommands::Get().UseInGameBound);
+					Section.AddMenuEntry(FAnimViewportShowCommands::Get().UseFixedBounds);
+					Section.AddMenuEntry(FAnimViewportShowCommands::Get().UsePreSkinnedBounds);
+					Section.AddMenuEntry(FAnimViewportShowCommands::Get().ShowPreviewMesh);
+					Section.AddMenuEntry(FAnimViewportShowCommands::Get().ShowMorphTargets);
+				}
+				{
+					FToolMenuSection& Section = InSubMenu->AddSection(
+						"AnimViewportMeshInfo", LOCTEXT("CharacterMenu_Actions_MeshInfo", "Mesh Info")
+					);
+					Section.AddMenuEntry(FAnimViewportShowCommands::Get().ShowDisplayInfoBasic);
+					Section.AddMenuEntry(FAnimViewportShowCommands::Get().ShowDisplayInfoDetailed);
+					Section.AddMenuEntry(FAnimViewportShowCommands::Get().ShowDisplayInfoSkelControls);
+					Section.AddMenuEntry(FAnimViewportShowCommands::Get().HideDisplayInfo);
+				}
+				{
+					FToolMenuSection& Section = InSubMenu->AddSection(
+						"AnimViewportPreviewOverlayDraw", LOCTEXT("CharacterMenu_Actions_Overlay", "Mesh Overlay Drawing")
+					);
+					Section.AddMenuEntry(FAnimViewportShowCommands::Get().ShowOverlayNone);
+					Section.AddMenuEntry(FAnimViewportShowCommands::Get().ShowBoneWeight);
+					Section.AddMenuEntry(FAnimViewportShowCommands::Get().ShowMorphTargetVerts);
+				}
+			}
+		)
+	);
+
+	Section.AddSubMenu(
+		"AnimationSubMenu",
+		LOCTEXT("CharacterMenu_AnimationSubMenu", "Animation"),
+		LOCTEXT("CharacterMenu_AnimationSubMenuToolTip", "Animation-related options"),
+		FNewToolMenuDelegate::CreateLambda(
+			[](UToolMenu* InSubMenu)
+			{
+				const TSharedPtr<SAnimationEditorViewportTabBody> AnimEditorViewportTab =
+					Private::GetAnimationEditorViewportTab(InSubMenu);
+				if (!AnimEditorViewportTab)
+				{
+					return;
+				}
+
+				{
+					FToolMenuSection& Section = InSubMenu->AddSection(
+						"AnimViewportRootMotion", LOCTEXT("CharacterMenu_RootMotionLabel", "Root Motion")
+					);
+					Section.AddMenuEntry(FAnimViewportShowCommands::Get().DoNotProcessRootMotion);
+					Section.AddMenuEntry(FAnimViewportShowCommands::Get().ProcessRootMotionLoop);
+					Section.AddMenuEntry(FAnimViewportShowCommands::Get().ProcessRootMotionLoopAndReset);
+				}
+
+				{
+					FToolMenuSection& Section = InSubMenu->AddSection(
+						"AnimViewportVisualization", LOCTEXT("CharacterMenu_VisualizationsLabel", "Visualizations")
+					);
+					Section.AddMenuEntry(FAnimViewportShowCommands::Get().ShowNotificationVisualizations);
+					Section.AddMenuEntry(FAnimViewportShowCommands::Get().DoNotVisualizeRootMotion);
+					Section.AddMenuEntry(FAnimViewportShowCommands::Get().VisualizeRootMotionTrajectory);
+					Section.AddMenuEntry(FAnimViewportShowCommands::Get().VisualizeRootMotionTrajectoryAndOrientation);
+					Section.AddMenuEntry(FAnimViewportShowCommands::Get().ShowAssetUserDataVisualizations);
+				}
+
+				{
+					FToolMenuSection& Section = InSubMenu->AddSection(
+						"AnimViewportAnimation", LOCTEXT("CharacterMenu_Actions_AnimationAsset", "Animation")
+					);
+					Section.AddMenuEntry(FAnimViewportShowCommands::Get().ShowRawAnimation);
+					Section.AddMenuEntry(FAnimViewportShowCommands::Get().ShowNonRetargetedAnimation);
+					Section.AddMenuEntry(FAnimViewportShowCommands::Get().ShowAdditiveBaseBones);
+					Section.AddMenuEntry(FAnimViewportShowCommands::Get().ShowSourceRawAnimation);
+
+					TSharedPtr<SEditorViewport> ViewportWidget = AnimEditorViewportTab->GetViewportWidget();
+					if (ViewportWidget && ViewportWidget->GetViewportClient())
+					{
+						if (UDebugSkelMeshComponent* PreviewComponent =
+								AnimEditorViewportTab->GetPreviewScene()->GetPreviewMeshComponent())
+						{
+							FUIAction DisableUnlessPreviewInstance;
+							DisableUnlessPreviewInstance.CanExecuteAction = FCanExecuteAction::CreateLambda(
+								[PreviewComponentWeak = TWeakObjectPtr<UDebugSkelMeshComponent>(PreviewComponent)]()
+								{
+									if (UDebugSkelMeshComponent* PreviewComponent = PreviewComponentWeak.Get())
+									{
+										return (
+											PreviewComponent->PreviewInstance
+											&& (PreviewComponent->PreviewInstance == PreviewComponent->GetAnimInstance())
+										);
+									}
+
+									return false;
+								}
+							);
+
+							Section.AddSubMenu(
+								"MirrorSubMenu",
+								LOCTEXT("CharacterMenu_AnimationSubMenu_MirrorSubMenu", "Mirror"),
+								LOCTEXT(
+									"CharacterMenu_AnimationSubMenu_MirrorSubMenuToolTip",
+									"Mirror the animation using the selected mirror data table"
+								),
+								FNewToolMenuChoice(FNewToolMenuDelegate::CreateStatic(
+									&UE::AnimationEditor::Private::FillCharacterMirrorMenu
+								)),
+								DisableUnlessPreviewInstance,
+								EUserInterfaceActionType::Button,
+								false,
+								FSlateIcon(),
+								false
+							);
+						}
+					}
+					Section.AddMenuEntry(FAnimViewportShowCommands::Get().ShowBakedAnimation);
+					Section.AddMenuEntry(FAnimViewportShowCommands::Get().DisablePostProcessBlueprint);
+				}
+
+				FToolMenuSection& Section =
+					InSubMenu->AddSection("SkinWeights", LOCTEXT("SkinWeights_Label", "Skin Weight Profiles"));
+				Section.AddEntry(FToolMenuEntry::InitWidget(
+					"SkinWeightCombo", AnimEditorViewportTab->SkinWeightCombo.ToSharedRef(), FText()
+				));
+			}
+		)
+	);
+
+	Section.AddSubMenu(
+		"BonesSubMenu",
+		LOCTEXT("CharacterMenu_BoneDrawSubMenu", "Bones"),
+		LOCTEXT("CharacterMenu_BoneDrawSubMenuToolTip", "Bone Drawing Options"),
+		FNewToolMenuDelegate::CreateLambda(
+			[](UToolMenu* InSubMenu)
+			{
+				TSharedPtr<SAnimationEditorViewportTabBody> AnimEditorViewportTab =
+					Private::GetAnimationEditorViewportTab(InSubMenu);
+				if (!AnimEditorViewportTab)
+				{
+					return;
+				}
+
+				{
+					FToolMenuSection& Section =
+						InSubMenu->AddSection("BonesAndSockets", LOCTEXT("CharacterMenu_BonesAndSocketsLabel", "Show"));
+					Section.AddMenuEntry(FAnimViewportShowCommands::Get().ShowSockets);
+					Section.AddMenuEntry(FAnimViewportShowCommands::Get().ShowAttributes);
+					Section.AddMenuEntry(FAnimViewportShowCommands::Get().ShowBoneNames);
+					Section.AddMenuEntry(FAnimViewportShowCommands::Get().ShowBoneColors);
+				}
+
+				{
+					FToolMenuSection& Section = InSubMenu->AddSection(
+						"AnimViewportPreviewHierarchyBoneDraw", LOCTEXT("CharacterMenu_Actions_BoneDrawing", "Bone Drawing")
+					);
+
+					TSharedPtr<SWidget> BoneSizeWidget =
+						SNew(UE::AnimationEditor::SBoneDrawSizeSetting).AnimEditorViewport(AnimEditorViewportTab);
+					Section.AddEntry(FToolMenuEntry::InitWidget(
+						"BoneDrawSize", BoneSizeWidget.ToSharedRef(), LOCTEXT("CharacterMenu_Actions_BoneDrawSize", "Bone Draw Size:")
+					));
+
+					Section.AddMenuEntry(FAnimViewportShowCommands::Get().ShowBoneDrawAll);
+					Section.AddMenuEntry(FAnimViewportShowCommands::Get().ShowBoneDrawSelected);
+					Section.AddMenuEntry(FAnimViewportShowCommands::Get().ShowBoneDrawSelectedAndParents);
+					Section.AddMenuEntry(FAnimViewportShowCommands::Get().ShowBoneDrawSelectedAndChildren);
+					Section.AddMenuEntry(FAnimViewportShowCommands::Get().ShowBoneDrawSelectedAndParentsAndChildren);
+					Section.AddMenuEntry(FAnimViewportShowCommands::Get().ShowBoneDrawNone);
+				}
+			}
+		)
+	);
+
+	Section.AddDynamicEntry(
+		"ClothingSubMenu",
+		FNewToolMenuSectionDelegate::CreateLambda(
+			[](FToolMenuSection& InSection)
+			{
+				TSharedPtr<SAnimationEditorViewportTabBody> AnimEditorViewportTab =
+					Private::GetAnimationEditorViewportTab(InSection);
+				if (!AnimEditorViewportTab)
+				{
+					return;
+				}
+
+				TSharedPtr<FAnimationEditorPreviewScene> PreviewScene = AnimEditorViewportTab->GetPreviewScene();
+				if (!PreviewScene)
+				{
+					return;
+				}
+
+				UDebugSkelMeshComponent* PreviewComp = PreviewScene->GetPreviewMeshComponent();
+				if (PreviewComp && GetDefault<UPersonaOptions>()->bExposeClothingSceneElementMenu)
+				{
+					constexpr bool bInOpenSubMenuOnClick = false;
+					constexpr bool bShouldCloseWindowAfterMenuSelection = false;
+					InSection.AddSubMenu(
+						"ClothingSubMenu",
+						LOCTEXT("CharacterMenu_ClothingSubMenu", "Clothing"),
+						LOCTEXT("CharacterMenu_ClothingSubMenuToolTip", "Options relating to clothing"),
+						FNewToolMenuChoice(
+							FNewToolMenuDelegate::CreateStatic(&UE::AnimationEditor::Private::FillCharacterClothingMenu)
+						),
+						bInOpenSubMenuOnClick,
+						TAttribute<FSlateIcon>(),
+						bShouldCloseWindowAfterMenuSelection
+					);
+				}
+			}
+		)
+	);
+
+	if (bInShowAudioEntries)
+	{
+		Section.AddSubMenu(
+			TEXT("AudioSubMenu"),
+			LOCTEXT("CharacterMenu_AudioSubMenu", "Audio"),
+			LOCTEXT("CharacterMenu_AudioSubMenuToolTip", "Audio options"),
+			FNewToolMenuDelegate::CreateLambda(
+				[](UToolMenu* InSubMenu)
+				{
+					FToolMenuSection& Section =
+						InSubMenu->AddSection("AnimViewportAudio", LOCTEXT("CharacterMenu_Audio", "Audio"));
+					Section.AddMenuEntry(FAnimViewportShowCommands::Get().MuteAudio);
+					Section.AddMenuEntry(FAnimViewportShowCommands::Get().UseAudioAttenuation);
+				}
+			)
+		);
+	}
+
+	if (bInShowTimecodeEntries)
+	{
+		Section.AddDynamicEntry(
+			"Timecode",
+			FNewToolMenuSectionDelegate::CreateLambda(
+				[](FToolMenuSection& InSection)
+				{
+					InSection.AddSubMenu(
+						TEXT("TimecodeSubMenu"),
+						LOCTEXT("CharacterMenu_TimecodeSubMenu", "Timecode"),
+						LOCTEXT("CharacterMenu_TimecodeSubMenuToolTip", "Timecode options"),
+						FNewToolMenuChoice(
+							FNewToolMenuDelegate::CreateStatic(&UE::AnimationEditor::Private::FillCharacterTimecodeMenu)
+						)
+					);
+				}
+			)
+		);
+	}
+
+	Section.Blocks.Sort(
+		[](const FToolMenuEntry& A, const FToolMenuEntry& B)
+		{
+			return A.Label.Get().ToLower().ToString() < B.Label.Get().ToLower().ToString();
+		}
+	);
+
+	Section.AddDynamicEntry(
+		"AdvancedSubMenu",
+		FNewToolMenuSectionDelegate::CreateLambda(
+			[](FToolMenuSection& InSection)
+			{
+				InSection.AddSubMenu(
+					TEXT("AdvancedSubMenu"),
+					LOCTEXT("CharacterMenu_AdvancedSubMenu", "Advanced"),
+					LOCTEXT("CharacterMenu_AdvancedSubMenuToolTip", "Advanced options"),
+					FNewToolMenuChoice(FNewToolMenuDelegate::CreateStatic(&UE::AnimationEditor::Private::FillCharacterAdvancedMenu
+					))
+				);
+			}
+		)
+	);
+}
+
+FToolMenuEntry UE::AnimationEditor::CreateLODSubmenu()
+{
+	return FToolMenuEntry::InitDynamicEntry(
+		"DynamicLODOptions",
+		FNewToolMenuSectionDelegate::CreateLambda(
+			[](FToolMenuSection& InDynamicSection) -> void
+			{
+				TSharedPtr<SAnimationEditorViewportTabBody> AnimEditorViewportTab =
+					Private::GetAnimationEditorViewportTab(InDynamicSection);
+				if (!AnimEditorViewportTab)
+				{
+					return;
+				}
+
+				// Label updates based on currently selected LOD
+				const TAttribute<FText> Label = TAttribute<FText>::CreateLambda(
+					[WeakTab = AnimEditorViewportTab.ToWeakPtr()]()
+					{
+						if (TSharedPtr<SAnimationEditorViewportTabBody> ViewportTab = WeakTab.Pin())
+						{
+							return GetLODMenuLabel(ViewportTab);
+						}
+
+						return LOCTEXT("LODSubmenuLabel", "LOD");
+					}
+				);
+
+				InDynamicSection.AddSubMenu(
+					"LOD",
+					Label,
+					FText(),
+					FNewToolMenuDelegate::CreateLambda(
+						[](UToolMenu* Submenu) -> void
+						{
+							UE::AnimationEditor::Private::PopulateLODSubmenu(Submenu);
+						}
+					)
+				);
+			}
+		)
+	);
+}
+
+FText UE::AnimationEditor::GetLODMenuLabel(const TWeakPtr<SAnimationEditorViewportTabBody>& InAnimEditorViewportTab)
+{
+	FText Label = LOCTEXT("LODMenu_AutoLabel", "LOD Auto");
+
+	if (TSharedPtr<SAnimationEditorViewportTabBody> ViewportTab = InAnimEditorViewportTab.Pin())
+	{
+		int32 LODSelectionType = ViewportTab->GetLODSelection();
+
+		if (ViewportTab->IsTrackingAttachedMeshLOD())
+		{
+			Label = FText::Format(LOCTEXT("LODMenu_DebugLabel", "LOD Debug ({0})"), FText::AsNumber(LODSelectionType - 1));
+		}
+		else if (LODSelectionType > 0)
+		{
+			FString TitleLabel = FString::Printf(TEXT("LOD %d"), LODSelectionType - 1);
+			Label = FText::FromString(TitleLabel);
+		}
+	}
+	return Label;
+}
+
+TSharedRef<SWidget> UE::AnimationEditor::MakeFloorOffsetWidget(const TWeakPtr<SAnimationEditorViewportTabBody>& InAnimEditorViewportTabWeak
+)
+{
+	constexpr float FOVMin = -100.0f;
+	constexpr float FOVMax = 100.0f;
+
+	return
+		// clang-format off
+		SNew(SBox)
+		.HAlign(HAlign_Right)
+		[
+			SNew(SBox)
+			.Padding(FMargin(4.0f, 0.0f, 0.0f, 0.0f))
+			.WidthOverride(100.0f)
+			[
+				SNew(SBorder)
+				.BorderImage(FAppStyle::Get().GetBrush("Menu.WidgetBorder"))
+				.Padding(FMargin(1.0f))
+				[
+					SNew(SSpinBox<float>)
+					.Style(&FAppStyle::Get(), "Menu.SpinBox")
+					.Font(FAppStyle::GetFontStyle(TEXT("MenuItem.Font")))
+					.MinValue(FOVMin)
+					.MaxValue(FOVMax)
+					.Value_Lambda(
+						[InAnimEditorViewportTabWeak]()
+						{
+							if (TSharedPtr<SAnimationEditorViewportTabBody> ViewportClient = InAnimEditorViewportTabWeak.Pin())
+							{
+								if (ViewportClient->GetViewportWidget() && ViewportClient->GetViewportWidget()->GetViewportClient())
+								{
+									FAnimationViewportClient& AnimViewportClient = (FAnimationViewportClient&)ViewportClient->GetLevelViewportClient();
+									return AnimViewportClient.GetFloorOffset();
+								}
+							}
+							return 0.0f;
+						})
+					.OnBeginSliderMovement_Lambda(
+						[InAnimEditorViewportTabWeak]()
+						{
+							if (TSharedPtr<SAnimationEditorViewportTabBody> ViewportClient = InAnimEditorViewportTabWeak.Pin())
+							{
+								ViewportClient->OnBeginSliderMovementFloorOffset();
+							}
+						})
+					.OnValueChanged_Lambda(
+						[InAnimEditorViewportTabWeak](float InNewValue)
+						{
+							if (TSharedPtr<SAnimationEditorViewportTabBody> ViewportClient = InAnimEditorViewportTabWeak.Pin())
+							{
+								ViewportClient->OnFloorOffsetChanged(InNewValue);
+							}
+						})
+					.OnValueCommitted_Lambda(
+						[InAnimEditorViewportTabWeak](float InNewValue, ETextCommit::Type InCommitType)
+						{
+							if (TSharedPtr<SAnimationEditorViewportTabBody> ViewportClient = InAnimEditorViewportTabWeak.Pin())
+							{
+								ViewportClient->OnFloorOffsetCommitted(InNewValue, InCommitType);
+							}
+						})
+					.ToolTipText(LOCTEXT("FloorOffsetToolTip", "Height offset for the floor mesh (stored per-mesh)"))
+				]
+			]
+		];
+	// clang-format on
+}
+
+FToolMenuEntry UE::AnimationEditor::CreateSettingsSubmenu()
+{
+	return FToolMenuEntry::InitSubMenu(
+		"Settings",
+		LOCTEXT("SettingsSubmenuLabel", "Settings"),
+		LOCTEXT("SettingsSubmenuTooltip", "Settings"),
+		FNewToolMenuDelegate::CreateLambda(
+			[](UToolMenu* Submenu) -> void
+			{
+				TSharedPtr<SAnimationEditorViewportTabBody> AnimEditorViewportTab =
+					Private::GetAnimationEditorViewportTab(Submenu);
+
+				if (!AnimEditorViewportTab)
+				{
+					return;
+				}
+
+				// Scene Setup
+				{
+					FToolMenuSection& SceneSetupSection =
+						Submenu->AddSection("SceneSetup", LOCTEXT("SceneSetupLabel", "Scene Setup"));
+					FToolMenuEntry FloorOffsetEntry = FToolMenuEntry::InitWidget(
+						"FloorOffset",
+						MakeFloorOffsetWidget(AnimEditorViewportTab),
+						LOCTEXT("FloorHeightOffset", "Floor Height Offset")
+					);
+					SceneSetupSection.AddEntry(FloorOffsetEntry);
+					SceneSetupSection.AddMenuEntry(FAnimViewportShowCommands::Get().AutoAlignFloorToMesh);
+					SceneSetupSection.AddEntry(UE::AnimationEditor::CreateTurnTableMenu());
+				}
+
+				// Preview Scene
+				{
+					FToolMenuSection& PreviewSceneSection =
+						Submenu->AddSection("PreviewScene", LOCTEXT("PreviewSceneLabel", "Preview Scene"));
+					FToolMenuEntry PreviewProfileSubmenu = UE::UnrealEd::CreateViewportToolbarAssetViewerProfileSubmenu();
+
+					PreviewSceneSection.AddEntry(PreviewProfileSubmenu);
+				}
+
+				// Audio
+				{
+					FToolMenuSection& AudioSection = Submenu->AddSection("Audio", LOCTEXT("AudioLabel", "Audio"));
+					AudioSection.AddMenuEntry(FAnimViewportShowCommands::Get().MuteAudio);
+					AudioSection.AddMenuEntry(FAnimViewportShowCommands::Get().UseAudioAttenuation);
+				}
+			}
+		)
+	);
+}
+
+void UE::AnimationEditor::ExtendCameraMenu(FName InCameraOptionsMenuName)
+{
+	UToolMenu* const Menu = UToolMenus::Get()->ExtendMenu(InCameraOptionsMenuName);
+	if (!Menu)
+	{
+		return;
+	}
+
+	Menu->AddDynamicSection(
+		"AnimEditorCameraExtensionDynamicSection",
+		FNewToolMenuDelegate::CreateLambda(
+			[](UToolMenu* InDynamicMenu)
+			{
+				UUnrealEdViewportToolbarContext* const EditorViewportContext =
+					InDynamicMenu->FindContext<UUnrealEdViewportToolbarContext>();
+				if (!EditorViewportContext)
+				{
+					return;
+				}
+
+				const TSharedPtr<SEditorViewport> EditorViewport = EditorViewportContext->Viewport.Pin();
+				if (!EditorViewport)
+				{
+					return;
+				}
+
+				FToolMenuSection& PositioningSection =
+					InDynamicMenu->FindOrAddSection("Positioning", LOCTEXT("PositioningLabel", "Positioning"));
+
+				// Camera Speed Submenu
+				{
+					PositioningSection.AddSubMenu(
+						"CameraSpeed",
+						LOCTEXT("CameraSpeedSubMenu", "Camera Speed"),
+						LOCTEXT("CameraSpeedSubMenu_ToolTip", "Camera Speed related actions"),
+						FNewToolMenuDelegate::CreateLambda(
+							[](UToolMenu* InMenu)
+							{
+								UE::UnrealEd::CreateCameraSpeedMenu(InMenu);
+							}
+						),
+						false,
+						FSlateIcon(FSlateIcon(FAppStyle::GetAppStyleSetName(), "EditorViewport.Speed"))
+					);
+				}
+
+				PositioningSection.AddSeparator("PositioningSeparator_1");
+
+				// Frame Selection
+				{
+					FToolMenuEntry FocusViewportToSelection =
+						FToolMenuEntry::InitMenuEntry(FEditorViewportCommands::Get().FocusViewportToSelection);
+					FocusViewportToSelection.UserInterfaceActionType = EUserInterfaceActionType::ToggleButton;
+					FocusViewportToSelection.Icon = FSlateIcon(FAppStyle::GetAppStyleSetName(), "Icons.FrameActor");
+					PositioningSection.AddEntry(FocusViewportToSelection);
+				}
+
+				PositioningSection.AddSubMenu(
+					"FollowMode",
+					LOCTEXT("CameraFollowModeLabel", "Follow Mode"),
+					LOCTEXT("CameraFollowModeTooltip", "Set various camera follow modes"),
+					FNewToolMenuDelegate::CreateLambda(
+						[](UToolMenu* InMenu)
+						{
+							UE::AnimationEditor::FillFollowModeSubmenu(InMenu);
+						}
+					),
+					false,
+					FSlateIcon(FAppStyle::GetAppStyleSetName(), "AnimViewportMenu.CameraFollow")
+				);
+
+				PositioningSection.AddMenuEntry(FAnimViewportMenuCommands::Get().TogglePauseAnimationOnCameraMove);
+
+				FToolMenuSection& DefaultCameraSection =
+					InDynamicMenu->FindOrAddSection("DefaultCamera", LOCTEXT("DefaultCameraLabel", "Default Camera"));
+
+				DefaultCameraSection.AddMenuEntry(FAnimViewportMenuCommands::Get().JumpToDefaultCamera);
+				DefaultCameraSection.AddMenuEntry(FAnimViewportMenuCommands::Get().SaveCameraAsDefault);
+				DefaultCameraSection.AddMenuEntry(FAnimViewportMenuCommands::Get().ClearDefaultCamera);
+			}
+		)
+	);
+}
+
+void UE::AnimationEditor::FillFollowModeSubmenu(UToolMenu* InMenu)
+{
+	if (!InMenu)
+	{
+		return;
+	}
+
+	FToolMenuSection& CameraFollowModeSection = InMenu->FindOrAddSection(
+		"AnimViewportCameraFollowMode", LOCTEXT("ViewMenu_CameraFollowModeLabel", "Camera Follow Mode")
+	);
+
+	CameraFollowModeSection.AddMenuEntry(FAnimViewportMenuCommands::Get().CameraFollowNone);
+	CameraFollowModeSection.AddMenuEntry(FAnimViewportMenuCommands::Get().CameraFollowRoot);
+	CameraFollowModeSection.AddMenuEntry(FAnimViewportMenuCommands::Get().CameraFollowBounds);
+
+	CameraFollowModeSection.AddSubMenu(
+		"CameraFollowBone",
+		LOCTEXT("CameraFollowBone_DisplayName", "Orbit Bone"),
+		LOCTEXT("CameraFollowBone_ToolTip", "Select a bone for the camera to follow and orbit around"),
+		FNewToolMenuDelegate::CreateLambda(
+			[](UToolMenu* InSubmenu)
+			{
+				InSubmenu->AddDynamicSection(
+					"CameraFollowModeBoneSubmenu",
+					FNewToolMenuDelegate::CreateLambda(
+						[](UToolMenu* InSectionMenu)
+						{
+							TSharedPtr<SAnimationEditorViewportTabBody> ViewportTab =
+								Private::GetAnimationEditorViewportTab(InSectionMenu);
+							if (!ViewportTab)
+							{
+								return;
+							}
+
+							FToolMenuSection& FollowModeBoneSection = InSectionMenu->FindOrAddSection(
+								"CameraFollowModeBoneSection",
+								LOCTEXT("CameraFollowModeBoneSection_Label", "Follow Bone Options")
+							);
+							FollowModeBoneSection.AddEntry(FToolMenuEntry::InitWidget(
+								"FollowBoneWidget", UE::AnimationEditor::MakeFollowBoneWidget(ViewportTab, nullptr), FText()
+							));
+
+							FollowModeBoneSection.AddMenuEntry(
+								"LockRotation",
+								LOCTEXT("LockRotation_DisplayName", "Lock Rotation"),
+								LOCTEXT("LockRotation_ToolTip", "Keep viewport camera rotation aligned to the orbited bone."),
+								FSlateIcon(),
+								FUIAction(
+									FExecuteAction::CreateLambda(
+										[ViewportWeak = ViewportTab.ToWeakPtr()]()
+										{
+											if (TSharedPtr<SAnimationEditorViewportTabBody> Viewport = ViewportWeak.Pin())
+											{
+												Viewport->ToggleRotateCameraToFollowBone();
+											}
+										}
+									),
+									FCanExecuteAction(),
+									FIsActionChecked::CreateLambda(
+										[ViewportWeak = ViewportTab.ToWeakPtr()]()
+										{
+											if (TSharedPtr<SAnimationEditorViewportTabBody> Viewport = ViewportWeak.Pin())
+											{
+												return Viewport->GetShouldRotateCameraToFollowBone();
+											}
+
+											return false;
+										}
+									)
+								),
+								EUserInterfaceActionType::ToggleButton
+							);
+						}
+					)
+				);
+			}
+		)
+	);
+}
+
+void UE::AnimationEditor::ExtendViewModesSubmenu(FName InViewModesSubmenuName)
+{
+	UToolMenu* const Submenu = UToolMenus::Get()->ExtendMenu(InViewModesSubmenuName);
+
+	Submenu->AddDynamicSection(
+		"LevelEditorViewModesExtensionDynamicSection",
+		FNewToolMenuDelegate::CreateLambda(
+			[](UToolMenu* InDynamicMenu)
+			{
+				TSharedPtr<SAnimationEditorViewportTabBody> ViewportTab =
+					Private::GetAnimationEditorViewportTab(InDynamicMenu);
+
+				if (!ViewportTab)
+				{
+					return;
+				}
+
+				FToolMenuSection& Section = InDynamicMenu->FindOrAddSection("ViewMode");
+				Section.AddSubMenu(
+					"VisualizeBufferViewMode",
+					LOCTEXT("VisualizeBufferViewModeDisplayName", "Buffer Visualization"),
+					LOCTEXT("BufferVisualizationMenu_ToolTip", "Select a mode for buffer visualization"),
+					FNewMenuDelegate::CreateStatic(&FBufferVisualizationMenuCommands::BuildVisualisationSubMenu),
+					FUIAction(
+						FExecuteAction(),
+						FCanExecuteAction(),
+						FIsActionChecked::CreateLambda(
+							[WeakTab = ViewportTab.ToWeakPtr()]()
+							{
+								if (const TSharedPtr<SAnimationEditorViewportTabBody> Viewport = WeakTab.Pin())
+								{
+									FEditorViewportClient& ViewportClient = Viewport->GetViewportClient();
+									return ViewportClient.IsViewModeEnabled(VMI_VisualizeBuffer);
+								}
+
+								return false;
+							}
+						)
+					),
+					EUserInterfaceActionType::RadioButton,
+					/* bInOpenSubMenuOnClick = */ false,
+					FSlateIcon(FAppStyle::GetAppStyleSetName(), "EditorViewport.VisualizeBufferMode")
+				);
+			}
+		)
+	);
+}
+
+TSharedRef<FExtender> UE::AnimationEditor::GetViewModesLegacyExtenders(const TWeakPtr<SAnimationEditorViewportTabBody>& InViewport)
+{
+	TSharedRef<FExtender> Extender(new FExtender());
+
+	TSharedPtr<SAnimationEditorViewportTabBody> ViewportTabPinned = InViewport.Pin();
+	if (!ViewportTabPinned)
+	{
+		return Extender;
+	}
+
+	TSharedPtr<SEditorViewport> EditorViewport = ViewportTabPinned->GetViewportWidget();
+	if (!EditorViewport)
+	{
+		return Extender;
+	}
+
+	Extender->AddMenuExtension(
+		TEXT("ViewMode"),
+		EExtensionHook::After,
+		EditorViewport->GetCommandList(),
+		FMenuExtensionDelegate::CreateLambda(
+			[InViewportWeak = InViewport](FMenuBuilder& InMenuBuilder)
+			{
+				InMenuBuilder.AddSubMenu(
+					LOCTEXT("VisualizeBufferViewModeDisplayName", "Buffer Visualization"),
+					LOCTEXT("BufferVisualizationMenu_ToolTip", "Select a mode for buffer visualization"),
+					FNewMenuDelegate::CreateStatic(&FBufferVisualizationMenuCommands::BuildVisualisationSubMenu),
+					FUIAction(
+						FExecuteAction(),
+						FCanExecuteAction(),
+						FIsActionChecked::CreateLambda(
+							[InViewportWeak]()
+							{
+								if (const TSharedPtr<SAnimationEditorViewportTabBody> ViewportPtr = InViewportWeak.Pin())
+								{
+									const FEditorViewportClient& ViewportClient = ViewportPtr->GetViewportClient();
+									return ViewportClient.IsViewModeEnabled(VMI_VisualizeBuffer);
+								}
+								return false;
+							}
+						)
+					),
+					"VisualizeBufferViewMode",
+					EUserInterfaceActionType::RadioButton,
+					/* bInOpenSubMenuOnClick = */ false,
+					FSlateIcon(FAppStyle::GetAppStyleSetName(), "EditorViewport.VisualizeBufferMode")
+				);
+			}
+		)
+	);
+
+	return Extender;
+}
+
+#undef LOCTEXT_NAMESPACE
diff --git a/Engine/Source/Editor/Persona/Private/ViewportToolbar/AnimationEditorMenus.h b/Engine/Source/Editor/Persona/Private/ViewportToolbar/AnimationEditorMenus.h
new file mode 100644
--- /dev/null
+++ b/Engine/Source/Editor/Persona/Private/ViewportToolbar/AnimationEditorMenus.h
@@ -0,0 +1,41 @@
+﻿// Copyright Epic Games, Inc. All Rights Reserved.
+
+#pragma once
+
+#include "Widgets/SCompoundWidget.h"
+
+class FExtender;
+class IPreviewProfileController;
+class SAnimationEditorViewportTabBody;
+class SComboButton;
+class SEditorViewport;
+class UToolMenu;
+struct FToolMenuEntry;
+
+namespace UE::AnimationEditor
+{
+void ExtendCameraMenu(FName InCameraOptionsMenuName);
+
+void FillFollowModeSubmenu(UToolMenu* InMenu);
+
+TSharedRef<SWidget> MakeFollowBoneWidget(
+	const TWeakPtr<SAnimationEditorViewportTabBody>& InViewport, const TWeakPtr<SComboButton>& InWeakComboButton
+);
+
+FToolMenuEntry CreateShowSubmenu();
+void FillShowSubmenu(UToolMenu* InMenu, bool bInShowViewportStatsToggle = true);
+
+FToolMenuEntry CreateTurnTableMenu();
+void FillTurnTableSubmenu(UToolMenu* InMenu);
+
+void AddSceneElementsSection(UToolMenu* InMenu, bool bInShowAudioEntries = false, bool bInShowTimecodeEntries = false);
+
+FToolMenuEntry CreateLODSubmenu();
+FText GetLODMenuLabel(const TWeakPtr<SAnimationEditorViewportTabBody>& InAnimEditorViewportTab);
+
+TSharedRef<SWidget> MakeFloorOffsetWidget(const TWeakPtr<SAnimationEditorViewportTabBody>& InAnimEditorViewportTabWeak);
+FToolMenuEntry CreateSettingsSubmenu();
+
+void ExtendViewModesSubmenu(FName InViewModesSubmenuName);
+TSharedRef<FExtender> GetViewModesLegacyExtenders(const TWeakPtr<SAnimationEditorViewportTabBody>& InViewport);
+} // namespace UE::AnimationEditor
diff --git a/Engine/Source/Editor/Persona/Private/ViewportToolbar/AnimationEditorWidgets.cpp b/Engine/Source/Editor/Persona/Private/ViewportToolbar/AnimationEditorWidgets.cpp
new file mode 100644
--- /dev/null
+++ b/Engine/Source/Editor/Persona/Private/ViewportToolbar/AnimationEditorWidgets.cpp
@@ -0,0 +1,175 @@
+﻿// Copyright Epic Games, Inc. All Rights Reserved.
+
+#include "AnimationEditorWidgets.h"
+
+#include "SAnimationEditorViewport.h"
+#include "Widgets/Input/SNumericEntryBox.h"
+#include "Widgets/Input/SSpinBox.h"
+
+#define LOCTEXT_NAMESPACE "AnimEditorViewportToolbar"
+
+void UE::AnimationEditor::SCustomAnimationSpeedSetting::Construct(const FArguments& InArgs)
+{
+	AnimViewportPtr = InArgs._AnimEditorViewport;
+	CustomSpeed = InArgs._CustomSpeed;
+	OnCustomSpeedChanged = InArgs._OnCustomSpeedChanged;
+
+	// clang-format off
+	ChildSlot
+	[
+		SNew(SBox)
+		.HAlign(HAlign_Right)
+		[
+			SNew(SBox)
+			.Padding(FMargin(4.0f, 0.0f, 0.0f, 0.0f))
+			.WidthOverride(100.0f)
+			[
+				SNew(SSpinBox<float>)
+				.Font(FAppStyle::GetFontStyle(TEXT("MenuItem.Font")))
+				.ToolTipText(LOCTEXT("AnimationCustomSpeed", "Set Custom Speed."))
+				.MinValue(0.f)
+				.MaxSliderValue(10.f)
+				.SupportDynamicSliderMaxValue(true)
+				.Value(CustomSpeed)
+				.OnValueChanged(OnCustomSpeedChanged)
+			]
+		]
+	];
+	// clang-format on
+}
+
+void UE::AnimationEditor::SBoneDrawSizeSetting::Construct(const FArguments& InArgs)
+{
+	AnimViewportPtr = InArgs._AnimEditorViewport;
+
+	TSharedPtr<SAnimationEditorViewportTabBody> AnimViewport = AnimViewportPtr.Pin();
+	if (!AnimViewport)
+	{
+		return;
+	}
+
+	// clang-format off
+	ChildSlot
+	[
+		SNew(SBox)
+		.HAlign(HAlign_Right)
+		[
+			SNew(SBox)
+			.Padding(FMargin(4.0f, 0.0f, 0.0f, 0.0f))
+			.WidthOverride(100.0f)
+			[
+				SNew(SSpinBox<float>)
+				.Font(FAppStyle::GetFontStyle(TEXT("MenuItem.Font")))
+				.ToolTipText(LOCTEXT("BoneDrawSize_ToolTip", "Change bone size in viewport."))
+				.MinValue(0.f)
+				.MaxSliderValue(10.f)
+				.SupportDynamicSliderMaxValue(true)
+				.Value(AnimViewport.ToSharedRef(), &SAnimationEditorViewportTabBody::GetBoneDrawSize)
+				.OnValueChanged(SSpinBox<float>::FOnValueChanged::CreateSP(AnimViewport.ToSharedRef(), &SAnimationEditorViewportTabBody::SetBoneDrawSize))
+			]
+		]
+	];
+	// clang-format on
+}
+
+void UE::AnimationEditor::SClothWindSettings::Construct(const FArguments& InArgs)
+{
+	AnimViewportPtr = InArgs._AnimEditorViewport;
+
+	TSharedPtr<SAnimationEditorViewportTabBody> AnimViewport = AnimViewportPtr.Pin();
+	if (!AnimViewport)
+	{
+		return;
+	}
+
+	// clang-format off
+	ChildSlot
+	[
+		SNew(SBox)
+		.HAlign(HAlign_Right)
+		[
+			SNew(SBox)
+			.Padding(FMargin(4.0f, 0.0f, 0.0f, 0.0f))
+			.WidthOverride(100.0f)
+			[
+				SNew(SNumericEntryBox<float>)
+				.Font(FAppStyle::GetFontStyle(TEXT("MenuItem.Font")))
+				.ToolTipText(LOCTEXT("WindStrength_ToolTip", "Change wind strength"))
+				.MinValue(0.f)
+				.AllowSpin(true)
+				.MinSliderValue(0.f)
+				.MaxSliderValue(10.f)
+				.Value(AnimViewport.ToSharedRef(), &SAnimationEditorViewportTabBody::GetWindStrengthSliderValue)
+				.OnValueChanged(SSpinBox<float>::FOnValueChanged::CreateSP(AnimViewport.ToSharedRef(), &SAnimationEditorViewportTabBody::SetWindStrength))
+			]
+		]
+	];
+	// clang-format on
+}
+
+bool UE::AnimationEditor::SClothWindSettings::IsWindEnabled() const
+{
+	if (TSharedPtr<SAnimationEditorViewportTabBody> AnimViewport = AnimViewportPtr.Pin())
+	{
+		return AnimViewport->IsApplyingClothWind();
+	}
+
+	return false;
+}
+
+void UE::AnimationEditor::SGravitySettings::Construct(const FArguments& InArgs)
+{
+	AnimViewportPtr = InArgs._AnimEditorViewport;
+
+	TSharedPtr<SAnimationEditorViewportTabBody> AnimViewport = AnimViewportPtr.Pin();
+	if (!AnimViewport)
+	{
+		return;
+	}
+
+	// clang-format off
+	ChildSlot
+	[
+		SNew(SBox)
+		.HAlign(HAlign_Right)
+		[
+			SNew(SBox)
+			.Padding(FMargin(4.0f, 0.0f, 0.0f, 0.0f))
+			.WidthOverride(100.0f)
+			[
+				SNew(SSpinBox<float>)
+				.Font(FAppStyle::GetFontStyle(TEXT("MenuItem.Font")))
+				.ToolTipText(LOCTEXT("GravityScale_ToolTip", "Change gravity scale"))
+				.MinValue(0.f)
+				.MaxValue(4.f)
+				.Value(AnimViewport.ToSharedRef(), &SAnimationEditorViewportTabBody::GetGravityScaleSliderValue)
+				.OnValueChanged(SSpinBox<float>::FOnValueChanged::CreateSP(AnimViewport.ToSharedRef(), &SAnimationEditorViewportTabBody::SetGravityScale))
+			]
+		]
+	];
+	// clang-format on
+}
+
+FReply UE::AnimationEditor::SGravitySettings::OnDecreaseGravityScale() const
+{
+	if (TSharedPtr<SAnimationEditorViewportTabBody> AnimViewport = AnimViewportPtr.Pin())
+	{
+		constexpr float DeltaValue = 0.025f;
+		AnimViewport->SetGravityScale(AnimViewport->GetGravityScaleSliderValue() - DeltaValue);
+	}
+
+	return FReply::Handled();
+}
+
+FReply UE::AnimationEditor::SGravitySettings::OnIncreaseGravityScale() const
+{
+	if (TSharedPtr<SAnimationEditorViewportTabBody> AnimViewport = AnimViewportPtr.Pin())
+	{
+		constexpr float DeltaValue = 0.025f;
+		AnimViewport->SetGravityScale(AnimViewport->GetGravityScaleSliderValue() + DeltaValue);
+	}
+
+	return FReply::Handled();
+}
+
+#undef LOCTEXT_NAMESPACE
diff --git a/Engine/Source/Editor/Persona/Private/ViewportToolbar/AnimationEditorWidgets.h b/Engine/Source/Editor/Persona/Private/ViewportToolbar/AnimationEditorWidgets.h
new file mode 100644
--- /dev/null
+++ b/Engine/Source/Editor/Persona/Private/ViewportToolbar/AnimationEditorWidgets.h
@@ -0,0 +1,97 @@
+﻿// Copyright Epic Games, Inc. All Rights Reserved.
+
+#pragma once
+
+#include "Widgets/SCompoundWidget.h"
+
+class SAnimationEditorViewportTabBody;
+
+namespace UE::AnimationEditor
+{
+// Class definition which represents widget to modify animation speed settings
+class SCustomAnimationSpeedSetting : public SCompoundWidget
+{
+public:
+	/** Notification for numeric value change */
+	DECLARE_DELEGATE_OneParam(FOnCustomSpeedChanged, float);
+
+	SLATE_BEGIN_ARGS(SCustomAnimationSpeedSetting)
+	{
+	}
+	SLATE_ARGUMENT(TWeakPtr<SAnimationEditorViewportTabBody>, AnimEditorViewport)
+	SLATE_ATTRIBUTE(float, CustomSpeed)
+	SLATE_EVENT(FOnCustomSpeedChanged, OnCustomSpeedChanged)
+SLATE_END_ARGS()
+
+void Construct(const FArguments& InArgs);
+
+protected:
+	TWeakPtr<SAnimationEditorViewportTabBody> AnimViewportPtr;
+	TAttribute<float> CustomSpeed = 1.0f;
+	FOnCustomSpeedChanged OnCustomSpeedChanged;
+};
+
+// Class definition which represents widget to modify Bone Draw Size in viewport
+class SBoneDrawSizeSetting : public SCompoundWidget
+{
+
+public:
+	SLATE_BEGIN_ARGS(SBoneDrawSizeSetting)
+	{
+	}
+	SLATE_ARGUMENT(TWeakPtr<SAnimationEditorViewportTabBody>, AnimEditorViewport)
+SLATE_END_ARGS()
+
+void Construct(const FArguments& InArgs);
+
+protected:
+	TWeakPtr<SAnimationEditorViewportTabBody> AnimViewportPtr;
+};
+
+
+// Class definition which represents widget to modify strength of wind for clothing
+class SClothWindSettings : public SCompoundWidget
+{
+public:
+	SLATE_BEGIN_ARGS(SClothWindSettings)
+		{
+		}
+
+		SLATE_ARGUMENT(TWeakPtr<SAnimationEditorViewportTabBody>, AnimEditorViewport)
+	SLATE_END_ARGS()
+
+	/** Constructs this widget from its declaration */
+	void Construct(const FArguments& InArgs);
+
+protected:
+	/** Callback function which determines whether this widget is enabled */
+	bool IsWindEnabled() const;
+
+protected:
+	/** The viewport hosting this widget */
+	TWeakPtr<SAnimationEditorViewportTabBody> AnimViewportPtr;
+};
+
+// Class definition which represents widget to modify gravity for preview
+class SGravitySettings : public SCompoundWidget
+{
+public:
+	SLATE_BEGIN_ARGS(SGravitySettings)
+		{
+		}
+
+		SLATE_ARGUMENT(TWeakPtr<SAnimationEditorViewportTabBody>, AnimEditorViewport)
+	SLATE_END_ARGS()
+
+	void Construct(const FArguments& InArgs);
+
+protected:
+	FReply OnDecreaseGravityScale() const;
+
+	FReply OnIncreaseGravityScale() const;
+
+protected:
+	TWeakPtr<SAnimationEditorViewportTabBody> AnimViewportPtr;
+};
+
+}
