Skeletal Mesh: Improve managing of editable geometry on LODs

diff --git a/Engine/Source/Editor/Persona/Private/PersonaMeshDetails.cpp b/Engine/Source/Editor/Persona/Private/PersonaMeshDetails.cpp
--- a/Engine/Source/Editor/Persona/Private/PersonaMeshDetails.cpp
+++ b/Engine/Source/Editor/Persona/Private/PersonaMeshDetails.cpp
@@ -161,23 +161,27 @@ IDetailCategoryBuilder&  GetLODIndexCategory(IDetailLayoutBuilder& DetailLayout,
 	return DetailLayout.EditCategory(*CategoryName, LODLevelString, ECategoryPriority::Important);
 }
 
-enum EButtonFlags
+enum class EButtonFlags
 {
 	/** No special property exporting flags */
-	BF_None = 0x00000000,
+	None = 0x00000000,
 
 	/** Show generate/apply button */
-	BF_Generate = 0x00000001,
+	Generate = 0x00000001,
 
 	/** Show reimport button */
-	BF_Reimport = 0x00000002,
+	Reimport = 0x00000002,
 
 	/** Show reimportnewfile button */
-	BF_ReimportNewFile = 0x00000004,
+	ReimportNewFile = 0x00000004,
 
 	/** Show remove button */
-	BF_Remove = 0x00000008
+	RemoveLOD = 0x00000008,
+
+	/** Show remove geometry button */
+	RemoveGeometry = 0x00000010
 };
+ENUM_CLASS_FLAGS(EButtonFlags)
 
 /**
  * Window for Nanite settings.
@@ -255,15 +259,16 @@ public:
 	SLATE_BEGIN_ARGS(SSkeletalLODActions)
 		: _LODIndex(INDEX_NONE)
 		, _PersonaToolkit(nullptr)
-		, _ButtonFlags(0)
+		, _ButtonFlags(EButtonFlags::None)
 	{}
 	SLATE_ARGUMENT(int32, LODIndex)
 	SLATE_ARGUMENT(TWeakPtr<IPersonaToolkit>, PersonaToolkit)
-	SLATE_ARGUMENT(uint32, ButtonFlags)
+	SLATE_ARGUMENT(EButtonFlags, ButtonFlags)
 	SLATE_ARGUMENT(FString, MeshDescriptionReferenceIDString)
 	SLATE_ARGUMENT(bool, BuildAvailable)
 	SLATE_EVENT(FOnClicked, OnApplyLODChangeClicked)
 	SLATE_EVENT(FOnClicked, OnRemoveLODClicked)
+	SLATE_EVENT(FOnClicked, OnRemoveGeometryClicked)
 	SLATE_EVENT(FOnClicked, OnReimportClicked)
 	SLATE_EVENT(FOnClicked, OnReimportNewFileClicked)
 
@@ -317,19 +322,19 @@ private:
 		if (!SkelMesh->IsValidLODIndex(LODIndex))
 		{
 			// Should be true for the button to exist except if we delete a LOD
-			return LOCTEXT("ReimportButton_NewFile_LODNotValid_ToolTip", "Cannot reimport, LOD was delete");
+			return LOCTEXT("ReimportButton_NewFile_LODNotValid_ToolTip", "Cannot reimport, LOD has been deleted.");
 		}
 
 		FSkeletalMeshLODInfo& LODInfo = *(SkelMesh->GetLODInfo(LODIndex));
 
 		FString Filename = FPaths::GetCleanFilename(LODInfo.SourceImportFilename);
 
-		return FText::Format(LOCTEXT("ReimportButton_NewFile_ToolTip", "Reimport LOD{0} using the current source file ({1})"), FText::AsNumber(LODIndex), FText::FromString(Filename));
+		return FText::Format(LOCTEXT("ReimportButton_NewFile_ToolTip", "Re-import LOD {0} using the current source file ({1})"), FText::AsNumber(LODIndex), FText::FromString(Filename));
 	}
 
 	FText GetReimportButtonNewFileToolTipText() const
 	{
-		return FText::Format(LOCTEXT("ReimportButton_ToolTip", "Choose a new file to reimport over this LOD (LOD{0})"), FText::AsNumber(LODIndex));
+		return FText::Format(LOCTEXT("ReimportButton_ToolTip", "Choose a new file to reimport over LOD {0}"), FText::AsNumber(LODIndex));
 	}
 
 	bool CanReimportFromSource() const
@@ -349,6 +354,18 @@ private:
 		return false;
 	}
 
+	bool CanRemoveGeometry() const
+	{
+		TSharedPtr<IPersonaToolkit> SharedToolkit = PersonaToolkit.Pin();
+		if (SharedToolkit.IsValid())
+		{
+			USkeletalMesh* SkeletalMesh = SharedToolkit->GetMesh();
+
+			return SkeletalMesh->IsValidLODIndex(LODIndex) && SkeletalMesh->HasMeshDescription(LODIndex);
+		}
+		return false;
+	}
+
 public:
 	static bool LODSettingsChanged(USkeletalMesh* SkeletalMesh, int32 LODIndex)
 	{
@@ -388,18 +405,19 @@ public:
 
 private:
 	// Incoming arg data
-	int32 LODIndex;
+	int32 LODIndex = INDEX_NONE;
 	TWeakPtr<IPersonaToolkit> PersonaToolkit;
-	uint32 ButtonFlags;
-	bool BuildAvailable;
+	EButtonFlags ButtonFlags = EButtonFlags::None;
+	bool BuildAvailable = false;
 
 	FOnClicked OnApplyLODChangeClicked;
 	FOnClicked OnRemoveLODClicked;
+	FOnClicked OnRemoveGeometryClicked;
 	FOnClicked OnReimportClicked;
 	FOnClicked OnReimportNewFileClicked;
 
 	// Cached exists flag so we don't constantly hit the file system
-	bool bDoesSourceFileExist_Cached;
+	bool bDoesSourceFileExist_Cached = false;
 
 	FString MeshDescriptionReferenceIDString;
 };
@@ -412,6 +430,7 @@ void SSkeletalLODActions::Construct(const FArguments& InArgs)
 	MeshDescriptionReferenceIDString = InArgs._MeshDescriptionReferenceIDString;
 	OnApplyLODChangeClicked = InArgs._OnApplyLODChangeClicked;
 	OnRemoveLODClicked = InArgs._OnRemoveLODClicked;
+	OnRemoveGeometryClicked = InArgs._OnRemoveGeometryClicked;
 	OnReimportClicked = InArgs._OnReimportClicked;
 	OnReimportNewFileClicked = InArgs._OnReimportNewFileClicked;
 	BuildAvailable = InArgs._BuildAvailable;
@@ -428,7 +447,7 @@ void SSkeletalLODActions::Construct(const FArguments& InArgs)
 			]
 		];
 
-	if (OnApplyLODChangeClicked.IsBound() && (ButtonFlags & EButtonFlags::BF_Generate))
+	if (OnApplyLODChangeClicked.IsBound() && EnumHasAllFlags(ButtonFlags, EButtonFlags::Generate))
 	{
 		FText ButtonNameText = BuildAvailable ? LOCTEXT("ApplyLODChange", "Apply Changes") : LOCTEXT("RegenerateLOD", "Regenerate LOD");
 		WrapBox->AddSlot()
@@ -446,7 +465,7 @@ void SSkeletalLODActions::Construct(const FArguments& InArgs)
 		];
 	}
 
-	if (OnRemoveLODClicked.IsBound() && (ButtonFlags & EButtonFlags::BF_Remove))
+	if (OnRemoveLODClicked.IsBound() && EnumHasAllFlags(ButtonFlags, EButtonFlags::RemoveLOD))
 	{
 		WrapBox->AddSlot()
 		[
@@ -462,7 +481,25 @@ void SSkeletalLODActions::Construct(const FArguments& InArgs)
 		];
 	}
 
-	if (OnReimportClicked.IsBound() && (ButtonFlags & EButtonFlags::BF_Reimport))
+	if (OnRemoveGeometryClicked.IsBound() && EnumHasAllFlags(ButtonFlags, EButtonFlags::RemoveGeometry))
+	{
+		WrapBox->AddSlot()
+		[
+			SNew(SButton)
+			.HAlign(HAlign_Center)
+			.VAlign(VAlign_Center)
+			.ToolTipText(LOCTEXT("RemoveGeometryTooltip", "Remove the geometry stored on this LOD. The LOD will instead be automatically generated, based on another LOD."))
+			.IsEnabled(this, &SSkeletalLODActions::CanRemoveGeometry)
+			.OnClicked(OnRemoveGeometryClicked)
+			[
+				SNew(STextBlock)
+				.Font(IDetailLayoutBuilder::GetDetailFont())
+				.Text(LOCTEXT("RemoveGeometry", "Remove Geometry"))
+			]
+		];
+	}
+	
+	if (OnReimportClicked.IsBound() && EnumHasAllFlags(ButtonFlags, EButtonFlags::Reimport))
 	{
 		WrapBox->AddSlot()
 		[
@@ -480,7 +517,7 @@ void SSkeletalLODActions::Construct(const FArguments& InArgs)
 		];
 	}
 
-	if (OnReimportNewFileClicked.IsBound() && (ButtonFlags & EButtonFlags::BF_ReimportNewFile))
+	if (OnReimportNewFileClicked.IsBound() && EnumHasAllFlags(ButtonFlags, EButtonFlags::ReimportNewFile))
 	{
 		WrapBox->AddSlot()
 		[
@@ -3101,8 +3138,21 @@ void FPersonaMeshDetails::AddLODLevelCategories(IDetailLayoutBuilder& DetailLayo
 					LODCategory.AddCustomBuilder(BuildSettingsWidgetsPerLOD[LODIndex].ToSharedRef());
 				}
 
-				uint32 ButtonFlag = (bShowGenerateButtons ? EButtonFlags::BF_Generate : 0) | (bShowReimportButtons ? EButtonFlags::BF_Reimport | EButtonFlags::BF_ReimportNewFile : 0) | (bShowRemoveButtons ? EButtonFlags::BF_Remove : 0);
-				if (ButtonFlag > 0)
+				EButtonFlags ButtonFlags = EButtonFlags::None;
+				if (bShowGenerateButtons)
+				{
+					ButtonFlags |= EButtonFlags::Generate;
+				}
+				if (bShowReimportButtons)
+				{
+					ButtonFlags |= EButtonFlags::Reimport | EButtonFlags::ReimportNewFile;
+				}
+				if (bShowRemoveButtons)
+				{
+					ButtonFlags |= EButtonFlags::RemoveLOD | EButtonFlags::RemoveGeometry;
+				}
+				
+				if (ButtonFlags != EButtonFlags::None)
 				{
 					FString MeshDescriptionReferenceIDString = LODModel.GetLODModelDeriveDataKey();
 					LODCategory.AddCustomRow(LOCTEXT("LODButtonsRow", "LOD Buttons"))
@@ -3113,12 +3163,12 @@ void FPersonaMeshDetails::AddLODLevelCategories(IDetailLayoutBuilder& DetailLayo
 							SNew(SSkeletalLODActions)
 							.LODIndex(LODIndex)
 							.PersonaToolkit(GetPersonaToolkit())
-							.ButtonFlags(ButtonFlag)
+							.ButtonFlags(ButtonFlags)
 							.MeshDescriptionReferenceIDString(MeshDescriptionReferenceIDString)
 							.BuildAvailable(bIsbuildAvailable)
-							//.OnApplyLODChangeClicked(this, &FPersonaMeshDetails::RegenerateLOD, LODIndex)
 							.OnApplyLODChangeClicked(this, &FPersonaMeshDetails::ApplyLODChanges, LODIndex)
 							.OnRemoveLODClicked(this, &FPersonaMeshDetails::RemoveOneLOD, LODIndex)
+							.OnRemoveGeometryClicked(this, &FPersonaMeshDetails::RemoveGeometryFromLOD, LODIndex)
 							.OnReimportClicked(this, &FPersonaMeshDetails::OnReimportLodClicked, EReimportButtonType::Reimport, LODIndex)
 							.OnReimportNewFileClicked(this, &FPersonaMeshDetails::OnReimportLodClicked, EReimportButtonType::ReimportWithNewFile, LODIndex)
 						];
@@ -4356,6 +4406,58 @@ FReply FPersonaMeshDetails::RemoveOneLOD(int32 LODIndex)
 	return FReply::Handled();
 }
 
+
+FReply FPersonaMeshDetails::RemoveGeometryFromLOD(int32 LODIndex)
+{
+	USkeletalMesh* SkelMesh = GetPersonaToolkit()->GetMesh();
+	check(SkelMesh);
+	check(SkelMesh->IsValidLODIndex(LODIndex));
+
+	if (LODIndex <= 0)
+	{
+		return FReply::Handled();
+	}
+	
+	FText ConfirmRemoveLODText = FText::Format( LOCTEXT("PersonaRemoveGeometry_Confirmation", "Are you sure you want to remove all geometry from LOD {0} on {1}?"), LODIndex, FText::FromString(SkelMesh->GetName()) );
+
+	if ( FMessageDialog::Open(EAppMsgType::YesNo, ConfirmRemoveLODText) != EAppReturnType::Yes )
+	{
+		return FReply::Handled();
+	}
+	
+	FText RemoveLODText = FText::Format(LOCTEXT("OnPersonaRemoveGeometry", "Remove geometry from LOD {0}"), LODIndex);
+	FScopedTransaction Transaction(TEXT(""), RemoveLODText, SkelMesh);
+
+	FScopedSuspendAlternateSkinWeightPreview ScopedSuspendAlternateSkinnWeightPreview(SkelMesh);
+	//PostEditChange scope
+	{
+		FScopedSkeletalMeshPostEditChange ScopedPostEditChange(SkelMesh);
+		SkelMesh->PreEditChange(nullptr);
+		SkelMesh->Modify();
+
+		SkelMesh->ModifyMeshDescription(LODIndex);
+		SkelMesh->ClearMeshDescriptionAndBulkData(LODIndex);
+
+		FSkeletalMeshLODInfo* LODInfo = SkelMesh->GetLODInfo(LODIndex);
+		if (LODInfo->ReductionSettings.BaseLOD == LODIndex)
+		{
+			LODInfo->ReductionSettings.BaseLOD--;
+		}
+		
+		// Make sure the reduction settings are kosher.
+		if (SkelMesh->GetLODSettings())
+		{
+			SkelMesh->GetLODSettings()->SetLODSettingsToMesh(SkelMesh);
+		}
+
+		RegenerateOneLOD(LODIndex);
+		RegenerateDependentLODs(LODIndex);
+	}
+
+	return FReply::Handled();
+}
+
+
 FText FPersonaMeshDetails::GetApplyButtonText() const
 {
 	if (IsApplyNeeded() || !IsGenerateAvailable())	
@@ -4499,7 +4601,14 @@ FText FPersonaMeshDetails::GetLODImportedText(int32 LODIndex) const
 	{
 		if (Mesh->GetLODInfo(LODIndex)->bHasBeenSimplified)
 		{
-			return  LOCTEXT("LODMeshReductionText_Label", "[generated]");
+			if (Mesh->HasMeshDescription(LODIndex))
+			{
+				return LOCTEXT("LODMeshReductionText_Label_Inline", "[inline reduced]");
+			}
+			else
+			{
+				return LOCTEXT("LODMeshReductionText_Label_Reduced", "[generated]");
+			}
 		}
 	}
 
diff --git a/Engine/Source/Editor/Persona/Private/PersonaMeshDetails.h b/Engine/Source/Editor/Persona/Private/PersonaMeshDetails.h
--- a/Engine/Source/Editor/Persona/Private/PersonaMeshDetails.h
+++ b/Engine/Source/Editor/Persona/Private/PersonaMeshDetails.h
@@ -646,6 +646,8 @@ private:
 	FReply RegenerateLOD(int32 LODIndex);
 	/** Removes the specified lod from the skeletal mesh */
 	FReply RemoveOneLOD(int32 LODIndex);
+	/** Removes any editable geometry from the given LOD, making it auto-generated */
+	FReply RemoveGeometryFromLOD(int32 LODIndex);
 	/** Restore the LOD imported data if the LOD is no longer reduced */
 	void RestoreNonReducedLOD(int32 LODIndex);
 	/** hide properties which don't need to be showed to end users */
diff --git a/Engine/Source/Editor/Persona/Private/SAnimViewportToolBar.cpp b/Engine/Source/Editor/Persona/Private/SAnimViewportToolBar.cpp
--- a/Engine/Source/Editor/Persona/Private/SAnimViewportToolBar.cpp
+++ b/Engine/Source/Editor/Persona/Private/SAnimViewportToolBar.cpp
@@ -1192,16 +1192,39 @@ TSharedRef<SWidget> SAnimViewportToolBar::GenerateLODMenu() const
 	InMenuBuilder.PushExtender(MenuExtender.ToSharedRef());
 
 	const USkeletalMesh* PreviewMesh =  Viewport.Pin()->GetPreviewScene()->GetPreviewMesh();
-	auto IsAutoGenerated = [PreviewMesh](const int32 InLODId) -> bool
+	auto IsLODSimplified = [PreviewMesh](const int32 InLODId) -> bool
 	{
-		if (!PreviewMesh || PreviewMesh->IsCompiling())
+		return PreviewMesh && !PreviewMesh->IsCompiling() && PreviewMesh->IsValidLODIndex(InLODId) &&
+			   PreviewMesh->GetLODInfo(InLODId)->bHasBeenSimplified;
+	};
+	
+	auto GetLODStatusExtraLabel = [IsLODSimplified, PreviewMesh](const int32 InLODId) -> FText
+	{
+		if (IsLODSimplified(InLODId))
 		{
-			return false;
+			if (PreviewMesh->HasMeshDescription(InLODId))
+			{
+				return LOCTEXT("LODStatus_Inline", " (Inline Reduced)");
+			}
+			
+			return LOCTEXT("LODStatus_Generated", " (Generated)");
+		}
+		return FText::GetEmpty();
+	};
+	auto GetLODStatusExtraTooltip = [IsLODSimplified, PreviewMesh](const int32 InLODId) -> FText
+	{
+		if (IsLODSimplified(InLODId))
+		{
+			if (PreviewMesh->HasMeshDescription(InLODId))
+			{
+				return LOCTEXT("LODStatusTooltip_Inline", "Generated from the editable geometry stored on this LOD but has been simplified in place.");
+			}
+
+			const int32 BaseLOD = PreviewMesh->GetLODInfo(InLODId)->ReductionSettings.BaseLOD;
+			return FText::Format(LOCTEXT("LODStatusTooltip_Generated", "Generated from a reduced version of LOD {0}.\nIt contains no editable geometry."), FText::AsNumber(BaseLOD));
 		}
-		return PreviewMesh->IsValidLODIndex(InLODId) && PreviewMesh->GetLODInfo(InLODId)->bHasBeenSimplified;
+		return LOCTEXT("LODStatusTooltip_Default", "Generated from the editable geometry stored on this LOD with no simplification applied.");
 	};
-	static const FText LODFormat = LOCTEXT("LODFmt", "LOD {0}");
-	static const FText GeneratedLODFormat = LOCTEXT("GeneratedLODFmt", "LOD {0} (generated)");
 	
 	{
 		// LOD Models
@@ -1215,13 +1238,14 @@ TSharedRef<SWidget> SAnimViewportToolBar::GenerateLODMenu() const
 			for (int32 LODId = 1; LODId < LODCount; ++LODId)
 			{
 				const FText LODNumber = FText::AsNumber(LODId);
-				const FText TitleLabel = IsAutoGenerated(LODId) ? FText::Format(GeneratedLODFormat, LODNumber) : FText::Format(LODFormat, LODNumber); 
+				const FText TitleLabel = FText::Format(LOCTEXT("LODFmt", "LOD {0}{1}"), LODNumber, GetLODStatusExtraLabel(LODId));
+				const FText TooltipText = FText::Format(LOCTEXT("LODTooltip", "Force select LOD {0}.\n\n{1}"), LODNumber, GetLODStatusExtraTooltip(LODId));
 
 				FUIAction Action(FExecuteAction::CreateSP(Viewport.Pin().ToSharedRef(), &SAnimationEditorViewportTabBody::OnSetLODModel, LODId + 1),
 					FCanExecuteAction(),
 					FIsActionChecked::CreateSP(Viewport.Pin().ToSharedRef(), &SAnimationEditorViewportTabBody::IsLODModelSelected, LODId + 1));
 
-				InMenuBuilder.AddMenuEntry(TitleLabel, FText::GetEmpty(), FSlateIcon(), Action, NAME_None, EUserInterfaceActionType::RadioButton);
+				InMenuBuilder.AddMenuEntry(TitleLabel, TooltipText, FSlateIcon(), Action, NAME_None, EUserInterfaceActionType::RadioButton);
 			}
 		}
 		InMenuBuilder.EndSection();
