[Viewport Toolbar] Skel/Anim Editor Toolbar - Share code between old and new toolbar menus

Old toolbar is now using shared code with new toolbar to generate most widgets.

- Removed PinnedCommands from SAnimViewportToolbar, since that is now retrieved from ViewportTabBody

- Removed PreviewProfileController, since it can be retrieved from ViewportTabBody as well, and just used there

- Removed duplicate widgets from SAnimViewportToolbar.cpp, since now Toolbar is using the same ones as the new toolbar

diff --git a/Engine/Source/Editor/Persona/Private/SAnimViewportToolBar.cpp b/Engine/Source/Editor/Persona/Private/SAnimViewportToolBar.cpp
--- a/Engine/Source/Editor/Persona/Private/SAnimViewportToolBar.cpp
+++ b/Engine/Source/Editor/Persona/Private/SAnimViewportToolBar.cpp
@@ -16,256 +16,62 @@
 #include "PropertyEditorModule.h"
 #include "IDetailsView.h"
 #include "Framework/Application/SlateApplication.h"
-#include "Widgets/Colors/SColorBlock.h"
 #include "Editor/EditorPerProjectUserSettings.h"
 #include "Preferences/PersonaOptions.h"
 #include "EditorViewportCommands.h"
 #include "AnimViewportMenuCommands.h"
 #include "AnimViewportShowCommands.h"
-#include "AnimViewportLODCommands.h"
 #include "AnimViewportPlaybackCommands.h"
-#include "SAnimPlusMinusSlider.h"
 #include "SEditorViewportToolBarMenu.h"
 #include "STransformViewportToolbar.h"
 #include "SEditorViewportViewMenu.h"
-#include "Widgets/Input/SSpinBox.h"
-#include "Widgets/Colors/SColorPicker.h"
 #include "Widgets/Input/SNumericEntryBox.h"
 #include "Widgets/Input/STextComboBox.h"
 #include "AssetViewerSettings.h"
-#include "PersonaPreviewSceneDescription.h"
-#include "Engine/PreviewMeshCollection.h"
 #include "PreviewSceneCustomizations.h"
-#include "ClothingSimulation.h"
 #include "SimulationEditorExtender.h"
 #include "ClothingSimulationFactory.h"
 #include "ClothingSystemEditorInterfaceModule.h"
 #include "Widgets/SWidget.h"
 #include "Types/ISlateMetaData.h"
 #include "Textures/SlateIcon.h"
-#include "ShowFlagMenuCommands.h"
 #include "BufferVisualizationMenuCommands.h"
 #include "IPinnedCommandList.h"
 #include "UICommandList_Pinnable.h"
-#include "BoneSelectionWidget.h"
 #include "Widgets/Text/SRichTextBlock.h"
 #include "ContentBrowserModule.h"
 #include "IContentBrowserSingleton.h"
 #include "SCommonEditorViewportToolbarBase.h"
 #include "Animation/MirrorDataTable.h"
 #include "ScopedTransaction.h"
-#include "SNameComboBox.h"
+#include "ViewportToolbar/AnimationEditorMenus.h"
+#include "ViewportToolbar/AnimationEditorWidgets.h"
 #include "ViewportToolbar/UnrealEdViewportToolbar.h"
 
 #define LOCTEXT_NAMESPACE "AnimViewportToolBar"
 
-//Class definition which represents widget to modify strength of wind for clothing
-class SClothWindSettings : public SCompoundWidget
-{
-public:
-
-	SLATE_BEGIN_ARGS(SClothWindSettings)
-	{}
-
-		SLATE_ARGUMENT(TWeakPtr<SAnimationEditorViewportTabBody>, AnimEditorViewport)
-	SLATE_END_ARGS()
-
-	/** Constructs this widget from its declaration */
-	void Construct(const FArguments& InArgs )
-	{
-		AnimViewportPtr = InArgs._AnimEditorViewport;
-
-		this->ChildSlot
-		[
-			SNew(SBox)
-			.HAlign(HAlign_Right)
-			[
-				SNew(SBox)
-				.Padding(FMargin(4.0f, 0.0f, 0.0f, 0.0f))
-				.WidthOverride(100.0f)
-				[
-					SNew(SNumericEntryBox<float>)
-					.Font(FAppStyle::GetFontStyle(TEXT("MenuItem.Font")))
-					.ToolTipText(LOCTEXT("WindStrength_ToolTip", "Change wind strength"))
-					.MinValue(0.f)
-					.AllowSpin(true)
-					.MinSliderValue(0.f)
-					.MaxSliderValue(10.f)
-					.Value(AnimViewportPtr.Pin().ToSharedRef(), &SAnimationEditorViewportTabBody::GetWindStrengthSliderValue)
-					.OnValueChanged(SSpinBox<float>::FOnValueChanged::CreateSP(AnimViewportPtr.Pin().ToSharedRef(), &SAnimationEditorViewportTabBody::SetWindStrength))
-				]
-			]
-		];
-	}
-
-protected:
-
-	/** Callback function which determines whether this widget is enabled */
-	bool IsWindEnabled() const
-	{
-		return AnimViewportPtr.Pin()->IsApplyingClothWind();
-	}
-
-protected:
-	/** The viewport hosting this widget */
-	TWeakPtr<SAnimationEditorViewportTabBody> AnimViewportPtr;
-};
-
-//Class definition which represents widget to modify gravity for preview
-class SGravitySettings : public SCompoundWidget
-{
-public:
-
-	SLATE_BEGIN_ARGS(SGravitySettings)
-	{}
-
-		SLATE_ARGUMENT(TWeakPtr<SAnimationEditorViewportTabBody>, AnimEditorViewport)
-	SLATE_END_ARGS()
-
-	void Construct(const FArguments& InArgs )
-	{
-		AnimViewportPtr = InArgs._AnimEditorViewport;
-
-		this->ChildSlot
-		[
-			SNew(SBox)
-			.HAlign(HAlign_Right)
-			[
-				SNew(SBox)
-				.Padding(FMargin(4.0f, 0.0f, 0.0f, 0.0f))
-				.WidthOverride(100.0f)
-				[
-					SNew(SSpinBox<float>)
-					.Font(FAppStyle::GetFontStyle(TEXT("MenuItem.Font")))
-					.ToolTipText(LOCTEXT("GravityScale_ToolTip", "Change gravity scale"))
-					.MinValue(0.f)
-					.MaxValue(4.f)
-					.Value(AnimViewportPtr.Pin().ToSharedRef(), &SAnimationEditorViewportTabBody::GetGravityScaleSliderValue)
-					.OnValueChanged(SSpinBox<float>::FOnValueChanged::CreateSP(AnimViewportPtr.Pin().ToSharedRef(), &SAnimationEditorViewportTabBody::SetGravityScale))
-				]
-			]
-		];
-	}
-
-protected:
-	FReply OnDecreaseGravityScale()
-	{
-		const float DeltaValue = 0.025f;
-		AnimViewportPtr.Pin()->SetGravityScale( AnimViewportPtr.Pin()->GetGravityScaleSliderValue() - DeltaValue );
-		return FReply::Handled();
-	}
-
-	FReply OnIncreaseGravityScale()
-	{
-		const float DeltaValue = 0.025f;
-		AnimViewportPtr.Pin()->SetGravityScale( AnimViewportPtr.Pin()->GetGravityScaleSliderValue() + DeltaValue );
-		return FReply::Handled();
-	}
-
-protected:
-	TWeakPtr<SAnimationEditorViewportTabBody> AnimViewportPtr;
-};
-
-//Class definition which represents widget to modify Bone Draw Size in viewport
-class SBoneDrawSizeSetting : public SCompoundWidget
-{
-	
-public:
-
-	SLATE_BEGIN_ARGS(SBoneDrawSizeSetting) {}
-	SLATE_ARGUMENT(TWeakPtr<SAnimationEditorViewportTabBody>, AnimEditorViewport)
-	SLATE_END_ARGS()
-
-	void Construct(const FArguments& InArgs )
-	{
-		AnimViewportPtr = InArgs._AnimEditorViewport;
-
-		this->ChildSlot
-		[
-			SNew(SBox)
-			.HAlign(HAlign_Right)
-			[
-				SNew(SBox)
-				.Padding(FMargin(4.0f, 0.0f, 0.0f, 0.0f))
-				.WidthOverride(100.0f)
-				[
-					SNew(SSpinBox<float>)
-					.Font(FAppStyle::GetFontStyle(TEXT("MenuItem.Font")))
-					.ToolTipText(LOCTEXT("BoneDrawSize_ToolTip", "Change bone size in viewport."))
-					.MinValue(0.f)
-					.MaxSliderValue(10.f)
-					.SupportDynamicSliderMaxValue(true)
-					.Value(AnimViewportPtr.Pin().ToSharedRef(), &SAnimationEditorViewportTabBody::GetBoneDrawSize)
-					.OnValueChanged(SSpinBox<float>::FOnValueChanged::CreateSP(AnimViewportPtr.Pin().ToSharedRef(), &SAnimationEditorViewportTabBody::SetBoneDrawSize))
-				]
-			]
-		];
-	}
-
-protected:
-	TWeakPtr<SAnimationEditorViewportTabBody> AnimViewportPtr;
-};
-
-//Class definition which represents widget to modify Bone Draw Size in viewport
-class SCusttomAnimationSpeedSetting : public SCompoundWidget
-{
-public:
-	/** Notification for numeric value change */
-	DECLARE_DELEGATE_OneParam(FOnCustomSpeedChanged, float);
-
-	SLATE_BEGIN_ARGS(SCusttomAnimationSpeedSetting) {}
-		SLATE_ARGUMENT(TWeakPtr<SAnimationEditorViewportTabBody>, AnimEditorViewport)
-		SLATE_ATTRIBUTE(float, CustomSpeed)
-		SLATE_EVENT(FOnCustomSpeedChanged, OnCustomSpeedChanged)
-	SLATE_END_ARGS()
-
-	void Construct(const FArguments& InArgs )
-	{
-		AnimViewportPtr = InArgs._AnimEditorViewport;
-		CustomSpeed = InArgs._CustomSpeed;
-		OnCustomSpeedChanged = InArgs._OnCustomSpeedChanged;
-
-		this->ChildSlot
-		[
-			SNew(SBox)
-			.HAlign(HAlign_Right)
-			[
-				SNew(SBox)
-				.Padding(FMargin(4.0f, 0.0f, 0.0f, 0.0f))
-				.WidthOverride(100.0f)
-				[
-					SNew(SSpinBox<float>)
-					.Font(FAppStyle::GetFontStyle(TEXT("MenuItem.Font")))
-					.ToolTipText(LOCTEXT("AnimationCustomSpeed", "Set Custom Speed."))
-					.MinValue(0.f)
-					.MaxSliderValue(10.f)
-					.SupportDynamicSliderMaxValue(true)
-					.Value(CustomSpeed)
-					.OnValueChanged(OnCustomSpeedChanged)
-				]
-			]
-		];
-	}
-
-protected:
-	TWeakPtr<SAnimationEditorViewportTabBody> AnimViewportPtr;
-	TAttribute<float> CustomSpeed = 1.0f;
-	FOnCustomSpeedChanged OnCustomSpeedChanged;
-};
 ///////////////////////////////////////////////////////////
 // SAnimViewportToolBar
 
-void SAnimViewportToolBar::Construct(const FArguments& InArgs, TSharedPtr<class SAnimationEditorViewportTabBody> InViewport, TSharedPtr<class SEditorViewport> InRealViewport)
+void SAnimViewportToolBar::Construct(
+	const FArguments& InArgs,
+	TSharedPtr<class SAnimationEditorViewportTabBody> InViewport,
+	TSharedPtr<class SEditorViewport> InRealViewport
+)
 {
-	PreviewProfileController = InArgs._PreviewProfileController;
 	bShowShowMenu = InArgs._ShowShowMenu;
-	bShowCharacterMenu= InArgs._ShowCharacterMenu;
+	bShowCharacterMenu = InArgs._ShowCharacterMenu;
 	bShowLODMenu = InArgs._ShowLODMenu;
 	bShowPlaySpeedMenu = InArgs._ShowPlaySpeedMenu;
 	bShowFloorOptions = InArgs._ShowFloorOptions;
 	bShowTurnTable = InArgs._ShowTurnTable;
 	bShowPhysicsMenu = InArgs._ShowPhysicsMenu;
 
+	if (!InRealViewport)
+	{
+		return;
+	}
+
 	CommandList = InRealViewport->GetCommandList();
 	Extenders = InArgs._Extenders;
 	Extenders.Add(GetViewMenuExtender(InRealViewport));
@@ -279,8 +85,11 @@ void SAnimViewportToolBar::Construct(const FArguments& InArgs, TSharedPtr<class
 	const FMargin ToolbarSlotPadding(4.0f, 1.0f);
 	const FMargin ToolbarButtonPadding(4.0f, 1.0f);
 
-	TSharedRef<SHorizontalBox> LeftToolbar = SNew(SHorizontalBox)
-		+SHorizontalBox::Slot()
+
+	// clang-format off
+	TSharedRef<SHorizontalBox> LeftToolbar =
+		SNew(SHorizontalBox)
+		+ SHorizontalBox::Slot()
 		.AutoWidth()
 		.Padding(ToolbarSlotPadding)
 		[
@@ -292,7 +101,7 @@ void SAnimViewportToolBar::Construct(const FArguments& InArgs, TSharedPtr<class
 			.AddMetaData<FTagMetaData>(FTagMetaData(TEXT("EditorViewportToolBar.MenuDropdown")))
 			.OnGetMenuContent(this, &SAnimViewportToolBar::GenerateViewMenu)
 		]
-		+SHorizontalBox::Slot()
+		+ SHorizontalBox::Slot()
 		.AutoWidth()
 		.Padding(ToolbarSlotPadding)
 		[
@@ -301,7 +110,7 @@ void SAnimViewportToolBar::Construct(const FArguments& InArgs, TSharedPtr<class
 			.ParentToolBar(SharedThis(this))
 			.Cursor(EMouseCursor::Default)
 			.Label(this, &SAnimViewportToolBar::GetCameraMenuLabel)
-			.LabelIcon( this, &SAnimViewportToolBar::GetCameraMenuLabelIcon )
+			.LabelIcon(this, &SAnimViewportToolBar::GetCameraMenuLabelIcon)
 			.AddMetaData<FTagMetaData>(FTagMetaData(TEXT("EditorViewportToolBar.CameraMenu")))
 			.OnGetMenuContent(this, &SAnimViewportToolBar::GenerateViewportTypeMenu)
 		]
@@ -330,17 +139,23 @@ void SAnimViewportToolBar::Construct(const FArguments& InArgs, TSharedPtr<class
 		.AutoWidth()
 		.Padding(ToolbarSlotPadding)
 		[
-			SNew(SPreviewSceneProfileSelector).PreviewProfileController(InArgs._PreviewProfileController)
-			.Visibility_Lambda([this]()
-			{
-				// only show this menu if the user has customized it by adding their own profiles
-				// this behavior was requested by UX to match the behavior of the static mesh editor
-				if (PreviewProfileController.IsValid())
+			SNew(SPreviewSceneProfileSelector)
+			.PreviewProfileController(InRealViewport->GetPreviewProfileController())
+			.Visibility_Lambda(
+				[ViewportWidgetWeak = InRealViewport.ToWeakPtr()]()
 				{
-					return PreviewProfileController->HasAnyUserProfiles() ? EVisibility::Visible : EVisibility::Collapsed;
-				}
-				return EVisibility::Hidden;
-			})
+					if (TSharedPtr<SEditorViewport> ViewportWidget = ViewportWidgetWeak.Pin())
+					{
+						// only show this menu if the user has customized it by adding their own profiles
+						// this behavior was requested by UX to match the behavior of the static mesh editor
+						if (TSharedPtr<IPreviewProfileController> PreviewProfileController = ViewportWidget->GetPreviewProfileController())
+						{
+							  return PreviewProfileController->HasAnyUserProfiles() ? EVisibility::Visible
+																					: EVisibility::Collapsed;
+						}
+					}
+					return EVisibility::Hidden;
+				})
 		]
 		+ SHorizontalBox::Slot()
 		.AutoWidth()
@@ -397,6 +212,7 @@ void SAnimViewportToolBar::Construct(const FArguments& InArgs, TSharedPtr<class
 			.OnCamSpeedChanged(this, &SAnimViewportToolBar::OnCamSpeedChanged)
 			.OnCamSpeedScalarChanged(this, &SAnimViewportToolBar::OnCamSpeedScalarChanged)
 		];
+	// clang-format on
 
 	TWeakPtr<SAnimViewportToolBar> AnimViewportToolbarWeak = SharedThis(this);
 	UE::UnrealEd::OnViewportClientCamSpeedChanged().BindLambda(
@@ -439,14 +255,12 @@ void SAnimViewportToolBar::Construct(const FArguments& InArgs, TSharedPtr<class
 
 	//@TODO: Need clipping horizontal box: LeftToolbar->AddWrapButton();
 
-	// Create our pinned commands before we bind commands
-	IPinnedCommandListModule& PinnedCommandListModule = FModuleManager::LoadModuleChecked<IPinnedCommandListModule>(TEXT("PinnedCommandList"));
-	PinnedCommands = PinnedCommandListModule.CreatePinnedCommandList((InArgs._ContextName != NAME_None) ? InArgs._ContextName : TEXT("PersonaViewport"));
-	PinnedCommands->SetStyle(&FAppStyle::Get(), TEXT("ViewportPinnedCommandList"));
+	TSharedPtr<IPinnedCommandList> PinnedCommands = InViewport->GetPinnedCommands();
 
+	// clang-format off
 	ChildSlot
 	[
-		SNew( SVerticalBox )
+		SNew(SVerticalBox)
 		+ SVerticalBox::Slot()
 		.AutoHeight()
 		[
@@ -457,12 +271,12 @@ void SAnimViewportToolBar::Construct(const FArguments& InArgs, TSharedPtr<class
 				LeftToolbar
 			]
 		]
-		+SVerticalBox::Slot()
+		+ SVerticalBox::Slot()
 		.AutoHeight()
 		[
-			PinnedCommands.ToSharedRef()
+			PinnedCommands.IsValid() ? PinnedCommands.ToSharedRef() : SNullWidget::NullWidget
 		]
-		+SVerticalBox::Slot()
+		+ SVerticalBox::Slot()
 		.Padding(FMargin(4.0f, 3.0f, 0.0f, 0.0f))
 		[
 			// Display text (e.g., item being previewed)
@@ -473,16 +287,28 @@ void SAnimViewportToolBar::Construct(const FArguments& InArgs, TSharedPtr<class
 			.TextStyle(&FAppStyle::Get().GetWidgetStyle<FTextBlockStyle>("AnimViewport.MessageText"))
 		]
 	];
-	
+	// clang-format on
+
 	SViewportToolBar::Construct(SViewportToolBar::FArguments());
 
-	// Register all the custom widgets we can use here
-	PinnedCommands->RegisterCustomWidget(IPinnedCommandList::FOnGenerateCustomWidget::CreateSP(this, &SAnimViewportToolBar::MakeFloorOffsetWidget), TEXT("FloorOffsetWidget"), LOCTEXT("FloorHeightOffset", "Floor Height Offset"));
-	PinnedCommands->RegisterCustomWidget(IPinnedCommandList::FOnGenerateCustomWidget::CreateSP(this, &SAnimViewportToolBar::MakeFOVWidget), TEXT("FOVWidget"), LOCTEXT("Viewport_FOVLabel", "Field Of View"));
+	if (PinnedCommands)
+	{
+		// Register all the custom widgets we can use here
+		PinnedCommands->RegisterCustomWidget(
+			IPinnedCommandList::FOnGenerateCustomWidget::CreateSP(this, &SAnimViewportToolBar::MakeFloorOffsetWidget),
+			TEXT("FloorOffsetWidget"),
+			LOCTEXT("FloorHeightOffset", "Floor Height Offset")
+		);
+		PinnedCommands->RegisterCustomWidget(
+			IPinnedCommandList::FOnGenerateCustomWidget::CreateSP(this, &SAnimViewportToolBar::MakeFOVWidget),
+			TEXT("FOVWidget"),
+			LOCTEXT("Viewport_FOVLabel", "Field Of View")
+		);
 
-	PinnedCommands->BindCommandList(InViewport->GetCommandList().ToSharedRef());
+		PinnedCommands->BindCommandList(InViewport->GetCommandList().ToSharedRef());
+	}
 
-	// We assign the viewport pointer here rather than initially, as SViewportToolbar::Construct 
+	// We assign the viewport pointer here rather than initially, as SViewportToolbar::Construct
 	// ends up calling through and attempting to perform operations on the not-yet-full-constructed viewport
 	Viewport = InViewport;
 }
@@ -494,25 +320,7 @@ EVisibility SAnimViewportToolBar::GetTransformToolbarVisibility() const
 
 TSharedRef<SWidget> SAnimViewportToolBar::MakeFloorOffsetWidget() const
 {
-	return
-		SNew(SBox)
-		.HAlign(HAlign_Right)
-		[
-			SNew(SBox)
-			.Padding(FMargin(4.0f, 0.0f, 0.0f, 0.0f))
-			.WidthOverride(100.0f)
-			[
-				SNew(SSpinBox<float>)
-				.Font(FAppStyle::GetFontStyle(TEXT("MenuItem.Font")))
-				.MinSliderValue(-100.0f)
-				.MaxSliderValue(100.0f)
-				.Value(this, &SAnimViewportToolBar::OnGetFloorOffset)
-				.OnBeginSliderMovement(const_cast<SAnimViewportToolBar*>(this), &SAnimViewportToolBar::OnBeginSliderMovementFloorOffset)
-				.OnValueChanged(const_cast<SAnimViewportToolBar*>(this), &SAnimViewportToolBar::OnFloorOffsetChanged)
-				.OnValueCommitted(const_cast<SAnimViewportToolBar*>(this), &SAnimViewportToolBar::OnFloorOffsetCommitted)
-				.ToolTipText(LOCTEXT("FloorOffsetToolTip", "Height offset for the floor mesh (stored per-mesh)"))
-			]
-		];
+	return UE::AnimationEditor::MakeFloorOffsetWidget(Viewport);
 }
 
 TSharedRef<SWidget> SAnimViewportToolBar::MakeFOVWidget() const
@@ -566,53 +374,19 @@ TSharedRef<SWidget> SAnimViewportToolBar::MakeFollowBoneComboWidget() const
 		];
 
 	TWeakPtr<SComboButton> WeakComboButton = ComboButton;
-	ComboButton->SetOnGetMenuContent(FOnGetContent::CreateSP(this, &SAnimViewportToolBar::MakeFollowBoneWidget, WeakComboButton));
+	ComboButton->SetOnGetMenuContent(FOnGetContent::CreateLambda(
+		[ViewportWeak = Viewport, WeakComboButton]()
+		{
+			return UE::AnimationEditor::MakeFollowBoneWidget(ViewportWeak, WeakComboButton);
+		}
+	));
 
 	return ComboButton;
 }
 
 TSharedRef<SWidget> SAnimViewportToolBar::MakeFollowBoneWidget() const
 {
-	return MakeFollowBoneWidget(nullptr);
-}
-
-TSharedRef<SWidget> SAnimViewportToolBar::MakeFollowBoneWidget(TWeakPtr<SComboButton> InWeakComboButton) const
-{
-	TSharedPtr<SBoneTreeMenu> BoneTreeMenu;
-
-	TSharedRef<SWidget> MenuWidget =
-		SNew(SBox)
-		.MaxDesiredHeight(400.0f)
-		[
-			SAssignNew(BoneTreeMenu, SBoneTreeMenu)
-			.bShowVirtualBones(true)
-			.OnBoneSelectionChanged_Lambda([this](FName InBoneName)
-			{
-				Viewport.Pin()->SetCameraFollowMode(EAnimationViewportCameraFollowMode::Bone, InBoneName);
-				FSlateApplication::Get().DismissAllMenus();
-
-				PinnedCommands->AddCustomWidget(TEXT("FollowBoneWidget"));
-			})
-			.SelectedBone(Viewport.Pin()->GetCameraFollowBoneName())
-			.OnGetReferenceSkeleton_Lambda([this]() -> const FReferenceSkeleton&
-			{
-				USkeletalMesh* PreviewMesh = Viewport.Pin()->GetPreviewScene()->GetPreviewMesh();
-				if (PreviewMesh)
-				{
-					return PreviewMesh->GetRefSkeleton();
-				}
-
-				static FReferenceSkeleton EmptySkeleton;
-				return EmptySkeleton;
-			})
-		];
-
-	if(InWeakComboButton.IsValid())
-	{
-		InWeakComboButton.Pin()->SetMenuContentWidgetToFocus(BoneTreeMenu->GetFilterTextWidget());
-	}
-
-	return MenuWidget;
+	return UE::AnimationEditor::MakeFollowBoneWidget(Viewport, nullptr);
 }
 
 TSharedRef<SWidget> SAnimViewportToolBar::GenerateViewMenu() const
@@ -641,7 +415,7 @@ TSharedRef<SWidget> SAnimViewportToolBar::GenerateViewMenu() const
 			InMenuBuilder.AddMenuEntry(FAnimViewportShowCommands::Get().AutoAlignFloorToMesh);
 			InMenuBuilder.PopCommandList();
 		}
-			
+
 		if (bShowTurnTable)
 		{
 			InMenuBuilder.AddSubMenu(
@@ -664,58 +438,11 @@ TSharedRef<SWidget> SAnimViewportToolBar::GenerateViewMenu() const
 		InMenuBuilder.AddSubMenu(
 			LOCTEXT("CameraFollowModeLabel", "Follow Mode"),
 			LOCTEXT("CameraFollowModeTooltip", "Set various camera follow modes"),
-			FNewMenuDelegate::CreateLambda([this](FMenuBuilder& InSubMenuBuilder)
+			FNewMenuDelegate::CreateLambda([ViewportWeak = Viewport](FMenuBuilder& InSubMenuBuilder)
 			{
-				InSubMenuBuilder.BeginSection("AnimViewportCameraFollowMode", LOCTEXT("ViewMenu_CameraFollowModeLabel", "Camera Follow Mode"));
+				InSubMenuBuilder.BeginSection(NAME_None, FText());
 				{
-					InSubMenuBuilder.PushCommandList(Viewport.Pin()->GetCommandList().ToSharedRef());
-
-					InSubMenuBuilder.AddMenuEntry(FAnimViewportMenuCommands::Get().CameraFollowNone);
-					InSubMenuBuilder.AddMenuEntry(FAnimViewportMenuCommands::Get().CameraFollowRoot);
-					InSubMenuBuilder.AddMenuEntry(FAnimViewportMenuCommands::Get().CameraFollowBounds);
-					InSubMenuBuilder.AddSubMenu(
-						LOCTEXT("CameraFollowBone_DisplayName", "Orbit Bone"),
-						LOCTEXT("CameraFollowBone_ToolTip", "Select a bone for the camera to follow and orbit around"),
-						FNewMenuDelegate::CreateLambda([this](FMenuBuilder& SubMenuBuilder)
-						{
-							SubMenuBuilder.BeginSection("CameraFollowModeBoneSubmenu", LOCTEXT("CameraFollowModeBoneSubmenu_Label", "Follow Bone Options"));
-							SubMenuBuilder.AddWidget(MakeFollowBoneWidget(), FText());
-							SubMenuBuilder.AddMenuEntry(
-								LOCTEXT("LockRotation_DisplayName", "Lock Rotation"),
-								LOCTEXT("LockRotation_ToolTip", "Keep viewport camera rotation aligned to the orbited bone."),
-								FSlateIcon(),
-								FUIAction(
-									FExecuteAction::CreateLambda([this]() {
-										Viewport.Pin()->ToggleRotateCameraToFollowBone();
-									}),
-									FCanExecuteAction(),
-									FIsActionChecked::CreateLambda([this]()
-									{
-										return Viewport.Pin()->GetShouldRotateCameraToFollowBone();
-									})
-								),
-								NAME_None,
-								EUserInterfaceActionType::ToggleButton
-							);
-							SubMenuBuilder.EndSection();
-						}),
-						FUIAction(
-							FExecuteAction(),
-							FCanExecuteAction::CreateLambda([this]()
-							{
-								return Viewport.Pin()->CanChangeCameraMode();
-							}),
-							FIsActionChecked::CreateLambda([this]()
-							{
-								return Viewport.Pin()->IsCameraFollowEnabled(EAnimationViewportCameraFollowMode::Bone);
-							})
-						),
-						"CameraFollowBone",
-						EUserInterfaceActionType::RadioButton,
-						/* bInOpenSubMenuOnClick = */ false,
-						FSlateIcon()
-					);
-					InSubMenuBuilder.PopCommandList();
+					InSubMenuBuilder.AddWidget(UE::AnimationEditor::CreateFollowModeMenuWidget(ViewportWeak), FText());
 				}
 				InSubMenuBuilder.EndSection();
 			}),
@@ -878,7 +605,7 @@ TSharedRef<SWidget> SAnimViewportToolBar::GenerateCharacterMenu() const
 						FToolMenuSection& Section = InSubMenu->AddSection("AnimViewportPreviewHierarchyBoneDraw", LOCTEXT("CharacterMenu_Actions_BoneDrawing", "Bone Drawing"));
 						if (ContextThis)
 						{
-							TSharedPtr<SWidget> BoneSizeWidget = SNew(SBoneDrawSizeSetting).AnimEditorViewport(ContextThis->Viewport);
+							TSharedPtr<SWidget> BoneSizeWidget = SNew(UE::AnimationEditor::SBoneDrawSizeSetting).AnimEditorViewport(ContextThis->Viewport);
 							Section.AddEntry(FToolMenuEntry::InitWidget(TEXT("BoneDrawSize"), BoneSizeWidget.ToSharedRef(), LOCTEXT("CharacterMenu_Actions_BoneDrawSize", "Bone Draw Size:")));
 						}
 						Section.AddMenuEntry(FAnimViewportShowCommands::Get().ShowBoneDrawAll);
@@ -1056,10 +783,10 @@ void SAnimViewportToolBar::FillCharacterClothingMenu(FMenuBuilder& MenuBuilder)
 		MenuBuilder.AddMenuEntry(Actions.EnableClothSimulation);
 		MenuBuilder.AddMenuEntry(Actions.ResetClothSimulation);
 
-		TSharedPtr<SWidget> WindWidget = SNew(SClothWindSettings).AnimEditorViewport(Viewport);
+		TSharedPtr<SWidget> WindWidget = SNew(UE::AnimationEditor::SClothWindSettings).AnimEditorViewport(Viewport);
 		MenuBuilder.AddWidget(WindWidget.ToSharedRef(), LOCTEXT("ClothPreview_WindStrength", "Wind Strength:"));
 
-		TSharedPtr<SWidget> GravityWidget = SNew(SGravitySettings).AnimEditorViewport(Viewport);
+		TSharedPtr<SWidget> GravityWidget = SNew(UE::AnimationEditor::SGravitySettings).AnimEditorViewport(Viewport);
 		MenuBuilder.AddWidget(GravityWidget.ToSharedRef(), LOCTEXT("ClothPreview_GravityScale", "Gravity Scale:"));
 
 		MenuBuilder.AddMenuEntry(Actions.EnableCollisionWithAttachedClothChildren);
@@ -1100,166 +827,29 @@ void SAnimViewportToolBar::FillCharacterClothingMenu(FMenuBuilder& MenuBuilder)
 
 TSharedRef<SWidget> SAnimViewportToolBar::GenerateShowMenu() const
 {
-	static const FName MenuName("Persona.AnimViewportToolBar");
-
-	if (!UToolMenus::Get()->IsMenuRegistered(MenuName))
+	if (TSharedPtr<SAnimationEditorViewportTabBody> ViewportPinned = Viewport.Pin())
 	{
-		UToolMenu* Menu = UToolMenus::Get()->RegisterMenu(MenuName);
-		Menu->AddDynamicSection("AnimSection", FNewToolMenuDelegate::CreateLambda([](UToolMenu* InMenu)
-		{
-			// Only include helpful show flags.
-			static const FShowFlagFilter ShowFlagFilter = FShowFlagFilter(FShowFlagFilter::ExcludeAllFlagsByDefault)
-				// General
-				.IncludeFlag(FEngineShowFlags::SF_AntiAliasing)
-				.IncludeFlag(FEngineShowFlags::SF_Collision)
-				.IncludeFlag(FEngineShowFlags::SF_Grid)
-				.IncludeFlag(FEngineShowFlags::SF_Particles)
-				.IncludeFlag(FEngineShowFlags::SF_Translucency)
-				// Post Processing
-				.IncludeFlag(FEngineShowFlags::SF_Bloom)
-				.IncludeFlag(FEngineShowFlags::SF_DepthOfField)
-				.IncludeFlag(FEngineShowFlags::SF_EyeAdaptation)
-				.IncludeFlag(FEngineShowFlags::SF_HMDDistortion)
-				.IncludeFlag(FEngineShowFlags::SF_MotionBlur)
-				.IncludeFlag(FEngineShowFlags::SF_Tonemapper)
-				// Lighting Components
-				.IncludeGroup(SFG_LightingComponents)
-				// Lighting Features
-				.IncludeFlag(FEngineShowFlags::SF_AmbientCubemap)
-				.IncludeFlag(FEngineShowFlags::SF_DistanceFieldAO)
-				.IncludeFlag(FEngineShowFlags::SF_IndirectLightingCache)
-				.IncludeFlag(FEngineShowFlags::SF_LightFunctions)
-				.IncludeFlag(FEngineShowFlags::SF_LightShafts)
-				.IncludeFlag(FEngineShowFlags::SF_ReflectionEnvironment)
-				.IncludeFlag(FEngineShowFlags::SF_ScreenSpaceAO)
-				.IncludeFlag(FEngineShowFlags::SF_ContactShadows)
-				.IncludeFlag(FEngineShowFlags::SF_ScreenSpaceReflections)
-				.IncludeFlag(FEngineShowFlags::SF_SubsurfaceScattering)
-				.IncludeFlag(FEngineShowFlags::SF_TexturedLightProfiles)
-				// Developer
-				.IncludeFlag(FEngineShowFlags::SF_Refraction)
-				// Advanced
-				.IncludeFlag(FEngineShowFlags::SF_DeferredLighting)
-				.IncludeFlag(FEngineShowFlags::SF_Selection)
-				.IncludeFlag(FEngineShowFlags::SF_SeparateTranslucency)
-				.IncludeFlag(FEngineShowFlags::SF_TemporalAA)
-				.IncludeFlag(FEngineShowFlags::SF_VertexColors)
-				.IncludeFlag(FEngineShowFlags::SF_MeshEdges)
-				;
-
-			FShowFlagMenuCommands::Get().BuildShowFlagsMenu(InMenu, ShowFlagFilter);
-		}));
+		constexpr bool bShowViewportStatsToggle = false;
+		return UE::AnimationEditor::CreateShowMenuWidget(
+			ViewportPinned->GetViewportWidget().ToSharedRef(), Extenders, bShowViewportStatsToggle
+		);
 	}
 
-	TSharedPtr<FExtender> MenuExtender = FExtender::Combine(Extenders);
-	FToolMenuContext MenuContext(CommandList, MenuExtender);
-	return UToolMenus::Get()->GenerateWidget(MenuName, MenuContext);
+	return SNullWidget::NullWidget;
 }
 
 FText SAnimViewportToolBar::GetLODMenuLabel() const
 {
-	FText Label = LOCTEXT("LODMenu_AutoLabel", "LOD Auto");
-	if (Viewport.IsValid())
-	{
-		int32 LODSelectionType = Viewport.Pin()->GetLODSelection();
-
-		if (Viewport.Pin()->IsTrackingAttachedMeshLOD())
-		{
-			Label = FText::Format(LOCTEXT("LODMenu_DebugLabel", "LOD Debug ({0})"), FText::AsNumber(LODSelectionType - 1));
-		}
-		else
-		{
-			if (LODSelectionType > 0)
-			{
-				FString TitleLabel = FString::Printf(TEXT("LOD %d"), LODSelectionType - 1);
-				Label = FText::FromString(TitleLabel);
-			}
-		}
-	}
-	return Label;
+	return UE::AnimationEditor::GetLODMenuLabel(Viewport);
 }
 
 TSharedRef<SWidget> SAnimViewportToolBar::GenerateLODMenu() const
 {
-	const FAnimViewportLODCommands& Actions = FAnimViewportLODCommands::Get();
-
-	TSharedPtr<FExtender> MenuExtender = FExtender::Combine(Extenders);
-
-	const bool bInShouldCloseWindowAfterMenuSelection = true;
-	FMenuBuilder InMenuBuilder(bInShouldCloseWindowAfterMenuSelection, Viewport.Pin()->GetCommandList(), MenuExtender);
-
-	InMenuBuilder.PushCommandList(Viewport.Pin()->GetCommandList().ToSharedRef());
-	InMenuBuilder.PushExtender(MenuExtender.ToSharedRef());
-
-	const USkeletalMesh* PreviewMesh =  Viewport.Pin()->GetPreviewScene()->GetPreviewMesh();
-	auto IsLODSimplified = [PreviewMesh](const int32 InLODId) -> bool
-	{
-		return PreviewMesh && !PreviewMesh->IsCompiling() && PreviewMesh->IsValidLODIndex(InLODId) &&
-			   PreviewMesh->GetLODInfo(InLODId)->bHasBeenSimplified;
-	};
-	
-	auto GetLODStatusExtraLabel = [IsLODSimplified, PreviewMesh](const int32 InLODId) -> FText
-	{
-		if (IsLODSimplified(InLODId))
-		{
-			if (PreviewMesh->HasMeshDescription(InLODId))
-			{
-				return LOCTEXT("LODStatus_Inline", " (Inline Reduced)");
-			}
-			
-			return LOCTEXT("LODStatus_Generated", " (Generated)");
-		}
-		return FText::GetEmpty();
-	};
-	auto GetLODStatusExtraTooltip = [IsLODSimplified, PreviewMesh](const int32 InLODId) -> FText
-	{
-		if (IsLODSimplified(InLODId))
-		{
-			if (PreviewMesh->HasMeshDescription(InLODId))
-			{
-				return LOCTEXT("LODStatusTooltip_Inline", "Generated from the editable geometry stored on this LOD but has been simplified in place.");
-			}
-
-			const int32 BaseLOD = PreviewMesh->GetLODInfo(InLODId)->ReductionSettings.BaseLOD;
-			return FText::Format(LOCTEXT("LODStatusTooltip_Generated", "Generated from a reduced version of LOD {0}.\nIt contains no editable geometry."), FText::AsNumber(BaseLOD));
-		}
-		return LOCTEXT("LODStatusTooltip_Default", "Generated from the editable geometry stored on this LOD with no simplification applied.");
-	};
-	
-	{
-		// LOD Models
-		InMenuBuilder.BeginSection("AnimViewportPreviewLODs", LOCTEXT("ShowLOD_PreviewLabel", "Preview LODs") );
-		{
-			InMenuBuilder.AddMenuEntry( Actions.LODDebug );
-			InMenuBuilder.AddMenuEntry( Actions.LODAuto );
-			InMenuBuilder.AddMenuEntry( Actions.LOD0 );
-
-			const int32 LODCount = Viewport.Pin()->GetLODModelCount();
-			for (int32 LODId = 1; LODId < LODCount; ++LODId)
-			{
-				const FText LODNumber = FText::AsNumber(LODId);
-				const FText TitleLabel = FText::Format(LOCTEXT("LODFmt", "LOD {0}{1}"), LODNumber, GetLODStatusExtraLabel(LODId));
-				const FText TooltipText = FText::Format(LOCTEXT("LODTooltip", "Force select LOD {0}.\n\n{1}"), LODNumber, GetLODStatusExtraTooltip(LODId));
-
-				FUIAction Action(FExecuteAction::CreateSP(Viewport.Pin().ToSharedRef(), &SAnimationEditorViewportTabBody::OnSetLODModel, LODId + 1),
-					FCanExecuteAction(),
-					FIsActionChecked::CreateSP(Viewport.Pin().ToSharedRef(), &SAnimationEditorViewportTabBody::IsLODModelSelected, LODId + 1));
-
-				InMenuBuilder.AddMenuEntry(TitleLabel, TooltipText, FSlateIcon(), Action, NAME_None, EUserInterfaceActionType::RadioButton);
-			}
-		}
-		InMenuBuilder.EndSection();
-	}
-
-	InMenuBuilder.PopCommandList();
-	InMenuBuilder.PopExtender();
-
-	return InMenuBuilder.MakeWidget();
+	return UE::AnimationEditor::GenerateLODMenuWidget(Viewport);
 }
 
 TSharedRef<SWidget> SAnimViewportToolBar::GenerateViewportTypeMenu() const
 {
-
 	TSharedPtr<FExtender> MenuExtender = FExtender::Combine(Extenders);
 
 	const bool bInShouldCloseWindowAfterMenuSelection = true;
@@ -1306,16 +896,23 @@ TSharedRef<SWidget> SAnimViewportToolBar::GeneratePlaybackMenu() const
 				{
 					InMenuBuilder.AddMenuEntry( Actions.PlaybackSpeedCommands[PlaybackSpeedIndex] );
 				}
-				TSharedPtr<SWidget> AnimSpeedWidget = SNew(SCusttomAnimationSpeedSetting)
-														.AnimEditorViewport(Viewport)
-														.CustomSpeed_Lambda([Viewport = Viewport]()
-															{
-																return Viewport.Pin()->GetCustomAnimationSpeed();
-															})
-														.OnCustomSpeedChanged_Lambda([Viewport = Viewport](float CustomSpeed)
-															{
-																return Viewport.Pin()->SetCustomAnimationSpeed(CustomSpeed);
-															});
+
+				TSharedPtr<SWidget> AnimSpeedWidget =
+					SNew(UE::AnimationEditor::SCustomAnimationSpeedSetting)
+						.AnimEditorViewport(Viewport)
+						.CustomSpeed_Lambda(
+							[Viewport = Viewport]()
+							{
+								return Viewport.Pin()->GetCustomAnimationSpeed();
+							}
+						)
+						.OnCustomSpeedChanged_Lambda(
+							[Viewport = Viewport](float CustomSpeed)
+							{
+								return Viewport.Pin()->SetCustomAnimationSpeed(CustomSpeed);
+							}
+						);
+
 				InMenuBuilder.AddWidget(AnimSpeedWidget.ToSharedRef(), LOCTEXT("PlaybackMenu_Speed_Custom", "Custom Speed:"));
 			}
 			InMenuBuilder.EndSection();
@@ -1329,39 +926,13 @@ TSharedRef<SWidget> SAnimViewportToolBar::GeneratePlaybackMenu() const
 
 void SAnimViewportToolBar::GenerateTurnTableMenu(FMenuBuilder& MenuBuilder) const
 {
-	const FAnimViewportPlaybackCommands& Actions = FAnimViewportPlaybackCommands::Get();
-
-	const bool bInShouldCloseWindowAfterMenuSelection = true;
-
-	MenuBuilder.PushCommandList(Viewport.Pin()->GetCommandList().ToSharedRef());
-	MenuBuilder.BeginSection("AnimViewportTurnTableMode", LOCTEXT("TurnTableMenu_ModeLabel", "Turn Table Mode"));
-	{
-		MenuBuilder.AddMenuEntry(Actions.PersonaTurnTablePlay);
-		MenuBuilder.AddMenuEntry(Actions.PersonaTurnTablePause);
-		MenuBuilder.AddMenuEntry(Actions.PersonaTurnTableStop);
-	}
-	MenuBuilder.EndSection();
-
-	MenuBuilder.BeginSection("AnimViewportTurnTableSpeed", LOCTEXT("TurnTableMenu_SpeedLabel", "Turn Table Speed"));
-	{
-		for (int i = 0; i < EAnimationPlaybackSpeeds::NumPlaybackSpeeds; ++i)
-		{
-			MenuBuilder.AddMenuEntry(Actions.TurnTableSpeeds[i]);
-		}
-		TSharedPtr<SWidget> AnimSpeedWidget = SNew(SCusttomAnimationSpeedSetting)
-												.AnimEditorViewport(Viewport)
-												.CustomSpeed_Lambda([Viewport = Viewport]()
-													{
-														return Viewport.Pin()->GetCustomTurnTableSpeed();
-													})
-												.OnCustomSpeedChanged_Lambda([Viewport = Viewport](float CustomSpeed)
-													{
-														Viewport.Pin()->SetCustomTurnTableSpeed(CustomSpeed);
-													});
-		MenuBuilder.AddWidget(AnimSpeedWidget.ToSharedRef(), LOCTEXT("PlaybackMenu_Speed_Custom", "Custom Speed:"));
-	}
-	MenuBuilder.EndSection();
-	MenuBuilder.PopCommandList();
+	MenuBuilder.AddWidget(
+		UE::AnimationEditor::GenerateTurnTableMenu(Viewport),
+		FText(),
+		false,
+		true,
+		LOCTEXT("TurnTableTooltip", "Set up auto-rotation of preview.")
+	);
 }
 
 FSlateColor SAnimViewportToolBar::GetFontColor() const
@@ -1460,7 +1031,13 @@ void SAnimViewportToolBar::OnFOVValueChanged( float NewValue )
 	ViewportClient.ViewFOV = NewValue;
 	ViewportClient.Invalidate();
 
-	PinnedCommands->AddCustomWidget(TEXT("FOVWidget"));
+	if (TSharedPtr<SAnimationEditorViewportTabBody> ViewportTabPinned = Viewport.Pin())
+	{
+		if (TSharedPtr<IPinnedCommandList> CommandsPinned = ViewportTabPinned->GetPinnedCommands())
+		{
+			CommandsPinned->AddCustomWidget(TEXT("FOVWidget"));
+		}
+	}
 }
 
 void SAnimViewportToolBar::OnFOVValueCommitted( float NewValue, ETextCommit::Type CommitInfo )
@@ -1482,50 +1059,6 @@ void SAnimViewportToolBar::OnCamSpeedScalarChanged(float NewValue)
 	AnimViewportClient.ConfigOption->SetCameraSpeedScalar(AnimViewportClient.GetAssetEditorToolkit()->GetEditorName(), NewValue, AnimViewportClient.GetViewportIndex());
 }
 
-float SAnimViewportToolBar::OnGetFloorOffset() const
-{
-	if(Viewport.IsValid())
-	{
-		FAnimationViewportClient& AnimViewportClient = (FAnimationViewportClient&)Viewport.Pin()->GetLevelViewportClient();
-		return AnimViewportClient.GetFloorOffset();
-	}
-
-	return 0.0f;
-}
-
-void SAnimViewportToolBar::OnBeginSliderMovementFloorOffset()
-{
-	// This value is saved in a UPROPERTY for the floor mesh, so changes are transactional
-	PendingTransaction = MakeUnique<FScopedTransaction>(LOCTEXT("SetFloorOffset", "Set Floor Offset"));
-	PinnedCommands->AddCustomWidget(TEXT("FloorOffsetWidget"));
-}
-
-void SAnimViewportToolBar::OnFloorOffsetChanged( float NewValue )
-{
-	FAnimationViewportClient& AnimViewportClient = (FAnimationViewportClient&)Viewport.Pin()->GetLevelViewportClient();
-
-	AnimViewportClient.SetFloorOffset( NewValue );
-
-	PinnedCommands->AddCustomWidget(TEXT("FloorOffsetWidget"));
-}
-
-void SAnimViewportToolBar::OnFloorOffsetCommitted( float NewValue, ETextCommit::Type CommitType )
-{
-	if (!PendingTransaction)
-	{
-		// Create the transaction here if it doesn't already exist. This can happen when changes come via text entry to the slider.
-		PendingTransaction = MakeUnique<FScopedTransaction>(LOCTEXT("SetFloorOffset", "Set Floor Offset"));
-	}
-
-	FAnimationViewportClient& AnimViewportClient = (FAnimationViewportClient&)Viewport.Pin()->GetLevelViewportClient();
-
-	AnimViewportClient.SetFloorOffset( NewValue );
-
-	PinnedCommands->AddCustomWidget(TEXT("FloorOffsetWidget"));
-
-	PendingTransaction.Reset();
-}
-
 void SAnimViewportToolBar::AddMenuExtender(FName MenuToExtend, FMenuExtensionDelegate MenuBuilderDelegate)
 {
 	TSharedRef<FExtender> Extender(new FExtender());
diff --git a/Engine/Source/Editor/Persona/Private/SAnimViewportToolBar.h b/Engine/Source/Editor/Persona/Private/SAnimViewportToolBar.h
--- a/Engine/Source/Editor/Persona/Private/SAnimViewportToolBar.h
+++ b/Engine/Source/Editor/Persona/Private/SAnimViewportToolBar.h
@@ -15,7 +15,6 @@
 
 class FMenuBuilder;
 class SComboButton;
-class IPreviewProfileController;
 
 /**
  * A level viewport toolbar widget that is placed in a viewport
@@ -24,8 +23,7 @@ class SAnimViewportToolBar : public SViewportToolBar
 {
 public:
 	SLATE_BEGIN_ARGS( SAnimViewportToolBar )
-		: _PreviewProfileController(nullptr)
-		, _ShowShowMenu(true)
+		: _ShowShowMenu(true)
 		, _ShowCharacterMenu(true)
 		, _ShowLODMenu(true)
 		, _ShowPlaySpeedMenu(true)
@@ -36,8 +34,6 @@ public:
 
 	SLATE_ARGUMENT(FName, ContextName)
 
-	SLATE_ARGUMENT(TSharedPtr<IPreviewProfileController>, PreviewProfileController)
-
 	SLATE_ARGUMENT(bool, ShowShowMenu)
 
 	SLATE_ARGUMENT(bool, ShowCharacterMenu)
@@ -56,9 +52,6 @@ public:
 
 	void Construct(const FArguments& InArgs, TSharedPtr<class SAnimationEditorViewportTabBody> InViewport, TSharedPtr<class SEditorViewport> InRealViewport);
 
-	/** Get the pinned commands widget */
-	TSharedPtr<IPinnedCommandList> GetPinnedCommandList() { return PinnedCommands; }
-
 private:
 	/**
 	 * Generates the toolbar view menu content 
@@ -168,13 +161,6 @@ private:
 	/** Called when Camera speed scalar is changed */
 	void OnCamSpeedScalarChanged(float NewValue);
 
-	/** Called by the floor offset slider in the perspective viewport to get the offset value */
-	float OnGetFloorOffset() const;
-	/** Called when the floor offset slider is adjusted in the perspective viewport */
-	void OnBeginSliderMovementFloorOffset();
-	void OnFloorOffsetChanged( float NewValue );
-	void OnFloorOffsetCommitted ( float NewValue, ETextCommit::Type CommitType );
-
 	/** Called when users select the mirror table in the menu */
 	void OnMirrorDataTableSelected(const FAssetData& SelectedMirrorTableData);
 	
@@ -189,7 +175,6 @@ private:
 
 	/** Build the follow bone widget */
 	TSharedRef<SWidget> MakeFollowBoneWidget() const;
-	TSharedRef<SWidget> MakeFollowBoneWidget(TWeakPtr<SComboButton> InWeakComboButton) const;
 
 	/** Build the follow bone widget (as a combo button) */
 	TSharedRef<SWidget> MakeFollowBoneComboWidget() const;
@@ -205,7 +190,6 @@ private:
 	/** Extend the view menu */
 	TSharedRef<FExtender> GetViewMenuExtender(TSharedPtr<class SEditorViewport> InRealViewport);
 
-private:
 	/** The viewport that we are in */
 	TWeakPtr<class SAnimationEditorViewportTabBody> Viewport;
 
@@ -215,15 +199,6 @@ private:
 	/** Extenders */
 	TArray<TSharedPtr<FExtender>> Extenders;
 
-	/** Pinned commands widget */
-	TSharedPtr<IPinnedCommandList> PinnedCommands;
-
-	/** Transaction to handle scoping updates for sliders we own that transact objects (i.e. floor mesh) */
-	TUniquePtr<FScopedTransaction> PendingTransaction;
-
-	/** The controller used to toggle Preview Scene Profiles */
-	TSharedPtr<IPreviewProfileController> PreviewProfileController;
-
 	/** Whether to show the 'Show' menu */
 	bool bShowShowMenu;
 
diff --git a/Engine/Source/Editor/Persona/Private/SAnimationEditorViewport.cpp b/Engine/Source/Editor/Persona/Private/SAnimationEditorViewport.cpp
--- a/Engine/Source/Editor/Persona/Private/SAnimationEditorViewport.cpp
+++ b/Engine/Source/Editor/Persona/Private/SAnimationEditorViewport.cpp
@@ -145,7 +145,6 @@ TSharedPtr<SWidget> SAnimationEditorViewport::MakeViewportToolbar()
 		.Cursor(EMouseCursor::Default)
 		.Extenders(Extenders)
 		.ContextName(ContextName)
-		.PreviewProfileController(GetPreviewProfileController()) // TODO: this line can be removed when updating SAnimViewportToolBar (CL 37180383)
 		.ShowShowMenu(bShowShowMenu)
 		.ShowLODMenu(bShowLODMenu)
 		.ShowPlaySpeedMenu(bShowPlaySpeedMenu)
@@ -398,6 +397,7 @@ SAnimationEditorViewportTabBody::SAnimationEditorViewportTabBody()
 	, SelectedTurnTableMode(EPersonaTurnTableMode::Stopped)
 	, SectionsDisplayMode(ESectionDisplayMode::None)
 {
+	CreatePinnedCommands();
 }
 
 SAnimationEditorViewportTabBody::~SAnimationEditorViewportTabBody()
@@ -559,7 +559,7 @@ FEditorViewportClient& SAnimationEditorViewportTabBody::GetViewportClient() cons
 
 TSharedRef<IPinnedCommandList> SAnimationEditorViewportTabBody::GetPinnedCommandList() const
 {
-	return ViewportWidget->GetViewportToolbar()->GetPinnedCommandList().ToSharedRef();
+	return PinnedCommands.ToSharedRef();
 }
 
 const TSharedPtr<IPinnedCommandList>& SAnimationEditorViewportTabBody::GetPinnedCommands()
@@ -802,8 +802,6 @@ void SAnimationEditorViewportTabBody::Construct(const FArguments& InArgs, const
 
 	TSharedPtr<SVerticalBox> ViewportContainer = nullptr;
 
-	CreatePinnedCommands();
-
 	this->ChildSlot
 	[
 		SAssignNew(ViewportContainer, SVerticalBox)
diff --git a/Engine/Source/Editor/Persona/Private/ViewportToolbar/AnimationEditorMenus.cpp b/Engine/Source/Editor/Persona/Private/ViewportToolbar/AnimationEditorMenus.cpp
--- a/Engine/Source/Editor/Persona/Private/ViewportToolbar/AnimationEditorMenus.cpp
+++ b/Engine/Source/Editor/Persona/Private/ViewportToolbar/AnimationEditorMenus.cpp
@@ -70,18 +70,53 @@ void PopulateLODSubmenu(UToolMenu* InMenu)
 
 	const FAnimViewportLODCommands& Actions = FAnimViewportLODCommands::Get();
 
-	const USkeletalMesh* PreviewMesh = AnimEditorViewportTab->GetPreviewScene()->GetPreviewMesh();
-	auto IsAutoGenerated = [PreviewMeshWeak = TWeakObjectPtr<const USkeletalMesh>(PreviewMesh)](const int32 InLODId) -> bool
+	const TWeakObjectPtr<const USkeletalMesh> PreviewMeshWeak = AnimEditorViewportTab->GetPreviewScene()->GetPreviewMesh();
+
+	auto IsLODSimplified = [PreviewMeshWeak](const int32 InLODId) -> bool
+	{
+		if (const TStrongObjectPtr<const USkeletalMesh>& PreviewMeshPinned = PreviewMeshWeak.Pin())
+		{
+			return !PreviewMeshPinned->IsCompiling() && PreviewMeshPinned->IsValidLODIndex(InLODId) &&
+				   PreviewMeshPinned->GetLODInfo(InLODId)->bHasBeenSimplified;
+		}
+
+		return false;
+	};
+
+	auto GetLODStatusExtraLabel = [IsLODSimplified, PreviewMeshWeak](const int32 InLODId) -> FText
 	{
-		TStrongObjectPtr<const USkeletalMesh> PreviewMeshPinned = PreviewMeshWeak.Pin();
-		if (!PreviewMeshPinned || PreviewMeshPinned->IsCompiling())
+		if (const TStrongObjectPtr<const USkeletalMesh>& PreviewMeshPinned = PreviewMeshWeak.Pin())
 		{
-			return false;
+			if (IsLODSimplified(InLODId))
+			{
+				if (PreviewMeshPinned->HasMeshDescription(InLODId))
+				{
+					return LOCTEXT("LODStatus_Inline", " (Inline Reduced)");
+				}
+
+				return LOCTEXT("LODStatus_Generated", " (Generated)");
+			}
 		}
-		return PreviewMeshPinned->IsValidLODIndex(InLODId) && PreviewMeshPinned->GetLODInfo(InLODId)->bHasBeenSimplified;
+		return FText::GetEmpty();
+	};
+
+	auto GetLODStatusExtraTooltip = [IsLODSimplified, PreviewMeshWeak](const int32 InLODId) -> FText
+	{
+		if (const TStrongObjectPtr<const USkeletalMesh>& PreviewMeshPinned = PreviewMeshWeak.Pin())
+		{
+			if (IsLODSimplified(InLODId))
+			{
+				if (PreviewMeshPinned->HasMeshDescription(InLODId))
+				{
+					return LOCTEXT("LODStatusTooltip_Inline", "Generated from the editable geometry stored on this LOD but has been simplified in place.");
+				}
+
+				const int32 BaseLOD = PreviewMeshPinned->GetLODInfo(InLODId)->ReductionSettings.BaseLOD;
+				return FText::Format(LOCTEXT("LODStatusTooltip_Generated", "Generated from a reduced version of LOD {0}.\nIt contains no editable geometry."), FText::AsNumber(BaseLOD));
+			}
+		}
+		return LOCTEXT("LODStatusTooltip_Default", "Generated from the editable geometry stored on this LOD with no simplification applied.");
 	};
-	static const FText LODFormat = LOCTEXT("LODFmt", "LOD {0}");
-	static const FText GeneratedLODFormat = LOCTEXT("GeneratedLODFmt", "LOD {0} (generated)");
 
 	{
 		// LOD Models
@@ -96,8 +131,8 @@ void PopulateLODSubmenu(UToolMenu* InMenu)
 			for (int32 LODId = 1; LODId < LODCount; ++LODId)
 			{
 				const FText LODNumber = FText::AsNumber(LODId);
-				const FText TitleLabel = IsAutoGenerated(LODId) ? FText::Format(GeneratedLODFormat, LODNumber)
-																: FText::Format(LODFormat, LODNumber);
+				const FText TitleLabel = FText::Format(LOCTEXT("LODFmt", "LOD {0}{1}"), LODNumber, GetLODStatusExtraLabel(LODId));
+				const FText TooltipText = FText::Format(LOCTEXT("LODTooltip", "Force select LOD {0}.\n\n{1}"), LODNumber, GetLODStatusExtraTooltip(LODId));
 
 				FUIAction Action(
 					FExecuteAction::CreateSP(
@@ -110,7 +145,7 @@ void PopulateLODSubmenu(UToolMenu* InMenu)
 				);
 
 				LODSection.AddMenuEntry(
-					FName(TitleLabel.ToString()), FText(), FText(), FSlateIcon(), Action, EUserInterfaceActionType::RadioButton
+					FName(TitleLabel.ToString()), TitleLabel, TooltipText, FSlateIcon(), Action, EUserInterfaceActionType::RadioButton
 				);
 			}
 		}
@@ -896,7 +931,7 @@ FToolMenuEntry UE::AnimationEditor::CreateLODSubmenu()
 				InDynamicSection.AddSubMenu(
 					"LOD",
 					Label,
-					FText(),
+					LOCTEXT("LODMenuTooltip", "LOD Options. Control how LODs are displayed.\nShift-clicking items will 'pin' them to the toolbar."),
 					FNewToolMenuDelegate::CreateLambda(
 						[](UToolMenu* Submenu) -> void
 						{
@@ -1315,4 +1350,128 @@ TSharedRef<FExtender> UE::AnimationEditor::GetViewModesLegacyExtenders(const TWe
 	return Extender;
 }
 
+TSharedRef<SWidget> UE::AnimationEditor::CreateFollowModeMenuWidget(const TWeakPtr<SAnimationEditorViewportTabBody>& InAnimEditorViewportTab
+)
+{
+	TSharedPtr<SAnimationEditorViewportTabBody> AnimEditorViewportTab = InAnimEditorViewportTab.Pin();
+	if (!AnimEditorViewportTab)
+	{
+		return SNullWidget::NullWidget;
+	}
+	// We generate a menu via UToolMenus, so we can use FillShowSubmenu call from both old and new toolbar
+	FName OldFollowModeMenuName = "AnimationEditor.OldViewportToolbar.FollowMode";
+	if (!UToolMenus::Get()->IsMenuRegistered(OldFollowModeMenuName))
+	{
+		if (UToolMenu* Menu = UToolMenus::Get()->RegisterMenu(OldFollowModeMenuName, NAME_None, EMultiBoxType::Menu, false))
+		{
+			Menu->AddDynamicSection(
+				"BaseSection",
+				FNewToolMenuDelegate::CreateLambda(
+					[](UToolMenu* InMenu)
+					{
+						UE::AnimationEditor::FillFollowModeSubmenu(InMenu);
+					}
+				)
+			);
+		}
+	}
+	FToolMenuContext MenuContext;
+	{
+		MenuContext.AppendCommandList(AnimEditorViewportTab->GetCommandList());
+		UAnimViewportContext* const ContextObject = NewObject<UAnimViewportContext>();
+		ContextObject->ViewportTabBody = AnimEditorViewportTab;
+		MenuContext.AddObject(ContextObject);
+	}
+	return UToolMenus::Get()->GenerateWidget(OldFollowModeMenuName, MenuContext);
+}
+
+TSharedRef<SWidget> UE::AnimationEditor::GenerateTurnTableMenu(const TWeakPtr<SAnimationEditorViewportTabBody>& InAnimEditorViewportTab
+)
+{
+	TSharedPtr<SAnimationEditorViewportTabBody> AnimEditorViewportTab = InAnimEditorViewportTab.Pin();
+	if (!AnimEditorViewportTab)
+	{
+		return SNullWidget::NullWidget;
+	}
+	// We generate a menu via UToolMenus, so we can use FillShowSubmenu call from both old and new toolbar
+	FName OldShowMenuName = "AnimationEditor.OldViewportToolbar.TurnTable";
+	if (!UToolMenus::Get()->IsMenuRegistered(OldShowMenuName))
+	{
+		if (UToolMenu* Menu = UToolMenus::Get()->RegisterMenu(OldShowMenuName, NAME_None, EMultiBoxType::Menu, false))
+		{
+			Menu->AddDynamicSection(
+				"BaseSection",
+				FNewToolMenuDelegate::CreateLambda(
+					[](UToolMenu* InMenu)
+					{
+						UE::AnimationEditor::FillTurnTableSubmenu(InMenu);
+					}
+				)
+			);
+		}
+	}
+	FToolMenuContext MenuContext;
+	{
+		MenuContext.AppendCommandList(AnimEditorViewportTab->GetCommandList());
+		UAnimViewportContext* const ContextObject = NewObject<UAnimViewportContext>();
+		ContextObject->ViewportTabBody = AnimEditorViewportTab;
+		MenuContext.AddObject(ContextObject);
+	}
+	return UToolMenus::Get()->GenerateWidget(OldShowMenuName, MenuContext);
+}
+
+TSharedRef<SWidget> UE::AnimationEditor::GenerateLODMenuWidget(const TWeakPtr<SAnimationEditorViewportTabBody>& InAnimEditorViewportTab
+)
+{
+	TSharedPtr<SAnimationEditorViewportTabBody> AnimEditorViewportTab = InAnimEditorViewportTab.Pin();
+	if (!AnimEditorViewportTab)
+	{
+		return SNullWidget::NullWidget;
+	}
+	// We generate a menu via UToolMenus, so we can use FillShowSubmenu call from both old and new toolbar
+	FName OldShowMenuName = "AnimationEditor.OldViewportToolbar.LODMenu";
+	if (!UToolMenus::Get()->IsMenuRegistered(OldShowMenuName))
+	{
+		if (UToolMenu* Menu = UToolMenus::Get()->RegisterMenu(OldShowMenuName, NAME_None, EMultiBoxType::Menu, false))
+		{
+			Menu->AddDynamicSection(
+				"BaseSection",
+				FNewToolMenuDelegate::CreateLambda(
+					[](UToolMenu* InMenu)
+					{
+						UE::AnimationEditor::Private::PopulateLODSubmenu(InMenu);
+					}
+				)
+			);
+		}
+	}
+	FToolMenuContext MenuContext;
+	{
+		MenuContext.AppendCommandList(AnimEditorViewportTab->GetCommandList());
+		UAnimViewportContext* const ContextObject = NewObject<UAnimViewportContext>();
+		ContextObject->ViewportTabBody = AnimEditorViewportTab;
+		MenuContext.AddObject(ContextObject);
+	}
+	return UToolMenus::Get()->GenerateWidget(OldShowMenuName, MenuContext);
+}
+
+TSharedRef<SWidget> UE::AnimationEditor::CreateShowMenuWidget(
+	const TSharedRef<SEditorViewport>& InViewport, const TArray<TSharedPtr<FExtender>>& InExtenders, bool bInShowViewportStatsToggle
+)
+{
+	static const FName MenuName("Persona.AnimViewportToolBar");
+	if (!UToolMenus::Get()->IsMenuRegistered(MenuName))
+	{
+		UToolMenu* Menu = UToolMenus::Get()->RegisterMenu(MenuName);
+		FillShowSubmenu(Menu, bInShowViewportStatsToggle);
+	}
+	FToolMenuContext MenuContext;
+	{
+		MenuContext.AppendCommandList(InViewport->GetCommandList());
+		TSharedPtr<FExtender> MenuExtender = FExtender::Combine(InExtenders);
+		MenuContext.AddExtender(MenuExtender);
+	}
+	return UToolMenus::Get()->GenerateWidget(MenuName, MenuContext);
+}
+
 #undef LOCTEXT_NAMESPACE
diff --git a/Engine/Source/Editor/Persona/Private/ViewportToolbar/AnimationEditorMenus.h b/Engine/Source/Editor/Persona/Private/ViewportToolbar/AnimationEditorMenus.h
--- a/Engine/Source/Editor/Persona/Private/ViewportToolbar/AnimationEditorMenus.h
+++ b/Engine/Source/Editor/Persona/Private/ViewportToolbar/AnimationEditorMenus.h
@@ -17,6 +17,7 @@ namespace UE::AnimationEditor
 void ExtendCameraMenu(FName InCameraOptionsMenuName);
 
 void FillFollowModeSubmenu(UToolMenu* InMenu);
+TSharedRef<SWidget> CreateFollowModeMenuWidget(const TWeakPtr<SAnimationEditorViewportTabBody>& InAnimEditorViewportTab);
 
 TSharedRef<SWidget> MakeFollowBoneWidget(
 	const TWeakPtr<SAnimationEditorViewportTabBody>& InViewport, const TWeakPtr<SComboButton>& InWeakComboButton
@@ -24,14 +25,19 @@ TSharedRef<SWidget> MakeFollowBoneWidget(
 
 FToolMenuEntry CreateShowSubmenu();
 void FillShowSubmenu(UToolMenu* InMenu, bool bInShowViewportStatsToggle = true);
+TSharedRef<SWidget> CreateShowMenuWidget(
+	const TSharedRef<SEditorViewport>& InViewport, const TArray<TSharedPtr<FExtender>>& InExtenders, bool bInShowViewportStatsToggle
+);
 
 FToolMenuEntry CreateTurnTableMenu();
 void FillTurnTableSubmenu(UToolMenu* InMenu);
+TSharedRef<SWidget> GenerateTurnTableMenu(const TWeakPtr<SAnimationEditorViewportTabBody>& InAnimEditorViewportTab);
 
 void AddSceneElementsSection(UToolMenu* InMenu, bool bInShowAudioEntries = false, bool bInShowTimecodeEntries = false);
 
 FToolMenuEntry CreateLODSubmenu();
 FText GetLODMenuLabel(const TWeakPtr<SAnimationEditorViewportTabBody>& InAnimEditorViewportTab);
+TSharedRef<SWidget> GenerateLODMenuWidget(const TWeakPtr<SAnimationEditorViewportTabBody>& InAnimEditorViewportTab);
 
 TSharedRef<SWidget> MakeFloorOffsetWidget(const TWeakPtr<SAnimationEditorViewportTabBody>& InAnimEditorViewportTabWeak);
 FToolMenuEntry CreateSettingsSubmenu();
