[Viewport Toolbar] Add missing Show and Settings icons to main toolbar entries

- Static Mesh
- Persona
- Water
- Smart Objects

diff --git a/Engine/Source/Editor/Persona/Private/ViewportToolbar/AnimationEditorMenus.cpp b/Engine/Source/Editor/Persona/Private/ViewportToolbar/AnimationEditorMenus.cpp
--- a/Engine/Source/Editor/Persona/Private/ViewportToolbar/AnimationEditorMenus.cpp
+++ b/Engine/Source/Editor/Persona/Private/ViewportToolbar/AnimationEditorMenus.cpp
@@ -454,7 +454,7 @@ TSharedRef<SWidget> UE::AnimationEditor::MakeFollowBoneWidget(
 
 FToolMenuEntry UE::AnimationEditor::CreateShowSubmenu()
 {
-	return FToolMenuEntry::InitSubMenu(
+	FToolMenuEntry Entry = FToolMenuEntry::InitSubMenu(
 		"Show",
 		LOCTEXT("ShowSubmenuLabel", "Show"),
 		LOCTEXT("ShowSubmenuTooltip", "Show options"),
@@ -466,6 +466,11 @@ FToolMenuEntry UE::AnimationEditor::CreateShowSubmenu()
 			}
 		)
 	);
+
+	Entry.ToolBarData.LabelOverride = FText();
+	Entry.Icon = FSlateIcon(FAppStyle::GetAppStyleSetName(), "Level.VisibleHighlightIcon16x");
+
+	return Entry;
 }
 
 void UE::AnimationEditor::FillShowSubmenu(UToolMenu* InMenu, bool bInShowViewportStatsToggle)
@@ -1035,7 +1040,7 @@ TSharedRef<SWidget> UE::AnimationEditor::MakeFloorOffsetWidget(const TWeakPtr<SA
 
 FToolMenuEntry UE::AnimationEditor::CreateSettingsSubmenu()
 {
-	return FToolMenuEntry::InitSubMenu(
+	FToolMenuEntry Entry = FToolMenuEntry::InitSubMenu(
 		"Settings",
 		LOCTEXT("SettingsSubmenuLabel", "Settings"),
 		LOCTEXT("SettingsSubmenuTooltip", "Settings"),
@@ -1082,6 +1087,11 @@ FToolMenuEntry UE::AnimationEditor::CreateSettingsSubmenu()
 			}
 		)
 	);
+
+	Entry.Icon = FSlateIcon(FAppStyle::GetAppStyleSetName(), "LevelEditor.GameSettings");
+	Entry.ToolBarData.LabelOverride = FText();
+
+	return Entry;
 }
 
 void UE::AnimationEditor::ExtendCameraMenu(FName InCameraOptionsMenuName)
diff --git a/Engine/Source/Editor/StaticMeshEditor/Private/StaticMeshEditorViewportToolbarSections.cpp b/Engine/Source/Editor/StaticMeshEditor/Private/StaticMeshEditorViewportToolbarSections.cpp
--- a/Engine/Source/Editor/StaticMeshEditor/Private/StaticMeshEditorViewportToolbarSections.cpp
+++ b/Engine/Source/Editor/StaticMeshEditor/Private/StaticMeshEditorViewportToolbarSections.cpp
@@ -123,7 +123,7 @@ FToolMenuEntry UE::StaticMeshEditor::CreateLODSubmenu()
 
 FToolMenuEntry UE::StaticMeshEditor::CreateShowSubmenu()
 {
-	return FToolMenuEntry::InitSubMenu(
+	FToolMenuEntry ShowSubmenu = FToolMenuEntry::InitSubMenu(
 		"Show",
 		LOCTEXT("ShowSubmenuLabel", "Show"),
 		FText(),
@@ -134,6 +134,11 @@ FToolMenuEntry UE::StaticMeshEditor::CreateShowSubmenu()
 			}
 		)
 	);
+
+	ShowSubmenu.ToolBarData.LabelOverride = FText();
+	ShowSubmenu.Icon = FSlateIcon(FAppStyle::GetAppStyleSetName(), "Level.VisibleHighlightIcon16x");
+
+	return ShowSubmenu;
 }
 
 TSharedRef<SWidget> UE::StaticMeshEditor::GenerateLODMenuWidget(const TSharedPtr<SStaticMeshEditorViewport>& InStaticMeshEditorViewport
diff --git a/Engine/Source/Editor/UnrealEd/Private/ViewportToolbar/UnrealEdViewportToolbar.cpp b/Engine/Source/Editor/UnrealEd/Private/ViewportToolbar/UnrealEdViewportToolbar.cpp
--- a/Engine/Source/Editor/UnrealEd/Private/ViewportToolbar/UnrealEdViewportToolbar.cpp
+++ b/Engine/Source/Editor/UnrealEd/Private/ViewportToolbar/UnrealEdViewportToolbar.cpp
@@ -3149,7 +3149,7 @@ FText GetScalabilityWarningTooltip()
 
 FToolMenuEntry CreateDefaultShowSubmenu()
 {
-	return FToolMenuEntry::InitSubMenu(
+	FToolMenuEntry Entry = FToolMenuEntry::InitSubMenu(
 		"Show",
 		LOCTEXT("ShowSubmenuLabel", "Show"),
 		LOCTEXT("ShowSubmenuTooltip", "Show flags related to the current viewport"),
@@ -3160,6 +3160,11 @@ FToolMenuEntry CreateDefaultShowSubmenu()
 			}
 		)
 	);
+
+	Entry.ToolBarData.LabelOverride = FText();
+	Entry.Icon = FSlateIcon(FAppStyle::GetAppStyleSetName(), "Level.VisibleHighlightIcon16x");
+
+	return Entry;
 }
 
 void AddDefaultShowFlags(UToolMenu* InMenu)
