[Viewport Toolbar] Persona - the cvar doesn't allow you to just show the new toolbar in the animBP editor

diff --git a/Engine/Source/Editor/Persona/Private/SAnimationEditorViewport.cpp b/Engine/Source/Editor/Persona/Private/SAnimationEditorViewport.cpp
--- a/Engine/Source/Editor/Persona/Private/SAnimationEditorViewport.cpp
+++ b/Engine/Source/Editor/Persona/Private/SAnimationEditorViewport.cpp
@@ -141,7 +141,17 @@ TSharedRef<FEditorViewportClient> SAnimationEditorViewport::MakeEditorViewportCl
 TSharedPtr<SWidget> SAnimationEditorViewport::MakeViewportToolbar()
 {
 	return SAssignNew(ViewportToolbar, SAnimViewportToolBar, TabBodyPtr.Pin(), SharedThis(this))
-		.Visibility(EVisibility::SelfHitTestInvisible)
+		.Visibility_Lambda(
+			[]() -> EVisibility
+			{
+				if (!UE::UnrealEd::ShowOldViewportToolbars())
+				{
+					return EVisibility::Collapsed;
+				}
+
+				return EVisibility::SelfHitTestInvisible;
+			}
+		)
 		.Cursor(EMouseCursor::Default)
 		.Extenders(Extenders)
 		.ContextName(ContextName)
