[Viewport Toolbar] Add Physics menu to PHAT editor viewport toolbar

- Shows Physics menu only when PHAT editor is open
- Same code used for both old and new toolbar
- In new toolbar, Persona Physics asset class icon is used instead of "Physics" label

diff --git a/Engine/Source/Editor/Persona/Private/SAnimViewportToolBar.cpp b/Engine/Source/Editor/Persona/Private/SAnimViewportToolBar.cpp
--- a/Engine/Source/Editor/Persona/Private/SAnimViewportToolBar.cpp
+++ b/Engine/Source/Editor/Persona/Private/SAnimViewportToolBar.cpp
@@ -470,18 +470,7 @@ TSharedRef<SWidget> SAnimViewportToolBar::GenerateViewMenu() const
 
 TSharedRef<SWidget> SAnimViewportToolBar::GeneratePhysicsMenu() const
 {
-	static const FName MenuName("Persona.AnimViewportPhysicsMenu");
-	if (!UToolMenus::Get()->IsMenuRegistered(MenuName))
-	{
-		UToolMenu* Menu = UToolMenus::Get()->RegisterMenu(MenuName);
-		Menu->AddSection("AnimViewportPhysicsMenu", LOCTEXT("ViewMenu_AnimViewportPhysicsMenu", "Physics Menu"));
-	}
-
-	TSharedPtr<FExtender> MenuExtender = FExtender::Combine(Extenders);
-	TSharedPtr<SAnimationEditorViewportTabBody> PinnedViewport = Viewport.Pin();
-	FToolMenuContext MenuContext(PinnedViewport->GetCommandList(), MenuExtender);
-	PinnedViewport->GetAssetEditorToolkit()->InitToolMenuContext(MenuContext);
-	return UToolMenus::Get()->GenerateWidget(MenuName, MenuContext);
+	return UE::AnimationEditor::GeneratePhysicsMenuWidget(Viewport, FExtender::Combine(Extenders));
 }
 
 TSharedRef<SWidget> SAnimViewportToolBar::GenerateCharacterMenu() const
diff --git a/Engine/Source/Editor/Persona/Private/SAnimationEditorViewport.cpp b/Engine/Source/Editor/Persona/Private/SAnimationEditorViewport.cpp
--- a/Engine/Source/Editor/Persona/Private/SAnimationEditorViewport.cpp
+++ b/Engine/Source/Editor/Persona/Private/SAnimationEditorViewport.cpp
@@ -243,6 +243,12 @@ TSharedPtr<SWidget> SAnimationEditorViewport::BuildViewportToolbar()
 			RightSection.AddEntry(UE::AnimationEditor::CreateShowSubmenu());
 			RightSection.AddEntry(UE::AnimationEditor::CreateLODSubmenu());
 			RightSection.AddEntry(UE::AnimationEditor::CreateSettingsSubmenu());
+
+			// Add the "Physics" sub menu (only shown when bShowPhysicsMenu is true)
+			{
+				FToolMenuInsert PhysicsSubmenuInsert("LOD", EToolMenuInsertType::Before);
+				UE::AnimationEditor::AddPhysicsMenu(ViewportToolbarName, PhysicsSubmenuInsert);
+			}
 		}
 	}
 
diff --git a/Engine/Source/Editor/Persona/Private/SAnimationEditorViewport.h b/Engine/Source/Editor/Persona/Private/SAnimationEditorViewport.h
--- a/Engine/Source/Editor/Persona/Private/SAnimationEditorViewport.h
+++ b/Engine/Source/Editor/Persona/Private/SAnimationEditorViewport.h
@@ -94,6 +94,7 @@ public:
 	/** Get the viewport toolbar widget */
 	TSharedPtr<SAnimViewportToolBar> GetViewportToolbar() const { return ViewportToolbar; }
 
+	bool IsPhysicsEditor() const { return bShowPhysicsMenu; }
 protected:
 	// SEditorViewport interface
 	virtual TSharedRef<FEditorViewportClient> MakeEditorViewportClient() override;
diff --git a/Engine/Source/Editor/Persona/Private/ViewportToolbar/AnimationEditorMenus.cpp b/Engine/Source/Editor/Persona/Private/ViewportToolbar/AnimationEditorMenus.cpp
--- a/Engine/Source/Editor/Persona/Private/ViewportToolbar/AnimationEditorMenus.cpp
+++ b/Engine/Source/Editor/Persona/Private/ViewportToolbar/AnimationEditorMenus.cpp
@@ -1368,6 +1368,93 @@ TSharedRef<FExtender> UE::AnimationEditor::GetViewModesLegacyExtenders(const TWe
 	return Extender;
 }
 
+void UE::AnimationEditor::AddPhysicsMenu(FName InPhysicsSubmenuName, FToolMenuInsert InInsertPosition)
+{
+	UToolMenu* const Submenu = UToolMenus::Get()->ExtendMenu(InPhysicsSubmenuName);
+	if (!Submenu)
+	{
+		return;
+	}
+
+	if (FToolMenuSection* RightSection = Submenu->FindSection("Right"))
+	{
+		RightSection->AddDynamicEntry(
+			"PhysicsDynamic",
+			FNewToolMenuSectionDelegate::CreateLambda(
+				[InInsertPosition](FToolMenuSection& InSection)
+				{
+					TSharedPtr<SAnimationEditorViewportTabBody> ViewportTab =
+						Private::GetAnimationEditorViewportTab(InSection);
+					if (!ViewportTab)
+					{
+						return;
+					}
+
+					if (TSharedPtr<SAnimationEditorViewport> ViewportWidget =
+							StaticCastSharedPtr<SAnimationEditorViewport>(ViewportTab->GetViewportWidget()))
+					{
+						// Only show physics sub menu when needed
+						if (ViewportWidget->IsPhysicsEditor())
+						{
+							FToolMenuEntry Entry = FToolMenuEntry::InitSubMenu(
+								"Physics",
+								LOCTEXT("PhysicsLabel", "Physics"),
+								LOCTEXT("PhysicsTooltip", "Physics Options. Use this to control the physics of the scene."),
+								FNewToolMenuDelegate::CreateLambda(
+									[](UToolMenu* Submenu) -> void
+									{
+										UE::AnimationEditor::FillPhysicsSubmenu(Submenu);
+									}
+								)
+							);
+
+							Entry.Icon = FSlateIcon(FAppStyle::GetAppStyleSetName(), "Persona.AssetClass.Physics");
+							Entry.ToolBarData.LabelOverride = FText();
+							Entry.InsertPosition = InInsertPosition;
+
+							InSection.AddEntry(Entry);
+						}
+					}
+				}
+			)
+		);
+	}
+}
+
+void UE::AnimationEditor::FillPhysicsSubmenu(UToolMenu* InMenu)
+{
+	if (!InMenu)
+	{
+		return;
+	}
+
+	TSharedPtr<SAnimationEditorViewportTabBody> ViewportTab = Private::GetAnimationEditorViewportTab(InMenu);
+	TSharedRef<SWidget> PhysicsMenuWidget = GeneratePhysicsMenuWidget(ViewportTab, InMenu->Context.GetAllExtenders());
+
+	InMenu->AddMenuEntry(NAME_None, FToolMenuEntry::InitWidget("Physics", PhysicsMenuWidget, FText()));
+}
+
+TSharedRef<SWidget> UE::AnimationEditor::GeneratePhysicsMenuWidget(
+	const TWeakPtr<SAnimationEditorViewportTabBody>& InAnimEditorViewportTab, const TSharedPtr<FExtender>& MenuExtender
+)
+{
+	if (TSharedPtr<SAnimationEditorViewportTabBody> ViewportTabPinned = InAnimEditorViewportTab.Pin())
+	{
+		static const FName MenuName("Persona.AnimViewportPhysicsMenu");
+		if (!UToolMenus::Get()->IsMenuRegistered(MenuName))
+		{
+			UToolMenu* Menu = UToolMenus::Get()->RegisterMenu(MenuName);
+			Menu->AddSection("AnimViewportPhysicsMenu", LOCTEXT("ViewMenu_AnimViewportPhysicsMenu", "Physics Menu"));
+		}
+
+		FToolMenuContext MenuContext(ViewportTabPinned->GetCommandList(), MenuExtender);
+		ViewportTabPinned->GetAssetEditorToolkit()->InitToolMenuContext(MenuContext);
+		return UToolMenus::Get()->GenerateWidget(MenuName, MenuContext);
+	}
+
+	return SNullWidget::NullWidget;
+}
+
 TSharedRef<SWidget> UE::AnimationEditor::CreateFollowModeMenuWidget(const TWeakPtr<SAnimationEditorViewportTabBody>& InAnimEditorViewportTab
 )
 {
diff --git a/Engine/Source/Editor/Persona/Private/ViewportToolbar/AnimationEditorMenus.h b/Engine/Source/Editor/Persona/Private/ViewportToolbar/AnimationEditorMenus.h
--- a/Engine/Source/Editor/Persona/Private/ViewportToolbar/AnimationEditorMenus.h
+++ b/Engine/Source/Editor/Persona/Private/ViewportToolbar/AnimationEditorMenus.h
@@ -3,6 +3,7 @@
 #pragma once
 
 #include "Widgets/SCompoundWidget.h"
+#include "ToolMenuMisc.h"
 
 class FExtender;
 class IPreviewProfileController;
@@ -44,4 +45,8 @@ FToolMenuEntry CreateSettingsSubmenu();
 
 void ExtendViewModesSubmenu(FName InViewModesSubmenuName);
 TSharedRef<FExtender> GetViewModesLegacyExtenders(const TWeakPtr<SAnimationEditorViewportTabBody>& InViewport);
+
+void AddPhysicsMenu(FName InPhysicsSubmenuName, FToolMenuInsert InInsertPosition);
+void FillPhysicsSubmenu(UToolMenu* InMenu);
+TSharedRef<SWidget> GeneratePhysicsMenuWidget(const TWeakPtr<SAnimationEditorViewportTabBody>& InAnimEditorViewportTab, const TSharedPtr<FExtender>& MenuExtender);
 } // namespace UE::AnimationEditor
