[Viewport Toolbar] Persona Editor - Use Pinned Commands and Viewport Display Text for both new and old toolbar

diff --git a/Engine/Source/Editor/Persona/Private/SAnimViewportToolBar.cpp b/Engine/Source/Editor/Persona/Private/SAnimViewportToolBar.cpp
--- a/Engine/Source/Editor/Persona/Private/SAnimViewportToolBar.cpp
+++ b/Engine/Source/Editor/Persona/Private/SAnimViewportToolBar.cpp
@@ -38,7 +38,6 @@
 #include "BufferVisualizationMenuCommands.h"
 #include "IPinnedCommandList.h"
 #include "UICommandList_Pinnable.h"
-#include "Widgets/Text/SRichTextBlock.h"
 #include "ContentBrowserModule.h"
 #include "IContentBrowserSingleton.h"
 #include "SCommonEditorViewportToolbarBase.h"
@@ -271,21 +270,6 @@ void SAnimViewportToolBar::Construct(
 				LeftToolbar
 			]
 		]
-		+ SVerticalBox::Slot()
-		.AutoHeight()
-		[
-			PinnedCommands.IsValid() ? PinnedCommands.ToSharedRef() : SNullWidget::NullWidget
-		]
-		+ SVerticalBox::Slot()
-		.Padding(FMargin(4.0f, 3.0f, 0.0f, 0.0f))
-		[
-			// Display text (e.g., item being previewed)
-			SNew(SRichTextBlock)
-			.Visibility(EVisibility::SelfHitTestInvisible)
-			.DecoratorStyleSet(&FAppStyle::Get())
-			.Text(InViewport.Get(), &SAnimationEditorViewportTabBody::GetDisplayString)
-			.TextStyle(&FAppStyle::Get().GetWidgetStyle<FTextBlockStyle>("AnimViewport.MessageText"))
-		]
 	];
 	// clang-format on
 
@@ -304,8 +288,6 @@ void SAnimViewportToolBar::Construct(
 			TEXT("FOVWidget"),
 			LOCTEXT("Viewport_FOVLabel", "Field Of View")
 		);
-
-		PinnedCommands->BindCommandList(InViewport->GetCommandList().ToSharedRef());
 	}
 
 	// We assign the viewport pointer here rather than initially, as SViewportToolbar::Construct
diff --git a/Engine/Source/Editor/Persona/Private/SAnimationEditorViewport.cpp b/Engine/Source/Editor/Persona/Private/SAnimationEditorViewport.cpp
--- a/Engine/Source/Editor/Persona/Private/SAnimationEditorViewport.cpp
+++ b/Engine/Source/Editor/Persona/Private/SAnimationEditorViewport.cpp
@@ -46,6 +46,7 @@
 #include "ViewportToolbar/AnimViewportContext.h"
 #include "ViewportToolbar/AnimationEditorMenus.h"
 #include "ViewportToolbar/UnrealEdViewportToolbar.h"
+#include "Widgets/Text/SRichTextBlock.h"
 
 #define LOCTEXT_NAMESPACE "PersonaViewportToolbar"
 
@@ -140,8 +141,13 @@ TSharedRef<FEditorViewportClient> SAnimationEditorViewport::MakeEditorViewportCl
 
 TSharedPtr<SWidget> SAnimationEditorViewport::MakeViewportToolbar()
 {
-	return SAssignNew(ViewportToolbar, SAnimViewportToolBar, TabBodyPtr.Pin(), SharedThis(this))
-		.Visibility_Lambda(
+	// clang-format off
+	TSharedRef<SVerticalBox> ToolbarWidget =
+		SNew(SVerticalBox)
+		+ SVerticalBox::Slot()
+		[
+			SAssignNew(ViewportToolbar, SAnimViewportToolBar, TabBodyPtr.Pin(), SharedThis(this))
+			.Visibility_Lambda(
 			[]() -> EVisibility
 			{
 				if (!UE::UnrealEd::ShowOldViewportToolbars())
@@ -150,17 +156,46 @@ TSharedPtr<SWidget> SAnimationEditorViewport::MakeViewportToolbar()
 				}
 
 				return EVisibility::SelfHitTestInvisible;
-			}
-		)
-		.Cursor(EMouseCursor::Default)
-		.Extenders(Extenders)
-		.ContextName(ContextName)
-		.ShowShowMenu(bShowShowMenu)
-		.ShowLODMenu(bShowLODMenu)
-		.ShowPlaySpeedMenu(bShowPlaySpeedMenu)
-		.ShowFloorOptions(bShowFloorOptions)
-		.ShowTurnTable(bShowTurnTable)
-		.ShowPhysicsMenu(bShowPhysicsMenu);
+			})
+			.Cursor(EMouseCursor::Default)
+			.Extenders(Extenders)
+			.ContextName(ContextName)
+			.ShowShowMenu(bShowShowMenu)
+			.ShowLODMenu(bShowLODMenu)
+			.ShowPlaySpeedMenu(bShowPlaySpeedMenu)
+			.ShowFloorOptions(bShowFloorOptions)
+			.ShowTurnTable(bShowTurnTable)
+			.ShowPhysicsMenu(bShowPhysicsMenu)
+		];
+	// clang-format on
+
+	if (TSharedPtr<SAnimationEditorViewportTabBody> ViewportTab = TabBodyPtr.Pin())
+	{
+		if (TSharedPtr<IPinnedCommandList> PinnedCommands = ViewportTab->GetPinnedCommands())
+		{
+			ToolbarWidget->AddSlot()
+			.AutoHeight()
+			[
+				PinnedCommands.ToSharedRef()
+			];
+		}
+
+		// clang-format off
+		ToolbarWidget->AddSlot()
+		.AutoHeight()
+		.Padding(FMargin(4.0f, 3.0f, 0.0f, 0.0f))
+		[
+			// Display text (e.g., item being previewed)
+			SNew(SRichTextBlock)
+			.Visibility(EVisibility::SelfHitTestInvisible)
+			.DecoratorStyleSet(&FAppStyle::Get())
+			.Text(ViewportTab.Get(), &SAnimationEditorViewportTabBody::GetDisplayString)
+			.TextStyle(&FAppStyle::Get().GetWidgetStyle<FTextBlockStyle>("AnimViewport.MessageText"))
+		];
+		// clang-format on
+	}
+
+	return ToolbarWidget;
 }
 
 TSharedPtr<SWidget> SAnimationEditorViewport::BuildViewportToolbar()
@@ -254,7 +289,10 @@ TSharedPtr<SWidget> SAnimationEditorViewport::BuildViewportToolbar()
 
 	FToolMenuContext ViewportToolbarContext;
 	{
-		ViewportToolbarContext.AppendCommandList(GetCommandList());
+		if (TSharedPtr<SAnimationEditorViewportTabBody> TabBody = TabBodyPtr.Pin())
+		{
+			ViewportToolbarContext.AppendCommandList(TabBody->GetCommandList());
+		}
 
 		// Add extenders
 		{
@@ -701,8 +739,6 @@ FReply SAnimationEditorViewportTabBody::OnKeyDown(const FGeometry& MyGeometry, c
 
 void SAnimationEditorViewportTabBody::Construct(const FArguments& InArgs, const TSharedRef<class IPersonaPreviewScene>& InPreviewScene, const TSharedRef<class FAssetEditorToolkit>& InAssetEditorToolkit, int32 InViewportIndex)
 {
-	UICommandList = MakeShareable(new FUICommandList_Pinnable);
-
 	PreviewScenePtr = StaticCastSharedRef<FAnimationEditorPreviewScene>(InPreviewScene);
 	AssetEditorToolkitPtr = InAssetEditorToolkit;
 	BlueprintEditorPtr = InArgs._BlueprintEditor;
@@ -1390,6 +1426,8 @@ void SAnimationEditorViewportTabBody::BindCommands()
 			ConstCastSharedRef<FUICommandList>(AssetEditorToolkit->GetToolkitCommands());
 		ToolkitCommandList->Append(UICommandList->AsShared());
 	}
+
+	PinnedCommands->BindCommandList(UICommandList.ToSharedRef());
 }
 
 void SAnimationEditorViewportTabBody::OnSetTurnTableSpeed(int32 SpeedIndex)
@@ -2344,6 +2382,8 @@ void SAnimationEditorViewportTabBody::CreatePinnedCommands()
 	PinnedCommands =
 		PinnedCommandListModule.CreatePinnedCommandList((ContextName != NAME_None) ? ContextName : TEXT("PersonaViewport"));
 	PinnedCommands->SetStyle(&FAppStyle::Get(), TEXT("ViewportPinnedCommandList"));
+
+	UICommandList = MakeShareable(new FUICommandList_Pinnable);
 }
 
 TSharedPtr<FAnimationViewportClient> SAnimationEditorViewportTabBody::GetAnimationViewportClient() const
