Morph target viewer now display the morph target origin this change complete a partial hotfix change cl: 37918191

#rnx

diff --git a/Engine/Source/Developer/SkeletalMeshUtilitiesCommon/Private/LODUtilities.cpp b/Engine/Source/Developer/SkeletalMeshUtilitiesCommon/Private/LODUtilities.cpp
--- a/Engine/Source/Developer/SkeletalMeshUtilitiesCommon/Private/LODUtilities.cpp
+++ b/Engine/Source/Developer/SkeletalMeshUtilitiesCommon/Private/LODUtilities.cpp
@@ -4122,13 +4122,18 @@ static void BuildMorphTargetsInternal(
 
 		//Reset the morph target imported source filename
 		MorphTarget->SetCustomImportedSourceFilename(LODIndex, FString());
+		MorphTarget->SetGeneratedByEngine(LODIndex, false);
 		//Now set the imported source filename from the imported data
 		if (BaseLodInfo)
 		{
 			if (const FMorphTargetImportedSourceFileInfo* MorphTargetImportedSourceFileInfo = BaseLodInfo->ImportedMorphTargetSourceFilename.Find(MorphTarget->GetName()))
 			{
 				const FString& SourceFilename = MorphTargetImportedSourceFileInfo->GetSourceFilename();
-				if (!SourceFilename.IsEmpty())
+				if (MorphTargetImportedSourceFileInfo->IsGeneratedByEngine())
+				{
+					MorphTarget->SetGeneratedByEngine(LODIndex, true);
+				}
+				else if (!SourceFilename.IsEmpty())
 				{
 					MorphTarget->SetCustomImportedSourceFilename(LODIndex, *SourceFilename);
 				}
diff --git a/Engine/Source/Editor/Persona/Private/SMorphTargetViewer.cpp b/Engine/Source/Editor/Persona/Private/SMorphTargetViewer.cpp
--- a/Engine/Source/Editor/Persona/Private/SMorphTargetViewer.cpp
+++ b/Engine/Source/Editor/Persona/Private/SMorphTargetViewer.cpp
@@ -135,23 +135,43 @@ TSharedRef< SWidget > SMorphTargetListRow::GenerateWidgetForColumn( const FName&
 		FText MorphNameText = FText::FromName(Item->Name);
 
 		bool bFirst = true;
-		for (int32 LodIndex = 0; LodIndex < Item->SourceFilenames.Num(); ++LodIndex)
+		for (int32 LodIndex = 0; LodIndex < Item->MorphLodInfos.Num(); ++LodIndex)
 		{
-			if (!Item->SourceFilenames[LodIndex].IsEmpty())
+			const FDisplayedMorphTargetInfo::FMorphLodInfo& MorphLodInfo = Item->MorphLodInfos[LodIndex];
+			const bool bLodCustomImported = !MorphLodInfo.SourceFilename.IsEmpty();
+			if (bLodCustomImported || MorphLodInfo.bIsGeneratedByEngine)
 			{
 				FText CarriageReturn;
 				if (bFirst)
 				{
-					MorphNameText = FText::Format(LOCTEXT("MorphRowName", "{0} (Imported)"), MorphNameText);
+					if(!MorphLodInfo.bIsGeneratedByEngine)
+					{
+						MorphNameText = FText::Format(LOCTEXT("MorphRowNameCustomImport", "{0} (Imported by File)"), MorphNameText);
+					}
+					else
+					{
+						MorphNameText = FText::Format(LOCTEXT("MorphRowNameGenerated", "{0} (Generated By Engine)"), MorphNameText);
+						SourceFilenamesTooltip = LOCTEXT("SourceFilenameTooltipEntry_Generated", "This morph target was generated by an engine tool.");
+					}
 					bFirst = false;
 				}
 				else
 				{
 					CarriageReturn = bFirst ? FText() : FText::FromString(TEXT("\n"));
 				}
-				SourceFilenamesTooltip = FText::Format(LOCTEXT("{0}{1}SourceFilenameTooltipEntry", "LOD {2} Source Filename: {3}"), CarriageReturn, SourceFilenamesTooltip, FText::AsNumber(LodIndex), FText::FromString(Item->SourceFilenames[LodIndex]));
+
+				if (bLodCustomImported)
+				{
+					SourceFilenamesTooltip = FText::Format(LOCTEXT("SourceFilenameTooltipEntry", "{0}{1}LOD {2} Source Filename: {3}"), SourceFilenamesTooltip, CarriageReturn, FText::AsNumber(LodIndex), FText::FromString(MorphLodInfo.SourceFilename));
+				}
 			}
 		}
+		//If we do not have a custom import lod, verify if we have a generated moprh
+		if (bFirst)
+		{
+			MorphNameText = FText::Format(LOCTEXT("MorphRowNameImportWithLodGeometry", "{0} (Imported with LOD geometry)"), MorphNameText);
+			SourceFilenamesTooltip = LOCTEXT("SourceFilenameTooltipEntry_ImportWithLodGeometry", "This morph target was imported with the LOD geometry.");
+		}
 
 		return
 			SNew( SVerticalBox )
@@ -654,13 +674,14 @@ void SMorphTargetViewer::CreateMorphTargetList( const FString& SearchText )
 
 			int32 NumberOfVerts = (MorphTargets[I]->GetMorphLODModels().Num() > 0)? MorphTargets[I]->GetMorphLODModels()[0].Vertices.Num() : 0;
 
-			TArray<FString> SourceFilenames;
-			SourceFilenames.AddDefaulted(SkeletalMesh->GetLODNum());
+			TArray<FDisplayedMorphTargetInfo::FMorphLodInfo> MorphLodInfos;
+			MorphLodInfos.AddDefaulted(SkeletalMesh->GetLODNum());
 			for (int32 LodIndex = 0; LodIndex < SkeletalMesh->GetLODNum(); ++LodIndex)
 			{
-				SourceFilenames[LodIndex] = MorphTargets[I]->GetCustomImportedSourceFilename(LodIndex);
+				MorphLodInfos[LodIndex].SourceFilename = MorphTargets[I]->GetCustomImportedSourceFilename(LodIndex);
+				MorphLodInfos[LodIndex].bIsGeneratedByEngine = MorphTargets[I]->IsGeneratedByEngine(LodIndex);
 			}
-			TSharedRef<FDisplayedMorphTargetInfo> Info = FDisplayedMorphTargetInfo::Make( MorphTargets[I]->GetFName(), NumberOfVerts, SourceFilenames);
+			TSharedRef<FDisplayedMorphTargetInfo> Info = FDisplayedMorphTargetInfo::Make( MorphTargets[I]->GetFName(), NumberOfVerts, MorphLodInfos);
 			if(MeshComponent)
 			{
 				const float *CurveValPtr = MeshComponent->GetMorphTargetCurves().Find( MorphTargets[I]->GetFName() );
diff --git a/Engine/Source/Editor/Persona/Private/SMorphTargetViewer.h b/Engine/Source/Editor/Persona/Private/SMorphTargetViewer.h
--- a/Engine/Source/Editor/Persona/Private/SMorphTargetViewer.h
+++ b/Engine/Source/Editor/Persona/Private/SMorphTargetViewer.h
@@ -24,22 +24,27 @@ public:
 	float Weight;
 	bool bAutoFillData;
 	int32 NumberOfVerts;
-	TArray<FString> SourceFilenames;
+	struct FMorphLodInfo
+	{
+		FString SourceFilename;
+		bool bIsGeneratedByEngine = false;
+	};
+	TArray<FMorphLodInfo> MorphLodInfos;
 
 	/** Static function for creating a new item, but ensures that you can only have a TSharedRef to one */
-	static TSharedRef<FDisplayedMorphTargetInfo> Make(const FName& Source, int32 InNumberOfVerts, const TArray<FString>& InSourceFilenames)
+	static TSharedRef<FDisplayedMorphTargetInfo> Make(const FName& Source, int32 InNumberOfVerts, const TArray<FMorphLodInfo>& InMorphLodInfos)
 	{
-		return MakeShareable(new FDisplayedMorphTargetInfo(Source, InNumberOfVerts, InSourceFilenames));
+		return MakeShareable(new FDisplayedMorphTargetInfo(Source, InNumberOfVerts, InMorphLodInfos));
 	}
 
 protected:
 	/** Hidden constructor, always use Make above */
-	FDisplayedMorphTargetInfo(const FName& InSource, int32 InNumberOfVerts, const TArray<FString>& InSourceFilenames)
+	FDisplayedMorphTargetInfo(const FName& InSource, int32 InNumberOfVerts, const TArray<FMorphLodInfo>& InMorphLodInfos)
 		: Name( InSource )
 		, Weight( 0 )
 		, bAutoFillData(true)
 		, NumberOfVerts(InNumberOfVerts)
-		, SourceFilenames(InSourceFilenames)
+		, MorphLodInfos(InMorphLodInfos)
 	{}
 
 	/** Hidden constructor, always use Make above */
diff --git a/Engine/Source/Runtime/Core/Private/UObject/FortniteMainBranchObjectVersion.cpp b/Engine/Source/Runtime/Core/Private/UObject/FortniteMainBranchObjectVersion.cpp
--- a/Engine/Source/Runtime/Core/Private/UObject/FortniteMainBranchObjectVersion.cpp
+++ b/Engine/Source/Runtime/Core/Private/UObject/FortniteMainBranchObjectVersion.cpp
@@ -12,7 +12,7 @@ TMap<FGuid, FGuid> FFortniteMainBranchObjectVersion::GetSystemGuids()
 	SystemGuids.Add(DevGuids.NANITE_DERIVEDDATA_VER, FGuid("42C46B24E1586986880AFE62B0F575AB"));
 	SystemGuids.Add(DevGuids.NIAGARASHADERMAP_DERIVEDDATA_VER, FGuid("C417A50271B3427D98785388266E3B66"));
 	SystemGuids.Add(DevGuids.Niagara_LatestScriptCompileVersion, FGuid("393081597FA7234C92F00869DE262D73"));
-	SystemGuids.Add(DevGuids.SkeletalMeshDerivedDataVersion, FGuid("B9AEAA3EE7AC4FCFB3860E8F3E06DD75"));
+	SystemGuids.Add(DevGuids.SkeletalMeshDerivedDataVersion, FGuid("0829C9BAB43B4B25A63747FDA68E7197"));
 	SystemGuids.Add(DevGuids.STATICMESH_DERIVEDDATA_VER, FGuid("68C7D98BE48042E0A35351325741F70E"));
 	SystemGuids.Add(DevGuids.MaterialTranslationDDCVersion, FGuid("D318C2D188414B7284D232E41F77BB74"));
 	return SystemGuids;
diff --git a/Engine/Source/Runtime/Engine/Classes/Animation/MorphTarget.h b/Engine/Source/Runtime/Engine/Classes/Animation/MorphTarget.h
--- a/Engine/Source/Runtime/Engine/Classes/Animation/MorphTarget.h
+++ b/Engine/Source/Runtime/Engine/Classes/Animation/MorphTarget.h
@@ -164,6 +164,10 @@ public:
 	ENGINE_API virtual FString GetCustomImportedSourceFilename(int32 LODIndex) const;
 	ENGINE_API virtual void SetCustomImportedSourceFilename(int32 LODIndex, const FString& InSourceFilename);
 
+	/** return true if this morphtarget was generated by an engine tool */
+	ENGINE_API virtual bool IsGeneratedByEngine(int32 LODIndex) const;
+	ENGINE_API virtual void SetGeneratedByEngine(int32 LODIndex, bool bInGeneratedByEngine);
+
 	/** Discard CPU Buffers after render resources have been created. */
 	UE_DEPRECATED(5.0, "No longer in use, will be deleted. Whether to discard vertex data is now determined during cooking instead of loading.")
 	ENGINE_API virtual void DiscardVertexData();
diff --git a/Engine/Source/Runtime/Engine/Classes/Engine/SkinnedAssetCommon.h b/Engine/Source/Runtime/Engine/Classes/Engine/SkinnedAssetCommon.h
--- a/Engine/Source/Runtime/Engine/Classes/Engine/SkinnedAssetCommon.h
+++ b/Engine/Source/Runtime/Engine/Classes/Engine/SkinnedAssetCommon.h
@@ -101,6 +101,9 @@ struct FMorphTargetImportedSourceFileInfo
 {
 	GENERATED_USTRUCT_BODY()
 
+	/** String that was used to mark generated by engine morph targets. */
+	static FString GeneratedByEnginePrefix;
+
 	ENGINE_API const FString& GetSourceFilename() const;
 
 	ENGINE_API void SetSourceFilename(const FString& Filename);
diff --git a/Engine/Source/Runtime/Engine/Private/MorphTools.cpp b/Engine/Source/Runtime/Engine/Private/MorphTools.cpp
--- a/Engine/Source/Runtime/Engine/Private/MorphTools.cpp
+++ b/Engine/Source/Runtime/Engine/Private/MorphTools.cpp
@@ -75,6 +75,13 @@ FArchive& operator<<(FArchive& Ar, FMorphTargetLODModel& M)
 				Ar << M.SourceFilename;
 			}
 		}
+
+		/* Fix the data for the generated by engine flag */
+		if(Ar.IsLoading() && M.SourceFilename.StartsWith(FMorphTargetImportedSourceFileInfo::GeneratedByEnginePrefix))
+		{
+			M.bGeneratedByEngine = true;
+			M.SourceFilename.Empty();
+		}
 	}
 
 	return Ar;
@@ -173,6 +180,32 @@ void UMorphTarget::SetCustomImportedSourceFilename(int32 LODIndex, const FString
 		// Calling GetMorphLODModels to potentially get from subclasses
 		FMorphTargetLODModel& MorphModel = GetMorphLODModels()[LODIndex];
 		MorphModel.SourceFilename = InSourceFilename;
+		MorphModel.bGeneratedByEngine = false;
+	}
+}
+
+bool UMorphTarget::IsGeneratedByEngine(int32 LODIndex) const
+{
+	if (LODIndex < MorphLODModels.Num())
+	{
+		// Calling GetMorphLODModels to potentially get from subclasses
+		const FMorphTargetLODModel& MorphModel = GetMorphLODModels()[LODIndex];
+		if(MorphModel.SourceFilename.IsEmpty())
+		{
+			return MorphModel.bGeneratedByEngine;
+		}
+	}
+	return false;
+}
+
+void UMorphTarget::SetGeneratedByEngine(int32 LODIndex, bool bInGeneratedByEngine)
+{
+	if (LODIndex < MorphLODModels.Num())
+	{
+		// Calling GetMorphLODModels to potentially get from subclasses
+		FMorphTargetLODModel& MorphModel = GetMorphLODModels()[LODIndex];
+		MorphModel.SourceFilename.Empty();
+		MorphModel.bGeneratedByEngine = bInGeneratedByEngine;
 	}
 }
 
diff --git a/Engine/Source/Runtime/Engine/Private/SkinnedAssetCommon.cpp b/Engine/Source/Runtime/Engine/Private/SkinnedAssetCommon.cpp
--- a/Engine/Source/Runtime/Engine/Private/SkinnedAssetCommon.cpp
+++ b/Engine/Source/Runtime/Engine/Private/SkinnedAssetCommon.cpp
@@ -24,14 +24,8 @@
 	FSkeletalMeshLODInfo
 -----------------------------------------------------------------------------*/
 
-namespace UE::MorphTargetSourceHelper
-{
-	const FString& GetGeneratedByEnginePrefix()
-	{
-		static FString GeneratedByEngineStr = TEXT("GeneratedByEngine");
-		return GeneratedByEngineStr;
-	}
-}
+//Init static string
+FString FMorphTargetImportedSourceFileInfo::GeneratedByEnginePrefix = TEXT("GeneratedByEngine_");
 
 const FString& FMorphTargetImportedSourceFileInfo::GetSourceFilename() const
 {
@@ -45,13 +39,13 @@ const FString& FMorphTargetImportedSourceFileInfo::GetSourceFilename() const
 
 bool FMorphTargetImportedSourceFileInfo::IsGeneratedByEngine() const
 {
-	return SourceFilename.StartsWith(UE::MorphTargetSourceHelper::GetGeneratedByEnginePrefix());
+	return SourceFilename.StartsWith(GeneratedByEnginePrefix);
 }
 
 void FMorphTargetImportedSourceFileInfo::SetGeneratedByEngine(bool bInGeneratedByEngine)
 {
-	FString NewSourceFilename = UE::MorphTargetSourceHelper::GetGeneratedByEnginePrefix();
-	NewSourceFilename += TEXT("_") + FGuid::NewGuid().ToString(EGuidFormats::Base36Encoded);
+	FString NewSourceFilename = GeneratedByEnginePrefix;
+	NewSourceFilename += FGuid::NewGuid().ToString(EGuidFormats::Base36Encoded);
 	SetSourceFilename(NewSourceFilename);
 }
 
