[Viewport Toolbar] Persona Editor - Update preview scene settings menu

Also removes settings submenu

diff --git a/Engine/Source/Editor/Persona/Private/SAnimationEditorViewport.cpp b/Engine/Source/Editor/Persona/Private/SAnimationEditorViewport.cpp
--- a/Engine/Source/Editor/Persona/Private/SAnimationEditorViewport.cpp
+++ b/Engine/Source/Editor/Persona/Private/SAnimationEditorViewport.cpp
@@ -268,7 +268,15 @@ TSharedPtr<SWidget> SAnimationEditorViewport::BuildViewportToolbar()
 
 			RightSection.AddEntry(UE::AnimationEditor::CreateShowSubmenu());
 			RightSection.AddEntry(UE::AnimationEditor::CreateLODSubmenu());
-			RightSection.AddEntry(UE::AnimationEditor::CreateSettingsSubmenu());
+
+			// Add Preview Scene Submenu
+			{
+				UE::AnimationEditor::ExtendPreviewSceneSettingsSubmenu(
+					"AnimationEditor.ViewportToolbar.AssetViewerProfile"
+				);
+
+				RightSection.AddEntry(UE::UnrealEd::CreateAssetViewerProfileSubmenu());
+			}
 
 			// Add the "Physics" sub menu (only shown when bShowPhysicsMenu is true)
 			{
@@ -296,6 +304,7 @@ TSharedPtr<SWidget> SAnimationEditorViewport::BuildViewportToolbar()
 			UUnrealEdViewportToolbarContext* const ContextObject =
 				UE::UnrealEd::CreateViewportToolbarDefaultContext(SharedThis(this));
 
+			ContextObject->bShowGridToggleInPreviewSceneSettings = true;
 			ViewportToolbarContext.AddObject(ContextObject);
 		}
 
diff --git a/Engine/Source/Editor/Persona/Private/ViewportToolbar/AnimationEditorMenus.cpp b/Engine/Source/Editor/Persona/Private/ViewportToolbar/AnimationEditorMenus.cpp
--- a/Engine/Source/Editor/Persona/Private/ViewportToolbar/AnimationEditorMenus.cpp
+++ b/Engine/Source/Editor/Persona/Private/ViewportToolbar/AnimationEditorMenus.cpp
@@ -503,7 +503,6 @@ void UE::AnimationEditor::FillShowSubmenu(UToolMenu* InMenu, bool bInShowViewpor
 						// General
 						.IncludeFlag(FEngineShowFlags::SF_AntiAliasing)
 						.IncludeFlag(FEngineShowFlags::SF_Collision)
-						.IncludeFlag(FEngineShowFlags::SF_Grid)
 						.IncludeFlag(FEngineShowFlags::SF_Particles)
 						.IncludeFlag(FEngineShowFlags::SF_Translucency)
 						// Post Processing
@@ -1044,64 +1043,6 @@ TSharedRef<SWidget> UE::AnimationEditor::MakeFloorOffsetWidget(const TWeakPtr<SA
 	// clang-format on
 }
 
-FToolMenuEntry UE::AnimationEditor::CreateSettingsSubmenu()
-{
-	FToolMenuEntry Entry = FToolMenuEntry::InitSubMenu(
-		"Settings",
-		LOCTEXT("SettingsSubmenuLabel", "Settings"),
-		LOCTEXT("SettingsSubmenuTooltip", "Settings"),
-		FNewToolMenuDelegate::CreateLambda(
-			[](UToolMenu* Submenu) -> void
-			{
-				TSharedPtr<SAnimationEditorViewportTabBody> AnimEditorViewportTab =
-					Private::GetAnimationEditorViewportTab(Submenu);
-
-				if (!AnimEditorViewportTab)
-				{
-					return;
-				}
-
-				// Scene Setup
-				{
-					FToolMenuSection& SceneSetupSection =
-						Submenu->AddSection("SceneSetup", LOCTEXT("SceneSetupLabel", "Scene Setup"));
-					FToolMenuEntry FloorOffsetEntry = FToolMenuEntry::InitWidget(
-						"FloorOffset",
-						MakeFloorOffsetWidget(AnimEditorViewportTab),
-						LOCTEXT("FloorHeightOffset", "Floor Height Offset")
-					);
-					SceneSetupSection.AddEntry(FloorOffsetEntry);
-					SceneSetupSection.AddMenuEntry(FAnimViewportShowCommands::Get().AutoAlignFloorToMesh);
-					SceneSetupSection.AddEntry(UE::AnimationEditor::CreateTurnTableMenu());
-				}
-
-				// Preview Scene
-				{
-					FToolMenuSection& PreviewSceneSection =
-						Submenu->AddSection("PreviewScene", LOCTEXT("PreviewSceneLabel", "Preview Scene"));
-					FToolMenuEntry PreviewProfileSubmenu = UE::UnrealEd::CreateAssetViewerProfileSubmenu();
-
-					PreviewSceneSection.AddEntry(PreviewProfileSubmenu);
-				}
-
-				// Audio
-				{
-					FToolMenuSection& AudioSection = Submenu->AddSection("Audio", LOCTEXT("AudioLabel", "Audio"));
-					AudioSection.AddMenuEntry(FAnimViewportShowCommands::Get().MuteAudio);
-					AudioSection.AddMenuEntry(FAnimViewportShowCommands::Get().UseAudioAttenuation);
-				}
-			}
-		)
-	);
-
-	Entry.Icon = FSlateIcon(FAppStyle::GetAppStyleSetName(), "LevelEditor.GameSettings");
-	Entry.InsertPosition.Position = EToolMenuInsertType::Last;
-	Entry.ToolBarData.LabelOverride = FText();
-	Entry.ToolBarData.ResizeParams.ClippingPriority = 1000;
-
-	return Entry;
-}
-
 void UE::AnimationEditor::ExtendCameraMenu(FName InCameraOptionsMenuName)
 {
 	UToolMenu* const Menu = UToolMenus::Get()->ExtendMenu(InCameraOptionsMenuName);
@@ -1455,6 +1396,54 @@ TSharedRef<SWidget> UE::AnimationEditor::GeneratePhysicsMenuWidget(
 	return SNullWidget::NullWidget;
 }
 
+void UE::AnimationEditor::ExtendPreviewSceneSettingsSubmenu(FName InSubmenuName)
+{
+	UToolMenu* const Submenu = UToolMenus::Get()->ExtendMenu(InSubmenuName);
+	if (!Submenu)
+	{
+		return;
+	}
+
+	Submenu->AddDynamicSection(
+		"StaticMeshEditorPreviewSceneDynamicSection",
+		FNewToolMenuDelegate::CreateLambda(
+			[](UToolMenu* InDynamicMenu)
+			{
+				FToolMenuInsert AudioInsertPosition("AssetViewerProfileSelectionSection", EToolMenuInsertType::Before);
+				FToolMenuSection& PreviewAudioSection = InDynamicMenu->FindOrAddSection(
+					"AssetViewerPreviewAudioSection",
+					LOCTEXT("AssetViewerPreviewAudioSectionLabel", "Preview Mesh Options"),
+					AudioInsertPosition
+				);
+				PreviewAudioSection.AddMenuEntry(FAnimViewportShowCommands::Get().MuteAudio);
+				PreviewAudioSection.AddMenuEntry(FAnimViewportShowCommands::Get().UseAudioAttenuation);
+
+				TSharedPtr<SAnimationEditorViewportTabBody> AnimEditorViewportTab =
+					Private::GetAnimationEditorViewportTab(InDynamicMenu);
+
+				if (!AnimEditorViewportTab)
+				{
+					return;
+				}
+
+				// Scene Setup
+				{
+					FToolMenuSection& SceneSetupSection =
+						InDynamicMenu->FindOrAddSection("SceneSetup", LOCTEXT("SceneSetupLabel", "Scene Setup"));
+					FToolMenuEntry FloorOffsetEntry = FToolMenuEntry::InitWidget(
+						"FloorOffset",
+						MakeFloorOffsetWidget(AnimEditorViewportTab),
+						LOCTEXT("FloorHeightOffset", "Floor Height Offset")
+					);
+					SceneSetupSection.AddEntry(FloorOffsetEntry);
+					SceneSetupSection.AddMenuEntry(FAnimViewportShowCommands::Get().AutoAlignFloorToMesh);
+					SceneSetupSection.AddEntry(UE::AnimationEditor::CreateTurnTableMenu());
+				}
+			}
+		)
+	);
+}
+
 TSharedRef<SWidget> UE::AnimationEditor::CreateFollowModeMenuWidget(const TWeakPtr<SAnimationEditorViewportTabBody>& InAnimEditorViewportTab
 )
 {
diff --git a/Engine/Source/Editor/Persona/Private/ViewportToolbar/AnimationEditorMenus.h b/Engine/Source/Editor/Persona/Private/ViewportToolbar/AnimationEditorMenus.h
--- a/Engine/Source/Editor/Persona/Private/ViewportToolbar/AnimationEditorMenus.h
+++ b/Engine/Source/Editor/Persona/Private/ViewportToolbar/AnimationEditorMenus.h
@@ -41,7 +41,6 @@ FText GetLODMenuLabel(const TWeakPtr<SAnimationEditorViewportTabBody>& InAnimEdi
 TSharedRef<SWidget> GenerateLODMenuWidget(const TWeakPtr<SAnimationEditorViewportTabBody>& InAnimEditorViewportTab);
 
 TSharedRef<SWidget> MakeFloorOffsetWidget(const TWeakPtr<SAnimationEditorViewportTabBody>& InAnimEditorViewportTabWeak);
-FToolMenuEntry CreateSettingsSubmenu();
 
 void ExtendViewModesSubmenu(FName InViewModesSubmenuName);
 TSharedRef<FExtender> GetViewModesLegacyExtenders(const TWeakPtr<SAnimationEditorViewportTabBody>& InViewport);
@@ -49,4 +48,6 @@ TSharedRef<FExtender> GetViewModesLegacyExtenders(const TWeakPtr<SAnimationEdito
 void AddPhysicsMenu(FName InPhysicsSubmenuName, FToolMenuInsert InInsertPosition);
 void FillPhysicsSubmenu(UToolMenu* InMenu);
 TSharedRef<SWidget> GeneratePhysicsMenuWidget(const TWeakPtr<SAnimationEditorViewportTabBody>& InAnimEditorViewportTab, const TSharedPtr<FExtender>& MenuExtender);
+
+void ExtendPreviewSceneSettingsSubmenu(FName InSubmenuName);
 } // namespace UE::AnimationEditor
