Add cvar option (default to true) for whether animeditorviewportclient outputs engine debug messages.

diff --git a/Engine/Source/Editor/Persona/Private/AnimationEditorViewportClient.cpp b/Engine/Source/Editor/Persona/Private/AnimationEditorViewportClient.cpp
--- a/Engine/Source/Editor/Persona/Private/AnimationEditorViewportClient.cpp
+++ b/Engine/Source/Editor/Persona/Private/AnimationEditorViewportClient.cpp
@@ -56,6 +56,8 @@
 #include "Animation/AnimSequenceHelpers.h"
 #include "AnimationBlueprintLibrary.h"
 
+TAutoConsoleVariable<bool> CVarShowEngineDebugMessageOnAnimViewport(TEXT("Persona.AnimationEditorViewportClient.ShowEngineDebugMessageOnAnimViewport"), true, TEXT("When true show engine onscreen debug messages on animation editor viewport clients."));
+
 namespace {
 	static const float AnimationEditorViewport_RotateSpeed = 0.02f;
 	static const float AnimationEditorViewport_TranslateSpeed = 0.25f;
@@ -784,7 +786,7 @@ void FAnimationViewportClient::DrawCanvas( FViewport& InViewport, FSceneView& Vi
 	}
 
 #if !(UE_BUILD_TEST)
-	if (World && GEngine->bEnableOnScreenDebugMessagesDisplay && GEngine->bEnableOnScreenDebugMessages)
+	if (World && GEngine->bEnableOnScreenDebugMessagesDisplay && GEngine->bEnableOnScreenDebugMessages && CVarShowEngineDebugMessageOnAnimViewport->GetBool())
 	{
 		constexpr int32 MessageX = 20;
 		constexpr int32 MessageY = 65;
