Add "Audio" and "Timecode" to the "Show" menu of the new Animation viewport toolbar

Update UE::AnimationEditor::AddSceneElementsSection to not take any bool parameters for controling if the "Timecode" and "Audio" submenus are added, instead they're always added.

diff --git a/Engine/Source/Editor/Persona/Private/ViewportToolbar/AnimationEditorMenus.cpp b/Engine/Source/Editor/Persona/Private/ViewportToolbar/AnimationEditorMenus.cpp
--- a/Engine/Source/Editor/Persona/Private/ViewportToolbar/AnimationEditorMenus.cpp
+++ b/Engine/Source/Editor/Persona/Private/ViewportToolbar/AnimationEditorMenus.cpp
@@ -617,7 +617,7 @@ void UE::AnimationEditor::FillTurnTableSubmenu(UToolMenu* InMenu)
 	}
 }
 
-void UE::AnimationEditor::AddSceneElementsSection(UToolMenu* InMenu, bool bInShowAudioEntries, bool bInShowTimecodeEntries)
+void UE::AnimationEditor::AddSceneElementsSection(UToolMenu* InMenu)
 {
 	if (!InMenu)
 	{
@@ -845,43 +845,35 @@ void UE::AnimationEditor::AddSceneElementsSection(UToolMenu* InMenu, bool bInSho
 		)
 	);
 
-	if (bInShowAudioEntries)
-	{
-		Section.AddSubMenu(
-			TEXT("AudioSubMenu"),
-			LOCTEXT("CharacterMenu_AudioSubMenu", "Audio"),
-			LOCTEXT("CharacterMenu_AudioSubMenuToolTip", "Audio options"),
-			FNewToolMenuDelegate::CreateLambda(
-				[](UToolMenu* InSubMenu)
-				{
-					FToolMenuSection& Section =
-						InSubMenu->AddSection("AnimViewportAudio", LOCTEXT("CharacterMenu_Audio", "Audio"));
-					Section.AddMenuEntry(FAnimViewportShowCommands::Get().MuteAudio);
-					Section.AddMenuEntry(FAnimViewportShowCommands::Get().UseAudioAttenuation);
-				}
-			)
-		);
-	}
-
-	if (bInShowTimecodeEntries)
-	{
-		Section.AddDynamicEntry(
-			"Timecode",
-			FNewToolMenuSectionDelegate::CreateLambda(
-				[](FToolMenuSection& InSection)
-				{
-					InSection.AddSubMenu(
-						TEXT("TimecodeSubMenu"),
-						LOCTEXT("CharacterMenu_TimecodeSubMenu", "Timecode"),
-						LOCTEXT("CharacterMenu_TimecodeSubMenuToolTip", "Timecode options"),
-						FNewToolMenuChoice(
-							FNewToolMenuDelegate::CreateStatic(&UE::AnimationEditor::Private::FillCharacterTimecodeMenu)
-						)
-					);
-				}
-			)
-		);
-	}
+	Section.AddSubMenu(
+		TEXT("AudioSubMenu"),
+		LOCTEXT("CharacterMenu_AudioSubMenu", "Audio"),
+		LOCTEXT("CharacterMenu_AudioSubMenuToolTip", "Audio options"),
+		FNewToolMenuDelegate::CreateLambda(
+			[](UToolMenu* InSubMenu)
+			{
+				FToolMenuSection& Section =
+					InSubMenu->AddSection("AnimViewportAudio", LOCTEXT("CharacterMenu_Audio", "Audio"));
+				Section.AddMenuEntry(FAnimViewportShowCommands::Get().MuteAudio);
+				Section.AddMenuEntry(FAnimViewportShowCommands::Get().UseAudioAttenuation);
+			}
+		)
+	);
+	Section.AddDynamicEntry(
+		"Timecode",
+		FNewToolMenuSectionDelegate::CreateLambda(
+			[](FToolMenuSection& InSection)
+			{
+				InSection.AddSubMenu(
+					TEXT("TimecodeSubMenu"),
+					LOCTEXT("CharacterMenu_TimecodeSubMenu", "Timecode"),
+					LOCTEXT("CharacterMenu_TimecodeSubMenuToolTip", "Timecode options"),
+					FNewToolMenuChoice(FNewToolMenuDelegate::CreateStatic(&UE::AnimationEditor::Private::FillCharacterTimecodeMenu
+					))
+				);
+			}
+		)
+	);
 
 	Section.Blocks.Sort(
 		[](const FToolMenuEntry& A, const FToolMenuEntry& B)
diff --git a/Engine/Source/Editor/Persona/Private/ViewportToolbar/AnimationEditorMenus.h b/Engine/Source/Editor/Persona/Private/ViewportToolbar/AnimationEditorMenus.h
--- a/Engine/Source/Editor/Persona/Private/ViewportToolbar/AnimationEditorMenus.h
+++ b/Engine/Source/Editor/Persona/Private/ViewportToolbar/AnimationEditorMenus.h
@@ -34,7 +34,7 @@ FToolMenuEntry CreateTurnTableMenu();
 void FillTurnTableSubmenu(UToolMenu* InMenu);
 TSharedRef<SWidget> GenerateTurnTableMenu(const TWeakPtr<SAnimationEditorViewportTabBody>& InAnimEditorViewportTab);
 
-void AddSceneElementsSection(UToolMenu* InMenu, bool bInShowAudioEntries = false, bool bInShowTimecodeEntries = false);
+void AddSceneElementsSection(UToolMenu* InMenu);
 
 FToolMenuEntry CreateLODSubmenu();
 FText GetLODMenuLabel(const TWeakPtr<SAnimationEditorViewportTabBody>& InAnimEditorViewportTab);
