Show shift-click pinned IPinnedCommandList commands whether or not the new or old viewport toolbar is used

* Add the IPinnedCommandList to the SAnimationEditorViewport overlay instead of to the old viewport toolbar.
* Also move the display text from SAnimationEditorViewportTabBody::GetDisplayString to the viewport overlay, so it doesn't overlap with the pinned commands.

diff --git a/Engine/Source/Editor/Persona/Private/SAnimationEditorViewport.cpp b/Engine/Source/Editor/Persona/Private/SAnimationEditorViewport.cpp
--- a/Engine/Source/Editor/Persona/Private/SAnimationEditorViewport.cpp
+++ b/Engine/Source/Editor/Persona/Private/SAnimationEditorViewport.cpp
@@ -121,6 +121,55 @@ void SAnimationEditorViewport::PopulateViewportOverlays(TSharedRef<SOverlay> Ove
 		[
 			BuildFeatureLevelWidget()
 		];
+
+	if (TSharedPtr<SAnimationEditorViewportTabBody> ViewportTab = TabBodyPtr.Pin())
+	{
+		TSharedRef<SVerticalBox> VerticalBox = SNew(SVerticalBox);
+
+		const TAttribute<FMargin> Padding = TAttribute<FMargin>::CreateLambda(
+			[]() -> FMargin
+			{
+				// The old viewport toolbar draws inside the viewport, so we need to pad to not overlap it.
+				const float TopPadding = UE::UnrealEd::ShowOldViewportToolbars() ? 29.0f : 0;
+				return FMargin(0, TopPadding, 0, 0);
+			}
+		);
+
+		// clang-format off
+		Overlay->AddSlot()
+			.VAlign(VAlign_Top)
+			.HAlign(HAlign_Fill)
+			.Padding(Padding)
+		[
+			VerticalBox
+		];
+		// clang-format on
+
+		if (TSharedPtr<IPinnedCommandList> PinnedCommands = ViewportTab->GetPinnedCommands())
+		{
+			// clang-format off
+			VerticalBox->AddSlot()
+			.AutoHeight()
+			[
+				PinnedCommands.ToSharedRef()
+			];
+			// clang-format on
+		}
+
+		// clang-format off
+		VerticalBox->AddSlot()
+		.AutoHeight()
+		.Padding(FMargin(4.0f, 3.0f, 0.0f, 0.0f))
+		[
+			// Display text (e.g., item being previewed)
+			SNew(SRichTextBlock)
+			.Visibility(EVisibility::SelfHitTestInvisible)
+			.DecoratorStyleSet(&FAppStyle::Get())
+			.Text(ViewportTab.Get(), &SAnimationEditorViewportTabBody::GetDisplayString)
+			.TextStyle(&FAppStyle::Get().GetWidgetStyle<FTextBlockStyle>("AnimViewport.MessageText"))
+		];
+		// clang-format on
+	}
 }
 
 TSharedRef<FEditorViewportClient> SAnimationEditorViewport::MakeEditorViewportClient()
@@ -160,32 +209,6 @@ TSharedPtr<SWidget> SAnimationEditorViewport::MakeViewportToolbar()
 		];
 	// clang-format on
 
-	if (TSharedPtr<SAnimationEditorViewportTabBody> ViewportTab = TabBodyPtr.Pin())
-	{
-		if (TSharedPtr<IPinnedCommandList> PinnedCommands = ViewportTab->GetPinnedCommands())
-		{
-			ToolbarWidget->AddSlot()
-			.AutoHeight()
-			[
-				PinnedCommands.ToSharedRef()
-			];
-		}
-
-		// clang-format off
-		ToolbarWidget->AddSlot()
-		.AutoHeight()
-		.Padding(FMargin(4.0f, 3.0f, 0.0f, 0.0f))
-		[
-			// Display text (e.g., item being previewed)
-			SNew(SRichTextBlock)
-			.Visibility(EVisibility::SelfHitTestInvisible)
-			.DecoratorStyleSet(&FAppStyle::Get())
-			.Text(ViewportTab.Get(), &SAnimationEditorViewportTabBody::GetDisplayString)
-			.TextStyle(&FAppStyle::Get().GetWidgetStyle<FTextBlockStyle>("AnimViewport.MessageText"))
-		];
-		// clang-format on
-	}
-
 	return ToolbarWidget;
 }
 
