Modify SAnimCurvePicker to optionally support multiselect
* Deprecate direct usage of SAnimCurvePicker using FOnCurvePicked, this API can still be used via FPersonaModule::CreateCurvePicker

diff --git a/Engine/Plugins/Experimental/Animation/HierarchyTableAnimation/Source/Editor/Private/SkeletonHierarchyTableTypeHandler.cpp b/Engine/Plugins/Experimental/Animation/HierarchyTableAnimation/Source/Editor/Private/SkeletonHierarchyTableTypeHandler.cpp
--- a/Engine/Plugins/Experimental/Animation/HierarchyTableAnimation/Source/Editor/Private/SkeletonHierarchyTableTypeHandler.cpp
+++ b/Engine/Plugins/Experimental/Animation/HierarchyTableAnimation/Source/Editor/Private/SkeletonHierarchyTableTypeHandler.cpp
@@ -119,13 +119,16 @@ void UHierarchyTable_TableTypeHandler_Skeleton::ExtendContextMenu(FMenuBuilder&
 				InSubMenuBuilder.BeginSection(NAME_None, LOCTEXT("ExistingCurveSection", "Existing Curve"));
 				{
 					InSubMenuBuilder.AddWidget(
-						PersonaModule.CreateCurvePicker(TableMetadata.Skeleton,
-							FOnCurvePicked::CreateLambda([this, SelectedIndex, &HierarchyTableView](const FName& InName)
+						PersonaModule.CreateMultiCurvePicker(TableMetadata.Skeleton,
+							FOnCurvesPicked::CreateLambda([this, SelectedIndex, &HierarchyTableView](const TConstArrayView<FName> InCurves)
 								{
 									FSlateApplication::Get().DismissAllMenus();
 
 									const int32 ParentIndex = SelectedIndex == INDEX_NONE ? 0 : SelectedIndex;
-									AddCurve(ParentIndex, InName);
+									for (FName Curve : InCurves)
+									{
+										AddCurve(ParentIndex, Curve);
+									}
 									HierarchyTableView.RegenerateHierarchy();
 								}),
 							FIsCurveNameMarkedForExclusion::CreateLambda([this](const FName& InName)
diff --git a/Engine/Source/Editor/Persona/Private/AnimTimeline/AnimTimelineTrack_Curves.cpp b/Engine/Source/Editor/Persona/Private/AnimTimeline/AnimTimelineTrack_Curves.cpp
--- a/Engine/Source/Editor/Persona/Private/AnimTimeline/AnimTimelineTrack_Curves.cpp
+++ b/Engine/Source/Editor/Persona/Private/AnimTimeline/AnimTimelineTrack_Curves.cpp
@@ -168,7 +168,11 @@ void FAnimTimelineTrack_Curves::FillMetadataEntryMenu(FMenuBuilder& Builder)
 	// Add existing curve to timeline using curve picker
 	{
 		const TSharedRef<SWidget> CurvePickerWidget = SNew(SAnimCurvePicker, &GetModel()->GetEditableSkeleton()->GetSkeleton())
-		.OnCurvePicked(this, &FAnimTimelineTrack_Curves::OnMetadataCurveNamePicked)
+		.OnCurvesPicked_Lambda([this](const TConstArrayView<FName> InCurves)
+		{
+			check(InCurves.Num() == 1);
+			return OnMetadataCurveNamePicked(InCurves[0]);
+		})
 		.IsCurveNameMarkedForExclusion(this, &FAnimTimelineTrack_Curves::IsCurveMarkedForExclusion);
 		Builder.AddWidget(CurvePickerWidget, FText::GetEmpty(), true);
 	}
@@ -192,7 +196,11 @@ void FAnimTimelineTrack_Curves::FillVariableCurveMenu(FMenuBuilder& Builder)
 	// Add existing curve to timeline using curve picker
 	{
 		const TSharedRef<SWidget> CurvePickerWidget = SNew(SAnimCurvePicker, &GetModel()->GetEditableSkeleton()->GetSkeleton())
-		.OnCurvePicked(this, &FAnimTimelineTrack_Curves::OnVariableCurveNamePicked)
+		.OnCurvesPicked_Lambda([this](const TConstArrayView<FName> InCurves)
+		{
+			check(InCurves.Num() == 1);
+			OnVariableCurveNamePicked(InCurves[0]);
+		})
 		.IsCurveNameMarkedForExclusion(this, &FAnimTimelineTrack_Curves::IsCurveMarkedForExclusion);
 		
 		Builder.AddWidget(CurvePickerWidget, FText::GetEmpty(), true);
diff --git a/Engine/Source/Editor/Persona/Private/PersonaModule.cpp b/Engine/Source/Editor/Persona/Private/PersonaModule.cpp
--- a/Engine/Source/Editor/Persona/Private/PersonaModule.cpp
+++ b/Engine/Source/Editor/Persona/Private/PersonaModule.cpp
@@ -1121,14 +1121,30 @@ void FPersonaModule::AddCommonToolbarExtensions(FToolBarBuilder& InToolbarBuilde
 TSharedRef<SWidget> FPersonaModule::CreateCurvePicker(TSharedRef<IEditableSkeleton> InEditableSkeleton, FOnCurvePicked InOnCurvePicked, FIsCurveNameMarkedForExclusion InIsCurveNameMarkedForExclusion)
 {
 	return SNew(SAnimCurvePicker, &InEditableSkeleton->GetSkeleton())
-		.OnCurvePicked(InOnCurvePicked)
+		.OnCurvesPicked_Lambda([InOnCurvePicked](const TConstArrayView<FName> InCurves)
+		{
+			check(InCurves.Num() == 1);
+			InOnCurvePicked.ExecuteIfBound(InCurves[0]);
+		})
+		.IsCurveNameMarkedForExclusion(InIsCurveNameMarkedForExclusion);
+}
+
+TSharedRef<SWidget> FPersonaModule::CreateMultiCurvePicker(const USkeleton* InSkeleton, FOnCurvesPicked InOnCurvesPicked, FIsCurveNameMarkedForExclusion InIsCurveNameMarkedForExclusion)
+{
+	return SNew(SAnimCurvePicker, InSkeleton)
+		.bEnableMultiselect(true)
+		.OnCurvesPicked(InOnCurvesPicked)
 		.IsCurveNameMarkedForExclusion(InIsCurveNameMarkedForExclusion);
 }
 
 TSharedRef<SWidget> FPersonaModule::CreateCurvePicker(const USkeleton* InSkeleton, FOnCurvePicked InOnCurvePicked, FIsCurveNameMarkedForExclusion InIsCurveNameMarkedForExclusion)
 {
 	return SNew(SAnimCurvePicker, InSkeleton)
-		.OnCurvePicked(InOnCurvePicked)
+		.OnCurvesPicked_Lambda([InOnCurvePicked](const TConstArrayView<FName> InCurves)
+		{
+			check(InCurves.Num() == 1);
+			InOnCurvePicked.ExecuteIfBound(InCurves[0]);
+		})
 		.IsCurveNameMarkedForExclusion(InIsCurveNameMarkedForExclusion);
 }
 
diff --git a/Engine/Source/Editor/Persona/Private/SAnimCurvePicker.cpp b/Engine/Source/Editor/Persona/Private/SAnimCurvePicker.cpp
--- a/Engine/Source/Editor/Persona/Private/SAnimCurvePicker.cpp
+++ b/Engine/Source/Editor/Persona/Private/SAnimCurvePicker.cpp
@@ -26,7 +26,7 @@ SAnimCurvePicker::~SAnimCurvePicker()
 
 void SAnimCurvePicker::Construct(const FArguments& InArgs, const USkeleton* InSkeleton)
 {
-	OnCurvePicked = InArgs._OnCurvePicked;
+	OnCurvesPicked = InArgs._OnCurvesPicked;
 	IsCurveNameMarkedForExclusion = InArgs._IsCurveNameMarkedForExclusion;
 	Skeleton = InSkeleton;
 	bShowOtherSkeletonCurves = false;
@@ -36,9 +36,8 @@ void SAnimCurvePicker::Construct(const FArguments& InArgs, const USkeleton* InSk
 	.OnTextChanged(this, &SAnimCurvePicker::HandleFilterTextChanged);
 
 	SAssignNew(NameListView, SListView<TSharedPtr<FName>>)
-	.SelectionMode(ESelectionMode::Single)
+	.SelectionMode(InArgs._bEnableMultiselect ? ESelectionMode::Multi : ESelectionMode::Single)
 	.ListItemsSource(&CurveNames)
-	.OnSelectionChanged(this, &SAnimCurvePicker::HandleSelectionChanged)
 	.OnGenerateRow(this, &SAnimCurvePicker::HandleGenerateRow);
 
 	const float HorizontalPadding = 8.0f;
@@ -116,22 +115,36 @@ void SAnimCurvePicker::Construct(const FArguments& InArgs, const USkeleton* InSk
 				]
 			];
 	}
-	
-	RefreshListItems();
-}
 
-void SAnimCurvePicker::HandleSelectionChanged(TSharedPtr<FName> InItem, ESelectInfo::Type InSelectionType)
-{
-	// When the user is navigating, do not act upon the selection change
-	if (InSelectionType == ESelectInfo::OnNavigation)
-	{
-		return;
-	}
-	
-	if (InItem.IsValid())
-	{
-		OnCurvePicked.ExecuteIfBound(*InItem);
-	}
+	VerticalBox->AddSlot()
+		.AutoHeight()
+		.Padding(HorizontalPadding, VerticalPadding)
+		[
+			SNew(SButton)
+				.Text(LOCTEXT("ConfirmButton", "Confirm"))
+				.IsEnabled_Lambda([this]()
+				{
+					return NameListView->GetNumItemsSelected() > 0;
+				})
+				.OnClicked_Lambda([this]()
+				{
+					TArray<TSharedPtr<FName>> SelectionPtrs;
+					NameListView->GetSelectedItems(SelectionPtrs);
+					check(SelectionPtrs.Num() > 0);
+
+					TArray<FName> Selection;
+					for (TSharedPtr<FName> Selected : SelectionPtrs)
+					{
+						Selection.Add(*Selected);
+					}
+
+					OnCurvesPicked.ExecuteIfBound(Selection);
+					
+					return FReply::Handled();
+				})
+		];
+
+	RefreshListItems();
 }
 
 TSharedRef<ITableRow> SAnimCurvePicker::HandleGenerateRow(TSharedPtr<FName> InItem, const TSharedRef<STableViewBase>& InOwnerTable)
diff --git a/Engine/Source/Editor/Persona/Private/SAnimCurvePicker.h b/Engine/Source/Editor/Persona/Private/SAnimCurvePicker.h
--- a/Engine/Source/Editor/Persona/Private/SAnimCurvePicker.h
+++ b/Engine/Source/Editor/Persona/Private/SAnimCurvePicker.h
@@ -19,11 +19,26 @@ public:
 	/** Virtual destructor. */
 	virtual ~SAnimCurvePicker() override;
 	
-	SLATE_BEGIN_ARGS(SAnimCurvePicker) {}
+	SLATE_BEGIN_ARGS(SAnimCurvePicker)
+		: _bEnableMultiselect(false)
+	{}
 
-	SLATE_EVENT(FOnCurvePicked, OnCurvePicked)
+	FOnCurvesPicked ConvertOnCurvePickedFn(const FOnCurvePicked& LegacyDelegate)
+	{
+		return FOnCurvesPicked::CreateLambda([LegacyDelegate](const TConstArrayView<FName> PickedCurves)
+		{
+			check(PickedCurves.Num() == 1)
+			LegacyDelegate.ExecuteIfBound(PickedCurves[0]);
+		});
+	}
+
+	SLATE_EVENT_DEPRECATED(5.6, "Use FOnCurvesPicked instead", FOnCurvePicked, OnCurvePicked, OnCurvesPicked, ConvertOnCurvePickedFn)
+
+	SLATE_EVENT(FOnCurvesPicked, OnCurvesPicked)
 
 	SLATE_EVENT(FIsCurveNameMarkedForExclusion, IsCurveNameMarkedForExclusion)
+
+	SLATE_ARGUMENT(bool, bEnableMultiselect)
 	
 	SLATE_END_ARGS()
 
@@ -42,13 +57,12 @@ private:
 	void FilterAvailableCurves();
 
 	/** UI handlers */
-	void HandleSelectionChanged(TSharedPtr<FName> InItem, ESelectInfo::Type InSelectionType);
 	TSharedRef<ITableRow> HandleGenerateRow(TSharedPtr<FName> InItem, const TSharedRef<STableViewBase>& InOwnerTable);
 	void HandleFilterTextChanged(const FText& InFilterText);
 
 private:
-	/** Delegate fired when a curve name is picked */
-	FOnCurvePicked OnCurvePicked;
+	/** Delegate fired when the curves selection is confirmed */
+	FOnCurvesPicked OnCurvesPicked;
 
 	/* Filter predicate to determine if curve should be excluded from the picker's curve list */
 	FIsCurveNameMarkedForExclusion IsCurveNameMarkedForExclusion;
diff --git a/Engine/Source/Editor/Persona/Private/SAnimationSequenceBrowser.cpp b/Engine/Source/Editor/Persona/Private/SAnimationSequenceBrowser.cpp
--- a/Engine/Source/Editor/Persona/Private/SAnimationSequenceBrowser.cpp
+++ b/Engine/Source/Editor/Persona/Private/SAnimationSequenceBrowser.cpp
@@ -404,10 +404,11 @@ public:
 			.WidthOverride(200.0f)
 			[
 				SNew(SAnimCurvePicker, &EditableSkeleton->GetSkeleton())
-				.OnCurvePicked_Lambda([this](const FName& InCurveName)
+				.OnCurvesPicked_Lambda([this](const TConstArrayView<FName> InCurveNames)
 				{
+					check(InCurveNames.Num() == 1);
 					FSlateApplication::Get().DismissAllMenus();
-					CurveString = InCurveName.ToString();
+					CurveString = InCurveNames[0].ToString();
 					OnChanged().Broadcast();
 				})
 			];
diff --git a/Engine/Source/Editor/Persona/Public/PersonaDelegates.h b/Engine/Source/Editor/Persona/Public/PersonaDelegates.h
--- a/Engine/Source/Editor/Persona/Public/PersonaDelegates.h
+++ b/Engine/Source/Editor/Persona/Public/PersonaDelegates.h
@@ -78,4 +78,7 @@ DECLARE_DELEGATE_OneParam(FOnCurvePicked, const FName& /*PickedName*/);
 DECLARE_DELEGATE_RetVal_OneParam(bool, FIsCurveNameMarkedForExclusion, const FName& /*CurveName*/);
 
 // Called by the notify picker to indicate that a notify was chosen by the user
-DECLARE_DELEGATE_OneParam(FOnNotifyPicked, const FName& /*PickedName*/);
\ No newline at end of file
+DECLARE_DELEGATE_OneParam(FOnNotifyPicked, const FName& /*PickedName*/);
+
+// Called by the curve picker to indicate the curves chosen by the user
+DECLARE_DELEGATE_OneParam(FOnCurvesPicked, const TConstArrayView<FName> /*Picked Curves*/);
\ No newline at end of file
diff --git a/Engine/Source/Editor/Persona/Public/PersonaModule.h b/Engine/Source/Editor/Persona/Public/PersonaModule.h
--- a/Engine/Source/Editor/Persona/Public/PersonaModule.h
+++ b/Engine/Source/Editor/Persona/Public/PersonaModule.h
@@ -454,6 +454,9 @@ public:
 	/** Register common tabs */
 	virtual FOnRegisterTabs& OnRegisterTabs() { return OnRegisterTabsDelegate; }
 
+	/** Create a widget that can choose multiple curve names. Derives available names from the asset registry list of assets that use the specified skeleton. */
+	virtual TSharedRef<SWidget> CreateMultiCurvePicker(const USkeleton* InSkeleton, FOnCurvesPicked InOnCurvesPicked, FIsCurveNameMarkedForExclusion InIsCurveNameMarkedForExclusion = FIsCurveNameMarkedForExclusion());
+
 	/** Create a widget that can choose a curve name. Derives available names from the asset registry list of assets that use the specified skeleton. */
 	virtual TSharedRef<SWidget> CreateCurvePicker(const USkeleton* InSkeleton, FOnCurvePicked InOnCurvePicked, FIsCurveNameMarkedForExclusion InIsCurveNameMarkedForExclusion = FIsCurveNameMarkedForExclusion());
 
