Clamp preview cursor of blend space to bounds of blend space

diff --git a/Engine/Source/Editor/Persona/Private/SAnimationBlendSpaceGridWidget.cpp b/Engine/Source/Editor/Persona/Private/SAnimationBlendSpaceGridWidget.cpp
--- a/Engine/Source/Editor/Persona/Private/SAnimationBlendSpaceGridWidget.cpp
+++ b/Engine/Source/Editor/Persona/Private/SAnimationBlendSpaceGridWidget.cpp
@@ -2603,7 +2603,7 @@ void SBlendSpaceGridWidget::Tick(const FGeometry& AllottedGeometry, const double
 			}
 			else
 			{
-				PreviewPosition = ScreenPositionToSampleValue(LastPreviewingMousePosition, false);
+				PreviewPosition = ScreenPositionToSampleValue(LastPreviewingMousePosition, true);
 			}
 			
 
