Avoid clipping LOD menu in Animation editor viewport toolbars early

diff --git a/Engine/Source/Editor/Persona/Private/ViewportToolbar/AnimationEditorMenus.cpp b/Engine/Source/Editor/Persona/Private/ViewportToolbar/AnimationEditorMenus.cpp
--- a/Engine/Source/Editor/Persona/Private/ViewportToolbar/AnimationEditorMenus.cpp
+++ b/Engine/Source/Editor/Persona/Private/ViewportToolbar/AnimationEditorMenus.cpp
@@ -926,10 +926,12 @@ FToolMenuEntry UE::AnimationEditor::CreateLODSubmenu()
 					}
 				);
 
-				InDynamicSection.AddSubMenu(
+				FToolMenuEntry& Entry = InDynamicSection.AddSubMenu(
 					"LOD",
 					Label,
-					LOCTEXT("LODMenuTooltip", "LOD Options. Control how LODs are displayed.\nShift-clicking items will 'pin' them to the toolbar."),
+					LOCTEXT(
+						"LODMenuTooltip", "LOD Options. Control how LODs are displayed.\nShift-clicking items will 'pin' them to the toolbar."
+					),
 					FNewToolMenuDelegate::CreateLambda(
 						[](UToolMenu* Submenu) -> void
 						{
@@ -937,11 +939,11 @@ FToolMenuEntry UE::AnimationEditor::CreateLODSubmenu()
 						}
 					)
 				);
+				Entry.InsertPosition.Position = EToolMenuInsertType::Last;
+				Entry.ToolBarData.ResizeParams.ClippingPriority = 800;
 			}
 		)
 	);
-	Entry.InsertPosition.Position = EToolMenuInsertType::Last;
-	Entry.ToolBarData.ResizeParams.ClippingPriority = 800;
 
 	return Entry;
 }
