Allow 'None' to be selected in SBoneSelectionWidget

diff --git a/Engine/Source/Editor/Persona/Private/AnimGraphNodeDetails.cpp b/Engine/Source/Editor/Persona/Private/AnimGraphNodeDetails.cpp
--- a/Engine/Source/Editor/Persona/Private/AnimGraphNodeDetails.cpp
+++ b/Engine/Source/Editor/Persona/Private/AnimGraphNodeDetails.cpp
@@ -530,6 +530,7 @@ void FBoneReferenceCustomization::CustomizeHeader( TSharedRef<IPropertyHandle> S
 		+ SWidgetSwitcher::Slot()
 		[
 			SNew(SBoneSelectionWidget)
+			.bShowNone(true)
 			.ToolTipText(StructPropertyHandle->GetToolTipText())
 			.OnBoneSelectionChanged(this, &FBoneReferenceCustomization::OnBoneSelectionChanged)
 			.OnGetSelectedBone(this, &FBoneReferenceCustomization::GetSelectedBone)
@@ -713,6 +714,7 @@ void FBoneSocketTargetCustomization::CustomizeHeader(TSharedRef<IPropertyHandle>
 			SNew(SBoneSelectionWidget)
 			.ToolTipText(StructPropertyHandle->GetToolTipText())
 			.bShowSocket(true)
+			.bShowNone(true)
 			.OnBoneSelectionChanged(this, &FBoneSocketTargetCustomization::OnBoneSelectionChanged)
 			.OnGetSelectedBone(this, &FBoneSocketTargetCustomization::GetSelectedBone)
 			.OnGetReferenceSkeleton(this, &FBoneReferenceCustomization::GetReferenceSkeleton)
diff --git a/Engine/Source/Editor/Persona/Private/Shared/BoneSelectionWidget.cpp b/Engine/Source/Editor/Persona/Private/Shared/BoneSelectionWidget.cpp
--- a/Engine/Source/Editor/Persona/Private/Shared/BoneSelectionWidget.cpp
+++ b/Engine/Source/Editor/Persona/Private/Shared/BoneSelectionWidget.cpp
@@ -20,6 +20,7 @@ void SBoneTreeMenu::Construct(const FArguments& InArgs)
 	OnGetSocketListDelegate = InArgs._OnGetSocketList;
 	bShowVirtualBones = InArgs._bShowVirtualBones;
 	bShowSocket = InArgs._bShowSocket;
+	bShowNone = InArgs._bShowNone;
 
 	FText TitleToUse = !InArgs._Title.IsEmpty() ? InArgs._Title  : LOCTEXT("BonePickerTitle", "Select...");
 
@@ -138,6 +139,13 @@ void SBoneTreeMenu::RebuildBoneList(const FName& SelectedBone)
 	SkeletonTreeInfo.Empty();
 	SkeletonTreeInfoFlat.Empty();
 
+	if (bShowNone)
+	{
+		TSharedRef<FBoneNameInfo> NoneBoneInfo = MakeShared<FBoneNameInfo>(NAME_None);
+		SkeletonTreeInfo.Add(NoneBoneInfo);
+		SkeletonTreeInfoFlat.Add(NoneBoneInfo);
+	}
+
 	if (ensure(OnGetReferenceSkeletonDelegate.IsBound()))
 	{
 		const FReferenceSkeleton& RefSkeleton = OnGetReferenceSkeletonDelegate.Execute();
@@ -254,6 +262,7 @@ void SBoneSelectionWidget::Construct(const FArguments& InArgs)
 	OnGetSocketList = InArgs._OnGetSocketList;
 	bShowSocket = InArgs._bShowSocket;
 	bShowVirtualBones = InArgs._bShowVirtualBones;
+	bShowNone = InArgs._bShowNone;
 
 	SuppliedToolTip = InArgs._ToolTipText.Get();
 
@@ -286,6 +295,7 @@ TSharedRef<SWidget> SBoneSelectionWidget::CreateSkeletonWidgetMenu()
 		.OnGetReferenceSkeleton(OnGetReferenceSkeleton)
 		.OnGetSocketList(OnGetSocketList)
 		.bShowSocket(bShowSocket)
+		.bShowNone(bShowNone)
 		.SelectedBone(CurrentBoneName)
 		.bShowVirtualBones(bShowVirtualBones);
 
diff --git a/Engine/Source/Editor/Persona/Public/BoneSelectionWidget.h b/Engine/Source/Editor/Persona/Public/BoneSelectionWidget.h
--- a/Engine/Source/Editor/Persona/Public/BoneSelectionWidget.h
+++ b/Engine/Source/Editor/Persona/Public/BoneSelectionWidget.h
@@ -34,6 +34,7 @@ public:
 	SLATE_BEGIN_ARGS(SBoneTreeMenu)
 		: _bShowVirtualBones(true)
 		, _bShowSocket(false)
+		, _bShowNone(false)
 		, _OnGetReferenceSkeleton()
 		, _OnBoneSelectionChanged()
 		, _OnGetSocketList()
@@ -42,6 +43,7 @@ public:
 		SLATE_ARGUMENT(FText, Title)
 		SLATE_ARGUMENT(bool, bShowVirtualBones)
 		SLATE_ARGUMENT(bool, bShowSocket)
+		SLATE_ARGUMENT(bool, bShowNone)
 		SLATE_ARGUMENT(FName, SelectedBone)
 		SLATE_EVENT(FGetReferenceSkeleton, OnGetReferenceSkeleton)
 		SLATE_EVENT(FOnBoneSelectionChanged, OnBoneSelectionChanged)
@@ -101,6 +103,7 @@ private:
 
 	bool bShowVirtualBones;
 	bool bShowSocket;
+	bool bShowNone;
 };
 
 class PERSONA_API SBoneSelectionWidget : public SCompoundWidget
@@ -110,6 +113,7 @@ public:
 	SLATE_BEGIN_ARGS( SBoneSelectionWidget )
 		: _bShowSocket(false)
 		, _bShowVirtualBones(true)
+		, _bShowNone(false)
 		, _OnBoneSelectionChanged()
 		, _OnGetSelectedBone()
 		, _OnGetReferenceSkeleton()
@@ -121,6 +125,9 @@ public:
 		/** Should show skeletons virtual bones in tree */
 		SLATE_ARGUMENT(bool, bShowVirtualBones)
 
+		/** Whether or not to show 'None' as a selectable bone name */
+		SLATE_ARGUMENT(bool, bShowNone)
+
 		/** set selected bone name */
 		SLATE_EVENT(FOnBoneSelectionChanged, OnBoneSelectionChanged);
 
@@ -164,6 +171,7 @@ private:
 	FGetSocketList			OnGetSocketList;
 	bool bShowSocket;
 	bool bShowVirtualBones;
+	bool bShowNone;
 
 	// Cache supplied tooltip
 	FText SuppliedToolTip;
