Remove audio-related entries from Animation editor preview scenes viewport toolbar menu

diff --git a/Engine/Source/Editor/Persona/Private/ViewportToolbar/AnimationEditorMenus.cpp b/Engine/Source/Editor/Persona/Private/ViewportToolbar/AnimationEditorMenus.cpp
--- a/Engine/Source/Editor/Persona/Private/ViewportToolbar/AnimationEditorMenus.cpp
+++ b/Engine/Source/Editor/Persona/Private/ViewportToolbar/AnimationEditorMenus.cpp
@@ -1,4 +1,4 @@
-﻿// Copyright Epic Games, Inc. All Rights Reserved.
+// Copyright Epic Games, Inc. All Rights Reserved.
 
 #include "AnimationEditorMenus.h"
 
@@ -1403,15 +1403,6 @@ void UE::AnimationEditor::ExtendPreviewSceneSettingsSubmenu(FName InSubmenuName)
 		FNewToolMenuDelegate::CreateLambda(
 			[](UToolMenu* InDynamicMenu)
 			{
-				FToolMenuInsert AudioInsertPosition("AssetViewerProfileSelectionSection", EToolMenuInsertType::Before);
-				FToolMenuSection& PreviewAudioSection = InDynamicMenu->FindOrAddSection(
-					"AssetViewerPreviewAudioSection",
-					LOCTEXT("AssetViewerPreviewAudioSectionLabel", "Preview Mesh Options"),
-					AudioInsertPosition
-				);
-				PreviewAudioSection.AddMenuEntry(FAnimViewportShowCommands::Get().MuteAudio);
-				PreviewAudioSection.AddMenuEntry(FAnimViewportShowCommands::Get().UseAudioAttenuation);
-
 				TSharedPtr<SAnimationEditorViewportTabBody> AnimEditorViewportTab =
 					Private::GetAnimationEditorViewportTab(InDynamicMenu);
 
