[Viewport Toolbar] Persona Editor - Add Playback Speed menu to Timeline controls

Also shares code with old toolbar.

diff --git a/Engine/Source/Editor/Persona/Private/AnimTimeline/SAnimTimeline.cpp b/Engine/Source/Editor/Persona/Private/AnimTimeline/SAnimTimeline.cpp
--- a/Engine/Source/Editor/Persona/Private/AnimTimeline/SAnimTimeline.cpp
+++ b/Engine/Source/Editor/Persona/Private/AnimTimeline/SAnimTimeline.cpp
@@ -134,8 +134,8 @@ void SAnimTimeline::Construct(const FArguments& InArgs, const TSharedRef<FAnimMo
 
 	TrackArea->SetOutliner(Outliner);
 
-	ColumnFillCoefficients[0] = 0.2f;
-	ColumnFillCoefficients[1] = 0.8f;
+	ColumnFillCoefficients[0] = 0.25f;
+	ColumnFillCoefficients[1] = 0.75f;
 
 	TAttribute<float> FillCoefficient_0, FillCoefficient_1;
 	{
@@ -262,7 +262,8 @@ void SAnimTimeline::Construct(const FArguments& InArgs, const TSharedRef<FAnimMo
 					// Transport controls
 					+SGridPanel::Slot(Column0, Row3, SGridPanel::Layer(10))
 					.VAlign(VAlign_Center)
-					.HAlign(HAlign_Center)
+					.HAlign(HAlign_Fill)
+					.Layer(-1)
 					[
 						SNew(SAnimTimelineTransportControls, InModel->GetPreviewScene(), InModel->GetAnimSequenceBase())
 					]
diff --git a/Engine/Source/Editor/Persona/Private/AnimTimeline/SAnimTimelineTransportControls.cpp b/Engine/Source/Editor/Persona/Private/AnimTimeline/SAnimTimelineTransportControls.cpp
--- a/Engine/Source/Editor/Persona/Private/AnimTimeline/SAnimTimelineTransportControls.cpp
+++ b/Engine/Source/Editor/Persona/Private/AnimTimeline/SAnimTimelineTransportControls.cpp
@@ -9,6 +9,12 @@
 #include "Animation/AnimSequence.h"
 #include "AnimPreviewInstance.h"
 #include "Modules/ModuleManager.h"
+#include "ToolMenu.h"
+#include "ToolMenus.h"
+#include "UICommandList_Pinnable.h"
+#include "ViewportToolbar/AnimViewportContext.h"
+#include "ViewportToolbar/AnimationEditorMenus.h"
+#include "Widgets/SBoxPanel.h"
 
 void SAnimTimelineTransportControls::Construct(const FArguments& InArgs, const TSharedRef<IPersonaPreviewScene>& InPreviewScene, UAnimSequenceBase* InAnimSequenceBase)
 {
@@ -32,10 +38,53 @@ void SAnimTimelineTransportControls::Construct(const FArguments& InArgs, const T
 	TransportControlArgs.OnGetPlaybackMode = FOnGetPlaybackMode::CreateSP(this, &SAnimTimelineTransportControls::GetPlaybackMode);
 	TransportControlArgs.OnGetRecording = FOnGetRecording::CreateSP(this, &SAnimTimelineTransportControls::IsRecording);
 
+	TSharedPtr<SHorizontalBox> HorizontalBox;
+	// clang-format off
 	ChildSlot
 	[
-		EditorWidgetsModule.CreateTransportControl(TransportControlArgs)
+		SAssignNew(HorizontalBox, SHorizontalBox)
+		+ SHorizontalBox::Slot()
+		.HAlign(HAlign_Center)
+		.FillWidth(1.0)
+		[
+			EditorWidgetsModule.CreateTransportControl(TransportControlArgs)
+		]
 	];
+	// clang-format on
+
+	// Add Playback Speed widget - we need to retrieve the viewport tab and the viewport it hosts
+	if (TSharedPtr<FAnimationEditorPreviewScene> PreviewScene =
+			StaticCastSharedPtr<FAnimationEditorPreviewScene>(WeakPreviewScene.Pin()))
+	{
+		const FName TimelinePlaybackMenu = "AnimationEditor.Timeline.PlaybackMenu";
+		if (!UToolMenus::Get()->IsMenuRegistered(TimelinePlaybackMenu))
+		{
+			if (UToolMenu* Menu = UToolMenus::Get()->RegisterMenu(
+					TimelinePlaybackMenu, NAME_None, EMultiBoxType::SlimHorizontalToolBar, false
+				))
+			{
+				FToolMenuSection& UnnamedSection = Menu->AddSection(NAME_None);
+				UnnamedSection.AddEntry(UE::AnimationEditor::CreatePlaybackSubmenu());
+			}
+		}
+
+		FToolMenuContext MenuContext;
+		{
+			MenuContext.AppendCommandList(PreviewScene->GetCommandList());
+			UAnimViewportContext* const ContextObject = NewObject<UAnimViewportContext>();
+			ContextObject->AnimationEditorPreviewScene = PreviewScene;
+			MenuContext.AddObject(ContextObject);
+		}
+
+		// clang-format off
+		HorizontalBox->AddSlot()
+		.HAlign(HAlign_Right)
+		.FillWidth(0.4)
+		[
+			UToolMenus::Get()->GenerateWidget(TimelinePlaybackMenu, MenuContext)
+		];
+		// clang-format on
+	}
 }
 
 UAnimSingleNodeInstance* SAnimTimelineTransportControls::GetPreviewInstance() const
diff --git a/Engine/Source/Editor/Persona/Private/AnimViewportPlaybackCommands.cpp b/Engine/Source/Editor/Persona/Private/AnimViewportPlaybackCommands.cpp
--- a/Engine/Source/Editor/Persona/Private/AnimViewportPlaybackCommands.cpp
+++ b/Engine/Source/Editor/Persona/Private/AnimViewportPlaybackCommands.cpp
@@ -20,30 +20,136 @@ FAnimViewportPlaybackCommands::FAnimViewportPlaybackCommands() : TCommands<FAnim
 
 void FAnimViewportPlaybackCommands::RegisterCommands()
 {
-	UI_COMMAND( PlaybackSpeedCommands[EAnimationPlaybackSpeeds::OneTenth],	"x0.1", "Set the animation playback speed to a tenth of normal", EUserInterfaceActionType::RadioButton, FInputChord() );
-	UI_COMMAND( PlaybackSpeedCommands[EAnimationPlaybackSpeeds::Quarter],		"x0.25", "Set the animation playback speed to a quarter of normal", EUserInterfaceActionType::RadioButton, FInputChord() );
-	UI_COMMAND( PlaybackSpeedCommands[EAnimationPlaybackSpeeds::Half],		"x0.5", "Set the animation playback speed to a half of normal", EUserInterfaceActionType::RadioButton, FInputChord() );
-	UI_COMMAND( PlaybackSpeedCommands[EAnimationPlaybackSpeeds::ThreeQuarters],		"x0.75", "Set the animation playback speed to three quarters of normal", EUserInterfaceActionType::RadioButton, FInputChord() );
-	UI_COMMAND( PlaybackSpeedCommands[EAnimationPlaybackSpeeds::Normal],		"x1.0", "Set the animation playback speed to normal", EUserInterfaceActionType::RadioButton, FInputChord() );
-	UI_COMMAND( PlaybackSpeedCommands[EAnimationPlaybackSpeeds::Double],		"x2.0", "Set the animation playback speed to double the speed of normal", EUserInterfaceActionType::RadioButton, FInputChord() );
-	UI_COMMAND( PlaybackSpeedCommands[EAnimationPlaybackSpeeds::FiveTimes],	"x5.0", "Set the animation playback speed to five times the normal speed", EUserInterfaceActionType::RadioButton, FInputChord() );
-	UI_COMMAND( PlaybackSpeedCommands[EAnimationPlaybackSpeeds::TenTimes],	"x10.0", "Set the animation playback speed to ten times the normal speed", EUserInterfaceActionType::RadioButton, FInputChord() );
-	UI_COMMAND( PlaybackSpeedCommands[EAnimationPlaybackSpeeds::Custom],	"xCustom", "Set the animation playback speed to assigned custom speed", EUserInterfaceActionType::RadioButton, FInputChord() );
-
-	UI_COMMAND(TurnTableSpeeds[EAnimationPlaybackSpeeds::OneTenth], "x0.1", "Set the turn table rotation speed to a tenth of normal", EUserInterfaceActionType::RadioButton, FInputChord());
-	UI_COMMAND(TurnTableSpeeds[EAnimationPlaybackSpeeds::Quarter], "x0.25", "Set the turn table rotation speed to a quarter of normal", EUserInterfaceActionType::RadioButton, FInputChord());
-	UI_COMMAND(TurnTableSpeeds[EAnimationPlaybackSpeeds::Half], "x0.5", "Set the turn table rotation speed to a half of normal", EUserInterfaceActionType::RadioButton, FInputChord());
-	UI_COMMAND(TurnTableSpeeds[EAnimationPlaybackSpeeds::ThreeQuarters], "x0.75", "Set the turn table rotation speed to three quarters of normal", EUserInterfaceActionType::RadioButton, FInputChord());
-	UI_COMMAND(TurnTableSpeeds[EAnimationPlaybackSpeeds::Normal], "x1.0", "Set the turn table rotation speed to normal", EUserInterfaceActionType::RadioButton, FInputChord());
-	UI_COMMAND(TurnTableSpeeds[EAnimationPlaybackSpeeds::Double], "x2.0", "Set the turn table rotation speed to double the speed of normal", EUserInterfaceActionType::RadioButton, FInputChord());
-	UI_COMMAND(TurnTableSpeeds[EAnimationPlaybackSpeeds::FiveTimes], "x5.0", "Set the turn table rotation speed to five times the normal speed", EUserInterfaceActionType::RadioButton, FInputChord());
-	UI_COMMAND(TurnTableSpeeds[EAnimationPlaybackSpeeds::TenTimes], "x10.0", "Set the turn table rotation speed to ten times the normal speed", EUserInterfaceActionType::RadioButton, FInputChord());
-	UI_COMMAND(TurnTableSpeeds[EAnimationPlaybackSpeeds::Custom], "xCustom", "Set the animation playback speed to assigned custom speed", EUserInterfaceActionType::RadioButton, FInputChord());
-
-	UI_COMMAND(PersonaTurnTablePlay, "Play", "Turn table rotates", EUserInterfaceActionType::RadioButton, FInputChord());
-	UI_COMMAND(PersonaTurnTablePause, "Pause", "Freeze with current rotation", EUserInterfaceActionType::RadioButton, FInputChord());
-	UI_COMMAND(PersonaTurnTableStop, "Stop", "Stop and Reset orientation", EUserInterfaceActionType::RadioButton, FInputChord());
+	const TSharedRef<FBindingContext> SharedThis = AsShared();
+	auto MakeLocalizedCommand =
+		[SharedThis](
+			TSharedPtr<FUICommandInfo>& OutCommandInfo, FName InCommandName, const FText& InCommandLabel, const FText& InCommandTooltip
+		)
+	{
+		FUICommandInfo::MakeCommandInfo(
+			SharedThis,
+			OutCommandInfo,
+			InCommandName,
+			InCommandLabel,
+			InCommandTooltip,
+			FSlateIcon(),
+			EUserInterfaceActionType::RadioButton,
+			FInputChord()
+		);
+	};
 
+	const FNumberFormattingOptions OneDecimalFormatting =
+		FNumberFormattingOptions().SetMinimumFractionalDigits(1).SetMaximumFractionalDigits(1);
+
+	// Playback Speed Commands
+	{
+		MakeLocalizedCommand(
+			PlaybackSpeedCommands[EAnimationPlaybackSpeeds::OneTenth],
+			"PlaybackSpeed_x0.1",
+			FText::Format(LOCTEXT("PlaybackSpeed_0.1_Label", "x{0}"), FText::AsNumber(0.1f)),
+			LOCTEXT("PlaybackSpeed_0.1_Tooltip", "Set the animation playback speed to a tenth of normal")
+		);
+		MakeLocalizedCommand(
+			PlaybackSpeedCommands[EAnimationPlaybackSpeeds::Quarter],
+			"PlaybackSpeed_x0.25",
+			FText::Format(LOCTEXT("PlaybackSpeed_0.25_Label", "x{0}"), FText::AsNumber(0.25f)),
+			LOCTEXT("PlaybackSpeed_0.25_Tooltip", "Set the animation playback speed to a tenth of normal")
+		);
+		MakeLocalizedCommand(
+			PlaybackSpeedCommands[EAnimationPlaybackSpeeds::Half],
+			"PlaybackSpeed_x0.5",
+			FText::Format(LOCTEXT("PlaybackSpeed_0.5_Label", "x{0}"), FText::AsNumber(0.5f)),
+			LOCTEXT("PlaybackSpeed_0.5_Tooltip", "Set the animation playback speed to a tenth of normal")
+		);
+		MakeLocalizedCommand(
+			PlaybackSpeedCommands[EAnimationPlaybackSpeeds::ThreeQuarters],
+			"PlaybackSpeed_x0.75",
+			FText::Format(LOCTEXT("PlaybackSpeed_0.75_Label", "x{0}"), FText::AsNumber(0.75f)),
+			LOCTEXT("PlaybackSpeed_0.75_Tooltip", "Set the animation playback speed to a tenth of normal")
+		);
+		MakeLocalizedCommand(
+			PlaybackSpeedCommands[EAnimationPlaybackSpeeds::Normal],
+			"PlaybackSpeed_x1.0",
+			FText::Format(LOCTEXT("PlaybackSpeed_1.0_Label", "x{0}"), FText::AsNumber(1.0f, &OneDecimalFormatting)),
+			LOCTEXT("PlaybackSpeed_1.0_Tooltip", "Set the animation playback speed to a tenth of normal")
+		);
+		MakeLocalizedCommand(
+			PlaybackSpeedCommands[EAnimationPlaybackSpeeds::Double],
+			"PlaybackSpeed_x2.0",
+			FText::Format(LOCTEXT("PlaybackSpeed_2.0_Label", "x{0}"), FText::AsNumber(2.0f, &OneDecimalFormatting)),
+			LOCTEXT("PlaybackSpeed_2.0_Tooltip", "Set the animation playback speed to a tenth of normal")
+		);
+		MakeLocalizedCommand(
+			PlaybackSpeedCommands[EAnimationPlaybackSpeeds::FiveTimes],
+			"PlaybackSpeed_x5.0",
+			FText::Format(LOCTEXT("PlaybackSpeed_5.0_Label", "x{0}"), FText::AsNumber(5.0f, &OneDecimalFormatting)),
+			LOCTEXT("PlaybackSpeed_5.0_Tooltip", "Set the animation playback speed to a tenth of normal")
+		);
+		MakeLocalizedCommand(
+			PlaybackSpeedCommands[EAnimationPlaybackSpeeds::TenTimes],
+			"PlaybackSpeed_x10.0",
+			FText::Format(LOCTEXT("PlaybackSpeed_10.0_Label", "x{0}"), FText::AsNumber(10.0f, &OneDecimalFormatting)),
+			LOCTEXT("PlaybackSpeed_10.0_Tooltip", "Set the animation playback speed to a tenth of normal")
+		);
+		UI_COMMAND( PlaybackSpeedCommands[EAnimationPlaybackSpeeds::Custom], "xCustom", "Set the animation playback speed to assigned custom speed", EUserInterfaceActionType::RadioButton, FInputChord() );
+	}
+
+	// Turntable Speed Commands
+	{
+		MakeLocalizedCommand(
+			TurnTableSpeeds[EAnimationPlaybackSpeeds::OneTenth],
+			"Turntable_speed_x0.1",
+			FText::Format(LOCTEXT("Turntable_speed_0.1_Label", "x{0}"), FText::AsNumber(0.1f)),
+			LOCTEXT("Turntable_speed_0.1_Tooltip", "Set the turn table rotation speed to a tenth of normal")
+		);
+		MakeLocalizedCommand(
+			TurnTableSpeeds[EAnimationPlaybackSpeeds::Quarter],
+			"Turntable_speed_x0.25",
+			FText::Format(LOCTEXT("Turntable_speed_0.25_Label", "x{0}"), FText::AsNumber(0.25f)),
+			LOCTEXT("Turntable_speed_0.25_Tooltip", "Set the turn table rotation speed to a tenth of normal")
+		);
+		MakeLocalizedCommand(
+			TurnTableSpeeds[EAnimationPlaybackSpeeds::Half],
+			"Turntable_speed_x0.5",
+			FText::Format(LOCTEXT("Turntable_speed_0.5_Label", "x{0}"), FText::AsNumber(0.5f)),
+			LOCTEXT("Turntable_speed_0.5_Tooltip", "Set the turn table rotation speed to a tenth of normal")
+		);
+		MakeLocalizedCommand(
+			TurnTableSpeeds[EAnimationPlaybackSpeeds::ThreeQuarters],
+			"Turntable_speed_x0.75",
+			FText::Format(LOCTEXT("Turntable_speed_0.75_Label", "x{0}"), FText::AsNumber(0.75f)),
+			LOCTEXT("Turntable_speed_0.75_Tooltip", "Set the turn table rotation speed to a tenth of normal")
+		);
+		MakeLocalizedCommand(
+			TurnTableSpeeds[EAnimationPlaybackSpeeds::Normal],
+			"Turntable_speed_x1.0",
+			FText::Format(LOCTEXT("Turntable_speed_1.0_Label", "x{0}"), FText::AsNumber(1.0f, &OneDecimalFormatting)),
+			LOCTEXT("Turntable_speed_1.0_Tooltip", "Set the turn table rotation speed to a tenth of normal")
+		);
+		MakeLocalizedCommand(
+			TurnTableSpeeds[EAnimationPlaybackSpeeds::Double],
+			"Turntable_speed_x2.0",
+			FText::Format(LOCTEXT("Turntable_speed_2.0_Label", "x{0}"), FText::AsNumber(2.0f, &OneDecimalFormatting)),
+			LOCTEXT("Turntable_speed_2.0_Tooltip", "Set the turn table rotation speed to a tenth of normal")
+		);
+		MakeLocalizedCommand(
+			TurnTableSpeeds[EAnimationPlaybackSpeeds::FiveTimes],
+			"Turntable_speed_x5.0",
+			FText::Format(LOCTEXT("Turntable_speed_5.0_Label", "x{0}"), FText::AsNumber(5.0f, &OneDecimalFormatting)),
+			LOCTEXT("Turntable_speed_5.0_Tooltip", "Set the turn table rotation speed to a tenth of normal")
+		);
+		MakeLocalizedCommand(
+			TurnTableSpeeds[EAnimationPlaybackSpeeds::TenTimes],
+			"Turntable_speed_x10.0",
+			FText::Format(LOCTEXT("Turntable_speed_10.0_Label", "x{0}"), FText::AsNumber(10.0f, &OneDecimalFormatting)),
+			LOCTEXT("Turntable_speed_10.0_Tooltip", "Set the turn table rotation speed to a tenth of normal")
+		);
+		UI_COMMAND(TurnTableSpeeds[EAnimationPlaybackSpeeds::Custom], "xCustom", "Set the animation playback speed to assigned custom speed", EUserInterfaceActionType::RadioButton, FInputChord());
+
+		UI_COMMAND(PersonaTurnTablePlay, "Play", "Turn table rotates", EUserInterfaceActionType::RadioButton, FInputChord());
+		UI_COMMAND(PersonaTurnTablePause, "Pause", "Freeze with current rotation", EUserInterfaceActionType::RadioButton, FInputChord());
+		UI_COMMAND(PersonaTurnTableStop, "Stop", "Stop and Reset orientation", EUserInterfaceActionType::RadioButton, FInputChord());
+	}
 }
 
 #undef LOCTEXT_NAMESPACE
diff --git a/Engine/Source/Editor/Persona/Private/AnimationEditorPreviewScene.cpp b/Engine/Source/Editor/Persona/Private/AnimationEditorPreviewScene.cpp
--- a/Engine/Source/Editor/Persona/Private/AnimationEditorPreviewScene.cpp
+++ b/Engine/Source/Editor/Persona/Private/AnimationEditorPreviewScene.cpp
@@ -25,10 +25,12 @@
 #include "Factories/PreviewMeshCollectionFactory.h"
 #include "AnimPreviewAttacheInstance.h"
 #include "AnimCustomInstanceHelper.h"
+#include "AnimViewportPlaybackCommands.h"
 #include "Animation/PreviewCollectionInterface.h"
 #include "ScopedTransaction.h"
 #include "Preferences/PersonaOptions.h"
 #include "ISkeletonTreeItem.h"
+#include "UICommandList_Pinnable.h"
 #include "Engine/SkeletalMeshSocket.h"
 
 #define LOCTEXT_NAMESPACE "AnimationEditorPreviewScene"
@@ -101,6 +103,8 @@ FAnimationEditorPreviewScene::FAnimationEditorPreviewScene(const ConstructionVal
 	// Disable killing actors outside of the world
 	AWorldSettings* WorldSettings = GetWorld()->GetWorldSettings(true);
 	WorldSettings->bEnableWorldBoundsChecks = false;
+
+	RegisterCommands();
 }
 
 FAnimationEditorPreviewScene::~FAnimationEditorPreviewScene()
@@ -1055,6 +1059,31 @@ TWeakObjectPtr<AWindDirectionalSource> FAnimationEditorPreviewScene::CreateWindA
 	return Wind;
 }
 
+void FAnimationEditorPreviewScene::OnSetPlaybackSpeed(int32 InPlaybackSpeedMode)
+{
+	SetAnimationPlaybackSpeedMode(static_cast<EAnimationPlaybackSpeeds::Type>(InPlaybackSpeedMode));
+}
+
+void FAnimationEditorPreviewScene::RegisterCommands()
+{
+	UICommandList = MakeShared<FUICommandList_Pinnable>();
+
+	FUICommandList_Pinnable& CommandList = *UICommandList;
+	FAnimViewportPlaybackCommands::Register();
+	const FAnimViewportPlaybackCommands& ViewportPlaybackCommands = FAnimViewportPlaybackCommands::Get();
+
+	// Create a menu item for each playback speed in EAnimationPlaybackSpeeds
+	for (int32 i = 0; i < static_cast<int>(EAnimationPlaybackSpeeds::NumPlaybackSpeeds); ++i)
+	{
+		CommandList.FUICommandList::MapAction(
+			ViewportPlaybackCommands.PlaybackSpeedCommands[i],
+			FExecuteAction::CreateRaw(this, &FAnimationEditorPreviewScene::OnSetPlaybackSpeed, i),
+			FCanExecuteAction(),
+			FIsActionChecked::CreateRaw(this, &FAnimationEditorPreviewScene::IsPlaybackSpeedSelected, i)
+		);
+	}
+}
+
 void FAnimationEditorPreviewScene::EnableWind(bool bEnableWind)
 {
 	UWorld* World = GetWorld();
@@ -1303,4 +1332,38 @@ TOptional<float> FAnimationEditorPreviewScene::GetCurrentTime() const
 	return {};
 }
 
+void FAnimationEditorPreviewScene::SetCustomAnimationSpeed(float InAnimationSpeed)
+{
+	CustomAnimationSpeed = InAnimationSpeed;
+	SetAnimationPlaybackSpeedMode(EAnimationPlaybackSpeeds::Custom);
+}
+
+float FAnimationEditorPreviewScene::GetCustomAnimationSpeed() const
+{
+	return CustomAnimationSpeed;
+}
+
+EAnimationPlaybackSpeeds::Type FAnimationEditorPreviewScene::GetAnimationPlaybackSpeedMode() const
+{
+	return AnimationPlaybackSpeedMode;
+}
+
+void FAnimationEditorPreviewScene::SetAnimationPlaybackSpeedMode(EAnimationPlaybackSpeeds::Type InMode)
+{
+	AnimationPlaybackSpeedMode = InMode;
+
+	if (GetWorld())
+	{
+		const float AnimationSpeed = (InMode == EAnimationPlaybackSpeeds::Custom)
+									   ? GetCustomAnimationSpeed()
+									   : EAnimationPlaybackSpeeds::Values[AnimationPlaybackSpeedMode];
+		GetWorld()->GetWorldSettings()->TimeDilation = AnimationSpeed;
+	}
+}
+
+bool FAnimationEditorPreviewScene::IsPlaybackSpeedSelected(int32 InPlaybackSpeedMode) const
+{
+	return InPlaybackSpeedMode == AnimationPlaybackSpeedMode;
+}
+
 #undef LOCTEXT_NAMESPACE
diff --git a/Engine/Source/Editor/Persona/Private/AnimationEditorPreviewScene.h b/Engine/Source/Editor/Persona/Private/AnimationEditorPreviewScene.h
--- a/Engine/Source/Editor/Persona/Private/AnimationEditorPreviewScene.h
+++ b/Engine/Source/Editor/Persona/Private/AnimationEditorPreviewScene.h
@@ -2,6 +2,7 @@
 
 #pragma once
 
+#include "AnimationEditorTypes.h"
 #include "CoreMinimal.h"
 #include "Textures/SlateIcon.h"
 #include "EditorUndoClient.h"
@@ -9,10 +10,13 @@
 #include "Animation/DebugSkelMeshComponent.h"
 #include "Engine/WindDirectionalSource.h"
 
+class FUICommandList_Pinnable;
 class IEditableSkeleton;
 class IPersonaToolkit;
 class UAnimSequence;
 
+
+
 class FAnimationEditorPreviewScene : public IPersonaPreviewScene, public FEditorUndoClient
 {
 public:
@@ -357,6 +361,24 @@ public:
 	/** Get the current time of the preview instance for the current animation sequence. */
 	TOptional<float> GetCurrentTime() const;
 
+	/** Set the current custom animation playback speed */
+	void SetCustomAnimationSpeed(float InAnimationSpeed);
+
+	/** Get the current custom animation playback speed */
+	float GetCustomAnimationSpeed() const;
+
+	/** Retrieve the current animation playback speed mode */
+	EAnimationPlaybackSpeeds::Type GetAnimationPlaybackSpeedMode() const;
+
+	/** Set the playback speed mode */
+	void SetAnimationPlaybackSpeedMode(EAnimationPlaybackSpeeds::Type InMode);
+
+	/** Function to return whether the supplied playback speed is the current active one */
+	bool IsPlaybackSpeedSelected(int32 InPlaybackSpeedMode) const;
+
+	/** Return the list of commands of the preview scene which are bound to delegates */
+	const TSharedPtr<FUICommandList_Pinnable>& GetCommandList() const { return UICommandList; }
+
 private:
 	/** Set preview mesh internal use only. The mesh should be verified by now. */
 	void SetPreviewMeshInternal(USkeletalMesh* NewPreviewMesh);
@@ -393,6 +415,11 @@ private:
 			OnRemoveAttachedComponentFilter.Unbind();
 		}
 	}
+
+	void OnSetPlaybackSpeed(int32 InPlaybackSpeedMode);
+
+	void RegisterCommands();
+
 private:
 	/** The one and only actor we have */
 	TObjectPtr<AActor> Actor;
@@ -509,4 +536,13 @@ private:
 
 	/** Should we display timecode in the preview viewport. */
 	bool bShowTimecode = true;
+
+	/** Custom Animation speed in the viewport. Transient setting. */
+	float CustomAnimationSpeed = 1.0f;
+
+	/** Selected playback speed mode, used for deciding scale */
+	EAnimationPlaybackSpeeds::Type AnimationPlaybackSpeedMode = EAnimationPlaybackSpeeds::Normal;
+
+	/** Commands that are bound to delegates*/
+	TSharedPtr<FUICommandList_Pinnable> UICommandList;
 };
diff --git a/Engine/Source/Editor/Persona/Private/AnimationEditorViewportClient.cpp b/Engine/Source/Editor/Persona/Private/AnimationEditorViewportClient.cpp
--- a/Engine/Source/Editor/Persona/Private/AnimationEditorViewportClient.cpp
+++ b/Engine/Source/Editor/Persona/Private/AnimationEditorViewportClient.cpp
@@ -180,7 +180,6 @@ FAnimationViewportClient::FAnimationViewportClient(const TSharedRef<IPersonaPrev
 	, PreviewScenePtr(InPreviewScene)
 	, AssetEditorToolkitPtr(InAssetEditorToolkit)
 	, bRotateCameraToFollowBone(false)
-	, AnimationPlaybackSpeedMode(EAnimationPlaybackSpeeds::Normal)
 	, bFocusOnDraw(false)
 	, bFocusUsingCustomCamera(false)
 	, CachedScreenSize(0.0f)
@@ -2651,31 +2650,38 @@ void FAnimationViewportClient::SetAdvancedShowFlagsForScene(const bool bAdvanced
 }
 void FAnimationViewportClient::SetPlaybackSpeedMode(EAnimationPlaybackSpeeds::Type InMode)
 {
-	AnimationPlaybackSpeedMode = InMode;
-
-	if (GetWorld())
+	if (FAnimationEditorPreviewScene* AnimationEditorPreviewScene = GetAnimPreviewScenePtr().Get())
 	{
-		const float AnimationSpeed = (InMode == EAnimationPlaybackSpeeds::Custom)
-			? GetCustomAnimationSpeed()
-			: EAnimationPlaybackSpeeds::Values[AnimationPlaybackSpeedMode];
-		GetWorld()->GetWorldSettings()->TimeDilation = AnimationSpeed;
+		AnimationEditorPreviewScene->SetAnimationPlaybackSpeedMode(InMode);
 	}
 }
 
 EAnimationPlaybackSpeeds::Type FAnimationViewportClient::GetPlaybackSpeedMode() const
 {
-	return AnimationPlaybackSpeedMode;
+	if (FAnimationEditorPreviewScene* AnimationEditorPreviewScene = GetAnimPreviewScenePtr().Get())
+	{
+		return AnimationEditorPreviewScene->GetAnimationPlaybackSpeedMode();
+	}
+
+	return EAnimationPlaybackSpeeds::Type::Normal;
 }
 
 void FAnimationViewportClient::SetCustomAnimationSpeed(const float InCustomAnimationSpeed)
 {
-	CustomAnimationSpeed = InCustomAnimationSpeed;
-	SetPlaybackSpeedMode(EAnimationPlaybackSpeeds::Custom);
+	if (FAnimationEditorPreviewScene* AnimationEditorPreviewScene = GetAnimPreviewScenePtr().Get())
+	{
+		AnimationEditorPreviewScene->SetCustomAnimationSpeed(InCustomAnimationSpeed);
+	}
 }
 
 float FAnimationViewportClient::GetCustomAnimationSpeed() const
 {
-	return CustomAnimationSpeed;
+	if (FAnimationEditorPreviewScene* AnimationEditorPreviewScene = GetAnimPreviewScenePtr().Get())
+	{
+		return AnimationEditorPreviewScene->GetCustomAnimationSpeed();
+	}
+
+	return 0.0f;
 }
 
 TSharedPtr<class FAnimationEditorPreviewScene> FAnimationViewportClient::GetAnimPreviewScenePtr() const
diff --git a/Engine/Source/Editor/Persona/Private/SAnimViewportToolBar.cpp b/Engine/Source/Editor/Persona/Private/SAnimViewportToolBar.cpp
--- a/Engine/Source/Editor/Persona/Private/SAnimViewportToolBar.cpp
+++ b/Engine/Source/Editor/Persona/Private/SAnimViewportToolBar.cpp
@@ -849,50 +849,12 @@ TSharedRef<SWidget> SAnimViewportToolBar::GenerateViewportTypeMenu() const
 
 TSharedRef<SWidget> SAnimViewportToolBar::GeneratePlaybackMenu() const
 {
-	const FAnimViewportPlaybackCommands& Actions = FAnimViewportPlaybackCommands::Get();
-
-	TSharedPtr<FExtender> MenuExtender = FExtender::Combine(Extenders);
-
-	const bool bInShouldCloseWindowAfterMenuSelection = true;
-	FMenuBuilder InMenuBuilder(bInShouldCloseWindowAfterMenuSelection, Viewport.Pin()->GetCommandList(), MenuExtender);
-
-	InMenuBuilder.PushCommandList(Viewport.Pin()->GetCommandList().ToSharedRef());
-	InMenuBuilder.PushExtender(MenuExtender.ToSharedRef());
+	if (const TSharedPtr<SAnimationEditorViewportTabBody>& ViewportTabPinned = Viewport.Pin())
 	{
-		// View modes
-		{
-			InMenuBuilder.BeginSection("AnimViewportPlaybackSpeed", LOCTEXT("PlaybackMenu_SpeedLabel", "Playback Speed") );
-			{
-				for(int32 PlaybackSpeedIndex = 0; PlaybackSpeedIndex < EAnimationPlaybackSpeeds::NumPlaybackSpeeds; ++PlaybackSpeedIndex)
-				{
-					InMenuBuilder.AddMenuEntry( Actions.PlaybackSpeedCommands[PlaybackSpeedIndex] );
-				}
-
-				TSharedPtr<SWidget> AnimSpeedWidget =
-					SNew(UE::AnimationEditor::SCustomAnimationSpeedSetting)
-						.AnimEditorViewport(Viewport)
-						.CustomSpeed_Lambda(
-							[Viewport = Viewport]()
-							{
-								return Viewport.Pin()->GetCustomAnimationSpeed();
-							}
-						)
-						.OnCustomSpeedChanged_Lambda(
-							[Viewport = Viewport](float CustomSpeed)
-							{
-								return Viewport.Pin()->SetCustomAnimationSpeed(CustomSpeed);
-							}
-						);
-
-				InMenuBuilder.AddWidget(AnimSpeedWidget.ToSharedRef(), LOCTEXT("PlaybackMenu_Speed_Custom", "Custom Speed:"));
-			}
-			InMenuBuilder.EndSection();
-		}
+		return UE::AnimationEditor::GeneratePlaybackMenu(ViewportTabPinned->GetPreviewScene(), Extenders);
 	}
-	InMenuBuilder.PopCommandList();
-	InMenuBuilder.PopExtender();
 
-	return InMenuBuilder.MakeWidget();
+	return SNullWidget::NullWidget;
 }
 
 void SAnimViewportToolBar::GenerateTurnTableMenu(FMenuBuilder& MenuBuilder) const
@@ -940,21 +902,9 @@ FSlateColor SAnimViewportToolBar::GetFontColor() const
 FText SAnimViewportToolBar::GetPlaybackMenuLabel() const
 {
 	FText Label = LOCTEXT("PlaybackError", "Error");
-	if (Viewport.IsValid())
+	if (const TSharedPtr<SAnimationEditorViewportTabBody>& ViewportTabPinned = Viewport.Pin())
 	{
-		for(int i = 0; i < EAnimationPlaybackSpeeds::NumPlaybackSpeeds; ++i)
-		{
-			if (Viewport.Pin()->IsPlaybackSpeedSelected(i))
-			{
-				const int32 NumFractionalDigits = (i == EAnimationPlaybackSpeeds::Quarter || i == EAnimationPlaybackSpeeds::ThreeQuarters) ? 2 : 1;
-
-				const FNumberFormattingOptions FormatOptions = FNumberFormattingOptions()
-					.SetMinimumFractionalDigits(NumFractionalDigits)
-					.SetMaximumFractionalDigits(NumFractionalDigits);
-
-				Label = FText::Format(LOCTEXT("AnimViewportPlaybackMenuLabel", "x{0}"), FText::AsNumber(EAnimationPlaybackSpeeds::Values[i], &FormatOptions));
-			}
-		}
+		return UE::AnimationEditor::GetPlaybackMenuLabel(ViewportTabPinned->GetPreviewScene());
 	}
 	return Label;
 }
diff --git a/Engine/Source/Editor/Persona/Private/SAnimationEditorViewport.cpp b/Engine/Source/Editor/Persona/Private/SAnimationEditorViewport.cpp
--- a/Engine/Source/Editor/Persona/Private/SAnimationEditorViewport.cpp
+++ b/Engine/Source/Editor/Persona/Private/SAnimationEditorViewport.cpp
@@ -1299,20 +1299,6 @@ void SAnimationEditorViewportTabBody::BindCommands()
 	//Bind LOD preview menu commands
 	const FAnimViewportPlaybackCommands& ViewportPlaybackCommands = FAnimViewportPlaybackCommands::Get();
 
-	CommandList.BeginGroup(TEXT("PlaybackSpeeds"));
-
-	//Create a menu item for each playback speed in EAnimationPlaybackSpeeds
-	for(int32 i = 0; i < int(EAnimationPlaybackSpeeds::NumPlaybackSpeeds); ++i)
-	{
-		CommandList.MapAction( 
-			ViewportPlaybackCommands.PlaybackSpeedCommands[i],
-			FExecuteAction::CreateSP(this, &SAnimationEditorViewportTabBody::OnSetPlaybackSpeed, i),
-			FCanExecuteAction(),
-			FIsActionChecked::CreateSP(this, &SAnimationEditorViewportTabBody::IsPlaybackSpeedSelected, i));
-	}
-
-	CommandList.EndGroup();
-
 	CommandList.MapAction(
 		ViewportShowMenuCommands.MuteAudio,
 		FExecuteAction::CreateSP(this, &SAnimationEditorViewportTabBody::OnToggleMuteAudio),
@@ -1431,6 +1417,11 @@ void SAnimationEditorViewportTabBody::BindCommands()
 	}
 
 	PinnedCommands->BindCommandList(UICommandList.ToSharedRef());
+
+	if (TSharedPtr<FAnimationEditorPreviewScene> AnimationEditorPreviewScene = PreviewScenePtr.Pin())
+	{
+		PinnedCommands->BindCommandList(AnimationEditorPreviewScene->GetCommandList().ToSharedRef());
+	}
 }
 
 void SAnimationEditorViewportTabBody::OnSetTurnTableSpeed(int32 SpeedIndex)
@@ -1747,24 +1738,6 @@ float SAnimationEditorViewportTabBody::GetBoneDrawSize() const
 	return 0.0f;
 }
 
-void SAnimationEditorViewportTabBody::SetCustomAnimationSpeed(float InCusteomAnimationSpeed)
-{
-	if (const TSharedPtr<FAnimationViewportClient>& AnimViewportClient = GetAnimationViewportClient())
-	{
-		AnimViewportClient->SetCustomAnimationSpeed(InCusteomAnimationSpeed);
-	}
-}
-
-float SAnimationEditorViewportTabBody::GetCustomAnimationSpeed() const
-{
-	if (const TSharedPtr<FAnimationViewportClient>& AnimViewportClient = GetAnimationViewportClient())
-	{
-		return AnimViewportClient->GetCustomAnimationSpeed();
-	}
-
-	return 0.0f;
-}
-
 void SAnimationEditorViewportTabBody::SetCustomTurnTableSpeed(float InCustomTurnTableSpeed)
 {
 	CustomTurnTableSpeed = InCustomTurnTableSpeed;
@@ -1864,25 +1837,6 @@ bool SAnimationEditorViewportTabBody::IsAutoAlignFloor() const
 	return false;
 }
 
-/** Function to set the current playback speed*/
-void SAnimationEditorViewportTabBody::OnSetPlaybackSpeed(int32 PlaybackSpeedMode)
-{
-	if (const TSharedPtr<FAnimationViewportClient>& AnimViewportClient = GetAnimationViewportClient())
-	{
-		AnimViewportClient->SetPlaybackSpeedMode((EAnimationPlaybackSpeeds::Type)PlaybackSpeedMode);
-	}
-}
-
-bool SAnimationEditorViewportTabBody::IsPlaybackSpeedSelected(int32 PlaybackSpeedMode)
-{
-	if (const TSharedPtr<FAnimationViewportClient>& AnimViewportClient = GetAnimationViewportClient())
-	{
-		return PlaybackSpeedMode == AnimViewportClient->GetPlaybackSpeedMode();
-	}
-
-	return false;
-}
-
 void SAnimationEditorViewportTabBody::ShowBound()
 {
 	if (const TSharedPtr<FAnimationViewportClient>& AnimViewportClient = GetAnimationViewportClient())
diff --git a/Engine/Source/Editor/Persona/Private/SAnimationEditorViewport.h b/Engine/Source/Editor/Persona/Private/SAnimationEditorViewport.h
--- a/Engine/Source/Editor/Persona/Private/SAnimationEditorViewport.h
+++ b/Engine/Source/Editor/Persona/Private/SAnimationEditorViewport.h
@@ -254,12 +254,6 @@ public:
 	bool IsTrackingAttachedMeshLOD() const;
 	int32 GetLODSelection() const;
 
-	/** Function to set the current playback speed*/
-	void OnSetPlaybackSpeed(int32 PlaybackSpeedMode);
-
-	/** Function to return whether the supplied playback speed is the current active one */
-	bool IsPlaybackSpeedSelected(int32 PlaybackSpeedMode);
-
 	/** Function to get anim viewport widget */
 	TSharedPtr<class SEditorViewport> GetViewportWidget() const { return ViewportWidget; }
 
@@ -299,9 +293,6 @@ public:
 	void SetBoneDrawSize(float BoneDrawSize);
 	float GetBoneDrawSize() const;
 
-	void SetCustomAnimationSpeed(float AnimationSpeed);
-	float GetCustomAnimationSpeed() const;
-
 	void SetCustomTurnTableSpeed(float InCustomTurnTableSpeed);
 	float GetCustomTurnTableSpeed() const;
 
diff --git a/Engine/Source/Editor/Persona/Private/TabSpawners.cpp b/Engine/Source/Editor/Persona/Private/TabSpawners.cpp
--- a/Engine/Source/Editor/Persona/Private/TabSpawners.cpp
+++ b/Engine/Source/Editor/Persona/Private/TabSpawners.cpp
@@ -118,19 +118,25 @@ void FPersonaModeSharedData::Save(const TSharedRef<FAnimationViewportClient>& In
 	bShowBound = InFromViewport->IsSetShowBoundsChecked();
 	LocalAxesMode = InFromViewport->GetLocalAxesMode();
 	ViewportType = InFromViewport->ViewportType;
-	PlaybackSpeedMode = InFromViewport->GetPlaybackSpeedMode();
+
+	const TSharedRef<FAnimationEditorPreviewScene> AnimationEditorPreviewScene =
+		StaticCastSharedRef<FAnimationEditorPreviewScene>(InFromViewport->GetPreviewScene());
+	PlaybackSpeedMode = AnimationEditorPreviewScene->GetAnimationPlaybackSpeedMode();
 }
 
 void FPersonaModeSharedData::Restore(const TSharedRef<FAnimationViewportClient>& InToViewport)
-{	
+{
 	InToViewport->SetViewportType((ELevelViewportType)ViewportType);
 	InToViewport->SetViewLocation(ViewLocation);
 	InToViewport->SetViewRotation(ViewRotation);
 	InToViewport->SetShowBounds(bShowBound);
 	InToViewport->SetLocalAxesMode((ELocalAxesMode::Type)LocalAxesMode);
 	InToViewport->SetOrthoZoom(OrthoZoom);
-	InToViewport->SetPlaybackSpeedMode(PlaybackSpeedMode);
-	
+
+	const TSharedRef<FAnimationEditorPreviewScene> AnimationEditorPreviewScene =
+		StaticCastSharedRef<FAnimationEditorPreviewScene>(InToViewport->GetPreviewScene());
+	AnimationEditorPreviewScene->SetAnimationPlaybackSpeedMode(PlaybackSpeedMode);
+
 	if (bCameraLock)
 	{
 		InToViewport->SetLookAtLocation(LookAtLocation);
diff --git a/Engine/Source/Editor/Persona/Private/ViewportToolbar/AnimViewportContext.h b/Engine/Source/Editor/Persona/Private/ViewportToolbar/AnimViewportContext.h
--- a/Engine/Source/Editor/Persona/Private/ViewportToolbar/AnimViewportContext.h
+++ b/Engine/Source/Editor/Persona/Private/ViewportToolbar/AnimViewportContext.h
@@ -12,4 +12,5 @@ class UAnimViewportContext : public UObject
 
 public:
 	TWeakPtr<SAnimationEditorViewportTabBody> ViewportTabBody;
+	TWeakPtr<FAnimationEditorPreviewScene> AnimationEditorPreviewScene;
 };
diff --git a/Engine/Source/Editor/Persona/Private/ViewportToolbar/AnimationEditorMenus.cpp b/Engine/Source/Editor/Persona/Private/ViewportToolbar/AnimationEditorMenus.cpp
--- a/Engine/Source/Editor/Persona/Private/ViewportToolbar/AnimationEditorMenus.cpp
+++ b/Engine/Source/Editor/Persona/Private/ViewportToolbar/AnimationEditorMenus.cpp
@@ -378,6 +378,59 @@ void FillCharacterTimecodeMenu(UToolMenu* InMenu)
 	FToolMenuSection& TimecodeSection = InMenu->AddSection("Timecode", LOCTEXT("Timecode_Label", "Timecode"));
 	TimecodeSection.AddMenuEntry(Actions.ShowTimecode);
 }
+
+void FillPlaybackMenu(UToolMenu* InMenu)
+{
+	if (!InMenu)
+	{
+		return;
+	}
+
+	FToolMenuSection& PlaybackSpeedSection =
+		InMenu->FindOrAddSection("AnimViewportPlaybackSpeed", LOCTEXT("PlaybackMenu_SpeedLabel", "Playback Speed"));
+
+	const FAnimViewportPlaybackCommands& Actions = FAnimViewportPlaybackCommands::Get();
+	for (int32 PlaybackSpeedIndex = 0; PlaybackSpeedIndex < EAnimationPlaybackSpeeds::NumPlaybackSpeeds; ++PlaybackSpeedIndex)
+	{
+		PlaybackSpeedSection.AddMenuEntry(Actions.PlaybackSpeedCommands[PlaybackSpeedIndex]);
+	}
+
+	if (UAnimViewportContext* const AnimViewportContext = InMenu->FindContext<UAnimViewportContext>())
+	{
+		if (TSharedPtr<FAnimationEditorPreviewScene> PreviewScenePinned =
+				AnimViewportContext->AnimationEditorPreviewScene.Pin())
+		{
+			TSharedPtr<SWidget> AnimSpeedWidget =
+				SNew(UE::AnimationEditor::SCustomAnimationSpeedSetting)
+					.CustomSpeed_Lambda(
+						[PreviewSceneWeak = PreviewScenePinned.ToWeakPtr()]()
+						{
+							if (TSharedPtr<FAnimationEditorPreviewScene> PreviewScenePinned = PreviewSceneWeak.Pin())
+							{
+								return PreviewScenePinned->GetCustomAnimationSpeed();
+							}
+
+							return 0.0f;
+						}
+					)
+					.OnCustomSpeedChanged_Lambda(
+						[PreviewSceneWeak = PreviewScenePinned.ToWeakPtr()](float InCustomSpeed)
+						{
+							if (TSharedPtr<FAnimationEditorPreviewScene> PreviewScenePinned = PreviewSceneWeak.Pin())
+							{
+								return PreviewScenePinned->SetCustomAnimationSpeed(InCustomSpeed);
+							}
+						}
+					);
+
+			PlaybackSpeedSection.AddEntry(
+				FToolMenuEntry::InitWidget(
+					"PlaybackSpeed", AnimSpeedWidget.ToSharedRef(), LOCTEXT("PlaybackMenu_Speed_Custom", "Custom Speed:")
+				)
+			);
+		}
+	}
+}
 } // namespace UE::AnimationEditor::Private
 
 TSharedRef<SWidget> UE::AnimationEditor::MakeFollowBoneWidget(
@@ -590,7 +643,6 @@ void UE::AnimationEditor::FillTurnTableSubmenu(UToolMenu* InMenu)
 		// clang-format off
 		TSharedPtr<SWidget> AnimSpeedWidget =
 			SNew(UE::AnimationEditor::SCustomAnimationSpeedSetting)
-			.AnimEditorViewport(AnimEditorViewportTab)
 			.CustomSpeed_Lambda(
 				[ViewportWeak = AnimEditorViewportTab.ToWeakPtr()]()
 				{
@@ -1414,6 +1466,111 @@ void UE::AnimationEditor::ExtendPreviewSceneSettingsSubmenu(FName InSubmenuName)
 	);
 }
 
+FText UE::AnimationEditor::GetPlaybackMenuLabel(const TWeakPtr<FAnimationEditorPreviewScene>& InAnimationEditorPreviewScene)
+{
+	FText Label = LOCTEXT("PlaybackError", "Error");
+	if (const TSharedPtr<FAnimationEditorPreviewScene>& AnimationEditorPreviewScene = InAnimationEditorPreviewScene.Pin())
+	{
+		for (int i = 0; i < EAnimationPlaybackSpeeds::NumPlaybackSpeeds; ++i)
+		{
+			if (AnimationEditorPreviewScene->IsPlaybackSpeedSelected(i))
+			{
+				const int32 NumFractionalDigits =
+					(i == EAnimationPlaybackSpeeds::Quarter || i == EAnimationPlaybackSpeeds::ThreeQuarters) ? 2 : 1;
+
+				const FNumberFormattingOptions FormatOptions = FNumberFormattingOptions()
+																   .SetMinimumFractionalDigits(NumFractionalDigits)
+																   .SetMaximumFractionalDigits(NumFractionalDigits);
+
+				float CurrentValue = i == EAnimationPlaybackSpeeds::Type::Custom
+									   ? AnimationEditorPreviewScene->GetCustomAnimationSpeed()
+									   : EAnimationPlaybackSpeeds::Values[i];
+				Label = FText::Format(
+					LOCTEXT("AnimViewportPlaybackMenuLabel", "x{0}"), FText::AsNumber(CurrentValue, &FormatOptions)
+				);
+			}
+		}
+	}
+	return Label;
+}
+
+FToolMenuEntry UE::AnimationEditor::CreatePlaybackSubmenu()
+{
+	FToolMenuEntry Entry = FToolMenuEntry::InitDynamicEntry(
+		"PlaybackMenu",
+		FNewToolMenuSectionDelegate::CreateLambda(
+			[](FToolMenuSection& InDynamicSection) -> void
+			{
+				if (UAnimViewportContext* const AnimViewportContext = InDynamicSection.FindContext<UAnimViewportContext>())
+				{
+					// Label updates based on currently selected LOD
+					const TAttribute<FText> Label = TAttribute<FText>::CreateLambda(
+						[PreviewSceneWeak = AnimViewportContext->AnimationEditorPreviewScene]()
+						{
+							return GetPlaybackMenuLabel(PreviewSceneWeak);
+						}
+					);
+
+					FToolMenuEntry& Entry = InDynamicSection.AddSubMenu(
+						"Playback",
+						Label,
+						LOCTEXT("PlaybackMenuTooltip", "Playback Speed Options. Control the time dilation of the scene's update."),
+						FNewToolMenuDelegate::CreateLambda(
+							[](UToolMenu* Submenu) -> void
+							{
+								UE::AnimationEditor::Private::FillPlaybackMenu(Submenu);
+							}
+						)
+					);
+
+					Entry.ToolBarData.ResizeParams.AllowClipping = false;
+				}
+			}
+		)
+	);
+
+	return Entry;
+}
+
+TSharedRef<SWidget> UE::AnimationEditor::GeneratePlaybackMenu(
+	const TWeakPtr<FAnimationEditorPreviewScene>& InAnimationEditorPreviewScene, const TArray<TSharedPtr<FExtender>>& InExtenders
+)
+{
+	// We generate a menu via UToolMenus, so we can use FillShowSubmenu call from both old and new toolbar
+	FName OldPlaybackMenuName = "AnimationEditor.OldViewportToolbar.PlaybackMenu";
+	if (!UToolMenus::Get()->IsMenuRegistered(OldPlaybackMenuName))
+	{
+		if (UToolMenu* Menu = UToolMenus::Get()->RegisterMenu(OldPlaybackMenuName, NAME_None, EMultiBoxType::Menu, false))
+		{
+			Menu->AddDynamicSection(
+				"BaseSection",
+				FNewToolMenuDelegate::CreateLambda(
+					[](UToolMenu* InMenu)
+					{
+						UE::AnimationEditor::Private::FillPlaybackMenu(InMenu);
+					}
+				)
+			);
+		}
+	}
+
+	FToolMenuContext MenuContext;
+	{
+		TSharedPtr<FExtender> MenuExtender = FExtender::Combine(InExtenders);
+		MenuContext.AddExtender(MenuExtender);
+		UAnimViewportContext* const ContextObject = NewObject<UAnimViewportContext>();
+		MenuContext.AddObject(ContextObject);
+
+		if (TSharedPtr<FAnimationEditorPreviewScene> AnimationPreviewScenePinned = InAnimationEditorPreviewScene.Pin())
+		{
+			MenuContext.AppendCommandList(AnimationPreviewScenePinned->GetCommandList());
+			ContextObject->AnimationEditorPreviewScene = AnimationPreviewScenePinned;
+		}
+	}
+
+	return UToolMenus::Get()->GenerateWidget(OldPlaybackMenuName, MenuContext);
+}
+
 TSharedRef<SWidget> UE::AnimationEditor::CreateFollowModeMenuWidget(const TWeakPtr<SAnimationEditorViewportTabBody>& InAnimEditorViewportTab
 )
 {
diff --git a/Engine/Source/Editor/Persona/Private/ViewportToolbar/AnimationEditorMenus.h b/Engine/Source/Editor/Persona/Private/ViewportToolbar/AnimationEditorMenus.h
--- a/Engine/Source/Editor/Persona/Private/ViewportToolbar/AnimationEditorMenus.h
+++ b/Engine/Source/Editor/Persona/Private/ViewportToolbar/AnimationEditorMenus.h
@@ -5,6 +5,7 @@
 #include "Widgets/SCompoundWidget.h"
 #include "ToolMenuMisc.h"
 
+class FAnimationEditorPreviewScene;
 class FExtender;
 class IPreviewProfileController;
 class SAnimationEditorViewportTabBody;
@@ -50,4 +51,8 @@ void FillPhysicsSubmenu(UToolMenu* InMenu);
 TSharedRef<SWidget> GeneratePhysicsMenuWidget(const TWeakPtr<SAnimationEditorViewportTabBody>& InAnimEditorViewportTab, const TSharedPtr<FExtender>& MenuExtender);
 
 void ExtendPreviewSceneSettingsSubmenu(FName InSubmenuName);
+
+FToolMenuEntry CreatePlaybackSubmenu();
+TSharedRef<SWidget> GeneratePlaybackMenu(const TWeakPtr<FAnimationEditorPreviewScene>& InAnimationEditorPreviewScene, const TArray<TSharedPtr<FExtender>>& InExtenders);
+FText GetPlaybackMenuLabel(const TWeakPtr<FAnimationEditorPreviewScene>& InAnimationEditorPreviewScene);
 } // namespace UE::AnimationEditor
diff --git a/Engine/Source/Editor/Persona/Private/ViewportToolbar/AnimationEditorWidgets.cpp b/Engine/Source/Editor/Persona/Private/ViewportToolbar/AnimationEditorWidgets.cpp
--- a/Engine/Source/Editor/Persona/Private/ViewportToolbar/AnimationEditorWidgets.cpp
+++ b/Engine/Source/Editor/Persona/Private/ViewportToolbar/AnimationEditorWidgets.cpp
@@ -10,7 +10,6 @@
 
 void UE::AnimationEditor::SCustomAnimationSpeedSetting::Construct(const FArguments& InArgs)
 {
-	AnimViewportPtr = InArgs._AnimEditorViewport;
 	CustomSpeed = InArgs._CustomSpeed;
 	OnCustomSpeedChanged = InArgs._OnCustomSpeedChanged;
 
diff --git a/Engine/Source/Editor/Persona/Private/ViewportToolbar/AnimationEditorWidgets.h b/Engine/Source/Editor/Persona/Private/ViewportToolbar/AnimationEditorWidgets.h
--- a/Engine/Source/Editor/Persona/Private/ViewportToolbar/AnimationEditorWidgets.h
+++ b/Engine/Source/Editor/Persona/Private/ViewportToolbar/AnimationEditorWidgets.h
@@ -18,7 +18,6 @@ public:
 	SLATE_BEGIN_ARGS(SCustomAnimationSpeedSetting)
 	{
 	}
-	SLATE_ARGUMENT(TWeakPtr<SAnimationEditorViewportTabBody>, AnimEditorViewport)
 	SLATE_ATTRIBUTE(float, CustomSpeed)
 	SLATE_EVENT(FOnCustomSpeedChanged, OnCustomSpeedChanged)
 SLATE_END_ARGS()
@@ -26,7 +25,6 @@ SLATE_END_ARGS()
 void Construct(const FArguments& InArgs);
 
 protected:
-	TWeakPtr<SAnimationEditorViewportTabBody> AnimViewportPtr;
 	TAttribute<float> CustomSpeed = 1.0f;
 	FOnCustomSpeedChanged OnCustomSpeedChanged;
 };
diff --git a/Engine/Source/Editor/Persona/Public/AnimationEditorTypes.h b/Engine/Source/Editor/Persona/Public/AnimationEditorTypes.h
new file mode 100644
--- /dev/null
+++ b/Engine/Source/Editor/Persona/Public/AnimationEditorTypes.h
@@ -0,0 +1,51 @@
+﻿// Copyright Epic Games, Inc. All Rights Reserved.
+
+#pragma once
+
+//////////////////////////////////////////////////////////////////////////
+// ELocalAxesMode
+
+namespace ELocalAxesMode
+{
+enum Type
+{
+	None,
+	Selected,
+	All,
+	NumAxesModes
+};
+}; // namespace ELocalAxesMode
+
+//////////////////////////////////////////////////////////////////////////
+// ELocalAxesMode
+
+namespace EDisplayInfoMode
+{
+enum Type
+{
+	None,
+	Basic,
+	Detailed,
+	SkeletalControls,
+	NumInfoModes
+};
+}; // namespace EDisplayInfoMode
+
+namespace EAnimationPlaybackSpeeds
+{
+enum Type
+{
+	OneTenth = 0,
+	Quarter,
+	Half,
+	ThreeQuarters,
+	Normal,
+	Double,
+	FiveTimes,
+	TenTimes,
+	Custom,
+	NumPlaybackSpeeds
+};
+
+extern float Values[NumPlaybackSpeeds];
+}; // namespace EAnimationPlaybackSpeeds
diff --git a/Engine/Source/Editor/Persona/Public/AnimationEditorViewportClient.h b/Engine/Source/Editor/Persona/Public/AnimationEditorViewportClient.h
--- a/Engine/Source/Editor/Persona/Public/AnimationEditorViewportClient.h
+++ b/Engine/Source/Editor/Persona/Public/AnimationEditorViewportClient.h
@@ -3,6 +3,7 @@
 #pragma once
 
 #include "CoreMinimal.h"
+#include "AnimationEditorTypes.h"
 #include "Delegates/IDelegateInstance.h"
 #include "Misc/Guid.h"
 #include "InputCoreTypes.h"
@@ -24,54 +25,6 @@ struct FSkelMeshRenderSection;
 DECLARE_DELEGATE_OneParam(FOnBoneSizeSet, float);
 DECLARE_DELEGATE_RetVal(float, FOnGetBoneSize)
 
-//////////////////////////////////////////////////////////////////////////
-// ELocalAxesMode
-
-namespace ELocalAxesMode
-{
-	enum Type
-	{
-		None,
-		Selected,
-		All,
-		NumAxesModes
-	};
-};
-
-//////////////////////////////////////////////////////////////////////////
-// ELocalAxesMode
-
-namespace EDisplayInfoMode
-{
-	enum Type
-	{
-		None,
-		Basic,	
-		Detailed,
-		SkeletalControls,
-		NumInfoModes
-	};
-};
-
-namespace EAnimationPlaybackSpeeds
-{
-	enum Type
-	{
-		OneTenth = 0,
-		Quarter,
-		Half,
-		ThreeQuarters,
-		Normal,
-		Double,
-		FiveTimes,
-		TenTimes,
-		Custom,
-		NumPlaybackSpeeds
-	};
-
-	extern float Values[NumPlaybackSpeeds];
-};
-
 /////////////////////////////////////////////////////////////////////////
 // FAnimationViewportClient
 
@@ -213,7 +166,9 @@ public:
 	float GetBoneDrawSize() const;
 
 	/** Access CustomAnimationSpeed config option*/
+	UE_DEPRECATED(5.6, "Animation Speed is now handled by FAnimationEditorPreviewScene. Please use FAnimationEditorPreviewScene::SetCustomAnimationSpeed instead.")
 	void SetCustomAnimationSpeed(const float InCustomAnimationSpeed);
+	UE_DEPRECATED(5.6, "Animation Speed is now handled by FAnimationEditorPreviewScene. Please use FAnimationEditorPreviewScene::GetCustomAnimationSpeed instead.")
 	float GetCustomAnimationSpeed() const;
 
 	/** Function to set Bone Draw  mode for the EBoneDrawType */
@@ -297,9 +252,11 @@ public:
 	int32 GetShowMeshStats() const;
 
 	/** Set the playback speed mode */
+	UE_DEPRECATED(5.6, "Animation Speed is now handled by FAnimationEditorPreviewScene. Please use FAnimationEditorPreviewScene::SetAnimationPlaybackSpeedMode instead.")
 	void SetPlaybackSpeedMode(EAnimationPlaybackSpeeds::Type InMode);
 
 	/** Get the playback speed mode */
+	UE_DEPRECATED(5.6, "Animation Speed is now handled by FAnimationEditorPreviewScene. Please use FAnimationEditorPreviewScene::GetAnimationPlaybackSpeedMode instead.")
 	EAnimationPlaybackSpeeds::Type GetPlaybackSpeedMode() const;
 
 	/** Get the preview scene we are viewing */
@@ -374,9 +331,6 @@ private:
 	/** User selected color using color picker */
 	FLinearColor SelectedHSVColor;
 
-	/** Selected playback speed mode, used for deciding scale */
-	EAnimationPlaybackSpeeds::Type AnimationPlaybackSpeedMode;
-
 	/** Flag for displaying the UV data in the viewport */
 	bool bDrawUVs;
 
@@ -513,9 +467,6 @@ private:
 
 	/** @return array of AssetUserData interfaces from editable objects on associated asset toolkit. */
 	TArray<IInterface_AssetUserData*> GetEditedObjectsWithAssetUserData() const;
-	
-	/** Custom Animation speed in the viewport. Transient setting. */
-	float CustomAnimationSpeed = 1.0f;
 
 	/** Size to draw bones in the viewport. Transient setting. */
 	float BoneDrawSize=1.0f;
