Fix double drawing of selected bone name when bone name drawing is enabled in the animation editors

diff --git a/Engine/Source/Editor/Persona/Private/AnimationEditorViewportClient.cpp b/Engine/Source/Editor/Persona/Private/AnimationEditorViewportClient.cpp
--- a/Engine/Source/Editor/Persona/Private/AnimationEditorViewportClient.cpp
+++ b/Engine/Source/Editor/Persona/Private/AnimationEditorViewportClient.cpp
@@ -963,6 +963,12 @@ void FAnimationViewportClient::ShowBoneNames( FCanvas* Canvas, FSceneView* View,
 			continue;
 		}
 
+		// Skip drawing bone name of selected bone, already drawn in SkeletonSelectionEditMode
+		if (PreviewMeshComponent == GetAnimPreviewScene()->GetPreviewMeshComponent() && GetAnimPreviewScene()->GetSelectedBoneIndex() == BoneIndex)
+		{
+			continue;
+		}
+
 		// If previewing a specific section, only show the bone names that belong to it
 		if ((PreviewMeshComponent->GetSectionPreview() >= 0) && !LODData.RenderSections[PreviewMeshComponent->GetSectionPreview()].BoneMap.Contains(BoneIndex))
 		{
