Animation Viewport Client: fix widget coord system computation

#rnx

diff --git a/Engine/Source/Editor/Persona/Private/AnimationEditorViewportClient.cpp b/Engine/Source/Editor/Persona/Private/AnimationEditorViewportClient.cpp
--- a/Engine/Source/Editor/Persona/Private/AnimationEditorViewportClient.cpp
+++ b/Engine/Source/Editor/Persona/Private/AnimationEditorViewportClient.cpp
@@ -1509,8 +1509,9 @@ FVector FAnimationViewportClient::GetWidgetLocation() const
 FMatrix FAnimationViewportClient::GetWidgetCoordSystem() const
 {
 	const ECoordSystem Space = GetWidgetCoordSystemSpace();
-	const bool bIsLocalOrParent = Space == COORD_Local || Space == COORD_Parent;
-	if( bIsLocalOrParent )
+	
+	const bool bComputeOrientation = Space == COORD_Local || Space == COORD_Parent || Space == COORD_Explicit;
+	if (bComputeOrientation)
 	{
 		return ModeTools->GetCustomInputCoordinateSystem();
 	}
