Improve SAnimCurvePicker workflow

diff --git a/Engine/Source/Editor/Persona/Private/SAnimCurvePicker.cpp b/Engine/Source/Editor/Persona/Private/SAnimCurvePicker.cpp
--- a/Engine/Source/Editor/Persona/Private/SAnimCurvePicker.cpp
+++ b/Engine/Source/Editor/Persona/Private/SAnimCurvePicker.cpp
@@ -37,6 +37,13 @@ void SAnimCurvePicker::Construct(const FArguments& InArgs, const USkeleton* InSk
 
 	SAssignNew(NameListView, SListView<TSharedPtr<FName>>)
 	.SelectionMode(InArgs._bEnableMultiselect ? ESelectionMode::Multi : ESelectionMode::Single)
+	.OnSelectionChanged_Lambda([this, InArgs](TSharedPtr<FName> InNameItem, ESelectInfo::Type SelectInfo)
+	{
+		if (InNameItem && !InArgs._bEnableMultiselect)
+		{
+			OnCurvesPicked.ExecuteIfBound(TArray<FName>{ *InNameItem });
+		}
+	})
 	.ListItemsSource(&CurveNames)
 	.OnGenerateRow(this, &SAnimCurvePicker::HandleGenerateRow);
 
@@ -116,33 +123,37 @@ void SAnimCurvePicker::Construct(const FArguments& InArgs, const USkeleton* InSk
 			];
 	}
 
-	VerticalBox->AddSlot()
-		.AutoHeight()
-		.Padding(HorizontalPadding, VerticalPadding)
-		[
-			SNew(SButton)
-				.Text(LOCTEXT("ConfirmButton", "Confirm"))
-				.IsEnabled_Lambda([this]()
-				{
-					return NameListView->GetNumItemsSelected() > 0;
-				})
-				.OnClicked_Lambda([this]()
-				{
-					TArray<TSharedPtr<FName>> SelectionPtrs;
-					NameListView->GetSelectedItems(SelectionPtrs);
-					check(SelectionPtrs.Num() > 0);
-
-					TArray<FName> Selection;
-					for (TSharedPtr<FName> Selected : SelectionPtrs)
+	if (InArgs._bEnableMultiselect)
+	{
+		VerticalBox->AddSlot()
+			.AutoHeight()
+			.Padding(HorizontalPadding, VerticalPadding)
+			[
+				SNew(SButton)
+					.ButtonStyle(FAppStyle::Get(), "PrimaryButton")
+					.Text(LOCTEXT("ConfirmButton", "Confirm"))
+					.IsEnabled_Lambda([this]()
 					{
-						Selection.Add(*Selected);
-					}
+						return NameListView->GetNumItemsSelected() > 0;
+					})
+					.OnClicked_Lambda([this]()
+					{
+						TArray<TSharedPtr<FName>> SelectionPtrs;
+						NameListView->GetSelectedItems(SelectionPtrs);
+						check(SelectionPtrs.Num() > 0);
 
-					OnCurvesPicked.ExecuteIfBound(Selection);
+						TArray<FName> Selection;
+						for (TSharedPtr<FName> Selected : SelectionPtrs)
+						{
+							Selection.Add(*Selected);
+						}
+
+						OnCurvesPicked.ExecuteIfBound(Selection);
 					
-					return FReply::Handled();
-				})
-		];
+						return FReply::Handled();
+					})
+			];
+	}
 
 	RefreshListItems();
 }
