Fixed comments from animation curves not being copied over to clipboard and/or pasted.

#rnx

diff --git a/Engine/Source/Editor/Persona/Private/AnimTimeline/AnimTimelineClipboard.cpp b/Engine/Source/Editor/Persona/Private/AnimTimeline/AnimTimelineClipboard.cpp
--- a/Engine/Source/Editor/Persona/Private/AnimTimeline/AnimTimelineClipboard.cpp
+++ b/Engine/Source/Editor/Persona/Private/AnimTimeline/AnimTimelineClipboard.cpp
@@ -233,6 +233,7 @@ void FAnimTimelineClipboardUtilities::OverwriteSelectedCurveDataFromClipboard(co
 				const UFloatCurveCopyObject* FloatCurve = Cast<UFloatCurveCopyObject>(ClipboardContent->Curves[0]);
 				Controller.SetCurveKeys(SelectedCurveTrackIdentifier, FloatCurve->Curve.FloatCurve.GetConstRefOfKeys());
 				Controller.SetCurveFlags(SelectedCurveTrackIdentifier, FloatCurve->Curve.GetCurveTypeFlags());
+				Controller.SetCurveComment(SelectedCurveTrackIdentifier, FloatCurve->Curve.Comment);
 			}
 			else if (SelectedCurveTrackType == ERawCurveTrackTypes::RCT_Transform)
 			{
@@ -244,6 +245,7 @@ void FAnimTimelineClipboardUtilities::OverwriteSelectedCurveDataFromClipboard(co
 				
 				Controller.SetTransformCurveKeys(SelectedCurveTrackIdentifier, TransformValues, TimeKeys);
 				Controller.SetCurveFlags(SelectedCurveTrackIdentifier, TransformCurve->Curve.GetCurveTypeFlags());
+				Controller.SetCurveComment(SelectedCurveTrackIdentifier, TransformCurve->Curve.Comment);
 			}
 		}
 	}
@@ -293,6 +295,7 @@ void FAnimTimelineClipboardUtilities::OverwriteOrAddCurvesFromClipboardContent(c
 				    
 				    Controller.SetCurveKeys(CurveIdentifier, FloatCopyObject->Curve.FloatCurve.GetConstRefOfKeys());
 				    Controller.SetCurveFlags(CurveIdentifier, FloatCopyObject->Curve.GetCurveTypeFlags());
+			    	Controller.SetCurveComment(CurveIdentifier, FloatCopyObject->Curve.Comment);
 			    }
 			    else if (InCurveCopyObj->CurveType == ERawCurveTrackTypes::RCT_Transform)
 			    {
@@ -304,6 +307,7 @@ void FAnimTimelineClipboardUtilities::OverwriteOrAddCurvesFromClipboardContent(c
 	    
 				    Controller.SetTransformCurveKeys(CurveIdentifier, TransformValues, TimeKeys);
 				    Controller.SetCurveFlags(CurveIdentifier, TransformCopyObject->Curve.GetCurveTypeFlags());
+			    	Controller.SetCurveComment(CurveIdentifier, TransformCopyObject->Curve.Comment);
 			    }
 		    }
 	    }
diff --git a/Engine/Source/Editor/Persona/Private/AnimTimeline/AnimTimelineTrack_FloatCurve.cpp b/Engine/Source/Editor/Persona/Private/AnimTimeline/AnimTimelineTrack_FloatCurve.cpp
--- a/Engine/Source/Editor/Persona/Private/AnimTimeline/AnimTimelineTrack_FloatCurve.cpp
+++ b/Engine/Source/Editor/Persona/Private/AnimTimeline/AnimTimelineTrack_FloatCurve.cpp
@@ -301,13 +301,14 @@ void FAnimTimelineTrack_FloatCurve::Copy(UAnimTimelineClipboardContent* InOutCli
 	CopyableCurve->Curve.SetName(FloatCurve->GetName());
 	CopyableCurve->Curve.SetCurveTypeFlags(FloatCurve->GetCurveTypeFlags());
 	CopyableCurve->Curve.CopyCurve(*FloatCurve);
-
+	CopyableCurve->Curve.Comment = FloatCurve->Comment;
+	
 	// Copy curve identifier data
 	CopyableCurve->CurveName = CurveName;
 	CopyableCurve->CurveType = ERawCurveTrackTypes::RCT_Float;
 	CopyableCurve->Channel = CurveId.Channel;
 	CopyableCurve->Axis = CurveId.Axis;
-
+	
 	// Origin data
 	CopyableCurve->OriginName = GetModel()->GetAnimSequenceBase()->GetFName();
 	
diff --git a/Engine/Source/Editor/Persona/Private/AnimTimeline/AnimTimelineTrack_TransformCurve.cpp b/Engine/Source/Editor/Persona/Private/AnimTimeline/AnimTimelineTrack_TransformCurve.cpp
--- a/Engine/Source/Editor/Persona/Private/AnimTimeline/AnimTimelineTrack_TransformCurve.cpp
+++ b/Engine/Source/Editor/Persona/Private/AnimTimeline/AnimTimelineTrack_TransformCurve.cpp
@@ -77,6 +77,7 @@ void FAnimTimelineTrack_TransformCurve::Copy(UAnimTimelineClipboardContent* InOu
 	CopyableCurve->Curve.SetName(TransformCurve->GetName());
 	CopyableCurve->Curve.SetCurveTypeFlags(TransformCurve->GetCurveTypeFlags());
 	CopyableCurve->Curve.CopyCurve(*TransformCurve);
+	CopyableCurve->Curve.Comment = TransformCurve->Comment;
 
 	// Copy curve identifier data
 	CopyableCurve->CurveName = CurveName;
