Animation Viewport: expose the near / far plane properties in the viewport toolbar

#rnx

diff --git a/Engine/Source/Editor/Persona/Private/SAnimationEditorViewport.cpp b/Engine/Source/Editor/Persona/Private/SAnimationEditorViewport.cpp
--- a/Engine/Source/Editor/Persona/Private/SAnimationEditorViewport.cpp
+++ b/Engine/Source/Editor/Persona/Private/SAnimationEditorViewport.cpp
@@ -244,8 +244,7 @@ TSharedPtr<SWidget> SAnimationEditorViewport::BuildViewportToolbar()
 				// Build the menu name our Camera menu will be using so we can extend it.
 				const FName SubmenuName = UToolMenus::JoinMenuPaths(ViewportToolbarName, "Camera");
 
-				constexpr bool bShowFarViewPlaneEntry = false;
-				UE::UnrealEd::ExtendCameraSubmenu(SubmenuName, bShowFarViewPlaneEntry);
+				UE::UnrealEd::ExtendCameraSubmenu(SubmenuName);
 				UE::AnimationEditor::ExtendCameraMenu(SubmenuName);
 
 				RightSection.AddEntry(UE::UnrealEd::CreateCameraSubmenu());
