Fix for crash when using a negative axis for angular velocity analysis in persona blend space editor
Deprecated templated parameter for AnalysisProperties from blendspace analysis functions, replace with strongly typed versions instead for compile time safety

#tests editor

diff --git a/Engine/Plugins/Animation/BlendSpaceMotionAnalysis/Source/BlendSpaceMotionAnalysis/Private/RootMotionAnalysis.cpp b/Engine/Plugins/Animation/BlendSpaceMotionAnalysis/Source/BlendSpaceMotionAnalysis/Private/RootMotionAnalysis.cpp
--- a/Engine/Plugins/Animation/BlendSpaceMotionAnalysis/Source/BlendSpaceMotionAnalysis/Private/RootMotionAnalysis.cpp
+++ b/Engine/Plugins/Animation/BlendSpaceMotionAnalysis/Source/BlendSpaceMotionAnalysis/Private/RootMotionAnalysis.cpp
@@ -68,7 +68,8 @@ static bool CalculateRootMotionVelocity(
 	const float                          RateScale)
 {
 	FVector Velocity;
-	if (BlendSpaceAnalysis::CalculateVelocity(Velocity, BlendSpace, AnalysisProperties, Animation, RateScale))
+	// Todo: we can remove the static_cast once the templated version of BlendSpaceAnalysis::CalculateVelocity is removed in 5.8+
+	if (BlendSpaceAnalysis::CalculateVelocity(Velocity, BlendSpace, static_cast<const ULinearAnalysisPropertiesBase*>(AnalysisProperties), Animation, RateScale))
 	{
 		FTransform FrameTM;
 		FVector FrameFacingDir = BlendSpaceAnalysis::GetAxisFromTM(FrameTM, AnalysisProperties->CharacterFacingAxis);
diff --git a/Engine/Plugins/Animation/BlendSpaceMotionAnalysis/Source/BlendSpaceMotionAnalysis/Public/RootMotionAnalysis.h b/Engine/Plugins/Animation/BlendSpaceMotionAnalysis/Source/BlendSpaceMotionAnalysis/Public/RootMotionAnalysis.h
--- a/Engine/Plugins/Animation/BlendSpaceMotionAnalysis/Source/BlendSpaceMotionAnalysis/Public/RootMotionAnalysis.h
+++ b/Engine/Plugins/Animation/BlendSpaceMotionAnalysis/Source/BlendSpaceMotionAnalysis/Public/RootMotionAnalysis.h
@@ -20,7 +20,7 @@ enum class EAnalysisRootMotionAxis : uint8
 
 //======================================================================================================================
 UCLASS()
-class URootMotionAnalysisProperties : public UAnalysisProperties
+class URootMotionAnalysisProperties : public ULinearAnalysisPropertiesBase
 {
 	GENERATED_BODY()
 public:
@@ -31,23 +31,6 @@ public:
 	UPROPERTY(EditAnywhere, DisplayName = "Axis", Category = AnalysisProperties)
 	EAnalysisRootMotionAxis FunctionAxis = EAnalysisRootMotionAxis::Speed;
 
-	/** The bone or socket used for analysis */
-	UPROPERTY(EditAnywhere, DisplayName = "Bone/Socket", Category = AnalysisProperties)
-	FBoneSocketTarget BoneSocket;
-
-	/**
-	* The space in which to perform the analysis. Fixed will use the analysis bone/socket at the first frame
-	* of the analysis time range. Changing will use the analysis bone/socket at the relevant frame during the
-	* analysis, but calculate velocities assuming that frame isn't moving. Moving will do the same but velocities
-	* as well as positions/rotations will be relative to this moving frame.
-	*/
-	UPROPERTY(EditAnywhere, Category = AnalysisProperties)
-	EAnalysisSpace Space = EAnalysisSpace::World;
-
-	/** Bone or socket that defines the analysis space (when it isn't World) */
-	UPROPERTY(EditAnywhere, DisplayName = "Analysis Space Bone/Socket", Category = AnalysisProperties, meta = (EditCondition = "Space != EAnalysisSpace::World"))
-	FBoneSocketTarget SpaceBoneSocket;
-
 	/** World or bone/socket axis that specifies the character's facing direction */
 	UPROPERTY(EditAnywhere, Category = AnalysisProperties)
 	EAnalysisLinearAxis CharacterFacingAxis = EAnalysisLinearAxis::PlusY;
@@ -55,14 +38,6 @@ public:
 	/** World or bone/socket axis that specifies the character's up direction */
 	UPROPERTY(EditAnywhere, Category = AnalysisProperties)
 	EAnalysisLinearAxis CharacterUpAxis = EAnalysisLinearAxis::PlusZ;
-
-	/** Fraction through each animation at which analysis starts */
-	UPROPERTY(EditAnywhere, DisplayName = "Start time fraction", Category = AnalysisProperties, meta = (ClampMin = "0", ClampMax = "1"))
-	float StartTimeFraction = 0.0f;
-
-	/** Fraction through each animation at which analysis ends */
-	UPROPERTY(EditAnywhere, DisplayName = "End time fraction", Category = AnalysisProperties, meta = (ClampMin = "0", ClampMax = "1"))
-	float EndTimeFraction = 1.0f;
 };
 
 /**
diff --git a/Engine/Source/Editor/Persona/Private/PersonaBlendSpaceAnalysis.cpp b/Engine/Source/Editor/Persona/Private/PersonaBlendSpaceAnalysis.cpp
--- a/Engine/Source/Editor/Persona/Private/PersonaBlendSpaceAnalysis.cpp
+++ b/Engine/Source/Editor/Persona/Private/PersonaBlendSpaceAnalysis.cpp
@@ -42,6 +42,435 @@ FVector BlendSpaceAnalysis::GetAxisFromTM(const FTransform& TM, EAnalysisLinearA
 	return FVector(0, 0, 0);
 }
 
+bool BlendSpaceAnalysis::CalculatePosition(FVector& Result, const UBlendSpace& BlendSpace, const ULinearAnalysisProperties* AnalysisProperties,
+	const UAnimSequence& Animation, const float RateScale)
+{
+	const int32 NumSampledKeys = Animation.GetNumberOfSampledKeys() - 1;
+	if (!AnalysisProperties || NumSampledKeys <= 0)
+	{
+		return false;
+	}
+
+	FTransform BoneOffset;
+	FName BoneName;
+	if (!GetBoneInfo(Animation, AnalysisProperties->BoneSocket, BoneOffset, BoneName))
+	{
+		return false;
+	}
+
+	int32 FirstKey = FMath::Clamp(
+		(int32) (float(NumSampledKeys) * AnalysisProperties->StartTimeFraction), 0, NumSampledKeys);
+	int32 LastKey = FMath::Clamp(
+		(int32) (float(NumSampledKeys) * AnalysisProperties->EndTimeFraction), FirstKey, NumSampledKeys);
+
+	FTransform FrameTM;
+	bool bNeedToUpdateFrameTM = true;
+	CalculateFrameTM(bNeedToUpdateFrameTM, FrameTM, FirstKey, AnalysisProperties, Animation);
+
+	Result.Set(0, 0, 0);
+	for (int32 Key = FirstKey; Key != LastKey + 1; ++Key)
+	{
+		CalculateFrameTM(bNeedToUpdateFrameTM, FrameTM, Key, AnalysisProperties, Animation);
+		FTransform BoneTM = GetBoneTransform(Animation, Key, BoneName);
+		FTransform TM = BoneOffset * BoneTM;
+		FVector RelativePos = FrameTM.InverseTransformPosition(TM.GetTranslation());
+		Result += RelativePos;
+	}
+	Result /= (1 + LastKey - FirstKey);
+	return true;
+}
+
+bool BlendSpaceAnalysis::CalculateDeltaPosition(FVector& Result, const UBlendSpace& BlendSpace, const ULinearAnalysisProperties* AnalysisProperties,
+	const UAnimSequence& Animation, const float RateScale)
+{
+	const int32 NumSampledKeys = Animation.GetNumberOfSampledKeys() - 1;
+	if (!AnalysisProperties || NumSampledKeys <= 0)
+	{
+		return false;
+	}
+
+	FTransform BoneOffset;
+	FName BoneName;
+	if (!GetBoneInfo(Animation, AnalysisProperties->BoneSocket, BoneOffset, BoneName))
+	{
+		return false;
+	}
+
+	int32 FirstKey = FMath::Clamp(
+		(int32) (float(NumSampledKeys) * AnalysisProperties->StartTimeFraction), 0, NumSampledKeys);
+	int32 LastKey = FMath::Clamp(
+		(int32) (float(NumSampledKeys) * AnalysisProperties->EndTimeFraction), FirstKey, NumSampledKeys);
+
+	FTransform FrameTM;
+	bool bNeedToUpdateFrameTM = true;
+	CalculateFrameTM(bNeedToUpdateFrameTM, FrameTM, FirstKey, AnalysisProperties, Animation);
+
+	FTransform BoneTM1 = GetBoneTransform(Animation, FirstKey, BoneName);
+	FTransform TM1 = BoneOffset * BoneTM1;
+	FVector RelativePos1 = FrameTM.InverseTransformPosition(TM1.GetTranslation());
+
+	CalculateFrameTM(bNeedToUpdateFrameTM, FrameTM, FirstKey, AnalysisProperties, Animation);
+	FTransform BoneTM2 = GetBoneTransform(Animation, LastKey, BoneName);
+	FTransform TM2 = BoneOffset * BoneTM2;
+	FVector RelativePos2 = FrameTM.InverseTransformPosition(TM2.GetTranslation());
+
+	Result = RelativePos2 - RelativePos1;
+	return true;
+}
+
+bool BlendSpaceAnalysis::CalculateVelocity(FVector& Result, const UBlendSpace& BlendSpace, const ULinearAnalysisPropertiesBase* AnalysisProperties,
+	const UAnimSequence& Animation, const float RateScale)
+{
+	const int32 NumSampledKeys = Animation.GetNumberOfSampledKeys() - 1;
+	if (!AnalysisProperties || NumSampledKeys <= 0)
+	{
+		return false;
+	}
+
+	if (NumSampledKeys == 1)
+	{
+		Result.Set(0, 0, 0);
+		return true;
+	}
+
+	FTransform BoneOffset;
+	FName BoneName;
+	if (!GetBoneInfo(Animation, AnalysisProperties->BoneSocket, BoneOffset, BoneName))
+	{
+		return false;
+	}
+
+	double DeltaTime = Animation.GetPlayLength() / double(NumSampledKeys);
+
+	int32 FirstKey = FMath::Clamp(
+		(int32) (float(NumSampledKeys) * AnalysisProperties->StartTimeFraction), 0, NumSampledKeys);
+	int32 LastKey = FMath::Clamp(
+		(int32) (float(NumSampledKeys) * AnalysisProperties->EndTimeFraction), FirstKey, NumSampledKeys-1);
+
+	// First and Last key are for averaging. However, the finite differencing always goes from one frame to the next
+	int32 NumKeys = FMath::Max(1 + LastKey - FirstKey, 1);
+
+	FTransform FrameTM;
+	bool bNeedToUpdateFrameTM = true;
+	CalculateFrameTM(bNeedToUpdateFrameTM, FrameTM, FirstKey, AnalysisProperties, Animation);
+
+	Result.Set(0, 0, 0);
+	for (int32 iKey = 0; iKey != NumKeys; ++iKey)
+	{
+		int32 Key = (FirstKey + iKey) % (NumSampledKeys + 1);
+		CalculateFrameTM(bNeedToUpdateFrameTM, FrameTM, Key, AnalysisProperties, Animation);
+
+		FTransform BoneTM1 = GetBoneTransform(Animation, Key, BoneName);
+		FTransform TM1 = BoneOffset * BoneTM1;
+		FVector RelativePos1 = FrameTM.InverseTransformPosition(TM1.GetTranslation());
+
+		int32 NextKey = (Key + 1) % (NumSampledKeys + 1);
+		if (AnalysisProperties->Space == EAnalysisSpace::Moving)
+		{
+			CalculateFrameTM(bNeedToUpdateFrameTM, FrameTM, NextKey, AnalysisProperties, Animation);
+		}
+
+		FTransform BoneTM2 = GetBoneTransform(Animation, NextKey, BoneName);
+		FTransform TM2 = BoneOffset * BoneTM2;
+		FVector RelativePos2 = FrameTM.InverseTransformPosition(TM2.GetTranslation());
+		FVector Velocity = (RelativePos2 - RelativePos1) / DeltaTime;
+
+#ifdef ANALYSIS_VERBOSE_LOG
+		UE_LOG(LogAnimation, Log, TEXT("%d Velocity = %f %f %f Height = %f"), 
+			   Key, Velocity.X, Velocity.Y, Velocity.Z, 0.5f * (RelativePos1 + RelativePos2).Z);
+#endif
+		Result += Velocity;
+	}
+	Result /= (1 + LastKey - FirstKey);
+	Result *= Animation.RateScale * RateScale;
+	UE_LOG(LogAnimation, Log, TEXT("%s vel = %f %f %f"), *BoneName.ToString(), Result.X, Result.Y, Result.Z);
+	return true;
+}
+
+void BlendSpaceAnalysis::CalculateBoneOrientation(FVector& RollPitchYaw, const UAnimSequence& Animation, const int32 Key, const FName BoneName,
+	const FTransform& BoneOffset, const UEulerAnalysisProperties* AnalysisProperties, const FVector& FrameFacingDir, const FVector& FrameRightDir,
+	const FVector& FrameUpDir)
+{
+	const FTransform BoneTM = GetBoneTransform(Animation, Key, BoneName);
+
+	const FTransform TM = BoneOffset * BoneTM;
+	const FVector AimForwardDir = BlendSpaceAnalysis::GetAxisFromTM(TM, AnalysisProperties->BoneFacingAxis);
+	const FVector AimRightDir = BlendSpaceAnalysis::GetAxisFromTM(TM, AnalysisProperties->BoneRightAxis);
+
+	double Yaw;
+	if (AnalysisProperties->EulerCalculationMethod == EEulerCalculationMethod::AimDirection)
+	{
+		// Yaw is taken from the AimRightDir to avoid problems when the gun is pointing up or down - especially if it 
+		// goes beyond 90 degrees in pitch. However, if there is roll around the gun axis, then this can produce
+		// incorrect/undesirable results.
+		Yaw = FMath::RadiansToDegrees(FMath::Atan2(
+			FVector::DotProduct(AimRightDir, -FrameFacingDir), FVector::DotProduct(AimRightDir, FrameRightDir)));
+	}
+	else
+	{
+		// This takes yaw directly from the forwards direction. Note that if the pose is really one with small yaw 
+		// and pitch more than 90 degrees, then this will calculate a yaw that is nearer to 180 degrees.
+		Yaw = FMath::RadiansToDegrees(FMath::Atan2(
+			FVector::DotProduct(AimForwardDir, FrameRightDir), FVector::DotProduct(AimForwardDir, FrameFacingDir)));
+	}
+
+	// Undo the yaw to get pitch
+	const FQuat YawQuat(FrameUpDir, FMath::DegreesToRadians(Yaw));
+	const FVector UnYawedAimForwardDir = YawQuat.UnrotateVector(AimForwardDir);
+	const double Up = UnYawedAimForwardDir | FrameUpDir;
+	const double Forward = UnYawedAimForwardDir | FrameFacingDir;
+	const double Pitch = FMath::RadiansToDegrees(FMath::Atan2(Up, Forward));
+
+	// Undo the pitch to get roll
+	const FVector UnYawedAimRightDir = YawQuat.UnrotateVector(AimRightDir);
+	const FQuat PitchQuat(FrameRightDir, -FMath::DegreesToRadians(Pitch));
+
+	const FVector UnYawedUnPitchedAimRightDir = PitchQuat.UnrotateVector(UnYawedAimRightDir);
+
+	const double Roll = FMath::RadiansToDegrees(FMath::Atan2(
+		FVector::DotProduct(UnYawedUnPitchedAimRightDir, -FrameUpDir), 
+		FVector::DotProduct(UnYawedUnPitchedAimRightDir, FrameRightDir)));
+
+	RollPitchYaw.Set(Roll, Pitch, Yaw);
+}
+
+bool BlendSpaceAnalysis::CalculateOrientation(FVector& Result, const UBlendSpace& BlendSpace, const UEulerAnalysisProperties* AnalysisProperties,
+	const UAnimSequence& Animation, const float RateScale)
+{
+	const int32 NumSampledKeys = Animation.GetNumberOfSampledKeys() - 1;
+	if (!AnalysisProperties || NumSampledKeys <= 0)
+	{
+		return false;
+	}
+
+	FTransform BoneOffset;
+	FName BoneName;
+	if (!GetBoneInfo(Animation, AnalysisProperties->BoneSocket, BoneOffset, BoneName))
+	{
+		return false;
+	}
+
+	int32 FirstKey = FMath::Clamp(
+		(int32) (float(NumSampledKeys) * AnalysisProperties->StartTimeFraction), 0, NumSampledKeys);
+	int32 LastKey = FMath::Clamp(
+		(int32) (float(NumSampledKeys) * AnalysisProperties->EndTimeFraction), FirstKey, NumSampledKeys);
+
+	FTransform FrameTM;
+	bool bNeedToUpdateFrameTM = true;
+	FVector FrameFacingDir, FrameUpDir, FrameRightDir;
+	CalculateFrameTM(bNeedToUpdateFrameTM, FrameTM, FirstKey, AnalysisProperties, Animation);
+	GetFrameDirs(FrameFacingDir, FrameUpDir, FrameRightDir, FrameTM, AnalysisProperties);
+
+	Result.Set(0, 0, 0);
+	for (int32 Key = FirstKey; Key != LastKey + 1; ++Key)
+	{
+		CalculateFrameTM(bNeedToUpdateFrameTM, FrameTM, Key, AnalysisProperties, Animation);
+		GetFrameDirs(FrameFacingDir, FrameUpDir, FrameRightDir, FrameTM, AnalysisProperties);
+
+		FVector RollPitchYaw;
+		CalculateBoneOrientation(
+			RollPitchYaw, Animation, Key, BoneName, BoneOffset, 
+			AnalysisProperties, FrameFacingDir, FrameRightDir, FrameUpDir);
+#ifdef ANALYSIS_VERBOSE_LOG
+		UE_LOG(LogAnimation, Log, TEXT("Roll/pitch/yaw = %f %f %f"), RollPitchYaw.X, RollPitchYaw.Y, RollPitchYaw.Z);
+#endif
+		Result += RollPitchYaw;
+	}
+	Result /= (1 + LastKey - FirstKey);
+	UE_LOG(LogAnimation, Log, TEXT("%s Orientation = %f %f %f"), *BoneName.ToString(), Result.X, Result.Y, Result.Z);
+	return true;
+}
+
+bool BlendSpaceAnalysis::CalculateDeltaOrientation(FVector& Result, const UBlendSpace& BlendSpace, const UEulerAnalysisProperties* AnalysisProperties,
+	const UAnimSequence& Animation, const float RateScale)
+{
+	const int32 NumSampledKeys = Animation.GetNumberOfSampledKeys() - 1;
+	if (!AnalysisProperties || NumSampledKeys <= 0)
+	{
+		return false;
+	}
+
+	FTransform BoneOffset;
+	FName BoneName;
+	if (!GetBoneInfo(Animation, AnalysisProperties->BoneSocket, BoneOffset, BoneName))
+	{
+		return false;
+	}
+
+	int32 FirstKey = FMath::Clamp(
+		(int32) (float(NumSampledKeys) * AnalysisProperties->StartTimeFraction), 0, NumSampledKeys);
+	int32 LastKey = FMath::Clamp(
+		(int32) (float(NumSampledKeys) * AnalysisProperties->EndTimeFraction), FirstKey, NumSampledKeys);
+
+	FTransform FrameTM;
+	bool bNeedToUpdateFrameTM = true;
+	FVector FrameFacingDir, FrameUpDir, FrameRightDir;
+	CalculateFrameTM(bNeedToUpdateFrameTM, FrameTM, FirstKey, AnalysisProperties, Animation);
+	GetFrameDirs(FrameFacingDir, FrameUpDir, FrameRightDir, FrameTM, AnalysisProperties);
+
+	FVector RollPitchYaw1;
+	CalculateBoneOrientation(
+		RollPitchYaw1, Animation, FirstKey, BoneName, BoneOffset,
+		AnalysisProperties, FrameFacingDir, FrameRightDir, FrameUpDir);
+
+	CalculateFrameTM(bNeedToUpdateFrameTM, FrameTM, LastKey, AnalysisProperties, Animation);
+	GetFrameDirs(FrameFacingDir, FrameUpDir, FrameRightDir, FrameTM, AnalysisProperties);
+
+	FVector RollPitchYaw2;
+	CalculateBoneOrientation(
+		RollPitchYaw2, Animation, LastKey, BoneName, BoneOffset,
+		AnalysisProperties, FrameFacingDir, FrameRightDir, FrameUpDir);
+
+	Result = RollPitchYaw2 - RollPitchYaw1;
+	return true;
+}
+
+bool BlendSpaceAnalysis::CalculateAngularVelocity(FVector& Result, const UBlendSpace& BlendSpace, const ULinearAnalysisProperties* AnalysisProperties,
+	const UAnimSequence& Animation, const float RateScale)
+{
+	const int32 NumSampledKeys = Animation.GetNumberOfSampledKeys() - 1;
+	if (!AnalysisProperties || NumSampledKeys <= 0)
+	{
+		return false;
+	}
+
+	if (NumSampledKeys == 1)
+	{
+		Result.Set(0, 0, 0);
+		return true;
+	}
+
+	FTransform BoneOffset;
+	FName BoneName;
+	if (!GetBoneInfo(Animation, AnalysisProperties->BoneSocket, BoneOffset, BoneName))
+	{
+		return false;
+	}
+
+	double DeltaTime = Animation.GetPlayLength() / double(NumSampledKeys);
+
+	int32 FirstKey = FMath::Clamp(
+		(int32) (float(NumSampledKeys) * AnalysisProperties->StartTimeFraction), 0, NumSampledKeys);
+	int32 LastKey = FMath::Clamp(
+		(int32) (float(NumSampledKeys) * AnalysisProperties->EndTimeFraction), FirstKey, NumSampledKeys-1);
+
+	// First and Last key are for averaging. However, the finite differencing always goes from one frame to the next
+	int32 NumKeys = FMath::Max(1 + LastKey - FirstKey, 1);
+
+	FTransform FrameTM;
+	bool bNeedToUpdateFrameTM = true;
+	CalculateFrameTM(bNeedToUpdateFrameTM, FrameTM, FirstKey, AnalysisProperties, Animation);
+
+	Result.Set(0, 0, 0);
+	for (int32 iKey = 0; iKey != NumKeys; ++iKey)
+	{
+		int32 Key = (FirstKey + iKey) % (NumSampledKeys + 1);
+		CalculateFrameTM(bNeedToUpdateFrameTM, FrameTM, Key, AnalysisProperties, Animation);
+
+		FTransform BoneTM1 = GetBoneTransform(Animation, Key, BoneName);
+		FTransform TM1 = BoneOffset * BoneTM1;
+		FQuat RelativeQuat1 = FrameTM.InverseTransformRotation(TM1.GetRotation());
+
+		int32 NextKey = (Key + 1) % (NumSampledKeys + 1);
+		if (AnalysisProperties->Space == EAnalysisSpace::Moving)
+		{
+			CalculateFrameTM(bNeedToUpdateFrameTM, FrameTM, NextKey, AnalysisProperties, Animation);
+		}
+
+		FTransform BoneTM2 = GetBoneTransform(Animation, NextKey, BoneName);
+		FTransform TM2 = BoneOffset * BoneTM2;
+		FQuat RelativeQuat2 = FrameTM.InverseTransformRotation(TM2.GetRotation());
+
+		FQuat Rotation = RelativeQuat2 * RelativeQuat1.Inverse();
+		FVector Axis;
+		double Angle;
+		Rotation.ToAxisAndAngle(Axis, Angle);
+		FVector AngularVelocity = FMath::RadiansToDegrees(Axis * (Angle / DeltaTime));
+#ifdef ANALYSIS_VERBOSE_LOG
+		UE_LOG(LogAnimation, Log, TEXT("Angular Velocity = %f %f %f"), AngularVelocity.X, AngularVelocity.Y, AngularVelocity.Z);
+#endif
+		Result += AngularVelocity;
+	}
+	Result /= (1 + LastKey - FirstKey);
+	Result *= Animation.RateScale * RateScale;
+	UE_LOG(LogAnimation, Log, TEXT("%s angular velocity = %f %f %f"), *BoneName.ToString(), Result.X, Result.Y, Result.Z);
+	return true;
+}
+
+bool BlendSpaceAnalysis::CalculateOrientationRate(FVector& Result, const UBlendSpace& BlendSpace, const UEulerAnalysisProperties* AnalysisProperties,
+	const UAnimSequence& Animation, const float RateScale)
+{
+	const int32 NumSampledKeys = Animation.GetNumberOfSampledKeys() - 1;
+	if (!AnalysisProperties || NumSampledKeys <= 0)
+	{
+		return false;
+	}
+
+	if (NumSampledKeys == 1)
+	{
+		Result.Set(0, 0, 0);
+		return true;
+	}
+
+	FTransform BoneOffset;
+	FName BoneName;
+	if (!GetBoneInfo(Animation, AnalysisProperties->BoneSocket, BoneOffset, BoneName))
+	{
+		return false;
+	}
+
+	double DeltaTime = Animation.GetPlayLength() / double(NumSampledKeys);
+
+	int32 FirstKey = FMath::Clamp(
+		(int32) (float(NumSampledKeys) * AnalysisProperties->StartTimeFraction), 0, NumSampledKeys);
+	int32 LastKey = FMath::Clamp(
+		(int32) (float(NumSampledKeys) * AnalysisProperties->EndTimeFraction), FirstKey, NumSampledKeys-1);
+
+	// First and Last key are for averaging. However, the finite differencing always goes from one frame to the next
+	int32 NumKeys = FMath::Max(1 + LastKey - FirstKey, 1);
+
+	FTransform FrameTM;
+	bool bNeedToUpdateFrameTM = true;
+	FVector FrameFacingDir, FrameUpDir, FrameRightDir;
+	CalculateFrameTM(bNeedToUpdateFrameTM, FrameTM, FirstKey, AnalysisProperties, Animation);
+	GetFrameDirs(FrameFacingDir, FrameUpDir, FrameRightDir, FrameTM, AnalysisProperties);
+
+	Result.Set(0, 0, 0);
+	for (int32 iKey = 0; iKey != NumKeys; ++iKey)
+	{
+		int32 Key = (FirstKey + iKey) % (NumSampledKeys + 1);
+		CalculateFrameTM(bNeedToUpdateFrameTM, FrameTM, Key, AnalysisProperties, Animation);
+		GetFrameDirs(FrameFacingDir, FrameUpDir, FrameRightDir, FrameTM, AnalysisProperties);
+
+		FVector RollPitchYaw1;
+		CalculateBoneOrientation(
+			RollPitchYaw1, Animation, Key, BoneName, BoneOffset, 
+			AnalysisProperties, FrameFacingDir, FrameRightDir, FrameUpDir);
+
+		int32 NextKey = (Key + 1) % (NumSampledKeys + 1);
+		if (AnalysisProperties->Space == EAnalysisSpace::Moving)
+		{
+			CalculateFrameTM(bNeedToUpdateFrameTM, FrameTM, NextKey, AnalysisProperties, Animation);
+			GetFrameDirs(FrameFacingDir, FrameUpDir, FrameRightDir, FrameTM, AnalysisProperties);
+		}
+
+		FVector RollPitchYaw2;
+		CalculateBoneOrientation(
+			RollPitchYaw2, Animation, NextKey, BoneName, BoneOffset, 
+			AnalysisProperties, FrameFacingDir, FrameRightDir, FrameUpDir);
+
+		const FVector OrientationRate = (RollPitchYaw2 - RollPitchYaw1) / DeltaTime;
+#ifdef ANALYSIS_VERBOSE_LOG
+		UE_LOG(LogAnimation, Log, TEXT("Orientation rate = %f %f %f"), OrientationRate.X, OrientationRate.Y, OrientationRate.Z);
+#endif
+		Result += OrientationRate;
+	}
+	Result /= (1 + LastKey - FirstKey);
+	Result *= Animation.RateScale * RateScale;
+	UE_LOG(LogAnimation, Log, TEXT("%s Orientation rate = %f %f %f"), *BoneName.ToString(), Result.X, Result.Y, Result.Z);
+	return true;
+}
+
 //======================================================================================================================
 bool BlendSpaceAnalysis::GetBoneInfo(const UAnimSequence&     Animation, 
 									 const FBoneSocketTarget& BoneSocket, 
@@ -222,48 +651,78 @@ bool FCoreBlendSpaceAnalysisFeature::CalculateSampleValue(float&
 	{
 		return false;
 	}
+
+	// Note we need to be explicit on function via static_cast until deprecated versions of the Calculate methods are removed otherwise the compiler cannot disambiguate the call
+
 	const FString& FunctionName = AnalysisProperties->Function;
 	if (FunctionName.Equals(TEXT("Position")))
 	{
 		return BlendSpaceAnalysis::CalculateComponentSampleValue(
-			Result, BlendSpaceAnalysis::CalculatePosition<ULinearAnalysisProperties>, BlendSpace, 
-			Cast<ULinearAnalysisProperties>(AnalysisProperties), Animation, RateScale);
+			Result,
+			static_cast<bool(*)(FVector&, const UBlendSpace&, const ULinearAnalysisProperties*, const UAnimSequence&, const float)>(&BlendSpaceAnalysis::CalculatePosition),
+			BlendSpace,
+			static_cast<const ULinearAnalysisProperties*>(AnalysisProperties),
+			Animation,
+			RateScale);
 	}
 	else if (FunctionName.Equals(TEXT("Velocity")))
 	{
 		return BlendSpaceAnalysis::CalculateComponentSampleValue(
-			Result, BlendSpaceAnalysis::CalculateVelocity<ULinearAnalysisProperties>, BlendSpace, 
-			Cast<ULinearAnalysisProperties>(AnalysisProperties), Animation, RateScale);
+			Result,
+			static_cast<bool(*)(FVector&, const UBlendSpace&, const ULinearAnalysisPropertiesBase*, const UAnimSequence&, const float)>(&BlendSpaceAnalysis::CalculateVelocity),
+			BlendSpace,
+			Cast<ULinearAnalysisProperties>(AnalysisProperties),
+			Animation,
+			RateScale);
 	}
 	else if (FunctionName.Equals(TEXT("DeltaPosition")))
 	{
 		return BlendSpaceAnalysis::CalculateComponentSampleValue(
-			Result, BlendSpaceAnalysis::CalculateDeltaPosition<ULinearAnalysisProperties>, BlendSpace, 
+			Result,
+			static_cast<bool(*)(FVector&, const UBlendSpace&, const ULinearAnalysisProperties*, const UAnimSequence&, const float)>(&BlendSpaceAnalysis::CalculateDeltaPosition),
+			BlendSpace,
 			Cast<ULinearAnalysisProperties>(AnalysisProperties), Animation, RateScale);
 	}
 	else if (FunctionName.Equals(TEXT("AngularVelocity")))
 	{
 		return BlendSpaceAnalysis::CalculateComponentSampleValue(
-			Result, BlendSpaceAnalysis::CalculateAngularVelocity<ULinearAnalysisProperties>, BlendSpace, 
-			Cast<ULinearAnalysisProperties>(AnalysisProperties), Animation, RateScale);
+			Result,
+			static_cast<bool(*)(FVector&, const UBlendSpace&, const ULinearAnalysisProperties*, const UAnimSequence&, const float)>(&BlendSpaceAnalysis::CalculateAngularVelocity),
+			BlendSpace,
+			Cast<ULinearAnalysisProperties>(AnalysisProperties),
+			Animation,
+			RateScale);
 	}
 	else if (FunctionName.Equals(TEXT("Orientation")))
 	{
 		return BlendSpaceAnalysis::CalculateComponentSampleValue(
-			Result, BlendSpaceAnalysis::CalculateOrientation<UEulerAnalysisProperties>, BlendSpace, 
-			Cast<UEulerAnalysisProperties>(AnalysisProperties), Animation, RateScale);
+			Result,
+			static_cast<bool(*)(FVector&, const UBlendSpace&, const UEulerAnalysisProperties*, const UAnimSequence&, const float)>(&BlendSpaceAnalysis::CalculateOrientation),
+			BlendSpace,
+			Cast<UEulerAnalysisProperties>(AnalysisProperties),
+			Animation,
+			RateScale);
 	}
 	else if (FunctionName.Equals(TEXT("OrientationRate")))
 	{
 		return BlendSpaceAnalysis::CalculateComponentSampleValue(
-			Result, BlendSpaceAnalysis::CalculateOrientationRate<UEulerAnalysisProperties>, BlendSpace, 
-			Cast<UEulerAnalysisProperties>(AnalysisProperties), Animation, RateScale);
+			Result,
+			static_cast<bool(*)(FVector&, const UBlendSpace&, const UEulerAnalysisProperties*, const UAnimSequence&, const float)>(&BlendSpaceAnalysis::CalculateOrientationRate),
+			BlendSpace,
+			Cast<UEulerAnalysisProperties>(AnalysisProperties),
+			Animation,
+			RateScale);
 	}
 	else if (FunctionName.Equals(TEXT("DeltaOrientation")))
 	{
 		return BlendSpaceAnalysis::CalculateComponentSampleValue(
-			Result, BlendSpaceAnalysis::CalculateDeltaOrientation<UEulerAnalysisProperties>, BlendSpace, 
-			Cast<UEulerAnalysisProperties>(AnalysisProperties), Animation, RateScale);
+			Result,
+			static_cast<bool(*)(FVector&, const UBlendSpace&, const UEulerAnalysisProperties*,
+				const UAnimSequence&, const float)>(&BlendSpaceAnalysis::CalculateDeltaOrientation),
+			BlendSpace,
+			Cast<UEulerAnalysisProperties>(AnalysisProperties),
+			Animation,
+			RateScale);
 	}
 	return false;
 }
diff --git a/Engine/Source/Editor/Persona/Public/BlendSpaceAnalysis.h b/Engine/Source/Editor/Persona/Public/BlendSpaceAnalysis.h
--- a/Engine/Source/Editor/Persona/Public/BlendSpaceAnalysis.h
+++ b/Engine/Source/Editor/Persona/Public/BlendSpaceAnalysis.h
@@ -103,17 +103,11 @@ public:
 	float                   EndTimeFraction = 1.0f;
 };
 
-UCLASS()
-class PERSONA_API ULinearAnalysisProperties : public UAnalysisProperties
+UCLASS(Abstract)
+class PERSONA_API ULinearAnalysisPropertiesBase : public UAnalysisProperties
 {
 	GENERATED_BODY()
 public:
-	void InitializeFromCache(TObjectPtr<UCachedAnalysisProperties> Cache) override;
-	void MakeCache(TObjectPtr<UCachedAnalysisProperties>& Cache, UBlendSpace* BlendSpace) override;
-
-	/** Axis for the analysis function */
-	UPROPERTY(EditAnywhere, DisplayName = "Axis", Category = AnalysisProperties)
-	EAnalysisLinearAxis FunctionAxis = EAnalysisLinearAxis::PlusX;
 
 	/** The bone or socket used for analysis */
 	UPROPERTY(EditAnywhere, DisplayName = "Bone/Socket", Category = AnalysisProperties)
@@ -141,6 +135,20 @@ public:
 	float EndTimeFraction = 1.0f;
 };
 
+UCLASS()
+class PERSONA_API ULinearAnalysisProperties : public ULinearAnalysisPropertiesBase
+{
+	GENERATED_BODY()
+public:
+	void InitializeFromCache(TObjectPtr<UCachedAnalysisProperties> Cache) override;
+	void MakeCache(TObjectPtr<UCachedAnalysisProperties>& Cache, UBlendSpace* BlendSpace) override;
+
+	/** Axis for the analysis function */
+	UPROPERTY(EditAnywhere, DisplayName = "Axis", Category = AnalysisProperties)
+	EAnalysisLinearAxis FunctionAxis = EAnalysisLinearAxis::PlusX;
+};
+
+
 UCLASS()
 class PERSONA_API UEulerAnalysisProperties : public UAnalysisProperties
 {
@@ -260,17 +268,40 @@ void GetFrameDirs(
 /**
 * Helper to extract the component from the FVector functions
 */
-template<typename FunctionType, typename T>
+template<typename FunctionType>
 static bool CalculateComponentSampleValue(
-	double&                    Result,
-	const FunctionType&        Fn,
-	const UBlendSpace&         BlendSpace, 
-	const T*                   AnalysisProperties, 
-	const UAnimSequence&       Animation,
-	const float                RateScale)
+	double& Result,
+	const FunctionType& Fn,
+	const UBlendSpace& BlendSpace,
+	const ULinearAnalysisProperties* AnalysisProperties,
+	const UAnimSequence& Animation,
+	const float	RateScale)
+{
+	FVector Value = FVector::ZeroVector;
+	int32 ComponentIndex = (int32)AnalysisProperties->FunctionAxis;
+	if (Fn(Value, BlendSpace, AnalysisProperties, Animation, RateScale))
+	{
+		Result = Value | BlendSpaceAnalysis::GetAxisFromTM(FTransform::Identity, AnalysisProperties->FunctionAxis);
+		return true;
+	}
+	return false;
+}
+
+//======================================================================================================================
+/**
+* Helper to extract the component from the FVector functions
+*/
+template<typename FunctionType>
+static bool CalculateComponentSampleValue(
+	double& Result,
+	const FunctionType& Fn,
+	const UBlendSpace& BlendSpace,
+	const UEulerAnalysisProperties* AnalysisProperties,
+	const UAnimSequence& Animation,
+	const float RateScale)
 {
 	FVector Value;
-	int32 ComponentIndex = (int32) AnalysisProperties->FunctionAxis;
+	int32 ComponentIndex = (int32)AnalysisProperties->FunctionAxis; // Roll, Pitch, Yaw -> 0, 1, 2
 	if (Fn(Value, BlendSpace, AnalysisProperties, Animation, RateScale))
 	{
 		Result = Value[ComponentIndex];
@@ -281,14 +312,149 @@ static bool CalculateComponentSampleValue(
 
 //======================================================================================================================
 /**
-* Helper to extract the component from the FVector functions
+* Helper to extract the component from the FVector functions as a float (pass through to double version)
 */
-template<typename FunctionType, typename T>
+template<typename FunctionType>
 static bool CalculateComponentSampleValue(
 	float&               Result,
 	const FunctionType&  Fn,
 	const UBlendSpace&   BlendSpace,
-	const T*             AnalysisProperties,
+	const ULinearAnalysisProperties* AnalysisProperties,
+	const UAnimSequence& Animation,
+	const float          RateScale)
+{
+	double DoubleResult = Result;
+	bool bResult = CalculateComponentSampleValue(DoubleResult, Fn, BlendSpace, AnalysisProperties, Animation, RateScale);
+	Result = float(DoubleResult);
+	return bResult;
+}
+
+//======================================================================================================================
+/**
+* Helper to extract the component from the FVector functions as a float (pass through to double version)
+*/
+template<typename FunctionType>
+static bool CalculateComponentSampleValue(
+	float& Result,
+	const FunctionType& Fn,
+	const UBlendSpace& BlendSpace,
+	const UEulerAnalysisProperties* AnalysisProperties,
+	const UAnimSequence& Animation,
+	const float          RateScale)
+{
+	double DoubleResult = Result;
+	bool bResult = CalculateComponentSampleValue(DoubleResult, Fn, BlendSpace, AnalysisProperties, Animation, RateScale);
+	Result = float(DoubleResult);
+	return bResult;
+}
+
+//======================================================================================================================
+extern PERSONA_API bool CalculatePosition(
+	FVector&                         Result,
+	const UBlendSpace&               BlendSpace,
+	const ULinearAnalysisProperties* AnalysisProperties,
+	const UAnimSequence&             Animation,
+	const float                      RateScale);
+
+//======================================================================================================================
+extern PERSONA_API bool CalculateDeltaPosition(
+	FVector&                         Result,
+	const UBlendSpace&               BlendSpace,
+	const ULinearAnalysisProperties* AnalysisProperties,
+	const UAnimSequence&             Animation,
+	const float                      RateScale);
+
+//======================================================================================================================
+extern PERSONA_API bool CalculateVelocity(
+	FVector&             Result,
+	const UBlendSpace&   BlendSpace,
+	const ULinearAnalysisPropertiesBase* AnalysisProperties,
+	const UAnimSequence& Animation,
+	const float          RateScale);
+
+//======================================================================================================================
+extern PERSONA_API void CalculateBoneOrientation(
+	FVector&                   RollPitchYaw,
+	const UAnimSequence&       Animation, 
+	const int32                Key, 
+	const FName                BoneName, 
+	const FTransform&          BoneOffset, 
+	const UEulerAnalysisProperties* AnalysisProperties,
+	const FVector&             FrameFacingDir, 
+	const FVector&             FrameRightDir, 
+	const FVector&             FrameUpDir);
+
+//======================================================================================================================
+// Note that if a looping animation has 56 keys, then its first key is 0 and last is 55, but these will be identical poses.
+// Thus it has one fewer intervals/unique keys
+extern PERSONA_API bool CalculateOrientation(
+	FVector&                        Result,
+	const UBlendSpace&              BlendSpace,
+	const UEulerAnalysisProperties* AnalysisProperties,
+	const UAnimSequence&            Animation,
+	const float                     RateScale);
+
+//======================================================================================================================
+// Note that if a looping animation has 56 keys, then its first key is 0 and last is 55, but these will be identical poses.
+// Thus it has one fewer intervals/unique keys
+extern PERSONA_API bool CalculateDeltaOrientation(
+	FVector&                        Result,
+	const UBlendSpace&              BlendSpace,
+	const UEulerAnalysisProperties* AnalysisProperties,
+	const UAnimSequence&            Animation,
+	const float                     RateScale);
+
+//======================================================================================================================
+extern PERSONA_API bool CalculateAngularVelocity(
+	FVector&                         Result,
+	const UBlendSpace&               BlendSpace,
+	const ULinearAnalysisProperties* AnalysisProperties,
+	const UAnimSequence&             Animation,
+	const float                      RateScale);
+
+//======================================================================================================================
+extern PERSONA_API bool CalculateOrientationRate(
+	FVector&                        Result,
+	const UBlendSpace&              BlendSpace,
+	const UEulerAnalysisProperties* AnalysisProperties,
+	const UAnimSequence&            Animation,
+	const float                     RateScale);
+
+//======================================================================================================================
+/**
+* Helper to extract the component from the FVector functions
+*/
+template<typename FunctionType, typename T>
+UE_DEPRECATED(5.6, "Use function with strongly typed AnalysisProperties")
+static bool CalculateComponentSampleValue(
+	double& Result,
+	const FunctionType& Fn,
+	const UBlendSpace& BlendSpace,
+	const T* AnalysisProperties,
+	const UAnimSequence& Animation,
+	const float                RateScale)
+{
+	FVector Value;
+	int32 ComponentIndex = (int32)AnalysisProperties->FunctionAxis;
+	if (Fn(Value, BlendSpace, AnalysisProperties, Animation, RateScale))
+	{
+		Result = Value[ComponentIndex];
+		return true;
+	}
+	return false;
+}
+
+//======================================================================================================================
+/**
+* Helper to extract the component from the FVector functions as a float (pass through to double version)
+*/
+template<typename FunctionType, typename T>
+UE_DEPRECATED(5.6, "Use function with strongly typed AnalysisProperties")
+static bool CalculateComponentSampleValue(
+	float& Result,
+	const FunctionType& Fn,
+	const UBlendSpace& BlendSpace,
+	const T AnalysisProperties,
 	const UAnimSequence& Animation,
 	const float          RateScale)
 {
@@ -300,6 +466,7 @@ static bool CalculateComponentSampleValue(
 
 //======================================================================================================================
 template <typename T>
+UE_DEPRECATED(5.6, "Use function with strongly typed AnalysisProperties")
 static bool CalculatePosition(
 	FVector&                         Result,
 	const UBlendSpace&               BlendSpace,
@@ -344,6 +511,7 @@ static bool CalculatePosition(
 
 //======================================================================================================================
 template <typename T>
+UE_DEPRECATED(5.6, "Use function with strongly typed AnalysisProperties")
 static bool CalculateDeltaPosition(
 	FVector&                         Result,
 	const UBlendSpace&               BlendSpace,
@@ -388,6 +556,7 @@ static bool CalculateDeltaPosition(
 
 //======================================================================================================================
 template <typename T>
+UE_DEPRECATED(5.6, "Use function with strongly typed AnalysisProperties")
 static bool CalculateVelocity(
 	FVector&             Result,
 	const UBlendSpace&   BlendSpace,
@@ -463,6 +632,7 @@ static bool CalculateVelocity(
 
 //======================================================================================================================
 template <typename T>
+UE_DEPRECATED(5.6, "Use function with strongly typed AnalysisProperties")
 void CalculateBoneOrientation(
 	FVector&                   RollPitchYaw,
 	const UAnimSequence&       Animation, 
@@ -521,6 +691,7 @@ void CalculateBoneOrientation(
 // Note that if a looping animation has 56 keys, then its first key is 0 and last is 55, but these will be identical poses.
 // Thus it has one fewer intervals/unique keys
 template <typename T>
+UE_DEPRECATED(5.6, "Use function with strongly typed AnalysisProperties")
 static bool CalculateOrientation(
 	FVector&                        Result,
 	const UBlendSpace&              BlendSpace,
@@ -576,6 +747,7 @@ static bool CalculateOrientation(
 // Note that if a looping animation has 56 keys, then its first key is 0 and last is 55, but these will be identical poses.
 // Thus it has one fewer intervals/unique keys
 template <typename T>
+UE_DEPRECATED(5.6, "Use function with strongly typed AnalysisProperties")
 static bool CalculateDeltaOrientation(
 	FVector&                        Result,
 	const UBlendSpace&              BlendSpace,
@@ -626,6 +798,7 @@ static bool CalculateDeltaOrientation(
 
 //======================================================================================================================
 template <typename T>
+UE_DEPRECATED(5.6, "Use function with strongly typed AnalysisProperties")
 static bool CalculateAngularVelocity(
 	FVector&                         Result,
 	const UBlendSpace&               BlendSpace,
@@ -704,6 +877,7 @@ static bool CalculateAngularVelocity(
 
 //======================================================================================================================
 template <typename T>
+UE_DEPRECATED(5.6, "Use function with strongly typed AnalysisProperties")
 static bool CalculateOrientationRate(
 	FVector&                        Result,
 	const UBlendSpace&              BlendSpace,
