Allow notifies/sync markers to be cut in Persona Editor

#rnx

diff --git a/Engine/Source/Editor/Persona/Private/SAnimNotifyPanel.cpp b/Engine/Source/Editor/Persona/Private/SAnimNotifyPanel.cpp
--- a/Engine/Source/Editor/Persona/Private/SAnimNotifyPanel.cpp
+++ b/Engine/Source/Editor/Persona/Private/SAnimNotifyPanel.cpp
@@ -3103,6 +3103,9 @@ UToolMenu* SAnimNotifyTrack::CreateContextMenuContent(FName BaseMenuName)
 				// allow it to delete
 				EditSection.AddMenuEntry(FAnimNotifyPanelCommands::Get().DeleteNotify);
 
+				// allow it to cut notifies
+				EditSection.AddMenuEntry(FAnimNotifyPanelCommands::Get().CutNotifies);
+				
 				if (MenuContext->NotifyEvent)
 				{
 					// For the "Replace With..." menu, make sure the current AnimNotify selection is valid for replacement
@@ -3961,6 +3964,7 @@ void FAnimNotifyPanelCommands::RegisterCommands()
 {
 	UI_COMMAND(DeleteNotify, "Delete", "Deletes the selected notifies.", EUserInterfaceActionType::Button, FInputChord(EKeys::Platform_Delete));
 	UI_COMMAND(CopyNotifies, "Copy", "Copy animation notify events.", EUserInterfaceActionType::Button, FInputChord(EModifierKey::Control, EKeys::C));
+	UI_COMMAND(CutNotifies, "Cut", "Cut animation notify events.", EUserInterfaceActionType::Button, FInputChord(EModifierKey::Control, EKeys::X));
 	UI_COMMAND(PasteNotifies, "Paste", "Paste animation notify event here.", EUserInterfaceActionType::Button, FInputChord(EModifierKey::Control, EKeys::V));
 }
 
@@ -4348,6 +4352,12 @@ float SAnimNotifyPanel::GetSequenceLength() const
 	return Sequence->GetPlayLength();
 }
 
+void SAnimNotifyPanel::CutSelectedNodesToClipboard()
+{
+	CopySelectedNodesToClipboard();
+	DeleteSelectedNodeObjects();
+}
+
 void SAnimNotifyPanel::PostUndo( bool bSuccess )
 {
 	if(Sequence != NULL)
@@ -4823,6 +4833,10 @@ void SAnimNotifyPanel::BindCommands()
 		Commands.DeleteNotify,
 		FExecuteAction::CreateSP(this, &SAnimNotifyPanel::OnDeletePressed));
 
+	CommandList->MapAction(
+		Commands.CutNotifies,
+		FExecuteAction::CreateSP(this, &SAnimNotifyPanel::CutSelectedNodesToClipboard));
+	
 	CommandList->MapAction(
 		Commands.CopyNotifies,
 		FExecuteAction::CreateSP(this, &SAnimNotifyPanel::CopySelectedNodesToClipboard));
diff --git a/Engine/Source/Editor/Persona/Private/SAnimNotifyPanel.h b/Engine/Source/Editor/Persona/Private/SAnimNotifyPanel.h
--- a/Engine/Source/Editor/Persona/Private/SAnimNotifyPanel.h
+++ b/Engine/Source/Editor/Persona/Private/SAnimNotifyPanel.h
@@ -157,6 +157,8 @@ public:
 	TSharedPtr<FUICommandInfo> DeleteNotify;
 
 	TSharedPtr<FUICommandInfo> CopyNotifies;
+	
+	TSharedPtr<FUICommandInfo> CutNotifies;
 
 	TSharedPtr<FUICommandInfo> PasteNotifies;
 
@@ -230,6 +232,8 @@ public:
 
 	virtual float GetSequenceLength() const override;
 
+	void CutSelectedNodesToClipboard();
+	
 	void CopySelectedNodesToClipboard() const;
 	void OnPasteNodes(SAnimNotifyTrack* RequestTrack, float ClickTime, ENotifyPasteMode::Type PasteMode, ENotifyPasteMultipleMode::Type MultiplePasteType);
 
