Refactor animation editor asset family to support plugins

diff --git a/Engine/Source/Editor/Persona/Private/PersonaAssetFamily.cpp b/Engine/Source/Editor/Persona/Private/PersonaAssetFamily.cpp
--- a/Engine/Source/Editor/Persona/Private/PersonaAssetFamily.cpp
+++ b/Engine/Source/Editor/Persona/Private/PersonaAssetFamily.cpp
@@ -9,71 +9,120 @@
 #include "Styling/AppStyle.h"
 #include "AssetToolsModule.h"
 #include "Preferences/PersonaOptions.h"
+#include "UObject/UObjectIterator.h"
 
 #define LOCTEXT_NAMESPACE "PersonaAssetFamily"
 
 FPersonaAssetFamily::FPersonaAssetFamily(const UObject* InFromObject)
-	: Skeleton(nullptr)
-	, Mesh(nullptr)
-	, AnimBlueprint(nullptr)
-	, AnimationAsset(nullptr)
-	, PhysicsAsset(nullptr)
 {
-	if (InFromObject)
+	for (TObjectIterator<UAnimationEditorsAssetFamilyExtension> ExtensionIterator(RF_NoFlags); ExtensionIterator; ++ExtensionIterator)
 	{
-		if (InFromObject->IsA<USkeleton>())
-		{
-			Skeleton = CastChecked<USkeleton>(InFromObject);
-		}
-		else if (InFromObject->IsA<UAnimationAsset>())
+		if (ExtensionIterator->GetClass() == UAnimationEditorsAssetFamilyExtension::StaticClass())
 		{
-			AnimationAsset = CastChecked<UAnimationAsset>(InFromObject);
+			continue;
 		}
-		else if (InFromObject->IsA<USkeletalMesh>())
+		FExtenderObjects& Extender = Extenders.AddDefaulted_GetRef();
+		Extender.Extension = *ExtensionIterator;
+		Extender.Asset = nullptr;
+	}
+
+	// Construct a directed graph to sort extenders according to position rules
+	{
+		TMap<FName, TArray<FName>> Graph;
+		TMap<FName, int32> InDegree;
+		TMap<FName, FExtenderObjects*> ExtenderMap;
+
+		for (FExtenderObjects& Extender : Extenders)
 		{
-			Mesh = CastChecked<USkeletalMesh>(InFromObject);
+			const FName AssetClassName = Extender.Extension->GetAssetClass()->GetFName();
+			ExtenderMap.Add(AssetClassName, &Extender);
+			InDegree.Add(AssetClassName, 0);
+			Graph.Add(AssetClassName, TArray<FName>());
 		}
-		else if (InFromObject->IsA<UAnimBlueprint>())
+
+		for (FExtenderObjects& Extender : Extenders)
 		{
-			AnimBlueprint = CastChecked<UAnimBlueprint>(InFromObject);
+			const FName AssetClassName = Extender.Extension->GetAssetClass()->GetFName();
+
+			FName BeforeName = NAME_None;
+			FName AfterName = NAME_None;
+			Extender.Extension->GetPosition(BeforeName, AfterName);
+
+			if (BeforeName != NAME_None && Graph.Contains(AssetClassName) && InDegree.Contains(BeforeName))
+			{
+				Graph[AssetClassName].Add(BeforeName);
+				++InDegree[BeforeName];
+			}
+			if (AfterName != NAME_None && Graph.Contains(AfterName) && InDegree.Contains(AssetClassName))
+			{
+				Graph[AfterName].Add(AssetClassName);
+				++InDegree[AssetClassName];
+			}
 		}
-		else if (InFromObject->IsA<UPhysicsAsset>())
+
+		TQueue<FName> Queue;
+		for (const auto& [ClassName, Degree] : InDegree)
 		{
-			PhysicsAsset = CastChecked<UPhysicsAsset>(InFromObject);
+			if (Degree == 0)
+			{
+				Queue.Enqueue(ClassName);
+			}
 		}
 
-		FindCounterpartAssets(InFromObject, Skeleton, Mesh);
-	}
-}
+		TArray<FExtenderObjects> SortedExtenders;
 
-FPersonaAssetFamily::FPersonaAssetFamily(const UObject* InFromObject, const TSharedRef<FPersonaAssetFamily> InFromFamily)
-	: Skeleton(InFromFamily->Skeleton)
-	, Mesh(InFromFamily->Mesh)
-	, AnimBlueprint(InFromFamily->AnimBlueprint)
-	, AnimationAsset(InFromFamily->AnimationAsset)
-	, PhysicsAsset(InFromFamily->PhysicsAsset)
-{
-	if (InFromObject)
-	{
-		if (InFromObject->IsA<USkeleton>())
+		while (!Queue.IsEmpty())
 		{
-			Skeleton = CastChecked<USkeleton>(InFromObject);
+			FName Current = NAME_None;
+			Queue.Dequeue(Current);
+
+			SortedExtenders.Add(*ExtenderMap[Current]);
+
+			for (const auto& Neighbour : Graph[Current])
+			{
+				--InDegree[Neighbour];
+				if (InDegree[Neighbour] == 0)
+				{
+					Queue.Enqueue(Neighbour);
+				}
+			}
 		}
-		else if (InFromObject->IsA<UAnimationAsset>())
+
+		if (SortedExtenders.Num() != Extenders.Num())
 		{
-			AnimationAsset = CastChecked<UAnimationAsset>(InFromObject);
+			UE_LOG(LogAnimation, Warning, TEXT("Unable to sort animation editor extenders"));
 		}
-		else if (InFromObject->IsA<USkeletalMesh>())
+		else
 		{
-			Mesh = CastChecked<USkeletalMesh>(InFromObject);
+			Extenders = SortedExtenders;
 		}
-		else if (InFromObject->IsA<UAnimBlueprint>())
+	}
+
+	if (InFromObject)
+	{
+		for (FExtenderObjects& Extender : Extenders)
 		{
-			AnimBlueprint = CastChecked<UAnimBlueprint>(InFromObject);
+			if (InFromObject->IsA(Extender.Extension->GetAssetClass()))
+			{
+				Extender.Asset = InFromObject;
+				Extender.Extension->FindCounterpartAssets(InFromObject, *this);
+			}
 		}
-		else if (InFromObject->IsA<UPhysicsAsset>())
+	}
+}
+
+FPersonaAssetFamily::FPersonaAssetFamily(const UObject* InFromObject, const TSharedRef<FPersonaAssetFamily> InFromFamily)
+	: Extenders(InFromFamily->Extenders)
+{
+	if (InFromObject)
+	{
+		for (FExtenderObjects& Extender : Extenders)
 		{
-			PhysicsAsset = CastChecked<UPhysicsAsset>(InFromObject);
+			if (InFromObject->IsA(Extender.Extension->GetAssetClass()))
+			{
+				Extender.Asset = InFromObject;
+				Extender.Extension->FindCounterpartAssets(InFromObject, *this);
+			}
 		}
 	}
 }
@@ -86,11 +135,11 @@ void FPersonaAssetFamily::Initialize()
 void FPersonaAssetFamily::GetAssetTypes(TArray<UClass*>& OutAssetTypes) const
 {
 	OutAssetTypes.Reset();
-	OutAssetTypes.Add(USkeleton::StaticClass());
-	OutAssetTypes.Add(USkeletalMesh::StaticClass());
-	OutAssetTypes.Add(UAnimationAsset::StaticClass());
-	OutAssetTypes.Add(UAnimBlueprint::StaticClass());
-	OutAssetTypes.Add(UPhysicsAsset::StaticClass());
+
+	for (const FExtenderObjects& Extender : Extenders)
+	{
+		OutAssetTypes.Add(Extender.Extension->GetAssetClass());
+	}
 }
 
 template<typename AssetType>
@@ -106,206 +155,49 @@ static void FindAssets(const USkeleton* InSkeleton, TArray<FAssetData>& OutAsset
 
 FAssetData FPersonaAssetFamily::FindAssetOfType(UClass* InAssetClass) const
 {
-	if (InAssetClass)
+	if (const FExtenderObjects* const Extender = GetExtensionForClass(InAssetClass))
 	{
-		if (InAssetClass->IsChildOf<USkeleton>())
+		if (Extender->Asset.IsValid())
 		{
-			// we should always have a skeleton here, this asset family is based on it
-			return FAssetData(Skeleton.Get());
+			return FAssetData(Extender->Asset.Get());
 		}
-		else if (InAssetClass->IsChildOf<UAnimationAsset>())
+		else
 		{
-			if (AnimationAsset.IsValid())
-			{
-				return FAssetData(AnimationAsset.Get());
-			}
-			else
-			{
-				TArray<FAssetData> Assets;
-				FindAssets<UAnimationAsset>(Skeleton.Get(), Assets, "Skeleton");
-				if (Assets.Num() > 0)
-				{
-					return Assets[0];
-				}
-			}
-		}
-		else if (InAssetClass->IsChildOf<USkeletalMesh>())
-		{
-			if (Mesh.IsValid())
-			{
-				return FAssetData(Mesh.Get());
-			}
-			else
-			{
-				TArray<FAssetData> Assets;
-				FindAssets<USkeletalMesh>(Skeleton.Get(), Assets, "Skeleton");
-				if (Assets.Num() > 0)
-				{
-					return Assets[0];
-				}
-			}
-		}
-		else if (InAssetClass->IsChildOf<UAnimBlueprint>())
-		{
-			if (AnimBlueprint.IsValid())
-			{
-				return FAssetData(AnimBlueprint.Get());
-			}
-			else
-			{
-				TArray<FAssetData> Assets;
-				FindAssets<UAnimBlueprint>(Skeleton.Get(), Assets, "TargetSkeleton");
-				if (Assets.Num() > 0)
-				{
-					return Assets[0];
-				}
-			}
-		}
-		else if (InAssetClass->IsChildOf<UPhysicsAsset>())
-		{
-			if (PhysicsAsset.IsValid())
-			{
-				return FAssetData(PhysicsAsset.Get());
-			}
-			else
-			{
-				// If we have a mesh and it has a physics asset, use it
-				if (Mesh.IsValid())
-				{
-					if (UPhysicsAsset* MeshPhysicsAsset = Mesh->GetPhysicsAsset())
-					{
-						return FAssetData(MeshPhysicsAsset);
-					}
-				}
-
-				TArray<FAssetData> Assets;
-
-				FAssetRegistryModule& AssetRegistryModule = FModuleManager::LoadModuleChecked<FAssetRegistryModule>("AssetRegistry");
-				FARFilter Filter;
-				Filter.bRecursiveClasses = true;
-				Filter.ClassPaths.Add(UPhysicsAsset::StaticClass()->GetClassPathName());
+			TArray<FAssetData> Assets;
+			Extender->Extension->FindAssetsOfType(Assets, *this);
 
-				// If we have a mesh, look for a physics asset that has that mesh as its preview mesh
-				if (Mesh.IsValid())
-				{
-					Filter.TagsAndValues.Add(GET_MEMBER_NAME_CHECKED(UPhysicsAsset, PreviewSkeletalMesh), FSoftObjectPath(Mesh.Get()).ToString());
-				}
-
-				AssetRegistryModule.Get().GetAssets(Filter, Assets);
-
-				if (Assets.Num() > 0)
-				{
-					return Assets[0];
-				}
+			if (Assets.Num() > 0)
+			{
+				return Assets[0];
 			}
 		}
 	}
-
 	return FAssetData();
 }
 
 void FPersonaAssetFamily::FindAssetsOfType(UClass* InAssetClass, TArray<FAssetData>& OutAssets) const
 {
-	if (InAssetClass)
+	if (const FExtenderObjects* const Extender = GetExtensionForClass(InAssetClass))
 	{
-		if (InAssetClass->IsChildOf<USkeleton>())
-		{
-			Skeleton.Get()->GetCompatibleSkeletonAssets(OutAssets);
-		}
-		else if (InAssetClass->IsChildOf<UAnimationAsset>())
-		{
-			FindAssets<UAnimationAsset>(Skeleton.Get(), OutAssets, "Skeleton");
-		}
-		else if (InAssetClass->IsChildOf<USkeletalMesh>())
-		{
-			FindAssets<USkeletalMesh>(Skeleton.Get(), OutAssets, "Skeleton");
-		}
-		else if (InAssetClass->IsChildOf<UAnimBlueprint>())
-		{
-			FindAssets<UAnimBlueprint>(Skeleton.Get(), OutAssets, "TargetSkeleton");
-		}
-		else if (InAssetClass->IsChildOf<UPhysicsAsset>())
-		{
-			FAssetRegistryModule& AssetRegistryModule = FModuleManager::LoadModuleChecked<FAssetRegistryModule>("AssetRegistry");
-			FARFilter Filter;
-			Filter.bRecursiveClasses = true;
-			Filter.ClassPaths.Add(UPhysicsAsset::StaticClass()->GetClassPathName());
-
-			// If we have a mesh, look for a physics asset that has that mesh as its preview mesh
-			if (Mesh.IsValid())
-			{
-				Filter.TagsAndValues.Add(GET_MEMBER_NAME_CHECKED(UPhysicsAsset, PreviewSkeletalMesh), FSoftObjectPath(Mesh.Get()).ToString());
-			}
-
-			AssetRegistryModule.Get().GetAssets(Filter, OutAssets);
-
-			// If we have a mesh and it has a physics asset, use it but only if its different from the one on the preview mesh
-			if (Mesh.IsValid())
-			{
-				if (UPhysicsAsset* MeshPhysicsAsset = Mesh->GetPhysicsAsset())
-				{
-					OutAssets.AddUnique(FAssetData(MeshPhysicsAsset));
-				}
-			}
-		}
+		Extender->Extension->FindAssetsOfType(OutAssets, *this);
 	}
 }
 
 FText FPersonaAssetFamily::GetAssetTypeDisplayName(UClass* InAssetClass) const
 {
-	if (InAssetClass)
+	if (const FExtenderObjects* const Extender = GetExtensionForClass(InAssetClass))
 	{
-		if (InAssetClass->IsChildOf<USkeleton>())
-		{
-			return LOCTEXT("SkeletonAssetDisplayName", "Skeleton");
-		}
-		else if (InAssetClass->IsChildOf<UAnimationAsset>())
-		{
-			return LOCTEXT("AnimationAssetDisplayName", "Animation");
-		}
-		else if (InAssetClass->IsChildOf<USkeletalMesh>())
-		{
-			return LOCTEXT("SkeletalMeshAssetDisplayName", "Mesh");
-		}
-		else if (InAssetClass->IsChildOf<UAnimBlueprint>())
-		{
-			return LOCTEXT("AnimBlueprintAssetDisplayName", "Blueprint");
-		}
-		else if (InAssetClass->IsChildOf<UPhysicsAsset>())
-		{
-			return LOCTEXT("PhysicsAssetDisplayName", "Physics");
-		}
+		return Extender->Extension->GetAssetTypeDisplayName();
 	}
-
 	return FText();
 }
 
 const FSlateBrush* FPersonaAssetFamily::GetAssetTypeDisplayIcon(UClass* InAssetClass) const
 {
-	if (InAssetClass)
+	if (const FExtenderObjects* const Extender = GetExtensionForClass(InAssetClass))
 	{
-		if (InAssetClass->IsChildOf<USkeleton>())
-		{
-			return FAppStyle::Get().GetBrush("Persona.AssetClass.Skeleton");
-		}
-		else if (InAssetClass->IsChildOf<UAnimationAsset>())
-		{
-			return FAppStyle::Get().GetBrush("Persona.AssetClass.Animation");
-		}
-		else if (InAssetClass->IsChildOf<USkeletalMesh>())
-		{
-			return FAppStyle::Get().GetBrush("Persona.AssetClass.SkeletalMesh");
-		}
-		else if (InAssetClass->IsChildOf<UAnimBlueprint>())
-		{
-			return FAppStyle::Get().GetBrush("Persona.AssetClass.Blueprint");
-		}
-		else if (InAssetClass->IsChildOf<UPhysicsAsset>())
-		{
-			return FAppStyle::Get().GetBrush("Persona.AssetClass.Physics");
-		}
+		return Extender->Extension->GetAssetTypeDisplayIcon();
 	}
-
 	return nullptr;
 }	
 
@@ -313,224 +205,443 @@ FSlateColor FPersonaAssetFamily::GetAssetTypeDisplayTint(UClass* InAssetClass) c
 {
 	static const FAssetToolsModule& AssetToolsModule = FModuleManager::LoadModuleChecked<FAssetToolsModule>("AssetTools");
 
-	UClass* UseAssetClass = nullptr;
 	if (InAssetClass)
 	{
-		if (InAssetClass->IsChildOf<USkeleton>())
-		{
-			UseAssetClass = USkeleton::StaticClass();
-		}
-		else if (InAssetClass->IsChildOf<UAnimationAsset>())
-		{
-			UseAssetClass = UAnimationAsset::StaticClass();
-		}
-		else if (InAssetClass->IsChildOf<USkeletalMesh>())
-		{
-			UseAssetClass = USkeletalMesh::StaticClass();
-		}
-		else if (InAssetClass->IsChildOf<UAnimBlueprint>())
+		if (const FExtenderObjects* const Extender = GetExtensionForClass(InAssetClass))
 		{
-			UseAssetClass = UAnimBlueprint::StaticClass();
-		}
-		else if (InAssetClass->IsChildOf<UPhysicsAsset>())
-		{
-			UseAssetClass = UPhysicsAsset::StaticClass();
+			TWeakPtr<IAssetTypeActions> AssetTypeActions = AssetToolsModule.Get().GetAssetTypeActionsForClass(Extender->Extension->GetAssetClass());
+			if (AssetTypeActions.IsValid())
+			{
+				return AssetTypeActions.Pin()->GetTypeColor();
+			}
 		}
 	}
 
-	if (UseAssetClass)
-	{
-		TWeakPtr<IAssetTypeActions> AssetTypeActions = AssetToolsModule.Get().GetAssetTypeActionsForClass(UseAssetClass);
-		if (AssetTypeActions.IsValid())
-		{
-			return AssetTypeActions.Pin()->GetTypeColor();
-		}
-	}
 	return FSlateColor::UseForeground();
 }
 
 bool FPersonaAssetFamily::IsAssetCompatible(const FAssetData& InAssetData) const
 {
 	UClass* Class = InAssetData.GetClass();
-	if (Class)
+
+	if (const FExtenderObjects* const Extender = GetExtensionForClass(Class))
 	{
-		if (Class->IsChildOf<USkeleton>())
+		return Extender->Extension->IsAssetCompatible(InAssetData, *this);
+	}
+
+	return false;
+}
+
+UClass* FPersonaAssetFamily::GetAssetFamilyClass(UClass* InClass) const
+{
+	if (const FExtenderObjects* const Extender = GetExtensionForClass(InClass))
+	{
+		return Extender->Extension->GetAssetClass();
+	}
+	return nullptr;
+}
+
+void FPersonaAssetFamily::RecordAssetOpened(const FAssetData& InAssetData)
+{
+	if (IsAssetCompatible(InAssetData))
+	{
+		UClass* Class = InAssetData.GetClass();
+		if (Class)
 		{
-			if (Skeleton.Get())
+			if (FExtenderObjects* Extender = GetExtensionForClass(Class))
 			{
-				return Skeleton.Get()->IsCompatibleForEditor(InAssetData);
+				Extender->Asset = InAssetData.GetAsset();
 			}
 		}
-		else if (Class->IsChildOf<UAnimationAsset>() || Class->IsChildOf<USkeletalMesh>())
-		{
-			FAssetDataTagMapSharedView::FFindTagResult Result = InAssetData.TagsAndValues.FindTag("Skeleton");
 
-			if (Result.IsSet())
-			{
-				if (Skeleton.Get())
-				{
-					return Skeleton.Get()->IsCompatibleForEditor(Result.GetValue());
-				}
-			}
-		}
-		else if (Class->IsChildOf<UAnimBlueprint>())
-		{
-			FAssetDataTagMapSharedView::FFindTagResult Result = InAssetData.TagsAndValues.FindTag("TargetSkeleton");
+		OnAssetOpened.Broadcast(InAssetData.GetAsset());
+	}
+}
 
-			if (Result.IsSet())
-			{
-				if (Skeleton.Get())
-				{
-					return Skeleton.Get()->IsCompatibleForEditor(Result.GetValue());
-				}
-			}
-		}
-		else if (Class->IsChildOf<UPhysicsAsset>())
-		{
-			// If our mesh is valid and this is the physics asset used on it, we are compatible
-			if (Mesh.IsValid() && InAssetData.GetSoftObjectPath() == FSoftObjectPath(Mesh.Get()->GetPhysicsAsset()))
-			{
-				return true;
-			}
+bool FPersonaAssetFamily::IsAssetTypeInFamily(const TObjectPtr<UClass> InClass) const
+{
+	return GetExtensionForClass(InClass) != nullptr;
+}
 
-			// Otherwise check if our mesh is the preview mesh of the physics asset
-			FAssetDataTagMapSharedView::FFindTagResult Result = InAssetData.TagsAndValues.FindTag(GET_MEMBER_NAME_CHECKED(UPhysicsAsset, PreviewSkeletalMesh));
-			if (Result.IsSet() && Mesh.IsValid())
-			{
-				return Result.GetValue() == FSoftObjectPath(Mesh.Get()).ToString();
-			}
-		}
+TWeakObjectPtr<const UObject> FPersonaAssetFamily::GetAssetOfType(const TObjectPtr<UClass> InClass) const
+{
+	if (const FExtenderObjects* const Extender = GetExtensionForClass(InClass))
+	{
+		return Extender->Asset;
 	}
+	return nullptr;
+}
+
+bool FPersonaAssetFamily::SetAssetOfType(const TObjectPtr<UClass> InClass, TWeakObjectPtr<const UObject> InObject)
+{
+	if (FExtenderObjects* Extender = GetExtensionForClass(InClass))
+	{
+		Extender->Asset = InObject;	
+		return true;
 
+	}
 	return false;
 }
 
-UClass* FPersonaAssetFamily::GetAssetFamilyClass(UClass* InClass) const
+void FPersonaAssetFamily::AddReferencedObjects(FReferenceCollector& Collector)
+{
+	for (FExtenderObjects& Extender : Extenders)
+	{
+		Collector.AddReferencedObject(Extender.Extension);
+	}
+}
+
+FString FPersonaAssetFamily::GetReferencerName() const
+{
+	return TEXT("FPersonaAssetFamily");
+}
+
+void FPersonaAssetFamily::OnSettingsChange(const UPersonaOptions* InOptions, EPropertyChangeType::Type InChangeType)
+{
+	OnAssetFamilyChanged.Broadcast();
+}
+
+FPersonaAssetFamily::FExtenderObjects* FPersonaAssetFamily::GetExtensionForClass(const TObjectPtr<const UClass> InClass)
 {
 	if (InClass)
 	{
-		if (InClass->IsChildOf<USkeleton>())
-		{
-			return USkeleton::StaticClass();
-		}
-		else if (InClass->IsChildOf<UAnimationAsset>())
-		{
-			return UAnimationAsset::StaticClass();
-		}
-		else if (InClass->IsChildOf<USkeletalMesh>())
+		for (FExtenderObjects& Extender : Extenders)
 		{
-			return USkeletalMesh::StaticClass();
-		}
-		else if (InClass->IsChildOf<UAnimBlueprint>())
-		{
-			return UAnimBlueprint::StaticClass();
-		}
-		else if (InClass->IsChildOf<UPhysicsAsset>())
-		{
-			return UPhysicsAsset::StaticClass();
+			if (InClass->IsChildOf(Extender.Extension->GetAssetClass()))
+			{
+				return &Extender;
+			}
 		}
 	}
-
 	return nullptr;
 }
 
-void FPersonaAssetFamily::RecordAssetOpened(const FAssetData& InAssetData)
+const FPersonaAssetFamily::FExtenderObjects* const FPersonaAssetFamily::GetExtensionForClass(const TObjectPtr<const UClass> InClass) const
 {
-	if (IsAssetCompatible(InAssetData))
+	if (InClass)
 	{
-		UClass* Class = InAssetData.GetClass();
-		if (Class)
+		for (const FExtenderObjects& Extender : Extenders)
 		{
-			if (Class->IsChildOf<UAnimationAsset>())
-			{
-				AnimationAsset = Cast<UAnimationAsset>(InAssetData.GetAsset());
-			}
-			else if (Class->IsChildOf<USkeletalMesh>())
-			{
-				Mesh = Cast<USkeletalMesh>(InAssetData.GetAsset());
-			}
-			else if (Class->IsChildOf<UAnimBlueprint>())
-			{
-				AnimBlueprint = Cast<UAnimBlueprint>(InAssetData.GetAsset());
-			}
-			else if (Class->IsChildOf<UPhysicsAsset>())
+			if (InClass->IsChildOf(Extender.Extension->GetAssetClass()))
 			{
-				PhysicsAsset = Cast<UPhysicsAsset>(InAssetData.GetAsset());
+				return &Extender;
 			}
 		}
+	}
+	return nullptr;
+}
 
-		OnAssetOpened.Broadcast(InAssetData.GetAsset());
+// UAnimationEditorsAssetFamilyExtension_SkeletonAsset
+
+TObjectPtr<UClass> UAnimationEditorsAssetFamilyExtension_SkeletonAsset::GetAssetClass() const
+{
+	return USkeleton::StaticClass();
+}
+
+FText UAnimationEditorsAssetFamilyExtension_SkeletonAsset::GetAssetTypeDisplayName() const
+{
+	return LOCTEXT("SkeletonAssetDisplayName", "Skeleton");
+}
+
+const FSlateBrush* UAnimationEditorsAssetFamilyExtension_SkeletonAsset::GetAssetTypeDisplayIcon() const
+{
+	return FAppStyle::Get().GetBrush("Persona.AssetClass.Skeleton");
+}
+
+void UAnimationEditorsAssetFamilyExtension_SkeletonAsset::FindAssetsOfType(TArray<FAssetData>& OutAssets, const IAnimationEditorsAssetFamilyInterface& AssetFamilyInterface) const
+{
+	if (TObjectPtr<const USkeleton> SkeletonAsset = AssetFamilyInterface.GetAssetOfType<USkeleton>())
+	{
+		SkeletonAsset->GetCompatibleSkeletonAssets(OutAssets);
+	}
+}
+
+bool UAnimationEditorsAssetFamilyExtension_SkeletonAsset::IsAssetCompatible(const FAssetData& InAssetData, const IAnimationEditorsAssetFamilyInterface& AssetFamilyInterface) const
+{
+	if (TObjectPtr<const USkeleton> SkeletonAsset = AssetFamilyInterface.GetAssetOfType<USkeleton>())
+	{
+		return SkeletonAsset.Get()->IsCompatibleForEditor(InAssetData);
 	}
+	return false;
+}
+
+void UAnimationEditorsAssetFamilyExtension_SkeletonAsset::FindCounterpartAssets(const UObject* InAsset, IAnimationEditorsAssetFamilyInterface& AssetFamilyInterface)
+{
+	const USkeleton* SkeletonAsset = CastChecked<const USkeleton>(InAsset);
+
+	AssetFamilyInterface.SetAssetOfType<USkeletalMesh>(SkeletonAsset->GetPreviewMesh());
+};
+
+void UAnimationEditorsAssetFamilyExtension_SkeletonAsset::GetPosition(FName& OutBeforeClass, FName& OutAfterClass) const
+{
+	OutBeforeClass = NAME_None;
+	OutAfterClass = NAME_None;
+}
+
+// UAnimationEditorsAssetFamilyExtension_SkeletalMeshAsset
+
+TObjectPtr<UClass> UAnimationEditorsAssetFamilyExtension_SkeletalMeshAsset::GetAssetClass() const
+{
+	return USkeletalMesh::StaticClass();
+}
+
+FText UAnimationEditorsAssetFamilyExtension_SkeletalMeshAsset::GetAssetTypeDisplayName() const
+{
+	return LOCTEXT("SkeletalMeshAssetDisplayName", "Skeletal Mesh");
 }
 
-void FPersonaAssetFamily::FindCounterpartAssets(const UObject* InAsset, TWeakObjectPtr<const USkeleton>& OutSkeleton, TWeakObjectPtr<const USkeletalMesh>& OutMesh)
+const FSlateBrush* UAnimationEditorsAssetFamilyExtension_SkeletalMeshAsset::GetAssetTypeDisplayIcon() const
 {
-	const USkeleton* CounterpartSkeleton = OutSkeleton.Get();
-	const USkeletalMesh* CounterpartMesh = OutMesh.Get();
-	FindCounterpartAssets(InAsset, CounterpartSkeleton, CounterpartMesh);
-	OutSkeleton = CounterpartSkeleton;
-	OutMesh = CounterpartMesh;
+	return FAppStyle::Get().GetBrush("Persona.AssetClass.SkeletalMesh");
 }
 
-void FPersonaAssetFamily::FindCounterpartAssets(const UObject* InAsset, const USkeleton*& OutSkeleton, const USkeletalMesh*& OutMesh)
+void UAnimationEditorsAssetFamilyExtension_SkeletalMeshAsset::FindAssetsOfType(TArray<FAssetData>& OutAssets, const IAnimationEditorsAssetFamilyInterface& AssetFamilyInterface) const
 {
-	if (InAsset->IsA<USkeleton>())
+	if (TObjectPtr<const USkeleton> SkeletonAsset = AssetFamilyInterface.GetAssetOfType<USkeleton>())
 	{
-		OutSkeleton = CastChecked<USkeleton>(InAsset);
-		OutMesh = OutSkeleton->GetPreviewMesh();
-		if(OutMesh == nullptr)
-		{
-			OutMesh = OutSkeleton->FindCompatibleMesh();
-		}
+		FindAssets<USkeletalMesh>(SkeletonAsset.Get(), OutAssets, "Skeleton");
 	}
-	else if (InAsset->IsA<UAnimationAsset>())
+}
+
+bool UAnimationEditorsAssetFamilyExtension_SkeletalMeshAsset::IsAssetCompatible(const FAssetData& InAssetData, const IAnimationEditorsAssetFamilyInterface& AssetFamilyInterface) const
+{
+	FAssetDataTagMapSharedView::FFindTagResult Result = InAssetData.TagsAndValues.FindTag("Skeleton");
+
+	if (Result.IsSet())
 	{
-		const UAnimationAsset* AnimationAsset = CastChecked<const UAnimationAsset>(InAsset);
-		OutSkeleton = AnimationAsset->GetSkeleton();
-		OutMesh = AnimationAsset->GetPreviewMesh();
-		if (OutMesh == nullptr)
+		if (TObjectPtr<const USkeleton> SkeletonAsset = AssetFamilyInterface.GetAssetOfType<USkeleton>())
 		{
-			OutMesh = OutSkeleton->GetPreviewMesh();
+			return SkeletonAsset->IsCompatibleForEditor(Result.GetValue());
 		}
-		if(OutMesh == nullptr)
+	}
+	return false;
+}
+
+void UAnimationEditorsAssetFamilyExtension_SkeletalMeshAsset::FindCounterpartAssets(const UObject* InAsset, IAnimationEditorsAssetFamilyInterface& AssetFamilyInterface)
+{
+	const USkeletalMesh* SkeletalMesh = CastChecked<const USkeletalMesh>(InAsset);
+
+	AssetFamilyInterface.SetAssetOfType<USkeleton>(SkeletalMesh->GetSkeleton());
+
+};
+
+void UAnimationEditorsAssetFamilyExtension_SkeletalMeshAsset::GetPosition(FName& OutBeforeClass, FName& OutAfterClass) const
+{
+	OutBeforeClass = NAME_None;
+	OutAfterClass = USkeleton::StaticClass()->GetFName();
+}
+
+// UAnimationEditorsAssetFamilyExtension_AnimationAsset
+
+TObjectPtr<UClass> UAnimationEditorsAssetFamilyExtension_AnimationAsset::GetAssetClass() const
+{
+	return UAnimationAsset::StaticClass();
+}
+
+FText UAnimationEditorsAssetFamilyExtension_AnimationAsset::GetAssetTypeDisplayName() const
+{
+	return LOCTEXT("AnimationAssetDisplayName", "Animation");
+}
+
+const FSlateBrush* UAnimationEditorsAssetFamilyExtension_AnimationAsset::GetAssetTypeDisplayIcon() const
+{
+	return FAppStyle::Get().GetBrush("Persona.AssetClass.Animation");
+}
+
+void UAnimationEditorsAssetFamilyExtension_AnimationAsset::FindAssetsOfType(TArray<FAssetData>& OutAssets, const IAnimationEditorsAssetFamilyInterface& AssetFamilyInterface) const
+{
+	if (TObjectPtr<const USkeleton> SkeletonAsset = AssetFamilyInterface.GetAssetOfType<USkeleton>())
+	{
+		FindAssets<UAnimationAsset>(SkeletonAsset, OutAssets, "Skeleton");
+	}
+}
+
+bool UAnimationEditorsAssetFamilyExtension_AnimationAsset::IsAssetCompatible(const FAssetData& InAssetData, const IAnimationEditorsAssetFamilyInterface& AssetFamilyInterface) const
+{
+	FAssetDataTagMapSharedView::FFindTagResult Result = InAssetData.TagsAndValues.FindTag("Skeleton");
+
+	if (Result.IsSet())
+	{
+		if (TObjectPtr<const USkeleton> SkeletonAsset = AssetFamilyInterface.GetAssetOfType<USkeleton>())
 		{
-			OutMesh = OutSkeleton->FindCompatibleMesh();
+			return SkeletonAsset->IsCompatibleForEditor(Result.GetValue());
 		}
 	}
-	else if (InAsset->IsA<USkeletalMesh>())
+
+	return false;
+}
+
+void UAnimationEditorsAssetFamilyExtension_AnimationAsset::FindCounterpartAssets(const UObject* InAsset, IAnimationEditorsAssetFamilyInterface& AssetFamilyInterface)
+{
+	const UAnimationAsset* AnimationAsset = CastChecked<const UAnimationAsset>(InAsset);
+
+	AssetFamilyInterface.SetAssetOfType<USkeleton>(AnimationAsset->GetSkeleton());
+
+	AssetFamilyInterface.SetAssetOfType<USkeletalMesh>(AnimationAsset->GetPreviewMesh());
+
+	if (AssetFamilyInterface.IsAssetTypeInFamilyAndUnassigned<USkeletalMesh>())
 	{
-		OutMesh = CastChecked<USkeletalMesh>(InAsset);
-		OutSkeleton = OutMesh->GetSkeleton();
+		AssetFamilyInterface.SetAssetOfType<USkeletalMesh>(AnimationAsset->GetSkeleton()->GetPreviewMesh());
 	}
-	else if (InAsset->IsA<UAnimBlueprint>())
+
+	if (AssetFamilyInterface.IsAssetTypeInFamilyAndUnassigned<USkeletalMesh>())
 	{
-		const UAnimBlueprint* AnimBlueprint = CastChecked<const UAnimBlueprint>(InAsset);
-		OutSkeleton = AnimBlueprint->TargetSkeleton;
-		OutMesh = AnimBlueprint->GetPreviewMesh();
-		check(AnimBlueprint->BlueprintType == BPTYPE_Interface || AnimBlueprint->bIsTemplate || AnimBlueprint->TargetSkeleton != nullptr);
-		if(OutMesh == nullptr && AnimBlueprint->TargetSkeleton)
+		AssetFamilyInterface.SetAssetOfType<USkeletalMesh>(AnimationAsset->GetSkeleton()->FindCompatibleMesh());
+	}
+};
+
+void UAnimationEditorsAssetFamilyExtension_AnimationAsset::GetPosition(FName& OutBeforeClass, FName& OutAfterClass) const
+{
+	OutBeforeClass = NAME_None;
+	OutAfterClass = USkeletalMesh::StaticClass()->GetFName();;
+}
+
+// UAnimationEditorsAssetFamilyExtension_AnimBlueprintAsset
+
+TObjectPtr<UClass> UAnimationEditorsAssetFamilyExtension_AnimBlueprintAsset::GetAssetClass() const
+{
+	return UAnimBlueprint::StaticClass();
+}
+
+FText UAnimationEditorsAssetFamilyExtension_AnimBlueprintAsset::GetAssetTypeDisplayName() const
+{
+	return LOCTEXT("AnimBlueprintAssetDisplayName", "Animation Blueprint");
+}
+
+const FSlateBrush* UAnimationEditorsAssetFamilyExtension_AnimBlueprintAsset::GetAssetTypeDisplayIcon() const
+{
+	return FAppStyle::Get().GetBrush("Persona.AssetClass.Blueprint");
+}
+
+void UAnimationEditorsAssetFamilyExtension_AnimBlueprintAsset::FindAssetsOfType(TArray<FAssetData>& OutAssets, const IAnimationEditorsAssetFamilyInterface& AssetFamilyInterface) const
+{
+	if (TObjectPtr<const USkeleton> SkeletonAsset = AssetFamilyInterface.GetAssetOfType<USkeleton>())
+	{
+		FindAssets<UAnimBlueprint>(SkeletonAsset, OutAssets, "TargetSkeleton");
+	}
+}
+
+bool UAnimationEditorsAssetFamilyExtension_AnimBlueprintAsset::IsAssetCompatible(const FAssetData& InAssetData, const IAnimationEditorsAssetFamilyInterface& AssetFamilyInterface) const
+{
+	FAssetDataTagMapSharedView::FFindTagResult Result = InAssetData.TagsAndValues.FindTag("TargetSkeleton");
+
+	if (Result.IsSet())
+	{
+		if (TObjectPtr<const USkeleton> SkeletonAsset = AssetFamilyInterface.GetAssetOfType<USkeleton>())
 		{
-			OutMesh = AnimBlueprint->TargetSkeleton->GetPreviewMesh();
-			if(OutMesh == nullptr)
-			{
-				OutMesh = AnimBlueprint->TargetSkeleton->FindCompatibleMesh();
-			}
+			return SkeletonAsset->IsCompatibleForEditor(Result.GetValue());
 		}
 	}
-	else if (InAsset->IsA<UPhysicsAsset>())
+
+	return false;
+}
+
+void UAnimationEditorsAssetFamilyExtension_AnimBlueprintAsset::FindCounterpartAssets(const UObject* InAsset, IAnimationEditorsAssetFamilyInterface& AssetFamilyInterface)
+{
+	const UAnimBlueprint* AnimBlueprint = CastChecked<const UAnimBlueprint>(InAsset);
+
+	AssetFamilyInterface.SetAssetOfType<USkeleton>(AnimBlueprint->TargetSkeleton);
+
+	AssetFamilyInterface.SetAssetOfType<USkeletalMesh>(AnimBlueprint->GetPreviewMesh());
+
+	check(AnimBlueprint->BlueprintType == BPTYPE_Interface || AnimBlueprint->bIsTemplate || AnimBlueprint->TargetSkeleton != nullptr);
+
+	if (AssetFamilyInterface.IsAssetTypeInFamilyAndUnassigned<USkeletalMesh>())
+	{
+		AssetFamilyInterface.SetAssetOfType<USkeletalMesh>( AnimBlueprint->TargetSkeleton->GetPreviewMesh());
+	}
+
+	if (AssetFamilyInterface.IsAssetTypeInFamilyAndUnassigned<USkeletalMesh>())
+	{
+		AssetFamilyInterface.SetAssetOfType<USkeletalMesh>(AnimBlueprint->TargetSkeleton->FindCompatibleMesh());
+	}
+};
+
+void UAnimationEditorsAssetFamilyExtension_AnimBlueprintAsset::GetPosition(FName& OutBeforeClass, FName& OutAfterClass) const
+{
+	OutBeforeClass = NAME_None;
+	OutAfterClass = UAnimationAsset::StaticClass()->GetFName();;
+}
+
+// UAnimationEditorsAssetFamilyExtension_PhysicsAsset
+
+TObjectPtr<UClass> UAnimationEditorsAssetFamilyExtension_PhysicsAsset::GetAssetClass() const
+{
+	return UPhysicsAsset::StaticClass();
+}
+
+FText UAnimationEditorsAssetFamilyExtension_PhysicsAsset::GetAssetTypeDisplayName() const
+{
+	return LOCTEXT("PhysicsAssetDisplayName", "Physics");
+}
+
+const FSlateBrush* UAnimationEditorsAssetFamilyExtension_PhysicsAsset::GetAssetTypeDisplayIcon() const
+{
+	return FAppStyle::Get().GetBrush("Persona.AssetClass.Physics");
+}
+
+void UAnimationEditorsAssetFamilyExtension_PhysicsAsset::FindAssetsOfType(TArray<FAssetData>& OutAssets, const IAnimationEditorsAssetFamilyInterface& AssetFamilyInterface) const
+{
+	FAssetRegistryModule& AssetRegistryModule = FModuleManager::LoadModuleChecked<FAssetRegistryModule>("AssetRegistry");
+	FARFilter Filter;
+	Filter.bRecursiveClasses = true;
+	Filter.ClassPaths.Add(UPhysicsAsset::StaticClass()->GetClassPathName());
+
+	// If we have a mesh, look for a physics asset that has that mesh as its preview mesh
+	if (TObjectPtr<const USkeletalMesh> SkeletalMeshAsset = AssetFamilyInterface.GetAssetOfType<USkeletalMesh>())
+	{
+		Filter.TagsAndValues.Add(GET_MEMBER_NAME_CHECKED(UPhysicsAsset, PreviewSkeletalMesh), FSoftObjectPath(SkeletalMeshAsset).ToString());
+	}
+
+	AssetRegistryModule.Get().GetAssets(Filter, OutAssets);
+
+	// If we have a mesh and it has a physics asset, use it but only if its different from the one on the preview mesh
+	if (TObjectPtr<const USkeletalMesh> SkeletalMeshAsset = AssetFamilyInterface.GetAssetOfType<USkeletalMesh>())
 	{
-		const UPhysicsAsset* PhysicsAsset = CastChecked<const UPhysicsAsset>(InAsset);
-		OutMesh = PhysicsAsset->PreviewSkeletalMesh.LoadSynchronous();
-		if(OutMesh != nullptr)
+		if (UPhysicsAsset* MeshPhysicsAsset = SkeletalMeshAsset->GetPhysicsAsset())
 		{
-			OutSkeleton = OutMesh->GetSkeleton();
+			OutAssets.AddUnique(FAssetData(MeshPhysicsAsset));
 		}
 	}
 }
 
-void FPersonaAssetFamily::OnSettingsChange(const UPersonaOptions* InOptions, EPropertyChangeType::Type InChangeType)
+bool UAnimationEditorsAssetFamilyExtension_PhysicsAsset::IsAssetCompatible(const FAssetData& InAssetData, const IAnimationEditorsAssetFamilyInterface& AssetFamilyInterface) const
 {
-	OnAssetFamilyChanged.Broadcast();
+	// If our mesh is valid and this is the physics asset used on it, we are compatible
+	TObjectPtr<const USkeletalMesh> SkeletalMeshAsset = AssetFamilyInterface.GetAssetOfType<USkeletalMesh>();
+
+	if (SkeletalMeshAsset && InAssetData.GetSoftObjectPath() == FSoftObjectPath(SkeletalMeshAsset->GetPhysicsAsset()))
+	{
+		return true;
+	}
+
+	// Otherwise check if our mesh is the preview mesh of the physics asset
+	FAssetDataTagMapSharedView::FFindTagResult Result = InAssetData.TagsAndValues.FindTag(GET_MEMBER_NAME_CHECKED(UPhysicsAsset, PreviewSkeletalMesh));
+	if (Result.IsSet() && SkeletalMeshAsset)
+	{
+		return Result.GetValue() == FSoftObjectPath(SkeletalMeshAsset).ToString();
+	}
+
+	return false;
+}
+
+void UAnimationEditorsAssetFamilyExtension_PhysicsAsset::FindCounterpartAssets(const UObject* InAsset, IAnimationEditorsAssetFamilyInterface& AssetFamilyInterface)
+{
+	const UPhysicsAsset* PhysicsAsset = CastChecked<const UPhysicsAsset>(InAsset);
+
+	TObjectPtr<USkeletalMesh> SkeletalMesh = PhysicsAsset->PreviewSkeletalMesh.LoadSynchronous();
+	if (SkeletalMesh)
+	{
+		AssetFamilyInterface.SetAssetOfType<USkeletalMesh>(SkeletalMesh);
+
+		AssetFamilyInterface.SetAssetOfType<USkeleton>(SkeletalMesh->GetSkeleton());
+	}
+};
+
+void UAnimationEditorsAssetFamilyExtension_PhysicsAsset::GetPosition(FName& OutBeforeClass, FName& OutAfterClass) const
+{
+	OutBeforeClass = NAME_None;
+	OutAfterClass = UAnimBlueprint::StaticClass()->GetFName();;
 }
 
 #undef LOCTEXT_NAMESPACE
diff --git a/Engine/Source/Editor/Persona/Private/PersonaAssetFamily.h b/Engine/Source/Editor/Persona/Private/PersonaAssetFamily.h
--- a/Engine/Source/Editor/Persona/Private/PersonaAssetFamily.h
+++ b/Engine/Source/Editor/Persona/Private/PersonaAssetFamily.h
@@ -8,13 +8,104 @@
 #include "Animation/AnimationAsset.h"
 #include "AssetRegistry/AssetData.h"
 #include "IAssetFamily.h"
+#include "AnimationEditorsAssetFamilyExtension.h"
+
+#include "PersonaAssetFamily.generated.h"
 
 class UAnimBlueprint;
 class UPhysicsAsset;
 class UPersonaOptions;
 
-class FPersonaAssetFamily : public IAssetFamily, public TSharedFromThis<FPersonaAssetFamily>
+UCLASS()
+class UAnimationEditorsAssetFamilyExtension_SkeletonAsset : public UAnimationEditorsAssetFamilyExtension
+{
+	GENERATED_BODY()
+
+public:
+	// UAnimationEditorsAssetFamilyExtension
+	virtual TObjectPtr<UClass> GetAssetClass() const override;
+	virtual FText GetAssetTypeDisplayName() const override;
+	virtual const FSlateBrush* GetAssetTypeDisplayIcon() const override;
+	virtual void FindAssetsOfType(TArray<FAssetData>& OutAssets, const IAnimationEditorsAssetFamilyInterface& AssetFamilyInterface) const override;
+	virtual bool IsAssetCompatible(const FAssetData& InAssetData, const IAnimationEditorsAssetFamilyInterface& AssetFamilyInterface) const override;
+	virtual void FindCounterpartAssets(const UObject* InAsset, IAnimationEditorsAssetFamilyInterface& AssetFamilyInterface) override;
+	virtual void GetPosition(FName& OutBeforeClass, FName& OutAfterClass) const override;
+};
+
+UCLASS()
+class UAnimationEditorsAssetFamilyExtension_SkeletalMeshAsset : public UAnimationEditorsAssetFamilyExtension
+{
+	GENERATED_BODY()
+
+public:
+	// UAnimationEditorsAssetFamilyExtension
+	virtual TObjectPtr<UClass> GetAssetClass() const override;
+	virtual FText GetAssetTypeDisplayName() const override;
+	virtual const FSlateBrush* GetAssetTypeDisplayIcon() const override;
+	virtual void FindAssetsOfType(TArray<FAssetData>& OutAssets, const IAnimationEditorsAssetFamilyInterface& AssetFamilyInterface) const override;
+	virtual bool IsAssetCompatible(const FAssetData& InAssetData, const IAnimationEditorsAssetFamilyInterface& AssetFamilyInterface) const override;
+	virtual void FindCounterpartAssets(const UObject* InAsset, IAnimationEditorsAssetFamilyInterface& AssetFamilyInterface) override;
+	virtual void GetPosition(FName& OutBeforeClass, FName& OutAfterClass) const override;
+};
+
+UCLASS()
+class UAnimationEditorsAssetFamilyExtension_AnimationAsset : public UAnimationEditorsAssetFamilyExtension
+{
+	GENERATED_BODY()
+
+public:
+	// UAnimationEditorsAssetFamilyExtension
+	virtual TObjectPtr<UClass> GetAssetClass() const override;
+	virtual FText GetAssetTypeDisplayName() const override;
+	virtual const FSlateBrush* GetAssetTypeDisplayIcon() const override;
+	virtual void FindAssetsOfType(TArray<FAssetData>& OutAssets, const IAnimationEditorsAssetFamilyInterface& AssetFamilyInterface) const override;
+	virtual bool IsAssetCompatible(const FAssetData& InAssetData, const IAnimationEditorsAssetFamilyInterface& AssetFamilyInterface) const override;
+	virtual void FindCounterpartAssets(const UObject* InAsset, IAnimationEditorsAssetFamilyInterface& AssetFamilyInterface) override;
+	virtual void GetPosition(FName& OutBeforeClass, FName& OutAfterClass) const override;
+};
+
+UCLASS()
+class UAnimationEditorsAssetFamilyExtension_AnimBlueprintAsset : public UAnimationEditorsAssetFamilyExtension
+{
+	GENERATED_BODY()
+
+public:
+	// UAnimationEditorsAssetFamilyExtension
+	virtual TObjectPtr<UClass> GetAssetClass() const override;
+	virtual FText GetAssetTypeDisplayName() const override;
+	virtual const FSlateBrush* GetAssetTypeDisplayIcon() const override;
+	virtual void FindAssetsOfType(TArray<FAssetData>& OutAssets, const IAnimationEditorsAssetFamilyInterface& AssetFamilyInterface) const override;
+	virtual bool IsAssetCompatible(const FAssetData& InAssetData, const IAnimationEditorsAssetFamilyInterface& AssetFamilyInterface) const override;
+	virtual void FindCounterpartAssets(const UObject* InAsset, IAnimationEditorsAssetFamilyInterface& AssetFamilyInterface) override;
+	virtual void GetPosition(FName& OutBeforeClass, FName& OutAfterClass) const override;
+};
+
+UCLASS()
+class UAnimationEditorsAssetFamilyExtension_PhysicsAsset : public UAnimationEditorsAssetFamilyExtension
 {
+	GENERATED_BODY()
+
+public:
+	// UAnimationEditorsAssetFamilyExtension
+	virtual TObjectPtr<UClass> GetAssetClass() const override;
+	virtual FText GetAssetTypeDisplayName() const override;
+	virtual const FSlateBrush* GetAssetTypeDisplayIcon() const override;
+	virtual void FindAssetsOfType(TArray<FAssetData>& OutAssets, const IAnimationEditorsAssetFamilyInterface& AssetFamilyInterface) const override;
+	virtual bool IsAssetCompatible(const FAssetData& InAssetData, const IAnimationEditorsAssetFamilyInterface& AssetFamilyInterface) const override;
+	virtual void FindCounterpartAssets(const UObject* InAsset, IAnimationEditorsAssetFamilyInterface& AssetFamilyInterface) override;
+	virtual void GetPosition(FName& OutBeforeClass, FName& OutAfterClass) const override;
+};
+
+class FPersonaAssetFamily : public IAssetFamily, public TSharedFromThis<FPersonaAssetFamily>, public IAnimationEditorsAssetFamilyInterface, public FGCObject
+{
+private:
+	struct FExtenderObjects
+	{
+		TObjectPtr<UAnimationEditorsAssetFamilyExtension> Extension;
+
+		TWeakObjectPtr<const UObject> Asset;
+	};
+
 public:
 	FPersonaAssetFamily(const UObject* InFromObject);
 	FPersonaAssetFamily(const UObject* InFromObject, const TSharedRef<FPersonaAssetFamily> InFromFamily);
@@ -36,11 +127,25 @@ public:
 	DECLARE_DERIVED_EVENT(FPersonaAssetFamily, IAssetFamily::FOnAssetFamilyChanged, FOnAssetFamilyChanged)
 	virtual FOnAssetFamilyChanged& GetOnAssetFamilyChanged() override { return OnAssetFamilyChanged; }
 
-	/** Helper functions for constructor and other systems that need to discover meshes/skeletons from related assets */
-	static void FindCounterpartAssets(const UObject* InAsset, TWeakObjectPtr<const USkeleton>& OutSkeleton, TWeakObjectPtr<const USkeletalMesh>& OutMesh);
-	static void FindCounterpartAssets(const UObject* InAsset, const USkeleton*& OutSkeleton, const USkeletalMesh*& OutMesh);
+	/* IAnimationEditorsAssetFamilyInterface */
+	virtual bool IsAssetTypeInFamily(const TObjectPtr<UClass> InClass) const override;
+	virtual TWeakObjectPtr<const UObject> GetAssetOfType(const TObjectPtr<UClass> InClass) const override;
+	virtual bool SetAssetOfType(const TObjectPtr<UClass> InClass, TWeakObjectPtr<const UObject> InObject) override;
+	
+	/* FGCObject */
+	virtual void AddReferencedObjects(FReferenceCollector& Collector) override;
+	virtual FString GetReferencerName() const override;
+
+	template <typename T>
+	TObjectPtr<const T> GetAssetOfType() const
+	{
+		return IAnimationEditorsAssetFamilyInterface::GetAssetOfType<T>();
+	}
 
 private:
+	FExtenderObjects* GetExtensionForClass(const TObjectPtr<const UClass> InClass);
+	const FExtenderObjects* const GetExtensionForClass(const TObjectPtr<const UClass> InClass) const;
+
 	/** Initialization to avoid shared ptr access in constructor */
 	void Initialize();
 	
@@ -50,20 +155,7 @@ private:
 private:
 	friend class FPersonaAssetFamilyManager;
 	
-	/** The skeleton that links all assets */
-	TWeakObjectPtr<const USkeleton> Skeleton;
-
-	/** The last mesh that was encountered */
-	TWeakObjectPtr<const USkeletalMesh> Mesh;
-
-	/** The last anim blueprint that was encountered */
-	TWeakObjectPtr<const UAnimBlueprint> AnimBlueprint;
-
-	/** The last animation asset that was encountered */
-	TWeakObjectPtr<const UAnimationAsset> AnimationAsset;
-
-	/** The last physics asset that was encountered */
-	TWeakObjectPtr<const UPhysicsAsset> PhysicsAsset;
+	TArray<FExtenderObjects> Extenders;
 
 	/** Event fired when an asset is opened */
 	FOnAssetOpened OnAssetOpened;
diff --git a/Engine/Source/Editor/Persona/Private/PersonaToolkit.cpp b/Engine/Source/Editor/Persona/Private/PersonaToolkit.cpp
--- a/Engine/Source/Editor/Persona/Private/PersonaToolkit.cpp
+++ b/Engine/Source/Editor/Persona/Private/PersonaToolkit.cpp
@@ -32,15 +32,6 @@ FPersonaToolkit::~FPersonaToolkit()
 	PreviewScene.Reset();
 }
 
-static void FindCounterpartAssets(const UObject* InAsset, TWeakObjectPtr<USkeleton>& OutSkeleton, TWeakObjectPtr <USkeletalMesh>& OutMesh)
-{
-	const USkeleton* CounterpartSkeleton = OutSkeleton.Get();
-	const USkeletalMesh* CounterpartMesh = OutMesh.Get();
-	FPersonaAssetFamily::FindCounterpartAssets(InAsset, CounterpartSkeleton, CounterpartMesh);
-	OutSkeleton = MakeWeakObjectPtr(const_cast<USkeleton*>(CounterpartSkeleton));
-	OutMesh = const_cast<USkeletalMesh*>(CounterpartMesh);
-}
-
 void FPersonaToolkit::Initialize(UObject* InAsset, const FPersonaToolkitArgs& PersonaToolkitArgs, USkeleton* InSkeleton)
 {
 	Asset = InAsset;
@@ -62,7 +53,8 @@ void FPersonaToolkit::Initialize(USkeleton* InSkeleton, const FPersonaToolkitArg
 	Skeleton = InSkeleton;
 	InitialAssetClass = USkeleton::StaticClass();
 
-	FindCounterpartAssets(InSkeleton, Skeleton, Mesh);
+	FPersonaAssetFamily AssetFamily(InSkeleton);
+	Mesh = const_cast<USkeletalMesh*>(AssetFamily.GetAssetOfType<USkeletalMesh>().Get());
 
 	CommonInitialSetup(PersonaToolkitArgs);
 }
@@ -73,7 +65,9 @@ void FPersonaToolkit::Initialize(UAnimationAsset* InAnimationAsset, const FPerso
 	AnimationAsset = InAnimationAsset;
 	InitialAssetClass = UAnimationAsset::StaticClass();
 
-	FindCounterpartAssets(InAnimationAsset, Skeleton, Mesh);
+	FPersonaAssetFamily AssetFamily(InAnimationAsset);
+	Skeleton = const_cast<USkeleton*>(AssetFamily.GetAssetOfType<USkeleton>().Get());
+	Mesh = const_cast<USkeletalMesh*>(AssetFamily.GetAssetOfType<USkeletalMesh>().Get());
 
 	CommonInitialSetup(PersonaToolkitArgs);
 
@@ -89,7 +83,8 @@ void FPersonaToolkit::Initialize(USkeletalMesh* InSkeletalMesh, const FPersonaTo
 	Mesh = InSkeletalMesh;
 	InitialAssetClass = USkeletalMesh::StaticClass();
 
-	FindCounterpartAssets(InSkeletalMesh, Skeleton, Mesh);
+	FPersonaAssetFamily AssetFamily(InSkeletalMesh);
+	Skeleton = const_cast<USkeleton*>(AssetFamily.GetAssetOfType<USkeleton>().Get());
 
 	CommonInitialSetup(PersonaToolkitArgs);
 }
@@ -100,7 +95,9 @@ void FPersonaToolkit::Initialize(UAnimBlueprint* InAnimBlueprint, const FPersona
 	AnimBlueprint = InAnimBlueprint;
 	InitialAssetClass = UAnimBlueprint::StaticClass();
 
-	FindCounterpartAssets(InAnimBlueprint, Skeleton, Mesh);
+	FPersonaAssetFamily AssetFamily(InAnimBlueprint);
+	Skeleton = const_cast<USkeleton*>(AssetFamily.GetAssetOfType<USkeleton>().Get());
+	Mesh = const_cast<USkeletalMesh*>(AssetFamily.GetAssetOfType<USkeletalMesh>().Get());
 
 	CommonInitialSetup(PersonaToolkitArgs);
 
@@ -116,7 +113,9 @@ void FPersonaToolkit::Initialize(UPhysicsAsset* InPhysicsAsset, const FPersonaTo
 	PhysicsAsset = InPhysicsAsset;
 	InitialAssetClass = UPhysicsAsset::StaticClass();
 
-	FindCounterpartAssets(InPhysicsAsset, Skeleton, Mesh);
+	FPersonaAssetFamily AssetFamily(InPhysicsAsset);
+	Skeleton = const_cast<USkeleton*>(AssetFamily.GetAssetOfType<USkeleton>().Get());
+	Mesh = const_cast<USkeletalMesh*>(AssetFamily.GetAssetOfType<USkeletalMesh>().Get());
 
 	CommonInitialSetup(PersonaToolkitArgs);
 }
diff --git a/Engine/Source/Editor/Persona/Public/AnimationEditorsAssetFamilyExtension.h b/Engine/Source/Editor/Persona/Public/AnimationEditorsAssetFamilyExtension.h
new file mode 100644
--- /dev/null
+++ b/Engine/Source/Editor/Persona/Public/AnimationEditorsAssetFamilyExtension.h
@@ -0,0 +1,86 @@
+// Copyright Epic Games, Inc. All Rights Reserved.
+
+#pragma once
+
+#include "CoreMinimal.h"
+
+#include "AnimationEditorsAssetFamilyExtension.generated.h"
+
+struct FSlateBrush;
+class USkeleton;
+class USkeletalMesh;
+
+struct IAnimationEditorsAssetFamilyInterface
+{
+	virtual bool IsAssetTypeInFamily(const TObjectPtr<UClass> InClass) const = 0;
+
+	template <typename T>
+	bool IsAssetTypeInFamily() const
+	{
+		return IsAssetTypeInFamily(T::StaticClass());
+	}
+
+	virtual TWeakObjectPtr<const UObject> GetAssetOfType(const TObjectPtr<UClass> InClass) const = 0;
+
+	template <typename T>
+	TObjectPtr<const T> GetAssetOfType() const
+	{
+		TWeakObjectPtr<const UObject> WeakAsset = GetAssetOfType(T::StaticClass());
+		if (TStrongObjectPtr<const UObject> StrongAsset = WeakAsset.Pin())
+		{
+			return CastChecked<T>(StrongAsset.Get());
+		}
+		return nullptr;
+	}
+
+	bool IsAssetTypeInFamilyAndUnassigned(const TObjectPtr<UClass> InClass) const
+	{
+		if (IsAssetTypeInFamily(InClass))
+		{
+			return !GetAssetOfType(InClass).IsValid();
+		}
+		return false;
+	}
+
+	template <typename T>
+	bool IsAssetTypeInFamilyAndUnassigned() const
+	{
+		return IsAssetTypeInFamilyAndUnassigned(T::StaticClass());
+	}
+
+	virtual bool SetAssetOfType(const TObjectPtr<UClass> InClass, TWeakObjectPtr<const UObject> InObject) = 0;
+
+	template <typename T>
+	bool SetAssetOfType(TWeakObjectPtr<const UObject> InObject)
+	{
+		return SetAssetOfType(T::StaticClass(), InObject);
+	}
+};
+
+UCLASS()
+class PERSONA_API UAnimationEditorsAssetFamilyExtension : public UObject
+{
+	GENERATED_BODY()
+
+public:
+	// The asset class that this extension is describing
+	virtual TObjectPtr<UClass> GetAssetClass() const { return nullptr; }
+
+	// Returns the display name to show for this asset class
+	virtual FText GetAssetTypeDisplayName() const { return FText(); }
+
+	// Returns the display icon to show for this asset class
+	virtual const FSlateBrush* GetAssetTypeDisplayIcon() const { return nullptr; }
+
+	// Fills OutAssets with the assets of this type that exist
+	virtual void FindAssetsOfType(TArray<FAssetData>& OutAssets, const IAnimationEditorsAssetFamilyInterface& AssetFamilyInterface) const {};
+
+	// Returns whether or now the provided asset is compatible
+	virtual bool IsAssetCompatible(const FAssetData& InAssetData, const IAnimationEditorsAssetFamilyInterface& AssetFamilyInterface) const { return false; }
+
+	// Sets the rest of the assets in the family
+	virtual void FindCounterpartAssets(const UObject* InAsset, IAnimationEditorsAssetFamilyInterface& AssetFamilyInterface) {};
+
+	// Gets the horizontal position of this asset
+	virtual void GetPosition(FName& OutBeforeClass, FName& OutAfterClass) const {};
+};
\ No newline at end of file
