Fix UI string that show morph target generated by engine when its not the case.

#rnx

diff --git a/Engine/Source/Editor/Persona/Private/SMorphTargetViewer.cpp b/Engine/Source/Editor/Persona/Private/SMorphTargetViewer.cpp
--- a/Engine/Source/Editor/Persona/Private/SMorphTargetViewer.cpp
+++ b/Engine/Source/Editor/Persona/Private/SMorphTargetViewer.cpp
@@ -133,45 +133,63 @@ TSharedRef< SWidget > SMorphTargetListRow::GenerateWidgetForColumn( const FName&
 	{
 		FText SourceFilenamesTooltip;
 		FText MorphNameText = FText::FromName(Item->Name);
+		FText CarriageReturn = FText::FromString(TEXT("\n"));
 
-		bool bFirst = true;
 		for (int32 LodIndex = 0; LodIndex < Item->MorphLodInfos.Num(); ++LodIndex)
 		{
 			const FDisplayedMorphTargetInfo::FMorphLodInfo& MorphLodInfo = Item->MorphLodInfos[LodIndex];
 			const bool bLodCustomImported = !MorphLodInfo.SourceFilename.IsEmpty();
-			if (bLodCustomImported || MorphLodInfo.bIsGeneratedByEngine)
+			const bool bIsValidLodMorph = MorphLodInfo.bIsValidLodMorph;
+			const int32 GeneratedFromLodIndex = MorphLodInfo.GeneratedFromLodIndex;
+			if (LodIndex == 0)
 			{
-				FText CarriageReturn;
-				if (bFirst)
+				if(!bIsValidLodMorph)
 				{
-					if(!MorphLodInfo.bIsGeneratedByEngine)
-					{
-						MorphNameText = FText::Format(LOCTEXT("MorphRowNameCustomImport", "{0} (Imported by File)"), MorphNameText);
-					}
-					else
-					{
-						MorphNameText = FText::Format(LOCTEXT("MorphRowNameGenerated", "{0} (Generated By Engine)"), MorphNameText);
-						SourceFilenamesTooltip = LOCTEXT("SourceFilenameTooltipEntry_Generated", "This morph target was generated by an engine tool.");
-					}
-					bFirst = false;
+					MorphNameText = FText::Format(LOCTEXT("MorphRowNameInvalid", "{0} (Empty or invalid Morph Data)"), MorphNameText);
+				}
+				else if (bLodCustomImported)
+				{
+					MorphNameText = FText::Format(LOCTEXT("MorphRowNameCustomImport", "{0} (Imported by File)"), MorphNameText);
+					SourceFilenamesTooltip = FText::Format(LOCTEXT("Lod0Tooltip_ImportByFile", "This morph target was imported from source filename: {0}"), FText::FromString(MorphLodInfo.SourceFilename));
+				}
+				else if(GeneratedFromLodIndex != INDEX_NONE)
+				{
+					//LOD 0 could be generated only from it's own import data (inline generated)
+					ensure(GeneratedFromLodIndex == 0);
+					MorphNameText = FText::Format(LOCTEXT("MorphRowNameGeneratedFromLOD", "{0} (Generated From Source Data)"), MorphNameText);
+					SourceFilenamesTooltip = LOCTEXT("Lod0Tooltip_GeneratedFromLOD", "This morph target was generated from is base mesh source data.");
+				}
+				else if(MorphLodInfo.bIsGeneratedByEngine)
+				{
+					MorphNameText = FText::Format(LOCTEXT("MorphRowNameGenerated", "{0} (Generated By Engine)"), MorphNameText);
+					SourceFilenamesTooltip = LOCTEXT("Lod0Tooltip_Generated", "This morph target was generated by an engine tool.");
 				}
 				else
 				{
-					CarriageReturn = bFirst ? FText() : FText::FromString(TEXT("\n"));
+					MorphNameText = FText::Format(LOCTEXT("MorphRowNameImportWithLodGeometry", "{0} (Imported with LOD geometry)"), MorphNameText);
+					SourceFilenamesTooltip = LOCTEXT("Lod0Tooltip_ImportWithLodGeometry", "This morph target was imported with the LOD geometry.");
 				}
-
+			}
+			else if (bIsValidLodMorph) //Do not add invalid morph to the tooltip
+			{
 				if (bLodCustomImported)
 				{
-					SourceFilenamesTooltip = FText::Format(LOCTEXT("SourceFilenameTooltipEntry", "{0}{1}LOD {2} Source Filename: {3}"), SourceFilenamesTooltip, CarriageReturn, FText::AsNumber(LodIndex), FText::FromString(MorphLodInfo.SourceFilename));
+					SourceFilenamesTooltip = FText::Format(LOCTEXT("LodXTooltip_ImportByFile", "{0}{1}LOD {2} was imported from source filename: {3}"), SourceFilenamesTooltip, CarriageReturn, FText::AsNumber(LodIndex), FText::FromString(MorphLodInfo.SourceFilename));
+				}
+				else if (GeneratedFromLodIndex != INDEX_NONE)
+				{
+					SourceFilenamesTooltip = FText::Format(LOCTEXT("LodXTooltip_GeneratedFromLOD", "{0}{1}LOD {2} was generated from a lower LOD {3}"), SourceFilenamesTooltip, CarriageReturn, FText::AsNumber(LodIndex), GeneratedFromLodIndex);
+				}
+				else if(MorphLodInfo.bIsGeneratedByEngine)
+				{
+					SourceFilenamesTooltip = FText::Format(LOCTEXT("LodXTooltip_Generated", "{0}{1}LOD {2} was generated by an engine tool"), SourceFilenamesTooltip, CarriageReturn, FText::AsNumber(LodIndex));
+				}
+				else
+				{
+					SourceFilenamesTooltip = FText::Format(LOCTEXT("LodXTooltip_ImportWithLodGeometry", "{0}{1}LOD {2} was imported with the LOD geometry"), SourceFilenamesTooltip, CarriageReturn, FText::AsNumber(LodIndex));
 				}
 			}
 		}
-		//If we do not have a custom import lod, verify if we have a generated moprh
-		if (bFirst)
-		{
-			MorphNameText = FText::Format(LOCTEXT("MorphRowNameImportWithLodGeometry", "{0} (Imported with LOD geometry)"), MorphNameText);
-			SourceFilenamesTooltip = LOCTEXT("SourceFilenameTooltipEntry_ImportWithLodGeometry", "This morph target was imported with the LOD geometry.");
-		}
 
 		return
 			SNew( SVerticalBox )
@@ -675,11 +693,26 @@ void SMorphTargetViewer::CreateMorphTargetList( const FString& SearchText )
 			int32 NumberOfVerts = (MorphTargets[I]->GetMorphLODModels().Num() > 0)? MorphTargets[I]->GetMorphLODModels()[0].Vertices.Num() : 0;
 
 			TArray<FDisplayedMorphTargetInfo::FMorphLodInfo> MorphLodInfos;
+			
 			MorphLodInfos.AddDefaulted(SkeletalMesh->GetLODNum());
 			for (int32 LodIndex = 0; LodIndex < SkeletalMesh->GetLODNum(); ++LodIndex)
 			{
-				MorphLodInfos[LodIndex].SourceFilename = MorphTargets[I]->GetCustomImportedSourceFilename(LodIndex);
-				MorphLodInfos[LodIndex].bIsGeneratedByEngine = MorphTargets[I]->IsGeneratedByEngine(LodIndex);
+				//Only query the data if the morph lod has some data (is valid)
+				MorphLodInfos[LodIndex].bIsValidLodMorph = MorphTargets[I]->HasDataForLOD(LodIndex);
+				if(MorphLodInfos[LodIndex].bIsValidLodMorph)
+				{
+					const bool bIsGeneratedByEngine = MorphTargets[I]->IsGeneratedByEngine(LodIndex);
+					// Morph target generated by the LOD reduction should be tag properly in the UI
+					// We will keep the generated lod index to be able to output it in the UI
+					if(LodIndex > 0 && bIsGeneratedByEngine && SkeletalMesh->IsReductionActive(LodIndex))
+					{
+						FSkeletalMeshOptimizationSettings ReductionSettings = SkeletalMesh->GetReductionSettings(LodIndex);
+						const int32 ReductionBaseLodIndex = ReductionSettings.BaseLOD;
+						MorphLodInfos[LodIndex].GeneratedFromLodIndex = ReductionBaseLodIndex;
+					}
+					MorphLodInfos[LodIndex].SourceFilename = MorphTargets[I]->GetCustomImportedSourceFilename(LodIndex);
+					MorphLodInfos[LodIndex].bIsGeneratedByEngine = bIsGeneratedByEngine;
+				}
 			}
 			TSharedRef<FDisplayedMorphTargetInfo> Info = FDisplayedMorphTargetInfo::Make( MorphTargets[I]->GetFName(), NumberOfVerts, MorphLodInfos);
 			if(MeshComponent)
diff --git a/Engine/Source/Editor/Persona/Private/SMorphTargetViewer.h b/Engine/Source/Editor/Persona/Private/SMorphTargetViewer.h
--- a/Engine/Source/Editor/Persona/Private/SMorphTargetViewer.h
+++ b/Engine/Source/Editor/Persona/Private/SMorphTargetViewer.h
@@ -28,6 +28,8 @@ public:
 	{
 		FString SourceFilename;
 		bool bIsGeneratedByEngine = false;
+		int32 GeneratedFromLodIndex = INDEX_NONE;
+		bool bIsValidLodMorph = false;
 	};
 	TArray<FMorphLodInfo> MorphLodInfos;
 
