Add framework for supporting custom blend profile implementations in animation blueprints
* Firstly converting state machine transitions to upgrade to be able to use new blend profiles, will follow up with additional CLs for remaining ABP uses

diff --git a/Engine/Plugins/Experimental/AnimNextAnimGraph/Source/AnimNextAnimGraph/Private/Traits/BlendSmootherPerBone.cpp b/Engine/Plugins/Experimental/AnimNextAnimGraph/Source/AnimNextAnimGraph/Private/Traits/BlendSmootherPerBone.cpp
--- a/Engine/Plugins/Experimental/AnimNextAnimGraph/Source/AnimNextAnimGraph/Private/Traits/BlendSmootherPerBone.cpp
+++ b/Engine/Plugins/Experimental/AnimNextAnimGraph/Source/AnimNextAnimGraph/Private/Traits/BlendSmootherPerBone.cpp
@@ -291,7 +291,7 @@ namespace UE::AnimNext
 		{
 			if (BlendProfile)
 			{
-				InstanceData->BlendProfileInterfaces.Add(MakeShared<FHierarchyTableBlendProfile>(BlendProfile));
+				InstanceData->BlendProfileInterfaces.Add(MakeShared<FHierarchyTableBlendProfile>(BlendProfile, EBlendProfileMode::TimeFactor));
 			}
 		}
 	}
@@ -329,6 +329,6 @@ namespace UE::AnimNext
 
 	void FBlendProfileProviderTrait::InitializeInstanceData(FExecutionContext& Context, const FTraitBinding& Binding, const FSharedData* SharedData, FInstanceData* InstanceData)
 	{
-		InstanceData->BlendProfileInterface = MakeShared<FHierarchyTableBlendProfile>(SharedData->TimeFactorBlendProfile);
+		InstanceData->BlendProfileInterface = MakeShared<FHierarchyTableBlendProfile>(SharedData->TimeFactorBlendProfile, EBlendProfileMode::TimeFactor);
 	}
 }
diff --git a/Engine/Plugins/Experimental/Animation/HierarchyTableAnimation/Source/Editor/HierarchyTableAnimationEditor.Build.cs b/Engine/Plugins/Experimental/Animation/HierarchyTableAnimation/Source/Editor/HierarchyTableAnimationEditor.Build.cs
--- a/Engine/Plugins/Experimental/Animation/HierarchyTableAnimation/Source/Editor/HierarchyTableAnimationEditor.Build.cs
+++ b/Engine/Plugins/Experimental/Animation/HierarchyTableAnimation/Source/Editor/HierarchyTableAnimationEditor.Build.cs
@@ -13,7 +13,8 @@ public class HierarchyTableAnimationEditor : ModuleRules
 			"AssetDefinition",
 			"Core",
 			"CoreUObject",
-			"Engine", 
+			"Engine",
+			"EditorWidgets",
 			"InputCore",
 			"Slate",
 			"SlateCore",
@@ -21,7 +22,9 @@ public class HierarchyTableAnimationEditor : ModuleRules
 			"HierarchyTableRuntime",
 			"HierarchyTableEditor",
 			"HierarchyTableAnimationRuntime",
-			"ToolMenus"
+			"ToolMenus",
+			"PropertyEditor",
+			"Persona"
 		);
 	}
 }
diff --git a/Engine/Plugins/Experimental/Animation/HierarchyTableAnimation/Source/Editor/Private/BlendProfileStandaloneEditorToolkit.cpp b/Engine/Plugins/Experimental/Animation/HierarchyTableAnimation/Source/Editor/Private/BlendProfileStandaloneEditorToolkit.cpp
--- a/Engine/Plugins/Experimental/Animation/HierarchyTableAnimation/Source/Editor/Private/BlendProfileStandaloneEditorToolkit.cpp
+++ b/Engine/Plugins/Experimental/Animation/HierarchyTableAnimation/Source/Editor/Private/BlendProfileStandaloneEditorToolkit.cpp
@@ -67,7 +67,7 @@ void FBlendProfileStandaloneEditorToolkit::RegisterTabSpawners(const TSharedRef<
 		FDetailsViewArgs DetailsViewArgs;
 		DetailsViewArgs.NameAreaSettings = FDetailsViewArgs::HideNameArea;
 		TSharedRef<IDetailsView> DetailsView = PropertyEditorModule.CreateDetailView(DetailsViewArgs);
-		DetailsView->SetObjects(TArray<UObject*>{ BlendProfile->Table });
+		DetailsView->SetObjects(TArray<UObject*>{ BlendProfile });
 
 		InTabManager->RegisterTabSpawner("BlendProfileStandaloneEditorDetailsTab", FOnSpawnTab::CreateLambda([=](const FSpawnTabArgs&)
 			{
diff --git a/Engine/Plugins/Experimental/Animation/HierarchyTableAnimation/Source/Editor/Private/BlendProfileStandaloneFactory.cpp b/Engine/Plugins/Experimental/Animation/HierarchyTableAnimation/Source/Editor/Private/BlendProfileStandaloneFactory.cpp
--- a/Engine/Plugins/Experimental/Animation/HierarchyTableAnimation/Source/Editor/Private/BlendProfileStandaloneFactory.cpp
+++ b/Engine/Plugins/Experimental/Animation/HierarchyTableAnimation/Source/Editor/Private/BlendProfileStandaloneFactory.cpp
@@ -2,12 +2,19 @@
 
 #include "BlendProfileStandaloneFactory.h"
 #include "BlendProfileStandalone.h"
+#include "Editor.h"
+#include "HierarchyTableEditorModule.h"
+#include "MaskProfile/HierarchyTableTypeMask.h"
 #include "Modules/ModuleManager.h"
+#include "SEnumCombo.h"
 #include "SkeletonHierarchyTableType.h"
-#include "MaskProfile/HierarchyTableTypeMask.h"
-#include "HierarchyTableEditorModule.h"
+#include "Widgets/SWindow.h"
+#include "Widgets/Input/SButton.h"
+
+#define LOCTEXT_NAMESPACE "BlendProfileStandaloneFactory"
 
 UBlendProfileStandaloneFactory::UBlendProfileStandaloneFactory()
+	: BlendProfileType(EBlendProfileStandaloneType::WeightFactor)
 {
 	SupportedClass = UBlendProfileStandalone::StaticClass();
 	bCreateNew = true;
@@ -17,19 +24,13 @@ UObject* UBlendProfileStandaloneFactory::FactoryCreateNew(UClass* Class, UObject
 {
 	TObjectPtr<UBlendProfileStandalone> BlendProfile = NewObject<UBlendProfileStandalone>(InParent, Class, Name, Flags, Context);
 	BlendProfile->Table = NewObject<UHierarchyTable>(BlendProfile);
+	BlendProfile->Type = BlendProfileType;
 
-	// TODO: Steamline hierarchy table creation API
+	// TODO: Streamline hierarchy table creation API
 
-	FInstancedStruct TableMetadata(FHierarchyTable_TableType_Skeleton::StaticStruct());
 	BlendProfile->Table->Initialize(TableMetadata, FHierarchyTable_ElementType_Mask::StaticStruct());
 
-	FHierarchyTableEditorModule& HierarchyTableModule = FModuleManager::GetModuleChecked<FHierarchyTableEditorModule>("HierarchyTableEditor");
-	const TObjectPtr<UHierarchyTable_TableTypeHandler> TableHandler = HierarchyTableModule.CreateTableHandler(BlendProfile->Table);
 	check(TableHandler);
-
-	TableHandler->FactoryConfigureProperties(TableMetadata);
-	BlendProfile->Table->Initialize(TableMetadata, FHierarchyTable_ElementType_Mask::StaticStruct());
-
 	TableHandler->SetHierarchyTable(BlendProfile->Table);
 	TableHandler->ConstructHierarchy();
 
@@ -38,5 +39,72 @@ UObject* UBlendProfileStandaloneFactory::FactoryCreateNew(UClass* Class, UObject
 
 bool UBlendProfileStandaloneFactory::ConfigureProperties()
 {
-	return true;
-}
\ No newline at end of file
+	if (!ConfigureBlendProfileType())
+	{
+		return false;
+	}
+
+	return ConfigureBlendProfileHierarchy();
+}
+
+bool UBlendProfileStandaloneFactory::ConfigureBlendProfileType()
+{
+	bool bConfirmClicked = false;
+
+	TSharedRef<SWindow> Window = SNew(SWindow)
+		.Title(LOCTEXT("Title", "Choose Blend Profile Type"))
+		.ClientSize(FVector2D(400, 400))
+		.SupportsMinimize(false)
+		.SupportsMaximize(false)
+		[
+			SNew(SVerticalBox)
+				+ SVerticalBox::Slot()
+				.AutoHeight()
+				[
+					SNew(SEnumComboBox, StaticEnum<EBlendProfileStandaloneType>())
+						.CurrentValue_Lambda([this]()
+							{
+								return static_cast<int32>(BlendProfileType);
+							})
+						.OnEnumSelectionChanged_Lambda([this](int32 InEnumValue, ESelectInfo::Type SelectInfo)
+							{
+								BlendProfileType = static_cast<EBlendProfileStandaloneType>(InEnumValue);
+							})
+				]
+				+ SVerticalBox::Slot()
+				.AutoHeight()
+				[
+					SNew(SButton)
+						.OnClicked_Lambda([&]()
+							{
+								bConfirmClicked = true;
+								Window->RequestDestroyWindow();
+								return FReply::Handled();
+							})
+						[
+							SNew(STextBlock)
+								.Text(LOCTEXT("Confirm", "Confirm"))
+						]
+				]
+		];
+
+
+	GEditor->EditorAddModalWindow(Window);
+
+	return bConfirmClicked;
+}
+
+bool UBlendProfileStandaloneFactory::ConfigureBlendProfileHierarchy()
+{
+	FHierarchyTableEditorModule& HierarchyTableModule = FModuleManager::GetModuleChecked<FHierarchyTableEditorModule>("HierarchyTableEditor");
+
+	TableHandler = HierarchyTableModule.CreateTableHandler(FHierarchyTable_TableType_Skeleton::StaticStruct());
+	check(TableHandler);
+
+	TableMetadata = FInstancedStruct(FHierarchyTable_TableType_Skeleton::StaticStruct());
+
+	// Displays window for setting skeleton
+	return TableHandler->FactoryConfigureProperties(TableMetadata);
+}
+
+#undef LOCTEXT_NAMESPACE
\ No newline at end of file
diff --git a/Engine/Plugins/Experimental/Animation/HierarchyTableAnimation/Source/Editor/Private/BlendProfileStandaloneFactory.h b/Engine/Plugins/Experimental/Animation/HierarchyTableAnimation/Source/Editor/Private/BlendProfileStandaloneFactory.h
--- a/Engine/Plugins/Experimental/Animation/HierarchyTableAnimation/Source/Editor/Private/BlendProfileStandaloneFactory.h
+++ b/Engine/Plugins/Experimental/Animation/HierarchyTableAnimation/Source/Editor/Private/BlendProfileStandaloneFactory.h
@@ -4,9 +4,13 @@
 
 #include "CoreMinimal.h"
 #include "Factories/Factory.h"
+#include "StructUtils/InstancedStruct.h"
 
 #include "BlendProfileStandaloneFactory.generated.h"
 
+enum class EBlendProfileStandaloneType;
+class UHierarchyTable_TableTypeHandler;
+
 UCLASS()
 class UBlendProfileStandaloneFactory : public UFactory
 {
@@ -18,4 +22,16 @@ public:
 	// UFactory interface
 	UObject* FactoryCreateNew(UClass* Class, UObject* InParent, FName Name, EObjectFlags Flags, UObject* Context, FFeedbackContext* Warn) override;
 	virtual bool ConfigureProperties() override;
+
+private:
+	bool ConfigureBlendProfileType();
+
+	bool ConfigureBlendProfileHierarchy();
+
+	EBlendProfileStandaloneType BlendProfileType;
+
+	FInstancedStruct TableMetadata;
+
+	UPROPERTY()
+	TObjectPtr<UHierarchyTable_TableTypeHandler> TableHandler;
 };
diff --git a/Engine/Plugins/Experimental/Animation/HierarchyTableAnimation/Source/Editor/Private/BlendProfileStandaloneProvider.cpp b/Engine/Plugins/Experimental/Animation/HierarchyTableAnimation/Source/Editor/Private/BlendProfileStandaloneProvider.cpp
new file mode 100644
--- /dev/null
+++ b/Engine/Plugins/Experimental/Animation/HierarchyTableAnimation/Source/Editor/Private/BlendProfileStandaloneProvider.cpp
@@ -0,0 +1,128 @@
+// Copyright Epic Games, Inc. All Rights Reserved.
+
+#include "BlendProfileStandaloneProvider.h"
+
+#include "Widgets/SCompoundWidget.h"
+#include "Widgets/Text/STextBlock.h"
+#include "ContentBrowserModule.h"
+#include "PropertyCustomizationHelpers.h"
+#include "HierarchyTableBlendProfile.h"
+
+void UBlendProfileStandaloneProvider::Initialize(TObjectPtr<UBlendProfileStandalone> InBlendProfile)
+{
+	BlendProfile = InBlendProfile;
+}
+
+void UBlendProfileStandaloneProvider::ConstructBlendProfile(const TObjectPtr<UBlendProfile> OutBlendProfile) const
+{
+	if (BlendProfile && BlendProfile->Table)
+	{
+		// Switch from EBlendProfileStandaloneType to EBlendProfileMode.
+		// These enums have the same values but EBlendProfileMode marks BlendMask as hidden 
+		// so it does not appear in UIs which is a problem for us so we define our own.
+
+		EBlendProfileMode Mode = EBlendProfileMode::WeightFactor;
+		if (BlendProfile->Type == EBlendProfileStandaloneType::TimeFactor)
+		{
+			Mode = EBlendProfileMode::TimeFactor;
+		}
+		else if (BlendProfile->Type == EBlendProfileStandaloneType::BlendMask)
+		{
+			Mode = EBlendProfileMode::BlendMask;
+		}
+
+		FHierarchyTableBlendProfile HierarchyTableBlendProfile(BlendProfile->Table, Mode);
+		HierarchyTableBlendProfile.ConstructBlendProfile(OutBlendProfile);
+	}
+}
+
+FName FBlendProfileStandalonePickerExtender::StaticGetId()
+{
+	return "Standalone";
+}
+
+FName FBlendProfileStandalonePickerExtender::GetId() const
+{
+	return StaticGetId();
+}
+
+FText  FBlendProfileStandalonePickerExtender::GetDisplayName() const
+{
+	return NSLOCTEXT("BlendProfileStandalone", "Asset", "Asset");
+}
+
+TSharedRef<SWidget> FBlendProfileStandalonePickerExtender::ConstructPickerWidget(const FPickerWidgetArgs& InWidgetArgs) const
+{
+	return SNew(FBlendProfileStandalonePickerExtender::SPicker, InWidgetArgs);
+}
+
+bool FBlendProfileStandalonePickerExtender::OwnsBlendProfileProvider(const TObjectPtr<const UObject> InObject) const
+{
+	return InObject.IsA<UBlendProfileStandaloneProvider>();
+}
+
+void FBlendProfileStandalonePickerExtender::SPicker::Construct(const FArguments& InArgs, IBlendProfilePickerExtender::FPickerWidgetArgs InPickerArgs)
+{
+	TObjectPtr<UBlendProfileStandaloneProvider> InitialProvider = Cast<UBlendProfileStandaloneProvider>(InPickerArgs.InitialSelection);
+	SelectedAsset = InitialProvider ? FAssetData(InitialProvider->BlendProfile) : FAssetData();
+
+	FContentBrowserModule& ContentBrowserModule = FModuleManager::Get().LoadModuleChecked<FContentBrowserModule>(TEXT("ContentBrowser"));
+
+	ChildSlot
+		[
+			SNew(SObjectPropertyEntryBox)
+				.ObjectPath_Lambda([this]()
+					{
+						return SelectedAsset.GetObjectPathString();
+					})
+				.OnObjectChanged_Lambda([this, InPickerArgs](const FAssetData& InAssetData)
+					{
+						SelectedAsset = InAssetData;
+
+						TObjectPtr<UObject> NewSelectedAsset = InAssetData.GetSoftObjectPath().TryLoad();
+						TObjectPtr<UBlendProfileStandalone> BlendProfile = Cast<UBlendProfileStandalone>(NewSelectedAsset);
+
+						TObjectPtr<UBlendProfileStandaloneProvider> BlendProfileProvider = NewObject<UBlendProfileStandaloneProvider>(InPickerArgs.Outer);
+						BlendProfileProvider->Initialize(BlendProfile);
+
+						InPickerArgs.OnProviderChanged.ExecuteIfBound(BlendProfileProvider, BlendProfileProvider);
+					})
+				.AllowedClass(UBlendProfileStandalone::StaticClass())
+				.OnShouldFilterAsset_Lambda([InPickerArgs](const FAssetData& InAssetData)
+					{
+						const FString BlendProfileTypeString = InAssetData.GetTagValueRef<FString>(GET_MEMBER_NAME_CHECKED(UBlendProfileStandalone, Type));
+
+						const EBlendProfileStandaloneType BlendProfileType = static_cast<EBlendProfileStandaloneType>(StaticEnum<EBlendProfileStandaloneType>()->GetValueByName(FName(*BlendProfileTypeString)));
+
+						// Only show blend profile assets with the required type
+						bool bKeep = false;
+						switch (InPickerArgs.SupportedBlendProfileModes) 
+						{
+							case EBlendProfilePickerMode::AllModes:
+								bKeep = true;
+								break;
+							case EBlendProfilePickerMode::BlendProfile:
+								bKeep = (BlendProfileType == EBlendProfileStandaloneType::WeightFactor) || (BlendProfileType ==  EBlendProfileStandaloneType::TimeFactor);
+								break;
+							case EBlendProfilePickerMode::BlendMask:
+								bKeep = (BlendProfileType == EBlendProfileStandaloneType::BlendMask);
+								break;
+						}
+
+						if (bKeep)
+						{
+							// Only display blend profile assets with matching skeletons
+							FString SkeletonSoftObjectPath;
+							if (InAssetData.GetTagValue("Skeleton", SkeletonSoftObjectPath))
+							{
+								if (InPickerArgs.Skeleton && FAssetData(InPickerArgs.Skeleton).ToSoftObjectPath().ToString() != SkeletonSoftObjectPath)
+								{
+									bKeep = false;
+								}
+							}
+						}
+
+						return !bKeep;
+					})
+		];
+}
\ No newline at end of file
diff --git a/Engine/Plugins/Experimental/Animation/HierarchyTableAnimation/Source/Editor/Private/BlendProfileStandaloneProvider.h b/Engine/Plugins/Experimental/Animation/HierarchyTableAnimation/Source/Editor/Private/BlendProfileStandaloneProvider.h
new file mode 100644
--- /dev/null
+++ b/Engine/Plugins/Experimental/Animation/HierarchyTableAnimation/Source/Editor/Private/BlendProfileStandaloneProvider.h
@@ -0,0 +1,58 @@
+// Copyright Epic Games, Inc. All Rights Reserved.
+
+#pragma once
+
+#include "Animation/BlendProfile.h"
+#include "BlendProfileStandalone.h"
+#include "IBlendProfilePickerExtender.h"
+#include "UObject/NoExportTypes.h"
+#include "Widgets/SCompoundWidget.h"
+
+#include "BlendProfileStandaloneProvider.generated.h"
+
+/**
+ * Class responsible for actually creating new blend profiles from custom data
+ */
+UCLASS()
+class HIERARCHYTABLEANIMATIONEDITOR_API UBlendProfileStandaloneProvider : public UObject, public IBlendProfileProviderInterface
+{
+	GENERATED_BODY()
+
+public:
+	// IBlendProfileProviderInterface
+	void Initialize(TObjectPtr<UBlendProfileStandalone> InBlendProfile);
+	void ConstructBlendProfile(const TObjectPtr<UBlendProfile> OutBlendProfile) const override;
+
+	UPROPERTY()
+	TObjectPtr<UBlendProfileStandalone> BlendProfile;
+};
+
+/**
+ * Class responsible for adding a new blend profile provider in the editor for FBlendProfileInterfaceWrapper
+ */
+class FBlendProfileStandalonePickerExtender : public IBlendProfilePickerExtender
+{
+public:
+	static FName StaticGetId();
+
+	// IBlendProfilePickerExtender
+	FName GetId() const override;
+	FText GetDisplayName() const override;
+	TSharedRef<SWidget> ConstructPickerWidget(const FPickerWidgetArgs& InWidgetArgs) const override;
+	bool OwnsBlendProfileProvider(const TObjectPtr<const UObject> InObject) const override;
+
+private:
+	class SPicker : public SCompoundWidget
+	{
+	public:
+		SLATE_BEGIN_ARGS(SPicker)
+		{
+		}
+		SLATE_END_ARGS()
+
+		void Construct(const FArguments& InArgs, IBlendProfilePickerExtender::FPickerWidgetArgs InPickerArgs);
+
+	private:
+		FAssetData SelectedAsset;
+	};
+};
\ No newline at end of file
diff --git a/Engine/Plugins/Experimental/Animation/HierarchyTableAnimation/Source/Editor/Private/HierachyTableAnimationEditorModule.cpp b/Engine/Plugins/Experimental/Animation/HierarchyTableAnimation/Source/Editor/Private/HierachyTableAnimationEditorModule.cpp
--- a/Engine/Plugins/Experimental/Animation/HierarchyTableAnimation/Source/Editor/Private/HierachyTableAnimationEditorModule.cpp
+++ b/Engine/Plugins/Experimental/Animation/HierarchyTableAnimation/Source/Editor/Private/HierachyTableAnimationEditorModule.cpp
@@ -11,6 +11,8 @@
 #include "TimeProfile/TimeProfileColumn.h"
 #include "SkeletonHierarchyTableTypeDetailsCustomization.h"
 #include "PropertyEditorModule.h"
+#include "PersonaModule.h"
+#include "BlendProfileStandaloneProvider.h"
 
 void FHierarchyTableAnimationEditorModule::StartupModule()
 {
@@ -36,6 +38,9 @@ void FHierarchyTableAnimationEditorModule::StartupModule()
 
 	FPropertyEditorModule& PropertyModule = FModuleManager::LoadModuleChecked<FPropertyEditorModule>("PropertyEditor");
 	PropertyModule.RegisterCustomPropertyTypeLayout(FHierarchyTable_TableType_Skeleton::StaticStruct()->GetFName(), FOnGetPropertyTypeCustomizationInstance::CreateStatic(&FHierarchyTableSkeletonTableTypeDetailsCustomization::MakeInstance));
+
+	FPersonaModule& PersonaModule = FModuleManager::LoadModuleChecked<FPersonaModule>("Persona");
+	PersonaModule.RegisterBlendProfilePickerExtender(MakeShared<FBlendProfileStandalonePickerExtender>());
 }
 
 void FHierarchyTableAnimationEditorModule::ShutdownModule()
@@ -60,6 +65,12 @@ void FHierarchyTableAnimationEditorModule::ShutdownModule()
 			}
 		}
 	}
+	
+	if (FModuleManager::Get().IsModuleLoaded("Persona"))
+	{
+		FPersonaModule& PersonaModule = FModuleManager::LoadModuleChecked<FPersonaModule>("Persona");
+		PersonaModule.UnregisterBlendProfilePickerExtender(FBlendProfileStandalonePickerExtender::StaticGetId());
+	}
 }
 
 IMPLEMENT_MODULE(FHierarchyTableAnimationEditorModule, HierarchyTableAnimationEditor)
diff --git a/Engine/Plugins/Experimental/Animation/HierarchyTableAnimation/Source/Runtime/Private/AnimNode_BlendProfileLayeredBlend.cpp b/Engine/Plugins/Experimental/Animation/HierarchyTableAnimation/Source/Runtime/Private/AnimNode_BlendProfileLayeredBlend.cpp
--- a/Engine/Plugins/Experimental/Animation/HierarchyTableAnimation/Source/Runtime/Private/AnimNode_BlendProfileLayeredBlend.cpp
+++ b/Engine/Plugins/Experimental/Animation/HierarchyTableAnimation/Source/Runtime/Private/AnimNode_BlendProfileLayeredBlend.cpp
@@ -42,12 +42,12 @@ void FAnimNode_BlendProfileLayeredBlend::UpdateCachedBoneData(const FBoneContain
 
 		if (SkeletonMetadata.Skeleton == Skeleton)
 		{
-			HierarchyTableBlendProfile = FHierarchyTableBlendProfile(MaskTable);
+			HierarchyTableBlendProfile = FHierarchyTableBlendProfile(MaskTable, EBlendProfileMode::BlendMask);
 		}
 	}
 	else
 	{
-		HierarchyTableBlendProfile = FHierarchyTableBlendProfile(Skeleton);
+		HierarchyTableBlendProfile = FHierarchyTableBlendProfile(Skeleton, EBlendProfileMode::BlendMask);
 	}
 	
 	// Update DesiredBoneBlendWeights
diff --git a/Engine/Plugins/Experimental/Animation/HierarchyTableAnimation/Source/Runtime/Private/BlendProfileStandalone.cpp b/Engine/Plugins/Experimental/Animation/HierarchyTableAnimation/Source/Runtime/Private/BlendProfileStandalone.cpp
new file mode 100644
--- /dev/null
+++ b/Engine/Plugins/Experimental/Animation/HierarchyTableAnimation/Source/Runtime/Private/BlendProfileStandalone.cpp
@@ -0,0 +1,26 @@
+// Copyright Epic Games, Inc. All Rights Reserved.
+
+#include "BlendProfileStandalone.h"
+
+#if WITH_EDITOR
+#include "Animation/Skeleton.h"
+#include "SkeletonHierarchyTableType.h"
+#include "UObject/AssetRegistryTagsContext.h"
+#include "UObject/NoExportTypes.h"
+#endif // WITH_EDITOR
+
+#if WITH_EDITOR
+
+void UBlendProfileStandalone::GetAssetRegistryTags(FAssetRegistryTagsContext Context) const
+{
+	Super::GetAssetRegistryTags(Context);
+
+	if (Table)
+	{
+		const FHierarchyTable_TableType_Skeleton& TableMetadata = Table->GetTableMetadata<FHierarchyTable_TableType_Skeleton>();
+
+		Context.AddTag( FAssetRegistryTag("Skeleton", FAssetData(TableMetadata.Skeleton).ToSoftObjectPath().ToString(), FAssetRegistryTag::TT_Alphabetical) );
+	}
+}
+
+#endif // WITH_EDITOR
\ No newline at end of file
diff --git a/Engine/Plugins/Experimental/Animation/HierarchyTableAnimation/Source/Runtime/Private/HierarchyTableBlendProfile.cpp b/Engine/Plugins/Experimental/Animation/HierarchyTableAnimation/Source/Runtime/Private/HierarchyTableBlendProfile.cpp
--- a/Engine/Plugins/Experimental/Animation/HierarchyTableAnimation/Source/Runtime/Private/HierarchyTableBlendProfile.cpp
+++ b/Engine/Plugins/Experimental/Animation/HierarchyTableAnimation/Source/Runtime/Private/HierarchyTableBlendProfile.cpp
@@ -6,7 +6,8 @@
 #include "HierarchyTableDefaultTypes.h"
 #include "MaskProfile/HierarchyTableTypeMask.h"
 
-FHierarchyTableBlendProfile::FHierarchyTableBlendProfile(TObjectPtr<UHierarchyTable> InHierarchyTable)
+FHierarchyTableBlendProfile::FHierarchyTableBlendProfile(TObjectPtr<UHierarchyTable> InHierarchyTable, const EBlendProfileMode InMode)
+	: Mode(InMode)
 {
 	const FHierarchyTable_TableType_Skeleton& TableMetadata = InHierarchyTable->GetTableMetadata<FHierarchyTable_TableType_Skeleton>();
 	Skeleton = TableMetadata.Skeleton;
@@ -75,7 +76,8 @@ FHierarchyTableBlendProfile::FHierarchyTableBlendProfile(TObjectPtr<UHierarchyTa
 	}
 }
 
-FHierarchyTableBlendProfile::FHierarchyTableBlendProfile(const TObjectPtr<USkeleton> InSkeleton)
+FHierarchyTableBlendProfile::FHierarchyTableBlendProfile(const TObjectPtr<USkeleton> InSkeleton, const EBlendProfileMode InMode)
+	: Mode(InMode)
 {
 	if (InSkeleton)
 	{
@@ -123,4 +125,22 @@ const UE::Anim::TNamedValueArray<FDefaultAllocator, UE::Anim::FCurveElement>& FH
 const TArray<FHierarchyTableBlendProfile::FNamedFloat>& FHierarchyTableBlendProfile::GetAttributeBlendWeights() const
 {
 	return AttributeBlendWeights;
+}
+
+EBlendProfileMode FHierarchyTableBlendProfile::GetMode() const
+{
+	return Mode;
+}
+
+void FHierarchyTableBlendProfile::ConstructBlendProfile(const TObjectPtr<UBlendProfile> OutBlendProfile) const
+{
+	OutBlendProfile->SetSkeleton(Skeleton);
+	OutBlendProfile->Mode = Mode;
+
+	for (int32 BoneIndex = 0; BoneIndex < BoneBlendWeights.Num(); ++BoneIndex)
+	{
+		const float BlendWeight = BoneBlendWeights[BoneIndex];
+
+		OutBlendProfile->SetBoneBlendScale(BoneIndex, BlendWeight, false /* bRecurse */, true /* bCreate */);
+	}
 }
\ No newline at end of file
diff --git a/Engine/Plugins/Experimental/Animation/HierarchyTableAnimation/Source/Runtime/Public/BlendProfileStandalone.h b/Engine/Plugins/Experimental/Animation/HierarchyTableAnimation/Source/Runtime/Public/BlendProfileStandalone.h
--- a/Engine/Plugins/Experimental/Animation/HierarchyTableAnimation/Source/Runtime/Public/BlendProfileStandalone.h
+++ b/Engine/Plugins/Experimental/Animation/HierarchyTableAnimation/Source/Runtime/Public/BlendProfileStandalone.h
@@ -7,12 +7,28 @@
 
 #include "BlendProfileStandalone.generated.h"
 
+UENUM()
+enum class EBlendProfileStandaloneType
+{
+	WeightFactor,
+	TimeFactor,
+	BlendMask
+};
+
 UCLASS()
 class HIERARCHYTABLEANIMATIONRUNTIME_API UBlendProfileStandalone : public UObject
 {
 	GENERATED_BODY()
 
 public:
+	UPROPERTY(VisibleAnywhere, Category=Default, AssetRegistrySearchable)
+	EBlendProfileStandaloneType Type;
+
 	UPROPERTY(VisibleAnywhere, Category=Default)
 	TObjectPtr<UHierarchyTable> Table;
-};
\ No newline at end of file
+
+#if WITH_EDITOR
+	virtual void GetAssetRegistryTags(FAssetRegistryTagsContext Context) const override;
+#endif
+};
+
diff --git a/Engine/Plugins/Experimental/Animation/HierarchyTableAnimation/Source/Runtime/Public/HierarchyTableBlendProfile.h b/Engine/Plugins/Experimental/Animation/HierarchyTableAnimation/Source/Runtime/Public/HierarchyTableBlendProfile.h
--- a/Engine/Plugins/Experimental/Animation/HierarchyTableAnimation/Source/Runtime/Public/HierarchyTableBlendProfile.h
+++ b/Engine/Plugins/Experimental/Animation/HierarchyTableAnimation/Source/Runtime/Public/HierarchyTableBlendProfile.h
@@ -11,15 +11,16 @@ public:
 	FHierarchyTableBlendProfile() = default;
 	virtual ~FHierarchyTableBlendProfile() = default;
 
-	FHierarchyTableBlendProfile(TObjectPtr<UHierarchyTable> InHierarchyTable);
+	FHierarchyTableBlendProfile(TObjectPtr<UHierarchyTable> InHierarchyTable, const EBlendProfileMode InMode);
 
-	FHierarchyTableBlendProfile(const TObjectPtr<USkeleton> InSkeleton);
+	FHierarchyTableBlendProfile(const TObjectPtr<USkeleton> InSkeleton, const EBlendProfileMode InMode);
 
 	// IBlendProfileInterface
 	virtual float GetBoneBlendScale(int32 InBoneIdx) const override;
 	virtual int32 GetNumBlendEntries() const override;
 	virtual int32 GetPerBoneInterpolationIndex(const FCompactPoseBoneIndex& InCompactPoseBoneIndex, const FBoneContainer& BoneContainer, const IInterpolationIndexProvider::FPerBoneInterpolationData* Data) const override;
 	virtual int32 GetPerBoneInterpolationIndex(const FSkeletonPoseBoneIndex InSkeletonBoneIndex, const USkeleton* TargetSkeleton, const IInterpolationIndexProvider::FPerBoneInterpolationData* Data) const override;
+	virtual EBlendProfileMode GetMode() const override;
 	virtual TObjectPtr<USkeleton> GetSkeleton() const override;
 
 	struct FNamedFloat
@@ -37,6 +38,8 @@ public:
 	const UE::Anim::TNamedValueArray<FDefaultAllocator, UE::Anim::FCurveElement>& GetCurveBlendWeights() const;
 	const TArray<FNamedFloat>& GetAttributeBlendWeights() const;
 
+	void ConstructBlendProfile(const TObjectPtr<UBlendProfile> OutBlendProfile) const;
+
 private:
 	TObjectPtr<USkeleton> Skeleton;
 
@@ -45,4 +48,6 @@ private:
 	UE::Anim::TNamedValueArray<FDefaultAllocator, UE::Anim::FCurveElement> CurveBlendWeights;
 
 	TArray<FNamedFloat> AttributeBlendWeights;
+
+	EBlendProfileMode Mode = EBlendProfileMode::WeightFactor;
 };
\ No newline at end of file
diff --git a/Engine/Source/Editor/AnimGraph/Private/AnimGraphNode_StateMachineBase.cpp b/Engine/Source/Editor/AnimGraph/Private/AnimGraphNode_StateMachineBase.cpp
--- a/Engine/Source/Editor/AnimGraph/Private/AnimGraphNode_StateMachineBase.cpp
+++ b/Engine/Source/Editor/AnimGraph/Private/AnimGraphNode_StateMachineBase.cpp
@@ -618,6 +618,8 @@ void UAnimGraphNode_StateMachineBase::OnProcessDuringCompilation(IAnimBlueprintC
 			{
 				continue;
 			}
+
+			TransitionNode->BlendProfileWrapper.UpdateCachedBlendProfile();
 			
 			const int32 TransitionIndex = Oven.FindOrAddTransition(FMachineCreator::FUniqueTransition(TransitionNode, StateNode));
 			FAnimationTransitionBetweenStates& BakedTransition = Oven.GetMachine().Transitions[TransitionIndex];
@@ -628,7 +630,7 @@ void UAnimGraphNode_StateMachineBase::OnProcessDuringCompilation(IAnimBlueprintC
 			BakedTransition.InterruptNotify = OutCompiledData.FindOrAddNotify(TransitionNode->TransitionInterrupt);
 			BakedTransition.BlendMode = TransitionNode->BlendMode;
 			BakedTransition.CustomCurve = TransitionNode->CustomBlendCurve;
-			BakedTransition.BlendProfile = TransitionNode->BlendProfile;
+			BakedTransition.BlendProfile = TransitionNode->BlendProfileWrapper.GetBlendProfile();
 			BakedTransition.LogicType = TransitionNode->LogicType;
 
 			UAnimStateNodeBase* PreviousState = StateNode;
diff --git a/Engine/Source/Editor/AnimGraph/Private/AnimStateTransitionNode.cpp b/Engine/Source/Editor/AnimGraph/Private/AnimStateTransitionNode.cpp
--- a/Engine/Source/Editor/AnimGraph/Private/AnimStateTransitionNode.cpp
+++ b/Engine/Source/Editor/AnimGraph/Private/AnimStateTransitionNode.cpp
@@ -150,24 +150,31 @@ void UAnimStateTransitionNode::PostLoad()
 			ParentGraph->SubGraphs.AddUnique(BoundGraph);
 		}
 	}
+
+	PRAGMA_DISABLE_DEPRECATION_WARNINGS
+	if (BlendProfile_DEPRECATED)
+	{
+		BlendProfileWrapper.SetSkeletonBlendProfile(BlendProfile_DEPRECATED);
+
+		BlendProfile_DEPRECATED = nullptr;
+
+		Modify();
+	}
+	PRAGMA_ENABLE_DEPRECATION_WARNINGS
 }
 
 bool UAnimStateTransitionNode::ValidateBlendProfile()
 {
-	if(BlendProfile)
+	if(TObjectPtr<UBlendProfile> BlendProfile = BlendProfileWrapper.GetBlendProfile())
 	{
 		// validate the skeleton of our blend profile
 		UBlueprint* Blueprint = FBlueprintEditorUtils::FindBlueprintForNodeChecked(this);
 		UAnimBlueprint* AnimBP = CastChecked<UAnimBlueprint>(Blueprint);
 
-		if(AnimBP->TargetSkeleton && !AnimBP->TargetSkeleton->BlendProfiles.Contains(BlendProfile))
-		{
-			BlendProfile = nullptr;
-			return false;
-		}
+		return AnimBP->TargetSkeleton == BlendProfile->GetSkeleton();
 	}
 
-	return true;
+	return false;
 }
 
 void UAnimStateTransitionNode::PostPasteNode()
@@ -426,7 +433,7 @@ void UAnimStateTransitionNode::PostEditChangeProperty(struct FPropertyChangedEve
 	if (PropertyName == GET_MEMBER_NAME_CHECKED(UAnimStateTransitionNode, CrossfadeDuration) || 
 		PropertyName == GET_MEMBER_NAME_CHECKED(UAnimStateTransitionNode, BlendMode) ||
 		PropertyName == GET_MEMBER_NAME_CHECKED(UAnimStateTransitionNode, CustomBlendCurve) ||
-		PropertyName == GET_MEMBER_NAME_CHECKED(UAnimStateTransitionNode, BlendProfile))
+		PropertyName == GET_MEMBER_NAME_CHECKED(UAnimStateTransitionNode, BlendProfileWrapper))
 	{
 		PropagateCrossfadeSettings();
 	}
@@ -572,7 +579,7 @@ void UAnimStateTransitionNode::CopyCrossfadeSettings(const UAnimStateTransitionN
 	CrossfadeMode_DEPRECATED = SrcNode->CrossfadeMode_DEPRECATED;
 	BlendMode = SrcNode->BlendMode;
 	CustomBlendCurve = SrcNode->CustomBlendCurve;
-	BlendProfile = SrcNode->BlendProfile;
+	BlendProfileWrapper = SrcNode->BlendProfileWrapper;
 	SharedCrossfadeIdx = SrcNode->SharedCrossfadeIdx;
 	SharedCrossfadeName = SrcNode->SharedCrossfadeName;
 	SharedCrossfadeGuid = SrcNode->SharedCrossfadeGuid;
diff --git a/Engine/Source/Editor/AnimGraph/Public/AnimStateTransitionNode.h b/Engine/Source/Editor/AnimGraph/Public/AnimStateTransitionNode.h
--- a/Engine/Source/Editor/AnimGraph/Public/AnimStateTransitionNode.h
+++ b/Engine/Source/Editor/AnimGraph/Public/AnimStateTransitionNode.h
@@ -48,9 +48,14 @@ class UAnimStateTransitionNode : public UAnimStateNodeBase
 	UPROPERTY(EditAnywhere, Category=Transition)
 	TObjectPtr<UCurveFloat> CustomBlendCurve;
 
+	// The blend profile to use to evaluate this transition per-bone
+	UE_DEPRECATED(5.6, "Use the interfaced blend profile instead")
+	UPROPERTY()
+	TObjectPtr<UBlendProfile> BlendProfile_DEPRECATED;
+
 	// The blend profile to use to evaluate this transition per-bone
 	UPROPERTY(EditAnywhere, Category=Transition, meta=(UseAsBlendProfile=true))
-	TObjectPtr<UBlendProfile> BlendProfile;
+	FBlendProfileInterfaceWrapper BlendProfileWrapper;
 
 	// Try setting the rule automatically based on most relevant asset player node's remaining time and the Automatic Rule Trigger Time of the transition, ignoring the internal time
 	UPROPERTY(EditAnywhere, Category=Transition)
diff --git a/Engine/Source/Editor/DetailCustomizations/Private/AnimTransitionNodeDetails.cpp b/Engine/Source/Editor/DetailCustomizations/Private/AnimTransitionNodeDetails.cpp
--- a/Engine/Source/Editor/DetailCustomizations/Private/AnimTransitionNodeDetails.cpp
+++ b/Engine/Source/Editor/DetailCustomizations/Private/AnimTransitionNodeDetails.cpp
@@ -106,7 +106,7 @@ void FAnimTransitionNodeDetails::CustomizeDetails( IDetailLayoutBuilder& DetailB
 		DetailBuilder.HideProperty(GET_MEMBER_NAME_CHECKED(UAnimStateTransitionNode, CustomBlendCurve));
 		DetailBuilder.HideProperty(GET_MEMBER_NAME_CHECKED(UAnimStateTransitionNode, LogicType));
 		DetailBuilder.HideProperty(GET_MEMBER_NAME_CHECKED(UAnimStateTransitionNode, PriorityOrder));
-		DetailBuilder.HideProperty(GET_MEMBER_NAME_CHECKED(UAnimStateTransitionNode, BlendProfile));
+		DetailBuilder.HideProperty(GET_MEMBER_NAME_CHECKED(UAnimStateTransitionNode, BlendProfileWrapper));
 	}
 	else
 	{
@@ -270,7 +270,7 @@ void FAnimTransitionNodeDetails::CustomizeDetails( IDetailLayoutBuilder& DetailB
 		CrossfadeCategory.AddProperty(GET_MEMBER_NAME_CHECKED(UAnimStateTransitionNode, CrossfadeDuration)).DisplayName( LOCTEXT("DurationLabel", "Duration") );
 		CrossfadeCategory.AddProperty(GET_MEMBER_NAME_CHECKED(UAnimStateTransitionNode, BlendMode)).DisplayName(LOCTEXT("ModeLabel", "Mode"));
 		CrossfadeCategory.AddProperty(GET_MEMBER_NAME_CHECKED(UAnimStateTransitionNode, CustomBlendCurve)).DisplayName(LOCTEXT("CurveLabel", "Custom Blend Curve"));
-		CrossfadeCategory.AddProperty(GET_MEMBER_NAME_CHECKED(UAnimStateTransitionNode, BlendProfile)).DisplayName(LOCTEXT("BlendProfileLabel", "Blend Profile"));
+		CrossfadeCategory.AddProperty(GET_MEMBER_NAME_CHECKED(UAnimStateTransitionNode, BlendProfileWrapper)).DisplayName(LOCTEXT("BlendProfileLabel", "Blend Profile"));
 
 		//////////////////////////////////////////////////////////////////////////
 
diff --git a/Engine/Source/Editor/Persona/Private/BlendProfileInterfaceWrapperDetailsCustomization.cpp b/Engine/Source/Editor/Persona/Private/BlendProfileInterfaceWrapperDetailsCustomization.cpp
new file mode 100644
--- /dev/null
+++ b/Engine/Source/Editor/Persona/Private/BlendProfileInterfaceWrapperDetailsCustomization.cpp
@@ -0,0 +1,376 @@
+// Copyright Epic Games, Inc. All Rights Reserved.
+
+#include "BlendProfileInterfaceWrapperDetailsCustomization.h"
+
+#include "Editor.h"
+#include "PropertyHandle.h"
+#include "DetailWidgetRow.h"
+#include "BlendProfilePicker.h"
+#include "Animation/BlendProfile.h"
+#include "ISkeletonEditorModule.h"
+#include "Kismet2/BlueprintEditorUtils.h"
+#include "Animation/AnimBlueprint.h"
+#include "Modules/ModuleManager.h"
+#include "Engine/PoseWatch.h"
+#include "Animation/BlendSpace.h"
+#include "Widgets/SWidget.h"
+#include "Widgets/DeclarativeSyntaxSupport.h"
+#include "IEditableSkeleton.h"
+#include "Widgets/SCompoundWidget.h"
+#include "EditorUndoClient.h"
+#include "BlendProfilePicker.h"
+#include "SEnumCombo.h"
+#include "PersonaModule.h"
+#include "IBlendProfilePickerExtender.h"
+
+#define LOCTEXT_NAMESPACE "BlendProfileStandaloneCustomization"
+
+class UBlendProfile;
+enum class EBlendProfileMode : uint8;
+
+class SBlendProfileInterfaceWrapperPicker : public SCompoundWidget
+{
+public:
+	SLATE_BEGIN_ARGS(SBlendProfileInterfaceWrapperPicker)
+	{
+	}
+	SLATE_END_ARGS()
+
+	struct FPickerArgs
+	{
+		FPickerArgs()
+			: SupportedBlendProfileModes(EBlendProfilePickerMode::AllModes)
+			, Outer(nullptr)
+		{
+		}
+
+		DECLARE_DELEGATE_OneParam(FOnBlendProfileChosen, UBlendProfile*)
+
+		TObjectPtr<USkeleton> Skeleton;
+
+		IBlendProfilePickerExtender::FPickerWidgetArgs::FOnBlendProfileProviderChanged OnProviderChanged;
+
+		FOnBlendProfileChosen OnBlendProfileChosen;
+
+		TSharedPtr<class IPropertyHandle> PropertyHandle;
+
+		EBlendProfilePickerMode SupportedBlendProfileModes;
+
+		TObjectPtr<UObject> Outer;
+	};
+
+	void Construct(const FArguments& InArgs, FPickerArgs ConstructArgs)
+	{
+		ISkeletonEditorModule& SkeletonEditorModule = FModuleManager::Get().LoadModuleChecked<ISkeletonEditorModule>("SkeletonEditor");
+
+		FBlendProfileInterfaceWrapper* BlendProfileInterface = nullptr;
+		{
+			void* StructAddress = nullptr;
+			if (ConstructArgs.PropertyHandle->GetValueData(/*out*/ StructAddress) == FPropertyAccess::Success)
+			{
+				BlendProfileInterface = static_cast<FBlendProfileInterfaceWrapper*>(StructAddress);
+			}
+		}
+		check(BlendProfileInterface);
+
+		CustomSource = "Skeleton";
+		CustomSourceText = LOCTEXT("Skeleton", "Skeleton");
+
+		FPersonaModule& PersonaModule = FModuleManager::GetModuleChecked<FPersonaModule>("Persona");
+		const TArray<TSharedPtr<IBlendProfilePickerExtender>> BlendProfileExtenders = PersonaModule.GetCustomBlendProfiles();
+
+		// Init CustomSources
+		{
+			CustomSources.Add("Skeleton");
+
+			for (const TSharedPtr<IBlendProfilePickerExtender>& BlendProfileExtender : BlendProfileExtenders)
+			{
+				CustomSources.Add(BlendProfileExtender->GetId());
+
+				if (!BlendProfileInterface->UsesSkeletonBlendProfile() && BlendProfileExtender->OwnsBlendProfileProvider(BlendProfileInterface->GetCustomProviderObject()))
+				{
+					CustomSource = BlendProfileExtender->GetId();
+					CustomSourceText = BlendProfileExtender->GetDisplayName();
+				}
+			}
+		}
+
+		TSharedPtr<SVerticalBox> VerticalBox;
+
+		ChildSlot
+			[
+				SAssignNew(VerticalBox, SVerticalBox)
+			];
+
+		// Add dropdown to toggle between blend profile types i.e. [Skeleton / Custom1 / ...]
+		VerticalBox->AddSlot()
+			.AutoHeight()
+			[
+				SNew(SComboButton)
+					.OnGetMenuContent_Lambda([this, ConstructArgs, BlendProfileExtenders]()
+						{
+							FMenuBuilder MenuBuilder(true, NULL);
+
+							// Action for switching to using skeleton blend profiles
+							{
+								FUIAction Action(FExecuteAction::CreateLambda([this, ConstructArgs]()
+									{
+										CustomSource = "Skeleton";
+										CustomSourceText = LOCTEXT("Skeleton", "Skeleton");
+										CustomWidgetBox->SetContent(SNullWidget::NullWidget);
+										ConstructArgs.OnProviderChanged.ExecuteIfBound(nullptr, nullptr);
+									}));
+
+								MenuBuilder.AddMenuEntry(LOCTEXT("Skeleton", "Skeleton"), FText::GetEmpty(), FSlateIcon(), Action);
+							}
+
+							for (const TSharedPtr<IBlendProfilePickerExtender>& BlendProfileExtender : BlendProfileExtenders)
+							{
+								FUIAction Action(FExecuteAction::CreateLambda([this, BlendProfileExtender, ConstructArgs]()
+									{
+										CustomSource = BlendProfileExtender->GetId();
+										CustomSourceText = BlendProfileExtender->GetDisplayName();
+
+										IBlendProfilePickerExtender::FPickerWidgetArgs Args;
+										Args.InitialSelection = nullptr;
+										Args.Outer = ConstructArgs.Outer;
+										Args.SupportedBlendProfileModes = ConstructArgs.SupportedBlendProfileModes;
+										Args.Skeleton = ConstructArgs.Skeleton;
+										Args.OnProviderChanged = IBlendProfilePickerExtender::FPickerWidgetArgs::FOnBlendProfileProviderChanged::CreateLambda([this, ConstructArgs](TObjectPtr<UObject> NewSelection, IBlendProfileProviderInterface* Interface)
+											{
+												ConstructArgs.OnProviderChanged.ExecuteIfBound(NewSelection, Interface);
+											});
+
+										CustomWidgetBox->SetContent(BlendProfileExtender->ConstructPickerWidget(Args));
+									}));
+
+								MenuBuilder.AddMenuEntry(BlendProfileExtender->GetDisplayName(), FText::GetEmpty(), FSlateIcon(), Action);
+							}
+
+							return MenuBuilder.MakeWidget();
+						})
+					.ButtonContent()
+					[
+						SNew(STextBlock)
+							.Text_Lambda([this]() { return CustomSourceText; })
+					]
+			];
+
+		// Add picker for skeleton blend profile to display when in skeleton mode
+		{
+			FBlendProfilePickerArgs SkeletonPickerArgs;
+			SkeletonPickerArgs.bAllowNew = false;
+			SkeletonPickerArgs.bAllowModify = false;
+			SkeletonPickerArgs.bAllowClear = true;
+			SkeletonPickerArgs.OnBlendProfileSelected = FOnBlendProfileSelected::CreateLambda([ConstructArgs](UBlendProfile* InBlendProfile)
+				{
+					ConstructArgs.OnBlendProfileChosen.ExecuteIfBound(InBlendProfile);
+				});
+			SkeletonPickerArgs.SupportedBlendProfileModes = ConstructArgs.SupportedBlendProfileModes;
+			if (BlendProfileInterface->UsesSkeletonBlendProfile())
+			{
+				SkeletonPickerArgs.InitialProfile = BlendProfileInterface->GetBlendProfile();
+			}
+
+			VerticalBox->AddSlot()
+				.AutoHeight()
+				[
+					SNew(SBox)
+						.Visibility_Lambda([this, ConstructArgs]() { return CustomSource == "Skeleton" ? EVisibility::Visible : EVisibility::Collapsed; })
+						[
+							SkeletonEditorModule.CreateBlendProfilePicker(ConstructArgs.Skeleton, SkeletonPickerArgs)
+						]
+				];
+		}
+
+		// Add picker for custom blend profile to display when in a custom mode, widget is mounted when custom mode is selected
+		{
+			VerticalBox->AddSlot()
+				.AutoHeight()
+				[
+					SAssignNew(CustomWidgetBox, SBox)
+						[
+							SNullWidget::NullWidget
+						]
+				];
+		}
+
+		// Setup custom widget if custom extender is being used
+		if (!BlendProfileInterface->UsesSkeletonBlendProfile())
+		{
+			for (const TSharedPtr<IBlendProfilePickerExtender>& BlendProfileExtender : BlendProfileExtenders)
+			{
+				if (BlendProfileExtender->OwnsBlendProfileProvider(BlendProfileInterface->GetCustomProviderObject()))
+				{
+					IBlendProfilePickerExtender::FPickerWidgetArgs PickerArgs;
+					PickerArgs.InitialSelection = BlendProfileInterface->GetCustomProviderObject();
+					PickerArgs.OnProviderChanged = IBlendProfilePickerExtender::FPickerWidgetArgs::FOnBlendProfileProviderChanged::CreateLambda([this, ConstructArgs](TObjectPtr<UObject> NewSelection, IBlendProfileProviderInterface* Interface)
+						{
+							ConstructArgs.OnProviderChanged.ExecuteIfBound(NewSelection, Interface);
+						});
+					PickerArgs.Outer = ConstructArgs.Outer;
+					PickerArgs.SupportedBlendProfileModes = ConstructArgs.SupportedBlendProfileModes;
+					PickerArgs.Skeleton = ConstructArgs.Skeleton;
+
+					CustomWidgetBox->SetContent(BlendProfileExtender->ConstructPickerWidget(PickerArgs));
+				}
+			}
+		}
+	}
+
+private:
+	TSharedPtr<SBox> CustomWidgetBox;
+	TArray<FName> CustomSources;
+
+	FName CustomSource = NAME_None;
+	FText CustomSourceText;
+};
+
+
+TSharedRef<IPropertyTypeCustomization> FBlendProfileInterfaceWrapperCustomization::MakeInstance()
+{
+	return MakeShareable(new FBlendProfileInterfaceWrapperCustomization);
+}
+
+void FBlendProfileInterfaceWrapperCustomization::CustomizeHeader(TSharedRef<class IPropertyHandle> InStructPropertyHandle, class FDetailWidgetRow& HeaderRow, IPropertyTypeCustomizationUtils& StructCustomizationUtils)
+{
+	TSharedRef<SWidget> ValueCustomWidget = SNullWidget::NullWidget;
+
+	TArray<UObject*> OuterObjects;
+	InStructPropertyHandle->GetOuterObjects(OuterObjects);
+
+	if (OuterObjects.Num() > 0)
+	{
+		// try to get skeleton from first outer
+		if (USkeleton* TargetSkeleton = GetSkeletonFromOuter(OuterObjects[0]))
+		{
+			TWeakPtr<IPropertyHandle> PropertyPtr(InStructPropertyHandle);
+
+			const bool bUseAsBlendMask = InStructPropertyHandle->GetBoolMetaData("UseAsBlendMask");
+			const bool bUseAsBlendProfile = InStructPropertyHandle->GetBoolMetaData("UseAsBlendProfile");
+
+			// If no mode is defined, show both.
+			EBlendProfilePickerMode SupportedBlendProfileModes = (bUseAsBlendMask || bUseAsBlendProfile) ? EBlendProfilePickerMode(0) : EBlendProfilePickerMode::AllModes;
+
+			if (bUseAsBlendProfile)
+			{
+				SupportedBlendProfileModes |= EBlendProfilePickerMode::BlendProfile;
+			}
+			if (bUseAsBlendMask)
+			{
+				SupportedBlendProfileModes |= EBlendProfilePickerMode::BlendMask;
+			}
+
+
+			/*
+			FBlendProfilePickerArgs Args;
+			Args.bAllowNew = false;
+			Args.bAllowModify = false;
+			Args.bAllowClear = true;
+			//Args.OnBlendProfileSelected = FOnBlendProfileSelected::CreateSP(this, &FBlendProfileInterfaceWrapperCustomization::OnBlendProfileChanged, PropertyPtr);
+			Args.InitialProfile = CurrentProfile;
+			Args.SupportedBlendProfileModes = SupportedBlendProfileModes;
+			Args.PropertyHandle = InStructPropertyHandle;
+			*/
+
+			//ISkeletonEditorModule& SkeletonEditorModule = FModuleManager::Get().LoadModuleChecked<ISkeletonEditorModule>("SkeletonEditor");
+			//ValueCustomWidget = SkeletonEditorModule.CreateBlendProfilePicker(TargetSkeleton, Args);
+
+			SBlendProfileInterfaceWrapperPicker::FPickerArgs WrapperArgs;
+			WrapperArgs.Skeleton = TargetSkeleton;
+			WrapperArgs.SupportedBlendProfileModes = SupportedBlendProfileModes;
+			WrapperArgs.PropertyHandle = InStructPropertyHandle;
+			WrapperArgs.Outer = OuterObjects[0];
+			WrapperArgs.OnProviderChanged = IBlendProfilePickerExtender::FPickerWidgetArgs::FOnBlendProfileProviderChanged::CreateSP(this, &FBlendProfileInterfaceWrapperCustomization::OnBlendProfileProviderChanged, PropertyPtr, OuterObjects[0]);
+			WrapperArgs.OnBlendProfileChosen = SBlendProfileInterfaceWrapperPicker::FPickerArgs::FOnBlendProfileChosen::CreateSP(this, &FBlendProfileInterfaceWrapperCustomization::OnBlendProfileChanged, PropertyPtr, OuterObjects[0]);
+
+			ValueCustomWidget = SNew(SBlendProfileInterfaceWrapperPicker, WrapperArgs);
+		}
+	}
+
+	HeaderRow.NameContent()
+		[
+			InStructPropertyHandle->CreatePropertyNameWidget()
+		]
+		.ValueContent()
+		.MinDesiredWidth(125.f)
+		.MaxDesiredWidth(400.f) //Slightly wider since expected names are a bit longer if users use BlendProfileModes as suffix
+		[
+			ValueCustomWidget
+		];
+}
+
+void FBlendProfileInterfaceWrapperCustomization::OnBlendProfileChanged(UBlendProfile* NewProfile, TWeakPtr<IPropertyHandle> WeakPropertyHandle, UObject* Outer)
+{
+	if (!GIsTransacting)
+	{
+		if (TSharedPtr<IPropertyHandle> PropertyHandle = WeakPropertyHandle.Pin())
+		{
+			FBlendProfileInterfaceWrapper* BlendProfileInterface = nullptr;
+			{
+				void* StructAddress = nullptr;
+				if (PropertyHandle->GetValueData(/*out*/ StructAddress) == FPropertyAccess::Success)
+				{
+					BlendProfileInterface = static_cast<FBlendProfileInterfaceWrapper*>(StructAddress);
+				}
+			}
+
+			BlendProfileInterface->SetSkeletonBlendProfile(NewProfile);
+		}
+	}
+}
+
+void FBlendProfileInterfaceWrapperCustomization::OnBlendProfileProviderChanged(TObjectPtr<UObject> NewProfile, IBlendProfileProviderInterface* Interface, TWeakPtr<IPropertyHandle> WeakPropertyHandle, UObject* Outer)
+{
+	if (!GIsTransacting)
+	{
+		if (TSharedPtr<IPropertyHandle> PropertyHandle = WeakPropertyHandle.Pin())
+		{
+			FBlendProfileInterfaceWrapper* BlendProfileInterface = nullptr;
+			{
+				void* StructAddress = nullptr;
+				if (PropertyHandle->GetValueData(/*out*/ StructAddress) == FPropertyAccess::Success)
+				{
+					BlendProfileInterface = static_cast<FBlendProfileInterfaceWrapper*>(StructAddress);
+				}
+			}
+
+			BlendProfileInterface->SetBlendProfileProvider(NewProfile, Interface, Outer);
+		}
+	}
+}
+
+USkeleton* FBlendProfileInterfaceWrapperCustomization::GetSkeletonFromOuter(const UObject* Outer)
+{
+	const UAnimBlueprint* AnimBlueprint = nullptr;
+	if (const UBlendSpace* BlendSpace = Cast<UBlendSpace>(Outer))
+	{
+		// Check for blend space graph nodes
+		if (!BlendSpace->IsAsset())
+		{
+			AnimBlueprint = BlendSpace->GetTypedOuter<UAnimBlueprint>();
+		}
+	}
+	if (const UEdGraphNode* OuterEdGraphNode = Cast<UEdGraphNode>(Outer))
+	{
+		AnimBlueprint = Cast<UAnimBlueprint>(FBlueprintEditorUtils::FindBlueprintForNode(OuterEdGraphNode));
+	}
+	else if (const UEdGraph* OuterEdGraph = Cast<UEdGraph>(Outer))
+	{
+		AnimBlueprint = Cast<UAnimBlueprint>(FBlueprintEditorUtils::FindBlueprintForGraph(OuterEdGraph));
+	}
+
+	if (AnimBlueprint)
+	{
+		return AnimBlueprint->TargetSkeleton;
+	}
+
+	if (const UAnimationAsset* OuterAnimAsset = Cast<UAnimationAsset>(Outer))
+	{
+		return OuterAnimAsset->GetSkeleton();
+	}
+
+	return nullptr;
+}
+
+#undef LOCTEXT_NAMESPACE
\ No newline at end of file
diff --git a/Engine/Source/Editor/Persona/Private/BlendProfileInterfaceWrapperDetailsCustomization.h b/Engine/Source/Editor/Persona/Private/BlendProfileInterfaceWrapperDetailsCustomization.h
new file mode 100644
--- /dev/null
+++ b/Engine/Source/Editor/Persona/Private/BlendProfileInterfaceWrapperDetailsCustomization.h
@@ -0,0 +1,26 @@
+// Copyright Epic Games, Inc. All Rights Reserved.
+
+#pragma once
+
+#include "CoreMinimal.h"
+#include "IPropertyTypeCustomization.h"
+
+class IBlendProfileProviderInterface;
+class UBlendProfile;
+
+class FBlendProfileInterfaceWrapperCustomization : public IPropertyTypeCustomization
+{
+public:
+	static TSharedRef<IPropertyTypeCustomization> MakeInstance();
+
+	virtual void CustomizeHeader(TSharedRef<class IPropertyHandle> StructPropertyHandle, class FDetailWidgetRow& HeaderRow, IPropertyTypeCustomizationUtils& StructCustomizationUtils) override;
+
+	virtual void CustomizeChildren(TSharedRef<IPropertyHandle> InStructPropertyHandle, class IDetailChildrenBuilder& ChildBuilder, IPropertyTypeCustomizationUtils& StructCustomizationUtils) override {}
+
+private:
+	void OnBlendProfileChanged(UBlendProfile* NewProfile, TWeakPtr<IPropertyHandle> WeakPropertyHandle, UObject* Outer);
+
+	void OnBlendProfileProviderChanged(TObjectPtr<UObject> NewProfile, IBlendProfileProviderInterface* Interface, TWeakPtr<IPropertyHandle> WeakPropertyHandle, UObject* Outer);
+
+	class USkeleton* GetSkeletonFromOuter(const UObject* Outer);
+};
diff --git a/Engine/Source/Editor/Persona/Private/PersonaModule.cpp b/Engine/Source/Editor/Persona/Private/PersonaModule.cpp
--- a/Engine/Source/Editor/Persona/Private/PersonaModule.cpp
+++ b/Engine/Source/Editor/Persona/Private/PersonaModule.cpp
@@ -21,6 +21,7 @@
 #include "AnimNotifyDetails.h"
 #include "AnimGraphNodeDetails.h"
 #include "BlendProfileCustomization.h"
+#include "BlendProfileInterfaceWrapperDetailsCustomization.h"
 #include "Engine/SkinnedAssetCommon.h"
 #include "AnimInstanceDetails.h"
 #include "IEditableSkeleton.h"
@@ -134,6 +135,7 @@ void FPersonaModule::StartupModule()
 		PropertyModule.RegisterCustomClassLayout("BlendSpace", FOnGetDetailCustomizationInstance::CreateStatic(&FBlendSpaceDetails::MakeInstance));	
 
 		PropertyModule.RegisterCustomPropertyTypeLayout("BlendProfile", FOnGetPropertyTypeCustomizationInstance::CreateStatic(&FBlendProfileCustomization::MakeInstance));
+		PropertyModule.RegisterCustomPropertyTypeLayout("BlendProfileInterfaceWrapper", FOnGetPropertyTypeCustomizationInstance::CreateStatic(&FBlendProfileInterfaceWrapperCustomization::MakeInstance));
 		PropertyModule.RegisterCustomPropertyTypeLayout( "InputScaleBias", FOnGetPropertyTypeCustomizationInstance::CreateStatic( &FInputScaleBiasCustomization::MakeInstance ) );
 		PropertyModule.RegisterCustomPropertyTypeLayout( "BoneReference", FOnGetPropertyTypeCustomizationInstance::CreateStatic( &FBoneReferenceCustomization::MakeInstance ) );
 		PropertyModule.RegisterCustomPropertyTypeLayout("BoneSocketTarget", FOnGetPropertyTypeCustomizationInstance::CreateStatic(&FBoneSocketTargetCustomization::MakeInstance));
@@ -1814,4 +1816,22 @@ void FPersonaModule::RegisterToolMenuExtensions()
 	}));
 }
 
+void FPersonaModule::RegisterBlendProfilePickerExtender(TSharedPtr<IBlendProfilePickerExtender> Extender)
+{
+	CustomBlendProfiles.Add(Extender);
+}
+
+void FPersonaModule::UnregisterBlendProfilePickerExtender(const FName ExtenderId)
+{
+	CustomBlendProfiles.RemoveAll([ExtenderId](const TSharedPtr<IBlendProfilePickerExtender>& Extender)
+		{
+			return Extender->GetId() == ExtenderId;
+		});
+}
+
+const TArray<TSharedPtr<IBlendProfilePickerExtender>>& FPersonaModule::GetCustomBlendProfiles() const
+{
+	return CustomBlendProfiles;
+}
+
 #undef LOCTEXT_NAMESPACE
diff --git a/Engine/Source/Editor/Persona/Public/IBlendProfilePickerExtender.h b/Engine/Source/Editor/Persona/Public/IBlendProfilePickerExtender.h
new file mode 100644
--- /dev/null
+++ b/Engine/Source/Editor/Persona/Public/IBlendProfilePickerExtender.h
@@ -0,0 +1,54 @@
+// Copyright Epic Games, Inc. All Rights Reserved.
+
+#pragma once
+
+#include "CoreMinimal.h"
+#include "BlendProfilePicker.h"
+#include "Delegates/DelegateCombinations.h"
+#include "UObject/ObjectPtr.h"
+
+class IBlendProfileProviderInterface;
+class SWidget;
+class USkeleton;
+
+class IBlendProfilePickerExtender
+{
+public:
+	struct FPickerWidgetArgs
+	{
+		FPickerWidgetArgs()
+			: Outer(nullptr)
+			, SupportedBlendProfileModes(EBlendProfilePickerMode::AllModes)
+		{
+		}
+
+		DECLARE_DELEGATE_TwoParams(FOnBlendProfileProviderChanged, TObjectPtr<UObject> /* NewProfileProvider */, IBlendProfileProviderInterface*);
+
+		// Should be fired when the blend profile provider object has changed
+		FOnBlendProfileProviderChanged OnProviderChanged;
+
+		// The initially selected provider object
+		TObjectPtr<UObject> InitialSelection;
+
+		// The outer to use for constructing new provider objects
+		TObjectPtr<UObject> Outer;
+
+		// Restrict which types of blend profiles are displayed in the picker
+		EBlendProfilePickerMode SupportedBlendProfileModes;
+
+		// Optional skeleton to restrict shown blend profiles relating to a particular skeleton
+		TObjectPtr<USkeleton> Skeleton;
+	};
+
+	// Returns an identifier to match extender instances
+	virtual FName GetId() const = 0;
+
+	// Text to display in the details panel when choosing blend profiles of this type
+	virtual FText GetDisplayName() const = 0;
+
+	// Construct the picker widget for choosing blend profiles of this type
+	virtual TSharedRef<SWidget> ConstructPickerWidget(const FPickerWidgetArgs& InWidgetArgs) const = 0;
+
+	// Return true if the provided object is the matching type for this extender
+	virtual bool OwnsBlendProfileProvider(const TObjectPtr<const UObject> InProviderObject) const = 0;
+};
\ No newline at end of file
diff --git a/Engine/Source/Editor/Persona/Public/PersonaModule.h b/Engine/Source/Editor/Persona/Public/PersonaModule.h
--- a/Engine/Source/Editor/Persona/Public/PersonaModule.h
+++ b/Engine/Source/Editor/Persona/Public/PersonaModule.h
@@ -11,6 +11,7 @@
 #include "PersonaDelegates.h"
 #include "Factories/FbxImportUI.h"
 #include "SPersonaToolBox.h"
+#include "IBlendProfilePickerExtender.h"
 
 class FBlueprintEditor;
 class IDetailsView;
@@ -561,7 +562,17 @@ public:
 	// Get the version information for the notify host asset registry
 	virtual uint32 GetNotifyHostAssetVersion() const;
 
+public:
+	// Register a new blend provider provider in the editor picker for FBlendProfileInterfaceWrapper
+	PERSONA_API void RegisterBlendProfilePickerExtender(TSharedPtr<IBlendProfilePickerExtender> Extender);
+
+	PERSONA_API void UnregisterBlendProfilePickerExtender(const FName ExtenderId);
+
+	const TArray<TSharedPtr<IBlendProfilePickerExtender>>& GetCustomBlendProfiles() const;
+
 private:
+	TArray<TSharedPtr<IBlendProfilePickerExtender>> CustomBlendProfiles;
+
 	/** When a new anim notify blueprint is created, this will handle post creation work such as adding non-event default nodes */
 	void HandleNewAnimNotifyBlueprintCreated(UBlueprint* InBlueprint);
 
diff --git a/Engine/Source/Runtime/Engine/Classes/Animation/BlendProfile.h b/Engine/Source/Runtime/Engine/Classes/Animation/BlendProfile.h
--- a/Engine/Source/Runtime/Engine/Classes/Animation/BlendProfile.h
+++ b/Engine/Source/Runtime/Engine/Classes/Animation/BlendProfile.h
@@ -59,6 +59,65 @@ public:
 	virtual int32 GetNumBlendEntries() const = 0;
 
 	virtual TObjectPtr<USkeleton> GetSkeleton() const = 0;
+
+	virtual EBlendProfileMode GetMode() const = 0;
+
+	void UpdateBoneWeights(FBlendSampleData& InOutCurrentData, const FAlphaBlend& BlendInfo, float BlendStartAlpha, float MainWeight, bool bInverse = false) const;
+};
+
+class UBlendProfile;
+
+UINTERFACE(MinimalAPI, Blueprintable)
+class UBlendProfileProviderInterface : public UInterface
+{
+	GENERATED_BODY()
+};
+
+/**
+ * Interface to be implemented by custom data structures that can create UBlendProfiles
+ */
+class IBlendProfileProviderInterface
+{
+	GENERATED_BODY()
+
+public:
+	// Convert the custom internal data struct into a regular blend profile
+	virtual void ConstructBlendProfile(const TObjectPtr<UBlendProfile> OutBlendProfile) const = 0;
+};
+
+/**
+ * Helper class for storing either USkeleton blend profiles and custom blend profile implementations.
+ * Has a custom details panel for switching between blend profile providers.
+ */
+USTRUCT()
+struct ENGINE_API FBlendProfileInterfaceWrapper
+{
+	GENERATED_BODY()
+
+public:
+	void UpdateCachedBlendProfile();
+
+	const TObjectPtr<UBlendProfile> GetBlendProfile() const { return BlendProfile; }
+
+	TObjectPtr<UObject> GetCustomProviderObject() const { return BlendProfileProvider.GetObject(); }
+
+	void SetBlendProfileProvider(TObjectPtr<UObject> Provider, IBlendProfileProviderInterface* ProviderInterface, UObject* Outer);
+
+	void SetSkeletonBlendProfile(TObjectPtr<UBlendProfile> InBlendProfile);
+
+	bool UsesSkeletonBlendProfile() const { return BlendProfileProvider.GetObject() == nullptr; }
+
+private:
+	UPROPERTY()
+	bool bIsSkeletonBlendProfile = true;
+
+	// The object responsible for constructing the blend profile
+	UPROPERTY()
+	TScriptInterface<IBlendProfileProviderInterface> BlendProfileProvider;
+
+	// The actual blend profile to use
+	UPROPERTY()
+	TObjectPtr<UBlendProfile> BlendProfile;
 };
 
 //////////////////////////////////////////////////////////////////////////
@@ -206,10 +265,13 @@ public:
 
 	EBlendProfileMode GetMode() const { return Mode; }
 
-private:
 	/** Sets the skeleton this blend profile is used with */
 	ENGINE_API void SetSkeleton(USkeleton* InSkeleton);
 
+	// Empties the bone enties array
+	void ClearEntries();
+
+private:
 	/** Set the blend scale for a single bone (ignore children) */
 	ENGINE_API void SetSingleBoneBlendScale(int32 InBoneIdx, float InScale, bool bCreate = false);
 
diff --git a/Engine/Source/Runtime/Engine/Private/Animation/AnimNode_StateMachine.cpp b/Engine/Source/Runtime/Engine/Private/Animation/AnimNode_StateMachine.cpp
--- a/Engine/Source/Runtime/Engine/Private/Animation/AnimNode_StateMachine.cpp
+++ b/Engine/Source/Runtime/Engine/Private/Animation/AnimNode_StateMachine.cpp
@@ -1032,7 +1032,7 @@ void FAnimNode_StateMachine::EvaluateTransitionStandardBlendInternal(FPoseContex
 
 		// If we have some skeleton remapping and the source data comes from another skeleton.
 		// This is a slightly slower path, so we made two branches, one with remapping and one without.
-		const FSkeletonRemapping& Remapping = UE::Anim::FSkeletonRemappingRegistry::Get().GetRemapping(Transition.BlendProfile->OwningSkeleton, Output.AnimInstanceProxy->GetSkeleton());
+		const FSkeletonRemapping& Remapping = UE::Anim::FSkeletonRemappingRegistry::Get().GetRemapping(Transition.BlendProfile->GetSkeleton(), Output.AnimInstanceProxy->GetSkeleton());
 		if (Remapping.IsValid())
 		{
 			for (const FCompactPoseBoneIndex TargetBoneIndex : Output.Pose.ForEachBoneIndex())
diff --git a/Engine/Source/Runtime/Engine/Private/Animation/BlendProfile.cpp b/Engine/Source/Runtime/Engine/Private/Animation/BlendProfile.cpp
--- a/Engine/Source/Runtime/Engine/Private/Animation/BlendProfile.cpp
+++ b/Engine/Source/Runtime/Engine/Private/Animation/BlendProfile.cpp
@@ -8,6 +8,60 @@
 
 #include UE_INLINE_GENERATED_CPP_BY_NAME(BlendProfile)
 
+void IBlendProfileInterface::UpdateBoneWeights(FBlendSampleData& InOutCurrentData, const FAlphaBlend& BlendInfo, float BlendStartAlpha, float MainWeight, bool bInverse) const
+{
+	const EBlendProfileMode Mode = GetMode();
+
+	for (int32 PerBoneIndex = 0; PerBoneIndex < InOutCurrentData.PerBoneBlendData.Num(); ++PerBoneIndex)
+	{
+		InOutCurrentData.PerBoneBlendData[PerBoneIndex] =  UBlendProfile::CalculateBoneWeight(GetBoneBlendScale(PerBoneIndex), Mode, BlendInfo, BlendStartAlpha, MainWeight, bInverse);
+	}
+}
+
+void FBlendProfileInterfaceWrapper::UpdateCachedBlendProfile()
+{
+	if (!bIsSkeletonBlendProfile && BlendProfileProvider.GetObject() && BlendProfile)
+	{
+		if (BlendProfileProvider.GetObject() && BlendProfile)
+		{
+			BlendProfileProvider->ConstructBlendProfile(BlendProfile);
+		}
+		else
+		{
+			BlendProfile = nullptr;
+		}
+	}
+}
+
+void FBlendProfileInterfaceWrapper::SetBlendProfileProvider(TObjectPtr<UObject> Provider, IBlendProfileProviderInterface* ProviderInterface, UObject* Outer)
+{
+	BlendProfileProvider.SetObject(Provider);
+	BlendProfileProvider.SetInterface(ProviderInterface);
+
+	if (BlendProfileProvider)
+	{
+		if (!BlendProfile || BlendProfile->GetOuter() != Outer)
+		{
+			BlendProfile = NewObject<UBlendProfile>(Outer);
+		}
+
+		BlendProfile->ClearEntries();
+		BlendProfileProvider->ConstructBlendProfile(BlendProfile);
+	}
+	
+	bIsSkeletonBlendProfile = false;
+}
+
+void FBlendProfileInterfaceWrapper::SetSkeletonBlendProfile(TObjectPtr<UBlendProfile> InBlendProfile)
+{
+	BlendProfileProvider.SetObject(nullptr);
+	BlendProfileProvider.SetInterface(nullptr);
+
+	BlendProfile = InBlendProfile;
+	
+	bIsSkeletonBlendProfile = true;
+}
+
 UBlendProfile::UBlendProfile()
 	: OwningSkeleton(nullptr)
 	, Mode(EBlendProfileMode::WeightFactor)
@@ -439,3 +493,7 @@ void UBlendProfile::UpdateBoneWeights(FBlendSampleData& InOutCurrentData, const
 	}
 }
 
+void UBlendProfile::ClearEntries()
+{
+	ProfileEntries.Empty();
+}
\ No newline at end of file
