Anim Timeline bug fix: Left clicking on notify track after marquee selecting pauses animation playback

diff --git a/Engine/Source/Editor/Persona/Private/AnimTimeline/AnimTimeSliderController.cpp b/Engine/Source/Editor/Persona/Private/AnimTimeline/AnimTimeSliderController.cpp
--- a/Engine/Source/Editor/Persona/Private/AnimTimeline/AnimTimeSliderController.cpp
+++ b/Engine/Source/Editor/Persona/Private/AnimTimeline/AnimTimeSliderController.cpp
@@ -567,79 +567,112 @@ FReply FAnimTimeSliderController::OnMouseButtonUp( SWidget& WidgetOwner, const F
 	}
 	else if ( bHandleLeftMouseButton )
 	{
-		if (MouseDragType == DRAG_PLAYBACK_START)
+		switch (MouseDragType)
 		{
-			TimeSliderArgs.OnPlaybackRangeEndDrag.ExecuteIfBound();
-		}
-		else if (MouseDragType == DRAG_PLAYBACK_END)
-		{
-			TimeSliderArgs.OnPlaybackRangeEndDrag.ExecuteIfBound();
-		}
-		else if (MouseDragType == DRAG_SELECTION_START)
-		{
-			TimeSliderArgs.OnSelectionRangeEndDrag.ExecuteIfBound();
-		}
-		else if (MouseDragType == DRAG_SELECTION_END)
-		{
-			TimeSliderArgs.OnSelectionRangeEndDrag.ExecuteIfBound();
-		}
-		else if (MouseDragType == DRAG_SETTING_RANGE)
-		{
-			// Zooming
-			FFrameTime MouseDownStart = ComputeFrameTimeFromMouse(MyGeometry, MouseDownPosition[0], RangeToScreen);
-
-			const bool bCanZoomIn  = MouseTime > MouseDownStart;
-			const bool bCanZoomOut = ViewRangeStack.Num() > 0;
-			if (bCanZoomIn || bCanZoomOut)
+		case DRAG_PLAYBACK_START:
 			{
-				TRange<double> ViewRange = TimeSliderArgs.ViewRange.Get();
-				if (!bCanZoomIn)
-				{
-					ViewRange = ViewRangeStack.Pop();
-				}
+				TimeSliderArgs.OnPlaybackRangeEndDrag.ExecuteIfBound();
+				
+				break;
+			}
+		case DRAG_PLAYBACK_END:
+			{
+				TimeSliderArgs.OnPlaybackRangeEndDrag.ExecuteIfBound();
+				
+				break;
+			}
+		case DRAG_SELECTION_START:
+			{
+				TimeSliderArgs.OnSelectionRangeEndDrag.ExecuteIfBound();
+				
+				break;
+			}
+		case DRAG_SELECTION_END:
+			{
+				TimeSliderArgs.OnSelectionRangeEndDrag.ExecuteIfBound();
+				
+				break;
+			}
+		case DRAG_SETTING_RANGE:
+			{
+				// Zooming
+				FFrameTime MouseDownStart = ComputeFrameTimeFromMouse(MyGeometry, MouseDownPosition[0], RangeToScreen);
 
-				if (bCanZoomIn)
+				const bool bCanZoomIn  = MouseTime > MouseDownStart;
+				const bool bCanZoomOut = ViewRangeStack.Num() > 0;
+				if (bCanZoomIn || bCanZoomOut)
 				{
-					// push the current value onto the stack
-					ViewRangeStack.Add(ViewRange);
-
-					ViewRange = TRange<double>(MouseDownStart.FrameNumber / GetTickResolution(), MouseTime.FrameNumber / GetTickResolution());
+					TRange<double> ViewRange = TimeSliderArgs.ViewRange.Get();
+					if (!bCanZoomIn)
+					{
+						ViewRange = ViewRangeStack.Pop();
+					}
+
+					if (bCanZoomIn)
+					{
+						// push the current value onto the stack
+						ViewRangeStack.Add(ViewRange);
+
+						ViewRange = TRange<double>(MouseDownStart.FrameNumber / GetTickResolution(), MouseTime.FrameNumber / GetTickResolution());
+					}
+				
+					TimeSliderArgs.OnViewRangeChanged.ExecuteIfBound(ViewRange, EViewRangeInterpolation::Immediate);
+					if( !TimeSliderArgs.ViewRange.IsBound() )
+					{
+						// The output is not bound to a delegate so we'll manage the value ourselves
+						TimeSliderArgs.ViewRange.Set(ViewRange);
+					}
 				}
 				
-				TimeSliderArgs.OnViewRangeChanged.ExecuteIfBound(ViewRange, EViewRangeInterpolation::Immediate);
-				if( !TimeSliderArgs.ViewRange.IsBound() )
+				break;
+			}
+		case DRAG_TIME:
+			{
+				double Time = (double)MouseTime.AsDecimal() / GetTickResolution().AsDecimal();
+
+				if(!MouseEvent.IsControlDown())
 				{
-					// The output is not bound to a delegate so we'll manage the value ourselves
-					TimeSliderArgs.ViewRange.Set(ViewRange);
+					const double SnapMargin = (ScrubConstants::SnapMarginInPixels / static_cast<double>(RangeToScreen.PixelsPerInput));
+					WeakModel.Pin()->Snap(Time, SnapMargin, { FName("MontageSection") });
 				}
-			}
-		}
-		else if(MouseDragType == DRAG_TIME)
-		{
-			double Time = (double)MouseTime.AsDecimal() / GetTickResolution().AsDecimal();
 
-			if(!MouseEvent.IsControlDown())
+				SetEditableTime(DraggedTimeIndex, static_cast<float>(Time), false);
+				DraggedTimeIndex = INDEX_NONE;
+				
+				break;
+			}
+		case DRAG_SCRUBBING_TIME:
 			{
-				const double SnapMargin = (ScrubConstants::SnapMarginInPixels / static_cast<double>(RangeToScreen.PixelsPerInput));
-				WeakModel.Pin()->Snap(Time, SnapMargin, { FName("MontageSection") });
+				TimeSliderArgs.OnEndScrubberMovement.ExecuteIfBound();
+				CommitScrubPosition(MouseTime, /*bIsScrubbing=*/false);
+				
+				break;
 			}
-
-			SetEditableTime(DraggedTimeIndex, static_cast<float>(Time), false);
-			DraggedTimeIndex = INDEX_NONE;
-		}
-		else
-		{
-			TimeSliderArgs.OnEndScrubberMovement.ExecuteIfBound();
-
-			CommitScrubPosition(MouseTime, /*bIsScrubbing=*/false);
+		case DRAG_NONE:
+			{
+				// DRAG_NONE is an expected case because Controller's OnMouseButtonUp 
+				// is called from SSequenceTimeSliders's OnMouseButtonUp
+				break;
+			}
+		default:
+			{
+				// if default case is called, then MouseDragType was not properly set
+				checkNoEntry();
+				
+				break;
+			}
+			
 		}
 
 		MouseDragType = DRAG_NONE;
 		DistanceDragged = 0.f;
-
+		
 		return FReply::Handled().ReleaseMouseCapture();
 	}
 
+	// reset distance dragged regardless if OnMouseButtonUp was handled
+	DistanceDragged = 0.f;
+	
 	return FReply::Unhandled();
 }
 
@@ -676,6 +709,13 @@ FReply FAnimTimeSliderController::OnMouseMove( SWidget& WidgetOwner, const FGeom
 			ClampViewRange(NewViewOutputMin, NewViewOutputMax);
 			SetViewRange(NewViewOutputMin, NewViewOutputMax, EViewRangeInterpolation::Immediate);
 		}
+		
+		if ( DistanceDragged != 0.f )
+		{
+			DistanceDragged = 0.f;
+		
+			return FReply::Handled().CaptureMouse(WidgetOwner.AsShared());
+		}
 	}
 	else if (bHandleLeftMouseButton)
 	{
@@ -788,13 +828,17 @@ FReply FAnimTimeSliderController::OnMouseMove( SWidget& WidgetOwner, const FGeom
 				SetEditableTime(DraggedTimeIndex, static_cast<float>(Time), true);
 			}
 		}
+		
+		// only capture mouse if MouseDragType is valid
+		if ( MouseDragType != DRAG_NONE && DistanceDragged != 0.f )
+		{
+			DistanceDragged = 0.f;
+		
+			return FReply::Handled().CaptureMouse(WidgetOwner.AsShared());
+		}
 	}
-
-	if ( DistanceDragged != 0.f && (bHandleLeftMouseButton || bHandleRightMouseButton) )
-	{
-		return FReply::Handled().CaptureMouse(WidgetOwner.AsShared());
-	}
-
+	
+	DistanceDragged = 0.f;
 
 	return FReply::Handled();
 }
