Add missing nullptr check in persona asset family blueprint extension

diff --git a/Engine/Source/Editor/Persona/Private/PersonaAssetFamily.cpp b/Engine/Source/Editor/Persona/Private/PersonaAssetFamily.cpp
--- a/Engine/Source/Editor/Persona/Private/PersonaAssetFamily.cpp
+++ b/Engine/Source/Editor/Persona/Private/PersonaAssetFamily.cpp
@@ -546,12 +546,12 @@ void UAnimationEditorsAssetFamilyExtension_AnimBlueprintAsset::FindCounterpartAs
 
 	check(AnimBlueprint->BlueprintType == BPTYPE_Interface || AnimBlueprint->bIsTemplate || AnimBlueprint->TargetSkeleton != nullptr);
 
-	if (AssetFamilyInterface.IsAssetTypeInFamilyAndUnassigned<USkeletalMesh>())
+	if (AssetFamilyInterface.IsAssetTypeInFamilyAndUnassigned<USkeletalMesh>() && AnimBlueprint->TargetSkeleton)
 	{
-		AssetFamilyInterface.SetAssetOfType<USkeletalMesh>( AnimBlueprint->TargetSkeleton->GetPreviewMesh());
+		AssetFamilyInterface.SetAssetOfType<USkeletalMesh>(AnimBlueprint->TargetSkeleton->GetPreviewMesh());
 	}
 
-	if (AssetFamilyInterface.IsAssetTypeInFamilyAndUnassigned<USkeletalMesh>())
+	if (AssetFamilyInterface.IsAssetTypeInFamilyAndUnassigned<USkeletalMesh>() && AnimBlueprint->TargetSkeleton)
 	{
 		AssetFamilyInterface.SetAssetOfType<USkeletalMesh>(AnimBlueprint->TargetSkeleton->FindCompatibleMesh());
 	}
