Add PreviewInstance null-check (can happen if there is no valid preview mesh)

diff --git a/Engine/Source/Editor/Persona/Private/AnimTimeline/AnimTimelineTrack_Curve.cpp b/Engine/Source/Editor/Persona/Private/AnimTimeline/AnimTimelineTrack_Curve.cpp
--- a/Engine/Source/Editor/Persona/Private/AnimTimeline/AnimTimelineTrack_Curve.cpp
+++ b/Engine/Source/Editor/Persona/Private/AnimTimeline/AnimTimelineTrack_Curve.cpp
@@ -230,7 +230,10 @@ TSharedRef<SWidget> FAnimTimelineTrack_Curve::GenerateContainerWidgetForTimeline
 		TimelineWidget->SetOnMouseButtonUp(FPointerEventHandler::CreateSP(this, &FAnimTimelineTrack_Curve::HandleMouseButtonUp));
 	}
 
-	GetModel()->GetPreviewScene()->GetPreviewMeshComponent()->PreviewInstance->AddKeyCompleteDelegate(FSimpleDelegate::CreateSP(this, &FAnimTimelineTrack_Curve::HandleCurveChanged));
+	if (GetModel()->GetPreviewScene()->GetPreviewMeshComponent()->PreviewInstance != nullptr)
+	{		
+		GetModel()->GetPreviewScene()->GetPreviewMeshComponent()->PreviewInstance->AddKeyCompleteDelegate(FSimpleDelegate::CreateSP(this, &FAnimTimelineTrack_Curve::HandleCurveChanged));
+	}
 
 	return TimelineWidget;
 }
