Curve editor bug fix : Curve Color remain the same after its changing (random or exact)

diff --git a/Engine/Source/Editor/CurveEditor/Private/CurveEditor.cpp b/Engine/Source/Editor/CurveEditor/Private/CurveEditor.cpp
--- a/Engine/Source/Editor/CurveEditor/Private/CurveEditor.cpp
+++ b/Engine/Source/Editor/CurveEditor/Private/CurveEditor.cpp
@@ -2283,15 +2283,8 @@ void FCurveEditor::SetRandomCurveColorsForSelected()
 	{
 		if (FCurveModel* Curve = FindCurve(CurveModelID))
 		{
-			UObject* Object = nullptr;
-			FString Name;
-			Curve->GetCurveColorObjectAndName(&Object, Name);
-			if (Object)
-			{
-				FLinearColor Color = UCurveEditorSettings::GetNextRandomColor();
-				Settings->SetCustomColor(Object->GetClass(), Name, Color);
-				Curve->SetColor(Color);
-			}
+			FLinearColor Color = UCurveEditorSettings::GetNextRandomColor();
+			Curve->SetColor(Color);
 		}
 	}
 }
@@ -2317,14 +2310,7 @@ void FCurveEditor::SetCurveColorsForSelected()
 			{
 				if (FCurveModel* Curve = Self->FindCurve(CurveModelID))
 				{
-					UObject* Object = nullptr;
-					FString Name;
-					Curve->GetCurveColorObjectAndName(&Object, Name);
-					if (Object)
-					{
-						Self->Settings->SetCustomColor(Object->GetClass(), Name, NewColor);
-						Curve->SetColor(NewColor);
-					}
+					Curve->SetColor(NewColor);
 				}
 			}
 		}
diff --git a/Engine/Source/Editor/CurveEditor/Public/CurveModel.h b/Engine/Source/Editor/CurveEditor/Public/CurveModel.h
--- a/Engine/Source/Editor/CurveEditor/Public/CurveModel.h
+++ b/Engine/Source/Editor/CurveEditor/Public/CurveModel.h
@@ -2,6 +2,7 @@
 
 #pragma once
 
+#include "CurveEditorSettings.h"
 #include "Containers/Array.h"
 #include "Containers/ArrayView.h"
 #include "Containers/UnrealString.h"
@@ -503,11 +504,19 @@ public:
 	{
 		return ChannelName;
 	}
-
-	/**
-	 */
-	FORCEINLINE void SetColor(const FLinearColor& InColor, bool bInModify = true)
+	
+	FORCEINLINE virtual void SetColor(const FLinearColor& InColor, bool bInModify = true)
 	{
+		UObject* Object = GetOwningObject();
+		if (Object)
+		{
+			UCurveEditorSettings* CurveEditorSettings = GetMutableDefault<UCurveEditorSettings>();
+			if (CurveEditorSettings)
+			{
+				CurveEditorSettings->SetCustomColor(Object->GetClass(), IntentionName, InColor);
+			}
+		}
+	
 		Color = InColor;
 		if (bInModify)
 		{
diff --git a/Engine/Source/Editor/Persona/Private/AnimTimeline/AnimTimelineTrack_Curve.cpp b/Engine/Source/Editor/Persona/Private/AnimTimeline/AnimTimelineTrack_Curve.cpp
--- a/Engine/Source/Editor/Persona/Private/AnimTimeline/AnimTimelineTrack_Curve.cpp
+++ b/Engine/Source/Editor/Persona/Private/AnimTimeline/AnimTimelineTrack_Curve.cpp
@@ -216,7 +216,9 @@ TSharedRef<SWidget> FAnimTimelineTrack_Curve::GenerateContainerWidgetForTimeline
 		GetCurveEditInfo(CurveIndex, Name, Type, EditIndex);
 
 		TUniquePtr<FRichCurveEditorModelNamed> NewCurveModel = MakeUnique<FRichCurveEditorModelNamed>(Name, Type, EditIndex, GetModel()->GetAnimSequenceBase());
-		NewCurveModel->SetColor(GetCurveColor(CurveIndex), false);
+		NewCurveModel->SetShortDisplayName(FText::FromName(Name));
+		NewCurveModel->SetLongDisplayName(FText::FromName(Name));
+		NewCurveModel->SetIntentionName(Name.ToString());
 		NewCurveModel->SetIsKeyDrawEnabled(MakeAttributeLambda([](){ return GetDefault<UPersonaOptions>()->bTimelineDisplayCurveKeys; }));
 		CurveEditor->AddCurve(MoveTemp(NewCurveModel));
 	}
diff --git a/Engine/Source/Editor/Persona/Private/AnimTimeline/AnimTimelineTrack_FloatCurve.cpp b/Engine/Source/Editor/Persona/Private/AnimTimeline/AnimTimelineTrack_FloatCurve.cpp
--- a/Engine/Source/Editor/Persona/Private/AnimTimeline/AnimTimelineTrack_FloatCurve.cpp
+++ b/Engine/Source/Editor/Persona/Private/AnimTimeline/AnimTimelineTrack_FloatCurve.cpp
@@ -33,7 +33,6 @@ FAnimTimelineTrack_FloatCurve::FAnimTimelineTrack_FloatCurve(const FFloatCurve*
 	, FloatCurve(InCurve)
 	, CurveName(InCurve->GetName())
 	, CurveId(FAnimationCurveIdentifier(InCurve->GetName(), ERawCurveTrackTypes::RCT_Float))
-	, Color(InCurve->Color)
 	, Comment(InCurve->Comment) 
 	, bIsMetadata(InCurve->GetCurveTypeFlag(AACF_Metadata))
 {
@@ -339,25 +338,19 @@ void FAnimTimelineTrack_FloatCurve::AddCurveTrackButton(TSharedPtr<SHorizontalBo
 		PersonaUtils::MakeTrackButton(LOCTEXT("EditCurveButtonText", "Curve"), FOnGetContent::CreateSP(this, &FAnimTimelineTrack_FloatCurve::BuildCurveTrackMenu), MakeAttributeSP(this, &FAnimTimelineTrack_FloatCurve::IsHovered))
 	];
 
-	FLinearColor CurveColor = FloatCurve->Color;
 	auto SetValue = [this](FLinearColor InNewColor)
 	{
-		UAnimSequenceBase* AnimSequenceBase = GetModel()->GetAnimSequenceBase();
-		AnimSequenceBase->GetController().SetCurveColor(CurveId, InNewColor);
-		Color = InNewColor;
-
-		// Set display curves too
 		for(const TPair<FCurveModelID, TUniquePtr<FCurveModel>>& CurvePair : CurveEditor->GetCurves())
 		{
 			CurvePair.Value->SetColor(InNewColor, false);
 		}
 	};
 
-	auto OnGetMenuContent = [this, CurveColor, SetValue]()
+	auto OnGetMenuContent = [this, SetValue]()
 	{
 		// Open a color picker
 		return SNew(SColorPicker)
-			.TargetColorAttribute(CurveColor)
+			.TargetColorAttribute_Raw(this, &FAnimTimelineTrack_FloatCurve::GetCurveColor, static_cast<int32>(INDEX_NONE))
 			.UseAlpha(false)
 			.DisplayInlineVersion(true)
 			.OnColorCommitted_Lambda(SetValue)
@@ -390,7 +383,7 @@ void FAnimTimelineTrack_FloatCurve::AddCurveTrackButton(TSharedPtr<SHorizontalBo
 		.ButtonContent()
 		[
 			SNew(SColorBlock)
-			.Color(CurveColor)
+			.Color_Raw(this, &FAnimTimelineTrack_FloatCurve::GetCurveColor, static_cast<int32>(INDEX_NONE))
 			.ShowBackgroundForAlpha(false)
 			.AlphaDisplayMode(EColorBlockAlphaDisplayMode::Ignore)
 			.Size(FVector2D(OutlinerRightPadding - 2.0f, OutlinerRightPadding))
@@ -399,8 +392,10 @@ void FAnimTimelineTrack_FloatCurve::AddCurveTrackButton(TSharedPtr<SHorizontalBo
 }
 
 FLinearColor FAnimTimelineTrack_FloatCurve::GetCurveColor(int32 InCurveIndex) const
-{ 
-	return Color; 
+{
+	FString	CurveIntentionName = CurveName.ToString();
+	UAnimSequenceBase* AnimSequenceBase = GetModel()->GetAnimSequenceBase();
+	return FAnimTimelineTrack_FloatCurve::GetFloatCurveColor(CurveName, CurveIntentionName, AnimSequenceBase, FLinearColor());
 }
 
 void FAnimTimelineTrack_FloatCurve::GetCurveEditInfo(int32 InCurveIndex, FName& OutName, ERawCurveTrackTypes& OutType, int32& OutCurveIndex) const
@@ -462,9 +457,38 @@ FSlateColor FAnimTimelineTrack_FloatCurve::GetTrackColor(bool bForComment) const
 	}
 	else
 	{
-		FLinearColor CurveColor = bIsMetadata ? Color.Desaturate(0.25f) : Color.Desaturate(0.75f);
+		FLinearColor CurveColor = bIsMetadata ? GetCurveColor(INDEX_NONE).Desaturate(0.25f) : GetCurveColor(INDEX_NONE).Desaturate(0.75f);
 		return bForComment ? CurveColor.CopyWithNewOpacity(CurveColor.A * 0.5f) : CurveColor;
 	}
 }
 
+FLinearColor FAnimTimelineTrack_FloatCurve::GetFloatCurveColor( const FName& InCurveName, const FString& InCurveIntentionName, const UAnimSequenceBase* InUAnimSequenceBase, FLinearColor InDefaultColor )
+{
+	// first check curve color from settings
+	const UCurveEditorSettings* CurveEditorSettings = GetDefault<UCurveEditorSettings>();
+	if (CurveEditorSettings)
+	{
+		TOptional<FLinearColor> ColorFromSettings = CurveEditorSettings->GetCustomColor(InUAnimSequenceBase->GetClass(), InCurveIntentionName);
+		if (ColorFromSettings.IsSet())
+		{
+			return ColorFromSettings.GetValue();
+		}
+	}
+	
+	// then check in UObject
+	IAnimationDataModel* AnimationDataModel = InUAnimSequenceBase->GetDataModel();
+	if (AnimationDataModel)
+	{
+		FAnimationCurveIdentifier CurveID(InCurveName, ERawCurveTrackTypes::RCT_Float);
+		const FFloatCurve* FloatCurve = AnimationDataModel->FindFloatCurve(CurveID);
+		if (FloatCurve)
+		{
+			return  FloatCurve->Color;
+		}
+	}
+	
+	return InDefaultColor;
+}
+
+
 #undef LOCTEXT_NAMESPACE
diff --git a/Engine/Source/Editor/Persona/Private/AnimTimeline/AnimTimelineTrack_FloatCurve.h b/Engine/Source/Editor/Persona/Private/AnimTimeline/AnimTimelineTrack_FloatCurve.h
--- a/Engine/Source/Editor/Persona/Private/AnimTimeline/AnimTimelineTrack_FloatCurve.h
+++ b/Engine/Source/Editor/Persona/Private/AnimTimeline/AnimTimelineTrack_FloatCurve.h
@@ -69,9 +69,6 @@ private:
 	FName CurveName;
 	FAnimationCurveIdentifier CurveId;
 
-	/** Cached color */
-	FLinearColor Color;
-
 	/** Cached comment */
 	FString Comment;
 
@@ -83,4 +80,7 @@ private:
 
 	/** Comment we can edit */
 	TSharedPtr<SInlineEditableTextBlock> EditableTextComment;
+	
+public:
+	static FLinearColor GetFloatCurveColor(const FName& InCurveName, const FString& InCurveIntentionName, const UAnimSequenceBase* InUAnimSequenceBase, FLinearColor InDefaultColor); 
 };
diff --git a/Engine/Source/Editor/Persona/Private/SAnimSequenceCurveEditor.cpp b/Engine/Source/Editor/Persona/Private/SAnimSequenceCurveEditor.cpp
--- a/Engine/Source/Editor/Persona/Private/SAnimSequenceCurveEditor.cpp
+++ b/Engine/Source/Editor/Persona/Private/SAnimSequenceCurveEditor.cpp
@@ -1,6 +1,7 @@
 // Copyright Epic Games, Inc. All Rights Reserved.
 
 #include "SAnimSequenceCurveEditor.h"
+#include "AnimTimeline/AnimTimelineTrack_FloatCurve.h"
 #include "AnimatedRange.h"
 #include "CurveEditor.h"
 #include "RichCurveEditorModel.h"
@@ -61,6 +62,11 @@ bool FRichCurveEditorModelNamed::IsValid() const
 	return AnimSequence->GetDataModel()->FindCurve(FAnimationCurveIdentifier(CurveName, Type)) != nullptr;
 }
 
+UObject* FRichCurveEditorModelNamed::GetOwningObject() const
+{
+	return AnimSequence.Get();
+}
+
 FRichCurve& FRichCurveEditorModelNamed::GetRichCurve()
 {
 	check(AnimSequence.Get() != nullptr);
@@ -118,6 +124,20 @@ void FRichCurveEditorModelNamed::SetCurveAttributes(const FCurveAttributes& InCu
 	Controller.SetCurveAttributes(CurveId, InCurveAttributes);
 }
 
+FLinearColor FRichCurveEditorModelNamed::GetColor() const
+{
+	return FAnimTimelineTrack_FloatCurve::GetFloatCurveColor(CurveName, IntentionName, AnimSequence.Get(), FRichCurveEditorModel::GetColor());
+}
+
+void FRichCurveEditorModelNamed::SetColor(const FLinearColor& InColor, bool bInModify)
+{
+	IAnimationDataController& Controller = AnimSequence->GetController();
+	Controller.SetCurveColor(CurveId, InColor);
+	
+	// call base class to update in UCurveEditorSettings and FCurveModel
+	FRichCurveEditorModel::SetColor(InColor, bInModify);
+}
+
 void FRichCurveEditorModelNamed::CurveHasChanged()
 {	
 	IAnimationDataController& Controller = AnimSequence->GetController();
@@ -278,7 +298,7 @@ public:
 		, CurveIndex(InCurveIndex)
 		, AnimSequence(InAnimSequence)
 		, CurveDisplayName(InCurveDisplayName)
-		, CurveColor(InCurveColor)
+		, CurveIntentionName(CurveDisplayName.ToString())
 		, OnCurveModified(InOnCurveModified)
 		, TreeId(InTreeId)
 	{
@@ -297,7 +317,7 @@ public:
 				[
 					SNew(STextBlock)
 					.Text(CurveDisplayName)
-					.ColorAndOpacity(FSlateColor(CurveColor))
+					.ColorAndOpacity_Raw(this, &FAnimSequenceCurveEditorItem::GetCurveColor)
 				];
 		}
 		else if (InColumnName == ColumnNames.SelectHeader)
@@ -317,7 +337,9 @@ public:
 		TUniquePtr<FRichCurveEditorModelNamed> NewCurveModel = MakeUnique<FRichCurveEditorModelNamed>(Name, Type, CurveIndex, AnimSequence.Get(), TreeId);
 		NewCurveModel->SetShortDisplayName(CurveDisplayName);
 		NewCurveModel->SetLongDisplayName(CurveDisplayName);
-		NewCurveModel->SetColor(CurveColor);
+		const FLinearColor UpdatedCurveColor = GetCurveColor().GetSpecifiedColor();
+		NewCurveModel->SetColor(UpdatedCurveColor);
+		NewCurveModel->SetIntentionName(CurveIntentionName);
 		NewCurveModel->OnCurveModified().Add(OnCurveModified);
 
 		OutCurveModels.Add(MoveTemp(NewCurveModel));
@@ -345,12 +367,17 @@ public:
 		return false;
 	}
 
+	FSlateColor GetCurveColor() const
+	{
+		return FAnimTimelineTrack_FloatCurve::GetFloatCurveColor(Name, CurveIntentionName, AnimSequence.Get(), FLinearColor());
+	}
+
 	FName Name;
 	ERawCurveTrackTypes Type;
 	int32 CurveIndex;
 	TWeakObjectPtr<UAnimSequenceBase> AnimSequence;
 	FText CurveDisplayName;
-	FLinearColor CurveColor;
+	FString CurveIntentionName;
 	FSimpleDelegate OnCurveModified;
 	FCurveEditorTreeItemID TreeId;
 };
diff --git a/Engine/Source/Editor/Persona/Private/SAnimSequenceCurveEditor.h b/Engine/Source/Editor/Persona/Private/SAnimSequenceCurveEditor.h
--- a/Engine/Source/Editor/Persona/Private/SAnimSequenceCurveEditor.h
+++ b/Engine/Source/Editor/Persona/Private/SAnimSequenceCurveEditor.h
@@ -31,12 +31,15 @@ public:
 	virtual ~FRichCurveEditorModelNamed();
 
 	virtual bool IsValid() const override;
+	virtual UObject* GetOwningObject() const override;
 	virtual FRichCurve& GetRichCurve() override;
 	virtual const FRichCurve& GetReadOnlyRichCurve() const override;
+	virtual FLinearColor GetColor() const override;
 
 	virtual void SetKeyPositions(TArrayView<const FKeyHandle> InKeys, TArrayView<const FKeyPosition> InKeyPositions, EPropertyChangeType::Type ChangeType) override;
 	virtual void SetKeyAttributes(TArrayView<const FKeyHandle> InKeys, TArrayView<const FKeyAttributes> InAttributes, EPropertyChangeType::Type ChangeType = EPropertyChangeType::Unspecified) override;
 	virtual void SetCurveAttributes(const FCurveAttributes& InCurveAttributes) override;
+	virtual void SetColor(const FLinearColor& InColor, bool bInModify = true) override;
 
 	void CurveHasChanged();
 	void OnModelHasChanged(const EAnimDataModelNotifyType& NotifyType, IAnimationDataModel* Model, const FAnimDataModelNotifPayload& Payload);
