Skeletal Mesh Editor: enable ITF based frustum selection
- DragTools.EnableITFTools true is still needed

#rnx

diff --git a/Engine/Source/Editor/Persona/Private/IPersonaEditorModeManager.cpp b/Engine/Source/Editor/Persona/Private/IPersonaEditorModeManager.cpp
--- a/Engine/Source/Editor/Persona/Private/IPersonaEditorModeManager.cpp
+++ b/Engine/Source/Editor/Persona/Private/IPersonaEditorModeManager.cpp
@@ -60,8 +60,10 @@ void UPersonaEditorModeManagerContext::SetFocusInViewport()
 
 IPersonaEditorModeManager::IPersonaEditorModeManager() : FAssetEditorModeManager()
 {
-	UContextObjectStore* ContextObjectStore = GetInteractiveToolsContext()->ContextObjectStore;
+	UModeManagerInteractiveToolsContext* ToolsContext = GetInteractiveToolsContext();
+	UContextObjectStore* ContextObjectStore = ToolsContext->ContextObjectStore;
 	ContextObjectStore->AddContextObject(PersonaModeManagerContext.Get());
+	ToolsContext->SetDragToolsEnabled(true);
 }
 
 IPersonaEditorModeManager::~IPersonaEditorModeManager()
