Fix for an issue when doing a replace on multiple selected notifies. When iterating through the array of nodes user chose, we now do it in reverse because we delete as we go. Previously this was changing the size of the array.

diff --git a/Engine/Source/Editor/Persona/Private/SAnimNotifyPanel.cpp b/Engine/Source/Editor/Persona/Private/SAnimNotifyPanel.cpp
--- a/Engine/Source/Editor/Persona/Private/SAnimNotifyPanel.cpp
+++ b/Engine/Source/Editor/Persona/Private/SAnimNotifyPanel.cpp
@@ -4586,8 +4586,10 @@ void SAnimNotifyPanel::OnReplaceSelectedWithNotify(FString NewNotifyName, UClass
 	const FScopedTransaction Transaction(LOCTEXT("ReplaceAnimNotify", "Replace Anim Notify"));
 	Sequence->Modify(true);
 
-	for (INodeObjectInterface* NodeObject : SelectedNodes)
+	// Reverse through the array, replacing nodes.
+	for(int32 index=SelectedNodes.Num()-1; index>=0; index--)
 	{
+		INodeObjectInterface* NodeObject = SelectedNodes[index];
 		FAnimNotifyEvent* OldEvent = NodeObject->GetNotifyEvent();
 		if (OldEvent)
 		{
