Randomize notify color in editor based on sync marker name

diff --git a/Engine/Source/Editor/Persona/Private/SAnimNotifyPanel.cpp b/Engine/Source/Editor/Persona/Private/SAnimNotifyPanel.cpp
--- a/Engine/Source/Editor/Persona/Private/SAnimNotifyPanel.cpp
+++ b/Engine/Source/Editor/Persona/Private/SAnimNotifyPanel.cpp
@@ -348,7 +348,14 @@ struct FSyncMarkerNodeInterface : public INodeObjectInterface
 	virtual bool IsBranchingPoint() override { return false; }
 	virtual TOptional<FLinearColor> GetEditorColor() override
 	{
-		return FLinearColor::Green;
+		// Choose a random color based on the marker name
+		// We don't use FLinearColor::MakeRandomSeededColor because it uses random RGB
+		// which makes dull colors
+		// Note: the 157 multiplier ensures similar hashes have different colors (see FLinearColor::MakeRandomColor)
+		const uint8 Saturation = 255;
+		const uint8 Luminosity = 255;
+		const uint8 Hue = (uint8)GetTypeHash(SyncMarker->MarkerName.ToString()) * 157;
+		return FLinearColor::MakeFromHSV8(Hue, Saturation, Luminosity);
 	}
 
 	virtual FText GetNodeTooltip(const UAnimSequenceBase* Sequence) override
