Skeletal Mesh: Fix localization issues.

#rnx

diff --git a/Engine/Source/Editor/Persona/Private/PersonaMeshDetails.cpp b/Engine/Source/Editor/Persona/Private/PersonaMeshDetails.cpp
--- a/Engine/Source/Editor/Persona/Private/PersonaMeshDetails.cpp
+++ b/Engine/Source/Editor/Persona/Private/PersonaMeshDetails.cpp
@@ -168,8 +168,8 @@ IDetailCategoryBuilder&  GetLODIndexCategory(IDetailLayoutBuilder& DetailLayout,
 {
 	FString CategoryName = FString(TEXT("LOD"));
 	CategoryName.AppendInt(LODIndex);
-	FText LODLevelString = FText::FromString(FString(TEXT("LOD ")) + FString::FromInt(LODIndex));
-	return DetailLayout.EditCategory(*CategoryName, LODLevelString, ECategoryPriority::Important);
+	FText LODLevelLabel = FText::Format(LOCTEXT("LODIndexLabel", "LOD {0}"), FText::AsNumber(LODIndex));
+	return DetailLayout.EditCategory(*CategoryName, LODLevelLabel, ECategoryPriority::Important);
 }
 
 enum class EButtonFlags
@@ -3240,7 +3240,7 @@ FText FPersonaMeshDetails::GetLODCustomModeNameContent(int32 LODIndex) const
 	{
 		return LOCTEXT("GetLODCustomModeNameContent_None", "Custom");
 	}
-	return FText::Format(LOCTEXT("GetLODCustomModeNameContent", "LOD{0}"), LODIndex);
+	return FText::Format(LOCTEXT("GetLODCustomModeNameContent", "LOD {0}"), LODIndex);
 }
 
 ECheckBoxState FPersonaMeshDetails::IsLODCustomModeCheck(int32 LODIndex) const
@@ -5115,14 +5115,12 @@ FReply FPersonaMeshDetails::AddMaterialSlot()
 
 FText FPersonaMeshDetails::GetMaterialArrayText() const
 {
-	FString MaterialArrayText = TEXT(" Material Slots");
 	int32 SlotNumber = 0;
 	if (SkeletalMeshPtr.IsValid() && !SkeletalMeshPtr->IsCompiling())
 	{
 		SlotNumber = SkeletalMeshPtr->GetMaterials().Num();
 	}
-	MaterialArrayText = FString::FromInt(SlotNumber) + MaterialArrayText;
-	return FText::FromString(MaterialArrayText);
+	return FText::Format(LOCTEXT("MaterialSlots", "{0} {0}|plural(one=Material Slot,other=Material Slots)"), SlotNumber);
 }
 
 void FPersonaMeshDetails::OnGetSectionsForView(ISectionListBuilder& OutSections, int32 LODIndex)
@@ -5333,42 +5331,6 @@ TSharedRef<SWidget> FPersonaMeshDetails::OnGenerateMaterialListExtraBottomWidget
 	return VerticalWidget;
 }
 
-FText FPersonaMeshDetails::GetFirstMaterialSlotUsedBySection(int32 MaterialIndex) const
-{
-	if (SkeletalMeshPtr.IsValid() && !SkeletalMeshPtr->IsCompiling() && MaterialUsedMap.Contains(MaterialIndex))
-	{
-		const TArray<FSectionLocalizer> *SectionLocalizers = MaterialUsedMap.Find(MaterialIndex);
-		if (SectionLocalizers->Num() > 0)
-		{
-			FString ArrayItemName = FString::FromInt(SectionLocalizers->Num()) + TEXT(" Sections");
-			return FText::FromString(ArrayItemName);
-		}
-	}
-	return FText();
-}
-
-TSharedRef<SWidget> FPersonaMeshDetails::OnGetMaterialSlotUsedByMenuContent(int32 MaterialIndex)
-{
-	FMenuBuilder MenuBuilder(true, NULL);
-
-	TArray<FSectionLocalizer> *SectionLocalizers;
-	if (SkeletalMeshPtr.IsValid() && MaterialUsedMap.Contains(MaterialIndex))
-{
-		SectionLocalizers = MaterialUsedMap.Find(MaterialIndex);
-		FUIAction Action;
-		FText EmptyTooltip;
-		// Add a menu item for each texture.  Clicking on the texture will display it in the content browser
-		for (const FSectionLocalizer& SectionUsingMaterial : (*SectionLocalizers))
-		{
-			FString ArrayItemName = TEXT("Lod ") + FString::FromInt(SectionUsingMaterial.LODIndex) + TEXT("  Index ") + FString::FromInt(SectionUsingMaterial.SectionIndex);
-			MenuBuilder.AddMenuEntry(FText::FromString(ArrayItemName), EmptyTooltip, FSlateIcon(), Action);
-		}
-	}
-	
-
-	return MenuBuilder.MakeWidget();
-}
-
 bool FPersonaMeshDetails::CanDeleteMaterialSlot(int32 MaterialIndex) const
 {
 	if (!SkeletalMeshPtr.IsValid() || SkeletalMeshPtr->IsCompiling())
@@ -5715,7 +5677,7 @@ TSharedRef<SWidget> FPersonaMeshDetails::OnGenerateCustomSectionWidgetsForSectio
 				SNew(STextBlock)
 				.Font(IDetailLayoutBuilder::GetDetailFont())
 				.Text(this, &FPersonaMeshDetails::GetCurrentRecomputeTangentsSetting, LODIndex, SectionIndex)
-				.ToolTipText(LOCTEXT("RecomputeTangentsToolTip", "Require (Support Skincache Shaders) enabled in Project Settings and compute shaders. None:Do_NOT_Recompute, All:Do_Recompute, R/G/B:Do_Recompute_And_Use_Vertex_Color_Channel_As_Blending_Mask"))
+				.ToolTipText(LOCTEXT("RecomputeTangentsToolTip", "Specifies how tangents are recomputed after skinning, at runtime. Requires skin cache to be enabled for this mesh."))
 			]
 		]
 	];
@@ -6006,8 +5968,13 @@ void FPersonaMeshDetails::UpdateLODCategoryVisibility() const
 
 FText FPersonaMeshDetails::GetCurrentLodName() const
 {
+	if (GetAutoLod(GetPersonaToolkit()->GetPreviewMeshComponent()))
+	{
+		return LOCTEXT("LODAutoLabel", "Auto (LOD 0)");
+	}
+
 	const int32 CurrentDisplayLOD = GetCurrentDisplayLODIndex(GetPersonaToolkit()->GetPreviewMeshComponent());
-	return FText::FromString(GetAutoLod(GetPersonaToolkit()->GetPreviewMeshComponent()) ? FString(TEXT("Auto (LOD0)")) : (FString(TEXT("LOD")) + FString::FromInt(CurrentDisplayLOD)));
+	return FText::Format(LOCTEXT("LODIndexLabel", "LOD {0}"), FText::AsNumber(CurrentDisplayLOD));
 }
 
 bool FPersonaMeshDetails::GetAutoLod(USkeletalMeshComponent* InMeshComponent)
@@ -6029,7 +5996,7 @@ FText FPersonaMeshDetails::GetCurrentLodTooltip() const
 {
 	if (GetPersonaToolkit()->GetPreviewMeshComponent() != nullptr && GetPersonaToolkit()->GetPreviewMeshComponent()->GetForcedLOD() == 0)
 	{
-		return LOCTEXT("PersonaLODPickerCurrentLODTooltip", "With Auto LOD selected, LOD0's properties are visible for editing");
+		return LOCTEXT("PersonaLODPickerCurrentLODTooltip", "With Auto LOD selected, LOD 0's properties are shown by default");
 	}
 	return FText::GetEmpty();
 }
@@ -6087,13 +6054,12 @@ TSharedRef<SWidget> FPersonaMeshDetails::OnGenerateLodMenuForLodPicker()
 	}
 	FMenuBuilder MenuBuilder(true, NULL);
 
-	FText AutoLodText = FText::FromString((TEXT("Auto LOD")));
 	FUIAction AutoLodAction(FExecuteAction::CreateSP(this, &FPersonaMeshDetails::SetCurrentLOD, 0));
-	MenuBuilder.AddMenuEntry(AutoLodText, LOCTEXT("OnGenerateLodMenuForSectionList_Auto_ToolTip", "With Auto LOD selected, LOD0's properties are visible for editing."), FSlateIcon(), AutoLodAction);
+	MenuBuilder.AddMenuEntry(LOCTEXT("LODAutoLabel", "Auto (LOD 0)"), LOCTEXT("OnGenerateLodMenuForSectionList_Auto_ToolTip", "With Auto LOD selected, LOD 0's properties are visible for editing."), FSlateIcon(), AutoLodAction);
 	// Add a menu item for each texture.  Clicking on the texture will display it in the content browser
 	for (int32 AllLodIndex = 0; AllLodIndex < SkelMeshLODCount; ++AllLodIndex)
 	{
-		FText LODLevelString = FText::FromString((TEXT("LOD ") + FString::FromInt(AllLodIndex)));
+		FText LODLevelString = FText::Format(LOCTEXT("LODIndexLabel", "LOD {0}"), FText::AsNumber(AllLodIndex));
 		FUIAction Action(FExecuteAction::CreateSP(this, &FPersonaMeshDetails::SetCurrentLOD, AllLodIndex+1));
 		MenuBuilder.AddMenuEntry(LODLevelString, FText::GetEmpty(), FSlateIcon(), Action);
 	}
@@ -6390,12 +6356,25 @@ void FPersonaMeshDetails::OnSectionVisibleInRayTracingChanged(ECheckBoxState New
 	}
 }
 
-const FString RecomputeTangentsChannelNames[] = {
-			TEXT("None"),
-			TEXT("All"),
-			TEXT("Red"),
-			TEXT("Green"),
-			TEXT("Blue")
+static const FText RecomputeTangentsChannelLabels[][2] = {
+	{
+		LOCTEXT("RecomputeTangentChannel_Disabled", "Disabled"),
+		LOCTEXT("RecomputeTangentChannelTooltip_Disabled", "Tangents will not be recomputed after skinning for this section."),
+	},
+	{LOCTEXT("RecomputeTangentChannel_Enabled", "Enabled with no mask"),
+		LOCTEXT("RecomputeTangentChannelTooltip_Enabled", "Tangents will be recomputed after skinning for this section."),
+	},
+	{LOCTEXT("RecomputeTangentChannel_Red", "Enabled using red mask"), 
+		LOCTEXT("RecomputeTangentChannelTooltip_Red", "Tangents will be recomputed after skinning, for this section, and blended with the base pose normal using the red vertex color channel as a mask to blend in the recomputed tangents."),
+	},
+	{
+		LOCTEXT("RecomputeTangentChannel_Green", "Enabled using green mask"),
+		LOCTEXT("RecomputeTangentChannelTooltip_Green", "Tangents will be recomputed after skinning, for this section, and blended with the base pose normal using the green vertex color channel as a mask to blend in the recomputed tangents."),
+	},
+	{
+		LOCTEXT("RecomputeTangentChannel_Blue", "Enabled using blue mask"),
+		LOCTEXT("RecomputeTangentChannelTooltip_Blue", "Tangents will be recomputed after skinning, for this section, and blended with the base pose normal using the blue vertex color channel as a mask to blend in the recomputed tangents."),
+	},
 };
 
 FText FPersonaMeshDetails::GetCurrentRecomputeTangentsSetting(int32 LODIndex, int32 SectionIndex) const
@@ -6430,16 +6409,16 @@ FText FPersonaMeshDetails::GetCurrentRecomputeTangentsSetting(int32 LODIndex, in
 			ChannelIndex = FMath::Clamp<int32>(static_cast<int32>(Section.RecomputeTangentsVertexMaskChannel) + 2, 0, 4);
 	}
 	
-	return FText::FromString(RecomputeTangentsChannelNames[ChannelIndex]);
+	return RecomputeTangentsChannelLabels[ChannelIndex][0];
 }
 
 TSharedRef<SWidget> FPersonaMeshDetails::OnGenerateRecomputeTangentsSetting(int32 LODIndex, int32 SectionIndex)
 {
-	FMenuBuilder MenuBuilder(true, NULL);
+	FMenuBuilder MenuBuilder(true, nullptr);
 	for (int32 ChannelIndex = 0; ChannelIndex < 5; ChannelIndex++) { 
 		FUIAction Action(FExecuteAction::CreateSP(this, &FPersonaMeshDetails::SetCurrentRecomputeTangentsSetting, LODIndex, SectionIndex, ChannelIndex));
 
-		MenuBuilder.AddMenuEntry(FText::FromString(RecomputeTangentsChannelNames[ChannelIndex]), FText::GetEmpty(), FSlateIcon(), Action);
+		MenuBuilder.AddMenuEntry(RecomputeTangentsChannelLabels[ChannelIndex][0], RecomputeTangentsChannelLabels[ChannelIndex][1], FSlateIcon(), Action);
 	}
 	return MenuBuilder.MakeWidget();
 }
@@ -6817,7 +6796,7 @@ TSharedRef<SUniformGridPanel> FPersonaMeshDetails::MakeClothingDetailsWidget(int
 		[
 			SNew(STextBlock)
 			.Font(DetailFontInfo)
-			.Text(FText::Format(LOCTEXT("LODIndex", "LOD {0}"), FText::AsNumber(LODIndex)))			
+			.Text(FText::Format(LOCTEXT("LODIndexLabel", "LOD {0}"), FText::AsNumber(LODIndex)))			
 		];
 
 		RowNumber++;
diff --git a/Engine/Source/Editor/Persona/Private/PersonaMeshDetails.h b/Engine/Source/Editor/Persona/Private/PersonaMeshDetails.h
--- a/Engine/Source/Editor/Persona/Private/PersonaMeshDetails.h
+++ b/Engine/Source/Editor/Persona/Private/PersonaMeshDetails.h
@@ -491,13 +491,9 @@ private:
 
 	void OnDeleteMaterialSlot(int32 MaterialIndex);
 
-	TSharedRef<SWidget> OnGetMaterialSlotUsedByMenuContent(int32 MaterialIndex);
-
 	/* Call when user change the material slot overlay material. */
 	void OnMaterialSlotOverlayMaterialChanged(UMaterialInterface* NewOverlayMaterial, int32 MaterialIndex);
 
-	FText GetFirstMaterialSlotUsedBySection(int32 MaterialIndex) const;
-
 	/**
 	* Handler for check box display based on whether the material is highlighted
 	*
