Fix for when removing a montage section, any sections that used that section as link will now become unlinked automatically.

diff --git a/Engine/Source/Editor/Persona/Private/SAnimMontagePanel.cpp b/Engine/Source/Editor/Persona/Private/SAnimMontagePanel.cpp
--- a/Engine/Source/Editor/Persona/Private/SAnimMontagePanel.cpp
+++ b/Engine/Source/Editor/Persona/Private/SAnimMontagePanel.cpp
@@ -438,6 +438,17 @@ void SAnimMontagePanel::RemoveSection(int32 SectionIndex)
 		const FScopedTransaction Transaction( LOCTEXT("DeleteSection", "Delete Section") );
 		Montage->Modify();
 
+		// Unlink section name that is being removed from other section's NextSection.
+		const FName SectionNameToRemove = Montage->CompositeSections[SectionIndex].SectionName;
+		for (FCompositeSection& CompositeSection : Montage->CompositeSections)
+		{
+			if(CompositeSection.NextSectionName == SectionNameToRemove)
+			{
+				CompositeSection.NextSectionName = NAME_None;
+			}
+		}
+
+		// Remove the section.
 		Montage->CompositeSections.RemoveAt(SectionIndex);
 		EnsureStartingSection();
 		OnMontageModified();
