Add missing playback speed control to the Animation Blueprint preview scrub controls.
- Elevate playback speed functions to IPersonaPreviewScene.

#rnx

diff --git a/Engine/Source/Editor/Persona/Private/AnimTimeline/SAnimTimelineTransportControls.cpp b/Engine/Source/Editor/Persona/Private/AnimTimeline/SAnimTimelineTransportControls.cpp
--- a/Engine/Source/Editor/Persona/Private/AnimTimeline/SAnimTimelineTransportControls.cpp
+++ b/Engine/Source/Editor/Persona/Private/AnimTimeline/SAnimTimelineTransportControls.cpp
@@ -72,7 +72,7 @@ void SAnimTimelineTransportControls::Construct(const FArguments& InArgs, const T
 		{
 			MenuContext.AppendCommandList(PreviewScene->GetCommandList());
 			UAnimViewportContext* const ContextObject = NewObject<UAnimViewportContext>();
-			ContextObject->AnimationEditorPreviewScene = PreviewScene;
+			ContextObject->PersonaPreviewScene = PreviewScene;
 			MenuContext.AddObject(ContextObject);
 		}
 
diff --git a/Engine/Source/Editor/Persona/Private/AnimationEditorPreviewScene.h b/Engine/Source/Editor/Persona/Private/AnimationEditorPreviewScene.h
--- a/Engine/Source/Editor/Persona/Private/AnimationEditorPreviewScene.h
+++ b/Engine/Source/Editor/Persona/Private/AnimationEditorPreviewScene.h
@@ -375,19 +375,19 @@ public:
 	TOptional<float> GetCurrentTime() const;
 
 	/** Set the current custom animation playback speed */
-	void SetCustomAnimationSpeed(float InAnimationSpeed);
+	virtual void SetCustomAnimationSpeed(float InAnimationSpeed) override;
 
 	/** Get the current custom animation playback speed */
-	float GetCustomAnimationSpeed() const;
+	virtual float GetCustomAnimationSpeed() const override;
 
 	/** Retrieve the current animation playback speed mode */
-	EAnimationPlaybackSpeeds::Type GetAnimationPlaybackSpeedMode() const;
+	virtual  EAnimationPlaybackSpeeds::Type GetAnimationPlaybackSpeedMode() const override;
 
 	/** Set the playback speed mode */
-	void SetAnimationPlaybackSpeedMode(EAnimationPlaybackSpeeds::Type InMode);
+	virtual void SetAnimationPlaybackSpeedMode(EAnimationPlaybackSpeeds::Type InMode) override;
 
 	/** Function to return whether the supplied playback speed is the current active one */
-	bool IsPlaybackSpeedSelected(int32 InPlaybackSpeedMode) const;
+	virtual bool IsPlaybackSpeedSelected(int32 InPlaybackSpeedMode) const override;
 
 	/** Return the list of commands of the preview scene which are bound to delegates */
 	TSharedPtr<FUICommandList_Pinnable> GetPinnedCommandList() const;
diff --git a/Engine/Source/Editor/Persona/Private/SAnimationScrubPanel.cpp b/Engine/Source/Editor/Persona/Private/SAnimationScrubPanel.cpp
--- a/Engine/Source/Editor/Persona/Private/SAnimationScrubPanel.cpp
+++ b/Engine/Source/Editor/Persona/Private/SAnimationScrubPanel.cpp
@@ -12,6 +12,9 @@
 #include "AnimationEditorPreviewScene.h"
 #include "Animation/AnimData/IAnimationDataModel.h"
 #include "Animation/AnimSequenceHelpers.h"
+#include "ViewportToolbar/AnimationEditorMenus.h"
+#include "ViewportToolbar/AnimViewportContext.h"
+#include "ToolMenus.h"
 
 #define LOCTEXT_NAMESPACE "AnimationScrubPanel"
 
@@ -26,9 +29,11 @@ void SAnimationScrubPanel::Construct( const SAnimationScrubPanel::FArguments& In
 	
 	PreviewScenePtr = InPreviewScene;
 
+	TSharedPtr<SHorizontalBox> HorizontalBox;
+
 	this->ChildSlot
 	[
-		SNew(SHorizontalBox)
+		SAssignNew(HorizontalBox, SHorizontalBox)
 		.AddMetaData<FTagMetaData>(TEXT("AnimScrub.Scrub"))
 		+SHorizontalBox::Slot()
 		.HAlign(HAlign_Fill) 
@@ -69,6 +74,35 @@ void SAnimationScrubPanel::Construct( const SAnimationScrubPanel::FArguments& In
 			.IsRealtimeStreamingMode(this, &SAnimationScrubPanel::IsRealtimeStreamingMode)
 		]
 	];
+	
+	{
+		const FName MenuName = "AnimationEditor.Scrub.PlaybackMenu";
+		if (!UToolMenus::Get()->IsMenuRegistered(MenuName))
+		{
+			UToolMenu* Menu = UToolMenus::Get()->RegisterMenu(MenuName);
+			Menu->MenuType = EMultiBoxType::SlimHorizontalToolBar;
+			
+			FToolMenuSection& UnnamedSection = Menu->AddSection(NAME_None);
+			UnnamedSection.AddEntry(UE::AnimationEditor::CreatePlaybackSubmenu());
+		}
+		
+		FToolMenuContext MenuContext;
+		{
+			MenuContext.AppendCommandList(InPreviewScene->GetCommandList());
+			UAnimViewportContext* const ContextObject = NewObject<UAnimViewportContext>();
+			ContextObject->PersonaPreviewScene = InPreviewScene;
+			MenuContext.AddObject(ContextObject);
+		}
+
+		// clang-format off
+		HorizontalBox->AddSlot()
+		.HAlign(HAlign_Right)
+		.AutoWidth()
+		[
+			UToolMenus::Get()->GenerateWidget(MenuName, MenuContext)
+		];
+		// clang-format on
+	}
 }
 
 FReply SAnimationScrubPanel::OnClick_Forward_Step()
diff --git a/Engine/Source/Editor/Persona/Private/ViewportToolbar/AnimViewportContext.h b/Engine/Source/Editor/Persona/Private/ViewportToolbar/AnimViewportContext.h
--- a/Engine/Source/Editor/Persona/Private/ViewportToolbar/AnimViewportContext.h
+++ b/Engine/Source/Editor/Persona/Private/ViewportToolbar/AnimViewportContext.h
@@ -12,5 +12,5 @@ class UAnimViewportContext : public UObject
 
 public:
 	TWeakPtr<SAnimationEditorViewportTabBody> ViewportTabBody;
-	TWeakPtr<FAnimationEditorPreviewScene> AnimationEditorPreviewScene;
+	TWeakPtr<IPersonaPreviewScene> PersonaPreviewScene;
 };
diff --git a/Engine/Source/Editor/Persona/Private/ViewportToolbar/AnimationEditorMenus.cpp b/Engine/Source/Editor/Persona/Private/ViewportToolbar/AnimationEditorMenus.cpp
--- a/Engine/Source/Editor/Persona/Private/ViewportToolbar/AnimationEditorMenus.cpp
+++ b/Engine/Source/Editor/Persona/Private/ViewportToolbar/AnimationEditorMenus.cpp
@@ -397,15 +397,15 @@ void FillPlaybackMenu(UToolMenu* InMenu)
 
 	if (UAnimViewportContext* const AnimViewportContext = InMenu->FindContext<UAnimViewportContext>())
 	{
-		if (TSharedPtr<FAnimationEditorPreviewScene> PreviewScenePinned =
-				AnimViewportContext->AnimationEditorPreviewScene.Pin())
+		if (TSharedPtr<IPersonaPreviewScene> PreviewScenePinned =
+				AnimViewportContext->PersonaPreviewScene.Pin())
 		{
 			TSharedPtr<SWidget> AnimSpeedWidget =
 				SNew(UE::AnimationEditor::SCustomAnimationSpeedSetting)
 					.CustomSpeed_Lambda(
 						[PreviewSceneWeak = PreviewScenePinned.ToWeakPtr()]()
 						{
-							if (TSharedPtr<FAnimationEditorPreviewScene> PreviewScenePinned = PreviewSceneWeak.Pin())
+							if (TSharedPtr<IPersonaPreviewScene> PreviewScenePinned = PreviewSceneWeak.Pin())
 							{
 								return PreviewScenePinned->GetCustomAnimationSpeed();
 							}
@@ -416,7 +416,7 @@ void FillPlaybackMenu(UToolMenu* InMenu)
 					.OnCustomSpeedChanged_Lambda(
 						[PreviewSceneWeak = PreviewScenePinned.ToWeakPtr()](float InCustomSpeed)
 						{
-							if (TSharedPtr<FAnimationEditorPreviewScene> PreviewScenePinned = PreviewSceneWeak.Pin())
+							if (TSharedPtr<IPersonaPreviewScene> PreviewScenePinned = PreviewSceneWeak.Pin())
 							{
 								return PreviewScenePinned->SetCustomAnimationSpeed(InCustomSpeed);
 							}
@@ -1452,10 +1452,10 @@ void UE::AnimationEditor::ExtendPreviewSceneSettingsSubmenu(FName InSubmenuName)
 	);
 }
 
-FText UE::AnimationEditor::GetPlaybackMenuLabel(const TWeakPtr<FAnimationEditorPreviewScene>& InAnimationEditorPreviewScene)
+FText UE::AnimationEditor::GetPlaybackMenuLabel(const TWeakPtr<IPersonaPreviewScene>& InPersonaPreviewScene)
 {
 	FText Label = LOCTEXT("PlaybackError", "Error");
-	if (const TSharedPtr<FAnimationEditorPreviewScene>& AnimationEditorPreviewScene = InAnimationEditorPreviewScene.Pin())
+	if (const TSharedPtr<IPersonaPreviewScene>& AnimationEditorPreviewScene = InPersonaPreviewScene.Pin())
 	{
 		for (int i = 0; i < EAnimationPlaybackSpeeds::NumPlaybackSpeeds; ++i)
 		{
@@ -1491,7 +1491,7 @@ FToolMenuEntry UE::AnimationEditor::CreatePlaybackSubmenu()
 				{
 					// Label updates based on currently selected LOD
 					const TAttribute<FText> Label = TAttribute<FText>::CreateLambda(
-						[PreviewSceneWeak = AnimViewportContext->AnimationEditorPreviewScene]()
+						[PreviewSceneWeak = AnimViewportContext->PersonaPreviewScene]()
 						{
 							return GetPlaybackMenuLabel(PreviewSceneWeak);
 						}
@@ -1550,7 +1550,7 @@ TSharedRef<SWidget> UE::AnimationEditor::GeneratePlaybackMenu(
 		if (TSharedPtr<FAnimationEditorPreviewScene> AnimationPreviewScenePinned = InAnimationEditorPreviewScene.Pin())
 		{
 			MenuContext.AppendCommandList(AnimationPreviewScenePinned->GetCommandList());
-			ContextObject->AnimationEditorPreviewScene = AnimationPreviewScenePinned;
+			ContextObject->PersonaPreviewScene = AnimationPreviewScenePinned;
 		}
 	}
 
diff --git a/Engine/Source/Editor/Persona/Private/ViewportToolbar/AnimationEditorMenus.h b/Engine/Source/Editor/Persona/Private/ViewportToolbar/AnimationEditorMenus.h
--- a/Engine/Source/Editor/Persona/Private/ViewportToolbar/AnimationEditorMenus.h
+++ b/Engine/Source/Editor/Persona/Private/ViewportToolbar/AnimationEditorMenus.h
@@ -2,6 +2,7 @@
 
 #pragma once
 
+#include "IPersonaPreviewScene.h"
 #include "Widgets/SCompoundWidget.h"
 #include "ToolMenuMisc.h"
 
@@ -54,5 +55,5 @@ void ExtendPreviewSceneSettingsSubmenu(FName InSubmenuName);
 
 FToolMenuEntry CreatePlaybackSubmenu();
 TSharedRef<SWidget> GeneratePlaybackMenu(const TWeakPtr<FAnimationEditorPreviewScene>& InAnimationEditorPreviewScene, const TArray<TSharedPtr<FExtender>>& InExtenders);
-FText GetPlaybackMenuLabel(const TWeakPtr<FAnimationEditorPreviewScene>& InAnimationEditorPreviewScene);
+FText GetPlaybackMenuLabel(const TWeakPtr<IPersonaPreviewScene>& InPersonaPreviewScene);
 } // namespace UE::AnimationEditor
diff --git a/Engine/Source/Editor/Persona/Public/IPersonaPreviewScene.h b/Engine/Source/Editor/Persona/Public/IPersonaPreviewScene.h
--- a/Engine/Source/Editor/Persona/Public/IPersonaPreviewScene.h
+++ b/Engine/Source/Editor/Persona/Public/IPersonaPreviewScene.h
@@ -8,6 +8,7 @@
 #include "Types/SlateEnums.h"
 #include "Animation/AnimBlueprint.h"
 #include "PersonaSelectionProxies.h"
+#include "AnimationEditorTypes.h"
 
 class UAnimationAsset;
 class UDebugSkelMeshComponent;
@@ -332,5 +333,20 @@ public:
 	 * detached from other Persona components (eg. skeleton list).
 	 */
 	virtual void SetEditableSkeleton(TSharedPtr<IEditableSkeleton> InEditableSkeleton) = 0;
+	
+	/** Set the current custom animation playback speed */
+	virtual void SetCustomAnimationSpeed(float InAnimationSpeed) = 0;
+
+	/** Get the current custom animation playback speed */
+	virtual float GetCustomAnimationSpeed() const = 0;
+
+	/** Retrieve the current animation playback speed mode */
+	virtual EAnimationPlaybackSpeeds::Type GetAnimationPlaybackSpeedMode() const = 0;
+
+	/** Set the playback speed mode */
+	virtual void SetAnimationPlaybackSpeedMode(EAnimationPlaybackSpeeds::Type InMode) = 0;
+
+	/** Function to return whether the supplied playback speed is the current active one */
+	virtual bool IsPlaybackSpeedSelected(int32 InPlaybackSpeedMode) const = 0;
 };
 
