When anim notifies snap to montage section markers, set trigger times so that notify doesn't cross over the snap (i.e. start will be after the section time, end will be before the section time).

diff --git a/Engine/Source/Editor/Persona/Private/AnimTimeline/AnimModel.cpp b/Engine/Source/Editor/Persona/Private/AnimTimeline/AnimModel.cpp
--- a/Engine/Source/Editor/Persona/Private/AnimTimeline/AnimModel.cpp
+++ b/Engine/Source/Editor/Persona/Private/AnimTimeline/AnimModel.cpp
@@ -239,20 +239,21 @@ void FAnimModel::SetEditableTime(int32 TimeIndex, double Time, bool bIsDragging)
 	OnSetEditableTime(TimeIndex, EditableTimes[TimeIndex], bIsDragging);
 }
 
-bool FAnimModel::Snap(float& InOutTime, float InSnapMargin, TArrayView<const FName> InSkippedSnapTypes) const
+bool FAnimModel::Snap(float& InOutTime, float InSnapMargin, TArrayView<const FName> InSkippedSnapTypes, FName& OutSnapType) const
 {
 	double DoubleTime = InOutTime;
-	bool bResult = Snap(DoubleTime, (double)InSnapMargin, InSkippedSnapTypes);
+	bool bResult = Snap(DoubleTime, (double)InSnapMargin, InSkippedSnapTypes, OutSnapType);
 	InOutTime = static_cast<float>(DoubleTime);
 	return bResult;
 }
 
-bool FAnimModel::Snap(double& InOutTime, double InSnapMargin, TArrayView<const FName> InSkippedSnapTypes) const
+bool FAnimModel::Snap(double& InOutTime, double InSnapMargin, TArrayView<const FName> InSkippedSnapTypes, FName& OutSnapType) const
 {
 	InSnapMargin = FMath::Max(InSnapMargin, (double)KINDA_SMALL_NUMBER);
 
 	double ClosestDelta = DBL_MAX;
 	double ClosestSnapTime = DBL_MAX;
+	FName ClosestSnapType = NAME_None;
 
 	// Check for enabled snap functions first
 	for(const TPair<FName, FSnapType>& SnapTypePair : SnapTypes)
@@ -269,6 +270,7 @@ bool FAnimModel::Snap(double& InOutTime, double InSnapMargin, TArrayView<const F
 						double Delta = FMath::Abs(SnappedTime - InOutTime);
 						if(Delta < InSnapMargin && Delta < ClosestDelta)
 						{
+							ClosestSnapType = SnapTypePair.Value.Type;
 							ClosestDelta = Delta;
 							ClosestSnapTime = SnappedTime;
 						}
@@ -290,6 +292,7 @@ bool FAnimModel::Snap(double& InOutTime, double InSnapMargin, TArrayView<const F
 				{
 					if(IsSnapChecked(SnapTime.Type))
 					{
+						ClosestSnapType = SnapTime.Type;
 						ClosestDelta = Delta;
 						ClosestSnapTime = SnapTime.Time;
 					}
@@ -300,6 +303,7 @@ bool FAnimModel::Snap(double& InOutTime, double InSnapMargin, TArrayView<const F
 
 	if(ClosestDelta != DBL_MAX)
 	{
+		OutSnapType = ClosestSnapType;
 		InOutTime = ClosestSnapTime;
 		return true;
 	}
diff --git a/Engine/Source/Editor/Persona/Private/AnimTimeline/AnimModel.h b/Engine/Source/Editor/Persona/Private/AnimTimeline/AnimModel.h
--- a/Engine/Source/Editor/Persona/Private/AnimTimeline/AnimModel.h
+++ b/Engine/Source/Editor/Persona/Private/AnimTimeline/AnimModel.h
@@ -177,12 +177,14 @@ public:
 
 	/** 
 	 * Snaps a time to the currently set snap times
-	 * @param	InOutTime		The time to snap
-	 * @param	InSnapMargin	The margin (in seconds) to limit the snap to
+	 * @param	InOutTime				The time to snap
+	 * @param	InSnapMargin			The margin (in seconds) to limit the snap to
+	 * @param	InSkippedSnapTypes		Array of snap types that should be ignored when checking for snap
+	 * @param	OutSnapType				The type of snap that was snapped to
 	 * @return true if a snap occurred
 	 */
-	bool Snap(float& InOutTime, float InSnapMargin, TArrayView<const FName> InSkippedSnapTypes) const;
-	bool Snap(double& InOutTime, double InSnapMargin, TArrayView<const FName> InSkippedSnapTypes) const;
+	bool Snap(float& InOutTime, float InSnapMargin, TArrayView<const FName> InSkippedSnapTypes, FName& OutSnapType) const;
+	bool Snap(double& InOutTime, double InSnapMargin, TArrayView<const FName> InSkippedSnapTypes, FName& OutSnapType) const;
 
 	/** Toggle the specified snap type */
 	void ToggleSnap(FName InSnapName);
diff --git a/Engine/Source/Editor/Persona/Private/AnimTimeline/AnimTimeSliderController.cpp b/Engine/Source/Editor/Persona/Private/AnimTimeline/AnimTimeSliderController.cpp
--- a/Engine/Source/Editor/Persona/Private/AnimTimeline/AnimTimeSliderController.cpp
+++ b/Engine/Source/Editor/Persona/Private/AnimTimeline/AnimTimeSliderController.cpp
@@ -633,7 +633,8 @@ FReply FAnimTimeSliderController::OnMouseButtonUp( SWidget& WidgetOwner, const F
 				if(!MouseEvent.IsControlDown())
 				{
 					const double SnapMargin = (ScrubConstants::SnapMarginInPixels / static_cast<double>(RangeToScreen.PixelsPerInput));
-					WeakModel.Pin()->Snap(Time, SnapMargin, { FName("MontageSection") });
+					FName SnapTypeUnused;
+					WeakModel.Pin()->Snap(Time, SnapMargin, { FName("MontageSection") }, SnapTypeUnused);
 				}
 
 				SetEditableTime(DraggedTimeIndex, static_cast<float>(Time), false);
@@ -826,7 +827,8 @@ FReply FAnimTimeSliderController::OnMouseMove( SWidget& WidgetOwner, const FGeom
 				if(!MouseEvent.IsControlDown())
 				{
 					const double SnapMargin = (ScrubConstants::SnapMarginInPixels / static_cast<double>(RangeToScreen.PixelsPerInput));
-					WeakModel.Pin()->Snap(Time, SnapMargin, { FName("MontageSection") });
+					FName SnapTypeUnused;
+					WeakModel.Pin()->Snap(Time, SnapMargin, { FName("MontageSection") }, SnapTypeUnused);
 				}
 
 				SetEditableTime(DraggedTimeIndex, static_cast<float>(Time), true);
diff --git a/Engine/Source/Editor/Persona/Private/SAnimNotifyPanel.cpp b/Engine/Source/Editor/Persona/Private/SAnimNotifyPanel.cpp
--- a/Engine/Source/Editor/Persona/Private/SAnimNotifyPanel.cpp
+++ b/Engine/Source/Editor/Persona/Private/SAnimNotifyPanel.cpp
@@ -584,7 +584,7 @@ private:
 	float HandleOverflowPan( const FVector2D& ScreenCursorPos, float TrackScreenSpaceXPosition, float TrackScreenSpaceMin, float TrackScreenSpaceMax);
 
 	/** Finds a snap position if possible for the provided scrub handle, if it is not possible, returns -1.0f */
-	float GetScrubHandleSnapPosition(float NotifyInputX, ENotifyStateHandleHit::Type HandleToCheck);
+	float GetScrubHandleSnapPosition(float NotifyInputX, ENotifyStateHandleHit::Type HandleToCheck, TOptional<EAnimEventTriggerOffsets::Type>& OffsetForSnap);
 
 	virtual FReply OnFocusReceived(const FGeometry& MyGeometry, const FFocusEvent& InFocusEvent) override;
 
@@ -1358,7 +1358,8 @@ public:
 
 		float CurrentMinSnapDest = MaxSnapDist;
 		float SnapPosition = ScaleInfo.LocalXToInput(WidgetSpaceNotifyPosition);
-		bOutSnapped = OnSnapPosition.IsBound() && !FSlateApplication::Get().GetModifierKeys().IsControlDown() && OnSnapPosition.Execute(SnapPosition, MaxSnapDist / ScaleInfo.PixelsPerInput, TArrayView<const FName>());
+		FName SnapTypeUnused;
+		bOutSnapped = OnSnapPosition.IsBound() && !FSlateApplication::Get().GetModifierKeys().IsControlDown() && OnSnapPosition.Execute(SnapPosition, MaxSnapDist / ScaleInfo.PixelsPerInput, TArrayView<const FName>(), SnapTypeUnused);
 		SnapPosition = ScaleInfo.InputToLocalX(SnapPosition);
 
 		float WidgetSpaceStartPosition = ScaleInfo.InputToLocalX(0.0f);
@@ -1951,11 +1952,21 @@ FReply SAnimNotifyNode::OnMouseMove( const FGeometry& MyGeometry, const FPointer
 		if (FAnimNotifyEvent* AnimNotifyEvent = NodeObjectInterface->GetNotifyEvent())
 		{
 			float InputStartTime = AnimNotifyEvent->GetTime();
-			float MarkerSnap = GetScrubHandleSnapPosition(InputStartTime, ENotifyStateHandleHit::Start);
+			TOptional<EAnimEventTriggerOffsets::Type> OffsetForSnap;
+			float MarkerSnap = GetScrubHandleSnapPosition(InputStartTime, ENotifyStateHandleHit::Start, OffsetForSnap);
 			if (MarkerSnap != -1.0f)
 			{
 				// We're near to a snap bar
-				EAnimEventTriggerOffsets::Type Offset = (MarkerSnap < InputStartTime) ? EAnimEventTriggerOffsets::OffsetAfter : EAnimEventTriggerOffsets::OffsetBefore;
+				EAnimEventTriggerOffsets::Type Offset = EAnimEventTriggerOffsets::NoOffset;
+				if (OffsetForSnap.IsSet())
+				{				
+					Offset = OffsetForSnap.GetValue();
+				}
+				else
+				{
+					Offset = (MarkerSnap < InputStartTime) ? EAnimEventTriggerOffsets::OffsetAfter : EAnimEventTriggerOffsets::OffsetBefore;
+				}
+
 				AnimNotifyEvent->TriggerTimeOffset = GetTriggerTimeOffsetForType(Offset);
 
 				// Adjust our start marker
@@ -2000,11 +2011,20 @@ FReply SAnimNotifyNode::OnMouseMove( const FGeometry& MyGeometry, const FPointer
 		if (FAnimNotifyEvent* AnimNotifyEvent = NodeObjectInterface->GetNotifyEvent())
 		{
 			float InputEndTime = AnimNotifyEvent->GetTime() + AnimNotifyEvent->GetDuration();
-			float MarkerSnap = GetScrubHandleSnapPosition(InputEndTime, ENotifyStateHandleHit::End);
+			TOptional<EAnimEventTriggerOffsets::Type> OffsetForSnap;
+			float MarkerSnap = GetScrubHandleSnapPosition(InputEndTime, ENotifyStateHandleHit::End, OffsetForSnap);
 			if (MarkerSnap != -1.0f)
 			{
 				// We're near to a snap bar
-				EAnimEventTriggerOffsets::Type Offset = (MarkerSnap < InputEndTime) ? EAnimEventTriggerOffsets::OffsetAfter : EAnimEventTriggerOffsets::OffsetBefore;
+				EAnimEventTriggerOffsets::Type Offset = EAnimEventTriggerOffsets::NoOffset;
+				if (OffsetForSnap.IsSet())
+				{
+					Offset = OffsetForSnap.GetValue();
+				}
+				else
+				{
+					Offset = (MarkerSnap < InputEndTime) ? EAnimEventTriggerOffsets::OffsetAfter : EAnimEventTriggerOffsets::OffsetBefore;
+				}
 				AnimNotifyEvent->EndTriggerTimeOffset = GetTriggerTimeOffsetForType(Offset);
 
 				// Adjust our end marker
@@ -2050,7 +2070,7 @@ FReply SAnimNotifyNode::OnMouseButtonUp( const FGeometry& MyGeometry, const FPoi
 	return FReply::Unhandled();
 }
 
-float SAnimNotifyNode::GetScrubHandleSnapPosition( float NotifyInputX, ENotifyStateHandleHit::Type HandleToCheck )
+float SAnimNotifyNode::GetScrubHandleSnapPosition( float NotifyInputX, ENotifyStateHandleHit::Type HandleToCheck, TOptional<EAnimEventTriggerOffsets::Type>& OffsetForSnap)
 {
 	FTrackScaleInfo ScaleInfo(ViewInputMin.Get(), ViewInputMax.Get(), 0, 0, CachedAllotedGeometrySize);
 
@@ -2058,8 +2078,22 @@ float SAnimNotifyNode::GetScrubHandleSnapPosition( float NotifyInputX, ENotifySt
 
 	if(OnSnapPosition.IsBound() && !FSlateApplication::Get().GetModifierKeys().IsControlDown())
 	{
-		if(OnSnapPosition.Execute(NotifyInputX, MaxSnapDist / ScaleInfo.PixelsPerInput, TArrayView<const FName>()))
+		FName SnapType = NAME_None;
+		if(OnSnapPosition.Execute(NotifyInputX, MaxSnapDist / ScaleInfo.PixelsPerInput, TArrayView<const FName>(), SnapType))
 		{
+			static const FName MontageSectionSnapTypeName("MontageSection");
+			if (SnapType == MontageSectionSnapTypeName)
+			{
+				if (HandleToCheck == ENotifyStateHandleHit::Start)
+				{
+					OffsetForSnap = EAnimEventTriggerOffsets::OffsetAfter;
+				}
+				else if (HandleToCheck == ENotifyStateHandleHit::End)
+				{
+					OffsetForSnap = EAnimEventTriggerOffsets::OffsetBefore;
+				}
+			}
+
 			return NotifyInputX;
 		}
 	}	
diff --git a/Engine/Source/Editor/Persona/Private/SAnimNotifyPanel.h b/Engine/Source/Editor/Persona/Private/SAnimNotifyPanel.h
--- a/Engine/Source/Editor/Persona/Private/SAnimNotifyPanel.h
+++ b/Engine/Source/Editor/Persona/Private/SAnimNotifyPanel.h
@@ -40,7 +40,7 @@ DECLARE_DELEGATE_OneParam( FReplaceWithSyncMarker, FString )
 DECLARE_DELEGATE( FDeselectAllNotifies )
 DECLARE_DELEGATE_OneParam( FOnGetBlueprintNotifyData, TArray<FAssetData>& )
 DECLARE_DELEGATE_OneParam( FOnGetNativeNotifyClasses, TArray<UClass*>&)
-DECLARE_DELEGATE_RetVal_ThreeParams(bool, FOnSnapPosition, float& /*InOutTimeToSnap*/, float /*InSnapMargin*/, TArrayView<const FName> /*InSkippedSnapTypes*/)
+DECLARE_DELEGATE_RetVal_FourParams(bool, FOnSnapPosition, float& /*InOutTimeToSnap*/, float /*InSnapMargin*/, TArrayView<const FName> /*InSkippedSnapTypes*/, FName& /*OutSnapType*/)
 
 class SAnimNotifyNode;
 class SAnimNotifyTrack;
