[Persona perview controller] Added editor world to collect preview-able skeletal mesh instance, useful for editing mesh running in sequencer.

diff --git a/Engine/Source/Editor/Persona/Private/PersonaPreviewSceneSkelMeshInstanceController.cpp b/Engine/Source/Editor/Persona/Private/PersonaPreviewSceneSkelMeshInstanceController.cpp
--- a/Engine/Source/Editor/Persona/Private/PersonaPreviewSceneSkelMeshInstanceController.cpp
+++ b/Engine/Source/Editor/Persona/Private/PersonaPreviewSceneSkelMeshInstanceController.cpp
@@ -164,7 +164,7 @@ void SSkeletalMeshDebugSelectionWidget::Refresh()
 			UWorld* World = *It;
 			// include only PIE and worlds that own the persistent level (i.e. non-streaming levels).
 			const bool bIsValidDebugWorld = (World != nullptr)
-				&& (World->WorldType == EWorldType::PIE || World->WorldType == EWorldType::EditorPreview)
+				&& (World->WorldType == EWorldType::PIE || World->WorldType == EWorldType::Editor || World->WorldType == EWorldType::EditorPreview)
 				&& World->PersistentLevel != nullptr
 				&& World->PersistentLevel->OwningWorld == World;
 			if (!bIsValidDebugWorld)
