[Skeletal Mesh Editor] Fixed running instance disconnnects when undo / entering tool

diff --git a/Engine/Source/Editor/Persona/Private/PersonaPreviewSceneSkelMeshInstanceController.cpp b/Engine/Source/Editor/Persona/Private/PersonaPreviewSceneSkelMeshInstanceController.cpp
--- a/Engine/Source/Editor/Persona/Private/PersonaPreviewSceneSkelMeshInstanceController.cpp
+++ b/Engine/Source/Editor/Persona/Private/PersonaPreviewSceneSkelMeshInstanceController.cpp
@@ -277,6 +277,12 @@ void SSkeletalMeshDebugSelectionWidget::Refresh()
 			ItemToActivate = Instance;
 			break;
 		}
+
+		if (PreviewComponent->PreviewInstance && PreviewComponent->PreviewInstance->GetDebugSkeletalMeshComponent() == Instance->SkeletalMeshComponent)
+		{
+			ItemToActivate = Instance;
+			break;
+		}
 	}
 
 	// assign selection to combobox
