[Persona SKM intance preview controller] picker now shows actor label instead of actor name. connection to live instance also disconnects if the preview controller change

diff --git a/Engine/Source/Editor/Persona/Private/PersonaPreviewSceneSkelMeshInstanceController.cpp b/Engine/Source/Editor/Persona/Private/PersonaPreviewSceneSkelMeshInstanceController.cpp
--- a/Engine/Source/Editor/Persona/Private/PersonaPreviewSceneSkelMeshInstanceController.cpp
+++ b/Engine/Source/Editor/Persona/Private/PersonaPreviewSceneSkelMeshInstanceController.cpp
@@ -31,6 +31,16 @@ void UPersonaPreviewSceneSkelMeshInstanceController::UninitializeView(
 	UPersonaPreviewSceneDescription* SceneDescription,
 	IPersonaPreviewScene* PreviewScene) const
 {
+	if (UDebugSkelMeshComponent* PreviewMeshComponent = PreviewScene->GetPreviewMeshComponent())
+	{
+        PreviewMeshComponent->bTrackAttachedInstanceLOD = false;
+		
+		if (UAnimPreviewInstance* PreviewAnimInstance = PreviewMeshComponent->PreviewInstance)
+		{
+			PreviewAnimInstance->SetDebugSkeletalMeshComponent(nullptr);
+		}
+	}
+
 	PreviewScene->ShowDefaultMode();
 }
 
@@ -210,7 +220,7 @@ void SSkeletalMeshDebugSelectionWidget::Refresh()
 					TSharedPtr<FSkeletalMeshDebugInstance> NewInstance = MakeShared<FSkeletalMeshDebugInstance>();
 					NewInstance->SkeletalMeshComponent = Component;
 					AActor* Actor = NewInstance->GetActor();
-					NewInstance->DisplayName = IsValid(Actor) ? FText::FromString(Actor->GetName()) : FText();
+					NewInstance->DisplayName = IsValid(Actor) ? FText::FromString(Actor->GetActorNameOrLabel()) : FText();
 					AllMeshInstances.Emplace(NewInstance);
 				}
 			});
@@ -294,7 +304,7 @@ TSharedRef<SWidget> SSkeletalMeshDebugSelectionWidget::OnGenerateComboBoxItemWid
 	// If we have the first item in the actor list, generate a special widget.
 	const TWeakObjectPtr<USkeletalMeshComponent> Component = Item.IsValid() ? Item->SkeletalMeshComponent : nullptr;
 	const AActor* Actor = Component.IsValid() ? Component->GetOwner() : nullptr;
-	FText ActorName = Actor ? FText::FromString(Actor->GetName()) : LOCTEXT("DestroyedActorText", "<Destroyed>");
+	FText ActorName = Actor ? FText::FromString(Actor->GetActorNameOrLabel()) : LOCTEXT("DestroyedActorText", "<Destroyed>");
 
 	const UWorld* World = Actor ? Actor->GetWorld() : nullptr;
 	FText WorldName = World ? FText::FromString(GetDebugStringForWorld(World)) : LOCTEXT("DestroyedWorldText", "<Destroyed>");
