Add support for sorting sections via a sort delegate.
- Ensure that the animation editor scene elements section is alphabetically sorted.
- Reorder the scene elements section to be at the top.

#rnx

diff --git a/Engine/Source/Developer/ToolMenus/Private/ToolMenuSection.cpp b/Engine/Source/Developer/ToolMenus/Private/ToolMenuSection.cpp
--- a/Engine/Source/Developer/ToolMenus/Private/ToolMenuSection.cpp
+++ b/Engine/Source/Developer/ToolMenus/Private/ToolMenuSection.cpp
@@ -37,6 +37,7 @@ void FToolMenuSection::InitGeneratedSectionCopy(const FToolMenuSection& Source,
 	Alignment = Source.Alignment;
 	Visibility = Source.Visibility;
 	ResizeParams = Source.ResizeParams;
+	Sorter = Source.Sorter;
 }
 
 bool FToolMenuSection::IsRegistering() const
diff --git a/Engine/Source/Developer/ToolMenus/Private/ToolMenus.cpp b/Engine/Source/Developer/ToolMenus/Private/ToolMenus.cpp
--- a/Engine/Source/Developer/ToolMenus/Private/ToolMenus.cpp
+++ b/Engine/Source/Developer/ToolMenus/Private/ToolMenus.cpp
@@ -1149,7 +1149,7 @@ void UToolMenus::ApplyCustomizationAndProfiles(UToolMenu* GeneratedMenu)
 			}
 		}
 	}
-
+	
 	// Apply the customization for the menu (if any)
 	FCustomizedToolMenuHierarchy CustomizationHierarchy = GeneratedMenu->GetMenuCustomizationHierarchy();
 	if (CustomizationHierarchy.Hierarchy.Num() != 0 || CustomizationHierarchy.RuntimeHierarchy.Num() != 0)
@@ -1312,6 +1312,17 @@ void UToolMenus::AssembleMenuHierarchy(UToolMenu* GeneratedMenu, const TArray<UT
 		AssembleMenu(GeneratedMenu, FoundParent);
 	}
 
+	for (FToolMenuSection& Section : GeneratedMenu->Sections)
+	{
+		if (Section.Sorter.IsBound())
+		{
+			Section.Blocks.StableSort([&Section](const FToolMenuEntry& A, const FToolMenuEntry& B)
+			{
+				return Section.Sorter.Execute(A, B, Section.Context);
+			});
+		}
+	}
+	
 	ApplyCustomizationAndProfiles(GeneratedMenu);
 }
 
diff --git a/Engine/Source/Developer/ToolMenus/Public/ToolMenuDelegates.h b/Engine/Source/Developer/ToolMenus/Public/ToolMenuDelegates.h
--- a/Engine/Source/Developer/ToolMenus/Public/ToolMenuDelegates.h
+++ b/Engine/Source/Developer/ToolMenus/Public/ToolMenuDelegates.h
@@ -32,6 +32,8 @@ DECLARE_DELEGATE_RetVal_OneParam(bool, FToolMenuIsActionChecked, const FToolMenu
 DECLARE_DELEGATE_RetVal_OneParam(ECheckBoxState, FToolMenuGetActionCheckState, const FToolMenuContext&);
 DECLARE_DELEGATE_RetVal_OneParam(bool, FToolMenuIsActionButtonVisible, const FToolMenuContext&);
 
+DECLARE_DELEGATE_RetVal_ThreeParams(bool, FToolMenuSectionSorter, const FToolMenuEntry&, const FToolMenuEntry&, const FToolMenuContext&);
+
 DECLARE_DELEGATE_TwoParams(FToolMenuExecuteString, const FString&, const FToolMenuContext&);
 
 DECLARE_DYNAMIC_DELEGATE_OneParam(FToolMenuDynamicExecuteAction, const FToolMenuContext&, Context);
diff --git a/Engine/Source/Developer/ToolMenus/Public/ToolMenuSection.h b/Engine/Source/Developer/ToolMenus/Public/ToolMenuSection.h
--- a/Engine/Source/Developer/ToolMenus/Public/ToolMenuSection.h
+++ b/Engine/Source/Developer/ToolMenus/Public/ToolMenuSection.h
@@ -144,6 +144,9 @@ public:
 
 	/** Section-level resize parameters. */
 	FMenuEntryResizeParams ResizeParams;
+	
+	/** A sorting function for blocks within this section. Blocks are sorted after all entries have been added. */
+	FToolMenuSectionSorter Sorter;
 
 private:
 
diff --git a/Engine/Source/Editor/Persona/Private/ViewportToolbar/AnimationEditorMenus.cpp b/Engine/Source/Editor/Persona/Private/ViewportToolbar/AnimationEditorMenus.cpp
--- a/Engine/Source/Editor/Persona/Private/ViewportToolbar/AnimationEditorMenus.cpp
+++ b/Engine/Source/Editor/Persona/Private/ViewportToolbar/AnimationEditorMenus.cpp
@@ -514,8 +514,8 @@ FToolMenuEntry UE::AnimationEditor::CreateShowSubmenu()
 		FNewToolMenuDelegate::CreateLambda(
 			[](UToolMenu* Submenu) -> void
 			{
-				UE::AnimationEditor::FillShowSubmenu(Submenu);
 				AddSceneElementsSection(Submenu);
+				UE::AnimationEditor::FillShowSubmenu(Submenu);
 			}
 		)
 	);
@@ -542,11 +542,10 @@ void UE::AnimationEditor::FillShowSubmenu(UToolMenu* InMenu, bool bInShowViewpor
 				if (bInShowViewportStatsToggle)
 				{
 					FToolMenuSection& UnnamedSection = InMenu->AddSection(NAME_None);
+					UnnamedSection.AddSeparator(NAME_None);
 					UnnamedSection.AddMenuEntry(
 						FEditorViewportCommands::Get().ToggleStats, LOCTEXT("ViewportStatsLabel", "Viewport Stats")
 					);
-
-					UnnamedSection.AddSeparator(NAME_None);
 				}
 
 				// Only include helpful show flags.
@@ -927,13 +926,6 @@ void UE::AnimationEditor::AddSceneElementsSection(UToolMenu* InMenu)
 		)
 	);
 
-	Section.Blocks.Sort(
-		[](const FToolMenuEntry& A, const FToolMenuEntry& B)
-		{
-			return A.Label.Get().ToLower().ToString() < B.Label.Get().ToLower().ToString();
-		}
-	);
-
 	Section.AddDynamicEntry(
 		"AdvancedSubMenu",
 		FNewToolMenuSectionDelegate::CreateLambda(
@@ -949,6 +941,20 @@ void UE::AnimationEditor::AddSceneElementsSection(UToolMenu* InMenu)
 			}
 		)
 	);
+	
+	Section.Sorter.BindLambda([](const FToolMenuEntry& A, const FToolMenuEntry& B, const FToolMenuContext& Context)
+	{
+		if (A.Name == "AdvancedSubMenu")
+		{
+			return false;
+		}
+		if (B.Name == "AdvancedSubMenu")
+		{
+			return true;
+		}
+		
+		return A.Label.Get().CompareTo(B.Label.Get()) < 0;
+	});
 }
 
 FToolMenuEntry UE::AnimationEditor::CreateLODSubmenu()
