Change behavior so overlapping AnimNotify(state) type properties can be edited during multiselect

diff --git a/Engine/Source/Editor/Persona/Private/AnimNotifyDetails.cpp b/Engine/Source/Editor/Persona/Private/AnimNotifyDetails.cpp
--- a/Engine/Source/Editor/Persona/Private/AnimNotifyDetails.cpp
+++ b/Engine/Source/Editor/Persona/Private/AnimNotifyDetails.cpp
@@ -14,6 +14,7 @@
 #include "AssetSearchBoxUtilPersona.h"
 #include "IDetailGroup.h"
 #include "ObjectEditorUtils.h"
+#include "Animation/AnimNotifies/AnimNotify.h"
 #include "Widgets/Input/STextComboBox.h"
 #include "Animation/AnimNotifies/AnimNotifyState.h"
 
@@ -176,23 +177,44 @@ void FAnimNotifyDetails::CustomizeDetails( IDetailLayoutBuilder& DetailBuilder )
 	const UClass* DetailObjectClass = nullptr;
 	const UClass* BaseClass = DetailBuilder.GetBaseClass();
 	TArray<TWeakObjectPtr<UObject>> SelectedObjects;
-	TArray<UClass*> NotifyClasses;
 	DetailBuilder.GetObjectsBeingCustomized(SelectedObjects);
 
 	check(SelectedObjects.Num() > 0);
-	UEditorNotifyObject* EditorObject = Cast<UEditorNotifyObject>(SelectedObjects[0].Get());
-	check(EditorObject);
-	UpdateSlotNames(EditorObject->AnimObject);
+
+	UObject* CommonOuter = nullptr;
+	for (const TWeakObjectPtr<UObject>& WeakObject : SelectedObjects)
+	{
+		if (const UEditorNotifyObject* EditorObject = Cast<UEditorNotifyObject>(WeakObject.Get()))
+		{
+			UpdateSlotNames(EditorObject->AnimObject);
+
+			// Find common outer object
+			if (CommonOuter == nullptr)
+			{
+				CommonOuter = EditorObject->AnimObject;
+			}
+			else if (EditorObject->AnimObject && CommonOuter != EditorObject->AnimObject)
+			{
+				CommonOuter = nullptr;
+			}
+		}		
+	}
+	
 
 	TSharedRef<IPropertyHandle> EventHandle = DetailBuilder.GetProperty(TEXT("Event"));
 	IDetailCategoryBuilder& EventCategory = DetailBuilder.EditCategory(TEXT("Category"));
 	EventCategory.AddProperty(EventHandle).OverrideResetToDefault(FResetToDefaultOverride::Hide());
 
 	// Hide notify objects that aren't set
-	UObject* NotifyPtr = nullptr;
+	UObject* NotifyObject = nullptr;
+	UObject* NotifyStateObject = nullptr;
 	FString NotifyClassName;
+	TSharedPtr<IPropertyHandle> PropertyHandleToUse = nullptr;
 	TSharedRef<IPropertyHandle> NotifyPropHandle = DetailBuilder.GetProperty(TEXT("Event.Notify"));
-	NotifyPropHandle->GetValue(NotifyPtr);
+	TSharedRef<IPropertyHandle> NotifyStatePropHandle = DetailBuilder.GetProperty(TEXT("Event.NotifyStateClass"));
+
+	const FPropertyAccess::Result NotifyPropertyAccessResult = NotifyPropHandle->GetValue(NotifyObject);
+	const FPropertyAccess::Result NotifyStatePropertyAccessResult = NotifyStatePropHandle->GetValue(NotifyStateObject);
 	
 	// Don't want to edit the notify name here.
 	DetailBuilder.HideProperty(TEXT("Event.NotifyName"));
@@ -200,44 +222,94 @@ void FAnimNotifyDetails::CustomizeDetails( IDetailLayoutBuilder& DetailBuilder )
 	IDetailCategoryBuilder& AnimNotifyCategory = DetailBuilder.EditCategory(
 		UE::AnimNotifyDetails::Private::AnimNotifyCategoryName, FText::GetEmpty(), ECategoryPriority::TypeSpecific);
 
-	// Check existence of notify, get rid of the property if not set
-	if(!NotifyPtr)
+	const bool bValidNotifyObjects = NotifyPropertyAccessResult != FPropertyAccess::Fail;
+	const bool bValidNotifyStateClasses = NotifyStatePropertyAccessResult != FPropertyAccess::Fail;
+
+	UObject* NotifyPtr = nullptr;
+	DetailBuilder.HideProperty(TEXT("Event.Notify"));
+	if (bValidNotifyStateClasses)
 	{
-		DetailBuilder.HideProperty(TEXT("Event.Notify"));
+		DetailBuilder.HideProperty(TEXT("Event.NotifyStateClass"));
+		DetailBuilder.HideProperty(TEXT("Event.EndLink"));
 
-		NotifyPropHandle = DetailBuilder.GetProperty(TEXT("Event.NotifyStateClass"));
-		NotifyPropHandle->GetValue(NotifyPtr);
+		UClass* NotifyStateClass = nullptr;
+		UAnimNotifyState* NotifyState = nullptr;
 
-		// Check existence of notify state, get rid of the property if not set
-		if(!NotifyPtr)
+		// Find common notify state class, only valid if all of the selected objects are of the same notify type
+		for (const TWeakObjectPtr<UObject>& WeakObject : SelectedObjects)
 		{
-			DetailBuilder.HideProperty(TEXT("Event.NotifyStateClass"));
-			DetailBuilder.HideProperty(TEXT("Event.EndLink"));
+			if (const UEditorNotifyObject* EditorObject = Cast<UEditorNotifyObject>(WeakObject.Get()))
+			{
+				NotifyState = EditorObject->Event.NotifyStateClass.Get();
+				if (NotifyStateClass == nullptr)
+				{
+					if (NotifyState)
+					{
+						NotifyStateClass = NotifyState->GetClass();
+					}
+					else
+					{
+						break;
+					}
+				}
+				else if (NotifyState && NotifyState->GetClass() != NotifyStateClass)
+				{
+					NotifyStateClass = nullptr;
+					break;
+				}
+			}
 		}
-		else
-		{
-			DetailObjectClass = NotifyPtr->GetClass();
 
-			// Get rid of the class selector in the details panel. It's not necessary for notifies
-			ClearInstancedSelectionDropDown(AnimNotifyCategory, NotifyPropHandle);
-		}
+		NotifyPtr = NotifyState;
+		PropertyHandleToUse = NotifyStatePropHandle;
+		DetailObjectClass = NotifyStateClass;
 	}
-	else
+
+	// In case there wasn't a valid AnimNotifyState class try to find an AnimNotify class instead
+	if (bValidNotifyObjects && NotifyPtr == nullptr)
 	{
-		DetailObjectClass = NotifyPtr->GetClass();
-		// Get rid of the class selector in the details panel. It's not necessary for notifies
-		ClearInstancedSelectionDropDown(AnimNotifyCategory, NotifyPropHandle);
+		UClass* NotifyClass = nullptr;
+		UAnimNotify* Notify = nullptr;
+		
+		// Find common notify class, only valid if all of the selected objects are of the same notify type
+		for (const TWeakObjectPtr<UObject>& WeakObject : SelectedObjects)
+		{
+			if (const UEditorNotifyObject* EditorObject = Cast<UEditorNotifyObject>(WeakObject.Get()))
+			{
+				Notify = EditorObject->Event.Notify.Get();
+				if (NotifyClass == nullptr)
+				{
+					if (Notify)
+					{
+						NotifyClass = Notify->GetClass();
+					}
+					else
+					{
+						break;
+					}
+				}
+				else if (Notify && Notify->GetClass() != NotifyClass)
+				{
+					NotifyClass = nullptr;
+					break;
+				}
+			}
+		}
 
-		// No state present, hide the entry
-		DetailBuilder.HideProperty(TEXT("Event.NotifyStateClass"));
+		NotifyPtr = Notify;
+		PropertyHandleToUse = NotifyPropHandle;
+		DetailObjectClass = NotifyClass;
+		
 	}
+	// This means we ended up in a situation where there is no valid notify or notify-state class - so why are we trying to customize a notify
+	//check(NotifyPtr != nullptr);
 
-	UAnimMontage* CurrentMontage = Cast<UAnimMontage>(EditorObject->AnimObject);
+	UAnimMontage* CurrentMontage = Cast<UAnimMontage>(CommonOuter);
 
 	// If we have a montage, and it has slots (which it should have) generate custom link properties
 	if(CurrentMontage && CurrentMontage->SlotAnimTracks.Num() > 0)
 	{
-		CustomizeLinkProperties(DetailBuilder, EventHandle, EditorObject);
+		CustomizeLinkProperties(DetailBuilder, EventHandle, CurrentMontage);
 	}
 	else
 	{
@@ -252,14 +324,14 @@ void FAnimNotifyDetails::CustomizeDetails( IDetailLayoutBuilder& DetailBuilder )
 
 	// Customizations do not run for instanced properties, so we have to resolve the properties and then
 	// customize them here instead.
-	if(NotifyPropHandle->IsValidHandle())
+	if(PropertyHandleToUse->IsValidHandle())
 	{
 		uint32 NumChildren = 0;
-		NotifyPropHandle->GetNumChildren(NumChildren);
+		PropertyHandleToUse->GetNumChildren(NumChildren);
 		if(NumChildren > 0)
 		{
-			TSharedPtr<IPropertyHandle> BaseHandle = NotifyPropHandle->GetChildHandle(0);
-			DetailBuilder.HideProperty(NotifyPropHandle);
+			TSharedPtr<IPropertyHandle> BaseHandle = PropertyHandleToUse->GetChildHandle(0);
+			DetailBuilder.HideProperty(PropertyHandleToUse);
 
 			BaseHandle->GetNumChildren(NumChildren);
 			DetailBuilder.HideProperty(BaseHandle);
@@ -473,7 +545,7 @@ void FAnimNotifyDetails::ClearInstancedSelectionDropDown(IDetailCategoryBuilder&
 	];
 }
 
-void FAnimNotifyDetails::CustomizeLinkProperties(IDetailLayoutBuilder& Builder, TSharedRef<IPropertyHandle> NotifyProperty, UEditorNotifyObject* EditorObject)
+void FAnimNotifyDetails::CustomizeLinkProperties(IDetailLayoutBuilder& Builder, TSharedRef<IPropertyHandle> NotifyProperty, UAnimSequenceBase* AnimSequenceBase)
 {
 	uint32 NumChildProperties = 0;
 	NotifyProperty->GetNumChildren(NumChildProperties);
@@ -495,6 +567,16 @@ void FAnimNotifyDetails::CustomizeLinkProperties(IDetailLayoutBuilder& Builder,
 				{
 					int32 SlotIdx = INDEX_NONE;
 					ChildHandle->GetValue(SlotIdx);
+					
+					TSharedPtr<FString> InitialItem = nullptr;
+					if (SlotNameItems.IsValidIndex(SlotIdx))
+					{
+						InitialItem = SlotNameItems[SlotIdx];
+					}
+					else
+					{
+						InitialItem = MakeShared<FString>(TEXT("Multiple Values"));
+					}
 
 					LinkCategory.AddProperty(ChildHandle)
 						.CustomWidget()
@@ -507,8 +589,9 @@ void FAnimNotifyDetails::CustomizeLinkProperties(IDetailLayoutBuilder& Builder,
 							SNew(STextComboBox)
 							.OptionsSource(&SlotNameItems)
 							.OnSelectionChanged(this, &FAnimNotifyDetails::OnSlotSelected, ChildHandle)
-							.OnComboBoxOpening(this, &FAnimNotifyDetails::UpdateSlotNames, EditorObject->AnimObject)
-							.InitiallySelectedItem(SlotNameItems[SlotIdx])
+							.OnComboBoxOpening(this, &FAnimNotifyDetails::UpdateSlotNames, AnimSequenceBase)
+							.InitiallySelectedItem(InitialItem)
+							.Font(Builder.GetDetailFont())
 						];
 				}
 				else
diff --git a/Engine/Source/Editor/Persona/Private/AnimNotifyDetails.h b/Engine/Source/Editor/Persona/Private/AnimNotifyDetails.h
--- a/Engine/Source/Editor/Persona/Private/AnimNotifyDetails.h
+++ b/Engine/Source/Editor/Persona/Private/AnimNotifyDetails.h
@@ -70,7 +70,7 @@ private:
 	 *	@param Builder Layout builder for this object
 	 *	@param NotifyProperty The property that contains the linking properties
 	 */
-	void CustomizeLinkProperties(IDetailLayoutBuilder& Builder, TSharedRef<IPropertyHandle> NotifyProperty, UEditorNotifyObject* EditorObject);
+	void CustomizeLinkProperties(IDetailLayoutBuilder& Builder, TSharedRef<IPropertyHandle> NotifyProperty, UAnimSequenceBase* AnimSequenceBase);
 
 	/** Hide any properties relating to notify linking
 	 *	@param Builder Layout builder for this object
