Animation Attributes improvements
- Added Animation Attribute Viewer to Animation editors (reusing Jack.Cai his previous work)
- Refactored out PoseWatch picker so it can be shared between Curve/Attribute viewers
- Pose watches now always capture full pose (deprecated old APIs)

diff --git a/Engine/Plugins/Animation/ControlRig/Source/ControlRigEditor/Private/Editor/RigAnimAttributeTabSummoner.cpp b/Engine/Plugins/Animation/ControlRig/Source/ControlRigEditor/Private/Editor/RigAnimAttributeTabSummoner.cpp
--- a/Engine/Plugins/Animation/ControlRig/Source/ControlRigEditor/Private/Editor/RigAnimAttributeTabSummoner.cpp
+++ b/Engine/Plugins/Animation/ControlRig/Source/ControlRigEditor/Private/Editor/RigAnimAttributeTabSummoner.cpp
@@ -14,6 +14,7 @@ FRigAnimAttributeTabSummoner::FRigAnimAttributeTabSummoner(const TSharedRef<ICon
 	, ControlRigEditor(InControlRigEditor)
 {
 	TabLabel = LOCTEXT("RigAnimAttributeTabLabel", "Animation Attributes");
+	TabIcon = FSlateIcon(FAppStyle::GetAppStyleSetName(), "AnimGraph.Attribute.Attributes.Icon");
 
 	bIsSingleton = true;
 
diff --git a/Engine/Source/Editor/AnimationBlueprintEditor/Private/AnimationBlueprintEditorMode.cpp b/Engine/Source/Editor/AnimationBlueprintEditor/Private/AnimationBlueprintEditorMode.cpp
--- a/Engine/Source/Editor/AnimationBlueprintEditor/Private/AnimationBlueprintEditorMode.cpp
+++ b/Engine/Source/Editor/AnimationBlueprintEditor/Private/AnimationBlueprintEditorMode.cpp
@@ -178,6 +178,7 @@ FAnimationBlueprintEditorMode::FAnimationBlueprintEditorMode(const TSharedRef<FA
 	TabFactories.RegisterFactory(PersonaModule.CreatePoseWatchTabFactory(InAnimationBlueprintEditor));
 	TSharedPtr<ISkeletonTree> SkeletonTree = InAnimationBlueprintEditor->GetSkeletonTree();
 	TabFactories.RegisterFactory(PersonaModule.CreateCurveViewerTabFactory(InAnimationBlueprintEditor, SkeletonTree.IsValid() ? SkeletonTree->GetEditableSkeleton().ToSharedPtr() : nullptr, InAnimationBlueprintEditor->GetPersonaToolkit()->GetPreviewScene(), FOnObjectsSelected::CreateSP(&InAnimationBlueprintEditor.Get(), &FAnimationBlueprintEditor::HandleObjectsSelected)));
+	TabFactories.RegisterFactory(PersonaModule.CreateAttributeViewerTabFactory(InAnimationBlueprintEditor, InAnimationBlueprintEditor->GetSkeletonTree()->GetEditableSkeleton(), InAnimationBlueprintEditor->GetPersonaToolkit()->GetPreviewScene()));	
 	TabFactories.RegisterFactory(PersonaModule.CreateAnimAssetFindReplaceTabFactory(InAnimationBlueprintEditor, FAnimAssetFindReplaceConfig()));
 
 	if (!bIsTemplate)
diff --git a/Engine/Source/Editor/AnimationEditor/Private/AnimationEditorMode.cpp b/Engine/Source/Editor/AnimationEditor/Private/AnimationEditorMode.cpp
--- a/Engine/Source/Editor/AnimationEditor/Private/AnimationEditorMode.cpp
+++ b/Engine/Source/Editor/AnimationEditor/Private/AnimationEditorMode.cpp
@@ -46,6 +46,7 @@ FAnimationEditorMode::FAnimationEditorMode(TSharedRef<FWorkflowCentricApplicatio
 	TabFactories.RegisterFactory(PersonaModule.CreateAnimNotifiesTabFactory(InHostingApp, InSkeletonTree->GetEditableSkeleton(), OnObjectsSelected));
 	TabFactories.RegisterFactory(PersonaModule.CreateAnimMontageSectionsTabFactory(InHostingApp, AnimationEditor->GetPersonaToolkit(), AnimationEditor->OnSectionsChanged));
 	TabFactories.RegisterFactory(PersonaModule.CreateAnimAssetFindReplaceTabFactory(InHostingApp, FAnimAssetFindReplaceConfig()));
+	TabFactories.RegisterFactory(PersonaModule.CreateAttributeViewerTabFactory(InHostingApp, AnimationEditor->GetSkeletonTree()->GetEditableSkeleton(), AnimationEditor->GetPersonaToolkit()->GetPreviewScene()));
 
 	TabLayout = FTabManager::NewLayout("Standalone_AnimationEditor_Layout_v1.5")
 		->AddArea
diff --git a/Engine/Source/Editor/Persona/Private/PersonaModule.cpp b/Engine/Source/Editor/Persona/Private/PersonaModule.cpp
--- a/Engine/Source/Editor/Persona/Private/PersonaModule.cpp
+++ b/Engine/Source/Editor/Persona/Private/PersonaModule.cpp
@@ -381,6 +381,12 @@ TSharedRef<FWorkflowTabFactory> FPersonaModule::CreatePersonaToolboxTabFactory(c
 	return MakeShareable(new FToolBoxSummoner(InHostingApp));
 }
 
+TSharedRef<FWorkflowTabFactory> FPersonaModule::CreateAttributeViewerTabFactory(const TSharedRef<FWorkflowCentricApplication>& InHostingApp, const TSharedRef<IEditableSkeleton>& InEditableSkeleton,
+	const TSharedRef<IPersonaPreviewScene>& InPreviewScene) const
+{
+	return MakeShareable(new FAnimAttributeViewerTabSummoner(InHostingApp, InEditableSkeleton, InPreviewScene));
+}
+
 TSharedRef<SWidget> FPersonaModule::CreateBlendSpacePreviewWidget(TAttribute<const UBlendSpace*> InBlendSpace, TAttribute<FVector> InPosition, TAttribute<FVector> InFilteredPosition) const
 {
 PRAGMA_DISABLE_DEPRECATION_WARNINGS
diff --git a/Engine/Source/Editor/Persona/Private/SAnimAttributeView.cpp b/Engine/Source/Editor/Persona/Private/SAnimAttributeView.cpp
--- a/Engine/Source/Editor/Persona/Private/SAnimAttributeView.cpp
+++ b/Engine/Source/Editor/Persona/Private/SAnimAttributeView.cpp
@@ -1,18 +1,24 @@
 // Copyright Epic Games, Inc. All Rights Reserved.
 
 #include "SAnimAttributeView.h"
+#include "AnimPreviewInstance.h"
+#include "IPersonaPreviewScene.h"
 #include "Animation/BuiltInAttributeTypes.h"
 #include "Components/SkeletalMeshComponent.h"
 #include "Engine/SkeletalMesh.h"
 #include "Modules/ModuleManager.h"
 #include "PropertyEditorModule.h"
 #include "IStructureDetailsView.h"
+#include "Animation/AnimBlueprintGeneratedClass.h"
+#include "Animation/DebugSkelMeshComponent.h"
 #include "Widgets/Input/SSearchBox.h"
 #include "Widgets/Layout/SExpandableArea.h"
 #include "Widgets/Views/SListView.h"
 #include "Framework/MultiBox/MultiBoxBuilder.h"
 #include "Styling/AppStyle.h"
 #include "Widgets/Layout/SScrollBox.h"
+#include "SPoseWatchPicker.h"
+#include "Engine/PoseWatch.h"
 
 #define LOCTEXT_NAMESPACE "SAnimAttributeView"
 
@@ -262,164 +268,6 @@ void SAnimAttributeView::Construct(const FArguments& InArgs)
 	];
 }
 
-void SAnimAttributeView::DisplayNewAttributeContainerSnapshots(
-	const TArray<TTuple<FName, const UE::Anim::FHeapAttributeContainer&>>& InSnapshots,
-	const USkeletalMeshComponent* InOwningComponent)
-{
-	if (!ensure(InSnapshots.Num() != 0))
-	{
-		ClearListView();
-		return;
-	}
-	
-	// we need the SKM to look up bone names
-	if (!InOwningComponent || !InOwningComponent->GetSkeletalMeshAsset())
-	{
-		ClearListView();
-		return;
-	}
-
-	if (ShouldInvalidateListViewCache(InSnapshots, InOwningComponent))
-	{
-		CachedNumSnapshots = InSnapshots.Num();
-		CachedSnapshotNameIndexMap.Reset();
-		
-		CachedAttributeIdentifierLists.Reset(InSnapshots.Num());
-		CachedAttributeSnapshotMap.Reset();
-		
-		CachedAttributeIdentifierLists.AddDefaulted( InSnapshots.Num());
-		
-		for (int32 SnapshotIndex = 0; SnapshotIndex < InSnapshots.Num(); SnapshotIndex++)
-		{
-			const TTuple<FName, const UE::Anim::FHeapAttributeContainer&> Snapshot = InSnapshots[SnapshotIndex];
-			CachedSnapshotNameIndexMap.FindOrAdd(Snapshot.Key) = SnapshotIndex;
-			
-			TTuple<FName, TArray<FAnimationAttributeIdentifier>>& CachedIdentifierList = CachedAttributeIdentifierLists[SnapshotIndex];
-			CachedIdentifierList.Key = Snapshot.Key;
-
-			const UE::Anim::FHeapAttributeContainer& AttributeContainer = Snapshot.Value;
-			TArray<FAnimationAttributeIdentifier>& CachedIdentifiers = CachedIdentifierList.Value;
-
-			CachedIdentifiers.AddDefaulted(AttributeContainer.Num());
-			
-			const TArray<TWeakObjectPtr<UScriptStruct>>& Types = AttributeContainer.GetUniqueTypes();
-
-			int32 CachedListIdentifierIndex = 0;
-			for (int32 TypeIndex = 0; TypeIndex < Types.Num(); TypeIndex++)
-			{
-				const TArray<UE::Anim::FAttributeId>& Ids = AttributeContainer.GetKeys(TypeIndex);
-			
-				for (int32 IdIndex = 0; IdIndex < Ids.Num(); IdIndex++)
-				{
-					const UE::Anim::FAttributeId& Id = Ids[IdIndex];
-					const FName BoneName = InOwningComponent->GetSkeletalMeshAsset()->GetRefSkeleton().GetBoneName(Id.GetIndex());
-				
-					const FAnimationAttributeIdentifier Identifier(
-						Id.GetName(), Id.GetIndex() ,BoneName, Types[TypeIndex].Get());
-
-					CachedIdentifiers[CachedListIdentifierIndex] = Identifier;
-
-					CachedAttributeSnapshotMap.FindOrAdd(Identifier).Add(Snapshot.Key);
-				
-					CachedListIdentifierIndex++;
-				}
-			}
-		}
-
-		// filtered list should also be refreshed since it depends on the cache
-		RefreshFilteredAttributeEntries();
-
-		// delay value view refresh until tick since this function can be called from animation thread
-		bShouldRefreshValueView = true;
-
-		return;
-	}
-
-	if (SelectedAttribute.IsSet())
-	{
-		for (const FAttributeValueView& ValueView : SelectedAttributeSnapshotValueViews)
-		{
-			const int32* SnapshotIndex = CachedSnapshotNameIndexMap.Find(ValueView.SnapshotName);
-			if (ensure(SnapshotIndex) && ensure(InSnapshots.IsValidIndex(*SnapshotIndex)))
-			{
-				const TTuple<FName, const UE::Anim::FHeapAttributeContainer&> Snapshot = InSnapshots[*SnapshotIndex];
-				const UE::Anim::FHeapAttributeContainer& AttributeContainer = Snapshot.Value;
-
-				ValueView.UpdateValue(AttributeContainer);
-			}
-		}
-	}	
-}
-
-
-
-bool SAnimAttributeView::ShouldInvalidateListViewCache(
-	const TArray<TTuple<FName, const UE::Anim::FHeapAttributeContainer&>>& InSnapshots,
-	const USkeletalMeshComponent* InOwningComponent)
-{
-	// we need the SKM to look up bone names
-	// it should have been checked
-	check(InOwningComponent);
-
-	if (InSnapshots.Num() != CachedAttributeIdentifierLists.Num())
-	{
-		return true;
-	}
-	
-	for (int32 SnapshotIndex = 0; SnapshotIndex < InSnapshots.Num(); SnapshotIndex++)
-	{
-		const TTuple<FName, const UE::Anim::FHeapAttributeContainer&> Snapshot = InSnapshots[SnapshotIndex];
-		const TTuple<FName, TArray<FAnimationAttributeIdentifier>>& CachedIdList = CachedAttributeIdentifierLists[SnapshotIndex];
-		
-		if (Snapshot.Key != CachedIdList.Key)
-		{
-			return true;
-		}
-
-		if (Snapshot.Value.Num() != CachedIdList.Value.Num())
-		{
-			return true;
-		}
-	}
-
-	for (int32 SnapshotIndex = 0; SnapshotIndex < InSnapshots.Num(); SnapshotIndex++)
-	{
-		const TTuple<FName, const UE::Anim::FHeapAttributeContainer&> Snapshot = InSnapshots[SnapshotIndex];
-		const TTuple<FName, TArray<FAnimationAttributeIdentifier>> CachedIdentifierList = CachedAttributeIdentifierLists[SnapshotIndex];
-
-		const UE::Anim::FHeapAttributeContainer& AttributeContainer = Snapshot.Value;
-		const TArray<FAnimationAttributeIdentifier>& CachedIdentifiers = CachedIdentifierList.Value;
-		
-		const TArray<TWeakObjectPtr<UScriptStruct>>& Types = AttributeContainer.GetUniqueTypes();
-
-		int32 CachedListIdentifierIndex = 0;
-		for (int32 TypeIndex = 0; TypeIndex < Types.Num(); TypeIndex++)
-		{
-			const TArray<UE::Anim::FAttributeId>& Ids = AttributeContainer.GetKeys(TypeIndex);
-			
-			for (int32 IdIndex = 0; IdIndex < Ids.Num(); IdIndex++)
-			{
-				const UE::Anim::FAttributeId& Id = Ids[IdIndex];
-				const FName BoneName = InOwningComponent->GetSkeletalMeshAsset()->GetRefSkeleton().GetBoneName(Id.GetIndex());
-				
-				FAnimationAttributeIdentifier Identifier(
-					Id.GetName(), Id.GetIndex() ,BoneName, Types[TypeIndex].Get());
-
-				if (!(Identifier == CachedIdentifiers[CachedListIdentifierIndex]))
-				{
-					return true;
-				}
-
-				CachedListIdentifierIndex++;
-			}
-		}
-	}	
-	
-
-	return false;
-}
-
-
 void SAnimAttributeView::OnSelectionChanged(TSharedPtr<FAnimAttributeEntry> InEntry, ESelectInfo::Type InSelectType)
 {
 	if (InEntry.IsValid())
@@ -624,7 +472,7 @@ void SAnimAttributeView::Tick(const FGeometry& AllottedGeometry, const double In
 		{
 			HeaderRow->RemoveColumn(ColumnId_AnimAttributeSnapshotName);
 		}
-		else if (CachedNumSnapshots >= 1 && HeaderRow->GetColumns().Num() <= MinimalNumColumns)
+		else if (CachedNumSnapshots > 1 && HeaderRow->GetColumns().Num() <= MinimalNumColumns)
 		{
 			SHeaderRow::FColumn::FArguments	ColumnArgs;
 
@@ -663,17 +511,103 @@ SAnimAttributeView::FAttributeValueView::FAttributeValueView(FName InSnapshotNam
 	ViewWidget->SetStructureData(StructData);
 }
 
-void SAnimAttributeView::FAttributeValueView::UpdateValue(const UE::Anim::FHeapAttributeContainer& InAttributeContainer) const
+void SAnimAttributeViewer::Construct(const FArguments& InArgs, const TSharedRef<IEditableSkeleton>& InEditableSkeleton, const TSharedRef<IPersonaPreviewScene>& InPreviewScene)
 {
-	const uint8* ValuePtr = InAttributeContainer.Find(SubjectAttribute.GetScriptStruct(), SubjectAttribute.GetAttributeId());
+	PreviewScenePtr = InPreviewScene;
+	EditableSkeletonPtr = InEditableSkeleton;
 
-	if (ensure(ValuePtr))
+	SAssignNew(AttributeView, SAnimAttributeView);
+	SAssignNew(PoseWatchPicker, SPoseWatchPicker)
+	.AnimBlueprintGeneratedClass_Lambda([this]() -> const UAnimBlueprintGeneratedClass*
+			{
+				if(const UAnimInstance* AnimInstance = GetAnimInstance())
+				{
+					return Cast<UAnimBlueprintGeneratedClass>(AnimInstance->GetClass());
+				}
+				
+				return nullptr;
+			})
+	.DefaultEntryDisplayText(LOCTEXT("DefaultEntryDisplayText", "Output Attributes"));
+
+	ChildSlot
+	[
+		SNew(SVerticalBox)
+		+SVerticalBox::Slot()
+		.AutoHeight()
+		[
+			PoseWatchPicker.ToSharedRef()
+		]
+		+SVerticalBox::Slot()
+		.FillHeight(1.f)
+		[
+			AttributeView.ToSharedRef()
+		]
+	];
+}
+
+UAnimInstance* SAnimAttributeViewer::GetAnimInstance() const
+{
+	UDebugSkelMeshComponent* MeshComponent = PreviewScenePtr.Pin()->GetPreviewMeshComponent();
+	UAnimInstance* AnimInstance = MeshComponent->GetAnimInstance();
+	if (AnimInstance)
+	{
+		// Look at the debugged anim instance if we are targeting one
+		if(AnimInstance == MeshComponent->PreviewInstance.Get())
+		{
+			UAnimPreviewInstance* AnimPreviewInstance = MeshComponent->PreviewInstance;
+			if(USkeletalMeshComponent* DebuggedComponent = AnimPreviewInstance->GetDebugSkeletalMeshComponent())
+			{
+				AnimInstance = DebuggedComponent->GetAnimInstance();
+			}
+		}
+	}
+
+	return AnimInstance;
+}
+
+void SAnimAttributeViewer::Tick(const FGeometry& AllottedGeometry, const double InCurrentTime, const float InDeltaTime)
+{	
+	const UAnimInstance* AnimInstance = GetAnimInstance();
+	if(AnimInstance && PoseWatchPicker.IsValid())
 	{
-		FMemory::Memcpy(
-			StructData->GetStructMemory(),
-			ValuePtr,
-			StructData->GetStruct()->GetStructureSize());
-	}		
+		// Find if we want to use a pose watch
+		if(UPoseWatchPoseElement* PoseWatchPoseElement = PoseWatchPicker->GetCurrentPoseWatch())
+		{
+			if(UAnimBlueprintGeneratedClass* AnimClass = Cast<UAnimBlueprintGeneratedClass>(AnimInstance->GetClass()))
+			{
+				// We have to grab our pose watches from the root class as no pose watches can be set on child anim BPs
+				if(const UAnimBlueprintGeneratedClass* RootClass = Cast<UAnimBlueprintGeneratedClass>(AnimClass->GetRootClass()))
+				{
+					const FAnimBlueprintDebugData& DebugData = RootClass->AnimBlueprintDebugData;
+					for(const FAnimNodePoseWatch& AnimNodePoseWatch : DebugData.AnimNodePoseWatch)
+					{
+						if(AnimNodePoseWatch.PoseWatchPoseElement == PoseWatchPoseElement)
+						{
+							const UE::Anim::FHeapAttributeContainer& OutputSnapshot = AnimNodePoseWatch.GetAttributes();
+							const TArray<TTuple<FName, const UE::Anim::FHeapAttributeContainer&>> Snapshots =
+							{
+								{*AnimNodePoseWatch.PoseWatch->GetLabel().ToString(), OutputSnapshot}
+							};
+
+							AttributeView->DisplayNewAttributeContainerSnapshots(Snapshots, PreviewScenePtr.Pin()->GetPreviewMeshComponent());
+
+							break;
+						}
+					}
+				}
+			}
+		}
+		else if (UDebugSkelMeshComponent* MeshComponent = PreviewScenePtr.Pin()->GetPreviewMeshComponent())
+		{
+			const UE::Anim::FMeshAttributeContainer& OutputSnapshot = MeshComponent->GetCustomAttributes();		
+			const TArray<TTuple<FName, const UE::Anim::FMeshAttributeContainer&>> Snapshots =
+			{
+				{TEXT("Evaluation"), OutputSnapshot }
+			};
+
+			AttributeView->DisplayNewAttributeContainerSnapshots(Snapshots, MeshComponent);
+		}
+	}
 }
 
 #undef LOCTEXT_NAMESPACE
diff --git a/Engine/Source/Editor/Persona/Private/SAnimCurveViewer.cpp b/Engine/Source/Editor/Persona/Private/SAnimCurveViewer.cpp
--- a/Engine/Source/Editor/Persona/Private/SAnimCurveViewer.cpp
+++ b/Engine/Source/Editor/Persona/Private/SAnimCurveViewer.cpp
@@ -30,6 +30,7 @@
 #include "PersonaTabs.h"
 #include "SAnimAssetFindReplace.h"
 #include "Widgets/Docking/SDockTab.h"
+#include "SPoseWatchPicker.h"
 
 #define LOCTEXT_NAMESPACE "SAnimCurveViewer"
 
@@ -101,10 +102,11 @@ private:
 
 bool FDisplayedAnimCurveInfo::GetActiveFlag(const TSharedPtr<SAnimCurveViewer>& InAnimCurveViewer, bool bInMorphTarget) const
 {
-	if(const UAnimInstance* AnimInstance = InAnimCurveViewer->GetAnimInstance())
+	const UAnimInstance* AnimInstance = InAnimCurveViewer->GetAnimInstance();
+	if(AnimInstance && InAnimCurveViewer->PoseWatchPicker.IsValid())
 	{
 		// Find if we want to use a pose watch
-		if(UPoseWatchPoseElement* PoseWatchPoseElement = InAnimCurveViewer->PoseWatch.Get())
+		if(UPoseWatchPoseElement* PoseWatchPoseElement = InAnimCurveViewer->PoseWatchPicker->GetCurrentPoseWatch())
 		{
 			if(UAnimBlueprintGeneratedClass* AnimClass = Cast<UAnimBlueprintGeneratedClass>(AnimInstance->GetClass()))
 			{
@@ -398,10 +400,11 @@ bool SAnimCurveListRow::GetActiveWeight(float& OutWeight) const
 	TSharedPtr<SAnimCurveViewer> AnimCurveViewer = AnimCurveViewerPtr.Pin();
 	if (AnimCurveViewer.IsValid())
 	{
-		if(const UAnimInstance* AnimInstance = AnimCurveViewer->GetAnimInstance())
+		const UAnimInstance* AnimInstance = AnimCurveViewer->GetAnimInstance();
+		if(AnimInstance && AnimCurveViewer->PoseWatchPicker.IsValid())
 		{
 			// Find if we want to use a pose watch
-			if(UPoseWatchPoseElement* PoseWatchPoseElement = AnimCurveViewer->PoseWatch.Get())
+			if(UPoseWatchPoseElement* PoseWatchPoseElement = AnimCurveViewer->PoseWatchPicker->GetCurrentPoseWatch())
 			{
 				if(UAnimBlueprintGeneratedClass* AnimClass = Cast<UAnimBlueprintGeneratedClass>(AnimInstance->GetClass()))
 				{
@@ -794,10 +797,11 @@ void SAnimCurveViewer::CreateAnimCurveList( const FString& SearchText, bool bInF
 
 	// Add active curves if required
 	TSet<FName> ActiveCurves;
-	if(const UAnimInstance* AnimInstance = GetAnimInstance())
+	const UAnimInstance* AnimInstance = GetAnimInstance();
+	if(AnimInstance && PoseWatchPicker.IsValid())
 	{
 		// Find if we want to use a pose watch
-		if(UPoseWatchPoseElement* PoseWatchPoseElement = PoseWatch.Get())
+		if(UPoseWatchPoseElement* PoseWatchPoseElement = PoseWatchPicker->GetCurrentPoseWatch())
 		{
 			if(UAnimBlueprintGeneratedClass* AnimClass = Cast<UAnimBlueprintGeneratedClass>(AnimInstance->GetClass()))
 			{
@@ -935,185 +939,24 @@ TSharedRef<SWidget> SAnimCurveViewer::CreateCurveSourceSelector()
 
 	if(bShowPoseWatches)
 	{
-		AnimationEditorUtils::OnPoseWatchesChanged().AddSP(this, &SAnimCurveViewer::HandlePoseWatchesChanged);
-
-		// Add selector for pose watches
-		SAssignNew(PoseWatchCombo, SComboBox<TSharedPtr<TWeakObjectPtr<UPoseWatchPoseElement>>>)
-		.ToolTipText_Lambda([this]()
-		{
-			return PoseWatch.Get() != nullptr && PoseWatch->GetParent() != nullptr ? FText::Format(LOCTEXT("PoseWatchTooltipFormat", "Viewing Pose Watch Curves for '{0}'"), PoseWatch->GetParent()->GetLabel()) : LOCTEXT("PoseWatchViewingOutputCurves", "Viewing output Curves");
-		})
-		.OptionsSource(&PoseWatches)
-		.OnGenerateWidget_Lambda([this](TSharedPtr<TWeakObjectPtr<UPoseWatchPoseElement>> InElement)
+		PoseWatchPicker = SNew(SPoseWatchPicker)
+		.AnimBlueprintGeneratedClass_Lambda([this]() -> const UAnimBlueprintGeneratedClass*
 		{
-			TWeakObjectPtr<UPoseWatchPoseElement> Element = *InElement;
-
-			return
-				SNew(SHorizontalBox)
-				+SHorizontalBox::Slot()
-				.AutoWidth()
-				.VAlign(VAlign_Center)
-				.Padding(2.0f, 0.0f)
-				[
-					SNew(SColorBlock)
-					.Visibility_Lambda([Element]()
-					{
-						return Element.Get() != nullptr ? EVisibility::Visible : EVisibility::Collapsed;
-					})
-					.AlphaBackgroundBrush(FAppStyle::Get().GetBrush("ColorPicker.RoundedAlphaBackground"))
-					.ShowBackgroundForAlpha(true)
-					.AlphaDisplayMode(EColorBlockAlphaDisplayMode::Ignore)
-					.Size(FVector2D(16.0f, 16.0f))
-					.CornerRadius(FVector4(4.0f, 4.0f, 4.0f, 4.0f))	
-					.Color_Lambda([Element]()
-					{
-						if(UPoseWatchPoseElement* PoseWatchPoseElement = Element.Get())
-						{
-							return FLinearColor(PoseWatchPoseElement->GetColor());
-						}
-						return FLinearColor::Gray;
-					})
-				]
-				+SHorizontalBox::Slot()
-				.AutoWidth()
-				.VAlign(VAlign_Center)
-				.Padding(2.0f, 0.0f)
-				[
-					SNew(STextBlock)
-					.Text_Lambda([Element]()
-					{
-						return Element.Get() != nullptr && Element->GetParent() != nullptr ? Element->GetParent()->GetLabel() : LOCTEXT("PoseWatchOutputCurves", "Output Curves");
-					})
-				];
-		})
-		.OnSelectionChanged_Lambda([this](TSharedPtr<TWeakObjectPtr<UPoseWatchPoseElement>> InElement, ESelectInfo::Type InSelectionType)
-		{
-			if(!InElement.IsValid())
-			{
-				return;
-			}
-			
-			if((*InElement).Get() == nullptr)
-			{
-				PoseWatch = nullptr;
-				return;
-			}
-
 			if(const UAnimInstance* AnimInstance = GetAnimInstance())
 			{
-				if(UAnimBlueprintGeneratedClass* AnimClass = Cast<UAnimBlueprintGeneratedClass>(AnimInstance->GetClass()))
-				{
-					TWeakObjectPtr<UPoseWatchPoseElement> Element = *InElement;
-
-					// We have to grab our pose watches from the root class as no pose watches can be set on child anim BPs
-					if(const UAnimBlueprintGeneratedClass* RootClass = Cast<UAnimBlueprintGeneratedClass>(AnimClass->GetRootClass()))
-					{
-						for(const FAnimNodePoseWatch& AnimNodePoseWatch : RootClass->AnimBlueprintDebugData.AnimNodePoseWatch)
-						{
-							if(AnimNodePoseWatch.PoseWatchPoseElement && Element.Get() == AnimNodePoseWatch.PoseWatchPoseElement)
-							{
-								PoseWatch = AnimNodePoseWatch.PoseWatchPoseElement;
-								break;
-							}
-						}
-					}
-				}
+				return Cast<UAnimBlueprintGeneratedClass>(AnimInstance->GetClass());
 			}
+			
+			return nullptr;
 		})
-		.Content()
-		[
-			SNew(SHorizontalBox)
-			.Clipping(EWidgetClipping::OnDemand)
-			+SHorizontalBox::Slot()
-			.AutoWidth()
-			.VAlign(VAlign_Center)
-			.Padding(2.0f, 0.0f)
-			[
-				SNew(SColorBlock)
-				.Visibility_Lambda([this]()
-				{
-					return PoseWatch.Get() != nullptr ? EVisibility::Visible : EVisibility::Collapsed;
-				})
-				.AlphaBackgroundBrush(FAppStyle::Get().GetBrush("ColorPicker.RoundedAlphaBackground"))
-				.ShowBackgroundForAlpha(true)
-				.AlphaDisplayMode(EColorBlockAlphaDisplayMode::Ignore)
-				.Size(FVector2D(16.0f, 16.0f))
-				.CornerRadius(FVector4(4.0f, 4.0f, 4.0f, 4.0f))	
-				.Color_Lambda([this]()
-				{
-					if(UPoseWatchPoseElement* Element = PoseWatch.Get())
-					{
-						return FLinearColor(Element->GetColor());
-					}
-					return FLinearColor::Gray;
-				})
-			]
-			+SHorizontalBox::Slot()
-			.AutoWidth()
-			.VAlign(VAlign_Center)
-			.Padding(2.0f, 0.0f)
-			[
-				SNew(STextBlock)
-				.Text_Lambda([this]()
-				{
-					return PoseWatch.Get() != nullptr && PoseWatch->GetParent() != nullptr ? PoseWatch->GetParent()->GetLabel() : LOCTEXT("PoseWatchOutputCurves", "Output Curves");
-				})
-			]
-		];
+		.DefaultEntryDisplayText(LOCTEXT("PoseWatchOutputCurves", "Output Curves"));
 
-		RebuildPoseWatches();
-
-		return PoseWatchCombo.ToSharedRef();
+		return PoseWatchPicker.ToSharedRef();
 	}
 
 	return SNullWidget::NullWidget;
 }
 
-void SAnimCurveViewer::RebuildPoseWatches()
-{
-	if(PoseWatchCombo.IsValid())
-	{
-		PoseWatches.Empty();
-
-		PoseWatches.Add(MakeShared<TWeakObjectPtr<UPoseWatchPoseElement>>());
-
-		if(const UAnimInstance* AnimInstance = GetAnimInstance())
-		{
-			if(UAnimBlueprintGeneratedClass* AnimClass = Cast<UAnimBlueprintGeneratedClass>(AnimInstance->GetClass()))
-			{
-				// We have to grab our pose watches from the root class as no pose watches can be set on child anim BPs
-				if(const UAnimBlueprintGeneratedClass* RootClass = Cast<UAnimBlueprintGeneratedClass>(AnimClass->GetRootClass()))
-				{
-					for(const FAnimNodePoseWatch& AnimNodePoseWatch : RootClass->AnimBlueprintDebugData.AnimNodePoseWatch)
-					{
-						if(AnimNodePoseWatch.PoseWatchPoseElement)
-						{
-							PoseWatches.Add(MakeShared<TWeakObjectPtr<UPoseWatchPoseElement>>(AnimNodePoseWatch.PoseWatchPoseElement));
-						}
-					}
-				}
-			}
-		}
-
-		PoseWatchCombo->RefreshOptions();
-	}
-}
-
-void SAnimCurveViewer::HandlePoseWatchesChanged(UAnimBlueprint* InAnimBlueprint, UEdGraphNode* /*InNode*/)
-{
-	if(UAnimInstance* AnimInstance = GetAnimInstance())
-	{
-		if(AnimInstance->GetClass()->IsChildOf(InAnimBlueprint->GeneratedClass))
-		{
-			RegisterActiveTimer(0.0f, FWidgetActiveTimerDelegate::CreateLambda([this](double /*InCurrentTime*/, float /*InDeltaTime*/)
-			{
-				RebuildPoseWatches();
-				return EActiveTimerReturnType::Stop;
-			}));
-		}
-	}
-}
-
 void SAnimCurveViewer::AddAnimCurveOverride( FName& Name, float Weight)
 {
 	float& Value = OverrideCurves.FindOrAdd(Name);
diff --git a/Engine/Source/Editor/Persona/Private/SAnimCurveViewer.h b/Engine/Source/Editor/Persona/Private/SAnimCurveViewer.h
--- a/Engine/Source/Editor/Persona/Private/SAnimCurveViewer.h
+++ b/Engine/Source/Editor/Persona/Private/SAnimCurveViewer.h
@@ -21,6 +21,8 @@
 #include "Filters/FilterBase.h"
 #include "Widgets/Input/SComboBox.h"
 
+class SPoseWatchPicker;
+
 namespace UE::Anim
 {
 enum class ECurveElementFlags : uint8;
@@ -280,10 +282,6 @@ private:
 
 	TSharedRef<SWidget> CreateCurveSourceSelector();
 	
-	/** Handle building list of available pose watches */
-	void HandlePoseWatchesChanged(UAnimBlueprint* /*InAnimBlueprint*/, UEdGraphNode* /*InNode*/);
-	void RebuildPoseWatches();
-
 	void BindCommands();
 	
 	/** Build context menu */
@@ -328,15 +326,14 @@ private:
 	/** Delegate called to select objects */
 	FOnObjectsSelected OnObjectsSelected;
 
-	/** Valid when viewing a pose watch's data */
-	TWeakObjectPtr<UPoseWatchPoseElement> PoseWatch;
-
 	/** Names of each pose watch that can be selected */
 	TArray<TSharedPtr<TWeakObjectPtr<UPoseWatchPoseElement>>> PoseWatches;
 	
 	/** Combobox used to select pose watch */
 	TSharedPtr<SComboBox<TSharedPtr<TWeakObjectPtr<UPoseWatchPoseElement>>>> PoseWatchCombo;
 
+	TSharedPtr<SPoseWatchPicker> PoseWatchPicker;
+
 	/** All filters that can be applied to the widget's display */
 	TArray<TSharedRef<FFilterBase<EAnimCurveViewerFilterFlags>>> Filters;
 };
\ No newline at end of file
diff --git a/Engine/Source/Editor/Persona/Private/SPoseWatchPicker.cpp b/Engine/Source/Editor/Persona/Private/SPoseWatchPicker.cpp
new file mode 100644
--- /dev/null
+++ b/Engine/Source/Editor/Persona/Private/SPoseWatchPicker.cpp
@@ -0,0 +1,189 @@
+// Copyright Epic Games, Inc. All Rights Reserved.
+
+#include "SPoseWatchPicker.h"
+
+#include "AnimationEditorUtils.h"
+#include "Animation/AnimBlueprintGeneratedClass.h"
+#include "Engine/PoseWatch.h"
+#include "Widgets/Colors/SColorBlock.h"
+
+#define LOCTEXT_NAMESPACE "SPoseWatchPicker"
+
+void SPoseWatchPicker::Construct(const FArguments& InArgs)
+{
+	AnimBlueprintAttribute = InArgs._AnimBlueprintGeneratedClass;
+
+	PoseWatchComboBox = SNew(SComboBox<TSharedPtr<TWeakObjectPtr<UPoseWatchPoseElement>>>)
+		.ToolTipText_Lambda([this, DefaultDisplayText = InArgs._DefaultEntryDisplayText]()
+		{
+			return FText::Format(LOCTEXT("PoseWatchTooltipFormat", "Previewing '{0}'"),  SelectedPoseWatch.Get() != nullptr && SelectedPoseWatch->GetParent() != nullptr ? SelectedPoseWatch->GetParent()->GetLabel() : DefaultDisplayText);
+		})
+		.OptionsSource(&CachedPoseWatches)
+		.OnGenerateWidget_Lambda([this, DefaultDisplayText = InArgs._DefaultEntryDisplayText](const TSharedPtr<TWeakObjectPtr<UPoseWatchPoseElement>>& InElement)
+		{			
+			return
+				SNew(SHorizontalBox)
+				+SHorizontalBox::Slot()
+				.AutoWidth()
+				.VAlign(VAlign_Center)
+				.Padding(2.0f, 0.0f)
+				[
+					SNew(SColorBlock)
+					.Visibility_Lambda([InElement]()
+					{
+						TWeakObjectPtr<UPoseWatchPoseElement> PoseWatchPoseElement = *InElement;
+						return PoseWatchPoseElement.Get() != nullptr ? EVisibility::Visible : EVisibility::Collapsed;
+					})
+					.AlphaBackgroundBrush(FAppStyle::Get().GetBrush("ColorPicker.RoundedAlphaBackground"))
+					.ShowBackgroundForAlpha(true)
+					.AlphaDisplayMode(EColorBlockAlphaDisplayMode::Ignore)
+					.Size(FVector2D(16.0f, 16.0f))
+					.CornerRadius(FVector4(4.0f, 4.0f, 4.0f, 4.0f))	
+					.Color_Lambda([InElement]()
+					{
+						TWeakObjectPtr<UPoseWatchPoseElement> PoseWatchPoseElement = *InElement;
+						if(PoseWatchPoseElement.Get())
+						{
+							return FLinearColor(PoseWatchPoseElement->GetColor());
+						}
+						return FLinearColor::Gray;
+					})
+				]
+				+SHorizontalBox::Slot()
+				.AutoWidth()
+				.VAlign(VAlign_Center)
+				.Padding(2.0f, 0.0f)
+				[
+					SNew(STextBlock)
+					.Text_Lambda([InElement, DefaultDisplayText]()
+					{
+						TWeakObjectPtr<UPoseWatchPoseElement> PoseWatchPoseElement = *InElement;
+						return PoseWatchPoseElement.Get() != nullptr && PoseWatchPoseElement->GetParent() != nullptr ? PoseWatchPoseElement->GetParent()->GetLabel() : DefaultDisplayText;
+					})
+				];
+		})
+		.OnSelectionChanged_Lambda([this](const TSharedPtr<TWeakObjectPtr<UPoseWatchPoseElement>>& InElement, ESelectInfo::Type InSelectionType)
+		{
+			if(!InElement.IsValid())
+			{
+				return;
+			}
+
+			TWeakObjectPtr<UPoseWatchPoseElement> PoseWatchPoseElement = *InElement;			
+			if(PoseWatchPoseElement.Get() == nullptr)
+			{
+				SelectedPoseWatch = nullptr;
+				return;
+			}
+
+			if(const UAnimBlueprintGeneratedClass* AnimClass = AnimBlueprintAttribute.Get())
+			{
+				// We have to grab our pose watches from the root class as no pose watches can be set on child anim BPs
+				if(const UAnimBlueprintGeneratedClass* RootClass = Cast<UAnimBlueprintGeneratedClass>(AnimClass->GetRootClass()))
+				{
+					for(const FAnimNodePoseWatch& AnimNodePoseWatch : RootClass->AnimBlueprintDebugData.AnimNodePoseWatch)
+					{
+						if(AnimNodePoseWatch.PoseWatchPoseElement && PoseWatchPoseElement.Get() == AnimNodePoseWatch.PoseWatchPoseElement)
+						{
+							SelectedPoseWatch = AnimNodePoseWatch.PoseWatchPoseElement;
+							break;
+						}
+					}
+				}
+			}
+		})
+		.Content()
+		[
+			SNew(SHorizontalBox)
+			.Clipping(EWidgetClipping::OnDemand)
+			+SHorizontalBox::Slot()
+			.AutoWidth()
+			.VAlign(VAlign_Center)
+			.Padding(2.0f, 0.0f)
+			[
+				SNew(SColorBlock)
+				.Visibility_Lambda([this]()
+				{
+					return SelectedPoseWatch.Get() != nullptr ? EVisibility::Visible : EVisibility::Collapsed;
+				})
+				.AlphaBackgroundBrush(FAppStyle::Get().GetBrush("ColorPicker.RoundedAlphaBackground"))
+				.ShowBackgroundForAlpha(true)
+				.AlphaDisplayMode(EColorBlockAlphaDisplayMode::Ignore)
+				.Size(FVector2D(16.0f, 16.0f))
+				.CornerRadius(FVector4(4.0f, 4.0f, 4.0f, 4.0f))	
+				.Color_Lambda([this]()
+				{
+					if(UPoseWatchPoseElement* Element = SelectedPoseWatch.Get())
+					{
+						return FLinearColor(Element->GetColor());
+					}
+					return FLinearColor::Gray;
+				})
+			]
+			+SHorizontalBox::Slot()
+			.AutoWidth()
+			.VAlign(VAlign_Center)
+			.Padding(2.0f, 0.0f)
+			[
+				SNew(STextBlock)
+				.Text_Lambda([this, DefaultDisplayText = InArgs._DefaultEntryDisplayText]()
+				{
+					return SelectedPoseWatch.Get() != nullptr && SelectedPoseWatch->GetParent() != nullptr ? SelectedPoseWatch->GetParent()->GetLabel() : DefaultDisplayText;
+				})
+			]
+		];
+	
+	ChildSlot
+	[
+		PoseWatchComboBox.ToSharedRef()
+	];
+
+	AnimationEditorUtils::OnPoseWatchesChanged().AddSP(this, &SPoseWatchPicker::OnPoseWatchesChanged);
+	RebuildPoseWatches();
+}
+
+UPoseWatchPoseElement* SPoseWatchPicker::GetCurrentPoseWatch() const
+{
+	return SelectedPoseWatch.Get();
+}
+
+void SPoseWatchPicker::OnPoseWatchesChanged(UAnimBlueprint* InAnimBlueprint, UEdGraphNode*)
+{
+	if(const UAnimBlueprintGeneratedClass* TargetClass = AnimBlueprintAttribute.Get())
+	{
+		if(TargetClass->IsChildOf(InAnimBlueprint->GeneratedClass))
+		{
+			RegisterActiveTimer(0.0f, FWidgetActiveTimerDelegate::CreateLambda([this](double /*InCurrentTime*/, float /*InDeltaTime*/)
+			{
+				RebuildPoseWatches();
+				return EActiveTimerReturnType::Stop;
+			}));
+		}
+	}
+}
+
+void SPoseWatchPicker::RebuildPoseWatches()
+{
+	CachedPoseWatches.Empty();
+
+	if(const UAnimBlueprintGeneratedClass* TargetClass = AnimBlueprintAttribute.Get())
+	{
+		CachedPoseWatches.Add(MakeShared<TWeakObjectPtr<UPoseWatchPoseElement>>());
+		// We have to grab our pose watches from the root class as no pose watches can be set on child anim BPs
+		if(const UAnimBlueprintGeneratedClass* RootClass = Cast<UAnimBlueprintGeneratedClass>(TargetClass->GetRootClass()))
+		{
+			for(const FAnimNodePoseWatch& AnimNodePoseWatch : RootClass->AnimBlueprintDebugData.AnimNodePoseWatch)
+			{
+				if(AnimNodePoseWatch.PoseWatchPoseElement)
+				{						
+					CachedPoseWatches.Add(MakeShared<TWeakObjectPtr<UPoseWatchPoseElement>>(AnimNodePoseWatch.PoseWatchPoseElement));
+				}
+			}
+		}
+
+		PoseWatchComboBox->RefreshOptions();
+	}
+	
+}
+
+#undef LOCTEXT_NAMESPACE // "SPoseWatchPicker"
\ No newline at end of file
diff --git a/Engine/Source/Editor/Persona/Private/SPoseWatchPicker.h b/Engine/Source/Editor/Persona/Private/SPoseWatchPicker.h
new file mode 100644
--- /dev/null
+++ b/Engine/Source/Editor/Persona/Private/SPoseWatchPicker.h
@@ -0,0 +1,31 @@
+// Copyright Epic Games, Inc. All Rights Reserved.
+
+#pragma once
+#include "Animation/AnimBlueprint.h"
+#include "Widgets/SCompoundWidget.h"
+#include "Widgets/Input/SComboBox.h"
+
+class UAnimBlueprint;
+class UPoseWatchPoseElement;
+
+class SPoseWatchPicker : public SCompoundWidget
+{
+	SLATE_BEGIN_ARGS(SPoseWatchPicker) : _AnimBlueprintGeneratedClass(nullptr) {}
+		SLATE_ATTRIBUTE(const UAnimBlueprintGeneratedClass*, AnimBlueprintGeneratedClass)
+		SLATE_ARGUMENT(FText, DefaultEntryDisplayText)
+	SLATE_END_ARGS()
+
+	void Construct(const FArguments& InArgs);
+
+	UPoseWatchPoseElement* GetCurrentPoseWatch() const;
+
+protected:
+	void OnPoseWatchesChanged(UAnimBlueprint* InAnimBlueprint, UEdGraphNode* /*InNode*/);
+	void RebuildPoseWatches();
+
+protected:
+	TSharedPtr<SComboBox<TSharedPtr<TWeakObjectPtr<UPoseWatchPoseElement>>>> PoseWatchComboBox;
+	TArray<TSharedPtr<TWeakObjectPtr<UPoseWatchPoseElement>>> CachedPoseWatches;
+	TAttribute<const UAnimBlueprintGeneratedClass*> AnimBlueprintAttribute;
+	TWeakObjectPtr<UPoseWatchPoseElement> SelectedPoseWatch = nullptr;
+};
\ No newline at end of file
diff --git a/Engine/Source/Editor/Persona/Private/TabSpawners.cpp b/Engine/Source/Editor/Persona/Private/TabSpawners.cpp
--- a/Engine/Source/Editor/Persona/Private/TabSpawners.cpp
+++ b/Engine/Source/Editor/Persona/Private/TabSpawners.cpp
@@ -16,6 +16,7 @@
 #include "PersonaUtils.h"
 #include "SMorphTargetViewer.h"
 #include "SAnimCurveViewer.h"
+#include "SAnimAttributeView.h"
 #include "SAnimationSequenceBrowser.h"
 #include "SAnimationEditorViewport.h"
 #include "SPoseWatchManager.h"
@@ -78,6 +79,8 @@ const FName FPersonaTabs::FindReplaceID("FindReplaceTab");
 
 const FName FPersonaTabs::ToolboxID("ToolBox");
 
+const FName FPersonaTabs::AnimAttributeViewID("AnimAttributeViewerTab");
+
 /////////////////////////////////////////////////////
 // FPersonaMode
 
@@ -830,5 +833,29 @@ TSharedRef<SWidget> FAssetPropertiesSummoner::CreateTabBody(const FWorkflowTabSp
 		.OnDetailsCreated(OnDetailsCreated);
 }
 
+/////////////////////////////////////////////////////
+// FAnimAttributeViewerTabSummoner
+
+FAnimAttributeViewerTabSummoner::FAnimAttributeViewerTabSummoner(TSharedPtr<class FAssetEditorToolkit> InHostingApp, const TSharedRef<IEditableSkeleton>& InEditableSkeleton, const TSharedRef<IPersonaPreviewScene>& InPreviewScene)
+	: FWorkflowTabFactory(FPersonaTabs::AnimAttributeViewID, InHostingApp)
+	, EditableSkeleton(InEditableSkeleton)
+	, PreviewScene(InPreviewScene)
+{
+	TabLabel = LOCTEXT("AnimAttributeViewTabTitle", "Animation Attributes");	
+	TabIcon = FSlateIcon(FAppStyle::GetAppStyleSetName(), "AnimGraph.Attribute.Attributes.Icon");
+
+	EnableTabPadding();
+	bIsSingleton = true;
+
+	ViewMenuDescription = LOCTEXT("AnimAttributeTabView", "Animation Attributes");
+	ViewMenuTooltip = LOCTEXT("AnimAttributeTabView_ToolTip", "Shows the animation attribute viewer");
+}
+
+TSharedRef<SWidget> FAnimAttributeViewerTabSummoner::CreateTabBody(const FWorkflowTabSpawnInfo& Info) const
+{
+	return SNew(SAnimAttributeViewer, EditableSkeleton.Pin().ToSharedRef(), PreviewScene.Pin().ToSharedRef());
+}
+
+
 #undef LOCTEXT_NAMESPACE
 
diff --git a/Engine/Source/Editor/Persona/Private/TabSpawners.h b/Engine/Source/Editor/Persona/Private/TabSpawners.h
--- a/Engine/Source/Editor/Persona/Private/TabSpawners.h
+++ b/Engine/Source/Editor/Persona/Private/TabSpawners.h
@@ -409,4 +409,26 @@ private:
 	FOnDetailsCreated OnDetailsCreated;
 };
 
+/////////////////////////////////////////////////////
+// FAnimAttributeViewerTabSummoner
+
+struct FAnimAttributeViewerTabSummoner : public FWorkflowTabFactory
+{
+public:
+	FAnimAttributeViewerTabSummoner(TSharedPtr<class FAssetEditorToolkit> InHostingApp, const TSharedRef<IEditableSkeleton>& InEditableSkeleton, const TSharedRef<IPersonaPreviewScene>& InPreviewScene);
+
+	virtual TSharedRef<SWidget> CreateTabBody(const FWorkflowTabSpawnInfo& Info) const override;
+
+	// Create a tooltip widget for the tab
+	virtual TSharedPtr<SToolTip> CreateTabToolTipWidget(const FWorkflowTabSpawnInfo& Info) const override
+	{
+		return IDocumentation::Get()->CreateToolTip(LOCTEXT("AnimAttributeViewTooltip", "The Animation Attribute Viewer tab lets you preview any Animation Attributes resulting from evaluating the current preview asset."), NULL, TEXT("Shared/Editors/Persona"), TEXT("AnimAttributeView_Window"));
+	}
+
+private:
+	TWeakPtr<class IEditableSkeleton> EditableSkeleton;
+	TWeakPtr<class IPersonaPreviewScene> PreviewScene;
+};
+
+
 #undef LOCTEXT_NAMESPACE
diff --git a/Engine/Source/Editor/Persona/Public/PersonaModule.h b/Engine/Source/Editor/Persona/Public/PersonaModule.h
--- a/Engine/Source/Editor/Persona/Public/PersonaModule.h
+++ b/Engine/Source/Editor/Persona/Public/PersonaModule.h
@@ -409,6 +409,9 @@ public:
 	/** Create a toolbox tab factory */
 	virtual TSharedRef<FWorkflowTabFactory> CreatePersonaToolboxTabFactory(const TSharedRef<class FPersonaAssetEditorToolkit>& InHostingApp) const;
 
+	/** Create a attributes viewer tab factory */
+	virtual TSharedRef<FWorkflowTabFactory> CreateAttributeViewerTabFactory(const TSharedRef<class FWorkflowCentricApplication>& InHostingApp, const TSharedRef<class IEditableSkeleton>& InEditableSkeleton, const TSharedRef<class IPersonaPreviewScene>& InPreviewScene) const;
+
 	/** Deprecated */
 	UE_DEPRECATED(5.0, "Please use the overload that takes a FBlendSpacePreviewArgs struct")
 	virtual TSharedRef<SWidget> CreateBlendSpacePreviewWidget(TAttribute<const UBlendSpace*> InBlendSpace, TAttribute<FVector> InBlendPosition, TAttribute<FVector> InFilteredBlendPosition) const;
diff --git a/Engine/Source/Editor/Persona/Public/PersonaTabs.h b/Engine/Source/Editor/Persona/Public/PersonaTabs.h
--- a/Engine/Source/Editor/Persona/Public/PersonaTabs.h
+++ b/Engine/Source/Editor/Persona/Public/PersonaTabs.h
@@ -48,6 +48,8 @@ struct PERSONA_API FPersonaTabs
 	static const FName FindReplaceID;
 
 	static const FName ToolboxID;
+
+	static const FName AnimAttributeViewID;
 private:
 	FPersonaTabs() {}
 };
\ No newline at end of file
diff --git a/Engine/Source/Editor/Persona/Public/SAnimAttributeView.h b/Engine/Source/Editor/Persona/Public/SAnimAttributeView.h
--- a/Engine/Source/Editor/Persona/Public/SAnimAttributeView.h
+++ b/Engine/Source/Editor/Persona/Public/SAnimAttributeView.h
@@ -9,6 +9,8 @@
 #include "Widgets/SCompoundWidget.h"
 #include "Animation/AttributesRuntime.h"
 #include "Animation/AnimData/AttributeIdentifier.h"
+#include "Components/SkeletalMeshComponent.h"
+#include "Engine/SkeletalMesh.h"
 
 class ITableRow;
 class STableViewBase;
@@ -16,6 +18,8 @@ class IStructureDetailsView;
 class FStructOnScope;
 class SHeaderRow;
 class SScrollBox;
+class SPoseWatchPicker;
+class UAnimInstance;
 
 namespace EColumnSortMode
 {
@@ -119,9 +123,95 @@ public:
 	
 	void Construct(const FArguments& InArgs);
 
+	template <typename ContainerType>
 	void DisplayNewAttributeContainerSnapshots(
-		const TArray<TTuple<FName, const UE::Anim::FHeapAttributeContainer&>>& InSnapshots, 
-		const USkeletalMeshComponent* InOwningComponent);
+		const TArray<TTuple<FName, const ContainerType&>>& InSnapshots, 
+		const USkeletalMeshComponent* InOwningComponent)
+	{
+		if (!ensure(InSnapshots.Num() != 0))
+		{
+			ClearListView();
+			return;
+		}
+		
+		// we need the SKM to look up bone names
+		if (!InOwningComponent || !InOwningComponent->GetSkeletalMeshAsset())
+		{
+			ClearListView();
+			return;
+		}
+
+		if (ShouldInvalidateListViewCache(InSnapshots, InOwningComponent))
+		{
+			CachedNumSnapshots = InSnapshots.Num();
+			CachedSnapshotNameIndexMap.Reset();
+			
+			CachedAttributeIdentifierLists.Reset(InSnapshots.Num());
+			CachedAttributeSnapshotMap.Reset();
+			
+			CachedAttributeIdentifierLists.AddDefaulted( InSnapshots.Num());
+			
+			for (int32 SnapshotIndex = 0; SnapshotIndex < InSnapshots.Num(); SnapshotIndex++)
+			{
+				const TTuple<FName, const ContainerType&> Snapshot = InSnapshots[SnapshotIndex];
+				CachedSnapshotNameIndexMap.FindOrAdd(Snapshot.Key) = SnapshotIndex;
+				
+				TTuple<FName, TArray<FAnimationAttributeIdentifier>>& CachedIdentifierList = CachedAttributeIdentifierLists[SnapshotIndex];
+				CachedIdentifierList.Key = Snapshot.Key;
+
+				const ContainerType& AttributeContainer = Snapshot.Value;
+				TArray<FAnimationAttributeIdentifier>& CachedIdentifiers = CachedIdentifierList.Value;
+
+				CachedIdentifiers.AddDefaulted(AttributeContainer.Num());
+				
+				const TArray<TWeakObjectPtr<UScriptStruct>>& Types = AttributeContainer.GetUniqueTypes();
+
+				int32 CachedListIdentifierIndex = 0;
+				for (int32 TypeIndex = 0; TypeIndex < Types.Num(); TypeIndex++)
+				{
+					const TArray<UE::Anim::FAttributeId>& Ids = AttributeContainer.GetKeys(TypeIndex);
+				
+					for (int32 IdIndex = 0; IdIndex < Ids.Num(); IdIndex++)
+					{
+						const UE::Anim::FAttributeId& Id = Ids[IdIndex];
+						const FName BoneName = InOwningComponent->GetSkeletalMeshAsset()->GetRefSkeleton().GetBoneName(Id.GetIndex());
+					
+						const FAnimationAttributeIdentifier Identifier(
+							Id.GetName(), Id.GetIndex() ,BoneName, Types[TypeIndex].Get());
+
+						CachedIdentifiers[CachedListIdentifierIndex] = Identifier;
+
+						CachedAttributeSnapshotMap.FindOrAdd(Identifier).Add(Snapshot.Key);
+					
+						CachedListIdentifierIndex++;
+					}
+				}
+			}
+
+			// filtered list should also be refreshed since it depends on the cache
+			RefreshFilteredAttributeEntries();
+
+			// delay value view refresh until tick since this function can be called from animation thread
+			bShouldRefreshValueView = true;
+
+			return;
+		}
+
+		if (SelectedAttribute.IsSet())
+		{
+			for (const FAttributeValueView& ValueView : SelectedAttributeSnapshotValueViews)
+			{
+				const int32* SnapshotIndex = CachedSnapshotNameIndexMap.Find(ValueView.SnapshotName);
+				if (ensure(SnapshotIndex) && ensure(InSnapshots.IsValidIndex(*SnapshotIndex)))
+				{
+					const TTuple<FName, const ContainerType&> Snapshot = InSnapshots[*SnapshotIndex];
+					const ContainerType& AttributeContainer = Snapshot.Value;
+
+					ValueView.UpdateValue(AttributeContainer);
+				}
+			}
+		}	
+	}
 
 	void ClearListView();
 	
@@ -129,9 +219,72 @@ public:
 	
 private:
 
+	template <typename ContainerType>
 	bool ShouldInvalidateListViewCache(
-		const TArray<TTuple<FName, const UE::Anim::FHeapAttributeContainer&>>& InSnapshots,
-		const USkeletalMeshComponent* InOwningComponent);
+		const TArray<TTuple<FName, const ContainerType&>>& InSnapshots,
+		const USkeletalMeshComponent* InOwningComponent)
+	{
+		// we need the SKM to look up bone names
+		// it should have been checked
+		check(InOwningComponent);
+
+		if (InSnapshots.Num() != CachedAttributeIdentifierLists.Num())
+		{
+			return true;
+		}
+	
+		for (int32 SnapshotIndex = 0; SnapshotIndex < InSnapshots.Num(); SnapshotIndex++)
+		{
+			const TTuple<FName, const ContainerType&> Snapshot = InSnapshots[SnapshotIndex];
+			const TTuple<FName, TArray<FAnimationAttributeIdentifier>>& CachedIdList = CachedAttributeIdentifierLists[SnapshotIndex];
+		
+			if (Snapshot.Key != CachedIdList.Key)
+			{
+				return true;
+			}
+
+			if (Snapshot.Value.Num() != CachedIdList.Value.Num())
+			{
+				return true;
+			}
+		}
+
+		for (int32 SnapshotIndex = 0; SnapshotIndex < InSnapshots.Num(); SnapshotIndex++)
+		{
+			const TTuple<FName, const ContainerType&> Snapshot = InSnapshots[SnapshotIndex];
+			const TTuple<FName, TArray<FAnimationAttributeIdentifier>> CachedIdentifierList = CachedAttributeIdentifierLists[SnapshotIndex];
+
+			const ContainerType& AttributeContainer = Snapshot.Value;
+			const TArray<FAnimationAttributeIdentifier>& CachedIdentifiers = CachedIdentifierList.Value;
+		
+			const TArray<TWeakObjectPtr<UScriptStruct>>& Types = AttributeContainer.GetUniqueTypes();
+
+			int32 CachedListIdentifierIndex = 0;
+			for (int32 TypeIndex = 0; TypeIndex < Types.Num(); TypeIndex++)
+			{
+				const TArray<UE::Anim::FAttributeId>& Ids = AttributeContainer.GetKeys(TypeIndex);
+			
+				for (int32 IdIndex = 0; IdIndex < Ids.Num(); IdIndex++)
+				{
+					const UE::Anim::FAttributeId& Id = Ids[IdIndex];
+					const FName BoneName = InOwningComponent->GetSkeletalMeshAsset()->GetRefSkeleton().GetBoneName(Id.GetIndex());
+				
+					FAnimationAttributeIdentifier Identifier(
+						Id.GetName(), Id.GetIndex() ,BoneName, Types[TypeIndex].Get());
+
+					if (!(Identifier == CachedIdentifiers[CachedListIdentifierIndex]))
+					{
+						return true;
+					}
+
+					CachedListIdentifierIndex++;
+				}
+			}
+		}
+
+		return false;
+	}
+
 	
 	void OnSelectionChanged(TSharedPtr<FAnimAttributeEntry> InEntry, ESelectInfo::Type InSelectType);
 
@@ -188,7 +341,20 @@ private:
 	{
 		FAttributeValueView(FName InSnapshotName, const FAnimAttributeEntry& InSelectedAttribute);
 
-		void UpdateValue(const UE::Anim::FHeapAttributeContainer& InAttributeContainer) const;
+		template<typename ContainerType>
+		void UpdateValue(const ContainerType& InAttributeContainer) const
+		{
+			const uint8* ValuePtr = InAttributeContainer.Find(SubjectAttribute.GetScriptStruct(), SubjectAttribute.GetAttributeId());
+
+			if (ensure(ValuePtr))
+			{
+				FMemory::Memcpy(
+					StructData->GetStructMemory(),
+					ValuePtr,
+					StructData->GetStruct()->GetStructureSize());
+			}
+		}
+
 
 		FAnimAttributeEntry SubjectAttribute;
 		FName SnapshotName;
@@ -210,3 +376,29 @@ inline void SAnimAttributeView::ClearListView()
 
 	bShouldRefreshListView = true;
 }
+
+class PERSONA_API SAnimAttributeViewer : public SCompoundWidget
+{
+	SLATE_BEGIN_ARGS( SAnimAttributeViewer )
+	{}
+	
+	SLATE_END_ARGS()
+	
+	void Construct( const FArguments& InArgs, const TSharedRef<class IEditableSkeleton>& InEditableSkeleton, const TSharedRef<class IPersonaPreviewScene>& InPreviewScene);
+	
+	virtual ~SAnimAttributeViewer() override {}
+	
+	virtual void Tick(const FGeometry& AllottedGeometry, const double InCurrentTime, const float InDeltaTime) override;
+
+	UAnimInstance* GetAnimInstance() const;
+
+protected:
+	/** Pointer to the preview scene we are bound to */
+	TWeakPtr<class IPersonaPreviewScene> PreviewScenePtr;
+
+	/** Pointer to the editable skeleton */
+	TWeakPtr<class IEditableSkeleton> EditableSkeletonPtr;
+
+	TSharedPtr<SAnimAttributeView> AttributeView;	
+	TSharedPtr<SPoseWatchPicker> PoseWatchPicker;
+};
diff --git a/Engine/Source/Runtime/Engine/Classes/Engine/PoseWatchRenderData.h b/Engine/Source/Runtime/Engine/Classes/Engine/PoseWatchRenderData.h
--- a/Engine/Source/Runtime/Engine/Classes/Engine/PoseWatchRenderData.h
+++ b/Engine/Source/Runtime/Engine/Classes/Engine/PoseWatchRenderData.h
@@ -3,6 +3,7 @@
 #pragma once
 
 #include "CoreMinimal.h"
+#include "Animation/AttributesRuntime.h"
 
 #if WITH_EDITORONLY_DATA
 
@@ -11,6 +12,12 @@
 #include "BoneIndices.h"
 #include "Animation/AnimCurveTypes.h"
 
+namespace UE::Anim
+{
+struct FStackAttributeContainer;
+struct FHeapAttributeContainer;
+}
+
 struct FReferenceSkeleton;
 
 struct FAnimNodePoseWatch
@@ -28,6 +35,8 @@ public:
 
 	ENGINE_API void SetCurves(const FBlendedCurve& InCurves);
 
+	ENGINE_API void SetAttributes(const UE::Anim::FStackAttributeContainer& InAttributes);
+
 	template<class TAllocator>
 	bool SetPose(const TArray<FBoneIndexType>& InRequiredBones, const TArray<FTransform, TAllocator>& InBoneTransforms)
 	{
@@ -50,6 +59,8 @@ public:
 
 	ENGINE_API const FBlendedHeapCurve& GetCurves() const;
 	
+	ENGINE_API const UE::Anim::FHeapAttributeContainer& GetAttributes() const;
+	
 	ENGINE_API const FTransform& GetWorldTransform() const;
 
 	ENGINE_API FLinearColor GetBoneColor() const;
@@ -65,6 +76,7 @@ private:
 	TArray<FBoneIndexType> RequiredBones;
 	TArray<FTransform> BoneTransforms;
 	FBlendedHeapCurve Curves;
+	UE::Anim::FHeapAttributeContainer HeapAttributes;
 	
 	// Mirrored properties updated on CopyPoseWatchData
 	FLinearColor BoneColor;
diff --git a/Engine/Source/Runtime/Engine/Private/Animation/AnimInstanceProxy.cpp b/Engine/Source/Runtime/Engine/Private/Animation/AnimInstanceProxy.cpp
--- a/Engine/Source/Runtime/Engine/Private/Animation/AnimInstanceProxy.cpp
+++ b/Engine/Source/Runtime/Engine/Private/Animation/AnimInstanceProxy.cpp
@@ -3430,12 +3430,25 @@ void FAnimInstanceProxy::RecordNodeAttribute(const FAnimInstanceProxy& InSourceP
 	}
 }
 
-void FAnimInstanceProxy::RegisterWatchedPose(const FCompactPose& Pose, int32 LinkID)
+void FAnimInstanceProxy::RegisterWatchedPose(const FCompactPose& Pose, const FBlendedCurve& Curve, int32 LinkID)
 {
-	RegisterWatchedPose(Pose, FBlendedCurve(), LinkID);
+	FPoseContext Context(this);
+	Context.Pose = Pose;
+	Context.Curve = Curve;
+
+	RegisterWatchedPose(Context, LinkID);
+}
+
+void FAnimInstanceProxy::RegisterWatchedPose(const FCSPose<FCompactPose>& Pose, const FBlendedCurve& Curve, int32 LinkID)
+{
+	FComponentSpacePoseContext Context(this);
+	Context.Pose = Pose;
+	Context.Curve = Curve;
+
+	RegisterWatchedPose(Context, LinkID);
 }
 
-void FAnimInstanceProxy::RegisterWatchedPose(const FCompactPose& Pose, const FBlendedCurve& InCurve, int32 LinkID)
+void FAnimInstanceProxy::RegisterWatchedPose(const FPoseContext& PoseContext, int32 LinkID)
 {
 	if (bIsBeingDebugged)
 	{
@@ -3452,22 +3465,22 @@ void FAnimInstanceProxy::RegisterWatchedPose(const FCompactPose& Pose, const FBl
 						PoseWatch.PoseWatch->SetIsNodeEnabled(true);
 
 						TArray<FTransform> BoneTransforms;
-						BoneTransforms.AddUninitialized(Pose.GetBoneContainer().GetNumBones());
+						BoneTransforms.AddUninitialized(PoseContext.Pose.GetBoneContainer().GetNumBones());
 
 						TArray<FBoneIndexType> TmpRequiredBones;
-						TmpRequiredBones.Reserve(Pose.GetBoneContainer().GetNumBones());
+						TmpRequiredBones.Reserve(PoseContext.Pose.GetBoneContainer().GetNumBones());
 
-						for (FCompactPoseBoneIndex BoneIndex : Pose.ForEachBoneIndex())
+						for (FCompactPoseBoneIndex BoneIndex : PoseContext.Pose.ForEachBoneIndex())
 						{
-							FMeshPoseBoneIndex MeshBoneIndex = Pose.GetBoneContainer().MakeMeshPoseIndex(BoneIndex);
-
-							BoneTransforms[MeshBoneIndex.GetInt()] = Pose[BoneIndex];
+							FMeshPoseBoneIndex MeshBoneIndex = PoseContext.Pose.GetBoneContainer().MakeMeshPoseIndex(BoneIndex);
+							BoneTransforms[MeshBoneIndex.GetInt()] = PoseContext.Pose[BoneIndex];
 							TmpRequiredBones.Add(MeshBoneIndex.GetInt());
 						}
 
 						PoseWatch.SetPose(TmpRequiredBones, BoneTransforms);
-						PoseWatch.SetCurves(InCurve);
+						PoseWatch.SetCurves(PoseContext.Curve);
 						PoseWatch.SetWorldTransform(SkelMeshComponent->GetComponentTransform());
+						PoseWatch.SetAttributes(PoseContext.CustomAttributes);
 
 						TRACE_ANIM_POSE_WATCH(*this, PoseWatch.PoseWatchPoseElement, PoseWatch.NodeID, PoseWatch.GetBoneTransforms(), PoseWatch.GetCurves(), PoseWatch.GetRequiredBones(), PoseWatch.GetWorldTransform(), PoseWatch.PoseWatchPoseElement->GetIsVisible());
 						break;
@@ -3478,17 +3491,38 @@ void FAnimInstanceProxy::RegisterWatchedPose(const FCompactPose& Pose, const FBl
 	}
 }
 
-void FAnimInstanceProxy::RegisterWatchedPose(const FCSPose<FCompactPose>& Pose, int32 LinkID)
+void FAnimInstanceProxy::RegisterWatchedPose(const FComponentSpacePoseContext& PoseContext, int32 LinkID)
 {
-	RegisterWatchedPose(Pose, FBlendedCurve(), LinkID);
-}
+	if (bIsBeingDebugged)
+	{
+		FAnimBlueprintDebugData* DebugData = GetAnimBlueprintDebugData();
+		if (DebugData)
+		{
+			if (USkeletalMeshComponent* SkelMeshComponent = GetSkelMeshComponent())
+			{
+				for (FAnimNodePoseWatch& PoseWatch : DebugData->AnimNodePoseWatch)
+				{
+					if (PoseWatch.PoseWatch && PoseWatch.NodeID == LinkID)
+					{
+						FCompactPose TempPose;
+						FCSPose<FCompactPose>::ConvertComponentPosesToLocalPoses(PoseContext.Pose, TempPose);
+						PoseWatch.Object = GetAnimInstanceObject();
+						PoseWatch.PoseWatch->SetIsNodeEnabled(true);
 
-void FAnimInstanceProxy::RegisterWatchedPose(const FCSPose<FCompactPose>& Pose, const FBlendedCurve& InCurve, int32 LinkID)
-{
-	FCompactPose LocalPose;
-	FCSPose<FCompactPose>::ConvertComponentPosesToLocalPoses(Pose, LocalPose);
+						const TArray<FTransform, FAnimStackAllocator>& BoneTransforms = TempPose.GetBones();
+						const TArray<FBoneIndexType>& TmpRequiredBones = TempPose.GetBoneContainer().GetBoneIndicesArray();
+						PoseWatch.SetPose(TmpRequiredBones, BoneTransforms);
+						PoseWatch.SetCurves(PoseContext.Curve);
+						PoseWatch.SetWorldTransform(SkelMeshComponent->GetComponentTransform());
+						PoseWatch.SetAttributes(PoseContext.CustomAttributes);
 
-	RegisterWatchedPose(LocalPose, InCurve, LinkID);
+						TRACE_ANIM_POSE_WATCH(*this, PoseWatch.PoseWatchPoseElement, PoseWatch.NodeID, PoseWatch.GetBoneTransforms(), PoseWatch.GetCurves(), PoseWatch.GetRequiredBones(), PoseWatch.GetWorldTransform(), PoseWatch.PoseWatchPoseElement->GetIsVisible());
+						break;
+					}
+				}
+			}
+		}
+	}
 }
 #endif
 
diff --git a/Engine/Source/Runtime/Engine/Private/Animation/AnimNodeBase.cpp b/Engine/Source/Runtime/Engine/Private/Animation/AnimNodeBase.cpp
--- a/Engine/Source/Runtime/Engine/Private/Animation/AnimNodeBase.cpp
+++ b/Engine/Source/Runtime/Engine/Private/Animation/AnimNodeBase.cpp
@@ -421,7 +421,7 @@ void FPoseLink::Evaluate(FPoseContext& Output)
 #endif
 
 #if WITH_EDITOR
-		Output.AnimInstanceProxy->RegisterWatchedPose(Output.Pose, Output.Curve, LinkID);
+		Output.AnimInstanceProxy->RegisterWatchedPose(Output, LinkID);
 #endif
 	}
 	else
@@ -489,7 +489,7 @@ void FComponentSpacePoseLink::EvaluateComponentSpace(FComponentSpacePoseContext&
 #endif
 
 #if WITH_EDITOR
-		Output.AnimInstanceProxy->RegisterWatchedPose(Output.Pose, Output.Curve, LinkID);
+		Output.AnimInstanceProxy->RegisterWatchedPose(Output, LinkID);
 #endif
 	}
 	else
diff --git a/Engine/Source/Runtime/Engine/Private/Engine/PoseWatchRenderData.cpp b/Engine/Source/Runtime/Engine/Private/Engine/PoseWatchRenderData.cpp
--- a/Engine/Source/Runtime/Engine/Private/Engine/PoseWatchRenderData.cpp
+++ b/Engine/Source/Runtime/Engine/Private/Engine/PoseWatchRenderData.cpp
@@ -1,6 +1,7 @@
 // Copyright Epic Games, Inc. All Rights Reserved.
 
 #include "Engine/PoseWatchRenderData.h"
+#include "Animation/AttributesRuntime.h"
 
 #if WITH_EDITORONLY_DATA
 
@@ -23,6 +24,11 @@ void FAnimNodePoseWatch::SetCurves(const FBlendedCurve& InCurves)
 	Curves.CopyFrom(InCurves);
 }
 
+void FAnimNodePoseWatch::SetAttributes(const UE::Anim::FStackAttributeContainer& InAttributes)
+{
+	HeapAttributes.CopyFrom(InAttributes);
+}
+
 void FAnimNodePoseWatch::SetWorldTransform(const FTransform& InWorldTransform)
 {
 	WorldTransform = InWorldTransform;
@@ -77,6 +83,11 @@ const FBlendedHeapCurve& FAnimNodePoseWatch::GetCurves() const
 	return Curves;
 }
 
+const UE::Anim::FHeapAttributeContainer& FAnimNodePoseWatch::GetAttributes() const
+{
+	return HeapAttributes;
+}
+
 const FTransform& FAnimNodePoseWatch::GetWorldTransform() const
 {
 	return WorldTransform;
diff --git a/Engine/Source/Runtime/Engine/Public/Animation/AnimInstanceProxy.h b/Engine/Source/Runtime/Engine/Public/Animation/AnimInstanceProxy.h
--- a/Engine/Source/Runtime/Engine/Public/Animation/AnimInstanceProxy.h
+++ b/Engine/Source/Runtime/Engine/Public/Animation/AnimInstanceProxy.h
@@ -14,6 +14,7 @@
 
 #include "AnimInstanceProxy.generated.h"
 
+struct FComponentSpacePoseContext;
 class UAnimInstance;
 class UAnimNotify;
 class UAnimNotifyState;
@@ -198,14 +199,15 @@ public:
 		return ActualAnimClass ? Cast<UAnimBlueprint>(ActualAnimClass->ClassGeneratedBy) : nullptr;
 	}
 
-	UE_DEPRECATED(5.3, "Please use RegisterWatchedPose that takes a pose and a curve")
-	ENGINE_API void RegisterWatchedPose(const FCompactPose& Pose, int32 LinkID);
-	UE_DEPRECATED(5.3, "Please use RegisterWatchedPose that takes a pose and a curve")
-	ENGINE_API void RegisterWatchedPose(const FCSPose<FCompactPose>& Pose, int32 LinkID);
-
 	// Record pose for node of ID LinkID if it is currently being watched
+	UE_DEPRECATED(5.6, "Please use RegisterWatchedPose that takes a pose context")
 	ENGINE_API void RegisterWatchedPose(const FCompactPose& Pose, const FBlendedCurve& InCurve, int32 LinkID);
+	UE_DEPRECATED(5.6, "Please use RegisterWatchedPose that takes a pose context")
 	ENGINE_API void RegisterWatchedPose(const FCSPose<FCompactPose>& Pose, const FBlendedCurve& InCurve, int32 LinkID);
+
+	// Record pose for node of ID LinkID if it is currently being watched
+	ENGINE_API void RegisterWatchedPose(const FPoseContext& Pose, int32 LinkID);
+	ENGINE_API void RegisterWatchedPose(const FComponentSpacePoseContext& Pose, int32 LinkID);
 #endif
 
 	UE_DEPRECATED(5.0, "Function renamed to FlipBufferWriteIndex as this no longer deals with sync groups.")
