Trim tools
- Made AnimationSequence Trim APIs framenumber (range) based rather than using time values
- UAnimSequencerController::SetMovieSceneRange now sets value directly as SetRange forcefully transacts object (clashes with controller actions)
- Changed Insert Frames to ask user to specify # of frames to insert (similar to the append operation)
- Adjusted some user-facing text/tooltips
- Ensure that AnimDataModel (insert/remove frame) notifies are not handled during transacting (notifies/syncmarkers follow regular transaction model so was doubling-up)
- Fixed up notify states not being adjusted when inserting frames

diff --git a/Engine/Plugins/Animation/AnimationData/Source/AnimationData/Private/AnimSequencerController.cpp b/Engine/Plugins/Animation/AnimationData/Source/AnimationData/Private/AnimSequencerController.cpp
--- a/Engine/Plugins/Animation/AnimationData/Source/AnimationData/Private/AnimSequencerController.cpp
+++ b/Engine/Plugins/Animation/AnimationData/Source/AnimationData/Private/AnimSequencerController.cpp
@@ -242,7 +242,7 @@ void UAnimSequencerController::SetMovieSceneRange(FFrameNumber InFrameNumber) co
 
 	if (UMovieSceneControlRigParameterSection* Section = Model->GetFKControlRigSection())
 	{
-		Section->SetRange(DataRange);
+		Section->SectionRange = DataRange;
 	}
 }
 
diff --git a/Engine/Source/Editor/Persona/Private/AnimTimeline/SAnimTimeline.cpp b/Engine/Source/Editor/Persona/Private/AnimTimeline/SAnimTimeline.cpp
--- a/Engine/Source/Editor/Persona/Private/AnimTimeline/SAnimTimeline.cpp
+++ b/Engine/Source/Editor/Persona/Private/AnimTimeline/SAnimTimeline.cpp
@@ -414,70 +414,64 @@ FReply SAnimTimeline::OnMouseButtonUp(const FGeometry& MyGeometry, const FPointe
 			FFrameTime MouseTime = TimeSliderController->GetFrameTimeFromMouse(MyGeometry, MouseEvent.GetScreenSpacePosition());
 			const float CurrentFrameTime = static_cast<float>(MouseTime.AsDecimal() / static_cast<double>(Model.Pin()->GetTickResolution()));
 			const float SequenceLength = AnimSequence->GetPlayLength();
-			const uint32 NumKeys = AnimSequence->GetNumberOfSampledKeys();
+
+			const FFrameTime MouseFrameTime = TimeSliderController->GetFrameTimeFromMouse(MyGeometry, MouseEvent.GetScreenSpacePosition());
+			const FFrameTime AssetFrameTime = ConvertFrameTime(MouseFrameTime, FFrameRate(Model.Pin()->GetTickResolution(), 1), Model.Pin()->GetFrameRate());
+			const FFrameNumber RoundedAssetFrame = AssetFrameTime.RoundToFrame();
+			const int32 NumKeys = AnimSequence->GetNumberOfSampledKeys();
 
 			MenuBuilder.BeginSection("SequenceEditingContext", LOCTEXT("SequenceEditing", "Sequence Editing") );
 			{
-				float CurrentFrameFraction = CurrentFrameTime / SequenceLength;
-				int32 CurrentKeyIndex = static_cast<int32>(CurrentFrameFraction * NumKeys);
-
 				FUIAction Action;
 				FText Label;
 
 				//Menu - "Remove Before"
 				//Only show this option if the selected frame is greater than frame 1 (first frame)
-				if (CurrentKeyIndex > 0)
+				if (RoundedAssetFrame > 0)
 				{
-					CurrentFrameFraction = static_cast<float>(CurrentKeyIndex) / static_cast<float>(NumKeys);
-
 					//Corrected frame time based on selected frame number
-					float CorrectedFrameTime = CurrentFrameFraction * SequenceLength;
-
-					Action = FUIAction(FExecuteAction::CreateSP(this, &SAnimTimeline::OnCropAnimSequence, true, CorrectedFrameTime));
-					Label = FText::Format(LOCTEXT("RemoveTillFrame", "Remove frame 0 to frame {0}"), FText::AsNumber(CurrentKeyIndex));
+					Action = FUIAction(FExecuteAction::CreateSP(this, &SAnimTimeline::OnCropAnimSequence, true, RoundedAssetFrame));
+					Label = FText::Format(LOCTEXT("RemoveTillFrame", "Remove frames in range 0 to {0}"), FText::AsNumber(RoundedAssetFrame.Value));
 					MenuBuilder.AddMenuEntry(Label, LOCTEXT("RemoveBefore_ToolTip", "Remove sequence before current position"), FSlateIcon(), Action);
 				}
 
-				uint32 NextKeyIndex = CurrentKeyIndex + 1;
+				const FFrameNumber NextFrameNumber = RoundedAssetFrame.Value + 1;
 
 				//Menu - "Remove After"
 				//Only show this option if next frame (CurrentKeyIndex + 1) is valid
-				if (NextKeyIndex < NumKeys)
+				if (NextFrameNumber < NumKeys)
 				{
-					float NextFrameFraction = (float)NextKeyIndex / (float)NumKeys;
-					float NextFrameTime = NextFrameFraction * SequenceLength;
-					Action = FUIAction(FExecuteAction::CreateSP(this, &SAnimTimeline::OnCropAnimSequence, false, NextFrameTime));
-					Label = FText::Format(LOCTEXT("RemoveFromFrame", "Remove from frame {0} to frame {1}"), FText::AsNumber(NextKeyIndex), FText::AsNumber(NumKeys));
+					Action = FUIAction(FExecuteAction::CreateSP(this, &SAnimTimeline::OnCropAnimSequence, false, NextFrameNumber));
+					Label = FText::Format(LOCTEXT("RemoveFromFrame", "Remove frames in range {0} to {1}"), FText::AsNumber(NextFrameNumber.Value), FText::AsNumber(NumKeys - 1));
 					MenuBuilder.AddMenuEntry(Label, LOCTEXT("RemoveAfter_ToolTip", "Remove sequence after current position"), FSlateIcon(), Action);
 				}
 
 				MenuBuilder.AddMenuSeparator();
 
 				//Corrected frame time based on selected frame number
-				float CorrectedFrameTime = CurrentFrameFraction * SequenceLength;
 
-				Action = FUIAction(FExecuteAction::CreateSP(this, &SAnimTimeline::OnInsertAnimSequence, true, CurrentKeyIndex));
-				Label = FText::Format(LOCTEXT("InsertBeforeCurrentFrame", "Insert frame before {0}"), FText::AsNumber(CurrentKeyIndex));
+				Action = FUIAction(FExecuteAction::CreateSP(this, &SAnimTimeline::OnPromptUserForNumberOfFrames, WidgetPath, TFunction<void(const int32)>([this, RoundedAssetFrame](const int32 NumberOfFrames) { OnInsertAnimSequence(true, RoundedAssetFrame.Value, NumberOfFrames); })));
+				Label = FText::Format(LOCTEXT("InsertBeforeCurrentFrame", "Insert frames before {0}"), FText::AsNumber(RoundedAssetFrame.Value));
 				MenuBuilder.AddMenuEntry(Label, LOCTEXT("InsertBefore_ToolTip", "Insert a frame before current position"), FSlateIcon(), Action);
 
-				Action = FUIAction(FExecuteAction::CreateSP(this, &SAnimTimeline::OnInsertAnimSequence, false, CurrentKeyIndex));
-				Label = FText::Format(LOCTEXT("InsertAfterCurrentFrame", "Insert frame after {0}"), FText::AsNumber(CurrentKeyIndex));
+				Action = FUIAction(FExecuteAction::CreateSP(this, &SAnimTimeline::OnPromptUserForNumberOfFrames, WidgetPath, TFunction<void(const int32)>([this, RoundedAssetFrame](const int32 NumberOfFrames) { OnInsertAnimSequence(false, RoundedAssetFrame.Value, NumberOfFrames); })));
+				Label = FText::Format(LOCTEXT("InsertAfterCurrentFrame", "Insert frames after {0}"), FText::AsNumber(RoundedAssetFrame.Value));
 				MenuBuilder.AddMenuEntry(Label, LOCTEXT("InsertAfter_ToolTip", "Insert a frame after current position"), FSlateIcon(), Action);
 
 				MenuBuilder.AddMenuSeparator();
 
 				//Corrected frame time based on selected frame number
-				Action = FUIAction(FExecuteAction::CreateSP(this, &SAnimTimeline::OnShowPopupOfAppendAnimation, WidgetPath, true));
-				MenuBuilder.AddMenuEntry(LOCTEXT("AppendBegin", "Append in the beginning"), LOCTEXT("AppendBegin_ToolTip", "Append in the beginning"), FSlateIcon(), Action);
+				Action = FUIAction(FExecuteAction::CreateSP(this, &SAnimTimeline::OnPromptUserForNumberOfFrames, WidgetPath, TFunction<void(const int32)>([this](const int32 NumberOfFrames) { OnAppendAnimSequence(true, NumberOfFrames); })));
+				MenuBuilder.AddMenuEntry(LOCTEXT("AppendBegin", "Add frames to start of sequence"), LOCTEXT("AppendBegin_ToolTip", "Adds user provided number of frames at the start of the sequence"), FSlateIcon(), Action);
 
-				Action = FUIAction(FExecuteAction::CreateSP(this, &SAnimTimeline::OnShowPopupOfAppendAnimation, WidgetPath, false));
-				MenuBuilder.AddMenuEntry(LOCTEXT("AppendEnd", "Append at the end"), LOCTEXT("AppendEnd_ToolTip", "Append at the end"), FSlateIcon(), Action);
+				Action = FUIAction(FExecuteAction::CreateSP(this, &SAnimTimeline::OnPromptUserForNumberOfFrames, WidgetPath, TFunction<void(const int32)>([this](const int32 NumberOfFrames) { OnAppendAnimSequence(false, NumberOfFrames); })));
+				MenuBuilder.AddMenuEntry(LOCTEXT("AppendEnd", "Append frames at end of sequence"), LOCTEXT("AppendEnd_ToolTip", "AAdds user provided number of frames at the end of the sequence"), FSlateIcon(), Action);
 
 				MenuBuilder.AddMenuSeparator();
 				//Menu - "ReZero"
-				Action = FUIAction(FExecuteAction::CreateSP(this, &SAnimTimeline::OnReZeroAnimSequence, CurrentKeyIndex));
-				Label = FText::Format(LOCTEXT("ReZeroAtFrame", "Re-zero at frame {0}"), FText::AsNumber(CurrentKeyIndex));
-				MenuBuilder.AddMenuEntry(Label, FText::Format(LOCTEXT("ReZeroAtFrame_ToolTip", "Resets the root track to (0, 0, 0) at frame {0} and apply the difference to all root transform of the sequence. It moves whole sequence to the amount of current root transform."), FText::AsNumber(CurrentKeyIndex)), FSlateIcon(), Action);
+				Action = FUIAction(FExecuteAction::CreateSP(this, &SAnimTimeline::OnReZeroAnimSequence, RoundedAssetFrame.Value));
+				Label = FText::Format(LOCTEXT("ReZeroAtFrame", "Re-zero at frame {0}"), FText::AsNumber(RoundedAssetFrame.Value));
+				MenuBuilder.AddMenuEntry(Label, FText::Format(LOCTEXT("ReZeroAtFrame_ToolTip", "Resets the root track to (0, 0, 0) at frame {0} and apply the difference to all root transform of the sequence. It moves whole sequence to the amount of current root transform."), FText::AsNumber(RoundedAssetFrame.Value)), FSlateIcon(), Action);
 
 				const int32 FrameNumberForCurrentTime = INDEX_NONE;
 				Action = FUIAction(FExecuteAction::CreateSP(this, &SAnimTimeline::OnReZeroAnimSequence, FrameNumberForCurrentTime));
@@ -495,12 +489,11 @@ FReply SAnimTimeline::OnMouseButtonUp(const FGeometry& MyGeometry, const FPointe
 	return FReply::Unhandled();
 }
 
-void SAnimTimeline::OnCropAnimSequence( bool bFromStart, float CurrentTime )
+void SAnimTimeline::OnCropAnimSequence(bool bFromStart, FFrameNumber StartFrame)
 {
 	UAnimSingleNodeInstance* PreviewInstance = GetPreviewInstance();
 	if(PreviewInstance)
 	{
-		float Length = PreviewInstance->GetLength();
 		if (PreviewInstance->GetCurrentAsset())
 		{
 			UAnimSequence* AnimSequence = Cast<UAnimSequence>( PreviewInstance->GetCurrentAsset() );
@@ -514,12 +507,11 @@ void SAnimTimeline::OnCropAnimSequence( bool bFromStart, float CurrentTime )
 				//Call modify to restore anim sequence current state
 				AnimSequence->Modify();
 				
-				const float TrimStart = bFromStart ? 0.f : CurrentTime;
-				const float TrimEnd = bFromStart ? CurrentTime : AnimSequence->GetPlayLength();
+				const TRange<FFrameNumber> TrimRange(TRangeBound<FFrameNumber>::Inclusive(bFromStart ? 0 : StartFrame),
+					 TRangeBound<FFrameNumber>::Exclusive(bFromStart ? StartFrame : AnimSequence->GetNumberOfSampledKeys()));
 
 				// Trim off the user-selected part of the raw anim data.
-				UE::Anim::AnimationData::Trim(AnimSequence, TrimStart, TrimEnd, !bFromStart);
-
+				UE::Anim::AnimationData::Trim(AnimSequence, TrimRange);
 
 				//Resetting slider position to the first frame
 				PreviewInstance->SetPosition( 0.0f, false );
@@ -551,7 +543,7 @@ void SAnimTimeline::OnAppendAnimSequence( bool bFromStart, int32 NumOfFrames )
 	}
 }
 
-void SAnimTimeline::OnInsertAnimSequence( bool bBefore, int32 CurrentFrame )
+void SAnimTimeline::OnInsertAnimSequence(bool bBefore, int32 CurrentFrame, int32 NumFrames)
 {
 	UAnimSingleNodeInstance* PreviewInstance = GetPreviewInstance();
 	if(PreviewInstance && PreviewInstance->GetCurrentAsset())
@@ -569,7 +561,7 @@ void SAnimTimeline::OnInsertAnimSequence( bool bBefore, int32 CurrentFrame )
 
 			// Crop the raw anim data.
 			const int32 StartFrame = (bBefore)? CurrentFrame : CurrentFrame + 1;
-			UE::Anim::AnimationData::DuplicateKeys(AnimSequence, StartFrame, 1, CurrentFrame);
+			UE::Anim::AnimationData::DuplicateKeys(AnimSequence, StartFrame, NumFrames, CurrentFrame);
 		}
 	}
 }
@@ -638,14 +630,22 @@ void SAnimTimeline::OnReZeroAnimSequence(int32 FrameIndex)
 	}
 }
 
-void SAnimTimeline::OnShowPopupOfAppendAnimation(FWidgetPath WidgetPath, bool bBegin)
+void SAnimTimeline::OnPromptUserForNumberOfFrames(FWidgetPath WidgetPath, TFunction<void(const int32 NumFrames)> Callback)
 {
 	TSharedRef<STextEntryPopup> TextEntry =
 		SNew(STextEntryPopup)
-		.Label(LOCTEXT("AppendAnim_AskNumFrames", "Number of Frames to Append"))
-		.OnTextCommitted(this, &SAnimTimeline::OnSequenceAppendedCalled, bBegin);
+		.Label(LOCTEXT("AskNumFrames", "Number of Frames to Add"))
+		.OnTextCommitted_Lambda([Callback](const FText & InNewGroupText, ETextCommit::Type CommitInfo)
+		{
+			// handle only onEnter. This is a big thing to apply when implicit focus change or any other event
+			if (CommitInfo == ETextCommit::OnEnter)
+			{
+				const int32 NumFrames = FMath::Clamp(FCString::Atoi(*InNewGroupText.ToString()), 0, 1000);
+				Callback(NumFrames);
+				FSlateApplication::Get().DismissAllMenus();
+			}
+		});
 
-	// Show dialog to enter new track name
 	FSlateApplication::Get().PushMenu(
 		SharedThis(this),
 		WidgetPath,
@@ -653,23 +653,7 @@ void SAnimTimeline::OnShowPopupOfAppendAnimation(FWidgetPath WidgetPath, bool bB
 		FSlateApplication::Get().GetCursorPos(),
 		FPopupTransitionEffect(FPopupTransitionEffect::TypeInPopup)
 		);
-}
-
-void SAnimTimeline::OnSequenceAppendedCalled(const FText & InNewGroupText, ETextCommit::Type CommitInfo, bool bBegin)
-{
-	// just a concern
-	const static int32 MaxFrame = 1000;
-
-	// handle only onEnter. This is a big thing to apply when implicit focus change or any other event
-	if (CommitInfo == ETextCommit::OnEnter)
-	{
-		int32 NumFrames = FCString::Atoi(*InNewGroupText.ToString());
-		if (NumFrames > 0 && NumFrames < MaxFrame)
-		{
-			OnAppendAnimSequence(bBegin, NumFrames);
-			FSlateApplication::Get().DismissAllMenus();
-		}
-	}
+	
 }
 
 TSharedRef<INumericTypeInterface<double>> SAnimTimeline::GetNumericTypeInterface() const
diff --git a/Engine/Source/Editor/Persona/Private/AnimTimeline/SAnimTimeline.h b/Engine/Source/Editor/Persona/Private/AnimTimeline/SAnimTimeline.h
--- a/Engine/Source/Editor/Persona/Private/AnimTimeline/SAnimTimeline.h
+++ b/Engine/Source/Editor/Persona/Private/AnimTimeline/SAnimTimeline.h
@@ -81,18 +81,15 @@ private:
 
 	void HandleScrubPositionChanged(FFrameTime NewScrubPosition, bool bIsScrubbing, bool bEvaluate) const;
 
-	void OnCropAnimSequence(bool bFromStart, float CurrentTime);
+	void OnCropAnimSequence(bool bFromStart, FFrameNumber StartFrame);
 
 	void OnAppendAnimSequence(bool bFromStart, int32 NumOfFrames);
 
-	void OnInsertAnimSequence(bool bBefore, int32 CurrentFrame);
+	void OnInsertAnimSequence(bool bBefore, int32 CurrentFrame, int32 NumFrames);
 
 	void OnReZeroAnimSequence(int32 FrameIndex);
 
-	void OnShowPopupOfAppendAnimation(FWidgetPath WidgetPath, bool bBegin);
-
-	void OnSequenceAppendedCalled(const FText & InNewGroupText, ETextCommit::Type CommitInfo, bool bBegin);
-
+	void OnPromptUserForNumberOfFrames(FWidgetPath WidgetPath, TFunction<void(const int32 NumFrames)> Callback);
 	double GetSpinboxDelta() const;
 
 	void SetPlayTime(double InFrameTime);
diff --git a/Engine/Source/Editor/Persona/Private/SAnimationScrubPanel.cpp b/Engine/Source/Editor/Persona/Private/SAnimationScrubPanel.cpp
--- a/Engine/Source/Editor/Persona/Private/SAnimationScrubPanel.cpp
+++ b/Engine/Source/Editor/Persona/Private/SAnimationScrubPanel.cpp
@@ -606,12 +606,13 @@ void SAnimationScrubPanel::OnCropAnimSequence( bool bFromStart, float CurrentTim
 				//Call modify to restore anim sequence current state
 				AnimSequence->Modify();
 
-
-				const float TrimStart = bFromStart ? 0.f : CurrentTime;
-				const float TrimEnd = bFromStart ? CurrentTime : AnimSequence->GetPlayLength();
+				const FFrameTime AssetFrameTime = AnimSequence->GetSamplingFrameRate().AsFrameTime(CurrentTime);
+				const FFrameNumber RoundedAssetFrame = AssetFrameTime.RoundToFrame();
+				const TRange<FFrameNumber> TrimRange(TRangeBound<FFrameNumber>::Inclusive(bFromStart ? 0 : RoundedAssetFrame),
+				bFromStart ? TRangeBound<FFrameNumber>::Exclusive(RoundedAssetFrame) : TRangeBound<FFrameNumber>::Exclusive(AnimSequence->GetNumberOfSampledKeys()));
 
 				// Trim off the user-selected part of the raw anim data.
-				UE::Anim::AnimationData::Trim(AnimSequence, TrimStart, TrimEnd);
+				UE::Anim::AnimationData::Trim(AnimSequence, TrimRange);
 
 				//Resetting slider position to the first frame
 				PreviewInstance->SetPosition( 0.0f, false );
diff --git a/Engine/Source/Runtime/Engine/Classes/Animation/AnimSequenceHelpers.h b/Engine/Source/Runtime/Engine/Classes/Animation/AnimSequenceHelpers.h
--- a/Engine/Source/Runtime/Engine/Classes/Animation/AnimSequenceHelpers.h
+++ b/Engine/Source/Runtime/Engine/Classes/Animation/AnimSequenceHelpers.h
@@ -132,12 +132,13 @@ namespace Anim {
 		* Trim a specific window from the animation data
 		*
 		* @param	InSequence		Animation Sequence to add the animated frame to
-		* @param	TrimStart		Time value at which to start removing frames
-		* @param	TrimEnd			Time value at which to stop removing frames
-		* @param	bInclusiveEnd   Whether or not the TrimEnd should be included in the trimming range
+		* @param	TrimRange		Range of _frame numbers_ to remove
 		*
 		* @return Whether or not the operation was succesful
 		*/
+		ENGINE_API bool Trim(UAnimSequence* InSequence, TRange<FFrameNumber> TrimRange);
+
+		UE_DEPRECATED(5.6, "Trimming operation using time values has been deprecated use Trim with FFrameNumbers instead")
 		ENGINE_API bool Trim(UAnimSequence* InSequence, float TrimStart, float TrimEnd, bool bInclusiveEnd=false);
 				
 		/**
diff --git a/Engine/Source/Runtime/Engine/Private/Animation/AnimSequence.cpp b/Engine/Source/Runtime/Engine/Private/Animation/AnimSequence.cpp
--- a/Engine/Source/Runtime/Engine/Private/Animation/AnimSequence.cpp
+++ b/Engine/Source/Runtime/Engine/Private/Animation/AnimSequence.cpp
@@ -4682,7 +4682,8 @@ void UAnimSequence::OnModelModified(const EAnimDataModelNotifyType& NotifyType,
 
 	auto HandleNumberOfFramesChanged = [this](FFrameNumber NewNumberOfFrames, FFrameNumber OldNumberOfFrames, FFrameNumber Frame0, FFrameNumber Frame1)
 	{
-		if (bPopulatingDataModel)
+		// Do not handle changes during model population, or undo-redo (notifies are transacted so will be restored/set, so only handle in case of, initial, user interaction)
+		if (bPopulatingDataModel || GIsTransacting)
 		{
 			return;
 		}
diff --git a/Engine/Source/Runtime/Engine/Private/Animation/AnimSequenceBase.cpp b/Engine/Source/Runtime/Engine/Private/Animation/AnimSequenceBase.cpp
--- a/Engine/Source/Runtime/Engine/Private/Animation/AnimSequenceBase.cpp
+++ b/Engine/Source/Runtime/Engine/Private/Animation/AnimSequenceBase.cpp
@@ -1095,7 +1095,8 @@ void UAnimSequenceBase::OnModelModified(const EAnimDataModelNotifyType& NotifyTy
 
 	auto HandleNumberOfFramesChange = [&](int32 NewNumFrames, int32 OldNumFrames, int32 Frame0, int32 Frame1)
 	{
-		if (bPopulatingDataModel)
+		// Do not handle changes during model population, or undo-redo (notifies are transacted so will be restored/set, so only handle in case of, initial, user interaction)
+		if (bPopulatingDataModel || GIsTransacting)
 		{
 			return;
 		}
@@ -1105,8 +1106,8 @@ void UAnimSequenceBase::OnModelModified(const EAnimDataModelNotifyType& NotifyTy
 		const float T1 = ModelFrameRate.AsSeconds(Frame1);
 		const float NewLength = ModelFrameRate.AsSeconds(NewNumFrames);
 		
-		const float StartRemoveTime = ModelFrameRate.AsSeconds(Frame0);
-		const float EndRemoveTime = ModelFrameRate.AsSeconds(Frame1);
+		const float StartTime = ModelFrameRate.AsSeconds(Frame0);
+		const float EndTime = ModelFrameRate.AsSeconds(Frame1);
 
 		// Total time value for frames that were removed
 		const float Duration = T1 - T0;
@@ -1123,21 +1124,24 @@ void UAnimSequenceBase::OnModelModified(const EAnimDataModelNotifyType& NotifyTy
 				}
 
 				float CurrentTime = Notify.GetTime();
-				float NewDuration = 0.f;
+				float NewDuration = Notify.GetDuration();
 
-				// if state, make sure to adjust end time
+				// if state, make sure to adjust timings
 				if (Notify.NotifyStateClass)
 				{
 					const float NotifyDuration = Notify.GetDuration();
 					const float NotifyEnd = CurrentTime + NotifyDuration;
-					if (NotifyEnd >= InsertTime)
-					{
-						NewDuration = NotifyDuration + Duration;
-					}
-					else
+
+					TRange<float> NotifyRange(CurrentTime, TRange<float>::BoundsType(CurrentTime+NotifyDuration));
+					TRange<float> InsertionRange(StartTime, EndTime);
+
+					// If insertion range or start time falls inside of notify range, add length to notify duration
+					if (NotifyRange.Contains(InsertionRange) || NotifyRange.Contains(StartTime))
 					{
-						NewDuration = NotifyDuration;
-					}
+						// Removal range falls entirely inside of notify state duration
+						const float InsertDuration = InsertionRange.GetUpperBound().GetValue() - InsertionRange.GetLowerBound().GetValue();
+						NewDuration += InsertDuration;
+					}					
 				}
 
 				// Shift out notify by the time allotted for the inserted keys
@@ -1183,26 +1187,45 @@ void UAnimSequenceBase::OnModelModified(const EAnimDataModelNotifyType& NotifyTy
 					const float NotifyEnd = CurrentTime + NotifyDuration;
 					NewDuration = NotifyDuration;
 
+					TRange<float> NotifyRange(CurrentTime, TRange<float>::BoundsType(CurrentTime+NotifyDuration));					
+					TRange<float> RemovalRange(StartTime, EndTime);
+
 					// If Notify is inside of the trimmed time frame, zero out the duration
-					if (CurrentTime >= StartRemoveTime && NotifyEnd <= EndRemoveTime)
+					if (RemovalRange.Contains(NotifyRange))
 					{
 						// small number @todo see if there is define for this
 						NewDuration = DataModelInterface->GetFrameRate().AsInterval();
 					}
+					else if (NotifyRange.Contains(RemovalRange))
+					{
+						// Removal range falls entirely inside of notify state duration
+						const float RemovalDuration = RemovalRange.GetUpperBound().GetValue() - RemovalRange.GetLowerBound().GetValue();
+						NewDuration -= RemovalDuration;
+					}
 					// If Notify overlaps trimmed time frame, remove trimmed duration
-					else if (CurrentTime < EndRemoveTime && NotifyEnd > EndRemoveTime)
+					else if (RemovalRange.Overlaps(NotifyRange))
 					{
-						NewDuration = NotifyEnd - Duration - CurrentTime;
+						TRange<float> Overlap = TRange<float>::Intersection(NotifyRange, RemovalRange);
+						if (Overlap.GetLowerBoundValue() > CurrentTime)
+						{
+							//Removing from back-end of notify
+							NewDuration = Overlap.GetLowerBoundValue() - CurrentTime;
+						}
+						else if (Overlap.GetUpperBoundValue() < NotifyEnd)
+						{
+							// Removing from front-end of notify
+							NewDuration = NotifyEnd - Overlap.GetUpperBoundValue();
+						}
 					}
 
 					NewDuration = FMath::Max(NewDuration, (float)DataModelInterface->GetFrameRate().AsInterval());
 				}
 
-				if (CurrentTime >= StartRemoveTime && CurrentTime <= EndRemoveTime)
+				if (CurrentTime >= StartTime && CurrentTime <= EndTime)
 				{
-					CurrentTime = StartRemoveTime;
+					CurrentTime = StartTime;
 				}
-				else if (CurrentTime > EndRemoveTime)
+				else if (CurrentTime > EndTime)
 				{
 					CurrentTime -= Duration;
 				}
diff --git a/Engine/Source/Runtime/Engine/Private/Animation/AnimSequenceHelpers.cpp b/Engine/Source/Runtime/Engine/Private/Animation/AnimSequenceHelpers.cpp
--- a/Engine/Source/Runtime/Engine/Private/Animation/AnimSequenceHelpers.cpp
+++ b/Engine/Source/Runtime/Engine/Private/Animation/AnimSequenceHelpers.cpp
@@ -802,7 +802,7 @@ bool AnimationData::AddLoopingInterpolation(UAnimSequence* InSequence)
 	return false;
 }
 
-bool AnimationData::Trim(UAnimSequence* InSequence, float TrimStart, float TrimEnd, bool bInclusiveEnd /*=false*/ )
+bool AnimationData::Trim(UAnimSequence* InSequence, TRange<FFrameNumber> TrimRange)
 {
 	const IAnimationDataModel* DataModel = InSequence->GetDataModel();
 	IAnimationDataController& Controller = InSequence->GetController();
@@ -813,19 +813,19 @@ bool AnimationData::Trim(UAnimSequence* InSequence, float TrimStart, float TrimE
 	if (NumTracks > 0 && NumKeys > 0)
 	{
 		const FFrameRate& FrameRate = DataModel->GetFrameRate();
-
-
 		// if there is only one key, there is nothing to trim away
 		if (NumKeys <= 1)
 		{
 			return false;
 		}
 
-		const FFrameNumber StartFrameTrim = FrameRate.AsFrameTime(TrimStart).RoundToFrame();
-		const FFrameNumber EndFrameTrim = FrameRate.AsFrameTime(TrimEnd).RoundToFrame();
+		// Convert frame range to key indices
+		const int32 StartTrimFrameIndex = TrimRange.GetLowerBound().IsInclusive() ? TrimRange.GetLowerBound().GetValue().Value : TrimRange.GetLowerBound().GetValue().Value + 1;
+		const int32 EndTrimFrameIndex = TrimRange.GetUpperBound().IsInclusive() ? TrimRange.GetUpperBound().GetValue().Value : TrimRange.GetUpperBound().GetValue().Value - 1;
+		const int32 NumTrimmedFrames = (EndTrimFrameIndex - StartTrimFrameIndex) + 1;
 
-		const int32 StartTrimKeyIndex = StartFrameTrim.Value;
-		const int32 NumTrimmedFrames = (EndFrameTrim.Value - StartFrameTrim.Value) + (bInclusiveEnd ? 1 : 0);
+		const int32 StartTrimKeyIndex = StartTrimFrameIndex;
+		const int32 NumTrimmedKeys = NumTrimmedFrames;
 
 		if (StartTrimKeyIndex == 0 && NumTrimmedFrames == NumKeys)
 		{
@@ -846,6 +846,17 @@ bool AnimationData::Trim(UAnimSequence* InSequence, float TrimStart, float TrimE
 	return false;
 }
 
+bool AnimationData::Trim(UAnimSequence* InSequence, float TrimStart, float TrimEnd, bool bInclusiveEnd /*=false*/ )
+{
+	const IAnimationDataModel* DataModel = InSequence->GetDataModel();
+	const FFrameRate& FrameRate = DataModel->GetFrameRate();
+	const FFrameNumber StartFrameTrim = FrameRate.AsFrameTime(TrimStart).RoundToFrame();
+	const FFrameNumber EndFrameTrim = FrameRate.AsFrameTime(TrimEnd).RoundToFrame();
+
+	const TRange<FFrameNumber> TrimRange(TRangeBound<FFrameNumber>::Inclusive(StartFrameTrim), bInclusiveEnd ? TRangeBound<FFrameNumber>::Inclusive(EndFrameTrim) : TRangeBound<FFrameNumber>::Exclusive(EndFrameTrim));
+	return Trim(InSequence, TrimRange);
+}
+
 void AnimationData::DuplicateKeys(UAnimSequence* InSequence, int32 StartKeyIndex, int32 NumDuplicates, int32 SourceKeyIndex /*= INDEX_NONE */)
 {
 	const IAnimationDataModel* Model = InSequence->GetDataModel();
@@ -1011,10 +1022,9 @@ void AnimationData::RemoveKeys(UAnimSequence* InSequence, int32 StartKeyIndex, i
 
 		const int32 NewNumberOfFrames = FMath::Max(NewNumberOfKeys - 1, 1);
 
-		const int32 StartKey = FMath::Max(StartKeyIndex, 0);
-
+		const int32 StartFrame = FMath::Max(StartKeyIndex - 1, 0);
 		// Notify will happen with time slice that was removed
-		Controller.ResizeInFrames(NewNumberOfFrames, StartKey, StartKey + NumKeysToRemove);
+		Controller.ResizeInFrames(NewNumberOfFrames, StartFrame, StartFrame + NumKeysToRemove);
 	}
 }
 
