TAttributeContainer::CopyFrom crashes when container contains invalid bone indices for active mesh
#fix Swapping control rig attribute container to be FMeshAttributeContainers, as Attribute inserting/retrieval is done using Skeleton indices (straight from SkeletalMesh->RefSkeleton). This ended up adding invalid "CompactPose" indices into the container which crashes during ::CopyFrom, as it tries to convert its into a FMeshPoseBoneIndex (while its value already is one). Adds overhead of converting to and from mesh attribute container in ControlRig AnimBP node.
#misc Changed back PoseWatch/AnimAttributeView to only be able to show FMeshAttributeContainer data, as it guarantees correct bone-index mappings for whatever is being previewed

diff --git a/Engine/Plugins/Animation/ControlRig/Source/ControlRig/Private/AnimNode_ControlRigBase.cpp b/Engine/Plugins/Animation/ControlRig/Source/ControlRig/Private/AnimNode_ControlRigBase.cpp
--- a/Engine/Plugins/Animation/ControlRig/Source/ControlRig/Private/AnimNode_ControlRigBase.cpp
+++ b/Engine/Plugins/Animation/ControlRig/Source/ControlRig/Private/AnimNode_ControlRigBase.cpp
@@ -234,10 +234,12 @@ void FAnimNode_ControlRigBase::ExecuteControlRig(FPoseContext& InOutput)
 		// Before we start modifying the RigHierarchy, we need to lock the rig to avoid
 		// corrupted state
 		UE::TScopeLock LockRig(ControlRig->GetEvaluateMutex());
-		
-		// temporarily give control rig access to the stack allocated attribute container
+
+		UE::Anim::FMeshAttributeContainer MeshAttributeContainer;
+		MeshAttributeContainer.CopyFrom(InOutput.CustomAttributes, InOutput.Pose.GetBoneContainer());
+		// temporarily give control rig access to the heap allocated attribute container
 		// control rig may have rig units that can add/get attributes to/from this container
-		UControlRig::FAnimAttributeContainerPtrScope AttributeScope(ControlRig, InOutput.CustomAttributes);
+		UControlRig::FAnimAttributeContainerPtrScope AttributeScope(ControlRig, MeshAttributeContainer);
 
 		URigHierarchy* Hierarchy = ControlRig->GetHierarchy();
 		if(Hierarchy == nullptr)
@@ -321,6 +323,7 @@ void FAnimNode_ControlRigBase::ExecuteControlRig(FPoseContext& InOutput)
 
 		// now update output
 		UpdateOutput(ControlRig, InOutput);
+		InOutput.CustomAttributes.CopyFrom(MeshAttributeContainer, InOutput.Pose.GetBoneContainer());
 	}
 }
 
diff --git a/Engine/Plugins/Animation/ControlRig/Source/ControlRig/Private/ControlRig.cpp b/Engine/Plugins/Animation/ControlRig/Source/ControlRig/Private/ControlRig.cpp
--- a/Engine/Plugins/Animation/ControlRig/Source/ControlRig/Private/ControlRig.cpp
+++ b/Engine/Plugins/Animation/ControlRig/Source/ControlRig/Private/ControlRig.cpp
@@ -775,7 +775,7 @@ void UControlRig::InitializeFromCDO()
 }
 
 UControlRig::FAnimAttributeContainerPtrScope::FAnimAttributeContainerPtrScope(UControlRig* InControlRig,
-	UE::Anim::FStackAttributeContainer& InExternalContainer)
+	UE::Anim::FMeshAttributeContainer& InExternalContainer)
 {
 	ControlRig = InControlRig;
 	ControlRig->ExternalAnimAttributeContainer = &InExternalContainer;
diff --git a/Engine/Plugins/Animation/ControlRig/Source/ControlRig/Private/ControlRigAnimInstance.cpp b/Engine/Plugins/Animation/ControlRig/Source/ControlRig/Private/ControlRigAnimInstance.cpp
--- a/Engine/Plugins/Animation/ControlRig/Source/ControlRig/Private/ControlRigAnimInstance.cpp
+++ b/Engine/Plugins/Animation/ControlRig/Source/ControlRig/Private/ControlRigAnimInstance.cpp
@@ -55,7 +55,7 @@ bool FControlRigAnimInstanceProxy::Evaluate(FPoseContext& Output)
 	UE::Anim::FCurveUtils::BuildUnsorted(Curve, StoredCurves);
 	Output.Curve.Combine(Curve);
 
-	Output.CustomAttributes.CopyFrom(StoredAttributes);
+	Output.CustomAttributes.CopyFrom(StoredAttributes, Output.Pose.GetBoneContainer());
 
 	return true;
 }
diff --git a/Engine/Plugins/Animation/ControlRig/Source/ControlRig/Private/ControlRigComponent.cpp b/Engine/Plugins/Animation/ControlRig/Source/ControlRig/Private/ControlRigComponent.cpp
--- a/Engine/Plugins/Animation/ControlRig/Source/ControlRig/Private/ControlRigComponent.cpp
+++ b/Engine/Plugins/Animation/ControlRig/Source/ControlRig/Private/ControlRigComponent.cpp
@@ -543,9 +543,8 @@ void UControlRigComponent::Update(float DeltaTime)
 				{
 					// Necessary for FStackAttributeContainer that uses a FAnimStackAllocator (TMemStackAllocator) which allocates from FMemStack.
 					// When allocating memory from FMemStack we need to explicitly use FMemMark to ensure items are freed when the scope exits. 
-					FMemMark Mark(FMemStack::Get());
-					TempAttributeContainer = MakeUnique<UE::Anim::FStackAttributeContainer>();
-					UControlRig::FAnimAttributeContainerPtrScope AttributeScope(CR, *TempAttributeContainer);
+					TempAttributeContainer.Empty();
+					UControlRig::FAnimAttributeContainerPtrScope AttributeScope(CR, TempAttributeContainer);
 				
 					CR->Evaluate_AnyThread();
 				}
@@ -1642,7 +1641,7 @@ void UControlRigComponent::TransferOutputs()
 
 						if (Proxy)
 						{
-							Proxy->StoredAttributes.CopyFrom(*TempAttributeContainer);
+							Proxy->StoredAttributes.CopyFrom(TempAttributeContainer);
 						}	
 					}
 				}
@@ -1750,7 +1749,7 @@ void UControlRigComponent::TransferOutputs()
 
 			if (Proxy)
 			{
-				Proxy->StoredAttributes.CopyFrom(*TempAttributeContainer);
+				Proxy->StoredAttributes.CopyFrom(TempAttributeContainer);
 			}			
 		}
 
diff --git a/Engine/Plugins/Animation/ControlRig/Source/ControlRig/Public/ControlRig.h b/Engine/Plugins/Animation/ControlRig/Source/ControlRig/Public/ControlRig.h
--- a/Engine/Plugins/Animation/ControlRig/Source/ControlRig/Public/ControlRig.h
+++ b/Engine/Plugins/Animation/ControlRig/Source/ControlRig/Public/ControlRig.h
@@ -510,7 +510,7 @@ public:
 	class FAnimAttributeContainerPtrScope
 	{
 	public:
-		CONTROLRIG_API FAnimAttributeContainerPtrScope(UControlRig* InControlRig, UE::Anim::FStackAttributeContainer& InExternalContainer);
+		CONTROLRIG_API FAnimAttributeContainerPtrScope(UControlRig* InControlRig, UE::Anim::FMeshAttributeContainer& InExternalContainer);
 		CONTROLRIG_API ~FAnimAttributeContainerPtrScope();
 
 		UControlRig* ControlRig;
@@ -520,7 +520,7 @@ private:
 	UPROPERTY(Transient)
 	FRigVMExtendedExecuteContext RigVMExtendedExecuteContext;
 
-	UE::Anim::FStackAttributeContainer* ExternalAnimAttributeContainer;
+	UE::Anim::FMeshAttributeContainer* ExternalAnimAttributeContainer;
 
 #if WITH_EDITOR
 	void SetEnableAnimAttributeTrace(bool bInEnable)
@@ -530,8 +530,8 @@ private:
 	
 	bool bEnableAnimAttributeTrace;
 	
-	UE::Anim::FHeapAttributeContainer InputAnimAttributeSnapshot;
-	UE::Anim::FHeapAttributeContainer OutputAnimAttributeSnapshot;
+	UE::Anim::FMeshAttributeContainer InputAnimAttributeSnapshot;
+	UE::Anim::FMeshAttributeContainer OutputAnimAttributeSnapshot;
 #endif
 	
 	/** The registry to access data source */
diff --git a/Engine/Plugins/Animation/ControlRig/Source/ControlRig/Public/ControlRigAnimInstance.h b/Engine/Plugins/Animation/ControlRig/Source/ControlRig/Public/ControlRigAnimInstance.h
--- a/Engine/Plugins/Animation/ControlRig/Source/ControlRig/Public/ControlRigAnimInstance.h
+++ b/Engine/Plugins/Animation/ControlRig/Source/ControlRig/Public/ControlRigAnimInstance.h
@@ -36,7 +36,7 @@ public:
 
 	TMap<FMeshPoseBoneIndex, FTransform> StoredTransforms;
 	TMap<FName, float> StoredCurves;
-	UE::Anim::FHeapAttributeContainer StoredAttributes;
+	UE::Anim::FMeshAttributeContainer StoredAttributes;
 };
 
 UCLASS(MinimalAPI, transient, NotBlueprintable)
diff --git a/Engine/Plugins/Animation/ControlRig/Source/ControlRig/Public/ControlRigComponent.h b/Engine/Plugins/Animation/ControlRig/Source/ControlRig/Public/ControlRigComponent.h
--- a/Engine/Plugins/Animation/ControlRig/Source/ControlRig/Public/ControlRigComponent.h
+++ b/Engine/Plugins/Animation/ControlRig/Source/ControlRig/Public/ControlRigComponent.h
@@ -671,7 +671,7 @@ private:
 	UPROPERTY(transient)
 	TObjectPtr<UControlRig> ControlRig;
 	
-	TUniquePtr<UE::Anim::FStackAttributeContainer> TempAttributeContainer;
+	UE::Anim::FMeshAttributeContainer TempAttributeContainer;
 
 	TMap<USkeletalMeshComponent*, FCachedSkeletalMeshComponentSettings> CachedSkeletalMeshComponentSettings;
 
diff --git a/Engine/Plugins/Animation/ControlRig/Source/ControlRig/Public/Units/RigUnitContext.h b/Engine/Plugins/Animation/ControlRig/Source/ControlRig/Public/Units/RigUnitContext.h
--- a/Engine/Plugins/Animation/ControlRig/Source/ControlRig/Public/Units/RigUnitContext.h
+++ b/Engine/Plugins/Animation/ControlRig/Source/ControlRig/Public/Units/RigUnitContext.h
@@ -77,7 +77,7 @@ struct FRigUnitContext
 	}
 
 	/** An external anim attribute container */
-	UE::Anim::FStackAttributeContainer* AnimAttributeContainer;
+	UE::Anim::FMeshAttributeContainer* AnimAttributeContainer;
 	
 	/** The registry to access data source */
 	const UAnimationDataSourceRegistry* DataSourceRegistry;
diff --git a/Engine/Plugins/Animation/ControlRig/Source/ControlRigEditor/Private/Editor/SControlRigAnimAttributeView.cpp b/Engine/Plugins/Animation/ControlRig/Source/ControlRigEditor/Private/Editor/SControlRigAnimAttributeView.cpp
--- a/Engine/Plugins/Animation/ControlRig/Source/ControlRigEditor/Private/Editor/SControlRigAnimAttributeView.cpp
+++ b/Engine/Plugins/Animation/ControlRig/Source/ControlRigEditor/Private/Editor/SControlRigAnimAttributeView.cpp
@@ -95,10 +95,10 @@ void SControlRigAnimAttributeView::HandleControlRigPostForwardSolve(
 		{
 			if (const USkeletalMeshComponent* SkeletalMeshComponent = Cast<USkeletalMeshComponent>(InControlRig->OuterSceneComponent.Get()))
 			{
-				const UE::Anim::FHeapAttributeContainer& InputSnapshot = InControlRig->InputAnimAttributeSnapshot;
-				const UE::Anim::FHeapAttributeContainer& OutputSnapshot = InControlRig->OutputAnimAttributeSnapshot;
+				const UE::Anim::FMeshAttributeContainer& InputSnapshot = InControlRig->InputAnimAttributeSnapshot;
+				const UE::Anim::FMeshAttributeContainer& OutputSnapshot = InControlRig->OutputAnimAttributeSnapshot;
 
-				const TArray<TTuple<FName, const UE::Anim::FHeapAttributeContainer&>> Snapshots =
+				const TArray<TTuple<FName, const UE::Anim::FMeshAttributeContainer&>> Snapshots =
 				{
 					{TEXT("Input") , InputSnapshot  },
 					{TEXT("Output"), OutputSnapshot }
diff --git a/Engine/Source/Editor/Persona/Private/SAnimAttributeView.cpp b/Engine/Source/Editor/Persona/Private/SAnimAttributeView.cpp
--- a/Engine/Source/Editor/Persona/Private/SAnimAttributeView.cpp
+++ b/Engine/Source/Editor/Persona/Private/SAnimAttributeView.cpp
@@ -582,8 +582,8 @@ void SAnimAttributeViewer::Tick(const FGeometry& AllottedGeometry, const double
 					{
 						if(AnimNodePoseWatch.PoseWatchPoseElement == PoseWatchPoseElement)
 						{
-							const UE::Anim::FHeapAttributeContainer& OutputSnapshot = AnimNodePoseWatch.GetAttributes();
-							const TArray<TTuple<FName, const UE::Anim::FHeapAttributeContainer&>> Snapshots =
+							const UE::Anim::FMeshAttributeContainer& OutputSnapshot = AnimNodePoseWatch.GetAttributes();
+							const TArray<TTuple<FName, const UE::Anim::FMeshAttributeContainer&>> Snapshots =
 							{
 								{*AnimNodePoseWatch.PoseWatch->GetLabel().ToString(), OutputSnapshot}
 							};
diff --git a/Engine/Source/Editor/Persona/Public/SAnimAttributeView.h b/Engine/Source/Editor/Persona/Public/SAnimAttributeView.h
--- a/Engine/Source/Editor/Persona/Public/SAnimAttributeView.h
+++ b/Engine/Source/Editor/Persona/Public/SAnimAttributeView.h
@@ -123,9 +123,8 @@ public:
 	
 	void Construct(const FArguments& InArgs);
 
-	template <typename ContainerType>
 	void DisplayNewAttributeContainerSnapshots(
-		const TArray<TTuple<FName, const ContainerType&>>& InSnapshots, 
+		const TArray<TTuple<FName, const UE::Anim::FMeshAttributeContainer&>>& InSnapshots, 
 		const USkeletalMeshComponent* InOwningComponent)
 	{
 		if (!ensure(InSnapshots.Num() != 0))
@@ -153,13 +152,13 @@ public:
 			
 			for (int32 SnapshotIndex = 0; SnapshotIndex < InSnapshots.Num(); SnapshotIndex++)
 			{
-				const TTuple<FName, const ContainerType&> Snapshot = InSnapshots[SnapshotIndex];
+				const TTuple<FName, const UE::Anim::FMeshAttributeContainer&> Snapshot = InSnapshots[SnapshotIndex];
 				CachedSnapshotNameIndexMap.FindOrAdd(Snapshot.Key) = SnapshotIndex;
 				
 				TTuple<FName, TArray<FAnimationAttributeIdentifier>>& CachedIdentifierList = CachedAttributeIdentifierLists[SnapshotIndex];
 				CachedIdentifierList.Key = Snapshot.Key;
 
-				const ContainerType& AttributeContainer = Snapshot.Value;
+				const UE::Anim::FMeshAttributeContainer& AttributeContainer = Snapshot.Value;
 				TArray<FAnimationAttributeIdentifier>& CachedIdentifiers = CachedIdentifierList.Value;
 
 				CachedIdentifiers.AddDefaulted(AttributeContainer.Num());
@@ -204,8 +203,8 @@ public:
 				const int32* SnapshotIndex = CachedSnapshotNameIndexMap.Find(ValueView.SnapshotName);
 				if (ensure(SnapshotIndex) && ensure(InSnapshots.IsValidIndex(*SnapshotIndex)))
 				{
-					const TTuple<FName, const ContainerType&> Snapshot = InSnapshots[*SnapshotIndex];
-					const ContainerType& AttributeContainer = Snapshot.Value;
+					const TTuple<FName, const UE::Anim::FMeshAttributeContainer&> Snapshot = InSnapshots[*SnapshotIndex];
+					const UE::Anim::FMeshAttributeContainer& AttributeContainer = Snapshot.Value;
 
 					ValueView.UpdateValue(AttributeContainer);
 				}
diff --git a/Engine/Source/Runtime/Engine/Classes/Animation/AttributesContainer.h b/Engine/Source/Runtime/Engine/Classes/Animation/AttributesContainer.h
--- a/Engine/Source/Runtime/Engine/Classes/Animation/AttributesContainer.h
+++ b/Engine/Source/Runtime/Engine/Classes/Animation/AttributesContainer.h
@@ -53,8 +53,9 @@ namespace UE
 			template<>
 			inline FMeshPoseBoneIndex MakeIndex(FCompactPoseBoneIndex BoneIndex, const FBoneContainer& BoneContainer)
 			{
-				
-				return BoneContainer.MakeMeshPoseIndex(BoneIndex);
+				// Guard against invalid compact-pose indexing
+				const bool bValidBoneIndex = BoneContainer.GetBoneIndicesArray().IsValidIndex(BoneIndex.GetInt());
+				return bValidBoneIndex ? BoneContainer.MakeMeshPoseIndex(BoneIndex) : FMeshPoseBoneIndex(INDEX_NONE);
 			}
 			
 			template<>
diff --git a/Engine/Source/Runtime/Engine/Classes/Engine/PoseWatchRenderData.h b/Engine/Source/Runtime/Engine/Classes/Engine/PoseWatchRenderData.h
--- a/Engine/Source/Runtime/Engine/Classes/Engine/PoseWatchRenderData.h
+++ b/Engine/Source/Runtime/Engine/Classes/Engine/PoseWatchRenderData.h
@@ -35,7 +35,8 @@ public:
 
 	ENGINE_API void SetCurves(const FBlendedCurve& InCurves);
 
-	ENGINE_API void SetAttributes(const UE::Anim::FStackAttributeContainer& InAttributes);
+	// _Moves_ Attributes out of InAttributes 
+	ENGINE_API void SetAttributes(UE::Anim::FMeshAttributeContainer& InAttributes);
 
 	template<class TAllocator>
 	bool SetPose(const TArray<FBoneIndexType>& InRequiredBones, const TArray<FTransform, TAllocator>& InBoneTransforms)
@@ -59,7 +60,7 @@ public:
 
 	ENGINE_API const FBlendedHeapCurve& GetCurves() const;
 	
-	ENGINE_API const UE::Anim::FHeapAttributeContainer& GetAttributes() const;
+	ENGINE_API const UE::Anim::FMeshAttributeContainer& GetAttributes() const;
 	
 	ENGINE_API const FTransform& GetWorldTransform() const;
 
@@ -76,7 +77,7 @@ private:
 	TArray<FBoneIndexType> RequiredBones;
 	TArray<FTransform> BoneTransforms;
 	FBlendedHeapCurve Curves;
-	UE::Anim::FHeapAttributeContainer HeapAttributes;
+	UE::Anim::FMeshAttributeContainer MeshAttributes;
 	
 	// Mirrored properties updated on CopyPoseWatchData
 	FLinearColor BoneColor;
diff --git a/Engine/Source/Runtime/Engine/Private/Animation/AnimInstanceProxy.cpp b/Engine/Source/Runtime/Engine/Private/Animation/AnimInstanceProxy.cpp
--- a/Engine/Source/Runtime/Engine/Private/Animation/AnimInstanceProxy.cpp
+++ b/Engine/Source/Runtime/Engine/Private/Animation/AnimInstanceProxy.cpp
@@ -3480,7 +3480,10 @@ void FAnimInstanceProxy::RegisterWatchedPose(const FPoseContext& PoseContext, in
 						PoseWatch.SetPose(TmpRequiredBones, BoneTransforms);
 						PoseWatch.SetCurves(PoseContext.Curve);
 						PoseWatch.SetWorldTransform(SkelMeshComponent->GetComponentTransform());
-						PoseWatch.SetAttributes(PoseContext.CustomAttributes);
+
+						UE::Anim::FMeshAttributeContainer MeshAttributeContainer;
+						MeshAttributeContainer.CopyFrom(PoseContext.CustomAttributes, PoseContext.Pose.GetBoneContainer());
+						PoseWatch.SetAttributes(MeshAttributeContainer);
 
 						TRACE_ANIM_POSE_WATCH(*this, PoseWatch.PoseWatchPoseElement, PoseWatch.NodeID, PoseWatch.GetBoneTransforms(), PoseWatch.GetCurves(), PoseWatch.GetRequiredBones(), PoseWatch.GetWorldTransform(), PoseWatch.PoseWatchPoseElement->GetIsVisible());
 						break;
@@ -3514,7 +3517,10 @@ void FAnimInstanceProxy::RegisterWatchedPose(const FComponentSpacePoseContext& P
 						PoseWatch.SetPose(TmpRequiredBones, BoneTransforms);
 						PoseWatch.SetCurves(PoseContext.Curve);
 						PoseWatch.SetWorldTransform(SkelMeshComponent->GetComponentTransform());
-						PoseWatch.SetAttributes(PoseContext.CustomAttributes);
+
+						UE::Anim::FMeshAttributeContainer MeshAttributeContainer;
+						MeshAttributeContainer.CopyFrom(PoseContext.CustomAttributes, TempPose.GetBoneContainer());
+						PoseWatch.SetAttributes(MeshAttributeContainer);
 
 						TRACE_ANIM_POSE_WATCH(*this, PoseWatch.PoseWatchPoseElement, PoseWatch.NodeID, PoseWatch.GetBoneTransforms(), PoseWatch.GetCurves(), PoseWatch.GetRequiredBones(), PoseWatch.GetWorldTransform(), PoseWatch.PoseWatchPoseElement->GetIsVisible());
 						break;
diff --git a/Engine/Source/Runtime/Engine/Private/Engine/PoseWatchRenderData.cpp b/Engine/Source/Runtime/Engine/Private/Engine/PoseWatchRenderData.cpp
--- a/Engine/Source/Runtime/Engine/Private/Engine/PoseWatchRenderData.cpp
+++ b/Engine/Source/Runtime/Engine/Private/Engine/PoseWatchRenderData.cpp
@@ -24,9 +24,9 @@ void FAnimNodePoseWatch::SetCurves(const FBlendedCurve& InCurves)
 	Curves.CopyFrom(InCurves);
 }
 
-void FAnimNodePoseWatch::SetAttributes(const UE::Anim::FStackAttributeContainer& InAttributes)
+void FAnimNodePoseWatch::SetAttributes(UE::Anim::FMeshAttributeContainer& InAttributes)
 {
-	HeapAttributes.CopyFrom(InAttributes);
+	MeshAttributes.MoveFrom(InAttributes);
 }
 
 void FAnimNodePoseWatch::SetWorldTransform(const FTransform& InWorldTransform)
@@ -83,9 +83,9 @@ const FBlendedHeapCurve& FAnimNodePoseWatch::GetCurves() const
 	return Curves;
 }
 
-const UE::Anim::FHeapAttributeContainer& FAnimNodePoseWatch::GetAttributes() const
+const UE::Anim::FMeshAttributeContainer& FAnimNodePoseWatch::GetAttributes() const
 {
-	return HeapAttributes;
+	return MeshAttributes;
 }
 
 const FTransform& FAnimNodePoseWatch::GetWorldTransform() const
