Fix for many SAssetShortcuts unnecessarily refreshing on asset opening resulting in slow ABP navigating

diff --git a/Engine/Source/Editor/Persona/Private/SAssetFamilyShortcutBar.cpp b/Engine/Source/Editor/Persona/Private/SAssetFamilyShortcutBar.cpp
--- a/Engine/Source/Editor/Persona/Private/SAssetFamilyShortcutBar.cpp
+++ b/Engine/Source/Editor/Persona/Private/SAssetFamilyShortcutBar.cpp
@@ -405,9 +405,14 @@ public:
 
 	void HandleAssetOpened(UObject* InAsset)
 	{
-		if (AssetFamily.IsValid() && AssetFamily->IsAssetCompatible(FAssetData(InAsset)))
+		const FAssetData InAssetData(InAsset);
+
+		if (AssetFamily.IsValid() && AssetFamily->IsAssetCompatible(InAssetData))
 		{
-			RefreshAsset();
+			if (AssetFamily->FindAssetOfType(InAsset->GetClass()) == InAssetData)
+			{
+				RefreshAsset();
+			}
 		}
 	}
 
