PR #12697: Fix ApplyLODChanges about InterchangeAssetImportData

#rnx

diff --git a/Engine/Source/Editor/Persona/Private/PersonaMeshDetails.cpp b/Engine/Source/Editor/Persona/Private/PersonaMeshDetails.cpp
--- a/Engine/Source/Editor/Persona/Private/PersonaMeshDetails.cpp
+++ b/Engine/Source/Editor/Persona/Private/PersonaMeshDetails.cpp
@@ -88,6 +88,8 @@
 #include "Interfaces/ITargetPlatform.h"
 #include "Interfaces/ITargetPlatformManagerModule.h"
 #include "Misc/CoreMisc.h"
+#include "InterchangeManager.h"
+#include "InterchangeAssetImportData.h"
 #include "SPerQualityLevelPropertiesWidget.h"
 #include "UObject/SavePackage.h"
 
@@ -4189,6 +4191,17 @@ FReply FPersonaMeshDetails::ApplyLODChanges(int32 LODIndex)
 			{
 				//Update the Asset Import Data
 				UFbxSkeletalMeshImportData* SKImportData = Cast<UFbxSkeletalMeshImportData>(SkelMesh->GetAssetImportData());
+				if (!SKImportData)
+				{
+					if (UInterchangeAssetImportData* InterchangeAssetImportData = Cast<UInterchangeAssetImportData>(SkelMesh->GetAssetImportData()))
+					{
+						if (UInterchangeManager::GetInterchangeManager().ConvertImportData(InterchangeAssetImportData, UFbxSkeletalMeshImportData::StaticClass(), reinterpret_cast<UObject**>(&SKImportData)))
+						{
+							SKImportData->SourceData = InterchangeAssetImportData->SourceData;
+						}
+					}
+				}
+
 				if (SKImportData)
 				{
 					check(bSrcBuildDataAvailable);
@@ -4209,6 +4222,17 @@ FReply FPersonaMeshDetails::ApplyLODChanges(int32 LODIndex)
 						SKImportData->MorphThresholdPosition = LODInfo->BuildSettings.MorphThresholdPosition;
 					}
 				}
+
+				UInterchangeAssetImportData* NewInterchangeAssetImportData = nullptr;
+				UInterchangeManager::GetInterchangeManager().ConvertImportData(SKImportData, UInterchangeAssetImportData::StaticClass(), reinterpret_cast<UObject**>(&NewInterchangeAssetImportData));
+				if (IsValid(NewInterchangeAssetImportData))
+				{
+					if (SKImportData)
+					{
+						NewInterchangeAssetImportData->SourceData = SKImportData->SourceData;
+					}
+					SkelMesh->SetAssetImportData(NewInterchangeAssetImportData);
+				}
 			}
 
 			RestoreNonReducedLOD(LODIndex);
