Updated headers using UnrealCodeFixup to make sure dllstorage is on methods/staticvars instead of types. Part 5/n

diff --git a/Engine/Source/Editor/MovieSceneTools/Public/TrackEditors/CommonAnimationTrackEditor.h b/Engine/Source/Editor/MovieSceneTools/Public/TrackEditors/CommonAnimationTrackEditor.h
--- a/Engine/Source/Editor/MovieSceneTools/Public/TrackEditors/CommonAnimationTrackEditor.h
+++ b/Engine/Source/Editor/MovieSceneTools/Public/TrackEditors/CommonAnimationTrackEditor.h
@@ -13,6 +13,8 @@
 #include "MovieSceneTrackEditor.h"
 #include "SequencerCoreFwd.h"
 
+#define UE_API MOVIESCENETOOLS_API
+
 struct FAssetData;
 struct FMovieSceneTimeWarpChannel;
 struct FMovieSceneSequenceTransform;
@@ -34,7 +36,7 @@ class ITrackExtension;
 /**
  * Tools for animation tracks
  */
-class MOVIESCENETOOLS_API FCommonAnimationTrackEditor : public FMovieSceneTrackEditor, public FGCObject
+class FCommonAnimationTrackEditor : public FMovieSceneTrackEditor, public FGCObject
 {
 public:
 	/**
@@ -42,13 +44,13 @@ public:
 	 *
 	 * @param InSequencer The sequencer instance to be used by this tool
 	 */
-	FCommonAnimationTrackEditor( TSharedRef<ISequencer> InSequencer );
+	UE_API FCommonAnimationTrackEditor( TSharedRef<ISequencer> InSequencer );
 
 	/** Virtual destructor. */
 	virtual ~FCommonAnimationTrackEditor() { }
 
 	//~ FGCObject
-	virtual void AddReferencedObjects(FReferenceCollector& Collector) override;
+	UE_API virtual void AddReferencedObjects(FReferenceCollector& Collector) override;
 	virtual FString GetReferencerName() const override
 	{
 		return TEXT("FCommonAnimationTrackEditor");
@@ -57,72 +59,72 @@ public:
 	/**
 	* Keeps track of how many skeletal animation track editors we have*
 	*/
-	static int32 NumberActive;
+	static UE_API int32 NumberActive;
 
-	static USkeletalMeshComponent* AcquireSkeletalMeshFromObjectGuid(const FGuid& Guid, TSharedPtr<ISequencer> SequencerPtr);
-	static USkeleton* AcquireSkeletonFromObjectGuid(const FGuid& Guid, TSharedPtr<ISequencer> SequencerPtr);
+	static UE_API USkeletalMeshComponent* AcquireSkeletalMeshFromObjectGuid(const FGuid& Guid, TSharedPtr<ISequencer> SequencerPtr);
+	static UE_API USkeleton* AcquireSkeletonFromObjectGuid(const FGuid& Guid, TSharedPtr<ISequencer> SequencerPtr);
 
 public:
 
 	// ISequencerTrackEditor interface
 
-	virtual FText GetDisplayName() const override;
-	virtual void BuildObjectBindingContextMenu(FMenuBuilder& MenuBuilder, const TArray<FGuid>& ObjectBindings, const UClass* ObjectClass) override;
-	virtual void BuildObjectBindingTrackMenu(FMenuBuilder& MenuBuilder, const TArray<FGuid>& ObjectBindings, const UClass* ObjectClass) override;
-	virtual bool HandleAssetAdded(UObject* Asset, const FGuid& TargetObjectGuid) override;
-	virtual TSharedRef<ISequencerSection> MakeSectionInterface( UMovieSceneSection& SectionObject, UMovieSceneTrack& Track, FGuid ObjectBinding ) override;
-	virtual TSharedPtr<SWidget> BuildOutlinerEditWidget(const FGuid& ObjectBinding, UMovieSceneTrack* Track, const FBuildEditWidgetParams& Params) override;
-	virtual bool OnAllowDrop(const FDragDropEvent& DragDropEvent, FSequencerDragDropParams& DragDropParams) override;
-	virtual FReply OnDrop(const FDragDropEvent& DragDropEvent, const FSequencerDragDropParams& DragDropParams) override;
-	virtual void OnInitialize() override;
-	virtual void OnRelease() override;
+	UE_API virtual FText GetDisplayName() const override;
+	UE_API virtual void BuildObjectBindingContextMenu(FMenuBuilder& MenuBuilder, const TArray<FGuid>& ObjectBindings, const UClass* ObjectClass) override;
+	UE_API virtual void BuildObjectBindingTrackMenu(FMenuBuilder& MenuBuilder, const TArray<FGuid>& ObjectBindings, const UClass* ObjectClass) override;
+	UE_API virtual bool HandleAssetAdded(UObject* Asset, const FGuid& TargetObjectGuid) override;
+	UE_API virtual TSharedRef<ISequencerSection> MakeSectionInterface( UMovieSceneSection& SectionObject, UMovieSceneTrack& Track, FGuid ObjectBinding ) override;
+	UE_API virtual TSharedPtr<SWidget> BuildOutlinerEditWidget(const FGuid& ObjectBinding, UMovieSceneTrack* Track, const FBuildEditWidgetParams& Params) override;
+	UE_API virtual bool OnAllowDrop(const FDragDropEvent& DragDropEvent, FSequencerDragDropParams& DragDropParams) override;
+	UE_API virtual FReply OnDrop(const FDragDropEvent& DragDropEvent, const FSequencerDragDropParams& DragDropParams) override;
+	UE_API virtual void OnInitialize() override;
+	UE_API virtual void OnRelease() override;
 
 protected:
 
 	virtual TSubclassOf<UMovieSceneCommonAnimationTrack> GetTrackClass() const = 0;
 
 	/** Animation sub menu */
-	TSharedRef<SWidget> BuildAddAnimationSubMenu(FGuid ObjectBinding, USkeleton* Skeleton, UE::Sequencer::TWeakViewModelPtr<UE::Sequencer::ITrackExtension> TrackModel);
-	TSharedRef<SWidget> BuildAnimationSubMenu(FGuid ObjectBinding, USkeleton* Skeleton, UMovieSceneTrack* Track);
-	void AddAnimationSubMenu(FMenuBuilder& MenuBuilder, TArray<FGuid> ObjectBindings, USkeleton* Skeleton, UMovieSceneTrack* Track);
+	UE_API TSharedRef<SWidget> BuildAddAnimationSubMenu(FGuid ObjectBinding, USkeleton* Skeleton, UE::Sequencer::TWeakViewModelPtr<UE::Sequencer::ITrackExtension> TrackModel);
+	UE_API TSharedRef<SWidget> BuildAnimationSubMenu(FGuid ObjectBinding, USkeleton* Skeleton, UMovieSceneTrack* Track);
+	UE_API void AddAnimationSubMenu(FMenuBuilder& MenuBuilder, TArray<FGuid> ObjectBindings, USkeleton* Skeleton, UMovieSceneTrack* Track);
 
 	/** Filter only compatible skeletons */
-	bool FilterAnimSequences(const FAssetData& AssetData, USkeleton* Skeleton);
+	UE_API bool FilterAnimSequences(const FAssetData& AssetData, USkeleton* Skeleton);
 
 	/** Animation sub menu filter function */
-	bool ShouldFilterAsset(const FAssetData& AssetData);
+	UE_API bool ShouldFilterAsset(const FAssetData& AssetData);
 
 	/** Animation asset selected */
-	void OnAnimationAssetSelected(const FAssetData& AssetData, TArray<FGuid> ObjectBindings, UMovieSceneTrack* Track);
+	UE_API void OnAnimationAssetSelected(const FAssetData& AssetData, TArray<FGuid> ObjectBindings, UMovieSceneTrack* Track);
 
 	/** Animation asset enter pressed */
-	void OnAnimationAssetEnterPressed(const TArray<FAssetData>& AssetData, TArray<FGuid> ObjectBindings, UMovieSceneTrack* Track);
+	UE_API void OnAnimationAssetEnterPressed(const TArray<FAssetData>& AssetData, TArray<FGuid> ObjectBindings, UMovieSceneTrack* Track);
 
 	/** Delegate for AnimatablePropertyChanged in AddKey */
-	FKeyPropertyResult AddKeyInternal(FFrameNumber KeyTime, UObject* Object, UAnimSequenceBase* AnimSequence, UMovieSceneTrack* Track, int32 RowIndex);
+	UE_API FKeyPropertyResult AddKeyInternal(FFrameNumber KeyTime, UObject* Object, UAnimSequenceBase* AnimSequence, UMovieSceneTrack* Track, int32 RowIndex);
 	
 	/** Construct the binding menu*/
-	void ConstructObjectBindingTrackMenu(FMenuBuilder& MenuBuilder, TArray<FGuid> ObjectBindings);
+	UE_API void ConstructObjectBindingTrackMenu(FMenuBuilder& MenuBuilder, TArray<FGuid> ObjectBindings);
 
 	/** Callback to Create the Animation Asset, pop open the dialog */
-	void HandleCreateAnimationSequence(USkeletalMeshComponent* SkelMeshComp, USkeleton* Skeleton, FGuid Binding, bool bCeateSoftLink);
+	UE_API void HandleCreateAnimationSequence(USkeletalMeshComponent* SkelMeshComp, USkeleton* Skeleton, FGuid Binding, bool bCeateSoftLink);
 
 	/** Callback to Creae the Animation Asset after getting the name*/
-	bool CreateAnimationSequence(const TArray<UObject*> NewAssets,USkeletalMeshComponent* SkelMeshComp, FGuid Binding, bool bCreateSoftLink);
+	UE_API bool CreateAnimationSequence(const TArray<UObject*> NewAssets,USkeletalMeshComponent* SkelMeshComp, FGuid Binding, bool bCreateSoftLink);
 
 	/** Open the linked Anim Sequence*/
-	void OpenLinkedAnimSequence(FGuid Binding);
+	UE_API void OpenLinkedAnimSequence(FGuid Binding);
 
 	/** Can Open the linked Anim Sequence*/
-	bool CanOpenLinkedAnimSequence(FGuid Binding);
+	UE_API bool CanOpenLinkedAnimSequence(FGuid Binding);
 
 	friend class FMovieSceneSkeletalAnimationParamsDetailCustomization;
 
 protected:
 	/* Was part of the the section but should be at the track level since it takes the final blended result at the current time, not the section instance value*/
-	bool CreatePoseAsset(const TArray<UObject*> NewAssets, FGuid InObjectBinding);
-	void HandleCreatePoseAsset(FGuid InObjectBinding);
-	bool CanCreatePoseAsset(FGuid InObjectBinding) const;
+	UE_API bool CreatePoseAsset(const TArray<UObject*> NewAssets, FGuid InObjectBinding);
+	UE_API void HandleCreatePoseAsset(FGuid InObjectBinding);
+	UE_API bool CanCreatePoseAsset(FGuid InObjectBinding) const;
 
 protected:
 	/* For Anim Sequence UI Option with be gc'd*/
@@ -132,49 +134,49 @@ protected:
 protected:
 	/* Delegate to handle sequencer changes for auto baking of anim sequences*/
 	FDelegateHandle SequencerSavedHandle;
-	void OnSequencerSaved(ISequencer& InSequence);
+	UE_API void OnSequencerSaved(ISequencer& InSequence);
 	FDelegateHandle SequencerChangedHandle;
-	void OnSequencerDataChanged(EMovieSceneDataChangeType DataChangeType);
+	UE_API void OnSequencerDataChanged(EMovieSceneDataChangeType DataChangeType);
 };
 
 
 /** Class for animation sections */
-class MOVIESCENETOOLS_API FCommonAnimationSection
+class FCommonAnimationSection
 	: public ISequencerSection
 	, public TSharedFromThis<FCommonAnimationSection>
 {
 public:
 
 	/** Constructor. */
-	FCommonAnimationSection( UMovieSceneSection& InSection, TWeakPtr<ISequencer> InSequencer);
+	UE_API FCommonAnimationSection( UMovieSceneSection& InSection, TWeakPtr<ISequencer> InSequencer);
 
 	/** Virtual destructor. */
-	virtual ~FCommonAnimationSection();
+	UE_API virtual ~FCommonAnimationSection();
 
 public:
 
 	// ISequencerSection interface
 
-	virtual UMovieSceneSection* GetSectionObject() override;
-	virtual FText GetSectionTitle() const override;
-	virtual FText GetSectionToolTip() const override;
-	virtual TOptional<FFrameTime> GetSectionTime(FSequencerSectionPainter& InPainter) const override;
-	virtual float GetSectionHeight(const UE::Sequencer::FViewDensityInfo& ViewDensity) const override;
-	virtual FMargin GetContentPadding() const override;
-	virtual int32 OnPaintSection( FSequencerSectionPainter& Painter ) const override;
-	virtual void BeginResizeSection() override;
-	virtual void ResizeSection(ESequencerSectionResizeMode ResizeMode, FFrameNumber ResizeTime) override;
-	virtual void BeginSlipSection() override;
-	virtual void SlipSection(FFrameNumber SlipTime) override;
-	virtual void CustomizePropertiesDetailsView(TSharedRef<IDetailsView> DetailsView, const FSequencerSectionPropertyDetailsViewCustomizationParams& InParams) const override;
-	virtual void BeginDilateSection() override;
-	virtual void DilateSection(const TRange<FFrameNumber>& NewRange, float DilationFactor) override;
-	virtual bool RequestDeleteKeyArea(const TArray<FName>& KeyAreaNamePath) override;
-	virtual void BuildSectionContextMenu(FMenuBuilder& MenuBuilder, const FGuid& InObjectBinding) override;
+	UE_API virtual UMovieSceneSection* GetSectionObject() override;
+	UE_API virtual FText GetSectionTitle() const override;
+	UE_API virtual FText GetSectionToolTip() const override;
+	UE_API virtual TOptional<FFrameTime> GetSectionTime(FSequencerSectionPainter& InPainter) const override;
+	UE_API virtual float GetSectionHeight(const UE::Sequencer::FViewDensityInfo& ViewDensity) const override;
+	UE_API virtual FMargin GetContentPadding() const override;
+	UE_API virtual int32 OnPaintSection( FSequencerSectionPainter& Painter ) const override;
+	UE_API virtual void BeginResizeSection() override;
+	UE_API virtual void ResizeSection(ESequencerSectionResizeMode ResizeMode, FFrameNumber ResizeTime) override;
+	UE_API virtual void BeginSlipSection() override;
+	UE_API virtual void SlipSection(FFrameNumber SlipTime) override;
+	UE_API virtual void CustomizePropertiesDetailsView(TSharedRef<IDetailsView> DetailsView, const FSequencerSectionPropertyDetailsViewCustomizationParams& InParams) const override;
+	UE_API virtual void BeginDilateSection() override;
+	UE_API virtual void DilateSection(const TRange<FFrameNumber>& NewRange, float DilationFactor) override;
+	UE_API virtual bool RequestDeleteKeyArea(const TArray<FName>& KeyAreaNamePath) override;
+	UE_API virtual void BuildSectionContextMenu(FMenuBuilder& MenuBuilder, const FGuid& InObjectBinding) override;
 
 
 protected:
-	void FindBestBlendSection(FGuid InObjectBinding);
+	UE_API void FindBestBlendSection(FGuid InObjectBinding);
 protected:
 
 	/** The section we are visualizing */
@@ -188,4 +190,6 @@ protected:
 	double PreDilatePlayRate;
 };
 
-} // namespace UE::Sequencer
\ No newline at end of file
+} // namespace UE::Sequencer
+
+#undef UE_API
diff --git a/Engine/Source/Editor/MovieSceneTools/Public/TrackEditors/CustomPrimitiveDataTrackEditor.h b/Engine/Source/Editor/MovieSceneTools/Public/TrackEditors/CustomPrimitiveDataTrackEditor.h
--- a/Engine/Source/Editor/MovieSceneTools/Public/TrackEditors/CustomPrimitiveDataTrackEditor.h
+++ b/Engine/Source/Editor/MovieSceneTools/Public/TrackEditors/CustomPrimitiveDataTrackEditor.h
@@ -13,6 +13,8 @@
 #include "Materials/MaterialParameters.h"
 #include "MovieSceneTrackEditor.h"
 
+#define UE_API MOVIESCENETOOLS_API
+
 class UMovieSceneCustomPrimitiveDataTrack;
 
 using FGetStartIndexDelegate = TFunction<uint8()>;
@@ -20,29 +22,29 @@ using FGetStartIndexDelegate = TFunction<uint8()>;
 /**
  * Track editor for custom primitive data tracks
  */
-class MOVIESCENETOOLS_API FCustomPrimitiveDataTrackEditor
+class FCustomPrimitiveDataTrackEditor
 	: public FMovieSceneTrackEditor
 {
 public:
 
 	/** Constructor. */
-	FCustomPrimitiveDataTrackEditor(TSharedRef<ISequencer> InSequencer);
+	UE_API FCustomPrimitiveDataTrackEditor(TSharedRef<ISequencer> InSequencer);
 
 	/** Virtual destructor. */
 	virtual ~FCustomPrimitiveDataTrackEditor() { }
 
-	static TSharedRef<ISequencerTrackEditor> CreateTrackEditor(TSharedRef<ISequencer> OwningSequencer);
+	static UE_API TSharedRef<ISequencerTrackEditor> CreateTrackEditor(TSharedRef<ISequencer> OwningSequencer);
 
 public:
 
 	// ISequencerTrackEditor interface
 
-	virtual FText GetDisplayName() const override;
-	virtual TSharedPtr<SWidget> BuildOutlinerEditWidget(const FGuid& ObjectBinding, UMovieSceneTrack* Track, const FBuildEditWidgetParams& Params) override;
-	virtual TSharedRef<ISequencerSection> MakeSectionInterface(UMovieSceneSection& SectionObject, UMovieSceneTrack& Track, FGuid ObjectBinding) override;
-	virtual bool SupportsType(TSubclassOf<UMovieSceneTrack> Type) const override;
-	virtual bool GetDefaultExpansionState(UMovieSceneTrack* InTrack) const override; 
-	virtual void ExtendObjectBindingTrackMenu(TSharedRef<FExtender> Extender, const TArray<FGuid>& ObjectBindings, const UClass* ObjectClass) override;
+	UE_API virtual FText GetDisplayName() const override;
+	UE_API virtual TSharedPtr<SWidget> BuildOutlinerEditWidget(const FGuid& ObjectBinding, UMovieSceneTrack* Track, const FBuildEditWidgetParams& Params) override;
+	UE_API virtual TSharedRef<ISequencerSection> MakeSectionInterface(UMovieSceneSection& SectionObject, UMovieSceneTrack& Track, FGuid ObjectBinding) override;
+	UE_API virtual bool SupportsType(TSubclassOf<UMovieSceneTrack> Type) const override;
+	UE_API virtual bool GetDefaultExpansionState(UMovieSceneTrack* InTrack) const override; 
+	UE_API virtual void ExtendObjectBindingTrackMenu(TSharedRef<FExtender> Extender, const TArray<FGuid>& ObjectBindings, const UClass* ObjectClass) override;
 
 protected:
 
@@ -50,50 +52,52 @@ private:
 
 
 	/** Provides the contents of the outliner edit widget */
-	TSharedRef<SWidget> OnGetAddMenuContent(FGuid ObjectBinding, UMovieSceneCustomPrimitiveDataTrack* CPDTrack, int32 TrackInsertRowIndex);
+	UE_API TSharedRef<SWidget> OnGetAddMenuContent(FGuid ObjectBinding, UMovieSceneCustomPrimitiveDataTrack* CPDTrack, int32 TrackInsertRowIndex);
 
 	/** Provides the contents of the add parameter menu. */
-	TSharedRef<SWidget> OnGetAddParameterMenuContent(FGuid ObjectBinding, UMovieSceneCustomPrimitiveDataTrack* CPDTrack);
+	UE_API TSharedRef<SWidget> OnGetAddParameterMenuContent(FGuid ObjectBinding, UMovieSceneCustomPrimitiveDataTrack* CPDTrack);
 
 	/** Provides the contents of the add parameter menu. */
-	void OnBuildAddParameterMenu(FMenuBuilder& MenuBuilder, FGuid ObjectBinding, UMovieSceneCustomPrimitiveDataTrack* CPDTrack);
+	UE_API void OnBuildAddParameterMenu(FMenuBuilder& MenuBuilder, FGuid ObjectBinding, UMovieSceneCustomPrimitiveDataTrack* CPDTrack);
 
 	/** Adds a scalar parameter and initial key to a custom primitive data track.
 	 * @param ObjectBinding The object binding which owns the custom primitive data track.
 	 * @param CPDTrack The track in which to look for sections to add the parameter to.
 	 * @param PrimitiveDataIndex The index for the primitive data to animate
 	 */
-	void AddScalarParameter(UPrimitiveComponent* Component, UMovieSceneCustomPrimitiveDataTrack* CPDTrack, FGetStartIndexDelegate GetStartIndexDelegate);
+	UE_API void AddScalarParameter(UPrimitiveComponent* Component, UMovieSceneCustomPrimitiveDataTrack* CPDTrack, FGetStartIndexDelegate GetStartIndexDelegate);
 
 	/** Adds a Vector2D parameter and initial key to a custom primitive data track.
 	 * @param ObjectBinding The object binding which owns the custom primitive data track.
 	 * @param CPDTrack The track in which to look for sections to add the parameter to.
 	 * @param PrimitiveDataStartIndex The start index in primitive data where we are animating a Vector2D
 	 */
-	void AddVector2DParameter(UPrimitiveComponent* Component, UMovieSceneCustomPrimitiveDataTrack* CPDTrack, FGetStartIndexDelegate GetStartIndexDelegate);
+	UE_API void AddVector2DParameter(UPrimitiveComponent* Component, UMovieSceneCustomPrimitiveDataTrack* CPDTrack, FGetStartIndexDelegate GetStartIndexDelegate);
 
 	/** Adds a Vector parameter and initial key to a custom primitive data track.
 	 * @param ObjectBinding The object binding which owns the custom primitive data track.
 	 * @param CPDTrack The track in which to look for sections to add the parameter to.
 	 * @param PrimitiveDataStartIndex The start index in primitive data where we are animating a Vector
 	 */
-	void AddVectorParameter(UPrimitiveComponent* Component, UMovieSceneCustomPrimitiveDataTrack* CPDTrack, FGetStartIndexDelegate GetStartIndexDelegate);
+	UE_API void AddVectorParameter(UPrimitiveComponent* Component, UMovieSceneCustomPrimitiveDataTrack* CPDTrack, FGetStartIndexDelegate GetStartIndexDelegate);
 
 	/** Adds a color parameter and initial key to a custom primitive data track.
 	 * @param ObjectBinding The object binding which owns the custom primitive data track.
 	 * @param CPDTrack The track in which to look for sections to add the parameter to.
 	 * @param PrimitiveDataStartIndex The start index in primitive data where we are animating a Color
 	 */
-	void AddColorParameter(UPrimitiveComponent* Component, UMovieSceneCustomPrimitiveDataTrack* CPDTrack, FGetStartIndexDelegate GetStartIndexDelegate);
+	UE_API void AddColorParameter(UPrimitiveComponent* Component, UMovieSceneCustomPrimitiveDataTrack* CPDTrack, FGetStartIndexDelegate GetStartIndexDelegate);
 
 	/** Returns whether we can add a parameter starting at this index of size ParameterSize. 
 	 * Will return false if another parameter added overlaps that index.
 	 */
-	bool CanAddParameter(UMovieSceneCustomPrimitiveDataTrack* CPDTrack, FGetStartIndexDelegate GetStartIndexDelegate, int ParameterSize);
+	UE_API bool CanAddParameter(UMovieSceneCustomPrimitiveDataTrack* CPDTrack, FGetStartIndexDelegate GetStartIndexDelegate, int ParameterSize);
 
-	void ConstructObjectBindingTrackMenu(FMenuBuilder& MenuBuilder, TArray<FGuid> ObjectBindings);
+	UE_API void ConstructObjectBindingTrackMenu(FMenuBuilder& MenuBuilder, TArray<FGuid> ObjectBindings);
 
-	void HandleAddCustomPrimitiveDataTrackExecute(UPrimitiveComponent* Component);
+	UE_API void HandleAddCustomPrimitiveDataTrackExecute(UPrimitiveComponent* Component);
 
 	uint8 StartIndex = 0;
-};
\ No newline at end of file
+};
+
+#undef UE_API
diff --git a/Engine/Source/Editor/MovieSceneTools/Public/TrackEditors/MaterialTrackEditor.h b/Engine/Source/Editor/MovieSceneTools/Public/TrackEditors/MaterialTrackEditor.h
--- a/Engine/Source/Editor/MovieSceneTools/Public/TrackEditors/MaterialTrackEditor.h
+++ b/Engine/Source/Editor/MovieSceneTools/Public/TrackEditors/MaterialTrackEditor.h
@@ -13,6 +13,8 @@
 #include "Materials/MaterialParameters.h"
 #include "MovieSceneTrackEditor.h"
 
+#define UE_API MOVIESCENETOOLS_API
+
 class UMaterial;
 class UMaterialInterface;
 class UMovieSceneMaterialTrack;
@@ -22,13 +24,13 @@ struct FComponentMaterialInfo;
 /**
  * Track editor for material parameters.
  */
-class MOVIESCENETOOLS_API FMaterialTrackEditor
+class FMaterialTrackEditor
 	: public FMovieSceneTrackEditor
 {
 public:
 
 	/** Constructor. */
-	FMaterialTrackEditor( TSharedRef<ISequencer> InSequencer );
+	UE_API FMaterialTrackEditor( TSharedRef<ISequencer> InSequencer );
 
 	/** Virtual destructor. */
 	virtual ~FMaterialTrackEditor() { }
@@ -37,9 +39,9 @@ public:
 
 	// ISequencerTrackEditor interface
 
-	virtual FText GetDisplayName() const override;
-	virtual TSharedPtr<SWidget> BuildOutlinerEditWidget( const FGuid& ObjectBinding, UMovieSceneTrack* Track, const FBuildEditWidgetParams& Params ) override;
-	virtual TSharedRef<ISequencerSection> MakeSectionInterface( UMovieSceneSection& SectionObject, UMovieSceneTrack& Track, FGuid ObjectBinding ) override;
+	UE_API virtual FText GetDisplayName() const override;
+	UE_API virtual TSharedPtr<SWidget> BuildOutlinerEditWidget( const FGuid& ObjectBinding, UMovieSceneTrack* Track, const FBuildEditWidgetParams& Params ) override;
+	UE_API virtual TSharedRef<ISequencerSection> MakeSectionInterface( UMovieSceneSection& SectionObject, UMovieSceneTrack& Track, FGuid ObjectBinding ) override;
 
 protected:
 
@@ -49,30 +51,30 @@ protected:
 private:
 
 	/** Provides the contents of the outliner edit widget */
-	TSharedRef<SWidget> OnGetAddMenuContent( FGuid ObjectBinding, UMovieSceneMaterialTrack* MaterialTrack, int32 TrackInsertRowIndex );
+	UE_API TSharedRef<SWidget> OnGetAddMenuContent( FGuid ObjectBinding, UMovieSceneMaterialTrack* MaterialTrack, int32 TrackInsertRowIndex );
 
 	/** Provides the contents of the add parameter menu. */
-	TSharedRef<SWidget> OnGetAddParameterMenuContent( FGuid ObjectBinding, UMovieSceneMaterialTrack* MaterialTrack );
+	UE_API TSharedRef<SWidget> OnGetAddParameterMenuContent( FGuid ObjectBinding, UMovieSceneMaterialTrack* MaterialTrack );
 
 	/** Provides the contents of the add parameter menu. */
-	void OnBuildAddParameterMenu( FMenuBuilder& MenuBuilder, FGuid ObjectBinding, UMovieSceneMaterialTrack* MaterialTrack );
+	UE_API void OnBuildAddParameterMenu( FMenuBuilder& MenuBuilder, FGuid ObjectBinding, UMovieSceneMaterialTrack* MaterialTrack );
 
 	/** Gets a material for a specific object binding and track */
-	UMaterial* GetMaterialForTrack( FGuid ObjectBinding, UMovieSceneMaterialTrack* MaterialTrack );
+	UE_API UMaterial* GetMaterialForTrack( FGuid ObjectBinding, UMovieSceneMaterialTrack* MaterialTrack );
 
 	/** Adds a scalar parameter and initial key to a material track.
 	 * @param ObjectBinding The object binding which owns the material track.
 	 * @param MaterialTrack The track in which to look for sections to add the parameter to.
 	 * @param ParameterName The name of the parameter to add an initial key for.
 	 */
-	void AddScalarParameter( FGuid ObjectBinding, UMovieSceneMaterialTrack* MaterialTrack, FMaterialParameterInfo ParameterInfo, FString InLayerName, FString InAssetName );
+	UE_API void AddScalarParameter( FGuid ObjectBinding, UMovieSceneMaterialTrack* MaterialTrack, FMaterialParameterInfo ParameterInfo, FString InLayerName, FString InAssetName );
 
 	/** Adds a color parameter and initial key to a material track.
 	* @param ObjectBinding The object binding which owns the material track.
 	 * @param MaterialTrack The track in which to look for sections to add the parameter to.
 	* @param ParameterName The name of the parameter to add an initial key for.
 	*/
-	void AddColorParameter( FGuid ObjectBinding, UMovieSceneMaterialTrack* MaterialTrack, FMaterialParameterInfo ParameterInfo, FString InLayerName, FString InAssetName);
+	UE_API void AddColorParameter( FGuid ObjectBinding, UMovieSceneMaterialTrack* MaterialTrack, FMaterialParameterInfo ParameterInfo, FString InLayerName, FString InAssetName);
 };
 
 
@@ -111,3 +113,5 @@ private:
 	void FillRebindMaterialTrackMenu(FMenuBuilder& MenuBuilder, class UMovieSceneComponentMaterialTrack* MaterialTrack, class UPrimitiveComponent* Component, FGuid ObjectBinding);
 
 };
+
+#undef UE_API
diff --git a/Engine/Source/Editor/MovieSceneTools/Public/TrackEditors/PropertyTrackEditors/BoolPropertyTrackEditor.h b/Engine/Source/Editor/MovieSceneTools/Public/TrackEditors/PropertyTrackEditors/BoolPropertyTrackEditor.h
--- a/Engine/Source/Editor/MovieSceneTools/Public/TrackEditors/PropertyTrackEditors/BoolPropertyTrackEditor.h
+++ b/Engine/Source/Editor/MovieSceneTools/Public/TrackEditors/PropertyTrackEditors/BoolPropertyTrackEditor.h
@@ -21,6 +21,8 @@
 #include "Tracks/MovieSceneBoolTrack.h"
 #include "UObject/UnrealNames.h"
 
+#define UE_API MOVIESCENETOOLS_API
+
 class FPropertyChangedParams;
 class ISequencer;
 class ISequencerSection;
@@ -32,7 +34,7 @@ struct FGuid;
 /**
  * A property track editor for Booleans.
  */
-class MOVIESCENETOOLS_API FBoolPropertyTrackEditor
+class FBoolPropertyTrackEditor
 	: public FPropertyTrackEditor<UMovieSceneBoolTrack>
 {
 public:
@@ -60,12 +62,14 @@ public:
 	 * @param OwningSequencer The sequencer instance to be used by this tool.
 	 * @return The new instance of this class.
 	 */
-	static TSharedRef<ISequencerTrackEditor> CreateTrackEditor( TSharedRef<ISequencer> OwningSequencer );
+	static UE_API TSharedRef<ISequencerTrackEditor> CreateTrackEditor( TSharedRef<ISequencer> OwningSequencer );
 
 protected:
 
 	//~ FPropertyTrackEditor interface
-	virtual FText GetDisplayName() const override;
-	virtual TSharedRef<ISequencerSection> MakeSectionInterface(UMovieSceneSection& SectionObject, UMovieSceneTrack& Track, FGuid ObjectBinding) override;
-	virtual void GenerateKeysFromPropertyChanged( const FPropertyChangedParams& PropertyChangedParams, UMovieSceneSection* SectionToKey, FGeneratedTrackKeys& OutGeneratedKeys ) override;
+	UE_API virtual FText GetDisplayName() const override;
+	UE_API virtual TSharedRef<ISequencerSection> MakeSectionInterface(UMovieSceneSection& SectionObject, UMovieSceneTrack& Track, FGuid ObjectBinding) override;
+	UE_API virtual void GenerateKeysFromPropertyChanged( const FPropertyChangedParams& PropertyChangedParams, UMovieSceneSection* SectionToKey, FGeneratedTrackKeys& OutGeneratedKeys ) override;
 };
+
+#undef UE_API
diff --git a/Engine/Source/Editor/MovieSceneTools/Public/TrackEditors/SpawnTrackEditor.h b/Engine/Source/Editor/MovieSceneTools/Public/TrackEditors/SpawnTrackEditor.h
--- a/Engine/Source/Editor/MovieSceneTools/Public/TrackEditors/SpawnTrackEditor.h
+++ b/Engine/Source/Editor/MovieSceneTools/Public/TrackEditors/SpawnTrackEditor.h
@@ -14,6 +14,8 @@
 #include "TrackEditors/PropertyTrackEditors/BoolPropertyTrackEditor.h"
 #include "UObject/NameTypes.h"
 
+#define UE_API MOVIESCENETOOLS_API
+
 class FMenuBuilder;
 class ISequencer;
 class ISequencerTrackEditor;
@@ -29,7 +31,7 @@ struct FGuid;
 /**
  * A property track editor for controlling the lifetime of a sapwnable object
  */
-class MOVIESCENETOOLS_API FSpawnTrackEditor
+class FSpawnTrackEditor
 	// TODO: Don't derive from bool property track editor here since it doesn't actually edit
 	// bool property tracks.
 	: public FBoolPropertyTrackEditor
@@ -42,7 +44,7 @@ public:
 	 * @param InSequencer The sequencer instance to be used by this tool.
 	 * @return The new instance of this class.
 	 */
-	static TSharedRef<ISequencerTrackEditor> CreateTrackEditor(TSharedRef<ISequencer> InSequencer);
+	static UE_API TSharedRef<ISequencerTrackEditor> CreateTrackEditor(TSharedRef<ISequencer> InSequencer);
 
 public:
 
@@ -51,22 +53,24 @@ public:
 	 *
 	 * @param InSequencer The sequencer instance to be used by this tool.
 	 */
-	FSpawnTrackEditor(TSharedRef<ISequencer> InSequencer);
+	UE_API FSpawnTrackEditor(TSharedRef<ISequencer> InSequencer);
 
 public:
 
 	// ISequencerTrackEditor interface
 
-	virtual UMovieSceneTrack* AddTrack(UMovieScene* FocusedMovieScene, const FGuid& ObjectHandle, TSubclassOf<class UMovieSceneTrack> TrackClass, FName UniqueTypeName) override;
-	virtual void BuildObjectBindingTrackMenu(FMenuBuilder& MenuBuilder, const TArray<FGuid>& ObjectBindings, const UClass* ObjectClass) override;
+	UE_API virtual UMovieSceneTrack* AddTrack(UMovieScene* FocusedMovieScene, const FGuid& ObjectHandle, TSubclassOf<class UMovieSceneTrack> TrackClass, FName UniqueTypeName) override;
+	UE_API virtual void BuildObjectBindingTrackMenu(FMenuBuilder& MenuBuilder, const TArray<FGuid>& ObjectBindings, const UClass* ObjectClass) override;
 	virtual TSharedPtr<SWidget> BuildOutlinerEditWidget(const FGuid& ObjectBinding, UMovieSceneTrack* Track, const FBuildEditWidgetParams& Params) override  { return TSharedPtr<SWidget>(); }
 	virtual bool HandleAssetAdded(UObject* Asset, const FGuid& TargetObjectGuid) override { return false; }
-	virtual bool SupportsType(TSubclassOf<UMovieSceneTrack> Type) const override;
-    virtual bool SupportsSequence(UMovieSceneSequence* InSequence) const override;
+	UE_API virtual bool SupportsType(TSubclassOf<UMovieSceneTrack> Type) const override;
+    UE_API virtual bool SupportsSequence(UMovieSceneSequence* InSequence) const override;
 
 private:
 
 	/** Callback for executing the "Add Spawn Track" menu entry. */
-	void HandleAddSpawnTrackMenuEntryExecute(TArray<FGuid> ObjectBindings);
-	bool CanAddSpawnTrack(FGuid ObjectBinding) const;
+	UE_API void HandleAddSpawnTrackMenuEntryExecute(TArray<FGuid> ObjectBindings);
+	UE_API bool CanAddSpawnTrack(FGuid ObjectBinding) const;
 };
+
+#undef UE_API
diff --git a/Engine/Source/Editor/MovieSceneTools/Public/TrackEditors/SubTrackEditor.h b/Engine/Source/Editor/MovieSceneTools/Public/TrackEditors/SubTrackEditor.h
--- a/Engine/Source/Editor/MovieSceneTools/Public/TrackEditors/SubTrackEditor.h
+++ b/Engine/Source/Editor/MovieSceneTools/Public/TrackEditors/SubTrackEditor.h
@@ -14,6 +14,8 @@
 #include "KeyframeTrackEditor.h"
 #include "MovieSceneTrackEditor.h"
 
+#define UE_API MOVIESCENETOOLS_API
+
 class AActor;
 struct FAssetData;
 class FMenuBuilder;
@@ -24,7 +26,7 @@ namespace UE::Sequencer { struct FAddKeyResult; }
 /**
  * Tools for subsequences
  */
-class MOVIESCENETOOLS_API FSubTrackEditor
+class FSubTrackEditor
 	: public FKeyframeTrackEditor<UMovieSceneSubTrack>
 {
 public:
@@ -34,7 +36,7 @@ public:
 	 *
 	 * @param InSequencer The sequencer instance to be used by this tool
 	 */
-	FSubTrackEditor(TSharedRef<ISequencer> InSequencer);
+	UE_API FSubTrackEditor(TSharedRef<ISequencer> InSequencer);
 
 	/** Virtual destructor. */
 	virtual ~FSubTrackEditor() { }
@@ -45,53 +47,53 @@ public:
 	 * @param OwningSequencer The sequencer instance to be used by this tool
 	 * @return The new instance of this class
 	 */
-	static TSharedRef<ISequencerTrackEditor> CreateTrackEditor(TSharedRef<ISequencer> OwningSequencer);
+	static UE_API TSharedRef<ISequencerTrackEditor> CreateTrackEditor(TSharedRef<ISequencer> OwningSequencer);
 
 public:
 
 	// ISequencerTrackEditor interface
-	virtual FText GetDisplayName() const override;
-	virtual void ProcessKeyOperation(FFrameNumber InKeyTime, const UE::Sequencer::FKeyOperation& Operation, ISequencer& InSequencer, TArray<UE::Sequencer::FAddKeyResult>* OutResults = nullptr) override;
-	virtual void BuildAddTrackMenu(FMenuBuilder& MenuBuilder) override;
-	virtual TSharedPtr<SWidget> BuildOutlinerEditWidget(const FGuid& ObjectBinding, UMovieSceneTrack* Track, const FBuildEditWidgetParams& Params) override;
-	virtual TSharedRef<ISequencerSection> MakeSectionInterface(UMovieSceneSection& SectionObject, UMovieSceneTrack& Track, FGuid ObjectBinding) override;
-	virtual bool HandleAssetAdded(UObject* Asset, const FGuid& TargetObjectGuid) override;
-	virtual bool SupportsSequence(UMovieSceneSequence* InSequence) const override;
-	virtual bool SupportsType(TSubclassOf<UMovieSceneTrack> Type) const override;
-	virtual const FSlateBrush* GetIconBrush() const override;
-	virtual bool OnAllowDrop(const FDragDropEvent& DragDropEvent, FSequencerDragDropParams& DragDropParams) override;
-	virtual FReply OnDrop(const FDragDropEvent& DragDropEvent, const FSequencerDragDropParams& DragDropParams) override;
-	virtual bool IsResizable(UMovieSceneTrack* InTrack) const override;
-	virtual void OnInitialize() override;
-	virtual void OnRelease() override;
-	virtual void Resize(float NewSize, UMovieSceneTrack* InTrack) override;
-	virtual bool GetDefaultExpansionState(UMovieSceneTrack* InTrack) const override;
+	UE_API virtual FText GetDisplayName() const override;
+	UE_API virtual void ProcessKeyOperation(FFrameNumber InKeyTime, const UE::Sequencer::FKeyOperation& Operation, ISequencer& InSequencer, TArray<UE::Sequencer::FAddKeyResult>* OutResults = nullptr) override;
+	UE_API virtual void BuildAddTrackMenu(FMenuBuilder& MenuBuilder) override;
+	UE_API virtual TSharedPtr<SWidget> BuildOutlinerEditWidget(const FGuid& ObjectBinding, UMovieSceneTrack* Track, const FBuildEditWidgetParams& Params) override;
+	UE_API virtual TSharedRef<ISequencerSection> MakeSectionInterface(UMovieSceneSection& SectionObject, UMovieSceneTrack& Track, FGuid ObjectBinding) override;
+	UE_API virtual bool HandleAssetAdded(UObject* Asset, const FGuid& TargetObjectGuid) override;
+	UE_API virtual bool SupportsSequence(UMovieSceneSequence* InSequence) const override;
+	UE_API virtual bool SupportsType(TSubclassOf<UMovieSceneTrack> Type) const override;
+	UE_API virtual const FSlateBrush* GetIconBrush() const override;
+	UE_API virtual bool OnAllowDrop(const FDragDropEvent& DragDropEvent, FSequencerDragDropParams& DragDropParams) override;
+	UE_API virtual FReply OnDrop(const FDragDropEvent& DragDropEvent, const FSequencerDragDropParams& DragDropParams) override;
+	UE_API virtual bool IsResizable(UMovieSceneTrack* InTrack) const override;
+	UE_API virtual void OnInitialize() override;
+	UE_API virtual void OnRelease() override;
+	UE_API virtual void Resize(float NewSize, UMovieSceneTrack* InTrack) override;
+	UE_API virtual bool GetDefaultExpansionState(UMovieSceneTrack* InTrack) const override;
 	virtual bool HasTransformKeyBindings() const override { return true; }
-	virtual bool CanAddTransformKeysForSelectedObjects() const override;
-	virtual void OnAddTransformKeysForSelectedObjects(EMovieSceneTransformChannel Channel) override;
+	UE_API virtual bool CanAddTransformKeysForSelectedObjects() const override;
+	UE_API virtual void OnAddTransformKeysForSelectedObjects(EMovieSceneTransformChannel Channel) override;
 
 public:
 	
 	/** Insert sequence into this track */
-	virtual void InsertSection(UMovieSceneTrack* Track);
+	UE_API virtual void InsertSection(UMovieSceneTrack* Track);
 
 	/** Duplicate the section into this track */
-	virtual void DuplicateSection(UMovieSceneSubSection* Section);
+	UE_API virtual void DuplicateSection(UMovieSceneSubSection* Section);
 
 	/** Create a new take of the given section */
-	virtual void CreateNewTake(UMovieSceneSubSection* Section);
+	UE_API virtual void CreateNewTake(UMovieSceneSubSection* Section);
 
 	/** Switch the selected section's take sequence */
-	virtual void ChangeTake(UMovieSceneSequence* Sequence);
+	UE_API virtual void ChangeTake(UMovieSceneSequence* Sequence);
 
 	/** Generate a menu for takes for this section */
-	virtual void AddTakesMenu(UMovieSceneSubSection* Section, FMenuBuilder& MenuBuilder);
+	UE_API virtual void AddTakesMenu(UMovieSceneSubSection* Section, FMenuBuilder& MenuBuilder);
 
 	/** Edit the section's metadata */
-	virtual void EditMetaData(UMovieSceneSubSection* Section);
+	UE_API virtual void EditMetaData(UMovieSceneSubSection* Section);
 
 	/** Update the current active edit mode when a subtrack or section is selected.*/
-	void UpdateActiveMode();
+	UE_API void UpdateActiveMode();
 
 	/**
 	 * Check whether the given sequence can be added as a sub-sequence.
@@ -102,94 +104,96 @@ public:
 	 * @param Sequence The sequence to check.
 	 * @return true if the sequence can be added as a sub-sequence, false otherwise.
 	 */
-	bool CanAddSubSequence(const UMovieSceneSequence& Sequence) const;
+	UE_API bool CanAddSubSequence(const UMovieSceneSequence& Sequence) const;
 
 public:
 
 	/** Get the name of the sub track */
-	virtual FText GetSubTrackName() const;
+	UE_API virtual FText GetSubTrackName() const;
 
 	/** Get the tooltip for this sub track editor */
-	virtual FText GetSubTrackToolTip() const;
+	UE_API virtual FText GetSubTrackToolTip() const;
 
 	/** Get the brush used for the sub track editor */
-	virtual FName GetSubTrackBrushName() const;
+	UE_API virtual FName GetSubTrackBrushName() const;
 
 	/** Get the display name for the sub section */
-	virtual FString GetSubSectionDisplayName(const UMovieSceneSubSection* Section) const;
+	UE_API virtual FString GetSubSectionDisplayName(const UMovieSceneSubSection* Section) const;
 
 	/** Get the default sub sequence name */
-	virtual FString GetDefaultSubsequenceName() const;
+	UE_API virtual FString GetDefaultSubsequenceName() const;
 
 	/** Get the sub sequence directory */
-	virtual FString GetDefaultSubsequenceDirectory() const;
+	UE_API virtual FString GetDefaultSubsequenceDirectory() const;
 
 	/** Get the UMovieSceneSubTrack class */
-	virtual TSubclassOf<UMovieSceneSubTrack> GetSubTrackClass() const;
+	UE_API virtual TSubclassOf<UMovieSceneSubTrack> GetSubTrackClass() const;
 
 	/** Called when the editor mode has made external changes to the origin data. */
-	void UpdateOrigin(FVector InPosition, FRotator InRotation);
+	UE_API void UpdateOrigin(FVector InPosition, FRotator InRotation);
 
 	/** Called when sequence playback updates to revert preview data on modified secitons */
-	void ResetSectionPreviews();
+	UE_API void ResetSectionPreviews();
 	void ResetSectionPreviews(FMovieSceneSequenceIDRef IDRef) { ResetSectionPreviews(); }
 	void ResetSectionPreviews(const FMovieSceneChannelMetaData* MetaData, UMovieSceneSection* InSection) { ResetSectionPreviews(); }
 
 	/** Query's the channel data directly (does not take parent transforms into account). Used for setting keyframes. */
-	FTransform GetTransformOriginDataForSubSection(const UMovieSceneSubSection* SubSection) const;
+	UE_API FTransform GetTransformOriginDataForSubSection(const UMovieSceneSubSection* SubSection) const;
 
 	/** Helper function that finds the previous key. Used to "unwind" rotators */
-	int32 GetPreviousKey(FMovieSceneDoubleChannel& Channel, FFrameNumber Time);
+	UE_API int32 GetPreviousKey(FMovieSceneDoubleChannel& Channel, FFrameNumber Time);
 
 	/** Helper function to fix-up Euler rotations if they would go over 180 degrees due to interpolation. */
-	double UnwindChannel(const double& OldValue, double NewValue);
+	UE_API double UnwindChannel(const double& OldValue, double NewValue);
 
 protected:
 
 	/** Get the list of supported sequence class paths */
-	virtual void GetSupportedSequenceClassPaths(TArray<FTopLevelAssetPath>& OutClassPaths) const;
+	UE_API virtual void GetSupportedSequenceClassPaths(TArray<FTopLevelAssetPath>& OutClassPaths) const;
 
 	/** Callback for executing the "Add Subsequence" menu entry. */
-	virtual void HandleAddSubTrackMenuEntryExecute();
+	UE_API virtual void HandleAddSubTrackMenuEntryExecute();
 
 	/** Callback for determining whether the "Add Subsequence" menu entry can execute. */
 	virtual bool HandleAddSubTrackMenuEntryCanExecute() const { return true; }
 
 	/** Whether to handle this asset being dropped onto the sequence as opposed to a specific track. */
-	virtual bool CanHandleAssetAdded(UMovieSceneSequence* Sequence) const;
+	UE_API virtual bool CanHandleAssetAdded(UMovieSceneSequence* Sequence) const;
 
 	UE_DEPRECATED(5.3, "CreateNewTrack has been deprecated, please implement GetSubTrackClass")
-	virtual UMovieSceneSubTrack* CreateNewTrack(UMovieScene* MovieScene) const;
+	UE_API virtual UMovieSceneSubTrack* CreateNewTrack(UMovieScene* MovieScene) const;
 
 	/** Find or create a sub track. If the given track is a subtrack, it will be returned. */
-	UMovieSceneSubTrack* FindOrCreateSubTrack(UMovieScene* MovieScene, UMovieSceneTrack* Track) const;
+	UE_API UMovieSceneSubTrack* FindOrCreateSubTrack(UMovieScene* MovieScene, UMovieSceneTrack* Track) const;
 
 	/** Callback for generating the menu of the "Add Sequence" combo button. */
-	TSharedRef<SWidget> HandleAddSubSequenceComboButtonGetMenuContent(UE::Sequencer::TWeakViewModelPtr<UE::Sequencer::ITrackExtension> TrackModel);
+	UE_API TSharedRef<SWidget> HandleAddSubSequenceComboButtonGetMenuContent(UE::Sequencer::TWeakViewModelPtr<UE::Sequencer::ITrackExtension> TrackModel);
 
 private:
 
 	/** Callback for executing a menu entry in the "Add Sequence" combo button. */
-	void HandleAddSubSequenceComboButtonMenuEntryExecute(const FAssetData& AssetData, UMovieSceneTrack* InTrack);
+	UE_API void HandleAddSubSequenceComboButtonMenuEntryExecute(const FAssetData& AssetData, UMovieSceneTrack* InTrack);
 
 	/** Callback for executing a menu entry in the "Add Sequence" combo button when enter pressed. */
-	void HandleAddSubSequenceComboButtonMenuEntryEnterPressed(const TArray<FAssetData>& AssetData, UMovieSceneTrack* InTrack);
+	UE_API void HandleAddSubSequenceComboButtonMenuEntryEnterPressed(const TArray<FAssetData>& AssetData, UMovieSceneTrack* InTrack);
 
 	/** Delegate for AnimatablePropertyChanged in AddKey */
-	FKeyPropertyResult AddKeyInternal(FFrameNumber KeyTime, UMovieSceneSequence* InMovieSceneSequence, UMovieSceneTrack* InTrack, int32 RowIndex);
+	UE_API FKeyPropertyResult AddKeyInternal(FFrameNumber KeyTime, UMovieSceneSequence* InMovieSceneSequence, UMovieSceneTrack* InTrack, int32 RowIndex);
 
 	/** Callback for AnimatablePropertyChanged in HandleAssetAdded. */
-	FKeyPropertyResult HandleSequenceAdded(FFrameNumber KeyTime, UMovieSceneSequence* Sequence, UMovieSceneTrack* Track, int32 RowIndex);
+	UE_API FKeyPropertyResult HandleSequenceAdded(FFrameNumber KeyTime, UMovieSceneSequence* Sequence, UMovieSceneTrack* Track, int32 RowIndex);
 
 	/** Handles adding keys to section.**/
-	void ProcessKeyOperationInternal(TArrayView<const UE::Sequencer::FKeySectionOperation> SectionsToKey, ISequencer& Sequencer, FFrameNumber KeyTime, TArray<UE::Sequencer::FAddKeyResult>* OutResults = nullptr);
+	UE_API void ProcessKeyOperationInternal(TArrayView<const UE::Sequencer::FKeySectionOperation> SectionsToKey, ISequencer& Sequencer, FFrameNumber KeyTime, TArray<UE::Sequencer::FAddKeyResult>* OutResults = nullptr);
 
 	/** Helper for creating new keys.*/
-	void GetOriginKeys(const FVector& CurrentPosition, const FRotator& CurrentRotation, EMovieSceneTransformChannel ChannelsToKey, UMovieSceneSection* Section, FGeneratedTrackKeys& OutGeneratedKeys);
+	UE_API void GetOriginKeys(const FVector& CurrentPosition, const FRotator& CurrentRotation, EMovieSceneTransformChannel ChannelsToKey, UMovieSceneSection* Section, FGeneratedTrackKeys& OutGeneratedKeys);
 
 	/** Helper for getting the sections that should be keyed. **/
-	void GetSectionsToKey(TArray<UMovieSceneSubSection*>& OutSectionsToKey) const;
+	UE_API void GetSectionsToKey(TArray<UMovieSceneSubSection*>& OutSectionsToKey) const;
 
 	/** Sections this editor has added preview data to for keyframing. Used to revert data when other edits are made, before adding a key */
 	TArray<UMovieSceneSubSection*> SectionsWithPreviews;
 };
+
+#undef UE_API
diff --git a/Engine/Source/Editor/NNEEditor/Private/NNEEditorModelDataFactory.h b/Engine/Source/Editor/NNEEditor/Private/NNEEditorModelDataFactory.h
--- a/Engine/Source/Editor/NNEEditor/Private/NNEEditorModelDataFactory.h
+++ b/Engine/Source/Editor/NNEEditor/Private/NNEEditorModelDataFactory.h
@@ -7,22 +7,26 @@
 
 #include "NNEEditorModelDataFactory.generated.h"
 
+#define UE_API NNEEDITOR_API
+
 /**
  * Class for importing and creating new UNNEModelData assets.
  *
  * The currently supported format is .onnx. The factory is invoked when a file of this format is dragged to the editor's content browser.
  */
-UCLASS()
-class NNEEDITOR_API UNNEModelDataFactory : public UFactory
+UCLASS(MinimalAPI)
+class UNNEModelDataFactory : public UFactory
 {
 	GENERATED_BODY()
 
 public:
-	UNNEModelDataFactory(const FObjectInitializer& ObjectInitializer);
+	UE_API UNNEModelDataFactory(const FObjectInitializer& ObjectInitializer);
 
 public:
 	//~ Begin UFactory Interface
-	virtual UObject* FactoryCreateFile(UClass* InClass, UObject* InParent, FName InName, EObjectFlags Flags, const FString& Filename, const TCHAR* Parms, FFeedbackContext* Warn, bool& bOutOperationCanceled) override;
-	virtual bool FactoryCanImport(const FString& Filename) override;
+	UE_API virtual UObject* FactoryCreateFile(UClass* InClass, UObject* InParent, FName InName, EObjectFlags Flags, const FString& Filename, const TCHAR* Parms, FFeedbackContext* Warn, bool& bOutOperationCanceled) override;
+	UE_API virtual bool FactoryCanImport(const FString& Filename) override;
 	//~ End UFactory Interface
-};
\ No newline at end of file
+};
+
+#undef UE_API
diff --git a/Engine/Source/Editor/PIEPreviewDeviceProfileSelector/Private/PIEPreviewSettings.h b/Engine/Source/Editor/PIEPreviewDeviceProfileSelector/Private/PIEPreviewSettings.h
--- a/Engine/Source/Editor/PIEPreviewDeviceProfileSelector/Private/PIEPreviewSettings.h
+++ b/Engine/Source/Editor/PIEPreviewDeviceProfileSelector/Private/PIEPreviewSettings.h
@@ -13,8 +13,8 @@
 
 #include "PIEPreviewSettings.generated.h"
 
-UCLASS(hidecategories = Object, config = PIEPreviewSettings)
-class PIEPREVIEWDEVICEPROFILESELECTOR_API UPIEPreviewSettings : public UObject
+UCLASS(MinimalAPI, hidecategories = Object, config = PIEPreviewSettings)
+class UPIEPreviewSettings : public UObject
 {
 	GENERATED_UCLASS_BODY()
 
@@ -27,4 +27,4 @@ public:
 
 	UPROPERTY(config)
 	float WindowScalingFactor;
-};
\ No newline at end of file
+};
diff --git a/Engine/Source/Editor/PIEPreviewDeviceProfileSelector/Private/PIEPreviewWindow.h b/Engine/Source/Editor/PIEPreviewDeviceProfileSelector/Private/PIEPreviewWindow.h
--- a/Engine/Source/Editor/PIEPreviewDeviceProfileSelector/Private/PIEPreviewWindow.h
+++ b/Engine/Source/Editor/PIEPreviewDeviceProfileSelector/Private/PIEPreviewWindow.h
@@ -9,108 +9,108 @@
 
 class SPIEPreviewWindowTitleBar;
 
-class PIEPREVIEWDEVICEPROFILESELECTOR_API SPIEPreviewWindow
+class SPIEPreviewWindow
 	: public SWindow
 {
 public:
 	/**
 	* Default constructor. Use SNew(SPIEPreviewWindow) instead.
 	*/
-	SPIEPreviewWindow();
+	PIEPREVIEWDEVICEPROFILESELECTOR_API SPIEPreviewWindow();
 
-	~SPIEPreviewWindow();
+	PIEPREVIEWDEVICEPROFILESELECTOR_API ~SPIEPreviewWindow();
 
 	// call this function set the 'preview device' and set initial values for the position and resolution scaling factor
-	void PrepareWindow(FVector2D WindowPosition, const float InitialScaleFactor, TSharedPtr<FPIEPreviewDevice> PreviewDevice);
+	PIEPREVIEWDEVICEPROFILESELECTOR_API void PrepareWindow(FVector2D WindowPosition, const float InitialScaleFactor, TSharedPtr<FPIEPreviewDevice> PreviewDevice);
 
 	// call this function to perform cleanup
-	void PrepareShutdown();
+	PIEPREVIEWDEVICEPROFILESELECTOR_API void PrepareShutdown();
 
 	// retrieves the display device
-	TSharedPtr<class FPIEPreviewDevice> GetDevice();
+	PIEPREVIEWDEVICEPROFILESELECTOR_API TSharedPtr<class FPIEPreviewDevice> GetDevice();
 
 	// call to retrieve the window current resolution scale factor
-	float GetWindowScaleFactor() const;
+	PIEPREVIEWDEVICEPROFILESELECTOR_API float GetWindowScaleFactor() const;
 
 	// call to set the window resolution scale factor
-	void SetWindowScaleFactor(const float ScaleFactor, const bool bStore = true);
+	PIEPREVIEWDEVICEPROFILESELECTOR_API void SetWindowScaleFactor(const float ScaleFactor, const bool bStore = true);
 
 	// call to rotate the window emulating a rotation of the physical device
-	void RotateWindow();
+	PIEPREVIEWDEVICEPROFILESELECTOR_API void RotateWindow();
 
 	// wrapper function that will query the same functionality on the device
-	bool IsRotationAllowed() const;
+	PIEPREVIEWDEVICEPROFILESELECTOR_API bool IsRotationAllowed() const;
 
 	// call to enable/disable bezel visibility
-	void FlipBezelVisibility();
+	PIEPREVIEWDEVICEPROFILESELECTOR_API void FlipBezelVisibility();
 
 	// call to determine whether or not the phone bezel is displayed 
-	bool GetBezelVisibility() const;
+	PIEPREVIEWDEVICEPROFILESELECTOR_API bool GetBezelVisibility() const;
 
-	static int32 GetDefaultTitleBarSize();
+	static PIEPREVIEWDEVICEPROFILESELECTOR_API int32 GetDefaultTitleBarSize();
 
 	// call with a true argument to restrict the window size to the desktop size
-	void SetClampWindowSize(const bool bClamp);
-	bool IsClampingWindowSize() const;
+	PIEPREVIEWDEVICEPROFILESELECTOR_API void SetClampWindowSize(const bool bClamp);
+	PIEPREVIEWDEVICEPROFILESELECTOR_API bool IsClampingWindowSize() const;
 
 	/** creates and returns the settings menu */
-	TSharedRef<SWidget> BuildSettingsMenu();
+	PIEPREVIEWDEVICEPROFILESELECTOR_API TSharedRef<SWidget> BuildSettingsMenu();
 
 	/** we need the game layer manager to control the DPI scaling behavior and this function can be called should be called when the manager is available */
-	void SetGameLayerManagerWidget(TSharedPtr<class SGameLayerManager> GameLayerManager);
+	PIEPREVIEWDEVICEPROFILESELECTOR_API void SetGameLayerManagerWidget(TSharedPtr<class SGameLayerManager> GameLayerManager);
 
 private:
-	virtual EHorizontalAlignment GetTitleAlignment() override;
-	virtual TSharedRef<SWidget> MakeWindowTitleBar(const TSharedRef<SWindow>& Window, const TSharedPtr<SWidget>& CenterContent, EHorizontalAlignment CenterContentAlignment) override;
+	PIEPREVIEWDEVICEPROFILESELECTOR_API virtual EHorizontalAlignment GetTitleAlignment() override;
+	PIEPREVIEWDEVICEPROFILESELECTOR_API virtual TSharedRef<SWidget> MakeWindowTitleBar(const TSharedRef<SWindow>& Window, const TSharedPtr<SWidget>& CenterContent, EHorizontalAlignment CenterContentAlignment) override;
 
 	// sets the device displayed by this window
-	void SetDevice(TSharedPtr<class FPIEPreviewDevice> InDevice);
+	PIEPREVIEWDEVICEPROFILESELECTOR_API void SetDevice(TSharedPtr<class FPIEPreviewDevice> InDevice);
 
 	// returns whether or not the provided ScaleFactor is a 'scale to physical device size' factor
-	bool IsScalingToDeviceSizeFactor(const float ScaleFactor) const;
+	PIEPREVIEWDEVICEPROFILESELECTOR_API bool IsScalingToDeviceSizeFactor(const float ScaleFactor) const;
 
 	// returns the scaling value reserved for 'scale to physical device size' factor
-	float GetScaleToDeviceSizeFactor() const;
+	PIEPREVIEWDEVICEPROFILESELECTOR_API float GetScaleToDeviceSizeFactor() const;
 
 	// utility function used to correct a given window position if it's out of the display area
-	void ValidatePosition(FVector2D& WindowPos);
+	PIEPREVIEWDEVICEPROFILESELECTOR_API void ValidatePosition(FVector2D& WindowPos);
 
 	// call to perform resolution scaling on the device
-	void ScaleWindow(float ScaleFactor);
+	PIEPREVIEWDEVICEPROFILESELECTOR_API void ScaleWindow(float ScaleFactor);
 
 	// this will properly scale and rotate the bezel to match the orientation of the device
-	void ComputeBezelOrientation();
+	PIEPREVIEWDEVICEPROFILESELECTOR_API void ComputeBezelOrientation();
 
 	// function that adds a widget to display the device's bezel
-	void CreatePIEPreviewBezelOverlay(UTexture2D* pBezelImage);
+	PIEPREVIEWDEVICEPROFILESELECTOR_API void CreatePIEPreviewBezelOverlay(UTexture2D* pBezelImage);
 
 	/** helper function used to create a menu item formed from a description text and a checkbox */
 	/** @ IsCheckedFunction - computes the state of the checkbox */
 	/** @ ExecuteActionFunction - function to be called when this menu item is clicked */
-	void CreateMenuEntry(class FMenuBuilder& MenuBuilder, FText&& TextEntry, TFunction<ECheckBoxState()>&& IsCheckedFunction, TFunction<void()>&& ExecuteActionFunction);
+	PIEPREVIEWDEVICEPROFILESELECTOR_API void CreateMenuEntry(class FMenuBuilder& MenuBuilder, FText&& TextEntry, TFunction<ECheckBoxState()>&& IsCheckedFunction, TFunction<void()>&& ExecuteActionFunction);
 
 	/** helper function used to create a simple text widget to the provided menu */
 	/** @ CreateTextFunction - function should return the menu text that we want to display */
-	void CreateTextMenuEntry(class FMenuBuilder& MenuBuilder, TFunction<FText()>&& CreateTextFunction);
+	PIEPREVIEWDEVICEPROFILESELECTOR_API void CreateTextMenuEntry(class FMenuBuilder& MenuBuilder, TFunction<FText()>&& CreateTextFunction);
 
 	// call to compute window size, position on screen and bezel orientation
-	void UpdateWindow();
+	PIEPREVIEWDEVICEPROFILESELECTOR_API void UpdateWindow();
 
 	// utility function that will compute a screen and DPI scale factor needed to scale the display window to the physical device size
-	float ComputeScaleToDeviceSizeFactor() const;
+	PIEPREVIEWDEVICEPROFILESELECTOR_API float ComputeScaleToDeviceSizeFactor() const;
 
 	// utility function used to retrieve the DPI scaling factor based on window's position
-	float ComputeDPIScaleFactor();
+	PIEPREVIEWDEVICEPROFILESELECTOR_API float ComputeDPIScaleFactor();
 
 	// callback used when DPI value varies, useful when we need to constrain the window to the physical device size
-	void OnWindowMoved(const TSharedRef<SWindow>& Window);
+	PIEPREVIEWDEVICEPROFILESELECTOR_API void OnWindowMoved(const TSharedRef<SWindow>& Window);
 
 	// callback used when DPI value varies, useful when we need to constrain the window to the physical device size
-	void OnDisplayDPIChanged(TSharedRef<SWindow> Window);
+	PIEPREVIEWDEVICEPROFILESELECTOR_API void OnDisplayDPIChanged(TSharedRef<SWindow> Window);
 
 	// utility function used to control game layer manager's DPI scaling behavior
 	// this should be called every time the window is rotated
-	void UpdateGameLayerManagerDefaultViewport();
+	PIEPREVIEWDEVICEPROFILESELECTOR_API void UpdateGameLayerManagerDefaultViewport();
 
 protected:
 	// brush created to display the bezel
@@ -179,4 +179,4 @@ FORCEINLINE float SPIEPreviewWindow::GetScaleToDeviceSizeFactor() const
 	return 0.0f;
 }
 
-#endif
\ No newline at end of file
+#endif
diff --git a/Engine/Source/Editor/PIEPreviewDeviceProfileSelector/Private/PIEPreviewWindowCoreStyle.h b/Engine/Source/Editor/PIEPreviewDeviceProfileSelector/Private/PIEPreviewWindowCoreStyle.h
--- a/Engine/Source/Editor/PIEPreviewDeviceProfileSelector/Private/PIEPreviewWindowCoreStyle.h
+++ b/Engine/Source/Editor/PIEPreviewDeviceProfileSelector/Private/PIEPreviewWindowCoreStyle.h
@@ -6,14 +6,14 @@
 
 #include "Styling/SlateStyle.h"
 
-class PIEPREVIEWDEVICEPROFILESELECTOR_API FPIEPreviewWindowCoreStyle
+class FPIEPreviewWindowCoreStyle
 {
 
 public:
 
-	static TSharedRef<class ISlateStyle> Create(const FName& InStyleSetName = "FPIECoreStyle");
+	static PIEPREVIEWDEVICEPROFILESELECTOR_API TSharedRef<class ISlateStyle> Create(const FName& InStyleSetName = "FPIECoreStyle");
 	
-	static void InitializePIECoreStyle();
+	static PIEPREVIEWDEVICEPROFILESELECTOR_API void InitializePIECoreStyle();
 	
 	/** @return the singleton instance. */
 	static const ISlateStyle& Get()
@@ -23,6 +23,6 @@ public:
 
 private:
 		
-	static TSharedPtr< class ISlateStyle > Instance;
+	static PIEPREVIEWDEVICEPROFILESELECTOR_API TSharedPtr< class ISlateStyle > Instance;
 };
-#endif
\ No newline at end of file
+#endif
diff --git a/Engine/Source/Editor/PIEPreviewDeviceProfileSelector/Private/PIEPreviewWindowStyle.h b/Engine/Source/Editor/PIEPreviewDeviceProfileSelector/Private/PIEPreviewWindowStyle.h
--- a/Engine/Source/Editor/PIEPreviewDeviceProfileSelector/Private/PIEPreviewWindowStyle.h
+++ b/Engine/Source/Editor/PIEPreviewDeviceProfileSelector/Private/PIEPreviewWindowStyle.h
@@ -10,20 +10,20 @@
 * Represents the appearance of an SPIEWindow
 */
 USTRUCT()
-struct PIEPREVIEWDEVICEPROFILESELECTOR_API FPIEPreviewWindowStyle : public FSlateWidgetStyle
+struct FPIEPreviewWindowStyle : public FSlateWidgetStyle
 {
 	GENERATED_USTRUCT_BODY()
 
-	FPIEPreviewWindowStyle();
+	PIEPREVIEWDEVICEPROFILESELECTOR_API FPIEPreviewWindowStyle();
 
 	virtual ~FPIEPreviewWindowStyle() {}
 
-	virtual void GetResources(TArray< const FSlateBrush* >& OutBrushes) const override;
+	PIEPREVIEWDEVICEPROFILESELECTOR_API virtual void GetResources(TArray< const FSlateBrush* >& OutBrushes) const override;
 
-	static const FName TypeName;
+	static PIEPREVIEWDEVICEPROFILESELECTOR_API const FName TypeName;
 	virtual const FName GetTypeName() const override { return TypeName; };
 
-	static const FPIEPreviewWindowStyle& GetDefault();
+	static PIEPREVIEWDEVICEPROFILESELECTOR_API const FPIEPreviewWindowStyle& GetDefault();
 	
 	/** Style used to draw the window ScreenRotationButton button */
 	UPROPERTY()
@@ -46,4 +46,4 @@ struct PIEPREVIEWDEVICEPROFILESELECTOR_API FPIEPreviewWindowStyle : public FSlat
 	FPIEPreviewWindowStyle& SetFullMobileContentScaleFactorButtonStyle(const FButtonStyle& InFullMobileContentScaleFactorButtonStyle) { FullMobileContentScaleFactorButtonStyle = InFullMobileContentScaleFactorButtonStyle; return *this; }
 
 };
-#endif
\ No newline at end of file
+#endif
diff --git a/Engine/Source/Editor/PIEPreviewDeviceProfileSelector/Private/PIEPreviewWindowTitleBar.h b/Engine/Source/Editor/PIEPreviewDeviceProfileSelector/Private/PIEPreviewWindowTitleBar.h
--- a/Engine/Source/Editor/PIEPreviewDeviceProfileSelector/Private/PIEPreviewWindowTitleBar.h
+++ b/Engine/Source/Editor/PIEPreviewDeviceProfileSelector/Private/PIEPreviewWindowTitleBar.h
@@ -12,18 +12,18 @@
 /**
 * Implements a window PIE title bar widget.
 */
-class PIEPREVIEWDEVICEPROFILESELECTOR_API SPIEPreviewWindowTitleBar
+class SPIEPreviewWindowTitleBar
 	: public SWindowTitleBar
 {
 private:
 	/**
 	* Creates widgets for this window's title bar area.
 	*/
-	virtual void MakeTitleBarContentWidgets(TSharedPtr< SWidget >& OutLeftContent, TSharedPtr< SWidget >& OutRightContent) override;
+	PIEPREVIEWDEVICEPROFILESELECTOR_API virtual void MakeTitleBarContentWidgets(TSharedPtr< SWidget >& OutLeftContent, TSharedPtr< SWidget >& OutRightContent) override;
 
-	const FSlateBrush* GetScreenRotationButtonImage() const;
+	PIEPREVIEWDEVICEPROFILESELECTOR_API const FSlateBrush* GetScreenRotationButtonImage() const;
 
-	TSharedPtr<class SPIEPreviewWindow> GetOwnerWindow() const;
+	PIEPREVIEWDEVICEPROFILESELECTOR_API TSharedPtr<class SPIEPreviewWindow> GetOwnerWindow() const;
 
 private:
 	// Holds the screen rotation button.
diff --git a/Engine/Source/Editor/PIEPreviewDeviceProfileSelector/Public/PIEPreviewDeviceEnumeration.h b/Engine/Source/Editor/PIEPreviewDeviceProfileSelector/Public/PIEPreviewDeviceEnumeration.h
--- a/Engine/Source/Editor/PIEPreviewDeviceProfileSelector/Public/PIEPreviewDeviceEnumeration.h
+++ b/Engine/Source/Editor/PIEPreviewDeviceProfileSelector/Public/PIEPreviewDeviceEnumeration.h
@@ -9,11 +9,13 @@
 #include "Templates/SharedPointer.h"
 #include "UObject/NameTypes.h"
 
+#define UE_API PIEPREVIEWDEVICEPROFILESELECTOR_API
+
 /**
 * 
 */
 
-class PIEPREVIEWDEVICEPROFILESELECTOR_API FPIEPreviewDeviceContainerCategory
+class FPIEPreviewDeviceContainerCategory
 {
 public:
 	FPIEPreviewDeviceContainerCategory(const FString& InSubDirectoryPath, FText DisplayName) : 
@@ -22,8 +24,8 @@ public:
 	{
 	}
 
-	FName GetCategoryName();
-	FText GetCategoryToolTip();
+	UE_API FName GetCategoryName();
+	UE_API FText GetCategoryToolTip();
 	const FText& GetCategoryDisplayName() const { return CategoryDisplayName; }
 	const FString& GetSubDirectoryPath() const { return SubDirectoryPath; }
 	const TArray<TSharedPtr<FPIEPreviewDeviceContainerCategory>>& GetSubCategories() const {return SubCategories;}
@@ -38,12 +40,12 @@ protected:
 	friend class FPIEPreviewDeviceContainer;
 };
 
-class PIEPREVIEWDEVICEPROFILESELECTOR_API FPIEPreviewDeviceContainer
+class FPIEPreviewDeviceContainer
 {
 public:
 	// Recursively iterate through 'RootDir' for all device json files.
 	// Sub directories are recorded as categories.
-	void EnumerateDeviceSpecifications(const FString& RootDir);
+	UE_API void EnumerateDeviceSpecifications(const FString& RootDir);
 
 	const TSharedPtr<FPIEPreviewDeviceContainerCategory> GetRootCategory() const { return RootCategory; }
 
@@ -52,7 +54,7 @@ public:
 	const TArray<FString>& GetDeviceSpecificationsLocalizedName() const { return DeviceSpecificationsLocalizedName; }
 
 	// return the category that contains DeviceIndex.
-	const TSharedPtr<FPIEPreviewDeviceContainerCategory> FindDeviceContainingCategory(int32 DeviceIndex) const;
+	UE_API const TSharedPtr<FPIEPreviewDeviceContainerCategory> FindDeviceContainingCategory(int32 DeviceIndex) const;
 
 private:
 	FString DeviceSpecificationRootDir;
@@ -62,6 +64,8 @@ private:
 	// All device specifications found.
 	TArray<FString> DeviceSpecifications;
 	TArray<FString> DeviceSpecificationsLocalizedName;
-	void EnumerateDeviceSpecifications(TSharedPtr<FPIEPreviewDeviceContainerCategory> SubCategory);
-	void EnumerateDeviceLocalizedNames();
-};
\ No newline at end of file
+	UE_API void EnumerateDeviceSpecifications(TSharedPtr<FPIEPreviewDeviceContainerCategory> SubCategory);
+	UE_API void EnumerateDeviceLocalizedNames();
+};
+
+#undef UE_API
diff --git a/Engine/Source/Editor/Persona/Private/IPoseWatchManagerTreeItem.h b/Engine/Source/Editor/Persona/Private/IPoseWatchManagerTreeItem.h
--- a/Engine/Source/Editor/Persona/Private/IPoseWatchManagerTreeItem.h
+++ b/Engine/Source/Editor/Persona/Private/IPoseWatchManagerTreeItem.h
@@ -7,10 +7,12 @@
 #include "Widgets/Views/STableRow.h"
 #include "PoseWatchManagerStandaloneTypes.h"
 
+#define UE_API PERSONA_API
+
 class UToolMenu;
 class SPoseWatchManager;
 
-struct PERSONA_API IPoseWatchManagerTreeItem : TSharedFromThis<IPoseWatchManagerTreeItem>
+struct IPoseWatchManagerTreeItem : TSharedFromThis<IPoseWatchManagerTreeItem>
 {
 	friend SPoseWatchManager;
 
@@ -118,7 +120,7 @@ protected:
 	TSet<TWeakPtr<IPoseWatchManagerTreeItem>> Children;
 
 	/** Static type identifier for the base class tree item */
-	static const EPoseWatchTreeItemType Type;
+	static UE_API const EPoseWatchTreeItemType Type;
 
 	/** Tree item type identifier */
 	EPoseWatchTreeItemType TreeType;
@@ -132,4 +134,6 @@ template <> inline const IPoseWatchManagerTreeItem* IPoseWatchManagerTreeItem::C
 template <> inline IPoseWatchManagerTreeItem* IPoseWatchManagerTreeItem::CastTo<IPoseWatchManagerTreeItem>()
 {
 	return this;
-}
\ No newline at end of file
+}
+
+#undef UE_API
diff --git a/Engine/Source/Editor/Persona/Private/PoseWatchManagerDragDrop.h b/Engine/Source/Editor/Persona/Private/PoseWatchManagerDragDrop.h
--- a/Engine/Source/Editor/Persona/Private/PoseWatchManagerDragDrop.h
+++ b/Engine/Source/Editor/Persona/Private/PoseWatchManagerDragDrop.h
@@ -17,6 +17,8 @@
 #include "UObject/WeakObjectPtr.h"
 #include "UObject/WeakObjectPtrTemplates.h"
 
+#define UE_API PERSONA_API
+
 class SWidget;
 class UClass;
 struct FSlateBrush;
@@ -77,11 +79,11 @@ struct FPoseWatchManagerDragValidationInfo
 };
 
 /** A drag/drop operation that was started from the pose watch manager */
-struct PERSONA_API FPoseWatchManagerDragDropOp : public FCompositeDragDropOp
+struct FPoseWatchManagerDragDropOp : public FCompositeDragDropOp
 {
 	DRAG_DROP_OPERATOR_TYPE(FPoseWatchManagerDragDropOp, FCompositeDragDropOp);
 
-	FPoseWatchManagerDragDropOp();
+	UE_API FPoseWatchManagerDragDropOp();
 
 	using FDragDropOperation::Construct;
 
@@ -97,12 +99,12 @@ struct PERSONA_API FPoseWatchManagerDragDropOp : public FCompositeDragDropOp
 		OverrideIcon = InOverrideIcon;
 	}
 
-	virtual TSharedPtr<SWidget> GetDefaultDecorator() const override;
+	UE_API virtual TSharedPtr<SWidget> GetDefaultDecorator() const override;
 
 private:
 
-	EVisibility GetOverrideVisibility() const;
-	EVisibility GetDefaultVisibility() const;
+	UE_API EVisibility GetOverrideVisibility() const;
+	UE_API EVisibility GetDefaultVisibility() const;
 
 	FText OverrideText;
 	FText GetOverrideText() const { return OverrideText; }
@@ -152,3 +154,5 @@ public:
 		SetupDefaults();
 	}
 };
+
+#undef UE_API
diff --git a/Engine/Source/Editor/Persona/Private/PoseWatchManagerStandaloneTypes.h b/Engine/Source/Editor/Persona/Private/PoseWatchManagerStandaloneTypes.h
--- a/Engine/Source/Editor/Persona/Private/PoseWatchManagerStandaloneTypes.h
+++ b/Engine/Source/Editor/Persona/Private/PoseWatchManagerStandaloneTypes.h
@@ -7,6 +7,8 @@
 #include "UObject/ObjectKey.h"
 #include "Templates/MaxSizeof.h"
 
+#define UE_API PERSONA_API
+
 struct IPoseWatchManagerTreeItem;
 class IPoseWatchManager;
 
@@ -22,14 +24,14 @@ enum EPoseWatchTreeItemType
 	Element
 };
 
-struct PERSONA_API FPoseWatchManagerCommonLabelData
+struct FPoseWatchManagerCommonLabelData
 {
 	TWeakPtr<IPoseWatchManager> WeakPoseWatchManager;
-	static const FLinearColor DisabledColor;
+	static UE_API const FLinearColor DisabledColor;
 
-	TOptional<FLinearColor> GetForegroundColor(const IPoseWatchManagerTreeItem& TreeItem) const;
+	UE_API TOptional<FLinearColor> GetForegroundColor(const IPoseWatchManagerTreeItem& TreeItem) const;
 
-	bool CanExecuteRenameRequest(const IPoseWatchManagerTreeItem& Item) const;
+	UE_API bool CanExecuteRenameRequest(const IPoseWatchManagerTreeItem& Item) const;
 };
 
 
@@ -47,3 +49,5 @@ struct FPoseWatchManagerHierarchyChangedData
 	/** Actions to apply to items */
 	uint8 ItemActions = 0;
 };
+
+#undef UE_API
diff --git a/Engine/Source/Editor/Persona/Public/AnimAssetFindReplace.h b/Engine/Source/Editor/Persona/Public/AnimAssetFindReplace.h
--- a/Engine/Source/Editor/Persona/Public/AnimAssetFindReplace.h
+++ b/Engine/Source/Editor/Persona/Public/AnimAssetFindReplace.h
@@ -1,4 +1,4 @@
-﻿// Copyright Epic Games, Inc. All Rights Reserved.
+// Copyright Epic Games, Inc. All Rights Reserved.
 
 #pragma once
 
@@ -9,6 +9,8 @@
 #include "Templates/SubclassOf.h"
 #include "AnimAssetFindReplace.generated.h"
 
+#define UE_API PERSONA_API
+
 class SAnimAssetFindReplace;
 struct FToolMenuSection;
 class SAutoCompleteSearchBox;
@@ -32,8 +34,8 @@ enum class EAnimAssetFindReplaceMode : int32
 };
 
 /** Processor base class to allow systems to add their own find/replace functionality */
-UCLASS(Abstract)
-class PERSONA_API UAnimAssetFindReplaceProcessor : public UObject
+UCLASS(MinimalAPI, Abstract)
+class UAnimAssetFindReplaceProcessor : public UObject
 {
 	GENERATED_BODY()
 
@@ -41,7 +43,7 @@ class PERSONA_API UAnimAssetFindReplaceProcessor : public UObject
 
 private:
 	/** Initially setup the processor, binding it to the specified widget */
-	void Initialize(TSharedRef<SAnimAssetFindReplace> InWidget);
+	UE_API void Initialize(TSharedRef<SAnimAssetFindReplace> InWidget);
 
 public:
 	/** Whether this processor supports the specified mode */
@@ -91,13 +93,13 @@ protected:
 	void RequestUIRefresh() { RequestRefreshUI(); }
 
 	/** Requests the entire find UI to be refreshed */
-	void RequestRefreshUI();
+	UE_API void RequestRefreshUI();
 
 	/** Requests cached data and search results */
-	void RequestRefreshCachedData();
+	UE_API void RequestRefreshCachedData();
 
 	/** Requests search results to be refreshed */
-	void RequestRefreshSearchResults();
+	UE_API void RequestRefreshSearchResults();
 
 private:
 	/** Weak ptr to the find/replace widget */
@@ -105,8 +107,8 @@ private:
 };
 
 /** Processor for string-based animation asset operations */
-UCLASS(Abstract)
-class PERSONA_API UAnimAssetFindReplaceProcessor_StringBase : public UAnimAssetFindReplaceProcessor
+UCLASS(MinimalAPI, Abstract)
+class UAnimAssetFindReplaceProcessor_StringBase : public UAnimAssetFindReplaceProcessor
 {
 	GENERATED_BODY()
 
@@ -116,7 +118,7 @@ public:
 	{}
 
 	/** Set the current find string */
-	void SetFindString(const FString& InString);
+	UE_API void SetFindString(const FString& InString);
 
 	/** Get the current find string */
 	FStringView GetFindString() const { return FindString; }
@@ -125,7 +127,7 @@ public:
 	const FString& GetFindStringRef() const { return FindString; }
 
 	/** Set the current replace string */
-	void SetReplaceString(const FString& InString);
+	UE_API void SetReplaceString(const FString& InString);
 
 	/** Get the current replace string */
 	FStringView GetReplaceString() const { return ReplaceString; }
@@ -134,19 +136,19 @@ public:
 	const FString& GetReplaceStringRef() const { return ReplaceString; }
 	
 	/** Set the current skeleton filter */
-	void SetSkeletonFilter(const FAssetData& InSkeletonFilter);
+	UE_API void SetSkeletonFilter(const FAssetData& InSkeletonFilter);
 
 	/** Get the current skeleton filter **/
 	const FAssetData& GetSkeletonFilter() const { return SkeletonFilter; }
 
 	/** Sets whether to match whole words when searching strings */
-	void SetFindWholeWord(bool bInFindWholeWord);
+	UE_API void SetFindWholeWord(bool bInFindWholeWord);
 
 	/** Gets whether to match whole words when searching strings */
 	bool GetFindWholeWord() const { return bFindWholeWord; }
 
 	/** Sets case comparison to use when searching */
-	void SetSearchCase(ESearchCase::Type InSearchCase);
+	UE_API void SetSearchCase(ESearchCase::Type InSearchCase);
 
 	/** Gets case comparison to use when searching */
 	ESearchCase::Type GetSearchCase() const { return SearchCase; }
@@ -157,16 +159,16 @@ public:
 protected:
 	/** UAnimAssetFindReplaceProcessor interface */
 	virtual bool SupportsMode(EAnimAssetFindReplaceMode InMode) const override { return true; }
-	virtual void ExtendToolbar(FToolMenuSection& InSection) override;
-	virtual TSharedRef<SWidget> MakeFindReplaceWidget() override;
-	virtual void FocusInitialWidget() override;
-	virtual bool CanCurrentlyReplace() const override;
-	virtual bool CanCurrentlyRemove() const override;
-	virtual void RefreshCachedData() override;
+	UE_API virtual void ExtendToolbar(FToolMenuSection& InSection) override;
+	UE_API virtual TSharedRef<SWidget> MakeFindReplaceWidget() override;
+	UE_API virtual void FocusInitialWidget() override;
+	UE_API virtual bool CanCurrentlyReplace() const override;
+	UE_API virtual bool CanCurrentlyRemove() const override;
+	UE_API virtual void RefreshCachedData() override;
 
 protected:
 	/** Check if the supplied string matches the find parameters */
-	bool NameMatches(FStringView InNameStringView) const;
+	UE_API bool NameMatches(FStringView InNameStringView) const;
 
 private:
 	/** The find string */
@@ -251,4 +253,6 @@ public:
 	{
 		return Cast<ProcessorType>(GetProcessor(ProcessorType::StaticClass()));
 	}
-};
\ No newline at end of file
+};
+
+#undef UE_API
diff --git a/Engine/Source/Editor/Persona/Public/AnimationEditorPreviewActor.h b/Engine/Source/Editor/Persona/Public/AnimationEditorPreviewActor.h
--- a/Engine/Source/Editor/Persona/Public/AnimationEditorPreviewActor.h
+++ b/Engine/Source/Editor/Persona/Public/AnimationEditorPreviewActor.h
@@ -5,12 +5,16 @@
 #include "GameFramework/Actor.h"
 #include "AnimationEditorPreviewActor.generated.h"
 
-UCLASS()
-class PERSONA_API AAnimationEditorPreviewActor : public AActor
+#define UE_API PERSONA_API
+
+UCLASS(MinimalAPI)
+class AAnimationEditorPreviewActor : public AActor
 {
 	GENERATED_BODY()
 
 public:
 	/** AActor interface */
-	virtual void K2_DestroyActor() override;
+	UE_API virtual void K2_DestroyActor() override;
 };
+
+#undef UE_API
diff --git a/Engine/Source/Editor/Persona/Public/AnimationEditorViewportClient.h b/Engine/Source/Editor/Persona/Public/AnimationEditorViewportClient.h
--- a/Engine/Source/Editor/Persona/Public/AnimationEditorViewportClient.h
+++ b/Engine/Source/Editor/Persona/Public/AnimationEditorViewportClient.h
@@ -16,6 +16,8 @@
 #include "Preferences/PersonaOptions.h"
 #include "SkeletalDebugRendering.h"
 
+#define UE_API PERSONA_API
+
 class FCanvas;
 class UPersonaOptions;
 class USkeletalMeshSocket;
@@ -28,199 +30,199 @@ DECLARE_DELEGATE_RetVal(float, FOnGetBoneSize)
 /////////////////////////////////////////////////////////////////////////
 // FAnimationViewportClient
 
-class PERSONA_API FAnimationViewportClient : public FEditorViewportClient, public TSharedFromThis<FAnimationViewportClient>
+class FAnimationViewportClient : public FEditorViewportClient, public TSharedFromThis<FAnimationViewportClient>
 {
 protected:
 
 	/** Function to display bone names*/
-	void ShowBoneNames(FCanvas* Canvas, FSceneView* View, UDebugSkelMeshComponent* MeshComponent);
+	UE_API void ShowBoneNames(FCanvas* Canvas, FSceneView* View, UDebugSkelMeshComponent* MeshComponent);
 
 	/** Function to display transform attribute names*/
-	void ShowAttributeNames(FCanvas* Canvas, FSceneView* View, UDebugSkelMeshComponent* MeshComponent) const;
+	UE_API void ShowAttributeNames(FCanvas* Canvas, FSceneView* View, UDebugSkelMeshComponent* MeshComponent) const;
 
 	/** Function to display debug lines generated from skeletal controls in animBP mode */
-	void DrawNodeDebugLines(TArray<FText>& Lines, FCanvas* Canvas, FSceneView* View);
+	UE_API void DrawNodeDebugLines(TArray<FText>& Lines, FCanvas* Canvas, FSceneView* View);
 
 public:
-	FAnimationViewportClient(const TSharedRef<class IPersonaPreviewScene>& InPreviewScene, const TSharedRef<class SAnimationEditorViewport>& InAnimationEditorViewport, const TSharedRef<class FAssetEditorToolkit>& InAssetEditorToolkit, int32 InViewportIndex, bool bInShowStats);
-	virtual ~FAnimationViewportClient() override;
+	UE_API FAnimationViewportClient(const TSharedRef<class IPersonaPreviewScene>& InPreviewScene, const TSharedRef<class SAnimationEditorViewport>& InAnimationEditorViewport, const TSharedRef<class FAssetEditorToolkit>& InAssetEditorToolkit, int32 InViewportIndex, bool bInShowStats);
+	UE_API virtual ~FAnimationViewportClient() override;
 
-	void Initialize();
+	UE_API void Initialize();
 
 	// FEditorViewportClient interface
-	virtual void Tick(float DeltaSeconds) override;
-	virtual void Draw(const FSceneView* View,FPrimitiveDrawInterface* PDI) override;
-	virtual void DrawCanvas( FViewport& InViewport, FSceneView& View, FCanvas& Canvas ) override;
-	virtual bool InputKey(const FInputKeyEventArgs& EventArgs) override;
-	virtual bool InputAxis(const FInputKeyEventArgs& EventArgs) override;
+	UE_API virtual void Tick(float DeltaSeconds) override;
+	UE_API virtual void Draw(const FSceneView* View,FPrimitiveDrawInterface* PDI) override;
+	UE_API virtual void DrawCanvas( FViewport& InViewport, FSceneView& View, FCanvas& Canvas ) override;
+	UE_API virtual bool InputKey(const FInputKeyEventArgs& EventArgs) override;
+	UE_API virtual bool InputAxis(const FInputKeyEventArgs& EventArgs) override;
 
 	// Sets what bones are drawn by DrawMeshBones and ShowBoneNames
-	virtual void UpdateBonesToDraw();
+	UE_API virtual void UpdateBonesToDraw();
 	
 //	virtual void ProcessClick(class FSceneView& View, class HHitProxy* HitProxy, FKey Key, EInputEvent Event, uint32 HitX, uint32 HitY) override;
 //	virtual bool InputWidgetDelta( FViewport* Viewport, EAxisList::Type CurrentAxis, FVector& Drag, FRotator& Rot, FVector& Scale ) override;
-	virtual void TrackingStarted( const struct FInputEventState& InInputState, bool bIsDragging, bool bNudge ) override;
-	virtual void TrackingStopped() override;
+	UE_API virtual void TrackingStarted( const struct FInputEventState& InInputState, bool bIsDragging, bool bNudge ) override;
+	UE_API virtual void TrackingStopped() override;
 //	virtual UE::Widget::EWidgetMode GetWidgetMode() const override;
 //	virtual void SetWidgetMode(UE::Widget::EWidgetMode InWidgetMode) override;
 //	virtual bool CanSetWidgetMode(UE::Widget::EWidgetMode NewMode) const override;
-	virtual FVector GetWidgetLocation() const override;
-	virtual FMatrix GetWidgetCoordSystem() const override;
-	virtual ECoordSystem GetWidgetCoordSystemSpace() const override;
-	virtual void SetWidgetCoordSystemSpace(ECoordSystem NewCoordSystem) override;
-	virtual void SetViewMode(EViewModeIndex InViewModeIndex) override;
-	virtual void SetViewportType(ELevelViewportType InViewportType) override;
-	virtual void RotateViewportType() override;
-	virtual bool CanCycleWidgetMode() const override;
-	virtual void SetupViewForRendering( FSceneViewFamily& ViewFamily, FSceneView& View ) override;
-	virtual void HandleToggleShowFlag(FEngineShowFlags::EShowFlag EngineShowFlagIndex) override;
-	virtual FMatrix CalcViewRotationMatrix(const FRotator& InViewRotation) const override;
+	UE_API virtual FVector GetWidgetLocation() const override;
+	UE_API virtual FMatrix GetWidgetCoordSystem() const override;
+	UE_API virtual ECoordSystem GetWidgetCoordSystemSpace() const override;
+	UE_API virtual void SetWidgetCoordSystemSpace(ECoordSystem NewCoordSystem) override;
+	UE_API virtual void SetViewMode(EViewModeIndex InViewModeIndex) override;
+	UE_API virtual void SetViewportType(ELevelViewportType InViewportType) override;
+	UE_API virtual void RotateViewportType() override;
+	UE_API virtual bool CanCycleWidgetMode() const override;
+	UE_API virtual void SetupViewForRendering( FSceneViewFamily& ViewFamily, FSceneView& View ) override;
+	UE_API virtual void HandleToggleShowFlag(FEngineShowFlags::EShowFlag EngineShowFlagIndex) override;
+	UE_API virtual FMatrix CalcViewRotationMatrix(const FRotator& InViewRotation) const override;
 	// End of FEditorViewportClient interface
 
 	/** Draw call to render UV overlay */
-	void DrawUVsForMesh(FViewport* InViewport, FCanvas* InCanvas, int32 InTextYPos, UDebugSkelMeshComponent* MeshComponent);
+	UE_API void DrawUVsForMesh(FViewport* InViewport, FCanvas* InCanvas, int32 InTextYPos, UDebugSkelMeshComponent* MeshComponent);
 
 	/** Set the camera follow mode */
-	void SetCameraFollowMode(EAnimationViewportCameraFollowMode Mode, FName InBoneName = NAME_None);
+	UE_API void SetCameraFollowMode(EAnimationViewportCameraFollowMode Mode, FName InBoneName = NAME_None);
 
 	/** Called when viewport focuses on a selection */
-	void OnFocusViewportToSelection();
+	UE_API void OnFocusViewportToSelection();
 
 	/** Get the camera follow mode */
-	EAnimationViewportCameraFollowMode GetCameraFollowMode() const;
+	UE_API EAnimationViewportCameraFollowMode GetCameraFollowMode() const;
 
 	/** Get the bone name to use when CameraFollowMode is EAnimationViewportCameraFollowMode::Bone */
-	FName GetCameraFollowBoneName() const;
+	UE_API FName GetCameraFollowBoneName() const;
 
 	/** Jump to the meshes default camera */
-	void JumpToDefaultCamera();
+	UE_API void JumpToDefaultCamera();
 
 	/** Save current camera as default for mesh */
-	void SaveCameraAsDefault();
+	UE_API void SaveCameraAsDefault();
 
 	/** Check whether we can save this camera as default */
-	bool CanSaveCameraAsDefault() const;
+	UE_API bool CanSaveCameraAsDefault() const;
 
 	/** Clear any default camera for mesh */
-	void ClearDefaultCamera();
+	UE_API void ClearDefaultCamera();
 
 	/** Returns whether we have a default camera set for this mesh */
-	bool HasDefaultCameraSet() const;
+	UE_API bool HasDefaultCameraSet() const;
 
 	/** Handle the skeletal mesh mesh component being used for preview changing */
-	void HandleSkeletalMeshChanged(class USkeletalMesh* OldSkeletalMesh, class USkeletalMesh* NewSkeletalMesh);
+	UE_API void HandleSkeletalMeshChanged(class USkeletalMesh* OldSkeletalMesh, class USkeletalMesh* NewSkeletalMesh);
 
 	/** Handle a change in the skeletal mesh mesh component being used for preview changing */
-	void HandleOnMeshChanged();
+	UE_API void HandleOnMeshChanged();
 
 	/** Handle a change in the skeletal mesh phusics component being used for preview changing */
-	void HandleOnSkelMeshPhysicsCreated();
+	UE_API void HandleOnSkelMeshPhysicsCreated();
 
 	/** Function to display bone names*/
-	void ShowBoneNames(FViewport* Viewport, FCanvas* Canvas, UDebugSkelMeshComponent* MeshComponent);
+	UE_API void ShowBoneNames(FViewport* Viewport, FCanvas* Canvas, UDebugSkelMeshComponent* MeshComponent);
 
 	/** Function to enable/disable floor auto align */
-	void OnToggleAutoAlignFloor();
+	UE_API void OnToggleAutoAlignFloor();
 
 	/** Function to check whether floor is auto align or not */
-	bool IsAutoAlignFloor() const;
+	UE_API bool IsAutoAlignFloor() const;
 
 	/** Function to mute/unmute audio in the viewport */
-	void OnToggleMuteAudio();
+	UE_API void OnToggleMuteAudio();
 
 	/** Function to check whether audio is muted or not */
-	bool IsAudioMuted() const;
+	UE_API bool IsAudioMuted() const;
 
 	/** Set whether to use audio attenuation */
-	void OnToggleUseAudioAttenuation();
+	UE_API void OnToggleUseAudioAttenuation();
 
 	/** Check whether we are using audio attenuation */
-	bool IsUsingAudioAttenuation() const;
+	UE_API bool IsUsingAudioAttenuation() const;
 
 	/** Function to set background color */
-	void SetBackgroundColor(FLinearColor InColor);
+	UE_API void SetBackgroundColor(FLinearColor InColor);
 
 	/** Function to get current brightness value */ 
-	float GetBrightnessValue() const;
+	UE_API float GetBrightnessValue() const;
 
 	/** Function to set brightness value */
-	void SetBrightnessValue(float Value);
+	UE_API void SetBrightnessValue(float Value);
 
 	/** Function to set Local axes mode for the ELocalAxesType */
-	void SetLocalAxesMode(ELocalAxesMode::Type AxesMode);
+	UE_API void SetLocalAxesMode(ELocalAxesMode::Type AxesMode);
 
 	/** Local axes mode checking function for the ELocalAxesType */
-	bool IsLocalAxesModeSet(ELocalAxesMode::Type AxesMode) const;
+	UE_API bool IsLocalAxesModeSet(ELocalAxesMode::Type AxesMode) const;
 
 	/** Get the Bone local axis mode */
-	ELocalAxesMode::Type GetLocalAxesMode() const;
+	UE_API ELocalAxesMode::Type GetLocalAxesMode() const;
 
 	/** Access Bone Draw size config option*/
-	void SetBoneDrawSize(const float InBoneDrawSize);
-	float GetBoneDrawSize() const;
+	UE_API void SetBoneDrawSize(const float InBoneDrawSize);
+	UE_API float GetBoneDrawSize() const;
 
 	/** Access CustomAnimationSpeed config option*/
 	UE_DEPRECATED(5.6, "Animation Speed is now handled by FAnimationEditorPreviewScene. Please use FAnimationEditorPreviewScene::SetCustomAnimationSpeed instead.")
-	void SetCustomAnimationSpeed(const float InCustomAnimationSpeed);
+	UE_API void SetCustomAnimationSpeed(const float InCustomAnimationSpeed);
 	UE_DEPRECATED(5.6, "Animation Speed is now handled by FAnimationEditorPreviewScene. Please use FAnimationEditorPreviewScene::GetCustomAnimationSpeed instead.")
-	float GetCustomAnimationSpeed() const;
+	UE_API float GetCustomAnimationSpeed() const;
 
 	/** Function to set Bone Draw  mode for the EBoneDrawType */
-	void SetBoneDrawMode(EBoneDrawMode::Type AxesMode);
+	UE_API void SetBoneDrawMode(EBoneDrawMode::Type AxesMode);
 
 	/** Bone Draw  mode checking function for the EBoneDrawType */
-	bool IsBoneDrawModeSet(EBoneDrawMode::Type AxesMode) const;
+	UE_API bool IsBoneDrawModeSet(EBoneDrawMode::Type AxesMode) const;
 
 	/** Get the Bone local axis mode */
-	EBoneDrawMode::Type GetBoneDrawMode() const;
+	UE_API EBoneDrawMode::Type GetBoneDrawMode() const;
 	
 	/** Returns the desired target of the camera */
-	FSphere GetCameraTarget();
+	UE_API FSphere GetCameraTarget();
 
 	/** Sets up the viewports camera (look-at etc) based on the current preview target*/
-	void UpdateCameraSetup();
+	UE_API void UpdateCameraSetup();
 
 	/* Places the viewport camera at a good location to view the supplied sphere */
-	void FocusViewportOnSphere( FSphere& Sphere, bool bInstant = true );
+	UE_API void FocusViewportOnSphere( FSphere& Sphere, bool bInstant = true );
 
 	/* Places the viewport camera at a good location to view the preview target */
-	void FocusViewportOnPreviewMesh(bool bUseCustomCamera);
+	UE_API void FocusViewportOnPreviewMesh(bool bUseCustomCamera);
 
 	/** Callback for toggling the normals show flag. */
-	void ToggleCPUSkinning();
+	UE_API void ToggleCPUSkinning();
 
 	/** Callback for checking the normals show flag. */
-	bool IsSetCPUSkinningChecked() const;
+	UE_API bool IsSetCPUSkinningChecked() const;
 
 	/** Toggles whether to lock the camera's rotation to a specified bone's orientation */
-	void ToggleRotateCameraToFollowBone();
+	UE_API void ToggleRotateCameraToFollowBone();
 
 	/** Whether or not to lock the camera's rotation to a specified bone's orientation */
-	bool GetShouldRotateCameraToFollowBone() const;
+	UE_API bool GetShouldRotateCameraToFollowBone() const;
 
 	/** Callback for toggling the normals show flag. */
-	void ToggleShowNormals();
+	UE_API void ToggleShowNormals();
 
 	/** Callback for checking the normals show flag. */
-	bool IsSetShowNormalsChecked() const;
+	UE_API bool IsSetShowNormalsChecked() const;
 
 	/** Callback for toggling the tangents show flag. */
-	void ToggleShowTangents();
+	UE_API void ToggleShowTangents();
 
 	/** Callback for checking the tangents show flag. */
-	bool IsSetShowTangentsChecked() const;
+	UE_API bool IsSetShowTangentsChecked() const;
 
 	/** Callback for toggling the binormals show flag. */
-	void ToggleShowBinormals();
+	UE_API void ToggleShowBinormals();
 
 	/** Callback for checking the binormals show flag. */
-	bool IsSetShowBinormalsChecked() const;
+	UE_API bool IsSetShowBinormalsChecked() const;
 
 	/** Callback for toggling UV drawing in the viewport */
-	void SetDrawUVOverlay(bool bInDrawUVs);
+	UE_API void SetDrawUVOverlay(bool bInDrawUVs);
 
 	/** Callback for checking whether the UV drawing is switched on. */
-	bool IsSetDrawUVOverlayChecked() const;
+	UE_API bool IsSetDrawUVOverlayChecked() const;
 
 	/** Returns the UV Channel that will be drawn when Draw UV Overlay is turned on */
 	int32 GetUVChannelToDraw() const { return UVChannelToDraw; }
@@ -229,29 +231,29 @@ public:
 	void SetUVChannelToDraw(int32 UVChannel) { UVChannelToDraw = UVChannel; }
 
 	/* Returns the floor height offset */	
-	float GetFloorOffset() const;
+	UE_API float GetFloorOffset() const;
 
 	/* Sets the floor height offset, saves it to config and invalidates the viewport so it shows up immediately */
-	void SetFloorOffset(float NewValue);
+	UE_API void SetFloorOffset(float NewValue);
 
 	/** Function to set mesh stat drawing state */
-	void OnSetShowMeshStats(int32 ShowMode);
+	UE_API void OnSetShowMeshStats(int32 ShowMode);
 	/** Whether or not mesh stats are being displayed */
-	bool IsShowingMeshStats() const;
+	UE_API bool IsShowingMeshStats() const;
 	/** Whether or not selected node stats are being displayed */
-	bool IsShowingSelectedNodeStats() const;
+	UE_API bool IsShowingSelectedNodeStats() const;
 	/** Whether detailed mesh stats are being displayed or basic mesh stats */
-	bool IsDetailedMeshStats() const;
+	UE_API bool IsDetailedMeshStats() const;
 
-	int32 GetShowMeshStats() const;
+	UE_API int32 GetShowMeshStats() const;
 
 	/** Set the playback speed mode */
 	UE_DEPRECATED(5.6, "Animation Speed is now handled by FAnimationEditorPreviewScene. Please use FAnimationEditorPreviewScene::SetAnimationPlaybackSpeedMode instead.")
-	void SetPlaybackSpeedMode(EAnimationPlaybackSpeeds::Type InMode);
+	UE_API void SetPlaybackSpeedMode(EAnimationPlaybackSpeeds::Type InMode);
 
 	/** Get the playback speed mode */
 	UE_DEPRECATED(5.6, "Animation Speed is now handled by FAnimationEditorPreviewScene. Please use FAnimationEditorPreviewScene::GetAnimationPlaybackSpeedMode instead.")
-	EAnimationPlaybackSpeeds::Type GetPlaybackSpeedMode() const;
+	UE_API EAnimationPlaybackSpeeds::Type GetPlaybackSpeedMode() const;
 
 	/** Get the preview scene we are viewing */
 	TSharedRef<class IPersonaPreviewScene> GetPreviewScene() const { return PreviewScenePtr.ToSharedRef(); }
@@ -260,23 +262,23 @@ public:
 	TSharedRef<class FAssetEditorToolkit> GetAssetEditorToolkit() const { return AssetEditorToolkitPtr.Pin().ToSharedRef(); }
 
 	/* Handle error checking for additive base pose */
-	bool ShouldDisplayAdditiveScaleErrorMessage() const;
+	UE_API bool ShouldDisplayAdditiveScaleErrorMessage() const;
 
 	/** Draws Mesh Sockets in foreground - bUseSkeletonSocketColor = true for grey (skeleton), false for red (mesh) **/
-	static void DrawSockets(const UDebugSkelMeshComponent* InPreviewMeshComponent, TArray<USkeletalMeshSocket*>& InSockets, FSelectedSocketInfo InSelectedSocket, FPrimitiveDrawInterface* PDI, bool bUseSkeletonSocketColor);
+	static UE_API void DrawSockets(const UDebugSkelMeshComponent* InPreviewMeshComponent, TArray<USkeletalMeshSocket*>& InSockets, FSelectedSocketInfo InSelectedSocket, FPrimitiveDrawInterface* PDI, bool bUseSkeletonSocketColor);
 
 	/** Draws Gizmo for the Transform in foreground **/
-	static void RenderGizmo(const FTransform& Transform, FPrimitiveDrawInterface* PDI);
+	static UE_API void RenderGizmo(const FTransform& Transform, FPrimitiveDrawInterface* PDI);
 
 	/** Function to display warning and info text on the viewport when outside of animBP mode */
-	FText GetDisplayInfo( bool bDisplayAllInfo ) const;
+	UE_API FText GetDisplayInfo( bool bDisplayAllInfo ) const;
 
 	/** Get the viewport index (0-3) for this client */
 	int32 GetViewportIndex() const { return ViewportIndex; }
 
 	/** Get the persona mode manager */
 	UE_DEPRECATED(5.1, "Use the UPersonaEditorModeManagerContext object stored in the editor mode tools' context store instead.")
-	class IPersonaEditorModeManager* GetPersonaModeManager() const;
+	UE_API class IPersonaEditorModeManager* GetPersonaModeManager() const;
 
 private:
 	/**
@@ -284,7 +286,7 @@ private:
 	 *
 	 * @param View	The scene view to use when calculate the listener position
 	 */
-	void UpdateAudioListener(const FSceneView& View);
+	UE_API void UpdateAudioListener(const FSceneView& View);
 
 public:
 
@@ -362,28 +364,28 @@ private:
 
 private:
 
-	void SetCameraTargetLocation(const FSphere &BoundSphere, float DeltaSeconds);
+	UE_API void SetCameraTargetLocation(const FSphere &BoundSphere, float DeltaSeconds);
 
 	/** Draws Mesh Bones in foreground **/
-	void DrawMeshBones(UDebugSkelMeshComponent * MeshComponent, FPrimitiveDrawInterface* PDI) const;
+	UE_API void DrawMeshBones(UDebugSkelMeshComponent * MeshComponent, FPrimitiveDrawInterface* PDI) const;
 	/** Draws the given array of transforms as bones */
-	void DrawBonesFromTransforms(TArray<FTransform>& Transforms, UDebugSkelMeshComponent * MeshComponent, FPrimitiveDrawInterface* PDI, FLinearColor BoneColour, FLinearColor RootBoneColour) const;
+	UE_API void DrawBonesFromTransforms(TArray<FTransform>& Transforms, UDebugSkelMeshComponent * MeshComponent, FPrimitiveDrawInterface* PDI, FLinearColor BoneColour, FLinearColor RootBoneColour) const;
 	/** Draws Bones for a compact pose */
-	void DrawBonesFromCompactPose(const FCompactHeapPose& Pose, UDebugSkelMeshComponent * MeshComponent, FPrimitiveDrawInterface* PDI, const FLinearColor& DrawColor) const;
+	UE_API void DrawBonesFromCompactPose(const FCompactHeapPose& Pose, UDebugSkelMeshComponent * MeshComponent, FPrimitiveDrawInterface* PDI, const FLinearColor& DrawColor) const;
 	/** Draws Bones for uncompressed animation **/
-	void DrawMeshBonesUncompressedAnimation(UDebugSkelMeshComponent * MeshComponent, FPrimitiveDrawInterface* PDI) const;
+	UE_API void DrawMeshBonesUncompressedAnimation(UDebugSkelMeshComponent * MeshComponent, FPrimitiveDrawInterface* PDI) const;
 	/** Draw Bones for non retargeted animation. */
-	void DrawMeshBonesNonRetargetedAnimation(UDebugSkelMeshComponent * MeshComponent, FPrimitiveDrawInterface* PDI) const;
+	UE_API void DrawMeshBonesNonRetargetedAnimation(UDebugSkelMeshComponent * MeshComponent, FPrimitiveDrawInterface* PDI) const;
 	/** Draws Bones for Additive Base Pose */
-	void DrawMeshBonesAdditiveBasePose(UDebugSkelMeshComponent * MeshComponent, FPrimitiveDrawInterface* PDI) const;
+	UE_API void DrawMeshBonesAdditiveBasePose(UDebugSkelMeshComponent * MeshComponent, FPrimitiveDrawInterface* PDI) const;
 	/** Draw Bones for non retargeted animation. */
-	void DrawMeshBonesSourceRawAnimation(UDebugSkelMeshComponent * MeshComponent, FPrimitiveDrawInterface* PDI) const;
+	UE_API void DrawMeshBonesSourceRawAnimation(UDebugSkelMeshComponent * MeshComponent, FPrimitiveDrawInterface* PDI) const;
 	/** Draw Bones for non retargeted animation. */
-	void DrawMeshBonesBakedAnimation(UDebugSkelMeshComponent * MeshComponent, FPrimitiveDrawInterface* PDI) const;
+	UE_API void DrawMeshBonesBakedAnimation(UDebugSkelMeshComponent * MeshComponent, FPrimitiveDrawInterface* PDI) const;
 	/** Draw Bones from skeleton reference pose. */
-	void DrawBonesFromSkeleton(UDebugSkelMeshComponent * MeshComponent, const USkeleton* Skeleton, const TArray<int32>& InSelectedBones, FPrimitiveDrawInterface* PDI) const;
+	UE_API void DrawBonesFromSkeleton(UDebugSkelMeshComponent * MeshComponent, const USkeleton* Skeleton, const TArray<int32>& InSelectedBones, FPrimitiveDrawInterface* PDI) const;
 	/** Draws Bones for RequiredBones with WorldTransform **/
-	void DrawBones(
+	UE_API void DrawBones(
 		const FVector& ComponentOrigin,
 		const TArray<FBoneIndexType>& RequiredBones,
 		const FReferenceSkeleton& RefSkeleton,
@@ -396,55 +398,55 @@ private:
 		bool bUseMultiColors) const;
 
 	/** Draws active transform attributes */
-	void DrawAttributes(UDebugSkelMeshComponent* MeshComponent, FPrimitiveDrawInterface* PDI) const;
+	UE_API void DrawAttributes(UDebugSkelMeshComponent* MeshComponent, FPrimitiveDrawInterface* PDI) const;
 
 	/** Draws visualization from animation notifies into viewport. */
-	void DrawNotifies(UDebugSkelMeshComponent* MeshComponent, FPrimitiveDrawInterface* PDI) const;
+	UE_API void DrawNotifies(UDebugSkelMeshComponent* MeshComponent, FPrimitiveDrawInterface* PDI) const;
 
 	/* Draws visualization from animation notifies into canvas. */
-	void DrawCanvasNotifies(UDebugSkelMeshComponent* MeshComponent, FCanvas& Canvas, FSceneView& View) const;
+	UE_API void DrawCanvasNotifies(UDebugSkelMeshComponent* MeshComponent, FCanvas& Canvas, FSceneView& View) const;
 
 	/* Draws visualization from Asset User Data into viewport. */
-	void DrawAssetUserData(FPrimitiveDrawInterface* PDI) const;
+	UE_API void DrawAssetUserData(FPrimitiveDrawInterface* PDI) const;
 
 	/* Draws visualization from Asset User Data into canvas. */
-	void DrawCanvasAssetUserData(FCanvas& Canvas, FSceneView& View) const;
+	UE_API void DrawCanvasAssetUserData(FCanvas& Canvas, FSceneView& View) const;
 
 	/** Draws root motion trajectory */
-	void DrawRootMotionTrajectory(UDebugSkelMeshComponent* MeshComponent, FPrimitiveDrawInterface* PDI) const;
+	UE_API void DrawRootMotionTrajectory(UDebugSkelMeshComponent* MeshComponent, FPrimitiveDrawInterface* PDI) const;
 
 	/** Draws bones from watched poses*/
-	void DrawWatchedPoses(UDebugSkelMeshComponent * MeshComponent, FPrimitiveDrawInterface* PDI);
+	UE_API void DrawWatchedPoses(UDebugSkelMeshComponent * MeshComponent, FPrimitiveDrawInterface* PDI);
 
 	/** Get the typed anim preview scene shared ptr*/
-	TSharedPtr<class FAnimationEditorPreviewScene> GetAnimPreviewScenePtr() const;
+	UE_API TSharedPtr<class FAnimationEditorPreviewScene> GetAnimPreviewScenePtr() const;
 
 	/** Get the typed anim preview scene */
-	TSharedRef<class FAnimationEditorPreviewScene> GetAnimPreviewScene() const;
+	UE_API TSharedRef<class FAnimationEditorPreviewScene> GetAnimPreviewScene() const;
 
 	/** Invalidate this view in response to a preview scene change */
-	void HandleInvalidateViews();
+	UE_API void HandleInvalidateViews();
 
 	/** Handle the view being focused from the preview scene */
-	void HandleFocusViews();
+	UE_API void HandleFocusViews();
 
 	/** Computes a bounding box for the selected section of the preview mesh component.
 	    If there is no selected section, returns an empty box. */
-	FBox ComputeBoundingBoxForSelectedEditorSection() const;
+	UE_API FBox ComputeBoundingBoxForSelectedEditorSection() const;
 
 	/** Converts all local vertex positions to world positions. */
-	void TransformVertexPositionsToWorld(TArray<FFinalSkinVertex>& LocalVertices) const;
+	UE_API void TransformVertexPositionsToWorld(TArray<FFinalSkinVertex>& LocalVertices) const;
 
 	/** Gets all vertex indices that the given section references. */
-	void GetAllVertexIndicesUsedInSection(const FRawStaticIndexBuffer16or32Interface& IndexBuffer,
+	UE_API void GetAllVertexIndicesUsedInSection(const FRawStaticIndexBuffer16or32Interface& IndexBuffer,
 										  const FSkelMeshRenderSection& SkelMeshSection,
 										  TArray<int32>& OutIndices) const;
 
 	/** Used for camera tracking - store data about the scene pre-tick */
-	void HandlePreviewScenePreTick();
+	UE_API void HandlePreviewScenePreTick();
 
 	/** Used for camera tracking - use stored data about the scene post-tick */
-	void HandlePreviewScenePostTick();
+	UE_API void HandlePreviewScenePostTick();
 
 private:
 	struct FTimecodeDisplayInfo
@@ -454,7 +456,7 @@ private:
 	};
 
 	/** @return array of AssetUserData interfaces from editable objects on associated asset toolkit. */
-	TArray<IInterface_AssetUserData*> GetEditedObjectsWithAssetUserData() const;
+	UE_API TArray<IInterface_AssetUserData*> GetEditedObjectsWithAssetUserData() const;
 
 	/** Size to draw bones in the viewport. Transient setting. */
 	float BoneDrawSize=1.0f;
@@ -478,7 +480,7 @@ private:
 	FVector LastLookAtLocation;
 
 	// Delegate Handler to allow changing of camera controller
-	void OnCameraControllerChanged();
+	UE_API void OnCameraControllerChanged();
 
 	/** True when the preview animation should resume playing upon finishing tracking */
 	bool bResumeAfterTracking;
@@ -486,3 +488,5 @@ private:
 	/** Timecode/slate information from current animation sequence. */
 	TOptional<FTimecodeDisplayInfo> TimecodeDisplay;
 };
+
+#undef UE_API
diff --git a/Engine/Source/Editor/Persona/Public/AnimationEditorsAssetFamilyExtension.h b/Engine/Source/Editor/Persona/Public/AnimationEditorsAssetFamilyExtension.h
--- a/Engine/Source/Editor/Persona/Public/AnimationEditorsAssetFamilyExtension.h
+++ b/Engine/Source/Editor/Persona/Public/AnimationEditorsAssetFamilyExtension.h
@@ -57,8 +57,8 @@ struct IAnimationEditorsAssetFamilyInterface
 	}
 };
 
-UCLASS()
-class PERSONA_API UAnimationEditorsAssetFamilyExtension : public UObject
+UCLASS(MinimalAPI)
+class UAnimationEditorsAssetFamilyExtension : public UObject
 {
 	GENERATED_BODY()
 
@@ -83,4 +83,4 @@ public:
 
 	// Gets the horizontal position of this asset
 	virtual void GetPosition(FName& OutBeforeClass, FName& OutAfterClass) const {};
-};
\ No newline at end of file
+};
diff --git a/Engine/Source/Editor/Persona/Public/AnimationSequenceBrowserMenuContexts.h b/Engine/Source/Editor/Persona/Public/AnimationSequenceBrowserMenuContexts.h
--- a/Engine/Source/Editor/Persona/Public/AnimationSequenceBrowserMenuContexts.h
+++ b/Engine/Source/Editor/Persona/Public/AnimationSequenceBrowserMenuContexts.h
@@ -10,8 +10,8 @@
 
 class SAnimationSequenceBrowser;
 
-UCLASS(BlueprintType)
-class PERSONA_API UAnimationSequenceBrowserContextMenuContext : public UObject
+UCLASS(MinimalAPI, BlueprintType)
+class UAnimationSequenceBrowserContextMenuContext : public UObject
 {
 	GENERATED_BODY()
 
