Fix window scaling issue when dragging multiple notifies in animation editor

diff --git a/Engine/Source/Editor/Persona/Private/SAnimNotifyPanel.cpp b/Engine/Source/Editor/Persona/Private/SAnimNotifyPanel.cpp
--- a/Engine/Source/Editor/Persona/Private/SAnimNotifyPanel.cpp
+++ b/Engine/Source/Editor/Persona/Private/SAnimNotifyPanel.cpp
@@ -4372,10 +4372,12 @@ FReply SAnimNotifyPanel::OnNotifyNodeDragStarted(TArray<TSharedPtr<SAnimNotifyNo
 	const FVector2D OverlayOrigin = OverlayBounds.Min;
 	const FVector2D OverlayExtents = OverlayBounds.GetSize();
 
+	const FSlateRenderTransform InvTrackTransform = GetCachedGeometry().GetAccumulatedRenderTransform().Inverse();
+	
 	for(const TSharedPtr<SAnimNotifyNode>& Node : Nodes)
 	{
-		const FVector2D OffsetFromFirst(Node->GetScreenPosition() - OverlayOrigin);
-
+		const FVector2D OffsetFromFirst(InvTrackTransform.TransformVector(Node->GetScreenPosition() - OverlayOrigin));
+		
 		NodeDragDecoratorOverlay->AddSlot()
 			.Padding(FMargin(static_cast<float>(OffsetFromFirst.X), static_cast<float>(OffsetFromFirst.Y), 0.0f, 0.0f))
 			[
