Add filter for attributes to anim asset browser

diff --git a/Engine/Source/Editor/Persona/Private/SAnimationSequenceBrowser.cpp b/Engine/Source/Editor/Persona/Private/SAnimationSequenceBrowser.cpp
--- a/Engine/Source/Editor/Persona/Private/SAnimationSequenceBrowser.cpp
+++ b/Engine/Source/Editor/Persona/Private/SAnimationSequenceBrowser.cpp
@@ -47,6 +47,7 @@
 #include "SAnimCurvePicker.h"
 #include "ToolMenus.h"
 #include "AnimationSequenceBrowserMenuContexts.h"
+#include "SAnimAttributePicker.h"
 #include "Viewports.h"
 #include "SceneInterface.h"
 #include "UObject/AssetRegistryTagsContext.h"
@@ -343,6 +344,127 @@ public:
 	TWeakPtr<IEditableSkeleton> EditableSkeleton;
 };
 
+/** A filter that displays animations that use an attribute */
+class FFrontendFilter_Attributes : public FFrontendFilter
+{
+public:
+	FFrontendFilter_Attributes(TSharedPtr<FFrontendFilterCategory> InCategory, TSharedPtr<IEditableSkeleton> InEditableSkeleton)
+		: FFrontendFilter(InCategory)
+		, EditableSkeleton(InEditableSkeleton)
+	{}
+
+	// FFrontendFilter implementation
+	virtual FString GetName() const override
+	{
+		return TEXT("AttributeAssets");
+	}
+
+	virtual FText GetDisplayName() const override
+	{
+		return AttributeName.IsEmpty()
+			? LOCTEXT("FFrontendFilter_AttributeAssetsEmpty", "Uses Attributes")
+			: FText::Format(LOCTEXT("FFrontendFilter_AttributeAssets", "Uses Attribute: {0}"), FText::FromString(AttributeName));
+	}
+
+	virtual FText GetToolTipText() const override
+	{
+		return AttributeName.IsEmpty()
+			? LOCTEXT("FFrontendFilter_AttributeAssetsEmptyTooltip", "Show assets that contains a (specified) attribute")
+			: FText::Format(LOCTEXT("FFrontendFilter_AttributeAssetsTooltip", "Show assets that use attribute '{0}'"), FText::FromString(AttributeName));
+	}
+
+	virtual void ModifyContextMenu(FMenuBuilder& MenuBuilder) override
+	{
+		MenuBuilder.BeginSection(TEXT("AttributeSection"), LOCTEXT("AttributeSectionHeading", "Choose Attribute"));
+		
+		MenuBuilder.AddMenuEntry(LOCTEXT("AllAttributeFilterName", "Any attribute"),
+			LOCTEXT("AllAttributeFilterName_ToolTip", "Consider all animation attributes for filtering."),
+			FSlateIcon(),
+			FUIAction(
+				FExecuteAction::CreateLambda([this]()
+				{
+					AttributeName.Empty();
+					OnChanged().Broadcast();
+				}), 
+				FCanExecuteAction::CreateLambda([this]()
+				{
+					return !AttributeName.IsEmpty();
+				}),
+				FGetActionCheckState::CreateLambda([this]()
+				{
+					return AttributeName.IsEmpty() ? ECheckBoxState::Checked : ECheckBoxState::Unchecked;
+				})
+			), 
+			NAME_None,
+			EUserInterfaceActionType::ToggleButton
+		);
+		
+		TSharedRef<SWidget> Widget = 
+			SNew(SBox)
+			.HeightOverride(300.0f)
+			.WidthOverride(200.0f)
+			[
+				SNew(SAnimAttributePicker, &EditableSkeleton->GetSkeleton())
+				.bEnableMultiSelect(false)
+				.OnAttributesPicked_Lambda([this](const TConstArrayView<FAnimationAttributeIdentifier> InAttributes)
+				{
+					check(InAttributes.Num() == 1);
+					FSlateApplication::Get().DismissAllMenus();
+					AttributeName = InAttributes[0].GetName().ToString();
+					OnChanged().Broadcast();
+				})
+			];
+
+		MenuBuilder.AddWidget(Widget, FText(), true);
+
+		MenuBuilder.EndSection();
+	}
+
+	virtual FLinearColor GetColor() const override
+	{
+		return FLinearColor(0.6f, 0.0f, 0.2f, 1.0f);
+	}
+
+	// IFilter implementation
+	virtual bool PassesFilter(FAssetFilterType InItem) const override
+	{
+		FAssetData ItemAssetData;
+		if (InItem.Legacy_TryGetAssetData(ItemAssetData))
+		{
+			const FString TagValue = ItemAssetData.GetTagValueRef<FString>(USkeleton::AttributeTag);
+			if (!TagValue.IsEmpty())
+			{
+				
+				FAnimationAttributeIdentifierArray AttributesIdentifierArray;
+				FAnimationAttributeIdentifierArray::StaticStruct()->ImportText(*TagValue, &AttributesIdentifierArray, nullptr, PPF_None, GLog, FAnimationAttributeIdentifierArray::StaticStruct()->GetName());
+				
+				if (AttributeName.IsEmpty() && !AttributesIdentifierArray.AttributeIdentifiers.IsEmpty())
+				{
+					return true;
+				}
+				
+				for (const FAnimationAttributeIdentifier& AttrId : AttributesIdentifierArray.AttributeIdentifiers)
+				{
+					if (AttrId.GetName().ToString() == AttributeName)
+					{
+						return true;
+					}
+				}
+			}
+		}
+
+		return false;
+	}
+
+public:
+	/** Attribute name string to use when filtering */
+	FString AttributeName;
+
+	/** Editable skeleton used to access available attributes */
+	TSharedPtr<IEditableSkeleton> EditableSkeleton;
+};
+
+
 /** A filter that displays animations that use a curve */
 class FFrontendFilter_Curves : public FFrontendFilter
 {
@@ -1046,7 +1168,8 @@ void SAnimationSequenceBrowser::Construct(const FArguments& InArgs, const TShare
 	Config.ExtraFrontendFilters.Add(SkeletonNotifyFilter.ToSharedRef());
 	Config.ExtraFrontendFilters.Add(MakeShareable(new FFrontendFilter_Curves(AnimCategory, InPersonaToolkit->GetEditableSkeleton())));
 	Config.ExtraFrontendFilters.Add(MakeShareable(new FFrontendFilter_SkeletonSyncMarker(AnimCategory, InPersonaToolkit->GetEditableSkeleton())));
-
+	Config.ExtraFrontendFilters.Add(MakeShareable(new FFrontendFilter_Attributes(AnimCategory, InPersonaToolkit->GetEditableSkeleton())));
+	
 	Config.OnIsAssetValidForCustomToolTip = FOnIsAssetValidForCustomToolTip::CreateLambda([](const FAssetData& AssetData) {return AssetData.IsAssetLoaded(); });
 	Config.OnGetCustomAssetToolTip = FOnGetCustomAssetToolTip::CreateSP(this, &SAnimationSequenceBrowser::CreateCustomAssetToolTip);
 	Config.OnVisualizeAssetToolTip = FOnVisualizeAssetToolTip::CreateSP(this, &SAnimationSequenceBrowser::OnVisualizeAssetToolTip);
