Persona: Ensure we check if the preview mesh is null when testing to see if we can show Nanite.

#rnx
#DisabledValidations checkCopyright, checkBalancedMacros

diff --git a/Engine/Source/Editor/Persona/Private/SAnimationEditorViewport.cpp b/Engine/Source/Editor/Persona/Private/SAnimationEditorViewport.cpp
--- a/Engine/Source/Editor/Persona/Private/SAnimationEditorViewport.cpp
+++ b/Engine/Source/Editor/Persona/Private/SAnimationEditorViewport.cpp
@@ -1904,7 +1904,8 @@ void SAnimationEditorViewportTabBody::ToggleShowNaniteFallback()
 
 bool SAnimationEditorViewportTabBody::CanShowNaniteFallback() const
 {
-	return GetPreviewScene()->GetPreviewMesh()->IsNaniteEnabled();
+	const USkeletalMesh* PreviewMesh = GetPreviewScene()->GetPreviewMesh();
+	return PreviewMesh && PreviewMesh->IsNaniteEnabled();
 }
 
 bool SAnimationEditorViewportTabBody::IsShowNaniteFallbackEnabled() const
