Cloth - Removed the defunct cloth simulation factory selector.

#rnx

diff --git a/Engine/Source/Editor/Persona/Private/PreviewSceneCustomizations.cpp b/Engine/Source/Editor/Persona/Private/PreviewSceneCustomizations.cpp
--- a/Engine/Source/Editor/Persona/Private/PreviewSceneCustomizations.cpp
+++ b/Engine/Source/Editor/Persona/Private/PreviewSceneCustomizations.cpp
@@ -357,43 +357,6 @@ void FPreviewSceneDescriptionCustomization::CustomizeDetails(IDetailLayoutBuilde
 		LinkedAnimGraphTagProperty->MarkHiddenByCustomization();
 	}
 
-	//
-	// Physics section...
-	//
-#if CHAOS_SIMULATION_DETAIL_VIEW_FACTORY_SELECTOR
-	// Physics settings
-	ClothSimulationFactoryList.Reset();
-	const TArray<IClothingSimulationFactoryClassProvider*> ClassProviders = IModularFeatures::Get().GetModularFeatureImplementations<IClothingSimulationFactoryClassProvider>(IClothingSimulationFactoryClassProvider::FeatureName);
-	for (const auto& ClassProvider : ClassProviders)
-	{
-		// Populate cloth factory list
-		ClothSimulationFactoryList.Add(MakeShared<TSubclassOf<class UClothingSimulationFactory>>(ClassProvider->GetClothingSimulationFactoryClass()));
-	}
-
-	DetailBuilder.EditCategory("Physics")
-	.AddCustomRow(LOCTEXT("PhysicsClothingSimulationFactory", "Clothing Simulation Factory Option"))
-	.RowTag("PhysicsClothingSimulationFactory")
-	.NameContent()
-	[
-		SNew(STextBlock)
-		.Font(IDetailLayoutBuilder::GetDetailFont())
-		.Text(LOCTEXT("PhysicsClothingSimulationFactory_Text", "Clothing Simulation Factory"))
-		.ToolTipText(LOCTEXT("PhysicsClothingSimulationFactory_ToolTip", "Select the cloth simulation used to preview the scene."))
-	]
-	.ValueContent()
-	.MinDesiredWidth(200.0f)
-	[
-		SNew(SComboBox<TSharedPtr<TSubclassOf<class UClothingSimulationFactory>>>)
-		.OptionsSource(&ClothSimulationFactoryList)
-		.OnGenerateWidget(this, &FPreviewSceneDescriptionCustomization::MakeClothingSimulationFactoryWidget)
-		.OnSelectionChanged(this, &FPreviewSceneDescriptionCustomization::OnClothingSimulationFactorySelectionChanged)
-		[
-			SNew(STextBlock)
-			.Text(this, &FPreviewSceneDescriptionCustomization::GetCurrentClothingSimulationFactoryText)
-		]
-	];
-#endif  // #if CHAOS_SIMULATION_DETAIL_VIEW_FACTORY_SELECTOR
-
 	//
 	// Additional Meshes section...
 	//
@@ -857,39 +820,4 @@ void FPreviewMeshCollectionEntryCustomization::HandleMeshesArrayChanged(TSharedP
 	}
 }
 
-#if CHAOS_SIMULATION_DETAIL_VIEW_FACTORY_SELECTOR
-TSharedRef<SWidget> FPreviewSceneDescriptionCustomization::MakeClothingSimulationFactoryWidget(TSharedPtr<TSubclassOf<class UClothingSimulationFactory>> Item) const
-{
-	return SNew(STextBlock)
-		.Text(*Item ? FText::FromName((*Item)->GetFName()) : LOCTEXT("PhysicsClothingSimulationFactory_NoneSelected", "None"))
-		.Font(IDetailLayoutBuilder::GetDetailFont());
-}
-
-void FPreviewSceneDescriptionCustomization::OnClothingSimulationFactorySelectionChanged(TSharedPtr<TSubclassOf<class UClothingSimulationFactory>> Item, ESelectInfo::Type SelectInfo) const
-{
-	// Set new factory to the preview mesh component:
-	if (const TSharedPtr<IPersonaToolkit> PersonaToolkitPin = PersonaToolkit.Pin())
-	{
-		if (UDebugSkelMeshComponent* const DebugSkelMeshComponent = PersonaToolkitPin->GetPreviewMeshComponent())
-		{
-			DebugSkelMeshComponent->UnregisterComponent();
-			DebugSkelMeshComponent->ClothingSimulationFactory = *Item;
-			DebugSkelMeshComponent->RegisterComponent();
-		}
-	}
-}
-
-FText FPreviewSceneDescriptionCustomization::GetCurrentClothingSimulationFactoryText() const
-{
-	TSubclassOf<class UClothingSimulationFactory> Item;
-	if (const TSharedPtr<IPersonaToolkit> PersonaToolkitPin = PersonaToolkit.Pin())
-	{
-		if (const UDebugSkelMeshComponent* const DebugSkelMeshComponent = PersonaToolkitPin->GetPreviewMeshComponent())
-		{
-			Item = DebugSkelMeshComponent->ClothingSimulationFactory;
-		}
-	}
-	return *Item ? FText::FromName((*Item)->GetFName()) : LOCTEXT("PhysicsClothingSimulationFactory_NoneSelected", "None");
-}
-#endif  // #if CHAOS_SIMULATION_DETAIL_VIEW_FACTORY_SELECTOR
 #undef LOCTEXT_NAMESPACE
diff --git a/Engine/Source/Editor/Persona/Private/PreviewSceneCustomizations.h b/Engine/Source/Editor/Persona/Private/PreviewSceneCustomizations.h
--- a/Engine/Source/Editor/Persona/Private/PreviewSceneCustomizations.h
+++ b/Engine/Source/Editor/Persona/Private/PreviewSceneCustomizations.h
@@ -12,10 +12,6 @@
 #include "AnimationEditorPreviewScene.h"
 #include "PersonaModule.h"
 
-#ifndef CHAOS_SIMULATION_DETAIL_VIEW_FACTORY_SELECTOR  // TODO: Decide whether to keep the detail view cloth selector after nvcloth has been removed
-	#define CHAOS_SIMULATION_DETAIL_VIEW_FACTORY_SELECTOR 1
-#endif
-
 struct FAssetData;
 class FDetailWidgetRow;
 class IDetailChildrenBuilder;
@@ -90,17 +86,6 @@ private:
 	// Reinitialize the preview controller in the preview scene.
 	void ReinitializePreviewController();
 
-#if CHAOS_SIMULATION_DETAIL_VIEW_FACTORY_SELECTOR
-	// Make the widget of each item in the preview cloth factory combo box
-	TSharedRef<SWidget> MakeClothingSimulationFactoryWidget(TSharedPtr<TSubclassOf<class UClothingSimulationFactory>> Item) const;
-
-	// Called when the cloth factory combo box selection changes, when a new parameter type is selected
-	void OnClothingSimulationFactorySelectionChanged(TSharedPtr<TSubclassOf<class UClothingSimulationFactory>> Item, ESelectInfo::Type SelectInfo) const;
-
-	// Delegate for getting the current preview cloth factory class name as text
-	FText GetCurrentClothingSimulationFactoryText() const;
-#endif  // #if CHAOS_SIMULATION_DETAIL_VIEW_FACTORY_SELECTOR
-
 private:
 	/** Cached skeleton name to check for asset registry tags */
 	FString SkeletonName;
@@ -123,11 +108,6 @@ private:
 	/** This is list of class available to filter asset by. This list doesn't change once loaded, so only collect once */
 	static TArray<FTopLevelAssetPath> AvailableClassNameList;
 
-#if CHAOS_SIMULATION_DETAIL_VIEW_FACTORY_SELECTOR
-	/** List of available cloth simulation factories. */
-	TArray<TSharedPtr<TSubclassOf<class UClothingSimulationFactory>>> ClothSimulationFactoryList;
-#endif  // #if CHAOS_SIMULATION_DETAIL_VIEW_FACTORY_SELECTOR
-
 	// Our layout builder (cached so we can refresh)
 	IDetailLayoutBuilder* MyDetailLayout;
 
