Pause preview animation while manipulating transform widget in animation editors
* It's pretty impossible to try position something relative to a parent when the parent keeps moving.

diff --git a/Engine/Source/Editor/Persona/Private/SkeletonSelectionEditMode.cpp b/Engine/Source/Editor/Persona/Private/SkeletonSelectionEditMode.cpp
--- a/Engine/Source/Editor/Persona/Private/SkeletonSelectionEditMode.cpp
+++ b/Engine/Source/Editor/Persona/Private/SkeletonSelectionEditMode.cpp
@@ -28,6 +28,7 @@ namespace SkeletonSelectionModeConstants
 
 FSkeletonSelectionEditMode::FSkeletonSelectionEditMode()
 	: bManipulating(false)
+	, bResumeAfterManipulating(false)
 	, bInTransaction(false)
 {
 	// Disable grid drawing for this mode as the viewport handles this
@@ -231,6 +232,14 @@ bool FSkeletonSelectionEditMode::StartTracking(FEditorViewportClient* InViewport
 			}
 
 			bManipulating = true;
+
+			UAnimPreviewInstance* PreviewInstance = PreviewMeshComponent && PreviewMeshComponent->IsPreviewOn() ? PreviewMeshComponent->PreviewInstance : nullptr;
+			if (PreviewInstance)
+			{
+				bResumeAfterManipulating = PreviewInstance->IsPlaying();
+				PreviewInstance->SetPlaying(false);
+			}
+
 			return true;
 		}
 	}
@@ -250,6 +259,22 @@ bool FSkeletonSelectionEditMode::EndTracking(FEditorViewportClient* InViewportCl
 		}
 
 		bManipulating = false;
+
+		if (bResumeAfterManipulating)
+		{
+			bResumeAfterManipulating = false;
+
+			const UDebugSkelMeshComponent* PreviewMeshComponent = GetAnimPreviewScene().GetPreviewMeshComponent();
+			if (PreviewMeshComponent)
+			{
+				UAnimPreviewInstance* PreviewInstance = PreviewMeshComponent && PreviewMeshComponent->IsPreviewOn() ? PreviewMeshComponent->PreviewInstance : nullptr;
+				if (PreviewInstance)
+				{
+					PreviewInstance->SetPlaying(true);
+				}
+			}
+		}
+
 		return true;
 	}
 
diff --git a/Engine/Source/Editor/Persona/Public/SkeletonSelectionEditMode.h b/Engine/Source/Editor/Persona/Public/SkeletonSelectionEditMode.h
--- a/Engine/Source/Editor/Persona/Public/SkeletonSelectionEditMode.h
+++ b/Engine/Source/Editor/Persona/Public/SkeletonSelectionEditMode.h
@@ -68,6 +68,9 @@ private:
 	/** Whether we are currently in a manipulation  */
 	bool bManipulating;	
 
+	/** Whether the preview animation should resume after the manipulating has ended */
+	bool bResumeAfterManipulating;
+
 	/** Whether we are currently in a transaction  */
 	bool bInTransaction;
 };
