Fix drag-dropping segments into composite track not matching selection order

diff --git a/Engine/Source/Editor/Persona/Private/SAnimSegmentsPanel.cpp b/Engine/Source/Editor/Persona/Private/SAnimSegmentsPanel.cpp
--- a/Engine/Source/Editor/Persona/Private/SAnimSegmentsPanel.cpp
+++ b/Engine/Source/Editor/Persona/Private/SAnimSegmentsPanel.cpp
@@ -506,6 +506,7 @@ void SAnimSegmentsPanel::OnTrackDragDrop( TSharedPtr<FDragDropOperation> DragDro
 					if (IsValidToAdd(DroppedSequence))
 					{
 						AddAnimSegment(DroppedSequence, DataPos);
+						DataPos += DroppedSequence->GetPlayLength();
 					}
 					else
 					{
