Fix build health warning

#rnx

diff --git a/Engine/Source/Editor/Persona/Private/SkeletonSelectionEditMode.cpp b/Engine/Source/Editor/Persona/Private/SkeletonSelectionEditMode.cpp
--- a/Engine/Source/Editor/Persona/Private/SkeletonSelectionEditMode.cpp
+++ b/Engine/Source/Editor/Persona/Private/SkeletonSelectionEditMode.cpp
@@ -233,7 +233,7 @@ bool FSkeletonSelectionEditMode::StartTracking(FEditorViewportClient* InViewport
 
 			bManipulating = true;
 
-			UAnimPreviewInstance* PreviewInstance = PreviewMeshComponent && PreviewMeshComponent->IsPreviewOn() ? PreviewMeshComponent->PreviewInstance : nullptr;
+			TObjectPtr<UAnimPreviewInstance> PreviewInstance = (PreviewMeshComponent && PreviewMeshComponent->IsPreviewOn()) ? PreviewMeshComponent->PreviewInstance : nullptr;
 			if (PreviewInstance)
 			{
 				bResumeAfterManipulating = PreviewInstance->IsPlaying();
