Cloth - Fixed bind command in LOD section cloth selection drop down when using non ClothingAssetCommon clothing data types that don't have an asset LOD specified.

#rnx

diff --git a/Engine/Source/Editor/Persona/Private/PersonaMeshDetails.cpp b/Engine/Source/Editor/Persona/Private/PersonaMeshDetails.cpp
--- a/Engine/Source/Editor/Persona/Private/PersonaMeshDetails.cpp
+++ b/Engine/Source/Editor/Persona/Private/PersonaMeshDetails.cpp
@@ -6967,7 +6967,7 @@ void FPersonaMeshDetails::OnClothingComboBoxOpening()
 
 TSharedRef<SWidget> FPersonaMeshDetails::OnGenerateWidgetForClothingEntry(TSharedPtr<FClothingEntry> InEntry)
 {
-	const UClothingAssetBase* const Asset = Cast<UClothingAssetBase>(InEntry->Asset.Get());
+	const UClothingAssetBase* const Asset = InEntry->Asset.Get();
 
 	FText EntryText;
 	if(Asset)
@@ -7052,7 +7052,8 @@ void FPersonaMeshDetails::OnClothingSelectionChanged(TSharedPtr<FClothingEntry>
 				ClearOriginalSectionUserData();
 			}
 
-			if (!ClothingAsset->BindToSkeletalMesh(Mesh, InLodIdx, InSectionIdx, InNewEntry->AssetLodIndex))
+			const int32 AssetLodIndex = InNewEntry->AssetLodIndex != INDEX_NONE ? InNewEntry->AssetLodIndex : InLodIdx;  //  Use the mesh LOD index if the asset type doesn't allow independent LOD setting
+			if (!ClothingAsset->BindToSkeletalMesh(Mesh, InLodIdx, InSectionIdx, AssetLodIndex))
 			{
 				// We failed to bind the clothing asset, reset box selection to "None"
 				SClothComboBoxPtr BoxPtr = ClothComboBoxes[BoxIndex];
@@ -7066,7 +7067,6 @@ void FPersonaMeshDetails::OnClothingSelectionChanged(TSharedPtr<FClothingEntry>
 				//Successful bind so set the SectionUserData
 				int32 AssetIndex = INDEX_NONE;
 				verify(Mesh->GetMeshClothingAssets().Find(ClothingAsset, AssetIndex));
-				const int32 AssetLodIndex = InNewEntry->AssetLodIndex != INDEX_NONE ? InNewEntry->AssetLodIndex : InLodIdx;  //  Use the mesh LOD index if the asset type doesn't allow independent LOD setting
 				OriginalSectionData.CorrespondClothAssetIndex = static_cast<int16>(AssetIndex);
 				OriginalSectionData.ClothingData.AssetGuid = ClothingAsset->GetAssetGuid();
 				OriginalSectionData.ClothingData.AssetLodIndex = AssetLodIndex;
