[UE] Preview Scene Settings ? Multiple unlocalized options in the Details panels
#fix  add persona details customization missing LOCTEXT

diff --git a/Engine/Source/Editor/Persona/Private/PreviewSceneCustomizations.cpp b/Engine/Source/Editor/Persona/Private/PreviewSceneCustomizations.cpp
--- a/Engine/Source/Editor/Persona/Private/PreviewSceneCustomizations.cpp
+++ b/Engine/Source/Editor/Persona/Private/PreviewSceneCustomizations.cpp
@@ -136,7 +136,7 @@ void FPreviewSceneDescriptionCustomization::CustomizeDetails(IDetailLayoutBuilde
 				return !GetMutableDefault<UPersonaOptions>()->IsAllowedClass(ControllerEntry->Class);
 			});
 
-		IDetailCategoryBuilder& AnimCategory = DetailBuilder.EditCategory("Animation");
+		IDetailCategoryBuilder& AnimCategory = DetailBuilder.EditCategory("Animation", LOCTEXT("AnimationLabel", "Animation"));
 		AnimCategory.AddCustomRow(PreviewControllerProperty->GetPropertyDisplayName())
 		.NameContent()
 		[
@@ -205,7 +205,7 @@ void FPreviewSceneDescriptionCustomization::CustomizeDetails(IDetailLayoutBuilde
 
 		const bool bCanUseDifferentSkeleton = PersonaToolkit.Pin()->CanPreviewMeshUseDifferentSkeleton();
 
-		DetailBuilder.EditCategory("Mesh")
+		DetailBuilder.EditCategory("Mesh", LOCTEXT("MeshLabel", "Mesh"))
 		.AddProperty(SkeletalMeshProperty)
 		.CustomWidget()
 		.OverrideResetToDefault(FResetToDefaultOverride::Create(
@@ -282,7 +282,7 @@ void FPreviewSceneDescriptionCustomization::CustomizeDetails(IDetailLayoutBuilde
 	
 	if (PersonaToolkit.Pin()->GetContext() == UAnimBlueprint::StaticClass()->GetFName())
 	{
-		DetailBuilder.EditCategory("Animation Blueprint")
+		DetailBuilder.EditCategory("Animation Blueprint", LOCTEXT("AnimationBlueprintLabel", "Animation Blueprint"))
 		.AddProperty(PreviewAnimationBlueprintProperty)
 		.CustomWidget()
 		.NameContent()
@@ -369,7 +369,7 @@ void FPreviewSceneDescriptionCustomization::CustomizeDetails(IDetailLayoutBuilde
 		TArray<UFactory*> FactoriesToUse({ FactoryToUse });
 
 		// bAllowPreviewMeshCollectionsToSelectFromDifferentSkeletons option
-		DetailBuilder.EditCategory("Additional Meshes")
+		DetailBuilder.EditCategory("Additional Meshes", LOCTEXT("AdditionalMeshesLabel", "Additional Meshes"))
 		.AddCustomRow(LOCTEXT("AdditionalMeshOption", "Additional Mesh Selection Option"))
 		.NameContent()
 		[
@@ -386,7 +386,7 @@ void FPreviewSceneDescriptionCustomization::CustomizeDetails(IDetailLayoutBuilde
 		];
 	
 		// bAllowPreviewMeshCollectionsToSelectFromDifferentSkeletons option
-		DetailBuilder.EditCategory("Additional Meshes")
+		DetailBuilder.EditCategory("Additional Meshes", LOCTEXT("AdditionalMeshesLabel", "Additional Meshes"))
 		.AddCustomRow(LOCTEXT("AdditionalMeshOption_AnimBP", "Additional Mesh Anim Selection Option"))
 		.NameContent()
 		[
@@ -407,7 +407,7 @@ void FPreviewSceneDescriptionCustomization::CustomizeDetails(IDetailLayoutBuilde
 			FResetToDefaultHandler::CreateSP(this, &FPreviewSceneDescriptionCustomization::OnResetToBaseClicked)
 		);
 
-		DetailBuilder.EditCategory("Additional Meshes")
+		DetailBuilder.EditCategory("Additional Meshes", LOCTEXT("AdditionalMeshesLabel", "Additional Meshes"))
 		.AddProperty(AdditionalMeshesProperty)
 		.CustomWidget()
 		.OverrideResetToDefault(ResetToDefaultOverride)
@@ -459,7 +459,7 @@ void FPreviewSceneDescriptionCustomization::CustomizeDetails(IDetailLayoutBuilde
 			TArray<UObject*> Objects;
 			Objects.Add(AdditionalMeshesAsset.GetAsset());
 
-			IDetailPropertyRow* PropertyRow = DetailBuilder.EditCategory("Additional Meshes")
+			IDetailPropertyRow* PropertyRow = DetailBuilder.EditCategory("Additional Meshes", LOCTEXT("AdditionalMeshesLabel", "Additional Meshes"))
 			.AddExternalObjectProperty(Objects, "SkeletalMeshes");
 
 			if (PropertyRow)
