[Persona] Make sure morph viewer removes registered callbacks when preview mesh changes to avoid them called on invalid mesh

diff --git a/Engine/Source/Editor/Persona/Private/SMorphTargetViewer.cpp b/Engine/Source/Editor/Persona/Private/SMorphTargetViewer.cpp
--- a/Engine/Source/Editor/Persona/Private/SMorphTargetViewer.cpp
+++ b/Engine/Source/Editor/Persona/Private/SMorphTargetViewer.cpp
@@ -429,7 +429,7 @@ void SMorphTargetViewer::Construct(const FArguments& InArgs, const TSharedRef<IP
 
 	const FText SkeletalMeshName = SkeletalMesh ? FText::FromString( SkeletalMesh->GetName() ) : LOCTEXT( "MorphTargetMeshNameLabel", "No Skeletal Mesh Present" );
 
-	SkeletalMesh->GetOnMeshChanged().Add(FSimpleDelegate::CreateSP(this, &SMorphTargetViewer::OnMeshChanged));
+	SkeletalMesh->GetOnMeshChanged().AddSP(this, &SMorphTargetViewer::OnMeshChanged);
 
 	ChildSlot
 	[
@@ -525,6 +525,7 @@ TSharedRef<SWidget> SMorphTargetViewer::AsWidget()
 
 void SMorphTargetViewer::OnPreviewMeshChanged(class USkeletalMesh* OldPreviewMesh, class USkeletalMesh* NewPreviewMesh)
 {
+	SkeletalMesh->GetOnMeshChanged().RemoveAll( this );
 	SkeletalMesh = NewPreviewMesh;
 	CreateMorphTargetList( NameFilterBox->GetText().ToString() );
 }
