[Persona] morph target viewer, added null checks

diff --git a/Engine/Source/Editor/Persona/Private/SMorphTargetViewer.cpp b/Engine/Source/Editor/Persona/Private/SMorphTargetViewer.cpp
--- a/Engine/Source/Editor/Persona/Private/SMorphTargetViewer.cpp
+++ b/Engine/Source/Editor/Persona/Private/SMorphTargetViewer.cpp
@@ -525,8 +525,17 @@ TSharedRef<SWidget> SMorphTargetViewer::AsWidget()
 
 void SMorphTargetViewer::OnPreviewMeshChanged(class USkeletalMesh* OldPreviewMesh, class USkeletalMesh* NewPreviewMesh)
 {
-	SkeletalMesh->GetOnMeshChanged().RemoveAll( this );
+	if (SkeletalMesh)
+	{
+		SkeletalMesh->GetOnMeshChanged().RemoveAll( this );
+	}
+	
 	SkeletalMesh = NewPreviewMesh;
+	
+	if (SkeletalMesh)
+	{
+		SkeletalMesh->GetOnMeshChanged().AddSP(this, &SMorphTargetViewer::OnMeshChanged);
+	}
 	CreateMorphTargetList( NameFilterBox->GetText().ToString() );
 }
 
