Animation Viewport: Allow switching between skin weight profiles directly in the viewport.

[RN]

diff --git a/Engine/Source/Editor/EditorStyle/Private/StarshipStyle.cpp b/Engine/Source/Editor/EditorStyle/Private/StarshipStyle.cpp
--- a/Engine/Source/Editor/EditorStyle/Private/StarshipStyle.cpp
+++ b/Engine/Source/Editor/EditorStyle/Private/StarshipStyle.cpp
@@ -5591,6 +5591,8 @@ void FStarshipEditorStyle::FStyle::SetupPersonaStyle()
 		Set( "AnimViewportMenu.AnimSetDrawUVs", new IMAGE_BRUSH( TEXT("Icons/icon_StaticMeshEd_UVOverlay_40x"), Icon40x40 ) );
 		Set( "AnimViewportMenu.AnimSetDrawUVs.Small", new IMAGE_BRUSH( TEXT("Icons/icon_StaticMeshEd_UVOverlay_40x"), Icon20x20 ) );
 
+		Set( "AnimViewportMenu.SkinWeightProfiles", new IMAGE_BRUSH_SVG( TEXT("Starship/Persona/SkinWeightProfiles"), Icon16x16 ) );
+
 		Set("AnimViewportMenu.PlayBackSpeed", new IMAGE_BRUSH_SVG("Starship/Common/play", Icon16x16));
 		Set("AnimViewportMenu.TurnTableSpeed", new IMAGE_BRUSH("Persona/Viewport/icon_turn_table_16x", Icon16x16));
 		Set("AnimViewportMenu.SceneSetup", new IMAGE_BRUSH("Icons/icon_tab_SceneOutliner_16x", Icon16x16));
diff --git a/Engine/Source/Editor/Persona/Private/SAnimationEditorViewport.cpp b/Engine/Source/Editor/Persona/Private/SAnimationEditorViewport.cpp
--- a/Engine/Source/Editor/Persona/Private/SAnimationEditorViewport.cpp
+++ b/Engine/Source/Editor/Persona/Private/SAnimationEditorViewport.cpp
@@ -22,7 +22,6 @@
 #include "AnimViewportLODCommands.h"
 #include "AnimViewportPlaybackCommands.h"
 #include "AnimPreviewInstance.h"
-#include "AssetViewerSettings.h"
 #include "Widgets/Input/STextComboBox.h"
 #include "IEditableSkeleton.h"
 #include "EditorAxisDisplayInfo.h"
@@ -249,7 +248,8 @@ TSharedPtr<SWidget> SAnimationEditorViewport::BuildViewportToolbar()
 
 			RightSection.AddEntry(UE::AnimationEditor::CreateShowSubmenu());
 			RightSection.AddEntry(UE::AnimationEditor::CreateLODSubmenu());
-
+			RightSection.AddEntry(UE::AnimationEditor::CreateSkinWeightProfileMenu());
+			
 			// Add Preview Scene Submenu
 			{
 				const FName AssetViewerProfileMenuName = "AnimationEditor.ViewportToolbar.AssetViewerProfile";
@@ -749,41 +749,6 @@ void SAnimationEditorViewportTabBody::Construct(const FArguments& InArgs, const
 		.Font(SmallLayoutFont)
 		.OnSelectionChanged(this, &SAnimationEditorViewportTabBody::ComboBoxSelectionChanged);
 
-	PopulateSkinWeightProfileNames();
-	
-	SkinWeightCombo = SNew(SNameComboBox)
-		.OptionsSource(&SkinWeightProfileNames)
-		.InitiallySelectedItem(SkinWeightProfileNames.Num() > 0 ? SkinWeightProfileNames[0] : nullptr)
-		.OnComboBoxOpening(FOnComboBoxOpening::CreateLambda([this]() 
-		{ 
-			// Retrieve currently selected value, and check whether or not it is still valid, it could be that a profile has been renamed or removed without updating the entries
-			FName Name = SkinWeightCombo->GetSelectedItem().IsValid() ? *SkinWeightCombo->GetSelectedItem().Get() : NAME_None;
-			PopulateSkinWeightProfileNames();
-			const int32 Index = SkinWeightProfileNames.IndexOfByPredicate([Name](TSharedPtr<FName> SearchName) { return Name == *SearchName; });
-			if (Index != INDEX_NONE)
-			{
-				SkinWeightCombo->SetSelectedItem(SkinWeightProfileNames[Index]);
-			}
-
-		}))
-		.OnSelectionChanged(SNameComboBox::FOnNameSelectionChanged::CreateLambda([WeakScenePtr = PreviewScenePtr](TSharedPtr<FName> SelectedProfile, ESelectInfo::Type SelectInfo)
-		{
-			// Apply the skin weight profile to the component, according to the selected the name, 
-			if (WeakScenePtr.IsValid() && SelectedProfile.IsValid())
-			{
-				UDebugSkelMeshComponent* MeshComponent = WeakScenePtr.Pin()->GetPreviewMeshComponent();
-				if (MeshComponent)
-				{
-					MeshComponent->ClearSkinWeightProfile();
-
-					if (*SelectedProfile != NAME_None)
-					{
-						MeshComponent->SetSkinWeightProfile(*SelectedProfile);
-					}					
-				}
-			}
-		}));	
-
 	FAnimationEditorViewportRequiredArgs ViewportArgs(InPreviewScene, SharedThis(this), InAssetEditorToolkit, InViewportIndex);
 
 	ViewportWidget = SNew(SAnimationEditorViewport, ViewportArgs)
@@ -878,8 +843,6 @@ void SAnimationEditorViewportTabBody::Construct(const FArguments& InArgs, const
 
 	AddMinLODNotification();
 
-	AddSkinWeightProfileNotification();
-	
 	if (TSharedPtr<FAnimationEditorPreviewScene> PreviewScene = PreviewScenePtr.Pin())
 	{
 		UE::AdvancedPreviewScene::BindDefaultOnSettingsChangedHandler(PreviewScene, LevelViewportClient);
@@ -2111,42 +2074,6 @@ void SAnimationEditorViewportTabBody::PopulateUVChoices()
 	}
 }
 
-void SAnimationEditorViewportTabBody::PopulateSkinWeightProfileNames()
-{
-	SkinWeightProfileNames.Empty();
-
-	// Always make sure we have a default 'none' option
-	const FName DefaultProfileName = NAME_None;
-	SkinWeightProfileNames.Add(MakeShared<FName>(DefaultProfileName));
-
-	// Retrieve all possible skin weight profiles from the component
-	if (UDebugSkelMeshComponent* PreviewComponent = GetPreviewScene()->GetPreviewMeshComponent())
-	{
-		if (USkeletalMesh* Mesh = PreviewComponent->GetSkeletalMeshAsset())
-		{
-			for (const FSkinWeightProfileInfo& Profile : Mesh->GetSkinWeightProfiles())
-			{
-				SkinWeightProfileNames.AddUnique(MakeShared<FName>(Profile.Name));
-			}
-		}
-	}
-}
-
-void SAnimationEditorViewportTabBody::UpdateSkinWeightSelection(USkeletalMesh* InSkeletalMesh)
-{
-	 // Check (post a mesh build) whether or not our currently selected profile name is still valid, and if not reset to 'none'
-	if (SkinWeightCombo->GetSelectedItem().IsValid())
-	{
-		const FName OldSelection = *SkinWeightCombo->GetSelectedItem();
-		PopulateSkinWeightProfileNames();
-		
-		const int32 SelectionIndex = SkinWeightProfileNames.IndexOfByPredicate([OldSelection](TSharedPtr<FName> InName) { return *InName == OldSelection; });		
-		
-		// Select new entry or otherwise select none
-		SkinWeightCombo->SetSelectedItem(SelectionIndex != INDEX_NONE ? SkinWeightProfileNames[SelectionIndex] : SkinWeightProfileNames[0]);
-	}
-}
-
 void SAnimationEditorViewportTabBody::UpdateScrubPanel(UAnimationAsset* AnimAsset)
 {
 	// We might not have a scrub panel if we're in animation mode.
@@ -2340,6 +2267,49 @@ void SAnimationEditorViewportTabBody::OnDebugForcedLODChanged()
 	}
 }
 
+void SAnimationEditorViewportTabBody::OnSetSkinWeightProfile(FName InProfileName, ESkinWeightProfileLayer InLayer)
+{
+	// Apply the skin weight profile to the component, according to the selected the name, 
+	if (PreviewScenePtr.IsValid())
+	{
+		if (UDebugSkelMeshComponent* MeshComponent = PreviewScenePtr.Pin()->GetPreviewMeshComponent())
+		{
+			MeshComponent->SetSkinWeightProfile(InProfileName, InLayer);
+		}
+	}
+}
+
+bool SAnimationEditorViewportTabBody::IsSkinWeightProfileSelected(FName InProfileName, ESkinWeightProfileLayer InLayer) const
+{
+	// Returns true if the given profile on that layer is enabeld. 
+	if (PreviewScenePtr.IsValid())
+	{
+		if (UDebugSkelMeshComponent* MeshComponent = PreviewScenePtr.Pin()->GetPreviewMeshComponent())
+		{
+			return MeshComponent->GetCurrentSkinWeightProfileName(InLayer) == InProfileName;
+		}
+	}
+	return InProfileName.IsNone();
+}
+
+TArray<FName> SAnimationEditorViewportTabBody::GetSelectedProfileNames() const
+{
+	// Get a list of the selected profile names. If only a single profile is selected, whether on primary or secondary,
+	// only one entry is returned. If no profile is selected, a single entry with NAME_None is returned. 
+	if (PreviewScenePtr.IsValid())
+	{
+		if (UDebugSkelMeshComponent* MeshComponent = PreviewScenePtr.Pin()->GetPreviewMeshComponent())
+		{
+			return MeshComponent->GetCurrentSkinWeightProfileLayerNames();
+		}
+	}
+
+	TArray<FName> ProfileNames;
+	ProfileNames.SetNum(FSkinWeightProfileStack::MaxLayerCount);
+	
+	return ProfileNames;
+}
+
 void SAnimationEditorViewportTabBody::OnBeginSliderMovementFloorOffset()
 {
 	// This value is saved in a UPROPERTY for the floor mesh, so changes are transactional
@@ -3132,66 +3102,6 @@ void SAnimationEditorViewportTabBody::AddMinLODNotification()
 	);
 }
 
-void SAnimationEditorViewportTabBody::AddSkinWeightProfileNotification()
-{
-	if(WeakSkinWeightPreviewNotification.IsValid())
-	{
-		return;
-	}
-
-	auto GetSkinWeightProfileNotificationVisibility = [this]()
-	{
-		if (GetPreviewScene()->GetPreviewMeshComponent() && GetPreviewScene()->GetPreviewMeshComponent()->IsUsingSkinWeightProfile())
-		{
-			return EVisibility::Visible;
-		}
-
-		return EVisibility::Collapsed;
-	};
-
-	auto GetSkinWeightProfileNotificationText = [this]() -> FText
-	{
-		FName ProfileName = NAME_None;
-		if (GetPreviewScene()->GetPreviewMeshComponent())
-		{
-			ProfileName = GetPreviewScene()->GetPreviewMeshComponent()->GetCurrentSkinWeightProfileName();
-		}
-
-		return FText::FormatOrdered(LOCTEXT("ProfileSkinWeightPreviewNotification", "Previewing Skin Weight Profile: {0}"), FText::FromName(ProfileName));
-	};
-
-	WeakSkinWeightPreviewNotification = AddNotification(EMessageSeverity::Info,
-		false,
-		SNew(SHorizontalBox)
-		.Visibility_Lambda(GetSkinWeightProfileNotificationVisibility)
-		.ToolTipText(LOCTEXT("ProfileSkinWeightPreviewTooltip", "Previewing a Skin Weight Profile."))
-		+SHorizontalBox::Slot()
-		.FillWidth(1.0f)
-		.Padding(2.0f, 4.0f)
-		[
-			SNew(SHorizontalBox)
-			+SHorizontalBox::Slot()
-			.AutoWidth()
-			.VAlign(VAlign_Center)
-			.Padding(0.0f, 0.0f, 4.0f, 0.0f)
-			[
-				SNew(STextBlock)
-				.TextStyle(FAppStyle::Get(), "AnimViewport.MessageText")
-				.Font(FAppStyle::Get().GetFontStyle("FontAwesome.9"))
-				.Text(FEditorFontGlyphs::Eye)
-			]
-			+SHorizontalBox::Slot()
-			.VAlign(VAlign_Center)
-			.FillWidth(1.0f)
-			[
-				SNew(STextBlock)
-				.Text_Lambda(GetSkinWeightProfileNotificationText)
-				.TextStyle(FAppStyle::Get(), "AnimViewport.MessageText")
-			]
-		],
-		FPersonaViewportNotificationOptions(TAttribute<EVisibility>::Create(GetSkinWeightProfileNotificationVisibility))
-	);
-}
 
 void SAnimationEditorViewportTabBody::HandleFocusCamera()
 {
diff --git a/Engine/Source/Editor/Persona/Private/SAnimationEditorViewport.h b/Engine/Source/Editor/Persona/Private/SAnimationEditorViewport.h
--- a/Engine/Source/Editor/Persona/Private/SAnimationEditorViewport.h
+++ b/Engine/Source/Editor/Persona/Private/SAnimationEditorViewport.h
@@ -302,6 +302,10 @@ public:
 	void OnLODModelChanged();
 	void OnDebugForcedLODChanged();
 
+	void OnSetSkinWeightProfile(FName InProfileName, ESkinWeightProfileLayer InLayer);
+	bool IsSkinWeightProfileSelected(FName InProfileName, ESkinWeightProfileLayer InLayer) const;
+	TArray<FName> GetSelectedProfileNames() const;
+
 	void OnBeginSliderMovementFloorOffset();
 	void OnFloorOffsetChanged(float InNewValue);
 	void OnFloorOffsetCommitted(float InNewValue, ETextCommit::Type InCommitType);
@@ -643,10 +647,6 @@ private:
 	TArray<int32> NumUVChannels;
 
 public:
-	/** Skin Weight Profile Selector */
-	TSharedPtr<SNameComboBox> SkinWeightCombo;
-	TArray<TSharedPtr<FName>> SkinWeightProfileNames;
-
 	/** Box that contains scrub panel */
 	TSharedPtr<SVerticalBox> ScrubPanelContainer;
 
@@ -662,9 +662,6 @@ public:
 	/** Min LOD notification */
 	TWeakPtr<SWidget> WeakMinLODNotification;
 
-	/** Min LOD notification */
-	TWeakPtr<SWidget> WeakSkinWeightPreviewNotification;
-
 	/** Current LOD selection*/
 	int32 LODSelection;
 
@@ -692,7 +689,4 @@ public:
 
 	/** Add a notification to tell the user a min LOD is being applied */
 	void AddMinLODNotification();
-
-	/** Add a notification to tell the user a Skin Weight Profile is being previewed */
-	void AddSkinWeightProfileNotification();
 };
diff --git a/Engine/Source/Editor/Persona/Private/ViewportToolbar/AnimationEditorMenus.cpp b/Engine/Source/Editor/Persona/Private/ViewportToolbar/AnimationEditorMenus.cpp
--- a/Engine/Source/Editor/Persona/Private/ViewportToolbar/AnimationEditorMenus.cpp
+++ b/Engine/Source/Editor/Persona/Private/ViewportToolbar/AnimationEditorMenus.cpp
@@ -25,6 +25,7 @@
 #include "SimulationEditorExtender.h"
 #include "ToolMenus.h"
 #include "UICommandList_Pinnable.h"
+#include "Animation/SkinWeightProfile.h"
 #include "ViewportToolbar/UnrealEdViewportToolbar.h"
 #include "Widgets/Input/SSpinBox.h"
 #include "Widgets/Input/STextComboBox.h"
@@ -155,6 +156,59 @@ void PopulateLODSubmenu(UToolMenu* InMenu)
 	}
 }
 
+void PopulateSkinWeightProfileMenu(UToolMenu* InMenu)
+{
+	if (!InMenu)
+	{
+		return;
+	}
+
+	const TSharedPtr<SAnimationEditorViewportTabBody> AnimEditorViewportTab = GetAnimationEditorViewportTab(InMenu);
+	if (!AnimEditorViewportTab)
+	{
+		return;
+	}
+
+	FText NoProfileLabel = LOCTEXT("SkinWeightProfileMenu_NoProfile", "No Profile");
+
+	const TWeakObjectPtr<const USkeletalMesh> PreviewMeshWeak = AnimEditorViewportTab->GetPreviewScene()->GetPreviewMesh();
+
+	if (const TStrongObjectPtr<const USkeletalMesh>& PreviewMeshPinned = PreviewMeshWeak.Pin())
+	{
+		const TArray<FSkinWeightProfileInfo>& SkinWeightProfiles = PreviewMeshPinned->GetSkinWeightProfiles();
+
+		FToolMenuSection& ProfilesSection = InMenu->AddSection("SkinWeightActiveProfile", LOCTEXT("SkinWeightProfileMenu_ActiveProfile", "Active Profile"));
+
+		auto AddProfileEntries = [Viewport=AnimEditorViewportTab, NoProfileLabel, &SkinWeightProfiles](FToolMenuSection& InSection, ESkinWeightProfileLayer InLayer)
+			{
+				FName NoneProfile;
+				FUIAction ActionSetBaseProfile(FExecuteAction::CreateSP(Viewport.ToSharedRef(), &SAnimationEditorViewportTabBody::OnSetSkinWeightProfile, NoneProfile, InLayer),
+					FCanExecuteAction(),
+					FIsActionChecked::CreateSP(Viewport.ToSharedRef(), &SAnimationEditorViewportTabBody::IsSkinWeightProfileSelected, NoneProfile, InLayer));
+
+				InSection.AddMenuEntry(FName(NoProfileLabel.ToString()), NoProfileLabel, FText::GetEmpty(), FSlateIcon(), ActionSetBaseProfile, EUserInterfaceActionType::RadioButton);
+				
+				// Retrieve all possible skin weight profiles from the component
+				for (const FSkinWeightProfileInfo& Profile : SkinWeightProfiles)
+				{
+					FUIAction Action(FExecuteAction::CreateSP(Viewport.ToSharedRef(), &SAnimationEditorViewportTabBody::OnSetSkinWeightProfile, Profile.Name, InLayer),
+						FCanExecuteAction(),
+						FIsActionChecked::CreateSP(Viewport.ToSharedRef(), &SAnimationEditorViewportTabBody::IsSkinWeightProfileSelected, Profile.Name, InLayer));
+					
+					InSection.AddMenuEntry(Profile.Name, FText::FromName(Profile.Name), FText::GetEmpty(), FSlateIcon(), Action, EUserInterfaceActionType::RadioButton);
+				}
+			};
+
+		AddProfileEntries(ProfilesSection, ESkinWeightProfileLayer::Primary);
+
+		if (SkinWeightProfiles.Num() > 1)
+		{
+			FToolMenuSection& SecondaryProfilesSection = InMenu->AddSection("SkinWeightActiveSecondaryProfiles", LOCTEXT("SkinWeightProfileMenu_ActiveSecondaryProfile", "Active Secondary Profile"));
+			AddProfileEntries(SecondaryProfilesSection, ESkinWeightProfileLayer::Secondary);
+		}
+	}	
+}
+
 void FillCharacterMirrorMenu(UToolMenu* InMenu)
 {
 	if (!InMenu)
@@ -809,12 +863,6 @@ void UE::AnimationEditor::AddSceneElementsSection(UToolMenu* InMenu)
 					Section.AddMenuEntry(FAnimViewportShowCommands::Get().ShowBakedAnimation);
 					Section.AddMenuEntry(FAnimViewportShowCommands::Get().DisablePostProcessBlueprint);
 				}
-
-				FToolMenuSection& Section =
-					InSubMenu->AddSection("SkinWeights", LOCTEXT("SkinWeights_Label", "Skin Weight Profiles"));
-				Section.AddEntry(FToolMenuEntry::InitWidget(
-					"SkinWeightCombo", AnimEditorViewportTab->SkinWeightCombo.ToSharedRef(), FText()
-				));
 			}
 		)
 	);
@@ -1700,6 +1748,102 @@ TSharedRef<SWidget> UE::AnimationEditor::GenerateLODMenuWidget(const TWeakPtr<SA
 	return UToolMenus::Get()->GenerateWidget(OldShowMenuName, MenuContext);
 }
 
+FToolMenuEntry UE::AnimationEditor::CreateSkinWeightProfileMenu()
+{
+	FToolMenuEntry Entry = FToolMenuEntry::InitDynamicEntry(
+		"DynamicSkinWeightProfileOptions",
+		FNewToolMenuSectionDelegate::CreateLambda(
+			[](FToolMenuSection& InDynamicSection) -> void
+			{
+				TSharedPtr<SAnimationEditorViewportTabBody> AnimEditorViewportTab =
+					Private::GetAnimationEditorViewportTab(InDynamicSection);
+				if (!AnimEditorViewportTab)
+				{
+					return;
+				}
+
+				// Label updates based on currently selected skin weight profile
+				const TAttribute<FText> Label = TAttribute<FText>::CreateLambda(
+					[WeakTab = AnimEditorViewportTab.ToWeakPtr()]()
+					{
+						if (TSharedPtr<SAnimationEditorViewportTabBody> ViewportTab = WeakTab.Pin())
+						{
+							return GetSkinWeightProfileMenuLabel(ViewportTab);
+						}
+
+						return LOCTEXT("SkinWeightProfileMenu_NoProfile", "No Profile");
+					}
+				);
+
+				FToolMenuEntry& Entry = InDynamicSection.AddSubMenu(
+					"SkinWeightProfile",
+					Label,
+					LOCTEXT(
+						"SkinWeightProfileMenuTooltip", "Control how skin weight profiles are displayed."
+					),
+					FNewToolMenuDelegate::CreateLambda(
+						[](UToolMenu* Submenu) -> void
+						{
+							UE::AnimationEditor::Private::PopulateSkinWeightProfileMenu(Submenu);
+						}
+					),
+					false,
+					FSlateIcon(FAppStyle::GetAppStyleSetName(), "AnimViewportMenu.SkinWeightProfiles")
+				);
+
+
+				Entry.Visibility = [WeakTab = AnimEditorViewportTab.ToWeakPtr()]() -> bool
+					{
+						if (TSharedPtr<SAnimationEditorViewportTabBody> ViewportTab = WeakTab.Pin())
+						{
+							if (const USkeletalMesh* SkeletalMesh = ViewportTab->GetPreviewScene()->GetPreviewMesh())
+							{
+								if (!SkeletalMesh->IsCompiling())
+								{
+									return SkeletalMesh->GetNumSkinWeightProfiles() > 0;
+								}
+							}
+						}
+						return false;
+					};
+
+				// Profiles are a lower priority than other menus. 
+				Entry.ToolBarData.ResizeParams.ClippingPriority = 200;
+			}
+		)
+	);
+
+	return Entry;
+}
+
+FText UE::AnimationEditor::GetSkinWeightProfileMenuLabel(const TWeakPtr<SAnimationEditorViewportTabBody>& InAnimEditorViewportTab)
+{
+	FText NoProfileLabel = LOCTEXT("SkinWeightProfileMenu_NoProfile", "No Profile");
+
+	if (TSharedPtr<SAnimationEditorViewportTabBody> ViewportTab = InAnimEditorViewportTab.Pin())
+	{
+		TArray<FName> ProfileNames = ViewportTab->GetSelectedProfileNames();
+
+		ProfileNames.RemoveAll([](FName InName) { return InName.IsNone(); });
+
+		if (ProfileNames.IsEmpty())
+		{
+			return NoProfileLabel;
+		}
+
+		TArray<FText> ProfileTextNames;
+		for (FName Name: ProfileNames)
+		{
+			ProfileTextNames.Add(FText::FromName(Name));
+		}
+
+		return FText::Join(LOCTEXT("SkinWeightProfileMenu_LabelJoin", " / "), ProfileTextNames);
+	}
+	
+	return NoProfileLabel;
+}
+
+
 TSharedRef<SWidget> UE::AnimationEditor::CreateShowMenuWidget(
 	const TSharedRef<SEditorViewport>& InViewport, const TArray<TSharedPtr<FExtender>>& InExtenders, bool bInShowViewportStatsToggle
 )
diff --git a/Engine/Source/Editor/Persona/Private/ViewportToolbar/AnimationEditorMenus.h b/Engine/Source/Editor/Persona/Private/ViewportToolbar/AnimationEditorMenus.h
--- a/Engine/Source/Editor/Persona/Private/ViewportToolbar/AnimationEditorMenus.h
+++ b/Engine/Source/Editor/Persona/Private/ViewportToolbar/AnimationEditorMenus.h
@@ -42,6 +42,10 @@ FToolMenuEntry CreateLODSubmenu();
 FText GetLODMenuLabel(const TWeakPtr<SAnimationEditorViewportTabBody>& InAnimEditorViewportTab);
 TSharedRef<SWidget> GenerateLODMenuWidget(const TWeakPtr<SAnimationEditorViewportTabBody>& InAnimEditorViewportTab);
 
+
+FToolMenuEntry CreateSkinWeightProfileMenu();
+FText GetSkinWeightProfileMenuLabel(const TWeakPtr<SAnimationEditorViewportTabBody>& InAnimEditorViewportTab);
+
 TSharedRef<SWidget> MakeFloorOffsetWidget(const TWeakPtr<SAnimationEditorViewportTabBody>& InAnimEditorViewportTabWeak);
 
 void ExtendViewModesSubmenu(FName InViewModesSubmenuName);
diff --git a/Engine/Source/Runtime/Engine/Private/Components/SkinnedMeshComponent.cpp b/Engine/Source/Runtime/Engine/Private/Components/SkinnedMeshComponent.cpp
--- a/Engine/Source/Runtime/Engine/Private/Components/SkinnedMeshComponent.cpp
+++ b/Engine/Source/Runtime/Engine/Private/Components/SkinnedMeshComponent.cpp
@@ -5264,12 +5264,15 @@ TArray<FName> USkinnedMeshComponent::GetCurrentSkinWeightProfileLayerNames() con
 		ProfileStack = FSkinWeightProfileStack{CurrentSkinWeightProfileLayers};
 	}
 
+	TArray<FName> ProfileNames;
+	ProfileNames.SetNum(FSkinWeightProfileStack::MaxLayerCount);
+	
 	if (USkinnedAsset* Asset = GetSkinnedAsset())
 	{
 		if (Asset->IsCompiling())
 		{
-			//Return an array with name none entry
-			return { NAME_None };
+			//Return an array with name none entries for all layers
+			return ProfileNames;
 		}
 
 		if (const FSkinWeightProfilesData* ProfileData = Asset->GetSkinWeightProfilesData(GetPredictedLODLevel()))
@@ -5281,8 +5284,6 @@ TArray<FName> USkinnedMeshComponent::GetCurrentSkinWeightProfileLayerNames() con
 		}
 	}
 
-	TArray<FName> ProfileNames;
-	ProfileNames.SetNum(FSkinWeightProfileStack::MaxLayerCount);
 	for (int32 LayerIndex = 0; LayerIndex < FSkinWeightProfileStack::MaxLayerCount; ++LayerIndex)
 	{
 		ProfileNames[LayerIndex] = ProfileStack[LayerIndex];
