Animation Viewport: Fix dangling code.

#rnx

diff --git a/Engine/Source/Editor/Persona/Private/SAnimationEditorViewport.cpp b/Engine/Source/Editor/Persona/Private/SAnimationEditorViewport.cpp
--- a/Engine/Source/Editor/Persona/Private/SAnimationEditorViewport.cpp
+++ b/Engine/Source/Editor/Persona/Private/SAnimationEditorViewport.cpp
@@ -831,13 +831,8 @@ void SAnimationEditorViewportTabBody::Construct(const FArguments& InArgs, const
 	BindCommands();
 
 	PopulateNumUVChannels();
-	PopulateSkinWeightProfileNames();
 
 	GetPreviewScene()->OnRecordingStateChanged().AddSP(this, &SAnimationEditorViewportTabBody::AddRecordingNotification);
-	if (GetPreviewScene()->GetPreviewMesh())
-	{
-		GetPreviewScene()->GetPreviewMesh()->OnPostMeshCached().AddSP(this, &SAnimationEditorViewportTabBody::UpdateSkinWeightSelection);
-	}
 
 	AddPostProcessNotification();
 
@@ -1973,17 +1968,11 @@ bool SAnimationEditorViewportTabBody::IsUsingPreSkinnedBounds() const
 void SAnimationEditorViewportTabBody::HandlePreviewMeshChanged(class USkeletalMesh* OldSkeletalMesh, class USkeletalMesh* NewSkeletalMesh)
 {
 	PopulateNumUVChannels();
-	PopulateSkinWeightProfileNames();
 
 	if (OldSkeletalMesh)
 	{
 		OldSkeletalMesh->OnPostMeshCached().RemoveAll(this);
 	}
-	
-	if (NewSkeletalMesh)
-	{
-		NewSkeletalMesh->OnPostMeshCached().AddSP(this, &SAnimationEditorViewportTabBody::UpdateSkinWeightSelection);
-	}
 }
 
 void SAnimationEditorViewportTabBody::AnimChanged(UAnimationAsset* AnimAsset)
diff --git a/Engine/Source/Editor/Persona/Private/SAnimationEditorViewport.h b/Engine/Source/Editor/Persona/Private/SAnimationEditorViewport.h
--- a/Engine/Source/Editor/Persona/Private/SAnimationEditorViewport.h
+++ b/Engine/Source/Editor/Persona/Private/SAnimationEditorViewport.h
@@ -8,7 +8,6 @@
 #include "Widgets/DeclarativeSyntaxSupport.h"
 #include "PersonaDelegates.h"
 #include "IPersonaViewport.h"
-#include "EngineDefines.h"
 #include "Toolkits/AssetEditorToolkit.h"
 #include "BlueprintEditor.h"
 #include "IPersonaPreviewScene.h"
@@ -280,9 +279,6 @@ public:
 	/** Function to get slider value which represents wind strength (0 - 1)*/
 	TOptional<float> GetWindStrengthSliderValue() const;
 
-	/** Function to get slider value which returns a string*/
-	FText GetWindStrengthLabel() const;
-
 	bool IsApplyingClothWind() const;
 
 	/** Show gravity scale */
@@ -483,12 +479,6 @@ private:
 	/** Populates choices for UV Channel combo box */
 	void PopulateUVChoices();
 
-	/** Populates choices for Skin Weight Profile combo box for each lod based on current preview asset */
-	void PopulateSkinWeightProfileNames();
-
-	/** Checks if Skin Weight Profile selection is still valid and otherwise resets it, called through PostCache from currenty Preview Skeletal Mesh */
-	void UpdateSkinWeightSelection(USkeletalMesh* InSkeletalMesh);
-
 	void AnimChanged(UAnimationAsset* AnimAsset);
 
 	/** Open the preview scene settings */
@@ -510,9 +500,6 @@ public:
 	bool CanChangeCameraMode() const;
 
 private:
-	/** Tests to see if bone move mode buttons should be visible */
-	EVisibility GetBoneMoveModeButtonVisibility() const;
-
 	/** Function to mute/unmute viewport audio */
 	void OnToggleMuteAudio();
 
