Montages need an easy way to set segment range to animation length

diff --git a/Engine/Source/Editor/Persona/Private/SAnimSegmentsPanel.cpp b/Engine/Source/Editor/Persona/Private/SAnimSegmentsPanel.cpp
--- a/Engine/Source/Editor/Persona/Private/SAnimSegmentsPanel.cpp
+++ b/Engine/Source/Editor/Persona/Private/SAnimSegmentsPanel.cpp
@@ -316,6 +316,10 @@ void SAnimSegmentsPanel::SummonSegmentNodeContextMenu(FMenuBuilder& MenuBuilder,
 			// open asset option
 			UIAction.ExecuteAction.BindRaw(this, &SAnimSegmentsPanel::OpenAsset, AnimSegmentIndex);
 			MenuBuilder.AddMenuEntry(LOCTEXT("OpenAssetOfSegment", "Open Asset"), LOCTEXT("OpenAssetOfSegment_ToolTip", "Open Asset"), FSlateIcon(), UIAction);
+			
+			// Reset segment start position/length to (0.f/LengthOfAnimationSequence)
+			UIAction.ExecuteAction.BindRaw(this, &SAnimSegmentsPanel::ResetSegmentLength, AnimSegmentIndex);
+			MenuBuilder.AddMenuEntry(LOCTEXT("ResetSegmentLength", "Reset Segment Length"), LOCTEXT("ResetSegment_ToolTip", "Resets Segment its StartOffset and Length to match the referenced Animation"), FSlateIcon(), UIAction);			
 		}
 		MenuBuilder.EndSection();
 	}
@@ -387,7 +391,7 @@ void SAnimSegmentsPanel::RemoveAnimSegment(int32 AnimSegmentIndex)
 {
 	if(ValidIndex(AnimSegmentIndex))
 	{
-		const FScopedTransaction Transaction( LOCTEXT("AnimSegmentseEditor", "Remove Segment") );
+		const FScopedTransaction Transaction( LOCTEXT("RemoveAnimSegmentTransaction", "Remove Segment") );
 		OnPreAnimUpdateDelegate.ExecuteIfBound();
 
 		AnimTrack->AnimSegments.RemoveAt(AnimSegmentIndex);
@@ -413,6 +417,23 @@ void SAnimSegmentsPanel::OpenAsset(int32 AnimSegmentIndex)
 	}
 }
 
+void SAnimSegmentsPanel::ResetSegmentLength(int32 AnimSegmentIndex)
+{
+	if (ValidIndex(AnimSegmentIndex))
+	{
+		if (UAnimSequenceBase* Asset = AnimTrack->AnimSegments[AnimSegmentIndex].GetAnimReference())
+		{
+			const FScopedTransaction Transaction( LOCTEXT("ResetAnimSegmentTransaction", "Reset Segment Length") );
+			OnPreAnimUpdateDelegate.ExecuteIfBound();
+			{
+				AnimTrack->AnimSegments[AnimSegmentIndex].AnimStartTime = 0.f;
+				AnimTrack->AnimSegments[AnimSegmentIndex].AnimEndTime = Asset->GetPlayLength();
+			}
+			OnPostAnimUpdateDelegate.ExecuteIfBound();
+		}
+	}
+}
+
 void SAnimSegmentsPanel::FillSubMenu(FMenuBuilder& MenuBuilder, int32 AnimSegmentIndex)
 {
 	if (ValidIndex(AnimSegmentIndex))
diff --git a/Engine/Source/Editor/Persona/Private/SAnimSegmentsPanel.h b/Engine/Source/Editor/Persona/Private/SAnimSegmentsPanel.h
--- a/Engine/Source/Editor/Persona/Private/SAnimSegmentsPanel.h
+++ b/Engine/Source/Editor/Persona/Private/SAnimSegmentsPanel.h
@@ -151,6 +151,7 @@ private:
 	void RemoveAnimSegment(int32 AnimSegmentIndex);
 	void RevertToParent(int32 AnimSegmentIndex);
 	void OpenAsset(int32 AnimSegmentIndex);
+	void ResetSegmentLength(int32 AnimSegmentIndex);
 	void FillSubMenu(FMenuBuilder& MenuBuilder, int32 AnimSegmentIndex);
 
 	/** Bind UI commands for this widget */
