[FBX legacy import] Fixed crash when reimport mesh using the legacy path since it is calling posteditchange on a temp/invalid skeletal mesh.

diff --git a/Engine/Source/Editor/Persona/Private/PersonaMeshDetails.cpp b/Engine/Source/Editor/Persona/Private/PersonaMeshDetails.cpp
--- a/Engine/Source/Editor/Persona/Private/PersonaMeshDetails.cpp
+++ b/Engine/Source/Editor/Persona/Private/PersonaMeshDetails.cpp
@@ -2687,7 +2687,7 @@ void FPersonaMeshDetails::OnPasteMaterialItem(int32 CurrentSlot)
 	}
 }
 
-void FPersonaMeshDetails::OnAttributePreChangePreventPostEditChange(const int32 LODIndex, const FName LODInfoPropertyName) const
+void FPersonaMeshDetails::OnAttributePreChangePreventPostEditChange(const int32 LODIndex, const FName LODInfoPropertyName)
 {
 	//Avoid PostEditChange call 
 	USkeletalMesh* Mesh = GetPersonaToolkit()->GetMesh();
@@ -2696,10 +2696,13 @@ void FPersonaMeshDetails::OnAttributePreChangePreventPostEditChange(const int32
 	{
 		return;
 	}
-	Mesh->StackPostEditChange();
+	
+	constexpr bool bCallPostEditChange = false;
+	constexpr bool bReregisterComponents = false;
+	ChangeScopes.Push(MakeUnique<FScopedSkeletalMeshPostEditChange>(Mesh, bCallPostEditChange, bReregisterComponents));
 }
 
-void FPersonaMeshDetails::OnAttributeChangedPreventPostEditChange(const int32 LODIndex, const FName LODInfoPropertyName, const bool bForceComponentRefresh) const
+void FPersonaMeshDetails::OnAttributeChangedPreventPostEditChange(const int32 LODIndex, const FName LODInfoPropertyName, const bool bForceComponentRefresh)
 {
 	USkeletalMesh* Mesh = GetPersonaToolkit()->GetMesh();
 
@@ -2707,7 +2710,7 @@ void FPersonaMeshDetails::OnAttributeChangedPreventPostEditChange(const int32 LO
 	{
 		return;
 	}
-	Mesh->UnStackPostEditChange();
+	ChangeScopes.Pop();
 	if(bForceComponentRefresh)
 	{
 		const bool bCallPostEditChange = false;
@@ -2716,7 +2719,7 @@ void FPersonaMeshDetails::OnAttributeChangedPreventPostEditChange(const int32 LO
 	}
 }
 
-void FPersonaMeshDetails::PreventAttributePostEditChange(TSharedPtr<IPropertyHandle> AttributeHandle, const int32 LODIndex, const FName PropertyName, const bool bForceComponentRefresh) const
+void FPersonaMeshDetails::PreventAttributePostEditChange(TSharedPtr<IPropertyHandle> AttributeHandle, const int32 LODIndex, const FName PropertyName, const bool bForceComponentRefresh)
 {
 	AttributeHandle->SetOnPropertyValuePreChange(FSimpleDelegate::CreateSP(this, &FPersonaMeshDetails::OnAttributePreChangePreventPostEditChange, LODIndex, PropertyName));
 	AttributeHandle->SetOnPropertyValueChanged(FSimpleDelegate::CreateSP(this, &FPersonaMeshDetails::OnAttributeChangedPreventPostEditChange, LODIndex, PropertyName, bForceComponentRefresh));
diff --git a/Engine/Source/Editor/Persona/Private/PersonaMeshDetails.h b/Engine/Source/Editor/Persona/Private/PersonaMeshDetails.h
--- a/Engine/Source/Editor/Persona/Private/PersonaMeshDetails.h
+++ b/Engine/Source/Editor/Persona/Private/PersonaMeshDetails.h
@@ -363,9 +363,9 @@ public:
 
 private:
 	//Prevent attribute change calling post edit change
-	void OnAttributePreChangePreventPostEditChange(int32 LODIndex, FName LODInfoPropertyName) const;
-	void OnAttributeChangedPreventPostEditChange(const int32 LODIndex, const FName LODInfoPropertyName, const bool bForceComponentRefresh) const;
-	void PreventAttributePostEditChange(TSharedPtr<IPropertyHandle> AttributeHandle, const int32 LODIndex, const FName PropertyName, const bool bForceComponentRefresh) const;
+	void OnAttributePreChangePreventPostEditChange(int32 LODIndex, FName LODInfoPropertyName);
+	void OnAttributeChangedPreventPostEditChange(const int32 LODIndex, const FName LODInfoPropertyName, const bool bForceComponentRefresh);
+	void PreventAttributePostEditChange(TSharedPtr<IPropertyHandle> AttributeHandle, const int32 LODIndex, const FName PropertyName, const bool bForceComponentRefresh);
 
 	//This function customize the LODInfo temporary object
 	void CustomizeLODInfoSetingsDetails(IDetailLayoutBuilder& DetailLayout, class ULODInfoUILayout* LODInfoUILayout, TSharedRef<IPropertyHandle> LODInfoProperty, IDetailCategoryBuilder& LODCategory);
@@ -771,6 +771,7 @@ private:
 	bool bDeleteWarningConsumed;
 
 private:
+	TArray<TUniquePtr<FScopedSkeletalMeshPostEditChange>> ChangeScopes;
 
 	// info about clothing combo boxes for multiple LOD
 	TArray<FClothingComboInfo>				ClothingComboLODInfos;
diff --git a/Engine/Source/Editor/UnrealEd/Private/Fbx/FbxSkeletalMeshImport.cpp b/Engine/Source/Editor/UnrealEd/Private/Fbx/FbxSkeletalMeshImport.cpp
--- a/Engine/Source/Editor/UnrealEd/Private/Fbx/FbxSkeletalMeshImport.cpp
+++ b/Engine/Source/Editor/UnrealEd/Private/Fbx/FbxSkeletalMeshImport.cpp
@@ -1973,9 +1973,6 @@ USkeletalMesh* UnFbx::FFbxImporter::ImportSkeletalMesh(FImportSkeletalMeshArgs &
 
 		SkeletalMesh = ExistingSkelMesh;
 
-		int32 PostEditChangeStackCounter = SkeletalMesh->GetPostEditChangeStackCounter();
-		SkeletalMesh->SetPostEditChangeStackCounter(0);
-
 		for (int32 LODIndex = 0, LODCount = SkeletalMesh->GetLODNum(); LODIndex < LODCount; ++LODIndex)
 		{
 			SkeletalMesh->ClearMeshDescriptionAndBulkData(LODIndex);
@@ -1993,9 +1990,6 @@ USkeletalMesh* UnFbx::FFbxImporter::ImportSkeletalMesh(FImportSkeletalMeshArgs &
 		SkeletalMesh->UnregisterAllMorphTarget();
 
 		SkeletalMesh->ReleaseResources();
-
-		SkeletalMesh->PostEditChange();
-		SkeletalMesh->SetPostEditChangeStackCounter(PostEditChangeStackCounter);
 	}
 
 
diff --git a/Engine/Source/Runtime/Engine/Classes/Engine/SkeletalMesh.h b/Engine/Source/Runtime/Engine/Classes/Engine/SkeletalMesh.h
--- a/Engine/Source/Runtime/Engine/Classes/Engine/SkeletalMesh.h
+++ b/Engine/Source/Runtime/Engine/Classes/Engine/SkeletalMesh.h
@@ -2516,7 +2516,8 @@ public:
 	 * This function will enforce valid material index in the sections and the LODMaterialMap of all LOD.
 	 */
 	ENGINE_API void ValidateAllLodMaterialIndexes();
-
+	
+private:
 	//We want to avoid calling post edit change multiple time during import and build process.
 
 	/*
@@ -2531,12 +2532,8 @@ public:
 	 */
 	ENGINE_API int32 UnStackPostEditChange();
 
-	int32 GetPostEditChangeStackCounter() { return PostEditChangeStackCounter; }
-	void SetPostEditChangeStackCounter(int32 InPostEditChangeStackCounter)
-	{
-		PostEditChangeStackCounter = InPostEditChangeStackCounter;
-	}
-
+	friend class FScopedSkeletalMeshPostEditChange;
+public:
 	/**
 	* If derive data cache key do not match, regenerate derived data and re-create any render state based on that.
 	*/
