Fix for the curve debugger window so that it can work in three contexts: in an anim bp when a pose watch element is active, in an anim bp when a pose watch element isn't active, and when in the anim sequence & skeleton editors.  Also fixed some small issues with filters not updating the curve list, incorrect mat/morph icons displaying, values updating incorrectly, etc.

diff --git a/Engine/Source/Editor/Persona/Private/SAnimCurveViewer.cpp b/Engine/Source/Editor/Persona/Private/SAnimCurveViewer.cpp
--- a/Engine/Source/Editor/Persona/Private/SAnimCurveViewer.cpp
+++ b/Engine/Source/Editor/Persona/Private/SAnimCurveViewer.cpp
@@ -103,37 +103,56 @@ private:
 bool FDisplayedAnimCurveInfo::GetActiveFlag(const TSharedPtr<SAnimCurveViewer>& InAnimCurveViewer, bool bInMorphTarget) const
 {
 	const UAnimInstance* AnimInstance = InAnimCurveViewer->GetAnimInstance();
-	if(AnimInstance && InAnimCurveViewer->PoseWatchPicker.IsValid())
+	if (AnimInstance)
 	{
+		auto FindMorphOrMatCurveOnAnimInstance = [AnimInstance](const FName& CurveName, bool bIsMorphTarget) {
+			const TMap<FName, float>& CurveList = AnimInstance->GetAnimationCurveList(bIsMorphTarget ? EAnimCurveType::MorphTargetCurve : EAnimCurveType::MaterialCurve);
+			return CurveList.Find(CurveName) != nullptr;
+		};
+
 		// Find if we want to use a pose watch
-		if(UPoseWatchPoseElement* PoseWatchPoseElement = InAnimCurveViewer->PoseWatchPicker->GetCurrentPoseWatch())
+		if (InAnimCurveViewer->PoseWatchPicker.IsValid())
 		{
-			if(UAnimBlueprintGeneratedClass* AnimClass = Cast<UAnimBlueprintGeneratedClass>(AnimInstance->GetClass()))
+			if (UPoseWatchPoseElement* PoseWatchPoseElement = InAnimCurveViewer->PoseWatchPicker->GetCurrentPoseWatch())
 			{
-				// We have to grab our pose watches from the root class as no pose watches can be set on child anim BPs
-				if(const UAnimBlueprintGeneratedClass* RootClass = Cast<UAnimBlueprintGeneratedClass>(AnimClass->GetRootClass()))
+				if (UAnimBlueprintGeneratedClass* AnimClass = Cast<UAnimBlueprintGeneratedClass>(AnimInstance->GetClass()))
 				{
-					const FAnimBlueprintDebugData& DebugData = RootClass->AnimBlueprintDebugData;
-					for(const FAnimNodePoseWatch& AnimNodePoseWatch : DebugData.AnimNodePoseWatch)
+					// We have to grab our pose watches from the root class as no pose watches can be set on child anim BPs
+					if (const UAnimBlueprintGeneratedClass* RootClass = Cast<UAnimBlueprintGeneratedClass>(AnimClass->GetRootClass()))
 					{
-						if(AnimNodePoseWatch.PoseWatchPoseElement == PoseWatchPoseElement)
+						const FAnimBlueprintDebugData& DebugData = RootClass->AnimBlueprintDebugData;
+						for (const FAnimNodePoseWatch& AnimNodePoseWatch : DebugData.AnimNodePoseWatch)
 						{
-							UE::Anim::ECurveElementFlags Flags = AnimNodePoseWatch.GetCurves().GetFlags(CurveName);
-							return EnumHasAnyFlags(Flags, bInMorphTarget ? UE::Anim::ECurveElementFlags::MorphTarget : UE::Anim::ECurveElementFlags::Material);
+							if (AnimNodePoseWatch.PoseWatchPoseElement == PoseWatchPoseElement)
+							{
+								UE::Anim::ECurveElementFlags Flags = AnimNodePoseWatch.GetCurves().GetFlags(CurveName);
+								return EnumHasAnyFlags(Flags, bInMorphTarget ? UE::Anim::ECurveElementFlags::MorphTarget : UE::Anim::ECurveElementFlags::Material);
+							}
 						}
 					}
 				}
 			}
+			else
+			{
+				// We are just displaying the output curves, so find if the curve exists on the anim instance
+				return FindMorphOrMatCurveOnAnimInstance(CurveName, bInMorphTarget);
+			}
 		}
 		else
 		{
-			// See if curve is in active set, attribute curve should have everything
-			const TMap<FName, float>& CurveList = AnimInstance->GetAnimationCurveList(bInMorphTarget ? EAnimCurveType::MorphTargetCurve : EAnimCurveType::MaterialCurve);
-			const float* CurrentValue = CurveList.Find(CurveName);
-			if (CurrentValue)
+			// We aren't using pose watch so check on the anim instance and the skeleton
+			if (FindMorphOrMatCurveOnAnimInstance(CurveName, bInMorphTarget))
 			{
 				return true;
 			}
+
+			if (TSharedPtr<IEditableSkeleton> EditableSkeleton = InAnimCurveViewer->EditableSkeletonPtr.Pin())
+			{
+				if (const FCurveMetaData* CurveMetaData = EditableSkeleton->GetSkeleton().GetCurveMetaData(CurveName))
+				{
+					return bInMorphTarget ? CurveMetaData->Type.bMorphtarget : CurveMetaData->Type.bMaterial;
+				}
+			}
 		}
 	}
 
@@ -394,56 +413,59 @@ FText SAnimCurveListRow::GetFilterText() const
 
 bool SAnimCurveListRow::GetActiveWeight(float& OutWeight) const
 {
-	bool bFoundActive = false;
-
 	// If anim viewer
 	TSharedPtr<SAnimCurveViewer> AnimCurveViewer = AnimCurveViewerPtr.Pin();
 	if (AnimCurveViewer.IsValid())
 	{
 		const UAnimInstance* AnimInstance = AnimCurveViewer->GetAnimInstance();
-		if(AnimInstance && AnimCurveViewer->PoseWatchPicker.IsValid())
+		if(AnimInstance)
 		{
-			// Find if we want to use a pose watch
-			if(UPoseWatchPoseElement* PoseWatchPoseElement = AnimCurveViewer->PoseWatchPicker->GetCurrentPoseWatch())
+			// If we're dealing with a pose watch, get the data from that
+			if (AnimCurveViewer->PoseWatchPicker.IsValid())
 			{
-				if(UAnimBlueprintGeneratedClass* AnimClass = Cast<UAnimBlueprintGeneratedClass>(AnimInstance->GetClass()))
+				// Find if we want to use a pose watch
+				if (UPoseWatchPoseElement* PoseWatchPoseElement = AnimCurveViewer->PoseWatchPicker->GetCurrentPoseWatch())
 				{
-					// We have to grab our pose watches from the root class as no pose watches can be set on child anim BPs
-					if(const UAnimBlueprintGeneratedClass* RootClass = Cast<UAnimBlueprintGeneratedClass>(AnimClass->GetRootClass()))
+					if (UAnimBlueprintGeneratedClass* AnimClass = Cast<UAnimBlueprintGeneratedClass>(AnimInstance->GetClass()))
 					{
-						const FAnimBlueprintDebugData& DebugData = RootClass->AnimBlueprintDebugData;
-						for(const FAnimNodePoseWatch& AnimNodePoseWatch : DebugData.AnimNodePoseWatch)
+						// We have to grab our pose watches from the root class as no pose watches can be set on child anim BPs
+						if (const UAnimBlueprintGeneratedClass* RootClass = Cast<UAnimBlueprintGeneratedClass>(AnimClass->GetRootClass()))
 						{
-							if(AnimNodePoseWatch.PoseWatchPoseElement == PoseWatchPoseElement)
+							const FAnimBlueprintDebugData& DebugData = RootClass->AnimBlueprintDebugData;
+							for (const FAnimNodePoseWatch& AnimNodePoseWatch : DebugData.AnimNodePoseWatch)
 							{
-								bool bHasElement = false;
-								float CurrentValue = AnimNodePoseWatch.GetCurves().Get(Item->CurveName, bHasElement);
-								if(bHasElement)
+								if (AnimNodePoseWatch.PoseWatchPoseElement == PoseWatchPoseElement)
 								{
-									OutWeight = CurrentValue;
-									bFoundActive = true;
+									bool bHasElement = false;
+									float CurrentValue = AnimNodePoseWatch.GetCurves().Get(Item->CurveName, bHasElement);
+									if (bHasElement)
+									{
+										OutWeight = CurrentValue;
+										return true;
+									}
+									break;
 								}
-								break;
 							}
 						}
 					}
+
+					// didn't find the curve on the pose watch element
+					return false;
 				}
 			}
-			else
+
+			// We aren't dealing with a pose watch so go directly to the anim instance curves
+			const TMap<FName, float>& CurveList = AnimInstance->GetAnimationCurveList(EAnimCurveType::AttributeCurve);
+			const float* CurrentValue = CurveList.Find(Item->CurveName);
+			if (CurrentValue)
 			{
-				// See if curve is in active set, attribute curve should have everything
-				const TMap<FName, float>& CurveList = AnimInstance->GetAnimationCurveList(EAnimCurveType::AttributeCurve);
-				const float* CurrentValue = CurveList.Find(Item->CurveName);
-				if (CurrentValue)
-				{
-					OutWeight = *CurrentValue;
-					bFoundActive = true;
-				}
+				OutWeight = *CurrentValue;
+				return true;
 			}
 		}
 	}
 
-	return bFoundActive;
+	return false;
 }
 
 
@@ -798,58 +820,66 @@ void SAnimCurveViewer::CreateAnimCurveList( const FString& SearchText, bool bInF
 	// Add active curves if required
 	TSet<FName> ActiveCurves;
 	const UAnimInstance* AnimInstance = GetAnimInstance();
-	if(AnimInstance && PoseWatchPicker.IsValid())
+
+	if(AnimInstance)
 	{
-		// Find if we want to use a pose watch
-		if(UPoseWatchPoseElement* PoseWatchPoseElement = PoseWatchPicker->GetCurrentPoseWatch())
+		bool bIsPoseWatchActive = false;
+		if (PoseWatchPicker.IsValid())
 		{
-			if(UAnimBlueprintGeneratedClass* AnimClass = Cast<UAnimBlueprintGeneratedClass>(AnimInstance->GetClass()))
+			// Find if we want to use a pose watch
+			if (UPoseWatchPoseElement* PoseWatchPoseElement = PoseWatchPicker->GetCurrentPoseWatch())
 			{
-				// We have to grab our pose watches from the root class as no pose watches can be set on child anim BPs
-				if(const UAnimBlueprintGeneratedClass* RootClass = Cast<UAnimBlueprintGeneratedClass>(AnimClass->GetRootClass()))
+				bIsPoseWatchActive = true;
+				if (UAnimBlueprintGeneratedClass* AnimClass = Cast<UAnimBlueprintGeneratedClass>(AnimInstance->GetClass()))
 				{
-					const FAnimBlueprintDebugData& DebugData = RootClass->AnimBlueprintDebugData;
-					for(const FAnimNodePoseWatch& AnimNodePoseWatch : DebugData.AnimNodePoseWatch)
+					// We have to grab our pose watches from the root class as no pose watches can be set on child anim BPs
+					if (const UAnimBlueprintGeneratedClass* RootClass = Cast<UAnimBlueprintGeneratedClass>(AnimClass->GetRootClass()))
 					{
-						if(AnimNodePoseWatch.PoseWatchPoseElement == PoseWatchPoseElement)
+						const FAnimBlueprintDebugData& DebugData = RootClass->AnimBlueprintDebugData;
+						for (const FAnimNodePoseWatch& AnimNodePoseWatch : DebugData.AnimNodePoseWatch)
 						{
-							AnimNodePoseWatch.GetCurves().ForEachElement([this, &AddCurve, &ActiveCurves](const UE::Anim::FCurveElement& InElement)
+							if (AnimNodePoseWatch.PoseWatchPoseElement == PoseWatchPoseElement)
 							{
-								if (EnumHasAnyFlags(CurrentCurveFlag, EAnimCurveViewerFilterFlags::MorphTarget))
-								{
-									if(EnumHasAnyFlags(InElement.Flags, UE::Anim::ECurveElementFlags::MorphTarget))
+								AnimNodePoseWatch.GetCurves().ForEachElement([this, &AddCurve, &ActiveCurves](const UE::Anim::FCurveElement& InElement)
 									{
-										AddCurve(InElement.Name, InElement.Flags);
-										ActiveCurves.Add(InElement.Name);
-									}
-								}
-								if (EnumHasAnyFlags(CurrentCurveFlag, EAnimCurveViewerFilterFlags::Material))
-								{
-									if(EnumHasAnyFlags(InElement.Flags, UE::Anim::ECurveElementFlags::Material))
-									{
-										AddCurve(InElement.Name, InElement.Flags);
-										ActiveCurves.Add(InElement.Name);
-									}
-								}
-
-								// If we arent filtering by curve type, just show all curves
-								if(!EnumHasAnyFlags(CurrentCurveFlag, EAnimCurveViewerFilterFlags::MorphTarget | EAnimCurveViewerFilterFlags::Material))
-								{
-									AddCurve(InElement.Name, InElement.Flags);
-									ActiveCurves.Add(InElement.Name);
-								}
-							});
-							break;
+										if (EnumHasAnyFlags(CurrentCurveFlag, EAnimCurveViewerFilterFlags::MorphTarget))
+										{
+											if (EnumHasAnyFlags(InElement.Flags, UE::Anim::ECurveElementFlags::MorphTarget))
+											{
+												AddCurve(InElement.Name, InElement.Flags);
+												ActiveCurves.Add(InElement.Name);
+											}
+										}
+										if (EnumHasAnyFlags(CurrentCurveFlag, EAnimCurveViewerFilterFlags::Material))
+										{
+											if (EnumHasAnyFlags(InElement.Flags, UE::Anim::ECurveElementFlags::Material))
+											{
+												AddCurve(InElement.Name, InElement.Flags);
+												ActiveCurves.Add(InElement.Name);
+											}
+										}
+
+										// If we arent filtering by curve type, just show all curves
+										if (!EnumHasAnyFlags(CurrentCurveFlag, EAnimCurveViewerFilterFlags::MorphTarget | EAnimCurveViewerFilterFlags::Material))
+										{
+											AddCurve(InElement.Name, InElement.Flags);
+											ActiveCurves.Add(InElement.Name);
+										}
+									});
+								break;
+							}
 						}
 					}
 				}
 			}
 		}
-		else
+
+		// If there isn't an active pose watch, grab the curve list directly from the anim instance
+		if (!bIsPoseWatchActive)
 		{
 			if (EnumHasAnyFlags(CurrentCurveFlag, EAnimCurveViewerFilterFlags::MorphTarget))
 			{
-				for(const TPair<FName, float>& CurveValuePair : AnimInstance->GetAnimationCurveList(EAnimCurveType::MorphTargetCurve))
+				for (const TPair<FName, float>& CurveValuePair : AnimInstance->GetAnimationCurveList(EAnimCurveType::MorphTargetCurve))
 				{
 					AddCurve(CurveValuePair.Key, UE::Anim::ECurveElementFlags::MorphTarget);
 					ActiveCurves.Add(CurveValuePair.Key);
@@ -857,7 +887,7 @@ void SAnimCurveViewer::CreateAnimCurveList( const FString& SearchText, bool bInF
 			}
 			if (EnumHasAnyFlags(CurrentCurveFlag, EAnimCurveViewerFilterFlags::Material))
 			{
-				for(const TPair<FName, float>& CurveValuePair : AnimInstance->GetAnimationCurveList(EAnimCurveType::MaterialCurve))
+				for (const TPair<FName, float>& CurveValuePair : AnimInstance->GetAnimationCurveList(EAnimCurveType::MaterialCurve))
 				{
 					AddCurve(CurveValuePair.Key, UE::Anim::ECurveElementFlags::Material);
 					ActiveCurves.Add(CurveValuePair.Key);
@@ -865,9 +895,9 @@ void SAnimCurveViewer::CreateAnimCurveList( const FString& SearchText, bool bInF
 			}
 
 			// If we arent filtering by curve type, show 'attribute' curves
-			if(!EnumHasAnyFlags(CurrentCurveFlag, EAnimCurveViewerFilterFlags::MorphTarget | EAnimCurveViewerFilterFlags::Material))
+			if (!EnumHasAnyFlags(CurrentCurveFlag, EAnimCurveViewerFilterFlags::MorphTarget | EAnimCurveViewerFilterFlags::Material))
 			{
-				for(const TPair<FName, float>& CurveValuePair : AnimInstance->GetAnimationCurveList(EAnimCurveType::AttributeCurve))
+				for (const TPair<FName, float>& CurveValuePair : AnimInstance->GetAnimationCurveList(EAnimCurveType::AttributeCurve))
 				{
 					AddCurve(CurveValuePair.Key, UE::Anim::ECurveElementFlags::None);
 					ActiveCurves.Add(CurveValuePair.Key);
@@ -890,6 +920,16 @@ void SAnimCurveViewer::CreateAnimCurveList( const FString& SearchText, bool bInF
 			}
 		}
 
+		// Look at whether we need to filter out any curves that aren't morphs or materials
+		if (bAddToList)
+		{
+			if ((EnumHasAnyFlags(CurrentCurveFlag, EAnimCurveViewerFilterFlags::MorphTarget) && !CurveNameValuePair.Value->bMorphTarget)
+				|| (EnumHasAnyFlags(CurrentCurveFlag, EAnimCurveViewerFilterFlags::Material) && !CurveNameValuePair.Value->bMaterial))
+			{
+				bAddToList = false;
+			}
+		}
+
 		// If we passed that, see if we are filtering to only active
 		if (bAddToList && EnumHasAnyFlags(CurrentCurveFlag, EAnimCurveViewerFilterFlags::Active))
 		{
@@ -909,7 +949,7 @@ void SAnimCurveViewer::CreateAnimCurveList( const FString& SearchText, bool bInF
 		}
 	}
 
-	if(bDirty)
+	if (bDirty)
 	{
 		// Sort final list
 		struct FSortSmartNamesAlphabetically
