[Persona] Fixed missing forward declaration

diff --git a/Engine/Source/Editor/Persona/Private/PersonaMeshDetails.h b/Engine/Source/Editor/Persona/Private/PersonaMeshDetails.h
--- a/Engine/Source/Editor/Persona/Private/PersonaMeshDetails.h
+++ b/Engine/Source/Editor/Persona/Private/PersonaMeshDetails.h
@@ -28,6 +28,7 @@ class SUniformGridPanel;
 struct FSectionLocalizer;
 class IDetailCategoryBuilder;
 class IDetailGroup;
+class FScopedSkeletalMeshPostEditChange;
 
 //The parameter is the LOD index so we know which LOD need to be rebuild
 DECLARE_DELEGATE_RetVal_OneParam(bool, FIsLODSettingsEnabledDelegate, int32)
