PR #13797: Enabling FBoneSocket Target Pickers in Contexts Outside Anim Montages

diff --git a/Engine/Source/Editor/Persona/Private/AnimGraphNodeDetails.cpp b/Engine/Source/Editor/Persona/Private/AnimGraphNodeDetails.cpp
--- a/Engine/Source/Editor/Persona/Private/AnimGraphNodeDetails.cpp
+++ b/Engine/Source/Editor/Persona/Private/AnimGraphNodeDetails.cpp
@@ -575,6 +575,28 @@ USkeleton* FBoneReferenceCustomization::GetSkeleton() const
 		}
 	}
 
+	TArray<UPackage*> CandidatePackages;
+	StructProperty->GetOuterPackages(CandidatePackages);
+
+	// Scan each candidate Package for a provider Object
+	for (const UPackage* Package : CandidatePackages)
+	{
+		if (!Package)
+		{
+			continue; // skip invalid packages
+		}
+		ForEachObjectWithPackage(Package, [&Objects](UObject* Object)
+		{
+			if (Object->Implements<UBoneReferenceSkeletonProvider>())
+			{
+				Objects.Add(Object);
+				return true; // stop scanning this package because the object with the skeleton provider has been found
+			}
+			return false;
+		},
+		false, RF_ClassDefaultObject | RF_Transient);
+	}
+
 	auto FindSkeletonForObject = [this](UObject* InObject) -> USkeleton* 
 	{
 		for( ; InObject; InObject = InObject->GetOuter())
