Fix up previous change to support bone socket references

Correctly return false (halt) when finding a UBoneReferenceSkeletonProvider, rather than true (continue).

diff --git a/Engine/Source/Editor/Persona/Private/AnimGraphNodeDetails.cpp b/Engine/Source/Editor/Persona/Private/AnimGraphNodeDetails.cpp
--- a/Engine/Source/Editor/Persona/Private/AnimGraphNodeDetails.cpp
+++ b/Engine/Source/Editor/Persona/Private/AnimGraphNodeDetails.cpp
@@ -590,9 +590,9 @@ USkeleton* FBoneReferenceCustomization::GetSkeleton() const
 			if (Object->Implements<UBoneReferenceSkeletonProvider>())
 			{
 				Objects.Add(Object);
-				return true; // stop scanning this package because the object with the skeleton provider has been found
+				return false; // stop scanning this package because the object with the skeleton provider has been found
 			}
-			return false;
+			return true;
 		},
 		false, RF_ClassDefaultObject | RF_Transient);
 	}
