Add filter for non-additive animations in the animation editor asset browser

diff --git a/Engine/Source/Editor/Persona/Private/SAnimationSequenceBrowser.cpp b/Engine/Source/Editor/Persona/Private/SAnimationSequenceBrowser.cpp
--- a/Engine/Source/Editor/Persona/Private/SAnimationSequenceBrowser.cpp
+++ b/Engine/Source/Editor/Persona/Private/SAnimationSequenceBrowser.cpp
@@ -51,6 +51,7 @@
 #include "Viewports.h"
 #include "SceneInterface.h"
 #include "UObject/AssetRegistryTagsContext.h"
+#include "SEnumCombo.h"
 
 #define LOCTEXT_NAMESPACE "SequenceBrowser"
 
@@ -64,8 +65,52 @@ public:
 
 	// FFrontendFilter implementation
 	virtual FString GetName() const override { return TEXT("AdditiveAnimAssets"); }
-	virtual FText GetDisplayName() const override { return LOCTEXT("FFrontendFilter_AdditiveAnimAssets", "Additive Animations"); }
-	virtual FText GetToolTipText() const override { return LOCTEXT("FFrontendFilter_AdditiveAnimAssetsToolTip", "Show only animations that are additive."); }
+	virtual FText GetDisplayName() const override
+	{
+		return FText::Format(LOCTEXT("FFrontendFilter_AdditiveAnimAssets", "Additive Type ({0})"), ModeDisplayText);
+	}
+	virtual FText GetToolTipText() const override { return LOCTEXT("FFrontendFilter_AdditiveAnimAssetsToolTip", "Show only animations that match a given additive type."); }
+
+	virtual void ModifyContextMenu(FMenuBuilder& MenuBuilder) override
+	{
+		MenuBuilder.BeginSection(TEXT("FilterMode"), LOCTEXT("FilterMode", "Filter Mode"));
+		{
+			TSharedRef<SWidget> ComboBox = SNew(SEnumComboBox, StaticEnum<EAdditiveTypeFilterMode>())
+				.CurrentValue_Lambda([this]()
+				{
+					return static_cast<int32>(Mode);
+				})
+				.OnEnumSelectionChanged_Lambda([this](uint32 NewValue, ESelectInfo::Type)
+				{
+					Mode = static_cast<EAdditiveTypeFilterMode>(NewValue);
+					ModeDisplayText = StaticEnum<EAdditiveTypeFilterMode>()->GetDisplayNameTextByValue(NewValue);
+					
+					switch (Mode)
+					{
+					case EAdditiveTypeFilterMode::NonAdditive:
+						ModeTagString = StaticEnum<EAdditiveAnimationType>()->GetNameStringByValue(AAT_None);
+						bInvertComparison = false;
+						break;
+					case EAdditiveTypeFilterMode::AnyAdditive:
+						ModeTagString = StaticEnum<EAdditiveAnimationType>()->GetNameStringByValue(AAT_None);
+						bInvertComparison = true;
+						break;
+					case EAdditiveTypeFilterMode::LocalSpaceBase:
+						ModeTagString = StaticEnum<EAdditiveAnimationType>()->GetNameStringByValue(AAT_LocalSpaceBase);
+						bInvertComparison = false;
+						break;
+					case EAdditiveTypeFilterMode::RotationOffsetMeshSpace:
+						ModeTagString = StaticEnum<EAdditiveAnimationType>()->GetNameStringByValue(AAT_RotationOffsetMeshSpace);
+						bInvertComparison = false;
+					}
+					
+					OnChanged().Broadcast();
+				});
+				
+			MenuBuilder.AddWidget(ComboBox, LOCTEXT("Mode", "Mode"));
+		}
+		MenuBuilder.EndSection();
+	}
 
 	// IFilter implementation
 	virtual bool PassesFilter(FAssetFilterType InItem) const override
@@ -74,10 +119,15 @@ public:
 		if (InItem.Legacy_TryGetAssetData(ItemAssetData))
 		{
 			const FString TagValue = ItemAssetData.GetTagValueRef<FString>(GET_MEMBER_NAME_CHECKED(UAnimSequence, AdditiveAnimType));
-			return !TagValue.IsEmpty() && !TagValue.Equals(TEXT("AAT_None"));
+			return !TagValue.IsEmpty() && (bInvertComparison ^ TagValue.Equals(ModeTagString));
 		}
 		return false;
 	}
+	
+	EAdditiveTypeFilterMode Mode = EAdditiveTypeFilterMode::AnyAdditive;
+	bool bInvertComparison = true;
+	FString ModeTagString = StaticEnum<EAdditiveAnimationType>()->GetNameStringByValue(AAT_None);
+	FText ModeDisplayText = StaticEnum<EAdditiveTypeFilterMode>()->GetDisplayNameTextByValue(static_cast<int32>(EAdditiveTypeFilterMode::AnyAdditive));
 };
 
 /** A filter that displays animations that use a skeleton notify */
diff --git a/Engine/Source/Editor/Persona/Private/SAnimationSequenceBrowser.h b/Engine/Source/Editor/Persona/Private/SAnimationSequenceBrowser.h
--- a/Engine/Source/Editor/Persona/Private/SAnimationSequenceBrowser.h
+++ b/Engine/Source/Editor/Persona/Private/SAnimationSequenceBrowser.h
@@ -18,6 +18,8 @@
 #include "IAnimationSequenceBrowser.h"
 #include "AssetRegistry/ARFilter.h"
 
+#include "SAnimationSequenceBrowser.generated.h"
+
 class FSceneViewport;
 class FUICommandList;
 class IPersonaToolkit;
@@ -29,6 +31,15 @@ class USoundWave;
 class FFrontendFilter_Folder;
 class FFrontendFilter_SkeletonNotify;
 
+UENUM()
+enum class EAdditiveTypeFilterMode : uint8
+{
+	NonAdditive UMETA(DisplayName = "Non-Additive"),
+	AnyAdditive UMETA(DisplayName = "Any Additive"),
+	LocalSpaceBase UMETA(DisplayName = "Local Space"),
+	RotationOffsetMeshSpace UMETA(DisplayName = "Mesh Space"),
+};
+
 //////////////////////////////////////////////////////////////////////////
 // FAnimationAssetViewportClient
 class FAnimationAssetViewportClient : public FEditorViewportClient
