Added support to mirror a blendspace sample.

diff --git a/Engine/Source/Editor/Persona/Private/Customization/BlendSampleDetails.cpp b/Engine/Source/Editor/Persona/Private/Customization/BlendSampleDetails.cpp
--- a/Engine/Source/Editor/Persona/Private/Customization/BlendSampleDetails.cpp
+++ b/Engine/Source/Editor/Persona/Private/Customization/BlendSampleDetails.cpp
@@ -203,6 +203,13 @@ void FBlendSampleDetails::CustomizeDetails(class IDetailLayoutBuilder& DetailBui
 			GET_MEMBER_NAME_CHECKED(FBlendSample, bIncludeInAnalyseAll), FBlendSample::StaticStruct());
 		CategoryBuilder.AddProperty(IncludeInAnalyseAllProperty);
 
+		TSharedPtr<IPropertyHandle> MirrorProperty = DetailBuilder.GetProperty(
+			GET_MEMBER_NAME_CHECKED(FBlendSample, bMirror), FBlendSample::StaticStruct());
+		IDetailPropertyRow& MirrorPropertyRow = CategoryBuilder.AddProperty(MirrorProperty);
+
+		TAttribute<bool> IsMirrorDataTableValid = TAttribute<bool>::CreateLambda([this]() { return BlendSpace->MirrorDataTable != nullptr; });
+		MirrorPropertyRow.EditCondition(IsMirrorDataTableValid, nullptr);
+		
 		bool bShowAnalysis = false;
 		bool bAnalyzed[3] = { false, false, false };
 		// Note that the analyzed position won't change whilst this menu is open, but the original value might.
diff --git a/Engine/Source/Editor/Persona/Private/Customization/BlendSpaceDetails.cpp b/Engine/Source/Editor/Persona/Private/Customization/BlendSpaceDetails.cpp
--- a/Engine/Source/Editor/Persona/Private/Customization/BlendSpaceDetails.cpp
+++ b/Engine/Source/Editor/Persona/Private/Customization/BlendSpaceDetails.cpp
@@ -24,6 +24,7 @@
 
 #include "DetailBuilderTypes.h"
 #include "SWarningOrErrorBox.h"
+#include "Animation/MirrorDataTable.h"
 
 #define LOCTEXT_NAMESPACE "BlendSpaceDetails"
 
@@ -173,6 +174,7 @@ void FBlendSpaceDetails::CustomizeDetails(class IDetailLayoutBuilder& DetailBuil
 			DetailBuilder.HideCategory("Analysis");
 			DetailBuilder.HideCategory("AnalysisProperties");
 			DetailBuilder.HideCategory("Graph");
+			DetailBuilder.HideCategory("Mirroring");
 		}
 
 		if(UBlendSpaceGraph* BlendSpaceGraph = Cast<UBlendSpaceGraph>(BlendSpace->GetOuter()))
@@ -479,6 +481,7 @@ void FBlendSpaceDetails::CustomizeDetails(class IDetailLayoutBuilder& DetailBuil
 				TSharedPtr<IPropertyHandle> RateScaleProperty = BlendSampleProperty->GetChildHandle(GET_MEMBER_NAME_CHECKED(FBlendSample, RateScale));
 				TSharedPtr<IPropertyHandle> IncludeInAnalyseAllProperty = BlendSampleProperty->GetChildHandle(GET_MEMBER_NAME_CHECKED(FBlendSample, bIncludeInAnalyseAll));
 				TSharedPtr<IPropertyHandle> UseSingleFrameBlendingProperty = BlendSampleProperty->GetChildHandle(GET_MEMBER_NAME_CHECKED(FBlendSample, FrameIndexToSample));
+				TSharedPtr<IPropertyHandle> MirrorProperty = BlendSampleProperty->GetChildHandle(GET_MEMBER_NAME_CHECKED(FBlendSample, bMirror));
 
 				IDetailGroup& Group = DetailCategoryBuilder.AddGroup(FName("BlendSamples_Samples"), FText::GetEmpty());
 				Group.HeaderRow()
@@ -547,6 +550,10 @@ void FBlendSpaceDetails::CustomizeDetails(class IDetailLayoutBuilder& DetailBuil
 					IDetailPropertyRow& UseSingleFrameBlendingPropertyRow = Group.AddPropertyRow(UseSingleFrameBlendingProperty.ToSharedRef());
 					FBlendSampleDetails::GenerateSingleFrameBlendingWidget(UseSingleFrameBlendingPropertyRow, UseSingleFrameBlendingProperty.ToSharedRef(), AnimationProperty.ToSharedRef());
 					Group.AddPropertyRow(IncludeInAnalyseAllProperty.ToSharedRef());
+
+					IDetailPropertyRow& MirrorPropertyRow = Group.AddPropertyRow(MirrorProperty.ToSharedRef());
+					TAttribute<bool> IsMirrorDataTableValid = TAttribute<bool>::CreateLambda([this]() { return BlendSpace->MirrorDataTable != nullptr; });
+					MirrorPropertyRow.EditCondition(IsMirrorDataTableValid, nullptr);
 				}
 				else if(BlendSpaceNode.Get())
 				{
@@ -555,6 +562,42 @@ void FBlendSpaceDetails::CustomizeDetails(class IDetailLayoutBuilder& DetailBuil
 				}
 			}
 		}
+
+		/// Mirroring - only allow to set MDT that share the same skeleton.
+		{
+			TSharedPtr<IPropertyHandle> MDTPropertyHandle =  DetailBuilder.GetProperty(GET_MEMBER_NAME_CHECKED(UBlendSpace, MirrorDataTable));
+
+			if (!MDTPropertyHandle->IsValidHandle())
+			{
+				return;
+			}
+			
+			IDetailCategoryBuilder& MirroringCategory = DetailBuilder.EditCategory(FName("Mirroring"));
+			
+			IDetailPropertyRow& MirroringPropertyRow = MirroringCategory.AddProperty(MDTPropertyHandle);
+			FDetailWidgetRow& Row = MirroringPropertyRow.CustomWidget(true);
+	
+			TSharedPtr<SWidget> ValueWidget = SNew(SObjectPropertyEntryBox)
+				.PropertyHandle(MDTPropertyHandle)
+				.AllowedClass(UMirrorDataTable::StaticClass())
+				.AllowClear(true)
+				.OnShouldFilterAsset(FOnShouldFilterAsset::CreateLambda([Skeleton = BlendSpace->GetSkeleton()](const FAssetData& InAssetData)
+				{
+					return Skeleton ? !Skeleton->IsCompatibleForEditor(InAssetData) : true;
+				}));
+
+			Row.NameContent()
+			[
+				MDTPropertyHandle->CreatePropertyNameWidget()
+			];
+
+			Row.ValueContent()
+			.MinDesiredWidth(Row.ValueWidget.MinWidth)
+			.MaxDesiredWidth(Row.ValueWidget.MaxWidth)
+			[
+				ValueWidget.ToSharedRef()
+			];
+		}
 	}	
 }
 
diff --git a/Engine/Source/Runtime/Engine/Classes/Animation/BlendSpace.h b/Engine/Source/Runtime/Engine/Classes/Animation/BlendSpace.h
--- a/Engine/Source/Runtime/Engine/Classes/Animation/BlendSpace.h
+++ b/Engine/Source/Runtime/Engine/Classes/Animation/BlendSpace.h
@@ -176,6 +176,9 @@ struct FBlendSample
 	
 	UPROPERTY(EditAnywhere, Category = BlendSample, meta=(UIMin="0.01", UIMax="2.0", ClampMin="0.01", ClampMax="64.0"))
 	float RateScale = 1.0f;
+
+	UPROPERTY(EditAnywhere, Category = BlendSample)
+	bool bMirror = false;
 	
 	UPROPERTY(EditAnywhere, Category = BlendSample, meta=(InlineEditConditionToggle))
 	bool bUseSingleFrameForBlending = false;
@@ -224,7 +227,7 @@ struct FBlendSample
 	
 	bool operator==( const FBlendSample& Other ) const 
 	{
-		return (Other.Animation == Animation && Other.SampleValue == SampleValue && Other.bUseSingleFrameForBlending == bUseSingleFrameForBlending && Other.FrameIndexToSample == FrameIndexToSample && FMath::IsNearlyEqual(Other.RateScale, RateScale));
+		return (Other.Animation == Animation && Other.SampleValue == SampleValue && Other.bUseSingleFrameForBlending == bUseSingleFrameForBlending && bMirror == Other.bMirror && Other.FrameIndexToSample == FrameIndexToSample && FMath::IsNearlyEqual(Other.RateScale, RateScale));
 	}
 
 	ENGINE_API float GetSamplePlayLength() const;
@@ -856,6 +859,10 @@ public:
 	UPROPERTY(EditAnywhere, Category = InputInterpolation)
 	EPreferredTriangulationDirection PreferredTriangulationDirection = EPreferredTriangulationDirection::Tangential;
 
+	// Used to mirror any blend samples with bMirror set to true
+	UPROPERTY(Category=Mirroring, EditAnywhere)
+	TObjectPtr<UMirrorDataTable> MirrorDataTable = nullptr;
+	
 protected:
 
 	/**
diff --git a/Engine/Source/Runtime/Engine/Private/Animation/BlendSpace.cpp b/Engine/Source/Runtime/Engine/Private/Animation/BlendSpace.cpp
--- a/Engine/Source/Runtime/Engine/Private/Animation/BlendSpace.cpp
+++ b/Engine/Source/Runtime/Engine/Private/Animation/BlendSpace.cpp
@@ -21,6 +21,8 @@
 #include "Logging/MessageLog.h"
 #include "UObject/UnrealType.h"
 #include "UObject/UObjectThreadContext.h"
+#include "Animation/MirrorDataTable.h"
+
 #if WITH_EDITOR
 #endif
 
@@ -1227,6 +1229,18 @@ void UBlendSpace::GetAnimationPose_Internal(TArray<FBlendSampleData>& BlendSampl
 					ChildExtractionContext.CurrentTime = static_cast<double>(Time);
 					ChildExtractionContext.DeltaTimeRecord = BlendSampleDataCache[I].DeltaTimeRecord;
 					Sample.Animation->GetAnimationPose(ChildAnimationPoseData, ChildExtractionContext);
+
+					if (Sample.bMirror && MirrorDataTable)
+					{
+						// @todo: To improve performance, we need to have a mirror map per LOD to avoid computing every time.
+						TCustomBoneIndexArray<FCompactPoseBoneIndex, FCompactPoseBoneIndex> CompactPoseMirrorBones;
+						TCustomBoneIndexArray<FQuat, FCompactPoseBoneIndex> ComponentSpaceRefRotations;
+						MirrorDataTable->FillCompactPoseAndComponentRefRotations(ChildAnimationPoseData.GetPose().GetBoneContainer(), CompactPoseMirrorBones, ComponentSpaceRefRotations);
+					
+						FAnimationRuntime::MirrorPose(ChildAnimationPoseData.GetPose(), MirrorDataTable->MirrorAxis, CompactPoseMirrorBones, ComponentSpaceRefRotations);
+						FAnimationRuntime::MirrorCurves(ChildAnimationPoseData.GetCurve(), *MirrorDataTable);
+						UE::Anim::Attributes::MirrorAttributes(ChildAnimationPoseData.GetAttributes(), *MirrorDataTable, CompactPoseMirrorBones);
+					}
 				}
 				else
 				{
@@ -1952,7 +1966,7 @@ void UBlendSpace::InitializePerBoneBlend()
 
 void UBlendSpace::TickFollowerSamples(
 	TArray<FBlendSampleData>& SampleDataList, const int32 HighestWeightIndex, FAnimAssetTickContext& Context, 
-	bool bResetMarkerDataOnFollowers, bool bLooping, const UMirrorDataTable* MirrorDataTable) const
+	bool bResetMarkerDataOnFollowers, bool bLooping, const UMirrorDataTable* InMirrorDataTable) const
 {
 	for (int32 SampleIndex = 0; SampleIndex < SampleDataList.Num(); ++SampleIndex)
 	{
@@ -1970,7 +1984,7 @@ void UBlendSpace::TickFollowerSamples(
 			{
 				Sample.Animation->TickByMarkerAsFollower(
 					SampleDataItem.MarkerTickRecord, Context.MarkerTickContext, SampleDataItem.Time, 
-					SampleDataItem.PreviousTime, Context.GetLeaderDelta(), bLooping, MirrorDataTable);
+					SampleDataItem.PreviousTime, Context.GetLeaderDelta(), bLooping, InMirrorDataTable);
 			}
 		}
 	}
