Added blendspace sample mirror support to blendspace graph

#rnx
#Synced-CL 47822097

diff --git a/Engine/Source/Editor/Persona/Private/Customization/BlendSampleDetails.cpp b/Engine/Source/Editor/Persona/Private/Customization/BlendSampleDetails.cpp
--- a/Engine/Source/Editor/Persona/Private/Customization/BlendSampleDetails.cpp
+++ b/Engine/Source/Editor/Persona/Private/Customization/BlendSampleDetails.cpp
@@ -397,6 +397,14 @@ void FBlendSampleDetails::CustomizeDetails(class IDetailLayoutBuilder& DetailBui
 	}
 	else
 	{
+		// Expose mirror option
+		TSharedPtr<IPropertyHandle> MirrorProperty = DetailBuilder.GetProperty(GET_MEMBER_NAME_CHECKED(FBlendSample, bMirror), FBlendSample::StaticStruct());
+		IDetailPropertyRow& MirrorPropertyRow = CategoryBuilder.AddProperty(MirrorProperty);
+
+		TAttribute<bool> IsMirrorDataTableValid = TAttribute<bool>::CreateLambda([this]() { return BlendSpace->MirrorDataTable != nullptr; });
+		MirrorPropertyRow.EditCondition(IsMirrorDataTableValid, nullptr);
+		
+		// Edit graph
 		check(BlendSpaceNode);
 		FDetailWidgetRow& GraphRow = CategoryBuilder.AddCustomRow(FText::FromString(TEXT("Graph")));
 		FBlendSampleDetails::GenerateSampleGraphWidget(GraphRow, BlendSpaceNode, SampleIndex);
diff --git a/Engine/Source/Editor/Persona/Private/Customization/BlendSpaceDetails.cpp b/Engine/Source/Editor/Persona/Private/Customization/BlendSpaceDetails.cpp
--- a/Engine/Source/Editor/Persona/Private/Customization/BlendSpaceDetails.cpp
+++ b/Engine/Source/Editor/Persona/Private/Customization/BlendSpaceDetails.cpp
@@ -174,7 +174,6 @@ void FBlendSpaceDetails::CustomizeDetails(class IDetailLayoutBuilder& DetailBuil
 			DetailBuilder.HideCategory("Analysis");
 			DetailBuilder.HideCategory("AnalysisProperties");
 			DetailBuilder.HideCategory("Graph");
-			DetailBuilder.HideCategory("Mirroring");
 		}
 
 		if(UBlendSpaceGraph* BlendSpaceGraph = Cast<UBlendSpaceGraph>(BlendSpace->GetOuter()))
@@ -557,6 +556,13 @@ void FBlendSpaceDetails::CustomizeDetails(class IDetailLayoutBuilder& DetailBuil
 				}
 				else if(BlendSpaceNode.Get())
 				{
+					// Expose mirror option
+					IDetailPropertyRow& MirrorPropertyRow = Group.AddPropertyRow(MirrorProperty.ToSharedRef());
+
+					TAttribute<bool> IsMirrorDataTableValid = TAttribute<bool>::CreateLambda([this]() { return BlendSpace->MirrorDataTable != nullptr; });
+					MirrorPropertyRow.EditCondition(IsMirrorDataTableValid, nullptr);
+
+					// Edit graph
 					FDetailWidgetRow& GraphRow = Group.AddWidgetRow();
 					FBlendSampleDetails::GenerateSampleGraphWidget(GraphRow, BlendSpaceNode.Get(), SampleIndex);
 				}
@@ -576,14 +582,26 @@ void FBlendSpaceDetails::CustomizeDetails(class IDetailLayoutBuilder& DetailBuil
 			
 			IDetailPropertyRow& MirroringPropertyRow = MirroringCategory.AddProperty(MDTPropertyHandle);
 			FDetailWidgetRow& Row = MirroringPropertyRow.CustomWidget(true);
-	
+			
 			TSharedPtr<SWidget> ValueWidget = SNew(SObjectPropertyEntryBox)
 				.PropertyHandle(MDTPropertyHandle)
 				.AllowedClass(UMirrorDataTable::StaticClass())
 				.AllowClear(true)
-				.OnShouldFilterAsset(FOnShouldFilterAsset::CreateLambda([Skeleton = BlendSpace->GetSkeleton()](const FAssetData& InAssetData)
+				.OnShouldFilterAsset(FOnShouldFilterAsset::CreateLambda([Skeleton = BlendSpace->GetSkeleton(), bGraph = !BlendSpace->IsAsset() && BlendSpaceNode.IsValid(), BlendSpaceNode = BlendSpaceNode](const FAssetData& InAssetData)
 				{
-					return Skeleton ? !Skeleton->IsCompatibleForEditor(InAssetData) : true;
+					if (!bGraph)
+					{
+						return Skeleton ? !Skeleton->IsCompatibleForEditor(InAssetData) : true;
+					}
+					else
+					{
+						if (UAnimBlueprint* AnimBlueprint = BlendSpaceNode->GetAnimBlueprint())
+						{
+							return AnimBlueprint->TargetSkeleton ? !AnimBlueprint->TargetSkeleton->IsCompatibleForEditor(InAssetData) : true;
+						}
+
+						return true;
+					}
 				}));
 
 			Row.NameContent()
diff --git a/Engine/Source/Runtime/Engine/Private/Animation/BlendSpace.cpp b/Engine/Source/Runtime/Engine/Private/Animation/BlendSpace.cpp
--- a/Engine/Source/Runtime/Engine/Private/Animation/BlendSpace.cpp
+++ b/Engine/Source/Runtime/Engine/Private/Animation/BlendSpace.cpp
@@ -1217,6 +1217,18 @@ void UBlendSpace::GetAnimationPose_Internal(TArray<FBlendSampleData>& BlendSampl
 				ChildrenPoses[I] = MoveTemp(ChildPoseContext.Pose);
 				ChildrenCurves[I] = MoveTemp(ChildPoseContext.Curve);
 				ChildrenAttributes[I] = MoveTemp(ChildPoseContext.CustomAttributes);
+
+				if (Sample.bMirror && MirrorDataTable)
+				{
+					// @todo: To improve performance, we need to have a mirror map per LOD to avoid computing every time.
+					TCustomBoneIndexArray<FCompactPoseBoneIndex, FCompactPoseBoneIndex> CompactPoseMirrorBones;
+					TCustomBoneIndexArray<FQuat, FCompactPoseBoneIndex> ComponentSpaceRefRotations;
+					MirrorDataTable->FillCompactPoseAndComponentRefRotations(ChildrenPoses[I].GetBoneContainer(), CompactPoseMirrorBones, ComponentSpaceRefRotations);
+					
+					FAnimationRuntime::MirrorPose(ChildrenPoses[I], MirrorDataTable->MirrorAxis, CompactPoseMirrorBones, ComponentSpaceRefRotations);
+					FAnimationRuntime::MirrorCurves(ChildrenCurves[I], *MirrorDataTable);
+					UE::Anim::Attributes::MirrorAttributes(ChildrenAttributes[I], *MirrorDataTable, CompactPoseMirrorBones);
+				}
 			}
 			else
 			{
