Cloth:
- Added Unbinding per section function to the UClothingAssetBase interface.
- Replaced all Unbinding implementations to a single function using the new interface.
- Removed the Clothing Data property inspector from the SkeletalMesh detail panel.
- Added SKM Cloth Asset Clothing Data stats.
- Made the MeshClothingAsset array read only to the Skeletal Mesh property UI.
- Reworked the stat display to use a WholeRowContent and use the Clothing Data name as header.
- Fixed Undo on deleting Clothing Data.
- Fixed Undo on renaming Clothing Data.

[RN] Skeletal Mesh Cloth Asset - Fixed unbinding of a single section. Fixed Undo renaming and deleting Clothing Data.
#Synced-CL 48035595

diff --git a/Engine/Plugins/ChaosClothAsset/Source/ChaosClothAssetEngine/Private/ChaosClothAsset/ClothAssetSKMClothingAsset.cpp b/Engine/Plugins/ChaosClothAsset/Source/ChaosClothAssetEngine/Private/ChaosClothAsset/ClothAssetSKMClothingAsset.cpp
--- a/Engine/Plugins/ChaosClothAsset/Source/ChaosClothAssetEngine/Private/ChaosClothAsset/ClothAssetSKMClothingAsset.cpp
+++ b/Engine/Plugins/ChaosClothAsset/Source/ChaosClothAssetEngine/Private/ChaosClothAsset/ClothAssetSKMClothingAsset.cpp
@@ -54,7 +54,7 @@ namespace UE::Chaos::ClothAsset::Private
 							if (Binding.LODIndex == LODIndex && Binding.Asset == &ClothingAsset)
 							{
 								const int32 OriginalDataSectionIndex = LODModel.Sections[Binding.SectionIndex].OriginalDataSectionIndex;
-								Binding.Asset->UnbindFromSkeletalMesh(&SkeletalMesh, Binding.LODIndex);
+								Binding.Asset->UnbindFromSkeletalMesh(&SkeletalMesh, Binding.LODIndex, Binding.SectionIndex);
 								Binding.SectionIndex = OriginalDataSectionIndex;
 
 								FSkelMeshSourceSectionUserData& SectionUserData = LODModel.UserSectionsData.FindChecked(OriginalDataSectionIndex);
@@ -473,25 +473,38 @@ bool UChaosClothAssetSKMClothingAsset::BindToSkeletalMesh(USkeletalMesh* Skeleta
 	return true;
 }
 
-void UChaosClothAssetSKMClothingAsset::UnbindFromSkeletalMesh(USkeletalMesh* SkeletalMesh)
+void UChaosClothAssetSKMClothingAsset::UnbindFromSkeletalMesh(USkeletalMesh* SkeletalMesh, const int32 MeshLodIndex, const int32 SectionIndex)
 {
 	if (SkeletalMesh)
 	{
-		ClothingAssetUtils::UnbindFromSkeletalMesh(*this, *SkeletalMesh);
+		ClothingAssetUtils::UnbindFromSkeletalMesh(*this, *SkeletalMesh, MeshLodIndex, SectionIndex);
 	}
 }
 
-void UChaosClothAssetSKMClothingAsset::UnbindFromSkeletalMesh(USkeletalMesh* SkeletalMesh, const int32 MeshLodIndex)
+void UChaosClothAssetSKMClothingAsset::UpdateAllLODBiasMappings(USkeletalMesh* SkeletalMesh)
 {
-	if (SkeletalMesh)
-	{
-		ClothingAssetUtils::UnbindFromSkeletalMesh(*this, *SkeletalMesh, MeshLodIndex);
-	}
+	// TODO: Implement a generic form of LOD Bias mapping in ClothingAssetUtils (need to rework the clothing asset interface for this)
 }
 
-void UChaosClothAssetSKMClothingAsset::UpdateAllLODBiasMappings(USkeletalMesh* SkeletalMesh)
+TArray<TPair<FText, FText>> UChaosClothAssetSKMClothingAsset::GetStats() const
 {
-	// TODO: Implement a generic form of LOD Bias mapping in ClothingAssetUtils (need to rework the clothing asset interface for this)
+	TArray<TPair<FText, FText>> Stats;
+	if (Asset && ClothSimulationModelIndex != INDEX_NONE)
+	{
+		if (const TSharedPtr<const FChaosClothSimulationModel> ClothSimulationModel = Asset->GetClothSimulationModel(ClothSimulationModelIndex))
+		{
+			for (int32 LodIndex = 0; LodIndex < ClothSimulationModel->GetNumLods(); ++LodIndex)
+			{
+				Stats.Emplace(FText::Format(LOCTEXT("LodIndexLabel", "LOD {0}"), FText::AsNumber(LodIndex)), FText{});
+				Stats.Emplace(LOCTEXT("VertexCount", "Vertices"), FText::AsNumber(ClothSimulationModel->GetNumVertices(LodIndex)));
+				Stats.Emplace(LOCTEXT("TriangleCount", "Triangles"), FText::AsNumber(ClothSimulationModel->GetNumTriangles(LodIndex)));
+				const FChaosClothSimulationLodModel& ClothSimulationLodModel = ClothSimulationModel->ClothSimulationLodModels[LodIndex];
+				Stats.Emplace(LOCTEXT("WeightMapCount", "Weight Maps"), FText::AsNumber(ClothSimulationLodModel.WeightMaps.Num()));
+				Stats.Emplace(LOCTEXT("MaxBoneWeights", "Bone Weights"), FText::AsNumber(ClothSimulationLodModel.CalcMaxBoneInfluences()));
+			}
+		}
+	}
+	return Stats;
 }
 
 void UChaosClothAssetSKMClothingAsset::OnModelChanged()
diff --git a/Engine/Plugins/ChaosClothAsset/Source/ChaosClothAssetEngine/Private/ChaosClothAsset/ClothSimulationModel.cpp b/Engine/Plugins/ChaosClothAsset/Source/ChaosClothAssetEngine/Private/ChaosClothAsset/ClothSimulationModel.cpp
--- a/Engine/Plugins/ChaosClothAsset/Source/ChaosClothAssetEngine/Private/ChaosClothAsset/ClothSimulationModel.cpp
+++ b/Engine/Plugins/ChaosClothAsset/Source/ChaosClothAssetEngine/Private/ChaosClothAsset/ClothSimulationModel.cpp
@@ -178,6 +178,16 @@ void FChaosClothSimulationLodModel::GetResourceSizeEx(FResourceSizeEx& Cumulativ
 	CumulativeResourceSize.AddDedicatedSystemMemoryBytes(Ar.TotalSize());
 }
 
+int32 FChaosClothSimulationLodModel::CalcMaxBoneInfluences() const
+{
+	int32 MaxBoneInfluences = 0;
+	for (const FClothVertBoneData& ClothVertBoneData : BoneData)
+	{
+		MaxBoneInfluences = FMath::Max(MaxBoneInfluences, ClothVertBoneData.NumInfluences);
+	}
+	return MaxBoneInfluences;
+}
+
 FChaosClothSimulationModel::FChaosClothSimulationModel(const TArray<TSharedRef<const FManagedArrayCollection>>& ClothCollections, 
 	const FReferenceSkeleton& ReferenceSkeleton,
 	TArray<FChaosClothAssetLodTransitionDataCache>* InOutTransitionCache)
diff --git a/Engine/Plugins/ChaosClothAsset/Source/ChaosClothAssetEngine/Public/ChaosClothAsset/ClothAssetSKMClothingAsset.h b/Engine/Plugins/ChaosClothAsset/Source/ChaosClothAssetEngine/Public/ChaosClothAsset/ClothAssetSKMClothingAsset.h
--- a/Engine/Plugins/ChaosClothAsset/Source/ChaosClothAssetEngine/Public/ChaosClothAsset/ClothAssetSKMClothingAsset.h
+++ b/Engine/Plugins/ChaosClothAsset/Source/ChaosClothAssetEngine/Public/ChaosClothAsset/ClothAssetSKMClothingAsset.h
@@ -61,9 +61,9 @@ private:
 
 	//~ Begin UClothingAssetBase Interface
 	virtual bool BindToSkeletalMesh(USkeletalMesh* SkeletalMesh, const int32 MeshLodIndex, const int32 SectionIndex, const int32 AssetLodIndex) override;
-	virtual void UnbindFromSkeletalMesh(USkeletalMesh* SkeletalMesh) override;
-	virtual void UnbindFromSkeletalMesh(USkeletalMesh* SkeletalMesh, const int32 MeshLodIndex) override;
+	virtual void UnbindFromSkeletalMesh(USkeletalMesh* SkeletalMesh, const int32 MeshLodIndexx, const int32 SectionIndex) override;
 	virtual void UpdateAllLODBiasMappings(USkeletalMesh* SkeletalMesh) override;
+	virtual TArray<TPair<FText, FText>> GetStats() const override;
 #endif  // #if WITH_EDITOR
 	virtual void RefreshBoneMapping(USkeletalMesh* SkeletalMesh) override {}
 	virtual bool IsValid() const override
diff --git a/Engine/Plugins/ChaosClothAsset/Source/ChaosClothAssetEngine/Public/ChaosClothAsset/ClothSimulationModel.h b/Engine/Plugins/ChaosClothAsset/Source/ChaosClothAssetEngine/Public/ChaosClothAsset/ClothSimulationModel.h
--- a/Engine/Plugins/ChaosClothAsset/Source/ChaosClothAssetEngine/Public/ChaosClothAsset/ClothSimulationModel.h
+++ b/Engine/Plugins/ChaosClothAsset/Source/ChaosClothAssetEngine/Public/ChaosClothAsset/ClothSimulationModel.h
@@ -79,6 +79,9 @@ struct FChaosClothSimulationLodModel
 	bool Serialize(FArchive& Ar);
 
 	void GetResourceSizeEx(FResourceSizeEx& CumulativeResourceSize);
+
+	/** Calculate the maximum number of bone influences used by the skinning weights. */
+	int32 CalcMaxBoneInfluences() const;
 };
 
 template<>
diff --git a/Engine/Source/Developer/SkeletalMeshUtilitiesCommon/Private/LODUtilities.cpp b/Engine/Source/Developer/SkeletalMeshUtilitiesCommon/Private/LODUtilities.cpp
--- a/Engine/Source/Developer/SkeletalMeshUtilitiesCommon/Private/LODUtilities.cpp
+++ b/Engine/Source/Developer/SkeletalMeshUtilitiesCommon/Private/LODUtilities.cpp
@@ -3503,7 +3503,7 @@ void FLODUtilities::UnbindClothingAndBackup(USkeletalMesh* SkeletalMesh, TArray<
 			int32 OriginalDataSectionIndex = LODModel.Sections[Binding.SectionIndex].OriginalDataSectionIndex;
 			if (Binding.Asset)
 			{
-				Binding.Asset->UnbindFromSkeletalMesh(SkeletalMesh, Binding.LODIndex);
+				Binding.Asset->UnbindFromSkeletalMesh(SkeletalMesh, Binding.LODIndex, INDEX_NONE);
 				Binding.SectionIndex = OriginalDataSectionIndex;
 			}
 			
diff --git a/Engine/Source/Editor/ClothPainter/Private/SClothAssetSelector.cpp b/Engine/Source/Editor/ClothPainter/Private/SClothAssetSelector.cpp
--- a/Engine/Source/Editor/ClothPainter/Private/SClothAssetSelector.cpp
+++ b/Engine/Source/Editor/ClothPainter/Private/SClothAssetSelector.cpp
@@ -179,6 +179,8 @@ public:
 							MakeUniqueObjectName(Asset->GetOuter(), UClothingAssetBase::StaticClass(), FName(SanitizedName)).ToString() :
 							SanitizedName;
 
+						FScopedTransaction CurrTransaction(LOCTEXT("RenameClothingAsset_Transaction", "Rename Clothing Data."));
+						Asset->Modify();
 						Asset->Rename(*NewName, Asset->GetOuter());
 					}
 				}
@@ -285,13 +287,17 @@ private:
 				int32 AssetIndex;
 				if(SkelMesh->GetMeshClothingAssets().Find(Asset, AssetIndex))
 				{
+					FScopedTransaction CurrTransaction(LOCTEXT("DeleteClothingAsset_Transaction", "Delete Clothing Data."));
+					SkelMesh->Modify();
+					Asset->Modify();
+
 					// Need to unregister our components so they shut down their current clothing simulation
 					FScopedSkeletalMeshPostEditChange ScopedPostEditChange(SkelMesh);
 					SkelMesh->PreEditChange(nullptr);
 
-					Asset->UnbindFromSkeletalMesh(SkelMesh);
+					Asset->UnbindFromSkeletalMesh(SkelMesh, INDEX_NONE, INDEX_NONE);
 					SkelMesh->GetMeshClothingAssets().RemoveAt(AssetIndex);
-					Asset->ConditionalBeginDestroy();  // Mark the clothing data as to be destroyed asap, so that the object name can be reused if needed
+					Asset->Rename(nullptr, GetTransientPackage());  // Rename so that the old object name can be reused if needed
 					Asset = nullptr;
 
 					// Need to fix up asset indices on sections.
@@ -1440,16 +1446,13 @@ TSharedRef<ITableRow> SClothAssetSelector::OnGenerateWidgetForClothingAssetItem(
 
 void SClothAssetSelector::OnAssetListSelectionChanged(TSharedPtr<FClothingAssetListItem> InSelectedItem, ESelectInfo::Type InSelectInfo)
 {
-	if (InSelectInfo != ESelectInfo::Direct)
+	if (InSelectedItem.IsValid())
 	{
-		if (InSelectedItem.IsValid())
-		{
-			SetSelectedAsset(InSelectedItem->ClothingAsset);
-		}
-		else
-		{
-			SetSelectedAsset(nullptr);
-		}
+		SetSelectedAsset(InSelectedItem->ClothingAsset);
+	}
+	else
+	{
+		SetSelectedAsset(nullptr);
 	}
 }
 
@@ -1561,6 +1564,10 @@ void SClothAssetSelector::RefreshAssetList()
 		{
 			AssetList->SetSelection(AssetListItems[SelectedItem]);
 		}
+		else
+		{
+			AssetList->ClearSelection();
+		}
 	}
 
 	// Refresh LOD list
diff --git a/Engine/Source/Editor/Persona/Private/PersonaMeshDetails.cpp b/Engine/Source/Editor/Persona/Private/PersonaMeshDetails.cpp
--- a/Engine/Source/Editor/Persona/Private/PersonaMeshDetails.cpp
+++ b/Engine/Source/Editor/Persona/Private/PersonaMeshDetails.cpp
@@ -6676,238 +6676,80 @@ void FPersonaMeshDetails::CustomizeClothingProperties(IDetailLayoutBuilder& Deta
 //
 void FPersonaMeshDetails::OnGenerateElementForClothingAsset( TSharedRef<IPropertyHandle> StructProperty, int32 ElementIndex, IDetailChildrenBuilder& ChildrenBuilder, IDetailLayoutBuilder* DetailLayout )
 {
-	const FSlateFontInfo DetailFontInfo = IDetailLayoutBuilder::GetDetailFont();
-
-	// Remove and reimport asset buttons
-	ChildrenBuilder.AddCustomRow( FText::GetEmpty() ) 
-	[
-		SNew(SHorizontalBox)
-		+ SHorizontalBox::Slot()
-		.FillWidth(1)
-		// remove button
-		+ SHorizontalBox::Slot()
-		.VAlign( VAlign_Center )
-		.Padding(2.f)
-		.AutoWidth()
-		[
-			SNew( SButton )
-			.Text( LOCTEXT("ClearButtonLabel", "Remove") )
-			.OnClicked( this, &FPersonaMeshDetails::OnRemoveClothingAssetClicked, ElementIndex, DetailLayout )
-			.IsFocusable( false )
-			.ContentPadding(0.f)
-			.ForegroundColor( FSlateColor::UseForeground() )
-			.ButtonColorAndOpacity(FLinearColor(1.0f,1.0f,1.0f,0.0f))
-			.ToolTipText(LOCTEXT("RemoveClothingAssetTip", "Remove this clothing asset"))
-			[ 
-				SNew( SImage )
-				.Image( FAppStyle::GetBrush("PropertyWindow.Button_Clear") )
-				.ColorAndOpacity( FSlateColor::UseForeground() )
-			]
-		]
-	];	
-
-	USkeletalMesh* CurrentMesh = GetPersonaToolkit()->GetMesh();
-	UClothingAssetBase* CurrentAsset = CurrentMesh->GetMeshClothingAssets()[ElementIndex];
-
-	ChildrenBuilder.AddCustomRow(LOCTEXT("ClothingAsset_Search_Name", "Name"))
-	.NameContent()
-	[
-		SNew(STextBlock)
-		.Text(LOCTEXT("ClothingAsset_Label_Name", "Name"))
-		.Font(DetailFontInfo)
-	]
-	.ValueContent()
-	.HAlign(HAlign_Fill)
-	[
-		SNew(STextBlock)
-		.Text(CurrentAsset ? FText::FromString(CurrentAsset->GetName()) : FText())
-	];
-
-	ChildrenBuilder.AddCustomRow(LOCTEXT("ClothingAsset_Search_Details", "Details"))
-	.NameContent()
-	[
-		SNew(STextBlock)
-		.Text(LOCTEXT("Details", "Details"))
-		.Font(DetailFontInfo)
-	]
-	.ValueContent()
-	.HAlign(HAlign_Fill)
-	[
-		MakeClothingDetailsWidget(ElementIndex)
-	];	
-
-	// Properties are now inside the clothing data, so we just add a new inspector and handle everything through that
-	FDetailWidgetRow& ClothPropRow = ChildrenBuilder.AddCustomRow(LOCTEXT("ClothingAsset_Search_Properties", "Properties"));
-
-	TSharedPtr<SKismetInspector> Inspector = nullptr;
-
-	ClothPropRow.WholeRowWidget
-	[
-		SNew(SExpandableArea)
-		.InitiallyCollapsed(true)
-		.HeaderContent()
-		[
-			SNew(STextBlock)
-			.Text(LOCTEXT("Properties_Header", "Clothing Properties"))
-		]
-		.BodyContent()
-		[
-			SAssignNew(Inspector, SKismetInspector)
-			.ShowTitleArea(false)
-			.ShowPublicViewControl(false)
-			.HideNameArea(true)
-			.IsPropertyEditingEnabledDelegate(FIsPropertyEditingEnabled::CreateSP(this, &FPersonaMeshDetails::IsClothingPanelEnabled))
-		]
-	];
-
-	SKismetInspector::FShowDetailsOptions Options;
-	Options.bHideFilterArea = true;
-	Options.bShowComponents = false;
-
-	if (CurrentAsset)
+	if (const USkeletalMesh* const CurrentMesh = GetPersonaToolkit()->GetMesh())
 	{
-		Inspector->ShowDetailsForSingleObject(CurrentAsset, Options);
-	}
-}
-
-TSharedRef<SUniformGridPanel> FPersonaMeshDetails::MakeClothingDetailsWidget(int32 AssetIndex) const
-{
-	TSharedRef<SUniformGridPanel> Grid = SNew(SUniformGridPanel).SlotPadding(2.0f);
-
-	USkeletalMesh* SkelMesh = GetPersonaToolkit()->GetMesh();
-
-	if (UClothingAssetCommon* const ClothingAsset = Cast<UClothingAssetCommon>(SkelMesh->GetMeshClothingAssets()[AssetIndex]))
-	{
-		const FSlateFontInfo DetailFontInfo = IDetailLayoutBuilder::GetDetailFont();
-		const int32 NumLODs = ClothingAsset->GetNumLods();
-		int32 RowNumber = 0;
-		for (int32 LODIndex = 0; LODIndex < NumLODs; ++LODIndex)
+		if (const UClothingAssetBase* const CurrentAsset = CurrentMesh->GetMeshClothingAssets()[ElementIndex])
 		{
-			Grid->AddSlot(0, RowNumber) // x, y
-				.HAlign(HAlign_Left)
-				[
-					SNew(STextBlock)
-					.Font(DetailFontInfo)
-					.Text(FText::Format(LOCTEXT("LODIndexLabel", "LOD {0}"), FText::AsNumber(LODIndex)))
-				];
-
-			++RowNumber;
-
-			const FClothLODDataCommon& LodData = ClothingAsset->LodData[LODIndex];
-			const FClothPhysicalMeshData& PhysMeshData = LodData.PhysicalMeshData;
-
-			Grid->AddSlot(0, RowNumber)
-				.HAlign(HAlign_Center)
-				[
-					SNew(STextBlock)
-					.Font(DetailFontInfo)
-					.Text(LOCTEXT("SimulVertexCount", "Simul Verts"))
-				];
-
-			Grid->AddSlot(0, RowNumber + 1)
-				.HAlign(HAlign_Center)
-				[
-					SNew(STextBlock)
-					.Font(DetailFontInfo)
-					.Text(FText::AsNumber(PhysMeshData.Vertices.Num() - PhysMeshData.NumFixedVerts))
-				];
+			const FSlateFontInfo DetailFontInfo = IDetailLayoutBuilder::GetDetailFont();
 
-			Grid->AddSlot(1, RowNumber)
-				.HAlign(HAlign_Center)
-				[
-					SNew(STextBlock)
-					.Font(DetailFontInfo)
-					.Text(LOCTEXT("FixedVertexCount", "Fixed Verts"))
-				];
-
-			Grid->AddSlot(1, RowNumber + 1)
-				.HAlign(HAlign_Center)
-				[
-					SNew(STextBlock)
-					.Font(DetailFontInfo)
-					.Text(FText::AsNumber(PhysMeshData.NumFixedVerts))
-				];
-
-			Grid->AddSlot(2, RowNumber)
-				.HAlign(HAlign_Center)
-				[
-					SNew(STextBlock)
-					.Font(DetailFontInfo)
-					.Text(LOCTEXT("TriangleCount", "Sim Triangles"))
-				];
-
-			Grid->AddSlot(2, RowNumber + 1)
-				.HAlign(HAlign_Center)
+			constexpr bool bStartExpanded = true;
+			IDetailGroup& Group = ChildrenBuilder.AddGroup(CurrentAsset->GetFName(), LOCTEXT("ClothingAssets_GroupName", "ClothingAssets"), bStartExpanded);
+			Group.HeaderRow()
+				.NameContent()
+				.HAlign(HAlign_Fill)
 				[
 					SNew(STextBlock)
-					.Font(DetailFontInfo)
-					.Text(FText::AsNumber(PhysMeshData.Indices.Num() / 3))
-				];
-
-			Grid->AddSlot(3, RowNumber)
-				.HAlign(HAlign_Center)
+						.Text(FText::Format(LOCTEXT("ClothingAsset_Label", "Index [ {0} ]"), FText::AsNumber(ElementIndex)))
+						.Font(DetailFontInfo)
+				]
+				.ValueContent()
 				[
 					SNew(STextBlock)
-					.Font(DetailFontInfo)
-					.Text(LOCTEXT("NumUsedBones", "Bones"))
+						.Text(FText::FromString(CurrentAsset->GetName()))
+						.Font(DetailFontInfo)
 				];
-
-			Grid->AddSlot(3, RowNumber + 1)
-				.HAlign(HAlign_Center)
+			Group.AddWidgetRow()
+				.WholeRowContent()
 				[
-					SNew(STextBlock)
-					.Font(DetailFontInfo)
-					.Text(FText::AsNumber(PhysMeshData.MaxBoneWeights))
+					MakeClothingDetailsWidget(ElementIndex)
 				];
-
-			RowNumber += 2;
 		}
 	}
-
-	return Grid;
 }
 
-FReply FPersonaMeshDetails::OnRemoveClothingAssetClicked(int32 AssetIndex, IDetailLayoutBuilder* DetailLayout)
+TSharedRef<SUniformGridPanel> FPersonaMeshDetails::MakeClothingDetailsWidget(int32 AssetIndex) const
 {
-	USkeletalMesh* SkelMesh = GetPersonaToolkit()->GetMesh();
-	check(SkelMesh);
+	const TSharedRef<SUniformGridPanel> Grid = SNew(SUniformGridPanel).SlotPadding(2.0f);
 
-	FScopedSuspendAlternateSkinWeightPreview ScopedSuspendAlternateSkinnWeightPreview(SkelMesh);
-	// Now we can remove the asset.
-	if(SkelMesh->GetMeshClothingAssets().IsValidIndex(AssetIndex))
+	if (const USkeletalMesh* const SkeletalMesh = GetPersonaToolkit()->GetMesh())
 	{
-		// Need to unregister our components so they shut down their current clothing simulation
-		FScopedSkeletalMeshPostEditChange ScopedPostEditChange(SkelMesh);
-
-		if (UClothingAssetBase* AssetToRemove = SkelMesh->GetMeshClothingAssets()[AssetIndex])
+		if (UClothingAssetBase* const ClothingAsset = SkeletalMesh->GetMeshClothingAssets()[AssetIndex])
 		{
-			AssetToRemove->UnbindFromSkeletalMesh(SkelMesh);
-		}
-		SkelMesh->GetMeshClothingAssets().RemoveAt(AssetIndex);
+			const FSlateFontInfo DetailFont = IDetailLayoutBuilder::GetDetailFont();
+			const FSlateFontInfo DetailFontBold = IDetailLayoutBuilder::GetDetailFontBold();
+			const TArray<TPair<FText, FText>> Stats = ClothingAsset->GetStats();
 
-		// Need to fix up asset indices on sections.
-		if(FSkeletalMeshModel* MeshResource = SkelMesh->GetImportedModel())
-		{
-			for(FSkeletalMeshLODModel& LodModel : MeshResource->LODModels)
+			int32 RowNumber = 0;
+			int32 ColumnNumber = 0;
+			for (const TPair<FText, FText>& Stat : Stats)
 			{
-				for(int32 SectionIndex = 0; SectionIndex < LodModel.Sections.Num(); ++SectionIndex)
+				if (Stat.Value.IsEmpty())  // Header
 				{
-					FSkelMeshSection& Section = LodModel.Sections[SectionIndex];
-					if(Section.CorrespondClothAssetIndex > AssetIndex)
+					if (ColumnNumber)
 					{
-						--Section.CorrespondClothAssetIndex;
-						//Keep the user section data (build source data) in sync
-						SetSkelMeshSourceSectionUserData(LodModel, SectionIndex, Section.OriginalDataSectionIndex);
+						RowNumber += 2;
+						ColumnNumber = 0;
 					}
 				}
+				Grid->AddSlot(ColumnNumber, RowNumber)
+					.HAlign(HAlign_Center)
+					[
+						SNew(STextBlock)
+							.Font(DetailFont)
+							.Text(Stat.Key)
+					];
+				Grid->AddSlot(ColumnNumber, RowNumber + 1)
+					.HAlign(HAlign_Center)
+					[
+						SNew(STextBlock)
+							.Font(DetailFont)
+							.Text(Stat.Value)
+					];
+				++ColumnNumber;
 			}
 		}
 	}
-
-	UpdateClothingEntries();
-	RefreshClothingComboBoxes();
-
-	return FReply::Handled();
+	return Grid;
 }
 
 void FPersonaMeshDetails::UpdateClothingEntries()
@@ -7051,7 +6893,7 @@ void FPersonaMeshDetails::OnClothingSelectionChanged(TSharedPtr<FClothingEntry>
 			if (UClothingAssetBase* CurrentAsset = Mesh->GetSectionClothingAsset(InLodIdx, InSectionIdx))
 			{
 				CurrentAsset->Modify();
-				CurrentAsset->UnbindFromSkeletalMesh(Mesh, InLodIdx);
+				CurrentAsset->UnbindFromSkeletalMesh(Mesh, InLodIdx, InSectionIdx);
 				ClearOriginalSectionUserData();
 			}
 
@@ -7081,7 +6923,7 @@ void FPersonaMeshDetails::OnClothingSelectionChanged(TSharedPtr<FClothingEntry>
 			if (UClothingAssetBase* CurrentAsset = Mesh->GetSectionClothingAsset(InLodIdx, InSectionIdx))
 			{
 				CurrentAsset->Modify();
-				CurrentAsset->UnbindFromSkeletalMesh(Mesh, InLodIdx);
+				CurrentAsset->UnbindFromSkeletalMesh(Mesh, InLodIdx, InSectionIdx);
 				ClearOriginalSectionUserData();
 			}
 		}
diff --git a/Engine/Source/Editor/SkeletalMeshEditor/Private/SkeletalMeshEditor.cpp b/Engine/Source/Editor/SkeletalMeshEditor/Private/SkeletalMeshEditor.cpp
--- a/Engine/Source/Editor/SkeletalMeshEditor/Private/SkeletalMeshEditor.cpp
+++ b/Engine/Source/Editor/SkeletalMeshEditor/Private/SkeletalMeshEditor.cpp
@@ -1002,6 +1002,10 @@ void FSkeletalMeshEditor::ApplyClothing(UClothingAssetBase* InAsset, int32 InLod
 	}
 
 	FSkeletalMeshLODModel& LODModel = Mesh->GetImportedModel()->LODModels[InLodIndex];
+	if (!LODModel.Sections.IsValidIndex(InSectionIndex))
+	{
+		return;
+	}
 	const FSkelMeshSection& Section = LODModel.Sections[InSectionIndex];
 	FScopedSuspendAlternateSkinWeightPreview ScopedSuspendAlternateSkinnWeightPreview(Mesh);
 	{
@@ -1024,7 +1028,7 @@ void FSkeletalMeshEditor::ApplyClothing(UClothingAssetBase* InAsset, int32 InLod
 			if (UClothingAssetBase* CurrentAsset = Mesh->GetSectionClothingAsset(InLodIndex, InSectionIndex))
 			{
 				CurrentAsset->Modify();
-				CurrentAsset->UnbindFromSkeletalMesh(Mesh, InLodIndex);
+				CurrentAsset->UnbindFromSkeletalMesh(Mesh, InLodIndex, InSectionIndex);
 				ClearOriginalSectionUserData();
 			}
 
@@ -1044,7 +1048,7 @@ void FSkeletalMeshEditor::ApplyClothing(UClothingAssetBase* InAsset, int32 InLod
 			if (UClothingAssetBase* CurrentAsset = Mesh->GetSectionClothingAsset(InLodIndex, InSectionIndex))
 			{
 				CurrentAsset->Modify();
-				CurrentAsset->UnbindFromSkeletalMesh(Mesh, InLodIndex);
+				CurrentAsset->UnbindFromSkeletalMesh(Mesh, InLodIndex, InSectionIndex);
 				ClearOriginalSectionUserData();
 			}
 		}
@@ -1053,14 +1057,19 @@ void FSkeletalMeshEditor::ApplyClothing(UClothingAssetBase* InAsset, int32 InLod
 
 void FSkeletalMeshEditor::RemoveClothing(int32 InLodIndex, int32 InSectionIndex)
 {
-	USkeletalMesh* Mesh = GetPersonaToolkit()->GetPreviewMesh();
+	USkeletalMesh* const Mesh = GetPersonaToolkit()->GetPreviewMesh();
 
-	if (Mesh->GetImportedModel() == nullptr || !Mesh->GetImportedModel()->LODModels.IsValidIndex(InLodIndex))
+	if (!Mesh || !Mesh->GetImportedModel() || !Mesh->GetImportedModel()->LODModels.IsValidIndex(InLodIndex))
 	{
 		return;
 	}
 
 	FSkeletalMeshLODModel& LODModel = Mesh->GetImportedModel()->LODModels[InLodIndex];
+	if (!LODModel.Sections.IsValidIndex(InSectionIndex))
+	{
+		return;
+	}
+
 	const FSkelMeshSection& Section = LODModel.Sections[InSectionIndex];
 	FScopedSuspendAlternateSkinWeightPreview ScopedSuspendAlternateSkinnWeightPreview(Mesh);
 	{
@@ -1079,7 +1088,7 @@ void FSkeletalMeshEditor::RemoveClothing(int32 InLodIndex, int32 InSectionIndex)
 		if (UClothingAssetBase* CurrentAsset = Mesh->GetSectionClothingAsset(InLodIndex, InSectionIndex))
 		{
 			CurrentAsset->Modify();
-			CurrentAsset->UnbindFromSkeletalMesh(Mesh, InLodIndex);
+			CurrentAsset->UnbindFromSkeletalMesh(Mesh, InLodIndex, InSectionIndex);
 			ClearOriginalSectionUserData();
 		}
 	}
diff --git a/Engine/Source/Runtime/ClothingSystemRuntimeCommon/Private/ClothingAsset.cpp b/Engine/Source/Runtime/ClothingSystemRuntimeCommon/Private/ClothingAsset.cpp
--- a/Engine/Source/Runtime/ClothingSystemRuntimeCommon/Private/ClothingAsset.cpp
+++ b/Engine/Source/Runtime/ClothingSystemRuntimeCommon/Private/ClothingAsset.cpp
@@ -41,63 +41,97 @@
 //==============================================================================
 
 #if WITH_EDITOR
-void ClothingAssetUtils::UnbindFromSkeletalMesh(const UClothingAssetBase& Asset, USkeletalMesh& SkeletalMesh)
+void ClothingAssetUtils::UnbindFromSkeletalMesh(const UClothingAssetBase& Asset, USkeletalMesh& SkeletalMesh, const int32 MeshLodIndex, const int32 LodSectionIndex)
 {
-	if (HasAnyMeshClothingAssetBindings(&SkeletalMesh))
+	FSkeletalMeshModel* const Mesh = SkeletalMesh.GetImportedModel();
+	if (!Mesh)
 	{
-		FScopedSkeletalMeshPostEditChange SkeletalMeshPostEditChange(&SkeletalMesh);
+		UE_LOG(LogClothingAsset, Warning,
+			TEXT("UnbindFromSkeletalMesh: Failed to unbind Clothing Data [%s] from Skeletal Mesh [%s], the Skeletal Mesh imported model doesn't exist."),
+			*Asset.GetName(),
+			*SkeletalMesh.GetName());
+		return;
+	}
+
+	int32 StartMeshLodIndex = 0;
+	int32 EndMeshLodIndex = -1;
+	if (MeshLodIndex == INDEX_NONE)
+	{
+		EndMeshLodIndex += Mesh->LODModels.Num();
+	}
+	else if (Mesh->LODModels.IsValidIndex(MeshLodIndex))
+	{
+		StartMeshLodIndex = EndMeshLodIndex = MeshLodIndex;
+	}
+	else
+	{
+		UE_LOG(LogClothingAsset, Warning,
+			TEXT("UnbindFromSkeletalMesh: Failed to unbind Clothing Data [%s] from Skeletal Mesh [%s], LOD%d doesn't exist."),
+			*Asset.GetName(),
+			*SkeletalMesh.GetName(),
+			MeshLodIndex);
+		return;
+	}
 
-		if (const FSkeletalMeshModel* const Mesh = SkeletalMesh.GetImportedModel())
+	bool bChangedMesh = false;
+	ON_SCOPE_EXIT  // If the mesh changed we need to re-register any components that use it to reflect the changes
 		{
-			const int32 NumLods = Mesh->LODModels.Num();
-			for (int32 LodIndex = 0; LodIndex < NumLods; ++LodIndex)
+			if (bChangedMesh)
 			{
-				UnbindFromSkeletalMesh(Asset, SkeletalMesh, LodIndex);
+				FScopedSkeletalMeshPostEditChange SkeletalMeshPostEditChange(&SkeletalMesh);
 			}
-		}
-	}
-}
-
-void ClothingAssetUtils::UnbindFromSkeletalMesh(const UClothingAssetBase& Asset, USkeletalMesh& SkeletalMesh, const int32 MeshLodIndex)
-{
-	if (HasAnyLodMeshClothingAssetBindings(&SkeletalMesh, MeshLodIndex))
+		};
+	for (int32 LodIndex = StartMeshLodIndex; LodIndex <= EndMeshLodIndex; ++LodIndex)
 	{
-		FSkeletalMeshModel* const Mesh = SkeletalMesh.GetImportedModel();
-		if (!Mesh || !Mesh->LODModels.IsValidIndex(MeshLodIndex))
+		FSkeletalMeshLODModel& LodModel = Mesh->LODModels[LodIndex];
+
+		int32 StartLodSectionIndex = 0;
+		int32 EndLodSectionIndex = -1;
+
+		if (MeshLodIndex == INDEX_NONE || LodSectionIndex == INDEX_NONE)
+		{
+			EndLodSectionIndex += LodModel.Sections.Num();
+		}
+		else if (LodModel.Sections.IsValidIndex(LodSectionIndex))
+		{
+			StartLodSectionIndex = EndLodSectionIndex = LodSectionIndex;
+		}
+		else
 		{
-			UE_LOG(LogClothingAsset, Warning, TEXT("Failed to remove clothing asset [%s] from mesh LOD%d as that LOD doesn't exist."),
+			UE_LOG(LogClothingAsset, Warning,
+				TEXT("UnbindFromSkeletalMesh: Failed to unbind Clothing Data [%s] from Skeletal Mesh [%s], LOD%d section %d doesn't exist."),
 				*Asset.GetName(),
-				MeshLodIndex);
+				*SkeletalMesh.GetName(),
+				MeshLodIndex,
+				LodSectionIndex);
 			return;
 		}
 
-		bool bChangedMesh = false;
-		FSkeletalMeshLODModel& LodModel = Mesh->LODModels[MeshLodIndex];
-
-		for (int32 SectionIdx = LodModel.Sections.Num() - 1; SectionIdx >= 0; --SectionIdx)
+		for (int32 SectionIndex = EndLodSectionIndex; SectionIndex >= StartLodSectionIndex; --SectionIndex)
 		{
-			FSkelMeshSection& Section = LodModel.Sections[SectionIdx];
+			FSkelMeshSection& Section = LodModel.Sections[SectionIndex];
 			if (Section.HasClothingData() && Section.ClothingData.AssetGuid == Asset.GetAssetGuid())
 			{
 				// Don't do this when in async task, this should have been done before initiating the build
-				if (IsInGameThread())
+				if (IsInGameThread() && !bChangedMesh)
 				{
 					SkeletalMesh.PreEditChange(nullptr);
 				}
 
 				// Log the LOD bias mapping data that will be removed
-				for (int32 LodIndex = 0; LodIndex < MeshLodIndex; ++LodIndex)
+				for (int32 BiasLodIndex = 0; BiasLodIndex < LodIndex; ++BiasLodIndex)
 				{
-					const int32 LODBias = MeshLodIndex - LodIndex;
+					const int32 LodBias = LodIndex - BiasLodIndex;
 
-					UE_CLOG(Section.ClothMappingDataLODs.IsValidIndex(LODBias) && Section.ClothMappingDataLODs[LODBias].Num(),
-						LogClothingAsset, Verbose, TEXT("Removed deformer data for [%s/%s], section %d LODBias = %d, render mesh LOD = %d, sim mesh LOD = %d"),
+					UE_CLOG(Section.ClothMappingDataLODs.IsValidIndex(LodBias) && Section.ClothMappingDataLODs[LodBias].Num(),
+						LogClothingAsset, Verbose,
+						TEXT("Removed deformer data for [%s/%s], section %d LOD bias = %d, render mesh LOD = %d, sim mesh LOD = %d"),
 						*SkeletalMesh.GetName(),
 						*Asset.GetName(),
-						SectionIdx,
-						LODBias,
-						MeshLodIndex,
-						LodIndex);
+						SectionIndex,
+						LodBias,
+						LodIndex,
+						BiasLodIndex);
 				}
 
 				Section.ClothingData.AssetGuid = FGuid();
@@ -113,18 +147,11 @@ void ClothingAssetUtils::UnbindFromSkeletalMesh(const UClothingAssetBase& Asset,
 				}
 
 				// Clear all remaining LOD bias mapping data from other LODs that relied on this section
-				ClothingAssetUtils::ClearLODBiasMappings(SkeletalMesh, MeshLodIndex, SectionIdx);
+				ClothingAssetUtils::ClearLODBiasMappings(SkeletalMesh, LodIndex, SectionIndex);
 
 				bChangedMesh = true;
 			}
 		}
-
-		// If the mesh changed we need to re-register any components that use it to reflect the changes
-		if (bChangedMesh)
-		{
-			//Register the scope post edit change
-			FScopedSkeletalMeshPostEditChange SkeletalMeshPostEditChange(&SkeletalMesh);
-		}
 	}
 }
 
@@ -708,32 +735,24 @@ void UClothingAssetCommon::UpdateLODBiasMappings(const USkeletalMesh* SkeletalMe
 	}
 }
 
-void UClothingAssetCommon::UnbindFromSkeletalMesh(USkeletalMesh* SkeletalMesh)
+void UClothingAssetCommon::UnbindFromSkeletalMesh(USkeletalMesh* InSkelMesh, const int32 InMeshLodIndex, const int32 InSectionIndex)
 {
-	// Clear the entire LOD map to allow for rebinding a new asset
-	for (int32& LodMapIndex : LodMap)
+	// Clear the LOD map entries to allow for rebinding a new asset
+	if (InMeshLodIndex == INDEX_NONE)
 	{
-		LodMapIndex = INDEX_NONE;  // Note: This must be done prior calling UnbindFromSkeletalMesh(), so that the info is correct when rebuilding the mesh
-	}
-
-	// Unbind from the Skeletal Mesh
-	if (SkeletalMesh)
-	{
-		ClothingAssetUtils::UnbindFromSkeletalMesh(*this, *SkeletalMesh);
+		for (int32& LodMapIndex : LodMap)
+		{
+			LodMapIndex = INDEX_NONE;  // Note: This must be done prior calling UnbindFromSkeletalMesh(), so that the info is correct when rebuilding the mesh
+		}
 	}
-}
-
-void UClothingAssetCommon::UnbindFromSkeletalMesh(USkeletalMesh* SkeletalMesh, const int32 MeshLodIndex)
-{
-	// Clear the LOD map entry to allow for rebinding a new asset
-	if (LodMap.IsValidIndex(MeshLodIndex))
+	else if (LodMap.IsValidIndex(InMeshLodIndex))
 	{
-		LodMap[MeshLodIndex] = INDEX_NONE;  // Note: This must be done prior calling UnbindFromSkeletalMesh(), so that the info is correct when rebuilding the mesh
+		LodMap[InMeshLodIndex] = INDEX_NONE;  // Note: This must be done prior calling UnbindFromSkeletalMesh(), so that the info is correct when rebuilding the mesh
 	}
 	// Unbind from the Skeletal Mesh
-	if (SkeletalMesh)
+	if (InSkelMesh)
 	{
-		ClothingAssetUtils::UnbindFromSkeletalMesh(*this, *SkeletalMesh, MeshLodIndex);
+		ClothingAssetUtils::UnbindFromSkeletalMesh(*this, *InSkelMesh, InMeshLodIndex, InSectionIndex);
 	}
 }
 
@@ -902,6 +921,22 @@ void UClothingAssetCommon::PostEditUndo()
 		}
 	}
 }
+
+TArray<TPair<FText, FText>> UClothingAssetCommon::GetStats() const
+{
+	TArray<TPair<FText, FText>> Stats;
+
+	for (int32 LodIndex = 0; LodIndex < LodData.Num(); ++LodIndex)
+	{
+		Stats.Emplace(FText::Format(LOCTEXT("LodIndexLabel", "LOD {0}"), FText::AsNumber(LodIndex)), FText{});
+		const FClothPhysicalMeshData& PhysMeshData = LodData[LodIndex].PhysicalMeshData;
+		Stats.Emplace(LOCTEXT("SimulVertexCount", "Simul Verts"), FText::AsNumber(PhysMeshData.Vertices.Num() - PhysMeshData.NumFixedVerts));
+		Stats.Emplace(LOCTEXT("FixedVertexCount", "Fixed Verts"), FText::AsNumber(PhysMeshData.NumFixedVerts));
+		Stats.Emplace(LOCTEXT("TriangleCount", "Sim Triangles"), FText::AsNumber(PhysMeshData.Indices.Num() / 3));
+		Stats.Emplace(LOCTEXT("MaxBoneWeights", "Bone Weights"), FText::AsNumber(PhysMeshData.MaxBoneWeights));
+	}
+	return Stats;
+}
 #endif // WITH_EDITOR
 
 void UClothingAssetCommon::RefreshBoneMapping(USkeletalMesh* InSkelMesh)
@@ -1578,7 +1613,7 @@ void UClothingAssetCommon::RemoveLastLod()
 			if (SkeletalMesh)
 			{
 				// Unbind from the Skeletal Mesh
-				ClothingAssetUtils::UnbindFromSkeletalMesh(*this, *SkeletalMesh, MeshLodIndex);
+				ClothingAssetUtils::UnbindFromSkeletalMesh(*this, *SkeletalMesh, MeshLodIndex, INDEX_NONE);
 			}
 		}
 		LodData.RemoveAt(LodIndex, EAllowShrinking::Yes);
diff --git a/Engine/Source/Runtime/ClothingSystemRuntimeCommon/Public/ClothingAsset.h b/Engine/Source/Runtime/ClothingSystemRuntimeCommon/Public/ClothingAsset.h
--- a/Engine/Source/Runtime/ClothingSystemRuntimeCommon/Public/ClothingAsset.h
+++ b/Engine/Source/Runtime/ClothingSystemRuntimeCommon/Public/ClothingAsset.h
@@ -32,19 +32,15 @@ namespace ClothingAssetUtils
 
 #if WITH_EDITOR
 	/**
-	 * Unbind this clothing asset from the provided skeletal mesh, will remove all LODs.
+	 * Unbind Clothing Data from the specified Skeletal Mesh LOD sections.
+	 * If MeshLodIndex is INDEX_NONE, the Clothing Data will get unbound from all LODs and all sections (ignoring SectionIndex in this case).
+	 * If SectionIndex is INDEX_NONE, the Clothing Data will get unbound from all sections of the specified LOD.
 	 * @param Asset The clothing data to unbind.
-	 * @param SkeletalMesh skeletal mesh to remove from.
+	 * @param SkeletalMesh The Skeletal Mesh to unbind the Clothing Data from.
+	 * @param MeshLodIndex Mesh LOD to remove this asset from (could still be bound to other LODs), or INDEX_NONE to unbind all LODs.
+	 * @param SectionIndex The Skeletal Mesh section to unbind, or INDEX_NONE to unbind all sections.
 	 */
-	UE_API void UnbindFromSkeletalMesh(const UClothingAssetBase& Asset, USkeletalMesh& SkeletalMesh);
-
-	/**
-	 * Unbind this clothing asset from the provided skeletal mesh.
-	 * @param Asset The clothing data to unbind.
-	 * @param SkeletalMesh The Skeletal mesh to remove from.
-	 * @param MeshLodIndex Mesh LOD to remove this asset from (could still be bound to other LODs).
-	 */
-	UE_API void UnbindFromSkeletalMesh(const UClothingAssetBase& Asset, USkeletalMesh& SkeletalMesh, const int32 MeshLodIndex);
+	UE_API void UnbindFromSkeletalMesh(const UClothingAssetBase& Asset, USkeletalMesh& SkeletalMesh, const int32 MeshLodIndex = INDEX_NONE, const int32 SectionIndex = INDEX_NONE);
 
 	/**
 	 * Return whether there are any cloth bindings whithin the skeletal mesh model.
@@ -141,11 +137,26 @@ public:
 	UE_API virtual bool BindToSkeletalMesh(USkeletalMesh* InSkelMesh, const int32 InMeshLodIndex, const int32 InSectionIndex, const int32 InAssetLodIndex) override;
 
 	/**
-	 * Helper that invokes \c UnbindFromSkeletalMesh() for each avilable entry in 
-	 * \p InSkelMesh->GetImportedModel()'s LODModel.
+	 * Unbind Clothing Data from the specified Skeletal Mesh LOD sections.
+	 * If MeshLodIndex is INDEX_NONE, the Clothing Data will get unbound from all LODs and all sections (ignoring SectionIndex in this case).
+	 * If SectionIndex is INDEX_NONE, the Clothing Data will get unbound from all sections of the specified LOD.
+	 * @param InSkelMesh The Skeletal Mesh to unbind the Clothing Data from.
+	 * @param InMeshLodIndex Mesh LOD to remove this asset from (could still be bound to other LODs), or INDEX_NONE to unbind all LODs.
+	 * @param InSectionIndex The Skeletal Mesh section to unbind, or INDEX_NONE to unbind all sections.
 	 */
-	UE_API virtual void UnbindFromSkeletalMesh(USkeletalMesh* InSkelMesh) override;
-	UE_API virtual void UnbindFromSkeletalMesh(USkeletalMesh* InSkelMesh, const int32 InMeshLodIndex) override;
+	UE_API virtual void UnbindFromSkeletalMesh(USkeletalMesh* InSkelMesh, const int32 InMeshLodIndex, const int32 InSectionIndex) override;
+	
+	UE_DEPRECATED(5.8, "Use UnbindFromSkeletalMesh(USkeletalMesh* InSkelMesh, const int32 InMeshLodIndex, const int32 InSectionIndex) instead.")
+	virtual void UnbindFromSkeletalMesh(USkeletalMesh* InSkelMesh) override
+	{
+		UnbindFromSkeletalMesh(InSkelMesh, INDEX_NONE, INDEX_NONE);
+	}
+
+	UE_DEPRECATED(5.8, "Use UnbindFromSkeletalMesh(USkeletalMesh* InSkelMesh, const int32 InMeshLodIndex, const int32 InSectionIndex) instead.")
+	virtual void UnbindFromSkeletalMesh(USkeletalMesh* InSkelMesh, const int32 InMeshLodIndex) override
+	{
+		UnbindFromSkeletalMesh(InSkelMesh, InMeshLodIndex, INDEX_NONE);
+	}
 
 	/**
 	 * Update all extra LOD deformer mappings.
@@ -343,6 +354,9 @@ private:
 	}
 
 #if WITH_EDITOR
+	// Return all LODs' stats.
+	virtual TArray<TPair<FText, FText>> GetStats() const override;
+
 	// Add extra cloth deformer mappings to cope with a different raytracing LOD than the one currently rendered.
 	void UpdateLODBiasMappings(const USkeletalMesh* SkeletalMesh, int32 UpdatedLODIndex, int32 SectionIndex);
 
diff --git a/Engine/Source/Runtime/ClothingSystemRuntimeInterface/Private/ClothingAssetBase.cpp b/Engine/Source/Runtime/ClothingSystemRuntimeInterface/Private/ClothingAssetBase.cpp
--- a/Engine/Source/Runtime/ClothingSystemRuntimeInterface/Private/ClothingAssetBase.cpp
+++ b/Engine/Source/Runtime/ClothingSystemRuntimeInterface/Private/ClothingAssetBase.cpp
@@ -12,6 +12,12 @@
 DEFINE_LOG_CATEGORY(LogClothingAsset)
 
 #if WITH_EDITOR
+void UClothingAssetBase::PostInitProperties()
+{
+	Super::PostInitProperties();
+	SetFlags(RF_Transactional);
+}
+
 void UClothingAssetBase::WarningNotification(const FText& Text)
 {
 	FNotificationInfo* const NotificationInfo = new FNotificationInfo(Text);
diff --git a/Engine/Source/Runtime/ClothingSystemRuntimeInterface/Public/ClothingAssetBase.h b/Engine/Source/Runtime/ClothingSystemRuntimeInterface/Public/ClothingAssetBase.h
--- a/Engine/Source/Runtime/ClothingSystemRuntimeInterface/Public/ClothingAssetBase.h
+++ b/Engine/Source/Runtime/ClothingSystemRuntimeInterface/Public/ClothingAssetBase.h
@@ -4,8 +4,14 @@
 
 #include "UObject/Object.h"
 #include "Misc/Guid.h"
+#if WITH_EDITOR
+#include "Containers/Array.h"
+#include "Internationalization/Text.h"
+#include "Templates/Tuple.h"
+#endif
 #include "ClothingAssetBase.generated.h"
 
+class ITransactionObjectAnnotation;
 class USkeletalMesh;
 
 #define UE_API CLOTHINGSYSTEMRUNTIMEINTERFACE_API
@@ -24,36 +30,41 @@ class UClothingAssetBase : public UObject
 public:
 
 #if WITH_EDITOR
-	/** Binds a clothing asset submesh to a skeletal mesh section
-	* @param InSkelMesh Skel mesh to bind to
-	* @param InSectionIndex Section in the skel mesh to replace
-	* @param InSubmeshIdx Submesh in this asset to replace section with
-	* @param InAssetLodIndex Internal clothing LOD to use
-	*/
+	/**
+	 * Binds a clothing asset submesh to a skeletal mesh section.
+	 * @param InSkelMesh Skel mesh to bind to.
+	 * @param InMeshLodIndex Mesh LOD to bind this asset to.
+	 * @param InSectionIndex Section in the skel mesh to replace.
+	 * @param InAssetLodIndex Internal clothing LOD to use.
+	 */
 	virtual bool BindToSkeletalMesh(USkeletalMesh* InSkelMesh, const int32 InMeshLodIndex, const int32 InSectionIndex, const int32 InAssetLodIndex)
 	PURE_VIRTUAL(UClothingAssetBase::BindToSkeletalMesh, return false;);
 
 	/**
-	* Unbinds this clothing asset from the provided skeletal mesh, will remove all LODs
-	* @param InSkelMesh skeletal mesh to remove from
-	*/
-	virtual void UnbindFromSkeletalMesh(USkeletalMesh* InSkelMesh)
-	PURE_VIRTUAL(UClothingAssetBase::UnbindFromSkeletalMesh, );
+	 * Unbind Clothing Data from the specified Skeletal Mesh LOD sections.
+	 * If MeshLodIndex is INDEX_NONE, the Clothing Data will get unbound from all LODs and all sections (ignoring SectionIndex in this case).
+	 * If SectionIndex is INDEX_NONE, the Clothing Data will get unbound from all sections of the specified LOD.
+	 * @param InSkelMesh The Skeletal Mesh to unbind the Clothing Data from.
+	 * @param InMeshLodIndex Mesh LOD to remove this asset from (could still be bound to other LODs), or INDEX_NONE to unbind all LODs.
+	 * @param InSectionIndex The Skeletal Mesh section to unbind, or INDEX_NONE to unbind all sections.
+	 */
+	virtual void UnbindFromSkeletalMesh(USkeletalMesh* InSkelMesh, const int32 InMeshLodIndex, const int32 InSectionIndex)
+PRAGMA_DISABLE_DEPRECATION_WARNINGS
+	PURE_VIRTUAL(UClothingAssetBase::UnbindFromSkeletalMesh, UnbindFromSkeletalMesh(InSkelMesh, InMeshLodIndex); return;);
+PRAGMA_ENABLE_DEPRECATION_WARNINGS
 
-	/**
-	* Unbinds this clothing asset from the provided skeletal mesh
-	* @param InSkelMesh skeletal mesh to remove from
-	* @param InMeshLodIndex Mesh LOD to remove this asset from (could still be bound to other LODs)
-	*/
-	virtual void UnbindFromSkeletalMesh(USkeletalMesh* InSkelMesh, const int32 InMeshLodIndex)
-	PURE_VIRTUAL(UClothingAssetBase::UnbindFromSkeletalMesh,);
+	UE_DEPRECATED(5.8, "Implement and use UnbindFromSkeletalMesh(USkeletalMesh* InSkelMesh, const int32 InMeshLodIndex, const int32 InSectionIndex) instead.")
+	virtual void UnbindFromSkeletalMesh(USkeletalMesh* InSkelMesh) {}
+
+	UE_DEPRECATED(5.8, "Implement and use UnbindFromSkeletalMesh(USkeletalMesh* InSkelMesh, const int32 InMeshLodIndex, const int32 InSectionIndex) instead.")
+	virtual void UnbindFromSkeletalMesh(USkeletalMesh* InSkelMesh, const int32 InMeshLodIndex) {}
 
 	/**
 	 * Update all extra LOD deformer mappings.
 	 * This should be called whenever the raytracing LOD bias is changed.
 	 */
 	virtual void UpdateAllLODBiasMappings(USkeletalMesh* SkeletalMesh)
-	PURE_VIRTUAL(UClothingAssetBase::UpdateAllLODBiasMappings,);
+	PURE_VIRTUAL(UClothingAssetBase::UpdateAllLODBiasMappings, return;);
 
 	/** Add a new LOD class instance. */
 	UE_DEPRECATED(5.7, "Use UClothingAssetCommon::AddNewLod() instead.")
@@ -63,12 +74,27 @@ public:
 	}
 
 	/**
-	*	Builds the LOD transition data
-	*	When we transition between LODs we skin the incoming mesh to the outgoing mesh
-	*	in exactly the same way the render mesh is skinned to create a smooth swap
-	*/
+	 * Builds the LOD transition data
+	 * When we transition between LODs we skin the incoming mesh to the outgoing mesh
+	 * in exactly the same way the render mesh is skinned to create a smooth swap
+	 */
 	UE_DEPRECATED(5.7, "Use UClothingAssetCommon::BuildLodTransitionData() instead.")
 	virtual void BuildLodTransitionData() {}
+
+	/**
+	 * Gather stats for the Skeletal Mesh editor details view.
+	 * @return An array of pairs stat name, stat value that will be displayed in a row.
+	 * @note Stats without values are displayed as the line header and act as a newline for the following entries.
+	 */
+	virtual TArray<TPair<FText, FText>> GetStats() const
+	{
+		return {};
+	}
+
+	//~ Begin UObject Interface
+	/** Make sure the object is transactional. */
+	UE_API virtual void PostInitProperties() override;
+	//~ End UObject Interface
 #endif
 
 #if WITH_EDITORONLY_DATA
@@ -87,7 +113,7 @@ public:
 	 * @param InSkelMesh - The mesh to use to remap the bones
 	 */
 	virtual void RefreshBoneMapping(USkeletalMesh* InSkelMesh)
-	PURE_VIRTUAL(UClothingAssetBase::RefreshBoneMapping, );
+	PURE_VIRTUAL(UClothingAssetBase::RefreshBoneMapping, return;);
 
 	/** Check the validity of this clothing data, for example for when it is dependent from an asset that isn't ready or accessible. */
 	virtual bool IsValid() const
@@ -111,7 +137,7 @@ public:
 
 	/** Called after all cloth assets sharing the same simulation are added or loaded */
 	virtual void PostUpdateAllAssets()
-	PURE_VIRTUAL(UClothingAssetBase::PostUpdateAllAssets(), );
+	PURE_VIRTUAL(UClothingAssetBase::PostUpdateAllAssets(), return;);
 
 	/** Get the guid identifying this asset */
 	const FGuid& GetAssetGuid() const
@@ -119,12 +145,6 @@ public:
 		return AssetGuid;
 	}
 
-#if WITH_EDITORONLY_DATA
-	UE_DEPRECATED(5.6, "This is an old Apex file path which is no longer used and will be removed.")
-	UPROPERTY()
-	FString ImportedFilePath;
-#endif
-
 protected:
 
 	/** The asset factory should have access, as it will assign the asset guid when building assets */
diff --git a/Engine/Source/Runtime/Engine/Classes/Engine/SkeletalMesh.h b/Engine/Source/Runtime/Engine/Classes/Engine/SkeletalMesh.h
--- a/Engine/Source/Runtime/Engine/Classes/Engine/SkeletalMesh.h
+++ b/Engine/Source/Runtime/Engine/Classes/Engine/SkeletalMesh.h
@@ -2256,7 +2256,7 @@ public:
 	}
 
 #if WITH_EDITOR
-	/** If the given section of the specified LOD has a clothing asset, unbind it's data and remove it from the asset array */
+	/** If the given section of the specified LOD has a clothing asset, unbind it and remove it from the asset array if it is no longer used. */
 	ENGINE_API void RemoveClothingAsset(int32 InLodIndex, int32 InSectionIndex);
 
 	/**
@@ -2290,7 +2290,7 @@ public:
 	 * whenever possible
 	 */
 	UE_DEPRECATED(4.27, "Please do not access this member directly; use USkeletalMesh::GetMeshClothingAssets() or USkeletalMesh::SetMeshClothingAssets().")
-	UPROPERTY(EditAnywhere, editfixedsize, BlueprintGetter = GetMeshClothingAssets, BlueprintSetter = SetMeshClothingAssets, Category = Clothing)
+	UPROPERTY(VisibleAnywhere, BlueprintGetter = GetMeshClothingAssets, BlueprintSetter = SetMeshClothingAssets, Category = Clothing)
 	TArray<TObjectPtr<UClothingAssetBase>> MeshClothingAssets;
 
 	static FName GetMeshClothingAssetsMemberName()
diff --git a/Engine/Source/Runtime/Engine/Private/SkeletalMesh.cpp b/Engine/Source/Runtime/Engine/Private/SkeletalMesh.cpp
--- a/Engine/Source/Runtime/Engine/Private/SkeletalMesh.cpp
+++ b/Engine/Source/Runtime/Engine/Private/SkeletalMesh.cpp
@@ -903,10 +903,15 @@ void USkeletalMesh::AddClothingAsset(UClothingAssetBase* InNewAsset)
 void USkeletalMesh::RemoveClothingAsset(int32 InLodIndex, int32 InSectionIndex)
 {
 	check(IsInGameThread());
-	if (UClothingAssetBase* Asset = GetSectionClothingAsset(InLodIndex, InSectionIndex))
+	if (UClothingAssetBase* const Asset = GetSectionClothingAsset(InLodIndex, InSectionIndex))
 	{
-		Asset->UnbindFromSkeletalMesh(this, InLodIndex);
-		GetMeshClothingAssets().Remove(Asset);
+		Asset->UnbindFromSkeletalMesh(this, InLodIndex, InSectionIndex);
+		TArray<UClothingAssetBase*> ClothingAssetsInUse;
+		GetClothingAssetsInUse(ClothingAssetsInUse);
+		if (!ClothingAssetsInUse.Contains(Asset))
+		{
+			GetMeshClothingAssets().Remove(Asset);
+		}
 		OnClothingChange.Broadcast();
 	}
 }
@@ -4248,7 +4253,7 @@ void USkeletalMesh::FinishPostLoadInternal(FSkinnedAssetPostLoadContext& Context
 				FScopedSkeletalMeshPostEditChange ScopedPostEditChange(this, false, false);
 				//Make sure the asset is unbind, some old code path was allowing to have bind cloth asset not present in the imported model.
 				//The old inline reduction code was not rebinding the cloth asset nor unbind it.
-				MeshClothingAsset->UnbindFromSkeletalMesh(this);
+				MeshClothingAsset->UnbindFromSkeletalMesh(this, INDEX_NONE, INDEX_NONE);
 			}
 		}
 	}
