AnimBP: Validate that users are not trying to apply an anim BP preview overlay to self

#Synced-CL 48457294

diff --git a/Engine/Source/Editor/Persona/Private/AnimationEditorPreviewScene.cpp b/Engine/Source/Editor/Persona/Private/AnimationEditorPreviewScene.cpp
--- a/Engine/Source/Editor/Persona/Private/AnimationEditorPreviewScene.cpp
+++ b/Engine/Source/Editor/Persona/Private/AnimationEditorPreviewScene.cpp
@@ -222,7 +222,17 @@ void FAnimationEditorPreviewScene::SetPreviewMesh(USkeletalMesh* NewPreviewMesh,
 
 void FAnimationEditorPreviewScene::SetPreviewAnimationBlueprint(UAnimBlueprint* InAnimBlueprint, UAnimBlueprint* InOverlayOrSubAnimBlueprint)
 {
-	if (InOverlayOrSubAnimBlueprint)
+	auto ValidateSelfOverlay = [InOverlayOrSubAnimBlueprint, InAnimBlueprint]()
+	{
+		if (InOverlayOrSubAnimBlueprint == InAnimBlueprint)
+		{
+			UE_LOGFMT(LogAnimation, Warning, "Anim BP '{AnimBP}' is using itself as a preview overlay, which would cause crashes or infinite recursion. Skipping and applying non-overlaid anim BP instead", InAnimBlueprint->GetName());
+			return false;
+		}
+		return true;
+	};
+
+	if (InOverlayOrSubAnimBlueprint && ValidateSelfOverlay())
 	{
 		SkeletalMeshComponent->SetAnimInstanceClass(InAnimBlueprint ? InAnimBlueprint->GeneratedClass : nullptr);
 		EPreviewAnimationBlueprintApplicationMethod ApplicationMethod = InOverlayOrSubAnimBlueprint->GetPreviewAnimationBlueprintApplicationMethod();
