Delay removal of curve model entries until next slate tick - to prevent memory corruption/destroyed object access
#Synced-CL 49030801

diff --git a/Engine/Source/Editor/Persona/Private/SAnimSequenceCurveEditor.cpp b/Engine/Source/Editor/Persona/Private/SAnimSequenceCurveEditor.cpp
--- a/Engine/Source/Editor/Persona/Private/SAnimSequenceCurveEditor.cpp
+++ b/Engine/Source/Editor/Persona/Private/SAnimSequenceCurveEditor.cpp
@@ -376,6 +376,21 @@ public:
 	TWeakPtr<ITimeSliderController> ExternalTimeSliderController;
 };
 
+void SAnimSequenceCurveEditor::Tick(const FGeometry& AllottedGeometry, const double InCurrentTime, const float InDeltaTime)
+{
+	IAnimSequenceCurveEditor::Tick(AllottedGeometry, InCurrentTime, InDeltaTime);
+	
+	for (const FAnimationCurveIdentifier& CurveToRemove : CurvesToRemove)
+	{
+		const int32 ChannelIndices = CurveToRemove.CurveType == ERawCurveTrackTypes::RCT_Transform ? 9 : 1;
+		for (int32 ChannelIndex = 0; ChannelIndex < ChannelIndices; ++ChannelIndex)
+		{
+			RemoveCurve(CurveToRemove.CurveName, CurveToRemove.CurveType, ChannelIndex);
+		}
+	}
+	CurvesToRemove.Reset();
+}
+
 SAnimSequenceCurveEditor::~SAnimSequenceCurveEditor()
 {
 	if(AnimSequence)
@@ -452,23 +467,13 @@ void SAnimSequenceCurveEditor::Construct(const FArguments& InArgs, const TShared
 
 void SAnimSequenceCurveEditor::OnModelHasChanged(const EAnimDataModelNotifyType& NotifyType, IAnimationDataModel* Model, const FAnimDataModelNotifPayload& Payload)
 {
-	auto StopEditingCurve = [this, &Payload, Model]()
-	{
-		const FCurvePayload& TypedPayload = Payload.GetPayload<FCurvePayload>();
-		const FAnimationCurveIdentifier& CurveId = TypedPayload.Identifier;
-		const int32 ChannelIndices = CurveId.CurveType == ERawCurveTrackTypes::RCT_Transform ? 9 : 1;
-		for (int32 ChannelIndex = 0; ChannelIndex < ChannelIndices; ++ChannelIndex)
-		{
-			RemoveCurve(CurveId.CurveName, CurveId.CurveType, ChannelIndex);
-		}
-	};
-	
 	switch(NotifyType)
 	{
 	case EAnimDataModelNotifyType::CurveRemoved:
 	case EAnimDataModelNotifyType::CurveRenamed:
 		{
-			StopEditingCurve();
+			const FCurvePayload& TypedPayload = Payload.GetPayload<FCurvePayload>();
+			CurvesToRemove.Add(TypedPayload.Identifier);
 			break;
 		}			
 	case EAnimDataModelNotifyType::CurveFlagsChanged:
@@ -476,7 +481,7 @@ void SAnimSequenceCurveEditor::OnModelHasChanged(const EAnimDataModelNotifyType&
 			const FCurveFlagsChangedPayload& TypedPayload = Payload.GetPayload<FCurveFlagsChangedPayload>();
 			if(Model->FindCurve(TypedPayload.Identifier)->GetCurveTypeFlag(AACF_Metadata))
 			{
-				StopEditingCurve();
+				CurvesToRemove.Add(TypedPayload.Identifier);
 			}
 			break;
 		}
@@ -602,4 +607,4 @@ void SAnimSequenceCurveEditor::ZoomToFit()
 	CurveEditor->ZoomToFit(EAxisList::Y);
 }
 
-#undef LOCTEXT_NAMESPACE
\ No newline at end of file
+#undef LOCTEXT_NAMESPACE
diff --git a/Engine/Source/Editor/Persona/Private/SAnimSequenceCurveEditor.h b/Engine/Source/Editor/Persona/Private/SAnimSequenceCurveEditor.h
--- a/Engine/Source/Editor/Persona/Private/SAnimSequenceCurveEditor.h
+++ b/Engine/Source/Editor/Persona/Private/SAnimSequenceCurveEditor.h
@@ -77,6 +77,7 @@ class SAnimSequenceCurveEditor : public IAnimSequenceCurveEditor
 	virtual void AddCurve(const FText& InCurveDisplayName, const FLinearColor& InCurveColor, const FName& InName, ERawCurveTrackTypes InType, int32 InCurveIndex, FSimpleDelegate InOnCurveModified) override;
 	virtual void RemoveCurve(const FName& InName, ERawCurveTrackTypes InType, int32 InCurveIndex) override;
 	virtual void ZoomToFit() override;
+	virtual void Tick(const FGeometry& AllottedGeometry, const double InCurrentTime, const float InDeltaTime) override;
 	
 	void OnModelHasChanged(const EAnimDataModelNotifyType& NotifyType, IAnimationDataModel* Model, const FAnimDataModelNotifPayload& Payload);
 private:
@@ -96,4 +97,7 @@ private:
 
 	/** The tree widget in the curve editor */
 	TSharedPtr<SCurveEditorTree> CurveEditorTree;
-};
\ No newline at end of file
+
+	/** Curves to be removed from editing on the next ::Tick */	
+	TArray<FAnimationCurveIdentifier> CurvesToRemove; 
+};
