PR #13996: Add new option to animation viewport editor to use a bone as the camera

#rnx
#Synced-CL 49029692

diff --git a/Engine/Source/Editor/Persona/Private/AnimationEditorViewportClient.cpp b/Engine/Source/Editor/Persona/Private/AnimationEditorViewportClient.cpp
--- a/Engine/Source/Editor/Persona/Private/AnimationEditorViewportClient.cpp
+++ b/Engine/Source/Editor/Persona/Private/AnimationEditorViewportClient.cpp
@@ -387,6 +387,12 @@ void FAnimationViewportClient::SetCameraFollowMode(EAnimationViewportCameraFollo
 		bCanFollow = PreviewMeshComponent->GetBoneIndex(InBoneName) != INDEX_NONE;
 	}
 
+	// When switching out of bone as camera we want to reset our coordinate system to something sensible
+	if (InCameraFollowMode != EAnimationViewportCameraFollowMode::BoneAsCamera && CameraFollowMode == EAnimationViewportCameraFollowMode::BoneAsCamera)
+	{
+		EnableCameraLock(true);
+	}
+
 	if(bCanFollow && InCameraFollowMode != EAnimationViewportCameraFollowMode::None)
 	{
 		ConfigOption->SetViewCameraFollow(AssetEditorToolkitPtr.Pin()->GetEditorName(), InCameraFollowMode, InBoneName, ViewportIndex);
@@ -395,13 +401,14 @@ void FAnimationViewportClient::SetCameraFollowMode(EAnimationViewportCameraFollo
 		CameraFollowBoneName = InBoneName;
 
 		bCameraLock = true;
-		bUsingOrbitCamera = true;
+		bUsingOrbitCamera = InCameraFollowMode != EAnimationViewportCameraFollowMode::BoneAsCamera;
+		bDisableInput = InCameraFollowMode == EAnimationViewportCameraFollowMode::BoneAsCamera;
 
 		if (PreviewMeshComponent != nullptr)
 		{
 			FVector LookAtLocation = LastLookAtLocation;
 
-			switch(CameraFollowMode)
+			switch (CameraFollowMode)
 			{
 			case EAnimationViewportCameraFollowMode::Bounds:
 				{
@@ -420,13 +427,19 @@ void FAnimationViewportClient::SetCameraFollowMode(EAnimationViewportCameraFollo
 					LookAtLocation = PreviewMeshComponent->GetBoneLocation(InBoneName);
 				}
 				break;
+				case EAnimationViewportCameraFollowMode::BoneAsCamera:
+				{
+					SetViewLocation(PreviewMeshComponent->GetBoneLocation(CameraFollowBoneName));
+					SetViewRotation(PreviewMeshComponent->GetBoneQuaternion(CameraFollowBoneName).Rotator());
+				}
+				break;
 			}
-
+		
 			OrbitRotation = FQuat::Identity;
-			SetLookAtLocation(LookAtLocation, true);
+			SetLookAtLocation(LookAtLocation, bUsingOrbitCamera);
 			LastLookAtLocation = LookAtLocation;
-			bUsingOrbitCamera = true;
 		}
+
 	}
 	else
 	{
@@ -436,6 +449,7 @@ void FAnimationViewportClient::SetCameraFollowMode(EAnimationViewportCameraFollo
 		CameraFollowBoneName = NAME_None;
 
 		OrbitRotation = FQuat::Identity;
+		bDisableInput = false;
 		EnableCameraLock(false);
 		FocusViewportOnPreviewMesh(false);
 		Invalidate();
@@ -481,6 +495,55 @@ FName FAnimationViewportClient::GetCameraFollowBoneName() const
 	return CameraFollowBoneName;
 }
 
+EAxis::Type FAnimationViewportClient::GetBoneAsCameraAxis(EDirection Direction) const
+{
+	if (Direction == EDirection::Forward)
+	{
+		return BoneAsCameraForward;
+	}
+	if (Direction == EDirection::Right)
+	{
+		return BoneAsCameraRight;
+	}
+	return BoneAsCameraUp;
+}
+
+UE::Math::TVector<double> FAnimationViewportClient::GetBoneAsCameraPositionOffset() const
+{
+	return BoneAsCameraPositionOffset;
+}
+
+FRotator FAnimationViewportClient::GetBoneAsCameraRotationOffset() const
+{
+	return BoneAsCameraRotationOffset;
+}
+
+void FAnimationViewportClient::SetBoneAsCameraAxis(EDirection Direction, EAxis::Type Axis)
+{
+	if (Direction == EDirection::Forward)
+	{
+		BoneAsCameraForward = Axis;
+	}
+	else if (Direction == EDirection::Right)
+	{
+		BoneAsCameraRight = Axis;
+	}
+	else
+	{
+		BoneAsCameraUp = Axis;
+	}
+}
+
+void FAnimationViewportClient::SetBoneAsCameraPositionOffset(UE::Math::TVector<double> Offset)
+{
+	BoneAsCameraPositionOffset = Offset;
+}
+
+void FAnimationViewportClient::SetBoneAsCameraRotationOffset(const FRotator& Offset)
+{
+	BoneAsCameraRotationOffset = Offset;
+}
+
 void FAnimationViewportClient::JumpToDefaultCamera()
 {
 	FocusViewportOnPreviewMesh(true);
@@ -899,6 +962,27 @@ void FAnimationViewportClient::HandlePreviewScenePostTick()
 					}
 				}
 				break;
+			case EAnimationViewportCameraFollowMode::BoneAsCamera:
+				{
+					const int32 BoneIndex = PreviewMeshComponent->GetBoneIndex(CameraFollowBoneName);
+					if (BoneIndex != INDEX_NONE)
+					{
+						FTransform BoneTransform = PreviewMeshComponent->GetBoneTransform(CameraFollowBoneName);
+
+						FViewportCameraTransform& ViewTransform = GetViewTransform();
+
+						const FVector Forward = BoneTransform.GetUnitAxis(BoneAsCameraForward);
+						const FVector Right = BoneTransform.GetUnitAxis(BoneAsCameraRight);
+						const FVector Up = BoneTransform.GetUnitAxis(BoneAsCameraUp);
+						const FVector Origin = BoneTransform.GetTranslation() + BoneAsCameraPositionOffset;
+						FTransform NewViewTransform = FTransform(Forward, Right, Up, Origin);
+						NewViewTransform.SetRotation(BoneAsCameraRotationOffset.Quaternion() * NewViewTransform.Rotator().Quaternion());
+
+						ViewTransform.SetLocation(NewViewTransform.GetLocation());
+						ViewTransform.SetRotation(NewViewTransform.GetRotation().Rotator());
+					}
+				}
+				break;
 			}
 
 			const FVector Offset = LookAtLocation - LastLookAtLocation;
diff --git a/Engine/Source/Editor/Persona/Private/SAnimationEditorViewport.cpp b/Engine/Source/Editor/Persona/Private/SAnimationEditorViewport.cpp
--- a/Engine/Source/Editor/Persona/Private/SAnimationEditorViewport.cpp
+++ b/Engine/Source/Editor/Persona/Private/SAnimationEditorViewport.cpp
@@ -2380,6 +2380,60 @@ void SAnimationEditorViewportTabBody::ToggleRotateCameraToFollowBone()
 	}
 }
 
+EAxis::Type SAnimationEditorViewportTabBody::GetBoneAsCameraAxis(FAnimationViewportClient::EDirection Direction) const
+{
+	if (const TSharedPtr<FAnimationViewportClient>& AnimViewportClient = GetAnimationViewportClient())
+	{
+		return AnimViewportClient->GetBoneAsCameraAxis(Direction);
+	}
+
+	return EAxis::Type::None;
+}
+
+UE::Math::TVector<double> SAnimationEditorViewportTabBody::GetBoneAsCameraPositionOffset() const
+{
+	if (const TSharedPtr<FAnimationViewportClient>& AnimViewportClient = GetAnimationViewportClient())
+	{
+		return AnimViewportClient->GetBoneAsCameraPositionOffset();
+	}
+
+	return UE::Math::TVector<double>::Zero();
+}
+
+FRotator SAnimationEditorViewportTabBody::GetBoneAsCameraRotationOffset() const
+{
+	if (const TSharedPtr<FAnimationViewportClient>& AnimViewportClient = GetAnimationViewportClient())
+	{
+		return AnimViewportClient->GetBoneAsCameraRotationOffset();
+	}
+
+	return FRotator();
+}
+
+void SAnimationEditorViewportTabBody::SetBoneAsCameraAxis(FAnimationViewportClient::EDirection Direction, EAxis::Type Axis)
+{
+	if (const TSharedPtr<FAnimationViewportClient>& AnimViewportClient = GetAnimationViewportClient())
+	{
+		return AnimViewportClient->SetBoneAsCameraAxis(Direction, Axis);
+	}
+}
+
+void SAnimationEditorViewportTabBody::SetBoneAsCameraPositionOffset(UE::Math::TVector<double> Offset)
+{
+	if (const TSharedPtr<FAnimationViewportClient>& AnimViewportClient = GetAnimationViewportClient())
+	{
+		AnimViewportClient->SetBoneAsCameraPositionOffset(Offset);
+	}
+}
+
+void SAnimationEditorViewportTabBody::SetBoneAsCameraRotationOffset(const FRotator& Offset)
+{
+	if (const TSharedPtr<FAnimationViewportClient>& AnimViewportClient = GetAnimationViewportClient())
+	{
+		AnimViewportClient->SetBoneAsCameraRotationOffset(Offset);
+	}
+}
+
 bool SAnimationEditorViewportTabBody::GetShouldRotateCameraToFollowBone() const
 {
 	if (const TSharedPtr<FAnimationViewportClient>& AnimViewportClient = GetAnimationViewportClient())
diff --git a/Engine/Source/Editor/Persona/Private/SAnimationEditorViewport.h b/Engine/Source/Editor/Persona/Private/SAnimationEditorViewport.h
--- a/Engine/Source/Editor/Persona/Private/SAnimationEditorViewport.h
+++ b/Engine/Source/Editor/Persona/Private/SAnimationEditorViewport.h
@@ -551,6 +551,26 @@ public:
 	bool IsCameraFollowEnabled(EAnimationViewportCameraFollowMode InCameraFollowMode) const;
 	FName GetCameraFollowBoneName() const;
 	void ToggleRotateCameraToFollowBone();
+
+	/** Gets the axis for a direction when using the Bone As Camera follow mode */
+	EAxis::Type GetBoneAsCameraAxis(FAnimationViewportClient::EDirection Direction) const;
+
+	/** Sets the axis for a direction when using the Bone As Camera follow mode */
+	void SetBoneAsCameraAxis(FAnimationViewportClient::EDirection Direction, EAxis::Type Axis);
+
+	/** Gets the position offset from the bone origin in world coordiantes for the Bone As Camera follow mode */
+	UE::Math::TVector<double> GetBoneAsCameraPositionOffset() const;
+
+	/** Gets the rotation offset from the bone in world coordiantes for the Bone As Camera follow mode */
+	FRotator GetBoneAsCameraRotationOffset() const;
+	/** Sets the forward axis for the Bone As Camera follow mode */
+
+	/** Sets the position offset from the bone origin in world coordiantes for the Bone As Camera follow mode */
+	void SetBoneAsCameraPositionOffset(UE::Math::TVector<double> Offset);
+
+	/** Sets the rotation offset from the bone in world coordiantes for the Bone As Camera follow mode */
+	void SetBoneAsCameraRotationOffset(const FRotator& Offset);
+
 	bool GetShouldRotateCameraToFollowBone() const;
 	void TogglePauseAnimationOnCameraMove();
 	bool GetShouldPauseAnimationOnCameraMove() const;
diff --git a/Engine/Source/Editor/Persona/Private/ViewportToolbar/AnimationEditorMenus.cpp b/Engine/Source/Editor/Persona/Private/ViewportToolbar/AnimationEditorMenus.cpp
--- a/Engine/Source/Editor/Persona/Private/ViewportToolbar/AnimationEditorMenus.cpp
+++ b/Engine/Source/Editor/Persona/Private/ViewportToolbar/AnimationEditorMenus.cpp
@@ -29,11 +29,468 @@
 #include "ViewportToolbar/UnrealEdViewportToolbar.h"
 #include "Widgets/Input/SSpinBox.h"
 #include "Widgets/Input/STextComboBox.h"
+#include "Widgets/Input/SVectorInputBox.h"
+#include "Runtime/SlateCore/Public/Styling/SlateStyleRegistry.h"
 
 #define LOCTEXT_NAMESPACE "AnimEditorViewportToolbar"
 
 namespace UE::AnimationEditor::Private
 {
+TSharedRef<SWidget> MakeBoneAsCameraWidget(const TWeakPtr<SAnimationEditorViewportTabBody>& InViewport, const TWeakPtr<SComboButton>& InWeakComboButton)
+{
+	TSharedPtr<SAnimationEditorViewportTabBody> Viewport = InViewport.Pin();
+
+	if (!Viewport)
+	{
+		return SNullWidget::NullWidget;
+	}
+
+	TSharedPtr<SEditorViewport> ViewportWidget = Viewport->GetViewportWidget();
+
+	if (!ViewportWidget || !ViewportWidget->GetViewportClient())
+	{
+		return SNullWidget::NullWidget;
+	}
+
+	TSharedPtr<SBoneTreeMenu> BoneTreeMenu;
+
+
+	// clang-format off
+	TSharedRef<SWidget> MenuWidget =
+		SNew(SBox)
+		.MaxDesiredHeight(400.0f)
+		[
+			SAssignNew(BoneTreeMenu, SBoneTreeMenu)
+				.bShowVirtualBones(true)
+				.OnBoneSelectionChanged_Lambda(
+					[ViewportWeak = Viewport.ToWeakPtr()](FName InBoneName)
+					{
+						TSharedPtr<SAnimationEditorViewportTabBody> Viewport = ViewportWeak.Pin();
+						if (!Viewport)
+						{
+							return;
+						}
+
+						Viewport->SetCameraFollowMode(EAnimationViewportCameraFollowMode::BoneAsCamera, InBoneName);
+
+						FSlateApplication::Get().DismissAllMenus();
+
+						if (const TSharedPtr<IPinnedCommandList>& PinnedCommands = Viewport->GetPinnedCommands())
+						{
+							PinnedCommands->SetStyle(&FAppStyle::Get(), TEXT("ViewportPinnedCommandList"));
+							PinnedCommands->AddCustomWidget(TEXT("BoneAsCameraWidget"));
+						}
+					})
+				.SelectedBone(Viewport->GetCameraFollowBoneName())
+				.OnGetReferenceSkeleton_Lambda(
+					[ViewportWeak = Viewport.ToWeakPtr()]() -> const FReferenceSkeleton&
+					{
+						static FReferenceSkeleton EmptySkeleton;
+
+						if (TSharedPtr<SAnimationEditorViewportTabBody> Viewport = ViewportWeak.Pin())
+						{
+							if (USkeletalMesh* PreviewMesh = Viewport->GetPreviewScene()->GetPreviewMesh())
+							{
+								return PreviewMesh->GetRefSkeleton();
+							}
+						}
+
+						return EmptySkeleton;
+					})
+		];
+	// clang-format on
+
+	if (TSharedPtr<SComboButton> ComboButton = InWeakComboButton.Pin())
+	{
+		ComboButton->SetMenuContentWidgetToFocus(BoneTreeMenu->GetFilterTextWidget());
+	}
+
+	return MenuWidget;
+}
+
+
+TSharedRef<SWidget> SetupBoneAsCameraAxisRadioButton(TSharedPtr<SAnimationEditorViewportTabBody>& ViewportTab,
+	const FName InName,
+	const TAttribute<FText>& InLabel,
+	const TAttribute<FText>& InToolTip,
+	FAnimationViewportClient::EDirection Direction,
+	EAxis::Type Axis)
+{
+	auto IsCheckedLambda = [ViewportWeak = ViewportTab.ToWeakPtr(), Direction, Axis]() //IsChecked
+		{
+			if (TSharedPtr<SAnimationEditorViewportTabBody> Viewport = ViewportWeak.Pin())
+			{
+				const bool bIsSameAxis = Viewport->GetBoneAsCameraAxis(Direction) == Axis;
+				return bIsSameAxis ? ECheckBoxState::Checked : ECheckBoxState::Unchecked;
+			}
+
+			return ECheckBoxState::Unchecked;
+		};
+
+	auto OnStateChangedLambda = [ViewportWeak = ViewportTab.ToWeakPtr(), Direction, Axis](ECheckBoxState State) //Exec
+		{
+			if (TSharedPtr<SAnimationEditorViewportTabBody> Viewport = ViewportWeak.Pin())
+			{
+				Viewport->SetBoneAsCameraAxis(Direction, Axis);
+			}
+		};
+
+	const ISlateStyle* StyleSet = FSlateStyleRegistry::FindSlateStyle("CoreStyle");
+
+	TSharedRef<SWidget> Button =
+		SNew(SCheckBox)
+		.Visibility(EVisibility::Visible)
+		.Style(StyleSet, "Menu.RadioButton")
+		.IsChecked_Lambda(IsCheckedLambda)
+		.OnCheckStateChanged_Lambda(OnStateChangedLambda);
+
+	return Button;
+}
+
+void SetupBoneAsCameraForward(FToolMenuSection& Section, TSharedPtr<SAnimationEditorViewportTabBody>& ViewportTab)
+{
+	TSharedRef<SWidget> XRadioButton = SetupBoneAsCameraAxisRadioButton(ViewportTab,
+		"CameraForwardX",
+		LOCTEXT("BoneAsCameraForwardSection_Label", "X"),
+		LOCTEXT("BoneAsCameraForwardSection_ToolTip", "Select X bone axis to use as camera forward"),
+		FAnimationViewportClient::EDirection::Forward,
+		EAxis::Type::X);
+
+	TSharedRef<SWidget> YRadioButton = SetupBoneAsCameraAxisRadioButton(ViewportTab,
+		"CameraForwardY",
+		LOCTEXT("BoneAsCameraForwardSection_Label", "Y"),
+		LOCTEXT("BoneAsCameraForwardSection_ToolTip", "Select Y bone axis to use as camera forward"),
+		FAnimationViewportClient::EDirection::Forward,
+		EAxis::Type::Y);
+
+	TSharedRef<SWidget> ZRadioButton = SetupBoneAsCameraAxisRadioButton(ViewportTab,
+		"CameraForwardZ",
+		LOCTEXT("BoneAsCameraForwardSection_Label", "Z"),
+		LOCTEXT("BoneAsCameraForwardSection_ToolTip", "Select Z bone axis to use as camera forward"),
+		FAnimationViewportClient::EDirection::Forward,
+		EAxis::Type::Z);
+
+	TSharedRef<SWidget> CheckboxRow = SNew(SHorizontalBox)
+		+ SHorizontalBox::Slot()
+		.VAlign(VAlign_Center)
+		.HAlign(HAlign_Left)
+		.AutoWidth()
+		.Padding(FMargin(10, 0, 0, 0))
+		[
+			XRadioButton
+		]
+		+ SHorizontalBox::Slot()
+		.VAlign(VAlign_Center)
+		.HAlign(HAlign_Left)
+		.AutoWidth()
+		.Padding(FMargin(10, 0, 50, 0))
+		[
+			SNew(STextBlock)
+				.Text(LOCTEXT("CameraForwardRadioButtonTextX", "X"))
+		]
+		+ SHorizontalBox::Slot()
+		.VAlign(VAlign_Center)
+		.HAlign(HAlign_Left)
+		.AutoWidth()
+		.Padding(FMargin(10, 0, 0, 0))
+		[
+			YRadioButton
+		]
+		+ SHorizontalBox::Slot()
+		.VAlign(VAlign_Center)
+		.HAlign(HAlign_Left)
+		.AutoWidth()
+		.Padding(FMargin(10, 0, 50, 0))
+		[
+			SNew(STextBlock)
+				.Text(LOCTEXT("CameraForwardRadioButtonTextY", "Y"))
+		]
+		+ SHorizontalBox::Slot()
+		.VAlign(VAlign_Center)
+		.HAlign(HAlign_Left)
+		.AutoWidth()
+		.Padding(FMargin(10, 0, 0, 0))
+		[
+			ZRadioButton
+		]
+		+ SHorizontalBox::Slot()
+		.VAlign(VAlign_Center)
+		.HAlign(HAlign_Left)
+		.AutoWidth()
+		.Padding(FMargin(10, 0, 50, 0))
+		[
+			SNew(STextBlock)
+				.Text(LOCTEXT("CameraForwardRadioButtonTextZ", "Z"))
+		];
+
+	Section.AddEntry(FToolMenuEntry::InitWidget("BoneAsCameraForwardAxisWidget", CheckboxRow, FText()));
+
+}
+
+void SetupBoneAsCameraRight(FToolMenuSection& Section, TSharedPtr<SAnimationEditorViewportTabBody>& ViewportTab)
+{
+	TSharedRef<SWidget> XRadioButton = SetupBoneAsCameraAxisRadioButton(ViewportTab,
+		"CameraRightX",
+		LOCTEXT("BoneAsCameraRightSection_Label", "X"),
+		LOCTEXT("BoneAsCameraRightSection_ToolTip", "Select X bone axis to use as camera forward"),
+		FAnimationViewportClient::EDirection::Right,
+		EAxis::Type::X);
+
+	TSharedRef<SWidget> YRadioButton = SetupBoneAsCameraAxisRadioButton(ViewportTab,
+		"CameraRightY",
+		LOCTEXT("BoneAsCameraRightSection_Label", "Y"),
+		LOCTEXT("BoneAsCameraRightSection_ToolTip", "Select Y bone axis to use as camera forward"),
+		FAnimationViewportClient::EDirection::Right,
+		EAxis::Type::Y);
+
+	TSharedRef<SWidget> ZRadioButton = SetupBoneAsCameraAxisRadioButton(ViewportTab,
+		"CameraRightZ",
+		LOCTEXT("BoneAsCameraRightSection_Label", "Z"),
+		LOCTEXT("BoneAsCameraRightSection_ToolTip", "Select Z bone axis to use as camera forward"),
+		FAnimationViewportClient::EDirection::Right,
+		EAxis::Type::Z);
+
+	TSharedRef<SWidget> CheckboxRow = SNew(SHorizontalBox)
+		+ SHorizontalBox::Slot()
+		.VAlign(VAlign_Center)
+		.HAlign(HAlign_Left)
+		.AutoWidth()
+		.Padding(FMargin(10, 0, 0, 0))
+		[
+			XRadioButton
+		]
+		+ SHorizontalBox::Slot()
+		.VAlign(VAlign_Center)
+		.HAlign(HAlign_Left)
+		.AutoWidth()
+		.Padding(FMargin(10, 0, 50, 0))
+		[
+			SNew(STextBlock)
+				.Text(LOCTEXT("CameraRightRadioButtonTextX", "X"))
+		]
+		+ SHorizontalBox::Slot()
+		.VAlign(VAlign_Center)
+		.HAlign(HAlign_Left)
+		.AutoWidth()
+		.Padding(FMargin(10, 0, 0, 0))
+		[
+			YRadioButton
+		]
+		+ SHorizontalBox::Slot()
+		.VAlign(VAlign_Center)
+		.HAlign(HAlign_Left)
+		.AutoWidth()
+		.Padding(FMargin(10, 0, 50, 0))
+		[
+			SNew(STextBlock)
+				.Text(LOCTEXT("CameraRightRadioButtonTextY", "Y"))
+		]
+		+ SHorizontalBox::Slot()
+		.VAlign(VAlign_Center)
+		.HAlign(HAlign_Left)
+		.AutoWidth()
+		.Padding(FMargin(10, 0, 0, 0))
+		[
+			ZRadioButton
+		]
+		+ SHorizontalBox::Slot()
+		.VAlign(VAlign_Center)
+		.HAlign(HAlign_Left)
+		.AutoWidth()
+		.Padding(FMargin(10, 0, 50, 0))
+		[
+			SNew(STextBlock)
+				.Text(LOCTEXT("CameraRightRadioButtonTextZ", "Z"))
+		];
+
+	Section.AddEntry(FToolMenuEntry::InitWidget("BoneAsCameraRightAxisWidget", CheckboxRow, FText()));
+}
+
+void SetupBoneAsCameraUp(FToolMenuSection& Section, TSharedPtr<SAnimationEditorViewportTabBody>& ViewportTab)
+{
+	TSharedRef<SWidget> XRadioButton = SetupBoneAsCameraAxisRadioButton(ViewportTab,
+		"CameraUpX",
+		LOCTEXT("BoneAsCameraUpection_Label", "X"),
+		LOCTEXT("BoneAsCameraUpSection_ToolTip", "Select X bone axis to use as camera forward"),
+		FAnimationViewportClient::EDirection::Up,
+		EAxis::Type::X);
+
+	TSharedRef<SWidget> YRadioButton = SetupBoneAsCameraAxisRadioButton(ViewportTab,
+		"CameraUpY",
+		LOCTEXT("BoneAsCameraUpSection_Label", "Y"),
+		LOCTEXT("BoneAsCameraUpSection_ToolTip", "Select Y bone axis to use as camera forward"),
+		FAnimationViewportClient::EDirection::Up,
+		EAxis::Type::Y);
+
+	TSharedRef<SWidget> ZRadioButton = SetupBoneAsCameraAxisRadioButton(ViewportTab,
+		"CameraUpZ",
+		LOCTEXT("BoneAsCameraUpSection_Label", "Z"),
+		LOCTEXT("BoneAsCameraUpSection_ToolTip", "Select Z bone axis to use as camera forward"),
+		FAnimationViewportClient::EDirection::Up,
+		EAxis::Type::Z);
+
+	TSharedRef<SWidget> CheckboxRow = SNew(SHorizontalBox)
+		+ SHorizontalBox::Slot()
+		.VAlign(VAlign_Center)
+		.HAlign(HAlign_Left)
+		.AutoWidth()
+		.Padding(FMargin(10, 0, 0, 0))
+		[
+			XRadioButton
+		]
+		+ SHorizontalBox::Slot()
+		.VAlign(VAlign_Center)
+		.HAlign(HAlign_Left)
+		.AutoWidth()
+		.Padding(FMargin(10, 0, 50, 0))
+		[
+			SNew(STextBlock)
+				.Text(LOCTEXT("CameraUpRadioButtonTextX", "X"))
+		]
+		+ SHorizontalBox::Slot()
+		.VAlign(VAlign_Center)
+		.HAlign(HAlign_Left)
+		.AutoWidth()
+		.Padding(FMargin(10, 0, 0, 0))
+		[
+			YRadioButton
+		]
+		+ SHorizontalBox::Slot()
+		.VAlign(VAlign_Center)
+		.HAlign(HAlign_Left)
+		.AutoWidth()
+		.Padding(FMargin(10, 0, 50, 0))
+		[
+			SNew(STextBlock)
+				.Text(LOCTEXT("CameraUpRadioButtonTextY", "Y"))
+		]
+		+ SHorizontalBox::Slot()
+		.VAlign(VAlign_Center)
+		.HAlign(HAlign_Left)
+		.AutoWidth()
+		.Padding(FMargin(10, 0, 0, 0))
+		[
+			ZRadioButton
+		]
+		+ SHorizontalBox::Slot()
+		.VAlign(VAlign_Center)
+		.HAlign(HAlign_Left)
+		.AutoWidth()
+		.Padding(FMargin(10, 0, 50, 0))
+		[
+			SNew(STextBlock)
+				.Text(LOCTEXT("CameraUpRadioButtonTextZ", "Z"))
+		];
+
+	Section.AddEntry(FToolMenuEntry::InitWidget("BoneAsCameraUpAxisWidget", CheckboxRow, FText()));
+}
+
+void SetupBoneAsCameraPositionOffset(FToolMenuSection& Section, TSharedPtr<SAnimationEditorViewportTabBody>& ViewportTab)
+{
+	const TSharedRef<SWidget> CameraOffsets =
+		SNew(SNumericVectorInputBox<double>)
+		.IsEnabled(true)
+		.Vector_Lambda([ViewportWeak = ViewportTab.ToWeakPtr()]()
+			{
+				if (TSharedPtr<SAnimationEditorViewportTabBody> Viewport = ViewportWeak.Pin())
+				{
+					return Viewport->GetBoneAsCameraPositionOffset();
+				}
+
+				return UE::Math::TVector<double>::Zero();
+			}
+		)
+		.OnXChanged_Lambda(
+			[ViewportWeak = ViewportTab.ToWeakPtr()](double Value)
+			{
+				if (TSharedPtr<SAnimationEditorViewportTabBody> Viewport = ViewportWeak.Pin())
+				{
+					UE::Math::TVector<double> Offset = Viewport->GetBoneAsCameraPositionOffset();
+					Offset.SetComponentForAxis(EAxis::Type::X, Value);
+					Viewport->SetBoneAsCameraPositionOffset(Offset);
+				}
+			}
+		)
+		.OnYChanged_Lambda(
+			[ViewportWeak = ViewportTab.ToWeakPtr()](double Value)
+			{
+				if (TSharedPtr<SAnimationEditorViewportTabBody> Viewport = ViewportWeak.Pin())
+				{
+					UE::Math::TVector<double> Offset = Viewport->GetBoneAsCameraPositionOffset();
+					Offset.SetComponentForAxis(EAxis::Type::Y, Value);
+					Viewport->SetBoneAsCameraPositionOffset(Offset);
+				}
+			}
+		)
+		.OnZChanged_Lambda(
+			[ViewportWeak = ViewportTab.ToWeakPtr()](double Value)
+			{
+				if (TSharedPtr<SAnimationEditorViewportTabBody> Viewport = ViewportWeak.Pin())
+				{
+					UE::Math::TVector<double> Offset = Viewport->GetBoneAsCameraPositionOffset();
+					Offset.SetComponentForAxis(EAxis::Type::Z, Value);
+					Viewport->SetBoneAsCameraPositionOffset(Offset);
+				}
+			}
+		)
+		.AllowSpin(true);
+
+	Section.AddEntry(FToolMenuEntry::InitWidget("BoneAsCameraPositionOffsetWidget", CameraOffsets, FText()));
+}
+
+void SetupBoneAsCameraRotationOffset(FToolMenuSection& Section, TSharedPtr<SAnimationEditorViewportTabBody>& ViewportTab)
+{
+	const TSharedRef<SWidget> CameraOffsets =
+		SNew(SNumericVectorInputBox<double>)
+		.IsEnabled(true)
+		.Vector_Lambda([ViewportWeak = ViewportTab.ToWeakPtr()]()
+			{
+				if (TSharedPtr<SAnimationEditorViewportTabBody> Viewport = ViewportWeak.Pin())
+				{
+					const FRotator Offset = Viewport->GetBoneAsCameraRotationOffset();
+					return FVector(Offset.Roll, Offset.Pitch, Offset.Yaw);
+				}
+
+				return UE::Math::TVector<double>::Zero();
+			}
+		)
+		.OnXChanged_Lambda(
+			[ViewportWeak = ViewportTab.ToWeakPtr()](double Value)
+			{
+				if (TSharedPtr<SAnimationEditorViewportTabBody> Viewport = ViewportWeak.Pin())
+				{
+					FRotator Offset = Viewport->GetBoneAsCameraRotationOffset();
+					Offset.SetComponentForAxis(EAxis::Type::X, Value);
+					Viewport->SetBoneAsCameraRotationOffset(Offset);
+				}
+			}
+		)
+		.OnYChanged_Lambda(
+			[ViewportWeak = ViewportTab.ToWeakPtr()](double Value)
+			{
+				if (TSharedPtr<SAnimationEditorViewportTabBody> Viewport = ViewportWeak.Pin())
+				{
+					FRotator Offset = Viewport->GetBoneAsCameraRotationOffset();
+					Offset.SetComponentForAxis(EAxis::Type::Y, Value);
+					Viewport->SetBoneAsCameraRotationOffset(Offset);
+				}
+			}
+		)
+		.OnZChanged_Lambda(
+			[ViewportWeak = ViewportTab.ToWeakPtr()](double Value)
+			{
+				if (TSharedPtr<SAnimationEditorViewportTabBody> Viewport = ViewportWeak.Pin())
+				{
+					FRotator Offset = Viewport->GetBoneAsCameraRotationOffset();
+					Offset.SetComponentForAxis(EAxis::Type::Z, Value);
+					Viewport->SetBoneAsCameraRotationOffset(Offset);
+				}
+			}
+		)
+		.AllowSpin(true);
+
+	Section.AddEntry(FToolMenuEntry::InitWidget("BoneAsCameraRotationOffsetWidget", CameraOffsets, FText()));
+}
 // Convenience function to retrieve Anim Editor Viewport Tab from UToolMenu
 TSharedPtr<SAnimationEditorViewportTabBody> GetAnimationEditorViewportTab(const UToolMenu* InMenu)
 {
@@ -1284,6 +1741,77 @@ void UE::AnimationEditor::FillFollowModeSubmenu(UToolMenu* InMenu)
 			}
 		)
 	);
+
+	CameraFollowModeSection.AddSubMenu(
+		"BoneAsCamera",
+		LOCTEXT("BoneAsCamera_DisplayName", "Bone As Camera"),
+		LOCTEXT("BoneAsCamera_ToolTip", "Select a bone for the camera to look out of"),
+		FNewToolMenuDelegate::CreateLambda(
+			[](UToolMenu* InSubmenu)
+			{
+				InSubmenu->AddDynamicSection(
+					"BoneAsCameraSubmenu",
+					FNewToolMenuDelegate::CreateLambda(
+						[](UToolMenu* InSectionMenu)
+						{
+							TSharedPtr<SAnimationEditorViewportTabBody> ViewportTab =
+								Private::GetAnimationEditorViewportTab(InSectionMenu);
+
+							if (!ViewportTab)
+							{
+								return;
+							}
+
+							FToolMenuSection& BoneAsCameraModeSelection = InSectionMenu->FindOrAddSection(
+								"BoneAsCameraSection",
+								LOCTEXT("BoneAsCameraSection_Label", "Bone Selector")
+							);
+
+							BoneAsCameraModeSelection.AddEntry(FToolMenuEntry::InitWidget(
+								"BoneAsCameraWidget", UE::AnimationEditor::Private::MakeBoneAsCameraWidget(ViewportTab, nullptr), FText()
+							));
+
+							FToolMenuSection& BoneAsCameraForwardAxisSection = InSectionMenu->FindOrAddSection(
+								"BoneAsCameraForwardAxisSection",
+								LOCTEXT("BoneAsCameraForwardSection_Label", "Forward Axis Selector")
+							);
+
+							UE::AnimationEditor::Private::SetupBoneAsCameraForward(BoneAsCameraForwardAxisSection, ViewportTab);
+
+							FToolMenuSection& BoneAsCameraRightAxisSection = InSectionMenu->FindOrAddSection(
+								"BoneAsCameraRightAxisSection",
+								LOCTEXT("BoneAsCameraRightSection_Label", "Right Axis Selector")
+							);
+
+							UE::AnimationEditor::Private::SetupBoneAsCameraRight(BoneAsCameraRightAxisSection, ViewportTab);
+
+							FToolMenuSection& BoneAsCameraUpAxisSection = InSectionMenu->FindOrAddSection(
+								"BoneAsCameraAxisUpSection",
+								LOCTEXT("BoneAsCameraUpSection_Label", "Up Axis Selector")
+							);
+
+							UE::AnimationEditor::Private::SetupBoneAsCameraUp(BoneAsCameraUpAxisSection, ViewportTab);
+
+							FToolMenuSection& BoneAsCameraPositionOffsetSection = InSectionMenu->FindOrAddSection(
+								"BoneAsCameraPositionOffset",
+								LOCTEXT("BoneAsCameraPositoinOffsetSection_Label", "Position Offset Selector")
+							);
+
+							UE::AnimationEditor::Private::SetupBoneAsCameraPositionOffset(BoneAsCameraPositionOffsetSection, ViewportTab);
+
+							FToolMenuSection& BoneAsCameraRotationOffsetSection = InSectionMenu->FindOrAddSection(
+								"BoneAsCameraRotationOffset",
+								LOCTEXT("BoneAsCameraRotationOffsetSection_Label", "Rotation Offset Selector")
+							);
+
+							UE::AnimationEditor::Private::SetupBoneAsCameraRotationOffset(BoneAsCameraRotationOffsetSection, ViewportTab);
+						}
+					)
+				);
+			}
+		),
+		false //ShouldCloseWindowAfterSelection
+	);
 }
 
 void UE::AnimationEditor::ExtendViewModesSubmenu(FName InViewModesSubmenuName)
diff --git a/Engine/Source/Editor/Persona/Public/AnimationEditorViewportClient.h b/Engine/Source/Editor/Persona/Public/AnimationEditorViewportClient.h
--- a/Engine/Source/Editor/Persona/Public/AnimationEditorViewportClient.h
+++ b/Engine/Source/Editor/Persona/Public/AnimationEditorViewportClient.h
@@ -95,6 +95,31 @@ public:
 	/** Get the bone name to use when CameraFollowMode is EAnimationViewportCameraFollowMode::Bone */
 	UE_API FName GetCameraFollowBoneName() const;
 
+	enum class EDirection
+	{
+		Forward,
+		Right,
+		Up
+	};
+
+	/** Gets the axis for a direction when using the Bone As Camera follow mode */
+	EAxis::Type GetBoneAsCameraAxis(EDirection Direction) const;
+
+	/** Gets the position offset from the bone origin in world coordiantes for the Bone As Camera follow mode */
+	UE::Math::TVector<double> GetBoneAsCameraPositionOffset() const;
+
+	/** Gets the rotation offset from the bone in world coordiantes for the Bone As Camera follow mode */
+	FRotator GetBoneAsCameraRotationOffset() const;
+
+	/** Sets the axis for a direction when using the Bone As Camera follow mode */
+	void SetBoneAsCameraAxis(EDirection Direction, EAxis::Type Axis);
+
+	/** Sets the position offset from the bone origin in world coordiantes for the Bone As Camera follow mode */
+	void SetBoneAsCameraPositionOffset(UE::Math::TVector<double> Offset);
+
+	/** Sets the rotation offset from the bone in world coordiantes for the Bone As Camera follow mode */
+	void SetBoneAsCameraRotationOffset(const FRotator& Offset);
+
 	/** Jump to the meshes default camera */
 	UE_API void JumpToDefaultCamera();
 
@@ -319,6 +344,21 @@ private:
 	/** The bone we will follow when in EAnimationViewportCameraFollowMode::Bone */
 	FName CameraFollowBoneName;
 
+	/** The forward axis to use when in EAnimationViewportCameraFollowMode::BoneAsCamera */
+	EAxis::Type BoneAsCameraForward = EAxis::Type::X;
+
+	/** The right axis to use when in EAnimationViewportCameraFollowMode::BoneAsCamera */
+	EAxis::Type BoneAsCameraRight = EAxis::Type::Y;
+
+	/** The up axis to use when in EAnimationViewportCameraFollowMode::BoneAsCamera */
+	EAxis::Type BoneAsCameraUp = EAxis::Type::Z;
+
+	/** The position offset from the bone origin in world space to use when in EAnimationViewportCameraFollowMode::BoneAsCamera */
+	UE::Math::TVector<double> BoneAsCameraPositionOffset = UE::Math::TVector<double>::ZeroVector;
+
+	/** The rotation offset from the bone origin in world space to use when in EAnimationViewportCameraFollowMode::BoneAsCamera */
+	FRotator BoneAsCameraRotationOffset = FRotator::ZeroRotator;
+
 	/** Should we auto align floor to mesh bounds */
 	bool bAutoAlignFloor;
 
diff --git a/Engine/Source/Editor/UnrealEd/Classes/Preferences/PersonaOptions.h b/Engine/Source/Editor/UnrealEd/Classes/Preferences/PersonaOptions.h
--- a/Engine/Source/Editor/UnrealEd/Classes/Preferences/PersonaOptions.h
+++ b/Engine/Source/Editor/UnrealEd/Classes/Preferences/PersonaOptions.h
@@ -33,7 +33,10 @@ enum class EAnimationViewportCameraFollowMode : uint8
 	Bone,
 
 	/** Follow the root bone while keeping the mesh vertically centered */
-	Root
+	Root,
+
+	/** Use transform of bone to build viewport transform */
+	BoneAsCamera,
 };
 
 /** Persistent per-viewport options */
