Make the Animation Curve Viewer read only in Animation Blueprint Editors

#Synced-CL 49678682

diff --git a/Engine/Source/Editor/AnimationBlueprintEditor/Public/IAnimationBlueprintEditor.h b/Engine/Source/Editor/AnimationBlueprintEditor/Public/IAnimationBlueprintEditor.h
--- a/Engine/Source/Editor/AnimationBlueprintEditor/Public/IAnimationBlueprintEditor.h
+++ b/Engine/Source/Editor/AnimationBlueprintEditor/Public/IAnimationBlueprintEditor.h
@@ -23,4 +23,7 @@ public:
 
 	/** Get the preview anim instance we are using, which can be a linked instance depending on preview settings */
 	virtual UAnimInstance* GetPreviewInstance() const = 0;
+	
+	// Name of EditorToolkit, used for testing against FAssetEditorToolkit::GetToolkitFName 
+	static FName GetAnimationBlueprintEditorToolkitName() { static FName ToolkitName("AnimationBlueprintEditor"); return ToolkitName; }
 };
diff --git a/Engine/Source/Editor/Persona/Private/SAnimCurveViewer.cpp b/Engine/Source/Editor/Persona/Private/SAnimCurveViewer.cpp
--- a/Engine/Source/Editor/Persona/Private/SAnimCurveViewer.cpp
+++ b/Engine/Source/Editor/Persona/Private/SAnimCurveViewer.cpp
@@ -225,6 +225,7 @@ TSharedRef< SWidget > SAnimCurveListRow::GenerateWidgetForColumn( const FName& C
 				SNew( SSpinBox<float> )
 				.Font(FAppStyle::Get().GetFontStyle("SmallFont"))
 				.Value(this, &SAnimCurveListRow::GetWeight)
+				.IsEnabled(this, &SAnimCurveListRow::CanBeEdited)
 				.MinSliderValue(this, &SAnimCurveListRow::GetMinWeight)
 				.MaxSliderValue(this, &SAnimCurveListRow::GetMaxWeight)
 				.OnValueChanged( this, &SAnimCurveListRow::OnAnimCurveWeightChanged )
@@ -468,6 +469,15 @@ bool SAnimCurveListRow::GetActiveWeight(float& OutWeight) const
 	return false;
 }
 
+bool SAnimCurveListRow::CanBeEdited() const
+{
+	if (const TSharedPtr<SAnimCurveViewer> AnimCurveViewer = AnimCurveViewerPtr.Pin())
+	{
+		return !AnimCurveViewer->bIsReadOnlyViewer;
+	}
+
+	return false;
+}
 
 FSlateFontInfo SAnimCurveListRow::GetItemFont() const
 {
@@ -516,6 +526,7 @@ void SAnimCurveViewer::Construct(const FArguments& InArgs,  const TSharedRef<IPe
 
 	PreviewScenePtr = InPreviewScene;
 	EditableSkeletonPtr = InArgs._EditableSkeleton;
+	bIsReadOnlyViewer = InArgs._bIsReadOnly;
 
 	InPreviewScene->RegisterOnPreviewMeshChanged(FOnPreviewMeshChanged::CreateSP(this, &SAnimCurveViewer::OnPreviewMeshChanged));
 	InPreviewScene->RegisterOnAnimChanged(FOnAnimChanged::CreateSP(this, &SAnimCurveViewer::OnPreviewAssetChanged));
@@ -525,6 +536,11 @@ void SAnimCurveViewer::Construct(const FArguments& InArgs,  const TSharedRef<IPe
 	BindCommands();
 
 	CurrentCurveFlag = EAnimCurveViewerFilterFlags::Active;
+	TArray<FName> HiddenColumns;
+	if (bIsReadOnlyViewer)
+	{
+		HiddenColumns.Add(CurveViewerColumns::AnimCurveEditLabel);
+	}
 
 	TSharedPtr<FFilterCategory> FilterCategory = MakeShared<FFilterCategory>(LOCTEXT("CurveFiltersLabel", "Curve Filters"), LOCTEXT("CurveFiltersToolTip", "Filter what kind fo curves can be displayed."));
 
@@ -661,6 +677,7 @@ void SAnimCurveViewer::Construct(const FArguments& InArgs,  const TSharedRef<IPe
 			.HeaderRow
 			(
 				SNew( SHeaderRow )
+				.HiddenColumnsList(HiddenColumns)
 				+ SHeaderRow::Column(CurveViewerColumns::AnimCurveNameLabel)
 				.FillWidth(1.f)
 				.DefaultLabel( LOCTEXT( "AnimCurveNameLabel", "Curve Name" ) )
diff --git a/Engine/Source/Editor/Persona/Private/SAnimCurveViewer.h b/Engine/Source/Editor/Persona/Private/SAnimCurveViewer.h
--- a/Engine/Source/Editor/Persona/Private/SAnimCurveViewer.h
+++ b/Engine/Source/Editor/Persona/Private/SAnimCurveViewer.h
@@ -138,6 +138,9 @@ private:
 	/** Get current active weight. Returns false if not currently active */
 	bool GetActiveWeight(float& OutWeight) const;
 
+	/** If the values of the curve represented by this row can be edited */
+	bool CanBeEdited() const;
+
 	/* The SAnimCurveViewer that we push the morph target weights into */
 	TWeakPtr<SAnimCurveViewer> AnimCurveViewerPtr;
 
@@ -166,7 +169,7 @@ public:
 	SLATE_BEGIN_ARGS( SAnimCurveViewer ) {}
 
 	SLATE_ARGUMENT(TSharedPtr<IEditableSkeleton>, EditableSkeleton)
-
+	SLATE_ARGUMENT(bool, bIsReadOnly)
 	SLATE_END_ARGS()
 
 	/**
@@ -313,6 +316,9 @@ private:
 	FText FilterText;
 
 	EAnimCurveViewerFilterFlags CurrentCurveFlag;
+	
+	/** Whatever this viewer is read only or allows editing of its values */
+	bool bIsReadOnlyViewer = false;
 
 	TMap<FName, float> OverrideCurves;
 
diff --git a/Engine/Source/Editor/Persona/Private/TabSpawners.cpp b/Engine/Source/Editor/Persona/Private/TabSpawners.cpp
--- a/Engine/Source/Editor/Persona/Private/TabSpawners.cpp
+++ b/Engine/Source/Editor/Persona/Private/TabSpawners.cpp
@@ -10,6 +10,7 @@
 #include "SAnimBlueprintParentPlayerList.h"
 #include "SSkeletonSlotNames.h"
 #include "AdvancedPreviewSceneModule.h"
+#include "IAnimationBlueprintEditor.h"
 #include "ISkeletonTree.h"
 #include "ISkeletonEditorModule.h"
 #include "SPersonaDetails.h"
@@ -188,12 +189,15 @@ FAnimCurveViewerTabSummoner::FAnimCurveViewerTabSummoner(TSharedPtr<class FAsset
 
 	ViewMenuDescription = LOCTEXT("AnimCurveTabView", "Curve Debugger");
 	ViewMenuTooltip = LOCTEXT("AnimCurveTabView_ToolTip", "Shows the animation curve debugger. This shows the state of animation curves.");
+
+	bSetCurveViewerReadOnly = InHostingApp->GetToolkitFName() == IAnimationBlueprintEditor::GetAnimationBlueprintEditorToolkitName();
 }
 
 TSharedRef<SWidget> FAnimCurveViewerTabSummoner::CreateTabBody(const FWorkflowTabSpawnInfo& Info) const
 {
 	return SNew(SAnimCurveViewer, PreviewScene.Pin().ToSharedRef(), OnObjectsSelected)
-		.EditableSkeleton(EditableSkeleton.Pin());
+		.EditableSkeleton(EditableSkeleton.Pin())
+		.bIsReadOnly(bSetCurveViewerReadOnly);
 }
 
 /////////////////////////////////////////////////////
diff --git a/Engine/Source/Editor/Persona/Private/TabSpawners.h b/Engine/Source/Editor/Persona/Private/TabSpawners.h
--- a/Engine/Source/Editor/Persona/Private/TabSpawners.h
+++ b/Engine/Source/Editor/Persona/Private/TabSpawners.h
@@ -157,6 +157,7 @@ private:
 	TWeakPtr<class IEditableSkeleton> EditableSkeleton;
 	TWeakPtr<class IPersonaPreviewScene> PreviewScene;
 	FOnObjectsSelected OnObjectsSelected;
+	bool bSetCurveViewerReadOnly;
 };
 
 /////////////////////////////////////////////////////
