Set the Anim Curve Debugger Widget to read only per default

#Synced-CL 50246811

diff --git a/Engine/Source/Editor/Persona/Private/PersonaModule.cpp b/Engine/Source/Editor/Persona/Private/PersonaModule.cpp
--- a/Engine/Source/Editor/Persona/Private/PersonaModule.cpp
+++ b/Engine/Source/Editor/Persona/Private/PersonaModule.cpp
@@ -289,9 +289,9 @@ TSharedRef<class FWorkflowTabFactory> FPersonaModule::CreateCurveViewerTabFactor
 	return MakeShareable(new FAnimCurveViewerTabSummoner(InHostingApp, InEditableSkeleton, InPreviewScene, InOnObjectsSelected));
 }
 
-TSharedRef<class FWorkflowTabFactory> FPersonaModule::CreateCurveViewerTabFactory(const TSharedRef<class FWorkflowCentricApplication>& InHostingApp, const TSharedPtr<class IEditableSkeleton>& InEditableSkeleton, const TSharedRef<IPersonaPreviewScene>& InPreviewScene, FOnObjectsSelected InOnObjectsSelected) const
+TSharedRef<class FWorkflowTabFactory> FPersonaModule::CreateCurveViewerTabFactory(const TSharedRef<class FWorkflowCentricApplication>& InHostingApp, const TSharedPtr<class IEditableSkeleton>& InEditableSkeleton, const TSharedRef<IPersonaPreviewScene>& InPreviewScene, FOnObjectsSelected InOnObjectsSelected, bool bInReadOnlyCurveViewer) const
 {
-	return MakeShareable(new FAnimCurveViewerTabSummoner(InHostingApp, InEditableSkeleton, InPreviewScene, InOnObjectsSelected));
+	return MakeShareable(new FAnimCurveViewerTabSummoner(InHostingApp, InEditableSkeleton, InPreviewScene, InOnObjectsSelected, bInReadOnlyCurveViewer));
 }
 
 TSharedRef<class FWorkflowTabFactory> FPersonaModule::CreateCurveMetadataEditorTabFactory(const TSharedRef<class FWorkflowCentricApplication>& InHostingApp, UObject* InMetadataHost, const TSharedRef<IPersonaPreviewScene>& InPreviewScene, FOnObjectsSelected InOnObjectsSelected) const
diff --git a/Engine/Source/Editor/Persona/Private/TabSpawners.cpp b/Engine/Source/Editor/Persona/Private/TabSpawners.cpp
--- a/Engine/Source/Editor/Persona/Private/TabSpawners.cpp
+++ b/Engine/Source/Editor/Persona/Private/TabSpawners.cpp
@@ -175,11 +175,12 @@ TSharedRef<SWidget> FMorphTargetTabSummoner::CreateTabBody(const FWorkflowTabSpa
 /////////////////////////////////////////////////////
 // FAnimCurveViewerTabSummoner
 
-FAnimCurveViewerTabSummoner::FAnimCurveViewerTabSummoner(TSharedPtr<class FAssetEditorToolkit> InHostingApp, const TSharedPtr<IEditableSkeleton>& InEditableSkeleton, const TSharedRef<IPersonaPreviewScene>& InPreviewScene, FOnObjectsSelected InOnObjectsSelected)
+FAnimCurveViewerTabSummoner::FAnimCurveViewerTabSummoner(TSharedPtr<class FAssetEditorToolkit> InHostingApp, const TSharedPtr<IEditableSkeleton>& InEditableSkeleton, const TSharedRef<IPersonaPreviewScene>& InPreviewScene, FOnObjectsSelected InOnObjectsSelected, bool bInSetCurveViewerReadOnly)
 	: FWorkflowTabFactory(FPersonaTabs::AnimCurveViewID, InHostingApp)
 	, EditableSkeleton(InEditableSkeleton)
 	, PreviewScene(InPreviewScene)
 	, OnObjectsSelected(InOnObjectsSelected)
+	, bSetCurveViewerReadOnly(bInSetCurveViewerReadOnly)
 {
 	TabLabel = LOCTEXT("AnimCurveViewerTabTitle", "Curve Debugger");
 	TabIcon = FSlateIcon(FAppStyle::GetAppStyleSetName(), "Persona.Tabs.AnimCurveDebugger");
@@ -189,8 +190,6 @@ FAnimCurveViewerTabSummoner::FAnimCurveViewerTabSummoner(TSharedPtr<class FAsset
 
 	ViewMenuDescription = LOCTEXT("AnimCurveTabView", "Curve Debugger");
 	ViewMenuTooltip = LOCTEXT("AnimCurveTabView_ToolTip", "Shows the animation curve debugger. This shows the state of animation curves.");
-
-	bSetCurveViewerReadOnly = InHostingApp->GetToolkitFName() == IAnimationBlueprintEditor::GetAnimationBlueprintEditorToolkitName();
 }
 
 TSharedRef<SWidget> FAnimCurveViewerTabSummoner::CreateTabBody(const FWorkflowTabSpawnInfo& Info) const
diff --git a/Engine/Source/Editor/Persona/Private/TabSpawners.h b/Engine/Source/Editor/Persona/Private/TabSpawners.h
--- a/Engine/Source/Editor/Persona/Private/TabSpawners.h
+++ b/Engine/Source/Editor/Persona/Private/TabSpawners.h
@@ -143,7 +143,7 @@ private:
 struct FAnimCurveViewerTabSummoner : public FWorkflowTabFactory
 {
 public:
-	FAnimCurveViewerTabSummoner(TSharedPtr<class FAssetEditorToolkit> InHostingApp, const TSharedPtr<IEditableSkeleton>& InEditableSkeleton, const TSharedRef<IPersonaPreviewScene>& InPreviewScene, FOnObjectsSelected InOnObjectsSelected);
+	FAnimCurveViewerTabSummoner(TSharedPtr<class FAssetEditorToolkit> InHostingApp, const TSharedPtr<IEditableSkeleton>& InEditableSkeleton, const TSharedRef<IPersonaPreviewScene>& InPreviewScene, FOnObjectsSelected InOnObjectsSelected, bool bInSetCurveViewerReadOnly = true);
 
 	virtual TSharedRef<SWidget> CreateTabBody(const FWorkflowTabSpawnInfo& Info) const override;
 
diff --git a/Engine/Source/Editor/Persona/Public/PersonaModule.h b/Engine/Source/Editor/Persona/Public/PersonaModule.h
--- a/Engine/Source/Editor/Persona/Public/PersonaModule.h
+++ b/Engine/Source/Editor/Persona/Public/PersonaModule.h
@@ -165,10 +165,10 @@ struct FBlendSpacePreviewArgs
 {
 	TAttribute<const UBlendSpace*> PreviewBlendSpace;
 
-	// Allows the target preview position to be programatically driven
+	// Allows the target preview position to be programmatically driven
 	TAttribute<FVector> PreviewPosition;
 
-	// Allows the current preview position to be programatically driven
+	// Allows the current preview position to be programmatically driven
 	TAttribute<FVector> PreviewFilteredPosition;
 
 	// Called to get the overridden name of a blend sample
@@ -380,8 +380,8 @@ public:
 	virtual TSharedRef<FWorkflowTabFactory> CreateCurveViewerTabFactory(const TSharedRef<class FWorkflowCentricApplication>& InHostingApp, const TSharedPtr<class IEditableSkeleton>& InEditableSkeleton, const TSharedRef<class IPersonaPreviewScene>& InPreviewScene, FSimpleMulticastDelegate& InOnPostUndo, FOnObjectsSelected InOnObjectsSelected) const;
 
 	/** Create a skeleton curve viewer tab factory */
-	virtual TSharedRef<FWorkflowTabFactory> CreateCurveViewerTabFactory(const TSharedRef<class FWorkflowCentricApplication>& InHostingApp, const TSharedPtr<class IEditableSkeleton>& InEditableSkeleton, const TSharedRef<class IPersonaPreviewScene>& InPreviewScene, FOnObjectsSelected InOnObjectsSelected) const;
-
+	virtual TSharedRef<FWorkflowTabFactory> CreateCurveViewerTabFactory(const TSharedRef<class FWorkflowCentricApplication>& InHostingApp, const TSharedPtr<class IEditableSkeleton>& InEditableSkeleton, const TSharedRef<class IPersonaPreviewScene>& InPreviewScene, FOnObjectsSelected InOnObjectsSelected, bool bInReadOnlyCurveViewer = true) const;
+	
 	/** Create a skeleton curve metadata editor tab factory */
 	virtual TSharedRef<FWorkflowTabFactory> CreateCurveMetadataEditorTabFactory(const TSharedRef<class FWorkflowCentricApplication>& InHostingApp, UObject* InMetadataHost, const TSharedRef<class IPersonaPreviewScene>& InPreviewScene, FOnObjectsSelected InOnObjectsSelected) const;
 
@@ -398,10 +398,10 @@ public:
 	/** Create a tab factory for editing a single object (like an animation asset) */
 	virtual TSharedRef<FWorkflowTabFactory> CreateAssetDetailsTabFactory(const TSharedRef<class FWorkflowCentricApplication>& InHostingApp, FOnGetAsset InOnGetAsset, FOnDetailsCreated InOnDetailsCreated) const;
 
-	/** Create a tab factory for for previewing morph targets */
+	/** Create a tab factory for previewing morph targets */
 	virtual TSharedRef<FWorkflowTabFactory> CreateMorphTargetTabFactory(const TSharedRef<class FWorkflowCentricApplication>& InHostingApp, const TSharedRef<IPersonaPreviewScene>& InPreviewScene, FSimpleMulticastDelegate& OnPostUndo) const;
 
-	/** Create a tab factory for for previewing morph targets */
+	/** Create a tab factory for previewing morph targets */
 	virtual TSharedRef<FWorkflowTabFactory> CreateMorphTargetTabFactory(const TSharedRef<class FWorkflowCentricApplication>& InHostingApp, TSharedRef<IMorphTargetViewer> InViewerWidget) const;
 
 	/** Create the default morph target viewer widget*/
@@ -596,7 +596,7 @@ public:
 	virtual uint32 GetNotifyHostAssetVersion() const;
 
 public:
-	// Register a new blend provider provider in the editor picker for FBlendProfileInterfaceWrapper
+	// Register a new blend provider in the editor picker for FBlendProfileInterfaceWrapper
 	PERSONA_API void RegisterBlendProfilePickerExtender(TSharedPtr<IBlendProfilePickerExtender> Extender);
 
 	PERSONA_API void UnregisterBlendProfilePickerExtender(const FName ExtenderId);
