Fix crash when updating blend samples in a Blendspace while PIE is running in a new window

#Synced-CL 50648665

diff --git a/Engine/Source/Editor/Persona/Private/Customization/BlendSpaceDetails.cpp b/Engine/Source/Editor/Persona/Private/Customization/BlendSpaceDetails.cpp
--- a/Engine/Source/Editor/Persona/Private/Customization/BlendSpaceDetails.cpp
+++ b/Engine/Source/Editor/Persona/Private/Customization/BlendSpaceDetails.cpp
@@ -155,9 +155,7 @@ void FBlendSpaceDetails::CustomizeDetails(class IDetailLayoutBuilder& DetailBuil
 	DetailBuilder.GetObjectsBeingCustomized(Objects);
 
 	Builder = &DetailBuilder;
-
-	FSimpleDelegate RefreshDelegate = FSimpleDelegate::CreateLambda([this]() { Builder->ForceRefreshDetails(); });
-
+	
 	TWeakObjectPtr<UObject>* BlendSpacePtr = Objects.FindByPredicate([](const TWeakObjectPtr<UObject>& ObjectPtr) { return ObjectPtr->IsA<UBlendSpace>(); });
 	if (BlendSpacePtr)
 	{
@@ -442,8 +440,8 @@ void FBlendSpaceDetails::CustomizeDetails(class IDetailLayoutBuilder& DetailBuil
 			{
 				 DefaultProperty->MarkHiddenByCustomization();
 			}
-
-			BlendSamplesArrayProperty->SetOnNumElementsChanged(RefreshDelegate);
+			
+			BlendSamplesArrayProperty->SetOnNumElementsChanged(FSimpleDelegate::CreateSP(this, &FBlendSpaceDetails::RefreshDetails));
 
 			// Add a "Remove all" button if there are some samples. Only in the asset blendspace for now.
 			if (NumBlendSampleEntries && BlendSpace->IsAsset())
@@ -624,9 +622,12 @@ void FBlendSpaceDetails::PostUndo(bool bSuccess)
 	GEditor->GetTimerManager()->SetTimerForNextTick(FTimerDelegate::CreateSP(this, &FBlendSpaceDetails::RefreshDetails));
 }
 
-void FBlendSpaceDetails::RefreshDetails()
+void FBlendSpaceDetails::RefreshDetails() const
 {
-	Builder->ForceRefreshDetails();
+	if (Builder)
+	{
+		Builder->ForceRefreshDetails();
+	}
 }
 
 #undef LOCTEXT_NAMESPACE
diff --git a/Engine/Source/Editor/Persona/Private/Customization/BlendSpaceDetails.h b/Engine/Source/Editor/Persona/Private/Customization/BlendSpaceDetails.h
--- a/Engine/Source/Editor/Persona/Private/Customization/BlendSpaceDetails.h
+++ b/Engine/Source/Editor/Persona/Private/Customization/BlendSpaceDetails.h
@@ -36,7 +36,7 @@ private:
 	FReply HandleClearSamples();
 	FReply HandleAnalyzeSamples();
 	void HandleAnalysisFunctionChanged(int32 AxisIndex, TSharedPtr<FString> NewItem);
-	void RefreshDetails();
+	void RefreshDetails() const;
 
 	IDetailLayoutBuilder* Builder;
 	UBlendSpace* BlendSpace;
