AnimationEditorPreviewScene - Request a new Audio Device so that we don't share the same Audio Device with Level Viewports

[RN] minor framework/audio

#Synced-CL 51007606

diff --git a/Engine/Source/Editor/Persona/Private/AnimationEditorPreviewScene.cpp b/Engine/Source/Editor/Persona/Private/AnimationEditorPreviewScene.cpp
--- a/Engine/Source/Editor/Persona/Private/AnimationEditorPreviewScene.cpp
+++ b/Engine/Source/Editor/Persona/Private/AnimationEditorPreviewScene.cpp
@@ -1,37 +1,39 @@
 // Copyright Epic Games, Inc. All Rights Reserved.
 
 #include "AnimationEditorPreviewScene.h"
-#include "Framework/Notifications/NotificationManager.h"
-#include "Widgets/Notifications/SNotificationList.h"
-#include "Misc/MessageDialog.h"
-#include "Modules/ModuleManager.h"
-#include "Components/StaticMeshComponent.h"
-#include "Styling/AppStyle.h"
 
-#include "Animation/AnimBlueprint.h"
+#include "AnimCustomInstanceHelper.h"
+#include "AnimPreviewAttacheInstance.h"
 #include "AnimPreviewInstance.h"
+#include "AnimViewportPlaybackCommands.h"
+#include "AudioDevice.h"
+#include "AudioDeviceManager.h"
+#include "ComponentAssetBroker.h"
 #include "IEditableSkeleton.h"
 #include "IPersonaToolkit.h"
-#include "PersonaUtils.h"
-#include "ComponentAssetBroker.h"
-#include "Engine/PreviewMeshCollection.h"
-#include "PersonaPreviewSceneDescription.h"
+#include "ISkeletonTreeItem.h"
+#include "PersonaModule.h"
 #include "PersonaPreviewSceneDefaultController.h"
+#include "PersonaPreviewSceneDescription.h"
+#include "PersonaUtils.h"
+#include "ScopedTransaction.h"
+#include "UICommandList_Pinnable.h"
+#include "Animation/AnimBlueprint.h"
+#include "Animation/PreviewCollectionInterface.h"
+#include "Components/StaticMeshComponent.h"
 #include "Components/WindDirectionalSourceComponent.h"
-#include "PhysicsEngine/PhysicsSettings.h"
-#include "PersonaModule.h"
+#include "Engine/PreviewMeshCollection.h"
+#include "Engine/SkeletalMeshSocket.h"
+#include "Factories/PreviewMeshCollectionFactory.h"
+#include "Framework/Notifications/NotificationManager.h"
 #include "GameFramework/WorldSettings.h"
+#include "Misc/MessageDialog.h"
+#include "Modules/ModuleManager.h"
 #include "Particles/ParticleSystemComponent.h"
-#include "Factories/PreviewMeshCollectionFactory.h"
-#include "AnimPreviewAttacheInstance.h"
-#include "AnimCustomInstanceHelper.h"
-#include "AnimViewportPlaybackCommands.h"
-#include "Animation/PreviewCollectionInterface.h"
-#include "ScopedTransaction.h"
+#include "PhysicsEngine/PhysicsSettings.h"
 #include "Preferences/PersonaOptions.h"
-#include "ISkeletonTreeItem.h"
-#include "UICommandList_Pinnable.h"
-#include "Engine/SkeletalMeshSocket.h"
+#include "Styling/AppStyle.h"
+#include "Widgets/Notifications/SNotificationList.h"
 
 #define LOCTEXT_NAMESPACE "AnimationEditorPreviewScene"
 
@@ -101,8 +103,31 @@
 	}
 
 	// Disable killing actors outside of the world
-	AWorldSettings* WorldSettings = GetWorld()->GetWorldSettings(true);
+	UWorld* World = GetWorld();
+	AWorldSettings* WorldSettings = World->GetWorldSettings(true);
 	WorldSettings->bEnableWorldBoundsChecks = false;
+	
+	// Create a new audio device for this scene
+	if (FAudioDeviceManager* AudioDeviceManager = GEngine->GetAudioDeviceManager())
+	{
+		FAudioDeviceParams DeviceParams = AudioDeviceManager->GetDefaultParamsForNewWorld();
+		DeviceParams.AssociatedWorld = World;
+		const FAudioDeviceHandle AudioDevice = AudioDeviceManager->RequestAudioDevice(DeviceParams);
+		if (AudioDevice.IsValid())
+		{
+			// Set the base mix of the new device based on the world settings of the world
+			AudioDevice.GetAudioDevice()->SetDefaultBaseSoundMix(World->GetWorldSettings()->DefaultBaseSoundMix);
 
+			// Set the world's audio device handle to use so that sounds which play in that world will use the correct audio device
+			World->SetAudioDevice(AudioDevice);
+				
+			FWorldContext& WorldContext = GEngine->GetWorldContextFromWorldChecked(World);
+
+			// Set this audio device handle on the world context so future world's set onto the world context
+			// will pass the audio device handle to them and audio will play on the correct audio device
+			WorldContext.AudioDeviceID = AudioDevice.GetDeviceID();
+		}
+	}
+
 	BindCommands();
 }
