AnimationEditorViewportClient & LevelEditorViewportClient UpdateAudioListener callsite changes

- Set Audio Focus on Received Focus
- Clear Audio Focus on Lost Focus
- UpdateAudioListener called on receive focus and when the perspective camera moves
- Additional call in AnimationEditorViewportClient when we focus on the preview mesh

#rnx

#Synced-CL 51460112
#Synced-CL 51460112

diff --git a/Engine/Source/Editor/Persona/Private/AnimationEditorViewportClient.cpp b/Engine/Source/Editor/Persona/Private/AnimationEditorViewportClient.cpp
--- a/Engine/Source/Editor/Persona/Private/AnimationEditorViewportClient.cpp
+++ b/Engine/Source/Editor/Persona/Private/AnimationEditorViewportClient.cpp
@@ -254,6 +254,8 @@ FAnimationViewportClient::FAnimationViewportClient(const TSharedRef<IPersonaPrev
 		check(AudioConfigOption);
 		AudioConfigOption->SetUseAudioAttenuation(true);
 	}
+	
+	SetAudioFocus();
 }
 
 FAnimationViewportClient::~FAnimationViewportClient()
@@ -1684,6 +1686,20 @@ bool FAnimationViewportClient::InputAxis(const FInputKeyEventArgs& EventArgs)
 	return bResult;
 }
 
+void FAnimationViewportClient::ReceivedFocus(FViewport* InViewport)
+{
+	FEditorViewportClient::ReceivedFocus(InViewport);
+	
+	UpdateAudioListener();
+	SetAudioFocus();
+}
+
+void FAnimationViewportClient::LostFocus(FViewport* InViewport)
+{
+	ClearAudioFocus();
+	FEditorViewportClient::LostFocus(InViewport);
+}
+
 void FAnimationViewportClient::SetLocalAxesMode(ELocalAxesMode::Type AxesMode)
 {
 	ConfigOption->SetDefaultLocalAxesSelection( AxesMode );
@@ -2524,6 +2540,11 @@ void FAnimationViewportClient::FocusViewportOnPreviewMesh(bool bUseCustomCamera)
 		bFocusUsingCustomCamera = bUseCustomCamera;
 		return;
 	}
+	
+	if (bHasAudioFocus)
+	{
+		UpdateAudioListener();
+	}
 
 	if (UDebugSkelMeshComponent* const PreviewMeshComponent = GetAnimPreviewScene()->GetPreviewMeshComponent())
 	{
@@ -2758,6 +2779,16 @@ IPersonaEditorModeManager* FAnimationViewportClient::GetPersonaModeManager() con
 	return ModeTools->GetInteractiveToolsContext()->ContextObjectStore->FindContext<UPersonaEditorModeManagerContext>()->GetPersonaEditorModeManager();
 }
 
+void FAnimationViewportClient::PerspectiveCameraMoved()
+{
+	FEditorViewportClient::PerspectiveCameraMoved();
+	
+	if (bHasAudioFocus)
+	{
+		UpdateAudioListener();
+	}
+}
+
 void FAnimationViewportClient::HandleInvalidateViews()
 {
 	Invalidate();
@@ -2775,6 +2806,11 @@ bool FAnimationViewportClient::CanCycleWidgetMode() const
 }
 
 void FAnimationViewportClient::UpdateAudioListener(const FSceneView& View)
+{
+	UpdateAudioListener();
+}
+
+void FAnimationViewportClient::UpdateAudioListener() const
 {
 	UWorld* ViewportWorld = GetWorld();
 
@@ -2789,6 +2825,7 @@ void FAnimationViewportClient::UpdateAudioListener(const FSceneView& View)
 			ListenerTransform.SetLocation(ViewLocation);
 
 			AudioDevice->SetListener(ViewportWorld, 0, ListenerTransform, 0.0f);
+			UE_LOGFMT(LogAnimation, Log, "Animation Viewport {Index}, {Transform}", ViewIndex, ListenerTransform.ToHumanReadableString());
 		}
 	}
 }
@@ -2797,11 +2834,6 @@ void FAnimationViewportClient::SetupViewForRendering( FSceneViewFamily& ViewFami
 {
 	FEditorViewportClient::SetupViewForRendering( ViewFamily, View );
 
-	if(bHasAudioFocus)
-	{
-		UpdateAudioListener(View);
-	}
-
 	// Cache screen size
 	const UDebugSkelMeshComponent* PreviewMeshComponent = GetAnimPreviewScene()->GetPreviewMeshComponent();
 	if (PreviewMeshComponent != nullptr && PreviewMeshComponent->MeshObject != nullptr)
diff --git a/Engine/Source/Editor/Persona/Public/AnimationEditorViewportClient.h b/Engine/Source/Editor/Persona/Public/AnimationEditorViewportClient.h
--- a/Engine/Source/Editor/Persona/Public/AnimationEditorViewportClient.h
+++ b/Engine/Source/Editor/Persona/Public/AnimationEditorViewportClient.h
@@ -55,6 +55,8 @@ public:
 	UE_API virtual void DrawCanvas( FViewport& InViewport, FSceneView& View, FCanvas& Canvas ) override;
 	UE_API virtual bool InputKey(const FInputKeyEventArgs& EventArgs) override;
 	UE_API virtual bool InputAxis(const FInputKeyEventArgs& EventArgs) override;
+	UE_API virtual void ReceivedFocus(FViewport* InViewport) override;
+	UE_API virtual void LostFocus(FViewport* InViewport) override;
 	UE_API virtual void SetCameraSpeedSettings(const FEditorViewportCameraSpeedSettings& InCameraSpeedSettings) override;
 
 	// Sets what bones are drawn by DrawMeshBones and ShowBoneNames
@@ -306,6 +308,9 @@ public:
 	UE_DEPRECATED(5.1, "Use the UPersonaEditorModeManagerContext object stored in the editor mode tools' context store instead.")
 	UE_API class IPersonaEditorModeManager* GetPersonaModeManager() const;
 
+protected:
+	virtual void PerspectiveCameraMoved() override;
+
 private:
 	/**
 	 * Updates the audio listener for this viewport 
@@ -313,6 +318,11 @@ private:
 	 * @param View	The scene view to use when calculate the listener position
 	 */
 	UE_API void UpdateAudioListener(const FSceneView& View);
+	
+	/**
+	 * Updates the audio listener for this viewport 
+	 */
+	void UpdateAudioListener() const;
 
 public:
 
@@ -523,6 +533,7 @@ private:
 	// Delegate Handler to allow changing of camera controller
 	UE_API void OnCameraControllerChanged();
 
+private:
 	/** True when the preview animation should resume playing upon finishing tracking */
 	bool bResumeAfterTracking;
 
diff --git a/Engine/Source/Editor/UnrealEd/Private/EditorEngine.cpp b/Engine/Source/Editor/UnrealEd/Private/EditorEngine.cpp
--- a/Engine/Source/Editor/UnrealEd/Private/EditorEngine.cpp
+++ b/Engine/Source/Editor/UnrealEd/Private/EditorEngine.cpp
@@ -1837,9 +1837,6 @@ void UEditorEngine::Tick( float DeltaSeconds, bool bIdleMode )
 
 		for(FEditorViewportClient* const ViewportClient : AllViewportClients)
 		{
-			// clear any previous audio focus flags
-			ViewportClient->ClearAudioFocus();
-
 			if (ViewportClient->IsPerspective())
 			{
 				if (ViewportClient->IsRealtime())
@@ -1879,8 +1876,6 @@ void UEditorEngine::Tick( float DeltaSeconds, bool bIdleMode )
 	// tell viewportclient it has audio focus
 	if (AudioFocusViewportClient)
 	{
-		AudioFocusViewportClient->SetAudioFocus();
-
 		// override realtime setting if viewport chooses (i.e. for cinematic preview)
 		if (AudioFocusViewportClient->IsForcedRealtimeAudio())
 		{
diff --git a/Engine/Source/Editor/UnrealEd/Private/LevelEditorViewport.cpp b/Engine/Source/Editor/UnrealEd/Private/LevelEditorViewport.cpp
--- a/Engine/Source/Editor/UnrealEd/Private/LevelEditorViewport.cpp
+++ b/Engine/Source/Editor/UnrealEd/Private/LevelEditorViewport.cpp
@@ -2689,6 +2689,11 @@ void FLevelEditorViewportClient::PerspectiveCameraMoved()
 
 	// Broadcast 'camera moved' delegate
 	FEditorDelegates::OnEditorCameraMoved.Broadcast(GetViewLocation(), GetViewRotation(), ViewportType, ViewIndex);
+	
+	if (bHasAudioFocus)
+	{
+		UpdateAudioListener();
+	}
 }
 
 /**
@@ -2769,6 +2774,9 @@ void FLevelEditorViewportClient::ReceivedFocus(FViewport* InViewport)
 	}
 
 	FEditorViewportClient::ReceivedFocus(InViewport);
+	
+	UpdateAudioListener();
+	SetAudioFocus();
 }
 
 void FLevelEditorViewportClient::LostFocus(FViewport* InViewport)
@@ -2787,6 +2795,7 @@ void FLevelEditorViewportClient::LostFocus(FViewport* InViewport)
 	}
 
 	GEditor->SetPreviewMeshMode(false);
+	ClearAudioFocus();
 }
 
 
@@ -5182,6 +5191,11 @@ void FLevelEditorViewportClient::DrawBrushDetails(const FSceneView* View, FPrimi
 }
 
 void FLevelEditorViewportClient::UpdateAudioListener(const FSceneView& View)
+{
+	UpdateAudioListener();
+}
+
+void FLevelEditorViewportClient::UpdateAudioListener()
 {
 	UWorld* ViewportWorld = GetWorld();
 
@@ -5210,6 +5224,7 @@ void FLevelEditorViewportClient::UpdateAudioListener(const FSceneView& View)
 			ListenerTransform.SetLocation(ViewLocation);
 
 			AudioDevice->SetListener(ViewportWorld, 0, ListenerTransform, 0.f);
+			UE_LOGFMT(LogEditorViewport, Log, "Level Editor Viewport {Index}, {Transform}", ViewIndex, ListenerTransform.ToHumanReadableString());
 		}
 	}
 }
@@ -5285,12 +5300,6 @@ void FLevelEditorViewportClient::SetupViewForRendering( FSceneViewFamily& ViewFa
 		// Hide the builder brush when not in geometry mode
 		ViewFamily.EngineShowFlags.SetBuilderBrush(false);
 	}
-
-	// Update the listener.
-	if (bHasAudioFocus)
-	{
-		UpdateAudioListener(View);
-	}
 }
 
 void FLevelEditorViewportClient::DrawCanvas( FViewport& InViewport, FSceneView& View, FCanvas& Canvas )
diff --git a/Engine/Source/Editor/UnrealEd/Public/LevelEditorViewport.h b/Engine/Source/Editor/UnrealEd/Public/LevelEditorViewport.h
--- a/Engine/Source/Editor/UnrealEd/Public/LevelEditorViewport.h
+++ b/Engine/Source/Editor/UnrealEd/Public/LevelEditorViewport.h
@@ -288,7 +288,13 @@ public:
 	 *
 	 * @param View	The scene view to use when calculate the listener position
 	 */
+	UE_DEPRECATED(5.8, "View is redundant so use method with no parameters")
 	UNREALED_API void UpdateAudioListener( const FSceneView& View );
+	
+	/**
+	 * Updates the audio listener for this viewport.
+	 */
+	UNREALED_API void UpdateAudioListener();
 
 	/** Determines if the new MoveCanvas movement should be used */
 	UNREALED_API bool ShouldUseMoveCanvasMovement (void);
