Persona Editors: forward internally handled asset editor selection to the ModeTools selection set

#Synced-CL 51856350

diff --git a/Engine/Source/Editor/AnimationEditor/AnimationEditor.Build.cs b/Engine/Source/Editor/AnimationEditor/AnimationEditor.Build.cs
--- a/Engine/Source/Editor/AnimationEditor/AnimationEditor.Build.cs
+++ b/Engine/Source/Editor/AnimationEditor/AnimationEditor.Build.cs
@@ -22,7 +22,8 @@ public class AnimationEditor : ModuleRules
                 "SkeletonEditor",
                 "Kismet",
                 "AnimGraph",
-				"ToolMenus"
+				"ToolMenus",
+				"TypedElementRuntime"
             }
 		);
 
diff --git a/Engine/Source/Editor/AnimationEditor/Private/AnimationEditor.cpp b/Engine/Source/Editor/AnimationEditor/Private/AnimationEditor.cpp
--- a/Engine/Source/Editor/AnimationEditor/Private/AnimationEditor.cpp
+++ b/Engine/Source/Editor/AnimationEditor/Private/AnimationEditor.cpp
@@ -417,6 +417,8 @@ void FAnimationEditor::HandleObjectsSelected(const TArray<UObject*>& InObjects)
 
 void FAnimationEditor::HandleSelectionChanged(const TArrayView<TSharedPtr<ISkeletonTreeItem>>& InSelectedItems, ESelectInfo::Type InSelectInfo)
 {
+	RouteTreeSelectionToModeManager(InSelectedItems);
+
 	if (DetailsView.IsValid())
 	{
 		TArray<UObject*> Objects;
diff --git a/Engine/Source/Editor/Persona/Persona.Build.cs b/Engine/Source/Editor/Persona/Persona.Build.cs
--- a/Engine/Source/Editor/Persona/Persona.Build.cs
+++ b/Engine/Source/Editor/Persona/Persona.Build.cs
@@ -85,6 +85,7 @@ public class Persona : ModuleRules
 				"InteractiveToolsFramework",
 				"ContentBrowser",
 	            "SkeletalMeshDescription",
+	            "TypedElementRuntime"
             }
         );
 
diff --git a/Engine/Source/Editor/Persona/Private/IPersonaEditorModeManager.cpp b/Engine/Source/Editor/Persona/Private/IPersonaEditorModeManager.cpp
--- a/Engine/Source/Editor/Persona/Private/IPersonaEditorModeManager.cpp
+++ b/Engine/Source/Editor/Persona/Private/IPersonaEditorModeManager.cpp
@@ -8,6 +8,7 @@
 #include "Slate/SceneViewport.h"
 #include "Framework/Application/SlateApplication.h"
 #include "Framework/Application/SlateUser.h"
+#include "Selection.h"
 
 #include UE_INLINE_GENERATED_CPP_BY_NAME(IPersonaEditorModeManager)
 
@@ -60,11 +61,18 @@ void UPersonaEditorModeManagerContext::SetFocusInViewport()
 	});
 }
 
-IPersonaEditorModeManager::IPersonaEditorModeManager() : FAssetEditorModeManager()
+IPersonaEditorModeManager::IPersonaEditorModeManager()
 {
 	UModeManagerInteractiveToolsContext* ToolsContext = GetInteractiveToolsContext();
 	UContextObjectStore* ContextObjectStore = ToolsContext->ContextObjectStore;
 	ContextObjectStore->AddContextObject(PersonaModeManagerContext.Get());
+
+	SelectionSet = NewObject<UTypedElementSelectionSet>(GetTransientPackage(), NAME_None, RF_Transactional);
+
+	// Support legacy USelections
+	ActorSet->SetElementSelectionSet(SelectionSet);
+	ComponentSet->SetElementSelectionSet(SelectionSet);
+	ObjectSet->SetElementSelectionSet(SelectionSet);
 }
 
 IPersonaEditorModeManager::~IPersonaEditorModeManager()
@@ -77,4 +85,10 @@ void IPersonaEditorModeManager::AddReferencedObjects(FReferenceCollector& Collec
 {
 	FAssetEditorModeManager::AddReferencedObjects(Collector);
 	Collector.AddReferencedObject(PersonaModeManagerContext);
+	Collector.AddReferencedObject(SelectionSet);
+}
+
+UTypedElementSelectionSet* IPersonaEditorModeManager::GetEditorSelectionSet() const
+{
+	return SelectionSet.Get();
 }
diff --git a/Engine/Source/Editor/Persona/Private/PersonaAssetEditorToolkit.cpp b/Engine/Source/Editor/Persona/Private/PersonaAssetEditorToolkit.cpp
--- a/Engine/Source/Editor/Persona/Private/PersonaAssetEditorToolkit.cpp
+++ b/Engine/Source/Editor/Persona/Private/PersonaAssetEditorToolkit.cpp
@@ -4,8 +4,11 @@
 
 #include "PersonaModule.h"
 #include "AssetEditorModeManager.h"
+#include "Elements/Framework/EngineElementsLibrary.h"
+#include "Elements/Framework/TypedElementSelectionSet.h"
 #include "Modules/ModuleManager.h"
 #include "IPersonaEditorModeManager.h"
+#include "ISkeletonTreeItem.h"
 #include "PersonaTabs.h"
 
 void FPersonaAssetEditorToolkit::CreateEditorModeManager()
@@ -15,6 +18,30 @@ void FPersonaAssetEditorToolkit::CreateEditorModeManager()
 	// Make sure we get told when the editor mode changes so we can switch to the appropriate tab
 	// if there's a toolbox available.
 	GetEditorModeManager().OnEditorModeIDChanged().AddSP(this, &FPersonaAssetEditorToolkit::OnEditorModeIdChanged);
+
+	if (UTypedElementSelectionSet* SelectionSet = GetEditorModeManager().GetEditorSelectionSet())
+	{
+		SelectionSetChangedHandle = SelectionSet->OnChanged().AddSP(this, &FPersonaAssetEditorToolkit::OnSelectionSetChanged);
+	}
+}
+
+void FPersonaAssetEditorToolkit::OnClose()
+{
+	if (SelectionSetChangedHandle.IsValid())
+	{
+		if (UTypedElementSelectionSet* SelectionSet = GetEditorModeManager().GetEditorSelectionSet())
+		{
+			SelectionSet->OnChanged().Remove(SelectionSetChangedHandle);
+		}
+		SelectionSetChangedHandle.Reset();
+	}
+
+	FWorkflowCentricApplication::OnClose();
+}
+
+void FPersonaAssetEditorToolkit::OnSelectionSetChanged(const UTypedElementSelectionSet* InSelectionSet)
+{
+	// Base implementation is empty. Subclasses override to update details views, etc.
 }
 
 void FPersonaAssetEditorToolkit::OnToolkitHostingStarted(const TSharedRef<IToolkit>& Toolkit)
@@ -58,3 +85,32 @@ void FPersonaAssetEditorToolkit::OnEditorModeIdChanged(const FEditorModeID& Mode
 		}
 	}
 }
+
+void FPersonaAssetEditorToolkit::RouteTreeSelectionToModeManager(const TConstArrayView<TSharedPtr<ISkeletonTreeItem>>& InItems)
+{
+	// @note: this is one-way. Selecting a TypedElement will not propagate to the SkeletonTree, it may be cleaner to implement an adapter
+	// (see: USelection::SetElementSelectionSet and associated plumbing)
+
+	UTypedElementSelectionSet* SelectionSet = GetEditorModeManager().GetEditorSelectionSet();
+	if (!SelectionSet)
+	{
+		return;
+	}
+
+	FTypedElementSelectionOptions SelectionOptions;
+	SelectionOptions.SetAllowLegacyNotifications(false);
+
+	SelectionSet->ClearSelection(SelectionOptions);
+
+	for (const TSharedPtr<ISkeletonTreeItem>& Item : InItems)
+	{
+		if (const UObject* Object = Item->GetObject())
+		{
+			FTypedElementHandle ElementHandle = UEngineElementsLibrary::AcquireEditorObjectElementHandle(Object);
+			if (ElementHandle.IsSet())
+			{
+				SelectionSet->SelectElement(ElementHandle, SelectionOptions);
+			}
+		}
+	}
+}
diff --git a/Engine/Source/Editor/Persona/Public/IPersonaEditorModeManager.h b/Engine/Source/Editor/Persona/Public/IPersonaEditorModeManager.h
--- a/Engine/Source/Editor/Persona/Public/IPersonaEditorModeManager.h
+++ b/Engine/Source/Editor/Persona/Public/IPersonaEditorModeManager.h
@@ -63,8 +63,11 @@ public:
 	IPersonaEditorModeManager();
 	virtual ~IPersonaEditorModeManager() override;
 	virtual void AddReferencedObjects(FReferenceCollector& Collector) override;
+	virtual UTypedElementSelectionSet* GetEditorSelectionSet() const override;
+
 private:
 	TObjectPtr<UPersonaEditorModeManagerContext> PersonaModeManagerContext{UPersonaEditorModeManagerContext::CreateFor(this)};
+	TObjectPtr<UTypedElementSelectionSet> SelectionSet;
 };
 
 #undef UE_API
diff --git a/Engine/Source/Editor/Persona/Public/PersonaAssetEditorToolkit.h b/Engine/Source/Editor/Persona/Public/PersonaAssetEditorToolkit.h
--- a/Engine/Source/Editor/Persona/Public/PersonaAssetEditorToolkit.h
+++ b/Engine/Source/Editor/Persona/Public/PersonaAssetEditorToolkit.h
@@ -6,6 +6,9 @@
 
 #define UE_API PERSONA_API
 
+class ISkeletonTreeItem;
+class UTypedElementSelectionSet;
+
 DECLARE_MULTICAST_DELEGATE_OneParam(FOnAttachToolkit, const TSharedRef<IToolkit>&);
 DECLARE_MULTICAST_DELEGATE_OneParam(FOnDetachToolkit, const TSharedRef<IToolkit>&);
 
@@ -25,16 +28,25 @@ public:
 	FOnAttachToolkit& GetOnAttachToolkit() { return OnAttachToolkit; }
 	FOnDetachToolkit& GetOnDetachToolkit() { return OnDetachToolkit; }
 	
+	UE_API virtual void OnClose() override;
+
 protected:
 	UE_API virtual void OnEditorModeIdChanged(const FEditorModeID& ModeChangedID, bool bIsEnteringMode);
 	
+	UE_API virtual void RouteTreeSelectionToModeManager(const TConstArrayView<TSharedPtr<ISkeletonTreeItem>>& InItems);
+
+	/** Called when the editor's TypedElementSelectionSet changes. Override to update details views, etc. */
+	UE_API virtual void OnSelectionSetChanged(const UTypedElementSelectionSet* InSelectionSet);
+
 protected:
 	FOnAttachToolkit OnAttachToolkit;
 	FOnDetachToolkit OnDetachToolkit;
 	
 	// The toolkit we're currently hosting.
 	TSharedPtr<IToolkit> HostedToolkit;
-		
+
+private:
+	FDelegateHandle SelectionSetChangedHandle;
 };
 
 #undef UE_API
diff --git a/Engine/Source/Editor/PhysicsAssetEditor/PhysicsAssetEditor.Build.cs b/Engine/Source/Editor/PhysicsAssetEditor/PhysicsAssetEditor.Build.cs
--- a/Engine/Source/Editor/PhysicsAssetEditor/PhysicsAssetEditor.Build.cs
+++ b/Engine/Source/Editor/PhysicsAssetEditor/PhysicsAssetEditor.Build.cs
@@ -42,6 +42,7 @@ public class PhysicsAssetEditor : ModuleRules
 				"ApplicationCore",
 				"EditorStyle",
 				"ToolWidgets",
+				"TypedElementRuntime"
 			}
         );
 		SetupModulePhysicsSupport(Target);
diff --git a/Engine/Source/Editor/PhysicsAssetEditor/Private/PhysicsAssetEditor.cpp b/Engine/Source/Editor/PhysicsAssetEditor/Private/PhysicsAssetEditor.cpp
--- a/Engine/Source/Editor/PhysicsAssetEditor/Private/PhysicsAssetEditor.cpp
+++ b/Engine/Source/Editor/PhysicsAssetEditor/Private/PhysicsAssetEditor.cpp
@@ -4493,8 +4493,11 @@ void FPhysicsAssetEditor::HandleSelectionChanged(const TArrayView<TSharedPtr<ISk
 			PhysAssetProperties->SetObjects(Objects);
 		}
 
-		// Only a user selection should change other view's selections
-		if (InSelectInfo != ESelectInfo::Direct)
+		// Update SharedData selection unconditionally so that
+		// FPhysicsAssetEditorEditMode::ShouldDrawWidget() returns true for all
+		// selection sources (tree, viewport, programmatic).  Without this the
+		// gizmo won't appear on the first tree selection because SharedData has
+		// no selected body/constraint when the selection comes from ESelectInfo::Direct.
 		{
 			// Block selection broadcast until we have selected all, as this can be an expensive operation
 			FScopedBulkSelection BulkSelection(SharedData);
@@ -4542,7 +4545,12 @@ void FPhysicsAssetEditor::HandleSelectionChanged(const TArrayView<TSharedPtr<ISk
 				GetPersonaToolkit()->GetPreviewScene()->ClearSelectedBone();
 			}
 
-			if (PhysicsAssetGraph.IsValid())
+			// Route selection to mode manager unconditionally so the gizmo framework
+			// sees the current selection for all selection sources (tree, viewport, programmatic).
+			RouteTreeSelectionToModeManager(InSelectedItems);
+
+			// Only a user selection should update secondary views (graph, etc.)
+			if (InSelectInfo != ESelectInfo::Direct && PhysicsAssetGraph.IsValid())
 			{
 				TSet<USkeletalBodySetup*> Bodies;
 				TSet<UPhysicsConstraintTemplate*> Constraints;
diff --git a/Engine/Source/Editor/SkeletalMeshEditor/Private/SkeletalMeshEditor.cpp b/Engine/Source/Editor/SkeletalMeshEditor/Private/SkeletalMeshEditor.cpp
--- a/Engine/Source/Editor/SkeletalMeshEditor/Private/SkeletalMeshEditor.cpp
+++ b/Engine/Source/Editor/SkeletalMeshEditor/Private/SkeletalMeshEditor.cpp
@@ -1279,6 +1279,15 @@ void FSkeletalMeshEditor::HandleObjectsSelected(const TArray<UObject*>& InObject
 	}
 }
 
+void FSkeletalMeshEditor::OnSelectionSetChanged(const UTypedElementSelectionSet* InSelectionSet)
+{
+	TArray<UObject*> SelectedObjects = InSelectionSet->GetSelectedObjects();
+	if (DetailsView.IsValid())
+	{
+		DetailsView->SetObjects(SelectedObjects);
+	}
+}
+
 void FSkeletalMeshEditor::HandleObjectSelected(UObject* InObject)
 {
 	if (DetailsView.IsValid())
@@ -1289,6 +1298,8 @@ void FSkeletalMeshEditor::HandleObjectSelected(UObject* InObject)
 
 void FSkeletalMeshEditor::HandleSelectionChanged(const TArrayView<TSharedPtr<ISkeletonTreeItem>>& InSelectedItems, ESelectInfo::Type InSelectInfo)
 {
+	RouteTreeSelectionToModeManager(InSelectedItems);
+
 	if (DetailsView.IsValid())
 	{
 		TArray<UObject*> Objects;
diff --git a/Engine/Source/Editor/SkeletalMeshEditor/Private/SkeletalMeshEditor.h b/Engine/Source/Editor/SkeletalMeshEditor/Private/SkeletalMeshEditor.h
--- a/Engine/Source/Editor/SkeletalMeshEditor/Private/SkeletalMeshEditor.h
+++ b/Engine/Source/Editor/SkeletalMeshEditor/Private/SkeletalMeshEditor.h
@@ -128,6 +128,8 @@ public:
 
 	void HandleObjectsSelected(const TArray<UObject*>& InObjects);
 
+	virtual void OnSelectionSetChanged(const UTypedElementSelectionSet* InSelectionSet) override;
+
 	// Returns the currently hosted toolkit. Can be invalid if no toolkit is being hosted.
 	TSharedPtr<IToolkit> GetHostedToolkit() const { return HostedToolkit; }
 	
diff --git a/Engine/Source/Editor/SkeletalMeshEditor/SkeletalMeshEditor.Build.cs b/Engine/Source/Editor/SkeletalMeshEditor/SkeletalMeshEditor.Build.cs
--- a/Engine/Source/Editor/SkeletalMeshEditor/SkeletalMeshEditor.Build.cs
+++ b/Engine/Source/Editor/SkeletalMeshEditor/SkeletalMeshEditor.Build.cs
@@ -43,7 +43,8 @@ public class SkeletalMeshEditor : ModuleRules
 				"PhysicsUtilities",
 				"InterchangeCore",
 				"InterchangeEngine",
-				"WorkspaceMenuStructure"
+				"WorkspaceMenuStructure",
+				"TypedElementRuntime"
 			}
 		);
 	}
diff --git a/Engine/Source/Editor/SkeletonEditor/Private/SkeletonEditor.cpp b/Engine/Source/Editor/SkeletonEditor/Private/SkeletonEditor.cpp
--- a/Engine/Source/Editor/SkeletonEditor/Private/SkeletonEditor.cpp
+++ b/Engine/Source/Editor/SkeletonEditor/Private/SkeletonEditor.cpp
@@ -329,6 +329,8 @@ void FSkeletonEditor::HandleObjectSelected(UObject* InObject)
 
 void FSkeletonEditor::HandleSelectionChanged(const TArrayView<TSharedPtr<ISkeletonTreeItem>>& InSelectedItems, ESelectInfo::Type InSelectInfo)
 {
+	RouteTreeSelectionToModeManager(InSelectedItems);
+
 	if (DetailsView.IsValid())
 	{
 		TArray<UObject*> Objects;
