Persona Editor SelectionSet sync selective clearing

#Synced-CL 52103551

diff --git a/Engine/Source/Editor/Persona/Private/PersonaAssetEditorToolkit.cpp b/Engine/Source/Editor/Persona/Private/PersonaAssetEditorToolkit.cpp
--- a/Engine/Source/Editor/Persona/Private/PersonaAssetEditorToolkit.cpp
+++ b/Engine/Source/Editor/Persona/Private/PersonaAssetEditorToolkit.cpp
@@ -100,7 +100,13 @@ void FPersonaAssetEditorToolkit::RouteTreeSelectionToModeManager(const TConstArr
 	FTypedElementSelectionOptions SelectionOptions;
 	SelectionOptions.SetAllowLegacyNotifications(false);
 
-	SelectionSet->ClearSelection(SelectionOptions);
+	// Only deselect elements that we previously routed, rather than clearing the entire SelectionSet.
+	for (const FTypedElementHandle& Handle : RoutedSelectionHandles)
+	{
+		SelectionSet->DeselectElement(Handle, SelectionOptions);
+	}
+
+	RoutedSelectionHandles.Reset();
 
 	for (const TSharedPtr<ISkeletonTreeItem>& Item : InItems)
 	{
@@ -110,6 +116,7 @@ void FPersonaAssetEditorToolkit::RouteTreeSelectionToModeManager(const TConstArr
 			if (ElementHandle.IsSet())
 			{
 				SelectionSet->SelectElement(ElementHandle, SelectionOptions);
+				RoutedSelectionHandles.Add(ElementHandle);
 			}
 		}
 	}
diff --git a/Engine/Source/Editor/Persona/Public/PersonaAssetEditorToolkit.h b/Engine/Source/Editor/Persona/Public/PersonaAssetEditorToolkit.h
--- a/Engine/Source/Editor/Persona/Public/PersonaAssetEditorToolkit.h
+++ b/Engine/Source/Editor/Persona/Public/PersonaAssetEditorToolkit.h
@@ -2,6 +2,7 @@
 
 #pragma once
 
+#include "Elements/Framework/TypedElementHandle.h"
 #include "WorkflowOrientedApp/WorkflowCentricApplication.h"
 
 #define UE_API PERSONA_API
@@ -47,6 +48,9 @@ protected:
 
 private:
 	FDelegateHandle SelectionSetChangedHandle;
+
+	/** Element handles routed by the last call to RouteTreeSelectionToModeManager, so we only deselect our own elements. */
+	TArray<FTypedElementHandle> RoutedSelectionHandles;
 };
 
 #undef UE_API
